#UE4
Don't clear GRenderThreadId until stats have been flushed

diff --git a/Engine/Source/Runtime/RenderCore/Private/RenderingThread.cpp b/Engine/Source/Runtime/RenderCore/Private/RenderingThread.cpp
--- a/Engine/Source/Runtime/RenderCore/Private/RenderingThread.cpp
+++ b/Engine/Source/Runtime/RenderCore/Private/RenderingThread.cpp
@@ -347,11 +347,11 @@ public:
 		{
 			RenderingThreadMain( TaskGraphBoundSyncEvent );
 		}
-		GRenderThreadId = 0;
 #if STATS
 		FThreadStats::ExplicitFlush();
 		FThreadStats::Shutdown();
-#endif
+#endif		
+		GRenderThreadId = 0;
 
 		return 0;
 	}
