UE4.Core - Removed numbering of the rendering threads, subsequent rendering thread can be detected by comparing the thread id

diff --git a/Engine/Source/Runtime/RenderCore/Private/RenderingThread.cpp b/Engine/Source/Runtime/RenderCore/Private/RenderingThread.cpp
--- a/Engine/Source/Runtime/RenderCore/Private/RenderingThread.cpp
+++ b/Engine/Source/Runtime/RenderCore/Private/RenderingThread.cpp
@@ -453,7 +453,7 @@ void StartRenderingThread()
 
 	EThreadPriority RenderingThreadPrio = TPri_Normal;
 
-	GRenderingThread = FRunnableThread::Create(GRenderingThreadRunnable, *FString::Printf(TEXT("RenderingThread %d"), ThreadCount), 0, 0, 0, RenderingThreadPrio);
+	GRenderingThread = FRunnableThread::Create(GRenderingThreadRunnable, *FName(NAME_RenderThread).ToString(), 0, 0, 0, RenderingThreadPrio);
 
 	// Wait for render thread to have taskgraph bound before we dispatch any tasks for it.
 	((FRenderingThread*)GRenderingThreadRunnable)->TaskGraphBoundSyncEvent->Wait();
