Stats - Fix for STATS=0

diff --git a/Engine/Source/Runtime/Core/Public/Stats/StatsData.h b/Engine/Source/Runtime/Core/Public/Stats/StatsData.h
--- a/Engine/Source/Runtime/Core/Public/Stats/StatsData.h
+++ b/Engine/Source/Runtime/Core/Public/Stats/StatsData.h
@@ -604,11 +604,6 @@ struct CORE_API FStatsUtils
 	/** Internal use, converts arbitrary string to and from an escaped notation for storage in an FName. **/
 	static FString ToEscapedFString(const TCHAR* Source);
 	static FString FromEscapedFString(const TCHAR* Escaped);
-
-	static FString BuildRenderThreadName( uint32 ThreadIndex )
-	{
-		return FString::Printf( TEXT( "%s %u" ), *FName(NAME_RenderThread).GetPlainNameString(), ThreadIndex );
-	}
 	
 	static FString BuildUniqueThreadName( uint32 InThreadID )
 	{
diff --git a/Engine/Source/Runtime/RenderCore/Private/RenderingThread.cpp b/Engine/Source/Runtime/RenderCore/Private/RenderingThread.cpp
--- a/Engine/Source/Runtime/RenderCore/Private/RenderingThread.cpp
+++ b/Engine/Source/Runtime/RenderCore/Private/RenderingThread.cpp
@@ -441,6 +441,11 @@ struct FConsoleRenderThreadPropagation : public IConsoleThreadPropagation
 
 };
 
+static FString BuildRenderingThreadName( uint32 ThreadIndex )
+{
+	return FString::Printf( TEXT( "%s %u" ), *FName( NAME_RenderThread ).GetPlainNameString(), ThreadIndex );
+}
+
 void StartRenderingThread()
 {
 	static uint32 ThreadCount = 0;
@@ -454,7 +459,7 @@ void StartRenderingThread()
 
 	EThreadPriority RenderingThreadPrio = TPri_Normal;
 
-	GRenderingThread = FRunnableThread::Create( GRenderingThreadRunnable, *FStatsUtils::BuildRenderThreadName( ThreadCount ), 0, 0, 0, RenderingThreadPrio );
+	GRenderingThread = FRunnableThread::Create( GRenderingThreadRunnable, *BuildRenderingThreadName( ThreadCount ), 0, 0, 0, RenderingThreadPrio );
 
 	// Wait for render thread to have taskgraph bound before we dispatch any tasks for it.
 	((FRenderingThread*)GRenderingThreadRunnable)->TaskGraphBoundSyncEvent->Wait();
