Editor: Fix a logic error in busy-waiting for the game thread to finish executing when a breakpoint is hit

diff --git a/Engine/Source/Runtime/RenderCore/Private/RenderingThread.cpp b/Engine/Source/Runtime/RenderCore/Private/RenderingThread.cpp
--- a/Engine/Source/Runtime/RenderCore/Private/RenderingThread.cpp
+++ b/Engine/Source/Runtime/RenderCore/Private/RenderingThread.cpp
@@ -522,11 +522,11 @@ void StopRenderingThread()
 
 			FGraphEventRef QuitTask = TGraphTask<FReturnGraphTask>::CreateTask(NULL, ENamedThreads::GameThread).ConstructAndDispatchWhenReady(ENamedThreads::RenderThread);
 
-			// Busy wait while Kismet debugging, to avoid opportunistic execution of game thread tasks
+			// Busy wait while BP debugging, to avoid opportunistic execution of game thread tasks
 			// If the game thread is already executing tasks, then we have no choice but to spin
-			if (GIntraFrameDebuggingGameThread || FTaskGraphInterface::Get().IsThreadProcessingTasks(ENamedThreads::GameThread) ) 
+			if (GIntraFrameDebuggingGameThread || FTaskGraphInterface::Get().IsThreadProcessingTasks(ENamedThreads::GameThread)) 
 			{
-				while ((QuitTask.GetReference() != NULL) && QuitTask->IsComplete())
+				while ((QuitTask.GetReference() != nullptr) && !QuitTask->IsComplete())
 				{
 					FPlatformProcess::Sleep(0.0f);
 				}
