Refactoring thread affinity settings.
There was a bug in setting affinity of a thread that assumed affinity from lookup table with key being a thread name. When names was appended with consecutive numbers (e.g. "RenderingThread 1") the mechanism failed. Refactored this to use special static consts describing affinity override'able by different platforms for different affinity types + possibility of setting affinity per thread.

diff --git a/Engine/Plugins/Messaging/UdpMessaging/Source/UdpMessaging/Private/Transport/UdpMessageBeacon.cpp b/Engine/Plugins/Messaging/UdpMessaging/Source/UdpMessaging/Private/Transport/UdpMessageBeacon.cpp
--- a/Engine/Plugins/Messaging/UdpMessaging/Source/UdpMessaging/Private/Transport/UdpMessageBeacon.cpp
+++ b/Engine/Plugins/Messaging/UdpMessaging/Source/UdpMessaging/Private/Transport/UdpMessageBeacon.cpp
@@ -28,7 +28,7 @@ FUdpMessageBeacon::FUdpMessageBeacon(FSocket* InSocket, const FGuid& InSocketId,
 	EndpointLeftEvent = FPlatformProcess::CreateSynchEvent();
 	MulticastAddress = InMulticastEndpoint.ToInternetAddr();
 
-	Thread = FRunnableThread::Create(this, TEXT("FUdpMessageBeacon"), 128 * 1024, TPri_AboveNormal);
+	Thread = FRunnableThread::Create(this, TEXT("FUdpMessageBeacon"), 128 * 1024, TPri_AboveNormal, FPlatformAffinity::GetPoolThreadMask());
 }
 
 
diff --git a/Engine/Plugins/Messaging/UdpMessaging/Source/UdpMessaging/Private/Transport/UdpMessageProcessor.cpp b/Engine/Plugins/Messaging/UdpMessaging/Source/UdpMessaging/Private/Transport/UdpMessageProcessor.cpp
--- a/Engine/Plugins/Messaging/UdpMessaging/Source/UdpMessaging/Private/Transport/UdpMessageProcessor.cpp
+++ b/Engine/Plugins/Messaging/UdpMessaging/Source/UdpMessaging/Private/Transport/UdpMessageProcessor.cpp
@@ -26,7 +26,7 @@ FUdpMessageProcessor::FUdpMessageProcessor( FSocket* InSocket, const FGuid& InNo
 	, Stopping(false)
 {
 	WorkEvent = FPlatformProcess::CreateSynchEvent();
-	Thread = FRunnableThread::Create(this, TEXT("FUdpMessageProcessor"), 128 * 1024, TPri_AboveNormal);
+	Thread = FRunnableThread::Create(this, TEXT("FUdpMessageProcessor"), 128 * 1024, TPri_AboveNormal, FPlatformAffinity::GetPoolThreadMask());
 
 	const UUdpMessagingSettings& Settings = *GetDefault<UUdpMessagingSettings>();
 
diff --git a/Engine/Source/Developer/DerivedDataCache/Private/DDCCleanup.cpp b/Engine/Source/Developer/DerivedDataCache/Private/DDCCleanup.cpp
--- a/Engine/Source/Developer/DerivedDataCache/Private/DDCCleanup.cpp
+++ b/Engine/Source/Developer/DerivedDataCache/Private/DDCCleanup.cpp
@@ -55,7 +55,7 @@ FDDCCleanup::FDDCCleanup()
 	, StopTaskCounter(0)
 {
 	// Don't delete the runnable automatically. It's going to be manually deleted in FDDCCleanup::Shutdown.
-	Thread = FRunnableThread::Create(this, TEXT("FDDCCleanup"), 0, TPri_BelowNormal);
+	Thread = FRunnableThread::Create(this, TEXT("FDDCCleanup"), 0, TPri_BelowNormal, FPlatformAffinity::GetPoolThreadMask());
 }
 
 FDDCCleanup::~FDDCCleanup()
diff --git a/Engine/Source/Runtime/Core/Private/Async/TaskGraph.cpp b/Engine/Source/Runtime/Core/Private/Async/TaskGraph.cpp
--- a/Engine/Source/Runtime/Core/Private/Async/TaskGraph.cpp
+++ b/Engine/Source/Runtime/Core/Private/Async/TaskGraph.cpp
@@ -675,7 +675,7 @@ public:
 		{
 			FString Name = FString::Printf(TEXT("TaskGraphThread %d"), ThreadIndex - (LastExternalThread + 1));
 			uint32 StackSize = 256 * 1024;
-			WorkerThreads[ThreadIndex].RunnableThread = FRunnableThread::Create(&Thread(ThreadIndex), *Name, StackSize, TPri_Normal); // these are below normal threads? so that they sleep when the named threads are active
+			WorkerThreads[ThreadIndex].RunnableThread = FRunnableThread::Create(&Thread(ThreadIndex), *Name, StackSize, TPri_Normal, FPlatformAffinity::GetTaskGraphThreadMask()); // these are below normal threads? so that they sleep when the named threads are active
 			WorkerThreads[ThreadIndex].bAttached = true;
 		}
 	}
diff --git a/Engine/Source/Runtime/Core/Private/HAL/AffinityManager.cpp b/Engine/Source/Runtime/Core/Private/HAL/AffinityManager.cpp
deleted file mode 100644
diff --git a/Engine/Source/Runtime/Core/Private/HAL/PThreadRunnableThread.h b/Engine/Source/Runtime/Core/Private/HAL/PThreadRunnableThread.h
--- a/Engine/Source/Runtime/Core/Private/HAL/PThreadRunnableThread.h
+++ b/Engine/Source/Runtime/Core/Private/HAL/PThreadRunnableThread.h
@@ -323,10 +323,6 @@ public:
 		return ThreadName;
 	}
 
-	virtual void SetThreadAffinityMask(uint64 AffinityMask) OVERRIDE
-	{
-	}
-
 protected:
 
 	virtual bool CreateInternal(FRunnable* InRunnable, const TCHAR* InThreadName,
@@ -354,7 +350,7 @@ protected:
 			SetThreadPriority(InThreadPri);
 
 			// set the affinity
-			SetThreadAffinityMask(InThreadAffinityMask);
+			FPlatformProcess::SetThreadAffinityMask( InThreadAffinityMask );
 		}
 		else // If it fails, clear all the vars
 		{
@@ -367,4 +363,3 @@ protected:
 		return Thread != PTHREAD_NULL;
 	}
 };
-
diff --git a/Engine/Source/Runtime/Core/Private/HAL/ThreadingBase.cpp b/Engine/Source/Runtime/Core/Private/HAL/ThreadingBase.cpp
--- a/Engine/Source/Runtime/Core/Private/HAL/ThreadingBase.cpp
+++ b/Engine/Source/Runtime/Core/Private/HAL/ThreadingBase.cpp
@@ -6,7 +6,7 @@
 #include "StatsData.h"
 
 /** The global thread pool */
-FQueuedThreadPool*		GThreadPool = NULL;
+FQueuedThreadPool*		GThreadPool						= NULL;
 
 CORE_API bool IsInGameThread()
 {
@@ -120,16 +120,6 @@ public:
 		}		
 		return Runnable != NULL;
 	}
-
-	/**	 
-	 * Change the thread processor affinity
-	 *
-	 * @param AffinityMask A bitfield indicating what processors the thread is allowed to run on
-	 */
-	virtual void SetThreadAffinityMask(uint64 AffinityMask)
-	{
-
-	}
 };
 uint32 FFakeThread::ThreadIdCounter = 0xffff;
 
@@ -201,9 +191,9 @@ FRunnableThread::~FRunnableThread()
 
 FRunnableThread* FRunnableThread::Create(
 	class FRunnable* InRunnable, 
-	const TCHAR* ThreadName, 
+	const TCHAR* ThreadName,
 	uint32 InStackSize,
-	EThreadPriority InThreadPri,
+	EThreadPriority InThreadPri, 
 	uint64 InThreadAffinityMask)
 {
 	FRunnableThread* NewThread = NULL;
@@ -214,11 +204,6 @@ FRunnableThread* FRunnableThread::Create(
 		NewThread = FPlatformProcess::CreateRunnableThread();
 		if (NewThread)
 		{
-			if( !InThreadAffinityMask )
-			{
-				InThreadAffinityMask = AffinityManagerGetAffinity( ThreadName );
-			}
-
 			// Call the thread's create method
 			if (NewThread->CreateInternal(InRunnable,ThreadName,InStackSize,InThreadPri,InThreadAffinityMask) == false)
 			{
@@ -341,7 +326,7 @@ public:
 
 		OwningThreadPool = InPool;
 		DoWorkEvent = FPlatformProcess::CreateSynchEvent();
-		Thread = FRunnableThread::Create(this, *PoolThreadName, InStackSize, ThreadPriority);
+		Thread = FRunnableThread::Create(this, *PoolThreadName, InStackSize, ThreadPriority, FPlatformAffinity::GetPoolThreadMask());
 		check(Thread);
 		return true;
 	}
@@ -598,4 +583,3 @@ FQueuedThreadPool* FQueuedThreadPool::Allocate()
 {
 	return new FQueuedThreadPoolBase;
 }
-
diff --git a/Engine/Source/Runtime/Core/Private/Serialization/AsyncIOSystemBase.cpp b/Engine/Source/Runtime/Core/Private/Serialization/AsyncIOSystemBase.cpp
--- a/Engine/Source/Runtime/Core/Private/Serialization/AsyncIOSystemBase.cpp
+++ b/Engine/Source/Runtime/Core/Private/Serialization/AsyncIOSystemBase.cpp
@@ -762,7 +762,7 @@ FIOSystem& FIOSystem::Get()
 			// the platform didn't have a specific need, so we just use the base class with the normal file system.
 			AsyncIOSystem = new FAsyncIOSystemBase(FPlatformFileManager::Get().GetPlatformFile());
 		}
-		AsyncIOThread = FRunnableThread::Create(AsyncIOSystem, TEXT("AsyncIOSystem"), 16384, TPri_AboveNormal);
+		AsyncIOThread = FRunnableThread::Create(AsyncIOSystem, TEXT("AsyncIOSystem"), 16384, TPri_AboveNormal, FPlatformAffinity::GetPoolThreadMask());
 		check(AsyncIOThread);
 	}
 	check(AsyncIOSystem);
diff --git a/Engine/Source/Runtime/Core/Private/Stats/Stats2.cpp b/Engine/Source/Runtime/Core/Private/Stats/Stats2.cpp
--- a/Engine/Source/Runtime/Core/Private/Stats/Stats2.cpp
+++ b/Engine/Source/Runtime/Core/Private/Stats/Stats2.cpp
@@ -648,7 +648,7 @@ public:
 
 	void Start()
 	{
-		Thread = FRunnableThread::Create(this, TEXT("StatsThread"), 128 * 1024, TPri_BelowNormal);
+		Thread = FRunnableThread::Create(this, TEXT("StatsThread"), 128 * 1024, TPri_BelowNormal, FPlatformAffinity::GetStatsThreadMask());
 		check(Thread != NULL);
 	}
 
diff --git a/Engine/Source/Runtime/Core/Private/WinRT/WinRTRunnableThread.h b/Engine/Source/Runtime/Core/Private/WinRT/WinRTRunnableThread.h
--- a/Engine/Source/Runtime/Core/Private/WinRT/WinRTRunnableThread.h
+++ b/Engine/Source/Runtime/Core/Private/WinRT/WinRTRunnableThread.h
@@ -158,11 +158,6 @@ public:
 		}
 	}
 
-	virtual void SetThreadAffinityMask(uint64 AffinityMask) OVERRIDE
-	{
-		// This probably needs implementing - added to fix the broken build
-	}
-
 	virtual void Suspend(bool bShouldPause = 1) OVERRIDE
 	{
 		check(Thread);
@@ -258,4 +253,3 @@ protected:
 		return Thread != NULL;
 	}
 };
-
diff --git a/Engine/Source/Runtime/Core/Private/Windows/WindowsPlatformProcess.cpp b/Engine/Source/Runtime/Core/Private/Windows/WindowsPlatformProcess.cpp
--- a/Engine/Source/Runtime/Core/Private/Windows/WindowsPlatformProcess.cpp
+++ b/Engine/Source/Runtime/Core/Private/Windows/WindowsPlatformProcess.cpp
@@ -369,7 +369,10 @@ uint32 FWindowsPlatformProcess::GetCurrentProcessId()
 
 void FWindowsPlatformProcess::SetThreadAffinityMask( uint64 AffinityMask )
 {
-	::SetThreadAffinityMask(::GetCurrentThread(), (DWORD_PTR)AffinityMask);
+	if( AffinityMask != FPlatformAffinity::GetNoAffinityMask() )
+	{
+		::SetThreadAffinityMask( ::GetCurrentThread(), (DWORD_PTR)AffinityMask );
+	}
 }
 
 bool FWindowsPlatformProcess::GetProcReturnCode( FProcHandle & ProcHandle, int32* ReturnCode )
diff --git a/Engine/Source/Runtime/Core/Private/Windows/WindowsRunnableThread.cpp b/Engine/Source/Runtime/Core/Private/Windows/WindowsRunnableThread.cpp
--- a/Engine/Source/Runtime/Core/Private/Windows/WindowsRunnableThread.cpp
+++ b/Engine/Source/Runtime/Core/Private/Windows/WindowsRunnableThread.cpp
@@ -9,13 +9,8 @@ uint32 FRunnableThreadWin::GuardedRun()
 {
 	uint32 ExitCode = 1;
 
-#if PLATFORM_XBOXONE
-	if( ThreadAffintyMask )
-	{
-		SetThreadAffinityMask(ThreadAffintyMask);
-	}
+	FPlatformProcess::SetThreadAffinityMask(ThreadAffintyMask);
 	UE_LOG(LogThreadingWindows, Log, TEXT("Runnable thread %s is on Process %d."), *ThreadName  , static_cast<uint32>(::GetCurrentProcessorNumber()) );
-#endif
 
 #if !PLATFORM_SEH_EXCEPTIONS_DISABLED
 	if( !FPlatformMisc::IsDebuggerPresent() || GAlwaysReportCrash )
@@ -56,14 +51,6 @@ uint32 FRunnableThreadWin::Run()
 	uint32 ExitCode = 1;
 	check(Runnable);
 
-#if PLATFORM_XBOXONE
-	if( ThreadAffintyMask )
-	{
-		SetThreadAffinityMask(ThreadAffintyMask);
-	}
-	UE_LOG(LogThreadingWindows, Log, TEXT("Runnable thread %s is on Process %d."), *ThreadName  , static_cast<uint32>(::GetCurrentProcessorNumber()) );
-#endif
-
 	// Initialize the runnable object
 	if (Runnable->Init() == true)
 	{
diff --git a/Engine/Source/Runtime/Core/Private/Windows/WindowsRunnableThread.h b/Engine/Source/Runtime/Core/Private/Windows/WindowsRunnableThread.h
--- a/Engine/Source/Runtime/Core/Private/Windows/WindowsRunnableThread.h
+++ b/Engine/Source/Runtime/Core/Private/Windows/WindowsRunnableThread.h
@@ -141,12 +141,6 @@ public:
 		}
 	}
 
-	virtual void SetThreadAffinityMask(uint64 AffinityMask) OVERRIDE
-	{
-		// Change the processor affinity on the thread
-		::SetThreadAffinityMask(Thread, (DWORD_PTR)AffinityMask);
-	}
-
 	virtual void Suspend(bool bShouldPause = 1) OVERRIDE
 	{
 		check(Thread);
@@ -237,4 +231,3 @@ protected:
 		return Thread != NULL;
 	}
 };
-
diff --git a/Engine/Source/Runtime/Core/Public/Android/AndroidAffinity.h b/Engine/Source/Runtime/Core/Public/Android/AndroidAffinity.h
new file mode 100644
--- /dev/null
+++ b/Engine/Source/Runtime/Core/Public/Android/AndroidAffinity.h
@@ -0,0 +1,45 @@
+// Copyright 1998-2014 Epic Games, Inc. All Rights Reserved.
+
+/*=============================================================================================
+AndroidAffinity.h: Android affinity profile masks definitions.
+==============================================================================================*/
+
+#pragma once
+
+#include "GenericPlatformAffinity.h"
+
+class FAndroidAffinity : public FGenericPlatformAffinity
+{
+public:
+	static const CORE_API uint64 GetMainGameMask() OVERRIDE
+	{
+		return MAKEAFFINITYMASK1(0);
+	}
+
+	static const CORE_API uint64 GetRenderingThreadMask() OVERRIDE
+	{
+		return MAKEAFFINITYMASK1(1);
+	}
+
+	static const CORE_API uint64 GetRTHeartBeatMask() OVERRIDE
+	{
+		return MAKEAFFINITYMASK1(4);
+	}
+
+	static const CORE_API uint64 GetPoolThreadMask() OVERRIDE
+	{
+		return MAKEAFFINITYMASK1(4);
+	}
+
+	static const CORE_API uint64 GetTaskGraphThreadMask() OVERRIDE
+	{
+		return MAKEAFFINITYMASK3(0, 2, 3);
+	}
+
+	static const CORE_API uint64 GetStatsThreadMask() OVERRIDE
+	{
+		return MAKEAFFINITYMASK1(5);
+	}
+};
+
+typedef FAndroidAffinity FPlatformAffinity;
diff --git a/Engine/Source/Runtime/Core/Public/Android/AndroidIncludes.h b/Engine/Source/Runtime/Core/Public/Android/AndroidIncludes.h
--- a/Engine/Source/Runtime/Core/Public/Android/AndroidIncludes.h
+++ b/Engine/Source/Runtime/Core/Public/Android/AndroidIncludes.h
@@ -61,6 +61,7 @@ struct RECT
 #include "Android/AndroidFile.h"
 #include "Android/AndroidSurvey.h"
 #include "Android/AndroidHttp.h"
+#include "Android/AndroidAffinity.h"
 
 // include platform properties and typedef it for the runtime
 #include "AndroidProperties.h"
diff --git a/Engine/Source/Runtime/Core/Public/Core.h b/Engine/Source/Runtime/Core/Public/Core.h
--- a/Engine/Source/Runtime/Core/Public/Core.h
+++ b/Engine/Source/Runtime/Core/Public/Core.h
@@ -426,7 +426,6 @@ extern CORE_API double GBlueprintCompileTime;
 #include "Internationalization/Culture.h"
 #include "Guid.h"						// FGuid class
 #include "NetworkGuid.h"				// FNetworkGuid class
-#include "AffinityManager.h"			// Non-platform Affinity information for threads.
 #include "UnrealMath.h"					// Vector math functions.
 #include "SHMath.h"						// SH math functions.
 #include "RandomStream.h"				// Random stream definitions.
diff --git a/Engine/Source/Runtime/Core/Public/GenericPlatform/GenericPlatformAffinity.h b/Engine/Source/Runtime/Core/Public/GenericPlatform/GenericPlatformAffinity.h
new file mode 100644
--- /dev/null
+++ b/Engine/Source/Runtime/Core/Public/GenericPlatform/GenericPlatformAffinity.h
@@ -0,0 +1,52 @@
+// Copyright 1998-2014 Epic Games, Inc. All Rights Reserved.
+
+#pragma once
+
+#include "Core.h"
+
+#define		MAKEAFFINITYMASK1(x)			(1<<x)
+#define		MAKEAFFINITYMASK2(x,y)			((1<<x)+(1<<y))
+#define		MAKEAFFINITYMASK3(x,y,z)		((1<<x)+(1<<y)+(1<<z))
+#define		MAKEAFFINITYMASK4(w,x,y,z)		((1<<w)+(1<<x)+(1<<y)+(1<<z))
+#define		MAKEAFFINITYMASK5(v,w,x,y,z)	((1<<v)+(1<<w)+(1<<x)+(1<<y)+(1<<z))
+
+class FGenericPlatformAffinity
+{
+public:
+	static const CORE_API uint64 GetMainGameMask()
+	{
+		return 0xFFFFFFFFFFFFFFFF;
+	}
+
+	static const CORE_API uint64 GetRenderingThreadMask()
+	{
+		return 0xFFFFFFFFFFFFFFFF;
+	}
+
+	static const CORE_API uint64 GetRTHeartBeatMask()
+	{
+		return 0xFFFFFFFFFFFFFFFF;
+	}
+
+	static const CORE_API uint64 GetPoolThreadMask()
+	{
+		return 0xFFFFFFFFFFFFFFFF;
+	}
+
+	static const CORE_API uint64 GetTaskGraphThreadMask()
+	{
+		return 0xFFFFFFFFFFFFFFFF;
+	}
+
+	static const CORE_API uint64 GetStatsThreadMask()
+	{
+		return 0xFFFFFFFFFFFFFFFF;
+	}
+
+	static const CORE_API uint64 GetNoAffinityMask()
+	{
+		return 0xFFFFFFFFFFFFFFFF;
+	}
+};
+
+
diff --git a/Engine/Source/Runtime/Core/Public/GenericPlatform/GenericPlatformProcess.h b/Engine/Source/Runtime/Core/Public/GenericPlatform/GenericPlatformProcess.h
--- a/Engine/Source/Runtime/Core/Public/GenericPlatform/GenericPlatformProcess.h
+++ b/Engine/Source/Runtime/Core/Public/GenericPlatform/GenericPlatformProcess.h
@@ -199,8 +199,14 @@ struct CORE_API FGenericPlatformProcess
 	 * @return the ProcessId of this process
 	 */
 	static uint32 GetCurrentProcessId();
-	/** Sets the CPU affinity mask for the current thread. */
+	
+	/**	 
+	 * Change the thread processor affinity
+	 *
+	 * @param AffinityMask A bitfield indicating what processors the thread is allowed to run on
+	 */
 	static void SetThreadAffinityMask( uint64 AffinityMask );
+	
 	/** Get startup directory.  NOTE: Only one return value is valid at a time! **/
 	static const TCHAR* BaseDir();
 	/** Get user directory.  NOTE: Only one return value is valid at a time! **/
diff --git a/Engine/Source/Runtime/Core/Public/HAL/AffinityManager.h b/Engine/Source/Runtime/Core/Public/HAL/AffinityManager.h
deleted file mode 100644
diff --git a/Engine/Source/Runtime/Core/Public/HAL/PlatformIncludes.h b/Engine/Source/Runtime/Core/Public/HAL/PlatformIncludes.h
--- a/Engine/Source/Runtime/Core/Public/HAL/PlatformIncludes.h
+++ b/Engine/Source/Runtime/Core/Public/HAL/PlatformIncludes.h
@@ -37,6 +37,7 @@
 #include "GenericPlatform/GenericPlatformFile.h"
 #include "GenericPlatform/GenericPlatformSurvey.h"
 #include "GenericPlatform/GenericPlatformHttp.h"
+#include "GenericPlatform/GenericPlatformAffinity.h"
 
 // Platform specific include.
 //@port This is where you include your "main" platform header
diff --git a/Engine/Source/Runtime/Core/Public/HAL/ThreadingBase.h b/Engine/Source/Runtime/Core/Public/HAL/ThreadingBase.h
--- a/Engine/Source/Runtime/Core/Public/HAL/ThreadingBase.h
+++ b/Engine/Source/Runtime/Core/Public/HAL/ThreadingBase.h
@@ -265,7 +265,7 @@ public:
 		const TCHAR* ThreadName,
 		uint32 InStackSize = 0,
 		EThreadPriority InThreadPri = TPri_Normal,
-		uint64 InThreadAffinityMask = 0);
+		uint64 InThreadAffinityMask = FPlatformAffinity::GetNoAffinityMask());
 
 	/**
 	 * Changes the thread priority of the currently running thread
@@ -274,13 +274,6 @@ public:
 	 */
 	virtual void SetThreadPriority (EThreadPriority NewPriority) = 0;
 
-	/**	 
-	 * Change the thread processor affinity
-	 *
-	 * @param AffinityMask A bitfield indicating what processors the thread is allowed to run on
-	 */
-	virtual void SetThreadAffinityMask(uint64 AffinityMask) = 0;
-
 	/**
 	 * Tells the thread to either pause execution or resume depending on the
 	 * passed in value.
diff --git a/Engine/Source/Runtime/Core/Public/HTML5/HTML5PlatformIncludes.h b/Engine/Source/Runtime/Core/Public/HTML5/HTML5PlatformIncludes.h
--- a/Engine/Source/Runtime/Core/Public/HTML5/HTML5PlatformIncludes.h
+++ b/Engine/Source/Runtime/Core/Public/HTML5/HTML5PlatformIncludes.h
@@ -58,6 +58,8 @@ struct RECT
 #include "HTML5/HTML5PlatformSurvey.h"
 #include "HTML5/HTML5PlatformHttp.h"
 
+typedef FGenericPlatformAffinity FPlatformAffinity;
+
 // include platform properties and typedef it for the runtime
 #include "HTML5PlatformProperties.h"
 typedef FHTML5PlatformProperties FPlatformProperties;
diff --git a/Engine/Source/Runtime/Core/Public/HTML5/HTML5PlatformRunnableThread.h b/Engine/Source/Runtime/Core/Public/HTML5/HTML5PlatformRunnableThread.h
--- a/Engine/Source/Runtime/Core/Public/HTML5/HTML5PlatformRunnableThread.h
+++ b/Engine/Source/Runtime/Core/Public/HTML5/HTML5PlatformRunnableThread.h
@@ -19,10 +19,6 @@ public:
 
 	}
 
-	virtual void SetThreadAffinityMask(uint64 AffinityMask) OVERRIDE
-	{
-	}
-
 	virtual void Suspend (bool bShouldPause = 1)
 	{
 
@@ -68,4 +64,3 @@ protected:
 		return false;
 	}
 };
-
diff --git a/Engine/Source/Runtime/Core/Public/IOS/IOSPlatformIncludes.h b/Engine/Source/Runtime/Core/Public/IOS/IOSPlatformIncludes.h
--- a/Engine/Source/Runtime/Core/Public/IOS/IOSPlatformIncludes.h
+++ b/Engine/Source/Runtime/Core/Public/IOS/IOSPlatformIncludes.h
@@ -99,6 +99,8 @@ struct RECT
 #include "IOS/IOSPlatformHttp.h"
 #import "IOSPlatformString.h"
 
+typedef FGenericPlatformAffinity FPlatformAffinity;
+
 // include platform properties and typedef it for the runtime
 
 #include "IOS/IOSPlatformProperties.h"
diff --git a/Engine/Source/Runtime/Core/Public/Linux/LinuxPlatformIncludes.h b/Engine/Source/Runtime/Core/Public/Linux/LinuxPlatformIncludes.h
--- a/Engine/Source/Runtime/Core/Public/Linux/LinuxPlatformIncludes.h
+++ b/Engine/Source/Runtime/Core/Public/Linux/LinuxPlatformIncludes.h
@@ -69,6 +69,7 @@ struct RECT
 #include "Linux/LinuxPlatformSurvey.h"
 #include "Linux/LinuxPlatformHttp.h"
 
+typedef FGenericPlatformAffinity FPlatformAffinity;
 
 // include platform properties and typedef it for the runtime
 #include "Linux/LinuxPlatformProperties.h"
diff --git a/Engine/Source/Runtime/Core/Public/Mac/MacPlatformIncludes.h b/Engine/Source/Runtime/Core/Public/Mac/MacPlatformIncludes.h
--- a/Engine/Source/Runtime/Core/Public/Mac/MacPlatformIncludes.h
+++ b/Engine/Source/Runtime/Core/Public/Mac/MacPlatformIncludes.h
@@ -92,6 +92,8 @@ struct RECT
 #include "Mac/MacPlatformSurvey.h"
 #include "Mac/MacPlatformHttp.h"
 
+typedef FGenericPlatformAffinity FPlatformAffinity;
+
 // include platform properties and typedef it for the runtime
 
 #include "Mac/MacPlatformProperties.h"
diff --git a/Engine/Source/Runtime/Core/Public/WinRT/WinRTIncludes.h b/Engine/Source/Runtime/Core/Public/WinRT/WinRTIncludes.h
--- a/Engine/Source/Runtime/Core/Public/WinRT/WinRTIncludes.h
+++ b/Engine/Source/Runtime/Core/Public/WinRT/WinRTIncludes.h
@@ -24,6 +24,8 @@
 #include "WinRT/WinRTSurvey.h"
 #include "WinRT/WinRTHttp.h"
 
+typedef FGenericPlatformAffinity FPlatformAffinity;
+
 // include platform properties and typedef it for the runtime
 #include "WinRTProperties.h"
 typedef FWinRTPlatformProperties FPlatformProperties;
diff --git a/Engine/Source/Runtime/Core/Public/Windows/WindowsPlatformIncludes.h b/Engine/Source/Runtime/Core/Public/Windows/WindowsPlatformIncludes.h
--- a/Engine/Source/Runtime/Core/Public/Windows/WindowsPlatformIncludes.h
+++ b/Engine/Source/Runtime/Core/Public/Windows/WindowsPlatformIncludes.h
@@ -56,6 +56,8 @@ using namespace vc_attributes;
 #include "Windows/WindowsPlatformSurvey.h"
 #include "Windows/WindowsPlatformHttp.h"
 
+typedef FGenericPlatformAffinity FPlatformAffinity;
+
 // include platform properties and typedef it for the runtime
 
 #include "Windows/WindowsPlatformProperties.h"
diff --git a/Engine/Source/Runtime/Engine/Private/ShaderCompiler/ShaderCompiler.cpp b/Engine/Source/Runtime/Engine/Private/ShaderCompiler/ShaderCompiler.cpp
--- a/Engine/Source/Runtime/Engine/Private/ShaderCompiler/ShaderCompiler.cpp
+++ b/Engine/Source/Runtime/Engine/Private/ShaderCompiler/ShaderCompiler.cpp
@@ -440,7 +440,7 @@ FShaderCompileThreadRunnable::FShaderCompileThreadRunnable(FShaderCompilingManag
 
 	if (Manager->bAllowAsynchronousShaderCompiling && !FPlatformProperties::RequiresCookedData())
 	{
-		Thread = FRunnableThread::Create(this, TEXT("ShaderCompilingThread"));
+		Thread = FRunnableThread::Create(this, TEXT("ShaderCompilingThread"), 0, TPri_Normal, FPlatformAffinity::GetPoolThreadMask());
 	}
 }
 
diff --git a/Engine/Source/Runtime/Engine/Private/UnrealEngine.cpp b/Engine/Source/Runtime/Engine/Private/UnrealEngine.cpp
--- a/Engine/Source/Runtime/Engine/Private/UnrealEngine.cpp
+++ b/Engine/Source/Runtime/Engine/Private/UnrealEngine.cpp
@@ -5555,7 +5555,7 @@ void UEngine::EnableScreenSaver( bool bEnable )
 				{
 					// Create thread inhibiting screen saver while it is running.
 					ScreenSaverInhibitorRunnable = new FScreenSaverInhibitor();
-					ScreenSaverInhibitor = FRunnableThread::Create(ScreenSaverInhibitorRunnable, TEXT("ScreenSaverInhibitor"), 16 * 1024);
+					ScreenSaverInhibitor = FRunnableThread::Create(ScreenSaverInhibitorRunnable, TEXT("ScreenSaverInhibitor"), 16 * 1024, TPri_Normal, FPlatformAffinity::GetPoolThreadMask());
 					// Only actually run when needed to not bypass group policies for screensaver, etc.
 					ScreenSaverInhibitor->Suspend( true );
 					ScreenSaverInhibitorSemaphore = 0;
diff --git a/Engine/Source/Runtime/Launch/Private/Android/LaunchAndroid.cpp b/Engine/Source/Runtime/Launch/Private/Android/LaunchAndroid.cpp
--- a/Engine/Source/Runtime/Launch/Private/Android/LaunchAndroid.cpp
+++ b/Engine/Source/Runtime/Launch/Private/Android/LaunchAndroid.cpp
@@ -299,8 +299,7 @@ static void* AndroidEventThreadWorker( void* param )
 {
 	struct android_app* state = (struct android_app*)param;
 
-	uint64 GameThreadAffinity = AffinityManagerGetAffinity( TEXT("MainGame"));
-	FPlatformProcess::SetThreadAffinityMask( GameThreadAffinity );
+	FPlatformProcess::SetThreadAffinityMask(FPlatformAffinity::GetMainGameMask());
 
 	FPlatformMisc::LowLevelOutputDebugString(L"Entering event processing thread engine entry point");
 
diff --git a/Engine/Source/Runtime/Launch/Private/LaunchEngineLoop.cpp b/Engine/Source/Runtime/Launch/Private/LaunchEngineLoop.cpp
--- a/Engine/Source/Runtime/Launch/Private/LaunchEngineLoop.cpp
+++ b/Engine/Source/Runtime/Launch/Private/LaunchEngineLoop.cpp
@@ -704,6 +704,8 @@ int32 FEngineLoop::PreInit( const TCHAR* CmdLine )
 	GGameThreadId = FPlatformTLS::GetCurrentThreadId();
 	GIsGameThreadIdInitialized = true;
 
+	FPlatformProcess::SetThreadAffinityMask(FPlatformAffinity::GetMainGameMask());
+
 #if PLATFORM_XBOXONE
 	HANDLE proc = ::GetCurrentProcess();
 	DWORD64 dwProcessAffinity, dwSystemAffinity;
@@ -711,16 +713,9 @@ int32 FEngineLoop::PreInit( const TCHAR* CmdLine )
 	{
 		UE_LOG(LogInit, Log, TEXT("Process affinity %ld System affinity %ld."), dwProcessAffinity , dwSystemAffinity);
 	}
-	uint64 GameThreadAffinity = AffinityManagerGetAffinity( TEXT("MainGame"));
-	FPlatformProcess::SetThreadAffinityMask( GameThreadAffinity );
 	UE_LOG(LogInit, Log, TEXT("Runnable thread Main Thread is on Process %d."), static_cast<uint32>(::GetCurrentProcessorNumber()) );	
 #endif // PLATFORM_XBOXONE
 
-#if PLATFORM_ANDROID
-	uint64 GameThreadAffinity = AffinityManagerGetAffinity( TEXT("MainGame"));
-	FPlatformProcess::SetThreadAffinityMask( GameThreadAffinity );
-#endif
-
 	// Figure out whether we're the editor, ucc or the game.
 	const SIZE_T CommandLineSize = FCString::Strlen(CmdLine)+1;
 	TCHAR* CommandLineCopy			= new TCHAR[ CommandLineSize ];
diff --git a/Engine/Source/Runtime/Messaging/Private/Bus/MessageBus.cpp b/Engine/Source/Runtime/Messaging/Private/Bus/MessageBus.cpp
--- a/Engine/Source/Runtime/Messaging/Private/Bus/MessageBus.cpp
+++ b/Engine/Source/Runtime/Messaging/Private/Bus/MessageBus.cpp
@@ -14,7 +14,7 @@ FMessageBus::FMessageBus( const IAuthorizeMessageRecipientsPtr& InRecipientAutho
 	: RecipientAuthorizer(InRecipientAuthorizer)
 {
 	Router = new FMessageRouter();
-	RouterThread = FRunnableThread::Create(Router, TEXT("FMessageBus.Router"), 128 * 1024, TPri_Normal);
+	RouterThread = FRunnableThread::Create(Router, TEXT("FMessageBus.Router"), 128 * 1024, TPri_Normal, FPlatformAffinity::GetPoolThreadMask());
 
 	check(Router != nullptr);
 }
diff --git a/Engine/Source/Runtime/Networking/Public/Common/UdpSocketReceiver.h b/Engine/Source/Runtime/Networking/Public/Common/UdpSocketReceiver.h
--- a/Engine/Source/Runtime/Networking/Public/Common/UdpSocketReceiver.h
+++ b/Engine/Source/Runtime/Networking/Public/Common/UdpSocketReceiver.h
@@ -45,7 +45,7 @@ public:
 		check(Socket->GetSocketType() == SOCKTYPE_Datagram);
 
 		SocketSubsystem = ISocketSubsystem::Get(PLATFORM_SOCKETSUBSYSTEM);
-		Thread = FRunnableThread::Create(this, ThreadDescription, 128 * 1024, TPri_AboveNormal);
+		Thread = FRunnableThread::Create(this, ThreadDescription, 128 * 1024, TPri_AboveNormal, FPlatformAffinity::GetPoolThreadMask());
 	}
 
 	/**
diff --git a/Engine/Source/Runtime/Networking/Public/Common/UdpSocketSender.h b/Engine/Source/Runtime/Networking/Public/Common/UdpSocketSender.h
--- a/Engine/Source/Runtime/Networking/Public/Common/UdpSocketSender.h
+++ b/Engine/Source/Runtime/Networking/Public/Common/UdpSocketSender.h
@@ -56,7 +56,7 @@ public:
 
 		WorkEvent = FPlatformProcess::CreateSynchEvent();
 
-		Thread = FRunnableThread::Create(this, ThreadDescription, 128 * 1024, TPri_AboveNormal);
+		Thread = FRunnableThread::Create(this, ThreadDescription, 128 * 1024, TPri_AboveNormal, FPlatformAffinity::GetPoolThreadMask());
 	}
 
 	/**
diff --git a/Engine/Source/Runtime/RenderCore/Private/RenderingThread.cpp b/Engine/Source/Runtime/RenderCore/Private/RenderingThread.cpp
--- a/Engine/Source/Runtime/RenderCore/Private/RenderingThread.cpp
+++ b/Engine/Source/Runtime/RenderCore/Private/RenderingThread.cpp
@@ -24,7 +24,7 @@ RENDERCORE_API bool GUseThreadedRendering = false;
 #if !(UE_BUILD_SHIPPING || UE_BUILD_TEST)
 	RENDERCORE_API bool GMainThreadBlockedOnRenderThread = false;
 #endif // #if !(UE_BUILD_SHIPPING || UE_BUILD_TEST)
-	
+
 static FRunnable* GRenderingThreadRunnable = NULL;
 
 /** If the rendering thread has been terminated by an unhandled exception, this contains the error message. */
@@ -191,7 +191,7 @@ void TickRenderingTickables()
 	{
 		return;
 	}
-	
+
 	// tick any rendering thread tickables
 	for (int32 ObjectIndex = 0; ObjectIndex < FTickableObjectRenderThread::RenderingThreadTickableObjects.Num(); ObjectIndex++)
 	{
@@ -302,12 +302,6 @@ public:
 	{
 		GRenderThreadId = FPlatformTLS::GetCurrentThreadId();
 
-#if PLATFORM_ANDROID
-		//@TODO: This should be using the affinitymask in FRunnableThread::Create() and tested on all platforms.
-		uint64 RenderThreadAffinityMask = AffinityManagerGetAffinity( TEXT("RenderingThread 0"));
-		FPlatformProcess::SetThreadAffinityMask( RenderThreadAffinityMask );
-#endif
-
 #if PLATFORM_WINDOWS
 		if ( !FPlatformMisc::IsDebuggerPresent() || GAlwaysReportCrash )
 		{
@@ -340,7 +334,7 @@ public:
 #if STATS
 		FThreadStats::ExplicitFlush();
 		FThreadStats::Shutdown();
-#endif		
+#endif
 		GRenderThreadId = 0;
 
 		return 0;
@@ -457,7 +451,7 @@ void StartRenderingThread()
 	// Create the rendering thread.
 	GRenderingThreadRunnable = new FRenderingThread();
 
-	GRenderingThread = FRunnableThread::Create( GRenderingThreadRunnable, *BuildRenderingThreadName( ThreadCount ));
+	GRenderingThread = FRunnableThread::Create( GRenderingThreadRunnable, *BuildRenderingThreadName( ThreadCount ), 0, TPri_Normal, FPlatformAffinity::GetRenderingThreadMask());
 
 	// Wait for render thread to have taskgraph bound before we dispatch any tasks for it.
 	((FRenderingThread*)GRenderingThreadRunnable)->TaskGraphBoundSyncEvent->Wait();
@@ -474,7 +468,7 @@ void StartRenderingThread()
 	// Create the rendering thread heartbeat
 	GRenderingThreadRunnableHeartbeat = new FRenderingThreadTickHeartbeat();
 
-	GRenderingThreadHeartbeat = FRunnableThread::Create(GRenderingThreadRunnableHeartbeat, *FString::Printf(TEXT("RTHeartBeat %d"), ThreadCount), 16 * 1024, TPri_AboveNormal);
+	GRenderingThreadHeartbeat = FRunnableThread::Create(GRenderingThreadRunnableHeartbeat, *FString::Printf(TEXT("RTHeartBeat %d"), ThreadCount), 16 * 1024, TPri_AboveNormal, FPlatformAffinity::GetRTHeartBeatMask());
 
 	ThreadCount++;
 }
@@ -522,7 +516,7 @@ void StopRenderingThread()
 
 			// Busy wait while BP debugging, to avoid opportunistic execution of game thread tasks
 			// If the game thread is already executing tasks, then we have no choice but to spin
-			if (GIntraFrameDebuggingGameThread || FTaskGraphInterface::Get().IsThreadProcessingTasks(ENamedThreads::GameThread)) 
+			if (GIntraFrameDebuggingGameThread || FTaskGraphInterface::Get().IsThreadProcessingTasks(ENamedThreads::GameThread) ) 
 			{
 				while ((QuitTask.GetReference() != nullptr) && !QuitTask->IsComplete())
 				{
