Move FShadow/LightMapRef into SceneTypes.h, stop including LightMap.h/ShadowMap.h everywhere
Add explicit constructor to UStaticMeshComponent, ULandscapeComponent and FModelElement so they can be implmented in cpp file, which includes LightMap/ShadowMap.h
Rename PackedNormal.inl to PackedNormal.h and move rest of declaration, so it can be included independent of RenderUtils.h
Move FMaterialRelevance definition from PrimitiveComponent.h to MaterialInterface.h
Move FExpressionInput/FExpressionOutput/FMaterialInput etc. from MaterialShared.h into new MaterialExpressionIO.h, to avoid Material.h pulling in RHI includes
Other required includes added

diff --git a/Engine/Source/Developer/GameplayDebugger/Classes/GameplayDebuggingHUDComponent.h b/Engine/Source/Developer/GameplayDebugger/Classes/GameplayDebuggingHUDComponent.h
--- a/Engine/Source/Developer/GameplayDebugger/Classes/GameplayDebuggingHUDComponent.h
+++ b/Engine/Source/Developer/GameplayDebugger/Classes/GameplayDebuggingHUDComponent.h
@@ -2,6 +2,7 @@
 
 #pragma once
 #include "GameFramework/HUD.h"
+#include "ShowFlags.h"
 #include "GameplayDebuggingHUDComponent.generated.h"
 
 class AGameplayDebuggingReplicator;
diff --git a/Engine/Source/Developer/GameplayDebugger/Private/GameplayDebuggingHUDComponent.cpp b/Engine/Source/Developer/GameplayDebugger/Private/GameplayDebuggingHUDComponent.cpp
--- a/Engine/Source/Developer/GameplayDebugger/Private/GameplayDebuggingHUDComponent.cpp
+++ b/Engine/Source/Developer/GameplayDebugger/Private/GameplayDebuggingHUDComponent.cpp
@@ -10,6 +10,7 @@
 #include "GameplayDebuggingControllerComponent.h"
 #include "CanvasItem.h"
 #include "AI/Navigation/NavigationSystem.h"
+#include "Engine/Texture2D.h"
 #include "Regex.h"
 #include "DrawDebugHelpers.h"
 #include "TimerManager.h"					// Game play timers
diff --git a/Engine/Source/Editor/StatsViewer/Private/StatsPages/PrimitiveStatsPage.cpp b/Engine/Source/Editor/StatsViewer/Private/StatsPages/PrimitiveStatsPage.cpp
--- a/Engine/Source/Editor/StatsViewer/Private/StatsPages/PrimitiveStatsPage.cpp
+++ b/Engine/Source/Editor/StatsViewer/Private/StatsPages/PrimitiveStatsPage.cpp
@@ -6,6 +6,8 @@
 #include "PrimitiveStats.h"
 #include "Landscape/LandscapeProxy.h"
 #include "Landscape/LandscapeComponent.h"
+#include "LightMap.h"
+#include "ShadowMap.h"
 
 #define LOCTEXT_NAMESPACE "Editor.StatsViewer.PrimitiveStats"
 
diff --git a/Engine/Source/Editor/UnrealEd/Classes/Editor/EditorEngine.h b/Engine/Source/Editor/UnrealEd/Classes/Editor/EditorEngine.h
--- a/Engine/Source/Editor/UnrealEd/Classes/Editor/EditorEngine.h
+++ b/Engine/Source/Editor/UnrealEd/Classes/Editor/EditorEngine.h
@@ -11,6 +11,7 @@
 class FAssetData;
 class UAnimSequence;
 class USkeleton;
+class FPoly;
 
 //
 // Things to set in mapSetBrush.
diff --git a/Engine/Source/Editor/UnrealEd/Private/EditorServer.cpp b/Engine/Source/Editor/UnrealEd/Private/EditorServer.cpp
--- a/Engine/Source/Editor/UnrealEd/Private/EditorServer.cpp
+++ b/Engine/Source/Editor/UnrealEd/Private/EditorServer.cpp
@@ -12,6 +12,7 @@
 #include "Layers/ILayers.h"
 #include "ScopedTransaction.h"
 #include "SurfaceIterators.h"
+#include "LightMap.h"
 #include "BSPOps.h"
 #include "EditorLevelUtils.h"
 #include "Editor/MainFrame/Public/MainFrame.h"
diff --git a/Engine/Source/Runtime/AIModule/Private/EnvironmentQuery/EQSTestingPawn.cpp b/Engine/Source/Runtime/AIModule/Private/EnvironmentQuery/EQSTestingPawn.cpp
--- a/Engine/Source/Runtime/AIModule/Private/EnvironmentQuery/EQSTestingPawn.cpp
+++ b/Engine/Source/Runtime/AIModule/Private/EnvironmentQuery/EQSTestingPawn.cpp
@@ -19,6 +19,7 @@ extern UNREALED_API class UEditorEngine* GEditor;
 #endif // WITH_EDITOR
 
 #include "Engine/Selection.h"
+#include "Engine/Texture2D.h"
 #include "Components/BillboardComponent.h"
 
 //----------------------------------------------------------------------//
diff --git a/Engine/Source/Runtime/Engine/Classes/Components/MeshComponent.h b/Engine/Source/Runtime/Engine/Classes/Components/MeshComponent.h
--- a/Engine/Source/Runtime/Engine/Classes/Components/MeshComponent.h
+++ b/Engine/Source/Runtime/Engine/Classes/Components/MeshComponent.h
@@ -2,6 +2,7 @@
 
 
 #pragma once
+#include "Components/PrimitiveComponent.h"
 #include "MeshComponent.generated.h"
 
 UCLASS(abstract)
diff --git a/Engine/Source/Runtime/Engine/Classes/Components/PrimitiveComponent.h b/Engine/Source/Runtime/Engine/Classes/Components/PrimitiveComponent.h
--- a/Engine/Source/Runtime/Engine/Classes/Components/PrimitiveComponent.h
+++ b/Engine/Source/Runtime/Engine/Classes/Components/PrimitiveComponent.h
@@ -5,13 +5,12 @@
 #include "StaticArray.h"
 #include "PhysicsEngine/BodyInstance.h"
 #include "Components/SceneComponent.h"
+#include "Materials/MaterialInterface.h"
 #include "SceneTypes.h"
 #include "CollisionQueryParams.h"
-
 #include "PrimitiveComponent.generated.h"
 
 class FPrimitiveSceneProxy;
-struct FPrimitiveViewRelevance;
 
 /** Information about a vertex of a primitive's triangle. */
 struct FPrimitiveTriangleVertex
@@ -78,68 +77,7 @@ namespace EHasCustomNavigableGeometry
 	};
 };
 
-/** Mirrored from Scene.h */
-USTRUCT()
-struct ENGINE_API FMaterialRelevance
-{
-	GENERATED_USTRUCT_BODY()
-
-	UPROPERTY()
-	uint32 bOpaque:1;
-
-	UPROPERTY()
-	uint32 bMasked:1;
-
-	UPROPERTY()
-	uint32 bDistortion:1;
-
-	UPROPERTY()
-	uint32 bUsesSceneColor:1;
-
-	UPROPERTY()
-	uint32 bSeparateTranslucency:1;
-
-	UPROPERTY()
-	uint32 bNormalTranslucency:1;
-
-	UPROPERTY()
-	uint32 bDisableDepthTest:1;
-
-	/** Default constructor. */
-	FMaterialRelevance()
-		:	bOpaque(false)
-		,	bMasked(false)
-		,	bDistortion(false)
-		,	bUsesSceneColor(false)
-		,	bSeparateTranslucency(false)
-		,	bNormalTranslucency(false)
-		,	bDisableDepthTest(false)
-		{}
-
-	/** Bitwise OR operator.  Sets any relevance bits which are present in either FMaterialRelevance. */
-	FMaterialRelevance& operator|=(const FMaterialRelevance& B)
-	{
-		bOpaque |= B.bOpaque;
-		bMasked |= B.bMasked;
-		bDistortion |= B.bDistortion;
-		bUsesSceneColor |= B.bUsesSceneColor;
-		bSeparateTranslucency |= B.bSeparateTranslucency;
-		bNormalTranslucency |= B.bNormalTranslucency;
-		bDisableDepthTest |= B.bDisableDepthTest;
-		return *this;
-	}
-	
-	/** Binary bitwise OR operator. */
-	friend FMaterialRelevance operator|(const FMaterialRelevance& A,const FMaterialRelevance& B)
-	{
-		FMaterialRelevance Result(A);
-		Result |= B;
-		return Result;
-	}
 
-	/** Copies the material's relevance flags to a primitive's view relevance flags. */
-	void SetPrimitiveViewRelevance(FPrimitiveViewRelevance& OutViewRelevance) const;
-};
 
 /** Information about the sprite category */
 USTRUCT()
diff --git a/Engine/Source/Runtime/Engine/Classes/Components/SphereComponent.h b/Engine/Source/Runtime/Engine/Classes/Components/SphereComponent.h
--- a/Engine/Source/Runtime/Engine/Classes/Components/SphereComponent.h
+++ b/Engine/Source/Runtime/Engine/Classes/Components/SphereComponent.h
@@ -2,6 +2,7 @@
 
 
 #pragma once
+#include "Components/ShapeComponent.h"
 #include "SphereComponent.generated.h"
 
 /** 
diff --git a/Engine/Source/Runtime/Engine/Classes/Components/StaticMeshComponent.h b/Engine/Source/Runtime/Engine/Classes/Components/StaticMeshComponent.h
--- a/Engine/Source/Runtime/Engine/Classes/Components/StaticMeshComponent.h
+++ b/Engine/Source/Runtime/Engine/Classes/Components/StaticMeshComponent.h
@@ -2,10 +2,9 @@
 
 
 #pragma once
-#include "LightMap.h"
-#include "ShadowMap.h"
+#include "SceneTypes.h"
 #include "Components/MeshComponent.h"
-#include "Runtime/RenderCore/Public/RenderUtils.h"
+#include "Runtime/RenderCore/Public/PackedNormal.h"
 #include "StaticMeshComponent.generated.h"
 
 class FColorVertexBuffer;
@@ -179,6 +178,8 @@ class ENGINE_API UStaticMeshComponent : public UMeshComponent
 	UPROPERTY(EditAnywhere, BlueprintReadWrite, Category=Lighting)
 	struct FLightmassPrimitiveSettings LightmassSettings;
 
+	virtual ~UStaticMeshComponent();
+
 	/** Change the StaticMesh used by this instance. */
 	UFUNCTION(BlueprintCallable, Category="Components|StaticMesh")
 	virtual bool SetStaticMesh(class UStaticMesh* NewMesh);
diff --git a/Engine/Source/Runtime/Engine/Classes/Components/TextRenderComponent.h b/Engine/Source/Runtime/Engine/Classes/Components/TextRenderComponent.h
--- a/Engine/Source/Runtime/Engine/Classes/Components/TextRenderComponent.h
+++ b/Engine/Source/Runtime/Engine/Classes/Components/TextRenderComponent.h
@@ -2,6 +2,7 @@
 
 
 #pragma once
+#include "Components/PrimitiveComponent.h"
 #include "TextRenderComponent.generated.h"
 
 UENUM()
diff --git a/Engine/Source/Runtime/Engine/Classes/GameFramework/ProjectileMovementComponent.h b/Engine/Source/Runtime/Engine/Classes/GameFramework/ProjectileMovementComponent.h
--- a/Engine/Source/Runtime/Engine/Classes/GameFramework/ProjectileMovementComponent.h
+++ b/Engine/Source/Runtime/Engine/Classes/GameFramework/ProjectileMovementComponent.h
@@ -1,6 +1,7 @@
 // Copyright 1998-2014 Epic Games, Inc. All Rights Reserved.
 
 #pragma once
+#include "GameFramework/MovementComponent.h"
 #include "ProjectileMovementComponent.generated.h"
 
 /**
diff --git a/Engine/Source/Runtime/Engine/Classes/Landscape/LandscapeComponent.h b/Engine/Source/Runtime/Engine/Classes/Landscape/LandscapeComponent.h
--- a/Engine/Source/Runtime/Engine/Classes/Landscape/LandscapeComponent.h
+++ b/Engine/Source/Runtime/Engine/Classes/Landscape/LandscapeComponent.h
@@ -2,8 +2,7 @@
 
 #pragma once
 #include "LandscapeLayerInfoObject.h"
-#include "LightMap.h"
-#include "ShadowMap.h"
+#include "SceneTypes.h"
 #include "LandscapeComponent.generated.h"
 
 class FLandscapeComponentDerivedData
@@ -237,6 +236,8 @@ public:
 	/** Platform Data where don't support texture sampling in vertex buffer */
 	FLandscapeComponentDerivedData PlatformData;
 
+	virtual ~ULandscapeComponent();
+
 	// Begin UObject interface.	
 	virtual void PostInitProperties() override;	
 	virtual void Serialize(FArchive& Ar) override;
diff --git a/Engine/Source/Runtime/Engine/Classes/Materials/Material.h b/Engine/Source/Runtime/Engine/Classes/Materials/Material.h
--- a/Engine/Source/Runtime/Engine/Classes/Materials/Material.h
+++ b/Engine/Source/Runtime/Engine/Classes/Materials/Material.h
@@ -3,7 +3,7 @@
 #pragma once
 #include "../Engine/BlendableInterface.h"
 #include "../Materials/MaterialInterface.h"
-#include "MaterialShared.h"
+#include "MaterialExpressionIO.h"
 #include "Material.generated.h"
 
 #if WITH_EDITOR
@@ -256,7 +256,7 @@ struct FMaterialParameterCollectionInfo
  * when light from the scene hits the surface, the shading model of the material is used to calculate how
  * that light interacts with the surface. 
  */
-UCLASS(hidecategories=Object, MinimalAPI, BlueprintType)
+UCLASS(hidecategories=Object, MinimalAPI, BlueprintType, DependsOn=UMaterialExpression)
 class UMaterial : public UMaterialInterface
 {
 	GENERATED_UCLASS_BODY()
diff --git a/Engine/Source/Runtime/Engine/Classes/Materials/MaterialInterface.h b/Engine/Source/Runtime/Engine/Classes/Materials/MaterialInterface.h
--- a/Engine/Source/Runtime/Engine/Classes/Materials/MaterialInterface.h
+++ b/Engine/Source/Runtime/Engine/Classes/Materials/MaterialInterface.h
@@ -10,6 +10,7 @@
 class UMaterial;
 class FMaterialResource;
 class FMaterialCompiler;
+struct FPrimitiveViewRelevance;
 
 UENUM()
 enum EMaterialUsage
@@ -27,6 +28,68 @@ enum EMaterialUsage
 	MATUSAGE_MAX,
 };
 
+USTRUCT()
+struct ENGINE_API FMaterialRelevance
+{
+	GENERATED_USTRUCT_BODY()
+
+	UPROPERTY()
+	uint32 bOpaque : 1;
+
+	UPROPERTY()
+		uint32 bMasked : 1;
+
+	UPROPERTY()
+		uint32 bDistortion : 1;
+
+	UPROPERTY()
+		uint32 bUsesSceneColor : 1;
+
+	UPROPERTY()
+		uint32 bSeparateTranslucency : 1;
+
+	UPROPERTY()
+		uint32 bNormalTranslucency : 1;
+
+	UPROPERTY()
+		uint32 bDisableDepthTest : 1;
+
+	/** Default constructor. */
+	FMaterialRelevance()
+		: bOpaque(false)
+		, bMasked(false)
+		, bDistortion(false)
+		, bUsesSceneColor(false)
+		, bSeparateTranslucency(false)
+		, bNormalTranslucency(false)
+		, bDisableDepthTest(false)
+	{}
+
+	/** Bitwise OR operator.  Sets any relevance bits which are present in either FMaterialRelevance. */
+	FMaterialRelevance& operator|=(const FMaterialRelevance& B)
+	{
+		bOpaque |= B.bOpaque;
+		bMasked |= B.bMasked;
+		bDistortion |= B.bDistortion;
+		bUsesSceneColor |= B.bUsesSceneColor;
+		bSeparateTranslucency |= B.bSeparateTranslucency;
+		bNormalTranslucency |= B.bNormalTranslucency;
+		bDisableDepthTest |= B.bDisableDepthTest;
+		return *this;
+	}
+
+	/** Binary bitwise OR operator. */
+	friend FMaterialRelevance operator|(const FMaterialRelevance& A, const FMaterialRelevance& B)
+	{
+		FMaterialRelevance Result(A);
+		Result |= B;
+		return Result;
+	}
+
+	/** Copies the material's relevance flags to a primitive's view relevance flags. */
+	void SetPrimitiveViewRelevance(FPrimitiveViewRelevance& OutViewRelevance) const;
+};
+
 /** 
  *	UMaterial interface settings for Lightmass
  */
diff --git a/Engine/Source/Runtime/Engine/Classes/Particles/ParticleSystemComponent.h b/Engine/Source/Runtime/Engine/Classes/Particles/ParticleSystemComponent.h
--- a/Engine/Source/Runtime/Engine/Classes/Particles/ParticleSystemComponent.h
+++ b/Engine/Source/Runtime/Engine/Classes/Particles/ParticleSystemComponent.h
@@ -3,6 +3,7 @@
 
 #pragma once
 #include "Components/PrimitiveComponent.h"
+#include "Materials/MaterialInterface.h"
 #include "ParticleHelper.h"
 #include "Particles/Emitter.h"
 #include "Particles/ParticleEmitter.h"
@@ -395,7 +396,7 @@ public:
 
 	/** The view relevance flags for each LODLevel. */
 	UPROPERTY(transient)
-	TArray<struct FMaterialRelevance> CachedViewRelevanceFlags;
+	TArray<FMaterialRelevance> CachedViewRelevanceFlags;
 
 	/** If true, the ViewRelevanceFlags are dirty and should be recached */
 	uint32 bIsViewRelevanceDirty:1;
diff --git a/Engine/Source/Runtime/Engine/Classes/PhysicsEngine/BodyInstance.h b/Engine/Source/Runtime/Engine/Classes/PhysicsEngine/BodyInstance.h
--- a/Engine/Source/Runtime/Engine/Classes/PhysicsEngine/BodyInstance.h
+++ b/Engine/Source/Runtime/Engine/Classes/PhysicsEngine/BodyInstance.h
@@ -3,6 +3,7 @@
 #pragma once
 
 #include "PhysxUserData.h"
+#include "CollisionQueryParams.h"
 #include "BodyInstance.generated.h"
 
 
@@ -507,7 +508,7 @@ public:
 	 *	@param	ObjectQueryParams	List of object types it's looking for. When this enters, we do object query with component shape
 	 *  @return TRUE if OutOverlaps contains any blocking results
 	 */
-	bool OverlapMulti(TArray<struct FOverlapResult>& InOutOverlaps, const class UWorld* World, const FTransform* pWorldToComponent, const FVector& Pos, const FRotator& Rot, ECollisionChannel TestChannel, const struct FComponentQueryParams& Params, const struct FCollisionResponseParams& ResponseParams, const struct FCollisionObjectQueryParams& ObjectQueryParams = FCollisionObjectQueryParams::DefaultObjectQueryParam) const;
+	bool OverlapMulti(TArray<struct FOverlapResult>& InOutOverlaps, const class UWorld* World, const FTransform* pWorldToComponent, const FVector& Pos, const FRotator& Rot, ECollisionChannel TestChannel, const struct FComponentQueryParams& Params, const struct FCollisionResponseParams& ResponseParams, const FCollisionObjectQueryParams& ObjectQueryParams = FCollisionObjectQueryParams::DefaultObjectQueryParam) const;
 
 	/**
 	 * Add an impulse to this bodyinstance, radiating out from the specified position.
diff --git a/Engine/Source/Runtime/Engine/Private/ContentStreaming.cpp b/Engine/Source/Runtime/Engine/Private/ContentStreaming.cpp
--- a/Engine/Source/Runtime/Engine/Private/ContentStreaming.cpp
+++ b/Engine/Source/Runtime/Engine/Private/ContentStreaming.cpp
@@ -92,7 +92,7 @@ TArray<IStreamingManager::FSlaveLocation> IStreamingManager::SlaveLocations;
 bool IStreamingManager::bPendingRemoveViews = false;
 
 /** The lightmap used by the currently selected component (toggledebugcamera), if it's a static mesh component. */
-extern ENGINE_API FLightMap2D* GDebugSelectedLightmap;
+extern ENGINE_API class FLightMap2D* GDebugSelectedLightmap;
 
 /** Whether the texture pool size has been artificially lowered for testing. */
 bool GIsOperatingWithReducedTexturePool = false;
diff --git a/Engine/Source/Runtime/Engine/Private/DebugCameraController.cpp b/Engine/Source/Runtime/Engine/Private/DebugCameraController.cpp
--- a/Engine/Source/Runtime/Engine/Private/DebugCameraController.cpp
+++ b/Engine/Source/Runtime/Engine/Private/DebugCameraController.cpp
@@ -9,6 +9,8 @@
 #include "GameFramework/SpectatorPawn.h"
 #include "GameFramework/SpectatorPawnMovement.h"
 #include "Engine/DebugCameraHUD.h"
+#include "LightMap.h"
+
 
 /** The currently selected actor. */
 AActor* GDebugSelectedActor = NULL;
diff --git a/Engine/Source/Runtime/Engine/Private/GameViewportClient.cpp b/Engine/Source/Runtime/Engine/Private/GameViewportClient.cpp
--- a/Engine/Source/Runtime/Engine/Private/GameViewportClient.cpp
+++ b/Engine/Source/Runtime/Engine/Private/GameViewportClient.cpp
@@ -30,7 +30,7 @@ extern ENGINE_API bool GShowDebugSelectedLightmap;
 /** The currently selected component in the actor. */
 extern ENGINE_API UPrimitiveComponent* GDebugSelectedComponent;
 /** The lightmap used by the currently selected component, if it's a static mesh component. */
-extern ENGINE_API FLightMap2D* GDebugSelectedLightmap;
+extern ENGINE_API class FLightMap2D* GDebugSelectedLightmap;
 
 /** Delegate called at the end of the frame when a screenshot is captured */
 FOnScreenshotCaptured UGameViewportClient::ScreenshotCapturedDelegate;
diff --git a/Engine/Source/Runtime/Engine/Private/Landscape/Landscape.cpp b/Engine/Source/Runtime/Engine/Private/Landscape/Landscape.cpp
--- a/Engine/Source/Runtime/Engine/Private/Landscape/Landscape.cpp
+++ b/Engine/Source/Runtime/Engine/Private/Landscape/Landscape.cpp
@@ -17,6 +17,8 @@ Landscape.cpp: Terrain rendering
 #include "Landscape/Landscape.h"
 #include "Landscape/LandscapeMeshCollisionComponent.h"
 #include "Landscape/LandscapeSplinesComponent.h"
+#include "LightMap.h"
+#include "ShadowMap.h"
 
 // Set this to 0 to disable landscape cooking and thus disable it on device.
 #define ENABLE_LANDSCAPE_COOKING 1
@@ -89,6 +91,12 @@ ULandscapeComponent::ULandscapeComponent(const class FPostConstructInitializePro
 	LpvBiasMultiplier = 0.0f; // Bias is 0 for landscape, since it's single sided
 }
 
+ULandscapeComponent::~ULandscapeComponent()
+{
+
+}
+
+
 void ULandscapeComponent::AddReferencedObjects(UObject* InThis, FReferenceCollector& Collector)
 {
 	ULandscapeComponent* This = CastChecked<ULandscapeComponent>(InThis);
diff --git a/Engine/Source/Runtime/Engine/Private/Materials/MaterialInterface.cpp b/Engine/Source/Runtime/Engine/Private/Materials/MaterialInterface.cpp
--- a/Engine/Source/Runtime/Engine/Private/Materials/MaterialInterface.cpp
+++ b/Engine/Source/Runtime/Engine/Private/Materials/MaterialInterface.cpp
@@ -6,6 +6,21 @@
 
 #include "EnginePrivate.h"
 
+//////////////////////////////////////////////////////////////////////////
+
+/** Copies the material's relevance flags to a primitive's view relevance flags. */
+void FMaterialRelevance::SetPrimitiveViewRelevance(FPrimitiveViewRelevance& OutViewRelevance) const
+{
+	OutViewRelevance.bOpaqueRelevance = bOpaque;
+	OutViewRelevance.bMaskedRelevance = bMasked;
+	OutViewRelevance.bDistortionRelevance = bDistortion;
+	OutViewRelevance.bSceneColorRelevance = bUsesSceneColor;
+	OutViewRelevance.bSeparateTranslucencyRelevance = bSeparateTranslucency;
+	OutViewRelevance.bNormalTranslucencyRelevance = bNormalTranslucency;
+}
+
+//////////////////////////////////////////////////////////////////////////
+
 UMaterialInterface::UMaterialInterface(const class FPostConstructInitializeProperties& PCIP)
 	: Super(PCIP)
 {
diff --git a/Engine/Source/Runtime/Engine/Private/ModelComponent.cpp b/Engine/Source/Runtime/Engine/Private/ModelComponent.cpp
--- a/Engine/Source/Runtime/Engine/Private/ModelComponent.cpp
+++ b/Engine/Source/Runtime/Engine/Private/ModelComponent.cpp
@@ -6,6 +6,8 @@
 
 #include "EnginePrivate.h"
 #include "Model.h"
+#include "LightMap.h"
+#include "ShadowMap.h"
 
 FModelElement::FModelElement(UModelComponent* InComponent,UMaterialInterface* InMaterial):
 	Component(InComponent),
@@ -18,6 +20,21 @@ FModelElement::FModelElement(UModelComponent* InComponent,UMaterialInterface* In
 	BoundingBox(ForceInitToZero)
 {}
 
+FModelElement::FModelElement():
+	Component(NULL), 
+	Material(NULL), 
+	IndexBuffer(NULL), 
+	FirstIndex(0), 
+	NumTriangles(0), 
+	MinVertexIndex(0), 
+	MaxVertexIndex(0), 
+	BoundingBox(ForceInitToZero)
+{}
+
+FModelElement::~FModelElement()
+{}
+
+
 /**
  * Serializer.
  */
diff --git a/Engine/Source/Runtime/Engine/Private/ModelRender.cpp b/Engine/Source/Runtime/Engine/Private/ModelRender.cpp
--- a/Engine/Source/Runtime/Engine/Private/ModelRender.cpp
+++ b/Engine/Source/Runtime/Engine/Private/ModelRender.cpp
@@ -8,7 +8,8 @@
 #include "LevelUtils.h"
 #include "Model.h"
 #include "HModel.h"
-
+#include "LightMap.h"
+#include "ShadowMap.h"
 
 namespace
 {
diff --git a/Engine/Source/Runtime/Engine/Private/NavMeshRenderingHelpers.h b/Engine/Source/Runtime/Engine/Private/NavMeshRenderingHelpers.h
--- a/Engine/Source/Runtime/Engine/Private/NavMeshRenderingHelpers.h
+++ b/Engine/Source/Runtime/Engine/Private/NavMeshRenderingHelpers.h
@@ -6,6 +6,7 @@
 #include "BatchedElements.h"
 #include "AI/Navigation/NavMeshRenderingComponent.h"
 #include "Materials/Material.h"
+#include "MaterialShared.h"
 
 static const FColor NavMeshRenderColor_Recast_TriangleEdges(255,255,255);
 static const FColor NavMeshRenderColor_Recast_TileEdges(16,16,16,32);
diff --git a/Engine/Source/Runtime/Engine/Private/PrimitiveComponent.cpp b/Engine/Source/Runtime/Engine/Private/PrimitiveComponent.cpp
--- a/Engine/Source/Runtime/Engine/Private/PrimitiveComponent.cpp
+++ b/Engine/Source/Runtime/Engine/Private/PrimitiveComponent.cpp
@@ -33,20 +33,6 @@ static FAutoConsoleVariable CVarAllowCachedOverlaps(
 	TEXT("0: disable cached overlaps, 1: enable\n"));
 
 
-//////////////////////////////////////////////////////////////////////////
-// FMaterialRelevance
-
-/** Copies the material's relevance flags to a primitive's view relevance flags. */
-void FMaterialRelevance::SetPrimitiveViewRelevance(FPrimitiveViewRelevance& OutViewRelevance) const
-{
-	OutViewRelevance.bOpaqueRelevance = bOpaque;
-	OutViewRelevance.bMaskedRelevance = bMasked;
-	OutViewRelevance.bDistortionRelevance = bDistortion;
-	OutViewRelevance.bSceneColorRelevance = bUsesSceneColor;
-	OutViewRelevance.bSeparateTranslucencyRelevance = bSeparateTranslucency;
-	OutViewRelevance.bNormalTranslucencyRelevance = bNormalTranslucency;
-}
-
 ///////////////////////////////////////////////////////////////////////////////
 // PRIMITIVE COMPONENT
 ///////////////////////////////////////////////////////////////////////////////
diff --git a/Engine/Source/Runtime/Engine/Private/ShadowMap.cpp b/Engine/Source/Runtime/Engine/Private/ShadowMap.cpp
--- a/Engine/Source/Runtime/Engine/Private/ShadowMap.cpp
+++ b/Engine/Source/Runtime/Engine/Private/ShadowMap.cpp
@@ -4,6 +4,8 @@
 
 #include "TextureLayout.h"
 #include "TargetPlatform.h"
+#include "LightMap.h"
+#include "ShadowMap.h"
 
 #if WITH_EDITOR
 	// NOTE: We're only counting the top-level mip-map for the following variables.
diff --git a/Engine/Source/Runtime/Engine/Private/StaticMeshComponent.cpp b/Engine/Source/Runtime/Engine/Private/StaticMeshComponent.cpp
--- a/Engine/Source/Runtime/Engine/Private/StaticMeshComponent.cpp
+++ b/Engine/Source/Runtime/Engine/Private/StaticMeshComponent.cpp
@@ -8,6 +8,8 @@
 #include "UObjectToken.h"
 #include "MapErrors.h"
 #include "ComponentInstanceDataCache.h"
+#include "LightMap.h"
+#include "ShadowMap.h"
 
 #define LOCTEXT_NAMESPACE "StaticMeshComponent"
 
@@ -139,6 +141,11 @@ UStaticMeshComponent::UStaticMeshComponent(const class FPostConstructInitializeP
 #endif
 }
 
+UStaticMeshComponent::~UStaticMeshComponent()
+{
+	// Empty, but required because we don't want to have to include LightMap.h and ShadowMap.h in StaticMeshComponent.h, and they are required to compile FLightMapRef and FShadowMapRef
+}
+
 void UStaticMeshComponent::GetLifetimeReplicatedProps(TArray< FLifetimeProperty > & OutLifetimeProps) const
 {
 	Super::GetLifetimeReplicatedProps(OutLifetimeProps);
diff --git a/Engine/Source/Runtime/Engine/Private/StaticMeshRender.cpp b/Engine/Source/Runtime/Engine/Private/StaticMeshRender.cpp
--- a/Engine/Source/Runtime/Engine/Private/StaticMeshRender.cpp
+++ b/Engine/Source/Runtime/Engine/Private/StaticMeshRender.cpp
@@ -8,6 +8,8 @@
 #include "StaticMeshResources.h"
 #include "LevelUtils.h"
 #include "TessellationRendering.h"
+#include "LightMap.h"
+#include "ShadowMap.h"
 
 /** If true, optimized depth-only index buffers are used for shadow rendering. */
 static bool GUseShadowIndexBuffer = true;
@@ -1079,6 +1081,16 @@ FLightInteraction FStaticMeshSceneProxy::FLODInfo::GetInteraction(const FLightSc
 	return FLightInteraction::Dynamic();
 }
 
+FLightMapInteraction FStaticMeshSceneProxy::FLODInfo::GetLightMapInteraction() const
+{
+	return LightMap ? LightMap->GetInteraction() : FLightMapInteraction();
+}
+
+FShadowMapInteraction FStaticMeshSceneProxy::FLODInfo::GetShadowMapInteraction() const
+{
+	return ShadowMap ? ShadowMap->GetInteraction() : FShadowMapInteraction();
+}
+
 float FStaticMeshSceneProxy::GetScreenSize( int32 LODIndex ) const
 {
 	return RenderData->ScreenSize[LODIndex];
diff --git a/Engine/Source/Runtime/Engine/Public/CanvasItem.h b/Engine/Source/Runtime/Engine/Public/CanvasItem.h
--- a/Engine/Source/Runtime/Engine/Public/CanvasItem.h
+++ b/Engine/Source/Runtime/Engine/Public/CanvasItem.h
@@ -8,6 +8,8 @@
 
 #include "BatchedElements.h"
 
+class FMaterialRenderProxy;
+
 class FCanvasItem
 {
 public:
diff --git a/Engine/Source/Runtime/Engine/Public/DynamicMeshBuilder.h b/Engine/Source/Runtime/Engine/Public/DynamicMeshBuilder.h
--- a/Engine/Source/Runtime/Engine/Public/DynamicMeshBuilder.h
+++ b/Engine/Source/Runtime/Engine/Public/DynamicMeshBuilder.h
@@ -6,6 +6,9 @@
 
 #pragma once
 
+class FPrimitiveDrawInterface;
+class FMaterialRenderProxy;
+
 /** The vertex type used for dynamic meshes. */
 struct FDynamicMeshVertex
 {
diff --git a/Engine/Source/Runtime/Engine/Public/Landscape/LandscapeRender.h b/Engine/Source/Runtime/Engine/Public/Landscape/LandscapeRender.h
--- a/Engine/Source/Runtime/Engine/Public/Landscape/LandscapeRender.h
+++ b/Engine/Source/Runtime/Engine/Public/Landscape/LandscapeRender.h
@@ -10,6 +10,9 @@ LandscapeRender.h: New terrain rendering
 #include "UniformBuffer.h"
 #include "Landscape/LandscapeComponent.h"
 #include "Landscape/LandscapeProxy.h"
+#include "LightMap.h"
+#include "ShadowMap.h"
+
 
 class FLandscapeVertexFactory;
 class FLandscapeVertexBuffer;
diff --git a/Engine/Source/Runtime/Engine/Public/LightMap.h b/Engine/Source/Runtime/Engine/Public/LightMap.h
--- a/Engine/Source/Runtime/Engine/Public/LightMap.h
+++ b/Engine/Source/Runtime/Engine/Public/LightMap.h
@@ -118,8 +118,7 @@ private:
 	virtual void FinishCleanup();
 };
 
-/** A reference to a light-map. */
-typedef TRefCountPtr<FLightMap> FLightMapRef;
+
 
 /** Lightmap reference serializer */
 extern FArchive& operator<<(FArchive& Ar,FLightMap*& R);
diff --git a/Engine/Source/Runtime/Engine/Public/MaterialExpressionIO.h b/Engine/Source/Runtime/Engine/Public/MaterialExpressionIO.h
new file mode 100644
--- /dev/null
+++ b/Engine/Source/Runtime/Engine/Public/MaterialExpressionIO.h
@@ -0,0 +1,119 @@
+// Copyright 1998-2014 Epic Games, Inc. All Rights Reserved.
+
+#pragma once
+
+//
+//	FExpressionInput
+//
+
+//@warning: FExpressionInput is mirrored in MaterialExpression.h and manually "subclassed" in Material.h (FMaterialInput)
+struct FExpressionInput
+{
+	/** Material expression that this input is connected to, or NULL if not connected. */
+	class UMaterialExpression*	Expression;
+
+	/** Index into Expression's outputs array that this input is connected to. */
+	int32						OutputIndex;
+
+	/** 
+	 * Optional name of the input.  
+	 * Note that this is the only member which is not derived from the output currently connected. 
+	 */
+	FString						InputName;
+
+	int32						Mask,
+								MaskR,
+								MaskG,
+								MaskB,
+								MaskA;
+	uint32						GCC64Padding; // @todo 64: if the C++ didn't mismirror this structure, we might not need this
+
+	FExpressionInput()
+		: Expression(NULL)
+		, OutputIndex(0)
+		, Mask(0)
+		, MaskR(0)
+		, MaskG(0)
+		, MaskB(0)
+		, MaskA(0)
+		, GCC64Padding(0)
+	{
+	}
+
+	int32 Compile(class FMaterialCompiler* Compiler, int32 MultiplexIndex=INDEX_NONE);
+
+	/**
+	 * Tests if the input has a material expression connected to it
+	 *
+	 * @return	true if an expression is connected, otherwise false
+	 */
+	bool IsConnected() const { return (NULL != Expression); }
+
+	/** Connects output of InExpression to this input */
+	ENGINE_API void Connect( int32 InOutputIndex, class UMaterialExpression* InExpression );
+};
+
+//
+//	FExpressionOutput
+//
+
+struct FExpressionOutput
+{
+	FString	OutputName;
+	int32	Mask,
+		MaskR,
+		MaskG,
+		MaskB,
+		MaskA;
+
+	FExpressionOutput(int32 InMask = 0, int32 InMaskR = 0, int32 InMaskG = 0, int32 InMaskB = 0, int32 InMaskA = 0) :
+		Mask(InMask),
+		MaskR(InMaskR),
+		MaskG(InMaskG),
+		MaskB(InMaskB),
+		MaskA(InMaskA)
+	{}
+
+	FExpressionOutput(FString InOutputName, int32 InMask = 0, int32 InMaskR = 0, int32 InMaskG = 0, int32 InMaskB = 0, int32 InMaskA = 0) :
+		OutputName(InOutputName),
+		Mask(InMask),
+		MaskR(InMaskR),
+		MaskG(InMaskG),
+		MaskB(InMaskB),
+		MaskA(InMaskA)
+	{}
+};
+
+//
+//	FMaterialInput
+//
+
+template<class InputType> struct FMaterialInput : FExpressionInput
+{
+	uint32	UseConstant : 1;
+	InputType	Constant;
+};
+
+struct FColorMaterialInput : FMaterialInput<FColor>
+{
+	ENGINE_API int32 Compile(class FMaterialCompiler* Compiler, const FColor& Default);
+};
+struct FScalarMaterialInput : FMaterialInput<float>
+{
+	ENGINE_API int32 Compile(class FMaterialCompiler* Compiler, float Default);
+};
+
+struct FVectorMaterialInput : FMaterialInput<FVector>
+{
+	ENGINE_API int32 Compile(class FMaterialCompiler* Compiler, const FVector& Default);
+};
+
+struct FVector2MaterialInput : FMaterialInput<FVector2D>
+{
+	int32 Compile(class FMaterialCompiler* Compiler, const FVector2D& Default);
+};
+
+struct FMaterialAttributesInput : FMaterialInput<int32>
+{
+	ENGINE_API int32 Compile(class FMaterialCompiler* Compiler, EMaterialProperty Property, float DefaultFloat, const FColor& DefaultColor, const FVector& DefaultVector);
+};
\ No newline at end of file
diff --git a/Engine/Source/Runtime/Engine/Public/MaterialShared.h b/Engine/Source/Runtime/Engine/Public/MaterialShared.h
--- a/Engine/Source/Runtime/Engine/Public/MaterialShared.h
+++ b/Engine/Source/Runtime/Engine/Public/MaterialShared.h
@@ -25,6 +25,7 @@ class UMaterial;
 class UMaterialInstance;
 class UMaterialExpression;
 class UMaterialInterface;
+struct FExpressionInput;
 
 #define ME_CAPTION_HEIGHT		18
 #define ME_STD_VPADDING			16
@@ -1378,121 +1379,11 @@ public:
 	virtual bool GetTextureValue(const FName ParameterName,const UTexture** OutValue, const FMaterialRenderContext& Context) const;
 };
 
-//
-//	FExpressionInput
-//
 
-//@warning: FExpressionInput is mirrored in MaterialExpression.h and manually "subclassed" in Material.h (FMaterialInput)
-struct FExpressionInput
-{
-	/** Material expression that this input is connected to, or NULL if not connected. */
-	class UMaterialExpression*	Expression;
-
-	/** Index into Expression's outputs array that this input is connected to. */
-	int32						OutputIndex;
-
-	/** 
-	 * Optional name of the input.  
-	 * Note that this is the only member which is not derived from the output currently connected. 
-	 */
-	FString						InputName;
-
-	int32						Mask,
-								MaskR,
-								MaskG,
-								MaskB,
-								MaskA;
-	uint32						GCC64Padding; // @todo 64: if the C++ didn't mismirror this structure, we might not need this
-
-	FExpressionInput()
-		: Expression(NULL)
-		, OutputIndex(0)
-		, Mask(0)
-		, MaskR(0)
-		, MaskG(0)
-		, MaskB(0)
-		, MaskA(0)
-		, GCC64Padding(0)
-	{
-	}
-
-	int32 Compile(class FMaterialCompiler* Compiler, int32 MultiplexIndex=INDEX_NONE);
-
-	/**
-	 * Tests if the input has a material expression connected to it
-	 *
-	 * @return	true if an expression is connected, otherwise false
-	 */
-	bool IsConnected() const { return (NULL != Expression); }
 
-	/** Connects output of InExpression to this input */
-	ENGINE_API void Connect( int32 InOutputIndex, class UMaterialExpression* InExpression );
-};
 
-//
-//	FMaterialInput
-//
 
-template<class InputType> struct FMaterialInput: FExpressionInput
-{
-	uint32	UseConstant:1;
-	InputType	Constant;
-};
 
-struct FColorMaterialInput: FMaterialInput<FColor>
-{
-	ENGINE_API int32 Compile(class FMaterialCompiler* Compiler,const FColor& Default);
-};
-struct FScalarMaterialInput: FMaterialInput<float>
-{
-	ENGINE_API int32 Compile(class FMaterialCompiler* Compiler,float Default);
-};
-
-struct FVectorMaterialInput: FMaterialInput<FVector>
-{
-	ENGINE_API int32 Compile(class FMaterialCompiler* Compiler,const FVector& Default);
-};
-
-struct FVector2MaterialInput: FMaterialInput<FVector2D>
-{
-	int32 Compile(class FMaterialCompiler* Compiler,const FVector2D& Default);
-};
-
-struct FMaterialAttributesInput: FMaterialInput<int32>
-{
-	ENGINE_API int32 Compile(class FMaterialCompiler* Compiler, EMaterialProperty Property, float DefaultFloat, const FColor& DefaultColor, const FVector& DefaultVector);
-};
-
-//
-//	FExpressionOutput
-//
-
-struct FExpressionOutput
-{
-	FString	OutputName;
-	int32	Mask,
-			MaskR,
-			MaskG,
-			MaskB,
-			MaskA;
-
-	FExpressionOutput(int32 InMask = 0,int32 InMaskR = 0,int32 InMaskG = 0,int32 InMaskB = 0,int32 InMaskA = 0):
-		Mask(InMask),
-		MaskR(InMaskR),
-		MaskG(InMaskG),
-		MaskB(InMaskB),
-		MaskA(InMaskA)
-	{}
-
-	FExpressionOutput(FString InOutputName, int32 InMask = 0,int32 InMaskR = 0,int32 InMaskG = 0,int32 InMaskB = 0,int32 InMaskA = 0):
-		OutputName(InOutputName),
-		Mask(InMask),
-		MaskR(InMaskR),
-		MaskG(InMaskG),
-		MaskB(InMaskB),
-		MaskA(InMaskA)
-	{}
-};
 
 /**
  * @return True if BlendMode is translucent (should be part of the translucent rendering).
diff --git a/Engine/Source/Runtime/Engine/Public/Model.h b/Engine/Source/Runtime/Engine/Public/Model.h
--- a/Engine/Source/Runtime/Engine/Public/Model.h
+++ b/Engine/Source/Runtime/Engine/Public/Model.h
@@ -9,8 +9,7 @@
 #include "RawIndexBuffer.h"
 #include "LocalVertexFactory.h"
 #include "RenderUtils.h"
-#include "LightMap.h"
-#include "ShadowMap.h"
+#include "SceneTypes.h"
 
 struct FStaticLightingVertex;
 class ABrush;
@@ -614,17 +613,9 @@ public:
 	/**
 	 * Minimal initialization constructor.
 	 */
-	FModelElement(UModelComponent* InComponent,UMaterialInterface* InMaterial);
-	FModelElement() 
-		: Component(NULL)
-		, Material(NULL)
-		, IndexBuffer(NULL)
-		, FirstIndex(0)
-		, NumTriangles(0)
-		, MinVertexIndex(0)
-		, MaxVertexIndex(0)
-		, BoundingBox(ForceInitToZero)
-	{}
+	ENGINE_API FModelElement(UModelComponent* InComponent,UMaterialInterface* InMaterial);
+	ENGINE_API FModelElement();
+	ENGINE_API virtual ~FModelElement();
 
 	/**
 	 * Serializer.
diff --git a/Engine/Source/Runtime/Engine/Public/SceneTypes.h b/Engine/Source/Runtime/Engine/Public/SceneTypes.h
--- a/Engine/Source/Runtime/Engine/Public/SceneTypes.h
+++ b/Engine/Source/Runtime/Engine/Public/SceneTypes.h
@@ -5,6 +5,14 @@
 class FLightMap;
 class FShadowMap;
 class FSceneViewStateInterface;
+class FLightMap;
+class FShadowMap;
+
+/** A reference to a light-map. */
+typedef TRefCountPtr<FLightMap> FLightMapRef;
+
+/** A reference to a shadow-map. */
+typedef TRefCountPtr<FShadowMap> FShadowMapRef;
 
 /** 
  * Class used to identify UPrimitiveComponents on the rendering thread without having to pass the pointer around, 
diff --git a/Engine/Source/Runtime/Engine/Public/ShadowMap.h b/Engine/Source/Runtime/Engine/Public/ShadowMap.h
--- a/Engine/Source/Runtime/Engine/Public/ShadowMap.h
+++ b/Engine/Source/Runtime/Engine/Public/ShadowMap.h
@@ -342,8 +342,7 @@ private:
 };
 
 
-/** A reference to a shadow-map. */
-typedef TRefCountPtr<FShadowMap> FShadowMapRef;
+
 
 /** Shadowmap reference serializer */
 extern FArchive& operator<<(FArchive& Ar,FShadowMap*& R);
diff --git a/Engine/Source/Runtime/Engine/Public/StaticLighting.h b/Engine/Source/Runtime/Engine/Public/StaticLighting.h
--- a/Engine/Source/Runtime/Engine/Public/StaticLighting.h
+++ b/Engine/Source/Runtime/Engine/Public/StaticLighting.h
@@ -7,6 +7,8 @@
 #pragma once
 
 #include "VertexFactory.h"
+#include "LightMap.h"
+#include "ShadowMap.h"
 
 // Forward declarations.
 class FStaticLightingTextureMapping;
diff --git a/Engine/Source/Runtime/Engine/Public/StaticMeshResources.h b/Engine/Source/Runtime/Engine/Public/StaticMeshResources.h
--- a/Engine/Source/Runtime/Engine/Public/StaticMeshResources.h
+++ b/Engine/Source/Runtime/Engine/Public/StaticMeshResources.h
@@ -949,15 +949,10 @@ protected:
 		// FLightCacheInterface.
 		virtual FLightInteraction GetInteraction(const FLightSceneProxy* LightSceneProxy) const;
 
-		virtual FLightMapInteraction GetLightMapInteraction() const
-		{
-			return LightMap ? LightMap->GetInteraction() : FLightMapInteraction();
-		}
+		virtual FLightMapInteraction GetLightMapInteraction() const;
+
+		virtual FShadowMapInteraction GetShadowMapInteraction() const;
 
-		virtual FShadowMapInteraction GetShadowMapInteraction() const
-		{
-			return ShadowMap ? ShadowMap->GetInteraction() : FShadowMapInteraction();
-		}
 
 	private:
 
diff --git a/Engine/Source/Runtime/RenderCore/Public/PackedNormal.h b/Engine/Source/Runtime/RenderCore/Public/PackedNormal.h
new file mode 100644
--- /dev/null
+++ b/Engine/Source/Runtime/RenderCore/Public/PackedNormal.h
@@ -0,0 +1,139 @@
+// Copyright 1998-2014 Epic Games, Inc. All Rights Reserved.
+
+
+#pragma once
+
+/** A normal vector, quantized and packed into 32-bits. */
+struct FPackedNormal
+{
+	union
+	{
+		struct
+		{
+#if PLATFORM_LITTLE_ENDIAN
+			uint8	X,
+			Y,
+			Z,
+			W;
+#else
+			uint8	W,
+			Z,
+			Y,
+			X;
+#endif
+		};
+		uint32		Packed;
+	}				Vector;
+
+	// Constructors.
+
+	FPackedNormal() { Vector.Packed = 0; }
+	FPackedNormal(uint32 InPacked) { Vector.Packed = InPacked; }
+	FPackedNormal(const FVector& InVector) { *this = InVector; }
+	FPackedNormal(uint8 InX, uint8 InY, uint8 InZ, uint8 InW) { Vector.X = InX; Vector.Y = InY; Vector.Z = InZ; Vector.W = InW; }
+
+	// Conversion operators.
+
+	void operator=(const FVector& InVector);
+	void operator=(const FVector4& InVector);
+	operator FVector() const;
+	VectorRegister GetVectorRegister() const;
+
+	// Set functions.
+	void Set(const FVector& InVector) { *this = InVector; }
+
+	// Equality operator.
+
+	bool operator==(const FPackedNormal& B) const;
+	bool operator!=(const FPackedNormal& B) const;
+
+	// Serializer.
+
+	friend RENDERCORE_API FArchive& operator<<(FArchive& Ar, FPackedNormal& N);
+
+	FString ToString() const
+	{
+		return FString::Printf(TEXT("X=%d Y=%d Z=%d W=%d"), Vector.X, Vector.Y, Vector.Z, Vector.W);
+	}
+
+	// Zero Normal
+	static RENDERCORE_API FPackedNormal ZeroNormal;
+};
+
+/** X=127.5, Y=127.5, Z=1/127.5f, W=-1.0 */
+extern RENDERCORE_API const VectorRegister GVectorPackingConstants;
+
+FORCEINLINE void FPackedNormal::operator=(const FVector& InVector)
+{
+	Vector.X = FMath::Clamp(FMath::TruncToInt(InVector.X * 127.5f + 127.5f),0,255);
+	Vector.Y = FMath::Clamp(FMath::TruncToInt(InVector.Y * 127.5f + 127.5f),0,255);
+	Vector.Z = FMath::Clamp(FMath::TruncToInt(InVector.Z * 127.5f + 127.5f),0,255);
+	Vector.W = 128;
+}
+
+FORCEINLINE void FPackedNormal::operator=(const FVector4& InVector)
+{
+	Vector.X = FMath::Clamp(FMath::TruncToInt(InVector.X * 127.5f + 127.5f),0,255);
+	Vector.Y = FMath::Clamp(FMath::TruncToInt(InVector.Y * 127.5f + 127.5f),0,255);
+	Vector.Z = FMath::Clamp(FMath::TruncToInt(InVector.Z * 127.5f + 127.5f),0,255);
+	Vector.W = FMath::Clamp(FMath::TruncToInt(InVector.W * 127.5f + 127.5f),0,255);
+}
+
+FORCEINLINE bool FPackedNormal::operator==(const FPackedNormal& B) const
+{
+	if(Vector.Packed != B.Vector.Packed)
+		return 0;
+
+	FVector	V1 = *this,
+			V2 = B;
+
+	if(FMath::Abs(V1.X - V2.X) > THRESH_NORMALS_ARE_SAME * 4.0f)
+		return 0;
+
+	if(FMath::Abs(V1.Y - V2.Y) > THRESH_NORMALS_ARE_SAME * 4.0f)
+		return 0;
+
+	if(FMath::Abs(V1.Z - V2.Z) > THRESH_NORMALS_ARE_SAME * 4.0f)
+		return 0;
+
+	return 1;
+}
+
+FORCEINLINE bool FPackedNormal::operator!=(const FPackedNormal& B) const
+{
+	if(Vector.Packed == B.Vector.Packed)
+		return 0;
+
+	FVector	V1 = *this,
+			V2 = B;
+
+	if(FMath::Abs(V1.X - V2.X) > THRESH_NORMALS_ARE_SAME * 4.0f)
+		return 1;
+
+	if(FMath::Abs(V1.Y - V2.Y) > THRESH_NORMALS_ARE_SAME * 4.0f)
+		return 1;
+
+	if(FMath::Abs(V1.Z - V2.Z) > THRESH_NORMALS_ARE_SAME * 4.0f)
+		return 1;
+
+	return 0;
+}
+
+FORCEINLINE FPackedNormal::operator FVector() const
+{
+	VectorRegister VectorToUnpack = GetVectorRegister();
+	// Write to FVector and return it.
+	FVector UnpackedVector;
+	VectorStoreFloat3( VectorToUnpack, &UnpackedVector );
+	return UnpackedVector;
+}
+
+FORCEINLINE VectorRegister FPackedNormal::GetVectorRegister() const
+{
+	// Rescale [0..255] range to [-1..1]
+	VectorRegister VectorToUnpack		= VectorLoadByte4( this );
+	VectorToUnpack						= VectorMultiplyAdd( VectorToUnpack, VectorReplicate(GVectorPackingConstants,2), VectorReplicate(GVectorPackingConstants,3) );
+	VectorResetFloatRegisters();
+	// Return unpacked vector register.
+	return VectorToUnpack;
+}
diff --git a/Engine/Source/Runtime/RenderCore/Public/PackedNormal.inl b/Engine/Source/Runtime/RenderCore/Public/PackedNormal.inl
deleted file mode 100644
diff --git a/Engine/Source/Runtime/RenderCore/Public/RenderUtils.h b/Engine/Source/Runtime/RenderCore/Public/RenderUtils.h
--- a/Engine/Source/Runtime/RenderCore/Public/RenderUtils.h
+++ b/Engine/Source/Runtime/RenderCore/Public/RenderUtils.h
@@ -5,69 +5,7 @@
 
 #include "RenderCore.h"
 #include "RenderResource.h"
-
-/** A normal vector, quantized and packed into 32-bits. */
-struct FPackedNormal
-{
-	union
-	{
-		struct
-		{
-#if PLATFORM_LITTLE_ENDIAN
-			uint8	X,
-					Y,
-					Z,
-					W;
-#else
-			uint8	W,
-					Z,
-					Y,
-					X;
-#endif
-		};
-		uint32		Packed;
-	}				Vector;
-
-	// Constructors.
-
-	FPackedNormal( ) { Vector.Packed = 0; }
-	FPackedNormal( uint32 InPacked ) { Vector.Packed = InPacked; }
-	FPackedNormal( const FVector& InVector ) { *this = InVector; }
-	FPackedNormal( uint8 InX, uint8 InY, uint8 InZ, uint8 InW ) { Vector.X = InX; Vector.Y = InY; Vector.Z = InZ; Vector.W = InW;}
-
-	// Conversion operators.
-
-	void operator=(const FVector& InVector);
-	void operator=(const FVector4& InVector);
-	operator FVector() const;
-	VectorRegister GetVectorRegister() const;
-
-	// Set functions.
-	void Set( const FVector& InVector ) { *this = InVector; }
-
-	// Equality operator.
-
-	bool operator==(const FPackedNormal& B) const;
-	bool operator!=(const FPackedNormal& B) const;
-
-	// Serializer.
-
-	friend RENDERCORE_API FArchive& operator<<(FArchive& Ar,FPackedNormal& N);
-	
-	FString ToString() const
-	{
-		return FString::Printf(TEXT("X=%d Y=%d Z=%d W=%d"), Vector.X, Vector.Y, Vector.Z, Vector.W);
-	}
-
-	// Zero Normal
-	static RENDERCORE_API FPackedNormal ZeroNormal;
-};
-
-/** X=127.5, Y=127.5, Z=1/127.5f, W=-1.0 */
-extern RENDERCORE_API const VectorRegister GVectorPackingConstants;
-
-// Include the packed normal inline functions.
-#include "PackedNormal.inl"
+#include "PackedNormal.h"
 
 /**
 * Constructs a basis matrix for the axis vectors and returns the sign of the determinant
diff --git a/Engine/Source/Runtime/Renderer/Private/LightMapRendering.h b/Engine/Source/Runtime/Renderer/Private/LightMapRendering.h
--- a/Engine/Source/Runtime/Renderer/Private/LightMapRendering.h
+++ b/Engine/Source/Runtime/Renderer/Private/LightMapRendering.h
@@ -8,7 +8,7 @@
 #define __LIGHTMAPRENDERING_H__
 
 extern ENGINE_API bool GShowDebugSelectedLightmap;
-extern ENGINE_API FLightMap2D* GDebugSelectedLightmap;
+extern ENGINE_API class FLightMap2D* GDebugSelectedLightmap;
 extern bool GVisualizeMipLevels;
 
 /**
