Don't ever try to allocate a global dynamic buffer of size 0.

diff --git a/Engine/Source/Runtime/RenderCore/Private/RenderResource.cpp b/Engine/Source/Runtime/RenderCore/Private/RenderResource.cpp
--- a/Engine/Source/Runtime/RenderCore/Private/RenderResource.cpp
+++ b/Engine/Source/Runtime/RenderCore/Private/RenderResource.cpp
@@ -228,7 +228,7 @@ public:
 	/** Default constructor. */
 	explicit FDynamicVertexBuffer(uint32 InMinBufferSize)
 		: MappedBuffer(NULL)
-		, BufferSize(Align(InMinBufferSize,ALIGNMENT))
+		, BufferSize(FMath::Max<uint32>(Align(InMinBufferSize,ALIGNMENT),ALIGNMENT))
 		, AllocatedByteCount(0)
 	{
 	}
@@ -409,7 +409,7 @@ public:
 	/** Initialization constructor. */
 	explicit FDynamicIndexBuffer(uint32 InMinBufferSize, uint32 InStride)
 		: MappedBuffer(NULL)
-		, BufferSize(Align(InMinBufferSize,ALIGNMENT))
+		, BufferSize(FMath::Max<uint32>(Align(InMinBufferSize,ALIGNMENT),ALIGNMENT))
 		, AllocatedByteCount(0)
 		, Stride(InStride)
 	{
