Fix for "ignored qualifiers" warning issued by clang 3.5.0

- The change is authored by [REDACTED_PERSON].
- Can make code more efficient by passing some args by value.
- Ugly due to typename not allowed outside of templates by MSVC (error C2899).

diff --git a/Engine/Source/Programs/UnrealBuildTool/Windows/VCToolChain.cs b/Engine/Source/Programs/UnrealBuildTool/Windows/VCToolChain.cs
--- a/Engine/Source/Programs/UnrealBuildTool/Windows/VCToolChain.cs
+++ b/Engine/Source/Programs/UnrealBuildTool/Windows/VCToolChain.cs
@@ -379,9 +379,6 @@ namespace UnrealBuildTool
 				// Allow Microsoft-specific syntax to slide, even though it may be non-standard.  Needed for Windows headers.
 				Arguments.Append(" -Wno-microsoft");								
 
-				// Don't complain about use of 'const' on reference types when it will have no effect
-				Arguments.Append( " -Wno-ignored-qualifiers" );	// @todo clang: This triggers for ENQUEUE_UNIQUE_RENDER_COMMAND_ONEPARAMETER  due to EURCMacro_##TypeName(const _ParamType1& In##ParamName1)
-
 				// @todo clang: Kind of a shame to turn these off.  We'd like to catch unused variables, but it is tricky with how our assertion macros work.
 				Arguments.Append(" -Wno-unused-variable");
 				Arguments.Append(" -Wno-unused-function");
diff --git a/Engine/Source/Runtime/RenderCore/Public/RenderingThread.h b/Engine/Source/Runtime/RenderCore/Public/RenderingThread.h
--- a/Engine/Source/Runtime/RenderCore/Public/RenderingThread.h
+++ b/Engine/Source/Runtime/RenderCore/Public/RenderingThread.h
@@ -204,12 +204,11 @@ DECLARE_STATS_GROUP(TEXT("Render Thread Commands"), STATGROUP_RenderThreadComman
 /**
  * Declares a rendering command type with 1 parameters.
  */
-#define ENQUEUE_UNIQUE_RENDER_COMMAND_ONEPARAMETER_DECLARE(TypeName,ParamType1,ParamName1,ParamValue1,Code) \
+#define ENQUEUE_UNIQUE_RENDER_COMMAND_ONEPARAMETER_DECLARE_OPTTYPENAME(TypeName,ParamType1,ParamName1,ParamValue1,OptTypename,Code) \
 	class EURCMacro_##TypeName : public FRenderCommand \
 	{ \
 	public: \
-		typedef ParamType1 _ParamType1; \
-		EURCMacro_##TypeName(const _ParamType1& In##ParamName1): \
+		EURCMacro_##TypeName(OptTypename TCallTraits<ParamType1>::ParamType In##ParamName1): \
 		  ParamName1(In##ParamName1) \
 		{} \
 		TASK_FUNCTION(Code) \
@@ -217,6 +216,8 @@ DECLARE_STATS_GROUP(TEXT("Render Thread Commands"), STATGROUP_RenderThreadComman
 	private: \
 		ParamType1 ParamName1; \
 	};
+#define ENQUEUE_UNIQUE_RENDER_COMMAND_ONEPARAMETER_DECLARE(TypeName,ParamType1,ParamName1,ParamValue1,Code) \
+	ENQUEUE_UNIQUE_RENDER_COMMAND_ONEPARAMETER_DECLARE_OPTTYPENAME(TypeName,ParamType1,ParamName1,ParamValue1,,Code)
 
 #define ENQUEUE_UNIQUE_RENDER_COMMAND_ONEPARAMETER_CREATE(TypeName,ParamType1,ParamValue1) \
 	{ \
@@ -239,7 +240,7 @@ DECLARE_STATS_GROUP(TEXT("Render Thread Commands"), STATGROUP_RenderThreadComman
 
 #define ENQUEUE_UNIQUE_RENDER_COMMAND_ONEPARAMETER_DECLARE_TEMPLATE(TypeName,TemplateParamName,ParamType1,ParamName1,ParamValue1,Code) \
 	template <typename TemplateParamName> \
-	ENQUEUE_UNIQUE_RENDER_COMMAND_ONEPARAMETER_DECLARE(TypeName,ParamType1,ParamName1,ParamValue1,Code)
+	ENQUEUE_UNIQUE_RENDER_COMMAND_ONEPARAMETER_DECLARE_OPTTYPENAME(TypeName,ParamType1,ParamName1,ParamValue1,typename,Code)
 
 #define ENQUEUE_UNIQUE_RENDER_COMMAND_ONEPARAMETER_CREATE_TEMPLATE(TypeName,TemplateParamName,ParamType1,ParamValue1) \
 	ENQUEUE_UNIQUE_RENDER_COMMAND_ONEPARAMETER_CREATE(TypeName<TemplateParamName>,ParamType1,ParamValue1)
@@ -247,13 +248,11 @@ DECLARE_STATS_GROUP(TEXT("Render Thread Commands"), STATGROUP_RenderThreadComman
 /**
  * Declares a rendering command type with 2 parameters.
  */
-#define ENQUEUE_UNIQUE_RENDER_COMMAND_TWOPARAMETER_DECLARE(TypeName,ParamType1,ParamName1,ParamValue1,ParamType2,ParamName2,ParamValue2,Code) \
+#define ENQUEUE_UNIQUE_RENDER_COMMAND_TWOPARAMETER_DECLARE_OPTTYPENAME(TypeName,ParamType1,ParamName1,ParamValue1,ParamType2,ParamName2,ParamValue2,OptTypename,Code) \
 	class EURCMacro_##TypeName : public FRenderCommand \
 	{ \
 	public: \
-		typedef ParamType1 _ParamType1; \
-		typedef ParamType2 _ParamType2; \
-		EURCMacro_##TypeName(const _ParamType1& In##ParamName1,const _ParamType2& In##ParamName2): \
+		EURCMacro_##TypeName(OptTypename TCallTraits<ParamType1>::ParamType In##ParamName1,OptTypename TCallTraits<ParamType2>::ParamType In##ParamName2): \
 		  ParamName1(In##ParamName1), \
 		  ParamName2(In##ParamName2) \
 		{} \
@@ -263,6 +262,8 @@ DECLARE_STATS_GROUP(TEXT("Render Thread Commands"), STATGROUP_RenderThreadComman
 		ParamType1 ParamName1; \
 		ParamType2 ParamName2; \
 	};
+#define ENQUEUE_UNIQUE_RENDER_COMMAND_TWOPARAMETER_DECLARE(TypeName,ParamType1,ParamName1,ParamValue1,ParamType2,ParamName2,ParamValue2,Code) \
+	ENQUEUE_UNIQUE_RENDER_COMMAND_TWOPARAMETER_DECLARE_OPTTYPENAME(TypeName,ParamType1,ParamName1,ParamValue1,ParamType2,ParamName2,ParamValue2,,Code)
 
 #define ENQUEUE_UNIQUE_RENDER_COMMAND_TWOPARAMETER_CREATE(TypeName,ParamType1,ParamValue1,ParamType2,ParamValue2) \
 	{ \
@@ -285,7 +286,7 @@ DECLARE_STATS_GROUP(TEXT("Render Thread Commands"), STATGROUP_RenderThreadComman
 
 #define ENQUEUE_UNIQUE_RENDER_COMMAND_TWOPARAMETER_DECLARE_TEMPLATE(TypeName,TemplateParamName,ParamType1,ParamName1,ParamValue1,ParamType2,ParamName2,ParamValue2,Code) \
 	template <typename TemplateParamName> \
-	ENQUEUE_UNIQUE_RENDER_COMMAND_TWOPARAMETER_DECLARE(TypeName,ParamType1,ParamName1,ParamValue1,ParamType2,ParamName2,ParamValue2,Code)
+	ENQUEUE_UNIQUE_RENDER_COMMAND_TWOPARAMETER_DECLARE_OPTTYPENAME(TypeName,ParamType1,ParamName1,ParamValue1,ParamType2,ParamName2,ParamValue2,typename,Code)
 
 #define ENQUEUE_UNIQUE_RENDER_COMMAND_TWOPARAMETER_CREATE_TEMPLATE(TypeName,TemplateParamName,ParamType1,ParamValue1,ParamType2,ParamValue2) \
 	ENQUEUE_UNIQUE_RENDER_COMMAND_TWOPARAMETER_CREATE(TypeName<TemplateParamName>,ParamType1,ParamValue1,ParamType2,ParamValue2)
@@ -293,14 +294,11 @@ DECLARE_STATS_GROUP(TEXT("Render Thread Commands"), STATGROUP_RenderThreadComman
 /**
  * Declares a rendering command type with 3 parameters.
  */
-#define ENQUEUE_UNIQUE_RENDER_COMMAND_THREEPARAMETER_DECLARE(TypeName,ParamType1,ParamName1,ParamValue1,ParamType2,ParamName2,ParamValue2,ParamType3,ParamName3,ParamValue3,Code) \
+#define ENQUEUE_UNIQUE_RENDER_COMMAND_THREEPARAMETER_DECLARE_OPTTYPENAME(TypeName,ParamType1,ParamName1,ParamValue1,ParamType2,ParamName2,ParamValue2,ParamType3,ParamName3,ParamValue3,OptTypename,Code) \
 	class EURCMacro_##TypeName : public FRenderCommand \
 	{ \
 	public: \
-		typedef ParamType1 _ParamType1; \
-		typedef ParamType2 _ParamType2; \
-		typedef ParamType3 _ParamType3; \
-		EURCMacro_##TypeName(const _ParamType1& In##ParamName1,const _ParamType2& In##ParamName2,const _ParamType3& In##ParamName3): \
+		EURCMacro_##TypeName(OptTypename TCallTraits<ParamType1>::ParamType In##ParamName1,OptTypename TCallTraits<ParamType2>::ParamType In##ParamName2,OptTypename TCallTraits<ParamType3>::ParamType In##ParamName3): \
 		  ParamName1(In##ParamName1), \
 		  ParamName2(In##ParamName2), \
 		  ParamName3(In##ParamName3) \
@@ -312,6 +310,8 @@ DECLARE_STATS_GROUP(TEXT("Render Thread Commands"), STATGROUP_RenderThreadComman
 		ParamType2 ParamName2; \
 		ParamType3 ParamName3; \
 	};
+#define ENQUEUE_UNIQUE_RENDER_COMMAND_THREEPARAMETER_DECLARE(TypeName,ParamType1,ParamName1,ParamValue1,ParamType2,ParamName2,ParamValue2,ParamType3,ParamName3,ParamValue3,Code) \
+	ENQUEUE_UNIQUE_RENDER_COMMAND_THREEPARAMETER_DECLARE_OPTTYPENAME(TypeName,ParamType1,ParamName1,ParamValue1,ParamType2,ParamName2,ParamValue2,ParamType3,ParamName3,ParamValue3,,Code)
 
 #define ENQUEUE_UNIQUE_RENDER_COMMAND_THREEPARAMETER_CREATE(TypeName,ParamType1,ParamValue1,ParamType2,ParamValue2,ParamType3,ParamValue3) \
 	{ \
@@ -334,7 +334,7 @@ DECLARE_STATS_GROUP(TEXT("Render Thread Commands"), STATGROUP_RenderThreadComman
 
 #define ENQUEUE_UNIQUE_RENDER_COMMAND_THREEPARAMETER_DECLARE_TEMPLATE(TypeName,TemplateParamName,ParamType1,ParamName1,ParamValue1,ParamType2,ParamName2,ParamValue2,ParamType3,ParamName3,ParamValue3,Code) \
 	template <typename TemplateParamName> \
-	ENQUEUE_UNIQUE_RENDER_COMMAND_THREEPARAMETER_DECLARE(TypeName,ParamType1,ParamName1,ParamValue1,ParamType2,ParamName2,ParamValue2,ParamType3,ParamName3,ParamValue3,Code)
+	ENQUEUE_UNIQUE_RENDER_COMMAND_THREEPARAMETER_DECLARE_OPTTYPENAME(TypeName,ParamType1,ParamName1,ParamValue1,ParamType2,ParamName2,ParamValue2,ParamType3,ParamName3,ParamValue3,typename,Code)
 
 #define ENQUEUE_UNIQUE_RENDER_COMMAND_THREEPARAMETER_CREATE_TEMPLATE(TypeName,TemplateParamName,ParamType1,ParamValue1,ParamType2,ParamValue2,ParamType3,ParamValue3) \
 	ENQUEUE_UNIQUE_RENDER_COMMAND_THREEPARAMETER_CREATE(TypeName<TemplateParamName>,ParamType1,ParamValue1,ParamType2,ParamValue2,ParamType3,ParamValue3)
@@ -342,15 +342,11 @@ DECLARE_STATS_GROUP(TEXT("Render Thread Commands"), STATGROUP_RenderThreadComman
 /**
  * Declares a rendering command type with 4 parameters.
  */
-#define ENQUEUE_UNIQUE_RENDER_COMMAND_FOURPARAMETER_DECLARE(TypeName,ParamType1,ParamName1,ParamValue1,ParamType2,ParamName2,ParamValue2,ParamType3,ParamName3,ParamValue3,ParamType4,ParamName4,ParamValue4,Code) \
+#define ENQUEUE_UNIQUE_RENDER_COMMAND_FOURPARAMETER_DECLARE_OPTTYPENAME(TypeName,ParamType1,ParamName1,ParamValue1,ParamType2,ParamName2,ParamValue2,ParamType3,ParamName3,ParamValue3,ParamType4,ParamName4,ParamValue4,OptTypename,Code) \
 	class EURCMacro_##TypeName : public FRenderCommand \
 	{ \
 	public: \
-		typedef ParamType1 _ParamType1; \
-		typedef ParamType2 _ParamType2; \
-		typedef ParamType3 _ParamType3; \
-		typedef ParamType4 _ParamType4; \
-		EURCMacro_##TypeName(const _ParamType1& In##ParamName1,const _ParamType2& In##ParamName2,const _ParamType3& In##ParamName3,const _ParamType4& In##ParamName4): \
+		EURCMacro_##TypeName(OptTypename TCallTraits<ParamType1>::ParamType In##ParamName1,OptTypename TCallTraits<ParamType2>::ParamType In##ParamName2,OptTypename TCallTraits<ParamType3>::ParamType In##ParamName3,OptTypename TCallTraits<ParamType4>::ParamType In##ParamName4): \
 		  ParamName1(In##ParamName1), \
 		  ParamName2(In##ParamName2), \
 		  ParamName3(In##ParamName3), \
@@ -364,6 +360,9 @@ DECLARE_STATS_GROUP(TEXT("Render Thread Commands"), STATGROUP_RenderThreadComman
 		ParamType3 ParamName3; \
 		ParamType4 ParamName4; \
 	};
+#define ENQUEUE_UNIQUE_RENDER_COMMAND_FOURPARAMETER_DECLARE(TypeName,ParamType1,ParamName1,ParamValue1,ParamType2,ParamName2,ParamValue2,ParamType3,ParamName3,ParamValue3,ParamType4,ParamName4,ParamValue4,Code) \
+	ENQUEUE_UNIQUE_RENDER_COMMAND_FOURPARAMETER_DECLARE_OPTTYPENAME(TypeName,ParamType1,ParamName1,ParamValue1,ParamType2,ParamName2,ParamValue2,ParamType3,ParamName3,ParamValue3,ParamType4,ParamName4,ParamValue4,,Code)
+
 
 #define ENQUEUE_UNIQUE_RENDER_COMMAND_FOURPARAMETER_CREATE(TypeName,ParamType1,ParamValue1,ParamType2,ParamValue2,ParamType3,ParamValue3,ParamType4,ParamValue4) \
 	{ \
@@ -386,7 +385,7 @@ DECLARE_STATS_GROUP(TEXT("Render Thread Commands"), STATGROUP_RenderThreadComman
 
 #define ENQUEUE_UNIQUE_RENDER_COMMAND_FOURPARAMETER_DECLARE_TEMPLATE(TypeName,TemplateParamName,ParamType1,ParamName1,ParamValue1,ParamType2,ParamName2,ParamValue2,ParamType3,ParamName3,ParamValue3,ParamType4,ParamName4,ParamValue4,Code) \
 	template <typename TemplateParamName> \
-	ENQUEUE_UNIQUE_RENDER_COMMAND_FOURPARAMETER_DECLARE(TypeName,ParamType1,ParamName1,ParamValue1,ParamType2,ParamName2,ParamValue2,ParamType3,ParamName3,ParamValue3,ParamType4,ParamName4,ParamValue4,Code)
+	ENQUEUE_UNIQUE_RENDER_COMMAND_FOURPARAMETER_DECLARE_OPTTYPENAME(TypeName,ParamType1,ParamName1,ParamValue1,ParamType2,ParamName2,ParamValue2,ParamType3,ParamName3,ParamValue3,ParamType4,ParamName4,ParamValue4,OptTypename,Code)
 
 #define ENQUEUE_UNIQUE_RENDER_COMMAND_FOURPARAMETER_CREATE_TEMPLATE(TypeName,TemplateParamName,ParamType1,ParamValue1,ParamType2,ParamValue2,ParamType3,ParamValue3,ParamType4,ParamValue4) \
 	ENQUEUE_UNIQUE_RENDER_COMMAND_FOURPARAMETER_CREATE(TypeName<TemplateParamName>,ParamType1,ParamValue1,ParamType2,ParamValue2,ParamType3,ParamValue3,ParamType4,ParamValue4)
