Uniform buffers can now use constructors. This is particularly useful for those that contain resources as they cannot be NULL.

diff --git a/Engine/Source/Runtime/Engine/Private/SceneManagement.cpp b/Engine/Source/Runtime/Engine/Private/SceneManagement.cpp
--- a/Engine/Source/Runtime/Engine/Private/SceneManagement.cpp
+++ b/Engine/Source/Runtime/Engine/Private/SceneManagement.cpp
@@ -286,3 +286,9 @@ int8 ComputeLODForMeshes( const TIndirectArray<class FStaticMesh>& StaticMeshes,
 	}
 	return LODToRender;
 }
+
+FViewUniformShaderParameters::FViewUniformShaderParameters()
+	: DirectionalLightShadowTexture(GWhiteTexture->TextureRHI)
+	, DirectionalLightShadowSampler(TStaticSamplerState<SF_Point, AM_Clamp, AM_Clamp, AM_Clamp>::GetRHI())
+{
+}
diff --git a/Engine/Source/Runtime/Engine/Public/SceneManagement.h b/Engine/Source/Runtime/Engine/Public/SceneManagement.h
--- a/Engine/Source/Runtime/Engine/Public/SceneManagement.h
+++ b/Engine/Source/Runtime/Engine/Public/SceneManagement.h
@@ -1518,7 +1518,7 @@ enum ETranslucencyVolumeCascade
 };
 
 /** The uniform shader parameters associated with a view. */
-BEGIN_UNIFORM_BUFFER_STRUCT(FViewUniformShaderParameters,ENGINE_API)
+BEGIN_UNIFORM_BUFFER_STRUCT_WITH_CONSTRUCTOR(FViewUniformShaderParameters,ENGINE_API)
 	DECLARE_UNIFORM_BUFFER_STRUCT_MEMBER(FMatrix,TranslatedWorldToClip)
 	DECLARE_UNIFORM_BUFFER_STRUCT_MEMBER(FMatrix,WorldToClip)
 	DECLARE_UNIFORM_BUFFER_STRUCT_MEMBER(FMatrix,TranslatedWorldToView)
diff --git a/Engine/Source/Runtime/RenderCore/Public/UniformBuffer.h b/Engine/Source/Runtime/RenderCore/Public/UniformBuffer.h
--- a/Engine/Source/Runtime/RenderCore/Public/UniformBuffer.h
+++ b/Engine/Source/Runtime/RenderCore/Public/UniformBuffer.h
@@ -580,12 +580,13 @@ private:
 	sizeof(StructTypeName), \
 	StructTypeName::zzGetMembers(), \
 	true);
-	
+
 /** Begins a uniform buffer struct declaration. */
-#define BEGIN_UNIFORM_BUFFER_STRUCT(StructTypeName,PrefixKeywords) \
+#define BEGIN_UNIFORM_BUFFER_STRUCT_EX(StructTypeName,PrefixKeywords,ConstructorSuffix) \
 	MS_ALIGN(UNIFORM_BUFFER_STRUCT_ALIGNMENT) class PrefixKeywords StructTypeName \
 	{ \
 	public: \
+		StructTypeName () ConstructorSuffix \
 		static FUniformBufferStruct StaticStruct; \
 		static FShaderUniformBufferParameter* ConstructUniformBufferParameter() { return new TShaderUniformBufferParameter<StructTypeName>(); } \
 		static FUniformBufferRHIRef CreateUniformBuffer(const StructTypeName& InContents, EUniformBufferUsage InUsage) \
@@ -669,6 +670,8 @@ private:
 		static const FUniformBufferStruct* GetStruct() { return &Name::StaticStruct; } \
 	};
 
+#define BEGIN_UNIFORM_BUFFER_STRUCT(StructTypeName,PrefixKeywords) BEGIN_UNIFORM_BUFFER_STRUCT_EX(StructTypeName,PrefixKeywords,{})
+#define BEGIN_UNIFORM_BUFFER_STRUCT_WITH_CONSTRUCTOR(StructTypeName,PrefixKeywords) BEGIN_UNIFORM_BUFFER_STRUCT_EX(StructTypeName,PrefixKeywords,;)
 
 /** Finds the FUniformBufferStruct corresponding to the given name, or NULL if not found. */
 extern RENDERCORE_API FUniformBufferStruct* FindUniformBufferStructByName(const TCHAR* StructName);
\ No newline at end of file
