UE4 - make sure bypass is latched right away for RHI thread and don't clear stencil on PS4 if we don't have a write address

diff --git a/Engine/Source/Runtime/RenderCore/Private/RenderingThread.cpp b/Engine/Source/Runtime/RenderCore/Private/RenderingThread.cpp
--- a/Engine/Source/Runtime/RenderCore/Private/RenderingThread.cpp
+++ b/Engine/Source/Runtime/RenderCore/Private/RenderingThread.cpp
@@ -514,6 +514,7 @@ void StartRenderingThread()
 		FTaskGraphInterface::Get().WaitUntilTaskCompletes(CompletionEvent, ENamedThreads::GameThread_Local);
 		GRHIThread = FRHIThread::Get().Thread;
 		check(GRHIThread);
+		GRHICommandList.LatchBypass();
 	}
 #endif
 
