Stats - Added explicit flush for stats after startup the rendering thread, to prevent fragmentation and creating a massive amout of stats packets

diff --git a/Engine/Source/Runtime/RenderCore/Private/RenderingThread.cpp b/Engine/Source/Runtime/RenderCore/Private/RenderingThread.cpp
--- a/Engine/Source/Runtime/RenderCore/Private/RenderingThread.cpp
+++ b/Engine/Source/Runtime/RenderCore/Private/RenderingThread.cpp
@@ -263,6 +263,13 @@ void RenderingThreadMain( FEvent* TaskGraphBoundSyncEvent )
 	// set the thread back to real time mode
 	FPlatformProcess::SetRealTimeMode();
 
+#if STATS
+	if (FThreadStats::WillEverCollectData())
+	{
+		FThreadStats::ExplicitFlush(); // flush the stats and set update the scope so we don't flush again until a frame update, this helps prevent fragmentation
+	}
+#endif
+
 	check(GIsThreadedRendering);
 	FTaskGraphInterface::Get().ProcessThreadUntilRequestReturn(ENamedThreads::RenderThread);
 	FPlatformMisc::MemoryBarrier();
@@ -287,6 +294,7 @@ static void AdvanceRenderingThreadStats(int64 StatsFrame, int32 MasterDisableCha
 	{
 		Frame = -StatsFrame; // mark this as a bad frame
 	}
+	// @TODO yrx 2014-10-17 Add AddAdvanceFrame message
 	static FStatNameAndInfo Adv(NAME_AdvanceFrame, "", "", TEXT(""), EStatDataType::ST_int64, true, false);
 	FThreadStats::AddMessage(Adv.GetEncodedName(), EStatOperation::AdvanceFrameEventRenderThread, Frame);
 	if( IsInActualRenderingThread() )
