Removed DrawDynamicElements in Ocean game

diff --git a/Engine/Source/Runtime/Engine/Private/NiagaraEffectRenderer.cpp b/Engine/Source/Runtime/Engine/Private/NiagaraEffectRenderer.cpp
--- a/Engine/Source/Runtime/Engine/Private/NiagaraEffectRenderer.cpp
+++ b/Engine/Source/Runtime/Engine/Private/NiagaraEffectRenderer.cpp
@@ -6,7 +6,6 @@
 #include "Particles/ParticleResources.h"
 
 
-DECLARE_CYCLE_STAT(TEXT("PreRenderView"), STAT_NiagaraPreRenderView, STATGROUP_Niagara);
 DECLARE_CYCLE_STAT(TEXT("Render Total"), STAT_NiagaraRender, STATGROUP_Niagara);
 DECLARE_CYCLE_STAT(TEXT("Render Sprites"), STAT_NiagaraRenderSprites, STATGROUP_Niagara);
 DECLARE_CYCLE_STAT(TEXT("Render Ribbons"), STAT_NiagaraRenderRibbons, STATGROUP_Niagara);
diff --git a/Engine/Source/Runtime/RenderCore/Private/RenderCore.cpp b/Engine/Source/Runtime/RenderCore/Private/RenderCore.cpp
--- a/Engine/Source/Runtime/RenderCore/Private/RenderCore.cpp
+++ b/Engine/Source/Runtime/RenderCore/Private/RenderCore.cpp
@@ -78,7 +78,6 @@ DEFINE_STAT(STAT_UpdatePreshadowCache);
 DEFINE_STAT(STAT_CreateWholeSceneProjectedShadow);
 DEFINE_STAT(STAT_AddViewDependentWholeSceneShadowsForView);
 DEFINE_STAT(STAT_SetupInteractionShadows);
-DEFINE_STAT(STAT_PreRenderView);
 DEFINE_STAT(STAT_GetDynamicMeshElements);
 DEFINE_STAT(STAT_UpdateStaticMeshesTime);
 DEFINE_STAT(STAT_StaticRelevance);
diff --git a/Engine/Source/Runtime/RenderCore/Public/RenderCore.h b/Engine/Source/Runtime/RenderCore/Public/RenderCore.h
--- a/Engine/Source/Runtime/RenderCore/Public/RenderCore.h
+++ b/Engine/Source/Runtime/RenderCore/Public/RenderCore.h
@@ -49,7 +49,6 @@ DECLARE_CYCLE_STAT_EXTERN(TEXT("View Relevance"),STAT_ViewRelevance,STATGROUP_In
 DECLARE_CYCLE_STAT_EXTERN(TEXT("Compute View Relevance"),STAT_ComputeViewRelevance,STATGROUP_InitViews, RENDERCORE_API);
 DECLARE_CYCLE_STAT_EXTERN(TEXT("Static Mesh Relevance"),STAT_StaticRelevance,STATGROUP_InitViews, RENDERCORE_API);
 DECLARE_CYCLE_STAT_EXTERN(TEXT("UpdateStaticMeshes"),STAT_UpdateStaticMeshesTime,STATGROUP_InitViews, RENDERCORE_API);
-DECLARE_CYCLE_STAT_EXTERN(TEXT("PreRenderView"),STAT_PreRenderView,STATGROUP_InitViews, RENDERCORE_API);
 DECLARE_CYCLE_STAT_EXTERN(TEXT("GetDynamicMeshElements"),STAT_GetDynamicMeshElements,STATGROUP_InitViews, RENDERCORE_API);
 DECLARE_CYCLE_STAT_EXTERN(TEXT("Init dynamic shadows"),STAT_InitDynamicShadowsTime,STATGROUP_InitViews, RENDERCORE_API);
 DECLARE_CYCLE_STAT_EXTERN(TEXT("Update Preshadow Cache"),STAT_UpdatePreshadowCache,STATGROUP_InitViews, RENDERCORE_API);
