Core - Replaced FPlatformProcess::CreateSyncEvent with FPlatformProcess::GetSynchEventFromPool, CreateSyncEvent is not longer publicly available

diff --git a/Engine/Plugins/Messaging/UdpMessaging/Source/UdpMessaging/Private/Transport/UdpMessageBeacon.cpp b/Engine/Plugins/Messaging/UdpMessaging/Source/UdpMessaging/Private/Transport/UdpMessageBeacon.cpp
--- a/Engine/Plugins/Messaging/UdpMessaging/Source/UdpMessaging/Private/Transport/UdpMessageBeacon.cpp
+++ b/Engine/Plugins/Messaging/UdpMessaging/Source/UdpMessaging/Private/Transport/UdpMessageBeacon.cpp
@@ -21,7 +21,7 @@ FUdpMessageBeacon::FUdpMessageBeacon(FSocket* InSocket, const FGuid& InSocketId,
 	, Socket(InSocket)
 	, Stopping(false)
 {
-	EndpointLeftEvent = FPlatformProcess::CreateSynchEvent();
+	EndpointLeftEvent = FPlatformProcess::GetSynchEventFromPool(false);
 	MulticastAddress = InMulticastEndpoint.ToInternetAddr();
 
 	Thread = FRunnableThread::Create(this, TEXT("FUdpMessageBeacon"), 128 * 1024, TPri_AboveNormal, FPlatformAffinity::GetPoolThreadMask());
@@ -38,7 +38,8 @@ FUdpMessageBeacon::~FUdpMessageBeacon()
 
 	MulticastAddress = nullptr;
 
-	delete EndpointLeftEvent;
+	FPlatformProcess::ReturnSynchEventToPool(EndpointLeftEvent);
+	EndpointLeftEvent = nullptr;
 }
 
 
diff --git a/Engine/Plugins/Messaging/UdpMessaging/Source/UdpMessaging/Private/Transport/UdpMessageProcessor.cpp b/Engine/Plugins/Messaging/UdpMessaging/Source/UdpMessaging/Private/Transport/UdpMessageProcessor.cpp
--- a/Engine/Plugins/Messaging/UdpMessaging/Source/UdpMessaging/Private/Transport/UdpMessageProcessor.cpp
+++ b/Engine/Plugins/Messaging/UdpMessaging/Source/UdpMessaging/Private/Transport/UdpMessageProcessor.cpp
@@ -21,7 +21,7 @@ FUdpMessageProcessor::FUdpMessageProcessor( FSocket* InSocket, const FGuid& InNo
 	, Socket(InSocket)
 	, Stopping(false)
 {
-	WorkEvent = FPlatformProcess::CreateSynchEvent();
+	WorkEvent = FPlatformProcess::GetSynchEventFromPool();
 	Thread = FRunnableThread::Create(this, TEXT("FUdpMessageProcessor"), 128 * 1024, TPri_AboveNormal, FPlatformAffinity::GetPoolThreadMask());
 
 	const UUdpMessagingSettings& Settings = *GetDefault<UUdpMessagingSettings>();
@@ -48,7 +48,8 @@ FUdpMessageProcessor::~FUdpMessageProcessor()
 	Thread->Kill(true);
 		
 	delete Thread;
-	delete WorkEvent;
+	FPlatformProcess::ReturnSynchEventToPool(WorkEvent);
+	WorkEvent = nullptr;
 }
 
 
diff --git a/Engine/Source/Developer/ProfilerService/Private/ProfilerServiceManager.cpp b/Engine/Source/Developer/ProfilerService/Private/ProfilerServiceManager.cpp
--- a/Engine/Source/Developer/ProfilerService/Private/ProfilerServiceManager.cpp
+++ b/Engine/Source/Developer/ProfilerService/Private/ProfilerServiceManager.cpp
@@ -26,7 +26,7 @@ public:
 	/** Default constructor. */
 	FFileTransferRunnable( FMessageEndpointPtr& InMessageEndpoint )
 		: Runnable( nullptr )
-		, WorkEvent( FPlatformProcess::CreateSynchEvent( true ) )
+		, WorkEvent( FPlatformProcess::GetSynchEventFromPool( true ) )
 		, MessageEndpoint( InMessageEndpoint )
 		, StopTaskCounter( 0 )
 	{
@@ -53,7 +53,8 @@ public:
 			UE_LOG(LogProfile, Log, TEXT( "File service-client sending aborted (srv): %s" ), *It.Key() );
 		}
 
-		delete WorkEvent;
+		FPlatformProcess::ReturnSynchEventToPool( WorkEvent );
+		WorkEvent = nullptr;
 	}
 
 	// Begin FRunnable interface.
diff --git a/Engine/Source/Editor/UnrealEd/Private/Commandlets/FileServerCommandlet.cpp b/Engine/Source/Editor/UnrealEd/Private/Commandlets/FileServerCommandlet.cpp
--- a/Engine/Source/Editor/UnrealEd/Private/Commandlets/FileServerCommandlet.cpp
+++ b/Engine/Source/Editor/UnrealEd/Private/Commandlets/FileServerCommandlet.cpp
@@ -42,7 +42,7 @@ int32 UFileServerCommandlet::Main( const FString& Params )
 #if PLATFORM_WINDOWS// Windows only
 	// Used by the .com wrapper to notify that the Ctrl-C handler was triggered.
 	// This shared event is checked each tick so that the log file can be cleanly flushed.
-	FEvent* ComWrapperShutdownEvent = FPlatformProcess::CreateSynchEvent(true);
+	FEvent* ComWrapperShutdownEvent = FPlatformProcess::GetSynchEventFromPool(true);
 #endif
 
 	// parse instance identifier
@@ -150,7 +150,8 @@ int32 UFileServerCommandlet::Main( const FString& Params )
 
 	//@todo abstract properly or delete 
 #if PLATFORM_WINDOWS
-	delete ComWrapperShutdownEvent;
+	FPlatformProcess::ReturnSynchEventToPool(ComWrapperShutdownEvent);
+	ComWrapperShutdownEvent = nullptr;
 #endif
 
 	GIsRunning = false;
diff --git a/Engine/Source/Runtime/Core/Private/Async/TaskGraph.cpp b/Engine/Source/Runtime/Core/Private/Async/TaskGraph.cpp
--- a/Engine/Source/Runtime/Core/Private/Async/TaskGraph.cpp
+++ b/Engine/Source/Runtime/Core/Private/Async/TaskGraph.cpp
@@ -614,14 +614,14 @@ private:
 		FEvent*												StallRestartEvent;
 
 		FThreadTaskQueue()
-			: StallRestartEvent(FPlatformProcess::CreateSynchEvent(true))
+			: StallRestartEvent(FPlatformProcess::GetSynchEventFromPool(true))
 		{
 
 		}
 		~FThreadTaskQueue()
 		{
-			delete StallRestartEvent;
-			StallRestartEvent = NULL;
+			FPlatformProcess::ReturnSynchEventToPool( StallRestartEvent );
+			StallRestartEvent = nullptr;
 		}
 	};
 
diff --git a/Engine/Source/Runtime/Core/Private/HAL/PThreadRunnableThread.h b/Engine/Source/Runtime/Core/Private/HAL/PThreadRunnableThread.h
--- a/Engine/Source/Runtime/Core/Private/HAL/PThreadRunnableThread.h
+++ b/Engine/Source/Runtime/Core/Private/HAL/PThreadRunnableThread.h
@@ -320,7 +320,7 @@ protected:
 		Runnable = InRunnable;
 
 		// Create a sync event to guarantee the Init() function is called first
-		ThreadInitSyncEvent	= FPlatformProcess::CreateSynchEvent(true);
+		ThreadInitSyncEvent	= FPlatformProcess::GetSynchEventFromPool(true);
 		// A name for the thread in for debug purposes. _ThreadProc will set it.
 		ThreadName = InThreadName ? InThreadName : TEXT("Unnamed UE4");
 		ThreadAffinityMask = InThreadAffinityMask;
@@ -344,8 +344,8 @@ protected:
 		}
 
 		// Cleanup the sync event
-		delete ThreadInitSyncEvent;
-		ThreadInitSyncEvent = NULL;
+		FPlatformProcess::ReturnSynchEventToPool( ThreadInitSyncEvent );
+		ThreadInitSyncEvent = nullptr;
 		return Thread != PTHREAD_NULL;
 	}
 };
diff --git a/Engine/Source/Runtime/Core/Private/HAL/ThreadingBase.cpp b/Engine/Source/Runtime/Core/Private/HAL/ThreadingBase.cpp
--- a/Engine/Source/Runtime/Core/Private/HAL/ThreadingBase.cpp
+++ b/Engine/Source/Runtime/Core/Private/HAL/ThreadingBase.cpp
@@ -330,7 +330,7 @@ public:
 		PoolThreadIndex++;
 
 		OwningThreadPool = InPool;
-		DoWorkEvent = FPlatformProcess::CreateSynchEvent();
+		DoWorkEvent = FPlatformProcess::GetSynchEventFromPool();
 		Thread = FRunnableThread::Create(this, *PoolThreadName, InStackSize, ThreadPriority, FPlatformAffinity::GetPoolThreadMask());
 		check(Thread);
 		return true;
@@ -356,9 +356,8 @@ public:
 		// brute force kill that thread. Very bad as that might leak.
 		Thread->WaitForCompletion();
 		// Clean up the event
-		delete DoWorkEvent;
+		FPlatformProcess::ReturnSynchEventToPool(DoWorkEvent);
 		DoWorkEvent = nullptr;
-		delete DoWorkEvent;
 		delete Thread;
 		return bDidExitOK;
 	}
diff --git a/Engine/Source/Runtime/Core/Private/Misc/EventPool.h b/Engine/Source/Runtime/Core/Private/Misc/EventPool.h
--- a/Engine/Source/Runtime/Core/Private/Misc/EventPool.h
+++ b/Engine/Source/Runtime/Core/Private/Misc/EventPool.h
@@ -15,6 +15,12 @@ enum class EEventPoolTypes
 	ManualReset
 };
 
+/** Helper class used to create a new stat id for events. */
+struct FEventStats
+{
+	/** Creates a new stats id that can be used for events. */
+	CORE_API static TStatId CreateStatID();
+};
 
 /**
  * Template class for event pools.
@@ -37,7 +43,7 @@ public:
 	 *
 	 * @return Pool singleton.
 	 */
-	static FEventPool& Get()
+	CORE_API static FEventPool& Get()
 	{
 		static FEventPool Singleton;
 		return Singleton;
diff --git a/Engine/Source/Runtime/Core/Private/Serialization/AsyncIOSystemBase.cpp b/Engine/Source/Runtime/Core/Private/Serialization/AsyncIOSystemBase.cpp
--- a/Engine/Source/Runtime/Core/Private/Serialization/AsyncIOSystemBase.cpp
+++ b/Engine/Source/Runtime/Core/Private/Serialization/AsyncIOSystemBase.cpp
@@ -514,7 +514,7 @@ bool FAsyncIOSystemBase::Init()
 {
 	CriticalSection				= new FCriticalSection();
 	ExclusiveReadCriticalSection= new FCriticalSection();
-	OutstandingRequestsEvent	= FPlatformProcess::CreateSynchEvent();
+	OutstandingRequestsEvent	= FPlatformProcess::GetSynchEventFromPool();
 	RequestIndex				= 1;
 	MinPriority					= AIOP_MIN;
 	IsRunning.Increment();
@@ -562,7 +562,8 @@ void FAsyncIOSystemBase::Exit()
 {
 	FlushHandles();
 	delete CriticalSection;
-	delete OutstandingRequestsEvent;
+	FPlatformProcess::ReturnSynchEventToPool(OutstandingRequestsEvent);
+	OutstandingRequestsEvent = nullptr;
 }
 
 void FAsyncIOSystemBase::Stop()
diff --git a/Engine/Source/Runtime/Core/Private/WinRT/WinRTRunnableThread.h b/Engine/Source/Runtime/Core/Private/WinRT/WinRTRunnableThread.h
--- a/Engine/Source/Runtime/Core/Private/WinRT/WinRTRunnableThread.h
+++ b/Engine/Source/Runtime/Core/Private/WinRT/WinRTRunnableThread.h
@@ -227,7 +227,7 @@ protected:
 		ThreadAffintyMask = InThreadAffinityMask;
 
 		// Create a sync event to guarantee the Init() function is called first
-		ThreadInitSyncEvent	= FPlatformProcess::CreateSynchEvent(true);
+		ThreadInitSyncEvent	= FPlatformProcess::GetSynchEventFromPool(true);
 
 		// Create the new thread
 		Thread = ThreadEmulation::CreateThread(NULL,InStackSize,_ThreadProc,this,0,(::DWORD *)&ThreadID);
@@ -248,8 +248,8 @@ protected:
 		}
 
 		// Cleanup the sync event
-		delete ThreadInitSyncEvent;
-		ThreadInitSyncEvent = NULL;
+		FPlatformProcess::ReturnSynchEventToPool(ThreadInitSyncEvent);
+		ThreadInitSyncEvent = nullptr;
 		return Thread != NULL;
 	}
 };
diff --git a/Engine/Source/Runtime/Core/Private/Windows/WindowsRunnableThread.h b/Engine/Source/Runtime/Core/Private/Windows/WindowsRunnableThread.h
--- a/Engine/Source/Runtime/Core/Private/Windows/WindowsRunnableThread.h
+++ b/Engine/Source/Runtime/Core/Private/Windows/WindowsRunnableThread.h
@@ -97,10 +97,10 @@ public:
 
 	FRunnableThreadWin( )
 		: Thread(NULL)
-		, Runnable(NULL)
-		, ThreadInitSyncEvent(NULL)
+		, Runnable(nullptr)
+		, ThreadInitSyncEvent(nullptr)
 		, ThreadPriority(TPri_Normal)
-		, ThreadID(NULL)
+		, ThreadID(0)
 	{
 
 	}
@@ -196,7 +196,7 @@ protected:
 		ThreadAffintyMask = InThreadAffinityMask;
 
 		// Create a sync event to guarantee the Init() function is called first
-		ThreadInitSyncEvent	= FPlatformProcess::CreateSynchEvent(true);
+		ThreadInitSyncEvent	= FPlatformProcess::GetSynchEventFromPool(true);
 
 		// Create the new thread
 		Thread = CreateThread(NULL,InStackSize,_ThreadProc,this,STACK_SIZE_PARAM_IS_A_RESERVATION,(::DWORD *)&ThreadID);
@@ -204,7 +204,7 @@ protected:
 		// If it fails, clear all the vars
 		if (Thread == NULL)
 		{
-			Runnable = NULL;
+			Runnable = nullptr;
 		}
 		else
 		{
@@ -216,8 +216,8 @@ protected:
 		}
 
 		// Cleanup the sync event
-		delete ThreadInitSyncEvent;
-		ThreadInitSyncEvent = NULL;
+		FPlatformProcess::ReturnSynchEventToPool(ThreadInitSyncEvent);
+		ThreadInitSyncEvent = nullptr;
 		return Thread != NULL;
 	}
 };
diff --git a/Engine/Source/Runtime/Core/Public/Async/AsyncWork.h b/Engine/Source/Runtime/Core/Public/Async/AsyncWork.h
--- a/Engine/Source/Runtime/Core/Public/Async/AsyncWork.h
+++ b/Engine/Source/Runtime/Core/Public/Async/AsyncWork.h
@@ -255,8 +255,8 @@ class FAsyncTask
 	**/
 	void DestroyEvent()
 	{
-		delete DoneEvent;
-		DoneEvent = NULL;
+		FPlatformProcess::ReturnSynchEventToPool(DoneEvent);
+		DoneEvent = nullptr;
 	}
 
 	/* Generic start function, not called directly
@@ -276,7 +276,7 @@ class FAsyncTask
 		{
 			if (!DoneEvent)
 			{
-				DoneEvent = FPlatformProcess::CreateSynchEvent(true);
+				DoneEvent = FPlatformProcess::GetSynchEventFromPool(true);
 			}
 			DoneEvent->Reset();
 			QueuedPool->AddQueuedWork(this);
diff --git a/Engine/Source/Runtime/Core/Public/GenericPlatform/GenericPlatformProcess.h b/Engine/Source/Runtime/Core/Public/GenericPlatform/GenericPlatformProcess.h
--- a/Engine/Source/Runtime/Core/Public/GenericPlatform/GenericPlatformProcess.h
+++ b/Engine/Source/Runtime/Core/Public/GenericPlatform/GenericPlatformProcess.h
@@ -436,6 +436,7 @@ struct CORE_API FGenericPlatformProcess
 	 * @return A new event, or nullptr none could be created.
 	 * @see GetSynchEventFromPool, ReturnSynchEventToPool
 	 */
+	DEPRECATED(4.8, "Please use GetSynchEventFromPool to create a new event, and GetSynchEventFromPool to release the event.")
 	static class FEvent* CreateSynchEvent(bool bIsManualReset = false);
 
 	/**
diff --git a/Engine/Source/Runtime/Engine/Private/AVIWriter.cpp b/Engine/Source/Runtime/Engine/Private/AVIWriter.cpp
--- a/Engine/Source/Runtime/Engine/Private/AVIWriter.cpp
+++ b/Engine/Source/Runtime/Engine/Private/AVIWriter.cpp
@@ -111,7 +111,7 @@ public:
 
 	void StartCapture(FViewport* Viewport , const FString& OutputPath /*=FString()*/)
 	{
-		GCaptureSyncEvent = FPlatformProcess::CreateSynchEvent();
+		GCaptureSyncEvent = FPlatformProcess::GetSynchEventFromPool();
 		if (!bCapturing)
 		{
 			if (!Viewport)
@@ -347,8 +347,8 @@ public:
 			SAFE_RELEASE(Control);
 			SAFE_RELEASE(Graph);
 			FWindowsPlatformMisc::CoUninitialize();
-			delete GCaptureSyncEvent;
-			GCaptureSyncEvent = NULL;
+			FPlatformProcess::ReturnSynchEventToPool(GCaptureSyncEvent);
+			GCaptureSyncEvent = nullptr;
 		}
 	}
 
@@ -607,13 +607,14 @@ public:
 			[AVFWriterInputRef markAsFinished];
 			// This will finish asynchronously and then destroy the relevant objects.
 			// We must wait for this to complete.
-			FEvent* Event = FPlatformProcess::CreateSynchEvent(true);
+			FEvent* Event = FPlatformProcess::GetSynchEventFromPool(true);
 			[AVFWriterRef finishWritingWithCompletionHandler:^{
 				check(AVFWriterRef.status == AVAssetWriterStatusCompleted);
 				Event->Trigger();
 			}];
 			Event->Wait(~0u);
-			delete Event;
+			FPlatformProcess::ReturnSynchEventToPool(Event);
+			Event = nullptr;
 			[AVFWriterInputRef release];
 			[AVFWriterRef release];
 			[AVFPixelBufferAdaptorRef release];
diff --git a/Engine/Source/Runtime/Messaging/Private/Bus/MessageRouter.cpp b/Engine/Source/Runtime/Messaging/Private/Bus/MessageRouter.cpp
--- a/Engine/Source/Runtime/Messaging/Private/Bus/MessageRouter.cpp
+++ b/Engine/Source/Runtime/Messaging/Private/Bus/MessageRouter.cpp
@@ -12,13 +12,14 @@ FMessageRouter::FMessageRouter()
 	, Tracer(MakeShareable(new FMessageTracer()))
 {
 	ActiveSubscriptions.FindOrAdd(NAME_All);
-	WorkEvent = FPlatformProcess::CreateSynchEvent(true);
+	WorkEvent = FPlatformProcess::GetSynchEventFromPool(true);
 }
 
 
 FMessageRouter::~FMessageRouter()
 {
-	delete WorkEvent;
+	FPlatformProcess::ReturnSynchEventToPool(WorkEvent);
+	WorkEvent = nullptr;
 }
 
 
diff --git a/Engine/Source/Runtime/Messaging/Private/Bus/MessageTracer.cpp b/Engine/Source/Runtime/Messaging/Private/Bus/MessageTracer.cpp
--- a/Engine/Source/Runtime/Messaging/Private/Bus/MessageTracer.cpp
+++ b/Engine/Source/Runtime/Messaging/Private/Bus/MessageTracer.cpp
@@ -11,13 +11,14 @@ FMessageTracer::FMessageTracer()
 	, ResetPending(false)
 	, Running(false)
 {
-	ContinueEvent = FPlatformProcess::CreateSynchEvent();
+	ContinueEvent = FPlatformProcess::GetSynchEventFromPool();
 }
 
 
 FMessageTracer::~FMessageTracer()
 {
-	delete ContinueEvent;
+	FPlatformProcess::ReturnSynchEventToPool(ContinueEvent);
+	ContinueEvent = nullptr;
 }
 
 
diff --git a/Engine/Source/Runtime/Networking/Public/Common/UdpSocketSender.h b/Engine/Source/Runtime/Networking/Public/Common/UdpSocketSender.h
--- a/Engine/Source/Runtime/Networking/Public/Common/UdpSocketSender.h
+++ b/Engine/Source/Runtime/Networking/Public/Common/UdpSocketSender.h
@@ -50,7 +50,7 @@ public:
 		int32 NewSize = 0;
 		Socket->SetSendBufferSize(512*1024, NewSize);
 
-		WorkEvent = FPlatformProcess::CreateSynchEvent();
+		WorkEvent = FPlatformProcess::GetSynchEventFromPool();
 
 		Thread = FRunnableThread::Create(this, ThreadDescription, 128 * 1024, TPri_AboveNormal, FPlatformAffinity::GetPoolThreadMask());
 	}
@@ -66,7 +66,8 @@ public:
 			delete Thread;
 		}
 
-		delete WorkEvent;
+		FPlatformProcess::ReturnSynchEventToPool(WorkEvent);
+		WorkEvent = nullptr;
 	}
 
 public:
diff --git a/Engine/Source/Runtime/PakFile/Private/SignedArchiveReader.cpp b/Engine/Source/Runtime/PakFile/Private/SignedArchiveReader.cpp
--- a/Engine/Source/Runtime/PakFile/Private/SignedArchiveReader.cpp
+++ b/Engine/Source/Runtime/PakFile/Private/SignedArchiveReader.cpp
@@ -24,7 +24,7 @@ FChunkCacheWorker::FChunkCacheWorker(FArchive* InReader)
 	// Public key should never be zero at this point. Check PublicKey.inl for more details.
 	check(!DecryptionKey.Exponent.IsZero() && !DecryptionKey.Modulus.IsZero());
 
-	QueuedRequestsEvent = FPlatformProcess::CreateSynchEvent();
+	QueuedRequestsEvent = FPlatformProcess::GetSynchEventFromPool();
 	Thread = FRunnableThread::Create(this, TEXT("FChunkCacheWorker"), 0, TPri_BelowNormal);
 }
 
@@ -32,8 +32,8 @@ FChunkCacheWorker::~FChunkCacheWorker()
 {
 	delete Thread;
 	Thread = NULL;
-	delete QueuedRequestsEvent;
-	QueuedRequestsEvent = NULL;
+	FPlatformProcess::ReturnSynchEventToPool(QueuedRequestsEvent);
+	QueuedRequestsEvent = nullptr;
 }
 
 bool FChunkCacheWorker::Init()
diff --git a/Engine/Source/Runtime/RenderCore/Private/RenderingThread.cpp b/Engine/Source/Runtime/RenderCore/Private/RenderingThread.cpp
--- a/Engine/Source/Runtime/RenderCore/Private/RenderingThread.cpp
+++ b/Engine/Source/Runtime/RenderCore/Private/RenderingThread.cpp
@@ -342,10 +342,16 @@ public:
 
 	FRenderingThread() : FRunnable()
 	{
-		TaskGraphBoundSyncEvent	= FPlatformProcess::CreateSynchEvent(true);
+		TaskGraphBoundSyncEvent	= FPlatformProcess::GetSynchEventFromPool(true);
 		RHIFlushResources();
 	}
 
+	virtual ~FRenderingThread()
+	{
+		FPlatformProcess::ReturnSynchEventToPool(TaskGraphBoundSyncEvent);
+		TaskGraphBoundSyncEvent = nullptr;
+	}
+
 	// FRunnable interface.
 	virtual bool Init(void) 
 	{ 
@@ -743,8 +749,7 @@ static void GameThreadWaitForTask(const FGraphEventRef& Task, bool bEmptyGameThr
 		SCOPE_CYCLE_COUNTER(STAT_GameIdleTime);
 		{
 			static int32 NumRecursiveCalls = 0;
-			static TArray<FEvent*> EventPool;
-			
+		
 			// Check for recursion. It's not completely safe but because we pump messages while 
 			// blocked it is expected.
 			NumRecursiveCalls++;
@@ -758,15 +763,7 @@ static void GameThreadWaitForTask(const FGraphEventRef& Task, bool bEmptyGameThr
 			}
 
 			// Grab an event from the pool and fire off a task to trigger it.
-			FEvent* Event = NULL;
-			if (EventPool.Num() > 0)
-			{
-				Event = EventPool.Pop();
-			}
-			else
-			{
-				Event = FPlatformProcess::CreateSynchEvent();
-			}
+			FEvent* Event = FPlatformProcess::GetSynchEventFromPool();
 			FTaskGraphInterface::Get().TriggerEventWhenTaskCompletes(Event, Task, ENamedThreads::GameThread);
 
 			// Check rendering thread health needs to be called from time to
@@ -780,7 +777,7 @@ static void GameThreadWaitForTask(const FGraphEventRef& Task, bool bEmptyGameThr
 				CheckRenderingThreadHealth();
 				if (bEmptyGameThreadTasks)
 				{
-					// process gamehtread tasks if there are any
+					// process gamethread tasks if there are any
 					FTaskGraphInterface::Get().ProcessThreadUntilIdle(ENamedThreads::GameThread);
 				}
 				bDone = Event->Wait(WaitTime);
@@ -788,7 +785,8 @@ static void GameThreadWaitForTask(const FGraphEventRef& Task, bool bEmptyGameThr
 			while (!bDone);
 
 			// Return the event to the pool and decrement the recursion counter.
-			EventPool.Push(Event);
+			FPlatformProcess::ReturnSynchEventToPool(Event);
+			Event = nullptr;
 			NumRecursiveCalls--;
 		}
 	}
