Stats - Removed ExplicitFlush from the rendering thread, no longer needed

diff --git a/Engine/Source/Runtime/RenderCore/Private/RenderingThread.cpp b/Engine/Source/Runtime/RenderCore/Private/RenderingThread.cpp
--- a/Engine/Source/Runtime/RenderCore/Private/RenderingThread.cpp
+++ b/Engine/Source/Runtime/RenderCore/Private/RenderingThread.cpp
@@ -283,10 +283,6 @@ void RenderingThreadMain( FEvent* TaskGraphBoundSyncEvent )
 	ENamedThreads::RenderThread = ENamedThreads::GameThread;
 	ENamedThreads::RenderThread_Local = ENamedThreads::GameThread_Local;
 	FPlatformMisc::MemoryBarrier();
-
-#if STATS
-	FThreadStats::ExplicitFlush();
-#endif
 }
 
 /**
@@ -408,10 +404,6 @@ public:
 		{
 			RenderingThreadMain( TaskGraphBoundSyncEvent );
 		}
-#if STATS
-		FThreadStats::ExplicitFlush();
-		FThreadStats::Shutdown();
-#endif
 		return 0;
 	}
 };
