UE4 - Added 5 & 6 param macros
ENQUEUE_UNIQUE_RENDER_COMMAND_FIVEPARAMETER and ENQUEUE_UNIQUE_RENDER_COMMAND_SIXPARAMETER

diff --git a/Engine/Source/Runtime/RenderCore/Public/RenderingThread.h b/Engine/Source/Runtime/RenderCore/Public/RenderingThread.h
--- a/Engine/Source/Runtime/RenderCore/Public/RenderingThread.h
+++ b/Engine/Source/Runtime/RenderCore/Public/RenderingThread.h
@@ -123,7 +123,6 @@ private:
 
 #if UE_SERVER
 #define ENQUEUE_UNIQUE_RENDER_COMMAND(TypeName,Code)
-#define ENQUEUE_UNIQUE_RENDER_COMMAND(TypeName,Code)
 #define ENQUEUE_UNIQUE_RENDER_COMMAND_ONEPARAMETER(TypeName,ParamType1,ParamName1,ParamValue1,Code)
 #define ENQUEUE_UNIQUE_RENDER_COMMAND_ONEPARAMETER_CREATE_TEMPLATE(TypeName,TemplateParamName,ParamType1,ParamValue1)
 #define ENQUEUE_UNIQUE_RENDER_COMMAND_ONEPARAMETER_DECLARE_TEMPLATE(TypeName,TemplateParamName,ParamType1,ParamName1,ParamValue1,Code)
@@ -136,6 +135,12 @@ private:
 #define ENQUEUE_UNIQUE_RENDER_COMMAND_FOURPARAMETER(TypeName,ParamType1,ParamName1,ParamValue1,ParamType2,ParamName2,ParamValue2,ParamType3,ParamName3,ParamValue3,ParamType4,ParamName4,ParamValue4,Code)
 #define ENQUEUE_UNIQUE_RENDER_COMMAND_FOURPARAMETER_CREATE_TEMPLATE(TypeName,TemplateParamName,ParamType1,ParamValue1,ParamType2,ParamValue2,ParamType3,ParamValue3,ParamType4,ParamValue4)
 #define ENQUEUE_UNIQUE_RENDER_COMMAND_FOURPARAMETER_DECLARE_TEMPLATE(TypeName,TemplateParamName,ParamType1,ParamName1,ParamValue1,ParamType2,ParamName2,ParamValue2,ParamType3,ParamName3,ParamValue3,ParamType4,ParamName4,ParamValue4,Code)
+#define ENQUEUE_UNIQUE_RENDER_COMMAND_FIVEPARAMETER(TypeName,ParamType1,ParamName1,ParamValue1,ParamType2,ParamName2,ParamValue2,ParamType3,ParamName3,ParamValue3,ParamType4,ParamName4,ParamValue4,ParamType5,ParamName5,ParamValue5,Code)
+#define ENQUEUE_UNIQUE_RENDER_COMMAND_FIVEPARAMETER_CREATE_TEMPLATE(TypeName,TemplateParamName,ParamType1,ParamValue1,ParamType2,ParamValue2,ParamType3,ParamValue3,ParamType4,ParamValue4,ParamType5,ParamValue5)
+#define ENQUEUE_UNIQUE_RENDER_COMMAND_FIVEPARAMETER_DECLARE_TEMPLATE(TypeName,TemplateParamName,ParamType1,ParamName1,ParamValue1,ParamType2,ParamName2,ParamValue2,ParamType3,ParamName3,ParamValue3,ParamType4,ParamName4,ParamValue4,ParamType5,ParamName5,ParamValue5,Code)
+#define ENQUEUE_UNIQUE_RENDER_COMMAND_SIXPARAMETER(TypeName,ParamType1,ParamName1,ParamValue1,ParamType2,ParamName2,ParamValue2,ParamType3,ParamName3,ParamValue3,ParamType4,ParamName4,ParamValue4,ParamType5,ParamName5,ParamValue5,ParamType6,ParamName6,ParamValue6,Code)
+#define ENQUEUE_UNIQUE_RENDER_COMMAND_SIXPARAMETER_CREATE_TEMPLATE(TypeName,TemplateParamName,ParamType1,ParamValue1,ParamType2,ParamValue2,ParamType3,ParamValue3,ParamType4,ParamValue4,ParamType5,ParamValue5,ParamType6,ParamValue6)
+#define ENQUEUE_UNIQUE_RENDER_COMMAND_SIXPARAMETER_DECLARE_TEMPLATE(TypeName,TemplateParamName,ParamType1,ParamName1,ParamValue1,ParamType2,ParamName2,ParamValue2,ParamType3,ParamName3,ParamValue3,ParamType4,ParamName4,ParamValue4,ParamType5,ParamName5,ParamValue5,ParamType6,ParamName6,ParamValue6,Code)
 #else
 
 /** The parent class of commands stored in the rendering command queue. */
@@ -390,6 +395,114 @@ DECLARE_STATS_GROUP(TEXT("Render Thread Commands"), STATGROUP_RenderThreadComman
 #define ENQUEUE_UNIQUE_RENDER_COMMAND_FOURPARAMETER_CREATE_TEMPLATE(TypeName,TemplateParamName,ParamType1,ParamValue1,ParamType2,ParamValue2,ParamType3,ParamValue3,ParamType4,ParamValue4) \
 	ENQUEUE_UNIQUE_RENDER_COMMAND_FOURPARAMETER_CREATE(TypeName<TemplateParamName>,ParamType1,ParamValue1,ParamType2,ParamValue2,ParamType3,ParamValue3,ParamType4,ParamValue4)
 
+/**
+ * Declares a rendering command type with 5 parameters.
+ */
+#define ENQUEUE_UNIQUE_RENDER_COMMAND_FIVEPARAMETER_DECLARE_OPTTYPENAME(TypeName,ParamType1,ParamName1,ParamValue1,ParamType2,ParamName2,ParamValue2,ParamType3,ParamName3,ParamValue3,ParamType4,ParamName4,ParamValue4,ParamType5,ParamName5,ParamValue5,OptTypename,Code) \
+	class EURCMacro_##TypeName : public FRenderCommand \
+	{ \
+	public: \
+		EURCMacro_##TypeName(OptTypename TCallTraits<ParamType1>::ParamType In##ParamName1,OptTypename TCallTraits<ParamType2>::ParamType In##ParamName2,OptTypename TCallTraits<ParamType3>::ParamType In##ParamName3,OptTypename TCallTraits<ParamType4>::ParamType In##ParamName4,OptTypename TCallTraits<ParamType5>::ParamType In##ParamName5): \
+		  ParamName1(In##ParamName1), \
+		  ParamName2(In##ParamName2), \
+		  ParamName3(In##ParamName3), \
+		  ParamName4(In##ParamName4), \
+		  ParamName5(In##ParamName5) \
+		{} \
+		TASK_FUNCTION(Code) \
+		TASKNAME_FUNCTION(TypeName) \
+	private: \
+		ParamType1 ParamName1; \
+		ParamType2 ParamName2; \
+		ParamType3 ParamName3; \
+		ParamType4 ParamName4; \
+		ParamType5 ParamName5; \
+	};
+#define ENQUEUE_UNIQUE_RENDER_COMMAND_FIVEPARAMETER_DECLARE(TypeName,ParamType1,ParamName1,ParamValue1,ParamType2,ParamName2,ParamValue2,ParamType3,ParamName3,ParamValue3,ParamType4,ParamName4,ParamValue4,ParamType5,ParamName5,ParamValue5,Code) \
+	ENQUEUE_UNIQUE_RENDER_COMMAND_FIVEPARAMETER_DECLARE_OPTTYPENAME(TypeName,ParamType1,ParamName1,ParamValue1,ParamType2,ParamName2,ParamValue2,ParamType3,ParamName3,ParamValue3,ParamType4,ParamName4,ParamValue4,ParamType5,ParamName5,ParamValue5,,Code)
+
+
+#define ENQUEUE_UNIQUE_RENDER_COMMAND_FIVEPARAMETER_CREATE(TypeName,ParamType1,ParamValue1,ParamType2,ParamValue2,ParamType3,ParamValue3,ParamType4,ParamValue4,ParamType5,ParamValue5) \
+	{ \
+		if(GIsThreadedRendering || !IsInGameThread()) \
+		{ \
+			CheckNotBlockedOnRenderThread(); \
+			TGraphTask<EURCMacro_##TypeName>::CreateTask().ConstructAndDispatchWhenReady(ParamValue1,ParamValue2,ParamValue3,ParamValue4,ParamValue5); \
+		} \
+		else \
+		{ \
+			EURCMacro_##TypeName TempCommand(ParamValue1,ParamValue2,ParamValue3,ParamValue4,ParamValue5); \
+			FScopeCycleCounter EURCMacro_Scope(TempCommand.GetStatId()); \
+			TempCommand.DoTask(ENamedThreads::GameThread, FGraphEventRef() ); \
+		} \
+	}
+
+#define ENQUEUE_UNIQUE_RENDER_COMMAND_FIVEPARAMETER(TypeName,ParamType1,ParamName1,ParamValue1,ParamType2,ParamName2,ParamValue2,ParamType3,ParamName3,ParamValue3,ParamType4,ParamName4,ParamValue4,ParamType5,ParamName5,ParamValue5,Code) \
+	ENQUEUE_UNIQUE_RENDER_COMMAND_FIVEPARAMETER_DECLARE(TypeName,ParamType1,ParamName1,ParamValue1,ParamType2,ParamName2,ParamValue2,ParamType3,ParamName3,ParamValue3,ParamType4,ParamName4,ParamValue4,ParamType5,ParamName5,ParamValue5,Code) \
+	ENQUEUE_UNIQUE_RENDER_COMMAND_FIVEPARAMETER_CREATE(TypeName,ParamType1,ParamValue1,ParamType2,ParamValue2,ParamType3,ParamValue3,ParamType4,ParamValue4,ParamType5,ParamValue5)
+
+#define ENQUEUE_UNIQUE_RENDER_COMMAND_FIVEPARAMETER_DECLARE_TEMPLATE(TypeName,TemplateParamName,ParamType1,ParamName1,ParamValue1,ParamType2,ParamName2,ParamValue2,ParamType3,ParamName3,ParamValue3,ParamType4,ParamName4,ParamValue4,ParamType5,ParamName5,ParamValue5,Code) \
+	template <typename TemplateParamName> \
+	ENQUEUE_UNIQUE_RENDER_COMMAND_FIVEPARAMETER_DECLARE_OPTTYPENAME(TypeName,ParamType1,ParamName1,ParamValue1,ParamType2,ParamName2,ParamValue2,ParamType3,ParamName3,ParamValue3,ParamType4,ParamName4,ParamValue4,ParamType5,ParamName5,ParamValue5,OptTypename,Code)
+
+#define ENQUEUE_UNIQUE_RENDER_COMMAND_FIVEPARAMETER_CREATE_TEMPLATE(TypeName,TemplateParamName,ParamType1,ParamValue1,ParamType2,ParamValue2,ParamType3,ParamValue3,ParamType4,ParamValue4,ParamType5,ParamValue5) \
+	ENQUEUE_UNIQUE_RENDER_COMMAND_FIVEPARAMETER_CREATE(TypeName<TemplateParamName>,ParamType1,ParamValue1,ParamType2,ParamValue2,ParamType3,ParamValue3,ParamType4,ParamValue4,ParamType5,ParamValue5)
+
+/**
+ * Declares a rendering command type with 6 parameters.
+ */
+#define ENQUEUE_UNIQUE_RENDER_COMMAND_SIXPARAMETER_DECLARE_OPTTYPENAME(TypeName,ParamType1,ParamName1,ParamValue1,ParamType2,ParamName2,ParamValue2,ParamType3,ParamName3,ParamValue3,ParamType4,ParamName4,ParamValue4,ParamType5,ParamName5,ParamValue5,ParamType6,ParamName6,ParamValue6,OptTypename,Code) \
+	class EURCMacro_##TypeName : public FRenderCommand \
+	{ \
+	public: \
+		EURCMacro_##TypeName(OptTypename TCallTraits<ParamType1>::ParamType In##ParamName1,OptTypename TCallTraits<ParamType2>::ParamType In##ParamName2,OptTypename TCallTraits<ParamType3>::ParamType In##ParamName3,OptTypename TCallTraits<ParamType4>::ParamType In##ParamName4,OptTypename TCallTraits<ParamType5>::ParamType In##ParamName5,OptTypename TCallTraits<ParamType6>::ParamType In##ParamName6): \
+		  ParamName1(In##ParamName1), \
+		  ParamName2(In##ParamName2), \
+		  ParamName3(In##ParamName3), \
+		  ParamName4(In##ParamName4), \
+		  ParamName5(In##ParamName5), \
+		  ParamName6(In##ParamName6) \
+		{} \
+		TASK_FUNCTION(Code) \
+		TASKNAME_FUNCTION(TypeName) \
+	private: \
+		ParamType1 ParamName1; \
+		ParamType2 ParamName2; \
+		ParamType3 ParamName3; \
+		ParamType4 ParamName4; \
+		ParamType5 ParamName5; \
+		ParamType6 ParamName6; \
+	};
+#define ENQUEUE_UNIQUE_RENDER_COMMAND_SIXPARAMETER_DECLARE(TypeName,ParamType1,ParamName1,ParamValue1,ParamType2,ParamName2,ParamValue2,ParamType3,ParamName3,ParamValue3,ParamType4,ParamName4,ParamValue4,ParamType5,ParamName5,ParamValue5,ParamType6,ParamName6,ParamValue6,Code) \
+	ENQUEUE_UNIQUE_RENDER_COMMAND_SIXPARAMETER_DECLARE_OPTTYPENAME(TypeName,ParamType1,ParamName1,ParamValue1,ParamType2,ParamName2,ParamValue2,ParamType3,ParamName3,ParamValue3,ParamType4,ParamName4,ParamValue4,ParamType5,ParamName5,ParamValue5,ParamType6,ParamName6,ParamValue6,,Code)
+
+
+#define ENQUEUE_UNIQUE_RENDER_COMMAND_SIXPARAMETER_CREATE(TypeName,ParamType1,ParamValue1,ParamType2,ParamValue2,ParamType3,ParamValue3,ParamType4,ParamValue4,ParamType5,ParamValue5,ParamType6,ParamValue6) \
+	{ \
+		if(GIsThreadedRendering || !IsInGameThread()) \
+		{ \
+			CheckNotBlockedOnRenderThread(); \
+			TGraphTask<EURCMacro_##TypeName>::CreateTask().ConstructAndDispatchWhenReady(ParamValue1,ParamValue2,ParamValue3,ParamValue4,ParamValue5,ParamValue6); \
+		} \
+		else \
+		{ \
+			EURCMacro_##TypeName TempCommand(ParamValue1,ParamValue2,ParamValue3,ParamValue4,ParamValue5,ParamValue6); \
+			FScopeCycleCounter EURCMacro_Scope(TempCommand.GetStatId()); \
+			TempCommand.DoTask(ENamedThreads::GameThread, FGraphEventRef() ); \
+		} \
+	}
+
+#define ENQUEUE_UNIQUE_RENDER_COMMAND_SIXPARAMETER(TypeName,ParamType1,ParamName1,ParamValue1,ParamType2,ParamName2,ParamValue2,ParamType3,ParamName3,ParamValue3,ParamType4,ParamName4,ParamValue4,ParamType5,ParamName5,ParamValue5,ParamType6,ParamName6,ParamValue6,Code) \
+	ENQUEUE_UNIQUE_RENDER_COMMAND_SIXPARAMETER_DECLARE(TypeName,ParamType1,ParamName1,ParamValue1,ParamType2,ParamName2,ParamValue2,ParamType3,ParamName3,ParamValue3,ParamType4,ParamName4,ParamValue4,ParamType5,ParamName5,ParamValue5,ParamType6,ParamName6,ParamValue6,Code) \
+	ENQUEUE_UNIQUE_RENDER_COMMAND_SIXPARAMETER_CREATE(TypeName,ParamType1,ParamValue1,ParamType2,ParamValue2,ParamType3,ParamValue3,ParamType4,ParamValue4,ParamType5,ParamValue5,ParamType6,ParamValue6)
+
+#define ENQUEUE_UNIQUE_RENDER_COMMAND_SIXPARAMETER_DECLARE_TEMPLATE(TypeName,TemplateParamName,ParamType1,ParamName1,ParamValue1,ParamType2,ParamName2,ParamValue2,ParamType3,ParamName3,ParamValue3,ParamType4,ParamName4,ParamValue4,ParamType5,ParamName5,ParamValue5,ParamType6,ParamName6,ParamValue6,Code) \
+	template <typename TemplateParamName> \
+	ENQUEUE_UNIQUE_RENDER_COMMAND_SIXPARAMETER_DECLARE_OPTTYPENAME(TypeName,ParamType1,ParamName1,ParamValue1,ParamType2,ParamName2,ParamValue2,ParamType3,ParamName3,ParamValue3,ParamType4,ParamName4,ParamValue4,ParamType5,ParamName5,ParamValue5,ParamType6,ParamName6,ParamValue6,OptTypename,Code)
+
+#define ENQUEUE_UNIQUE_RENDER_COMMAND_SIXPARAMETER_CREATE_TEMPLATE(TypeName,TemplateParamName,ParamType1,ParamValue1,ParamType2,ParamValue2,ParamType3,ParamValue3,ParamType4,ParamValue4,ParamType5,ParamValue5,ParamType6,ParamValue6) \
+	ENQUEUE_UNIQUE_RENDER_COMMAND_SIXPARAMETER_CREATE(TypeName<TemplateParamName>,ParamType1,ParamValue1,ParamType2,ParamValue2,ParamType3,ParamValue3,ParamType4,ParamValue4,ParamType5,ParamValue5,ParamType6,ParamValue6)
+
 #endif //UE_SERVER
 
 
