Tick movie player at a higher frequency for smoother playback.
Also fix  tearing issues with PS4 movie playback.

diff --git a/Engine/Source/Runtime/MoviePlayer/Private/DefaultGameMoviePlayer.cpp b/Engine/Source/Runtime/MoviePlayer/Private/DefaultGameMoviePlayer.cpp
--- a/Engine/Source/Runtime/MoviePlayer/Private/DefaultGameMoviePlayer.cpp
+++ b/Engine/Source/Runtime/MoviePlayer/Private/DefaultGameMoviePlayer.cpp
@@ -27,7 +27,7 @@ TSharedPtr<FDefaultGameMoviePlayer> FDefaultGameMoviePlayer::Get()
 }
 
 FDefaultGameMoviePlayer::FDefaultGameMoviePlayer()
-	: FTickableObjectRenderThread(false)
+	: FTickableObjectRenderThread(false, true)
 	, SyncMechanism(NULL)
 	, MovieStreamingIsDone(1)
 	, LoadingIsDone(1)
@@ -249,6 +249,18 @@ void FDefaultGameMoviePlayer::WaitForMovieToFinish()
 				// Gives widgets a chance to process any accumulated input
 				SlateApp.FinishedInputThisFrame();
 
+				float DeltaTime = SlateApp.GetDeltaTime();				
+
+				//pass this rather than doing ::Get() because the sharedptr isn't threadsafe.
+				ENQUEUE_UNIQUE_RENDER_COMMAND_TWOPARAMETER(
+					WaitingTickStreamer,
+					FDefaultGameMoviePlayer*, MoviePlayer, this,
+					float, DeltaTime, DeltaTime,
+					{
+					MoviePlayer->TickStreamer(DeltaTime);
+				}
+				);
+				
 				SlateApp.Tick();
 
 				// Synchronize the game thread and the render thread so that the render thread doesn't get too far behind.
@@ -311,21 +323,14 @@ bool FDefaultGameMoviePlayer::IsMovieStreamingFinished() const
 
 void FDefaultGameMoviePlayer::Tick( float DeltaTime )
 {
+	check(IsInRenderingThread());
 	if (LoadingScreenWindowPtr.IsValid() && RendererPtr.IsValid())
 	{
-		if (MovieStreamingIsPrepared() && !IsMovieStreamingFinished())
-		{
-			const bool bMovieIsDone = MovieStreamer->Tick(DeltaTime);
-			if (bMovieIsDone)
-			{
-				MovieStreamingIsDone.Set(1);
-			}
-		}
-
 		if (!IsLoadingFinished() && SyncMechanism)
 		{
 			if (SyncMechanism->IsSlateDrawPassEnqueued())
-			{
+			{				
+				TickStreamer(DeltaTime);
 				RendererPtr.Pin()->DrawWindows();
 				SyncMechanism->ResetSlateDrawPassEnqueued();
 			}
@@ -333,6 +338,18 @@ void FDefaultGameMoviePlayer::Tick( float DeltaTime )
 	}
 }
 
+void FDefaultGameMoviePlayer::TickStreamer(float DeltaTime)
+{	
+	if (MovieStreamingIsPrepared() && !IsMovieStreamingFinished())
+	{
+		const bool bMovieIsDone = MovieStreamer->Tick(DeltaTime);		
+		if (bMovieIsDone)
+		{
+			MovieStreamingIsDone.Set(1);
+		}
+	}
+}
+
 TStatId FDefaultGameMoviePlayer::GetStatId() const
 {
 	RETURN_QUICK_DECLARE_CYCLE_STAT(FDefaultGameMoviePlayer, STATGROUP_Tickables);
diff --git a/Engine/Source/Runtime/MoviePlayer/Private/DefaultGameMoviePlayer.h b/Engine/Source/Runtime/MoviePlayer/Private/DefaultGameMoviePlayer.h
--- a/Engine/Source/Runtime/MoviePlayer/Private/DefaultGameMoviePlayer.h
+++ b/Engine/Source/Runtime/MoviePlayer/Private/DefaultGameMoviePlayer.h
@@ -33,8 +33,12 @@ public:
 	virtual TStatId GetStatId() const override;
 	virtual bool IsTickable() const override;
 	
+	
 private:
 
+	/** Ticks the underlying MovieStreamer.  Must be done exactly once before each DrawWindows call. */
+	void TickStreamer(float DeltaTime);
+
 	/** True if we have both a registered movie streamer and movies to stream */
 	bool MovieStreamingIsPrepared() const;
 
diff --git a/Engine/Source/Runtime/RenderCore/Private/RenderingThread.cpp b/Engine/Source/Runtime/RenderCore/Private/RenderingThread.cpp
--- a/Engine/Source/Runtime/RenderCore/Private/RenderingThread.cpp
+++ b/Engine/Source/Runtime/RenderCore/Private/RenderingThread.cpp
@@ -183,15 +183,39 @@ FSuspendRenderingThread::~FSuspendRenderingThread()
 
 /** Static array of tickable objects that are ticked from rendering thread*/
 FTickableObjectRenderThread::FRenderingThreadTickableObjectsArray FTickableObjectRenderThread::RenderingThreadTickableObjects;
+FTickableObjectRenderThread::FRenderingThreadTickableObjectsArray FTickableObjectRenderThread::RenderingThreadHighFrequencyTickableObjects;
+
+void TickHighFrequencyTickables(double CurTime)
+{
+	static double LastHighFreqTime = FPlatformTime::Seconds();
+	float DeltaSecondsHighFreq = CurTime - LastHighFreqTime;
+
+	// tick any high frequency rendering thread tickables.
+	for (int32 ObjectIndex = 0; ObjectIndex < FTickableObjectRenderThread::RenderingThreadHighFrequencyTickableObjects.Num(); ObjectIndex++)
+	{
+		FTickableObjectRenderThread* TickableObject = FTickableObjectRenderThread::RenderingThreadHighFrequencyTickableObjects[ObjectIndex];
+		// make sure it wants to be ticked and the rendering thread isn't suspended
+		if (TickableObject->IsTickable())
+		{
+			STAT(FScopeCycleCounter(TickableObject->GetStatId());)
+				TickableObject->Tick(DeltaSecondsHighFreq);
+		}
+	}
+
+	LastHighFreqTime = CurTime;
+}
 
 void TickRenderingTickables()
 {
 	static double LastTickTime = FPlatformTime::Seconds();
+	
 
 	// calc how long has passed since last tick
 	double CurTime = FPlatformTime::Seconds();
 	float DeltaSeconds = CurTime - LastTickTime;
-	
+		
+	TickHighFrequencyTickables(CurTime);
+
 	if (DeltaSeconds < (1.f/GRenderingThreadMaxIdleTickFrequency))
 	{
 		return;
diff --git a/Engine/Source/Runtime/RenderCore/Public/TickableObjectRenderThread.h b/Engine/Source/Runtime/RenderCore/Public/TickableObjectRenderThread.h
--- a/Engine/Source/Runtime/RenderCore/Public/TickableObjectRenderThread.h
+++ b/Engine/Source/Runtime/RenderCore/Public/TickableObjectRenderThread.h
@@ -18,7 +18,7 @@ class RENDERCORE_API FTickableObjectRenderThread
 public:
 
 	/** Static array of tickable objects that are ticked from rendering thread*/
-	static struct FRenderingThreadTickableObjectsArray : public TArray<FTickableObjectRenderThread*>
+	struct FRenderingThreadTickableObjectsArray : public TArray<FTickableObjectRenderThread*>
 	{
 		~FRenderingThreadTickableObjectsArray()
 		{
@@ -33,15 +33,19 @@ public:
 			// if we exited uncleanly from a runaway loop, crash explicitly in Dev
 			check(Num() == 0);
 		}
-	} RenderingThreadTickableObjects;
+	};
+
+	static FRenderingThreadTickableObjectsArray RenderingThreadTickableObjects;
+	static FRenderingThreadTickableObjectsArray RenderingThreadHighFrequencyTickableObjects;
 
 	/**
 	 * Registers this instance with the static array of tickable objects.	
 	 *
 	 * @param bRegisterImmediately true if the object should be registered immediately.
 	 */
-	FTickableObjectRenderThread(bool bRegisterImmediately=true) :
-		bRegistered(false)
+	FTickableObjectRenderThread(bool bRegisterImmediately=true, bool bInHighFrequency=false) :
+		bRegistered(false),
+		bHighFrequency(bInHighFrequency)
 	{
 		if(bRegisterImmediately)
 		{
@@ -63,9 +67,10 @@ public:
 		checkf(IsInRenderingThread(), TEXT("Game thread attempted to unregister an object in the RenderingThreadTickableObjects array."));
 		if (bRegistered)
 		{
-			const int32 Pos=RenderingThreadTickableObjects.Find(this);
+			FRenderingThreadTickableObjectsArray& TickableObjectArray = bHighFrequency ? RenderingThreadHighFrequencyTickableObjects : RenderingThreadTickableObjects;
+			const int32 Pos = TickableObjectArray.Find(this);
 			check(Pos!=INDEX_NONE);
-			RenderingThreadTickableObjects.RemoveAt(Pos);
+			TickableObjectArray.RemoveAt(Pos);
 			bRegistered = false;
 		}
 	}
@@ -79,8 +84,16 @@ public:
 		// make sure that only the rendering thread is attempting to add items to the RenderingThreadTickableObjects list
 		checkf(IsInRenderingThread(), TEXT("Game thread attempted to register an object in the RenderingThreadTickableObjects array."));
 		check(!RenderingThreadTickableObjects.Contains(this));
+		check(!RenderingThreadHighFrequencyTickableObjects.Contains(this));
 		check(!bRegistered);
-		RenderingThreadTickableObjects.Add( this );
+		if (bHighFrequency)
+		{
+			RenderingThreadHighFrequencyTickableObjects.Add(this);
+		}
+		else
+		{
+			RenderingThreadTickableObjects.Add(this);
+		}		
 		bRegistered = true;
 	}
 
@@ -119,6 +132,7 @@ public:
 
 private:
 	bool bRegistered;
+	bool bHighFrequency;
 };
 
 #endif
