UE4 - D3D12 fix

diff --git a/Engine/Source/Runtime/RenderCore/Public/UniformBuffer.h b/Engine/Source/Runtime/RenderCore/Public/UniformBuffer.h
--- a/Engine/Source/Runtime/RenderCore/Public/UniformBuffer.h
+++ b/Engine/Source/Runtime/RenderCore/Public/UniformBuffer.h
@@ -84,7 +84,7 @@ public:
 	/** Creates a uniform buffer with the given value, and returns a structured reference to it. */
 	static TUniformBufferRef<TBufferStruct> CreateUniformBufferImmediate(const TBufferStruct& Value, EUniformBufferUsage Usage)
 	{
-		check(IsInRenderingThread());
+		check(IsInRenderingThread() || IsInRHIThread());
 		return TUniformBufferRef<TBufferStruct>(RHICreateUniformBuffer(&Value,TBufferStruct::StaticStruct.GetLayout(),Usage));
 	}
 	/** Creates a uniform buffer with the given value, and returns a structured reference to it. */
