UE4 - Fix for togglerenderingthread on D3D12

diff --git a/Engine/Source/Runtime/RenderCore/Private/RenderingThread.cpp b/Engine/Source/Runtime/RenderCore/Private/RenderingThread.cpp
--- a/Engine/Source/Runtime/RenderCore/Private/RenderingThread.cpp
+++ b/Engine/Source/Runtime/RenderCore/Private/RenderingThread.cpp
@@ -654,6 +654,8 @@ void StopRenderingThread()
 			delete GRenderingThread;
 			GRenderingThread = NULL;
 			
+			GRHICommandList.LatchBypass();
+
 			delete GRenderingThreadRunnable;
 			GRenderingThreadRunnable = NULL;
 		}
