Fix to round up for mip size estimation

diff --git a/Engine/Source/Runtime/RenderCore/Private/RenderUtils.cpp b/Engine/Source/Runtime/RenderCore/Private/RenderUtils.cpp
--- a/Engine/Source/Runtime/RenderCore/Private/RenderUtils.cpp
+++ b/Engine/Source/Runtime/RenderCore/Private/RenderUtils.cpp
@@ -584,6 +584,11 @@ SIZE_T CalcTextureMipMapSize3D( uint32 TextureSizeX, uint32 TextureSizeY, uint32
 	uint32 ZExtent;
 	CalcMipMapExtent3D(TextureSizeX, TextureSizeY, TextureSizeZ, Format, MipIndex, XExtent, YExtent, ZExtent);
 
+	// Offset MipExtent to round up result
+	XExtent += GPixelFormats[Format].BlockSizeX - 1;
+	YExtent += GPixelFormats[Format].BlockSizeY - 1;
+	ZExtent += GPixelFormats[Format].BlockSizeZ - 1;
+
 	const uint32 XPitch = (XExtent / GPixelFormats[Format].BlockSizeX) * GPixelFormats[Format].BlockBytes;
 	const uint32 NumRows = YExtent / GPixelFormats[Format].BlockSizeY;
 	const uint32 NumLayers = ZExtent / GPixelFormats[Format].BlockSizeZ;
@@ -610,6 +615,9 @@ SIZE_T CalcTextureMipMapSize( uint32 TextureSizeX, uint32 TextureSizeY, EPixelFo
 {
 	FIntPoint MipExtent = CalcMipMapExtent(TextureSizeX, TextureSizeY, Format, MipIndex);
 
+	// Offset MipExtent to round up result
+	MipExtent += FIntPoint(GPixelFormats[Format].BlockSizeX, GPixelFormats[Format].BlockSizeY) - FIntPoint(1, 1);
+
 	const uint32 Pitch = (MipExtent.X / GPixelFormats[Format].BlockSizeX) * GPixelFormats[Format].BlockBytes;
 	const uint32 NumRows = MipExtent.Y / GPixelFormats[Format].BlockSizeY;
 
