UE4 - Fixed a critical bug in the render thread fencing that could sometimes cause fences to complete before previous commands.

diff --git a/Engine/Source/Runtime/RenderCore/Private/RenderingThread.cpp b/Engine/Source/Runtime/RenderCore/Private/RenderingThread.cpp
--- a/Engine/Source/Runtime/RenderCore/Private/RenderingThread.cpp
+++ b/Engine/Source/Runtime/RenderCore/Private/RenderingThread.cpp
@@ -747,19 +747,8 @@ void FRenderCommandFence::BeginFence()
 			STAT_FNullGraphTask_FenceRenderCommand,
 			STATGROUP_TaskGraphTasks);
 
-		if (IsFenceComplete())
-		{
-			CompletionEvent = TGraphTask<FNullGraphTask>::CreateTask(NULL, ENamedThreads::GameThread).ConstructAndDispatchWhenReady(
-				GET_STATID(STAT_FNullGraphTask_FenceRenderCommand), ENamedThreads::RenderThread);
-		}
-		else
-		{
-			// we already had a fence, so we will chain this one to the old one as a prerequisite
-			FGraphEventArray Prerequistes;
-			Prerequistes.Add(CompletionEvent);
-			CompletionEvent = TGraphTask<FNullGraphTask>::CreateTask(&Prerequistes, ENamedThreads::GameThread).ConstructAndDispatchWhenReady(
-				GET_STATID(STAT_FNullGraphTask_FenceRenderCommand), ENamedThreads::RenderThread);
-		}
+		CompletionEvent = TGraphTask<FNullGraphTask>::CreateTask(NULL, ENamedThreads::GameThread).ConstructAndDispatchWhenReady(
+			GET_STATID(STAT_FNullGraphTask_FenceRenderCommand), ENamedThreads::RenderThread);
 	}
 }
 
