UE4 - Non-unity compile fix

diff --git a/Engine/Source/Runtime/RenderCore/Public/RenderingThread.h b/Engine/Source/Runtime/RenderCore/Public/RenderingThread.h
--- a/Engine/Source/Runtime/RenderCore/Public/RenderingThread.h
+++ b/Engine/Source/Runtime/RenderCore/Public/RenderingThread.h
@@ -277,7 +277,7 @@ namespace Detail
 	template<typename T>
 	struct FalseType
 	{
-		static constexpr const bool False = false;
+		static CONSTEXPR const bool False = false;
 	};
 }
 
