UE4 - Change pending delete data structure for editor, and more specifically the cooker. Render resources can accumulate for hours in the cooker.

diff --git a/Engine/Source/Runtime/RenderCore/Private/RenderingThread.cpp b/Engine/Source/Runtime/RenderCore/Private/RenderingThread.cpp
--- a/Engine/Source/Runtime/RenderCore/Private/RenderingThread.cpp
+++ b/Engine/Source/Runtime/RenderCore/Private/RenderingThread.cpp
@@ -1100,6 +1100,43 @@ FRHICommandListImmediate& GetImmediateCommandList_ForRenderCommand()
 	return FRHICommandListExecutor::GetImmediateCommandList();
 }
 
+#if WITH_EDITOR || IS_PROGRAM
+
+// mainly concerned about the cooker here, but anyway, the editor can run without a frame for a very long time (hours) and we do not have enough lock free links. 
+
+/** The set of deferred cleanup objects which are pending cleanup. */
+TArray<FDeferredCleanupInterface*> PendingCleanupObjectsList;
+FCriticalSection PendingCleanupObjectsListLock;
+
+FPendingCleanupObjects::FPendingCleanupObjects()
+{
+	check(IsInGameThread());
+	{
+		FScopeLock Lock(&PendingCleanupObjectsListLock);
+		Exchange(CleanupArray, PendingCleanupObjectsList);
+	}
+}
+
+FPendingCleanupObjects::~FPendingCleanupObjects()
+{
+	QUICK_SCOPE_CYCLE_COUNTER(STAT_FPendingCleanupObjects_Destruct);
+
+	for (int32 ObjectIndex = 0; ObjectIndex < CleanupArray.Num(); ObjectIndex++)
+	{
+		CleanupArray[ObjectIndex]->FinishCleanup();
+	}
+}
+
+void BeginCleanup(FDeferredCleanupInterface* CleanupObject)
+{
+	{
+		FScopeLock Lock(&PendingCleanupObjectsListLock);
+		PendingCleanupObjectsList.Add(CleanupObject);
+	}
+}
+
+#else
+
 /** The set of deferred cleanup objects which are pending cleanup. */
 static TLockFreePointerListUnordered<FDeferredCleanupInterface, PLATFORM_CACHE_LINE_SIZE>	PendingCleanupObjectsList;
 
@@ -1113,7 +1150,7 @@ FPendingCleanupObjects::~FPendingCleanupObjects()
 {
 	QUICK_SCOPE_CYCLE_COUNTER(STAT_FPendingCleanupObjects_Destruct);
 
-	for(int32 ObjectIndex = 0;ObjectIndex < CleanupArray.Num();ObjectIndex++)
+	for (int32 ObjectIndex = 0; ObjectIndex < CleanupArray.Num(); ObjectIndex++)
 	{
 		CleanupArray[ObjectIndex]->FinishCleanup();
 	}
@@ -1124,6 +1161,8 @@ void BeginCleanup(FDeferredCleanupInterface* CleanupObject)
 	PendingCleanupObjectsList.Push(CleanupObject);
 }
 
+#endif
+
 FPendingCleanupObjects* GetPendingCleanupObjects()
 {
 	return new FPendingCleanupObjects;
