Fix for non-unity.

diff --git a/Engine/Source/Runtime/RenderCore/Private/RenderingThread.cpp b/Engine/Source/Runtime/RenderCore/Private/RenderingThread.cpp
--- a/Engine/Source/Runtime/RenderCore/Private/RenderingThread.cpp
+++ b/Engine/Source/Runtime/RenderCore/Private/RenderingThread.cpp
@@ -12,6 +12,7 @@
 #include "Misc/CoreStats.h"
 #include "Misc/TimeGuard.h"
 #include "Misc/CoreDelegates.h"
+#include "Misc/ScopeLock.h"
 #include "RenderCore.h"
 #include "RenderCommandFence.h"
 #include "RHI.h"
