Added a temp ensure get callstacks for "FlushRenderingCommands called recursively" warnings being triggered during automation tests

diff --git a/Engine/Source/Runtime/RenderCore/Private/RenderingThread.cpp b/Engine/Source/Runtime/RenderCore/Private/RenderingThread.cpp
--- a/Engine/Source/Runtime/RenderCore/Private/RenderingThread.cpp
+++ b/Engine/Source/Runtime/RenderCore/Private/RenderingThread.cpp
@@ -953,6 +953,11 @@ static void GameThreadWaitForTask(const FGraphEventRef& Task, bool bEmptyGameThr
 			NumRecursiveCalls++;
 			if (NumRecursiveCalls > 1)
 			{
+				if(GIsAutomationTesting)
+				{
+					// temp test to log callstacks for this being triggered during automation tests
+					ensureMsgf(false, TEXT("FlushRenderingCommands called recursively! %d calls on the stack."));
+				}
 				UE_LOG(LogRendererCore,Warning,TEXT("FlushRenderingCommands called recursively! %d calls on the stack."), NumRecursiveCalls);
 			}
 			if (NumRecursiveCalls > 1 || FTaskGraphInterface::Get().IsThreadProcessingTasks(ENamedThreads::GameThread))
