RL - Add additional compiler error for previous change of FDeferredCleanupInterface

diff --git a/Engine/Source/Runtime/RenderCore/Public/RenderingThread.h b/Engine/Source/Runtime/RenderCore/Public/RenderingThread.h
--- a/Engine/Source/Runtime/RenderCore/Public/RenderingThread.h
+++ b/Engine/Source/Runtime/RenderCore/Public/RenderingThread.h
@@ -623,6 +623,9 @@ class RENDERCORE_API FDeferredCleanupInterface
 {
 public:
 	virtual ~FDeferredCleanupInterface() {}
+
+	DEPRECATED(4.20, "FinishCleanup is deprecated. Use RAII in the destructor instead.")
+	virtual void FinishCleanup() final {}
 };
 
 /**
