UE4.22 - More info to track down failure
#rnx

diff --git a/Engine/Source/Runtime/RenderCore/Public/Shader.h b/Engine/Source/Runtime/RenderCore/Public/Shader.h
--- a/Engine/Source/Runtime/RenderCore/Public/Shader.h
+++ b/Engine/Source/Runtime/RenderCore/Public/Shader.h
@@ -2194,7 +2194,13 @@ public:
 				check(Type);
 				checkSlow(FName(Type->GetName()) != NAME_None);
 				Ar << Type;
+#if WITH_EDITOR
+				TRefCountPtr<FShader>* Found = Shaders.Find(Key);
+				checkf(Found, TEXT("Unable to find FShaderType %s!"), Type->GetName());
+				FShader* CurrentShader = *Found;
+#else
 				FShader* CurrentShader = Shaders.FindChecked(Key);
+#endif
 				SerializeShaderForSaving(CurrentShader, Ar, bHandleShaderKeyChanges, bInlineShaderResource);
 			}
 
