REL - Speculative fix for UE-70399

diff --git a/Engine/Source/Runtime/RenderCore/Public/Shader.h b/Engine/Source/Runtime/RenderCore/Public/Shader.h
--- a/Engine/Source/Runtime/RenderCore/Public/Shader.h
+++ b/Engine/Source/Runtime/RenderCore/Public/Shader.h
@@ -1954,7 +1954,7 @@ public:
 	/** Destructor ensures pipelines cleared up. */
 	virtual ~TShaderMap()
 	{
-		EmptyShaderPipelines();
+		Empty();
 	}
 
 	EShaderPlatform GetShaderPlatform() const { return Platform; }
