fix for clang build failing

diff --git a/Engine/Source/Runtime/RenderCore/Private/RenderTargetPool.cpp b/Engine/Source/Runtime/RenderCore/Private/RenderTargetPool.cpp
--- a/Engine/Source/Runtime/RenderCore/Private/RenderTargetPool.cpp
+++ b/Engine/Source/Runtime/RenderCore/Private/RenderTargetPool.cpp
@@ -16,7 +16,7 @@ TGlobalResource<FRenderTargetPool> GRenderTargetPool;
 
 DEFINE_LOG_CATEGORY_STATIC(LogRenderTargetPool, Warning, All);
 
-static void DumpRenderTargetPoolMemory(FOutputDevice& OutputDevice)
+RENDERCORE_API void DumpRenderTargetPoolMemory(FOutputDevice& OutputDevice)
 {
 	GRenderTargetPool.DumpMemoryUsage(OutputDevice);
 }
@@ -381,6 +381,10 @@ bool FRenderTargetPool::FindFreeElement(FRHICommandList& RHICmdList, const FPool
 			for(uint32 i = 0, Num = (uint32)PooledRenderTargets.Num(); i < Num; ++i)
 			{
 				FPooledRenderTarget* Element = PooledRenderTargets[i];
+				if(Element && Element->GetDesc().Compare(Desc, bExactMatch))
+				{
+					int a = 0;
+				}
     
 				if(Element && Element->IsFree() && Element->GetDesc().Compare(Desc, bExactMatch))
 				{
@@ -1143,6 +1147,36 @@ void FRenderTargetPool::DumpMemoryUsage(FOutputDevice& OutputDevice)
 	uint32 PoolKB=0;
 	GetStats(NumTargets,PoolKB,UsedKB);
 	OutputDevice.Logf(TEXT("%.3fMB total, %.3fMB used, %d render targets"), PoolKB / 1024.f, UsedKB / 1024.f, NumTargets);
+
+
+	uint32 DeferredTotal = 0;
+	OutputDevice.Logf(TEXT("Deferred Render Targets:"));
+	for (int32 i = 0; i < DeferredDeleteArray.Num(); ++i)
+	{
+		FPooledRenderTarget* Element = DeferredDeleteArray[i];
+
+		if (Element)
+		{
+			check(!Element->IsSnapshot());
+			OutputDevice.Logf(
+				TEXT("  %6.3fMB %4dx%4d%s%s %2dmip(s) %s (%s) %s %s"),
+				ComputeSizeInKB(*Element) / 1024.0f,
+				Element->Desc.Extent.X,
+				Element->Desc.Extent.Y,
+				Element->Desc.Depth > 1 ? *FString::Printf(TEXT("x%3d"), Element->Desc.Depth) : (Element->Desc.IsCubemap() ? TEXT("cube") : TEXT("    ")),
+				Element->Desc.bIsArray ? *FString::Printf(TEXT("[%3d]"), Element->Desc.ArraySize) : TEXT("     "),
+				Element->Desc.NumMips,
+				Element->Desc.DebugName,
+				GPixelFormats[Element->Desc.Format].Name,
+				Element->IsTransient() ? TEXT("(transient)") : TEXT(""),
+				GSupportsTransientResourceAliasing ? *FString::Printf(TEXT("Frames since last discard: %d"), GFrameNumberRenderThread - Element->FrameNumberLastDiscard) : TEXT("")
+			);
+			uint32 SizeInKB = ComputeSizeInKB(*Element);
+			DeferredTotal += SizeInKB;
+		}
+	}
+	OutputDevice.Logf(TEXT("%.3fMB Deferred total"), DeferredTotal / 1024.f);
+
 }
 
 uint32 FPooledRenderTarget::AddRef() const
