Adding more information to be able to detect why FMaterialShaderMap::Serialize randomly fails in CIS.

diff --git a/Engine/Source/Runtime/RenderCore/Public/Shader.h b/Engine/Source/Runtime/RenderCore/Public/Shader.h
--- a/Engine/Source/Runtime/RenderCore/Public/Shader.h
+++ b/Engine/Source/Runtime/RenderCore/Public/Shader.h
@@ -2194,7 +2194,15 @@ public:
 				Ar << Type;
 #if WITH_EDITOR
 				TRefCountPtr<FShader>* Found = Shaders.Find(Key);
-				checkf(Found, TEXT("Unable to find FShaderType %s!"), Type->GetName());
+				if (!Found)
+				{
+					UE_LOG(LogShaders, Error, TEXT("Unable to find type %s permutation id %d inside ShadersMap"), Type->GetName(), Key.PermutationId);
+					for (TMap<FShaderPrimaryKey, TRefCountPtr<FShader> >::TIterator It(Shaders); It; ++It)
+					{
+						UE_LOG(LogShaders, Error, TEXT("ShadersMap type %s Permutation id %d"), It->Key.Type->GetName(), It->Key.PermutationId);
+					}
+				}
+				checkf(Found, TEXT("Unable to find FShaderType %s permutation id %d!"), Type->GetName(), Key.PermutationId);
 				FShader* CurrentShader = *Found;
 #else
 				FShader* CurrentShader = Shaders.FindChecked(Key);
