#UE4 Fix a copy error due to an open file handle in FEditorShaderStableInfo::Finalize. Also added some additional logging in case copy failures happen again in the future.

diff --git a/Engine/Source/Runtime/RenderCore/Private/ShaderCodeLibrary.cpp b/Engine/Source/Runtime/RenderCore/Private/ShaderCodeLibrary.cpp
--- a/Engine/Source/Runtime/RenderCore/Private/ShaderCodeLibrary.cpp
+++ b/Engine/Source/Runtime/RenderCore/Private/ShaderCodeLibrary.cpp
@@ -1119,7 +1119,12 @@ struct FEditorShaderCodeArchive
 					FString OutputFilePath = GetCodeArchiveFilename(OutputDir, LibraryName, FormatName);
 
 					// Copy to output location - support for iterative native library cooking
-					IFileManager::Get().Copy(*OutputFilePath, *IntermediateFormatPath, true, true);
+					uint32 Result = IFileManager::Get().Copy(*OutputFilePath, *IntermediateFormatPath, true, true);
+					if (Result != COPY_OK)
+					{
+						UE_LOG(LogShaderLibrary, Error, TEXT("FEditorShaderCodeArchive shader library copy failed to %s. Failed to finalize Shared Shader Library %s with format %s"), *OutputFilePath, *LibraryName, *FormatName.ToString());
+						bSuccess = false;
+					}
 				}
 			}
 		}
@@ -1144,7 +1149,12 @@ struct FEditorShaderCodeArchive
 				FString OutputFilePath = GetPipelinesArchiveFilename(OutputDir, LibraryName, FormatName);
 
 				// Copy to output location - support for iterative native library cooking
-				IFileManager::Get().Copy(*OutputFilePath, *TempFilePath, true, true);
+				uint32 Result = IFileManager::Get().Copy(*OutputFilePath, *TempFilePath, true, true);
+				if (Result != COPY_OK)
+				{
+					UE_LOG(LogShaderLibrary, Error, TEXT("FEditorShaderCodeArchive pipeline copy failed to %s. Failed to finalize Shared Shader Library %s with format %s"), *OutputFilePath, *LibraryName, *FormatName.ToString());
+					bSuccess = false;
+				}
 			}
 		}
 
@@ -1315,21 +1325,23 @@ struct FEditorShaderStableInfo
 			FString IntermediateFormatPath = GetStableInfoArchiveFilename(FPaths::ProjectSavedDir() / TEXT("Shaders") / FormatName.ToString(), LibraryName, FormatName);
 
 			// Write directly to the file
-			TUniquePtr<FArchive> IntermediateFormatAr(IFileManager::Get().CreateFileWriter(*IntermediateFormatPath));
+			{
+				TUniquePtr<FArchive> IntermediateFormatAr(IFileManager::Get().CreateFileWriter(*IntermediateFormatPath));
 
-			const FString HeaderText = FStableShaderKeyAndValue::HeaderLine();
-			auto HeaderSrc = StringCast<ANSICHAR>(*HeaderText, HeaderText.Len());
+				const FString HeaderText = FStableShaderKeyAndValue::HeaderLine();
+				auto HeaderSrc = StringCast<ANSICHAR>(*HeaderText, HeaderText.Len());
 
-			IntermediateFormatAr->Serialize((ANSICHAR*)HeaderSrc.Get(), HeaderSrc.Length() * sizeof(ANSICHAR));
+				IntermediateFormatAr->Serialize((ANSICHAR*)HeaderSrc.Get(), HeaderSrc.Length() * sizeof(ANSICHAR));
 
-			FString LineBuffer;
-			LineBuffer.Reserve(512);
+				FString LineBuffer;
+				LineBuffer.Reserve(512);
 
-			for (const FStableShaderKeyAndValue& Item : StableMap)
-			{
-				Item.ToString(LineBuffer);
-				auto LineConverted = StringCast<ANSICHAR>(*LineBuffer, LineBuffer.Len());
-				IntermediateFormatAr->Serialize((ANSICHAR*)LineConverted.Get(), LineConverted.Length() * sizeof(ANSICHAR));
+				for (const FStableShaderKeyAndValue& Item : StableMap)
+				{
+					Item.ToString(LineBuffer);
+					auto LineConverted = StringCast<ANSICHAR>(*LineBuffer, LineBuffer.Len());
+					IntermediateFormatAr->Serialize((ANSICHAR*)LineConverted.Get(), LineConverted.Length() * sizeof(ANSICHAR));
+				}
 			}
 
 			// Only the master cooker needs to write to the output directory, child cookers only write to the Saved directory
@@ -1338,8 +1350,16 @@ struct FEditorShaderStableInfo
 				FString OutputFilePath = GetStableInfoArchiveFilename(OutputDir, LibraryName, FormatName);
 
 				// Copy to output location - support for iterative native library cooking
-				IFileManager::Get().Copy(*OutputFilePath, *IntermediateFormatPath, true, true);
-				OutSCLCSVPath = OutputFilePath;
+				uint32 Result = IFileManager::Get().Copy(*OutputFilePath, *IntermediateFormatPath, true, true);
+				if (Result == COPY_OK)
+				{
+					OutSCLCSVPath = OutputFilePath;
+				}
+				else
+				{
+					UE_LOG(LogShaderLibrary, Error, TEXT("FEditorShaderStableInfo copy failed to %s. Failed to finalize Shared Shader Library %s with format %s"), *OutputFilePath, *LibraryName, *FormatName.ToString());
+					bSuccess = false;
+				}
 			}
 		}
 
