UE-71809 is fixed, so it's time to dismantle extra logs, which were used to catch UE-71809.

diff --git a/Engine/Source/Runtime/RenderCore/Public/Shader.h b/Engine/Source/Runtime/RenderCore/Public/Shader.h
--- a/Engine/Source/Runtime/RenderCore/Public/Shader.h
+++ b/Engine/Source/Runtime/RenderCore/Public/Shader.h
@@ -2192,21 +2192,7 @@ public:
 				check(Type);
 				checkSlow(FName(Type->GetName()) != NAME_None);
 				Ar << Type;
-#if WITH_EDITOR
-				TRefCountPtr<FShader>* Found = Shaders.Find(Key);
-				if (!Found)
-				{
-					UE_LOG(LogShaders, Error, TEXT("Unable to find type %s permutation id %d inside ShadersMap"), Type->GetName(), Key.PermutationId);
-					for (TMap<FShaderPrimaryKey, TRefCountPtr<FShader> >::TIterator It(Shaders); It; ++It)
-					{
-						UE_LOG(LogShaders, Error, TEXT("ShadersMap type %s Permutation id %d"), It->Key.Type->GetName(), It->Key.PermutationId);
-					}
-				}
-				checkf(Found, TEXT("Unable to find FShaderType %s permutation id %d!"), Type->GetName(), Key.PermutationId);
-				FShader* CurrentShader = *Found;
-#else
 				FShader* CurrentShader = Shaders.FindChecked(Key);
-#endif
 				SerializeShaderForSaving(CurrentShader, Ar, bHandleShaderKeyChanges, bInlineShaderResource);
 			}
 
