UE4.22 - Get more info to track down crash
#rnx

diff --git a/Engine/Source/Runtime/RenderCore/Private/RenderTargetPool.cpp b/Engine/Source/Runtime/RenderCore/Private/RenderTargetPool.cpp
--- a/Engine/Source/Runtime/RenderCore/Private/RenderTargetPool.cpp
+++ b/Engine/Source/Runtime/RenderCore/Private/RenderTargetPool.cpp
@@ -1194,7 +1194,7 @@ uint32 FPooledRenderTarget::Release()
 {
 	if (!bSnapshot)
 	{
-		check(IsInRenderingThread());
+		checkf(IsInRenderingThread(), TEXT("Tried to delete on non-render thread, PooledRT %s %s"), Desc.DebugName ? Desc.DebugName : TEXT("<Unnamed>"), *Desc.GenerateInfoString());
 		uint32 Refs = uint32(--NumRefs);
 		if(Refs == 0)
 		{
