Speculative fix for macOS client builds crashing on launch due to bad PSO cache file.  Don't write to PSO cache file while the Engine is shutting down - auto saves should have kept the file upto date.

diff --git a/Engine/Source/Runtime/RenderCore/Private/ShaderPipelineCache.cpp b/Engine/Source/Runtime/RenderCore/Private/ShaderPipelineCache.cpp
--- a/Engine/Source/Runtime/RenderCore/Private/ShaderPipelineCache.cpp
+++ b/Engine/Source/Runtime/RenderCore/Private/ShaderPipelineCache.cpp
@@ -18,7 +18,7 @@
 #include "Misc/EngineVersion.h"
 #include "PipelineStateCache.h"
 #include "PipelineFileCache.h"
-#include "Misc/ScopeRWLock.h"
+#include "Misc/ScopeLock.h"
 #include "Misc/CoreDelegates.h"
 #include "ShaderCodeLibrary.h"
 #include "TickableObjectRenderThread.h"
@@ -1099,16 +1099,14 @@ FShaderPipelineCache::FShaderPipelineCache(EShaderPlatform Platform)
 
 FShaderPipelineCache::~FShaderPipelineCache()
 {
+	// Only save PSO Record / Logging at shutdown
 	if (GetShaderPipelineCacheSaveBoundPSOLog())
 	{
 		FShaderPipelineCache::SavePipelineFileCache(FPipelineFileCache::SaveMode::BoundPSOsOnly);
 	}
-	if (GetPSOFileCacheSaveUserCache())
-	{
-		FShaderPipelineCache::SavePipelineFileCache(FPipelineFileCache::SaveMode::Incremental);
-	}
 	
-	Close();
+	// Close with shutdown flag
+	Close(true);
 	
 	// The render thread tick should be dead now and we are safe to destroy everything that needs to wait or manual destruction
 
@@ -1432,7 +1430,7 @@ bool FShaderPipelineCache::Save(FPipelineFileCache::SaveMode Mode)
 	return bOK;
 }
 
-void FShaderPipelineCache::Close()
+void FShaderPipelineCache::Close(bool bShuttingDown)
 {
 	FScopeLock Lock(&Mutex);
 		
@@ -1448,8 +1446,8 @@ void FShaderPipelineCache::Close()
 		Save(FPipelineFileCache::SaveMode::BoundPSOsOnly);
 	}
 	
-	// Force a fast save, just in case
-	if (GetPSOFileCacheSaveUserCache())
+	// Force a fast save, just in case - but not at shutdown / destruction
+	if (GetPSOFileCacheSaveUserCache() && !bShuttingDown)
 	{
 		Save(FPipelineFileCache::SaveMode::Incremental);
 	}
diff --git a/Engine/Source/Runtime/RenderCore/Public/ShaderPipelineCache.h b/Engine/Source/Runtime/RenderCore/Public/ShaderPipelineCache.h
--- a/Engine/Source/Runtime/RenderCore/Public/ShaderPipelineCache.h
+++ b/Engine/Source/Runtime/RenderCore/Public/ShaderPipelineCache.h
@@ -200,7 +200,7 @@ public:
 private:
 	bool Open(FString const& Name, EShaderPlatform Platform);
 	bool Save(FPipelineFileCache::SaveMode Mode);
-	void Close();
+	void Close(bool bShuttingDown = false);
 
 	bool Precompile(FRHICommandListImmediate& RHICmdList, EShaderPlatform Platform, FPipelineCacheFileFormatPSO const& PSO);
 	void PreparePipelineBatch(TDoubleLinkedList<FPipelineCacheFileFormatPSORead*>& PipelineBatch);
