Remove the error log (logging still reports internally) and report result back to higher level for it to decide if this is an error - we can not know at the engine level.  Soften the language of some of the other internal fail messages to steer the thinking away from errors and more informational.

diff --git a/Engine/Source/Runtime/RenderCore/Private/ShaderCodeLibrary.cpp b/Engine/Source/Runtime/RenderCore/Private/ShaderCodeLibrary.cpp
--- a/Engine/Source/Runtime/RenderCore/Private/ShaderCodeLibrary.cpp
+++ b/Engine/Source/Runtime/RenderCore/Private/ShaderCodeLibrary.cpp
@@ -1488,14 +1488,18 @@ public:
 #endif
 	}
 
-	void OpenLibrary(FString const& Name, FString const& Directory)
+	bool OpenLibrary(FString const& Name, FString const& Directory)
 	{
 		LLM_SCOPE(ELLMTag::Shaders);
+		
+		bool bResult = false;
 
 		if (ShaderPlatform < SP_NumPlatforms)
 		{
 			if (OpenShaderCode(Directory, ShaderPlatform, Name))
 			{
+				bResult = true;
+				
 				// Attempt to open the shared-cooked override code library if there is one.
 				// This is probably not ideal, but it should get shared-cooks working.
 				OpenShaderCode(Directory, ShaderPlatform, Name + TEXT("_SC"));
@@ -1503,11 +1507,6 @@ public:
 				// Inform the pipeline cache that the state of loaded libraries has changed
 				FShaderPipelineCache::ShaderLibraryStateChanged(FShaderPipelineCache::Opened, ShaderPlatform, Name);
 			}
-			else
-			{
-				FName PlatformName = LegacyShaderPlatformToShaderFormat(ShaderPlatform);
-				UE_LOG(LogShaderLibrary, Error, TEXT("Cooked Context: Failed to load Shared Shader Library %s from %s for %s"), *Name, *Directory, *PlatformName.GetPlainNameString());
-			}
 		}
 
 #if WITH_EDITOR
@@ -1528,6 +1527,8 @@ public:
 			}
 		}
 #endif
+		
+		return bResult;
 	}
 
 	void CloseLibrary(FString const& Name)
@@ -1588,12 +1589,12 @@ public:
 			}
 			else
 			{
-				UE_LOG(LogShaderLibrary, Display, TEXT("Failed to load Native Shared Shader Library: %s."), *Library);
+				UE_LOG(LogShaderLibrary, Display, TEXT("Cooked Context: No Native Shared Shader Library for %s"), *Library);
 			}
 		}
 		else
 		{
-			UE_LOG(LogShaderLibrary, Display, TEXT("Failed to load Shared Shader Library: %s and no native library supported."), *Library);
+			UE_LOG(LogShaderLibrary, Display, TEXT("Cooked Context: No Shared Shader Library for: %s and native library not supported."), *Library);
 		}
 
 		bool const bOK = IsValidRef(ShaderCodeArchive);
@@ -2253,12 +2254,14 @@ EShaderPlatform FShaderCodeLibrary::GetRuntimeShaderPlatform(void)
 	return Platform;
 }
 
-void FShaderCodeLibrary::OpenLibrary(FString const& Name, FString const& Directory)
+bool FShaderCodeLibrary::OpenLibrary(FString const& Name, FString const& Directory)
 {
+	bool bResult = false;
 	if (FShaderCodeLibraryImpl::Impl)
 	{
-		FShaderCodeLibraryImpl::Impl->OpenLibrary(Name, Directory);
+		bResult = FShaderCodeLibraryImpl::Impl->OpenLibrary(Name, Directory);
 	}
+	return bResult;
 }
 
 void FShaderCodeLibrary::CloseLibrary(FString const& Name)
diff --git a/Engine/Source/Runtime/RenderCore/Public/ShaderCodeLibrary.h b/Engine/Source/Runtime/RenderCore/Public/ShaderCodeLibrary.h
--- a/Engine/Source/Runtime/RenderCore/Public/ShaderCodeLibrary.h
+++ b/Engine/Source/Runtime/RenderCore/Public/ShaderCodeLibrary.h
@@ -145,7 +145,7 @@ struct RENDERCORE_API FShaderCodeLibrary
 	// Open a named library.
 	// For cooking this will place all added shaders & pipelines into the library file with this name.
 	// At runtime this will open the shader library with this name.
-	static void OpenLibrary(FString const& Name, FString const& Directory);
+	static bool OpenLibrary(FString const& Name, FString const& Directory);
     
 	// Close a named library.
 	// For cooking, after this point any AddShaderCode/AddShaderPipeline calls will be invalid until OpenLibrary is called again.
