Initial work to support creating additive 'patch' shader libraries from the metadata of two different cooked builds.

diff --git a/Engine/Source/Editor/UnrealEd/Classes/Commandlets/ShaderCodeLibraryToolsCommandlet.h b/Engine/Source/Editor/UnrealEd/Classes/Commandlets/ShaderCodeLibraryToolsCommandlet.h
new file mode 100644
--- /dev/null
+++ b/Engine/Source/Editor/UnrealEd/Classes/Commandlets/ShaderCodeLibraryToolsCommandlet.h
@@ -0,0 +1,22 @@
+// Copyright 1998-2019 Epic Games, Inc. All Rights Reserved.
+
+#pragma once
+
+#include "CoreMinimal.h"
+#include "UObject/ObjectMacros.h"
+#include "Commandlets/Commandlet.h"
+#include "ShaderCodeLibraryToolsCommandlet.generated.h"
+
+UCLASS()
+class UShaderCodeLibraryToolsCommandlet
+	: public UCommandlet
+{
+	GENERATED_UCLASS_BODY()
+
+	//~ Begin UCommandlet Interface
+	virtual int32 Main(const FString& Params) override;
+	//~ End UCommandlet Interface
+
+	static int32 StaticMain(const FString& Params);
+};
+
diff --git a/Engine/Source/Editor/UnrealEd/Private/Commandlets/ShaderCodeLibraryToolsCommandlet.cpp b/Engine/Source/Editor/UnrealEd/Private/Commandlets/ShaderCodeLibraryToolsCommandlet.cpp
new file mode 100644
--- /dev/null
+++ b/Engine/Source/Editor/UnrealEd/Private/Commandlets/ShaderCodeLibraryToolsCommandlet.cpp
@@ -0,0 +1,51 @@
+// Copyright 1998-2019 Epic Games, Inc. All Rights Reserved.
+
+#include "Commandlets/ShaderCodeLibraryToolsCommandlet.h"
+#include "Misc/Paths.h"
+
+#include "PipelineFileCache.h"
+#include "ShaderCodeLibrary.h"
+#include "Misc/FileHelper.h"
+#include "Misc/ConfigCacheIni.h"
+
+DEFINE_LOG_CATEGORY_STATIC(LogShaderCodeLibraryTools, Log, All);
+
+UShaderCodeLibraryToolsCommandlet::UShaderCodeLibraryToolsCommandlet(const FObjectInitializer& ObjectInitializer)
+	: Super(ObjectInitializer)
+{
+}
+
+int32 UShaderCodeLibraryToolsCommandlet::Main(const FString& Params)
+{
+	return StaticMain(Params);
+}
+
+int32 UShaderCodeLibraryToolsCommandlet::StaticMain(const FString& Params)
+{
+	TArray<FString> Tokens;
+	TArray<FString> Switches;
+	TMap<FString, FString> ParamVals;
+	UCommandlet::ParseCommandLine(*Params, Tokens, Switches, ParamVals);
+
+	if (Tokens.Num() >= 3)
+	{
+		FString const& Left = Tokens[0];
+		FString const& Right = Tokens[1];
+		FString const& Output = Tokens[2];
+		bool bNativeFormat = Switches.Contains(TEXT("PreferNativeArchives")) || Switches.Contains(TEXT("-PreferNativeArchives"));
+		
+		bool bArchive = false;
+		if (GConfig->GetBool(TEXT("/Script/UnrealEd.ProjectPackagingSettings"), TEXT("bSharedMaterialNativeLibraries"), bArchive, GGameIni));
+		{
+			bNativeFormat |= bArchive;
+		}
+		
+		TArray<FString> OldMetaDataDirs;
+		OldMetaDataDirs.Add(Left);
+		
+		return FShaderCodeLibrary::CreatePatchLibrary(OldMetaDataDirs, Right, Output, bNativeFormat) ? 0 : 1;
+	}
+	
+	UE_LOG(LogShaderCodeLibraryTools, Warning, TEXT("Usage: <Path-To-Old-MetaData> <Path-To-New-MetaData> <Output-Path> [-PreferNativeArchives]\n"));
+	return 0;
+}
diff --git a/Engine/Source/Runtime/RenderCore/Private/ShaderCodeLibrary.cpp b/Engine/Source/Runtime/RenderCore/Private/ShaderCodeLibrary.cpp
--- a/Engine/Source/Runtime/RenderCore/Private/ShaderCodeLibrary.cpp
+++ b/Engine/Source/Runtime/RenderCore/Private/ShaderCodeLibrary.cpp
@@ -941,6 +941,71 @@ struct FEditorShaderCodeArchive
 		return false;
 	}
 
+	bool LoadExistingShaderCodeLibrary(FString const& MetaDataDir)
+	{
+		FString IntermediateFormatPath = GetCodeArchiveFilename(MetaDataDir / TEXT("ShaderLibrarySource"), LibraryName, FormatName);
+		FArchive* PrevCookedAr = IFileManager::Get().CreateFileReader(*IntermediateFormatPath);
+		bool bOK = true;
+		if (PrevCookedAr)
+		{
+			uint32 ArchiveVersion = 0;
+			*PrevCookedAr << ArchiveVersion;
+			if (ArchiveVersion == GShaderCodeArchiveVersion)
+			{
+				// Read shader library
+				*PrevCookedAr << Shaders;
+				for (auto& Entry : Shaders)
+				{
+					Entry.Value.LoadedCode.SetNumUninitialized(Entry.Value.Size);
+					PrevCookedAr->Serialize(Entry.Value.LoadedCode.GetData(), Entry.Value.Size);
+					bOK = !PrevCookedAr->GetError();
+					if (!bOK)
+					{
+						UE_LOG(LogShaderLibrary, Error, TEXT("Failed to deserialize shader code for %s from %s"), *Entry.Key.ToString(), *IntermediateFormatPath);
+						break;
+					}
+				}
+			}
+			else
+			{
+				bOK = false;
+				UE_LOG(LogShaderLibrary, Warning, TEXT("Failed to deserialize shader code from %s because the archive format %u is incompatible with the current version %u"), *IntermediateFormatPath, ArchiveVersion, GShaderCodeArchiveVersion);
+			}
+			
+			PrevCookedAr->Close();
+			delete PrevCookedAr;
+		}
+		else
+		{
+			bOK = false;
+			UE_LOG(LogShaderLibrary, Error, TEXT("Failed to open shader code library from %s"), *IntermediateFormatPath);
+		}
+		
+		if (bOK)
+		{
+			FString PipelinesPath = GetPipelinesArchiveFilename(MetaDataDir / TEXT("ShaderLibrarySource"), LibraryName, FormatName);
+			FArchive* PipelinesArchive = IFileManager::Get().CreateFileReader(*PipelinesPath);
+			if (PipelinesArchive)
+			{
+				uint32 ArchiveVersion = 0;
+				*PipelinesArchive << ArchiveVersion;
+				if (ArchiveVersion == GShaderPipelineArchiveVersion)
+				{
+					*PipelinesArchive << Pipelines;
+				}
+				else
+				{
+					bOK = false;
+					UE_LOG(LogShaderLibrary, Warning, TEXT("Failed to deserialize shader pipelines from %s because the archive format %u is incompatible with the current version %u"), *PipelinesPath, ArchiveVersion, GShaderPipelineArchiveVersion);
+				}
+				
+				PipelinesArchive->Close();
+				delete PipelinesArchive;
+			}
+		}
+		return bOK;
+	}
+
 	void AddExistingShaderCodeLibrary(FString const& OutputDir)
 	{
 		check(LibraryName.Len() > 0);
@@ -1028,7 +1093,7 @@ struct FEditorShaderCodeArchive
 		}
 	}
 
-	bool Finalize(FString OutputDir, bool bNativeFormat, bool bMasterCooker)
+	bool Finalize(FString OutputDir, const FString& MetaOutputDir, bool bNativeFormat, bool bMasterCooker)
 	{
 		check(LibraryName.Len() > 0);
 
@@ -1054,61 +1119,12 @@ struct FEditorShaderCodeArchive
 
 				*FileWriter << GShaderCodeArchiveVersion;
 
-				if (Format->CanStripShaderCode(bNativeFormat) && !bNativeFormat)
-				{
-					FString DebugPlatformDir = GetShaderDebugFolder(FPaths::ProjectSavedDir() / TEXT("Shaders") / FormatName.ToString(), LibraryName, FormatName);
-					IFileManager::Get().MakeDirectory(*DebugPlatformDir, true);
-
-					TMap<FSHAHash, FShaderCodeEntry> StrippedShaders;
-					uint64 TotalSize = 0;
-					for (const auto& Pair : Shaders)
-					{
-						TArray<uint8> CompressedCode;
-						CompressedCode.Append(Pair.Value.LoadedCode.GetData(), Pair.Value.Size);
-
-						int32 UncompressedSize = Pair.Value.UncompressedSize;
-
-						TArray<uint8> UCode;
-						TArray<uint8>& UncompressedCode = FShaderLibraryHelperUncompressCode(Platform, UncompressedSize, CompressedCode, UCode);
-
-						if (!Format->StripShaderCode(UncompressedCode, DebugPlatformDir, bNativeFormat))
-						{
-							bSuccess = false;
-						}
-
-						FShaderLibraryHelperCompressCode(Platform, UncompressedCode, CompressedCode);
-
-						FShaderCodeEntry StrippedEntry;
-						StrippedEntry.Size = CompressedCode.Num();
-						StrippedEntry.Offset = TotalSize;
-						StrippedEntry.UncompressedSize = UncompressedCode.Num();
-						StrippedEntry.Frequency = Pair.Value.Frequency;
-						StrippedEntry.LoadedCode = CompressedCode;
-
-						TotalSize += StrippedEntry.Size;
-
-						StrippedShaders.Add(Pair.Key, StrippedEntry);
-					}
-
-					// Write stripped shader library
-					*FileWriter << StrippedShaders;
-					for (auto& Pair : StrippedShaders)
-					{
-						FileWriter->Serialize(Pair.Value.LoadedCode.GetData(), Pair.Value.Size);
-					}
-
-					// Delete the temp dir
-					IFileManager::Get().DeleteDirectory(*DebugPlatformDir, false, true);
-				}
-				else
-				{
-					// Write shader library
-					*FileWriter << Shaders;
-					for (auto& Pair : Shaders)
-					{
-						FileWriter->Serialize(Pair.Value.LoadedCode.GetData(), Pair.Value.Size);
-					}
-				}
+                // Write shader library
+                *FileWriter << Shaders;
+                for (auto& Pair : Shaders)
+                {
+                    FileWriter->Serialize(Pair.Value.LoadedCode.GetData(), Pair.Value.Size);
+                }
 
 				FileWriter->Close();
 				delete FileWriter;
@@ -1124,7 +1140,18 @@ struct FEditorShaderCodeArchive
 					{
 						UE_LOG(LogShaderLibrary, Error, TEXT("FEditorShaderCodeArchive shader library copy failed to %s. Failed to finalize Shared Shader Library %s with format %s"), *OutputFilePath, *LibraryName, *FormatName.ToString());
 						bSuccess = false;
-					}
+                    }
+                    
+                    if (MetaOutputDir.Len())
+                    {
+                        FString MetaFormatPath = GetCodeArchiveFilename(MetaOutputDir / TEXT("../ShaderLibrarySource"), LibraryName, FormatName);
+                        Result = IFileManager::Get().Copy(*MetaFormatPath, *IntermediateFormatPath, true, true);
+                        if (Result != COPY_OK)
+                        {
+                            UE_LOG(LogShaderLibrary, Error, TEXT("FEditorShaderCodeArchive shader library copy failed to %s. Failed to saved metadata copy of Shared Shader Library %s with format %s"), *OutputFilePath, *LibraryName, *FormatName.ToString());
+                            bSuccess = false;
+                        }
+                    }
 				}
 			}
 		}
@@ -1155,6 +1182,17 @@ struct FEditorShaderCodeArchive
 					UE_LOG(LogShaderLibrary, Error, TEXT("FEditorShaderCodeArchive pipeline copy failed to %s. Failed to finalize Shared Shader Library %s with format %s"), *OutputFilePath, *LibraryName, *FormatName.ToString());
 					bSuccess = false;
 				}
+                
+                if (MetaOutputDir.Len())
+                {
+                    FString MetaFormatPath = GetPipelinesArchiveFilename(MetaOutputDir / TEXT("../ShaderLibrarySource"), LibraryName, FormatName);
+                    Result = IFileManager::Get().Copy(*MetaFormatPath, *TempFilePath, true, true);
+                    if (Result != COPY_OK)
+                    {
+                        UE_LOG(LogShaderLibrary, Error, TEXT("FEditorShaderCodeArchive pipeline copy failed to %s. Failed to save metadata copy of Shared Shader Library %s with format %s"), *OutputFilePath, *LibraryName, *FormatName.ToString());
+                        bSuccess = false;
+                    }
+                }
 			}
 		}
 
@@ -1230,6 +1268,85 @@ struct FEditorShaderCodeArchive
 		return bOK;
 	}
 
+	void MakePatchLibrary(TArray<FEditorShaderCodeArchive*> const& OldLibraries, FEditorShaderCodeArchive const& NewLibrary)
+	{
+		for(auto const& Pair : NewLibrary.Shaders)
+		{
+			if (!HasShader(Pair.Key))
+			{
+				bool bInPreviousPatch = false;
+				for (FEditorShaderCodeArchive const* OldLibrary : OldLibraries)
+				{
+					bInPreviousPatch |= OldLibrary->HasShader(Pair.Key);
+					if (bInPreviousPatch)
+					{
+						break;
+					}
+				}
+				if (!bInPreviousPatch)
+				{
+					FShaderCodeEntry Entry = Pair.Value;
+					Entry.Offset = Offset;
+					Offset += Entry.Size;
+					Shaders.Add(Pair.Key, Entry);
+				}
+			}
+		}
+		
+		Pipelines = NewLibrary.Pipelines;
+	}
+	
+	static bool CreatePatchLibrary(FName FormatName, FString const& LibraryName, TArray<FString> const& OldMetaDataDirs, FString const& NewMetaDataDir, FString const& OutDir, bool bNativeFormat)
+	{
+		TArray<FEditorShaderCodeArchive*> OldLibraries;
+		for (FString const& OldMetaDataDir : OldMetaDataDirs)
+		{
+			FEditorShaderCodeArchive* OldLibrary = new FEditorShaderCodeArchive(FormatName);
+			OldLibrary->OpenLibrary(LibraryName);
+			if (OldLibrary->LoadExistingShaderCodeLibrary(OldMetaDataDir))
+			{
+				OldLibraries.Add(OldLibrary);
+			}
+		}
+
+		FEditorShaderCodeArchive NewLibrary(FormatName);
+		NewLibrary.OpenLibrary(LibraryName);
+		bool bOK = NewLibrary.LoadExistingShaderCodeLibrary(NewMetaDataDir);
+		if (bOK)
+		{
+			FEditorShaderCodeArchive OutLibrary(FormatName);
+			OutLibrary.OpenLibrary(LibraryName);
+			OutLibrary.MakePatchLibrary(OldLibraries, NewLibrary);
+			bOK = OutLibrary.Offset > 0;
+			if (bOK)
+			{
+				FString Empty;
+				bOK = OutLibrary.Finalize(OutDir, Empty, bNativeFormat, true);
+				UE_CLOG(!bOK, LogShaderLibrary, Error, TEXT("Failed to save %s shader patch library %s, %s, %s"), bNativeFormat ? TEXT("native") : TEXT(""), *FormatName.ToString(), *LibraryName, *OutDir);
+				
+				if (bOK && bNativeFormat && OutLibrary.GetFormat()->SupportsShaderArchives())
+				{
+					bOK = OutLibrary.PackageNativeShaderLibrary(OutDir);
+					UE_CLOG(!bOK, LogShaderLibrary, Error, TEXT("Failed to package native shader patch library %s, %s, %s"), *FormatName.ToString(), *LibraryName, *OutDir);
+				}
+			}
+			else
+			{
+				UE_LOG(LogShaderLibrary, Verbose, TEXT("No shaders to patch for library %s, %s, %s"), *FormatName.ToString(), *LibraryName, *OutDir);
+			}
+		}
+		else
+		{
+			UE_LOG(LogShaderLibrary, Error, TEXT("Failed to open the shader library to patch against %s, %s, %s"), *FormatName.ToString(), *LibraryName, *NewMetaDataDir);
+		}
+		
+		for (FEditorShaderCodeArchive* Lib : OldLibraries)
+		{
+			delete Lib;
+		}
+		return bOK;
+	}
+
 private:
 	FName FormatName;
 	FString LibraryName;
@@ -1919,7 +2036,7 @@ public:
 				FEditorShaderCodeArchive* CodeArchive = EditorShaderCodeArchive[SPlatform];
 				if (CodeArchive)
 				{
-					bOk &= CodeArchive->Finalize(ShaderCodeDir, bNativeFormat, bMaster);
+					bOk &= CodeArchive->Finalize(ShaderCodeDir, MetaOutputDir, bNativeFormat, bMaster);
 				}
 			}
 			{
@@ -2379,6 +2496,38 @@ void FShaderCodeLibrary::DumpShaderCodeStats()
 		FShaderCodeLibraryImpl::Impl->DumpShaderCodeStats();
 	}
 }
+
+bool FShaderCodeLibrary::CreatePatchLibrary(TArray<FString> const& OldMetaDataDirs, FString const& NewMetaDataDir, FString const& OutDir, bool bNativeFormat)
+{
+	TMap<FName, TSet<FString>> FormatLibraryMap;
+	TArray<FString> LibraryFiles;
+	IFileManager::Get().FindFiles(LibraryFiles, *(NewMetaDataDir / TEXT("ShaderLibrarySource")), *ShaderExtension);
+	
+	for (FString const& Path : LibraryFiles)
+	{
+		FString Name = FPaths::GetBaseFilename(Path);
+		if (Name.RemoveFromStart(TEXT("ShaderArchive-")))
+		{
+			TArray<FString> Components;
+			if (Name.ParseIntoArray(Components, TEXT("-")) == 2)
+			{
+				FName Format(*Components[1]);
+				TSet<FString>& Libraries = FormatLibraryMap.FindOrAdd(Format);
+				Libraries.Add(Components[0]);
+			}
+		}
+	}
+	
+	bool bOK = true;
+	for (auto const& Entry : FormatLibraryMap)
+	{
+		for (auto const& Library : Entry.Value)
+		{
+			bOK |= FEditorShaderCodeArchive::CreatePatchLibrary(Entry.Key, Library, OldMetaDataDirs, NewMetaDataDir, OutDir, bNativeFormat);
+		}
+	}
+	return bOK;
+}
 #endif// WITH_EDITOR
 
 void FShaderCodeLibrary::SafeAssignHash(FRHIShader* InShader, const FSHAHash& Hash)
diff --git a/Engine/Source/Runtime/RenderCore/Public/ShaderCodeLibrary.h b/Engine/Source/Runtime/RenderCore/Public/ShaderCodeLibrary.h
--- a/Engine/Source/Runtime/RenderCore/Public/ShaderCodeLibrary.h
+++ b/Engine/Source/Runtime/RenderCore/Public/ShaderCodeLibrary.h
@@ -238,6 +238,9 @@ struct RENDERCORE_API FShaderCodeLibrary
 	
 	// Dump collected stats for each shader platform
 	static void DumpShaderCodeStats();
+	
+	// Create a smaller 'patch' library that only contains data from 'NewMetaDataDir' not contained in any of 'OldMetaDataDirs'
+	static bool CreatePatchLibrary(TArray<FString> const& OldMetaDataDirs, FString const& NewMetaDataDir, FString const& OutDir, bool bNativeFormat);
 #endif
 	
 	// Safely assign the hash to a shader object
