Shader Library: use proper shader format names (UE-72416).

- Vulkan (and GL) can have variant names like _NOUB.

diff --git a/Engine/Source/Runtime/RenderCore/Private/ShaderCodeLibrary.cpp b/Engine/Source/Runtime/RenderCore/Private/ShaderCodeLibrary.cpp
--- a/Engine/Source/Runtime/RenderCore/Private/ShaderCodeLibrary.cpp
+++ b/Engine/Source/Runtime/RenderCore/Private/ShaderCodeLibrary.cpp
@@ -1934,10 +1934,11 @@ public:
 		for (FName const& Format : ShaderFormats)
 		{
 			EShaderPlatform Platform = ShaderFormatToLegacyShaderPlatform(Format);
+			FName PossiblyAdjustedFormat = LegacyShaderPlatformToShaderFormat(Platform);	// Vulkan and GL switch between name variants depending on CVars (e.g. see r.Vulkan.UseRealUBs)
 			FEditorShaderCodeArchive* CodeArchive = EditorShaderCodeArchive[Platform];
 			if (!CodeArchive)
 			{
-				CodeArchive = new FEditorShaderCodeArchive(Format);
+				CodeArchive = new FEditorShaderCodeArchive(PossiblyAdjustedFormat);
 				EditorShaderCodeArchive[Platform] = CodeArchive;
 				EditorArchivePipelines[Platform] = !bNativeFormat;
 			}
@@ -1946,10 +1947,11 @@ public:
 		for (FName const& Format : ShaderFormats)
 		{
 			EShaderPlatform Platform = ShaderFormatToLegacyShaderPlatform(Format);
+			FName PossiblyAdjustedFormat = LegacyShaderPlatformToShaderFormat(Platform);	// Vulkan and GL switch between name variants depending on CVars (e.g. see r.Vulkan.UseRealUBs)
 			FEditorShaderStableInfo* StableArchive = EditorShaderStableInfo[Platform];
 			if (!StableArchive)
 			{
-				StableArchive = new FEditorShaderStableInfo(Format);
+				StableArchive = new FEditorShaderStableInfo(PossiblyAdjustedFormat);
 				EditorShaderStableInfo[Platform] = StableArchive;
 			}
 			check(StableArchive);
