Cherry-pick 5814346: Fixes 2 resource transitions bugs on FRDGBuffers

1: Two passes writing to same UAV was not getting the RWBarrier transition.
2: Incorrectly transitioned to readable when accessing for indirect draw call or SRV.

diff --git a/Engine/Source/Runtime/RenderCore/Private/RenderGraphBuilder.cpp b/Engine/Source/Runtime/RenderCore/Private/RenderGraphBuilder.cpp
--- a/Engine/Source/Runtime/RenderCore/Private/RenderGraphBuilder.cpp
+++ b/Engine/Source/Runtime/RenderCore/Private/RenderGraphBuilder.cpp
@@ -708,9 +708,16 @@ void FRDGBuilder::TransitionTexture( const FRDGTexture* Texture, EResourceTransi
 
 void FRDGBuilder::TransitionUAV(FUnorderedAccessViewRHIParamRef UAV, const FRDGResource* UnderlyingResource, EResourceTransitionAccess TransitionAccess, bool bRequiredCompute ) const
 {
-	const bool bRequiredWritable = true;
+	const bool bRequiredWritable = TransitionAccess != EResourceTransitionAccess::EReadable;
 
-	if(UnderlyingResource->bWritable != bRequiredWritable || UnderlyingResource->bCompute != bRequiredCompute )
+	if (bRequiredWritable && UnderlyingResource->bWritable)
+	{
+		// Force a RW barrier between UAV write.
+		// TODO(RDG): allow to have no barriere in the API when multiple pass write concurrently to same resource.
+		EResourceTransitionPipeline TransitionPipeline = CalcTransitionPipeline(UnderlyingResource->bCompute, bRequiredCompute);
+		RHICmdList.TransitionResource(EResourceTransitionAccess::ERWBarrier, TransitionPipeline, UAV);
+	}
+	else if(UnderlyingResource->bWritable != bRequiredWritable || UnderlyingResource->bCompute != bRequiredCompute )
 	{
 		EResourceTransitionPipeline TransitionPipeline = CalcTransitionPipeline(UnderlyingResource->bCompute, bRequiredCompute );
 		RHICmdList.TransitionResource( TransitionAccess, TransitionPipeline, UAV);
