potential fix for PSO problems

diff --git a/Engine/Source/Runtime/RenderCore/Private/ShaderCodeLibrary.cpp b/Engine/Source/Runtime/RenderCore/Private/ShaderCodeLibrary.cpp
--- a/Engine/Source/Runtime/RenderCore/Private/ShaderCodeLibrary.cpp
+++ b/Engine/Source/Runtime/RenderCore/Private/ShaderCodeLibrary.cpp
@@ -1456,6 +1456,7 @@ struct FEditorShaderStableInfo
 				for (const FStableShaderKeyAndValue& Item : StableMap)
 				{
 					Item.ToString(LineBuffer);
+					LineBuffer += TCHAR('\n');
 					auto LineConverted = StringCast<ANSICHAR>(*LineBuffer, LineBuffer.Len());
 					IntermediateFormatAr->Serialize((ANSICHAR*)LineConverted.Get(), LineConverted.Length() * sizeof(ANSICHAR));
 				}
