Add cvar to set IO priority for shader code library reads: r.ShaderCodeLibrary.DefaultAsyncIOPriority. Sometimes they are thrashing during uobject reads and causing seeks.
Setting is unchanged for now

diff --git a/Engine/Source/Runtime/RenderCore/Private/ShaderCodeLibrary.cpp b/Engine/Source/Runtime/RenderCore/Private/ShaderCodeLibrary.cpp
--- a/Engine/Source/Runtime/RenderCore/Private/ShaderCodeLibrary.cpp
+++ b/Engine/Source/Runtime/RenderCore/Private/ShaderCodeLibrary.cpp
@@ -45,6 +45,15 @@ static FString ShaderExtension = TEXT(".ushaderbytecode");
 static FString StableExtension = TEXT(".scl.csv");
 static FString PipelineExtension = TEXT(".ushaderpipelines");
 
+int32 GShaderCodeLibraryAsyncLoadingPriority = int32(AIOP_Normal);
+static FAutoConsoleVariableRef CVarShaderCodeLibraryAsyncLoadingPriority(
+	TEXT("r.ShaderCodeLibrary.DefaultAsyncIOPriority"),
+	GShaderCodeLibraryAsyncLoadingPriority,
+	TEXT(""),
+	ECVF_Default
+);
+
+
 static FString GetCodeArchiveFilename(const FString& BaseDir, const FString& LibraryName, FName Platform)
 {
 	return BaseDir / FString::Printf(TEXT("ShaderArchive-%s-"), *LibraryName) + Platform.ToString() + ShaderExtension;
@@ -518,7 +527,8 @@ public:
 				int64 ReadSize = Entry->Size;
 				int64 ReadOffset = LibraryCodeOffset + Entry->Offset;
 				Entry->LoadedCode.SetNumUninitialized(ReadSize);
-				LocalReadRequest = MakeShareable(LibraryAsyncFileHandle->ReadRequest(ReadOffset, ReadSize, bHiPriSync ? AIOP_CriticalPath : AIOP_Normal, nullptr, Entry->LoadedCode.GetData()));
+				EAsyncIOPriorityAndFlags IOPriority = bHiPriSync ? AIOP_CriticalPath : (EAsyncIOPriorityAndFlags)GShaderCodeLibraryAsyncLoadingPriority;
+				LocalReadRequest = MakeShareable(LibraryAsyncFileHandle->ReadRequest(ReadOffset, ReadSize, IOPriority, nullptr, Entry->LoadedCode.GetData()));
 
 				Entry->ReadRequest = LocalReadRequest;
 				bHasReadRequest = true;
