Modified error handling code in ShaderCodeLibrary to display a message box telling players the game files are missing instead of crashing with fatal UE_LOG message.

diff --git a/Engine/Source/Runtime/RenderCore/Private/ShaderCodeLibrary.cpp b/Engine/Source/Runtime/RenderCore/Private/ShaderCodeLibrary.cpp
--- a/Engine/Source/Runtime/RenderCore/Private/ShaderCodeLibrary.cpp
+++ b/Engine/Source/Runtime/RenderCore/Private/ShaderCodeLibrary.cpp
@@ -11,6 +11,7 @@ ShaderCodeLibrary.cpp: Bound shader state cache implementation.
 #include "Math/UnitConversion.h"
 #include "HAL/FileManagerGeneric.h"
 #include "HAL/LowLevelMemTracker.h"
+#include "HAL/PlatformSplash.h"
 #include "Misc/ScopeLock.h"
 #include "Misc/ScopeRWLock.h"
 #include "Async/AsyncFileHandle.h"
@@ -2167,9 +2168,22 @@ void FShaderCodeLibrary::InitForRuntime(EShaderPlatform ShaderPlatform)
 		else
 		{
 #if !WITH_EDITOR
-            UE_LOG(LogShaderLibrary, Fatal, TEXT("Failed to initialise ShaderCodeLibrary required by the project because part of the Global shader library is missing from %s."), *FPaths::ProjectContentDir());
+			if (FPlatformProperties::SupportsWindowedMode())
+			{
+				FPlatformSplash::Hide();
+
+				UE_LOG(LogShaderLibrary, Error, TEXT("Failed to initialize ShaderCodeLibrary required by the project because part of the Global shader library is missing from %s."), *FPaths::ProjectContentDir());
+
+				FText LocalizedMsg = FText::Format(NSLOCTEXT("MessageDialog", "MissingGlobalShaderLibraryFiles_Body", "Game files required to initialize the global shader library are missing from:\n\n{0}\n\nPlease make sure the game is installed correctly."), FText::FromString(FPaths::ConvertRelativePathToFull(FPaths::ProjectContentDir())));
+				FPlatformMisc::MessageBoxExt(EAppMsgType::Ok, *LocalizedMsg.ToString(), *NSLOCTEXT("MessageDialog", "MissingGlobalShaderLibraryFiles_Title", "Missing game files").ToString());
+			}
+			else
+			{
+				UE_LOG(LogShaderLibrary, Fatal, TEXT("Failed to initialize ShaderCodeLibrary required by the project because part of the Global shader library is missing from %s."), *FPaths::ProjectContentDir());
+			}
 #endif
 			Shutdown();
+			FPlatformMisc::RequestExit(true);
 		}
 	}
 }
