Richard's change to skip PSO cache entries with invalid render target or samples count

diff --git a/Engine/Source/Runtime/RenderCore/Private/ShaderPipelineCache.cpp b/Engine/Source/Runtime/RenderCore/Private/ShaderPipelineCache.cpp
--- a/Engine/Source/Runtime/RenderCore/Private/ShaderPipelineCache.cpp
+++ b/Engine/Source/Runtime/RenderCore/Private/ShaderPipelineCache.cpp
@@ -594,6 +594,11 @@ bool FShaderPipelineCache::Precompile(FRHICommandListImmediate& RHICmdList, ESha
 		
 		GraphicsInitializer.RenderTargetsEnabled = PSO.GraphicsDesc.RenderTargetsActive;
 		GraphicsInitializer.NumSamples = PSO.GraphicsDesc.MSAASamples;
+		
+		if(GraphicsInitializer.RenderTargetsEnabled > MaxSimultaneousRenderTargets || GraphicsInitializer.NumSamples > 16)
+		{
+			return false;
+		}
 
 		GraphicsInitializer.SubpassHint = (ESubpassHint)PSO.GraphicsDesc.SubpassHint;
 		GraphicsInitializer.SubpassIndex = PSO.GraphicsDesc.SubpassIndex;
