Workaround for corrupt PSO data causing crash at startup on Mac

diff --git a/Engine/Source/Runtime/RenderCore/Private/ShaderPipelineCache.cpp b/Engine/Source/Runtime/RenderCore/Private/ShaderPipelineCache.cpp
--- a/Engine/Source/Runtime/RenderCore/Private/ShaderPipelineCache.cpp
+++ b/Engine/Source/Runtime/RenderCore/Private/ShaderPipelineCache.cpp
@@ -595,7 +595,7 @@ bool FShaderPipelineCache::Precompile(FRHICommandListImmediate& RHICmdList, ESha
 		GraphicsInitializer.RenderTargetsEnabled = PSO.GraphicsDesc.RenderTargetsActive;
 		GraphicsInitializer.NumSamples = PSO.GraphicsDesc.MSAASamples;
 		
-		if(GraphicsInitializer.RenderTargetsEnabled > MaxSimultaneousRenderTargets || GraphicsInitializer.NumSamples > 16)
+		if(GraphicsInitializer.RenderTargetsEnabled > MaxSimultaneousRenderTargets || GraphicsInitializer.NumSamples > 16 || GraphicsInitializer.PrimitiveType >= PT_Num)
 		{
 			return false;
 		}
