Fix precompiles of Android shaders

#ue4
#android

diff --git a/Engine/Source/Runtime/RenderCore/Private/ShaderPipelineCache.cpp b/Engine/Source/Runtime/RenderCore/Private/ShaderPipelineCache.cpp
--- a/Engine/Source/Runtime/RenderCore/Private/ShaderPipelineCache.cpp
+++ b/Engine/Source/Runtime/RenderCore/Private/ShaderPipelineCache.cpp
@@ -595,10 +595,16 @@ bool FShaderPipelineCache::Precompile(FRHICommandListImmediate& RHICmdList, ESha
 		GraphicsInitializer.RenderTargetsEnabled = PSO.GraphicsDesc.RenderTargetsActive;
 		GraphicsInitializer.NumSamples = PSO.GraphicsDesc.MSAASamples;
 		
-		if(GraphicsInitializer.RenderTargetsEnabled > MaxSimultaneousRenderTargets || GraphicsInitializer.NumSamples > 16 || GraphicsInitializer.PrimitiveType >= PT_Num)
+		if(GraphicsInitializer.RenderTargetsEnabled > MaxSimultaneousRenderTargets || GraphicsInitializer.NumSamples > 16)
 		{
 			return false;
 		}
+#if PLATFORM_MAC
+		if(GraphicsInitializer.PrimitiveType >= PT_Num)
+		{
+			return false;
+		}
+#endif
 
 		GraphicsInitializer.SubpassHint = (ESubpassHint)PSO.GraphicsDesc.SubpassHint;
 		GraphicsInitializer.SubpassIndex = PSO.GraphicsDesc.SubpassIndex;
