Update PSO validity check to compare against valid data, was comparing against a default out of bounds value.

diff --git a/Engine/Source/Runtime/RenderCore/Private/ShaderPipelineCache.cpp b/Engine/Source/Runtime/RenderCore/Private/ShaderPipelineCache.cpp
--- a/Engine/Source/Runtime/RenderCore/Private/ShaderPipelineCache.cpp
+++ b/Engine/Source/Runtime/RenderCore/Private/ShaderPipelineCache.cpp
@@ -599,8 +599,9 @@ bool FShaderPipelineCache::Precompile(FRHICommandListImmediate& RHICmdList, ESha
 		{
 			return false;
 		}
+
 #if PLATFORM_MAC
-		if(GraphicsInitializer.PrimitiveType >= PT_Num)
+		if(PSO.GraphicsDesc.PrimitiveType >= PT_Num)
 		{
 			return false;
 		}
