RDG Resource Barrier Fixes and Improvements.

 - Resource barriers are now batched and submitted as a group at the beginning of a pass.
 - UAVs referencing the same underlying resource now only submit a single write barrier instead of a write + N UAV barriers.
 - During a mip generation pass, barriers are not issued for textures (they are implicitly converted to the SubRWBarrier state).
 - Depth / stencil read now transitions to the read state; this allows DSV + SRV to work correctly.

diff --git a/Engine/Source/Runtime/RenderCore/Private/RenderGraphBarrierBatcher.cpp b/Engine/Source/Runtime/RenderCore/Private/RenderGraphBarrierBatcher.cpp
new file mode 100644
--- /dev/null
+++ b/Engine/Source/Runtime/RenderCore/Private/RenderGraphBarrierBatcher.cpp
@@ -0,0 +1,259 @@
+// Copyright 1998-2019 Epic Games, Inc. All Rights Reserved.
+
+#include "RenderGraphBarrierBatcher.h"
+#include "RenderGraphPass.h"
+
+namespace
+{
+
+/** Number of entries to reserve in the batch array. */
+const uint32 kBatchReservationSize = 8;
+
+/** Returns whether the resource is transitioning to a writable state. */
+inline bool IsWriteAccessBegin(FRDGResourceState::EAccess AccessBefore, FRDGResourceState::EAccess AccessAfter)
+{
+	return AccessBefore == FRDGResourceState::EAccess::Read && AccessAfter == FRDGResourceState::EAccess::Write;
+}
+
+/** Returns whether the resource is transitioning from a writable state. */
+inline bool IsWriteAccessEnd(FRDGResourceState::EAccess AccessBefore, FRDGResourceState::EAccess AccessAfter)
+{
+	return AccessBefore == FRDGResourceState::EAccess::Write && AccessAfter == FRDGResourceState::EAccess::Read;
+}
+
+inline EResourceTransitionPipeline GetResourceTransitionPipeline(FRDGResourceState::EPipeline PipelineBefore, FRDGResourceState::EPipeline PipelineAfter)
+{
+	switch (PipelineBefore)
+	{
+	case FRDGResourceState::EPipeline::Graphics:
+		switch (PipelineAfter)
+		{
+		case FRDGResourceState::EPipeline::Graphics:
+			return EResourceTransitionPipeline::EGfxToGfx;
+		case FRDGResourceState::EPipeline::Compute:
+			return EResourceTransitionPipeline::EGfxToCompute;
+		}
+		break;
+	case FRDGResourceState::EPipeline::Compute:
+		switch (PipelineAfter)
+		{
+		case FRDGResourceState::EPipeline::Graphics:
+			return EResourceTransitionPipeline::EComputeToGfx;
+		case FRDGResourceState::EPipeline::Compute:
+			return EResourceTransitionPipeline::EComputeToCompute;
+		}
+		break;
+	}
+	check(false);
+	return EResourceTransitionPipeline::EGfxToGfx;
+}
+
+inline EResourceTransitionAccess GetResourceTransitionAccess(FRDGResourceState::EAccess AccessAfter)
+{
+	return AccessAfter == FRDGResourceState::EAccess::Write ? EResourceTransitionAccess::EWritable : EResourceTransitionAccess::EReadable;
+}
+
+inline EResourceTransitionAccess GetResourceTransitionAccessForUAV(FRDGResourceState::EAccess AccessBefore, FRDGResourceState::EAccess AccessAfter)
+{
+	switch (AccessAfter)
+	{
+	case FRDGResourceState::EAccess::Read:
+		return EResourceTransitionAccess::EReadable;
+	case FRDGResourceState::EAccess::Write:
+		// If doing a Write -> Write transition, we use a UAV barrier.
+		if (AccessBefore == FRDGResourceState::EAccess::Write)
+		{
+			return EResourceTransitionAccess::ERWBarrier;
+		}
+		else
+		{
+			return EResourceTransitionAccess::EWritable;
+		}
+	}
+	check(false);
+	return EResourceTransitionAccess::EMaxAccess;
+}
+
+} //! namespace
+
+FRDGBarrierBatcher::~FRDGBarrierBatcher()
+{
+	check(!bAllowQueueing);
+}
+
+void FRDGBarrierBatcher::Begin()
+{
+#if RDG_ENABLE_DEBUG
+	check(!bAllowQueueing);
+
+	check(!TextureUpdateMultiFrameBegins.Num());
+	check(!TextureUpdateMultiFrameEnds.Num());
+	check(!TextureBatchMap.Num());
+
+	check(!UAVUpdateMultiFrameBegins.Num());
+	check(!UAVUpdateMultiFrameEnds.Num());
+	check(!UAVBatchMap.Num());
+#endif
+
+	bAllowQueueing = true;
+}
+
+void FRDGBarrierBatcher::QueueTransitionTexture(FRDGTexture* Texture, FRDGResourceState StateAfter)
+{
+	check(Texture);
+
+	const FRDGResourceState StateBefore = Texture->State;
+
+	ValidateTransition(StateBefore, StateAfter, Texture);
+
+	if (StateBefore != StateAfter)
+	{
+		FRHITexture* RHITexture = Texture->GetRHIUnchecked();
+
+		const bool bIsMultiFrameResource = (Texture->Flags & ERDGResourceFlags::MultiFrame) == ERDGResourceFlags::MultiFrame;
+
+		if (bIsMultiFrameResource && IsWriteAccessBegin(StateBefore.Access, StateAfter.Access))
+		{
+			TextureUpdateMultiFrameBegins.AddUnique(RHITexture);
+		}
+
+		// Add transition to the respective batch bucket.
+		{
+			FTransitionParameters TransitionParameters;
+			TransitionParameters.TransitionAccess = GetResourceTransitionAccess(StateAfter.Access);
+			TransitionParameters.TransitionPipeline = EResourceTransitionPipeline::EGfxToGfx; // NOTE: Transition API for textures doesn't currently expose pipeline transitions.
+
+			FTextureBatch& TextureBatch = TextureBatchMap.FindOrAdd(TransitionParameters);
+			TextureBatch.Reserve(kBatchReservationSize);
+
+			#if RDG_ENABLE_DEBUG
+			{
+				// We should have filtered out duplicates in the first branch of this function.
+				check(TextureBatch.Find(RHITexture) == INDEX_NONE);
+			}
+			#endif
+
+			TextureBatch.Add(RHITexture);
+		}
+
+		if (bIsMultiFrameResource && IsWriteAccessEnd(StateBefore.Access, StateAfter.Access))
+		{
+			TextureUpdateMultiFrameEnds.AddUnique(RHITexture);
+		}
+
+		Texture->State = StateAfter;
+	}
+}
+
+void FRDGBarrierBatcher::QueueTransitionUAV(
+	FRHIUnorderedAccessView* UAV,
+	FRDGTrackedResource* UnderlyingResource,
+	FRDGResourceState StateAfter)
+{
+	check(UAV);
+	check(UnderlyingResource);
+
+	const FRDGResourceState StateBefore = UnderlyingResource->State;
+
+	ValidateTransition(StateBefore, StateAfter, UnderlyingResource);
+
+	if (StateBefore != StateAfter)
+	{
+		const bool bIsMultiFrameResource = (UnderlyingResource->Flags & ERDGResourceFlags::MultiFrame) == ERDGResourceFlags::MultiFrame;
+
+		if (bIsMultiFrameResource && IsWriteAccessBegin(StateBefore.Access, StateAfter.Access))
+		{
+			UAVUpdateMultiFrameBegins.AddUnique(UAV);
+		}
+
+		// Add transition to the correct batch bucket.
+		{
+			FTransitionParameters TransitionParameters;
+			TransitionParameters.TransitionAccess = GetResourceTransitionAccessForUAV(StateBefore.Access, StateAfter.Access);
+			TransitionParameters.TransitionPipeline = GetResourceTransitionPipeline(StateBefore.Pipeline, StateAfter.Pipeline);
+
+			FUAVBatch& UAVBatch = UAVBatchMap.FindOrAdd(TransitionParameters);
+			UAVBatch.Reserve(kBatchReservationSize);
+
+			#if RDG_ENABLE_DEBUG
+			{
+				// We should have filtered out duplicates in the first branch of this function.
+				check(UAVBatch.Find(UAV) == INDEX_NONE);
+			}
+			#endif
+
+			UAVBatch.Add(UAV);
+		}
+
+		if (bIsMultiFrameResource && IsWriteAccessEnd(StateBefore.Access, StateAfter.Access))
+		{
+			UAVUpdateMultiFrameEnds.AddUnique(UAV);
+		}
+
+		UnderlyingResource->State = StateAfter;
+	}
+}
+
+void FRDGBarrierBatcher::End(FRHICommandList& RHICmdList)
+{
+	check(bAllowQueueing);
+	bAllowQueueing = false;
+
+	for (FRHITexture* RHITexture : TextureUpdateMultiFrameBegins)
+	{
+		RHICmdList.BeginUpdateMultiFrameResource(RHITexture);
+	}
+	TextureUpdateMultiFrameBegins.Empty();
+
+	for (FRHIUnorderedAccessView* RHIUAV : UAVUpdateMultiFrameBegins)
+	{
+		RHICmdList.BeginUpdateMultiFrameResource(RHIUAV);
+	}
+	UAVUpdateMultiFrameBegins.Empty();
+
+	for (auto& Element : TextureBatchMap)
+	{
+		FTransitionParameters TransitionParameters = Element.Key;
+		FTextureBatch& Batch = Element.Value;
+		RHICmdList.TransitionResources(TransitionParameters.TransitionAccess, Batch.GetData(), Batch.Num());
+	}
+	TextureBatchMap.Empty();
+
+	for (auto& Element : UAVBatchMap)
+	{
+		FTransitionParameters TransitionParameters = Element.Key;
+		FUAVBatch& Batch = Element.Value;
+		RHICmdList.TransitionResources(TransitionParameters.TransitionAccess, TransitionParameters.TransitionPipeline, Batch.GetData(), Batch.Num());
+	}
+	UAVBatchMap.Empty();
+
+	for (FRHITexture* RHITexture : TextureUpdateMultiFrameEnds)
+	{
+		RHICmdList.EndUpdateMultiFrameResource(RHITexture);
+	}
+	TextureUpdateMultiFrameEnds.Empty();
+
+	for (FRHIUnorderedAccessView* RHIUAV : UAVUpdateMultiFrameEnds)
+	{
+		RHICmdList.EndUpdateMultiFrameResource(RHIUAV);
+	}
+	UAVUpdateMultiFrameEnds.Empty();
+}
+
+void FRDGBarrierBatcher::ValidateTransition(FRDGResourceState StateBefore, FRDGResourceState StateAfter, const FRDGTrackedResource* Resource)
+{
+#if RDG_ENABLE_DEBUG
+	check(StateAfter.Pipeline != FRDGResourceState::EPipeline::MAX);
+	check(StateAfter.Access != FRDGResourceState::EAccess::Unknown);
+
+	if (StateBefore != StateAfter && StateAfter.Pass)
+	{
+		// We allow duplicate transitions of the same resource within the same pass, but not conflicting ones.
+		checkf(
+			StateBefore.Pass != StateAfter.Pass,
+			TEXT("Pass %s attempted to transition resource %s to different states. Make sure the resource isn't being used\n")
+			TEXT("for both read and write at the same time. This can occur if the resource is used as both an SRV and UAV, or\n")
+			TEXT("SRV and Render Target, for example.\n"), StateAfter.Pass->GetName(), Resource->Name);
+	}
+#endif
+}
\ No newline at end of file
diff --git a/Engine/Source/Runtime/RenderCore/Private/RenderGraphBarrierBatcher.h b/Engine/Source/Runtime/RenderCore/Private/RenderGraphBarrierBatcher.h
new file mode 100644
--- /dev/null
+++ b/Engine/Source/Runtime/RenderCore/Private/RenderGraphBarrierBatcher.h
@@ -0,0 +1,71 @@
+// Copyright 1998-2019 Epic Games, Inc. All Rights Reserved.
+
+#pragma once
+
+#include "RenderGraphResources.h"
+
+/** Class that queues up batches of resource barriers and then submits them to a command list. */
+class RENDERCORE_API FRDGBarrierBatcher final
+{
+public:
+	FRDGBarrierBatcher() = default;
+	~FRDGBarrierBatcher();
+
+	void Begin();
+
+	void QueueTransitionTexture(FRDGTexture* Texture, FRDGResourceState StateAfter);
+
+	void QueueTransitionUAV(
+		FRHIUnorderedAccessView* UAV,
+		FRDGTrackedResource* UnderlyingResource,
+		FRDGResourceState StateAfter);
+
+	/** Submits and clears all queued transitions to the provided RHI command list. */
+	void End(FRHICommandList& RHICmdList);
+
+private:
+	static void ValidateTransition(FRDGResourceState StateBefore, FRDGResourceState StateAfter, const FRDGTrackedResource* Resource);
+
+	struct FTransitionParameters
+	{
+		FTransitionParameters() = default;
+
+		bool operator==(const FTransitionParameters& Other) const
+		{
+			return TransitionAccess == Other.TransitionAccess && TransitionPipeline == Other.TransitionPipeline;
+		}
+
+		uint32 GetHash() const
+		{
+			return uint32(TransitionAccess) | (uint32(TransitionPipeline) << 8);
+		}
+
+		EResourceTransitionAccess TransitionAccess;
+		EResourceTransitionPipeline TransitionPipeline = EResourceTransitionPipeline::EGfxToGfx;
+	};
+
+	template <typename TBatchType>
+	struct TBatchMapKeyFuncs : public TDefaultMapKeyFuncs<FTransitionParameters, TBatchType, /** bAllowDuplicateKeys */ false>
+	{
+		static uint32 GetKeyHash(FTransitionParameters Key)
+		{
+			return Key.GetHash();
+		}
+	};
+
+	using FUAVBatch = TArray<FRHIUnorderedAccessView*, SceneRenderingAllocator>;
+	using FUAVBatchMap = TMap<FTransitionParameters, FUAVBatch, SceneRenderingSetAllocator, TBatchMapKeyFuncs<FUAVBatch>>;
+
+	using FTextureBatch = TArray<FRHITexture*, SceneRenderingAllocator>;
+	using FTextureBatchMap = TMap<FTransitionParameters, FTextureBatch, SceneRenderingSetAllocator, TBatchMapKeyFuncs<FTextureBatch>>;
+
+	FTextureBatch TextureUpdateMultiFrameBegins;
+	FTextureBatch TextureUpdateMultiFrameEnds;
+	FTextureBatchMap TextureBatchMap;
+
+	FUAVBatch UAVUpdateMultiFrameBegins;
+	FUAVBatch UAVUpdateMultiFrameEnds;
+	FUAVBatchMap UAVBatchMap;
+
+	bool bAllowQueueing = false;
+};
\ No newline at end of file
diff --git a/Engine/Source/Runtime/RenderCore/Private/RenderGraphBuilder.cpp b/Engine/Source/Runtime/RenderCore/Private/RenderGraphBuilder.cpp
--- a/Engine/Source/Runtime/RenderCore/Private/RenderGraphBuilder.cpp
+++ b/Engine/Source/Runtime/RenderCore/Private/RenderGraphBuilder.cpp
@@ -4,6 +4,7 @@
 #include "RenderCore.h"
 #include "RenderTargetPool.h"
 #include "RenderGraphResourcePool.h"
+#include "RenderGraphBarrierBatcher.h"
 #include "VisualizeTexture.h"
 #include "ProfilingDebugging/CsvProfiler.h"
 
@@ -621,9 +622,10 @@ void FRDGBuilder::AllocateRHITextureIfNeeded(FRDGTexture* Texture)
 
 	check(Texture->ReferenceCount > 0 || GRDGImmediateMode);
 
-	// TODO(RDG): should avoid bDoWritableBarrier = true
 	TRefCountPtr<IPooledRenderTarget>& PooledRenderTarget = AllocatedTextures.FindOrAdd(Texture);
-	GRenderTargetPool.FindFreeElement(RHICmdList, Texture->Desc, PooledRenderTarget, Texture->Name, /* bDoWritableBarrier = */ true);
+
+	const bool bDoWriteBarrier = false;
+	GRenderTargetPool.FindFreeElement(RHICmdList, Texture->Desc, PooledRenderTarget, Texture->Name, bDoWriteBarrier);
 
 	Texture->PooledRenderTarget = PooledRenderTarget;
 	Texture->ResourceRHI = PooledRenderTarget->GetRenderTargetItem().ShaderResourceTexture;
@@ -644,6 +646,23 @@ void FRDGBuilder::AllocateRHITextureUAVIfNeeded(FRDGTextureUAV* UAV)
 	UAV->ResourceRHI = UAV->Desc.Texture->PooledRenderTarget->GetRenderTargetItem().MipUAVs[UAV->Desc.MipLevel];
 }
 
+void FRDGBuilder::AllocateRHIBufferIfNeeded(FRDGBuffer* Buffer)
+{
+	check(Buffer);
+
+	if (Buffer->PooledBuffer)
+	{
+		return;
+	}
+
+	check(Buffer->ReferenceCount > 0 || GRDGImmediateMode);
+
+	TRefCountPtr<FPooledRDGBuffer>& AllocatedBuffer = AllocatedBuffers.FindOrAdd(Buffer);
+	GRenderGraphResourcePool.FindFreeBuffer(RHICmdList, Buffer->Desc, AllocatedBuffer, Buffer->Name);
+	check(AllocatedBuffer);
+	Buffer->PooledBuffer = AllocatedBuffer;
+}
+
 void FRDGBuilder::AllocateRHIBufferSRVIfNeeded(FRDGBufferSRV* SRV)
 {
 	check(SRV);
@@ -699,17 +718,7 @@ void FRDGBuilder::AllocateRHIBufferUAVIfNeeded(FRDGBufferUAV* UAV)
 	}
 	
 	FRDGBufferRef Buffer = UAV->Desc.Buffer;
-
-	// Allocate a buffer resource.
-	if (!Buffer->PooledBuffer)
-	{
-		check(Buffer->ReferenceCount > 0 || GRDGImmediateMode);
-
-		TRefCountPtr<FPooledRDGBuffer>& AllocatedBuffer = AllocatedBuffers.FindOrAdd(Buffer);
-		GRenderGraphResourcePool.FindFreeBuffer(RHICmdList, Buffer->Desc, AllocatedBuffer, Buffer->Name);
-
-		Buffer->PooledBuffer = AllocatedBuffer;
-	}
+	AllocateRHIBufferIfNeeded(Buffer);
 
 	if (Buffer->PooledBuffer->UAVs.Contains(UAV->Desc))
 	{
@@ -739,115 +748,6 @@ void FRDGBuilder::AllocateRHIBufferUAVIfNeeded(FRDGBufferUAV* UAV)
 	Buffer->PooledBuffer->UAVs.Add(UAV->Desc, RHIUnorderedAccessView);
 }
 
-static EResourceTransitionPipeline GetResourceTransitionPipeline(ERDGPassPipeline PassPipelineBefore, ERDGPassPipeline PassPipelineAfter)
-{
-	static const EResourceTransitionPipeline Table[(uint32)ERDGPassPipeline::MAX][(uint32)ERDGPassPipeline::MAX] =
-	{
-		// Graphics
-		{
-			EResourceTransitionPipeline::EGfxToGfx,
-			EResourceTransitionPipeline::EGfxToCompute
-		},
-
-		// Compute
-		{
-			EResourceTransitionPipeline::EComputeToGfx,
-			EResourceTransitionPipeline::EComputeToCompute
-		}
-	};
-
-#if RDG_ENABLE_DEBUG
-	check(PassPipelineBefore < ERDGPassPipeline::MAX);
-	check(PassPipelineAfter < ERDGPassPipeline::MAX);
-#endif
-
-	return Table[(uint32)PassPipelineBefore][(uint32)PassPipelineAfter];
-};
-
-static EResourceTransitionAccess GetResourceTransitionAccess(ERDGPassAccess PassAccess)
-{
-	return PassAccess == ERDGPassAccess::Read ? EResourceTransitionAccess::EReadable : EResourceTransitionAccess::EWritable;
-}
-
-void FRDGBuilder::TransitionTexture(FRDGTexture* Texture, ERDGPassAccess PassAccessAfter, ERDGPassPipeline PassPipelineAfter) const
-{
-	const ERDGPassAccess PassAccessBefore = Texture->PassAccess;
-
-	if (PassAccessBefore != PassAccessAfter || Texture->PassPipeline != PassPipelineAfter)
-	{
-		FRHITexture* RHITexture = Texture->GetRHIUnchecked();
-
-		const bool bIsMultiFrameResource = (Texture->Flags & ERDGResourceFlags::MultiFrame) == ERDGResourceFlags::MultiFrame;
-
-		const bool bIsWriteAccessBegin = PassAccessBefore == ERDGPassAccess::Read && PassAccessAfter == ERDGPassAccess::Write;
-		const bool bIsWriteAccessEnd = PassAccessBefore == ERDGPassAccess::Write && PassAccessAfter == ERDGPassAccess::Read;
-
-		if (bIsMultiFrameResource && bIsWriteAccessBegin)
-		{
-			RHICmdList.BeginUpdateMultiFrameResource(RHITexture);
-		}
-
-		const EResourceTransitionAccess TransitionAccess = GetResourceTransitionAccess(PassAccessAfter);
-
-		RHICmdList.TransitionResource(TransitionAccess, RHITexture);
-		
-		if (bIsMultiFrameResource && bIsWriteAccessEnd)
-		{
-			RHICmdList.EndUpdateMultiFrameResource(RHITexture);
-		}
-		
-		Texture->PassAccess = PassAccessAfter;
-		Texture->PassPipeline = PassPipelineAfter;
-	}
-}
-
-void FRDGBuilder::TransitionUAV(
-	FRHIUnorderedAccessView* UAV,
-	FRDGTrackedResource* UnderlyingResource,
-	ERDGPassAccess PassAccessAfter,
-	ERDGPassPipeline PassPipelineAfter) const
-{
-	check(UAV);
-	check(UnderlyingResource);
-
-	const ERDGPassAccess PassAccessBefore = UnderlyingResource->PassAccess;
-	const ERDGPassPipeline PassPipelineBefore = UnderlyingResource->PassPipeline;
-
-	if (PassAccessBefore == ERDGPassAccess::Write && PassAccessAfter == ERDGPassAccess::Write)
-	{
-		// Force a RW barrier between UAV write.
-		// TODO(RDG): allow to have no barrier in the API when multiple pass write concurrently to same resource.
-		const EResourceTransitionPipeline TransitionPipeline = GetResourceTransitionPipeline(PassPipelineBefore, PassPipelineAfter);
-
-		RHICmdList.TransitionResource(EResourceTransitionAccess::ERWBarrier, TransitionPipeline, UAV);
-	}
-	else if(PassAccessBefore != PassAccessAfter|| PassPipelineBefore != PassPipelineAfter)
-	{
-		const bool bIsMultiFrameResource = (UnderlyingResource->Flags & ERDGResourceFlags::MultiFrame) == ERDGResourceFlags::MultiFrame;
-
-		const bool bIsWriteAccessBegin = PassAccessBefore == ERDGPassAccess::Read && PassAccessAfter == ERDGPassAccess::Write;
-		const bool bIsWriteAccessEnd = PassAccessBefore == ERDGPassAccess::Write && PassAccessAfter == ERDGPassAccess::Read;
-
-		if (bIsMultiFrameResource && bIsWriteAccessBegin)
-		{
-			RHICmdList.BeginUpdateMultiFrameResource(UAV);
-		}
-
-		const EResourceTransitionPipeline TransitionPipeline = GetResourceTransitionPipeline(PassPipelineBefore, PassPipelineAfter);
-		const EResourceTransitionAccess TransitionAccess = GetResourceTransitionAccess(PassAccessAfter);
-
-		RHICmdList.TransitionResource(TransitionAccess, TransitionPipeline, UAV);
-		
-		if (bIsMultiFrameResource && bIsWriteAccessEnd)
-		{
-			RHICmdList.EndUpdateMultiFrameResource(UAV);
-		}
-		
-		UnderlyingResource->PassAccess = PassAccessAfter;
-		UnderlyingResource->PassPipeline = PassPipelineAfter;
-	}
-}
-
 void FRDGBuilder::ExecutePass(const FRDGPass* Pass)
 {
 	QUICK_SCOPE_CYCLE_COUNTER(STAT_FRDGBuilder_ExecutePass);
@@ -900,12 +800,18 @@ void FRDGBuilder::ExecutePass(const FRDGPass* Pass)
 
 void FRDGBuilder::PrepareResourcesForExecute(const FRDGPass* Pass, struct FRHIRenderPassInfo* OutRPInfo, bool* bOutHasRenderTargets)
 {
+	check(Pass);
+
 	OutRPInfo->NumUAVs = 0;
 	OutRPInfo->UAVIndex = 0;
 
-	const ERDGPassPipeline PassPipeline = Pass->IsCompute() ? ERDGPassPipeline::Compute : ERDGPassPipeline::Graphics;
+	const bool bIsCompute = Pass->IsCompute();
+
+	FRDGBarrierBatcher BarrierBatcher;
+	BarrierBatcher.Begin();
 
-	const bool bGeneratingMips = (Pass->GetFlags() & ERDGPassFlags::GenerateMips) == ERDGPassFlags::GenerateMips;
+	// NOTE: When generating mips, we don't perform any transitions on textures. They are done implicitly by the RHI.
+	const bool bGeneratingMips = Pass->IsGenerateMips();
 
 	FRDGPassParameterStruct ParameterStruct = Pass->GetParameters();
 
@@ -931,7 +837,11 @@ void FRDGBuilder::PrepareResourcesForExecute(const FRDGPass* Pass, struct FRHIRe
 
 				check(Texture->PooledRenderTarget);
 				check(Texture->ResourceRHI);
-				TransitionTexture(Texture, ERDGPassAccess::Read, PassPipeline);
+
+				if (!bGeneratingMips)
+				{
+					BarrierBatcher.QueueTransitionTexture(Texture, FRDGResourceState::CreateRead(Pass));
+				}
 			}
 		}
 		break;
@@ -955,7 +865,10 @@ void FRDGBuilder::PrepareResourcesForExecute(const FRDGPass* Pass, struct FRHIRe
 					SRV->ResourceRHI = Texture->PooledRenderTarget->GetRenderTargetItem().MipSRVs[SRV->Desc.MipLevel];
 				}
 
-				TransitionTexture(Texture, ERDGPassAccess::Read, PassPipeline);
+				if (!bGeneratingMips)
+				{
+					BarrierBatcher.QueueTransitionTexture(Texture, FRDGResourceState::CreateRead(Pass));
+				}
 			}
 		}
 		break;
@@ -975,12 +888,12 @@ void FRDGBuilder::PrepareResourcesForExecute(const FRDGPass* Pass, struct FRHIRe
 
 				FRHIUnorderedAccessView* UAVRHI = UAV->GetRHI();
 
-				if (PassPipeline == ERDGPassPipeline::Graphics)
+				if (!bIsCompute)
 				{
 					OutRPInfo->UAVs[OutRPInfo->NumUAVs++] = UAVRHI;	// Bind UAVs in declaration order
 				}
 
-				TransitionUAV(UAVRHI, Texture, ERDGPassAccess::Write, PassPipeline);
+				BarrierBatcher.QueueTransitionUAV(UAVRHI, Texture, FRDGResourceState::CreateWrite(Pass));
 			}
 		}
 		break;
@@ -999,9 +912,9 @@ void FRDGBuilder::PrepareResourcesForExecute(const FRDGPass* Pass, struct FRHIRe
 				// TODO(RDG): super hacky, find the UAV and transition it. Hopefully there is one...
 				check(Buffer->PooledBuffer);
 				check(Buffer->PooledBuffer->UAVs.Num() == 1);
-				FRHIUnorderedAccessView* UAV = Buffer->PooledBuffer->UAVs.CreateIterator().Value();
+				FRHIUnorderedAccessView* UAVRHI = Buffer->PooledBuffer->UAVs.CreateIterator().Value();
 
-				TransitionUAV(UAV, Buffer, ERDGPassAccess::Write, PassPipeline);
+				BarrierBatcher.QueueTransitionUAV(UAVRHI, Buffer, FRDGResourceState::CreateRead(Pass));
 			}
 		}
 		break;
@@ -1024,9 +937,9 @@ void FRDGBuilder::PrepareResourcesForExecute(const FRDGPass* Pass, struct FRHIRe
 				// TODO(RDG): super hacky, find the UAV and transition it. Hopefully there is one...
 				check(Buffer->PooledBuffer);
 				check(Buffer->PooledBuffer->UAVs.Num() == 1);
-				FRHIUnorderedAccessView* UAV = Buffer->PooledBuffer->UAVs.CreateIterator().Value();
+				FRHIUnorderedAccessView* UAVRHI = Buffer->PooledBuffer->UAVs.CreateIterator().Value();
 
-				TransitionUAV(UAV, Buffer, ERDGPassAccess::Read, PassPipeline);
+				BarrierBatcher.QueueTransitionUAV(UAVRHI, Buffer, FRDGResourceState::CreateRead(Pass));
 			}
 		}
 		break;
@@ -1046,18 +959,18 @@ void FRDGBuilder::PrepareResourcesForExecute(const FRDGPass* Pass, struct FRHIRe
 
 				FRHIUnorderedAccessView* UAVRHI = UAV->GetRHI();
 
-				if (PassPipeline == ERDGPassPipeline::Graphics)
+				if (!bIsCompute)
 				{
 					OutRPInfo->UAVs[OutRPInfo->NumUAVs++] = UAVRHI;	// Bind UAVs in declaration order
 				}
-	
-				TransitionUAV(UAVRHI, Buffer, ERDGPassAccess::Write, PassPipeline);
+
+				BarrierBatcher.QueueTransitionUAV(UAVRHI, Buffer, FRDGResourceState::CreateWrite(Pass));
 			}
 		}
 		break;
 		case UBMT_RENDER_TARGET_BINDING_SLOTS:
 		{
-			check(PassPipeline == ERDGPassPipeline::Graphics);
+			check(!Pass->IsCompute());
 
 			const FRenderTargetBindingSlots& RenderTargetBindingSlots = Parameter.GetAsRenderTargetBindingSlots();
 			const auto& RenderTargets = RenderTargetBindingSlots.Output;
@@ -1106,8 +1019,7 @@ void FRDGBuilder::PrepareResourcesForExecute(const FRDGPass* Pass, struct FRHIRe
 
 					if (!bGeneratingMips)
 					{
-						// Implicit assurance the RHI will do the correct transitions
-						TransitionTexture(Texture, ERDGPassAccess::Write, PassPipeline);
+						BarrierBatcher.QueueTransitionTexture(Texture, FRDGResourceState::CreateWrite(Pass));
 					}
 
 					SampleCount |= OutRenderTarget.RenderTarget->GetNumSamples();
@@ -1142,8 +1054,11 @@ void FRDGBuilder::PrepareResourcesForExecute(const FRDGPass* Pass, struct FRHIRe
 					MakeRenderTargetActions(DepthStencil.GetStencilLoadAction(), DepthStencil.GetStencilStoreAction()));
 				OutDepthStencil.ExclusiveDepthStencil = DepthStencil.GetDepthStencilAccess();
 
-				TransitionTexture(Texture, ERDGPassAccess::Write, PassPipeline);
-
+				BarrierBatcher.QueueTransitionTexture(Texture,
+					DepthStencil.GetDepthStencilAccess().IsAnyWrite() ?
+						FRDGResourceState::CreateWrite(Pass) :
+						FRDGResourceState::CreateRead(Pass));
+	
 				SampleCount |= OutDepthStencil.DepthStencilTarget->GetNumSamples();
 				ValidDepthStencilCount++;
 
@@ -1165,6 +1080,8 @@ void FRDGBuilder::PrepareResourcesForExecute(const FRDGPass* Pass, struct FRHIRe
 	}
 
 	OutRPInfo->bGeneratingMips = bGeneratingMips;
+
+	BarrierBatcher.End(RHICmdList);
 }
 
 void FRDGBuilder::UpdateAccessGuardForPassResources(const FRDGPass* Pass, bool bAllowAccess)
@@ -1403,13 +1320,21 @@ void FRDGBuilder::ReleaseUnreferencedResources(const FRDGPass* Pass)
 
 void FRDGBuilder::ProcessDeferredInternalResourceQueries()
 {
+	FRDGBarrierBatcher BarrierBatcher;
+	BarrierBatcher.Begin();
+
 	for (const auto& Query : DeferredInternalTextureQueries)
 	{
 		check(Query.Texture->PooledRenderTarget);
 
 		if (Query.bTransitionToRead)
 		{
-			TransitionTexture(Query.Texture, ERDGPassAccess::Read, ERDGPassPipeline::Graphics);
+			BarrierBatcher.QueueTransitionTexture(
+				Query.Texture,
+				FRDGResourceState(
+					nullptr,
+					FRDGResourceState::EPipeline::Graphics,
+					FRDGResourceState::EAccess::Read));
 		}
 
 		*Query.OutTexturePtr = AllocatedTextures.FindChecked(Query.Texture);
@@ -1447,6 +1372,8 @@ void FRDGBuilder::ProcessDeferredInternalResourceQueries()
 			ReleaseRHIBufferIfUnreferenced(Query.Buffer);
 		}
 	}
+
+	BarrierBatcher.End(RHICmdList);
 }
 
 void FRDGBuilder::DestructPasses()
diff --git a/Engine/Source/Runtime/RenderCore/Private/RenderGraphResources.cpp b/Engine/Source/Runtime/RenderCore/Private/RenderGraphResources.cpp
new file mode 100644
--- /dev/null
+++ b/Engine/Source/Runtime/RenderCore/Private/RenderGraphResources.cpp
@@ -0,0 +1,26 @@
+// Copyright 1998-2019 Epic Games, Inc. All Rights Reserved.
+
+#include "RenderGraphResources.h"
+#include "RenderGraphPass.h"
+
+FRDGResourceState FRDGResourceState::CreateRead(const FRDGPass* Pass)
+{
+	check(Pass);
+
+	FRDGResourceState State;
+	State.Pass = Pass;
+	State.Pipeline = Pass->IsCompute() ? EPipeline::Compute : EPipeline::Graphics;
+	State.Access = FRDGResourceState::EAccess::Read;
+	return State;
+}
+
+FRDGResourceState FRDGResourceState::CreateWrite(const FRDGPass* Pass)
+{
+	check(Pass);
+
+	FRDGResourceState State;
+	State.Pass = Pass;
+	State.Pipeline = Pass->IsCompute() ? EPipeline::Compute : EPipeline::Graphics;
+	State.Access = FRDGResourceState::EAccess::Write;
+	return State;
+}
\ No newline at end of file
diff --git a/Engine/Source/Runtime/RenderCore/Public/RenderGraphBuilder.h b/Engine/Source/Runtime/RenderCore/Public/RenderGraphBuilder.h
--- a/Engine/Source/Runtime/RenderCore/Public/RenderGraphBuilder.h
+++ b/Engine/Source/Runtime/RenderCore/Public/RenderGraphBuilder.h
@@ -194,20 +194,10 @@ private:
 
 	void AllocateRHITextureIfNeeded(FRDGTexture* Texture);
 	void AllocateRHITextureUAVIfNeeded(FRDGTextureUAV* UAV);
+	void AllocateRHIBufferIfNeeded(FRDGBuffer* Buffer);
 	void AllocateRHIBufferSRVIfNeeded(FRDGBufferSRV* SRV);
 	void AllocateRHIBufferUAVIfNeeded(FRDGBufferUAV* UAV);
 
-	void TransitionTexture(
-		FRDGTexture* Texture,
-		ERDGPassAccess PassAccess,
-		ERDGPassPipeline PassPipeline) const;
-
-	void TransitionUAV(
-		FRHIUnorderedAccessView* UAV,
-		FRDGTrackedResource* UnderlyingResource,
-		ERDGPassAccess PassAccess,
-		ERDGPassPipeline PassPipeline) const;
-
 	void ExecutePass(const FRDGPass* Pass);
 	void PrepareResourcesForExecute(const FRDGPass* Pass, struct FRHIRenderPassInfo* OutRPInfo, bool* bOutHasRenderTargets);
 
diff --git a/Engine/Source/Runtime/RenderCore/Public/RenderGraphPass.h b/Engine/Source/Runtime/RenderCore/Public/RenderGraphPass.h
--- a/Engine/Source/Runtime/RenderCore/Public/RenderGraphPass.h
+++ b/Engine/Source/Runtime/RenderCore/Public/RenderGraphPass.h
@@ -198,6 +198,11 @@ public:
 		return PassFlags;
 	}
 
+	bool IsGenerateMips() const
+	{
+		return (PassFlags & ERDGPassFlags::GenerateMips) == ERDGPassFlags::GenerateMips;
+	}
+
 	bool IsCompute() const
 	{
 		return (PassFlags & ERDGPassFlags::Compute) == ERDGPassFlags::Compute;
diff --git a/Engine/Source/Runtime/RenderCore/Public/RenderGraphResources.h b/Engine/Source/Runtime/RenderCore/Public/RenderGraphResources.h
--- a/Engine/Source/Runtime/RenderCore/Public/RenderGraphResources.h
+++ b/Engine/Source/Runtime/RenderCore/Public/RenderGraphResources.h
@@ -16,9 +16,6 @@ class FRDGBuffer;
 class FRDGBufferSRV;
 class FRDGBufferUAV;
 
-struct FPooledRDGBuffer;
-struct FRayTracingShaderBindingsWriter;
-
 /** Defines the RDG resource references for user code not forgetting the const every time. */
 using FRDGResourceRef = FRDGResource*;
 using FRDGTrackedResourceRef = FRDGTrackedResource*;
@@ -29,20 +26,56 @@ using FRDGBufferRef = FRDGBuffer*;
 using FRDGBufferSRVRef = FRDGBufferSRV*;
 using FRDGBufferUAVRef = FRDGBufferUAV*;
 
-/** The hardware pipeline for pass scheduling. */
-enum class ERDGPassPipeline : uint8
+/** Used for tracking resource state during execution. */
+struct FRDGResourceState
 {
-	Graphics,
-	Compute,
-	MAX
-};
+	// The hardware pipeline the resource state is compatible with.
+	enum class EPipeline : uint8
+	{
+		Graphics,
+		Compute,
+		MAX
+	};
 
-/** Describes how a pass accesses a resource. */
-enum class ERDGPassAccess : uint8
-{
-	Read,
-	Write,
-	MAX
+	// The access permissions the resource state is compatible with.
+	enum class EAccess : uint8
+	{
+		Read,
+		Write,
+		Unknown,
+		MAX = Unknown
+	};
+
+	/** Constructs a read / write state suitable for the provided pass. */
+	static FRDGResourceState CreateRead(const FRDGPass* Pass);
+	static FRDGResourceState CreateWrite(const FRDGPass* Pass);
+
+	FRDGResourceState() = default;
+
+	FRDGResourceState(const FRDGPass* InPass, EPipeline InPipeline, EAccess InAccess)
+		: Pass(InPass)
+		, Pipeline(InPipeline)
+		, Access(InAccess)
+	{}
+
+	bool operator==(const FRDGResourceState& Other) const
+	{
+		return Pass == Other.Pass && Pipeline == Other.Pipeline && Access == Other.Access;
+	}
+
+	bool operator!=(const FRDGResourceState& Other) const
+	{
+		return !(*this == Other);
+	}
+
+	// The last pass the resource was used.
+	const FRDGPass* Pass = nullptr;
+
+	// The last used pass hardware pipeline.
+	EPipeline Pipeline = EPipeline::Graphics;
+
+	// The last used access on the pass.
+	EAccess Access = EAccess::Unknown;
 };
 
 /** Render graph specific flags for resources. */
@@ -157,6 +190,7 @@ private:
 	{
 		checkf(!FirstProducer, TEXT("Resource %s with producer pass marked as external."), Name);
 		bHasBeenProduced = true;
+		bIsExternal = true;
 	}
 
 	bool HasBeenProduced() const
@@ -164,6 +198,11 @@ private:
 		return bHasBeenProduced;
 	}
 
+	bool IsExternal() const
+	{
+		return bIsExternal;
+	}
+
 	/** Pointer towards the pass that is the first to produce it, for even more convenient error message. */
 	const FRDGPass* FirstProducer = nullptr;
 
@@ -172,13 +211,15 @@ private:
 
 	/** Boolean to track at wiring time if a resource has ever been produced by a pass, to error out early if accessing a resource that has not been produced. */
 	bool bHasBeenProduced = false;
+
+	/** Tracks whether this resource is an externally imported resource. */
+	bool bIsExternal = false;
 #endif
 
-	/** Used for tracking resource state during execution. */
-	ERDGPassAccess PassAccess = ERDGPassAccess::Read;
-	ERDGPassPipeline PassPipeline = ERDGPassPipeline::Graphics;
+	FRDGResourceState State;
 
 	friend class FRDGBuilder;
+	friend class FRDGBarrierBatcher;
 };
 
 /** Descriptor of a graph tracked texture. */
@@ -232,6 +273,7 @@ private:
 	IPooledRenderTarget* PooledRenderTarget = nullptr;
 
 	friend class FRDGBuilder;
+	friend class FRDGBarrierBatcher;
 };
 
 /** Render graph tracked SRV. */
