Change check() to return success/failure in CheckPSOStringInveribility().  Use this and a new FPipelineCacheFileFormatPSO::Verify() function to decide if we should actually use this PSO in the ShaderPipelineCacheToolsCommandlet expand operation.  In non shipping builds also use the verfy function during FPipelineCacheFileFormatPSO::Init() to catch bad PSOs during PSO record runs.  E.g use the current check() assert in FPipelineFileCache::CacheGraphicsPSO().

diff --git a/Engine/Source/Editor/UnrealEd/Private/Commandlets/ShaderPipelineCacheToolsCommandlet.cpp b/Engine/Source/Editor/UnrealEd/Private/Commandlets/ShaderPipelineCacheToolsCommandlet.cpp
--- a/Engine/Source/Editor/UnrealEd/Private/Commandlets/ShaderPipelineCacheToolsCommandlet.cpp
+++ b/Engine/Source/Editor/UnrealEd/Private/Commandlets/ShaderPipelineCacheToolsCommandlet.cpp
@@ -226,7 +226,7 @@ static void PrintShaders(const TMap<FSHAHash, TArray<FString>>& InverseMap, cons
 	}
 }
 
-void CheckPSOStringInveribility(const FPipelineCacheFileFormatPSO& Item)
+bool CheckPSOStringInveribility(const FPipelineCacheFileFormatPSO& Item)
 {
 	FPipelineCacheFileFormatPSO TempItem(Item);
 	TempItem.Hash = 0;
@@ -254,8 +254,7 @@ void CheckPSOStringInveribility(const FPipelineCacheFileFormatPSO& Item)
 	}
 	UE_LOG(LogShaderPipelineCacheTools, Verbose, TEXT("CheckPSOStringInveribility: %s"), *StringRep);
 
-	check(DupItem == TempItem);
-	check(GetTypeHash(DupItem) == GetTypeHash(TempItem));
+	return (DupItem == TempItem) && (GetTypeHash(DupItem) == GetTypeHash(TempItem));
 }
 
 int32 DumpPSOSC(FString& Token)
@@ -523,6 +522,7 @@ int32 ExpandPSOSC(const TArray<FString>& Tokens)
 				auto* ExistingPSO = PSOs.Find(TempPSO);
 				if(ExistingPSO != nullptr)
 				{
+					// Existing PSO must have already gone through verify and invertibility checks
 					check(*ExistingPSO == TempPSO);
 					
 					// Get More accurate stats by testing for diff - we could just merge and be done
@@ -536,7 +536,17 @@ int32 ExpandPSOSC(const TArray<FString>& Tokens)
 				}
 				else
 				{
-					PSOs.Add(TempPSO);
+					bool bInvertibilityResult = CheckPSOStringInveribility(TempPSO);
+					bool bVerifyResult = TempPSO.Verify();
+					if(bInvertibilityResult && bVerifyResult)
+					{
+						PSOs.Add(TempPSO);
+					}
+					else
+					{
+						// Log Found Bad PSO, this is in the context of the logged current file above so we can see where this has come from
+						UE_LOG(LogShaderPipelineCacheTools, Warning, TEXT("Bad PSO found discarding [Invertibility=%s Verify=%s in: %s]"), bInvertibilityResult ? TEXT("PASS") :  TEXT("FAIL") , bVerifyResult ? TEXT("PASS") :  TEXT("FAIL"), *Tokens[Index]);
+					}
 				}
 			}
 		}
@@ -552,12 +562,6 @@ int32 ExpandPSOSC(const TArray<FString>& Tokens)
 	}
 	UE_LOG(LogShaderPipelineCacheTools, Display, TEXT("Loaded %d PSOs total [Usage Mask Merged = %d]."), PSOs.Num(), MergeCount);
 
-	//self test
-	for (const FPipelineCacheFileFormatPSO& Item : PSOs)
-	{
-		CheckPSOStringInveribility(Item);
-	}
-	// end self test
 	if (UE_LOG_ACTIVE(LogShaderPipelineCacheTools, Verbose))
 	{
 		TMap<FSHAHash, TArray<FString>> InverseMap;
@@ -1151,11 +1155,12 @@ int32 BuildPSOSC(const TArray<FString>& Tokens)
 				{
 					PSO.Type = FPipelineCacheFileFormatPSO::DescriptorType::Graphics;
 					check(!bLooksLikeAComputeShader);
-					if (PSO.GraphicsDesc.VertexShader == FSHAHash())
-					{
-						UE_LOG(LogShaderPipelineCacheTools, Warning, TEXT("Stable PSO with null vertex shader, ignored."));
-						bValid = false;
-					}
+				}
+				
+				bValid = PSO.Verify();
+				if(!bValid)
+				{
+					UE_LOG(LogShaderPipelineCacheTools, Warning, TEXT("Bad PSO found - verify failed - PSO discarded [Line %d in: %s]"), Index, *FileName);
 				}
 			}
 
diff --git a/Engine/Source/Runtime/RHI/Private/PipelineFileCache.cpp b/Engine/Source/Runtime/RHI/Private/PipelineFileCache.cpp
--- a/Engine/Source/Runtime/RHI/Private/PipelineFileCache.cpp
+++ b/Engine/Source/Runtime/RHI/Private/PipelineFileCache.cpp
@@ -603,6 +603,115 @@ void FPipelineCacheFileFormatPSO::CommonFromString(const FString& Src)
 #endif
 }
 
+bool FPipelineCacheFileFormatPSO::Verify() const
+{
+	if(Type == DescriptorType::Compute)
+	{
+		return ComputeDesc.ComputeShader != FSHAHash();
+	}
+	else if(Type == DescriptorType::Graphics)
+	{
+		if(GraphicsDesc.VertexShader == FSHAHash())
+		{
+			// No vertex shader - no graphics - nothing else matters
+			return false;
+		}
+		
+#if PLATFORM_SUPPORTS_TESSELLATION_SHADERS
+		if( GraphicsDesc.HullShader == FSHAHash() && GraphicsDesc.DomainShader == FSHAHash() && GraphicsDesc.PrimitiveType >= PT_1_ControlPointPatchList && GraphicsDesc.PrimitiveType <= PT_32_ControlPointPatchList)
+		{
+			// Not using tessellation - we shouldn't try to draw patches
+			return false;
+		}
+		else if( GraphicsDesc.HullShader != FSHAHash() && GraphicsDesc.DomainShader == FSHAHash() ||
+				 GraphicsDesc.DomainShader == FSHAHash() && GraphicsDesc.HullShader != FSHAHash() )
+		{
+			// Hull without Domain or vice-versa
+			return false;
+		}
+#else
+		if(GraphicsDesc.HullShader != FSHAHash() || GraphicsDesc.DomainShader != FSHAHash())
+		{
+			// Define says we don't support tessellation - why have we got tessellation shaders - not a valid PSO for target platform
+			return false;
+		}
+		
+		if(GraphicsDesc.PrimitiveType >= PT_1_ControlPointPatchList && GraphicsDesc.PrimitiveType <= PT_32_ControlPointPatchList)
+		{
+			// Define says we don't support tessellation - can't draw patches - not a valid PSO for target platform
+			return false;
+		}
+#endif
+
+#if PLATFORM_SUPPORTS_GEOMETRY_SHADERS
+		// Is there anything to actually test here?
+#endif
+		
+		if( GraphicsDesc.RenderTargetsActive > MaxSimultaneousRenderTargets ||
+			GraphicsDesc.MSAASamples > 16 ||
+			(uint32)GraphicsDesc.PrimitiveType >= (uint32)EPrimitiveType::PT_Num ||
+			(uint32)GraphicsDesc.DepthStencilFormat >= (uint32)EPixelFormat::PF_MAX ||
+			(uint8)GraphicsDesc.DepthLoad >= (uint8)ERenderTargetLoadAction::Num ||
+			(uint8)GraphicsDesc.StencilLoad >= (uint8)ERenderTargetLoadAction::Num ||
+			(uint8)GraphicsDesc.DepthStore >= (uint8)ERenderTargetStoreAction::Num ||
+			(uint8)GraphicsDesc.StencilStore >= (uint8)ERenderTargetStoreAction::Num )
+		{
+			return false;
+		}
+		
+		for(uint32 rt = 0;rt < GraphicsDesc.RenderTargetsActive;++rt)
+		{
+			if((uint32)GraphicsDesc.RenderTargetFormats[rt] >= (uint32)EPixelFormat::PF_MAX)
+			{
+				return false;
+			}
+			
+			if( GraphicsDesc.BlendState.RenderTargets[rt].ColorBlendOp >= EBlendOperation::EBlendOperation_Num ||
+				GraphicsDesc.BlendState.RenderTargets[rt].AlphaBlendOp >= EBlendOperation::EBlendOperation_Num ||
+				GraphicsDesc.BlendState.RenderTargets[rt].ColorSrcBlend >= EBlendFactor::EBlendFactor_Num ||
+				GraphicsDesc.BlendState.RenderTargets[rt].ColorDestBlend >= EBlendFactor::EBlendFactor_Num ||
+				GraphicsDesc.BlendState.RenderTargets[rt].AlphaSrcBlend >= EBlendFactor::EBlendFactor_Num ||
+				GraphicsDesc.BlendState.RenderTargets[rt].AlphaDestBlend >= EBlendFactor::EBlendFactor_Num ||
+				GraphicsDesc.BlendState.RenderTargets[rt].ColorWriteMask > 0xf)
+			{
+				return false;
+			}
+		}
+		
+		if( (uint8)GraphicsDesc.RasterizerState.FillMode >= (uint8)ERasterizerFillMode::ERasterizerFillMode_Num ||
+			(uint8)GraphicsDesc.RasterizerState.CullMode >= (uint8)ERasterizerCullMode_Num)
+		{
+			return false;
+		}
+		
+		if( (uint8)GraphicsDesc.DepthStencilState.DepthTest >= (uint8)ECompareFunction::ECompareFunction_Num ||
+			(uint8)GraphicsDesc.DepthStencilState.FrontFaceStencilTest >= (uint8)ECompareFunction::ECompareFunction_Num ||
+			(uint8)GraphicsDesc.DepthStencilState.BackFaceStencilTest >= (uint8)ECompareFunction::ECompareFunction_Num ||
+			(uint8)GraphicsDesc.DepthStencilState.FrontFaceStencilFailStencilOp >= (uint8)EStencilOp::EStencilOp_Num ||
+			(uint8)GraphicsDesc.DepthStencilState.FrontFaceDepthFailStencilOp >= (uint8)EStencilOp::EStencilOp_Num ||
+			(uint8)GraphicsDesc.DepthStencilState.FrontFacePassStencilOp >= (uint8)EStencilOp::EStencilOp_Num ||
+			(uint8)GraphicsDesc.DepthStencilState.BackFaceStencilFailStencilOp >= (uint8)EStencilOp::EStencilOp_Num ||
+			(uint8)GraphicsDesc.DepthStencilState.BackFaceDepthFailStencilOp >= (uint8)EStencilOp::EStencilOp_Num ||
+			(uint8)GraphicsDesc.DepthStencilState.BackFacePassStencilOp >= (uint8)EStencilOp::EStencilOp_Num)
+		{
+			return false;
+		}
+
+		uint32 ElementCount = (uint32)GraphicsDesc.VertexDescriptor.Num();
+		for (uint32 i = 0; i < ElementCount;++i)
+		{
+			if(GraphicsDesc.VertexDescriptor[i].Type >= EVertexElementType::VET_MAX)
+			{
+				return false;
+			}
+		}
+		
+		return true;
+	}
+	
+	return false;
+}
+
 /**
   * FPipelineCacheFileFormatPSO
   **/
@@ -848,7 +957,13 @@ FPipelineCacheFileFormatPSO::FPipelineCacheFileFormatPSO()
 	
 	PSO.ComputeDesc.ComputeShader = Init->GetHash();
 	
-	return true;
+	bool bOK = true;
+	
+#if !UE_BUILD_SHIPPING
+	bOK = PSO.Verify();
+#endif
+	
+	return bOK;
 }
 
 /*static*/ bool FPipelineCacheFileFormatPSO::Init(FPipelineCacheFileFormatPSO& PSO, FGraphicsPipelineStateInitializer const& Init)
@@ -970,6 +1085,10 @@ FPipelineCacheFileFormatPSO::FPipelineCacheFileFormatPSO()
 	PSO.GraphicsDesc.SubpassHint = (uint8)Init.SubpassHint;
 	PSO.GraphicsDesc.SubpassIndex = Init.SubpassIndex;
 	
+#if !UE_BUILD_SHIPPING
+	bOK = bOK && PSO.Verify();
+#endif
+	
 	return bOK;
 }
 
diff --git a/Engine/Source/Runtime/RHI/Public/PipelineFileCache.h b/Engine/Source/Runtime/RHI/Public/PipelineFileCache.h
--- a/Engine/Source/Runtime/RHI/Public/PipelineFileCache.h
+++ b/Engine/Source/Runtime/RHI/Public/PipelineFileCache.h
@@ -188,6 +188,9 @@ struct RHI_API FPipelineCacheFileFormatPSO
 	FString CommonToString() const;
 	static FString CommonHeaderLine();
 	void CommonFromString(const FString& Src);
+	
+	// Potential cases for seperating verify logic if requiired: RunTime-Logging, RunTime-UserCaching, RunTime-PreCompile, CommandLet-Cooking
+	bool Verify() const;
 };
 
 struct RHI_API FPipelineCacheFileFormatPSORead
diff --git a/Engine/Source/Runtime/RenderCore/Private/ShaderPipelineCache.cpp b/Engine/Source/Runtime/RenderCore/Private/ShaderPipelineCache.cpp
--- a/Engine/Source/Runtime/RenderCore/Private/ShaderPipelineCache.cpp
+++ b/Engine/Source/Runtime/RenderCore/Private/ShaderPipelineCache.cpp
@@ -528,121 +528,119 @@ bool FShaderPipelineCache::Precompile(FRHICommandListImmediate& RHICmdList, ESha
 	INC_DWORD_STAT(STAT_PreCompileShadersNum);
     
     uint64 StartTime = FPlatformTime::Cycles64();
-		
+
 	bool bOk = false;
 	
-	TArray<uint8> DummyCode;
-	
-	if(FPipelineCacheFileFormatPSO::DescriptorType::Graphics == PSO.Type)
+	if(PSO.Verify())
 	{
-		FGraphicsPipelineStateInitializer GraphicsInitializer;
-		
-		FRHIVertexDeclaration* VertexDesc = PipelineStateCache::GetOrCreateVertexDeclaration(PSO.GraphicsDesc.VertexDescriptor);
-		GraphicsInitializer.BoundShaderState.VertexDeclarationRHI = VertexDesc;
+		TArray<uint8> DummyCode;
 		
-		FVertexShaderRHIRef VertexShader;
-		if (PSO.GraphicsDesc.VertexShader != FSHAHash())
+		if(FPipelineCacheFileFormatPSO::DescriptorType::Graphics == PSO.Type)
 		{
-			VertexShader = FShaderCodeLibrary::CreateVertexShader(Platform, PSO.GraphicsDesc.VertexShader, DummyCode);
-			GraphicsInitializer.BoundShaderState.VertexShaderRHI = VertexShader;
-		}
+			FGraphicsPipelineStateInitializer GraphicsInitializer;
+			
+			FRHIVertexDeclaration* VertexDesc = PipelineStateCache::GetOrCreateVertexDeclaration(PSO.GraphicsDesc.VertexDescriptor);
+			GraphicsInitializer.BoundShaderState.VertexDeclarationRHI = VertexDesc;
+			
+			FVertexShaderRHIRef VertexShader;
+			if (PSO.GraphicsDesc.VertexShader != FSHAHash())
+			{
+				VertexShader = FShaderCodeLibrary::CreateVertexShader(Platform, PSO.GraphicsDesc.VertexShader, DummyCode);
+				GraphicsInitializer.BoundShaderState.VertexShaderRHI = VertexShader;
+			}
 
-#if PLATFORM_SUPPORTS_TESSELLATION_SHADERS
-		FHullShaderRHIRef HullShader;
-		if (PSO.GraphicsDesc.HullShader != FSHAHash())
-		{
-			HullShader = FShaderCodeLibrary::CreateHullShader(Platform, PSO.GraphicsDesc.HullShader, DummyCode);
-			GraphicsInitializer.BoundShaderState.HullShaderRHI = HullShader;
-		}
+	#if PLATFORM_SUPPORTS_TESSELLATION_SHADERS
+			FHullShaderRHIRef HullShader;
+			if (PSO.GraphicsDesc.HullShader != FSHAHash())
+			{
+				HullShader = FShaderCodeLibrary::CreateHullShader(Platform, PSO.GraphicsDesc.HullShader, DummyCode);
+				GraphicsInitializer.BoundShaderState.HullShaderRHI = HullShader;
+			}
 
-		FDomainShaderRHIRef DomainShader;
-		if (PSO.GraphicsDesc.DomainShader != FSHAHash())
-		{
-			DomainShader = FShaderCodeLibrary::CreateDomainShader(Platform, PSO.GraphicsDesc.DomainShader, DummyCode);
-			GraphicsInitializer.BoundShaderState.DomainShaderRHI = DomainShader;
-		}
-#endif
-		FPixelShaderRHIRef FragmentShader;
-		if (PSO.GraphicsDesc.FragmentShader != FSHAHash())
-		{
-			FragmentShader = FShaderCodeLibrary::CreatePixelShader(Platform, PSO.GraphicsDesc.FragmentShader, DummyCode);
-			GraphicsInitializer.BoundShaderState.PixelShaderRHI = FragmentShader;
-		}
+			FDomainShaderRHIRef DomainShader;
+			if (PSO.GraphicsDesc.DomainShader != FSHAHash())
+			{
+				DomainShader = FShaderCodeLibrary::CreateDomainShader(Platform, PSO.GraphicsDesc.DomainShader, DummyCode);
+				GraphicsInitializer.BoundShaderState.DomainShaderRHI = DomainShader;
+			}
+	#endif
+			FPixelShaderRHIRef FragmentShader;
+			if (PSO.GraphicsDesc.FragmentShader != FSHAHash())
+			{
+				FragmentShader = FShaderCodeLibrary::CreatePixelShader(Platform, PSO.GraphicsDesc.FragmentShader, DummyCode);
+				GraphicsInitializer.BoundShaderState.PixelShaderRHI = FragmentShader;
+			}
 
-#if PLATFORM_SUPPORTS_GEOMETRY_SHADERS
-		FGeometryShaderRHIRef GeometryShader;
-		if (PSO.GraphicsDesc.GeometryShader != FSHAHash())
-		{
-			GeometryShader = FShaderCodeLibrary::CreateGeometryShader(Platform, PSO.GraphicsDesc.GeometryShader, DummyCode);
-			GraphicsInitializer.BoundShaderState.GeometryShaderRHI = GeometryShader;
-		}
-#endif
-		auto BlendState = RHICmdList.CreateBlendState(PSO.GraphicsDesc.BlendState);
-		GraphicsInitializer.BlendState = BlendState;
-		
-		auto RasterState = RHICmdList.CreateRasterizerState(PSO.GraphicsDesc.RasterizerState);
-		GraphicsInitializer.RasterizerState = RasterState;
-		
-		auto DepthState = RHICmdList.CreateDepthStencilState(PSO.GraphicsDesc.DepthStencilState);
-		GraphicsInitializer.DepthStencilState = DepthState;
-		
-		for (uint32 i = 0; i < MaxSimultaneousRenderTargets; ++i)
-		{
-			GraphicsInitializer.RenderTargetFormats[i] = PSO.GraphicsDesc.RenderTargetFormats[i];
-			GraphicsInitializer.RenderTargetFlags[i] = PSO.GraphicsDesc.RenderTargetFlags[i];
-		}
-		
-		GraphicsInitializer.RenderTargetsEnabled = PSO.GraphicsDesc.RenderTargetsActive;
-		GraphicsInitializer.NumSamples = PSO.GraphicsDesc.MSAASamples;
-		
-		if(GraphicsInitializer.RenderTargetsEnabled > MaxSimultaneousRenderTargets || GraphicsInitializer.NumSamples > 16)
-		{
-			return false;
-		}
+	#if PLATFORM_SUPPORTS_GEOMETRY_SHADERS
+			FGeometryShaderRHIRef GeometryShader;
+			if (PSO.GraphicsDesc.GeometryShader != FSHAHash())
+			{
+				GeometryShader = FShaderCodeLibrary::CreateGeometryShader(Platform, PSO.GraphicsDesc.GeometryShader, DummyCode);
+				GraphicsInitializer.BoundShaderState.GeometryShaderRHI = GeometryShader;
+			}
+	#endif
+			auto BlendState = RHICmdList.CreateBlendState(PSO.GraphicsDesc.BlendState);
+			GraphicsInitializer.BlendState = BlendState;
+			
+			auto RasterState = RHICmdList.CreateRasterizerState(PSO.GraphicsDesc.RasterizerState);
+			GraphicsInitializer.RasterizerState = RasterState;
+			
+			auto DepthState = RHICmdList.CreateDepthStencilState(PSO.GraphicsDesc.DepthStencilState);
+			GraphicsInitializer.DepthStencilState = DepthState;
+			
+			for (uint32 i = 0; i < MaxSimultaneousRenderTargets; ++i)
+			{
+				GraphicsInitializer.RenderTargetFormats[i] = PSO.GraphicsDesc.RenderTargetFormats[i];
+				GraphicsInitializer.RenderTargetFlags[i] = PSO.GraphicsDesc.RenderTargetFlags[i];
+			}
+			
+			GraphicsInitializer.RenderTargetsEnabled = PSO.GraphicsDesc.RenderTargetsActive;
+			GraphicsInitializer.NumSamples = PSO.GraphicsDesc.MSAASamples;
 
-#if PLATFORM_MAC
-		if(PSO.GraphicsDesc.PrimitiveType >= PT_Num)
+			GraphicsInitializer.SubpassHint = (ESubpassHint)PSO.GraphicsDesc.SubpassHint;
+			GraphicsInitializer.SubpassIndex = PSO.GraphicsDesc.SubpassIndex;
+			
+			GraphicsInitializer.DepthStencilTargetFormat = PSO.GraphicsDesc.DepthStencilFormat;
+			GraphicsInitializer.DepthStencilTargetFlag = PSO.GraphicsDesc.DepthStencilFlags;
+			GraphicsInitializer.DepthTargetLoadAction = PSO.GraphicsDesc.DepthLoad;
+			GraphicsInitializer.StencilTargetLoadAction = PSO.GraphicsDesc.StencilLoad;
+			GraphicsInitializer.DepthTargetStoreAction = PSO.GraphicsDesc.DepthStore;
+			GraphicsInitializer.StencilTargetStoreAction = PSO.GraphicsDesc.StencilStore;
+			
+			GraphicsInitializer.PrimitiveType = PSO.GraphicsDesc.PrimitiveType;
+			GraphicsInitializer.bFromPSOFileCache = true;
+			
+			// This indicates we do not want a fatal error if this compilation fails
+			// (ie, if this entry in the file cache is bad)
+			GraphicsInitializer.bFromPSOFileCache = 1;
+			
+			// Use SetGraphicsPipelineState to call down into PipelineStateCache and also handle the fallback case used by OpenGL.
+			SetGraphicsPipelineState(RHICmdList, GraphicsInitializer, EApplyRendertargetOption::DoNothing);
+			bOk = true;
+		}
+		else if(FPipelineCacheFileFormatPSO::DescriptorType::Compute == PSO.Type)
 		{
-			return false;
+			FComputeShaderRHIRef ComputeInitializer = FShaderCodeLibrary::CreateComputeShader(GMaxRHIShaderPlatform, PSO.ComputeDesc.ComputeShader, DummyCode);
+			if(ComputeInitializer.IsValid())
+			{
+				FComputePipelineState* ComputeResult = PipelineStateCache::GetAndOrCreateComputePipelineState(RHICmdList, ComputeInitializer);
+				bOk = ComputeResult != nullptr;
+			}
 		}
-#endif
-
-		GraphicsInitializer.SubpassHint = (ESubpassHint)PSO.GraphicsDesc.SubpassHint;
-		GraphicsInitializer.SubpassIndex = PSO.GraphicsDesc.SubpassIndex;
-		
-		GraphicsInitializer.DepthStencilTargetFormat = PSO.GraphicsDesc.DepthStencilFormat;
-		GraphicsInitializer.DepthStencilTargetFlag = PSO.GraphicsDesc.DepthStencilFlags;
-		GraphicsInitializer.DepthTargetLoadAction = PSO.GraphicsDesc.DepthLoad;
-		GraphicsInitializer.StencilTargetLoadAction = PSO.GraphicsDesc.StencilLoad;
-		GraphicsInitializer.DepthTargetStoreAction = PSO.GraphicsDesc.DepthStore;
-		GraphicsInitializer.StencilTargetStoreAction = PSO.GraphicsDesc.StencilStore;
-		
-		GraphicsInitializer.PrimitiveType = PSO.GraphicsDesc.PrimitiveType;
-		GraphicsInitializer.bFromPSOFileCache = true;
-		
-		// This indicates we do not want a fatal error if this compilation fails
-		// (ie, if this entry in the file cache is bad)
-		GraphicsInitializer.bFromPSOFileCache = 1;
-		
-		// Use SetGraphicsPipelineState to call down into PipelineStateCache and also handle the fallback case used by OpenGL.
-		SetGraphicsPipelineState(RHICmdList, GraphicsInitializer, EApplyRendertargetOption::DoNothing);
-		bOk = true;
-	}
-	else if(FPipelineCacheFileFormatPSO::DescriptorType::Compute == PSO.Type)
-	{
-		FComputeShaderRHIRef ComputeInitializer = FShaderCodeLibrary::CreateComputeShader(GMaxRHIShaderPlatform, PSO.ComputeDesc.ComputeShader, DummyCode);
-		if(ComputeInitializer.IsValid())
+		else
 		{
-			FComputePipelineState* ComputeResult = PipelineStateCache::GetAndOrCreateComputePipelineState(RHICmdList, ComputeInitializer);
-			bOk = ComputeResult != nullptr;
+			check(false);
 		}
 	}
+#if !UE_BUILD_SHIPPING
 	else
 	{
-		check(false);
+		UE_LOG(LogRHI, Warning, TEXT("FShaderPipelineCache::Precompile() - PSO Verify failure - Did not attempt to precompile"));
 	}
-    
-    if (bOk)
+#endif
+	
+    // All read dependencies have given the green light - always update task counts
+    // Otherwise we end up with outstanding compiles that we can't progress or external tools may think this has not been completed and may run again.
     {
         uint64 TimeDelta = FPlatformTime::Cycles64() - StartTime;
         FPlatformAtomics::InterlockedIncrement(&TotalCompleteTasks);
