Hand-Integrated fixes from Main:

7064877 - RDG Resource Barrier Fixes and Improvements.

 - Resource barriers are now batched and submitted as a group at the beginning of a pass.
 - UAVs referencing the same underlying resource now only submit a single write barrier instead of a write + N UAV barriers.
 - During a mip generation pass, barriers are not issued for textures (they are implicitly converted to the SubRWBarrier state).
 - Depth / stencil read now transitions to the read state; this allows DSV + SRV to work correctly.

7064884 - Fixed RHI resource transitions causing validation failures.

7064898 - Fixed issue where implementation of RHICopyTexture would only copy mip 0, breaking environment probe copies.

diff --git a/Engine/Source/Runtime/RenderCore/Private/RenderGraphBuilder.cpp b/Engine/Source/Runtime/RenderCore/Private/RenderGraphBuilder.cpp
--- a/Engine/Source/Runtime/RenderCore/Private/RenderGraphBuilder.cpp
+++ b/Engine/Source/Runtime/RenderCore/Private/RenderGraphBuilder.cpp
@@ -632,6 +632,23 @@ void FRDGBuilder::AllocateRHITextureIfNeeded(FRDGTexture* Texture)
 	check(Texture->ResourceRHI);
 }
 
+void FRDGBuilder::AllocateRHIBufferIfNeeded(FRDGBuffer* Buffer)
+{
+	check(Buffer);
+
+	if (Buffer->PooledBuffer)
+	{
+		return;
+	}
+
+	check(Buffer->ReferenceCount > 0 || GRDGImmediateMode);
+
+	TRefCountPtr<FPooledRDGBuffer>& AllocatedBuffer = AllocatedBuffers.FindOrAdd(Buffer);
+	GRenderGraphResourcePool.FindFreeBuffer(RHICmdList, Buffer->Desc, AllocatedBuffer, Buffer->Name);
+	check(AllocatedBuffer);
+	Buffer->PooledBuffer = AllocatedBuffer;
+}
+
 void FRDGBuilder::AllocateRHITextureUAVIfNeeded(FRDGTextureUAV* UAV)
 {
 	check(UAV);
@@ -1410,6 +1427,8 @@ void FRDGBuilder::DestructPasses()
 
 	DeferredInternalTextureQueries.Empty();
 	DeferredInternalBufferQueries.Empty();
+	ExternalTextures.Empty();
+	ExternalBuffers.Empty();
 	AllocatedTextures.Empty();
 	AllocatedBuffers.Empty();
 }
diff --git a/Engine/Source/Runtime/RenderCore/Public/RenderGraphBuilder.h b/Engine/Source/Runtime/RenderCore/Public/RenderGraphBuilder.h
--- a/Engine/Source/Runtime/RenderCore/Public/RenderGraphBuilder.h
+++ b/Engine/Source/Runtime/RenderCore/Public/RenderGraphBuilder.h
@@ -151,6 +151,10 @@ private:
 	/** Keep the references over the pooled render target, since FRDGTexture is allocated on MemStack. */
 	TMap<FRDGBuffer*, TRefCountPtr<FPooledRDGBuffer>, SceneRenderingSetAllocator> AllocatedBuffers;
 
+	/** Tracks external resources to their registered render graph counterparts for de-duplication. */
+	TMap<const IPooledRenderTarget*, FRDGTexture*, SceneRenderingSetAllocator> ExternalTextures;
+	TMap<const FPooledRDGBuffer*, FRDGBuffer*, SceneRenderingSetAllocator> ExternalBuffers;
+
 	/** Array of all deferred access to internal textures. */
 	struct FDeferredInternalTextureQuery
 	{
diff --git a/Engine/Source/Runtime/RenderCore/Public/RenderGraphBuilder.inl b/Engine/Source/Runtime/RenderCore/Public/RenderGraphBuilder.inl
--- a/Engine/Source/Runtime/RenderCore/Public/RenderGraphBuilder.inl
+++ b/Engine/Source/Runtime/RenderCore/Public/RenderGraphBuilder.inl
@@ -11,6 +11,11 @@ inline FRDGTextureRef FRDGBuilder::RegisterExternalTexture(const TRefCountPtr<IP
 	}
 #endif
 
+	if (FRDGTextureRef* Texture = ExternalTextures.Find(ExternalPooledTexture.GetReference()))
+	{
+		return *Texture;
+	}
+
 	FRDGTexture* OutTexture = AllocateForRHILifeTime<FRDGTexture>(Name, ExternalPooledTexture->GetDesc(), ERDGResourceFlags::None);
 	OutTexture->PooledRenderTarget = ExternalPooledTexture;
 	OutTexture->ResourceRHI = ExternalPooledTexture->GetRenderTargetItem().ShaderResourceTexture;
@@ -23,6 +28,8 @@ inline FRDGTextureRef FRDGBuilder::RegisterExternalTexture(const TRefCountPtr<IP
 	}
 #endif
 
+	ExternalTextures.Add(ExternalPooledTexture.GetReference(), OutTexture);
+
 	return OutTexture;
 }
 
@@ -34,6 +41,11 @@ inline FRDGBufferRef FRDGBuilder::RegisterExternalBuffer(const TRefCountPtr<FPoo
 	}
 #endif
 
+	if (FRDGBufferRef* Buffer = ExternalBuffers.Find(ExternalPooledBuffer.GetReference()))
+	{
+		return *Buffer;
+	}
+
 	FRDGBuffer* OutBuffer = AllocateForRHILifeTime<FRDGBuffer>(Name, ExternalPooledBuffer->Desc, ERDGResourceFlags::None);
 	OutBuffer->PooledBuffer = ExternalPooledBuffer;
 	AllocatedBuffers.Add(OutBuffer, ExternalPooledBuffer);
@@ -45,6 +57,8 @@ inline FRDGBufferRef FRDGBuilder::RegisterExternalBuffer(const TRefCountPtr<FPoo
 	}
 #endif
 
+	ExternalBuffers.Add(ExternalPooledBuffer.GetReference(), OutBuffer);
+
 	return OutBuffer;
 }
 
