Added missing files in integration from main

diff --git a/Engine/Source/Runtime/RenderCore/Private/RenderGraphBarrierBatcher.cpp b/Engine/Source/Runtime/RenderCore/Private/RenderGraphBarrierBatcher.cpp
--- a/Engine/Source/Runtime/RenderCore/Private/RenderGraphBarrierBatcher.cpp
+++ b/Engine/Source/Runtime/RenderCore/Private/RenderGraphBarrierBatcher.cpp
@@ -6,73 +6,73 @@
 namespace
 {
 
-/** Number of entries to reserve in the batch array. */
-const uint32 kBatchReservationSize = 8;
+	/** Number of entries to reserve in the batch array. */
+	const uint32 kBatchReservationSize = 8;
 
-/** Returns whether the resource is transitioning to a writable state. */
-inline bool IsWriteAccessBegin(FRDGResourceState::EAccess AccessBefore, FRDGResourceState::EAccess AccessAfter)
-{
-	return AccessBefore == FRDGResourceState::EAccess::Read && AccessAfter == FRDGResourceState::EAccess::Write;
-}
+	/** Returns whether the resource is transitioning to a writable state. */
+	inline bool IsWriteAccessBegin(FRDGResourceState::EAccess AccessBefore, FRDGResourceState::EAccess AccessAfter)
+	{
+		return AccessBefore == FRDGResourceState::EAccess::Read && AccessAfter == FRDGResourceState::EAccess::Write;
+	}
 
-/** Returns whether the resource is transitioning from a writable state. */
-inline bool IsWriteAccessEnd(FRDGResourceState::EAccess AccessBefore, FRDGResourceState::EAccess AccessAfter)
-{
-	return AccessBefore == FRDGResourceState::EAccess::Write && AccessAfter == FRDGResourceState::EAccess::Read;
-}
+	/** Returns whether the resource is transitioning from a writable state. */
+	inline bool IsWriteAccessEnd(FRDGResourceState::EAccess AccessBefore, FRDGResourceState::EAccess AccessAfter)
+	{
+		return AccessBefore == FRDGResourceState::EAccess::Write && AccessAfter == FRDGResourceState::EAccess::Read;
+	}
 
-inline EResourceTransitionPipeline GetResourceTransitionPipeline(FRDGResourceState::EPipeline PipelineBefore, FRDGResourceState::EPipeline PipelineAfter)
-{
-	switch (PipelineBefore)
+	inline EResourceTransitionPipeline GetResourceTransitionPipeline(FRDGResourceState::EPipeline PipelineBefore, FRDGResourceState::EPipeline PipelineAfter)
 	{
-	case FRDGResourceState::EPipeline::Graphics:
-		switch (PipelineAfter)
+		switch (PipelineBefore)
 		{
 		case FRDGResourceState::EPipeline::Graphics:
-			return EResourceTransitionPipeline::EGfxToGfx;
-		case FRDGResourceState::EPipeline::Compute:
-			return EResourceTransitionPipeline::EGfxToCompute;
-		}
-		break;
-	case FRDGResourceState::EPipeline::Compute:
-		switch (PipelineAfter)
-		{
-		case FRDGResourceState::EPipeline::Graphics:
-			return EResourceTransitionPipeline::EComputeToGfx;
+			switch (PipelineAfter)
+			{
+			case FRDGResourceState::EPipeline::Graphics:
+				return EResourceTransitionPipeline::EGfxToGfx;
+			case FRDGResourceState::EPipeline::Compute:
+				return EResourceTransitionPipeline::EGfxToCompute;
+			}
+			break;
 		case FRDGResourceState::EPipeline::Compute:
-			return EResourceTransitionPipeline::EComputeToCompute;
+			switch (PipelineAfter)
+			{
+			case FRDGResourceState::EPipeline::Graphics:
+				return EResourceTransitionPipeline::EComputeToGfx;
+			case FRDGResourceState::EPipeline::Compute:
+				return EResourceTransitionPipeline::EComputeToCompute;
+			}
+			break;
 		}
-		break;
+		check(false);
+		return EResourceTransitionPipeline::EGfxToGfx;
 	}
-	check(false);
-	return EResourceTransitionPipeline::EGfxToGfx;
-}
 
-inline EResourceTransitionAccess GetResourceTransitionAccess(FRDGResourceState::EAccess AccessAfter)
-{
-	return AccessAfter == FRDGResourceState::EAccess::Write ? EResourceTransitionAccess::EWritable : EResourceTransitionAccess::EReadable;
-}
+	inline EResourceTransitionAccess GetResourceTransitionAccess(FRDGResourceState::EAccess AccessAfter)
+	{
+		return AccessAfter == FRDGResourceState::EAccess::Write ? EResourceTransitionAccess::EWritable : EResourceTransitionAccess::EReadable;
+	}
 
-inline EResourceTransitionAccess GetResourceTransitionAccessForUAV(FRDGResourceState::EAccess AccessBefore, FRDGResourceState::EAccess AccessAfter)
-{
-	switch (AccessAfter)
+	inline EResourceTransitionAccess GetResourceTransitionAccessForUAV(FRDGResourceState::EAccess AccessBefore, FRDGResourceState::EAccess AccessAfter)
 	{
-	case FRDGResourceState::EAccess::Read:
-		return EResourceTransitionAccess::EReadable;
-	case FRDGResourceState::EAccess::Write:
-		// If doing a Write -> Write transition, we use a UAV barrier.
-		if (AccessBefore == FRDGResourceState::EAccess::Write)
-		{
-			return EResourceTransitionAccess::ERWBarrier;
-		}
-		else
+		switch (AccessAfter)
 		{
-			return EResourceTransitionAccess::EWritable;
+		case FRDGResourceState::EAccess::Read:
+			return EResourceTransitionAccess::EReadable;
+		case FRDGResourceState::EAccess::Write:
+			// If doing a Write -> Write transition, we use a UAV barrier.
+			if (AccessBefore == FRDGResourceState::EAccess::Write)
+			{
+				return EResourceTransitionAccess::ERWBarrier;
+			}
+			else
+			{
+				return EResourceTransitionAccess::EWritable;
+			}
 		}
+		check(false);
+		return EResourceTransitionAccess::EMaxAccess;
 	}
-	check(false);
-	return EResourceTransitionAccess::EMaxAccess;
-}
 
 } //! namespace
 
@@ -126,12 +126,12 @@ void FRDGBarrierBatcher::QueueTransitionTexture(FRDGTexture* Texture, FRDGResour
 			FTextureBatch& TextureBatch = TextureBatchMap.FindOrAdd(TransitionParameters);
 			TextureBatch.Reserve(kBatchReservationSize);
 
-			#if RDG_ENABLE_DEBUG
+#if RDG_ENABLE_DEBUG
 			{
 				// We should have filtered out duplicates in the first branch of this function.
 				check(TextureBatch.Find(RHITexture) == INDEX_NONE);
 			}
-			#endif
+#endif
 
 			TextureBatch.Add(RHITexture);
 		}
@@ -175,12 +175,12 @@ void FRDGBarrierBatcher::QueueTransitionUAV(
 			FUAVBatch& UAVBatch = UAVBatchMap.FindOrAdd(TransitionParameters);
 			UAVBatch.Reserve(kBatchReservationSize);
 
-			#if RDG_ENABLE_DEBUG
+#if RDG_ENABLE_DEBUG
 			{
 				// We should have filtered out duplicates in the first branch of this function.
 				check(UAVBatch.Find(UAV) == INDEX_NONE);
 			}
-			#endif
+#endif
 
 			UAVBatch.Add(UAV);
 		}
