Remove duplicated function
#rnx

diff --git a/Engine/Source/Runtime/RenderCore/Private/RenderGraphBuilder.cpp b/Engine/Source/Runtime/RenderCore/Private/RenderGraphBuilder.cpp
--- a/Engine/Source/Runtime/RenderCore/Private/RenderGraphBuilder.cpp
+++ b/Engine/Source/Runtime/RenderCore/Private/RenderGraphBuilder.cpp
@@ -632,23 +632,6 @@ void FRDGBuilder::AllocateRHITextureIfNeeded(FRDGTexture* Texture)
 	check(Texture->ResourceRHI);
 }
 
-void FRDGBuilder::AllocateRHIBufferIfNeeded(FRDGBuffer* Buffer)
-{
-	check(Buffer);
-
-	if (Buffer->PooledBuffer)
-	{
-		return;
-	}
-
-	check(Buffer->ReferenceCount > 0 || GRDGImmediateMode);
-
-	TRefCountPtr<FPooledRDGBuffer>& AllocatedBuffer = AllocatedBuffers.FindOrAdd(Buffer);
-	GRenderGraphResourcePool.FindFreeBuffer(RHICmdList, Buffer->Desc, AllocatedBuffer, Buffer->Name);
-	check(AllocatedBuffer);
-	Buffer->PooledBuffer = AllocatedBuffer;
-}
-
 void FRDGBuilder::AllocateRHITextureUAVIfNeeded(FRDGTextureUAV* UAV)
 {
 	check(UAV);
