Fixed issue where RDG buffer pool would grow unbounded.

diff --git a/Engine/Source/Editor/UnrealEd/Private/EditorEngine.cpp b/Engine/Source/Editor/UnrealEd/Private/EditorEngine.cpp
--- a/Engine/Source/Editor/UnrealEd/Private/EditorEngine.cpp
+++ b/Engine/Source/Editor/UnrealEd/Private/EditorEngine.cpp
@@ -211,6 +211,7 @@
 #include "Materials/MaterialInstance.h"
 #include "ComponentRecreateRenderStateContext.h"
 #include "RenderTargetPool.h"
+#include "RenderGraphBuilder.h"
 
 DEFINE_LOG_CATEGORY_STATIC(LogEditor, Log, All);
 
@@ -1895,6 +1896,7 @@ void UEditorEngine::Tick( float DeltaSeconds, bool bIdleMode )
 					GRenderingRealtimeClock.Tick(DeltaTime);
 				}
 				GRenderTargetPool.TickPoolElements();
+				FRDGBuilder::TickPoolElements();
 			});
 	}
 
diff --git a/Engine/Source/Runtime/Engine/Private/GameEngine.cpp b/Engine/Source/Runtime/Engine/Private/GameEngine.cpp
--- a/Engine/Source/Runtime/Engine/Private/GameEngine.cpp
+++ b/Engine/Source/Runtime/Engine/Private/GameEngine.cpp
@@ -59,6 +59,7 @@
 #include "DynamicResolutionState.h"
 #include "ProfilingDebugging/CsvProfiler.h"
 #include "RenderTargetPool.h"
+#include "RenderGraphBuilder.h"
 
 #if WITH_EDITOR
 #include "PIEPreviewDeviceProfileSelectorModule.h"
@@ -1804,6 +1805,7 @@ void UGameEngine::Tick( float DeltaSeconds, bool bIdleMode )
 			}
 			
 			GRenderTargetPool.TickPoolElements();
+			FRDGBuilder::TickPoolElements();
 		});
 	}
 
diff --git a/Engine/Source/Runtime/RenderCore/Private/RenderGraphBuilder.cpp b/Engine/Source/Runtime/RenderCore/Private/RenderGraphBuilder.cpp
--- a/Engine/Source/Runtime/RenderCore/Private/RenderGraphBuilder.cpp
+++ b/Engine/Source/Runtime/RenderCore/Private/RenderGraphBuilder.cpp
@@ -89,6 +89,11 @@ void EmitRDGWarning(const FString& WarningMessage)
 #define EmitRDGWarningf(WarningMessageFormat, ...) \
 	EmitRDGWarning(FString::Printf(WarningMessageFormat, ##__VA_ARGS__));
 
+void FRDGBuilder::TickPoolElements()
+{
+	GRenderGraphResourcePool.TickPoolElements();
+}
+
 FRDGBuilder::FRDGBuilder(FRHICommandListImmediate& InRHICmdList)
 	: RHICmdList(InRHICmdList)
 	, MemStack(FMemStack::Get())
diff --git a/Engine/Source/Runtime/RenderCore/Private/RenderGraphResourcePool.cpp b/Engine/Source/Runtime/RenderCore/Private/RenderGraphResourcePool.cpp
--- a/Engine/Source/Runtime/RenderCore/Private/RenderGraphResourcePool.cpp
+++ b/Engine/Source/Runtime/RenderCore/Private/RenderGraphResourcePool.cpp
@@ -29,6 +29,8 @@ void FRenderGraphResourcePool::FindFreeBuffer(
 		if (PooledBuffer->Desc == Desc)
 		{
 			Out = PooledBuffer;
+			Out->LastUsedFrame = FrameCounter;
+			Out->Name = InDebugName;
 			// TODO(RDG): assign name on RHI.
 			return;
 		}
@@ -41,6 +43,8 @@ void FRenderGraphResourcePool::FindFreeBuffer(
 		check(Out->GetRefCount() == 2);
 
 		Out->Desc = Desc;
+		Out->Name = InDebugName;
+		Out->LastUsedFrame = FrameCounter;
 
 		uint32 NumBytes = Desc.GetTotalNumBytes();
 
@@ -67,15 +71,32 @@ void FRenderGraphResourcePool::ReleaseDynamicRHI()
 	AllocatedBuffers.Empty();
 }
 
-void FRenderGraphResourcePool::FreeUnusedResources()
+void FRenderGraphResourcePool::TickPoolElements()
 {
-	for (auto& PooledBuffer : AllocatedBuffers)
+	const uint32 kFramesUntilRelease = 30;
+
+	int32 BufferIndex = 0;
+
+	while (BufferIndex < AllocatedBuffers.Num())
 	{
-		if (PooledBuffer.GetRefCount() == 1)
+		TRefCountPtr<FPooledRDGBuffer>& Buffer = AllocatedBuffers[BufferIndex];
+
+		const bool bIsUnused = Buffer.GetRefCount() == 1;
+
+		const bool bNotRequestedRecently = (FrameCounter - Buffer->LastUsedFrame) > kFramesUntilRelease;
+
+		if (bIsUnused && bNotRequestedRecently)
 		{
-			PooledBuffer = nullptr;
+			Swap(Buffer, AllocatedBuffers.Last());
+			AllocatedBuffers.Pop();
+		}
+		else
+		{
+			++BufferIndex;
 		}
 	}
+
+	++FrameCounter;
 }
 
 TGlobalResource<FRenderGraphResourcePool> GRenderGraphResourcePool;
diff --git a/Engine/Source/Runtime/RenderCore/Private/RenderGraphResourcePool.h b/Engine/Source/Runtime/RenderCore/Private/RenderGraphResourcePool.h
--- a/Engine/Source/Runtime/RenderCore/Private/RenderGraphResourcePool.h
+++ b/Engine/Source/Runtime/RenderCore/Private/RenderGraphResourcePool.h
@@ -31,12 +31,14 @@ public:
 	/** Free renderer resources */
 	virtual void ReleaseDynamicRHI() override;
 
-	/** Good to call between levels or before memory intense operations. */
-	void FreeUnusedResources();
+	/** Call once per frame to trim elements from the pool. */
+	void TickPoolElements();
 
 private:
 	/** Elements can be 0, we compact the buffer later. */
-	TArray< TRefCountPtr<FPooledRDGBuffer> > AllocatedBuffers;
+	TArray<TRefCountPtr<FPooledRDGBuffer>> AllocatedBuffers;
+
+	uint32 FrameCounter = 0;
 };
 
 /** The global render targets for easy shading. */
diff --git a/Engine/Source/Runtime/RenderCore/Public/RenderGraphBuilder.h b/Engine/Source/Runtime/RenderCore/Public/RenderGraphBuilder.h
--- a/Engine/Source/Runtime/RenderCore/Public/RenderGraphBuilder.h
+++ b/Engine/Source/Runtime/RenderCore/Public/RenderGraphBuilder.h
@@ -135,6 +135,9 @@ public:
 	/** Executes the queued passes, managing setting of render targets (RHI RenderPasses), resource transitions and queued texture extraction. */
 	void Execute();
 
+	/** Per-frame update of the render graph resource pool. */
+	static void TickPoolElements();
+
 	/** The RHI command list used for the render graph. */
 	FRHICommandListImmediate& RHICmdList;
 
diff --git a/Engine/Source/Runtime/RenderCore/Public/RenderGraphResources.h b/Engine/Source/Runtime/RenderCore/Public/RenderGraphResources.h
--- a/Engine/Source/Runtime/RenderCore/Public/RenderGraphResources.h
+++ b/Engine/Source/Runtime/RenderCore/Public/RenderGraphResources.h
@@ -612,7 +612,11 @@ struct FPooledRDGBuffer
 	}
 
 private:
+	const TCHAR* Name = nullptr;
 	uint32 RefCount = 0;
+	uint32 LastUsedFrame = 0;
+
+	friend class FRenderGraphResourcePool;
 };
 
 /** Render graph tracked buffers. */
