Fixed dereferencing null ptr in FRealtimeGPUProfiler at exit

diff --git a/Engine/Source/Runtime/Launch/Private/LaunchEngineLoop.cpp b/Engine/Source/Runtime/Launch/Private/LaunchEngineLoop.cpp
--- a/Engine/Source/Runtime/Launch/Private/LaunchEngineLoop.cpp
+++ b/Engine/Source/Runtime/Launch/Private/LaunchEngineLoop.cpp
@@ -431,7 +431,7 @@ static TUniquePtr<FOutputDeviceTestExit> GScopedTestExit;
 static void RHIExitAndStopRHIThread()
 {
 #if HAS_GPU_STATS
-	FRealtimeGPUProfiler::Get()->Release();
+	FRealtimeGPUProfiler::SafeRelease();
 #endif
 
 	// Stop the RHI Thread (using GRHIThread_InternalUseOnly is unreliable since RT may be stopped)
diff --git a/Engine/Source/Runtime/RenderCore/Private/ProfilingDebugging/RealtimeGPUProfiler.cpp b/Engine/Source/Runtime/RenderCore/Private/ProfilingDebugging/RealtimeGPUProfiler.cpp
--- a/Engine/Source/Runtime/RenderCore/Private/ProfilingDebugging/RealtimeGPUProfiler.cpp
+++ b/Engine/Source/Runtime/RenderCore/Private/ProfilingDebugging/RealtimeGPUProfiler.cpp
@@ -538,6 +538,15 @@ FRealtimeGPUProfiler* FRealtimeGPUProfiler::Get()
 	return Instance;
 }
 
+
+void FRealtimeGPUProfiler::SafeRelease()
+{
+	if (Instance)
+		Instance->Cleanup();
+	Instance = nullptr;
+}
+
+
 FRealtimeGPUProfiler::FRealtimeGPUProfiler()
 	: WriteBufferIndex(0)
 	, ReadBufferIndex(1) 
@@ -554,6 +563,11 @@ FRealtimeGPUProfiler::FRealtimeGPUProfiler()
 }
 
 void FRealtimeGPUProfiler::Release()
+{
+	Cleanup();
+}
+
+void FRealtimeGPUProfiler::Cleanup()
 {
 	for (int Index = 0; Index < Frames.Num(); Index++)
 	{
diff --git a/Engine/Source/Runtime/RenderCore/Public/ProfilingDebugging/RealtimeGPUProfiler.h b/Engine/Source/Runtime/RenderCore/Public/ProfilingDebugging/RealtimeGPUProfiler.h
--- a/Engine/Source/Runtime/RenderCore/Public/ProfilingDebugging/RealtimeGPUProfiler.h
+++ b/Engine/Source/Runtime/RenderCore/Public/ProfilingDebugging/RealtimeGPUProfiler.h
@@ -233,11 +233,15 @@ public:
 	// Singleton interface
 	static RENDERCORE_API FRealtimeGPUProfiler* Get();
 
+	/** *Safe release of the singleton */
+	static RENDERCORE_API void SafeRelease();
+
 	/** Per-frame update */
 	RENDERCORE_API void BeginFrame(FRHICommandListImmediate& RHICmdList);
 	RENDERCORE_API void EndFrame(FRHICommandListImmediate& RHICmdList);
 
 	/** Final cleanup */
+	UE_DEPRECATED(4.23, "Use FRealtimeGPUProfiler::SafeRelease() instead.")
 	RENDERCORE_API void Release();
 
 	/** Push/pop events */
@@ -247,6 +251,10 @@ public:
 private:
 	FRealtimeGPUProfiler();
 
+	/** Deinitialize of the object*/
+	void Cleanup();
+
+
 	/** Ringbuffer of profiler frames */
 	TArray<FRealtimeGPUProfilerFrame*> Frames;
 
