Refactored render graph validation logic. Hardened mip generation pass barrier handling. Fixed compute-based mip generation pass helper.

diff --git a/Engine/Source/Runtime/Engine/Private/TextureRenderTarget2D.cpp b/Engine/Source/Runtime/Engine/Private/TextureRenderTarget2D.cpp
--- a/Engine/Source/Runtime/Engine/Private/TextureRenderTarget2D.cpp
+++ b/Engine/Source/Runtime/Engine/Private/TextureRenderTarget2D.cpp
@@ -53,6 +53,11 @@ FTextureResource* UTextureRenderTarget2D::CreateResource()
 	if (bAutoGenerateMips)
 	{
 		NumMips = FGenericPlatformMath::CeilToInt(FGenericPlatformMath::Log2(FGenericPlatformMath::Max(SizeX, SizeY)));
+
+		if (RHIRequiresComputeGenerateMips())
+		{
+			bCanCreateUAV = 1;
+		}
 	}
 	else
 	{
diff --git a/Engine/Source/Runtime/RenderCore/Private/GenerateMips.cpp b/Engine/Source/Runtime/RenderCore/Private/GenerateMips.cpp
--- a/Engine/Source/Runtime/RenderCore/Private/GenerateMips.cpp
+++ b/Engine/Source/Runtime/RenderCore/Private/GenerateMips.cpp
@@ -126,12 +126,16 @@ void FGenerateMips::Compute(FRHICommandListImmediate& RHIImmCmdList, FRHITexture
 			FMath::Max(DestTextureSizeY / MIPSSHADER_NUMTHREADS, 1),
 			1);
 		//Pass added per mip level to be written.
-		FComputeShaderUtils::AddPass(
-			GraphBuilder,
+		ClearUnusedGraphResources(*ComputeShader, PassParameters);
+
+		GraphBuilder.AddPass(
 			RDG_EVENT_NAME("Generate2DTextureMips DestMipLevel=%d", MipLevel),
-			*ComputeShader,
 			PassParameters,
-			GenMipsGroupCount);
+			ERenderGraphPassFlags::Compute | ERDGPassFlags::GenerateMips,
+			[PassParameters, ComputeShader, GenMipsGroupCount](FRHICommandList& RHICmdList)
+		{
+			FComputeShaderUtils::Dispatch(RHICmdList, *ComputeShader, *PassParameters, GenMipsGroupCount);
+		});
 	}
 	GraphBuilder.QueueTextureExtraction(GraphTexture, &GenMipsStruct->RenderTarget);
 	GraphBuilder.Execute();	
diff --git a/Engine/Source/Runtime/RenderCore/Private/RenderGraphBarrierBatcher.cpp b/Engine/Source/Runtime/RenderCore/Private/RenderGraphBarrierBatcher.cpp
--- a/Engine/Source/Runtime/RenderCore/Private/RenderGraphBarrierBatcher.cpp
+++ b/Engine/Source/Runtime/RenderCore/Private/RenderGraphBarrierBatcher.cpp
@@ -48,63 +48,71 @@ namespace
 		return EResourceTransitionPipeline::EGfxToGfx;
 	}
 
-	inline EResourceTransitionAccess GetResourceTransitionAccess(FRDGResourceState::EAccess AccessAfter)
-	{
-		return AccessAfter == FRDGResourceState::EAccess::Write ? EResourceTransitionAccess::EWritable : EResourceTransitionAccess::EReadable;
-	}
+} //! namespace
 
-	inline EResourceTransitionAccess GetResourceTransitionAccessForUAV(FRDGResourceState::EAccess AccessBefore, FRDGResourceState::EAccess AccessAfter)
+FRDGBarrierBatcher::FRDGBarrierBatcher(FRHICommandList& InRHICmdList, const FRDGPass* InPass)
+	: RHICmdList(InRHICmdList)
+	, Pass(InPass)
+{
+	if (Pass)
 	{
-		switch (AccessAfter)
-		{
-		case FRDGResourceState::EAccess::Read:
-			return EResourceTransitionAccess::EReadable;
-		case FRDGResourceState::EAccess::Write:
-			// If doing a Write -> Write transition, we use a UAV barrier.
-			if (AccessBefore == FRDGResourceState::EAccess::Write)
-			{
-				return EResourceTransitionAccess::ERWBarrier;
-			}
-			else
-			{
-				return EResourceTransitionAccess::EWritable;
-			}
-		}
-		check(false);
-		return EResourceTransitionAccess::EMaxAccess;
+		bIsGeneratingMips = Pass->IsGenerateMips();
+		Pipeline = Pass->IsCompute() ? FRDGResourceState::EPipeline::Compute : FRDGResourceState::EPipeline::Graphics;
 	}
-
-} //! namespace
+}
 
 FRDGBarrierBatcher::~FRDGBarrierBatcher()
 {
-	check(!bAllowQueueing);
-}
+	for (FRHITexture* RHITexture : TextureUpdateMultiFrameBegins)
+	{
+		RHICmdList.BeginUpdateMultiFrameResource(RHITexture);
+	}
 
-void FRDGBarrierBatcher::Begin()
-{
-#if RDG_ENABLE_DEBUG
-	check(!bAllowQueueing);
+	for (FRHIUnorderedAccessView* RHIUAV : UAVUpdateMultiFrameBegins)
+	{
+		RHICmdList.BeginUpdateMultiFrameResource(RHIUAV);
+	}
+
+	for (auto& Element : TextureBatchMap)
+	{
+		FTransitionParameters TransitionParameters = Element.Key;
+		FTextureBatch& Batch = Element.Value;
+		RHICmdList.TransitionResources(TransitionParameters.TransitionAccess, Batch.GetData(), Batch.Num());
+	}
 
-	check(!TextureUpdateMultiFrameBegins.Num());
-	check(!TextureUpdateMultiFrameEnds.Num());
-	check(!TextureBatchMap.Num());
+	for (auto& Element : UAVBatchMap)
+	{
+		FTransitionParameters TransitionParameters = Element.Key;
+		FUAVBatch& Batch = Element.Value;
+		RHICmdList.TransitionResources(TransitionParameters.TransitionAccess, TransitionParameters.TransitionPipeline, Batch.GetData(), Batch.Num());
+	}
 
-	check(!UAVUpdateMultiFrameBegins.Num());
-	check(!UAVUpdateMultiFrameEnds.Num());
-	check(!UAVBatchMap.Num());
-#endif
+	for (FRHITexture* RHITexture : TextureUpdateMultiFrameEnds)
+	{
+		RHICmdList.EndUpdateMultiFrameResource(RHITexture);
+	}
 
-	bAllowQueueing = true;
+	for (FRHIUnorderedAccessView* RHIUAV : UAVUpdateMultiFrameEnds)
+	{
+		RHICmdList.EndUpdateMultiFrameResource(RHIUAV);
+	}
 }
 
-void FRDGBarrierBatcher::QueueTransitionTexture(FRDGTexture* Texture, FRDGResourceState StateAfter)
+void FRDGBarrierBatcher::QueueTransitionTexture(FRDGTexture* Texture, FRDGResourceState::EAccess AccessAfter)
 {
 	check(Texture);
 
+	// Texture transitions are ignored when generating mips, since the render target binding call or UAV will
+	// perform the subresource transition.
+	if (bIsGeneratingMips)
+	{
+		return;
+	}
+
 	const FRDGResourceState StateBefore = Texture->State;
+	const FRDGResourceState StateAfter(Pass, Pipeline, AccessAfter);
 
-	ValidateTransition(StateBefore, StateAfter, Texture);
+	ValidateTransition(Texture, StateBefore, StateAfter);
 
 	if (StateBefore != StateAfter)
 	{
@@ -126,12 +134,12 @@ void FRDGBarrierBatcher::QueueTransitionTexture(FRDGTexture* Texture, FRDGResour
 			FTextureBatch& TextureBatch = TextureBatchMap.FindOrAdd(TransitionParameters);
 			TextureBatch.Reserve(kBatchReservationSize);
 
-#if RDG_ENABLE_DEBUG
+			#if RDG_ENABLE_DEBUG
 			{
 				// We should have filtered out duplicates in the first branch of this function.
-				check(TextureBatch.Find(RHITexture) == INDEX_NONE);
+				ensure(TextureBatch.Find(RHITexture) == INDEX_NONE);
 			}
-#endif
+			#endif
 
 			TextureBatch.Add(RHITexture);
 		}
@@ -148,14 +156,15 @@ void FRDGBarrierBatcher::QueueTransitionTexture(FRDGTexture* Texture, FRDGResour
 void FRDGBarrierBatcher::QueueTransitionUAV(
 	FRHIUnorderedAccessView* UAV,
 	FRDGTrackedResource* UnderlyingResource,
-	FRDGResourceState StateAfter)
+	FRDGResourceState::EAccess AccessAfter)
 {
 	check(UAV);
 	check(UnderlyingResource);
 
 	const FRDGResourceState StateBefore = UnderlyingResource->State;
+	const FRDGResourceState StateAfter(Pass, Pipeline, AccessAfter);
 
-	ValidateTransition(StateBefore, StateAfter, UnderlyingResource);
+	ValidateTransition(UnderlyingResource, StateBefore, StateAfter);
 
 	if (StateBefore != StateAfter)
 	{
@@ -175,12 +184,12 @@ void FRDGBarrierBatcher::QueueTransitionUAV(
 			FUAVBatch& UAVBatch = UAVBatchMap.FindOrAdd(TransitionParameters);
 			UAVBatch.Reserve(kBatchReservationSize);
 
-#if RDG_ENABLE_DEBUG
+			#if RDG_ENABLE_DEBUG
 			{
 				// We should have filtered out duplicates in the first branch of this function.
-				check(UAVBatch.Find(UAV) == INDEX_NONE);
+				ensure(UAVBatch.Find(UAV) == INDEX_NONE);
 			}
-#endif
+			#endif
 
 			UAVBatch.Add(UAV);
 		}
@@ -194,53 +203,7 @@ void FRDGBarrierBatcher::QueueTransitionUAV(
 	}
 }
 
-void FRDGBarrierBatcher::End(FRHICommandList& RHICmdList)
-{
-	check(bAllowQueueing);
-	bAllowQueueing = false;
-
-	for (FRHITexture* RHITexture : TextureUpdateMultiFrameBegins)
-	{
-		RHICmdList.BeginUpdateMultiFrameResource(RHITexture);
-	}
-	TextureUpdateMultiFrameBegins.Empty();
-
-	for (FRHIUnorderedAccessView* RHIUAV : UAVUpdateMultiFrameBegins)
-	{
-		RHICmdList.BeginUpdateMultiFrameResource(RHIUAV);
-	}
-	UAVUpdateMultiFrameBegins.Empty();
-
-	for (auto& Element : TextureBatchMap)
-	{
-		FTransitionParameters TransitionParameters = Element.Key;
-		FTextureBatch& Batch = Element.Value;
-		RHICmdList.TransitionResources(TransitionParameters.TransitionAccess, Batch.GetData(), Batch.Num());
-	}
-	TextureBatchMap.Empty();
-
-	for (auto& Element : UAVBatchMap)
-	{
-		FTransitionParameters TransitionParameters = Element.Key;
-		FUAVBatch& Batch = Element.Value;
-		RHICmdList.TransitionResources(TransitionParameters.TransitionAccess, TransitionParameters.TransitionPipeline, Batch.GetData(), Batch.Num());
-	}
-	UAVBatchMap.Empty();
-
-	for (FRHITexture* RHITexture : TextureUpdateMultiFrameEnds)
-	{
-		RHICmdList.EndUpdateMultiFrameResource(RHITexture);
-	}
-	TextureUpdateMultiFrameEnds.Empty();
-
-	for (FRHIUnorderedAccessView* RHIUAV : UAVUpdateMultiFrameEnds)
-	{
-		RHICmdList.EndUpdateMultiFrameResource(RHIUAV);
-	}
-	UAVUpdateMultiFrameEnds.Empty();
-}
-
-void FRDGBarrierBatcher::ValidateTransition(FRDGResourceState StateBefore, FRDGResourceState StateAfter, const FRDGTrackedResource* Resource)
+void FRDGBarrierBatcher::ValidateTransition(const FRDGTrackedResource* Resource, FRDGResourceState StateBefore, FRDGResourceState StateAfter)
 {
 #if RDG_ENABLE_DEBUG
 	check(StateAfter.Pipeline != FRDGResourceState::EPipeline::MAX);
@@ -249,11 +212,44 @@ void FRDGBarrierBatcher::ValidateTransition(FRDGResourceState StateBefore, FRDGR
 	if (StateBefore != StateAfter && StateAfter.Pass)
 	{
 		// We allow duplicate transitions of the same resource within the same pass, but not conflicting ones.
-		checkf(
+		ensureMsgf(
 			StateBefore.Pass != StateAfter.Pass,
 			TEXT("Pass %s attempted to transition resource %s to different states. Make sure the resource isn't being used\n")
 			TEXT("for both read and write at the same time. This can occur if the resource is used as both an SRV and UAV, or\n")
-			TEXT("SRV and Render Target, for example.\n"), StateAfter.Pass->GetName(), Resource->Name);
+			TEXT("SRV and Render Target, for example. If this pass is meant to generate mip maps, make sure the GenerateMips flag\n")
+			TEXT("is set.\n"), StateAfter.Pass->GetName(), Resource->Name);
 	}
 #endif
-}
\ No newline at end of file
+}
+
+EResourceTransitionAccess FRDGBarrierBatcher::GetResourceTransitionAccess(FRDGResourceState::EAccess AccessAfter) const
+{
+	return AccessAfter == FRDGResourceState::EAccess::Write ? EResourceTransitionAccess::EWritable : EResourceTransitionAccess::EReadable;
+}
+
+EResourceTransitionAccess FRDGBarrierBatcher::GetResourceTransitionAccessForUAV(FRDGResourceState::EAccess AccessBefore, FRDGResourceState::EAccess AccessAfter) const
+{
+	switch (AccessAfter)
+	{
+	case FRDGResourceState::EAccess::Read:
+		return EResourceTransitionAccess::EReadable;
+
+	case FRDGResourceState::EAccess::Write:
+		// Mip-map generation uses its own barrier.
+		if (bIsGeneratingMips)
+		{
+			return EResourceTransitionAccess::ERWSubResBarrier;
+		}
+		// If doing a Write -> Write transition, we use a UAV barrier.
+		else if (AccessBefore == FRDGResourceState::EAccess::Write)
+		{
+			return EResourceTransitionAccess::ERWBarrier;
+		}
+		else
+		{
+			return EResourceTransitionAccess::EWritable;
+		}
+	}
+	check(false);
+	return EResourceTransitionAccess::EMaxAccess;
+}
diff --git a/Engine/Source/Runtime/RenderCore/Private/RenderGraphBarrierBatcher.h b/Engine/Source/Runtime/RenderCore/Private/RenderGraphBarrierBatcher.h
--- a/Engine/Source/Runtime/RenderCore/Private/RenderGraphBarrierBatcher.h
+++ b/Engine/Source/Runtime/RenderCore/Private/RenderGraphBarrierBatcher.h
@@ -8,24 +8,25 @@
 class RENDERCORE_API FRDGBarrierBatcher final
 {
 public:
-	FRDGBarrierBatcher() = default;
+	/** RAII initialization of batcher. Batcher will flush all queued transitions in the destructor.
+	 *  @param Pass The current pass if performing inter-pass barriers. This can be null (e.g. for post-execution barriers).
+	 */
+	FRDGBarrierBatcher(FRHICommandList& RHICmdList, const FRDGPass* Pass);
 	~FRDGBarrierBatcher();
 
-	void Begin();
-
-	void QueueTransitionTexture(FRDGTexture* Texture, FRDGResourceState StateAfter);
+	/** Queues a transition of the texture to the requested access state. */
+	void QueueTransitionTexture(FRDGTexture* Texture, FRDGResourceState::EAccess AccessAfter);
 
+	/** Queues a transition of the UAV to the requested access state. The current state of the resource is checked and the transition
+	 *  is only performed if a change occurs. It is valid to queue the same state multiple times. However, it is invalid to queue a
+	 *  resource into multiple states at the same time. The underlying resource is the texture / buffer instance referenced by the view.
+	 */
 	void QueueTransitionUAV(
 		FRHIUnorderedAccessView* UAV,
 		FRDGTrackedResource* UnderlyingResource,
-		FRDGResourceState StateAfter);
-
-	/** Submits and clears all queued transitions to the provided RHI command list. */
-	void End(FRHICommandList& RHICmdList);
+		FRDGResourceState::EAccess AccessAfter);
 
 private:
-	static void ValidateTransition(FRDGResourceState StateBefore, FRDGResourceState StateAfter, const FRDGTrackedResource* Resource);
-
 	struct FTransitionParameters
 	{
 		FTransitionParameters() = default;
@@ -59,6 +60,8 @@ private:
 	using FTextureBatch = TArray<FRHITexture*, SceneRenderingAllocator>;
 	using FTextureBatchMap = TMap<FTransitionParameters, FTextureBatch, SceneRenderingSetAllocator, TBatchMapKeyFuncs<FTextureBatch>>;
 
+	FRHICommandList& RHICmdList;
+
 	FTextureBatch TextureUpdateMultiFrameBegins;
 	FTextureBatch TextureUpdateMultiFrameEnds;
 	FTextureBatchMap TextureBatchMap;
@@ -67,5 +70,12 @@ private:
 	FUAVBatch UAVUpdateMultiFrameEnds;
 	FUAVBatchMap UAVBatchMap;
 
-	bool bAllowQueueing = false;
+	void ValidateTransition(const FRDGTrackedResource* Resource, FRDGResourceState StateBefore, FRDGResourceState StateAfter);
+
+	EResourceTransitionAccess GetResourceTransitionAccess(FRDGResourceState::EAccess AccessAfter) const;
+	EResourceTransitionAccess GetResourceTransitionAccessForUAV(FRDGResourceState::EAccess AccessBefore, FRDGResourceState::EAccess AccessAfter) const;
+
+	const FRDGPass* Pass = nullptr;
+	FRDGResourceState::EPipeline Pipeline = FRDGResourceState::EPipeline::Graphics;
+	bool bIsGeneratingMips = false;
 };
\ No newline at end of file
diff --git a/Engine/Source/Runtime/RenderCore/Private/RenderGraphBuilder.cpp b/Engine/Source/Runtime/RenderCore/Private/RenderGraphBuilder.cpp
--- a/Engine/Source/Runtime/RenderCore/Private/RenderGraphBuilder.cpp
+++ b/Engine/Source/Runtime/RenderCore/Private/RenderGraphBuilder.cpp
@@ -5,6 +5,7 @@
 #include "RenderTargetPool.h"
 #include "RenderGraphResourcePool.h"
 #include "RenderGraphBarrierBatcher.h"
+#include "RenderGraphValidation.h"
 #include "VisualizeTexture.h"
 #include "ProfilingDebugging/CsvProfiler.h"
 
@@ -48,6 +49,16 @@ bool GetEmitRDGEvents()
 #endif
 }
 
+bool IsRDGDebugEnabled()
+{
+	return GRDGDebug != 0;
+}
+
+bool IsRDGImmediateModeEnabled()
+{
+	return GRDGImmediateMode != 0;
+}
+
 void InitRenderGraph()
 {
 #if RDG_ENABLE_DEBUG_WITH_ENGINE
@@ -86,9 +97,6 @@ void EmitRDGWarning(const FString& WarningMessage)
 	}
 }
 
-#define EmitRDGWarningf(WarningMessageFormat, ...) \
-	EmitRDGWarning(FString::Printf(WarningMessageFormat, ##__VA_ARGS__));
-
 void FRDGBuilder::TickPoolElements()
 {
 	GRenderGraphResourcePool.TickPoolElements();
@@ -101,34 +109,11 @@ FRDGBuilder::FRDGBuilder(FRHICommandListImmediate& InRHICmdList)
 	, StatScopeStack(RHICmdList)
 {}
 
-FRDGBuilder::~FRDGBuilder()
-{
-	#if RDG_ENABLE_DEBUG
-	{
-		checkf(bHasExecuted, TEXT("Render graph execution is required to ensure consistency with immediate mode."));
-	}
-	#endif
-}
-
 void FRDGBuilder::Execute()
 {
 	CSV_SCOPED_TIMING_STAT_EXCLUSIVE(FRDGBuilder_Execute);
 
-	#if RDG_ENABLE_DEBUG
-	{
-		checkf(!bHasExecuted, TEXT("Render graph execution should only happen once to ensure consistency with immediate mode."));
-
-		/** FRDGBuilder::AllocParameters() allocates shader parameter structure for the life time until pass execution.
-		 * But they are allocated on a FMemStack for CPU performance reason, and have their destructor called right after
-		 * the pass execution. Therefore allocating pass parameter unused by a FRDGBuilder::AddPass() can lead on a memory
-		 * leak of RHI resource that have been reference in the parameter structure.
-		 */
-		checkf(
-			AllocatedUnusedPassParameters.Num() == 0,
-			TEXT("%i pass parameter structure has been allocated with FRDGBuilder::AllocParameters(), but has not be used by a ")
-			TEXT("FRDGBuilder::AddPass() that can cause RHI resource leak."), AllocatedUnusedPassParameters.Num());
-	}
-	#endif
+	IF_RDG_ENABLE_DEBUG(Validation.ValidateExecuteBegin());
 
 	EventScopeStack.BeginExecute();
 	StatScopeStack.BeginExecute();
@@ -149,43 +134,19 @@ void FRDGBuilder::Execute()
 
 	ProcessDeferredInternalResourceQueries();
 
-	DestructPasses();
+	IF_RDG_ENABLE_DEBUG(Validation.ValidateExecuteEnd());
 
-	#if RDG_ENABLE_DEBUG
-	{
-		bHasExecuted = true;
-	}
-	#endif
+	DestructPasses();
 }
 
 void FRDGBuilder::AddPassInternal(FRDGPass* Pass)
 {
-	#if RDG_ENABLE_DEBUG
-	{
-		checkf(!bHasExecuted, TEXT("Render graph pass %s needs to be added before the builder execution."), Pass->GetName());
-
-		const void* ParameterStructData = Pass->GetParameters().GetContents();
-
-		/** The lifetime of each pass parameter structure must extend until deferred pass execution; therefore, it needs to be
-		 *  allocated with FRDGBuilder::AllocParameters(). Also, all references held by the parameter structure are released
-		 *  immediately after pass execution, so a pass parameter struct instance must be 1-to-1 with a pass instance (i.e. one
-		 *  per AddPass() call).
-		 */
-		checkf(
-			AllocatedUnusedPassParameters.Contains(ParameterStructData),
-			TEXT("The pass parameter structure has not been allocated for correct life time FRDGBuilder::AllocParameters() or has already ")
-			TEXT("been used by another previous FRDGBuilder::AddPass()."));
-
-		AllocatedUnusedPassParameters.Remove(ParameterStructData);
-	}
-	#endif
+	IF_RDG_ENABLE_DEBUG(Validation.ValidateAddPass(Pass));
 
 	Pass->EventScope = EventScopeStack.GetCurrentScope();
 	Pass->StatScope = StatScopeStack.GetCurrentScope();
 	Passes.Emplace(Pass);
 
-	ValidatePass(Pass);
-
 	if (GRDGImmediateMode)
 	{
 		ExecutePass(Pass);
@@ -194,230 +155,6 @@ void FRDGBuilder::AddPassInternal(FRDGPass* Pass)
 	VisualizePassOutputs(Pass);
 }
 
-void FRDGBuilder::ValidatePass(const FRDGPass* Pass) const
-{
-#if RDG_ENABLE_DEBUG
-	const FRenderTargetBindingSlots* RenderTargetBindingSlots = nullptr;
-
-	const TCHAR* PassName = Pass->GetName();
-	const bool bIsCompute = Pass->IsCompute();
-	const bool bRenderTargetSlotsRequired = !bIsCompute;
-
-	FRDGPassParameterStruct ParameterStruct = Pass->GetParameters();
-
-	const uint32 ParameterCount = ParameterStruct.GetParameterCount();
-
-	for (uint32 ParameterIndex = 0; ParameterIndex < ParameterCount; ++ParameterIndex)
-	{
-		FRDGPassParameter Parameter = ParameterStruct.GetParameter(ParameterIndex);
-
-		switch (Parameter.GetType())
-		{
-		case UBMT_RDG_TEXTURE:
-		{
-			if (FRDGTextureRef Texture = Parameter.GetAsTexture())
-			{
-				checkf(Texture->HasBeenProduced(),
-					TEXT("Pass %s has a dependency over the texture %s that has never been produced."),
-					PassName, Texture->Name);
-			}
-		}
-		break;
-		case UBMT_RDG_TEXTURE_SRV:
-		{
-			if (FRDGTextureSRVRef SRV = Parameter.GetAsTextureSRV())
-			{
-				FRDGTextureRef Texture = SRV->Desc.Texture;
-
-				checkf(Texture->HasBeenProduced(),
-					TEXT("Pass %s has a dependency over the texture %s that has never been produced."),
-					PassName, Texture->Name);
-			}
-		}
-		break;
-		case UBMT_RDG_TEXTURE_UAV:
-		{
-			if (FRDGTextureUAVRef UAV = Parameter.GetAsTextureUAV())
-			{
-				FRDGTextureRef Texture = UAV->Desc.Texture;
-
-				Texture->MarkAsProducedBy(Pass);
-			}
-		}
-		break;
-		case UBMT_RDG_BUFFER:
-		{
-			if (FRDGBufferRef Buffer = Parameter.GetAsBuffer())
-			{
-				checkf(Buffer->HasBeenProduced(),
-					TEXT("Pass %s has a dependency over the buffer %s that has never been produced."),
-					PassName, Buffer->Name);
-			}
-		}
-		break;
-		case UBMT_RDG_BUFFER_SRV:
-		{
-			if (FRDGBufferSRVRef SRV = Parameter.GetAsBufferSRV())
-			{
-				FRDGBufferRef Buffer = SRV->Desc.Buffer;
-
-				checkf(Buffer->HasBeenProduced(),
-					TEXT("Pass %s has a dependency over the buffer %s that has never been produced."),
-					PassName, SRV->Desc.Buffer->Name);
-			}
-		}
-		break;
-		case UBMT_RDG_BUFFER_UAV:
-		{
-			if (FRDGBufferUAVRef UAV = Parameter.GetAsBufferUAV())
-			{
-				FRDGBufferRef Buffer = UAV->Desc.Buffer;
-
-				Buffer->MarkAsProducedBy(Pass);
-			}
-		}
-		break;
-		case UBMT_RENDER_TARGET_BINDING_SLOTS:
-		{
-			if (!RenderTargetBindingSlots)
-			{
-				RenderTargetBindingSlots = &Parameter.GetAsRenderTargetBindingSlots();
-			}
-			else if (GRDGDebug)
-			{
-				EmitRDGWarningf(
-					TEXT("Pass %s have duplicated render target binding slots."),
-					PassName);
-			}
-		}
-		break;
-		default:
-			break;
-		}
-	}
-
-	/** Validate that raster passes have render target binding slots and compute passes don't. */
-	if (RenderTargetBindingSlots)
-	{
-		checkf(!bIsCompute, TEXT("Pass '%s' has render target binding slots but is flagged as 'Compute'."), PassName);
-	}
-	else
-	{
-		checkf(bIsCompute, TEXT("Pass '%s' is missing render target binding slots. Add the 'Compute' flag if render targets are not required."), PassName);
-	}
-
-	/** Validate render target / depth stencil binding usage. */
-	if (RenderTargetBindingSlots)
-	{
-		const auto& RenderTargets = RenderTargetBindingSlots->Output;
-
-		const uint32 RenderTargetCount = RenderTargets.Num();
-
-		{
-			/** Tracks the number of contiguous, non-null textures in the render target output array. */
-			uint32 ValidRenderTargetCount = 0;
-
-			for (uint32 RenderTargetIndex = 0; RenderTargetIndex < RenderTargetCount; ++RenderTargetIndex)
-			{
-				const FRenderTargetBinding& RenderTarget = RenderTargets[RenderTargetIndex];
-
-				if (FRDGTextureRef Texture = RenderTarget.GetTexture())
-				{
-					/** Validate that load action is correct. */
-					const bool bIsLoadAction = RenderTarget.GetLoadAction() == ERenderTargetLoadAction::ELoad;
-					const bool bIsLoadActionInvalid = bIsLoadAction && !Texture->HasBeenProduced();
-					checkf(
-						!bIsLoadActionInvalid,
-						TEXT("Can't load a render target '%s' that has never been produced."),
-						Texture->Name);
-					
-					const bool bOverwritesResourceEntirely = !bIsLoadAction && Texture->HasBeenProduced() && Texture->Desc.NumMips == 1;
-					ensureMsgf(!bOverwritesResourceEntirely,
-						TEXT("Clobbering render target %s result that has been produced previously by another pass. Should instead render to a new FRDGTexture to reduce memory pressure."),
-						Texture->Name);
-					
-					/** Mark the pass as a producer for render targets with a store action. */
-					{
-						const bool bIsStoreAction = RenderTarget.GetStoreAction() != ERenderTargetStoreAction::ENoAction;
-						check(bIsStoreAction); // already been validated in FRenderTargetBinding::Validate()
-						Texture->MarkAsProducedBy(Pass);
-					}
-				}
-				else
-				{
-					/** Found end of contiguous interval of valid render targets. */
-					ValidRenderTargetCount = RenderTargetIndex;
-					break;
-				}
-			}
-
-			/** Validate that no holes exist in the render target output array. Render targets must be bound contiguously. */
-			for (uint32 RenderTargetIndex = ValidRenderTargetCount; RenderTargetIndex < RenderTargetCount; ++RenderTargetIndex)
-			{
-				const FRenderTargetBinding& RenderTarget = RenderTargets[RenderTargetIndex];
-				checkf(RenderTarget.GetTexture() == nullptr, TEXT("Render targets must be packed. No empty spaces in the array."));
-			}
-		}
-
-		const auto IsMipLevelBoundForRead = [ParameterStruct, ParameterCount](FRDGTextureRef Texture, uint32 MipLevel)
-		{
-			for (uint32 ParameterIndex = 0; ParameterIndex < ParameterCount; ++ParameterIndex)
-			{
-				FRDGPassParameter Parameter = ParameterStruct.GetParameter(ParameterIndex);
-
-				switch (Parameter.GetType())
-				{
-				case UBMT_RDG_TEXTURE:
-				{
-					if (Parameter.GetAsTexture() == Texture)
-					{
-						/** Full texture view bound. */
-						return true;
-					}
-					break;
-				}
-				case UBMT_RDG_TEXTURE_SRV:
-				{
-					if (const FRDGTextureSRVRef InputSRV = Parameter.GetAsTextureSRV())
-					{
-						const FRDGTextureSRVDesc& InputSRVDesc = InputSRV->Desc;
-
-						if (InputSRVDesc.Texture == Texture)
-						{
-							return InputSRVDesc.MipLevel == MipLevel;
-						}
-					}
-					break;
-				}
-				default:
-					break;
-				}
-			}
-
-			return false;
-		};
-
-		/** Validate that texture mips are not bound as a render target and SRV at the same time. */
-		for (uint32 RenderTargetIndex = 0; RenderTargetIndex < RenderTargetCount; ++RenderTargetIndex)
-		{
-			const FRenderTargetBinding& RenderTarget = RenderTargets[RenderTargetIndex];
-
-			if (FRDGTextureRef Texture = RenderTarget.GetTexture())
-			{
-				if (IsMipLevelBoundForRead(Texture, RenderTarget.GetMipIndex()))
-				{
-					checkf(false, TEXT("Texture '%s', Mip '%d' is bound as a render target and SRV at the same time."), Texture->Name, RenderTarget.GetMipIndex());
-				}
-			}
-			else
-			{
-				break;
-			}
-		}
-	}
-#endif // RDG_ENABLE_DEBUG
-}
-
 void FRDGBuilder::VisualizePassOutputs(const FRDGPass* Pass)
 {
 #if SUPPORTS_VISUALIZE_TEXTURE
@@ -634,7 +371,6 @@ void FRDGBuilder::AllocateRHITextureIfNeeded(FRDGTexture* Texture)
 
 	Texture->PooledRenderTarget = PooledRenderTarget;
 	Texture->ResourceRHI = PooledRenderTarget->GetRenderTargetItem().ShaderResourceTexture;
-	check(Texture->ResourceRHI);
 }
 
 void FRDGBuilder::AllocateRHITextureUAVIfNeeded(FRDGTextureUAV* UAV)
@@ -679,13 +415,6 @@ void FRDGBuilder::AllocateRHIBufferSRVIfNeeded(FRDGBufferSRV* SRV)
 
 	FRDGBufferRef Buffer = SRV->Desc.Buffer;
 
-	// The underlying buffer have already been allocated by a prior pass through AllocateRHIBufferUAVIfNeeded().
-	#if RDG_ENABLE_DEBUG
-	{
-		check(Buffer->HasBeenProduced());
-	}
-	#endif
-
 	check(Buffer->PooledBuffer);
 
 	if (Buffer->PooledBuffer->SRVs.Contains(SRV->Desc))
@@ -757,10 +486,7 @@ void FRDGBuilder::ExecutePass(const FRDGPass* Pass)
 {
 	QUICK_SCOPE_CYCLE_COUNTER(STAT_FRDGBuilder_ExecutePass);
 
-	{
-		const bool bAllowAccess = true;
-		UpdateAccessGuardForPassResources(Pass, bAllowAccess);
-	}
+	IF_RDG_ENABLE_DEBUG(Validation.ValidateExecutePassBegin(Pass));
 
 	FRHIRenderPassInfo RPInfo;
 	bool bHasRenderTargets = false;
@@ -789,12 +515,7 @@ void FRDGBuilder::ExecutePass(const FRDGPass* Pass)
 
 	EventScopeStack.EndExecutePass();
 
-	{
-		const bool bAllowAccess = false;
-		UpdateAccessGuardForPassResources(Pass, bAllowAccess);
-	}
-
-	UnmarkUsedResources(Pass);
+	IF_RDG_ENABLE_DEBUG(Validation.ValidateExecutePassEnd(Pass));
 
 	// Can't release resources with immediate mode, because don't know if whether they are gonna be used.
 	if (!GRDGImmediateMode)
@@ -812,8 +533,7 @@ void FRDGBuilder::PrepareResourcesForExecute(const FRDGPass* Pass, struct FRHIRe
 
 	const bool bIsCompute = Pass->IsCompute();
 
-	FRDGBarrierBatcher BarrierBatcher;
-	BarrierBatcher.Begin();
+	FRDGBarrierBatcher BarrierBatcher(RHICmdList, Pass);
 
 	// NOTE: When generating mips, we don't perform any transitions on textures. They are done implicitly by the RHI.
 	const bool bGeneratingMips = Pass->IsGenerateMips();
@@ -832,21 +552,10 @@ void FRDGBuilder::PrepareResourcesForExecute(const FRDGPass* Pass, struct FRHIRe
 		{
 			if (FRDGTextureRef Texture = Parameter.GetAsTexture())
 			{
-				// The underlying texture have already been allocated by a prior pass.
-				#if RDG_ENABLE_DEBUG
-				{
-					check(Texture->HasBeenProduced());
-					Texture->PassAccessCount++;
-				}
-				#endif
-
 				check(Texture->PooledRenderTarget);
 				check(Texture->ResourceRHI);
 
-				if (!bGeneratingMips)
-				{
-					BarrierBatcher.QueueTransitionTexture(Texture, FRDGResourceState::CreateRead(Pass));
-				}
+				BarrierBatcher.QueueTransitionTexture(Texture, FRDGResourceState::EAccess::Read);
 			}
 		}
 		break;
@@ -856,24 +565,13 @@ void FRDGBuilder::PrepareResourcesForExecute(const FRDGPass* Pass, struct FRHIRe
 			{
 				FRDGTextureRef Texture = SRV->Desc.Texture;
 
-				// The underlying texture have already been allocated by a prior pass.
-				#if RDG_ENABLE_DEBUG
-				{
-					check(Texture->HasBeenProduced());
-					Texture->PassAccessCount++;
-				}
-				#endif
-
 				// Might be the first time using this render graph SRV, so need to setup the cached rhi resource.
 				if (!SRV->ResourceRHI)
 				{
 					SRV->ResourceRHI = Texture->PooledRenderTarget->GetRenderTargetItem().MipSRVs[SRV->Desc.MipLevel];
 				}
 
-				if (!bGeneratingMips)
-				{
-					BarrierBatcher.QueueTransitionTexture(Texture, FRDGResourceState::CreateRead(Pass));
-				}
+				BarrierBatcher.QueueTransitionTexture(Texture, FRDGResourceState::EAccess::Read);
 			}
 		}
 		break;
@@ -883,12 +581,6 @@ void FRDGBuilder::PrepareResourcesForExecute(const FRDGPass* Pass, struct FRHIRe
 			{
 				FRDGTextureRef Texture = UAV->Desc.Texture;
 	
-				#if RDG_ENABLE_DEBUG
-				{
-					Texture->PassAccessCount++;
-				}
-				#endif
-
 				AllocateRHITextureUAVIfNeeded(UAV);
 
 				FRHIUnorderedAccessView* UAVRHI = UAV->GetRHI();
@@ -898,7 +590,7 @@ void FRDGBuilder::PrepareResourcesForExecute(const FRDGPass* Pass, struct FRHIRe
 					OutRPInfo->UAVs[OutRPInfo->NumUAVs++] = UAVRHI;	// Bind UAVs in declaration order
 				}
 
-				BarrierBatcher.QueueTransitionUAV(UAVRHI, Texture, FRDGResourceState::CreateWrite(Pass));
+				BarrierBatcher.QueueTransitionUAV(UAVRHI, Texture, FRDGResourceState::EAccess::Write);
 			}
 		}
 		break;
@@ -906,20 +598,12 @@ void FRDGBuilder::PrepareResourcesForExecute(const FRDGPass* Pass, struct FRHIRe
 		{
 			if (FRDGBufferRef Buffer = Parameter.GetAsBuffer())
 			{
-				// The underlying buffer have already been allocated by a prior pass through AllocateRHIBufferUAVIfNeeded().
-				#if RDG_ENABLE_DEBUG
-				{
-					check(Buffer->HasBeenProduced());
-					Buffer->PassAccessCount++;
-				}
-				#endif
-
 				// TODO(RDG): super hacky, find the UAV and transition it. Hopefully there is one...
 				check(Buffer->PooledBuffer);
 				check(Buffer->PooledBuffer->UAVs.Num() == 1);
 				FRHIUnorderedAccessView* UAVRHI = Buffer->PooledBuffer->UAVs.CreateIterator().Value();
 
-				BarrierBatcher.QueueTransitionUAV(UAVRHI, Buffer, FRDGResourceState::CreateRead(Pass));
+				BarrierBatcher.QueueTransitionUAV(UAVRHI, Buffer, FRDGResourceState::EAccess::Read);
 			}
 		}
 		break;
@@ -929,14 +613,6 @@ void FRDGBuilder::PrepareResourcesForExecute(const FRDGPass* Pass, struct FRHIRe
 			{
 				FRDGBufferRef Buffer = SRV->Desc.Buffer;
 
-				// The underlying buffer have already been allocated by a prior pass through AllocateRHIBufferUAVIfNeeded().
-				#if RDG_ENABLE_DEBUG
-				{
-					check(Buffer->HasBeenProduced());
-					Buffer->PassAccessCount++;
-				}
-				#endif
-
 				AllocateRHIBufferSRVIfNeeded(SRV);
 
 				// TODO(RDG): super hacky, find the UAV and transition it. Hopefully there is one...
@@ -944,7 +620,7 @@ void FRDGBuilder::PrepareResourcesForExecute(const FRDGPass* Pass, struct FRHIRe
 				check(Buffer->PooledBuffer->UAVs.Num() == 1);
 				FRHIUnorderedAccessView* UAVRHI = Buffer->PooledBuffer->UAVs.CreateIterator().Value();
 
-				BarrierBatcher.QueueTransitionUAV(UAVRHI, Buffer, FRDGResourceState::CreateRead(Pass));
+				BarrierBatcher.QueueTransitionUAV(UAVRHI, Buffer, FRDGResourceState::EAccess::Read);
 			}
 		}
 		break;
@@ -954,12 +630,6 @@ void FRDGBuilder::PrepareResourcesForExecute(const FRDGPass* Pass, struct FRHIRe
 			{
 				FRDGBufferRef Buffer = UAV->Desc.Buffer;
 
-				#if RDG_ENABLE_DEBUG
-				{
-					Buffer->PassAccessCount++;
-				}
-				#endif
-
 				AllocateRHIBufferUAVIfNeeded(UAV);
 
 				FRHIUnorderedAccessView* UAVRHI = UAV->GetRHI();
@@ -969,14 +639,12 @@ void FRDGBuilder::PrepareResourcesForExecute(const FRDGPass* Pass, struct FRHIRe
 					OutRPInfo->UAVs[OutRPInfo->NumUAVs++] = UAVRHI;	// Bind UAVs in declaration order
 				}
 
-				BarrierBatcher.QueueTransitionUAV(UAVRHI, Buffer, FRDGResourceState::CreateWrite(Pass));
+				BarrierBatcher.QueueTransitionUAV(UAVRHI, Buffer, FRDGResourceState::EAccess::Write);
 			}
 		}
 		break;
 		case UBMT_RENDER_TARGET_BINDING_SLOTS:
 		{
-			check(!Pass->IsCompute());
-
 			const FRenderTargetBindingSlots& RenderTargetBindingSlots = Parameter.GetAsRenderTargetBindingSlots();
 			const auto& RenderTargets = RenderTargetBindingSlots.Output;
 			const auto& DepthStencil = RenderTargetBindingSlots.DepthStencil;
@@ -1022,19 +690,10 @@ void FRDGBuilder::PrepareResourcesForExecute(const FRDGPass* Pass, struct FRHIRe
 					OutRenderTarget.MipIndex = RenderTarget.GetMipIndex();
 					OutRenderTarget.Action = MakeRenderTargetActions(RenderTarget.GetLoadAction(), StoreAction);
 
-					if (!bGeneratingMips)
-					{
-						BarrierBatcher.QueueTransitionTexture(Texture, FRDGResourceState::CreateWrite(Pass));
-					}
+					BarrierBatcher.QueueTransitionTexture(Texture, FRDGResourceState::EAccess::Write);
 
 					SampleCount |= OutRenderTarget.RenderTarget->GetNumSamples();
 					ValidRenderTargetCount++;
-
-					#if RDG_ENABLE_DEBUG
-					{
-						Texture->PassAccessCount++;
-					}
-					#endif
 				}
 				else
 				{
@@ -1061,17 +720,11 @@ void FRDGBuilder::PrepareResourcesForExecute(const FRDGPass* Pass, struct FRHIRe
 
 				BarrierBatcher.QueueTransitionTexture(Texture,
 					DepthStencil.GetDepthStencilAccess().IsAnyWrite() ?
-						FRDGResourceState::CreateWrite(Pass) :
-						FRDGResourceState::CreateRead(Pass));
-	
+					FRDGResourceState::EAccess::Write :
+					FRDGResourceState::EAccess::Read);
+
 				SampleCount |= OutDepthStencil.DepthStencilTarget->GetNumSamples();
 				ValidDepthStencilCount++;
-
-				#if RDG_ENABLE_DEBUG
-				{
-					Texture->PassAccessCount++;
-				}
-				#endif
 			}
 
 			OutRPInfo->bIsMSAA = SampleCount > 1;
@@ -1085,123 +738,6 @@ void FRDGBuilder::PrepareResourcesForExecute(const FRDGPass* Pass, struct FRHIRe
 	}
 
 	OutRPInfo->bGeneratingMips = bGeneratingMips;
-
-	BarrierBatcher.End(RHICmdList);
-}
-
-void FRDGBuilder::UpdateAccessGuardForPassResources(const FRDGPass* Pass, bool bAllowAccess)
-{
-#if RDG_ENABLE_DEBUG
-	FRDGPassParameterStruct ParameterStruct = Pass->GetParameters();
-
-	const uint32 ParameterCount = ParameterStruct.GetParameterCount();
-
-	for (uint32 ParameterIndex = 0; ParameterIndex < ParameterCount; ++ParameterIndex)
-	{
-		FRDGPassParameter Parameter = ParameterStruct.GetParameter(ParameterIndex);
-
-		if (Parameter.IsResource())
-		{
-			if (FRDGResourceRef Resource = Parameter.GetAsResource())
-			{
-				Resource->bAllowRHIAccess = bAllowAccess;
-			}
-		}
-		else if (Parameter.GetType() == UBMT_RENDER_TARGET_BINDING_SLOTS)
-		{
-			const FRenderTargetBindingSlots& RenderTargetBindingSlots = Parameter.GetAsRenderTargetBindingSlots();
-			const auto& RenderTargets = RenderTargetBindingSlots.Output;
-			const auto& DepthStencil = RenderTargetBindingSlots.DepthStencil;
-			const uint32 RenderTargetCount = RenderTargets.Num();
-
-			for (uint32 RenderTargetIndex = 0; RenderTargetIndex < RenderTargetCount; RenderTargetIndex++)
-			{
-				const FRenderTargetBinding& RenderTarget = RenderTargets[RenderTargetIndex];
-
-				if (FRDGTextureRef Texture = RenderTarget.GetTexture())
-				{
-					Texture->bAllowRHIAccess = bAllowAccess;
-				}
-				else
-				{
-					break;
-				}
-			}
-			
-			if (FRDGTextureRef Texture = DepthStencil.GetTexture())
-			{
-				Texture->bAllowRHIAccess = bAllowAccess;
-			}
-		}
-	}
-#endif
-}
-
-void FRDGBuilder::UnmarkUsedResources(const FRDGPass* Pass)
-{
-#if RDG_ENABLE_DEBUG
-	FRDGPassParameterStruct ParameterStruct = Pass->GetParameters();
-
-	const uint32 ParameterCount = ParameterStruct.GetParameterCount();
-
-	if (GRDGDebug)
-	{
-		uint32 TrackedResourceCount = 0;
-		uint32 UsedResourceCount = 0;
-
-		for (uint32 ParameterIndex = 0; ParameterIndex < ParameterCount; ++ParameterIndex)
-		{
-			FRDGPassParameter Parameter = ParameterStruct.GetParameter(ParameterIndex);
-
-			if (Parameter.IsResource())
-			{
-				if (FRDGResourceRef Resource = Parameter.GetAsResource())
-				{
-					TrackedResourceCount++;
-					UsedResourceCount += Resource->bIsActuallyUsedByPass ? 1 : 0;
-				}
-			}
-		}
-
-		if (TrackedResourceCount != UsedResourceCount)
-		{
-			FString WarningMessage = FString::Printf(
-				TEXT("'%d' of the '%d' resources of the pass '%s' where not actually used."),
-				TrackedResourceCount - UsedResourceCount, TrackedResourceCount, Pass->GetName());
-
-			for (uint32 ParameterIndex = 0; ParameterIndex < ParameterCount; ++ParameterIndex)
-			{
-				FRDGPassParameter Parameter = ParameterStruct.GetParameter(ParameterIndex);
-
-				if (Parameter.IsResource())
-				{
-					if (const FRDGResourceRef Resource = Parameter.GetAsResource())
-					{
-						if (!Resource->bIsActuallyUsedByPass)
-						{
-							WarningMessage += FString::Printf(TEXT("\n    %s"), Resource->Name);
-						}
-					}
-				}
-			}
-
-			EmitRDGWarning(WarningMessage);
-		}
-	}
-
-	for (uint32 ParameterIndex = 0; ParameterIndex < ParameterCount; ++ParameterIndex)
-	{
-		FRDGPassParameter Parameter = ParameterStruct.GetParameter(ParameterIndex);
-
-		if (Parameter.IsResource())
-		{
-			if (const FRDGResourceRef Resource = Parameter.GetAsResource())
-			{
-				Resource->bIsActuallyUsedByPass = false;
-			}
-		}
-	}
-#endif
 }
 
 void FRDGBuilder::ReleaseRHITextureIfUnreferenced(FRDGTexture* Texture)
@@ -1325,8 +861,7 @@ void FRDGBuilder::ReleaseUnreferencedResources(const FRDGPass* Pass)
 
 void FRDGBuilder::ProcessDeferredInternalResourceQueries()
 {
-	FRDGBarrierBatcher BarrierBatcher;
-	BarrierBatcher.Begin();
+	FRDGBarrierBatcher BarrierBatcher(RHICmdList, nullptr);
 
 	for (const auto& Query : DeferredInternalTextureQueries)
 	{
@@ -1334,25 +869,11 @@ void FRDGBuilder::ProcessDeferredInternalResourceQueries()
 
 		if (Query.bTransitionToRead)
 		{
-			BarrierBatcher.QueueTransitionTexture(
-				Query.Texture,
-				FRDGResourceState(
-					nullptr,
-					FRDGResourceState::EPipeline::Graphics,
-					FRDGResourceState::EAccess::Read));
+			BarrierBatcher.QueueTransitionTexture(Query.Texture, FRDGResourceState::EAccess::Read);
 		}
 
 		*Query.OutTexturePtr = AllocatedTextures.FindChecked(Query.Texture);
-		
-		#if RDG_ENABLE_DEBUG
-		{
-			// Increment the number of times the texture has been accessed to avoid warning on produced but never used resources that were produced
-			// only to be extracted for the graph.
-			Query.Texture->PassAccessCount += 1;
-		}
-		#endif
 
-		// No need to manually release in immediate mode, since it is done directly when emptying AllocatedTextures in DestructPasses().
 		if (!GRDGImmediateMode)
 		{
 			ReleaseRHITextureIfUnreferenced(Query.Texture);
@@ -1363,50 +884,15 @@ void FRDGBuilder::ProcessDeferredInternalResourceQueries()
 	{
 		*Query.OutBufferPtr = AllocatedBuffers.FindChecked(Query.Buffer);
 
-		#if RDG_ENABLE_DEBUG
-		{
-			// Increment the number of times the buffer has been accessed to avoid warning on produced but never used resources that were produced
-			// only to be extracted for the graph.
-			Query.Buffer->PassAccessCount += 1;
-		}
-		#endif
-
-		// No need to manually release in immediate mode, since it is done directly when emptying AllocatedBuffer in DestructPasses().
 		if (!GRDGImmediateMode)
 		{
 			ReleaseRHIBufferIfUnreferenced(Query.Buffer);
 		}
 	}
-
-	BarrierBatcher.End(RHICmdList);
 }
 
 void FRDGBuilder::DestructPasses()
 {
-	#if RDG_ENABLE_DEBUG
-	{
-		const bool bEmitWarnings = (GRDGDebug > 0);
-
-		for (const FRDGTrackedResourceRef Resource : TrackedResources)
-		{
-			check(Resource->ReferenceCount == 0);
-
-			const bool bProducedButNeverUsed = Resource->PassAccessCount == 1 && Resource->FirstProducer;
-
-			if (bEmitWarnings && bProducedButNeverUsed)
-			{
-				check(Resource->HasBeenProduced());
-
-				EmitRDGWarningf(
-					TEXT("Resources %s has been produced by the pass %s, but never used by another pass."),
-					Resource->Name, Resource->FirstProducer->GetName());
-			}
-		}
-
-		TrackedResources.Empty();
-	}
-	#endif
-
 	for (int32 PassIndex = Passes.Num() - 1; PassIndex >= 0; --PassIndex)
 	{
 		Passes[PassIndex]->~FRDGPass();
diff --git a/Engine/Source/Runtime/RenderCore/Private/RenderGraphResources.cpp b/Engine/Source/Runtime/RenderCore/Private/RenderGraphResources.cpp
deleted file mode 100644
diff --git a/Engine/Source/Runtime/RenderCore/Private/RenderGraphValidation.cpp b/Engine/Source/Runtime/RenderCore/Private/RenderGraphValidation.cpp
new file mode 100644
--- /dev/null
+++ b/Engine/Source/Runtime/RenderCore/Private/RenderGraphValidation.cpp
@@ -0,0 +1,629 @@
+// Copyright 1998-2019 Epic Games, Inc. All Rights Reserved.
+
+#include "RenderGraphValidation.h"
+
+#if RDG_ENABLE_DEBUG
+
+/** This utility class tracks read / write events on individual resources within a pass in order to ensure
+ *  that sub-resources are not being simultaneously bound for read / write. The RHI currently only supports
+ *  disjoint read-write access on a texture for mip-map generation purposes.
+ */
+class FRDGPassResourceValidator
+{
+public:
+	static const uint32 AllMipLevels = ~0;
+
+	FRDGPassResourceValidator(const FRDGPass& InPass)
+		: Pass(InPass)
+		, bIsGenerateMips(InPass.IsGenerateMips())
+	{}
+
+	void Read(FRDGBufferRef Buffer)
+	{
+		check(Buffer);
+
+		ensureMsgf(!bIsGenerateMips,
+			TEXT("Attempting to read from Buffer %s on Pass %s which is marked as 'GenerateMips'. Only a single texture is supported on this type of pass."),
+			Buffer->Name, Pass.GetName());
+
+		FBufferAccess& BufferAccess = BufferAccessMap.FindOrAdd(Buffer);
+
+		ensureMsgf(BufferAccess.Access != FRDGResourceState::EAccess::Write,
+			TEXT("Buffer %s is being used for both read and write access on Pass %s. This is not supported."),
+			Buffer->Name, Pass.GetName());
+
+		BufferAccess.Access = FRDGResourceState::EAccess::Read;
+	}
+
+	void Write(FRDGBufferRef Buffer)
+	{
+		check(Buffer);
+
+		ensureMsgf(!bIsGenerateMips,
+			TEXT("Attempting to write to Buffer %s on Pass %s which is marked as 'GenerateMips'. Only a single texture is supported on this type of pass."),
+			Buffer->Name, Pass.GetName());
+
+		FBufferAccess& BufferAccess = BufferAccessMap.FindOrAdd(Buffer);
+
+		ensureMsgf(BufferAccess.Access != FRDGResourceState::EAccess::Read,
+			TEXT("Attempting to simultaneously read and write from Buffer %s on Pass %s. This is not supported."),
+			Buffer->Name, Pass.GetName());
+
+		BufferAccess.Access = FRDGResourceState::EAccess::Read;
+	}
+
+	void Read(FRDGTextureRef Texture, uint32 MipLevel = AllMipLevels)
+	{
+		check(Texture);
+
+		const uint32 MipMask = (MipLevel == AllMipLevels) ? AllMipLevels : (1 << MipLevel);
+
+		FTextureAccess& TextureAccess = TextureAccessMap.FindOrAdd(Texture);
+		TextureAccess.SetAllRead(MipMask);
+
+		ensureMsgf(!TextureAccess.IsAnyWrite(MipMask),
+			TEXT("Attempting to access mips of Texture %s for both read and write access on Pass %s. This is not supported."),
+			MipLevel, Texture->Name, Pass.GetName());
+
+		if (bIsGenerateMips)
+		{
+			ensureMsgf(TextureAccessMap.Num() == 1,
+				TEXT("Attempting to access more than one texture in Pass %s when 'GenerateMips' is set. Only a single texture is supported on this type of pass."),
+				Pass.GetName());
+		}
+		else
+		{
+			ensureMsgf(!TextureAccess.Mips.Write,
+				TEXT("Texture %s is being used for both read and write access on the same Pass %s. This is only supported via the 'GenerateMips' flag."),
+				Pass.GetName());
+		}
+	}
+
+	void Write(FRDGTextureRef Texture, uint32 MipLevel)
+	{
+		check(Texture);
+
+		const uint32 MipMask = 1 << MipLevel;
+
+		FTextureAccess& TextureAccess = TextureAccessMap.FindOrAdd(Texture);
+		TextureAccess.SetAllWrite(MipMask);
+
+		ensureMsgf(!TextureAccess.IsAnyRead(MipMask),
+			TEXT("Attempting to use Mip %d of Texture %s for both read and write access on Pass %s. This is not supported."),
+			MipLevel, Texture->Name, Pass.GetName());
+
+		if (bIsGenerateMips)
+		{
+			ensureMsgf(TextureAccessMap.Num() == 1,
+				TEXT("Attempting to access more than one texture in Pass %s when 'GenerateMips' is set. Only a single texture is supported on this type of pass."),
+				Pass.GetName());
+		}
+		else
+		{
+			ensureMsgf(!TextureAccess.Mips.Read,
+				TEXT("Texture %s is being used for both read and write access on the same Pass %s. This is only supported via the 'GenerateMips' flag."),
+				Texture->Name, Pass.GetName());
+		}
+	}
+
+private:
+	struct FBufferAccess
+	{
+		FRDGResourceState::EAccess Access = FRDGResourceState::EAccess::Unknown;
+
+		FORCEINLINE bool operator!= (FBufferAccess Other) const
+		{
+			return Access != Other.Access;
+		}
+
+		FORCEINLINE bool operator< (FBufferAccess Other) const
+		{
+			return Access < Other.Access;
+		}
+	};
+
+	struct FTextureAccess
+	{
+		FTextureAccess() = default;
+
+		union
+		{
+			uint64 PackedData = 0;
+
+			struct
+			{
+				uint32 Read;
+				uint32 Write;
+			} Mips;
+		};
+
+		FORCEINLINE bool IsAnyRead(uint32 MipMask) const
+		{
+			return (Mips.Read & MipMask) != 0;
+		}
+
+		FORCEINLINE bool IsAnyWrite(uint32 MipMask) const
+		{
+			return (Mips.Write & MipMask) != 0;
+		}
+
+		FORCEINLINE void SetAllRead(uint32 MipMask)
+		{
+			Mips.Read |= MipMask;
+		}
+
+		FORCEINLINE void SetAllWrite(uint32 MipMask)
+		{
+			Mips.Write |= MipMask;
+		}
+
+		FORCEINLINE bool operator!= (FTextureAccess Other) const
+		{
+			return PackedData != Other.PackedData;
+		}
+
+		FORCEINLINE bool operator< (FTextureAccess Other) const
+		{
+			return PackedData < Other.PackedData;
+		}
+	};
+
+	TMap<FRDGBufferRef, FBufferAccess, SceneRenderingSetAllocator> BufferAccessMap;
+	TMap<FRDGTextureRef, FTextureAccess, SceneRenderingSetAllocator> TextureAccessMap;
+	const FRDGPass& Pass;
+	const bool bIsGenerateMips;
+};
+
+FRDGUserValidation::~FRDGUserValidation()
+{
+	checkf(bHasExecuted, TEXT("Render graph execution is required to ensure consistency with immediate mode."));
+}
+
+void FRDGUserValidation::ExecuteGuard(const TCHAR* Operation, const TCHAR* ResourceName)
+{
+	checkf(!bHasExecuted, TEXT("Render graph operation '%s' with resource '%s' must be performed prior to graph execution."), Operation, ResourceName);
+}
+
+void FRDGUserValidation::ValidateCreateResource(FRDGTrackedResourceRef Resource)
+{
+	check(Resource);
+	TrackedResources.Add(Resource);
+}
+
+void FRDGUserValidation::ValidateCreateExternalResource(FRDGTrackedResourceRef Resource)
+{
+	ValidateCreateResource(Resource);
+	Resource->MarkAsExternal();
+}
+
+void FRDGUserValidation::ValidateExtractResource(FRDGTrackedResourceRef Resource)
+{
+	check(Resource);
+
+	checkf(Resource->HasBeenProduced(),
+		TEXT("Unable to queue the extraction of the resource %s because it has not been produced by any pass."),
+		Resource->Name);
+
+	/** Increment pass access counts for externally registered buffers and textures to avoid
+	 *  emitting a 'produced but never used' warning. We don't have the history of registered
+	 *  resources to be able to emit a proper warning.
+	 */
+
+	Resource->PassAccessCount++;
+}
+
+void FRDGUserValidation::RemoveUnusedWarning(FRDGTrackedResourceRef Resource)
+{
+	check(Resource);
+	Resource->PassAccessCount++;
+}
+
+void FRDGUserValidation::ValidateAllocPassParameters(const void* Parameters)
+{
+	check(Parameters);
+	AllocatedUnusedPassParameters.Add(Parameters);
+}
+
+void FRDGUserValidation::ValidateAddPass(const FRDGPass* Pass)
+{
+	checkf(!bHasExecuted, TEXT("Render graph pass %s needs to be added before the builder execution."), Pass->GetName());
+
+	{
+		const void* ParameterStructData = Pass->GetParameters().GetContents();
+
+		/** The lifetime of each pass parameter structure must extend until deferred pass execution; therefore, it needs to be
+		 *  allocated with FRDGBuilder::AllocParameters(). Also, all references held by the parameter structure are released
+		 *  immediately after pass execution, so a pass parameter struct instance must be 1-to-1 with a pass instance (i.e. one
+		 *  per AddPass() call).
+		 */
+		checkf(
+			AllocatedUnusedPassParameters.Contains(ParameterStructData),
+			TEXT("The pass parameter structure has not been allocated for correct life time FRDGBuilder::AllocParameters() or has already ")
+			TEXT("been used by another previous FRDGBuilder::AddPass()."));
+
+		AllocatedUnusedPassParameters.Remove(ParameterStructData);
+	}
+
+	const FRenderTargetBindingSlots* RenderTargetBindingSlots = nullptr;
+
+	const TCHAR* PassName = Pass->GetName();
+	const bool bIsCompute = Pass->IsCompute();
+	const bool bIsGeneratingMips = Pass->IsGenerateMips();
+	const bool bRenderTargetSlotsRequired = !bIsCompute;
+
+	FRDGPassResourceValidator PassResourceValidator(*Pass);
+
+	FRDGPassParameterStruct ParameterStruct = Pass->GetParameters();
+
+	const uint32 ParameterCount = ParameterStruct.GetParameterCount();
+
+	for (uint32 ParameterIndex = 0; ParameterIndex < ParameterCount; ++ParameterIndex)
+	{
+		FRDGPassParameter Parameter = ParameterStruct.GetParameter(ParameterIndex);
+
+		switch (Parameter.GetType())
+		{
+		case UBMT_RDG_TEXTURE:
+		{
+			if (FRDGTextureRef Texture = Parameter.GetAsTexture())
+			{
+				PassResourceValidator.Read(Texture);
+
+				checkf(Texture->HasBeenProduced(),
+					TEXT("Pass %s has a dependency over the texture %s that has never been produced."),
+					PassName, Texture->Name);
+
+				Texture->PassAccessCount++;
+			}
+		}
+		break;
+		case UBMT_RDG_TEXTURE_SRV:
+		{
+			if (FRDGTextureSRVRef SRV = Parameter.GetAsTextureSRV())
+			{
+				FRDGTextureRef Texture = SRV->Desc.Texture;
+
+				PassResourceValidator.Read(Texture, SRV->Desc.MipLevel);
+
+				checkf(Texture->HasBeenProduced(),
+					TEXT("Pass %s has a dependency over the texture %s that has never been produced."),
+					PassName, Texture->Name);
+
+				Texture->PassAccessCount++;
+			}
+		}
+		break;
+		case UBMT_RDG_TEXTURE_UAV:
+		{
+			if (FRDGTextureUAVRef UAV = Parameter.GetAsTextureUAV())
+			{
+				FRDGTextureRef Texture = UAV->Desc.Texture;
+
+				PassResourceValidator.Write(Texture, UAV->Desc.MipLevel);
+
+				Texture->MarkAsProducedBy(Pass);
+				Texture->PassAccessCount++;
+			}
+		}
+		break;
+		case UBMT_RDG_BUFFER:
+		{
+			if (FRDGBufferRef Buffer = Parameter.GetAsBuffer())
+			{
+				PassResourceValidator.Read(Buffer);
+
+				checkf(Buffer->HasBeenProduced(),
+					TEXT("Pass %s has a dependency over the buffer %s that has never been produced."),
+					PassName, Buffer->Name);
+
+				Buffer->PassAccessCount++;
+			}
+		}
+		break;
+		case UBMT_RDG_BUFFER_SRV:
+		{
+			if (FRDGBufferSRVRef SRV = Parameter.GetAsBufferSRV())
+			{
+				FRDGBufferRef Buffer = SRV->Desc.Buffer;
+
+				PassResourceValidator.Read(Buffer);
+
+				checkf(Buffer->HasBeenProduced(),
+					TEXT("Pass %s has a dependency over the buffer %s that has never been produced."),
+					PassName, SRV->Desc.Buffer->Name);
+
+				Buffer->PassAccessCount++;
+			}
+		}
+		break;
+		case UBMT_RDG_BUFFER_UAV:
+		{
+			if (FRDGBufferUAVRef UAV = Parameter.GetAsBufferUAV())
+			{
+				FRDGBufferRef Buffer = UAV->Desc.Buffer;
+
+				PassResourceValidator.Write(Buffer);
+
+				Buffer->MarkAsProducedBy(Pass);
+				Buffer->PassAccessCount++;
+			}
+		}
+		break;
+		case UBMT_RENDER_TARGET_BINDING_SLOTS:
+		{
+			if (!RenderTargetBindingSlots)
+			{
+				RenderTargetBindingSlots = &Parameter.GetAsRenderTargetBindingSlots();
+			}
+			else if (IsRDGDebugEnabled())
+			{
+				EmitRDGWarningf(
+					TEXT("Pass %s have duplicated render target binding slots."),
+					PassName);
+			}
+		}
+		break;
+		default:
+			break;
+		}
+	}
+
+	/** Validate that raster passes have render target binding slots and compute passes don't. */
+	if (RenderTargetBindingSlots)
+	{
+		checkf(!bIsCompute, TEXT("Pass '%s' has render target binding slots but is flagged as 'Compute'."), PassName);
+	}
+	else
+	{
+		checkf(bIsCompute, TEXT("Pass '%s' is missing render target binding slots. Add the 'Compute' flag if render targets are not required."), PassName);
+	}
+
+	/** Validate render target / depth stencil binding usage. */
+	if (RenderTargetBindingSlots)
+	{
+		const auto& RenderTargets = RenderTargetBindingSlots->Output;
+
+		{
+			const auto& DepthStencil = RenderTargetBindingSlots->DepthStencil;
+
+			if (FRDGTextureRef Texture = DepthStencil.GetTexture())
+			{
+				ensureMsgf(Texture->Desc.NumSamples == 1,
+					TEXT("Pass '%s' uses an MSAA depth-stencil render target. This is not yet supported."),
+					Pass->GetName());
+
+				ensureMsgf(!bIsGeneratingMips,
+					TEXT("Pass '%s' is marked to generate mips but has a depth stencil texture. This is not supported."),
+					Pass->GetName());
+
+				// Depth stencil only supports one mip, since there isn't actually a way to select the mip level.
+				const uint32 MipLevel = 0;
+				check(Texture->Desc.NumMips == 1);
+
+				if (DepthStencil.GetDepthStencilAccess().IsAnyWrite())
+				{
+					PassResourceValidator.Write(Texture, MipLevel);
+					Texture->MarkAsProducedBy(Pass);
+				}
+				else
+				{
+					PassResourceValidator.Read(Texture, MipLevel);
+				}
+
+				Texture->PassAccessCount++;
+			}
+		}
+
+		const uint32 RenderTargetCount = RenderTargets.Num();
+
+		{
+			/** Tracks the number of contiguous, non-null textures in the render target output array. */
+			uint32 ValidRenderTargetCount = 0;
+
+			for (uint32 RenderTargetIndex = 0; RenderTargetIndex < RenderTargetCount; ++RenderTargetIndex)
+			{
+				const FRenderTargetBinding& RenderTarget = RenderTargets[RenderTargetIndex];
+
+				if (FRDGTextureRef Texture = RenderTarget.GetTexture())
+				{
+					PassResourceValidator.Write(Texture, RenderTarget.GetMipIndex());
+
+					/** Validate that load action is correct. */
+					const bool bIsLoadAction = RenderTarget.GetLoadAction() == ERenderTargetLoadAction::ELoad;
+					const bool bIsLoadActionInvalid = bIsLoadAction && !Texture->HasBeenProduced();
+					checkf(
+						!bIsLoadActionInvalid,
+						TEXT("Can't load a render target '%s' that has never been produced."),
+						Texture->Name);
+					
+					const bool bOverwritesResourceEntirely = !bIsLoadAction && Texture->HasBeenProduced() && Texture->Desc.NumMips == 1;
+					ensureMsgf(!bOverwritesResourceEntirely,
+						TEXT("Clobbering render target %s result that has been produced previously by another pass. Should instead render to a new FRDGTexture to reduce memory pressure."),
+						Texture->Name);
+					
+					/** Mark the pass as a producer for render targets with a store action. */
+					{
+						const bool bIsStoreAction = RenderTarget.GetStoreAction() != ERenderTargetStoreAction::ENoAction;
+						check(bIsStoreAction); // already been validated in FRenderTargetBinding::Validate()
+						Texture->MarkAsProducedBy(Pass);
+					}
+
+					Texture->PassAccessCount++;
+				}
+				else
+				{
+					/** Found end of contiguous interval of valid render targets. */
+					ValidRenderTargetCount = RenderTargetIndex;
+					break;
+				}
+			}
+
+			/** Validate that no holes exist in the render target output array. Render targets must be bound contiguously. */
+			for (uint32 RenderTargetIndex = ValidRenderTargetCount; RenderTargetIndex < RenderTargetCount; ++RenderTargetIndex)
+			{
+				const FRenderTargetBinding& RenderTarget = RenderTargets[RenderTargetIndex];
+				checkf(RenderTarget.GetTexture() == nullptr, TEXT("Render targets must be packed. No empty spaces in the array."));
+			}
+		}
+	}
+}
+
+void FRDGUserValidation::ValidateExecuteBegin()
+{
+	checkf(!bHasExecuted, TEXT("Render graph execution should only happen once to ensure consistency with immediate mode."));
+
+	/** FRDGBuilder::AllocParameters() allocates shader parameter structure for the lifetime until pass execution.
+	 *  But they are allocated on a FMemStack for CPU performance, and are destructed immediately after pass execution.
+	 *  Therefore allocating pass parameter unused by a FRDGBuilder::AddPass() can lead on a memory leak of RHI resources
+	 *  that have been reference in the parameter structure.
+	 */
+	checkf(
+		AllocatedUnusedPassParameters.Num() == 0,
+		TEXT("%i pass parameter structure has been allocated with FRDGBuilder::AllocParameters(), but has not be used by a ")
+		TEXT("FRDGBuilder::AddPass() that can cause RHI resource leak."), AllocatedUnusedPassParameters.Num());
+}
+
+void FRDGUserValidation::ValidateExecuteEnd()
+{
+	bHasExecuted = true;
+
+	const bool bEmitWarnings = IsRDGDebugEnabled();
+
+	for (const FRDGTrackedResourceRef Resource : TrackedResources)
+	{
+		check(Resource->ReferenceCount == 0);
+
+		const bool bProducedButNeverUsed = Resource->PassAccessCount == 1 && Resource->FirstProducer;
+
+		if (bEmitWarnings && bProducedButNeverUsed)
+		{
+			check(Resource->HasBeenProduced());
+
+			EmitRDGWarningf(
+				TEXT("Resources %s has been produced by the pass %s, but never used by another pass."),
+				Resource->Name, Resource->FirstProducer->GetName());
+		}
+	}
+
+	TrackedResources.Empty();
+}
+
+void FRDGUserValidation::ValidateExecutePassBegin(const FRDGPass* Pass)
+{
+	SetAllowRHIAccess(Pass, true);
+}
+
+void FRDGUserValidation::ValidateExecutePassEnd(const FRDGPass* Pass)
+{
+	SetAllowRHIAccess(Pass, false);
+
+	FRDGPassParameterStruct ParameterStruct = Pass->GetParameters();
+
+	const uint32 ParameterCount = ParameterStruct.GetParameterCount();
+
+	if (IsRDGDebugEnabled())
+	{
+		uint32 TrackedResourceCount = 0;
+		uint32 UsedResourceCount = 0;
+
+		for (uint32 ParameterIndex = 0; ParameterIndex < ParameterCount; ++ParameterIndex)
+		{
+			FRDGPassParameter Parameter = ParameterStruct.GetParameter(ParameterIndex);
+
+			if (Parameter.IsResource())
+			{
+				if (FRDGResourceRef Resource = Parameter.GetAsResource())
+				{
+					TrackedResourceCount++;
+					UsedResourceCount += Resource->bIsActuallyUsedByPass ? 1 : 0;
+				}
+			}
+		}
+
+		if (TrackedResourceCount != UsedResourceCount)
+		{
+			FString WarningMessage = FString::Printf(
+				TEXT("'%d' of the '%d' resources of the pass '%s' where not actually used."),
+				TrackedResourceCount - UsedResourceCount, TrackedResourceCount, Pass->GetName());
+
+			for (uint32 ParameterIndex = 0; ParameterIndex < ParameterCount; ++ParameterIndex)
+			{
+				FRDGPassParameter Parameter = ParameterStruct.GetParameter(ParameterIndex);
+
+				if (Parameter.IsResource())
+				{
+					if (const FRDGResourceRef Resource = Parameter.GetAsResource())
+					{
+						if (!Resource->bIsActuallyUsedByPass)
+						{
+							WarningMessage += FString::Printf(TEXT("\n    %s"), Resource->Name);
+						}
+					}
+				}
+			}
+
+			EmitRDGWarning(WarningMessage);
+		}
+	}
+
+	for (uint32 ParameterIndex = 0; ParameterIndex < ParameterCount; ++ParameterIndex)
+	{
+		FRDGPassParameter Parameter = ParameterStruct.GetParameter(ParameterIndex);
+
+		if (Parameter.IsResource())
+		{
+			if (const FRDGResourceRef Resource = Parameter.GetAsResource())
+			{
+				Resource->bIsActuallyUsedByPass = false;
+			}
+		}
+	}
+}
+
+void FRDGUserValidation::SetAllowRHIAccess(const FRDGPass* Pass, bool bAllowAccess)
+{
+	FRDGPassParameterStruct ParameterStruct = Pass->GetParameters();
+
+	const uint32 ParameterCount = ParameterStruct.GetParameterCount();
+
+	for (uint32 ParameterIndex = 0; ParameterIndex < ParameterCount; ++ParameterIndex)
+	{
+		FRDGPassParameter Parameter = ParameterStruct.GetParameter(ParameterIndex);
+
+		if (Parameter.IsResource())
+		{
+			if (FRDGResourceRef Resource = Parameter.GetAsResource())
+			{
+				Resource->bAllowRHIAccess = bAllowAccess;
+			}
+		}
+		else if (Parameter.GetType() == UBMT_RENDER_TARGET_BINDING_SLOTS)
+		{
+			const FRenderTargetBindingSlots& RenderTargetBindingSlots = Parameter.GetAsRenderTargetBindingSlots();
+			const auto& RenderTargets = RenderTargetBindingSlots.Output;
+			const auto& DepthStencil = RenderTargetBindingSlots.DepthStencil;
+			const uint32 RenderTargetCount = RenderTargets.Num();
+
+			for (uint32 RenderTargetIndex = 0; RenderTargetIndex < RenderTargetCount; RenderTargetIndex++)
+			{
+				const FRenderTargetBinding& RenderTarget = RenderTargets[RenderTargetIndex];
+
+				if (FRDGTextureRef Texture = RenderTarget.GetTexture())
+				{
+					Texture->bAllowRHIAccess = bAllowAccess;
+				}
+				else
+				{
+					break;
+				}
+			}
+
+			if (FRDGTextureRef Texture = DepthStencil.GetTexture())
+			{
+				Texture->bAllowRHIAccess = bAllowAccess;
+			}
+		}
+	}
+}
+
+#endif
\ No newline at end of file
diff --git a/Engine/Source/Runtime/RenderCore/Public/RenderGraphBuilder.h b/Engine/Source/Runtime/RenderCore/Public/RenderGraphBuilder.h
--- a/Engine/Source/Runtime/RenderCore/Public/RenderGraphBuilder.h
+++ b/Engine/Source/Runtime/RenderCore/Public/RenderGraphBuilder.h
@@ -6,6 +6,7 @@
 #include "RenderGraphEvent.h"
 #include "RenderGraphResources.h"
 #include "RenderGraphPass.h"
+#include "RenderGraphValidation.h"
 #include "ShaderParameterMacros.h"
 
 /** Builds the per-frame render graph.
@@ -17,7 +18,7 @@ class RENDERCORE_API FRDGBuilder
 public:
 	/** A RHI cmd list is required, if using the immediate mode. */
 	FRDGBuilder(FRHICommandListImmediate& InRHICmdList);
-	~FRDGBuilder();
+	FRDGBuilder(const FRDGBuilder&) = delete;
 
 	/** Register a external texture to be tracked by the render graph. */
 	FRDGTextureRef RegisterExternalTexture(
@@ -117,9 +118,12 @@ public:
 	 */
 	void QueueBufferExtraction(FRDGBufferRef Buffer, TRefCountPtr<FPooledRDGBuffer>* OutBufferPtr);
 
-	/** Flag a texture that is only produced by only 1 pass, but never used or extracted, to avoid generating a warning at runtime. */
+	/** Flag a texture that is produced by a pass but never used or extracted to not emit an 'unused' warning. */
 	void RemoveUnusedTextureWarning(FRDGTextureRef Texture);
 
+	/** Flag a buffer that is produced by a pass but never used or extracted to not emit an 'unused' warning. */
+	void RemoveUnusedBufferWarning(FRDGBufferRef Buffer);
+
 	/** Begins / ends a named event scope. These scopes are visible in most external GPU profilers.
 	 *  Prefer to use RDG_EVENT_SCOPE instead of calling this manually.
 	 */
@@ -179,17 +183,9 @@ private:
 	FRDGStatScopeStack StatScopeStack;
 
 #if RDG_ENABLE_DEBUG
-	/** All recently allocated pass parameter structure, but not used by a AddPass() yet. */
-	TSet<const void*> AllocatedUnusedPassParameters;
-
-	/** List of tracked resources for validation prior to shutdown. */
-	TArray<FRDGTrackedResourceRef, SceneRenderingAllocator> TrackedResources;
-
-	/** Whether the Execute() has already been called. */
-	bool bHasExecuted = false;
+	FRDGUserValidation Validation;
 #endif
 
-	void ValidatePass(const FRDGPass* Pass) const;
 	void VisualizePassOutputs(const FRDGPass* Pass);
 
 	void WalkGraphDependencies();
@@ -208,9 +204,6 @@ private:
 	void ExecutePass(const FRDGPass* Pass);
 	void PrepareResourcesForExecute(const FRDGPass* Pass, struct FRHIRenderPassInfo* OutRPInfo, bool* bOutHasRenderTargets);
 
-	void UpdateAccessGuardForPassResources(const FRDGPass* Pass, bool bAllowAccess);
-	void UnmarkUsedResources(const FRDGPass* Pass);
-
 	void ReleaseRHITextureIfUnreferenced(FRDGTexture* Texture);
 	void ReleaseRHIBufferIfUnreferenced(FRDGBuffer* Buffer);
 	void ReleaseUnreferencedResources(const FRDGPass* Pass);
diff --git a/Engine/Source/Runtime/RenderCore/Public/RenderGraphBuilder.inl b/Engine/Source/Runtime/RenderCore/Public/RenderGraphBuilder.inl
--- a/Engine/Source/Runtime/RenderCore/Public/RenderGraphBuilder.inl
+++ b/Engine/Source/Runtime/RenderCore/Public/RenderGraphBuilder.inl
@@ -6,6 +6,7 @@ inline FRDGTextureRef FRDGBuilder::RegisterExternalTexture(const TRefCountPtr<IP
 {
 #if RDG_ENABLE_DEBUG
 	{
+		Validation.ExecuteGuard(TEXT("RegisterExternalTexture"), Name);
 		checkf(ExternalPooledTexture.IsValid(), TEXT("Attempted to register NULL external texture: %s"), Name);
 		checkf(Name, TEXT("Externally allocated texture requires a debug name when registering them to render graph."));
 	}
@@ -21,12 +22,7 @@ inline FRDGTextureRef FRDGBuilder::RegisterExternalTexture(const TRefCountPtr<IP
 	OutTexture->ResourceRHI = ExternalPooledTexture->GetRenderTargetItem().ShaderResourceTexture;
 	AllocatedTextures.Add(OutTexture, ExternalPooledTexture);
 
-#if RDG_ENABLE_DEBUG
-	{
-		OutTexture->MarkAsExternal();
-		TrackedResources.Add(OutTexture);
-	}
-#endif
+	IF_RDG_ENABLE_DEBUG(Validation.ValidateCreateExternalResource(OutTexture));
 
 	ExternalTextures.Add(ExternalPooledTexture.GetReference(), OutTexture);
 
@@ -37,6 +33,7 @@ inline FRDGBufferRef FRDGBuilder::RegisterExternalBuffer(const TRefCountPtr<FPoo
 {
 #if RDG_ENABLE_DEBUG
 	{
+		Validation.ExecuteGuard(TEXT("RegisterExternalBuffer"), Name);
 		checkf(ExternalPooledBuffer.IsValid(), TEXT("Attempted to register NULL external buffer: %s"), Name);
 	}
 #endif
@@ -50,12 +47,7 @@ inline FRDGBufferRef FRDGBuilder::RegisterExternalBuffer(const TRefCountPtr<FPoo
 	OutBuffer->PooledBuffer = ExternalPooledBuffer;
 	AllocatedBuffers.Add(OutBuffer, ExternalPooledBuffer);
 
-#if RDG_ENABLE_DEBUG
-	{
-		OutBuffer->MarkAsExternal();
-		TrackedResources.Add(OutBuffer);
-	}
-#endif
+	IF_RDG_ENABLE_DEBUG(Validation.ValidateCreateExternalResource(OutBuffer));
 
 	ExternalBuffers.Add(ExternalPooledBuffer.GetReference(), OutBuffer);
 
@@ -64,31 +56,28 @@ inline FRDGBufferRef FRDGBuilder::RegisterExternalBuffer(const TRefCountPtr<FPoo
 
 inline FRDGTextureRef FRDGBuilder::CreateTexture(
 	const FPooledRenderTargetDesc& Desc,
-	const TCHAR* DebugName,
+	const TCHAR* Name,
 	ERDGResourceFlags Flags)
 {
 #if RDG_ENABLE_DEBUG
 	{
-		checkf(DebugName, TEXT("Creating a render graph texture requires a valid debug name."));
-		checkf(!bHasExecuted, TEXT("Render graph texture %s needs to be created before the builder execution."), DebugName);
-		checkf(Desc.Format != PF_Unknown, TEXT("Illegal to create texture %s with an invalid pixel format."), DebugName);
+		Validation.ExecuteGuard(TEXT("CreateTexture"), Name);
+
+		checkf(Name, TEXT("Creating a render graph texture requires a valid debug name."));
+		checkf(Desc.Format != PF_Unknown, TEXT("Illegal to create texture %s with an invalid pixel format."), Name);
 
 		const bool bCanHaveUAV = Desc.TargetableFlags & TexCreate_UAV;
 		const bool bIsMSAA = Desc.NumSamples > 1;
 
 		// D3D11 doesn't allow creating a UAV on MSAA texture.
 		const bool bIsUAVForMSAATexture = bIsMSAA && bCanHaveUAV;
-		checkf(!bIsUAVForMSAATexture, TEXT("TexCreate_UAV is not allowed on MSAA texture %s."), DebugName);
+		checkf(!bIsUAVForMSAATexture, TEXT("TexCreate_UAV is not allowed on MSAA texture %s."), Name);
 	}
 #endif
 
-	FRDGTexture* Texture = AllocateForRHILifeTime<FRDGTexture>(DebugName, Desc, Flags);
+	FRDGTexture* Texture = AllocateForRHILifeTime<FRDGTexture>(Name, Desc, Flags);
 
-#if RDG_ENABLE_DEBUG
-	{
-		TrackedResources.Add(Texture);
-	}
-#endif
+	IF_RDG_ENABLE_DEBUG(Validation.ValidateCreateResource(Texture));
 
 	return Texture;
 }
@@ -100,18 +89,14 @@ inline FRDGBufferRef FRDGBuilder::CreateBuffer(
 {
 #if RDG_ENABLE_DEBUG
 	{
+		Validation.ExecuteGuard(TEXT("CreateBuffer"), Name);
 		checkf(Name, TEXT("Creating a render graph buffer requires a valid debug name."));
-		checkf(!bHasExecuted, TEXT("Render graph buffer %s needs to be created before the builder execution."), Name);
 	}
 #endif
 
 	FRDGBufferRef Buffer = AllocateForRHILifeTime<FRDGBuffer>(Name, Desc, Flags);
 
-#if RDG_ENABLE_DEBUG
-	{
-		TrackedResources.Add(Buffer);
-	}
-#endif
+	IF_RDG_ENABLE_DEBUG(Validation.ValidateCreateResource(Buffer));
 
 	return Buffer;
 }
@@ -122,13 +107,13 @@ inline FRDGTextureSRVRef FRDGBuilder::CreateSRV(const FRDGTextureSRVDesc& Desc)
 
 #if RDG_ENABLE_DEBUG
 	{
+		Validation.ExecuteGuard(TEXT("CreateSRV"), Texture->Name);
 		checkf(Texture, TEXT("RenderGraph texture SRV created with a null texture."));
-		checkf(!bHasExecuted, TEXT("Render graph SRV %s needs to be created before the builder execution."), Desc.Texture->Name);
 		checkf(Desc.Texture->Desc.TargetableFlags & TexCreate_ShaderResource, TEXT("Attempted to create SRV from texture %s which was not created with TexCreate_ShaderResource"), Desc.Texture->Name);
 	}
 #endif
 
-	return AllocateForRHILifeTime<FRDGTextureSRV>(Desc.Texture->Name, Desc);
+	return AllocateForRHILifeTime<FRDGTextureSRV>(Texture->Name, Desc);
 }
 
 inline FRDGBufferSRVRef FRDGBuilder::CreateSRV(const FRDGBufferSRVDesc& Desc)
@@ -137,8 +122,8 @@ inline FRDGBufferSRVRef FRDGBuilder::CreateSRV(const FRDGBufferSRVDesc& Desc)
 
 #if RDG_ENABLE_DEBUG
 	{
+		Validation.ExecuteGuard(TEXT("CreateSRV"), Buffer->Name);
 		checkf(Buffer, TEXT("RenderGraph buffer SRV created with a null buffer."));
-		checkf(!bHasExecuted, TEXT("Render graph SRV %s needs to be created before the builder execution."), Desc.Buffer->Name);
 	}
 #endif
 
@@ -151,8 +136,8 @@ inline FRDGTextureUAVRef FRDGBuilder::CreateUAV(const FRDGTextureUAVDesc& Desc)
 
 #if RDG_ENABLE_DEBUG
 	{
+		Validation.ExecuteGuard(TEXT("CreateUAV"), Texture->Name);
 		checkf(Texture, TEXT("RenderGraph texture UAV created with a null texture."));
-		checkf(!bHasExecuted, TEXT("Render graph UAV %s needs to be created before the builder execution."), Texture->Name);
 		checkf(Texture->Desc.TargetableFlags & TexCreate_UAV, TEXT("Attempted to create UAV from texture %s which was not created with TexCreate_UAV"), Texture->Name);
 	}
 #endif
@@ -166,8 +151,8 @@ inline FRDGBufferUAVRef FRDGBuilder::CreateUAV(const FRDGBufferUAVDesc& Desc)
 
 #if RDG_ENABLE_DEBUG
 	{
+		Validation.ExecuteGuard(TEXT("CreateUAV"), Buffer->Name);
 		checkf(Buffer, TEXT("RenderGraph buffer UAV created with a null buffer."));
-		checkf(!bHasExecuted, TEXT("Render graph UAV %s needs to be created before the builder execution."), Buffer->Name);
 	}
 #endif
 
@@ -180,11 +165,7 @@ ParameterStructType* FRDGBuilder::AllocParameters()
 	// TODO(RDG): could allocate using AllocateForRHILifeTime() to avoid the copy done when using FRHICommandList::BuildLocalUniformBuffer()
 	ParameterStructType* OutParameterPtr = new(MemStack) ParameterStructType;
 	FMemory::Memzero(OutParameterPtr, sizeof(ParameterStructType));
-#if RDG_ENABLE_DEBUG
-	{
-		AllocatedUnusedPassParameters.Add(static_cast<void *>(OutParameterPtr));
-	}
-#endif
+	IF_RDG_ENABLE_DEBUG(Validation.ValidateAllocPassParameters(OutParameterPtr));
 	return OutParameterPtr;
 }
 
@@ -214,19 +195,10 @@ inline void FRDGBuilder::QueueTextureExtraction(
 	TRefCountPtr<IPooledRenderTarget>* OutTexturePtr,
 	bool bTransitionToRead)
 {
-	check(Texture);
+	IF_RDG_ENABLE_DEBUG(Validation.ValidateExtractResource(Texture));
+
 	check(OutTexturePtr);
-#if RDG_ENABLE_DEBUG
-	{
-		checkf(!bHasExecuted,
-			TEXT("Accessing render graph internal texture %s with QueueTextureExtraction() needs to happen before the builder's execution."),
-			Texture->Name);
 
-		checkf(Texture->HasBeenProduced(),
-			TEXT("Unable to queue the extraction of the texture %s because it has not been produced by any pass."),
-			Texture->Name);
-	}
-#endif
 	FDeferredInternalTextureQuery Query;
 	Query.Texture = Texture;
 	Query.OutTexturePtr = OutTexturePtr;
@@ -236,19 +208,10 @@ inline void FRDGBuilder::QueueTextureExtraction(
 
 inline void FRDGBuilder::QueueBufferExtraction(FRDGBufferRef Buffer, TRefCountPtr<FPooledRDGBuffer>* OutBufferPtr)
 {
-	check(Buffer);
+	IF_RDG_ENABLE_DEBUG(Validation.ValidateExtractResource(Buffer));
+
 	check(OutBufferPtr);
-#if RDG_ENABLE_DEBUG
-	{
-		checkf(!bHasExecuted,
-			TEXT("Accessing render graph internal buffer %s with QueueBufferExtraction() needs to happen before the builder's execution."),
-			Buffer->Name);
 
-		checkf(Buffer->HasBeenProduced(),
-			TEXT("Unable to queue the extraction of the buffer %s because it has not been produced by any pass."),
-			Buffer->Name);
-	}
-#endif
 	FDeferredInternalBufferQuery Query;
 	Query.Buffer = Buffer;
 	Query.OutBufferPtr = OutBufferPtr;
@@ -257,17 +220,19 @@ inline void FRDGBuilder::QueueBufferExtraction(FRDGBufferRef Buffer, TRefCountPt
 
 inline void FRDGBuilder::RemoveUnusedTextureWarning(FRDGTextureRef Texture)
 {
-	check(Texture);
 #if RDG_ENABLE_DEBUG
-	{
-		checkf(!bHasExecuted,
-			TEXT("Flaging texture %s with FlagUnusedTexture() needs to happen before the builder's execution."),
-			Texture->Name);
+	check(Texture);
+	Validation.ExecuteGuard(TEXT("RemoveUnusedTextureWarning"), Texture->Name);
+	Validation.RemoveUnusedWarning(Texture);
+#endif
+}
 
-		// Increment the number of time the texture has been accessed to avoid warning on produced but never used resources that were produced
-		// only to be extracted for the graph.
-		Texture->PassAccessCount += 1;
-	}
+inline void FRDGBuilder::RemoveUnusedBufferWarning(FRDGBufferRef Buffer)
+{
+#if RDG_ENABLE_DEBUG
+	check(Buffer);
+	Validation.ExecuteGuard(TEXT("RemoveUnusedBufferWarning"), Buffer->Name);
+	Validation.RemoveUnusedWarning(Buffer);
 #endif
 }
 
diff --git a/Engine/Source/Runtime/RenderCore/Public/RenderGraphDefinitions.h b/Engine/Source/Runtime/RenderCore/Public/RenderGraphDefinitions.h
--- a/Engine/Source/Runtime/RenderCore/Public/RenderGraphDefinitions.h
+++ b/Engine/Source/Runtime/RenderCore/Public/RenderGraphDefinitions.h
@@ -9,6 +9,13 @@
 /** Whether render graph debugging is enabled. */
 #define RDG_ENABLE_DEBUG (!UE_BUILD_SHIPPING && !UE_BUILD_TEST)
 
+/** Performs the operation if RDG_ENABLE_DEBUG is enabled. Useful for one-line checks without explicitly wrapping in #if. */
+#if RDG_ENABLE_DEBUG
+#define IF_RDG_ENABLE_DEBUG(Op) Op
+#else
+#define IF_RDG_ENABLE_DEBUG(Op)
+#endif
+
 /** Whether render graph debugging is enabled and we are compiling with the engine. */
 #define RDG_ENABLE_DEBUG_WITH_ENGINE (RDG_ENABLE_DEBUG && WITH_ENGINE)
 
diff --git a/Engine/Source/Runtime/RenderCore/Public/RenderGraphResources.h b/Engine/Source/Runtime/RenderCore/Public/RenderGraphResources.h
--- a/Engine/Source/Runtime/RenderCore/Public/RenderGraphResources.h
+++ b/Engine/Source/Runtime/RenderCore/Public/RenderGraphResources.h
@@ -46,10 +46,6 @@ struct FRDGResourceState
 		MAX = Unknown
 	};
 
-	/** Constructs a read / write state suitable for the provided pass. */
-	static FRDGResourceState CreateRead(const FRDGPass* Pass);
-	static FRDGResourceState CreateWrite(const FRDGPass* Pass);
-
 	FRDGResourceState() = default;
 
 	FRDGResourceState(const FRDGPass* InPass, EPipeline InPipeline, EAccess InAccess)
@@ -156,6 +152,7 @@ private:
 #endif
 
 	friend class FRDGBuilder;
+	friend class FRDGUserValidation;
 };
 
 /** A render graph resource with an allocation lifetime tracked by the graph. */
@@ -219,6 +216,7 @@ private:
 	FRDGResourceState State;
 
 	friend class FRDGBuilder;
+	friend class FRDGUserValidation;
 	friend class FRDGBarrierBatcher;
 };
 
diff --git a/Engine/Source/Runtime/RenderCore/Public/RenderGraphValidation.h b/Engine/Source/Runtime/RenderCore/Public/RenderGraphValidation.h
new file mode 100644
--- /dev/null
+++ b/Engine/Source/Runtime/RenderCore/Public/RenderGraphValidation.h
@@ -0,0 +1,85 @@
+// Copyright 1998-2019 Epic Games, Inc. All Rights Reserved.
+
+#pragma once
+
+#include "RenderGraphPass.h"
+
+/** Returns whether render graph runtime debugging is enabled. */
+extern bool IsRDGDebugEnabled();
+
+/** Returns whether render graph is running in immediate mode. */
+extern bool IsRDGImmediateModeEnabled();
+
+/** Emits a render graph warning. */
+extern void EmitRDGWarning(const FString& WarningMessage);
+
+#define EmitRDGWarningf(WarningMessageFormat, ...) \
+	EmitRDGWarning(FString::Printf(WarningMessageFormat, ##__VA_ARGS__));
+
+#if RDG_ENABLE_DEBUG
+
+/** Used by the render graph builder to validate correct usage of the graph API from setup to execution.
+ *  Validation is compiled out in shipping builds. This class tracks resources and passes as they are
+ *  added to the graph. It will then validate execution of the graph, including whether resources are
+ *  used during execution, and that they are properly produced before being consumed. All found issues
+ *  must be clear enough to help the user identify the problem in client code. Validation should occur
+ *  as soon as possible in the graph lifecycle. It's much easier to catch an issue at the setup location
+ *  rather than during deferred execution.
+ *
+ *  Finally, this class is designed for user validation, not for internal graph validation. In other words,
+ *  if the user can break the graph externally via the client-facing API, this validation layer should catch it.
+ *  Any internal validation of the graph state should be kept out of this class in order to provide a clear
+ *  and modular location to extend the validation layer as well as clearly separate the graph implementation
+ *  details from events in the graph.
+ */
+class RENDERCORE_API FRDGUserValidation final
+{
+public:
+	FRDGUserValidation() = default;
+	FRDGUserValidation(const FRDGUserValidation&) = delete;
+	~FRDGUserValidation();
+
+	/** Validates that the graph has not executed yet. */
+	void ExecuteGuard(const TCHAR* Operation, const TCHAR* ResourceName);
+
+	/** Tracks and validates the creation of a new resource in the graph. */
+	void ValidateCreateResource(FRDGTrackedResourceRef Resource);
+
+	/** Tracks and validates the creation of a new externally registered resource. */
+	void ValidateCreateExternalResource(FRDGTrackedResourceRef Resource);
+
+	/** Tracks and validates usage of a pass parameters allocation. */
+	void ValidateAllocPassParameters(const void* Parameters);
+
+	/** Validates a resource extraction operation. */
+	void ValidateExtractResource(FRDGTrackedResourceRef Resource);
+
+	/** Tracks and validates the addition of a new pass to the graph. */
+	void ValidateAddPass(const FRDGPass* Pass);
+
+	/** Validate pass state before and after execution. */
+	void ValidateExecutePassBegin(const FRDGPass* Pass);
+	void ValidateExecutePassEnd(const FRDGPass* Pass);
+
+	/** Validate graph state before and after execution. */
+	void ValidateExecuteBegin();
+	void ValidateExecuteEnd();
+
+	/** Removes the 'produced but not used' warning from the requested resource. */
+	void RemoveUnusedWarning(FRDGTrackedResourceRef Resource);
+
+private:
+	/** Traverses all resources in the pass and marks whether they are externally accessible by user pass implementations. */
+	static void SetAllowRHIAccess(const FRDGPass* Pass, bool bAllowAccess);
+
+	/** All recently allocated pass parameter structure, but not used by a AddPass() yet. */
+	TSet<const void*, DefaultKeyFuncs<const void*>, SceneRenderingSetAllocator> AllocatedUnusedPassParameters;
+
+	/** List of tracked resources for validation prior to shutdown. */
+	TArray<FRDGTrackedResourceRef, SceneRenderingAllocator> TrackedResources;
+
+	/** Whether the Execute() has already been called. */
+	bool bHasExecuted = false;
+};
+
+#endif
\ No newline at end of file
