Minor fixes to RenderGraph implementation.

diff --git a/Engine/Source/Runtime/RenderCore/Private/RenderGraphBuilder.cpp b/Engine/Source/Runtime/RenderCore/Private/RenderGraphBuilder.cpp
--- a/Engine/Source/Runtime/RenderCore/Private/RenderGraphBuilder.cpp
+++ b/Engine/Source/Runtime/RenderCore/Private/RenderGraphBuilder.cpp
@@ -5,7 +5,6 @@
 #include "RenderTargetPool.h"
 #include "RenderGraphResourcePool.h"
 #include "RenderGraphBarrierBatcher.h"
-#include "RenderGraphValidation.h"
 #include "VisualizeTexture.h"
 #include "ProfilingDebugging/CsvProfiler.h"
 
@@ -371,6 +370,7 @@ void FRDGBuilder::AllocateRHITextureIfNeeded(FRDGTexture* Texture)
 
 	Texture->PooledRenderTarget = PooledRenderTarget;
 	Texture->ResourceRHI = PooledRenderTarget->GetRenderTargetItem().ShaderResourceTexture;
+	check(Texture->ResourceRHI);
 }
 
 void FRDGBuilder::AllocateRHITextureUAVIfNeeded(FRDGTextureUAV* UAV)
@@ -884,6 +884,7 @@ void FRDGBuilder::ProcessDeferredInternalResourceQueries()
 	{
 		*Query.OutBufferPtr = AllocatedBuffers.FindChecked(Query.Buffer);
 
+		// No need to manually release in immediate mode, since it is done directly when emptying AllocatedTextures in DestructPasses().
 		if (!GRDGImmediateMode)
 		{
 			ReleaseRHIBufferIfUnreferenced(Query.Buffer);
