UE4.21 - FIx crash on html5

diff --git a/Engine/Source/Runtime/RenderCore/Private/RenderUtils.cpp b/Engine/Source/Runtime/RenderCore/Private/RenderUtils.cpp
--- a/Engine/Source/Runtime/RenderCore/Private/RenderUtils.cpp
+++ b/Engine/Source/Runtime/RenderCore/Private/RenderUtils.cpp
@@ -284,11 +284,19 @@ public:
 			FRHIResourceCreateInfo CreateInfo(&BlackTextureBulkData);
 			FTexture3DRHIRef Texture3D = RHICreateTexture3D(1,1,1,PixelFormat,1,TexCreate_ShaderResource,CreateInfo);
 			TextureRHI = Texture3D;	
-
-			// Create the sampler state.
-			FSamplerStateInitializerRHI SamplerStateInitializer(SF_Point,AM_Wrap,AM_Wrap,AM_Wrap);
-			SamplerStateRHI = RHICreateSamplerState(SamplerStateInitializer);
 		}
+		else
+		{
+			// Create a texture, even though it's not a volume texture
+			FBlackVolumeTextureResourceBulkDataInterface BlackTextureBulkData;
+			FRHIResourceCreateInfo CreateInfo(&BlackTextureBulkData);
+			FTexture2DRHIRef Texture2D = RHICreateTexture2D(1, 1, PixelFormat, 1, 1, TexCreate_ShaderResource, CreateInfo);
+			TextureRHI = Texture2D;
+		}
+
+		// Create the sampler state.
+		FSamplerStateInitializerRHI SamplerStateInitializer(SF_Point, AM_Wrap, AM_Wrap, AM_Wrap);
+		SamplerStateRHI = RHICreateSamplerState(SamplerStateInitializer);
 	}
 
 	/**
