Unshelved from pending changelist '4398174':

UE-54804 Graphical stuttering for weapon model on Lenovo S939 in First Person Template Project

#Android
#4.21

The bug was generated by CL3840215 - the fix for UE-53592: Assertion right after rendering scene on Lenovo S939. I've retested UE-53592 on my side and it's OK.

Also, when building a sample (eg. First Person) on //Dev-Main, the game crashed at start with:
D/UE4 (22619): Assertion failed: glProgramParameteri [File:C:/Work/UE4_Main/Engine/Source/Runtime/OpenGLDrv/Private\Android/AndroidOpenGL.h] [Line: 479]

diff --git a/Engine/Source/Runtime/OpenGLDrv/Private/Android/AndroidOpenGL.h b/Engine/Source/Runtime/OpenGLDrv/Private/Android/AndroidOpenGL.h
--- a/Engine/Source/Runtime/OpenGLDrv/Private/Android/AndroidOpenGL.h
+++ b/Engine/Source/Runtime/OpenGLDrv/Private/Android/AndroidOpenGL.h
@@ -476,8 +476,15 @@ struct FAndroidOpenGL : public FOpenGLES2
 
 	static FORCEINLINE void ProgramParameter(GLuint Program, GLenum PName, GLint Value)
 	{
-		check(glProgramParameteri);
-		glProgramParameteri(Program, PName, Value);
+		if (bES30Support)
+		{
+			check(glProgramParameteri);
+			glProgramParameteri(Program, PName, Value);
+		}
+		else
+		{
+			FOpenGLBase::ProgramParameter(Program, PName, Value);
+		}
 	}
 
 	static FORCEINLINE void BindBufferBase(GLenum Target, GLuint Index, GLuint Buffer)
diff --git a/Engine/Source/Runtime/RenderCore/Public/RenderingThread.h b/Engine/Source/Runtime/RenderCore/Public/RenderingThread.h
--- a/Engine/Source/Runtime/RenderCore/Public/RenderingThread.h
+++ b/Engine/Source/Runtime/RenderCore/Public/RenderingThread.h
@@ -244,9 +244,7 @@ FORCEINLINE_DEBUGGABLE void EnqueueUniqueRenderCommand(LAMBDA&& Lambda)
 #if 0 // UE_SERVER && UE_BUILD_DEBUG
 	UE_LOG(LogRHI, Warning, TEXT("Render command '%s' is being executed on a dedicated server."), TSTR::TStr())
 #endif
-	// always use a new task for devices that have GUseThreadedRendering=false
-	// even when the call is from the rendering thread
-	if (GUseThreadedRendering && IsInRenderingThread())
+	if (IsInRenderingThread())
 	{
 		FRHICommandListImmediate& RHICmdList = GetImmediateCommandList_ForRenderCommand();
 		Lambda(RHICmdList);
