DR - Move some static globals out of function scopes to help debugging
- Remove unused member variable

diff --git a/Engine/Source/Runtime/RenderCore/Private/RenderCore.cpp b/Engine/Source/Runtime/RenderCore/Private/RenderCore.cpp
--- a/Engine/Source/Runtime/RenderCore/Private/RenderCore.cpp
+++ b/Engine/Source/Runtime/RenderCore/Private/RenderCore.cpp
@@ -178,6 +178,11 @@ DEFINE_STAT(STAT_RenderTargetPoolSize);
 DEFINE_STAT(STAT_RenderTargetPoolUsed);
 DEFINE_STAT(STAT_RenderTargetPoolCount);
 
+
+static TLinkedList<FUniformBufferStruct*>* GUniformStructList = nullptr;
+static TMap<FName, FUniformBufferStruct*> GGlobalNameStructMap;
+
+
 #define EXPOSE_FORCE_LOD !(UE_BUILD_SHIPPING || UE_BUILD_TEST)
 
 #if EXPOSE_FORCE_LOD
@@ -234,14 +239,12 @@ void FInputLatencyTimer::GameThreadTick()
 
 TLinkedList<FUniformBufferStruct*>*& FUniformBufferStruct::GetStructList()
 {
-	static TLinkedList<FUniformBufferStruct*>* GUniformStructList = NULL;
 	return GUniformStructList;
 }
 
 TMap<FName, FUniformBufferStruct*>& FUniformBufferStruct::GetNameStructMap()
 {
-	static 	TMap<FName, FUniformBufferStruct*> GlobalNameStructMap;
-	return GlobalNameStructMap;
+	return GGlobalNameStructMap;
 }
 
 
diff --git a/Engine/Source/Runtime/Renderer/Private/RendererModule.h b/Engine/Source/Runtime/Renderer/Private/RendererModule.h
--- a/Engine/Source/Runtime/Renderer/Private/RendererModule.h
+++ b/Engine/Source/Runtime/Renderer/Private/RendererModule.h
@@ -104,7 +104,6 @@ public:
 
 private:
 	TSet<FSceneInterface*> AllocatedScenes;
-	ICustomCulling* CustomCullingImpl;
 	FPostOpaqueRenderDelegate PostOpaqueRenderDelegate;
 	FPostOpaqueRenderDelegate OverlayRenderDelegate;
 	FOnResolvedSceneColor PostResolvedSceneColorCallbacks;
diff --git a/Engine/Source/Runtime/Renderer/Private/SceneRendering.cpp b/Engine/Source/Runtime/Renderer/Private/SceneRendering.cpp
--- a/Engine/Source/Runtime/Renderer/Private/SceneRendering.cpp
+++ b/Engine/Source/Runtime/Renderer/Private/SceneRendering.cpp
@@ -3053,7 +3053,6 @@ void OnChangeCVarRequiringRecreateRenderState(IConsoleVariable* Var)
 }
 
 FRendererModule::FRendererModule()
-	: CustomCullingImpl(nullptr)
 {
 	CVarSimpleForwardShading_PreviousValue = CVarSimpleForwardShading.AsVariable()->GetInt();
 	CVarSimpleForwardShading.AsVariable()->SetOnChangedCallback(FConsoleVariableDelegate::CreateStatic(&OnChangeSimpleForwardShading));
diff --git a/Engine/Source/Runtime/ShaderCore/Private/Shader.cpp b/Engine/Source/Runtime/ShaderCore/Private/Shader.cpp
--- a/Engine/Source/Runtime/ShaderCore/Private/Shader.cpp
+++ b/Engine/Source/Runtime/ShaderCore/Private/Shader.cpp
@@ -37,6 +37,7 @@ static TAutoConsoleVariable<int32> CVarUsePipelines(
 
 static TLinkedList<FShaderType*>*			GShaderTypeList = nullptr;
 static TLinkedList<FShaderPipelineType*>*	GShaderPipelineList = nullptr;
+static TMap<FName, FShaderType*>*			GShaderNameToTypeMap = nullptr;
 
 static FSHAHash ShaderSourceDefaultHash; //will only be read (never written) for the cooking case
 
@@ -216,7 +217,6 @@ TArray<FShaderType*> FShaderType::GetShaderTypesByFilename(const TCHAR* Filename
 
 TMap<FName, FShaderType*>& FShaderType::GetNameToTypeMap()
 {
-	static TMap<FName, FShaderType*>* GShaderNameToTypeMap = NULL;
 	if(!GShaderNameToTypeMap)
 	{
 		GShaderNameToTypeMap = new TMap<FName, FShaderType*>();
diff --git a/Engine/Source/Runtime/ShaderCore/Private/VertexFactory.cpp b/Engine/Source/Runtime/ShaderCore/Private/VertexFactory.cpp
--- a/Engine/Source/Runtime/ShaderCore/Private/VertexFactory.cpp
+++ b/Engine/Source/Runtime/ShaderCore/Private/VertexFactory.cpp
@@ -11,13 +11,14 @@
 uint32 FVertexFactoryType::NextHashIndex = 0;
 bool FVertexFactoryType::bInitializedSerializationHistory = false;
 
+static TLinkedList<FVertexFactoryType*>* GVFTypeList = nullptr;
+
 /**
  * @return The global shader factory list.
  */
 TLinkedList<FVertexFactoryType*>*& FVertexFactoryType::GetTypeList()
 {
-	static TLinkedList<FVertexFactoryType*>* TypeList = NULL;
-	return TypeList;
+	return GVFTypeList;
 }
 
 /**
