Fix Linux crash on start (due to different order of global constructors).

- Partial backout of 4366793 (most of the CL left intact).

diff --git a/Engine/Source/Runtime/RenderCore/Private/RenderCore.cpp b/Engine/Source/Runtime/RenderCore/Private/RenderCore.cpp
--- a/Engine/Source/Runtime/RenderCore/Private/RenderCore.cpp
+++ b/Engine/Source/Runtime/RenderCore/Private/RenderCore.cpp
@@ -178,10 +178,7 @@ DEFINE_STAT(STAT_RenderTargetPoolSize);
 DEFINE_STAT(STAT_RenderTargetPoolUsed);
 DEFINE_STAT(STAT_RenderTargetPoolCount);
 
-
 static TLinkedList<FUniformBufferStruct*>* GUniformStructList = nullptr;
-static TMap<FName, FUniformBufferStruct*> GGlobalNameStructMap;
-
 
 #define EXPOSE_FORCE_LOD !(UE_BUILD_SHIPPING || UE_BUILD_TEST)
 
@@ -244,7 +241,8 @@ TLinkedList<FUniformBufferStruct*>*& FUniformBufferStruct::GetStructList()
 
 TMap<FName, FUniformBufferStruct*>& FUniformBufferStruct::GetNameStructMap()
 {
-	return GGlobalNameStructMap;
+	static TMap<FName, FUniformBufferStruct*> GlobalNameStructMap;
+	return GlobalNameStructMap;
 }
 
 
