Merging //UE4/Dev-Rendering-Graph@4492585 to Dev-Rendering (//UE4/Dev-Rendering)

diff --git a/Engine/Source/Runtime/Engine/Public/ParticleVertexFactory.h b/Engine/Source/Runtime/Engine/Public/ParticleVertexFactory.h
--- a/Engine/Source/Runtime/Engine/Public/ParticleVertexFactory.h
+++ b/Engine/Source/Runtime/Engine/Public/ParticleVertexFactory.h
@@ -116,22 +116,22 @@ private:
 /**
  * Uniform buffer for particle sprite vertex factories.
  */
-BEGIN_UNIFORM_BUFFER_STRUCT( FParticleSpriteUniformParameters, ENGINE_API)
-	UNIFORM_MEMBER_EX( FVector4, AxisLockRight, EShaderPrecisionModifier::Half )
-	UNIFORM_MEMBER_EX( FVector4, AxisLockUp, EShaderPrecisionModifier::Half )
-	UNIFORM_MEMBER_EX( FVector4, TangentSelector, EShaderPrecisionModifier::Half )
-	UNIFORM_MEMBER_EX( FVector4, NormalsSphereCenter, EShaderPrecisionModifier::Half )
-	UNIFORM_MEMBER_EX( FVector4, NormalsCylinderUnitDirection, EShaderPrecisionModifier::Half )
-	UNIFORM_MEMBER_EX( FVector4, SubImageSize, EShaderPrecisionModifier::Half )
-	UNIFORM_MEMBER_EX( FVector, CameraFacingBlend, EShaderPrecisionModifier::Half )
-	UNIFORM_MEMBER_EX( float, RemoveHMDRoll, EShaderPrecisionModifier::Half )
-	UNIFORM_MEMBER( FVector4, MacroUVParameters )
-	UNIFORM_MEMBER_EX( float, RotationScale, EShaderPrecisionModifier::Half )
-	UNIFORM_MEMBER_EX( float, RotationBias, EShaderPrecisionModifier::Half )
-	UNIFORM_MEMBER_EX( float, NormalsType, EShaderPrecisionModifier::Half )
-	UNIFORM_MEMBER_EX( float, InvDeltaSeconds, EShaderPrecisionModifier::Half )
-	UNIFORM_MEMBER_EX( FVector2D, PivotOffset, EShaderPrecisionModifier::Half )
-END_UNIFORM_BUFFER_STRUCT( FParticleSpriteUniformParameters )
+BEGIN_GLOBAL_SHADER_PARAMETER_STRUCT( FParticleSpriteUniformParameters, ENGINE_API)
+	SHADER_PARAMETER_EX( FVector4, AxisLockRight, EShaderPrecisionModifier::Half )
+	SHADER_PARAMETER_EX( FVector4, AxisLockUp, EShaderPrecisionModifier::Half )
+	SHADER_PARAMETER_EX( FVector4, TangentSelector, EShaderPrecisionModifier::Half )
+	SHADER_PARAMETER_EX( FVector4, NormalsSphereCenter, EShaderPrecisionModifier::Half )
+	SHADER_PARAMETER_EX( FVector4, NormalsCylinderUnitDirection, EShaderPrecisionModifier::Half )
+	SHADER_PARAMETER_EX( FVector4, SubImageSize, EShaderPrecisionModifier::Half )
+	SHADER_PARAMETER_EX( FVector, CameraFacingBlend, EShaderPrecisionModifier::Half )
+	SHADER_PARAMETER_EX( float, RemoveHMDRoll, EShaderPrecisionModifier::Half )
+	SHADER_PARAMETER( FVector4, MacroUVParameters )
+	SHADER_PARAMETER_EX( float, RotationScale, EShaderPrecisionModifier::Half )
+	SHADER_PARAMETER_EX( float, RotationBias, EShaderPrecisionModifier::Half )
+	SHADER_PARAMETER_EX( float, NormalsType, EShaderPrecisionModifier::Half )
+	SHADER_PARAMETER_EX( float, InvDeltaSeconds, EShaderPrecisionModifier::Half )
+	SHADER_PARAMETER_EX( FVector2D, PivotOffset, EShaderPrecisionModifier::Half )
+END_GLOBAL_SHADER_PARAMETER_STRUCT()
 typedef TUniformBufferRef<FParticleSpriteUniformParameters> FParticleSpriteUniformBufferRef;
 
 /**
diff --git a/Engine/Source/Runtime/Engine/Public/PrimitiveUniformShaderParameters.h b/Engine/Source/Runtime/Engine/Public/PrimitiveUniformShaderParameters.h
--- a/Engine/Source/Runtime/Engine/Public/PrimitiveUniformShaderParameters.h
+++ b/Engine/Source/Runtime/Engine/Public/PrimitiveUniformShaderParameters.h
@@ -11,25 +11,25 @@
 
 
 /** The uniform shader parameters associated with a primitive. */
-BEGIN_UNIFORM_BUFFER_STRUCT(FPrimitiveUniformShaderParameters,ENGINE_API)
-	UNIFORM_MEMBER(FMatrix,LocalToWorld)		// always needed
-	UNIFORM_MEMBER_EX(FVector4,InvNonUniformScaleAndDeterminantSign,EShaderPrecisionModifier::Half) //often needed
-	UNIFORM_MEMBER(FVector4,ObjectWorldPositionAndRadius)	// needed by some materials
-	UNIFORM_MEMBER(FMatrix,WorldToLocal)		// rarely needed
-	UNIFORM_MEMBER(FVector,ActorWorldPosition)
-	UNIFORM_MEMBER_EX(float,UseSingleSampleShadowFromStationaryLights,EShaderPrecisionModifier::Half)	
-	UNIFORM_MEMBER(FVector,ObjectBounds)		// only needed for editor/development
-	UNIFORM_MEMBER(float,LpvBiasMultiplier)
-	UNIFORM_MEMBER_EX(float,DecalReceiverMask,EShaderPrecisionModifier::Half)
-	UNIFORM_MEMBER_EX(float,PerObjectGBufferData,EShaderPrecisionModifier::Half)		// 0..1, 2 bits, bDistanceFieldRepresentation, bHeightfieldRepresentation
-	UNIFORM_MEMBER_EX(float,UseVolumetricLightmapShadowFromStationaryLights,EShaderPrecisionModifier::Half)		
-	UNIFORM_MEMBER_EX(float,UseEditorDepthTest,EShaderPrecisionModifier::Half)
-	UNIFORM_MEMBER_EX(FVector4,ObjectOrientation,EShaderPrecisionModifier::Half)
-	UNIFORM_MEMBER_EX(FVector4,NonUniformScale,EShaderPrecisionModifier::Half)
-	UNIFORM_MEMBER(FVector, LocalObjectBoundsMin)		// This is used in a custom material function (ObjectLocalBounds.uasset)
-	UNIFORM_MEMBER(FVector, LocalObjectBoundsMax)		// This is used in a custom material function (ObjectLocalBounds.uasset)
-	UNIFORM_MEMBER(uint32,LightingChannelMask)
-END_UNIFORM_BUFFER_STRUCT(FPrimitiveUniformShaderParameters)
+BEGIN_GLOBAL_SHADER_PARAMETER_STRUCT(FPrimitiveUniformShaderParameters,ENGINE_API)
+	SHADER_PARAMETER(FMatrix,LocalToWorld)		// always needed
+	SHADER_PARAMETER_EX(FVector4,InvNonUniformScaleAndDeterminantSign,EShaderPrecisionModifier::Half) //often needed
+	SHADER_PARAMETER(FVector4,ObjectWorldPositionAndRadius)	// needed by some materials
+	SHADER_PARAMETER(FMatrix,WorldToLocal)		// rarely needed
+	SHADER_PARAMETER(FVector,ActorWorldPosition)
+	SHADER_PARAMETER_EX(float,UseSingleSampleShadowFromStationaryLights,EShaderPrecisionModifier::Half)	
+	SHADER_PARAMETER(FVector,ObjectBounds)		// only needed for editor/development
+	SHADER_PARAMETER(float,LpvBiasMultiplier)
+	SHADER_PARAMETER_EX(float,DecalReceiverMask,EShaderPrecisionModifier::Half)
+	SHADER_PARAMETER_EX(float,PerObjectGBufferData,EShaderPrecisionModifier::Half)		// 0..1, 2 bits, bDistanceFieldRepresentation, bHeightfieldRepresentation
+	SHADER_PARAMETER_EX(float,UseVolumetricLightmapShadowFromStationaryLights,EShaderPrecisionModifier::Half)		
+	SHADER_PARAMETER_EX(float,UseEditorDepthTest,EShaderPrecisionModifier::Half)
+	SHADER_PARAMETER_EX(FVector4,ObjectOrientation,EShaderPrecisionModifier::Half)
+	SHADER_PARAMETER_EX(FVector4,NonUniformScale,EShaderPrecisionModifier::Half)
+	SHADER_PARAMETER(FVector, LocalObjectBoundsMin)		// This is used in a custom material function (ObjectLocalBounds.uasset)
+	SHADER_PARAMETER(FVector, LocalObjectBoundsMax)		// This is used in a custom material function (ObjectLocalBounds.uasset)
+	SHADER_PARAMETER(uint32,LightingChannelMask)
+END_GLOBAL_SHADER_PARAMETER_STRUCT()
 
 /** Initializes the primitive uniform shader parameters. */
 inline FPrimitiveUniformShaderParameters GetPrimitiveUniformShaderParameters(
diff --git a/Engine/Source/Runtime/Engine/Public/SceneView.h b/Engine/Source/Runtime/Engine/Public/SceneView.h
--- a/Engine/Source/Runtime/Engine/Public/SceneView.h
+++ b/Engine/Source/Runtime/Engine/Public/SceneView.h
@@ -583,16 +583,16 @@ static const int MAX_MOBILE_SHADOWCASCADES = 4;
 /** The uniform shader parameters for a mobile directional light and its shadow.
   * One uniform buffer will be created for the first directional light in each lighting channel.
   */
-BEGIN_UNIFORM_BUFFER_STRUCT_WITH_CONSTRUCTOR(FMobileDirectionalLightShaderParameters, ENGINE_API)
-	UNIFORM_MEMBER_EX(FLinearColor, DirectionalLightColor, EShaderPrecisionModifier::Half)
-	UNIFORM_MEMBER_EX(FVector4, DirectionalLightDirectionAndShadowTransition, EShaderPrecisionModifier::Half)
-	UNIFORM_MEMBER_EX(FVector4, DirectionalLightShadowSize, EShaderPrecisionModifier::Half)
-	UNIFORM_MEMBER_EX(FVector4, DirectionalLightDistanceFadeMAD, EShaderPrecisionModifier::Half) // .zw is not used atm
-	UNIFORM_MEMBER_EX(FVector4, DirectionalLightShadowDistances, EShaderPrecisionModifier::Half)
-	UNIFORM_MEMBER_ARRAY(FMatrix, DirectionalLightScreenToShadow, [MAX_MOBILE_SHADOWCASCADES])
-	UNIFORM_MEMBER_TEXTURE(Texture2D, DirectionalLightShadowTexture)
-	UNIFORM_MEMBER_SAMPLER(SamplerState, DirectionalLightShadowSampler)
-END_UNIFORM_BUFFER_STRUCT(FMobileDirectionalLightShaderParameters)
+BEGIN_GLOBAL_SHADER_PARAMETER_STRUCT_WITH_CONSTRUCTOR(FMobileDirectionalLightShaderParameters, ENGINE_API)
+	SHADER_PARAMETER_EX(FLinearColor, DirectionalLightColor, EShaderPrecisionModifier::Half)
+	SHADER_PARAMETER_EX(FVector4, DirectionalLightDirectionAndShadowTransition, EShaderPrecisionModifier::Half)
+	SHADER_PARAMETER_EX(FVector4, DirectionalLightShadowSize, EShaderPrecisionModifier::Half)
+	SHADER_PARAMETER_EX(FVector4, DirectionalLightDistanceFadeMAD, EShaderPrecisionModifier::Half) // .zw is not used atm
+	SHADER_PARAMETER_EX(FVector4, DirectionalLightShadowDistances, EShaderPrecisionModifier::Half)
+	SHADER_PARAMETER_ARRAY(FMatrix, DirectionalLightScreenToShadow, [MAX_MOBILE_SHADOWCASCADES])
+	SHADER_PARAMETER_TEXTURE(Texture2D, DirectionalLightShadowTexture)
+	SHADER_PARAMETER_SAMPLER(SamplerState, DirectionalLightShadowSampler)
+END_GLOBAL_SHADER_PARAMETER_STRUCT()
 
 //////////////////////////////////////////////////////////////////////////
 
@@ -747,79 +747,79 @@ enum ETranslucencyVolumeCascade
 	VIEW_UNIFORM_BUFFER_MEMBER(float, StereoIPD)
 
 #define VIEW_UNIFORM_BUFFER_MEMBER(type, identifier) \
-	UNIFORM_MEMBER(type, identifier)
+	SHADER_PARAMETER(type, identifier)
 
 #define VIEW_UNIFORM_BUFFER_MEMBER_EX(type, identifier, precision) \
-	UNIFORM_MEMBER_EX(type, identifier, precision)
+	SHADER_PARAMETER_EX(type, identifier, precision)
 
 #define VIEW_UNIFORM_BUFFER_MEMBER_ARRAY(type, identifier, dimension) \
-	UNIFORM_MEMBER_ARRAY(type, identifier, dimension)
+	SHADER_PARAMETER_ARRAY(type, identifier, dimension)
 
 /** The uniform shader parameters associated with a view. */
-BEGIN_UNIFORM_BUFFER_STRUCT_WITH_CONSTRUCTOR(FViewUniformShaderParameters, ENGINE_API)
+BEGIN_GLOBAL_SHADER_PARAMETER_STRUCT_WITH_CONSTRUCTOR(FViewUniformShaderParameters, ENGINE_API)
 
 	VIEW_UNIFORM_BUFFER_MEMBER_TABLE
 
 	// Same as Wrap_WorldGroupSettings and Clamp_WorldGroupSettings, but with mipbias=MaterialTextureMipBias.
-	UNIFORM_MEMBER_SAMPLER(SamplerState, MaterialTextureBilinearWrapedSampler)
-	UNIFORM_MEMBER_SAMPLER(SamplerState, MaterialTextureBilinearClampedSampler)
-
-	UNIFORM_MEMBER_TEXTURE(Texture3D<uint4>, VolumetricLightmapIndirectionTexture) // FPrecomputedVolumetricLightmapLightingPolicy
-	UNIFORM_MEMBER_TEXTURE(Texture3D, VolumetricLightmapBrickAmbientVector) // FPrecomputedVolumetricLightmapLightingPolicy
-	UNIFORM_MEMBER_TEXTURE(Texture3D, VolumetricLightmapBrickSHCoefficients0) // FPrecomputedVolumetricLightmapLightingPolicy
-	UNIFORM_MEMBER_TEXTURE(Texture3D, VolumetricLightmapBrickSHCoefficients1) // FPrecomputedVolumetricLightmapLightingPolicy
-	UNIFORM_MEMBER_TEXTURE(Texture3D, VolumetricLightmapBrickSHCoefficients2) // FPrecomputedVolumetricLightmapLightingPolicy
-	UNIFORM_MEMBER_TEXTURE(Texture3D, VolumetricLightmapBrickSHCoefficients3) // FPrecomputedVolumetricLightmapLightingPolicy
-	UNIFORM_MEMBER_TEXTURE(Texture3D, VolumetricLightmapBrickSHCoefficients4) // FPrecomputedVolumetricLightmapLightingPolicy
-	UNIFORM_MEMBER_TEXTURE(Texture3D, VolumetricLightmapBrickSHCoefficients5) // FPrecomputedVolumetricLightmapLightingPolicy
-	UNIFORM_MEMBER_TEXTURE(Texture3D, SkyBentNormalBrickTexture) // FPrecomputedVolumetricLightmapLightingPolicy
-	UNIFORM_MEMBER_TEXTURE(Texture3D, DirectionalLightShadowingBrickTexture) // FPrecomputedVolumetricLightmapLightingPolicy
-
-	UNIFORM_MEMBER_SAMPLER(SamplerState, VolumetricLightmapBrickAmbientVectorSampler) // FPrecomputedVolumetricLightmapLightingPolicy
-	UNIFORM_MEMBER_SAMPLER(SamplerState, VolumetricLightmapTextureSampler0) // FPrecomputedVolumetricLightmapLightingPolicy
-	UNIFORM_MEMBER_SAMPLER(SamplerState, VolumetricLightmapTextureSampler1) // FPrecomputedVolumetricLightmapLightingPolicy
-	UNIFORM_MEMBER_SAMPLER(SamplerState, VolumetricLightmapTextureSampler2) // FPrecomputedVolumetricLightmapLightingPolicy
-	UNIFORM_MEMBER_SAMPLER(SamplerState, VolumetricLightmapTextureSampler3) // FPrecomputedVolumetricLightmapLightingPolicy
-	UNIFORM_MEMBER_SAMPLER(SamplerState, VolumetricLightmapTextureSampler4) // FPrecomputedVolumetricLightmapLightingPolicy
-	UNIFORM_MEMBER_SAMPLER(SamplerState, VolumetricLightmapTextureSampler5) // FPrecomputedVolumetricLightmapLightingPolicy
-	UNIFORM_MEMBER_SAMPLER(SamplerState, SkyBentNormalTextureSampler) // FPrecomputedVolumetricLightmapLightingPolicy
-	UNIFORM_MEMBER_SAMPLER(SamplerState, DirectionalLightShadowingTextureSampler) // FPrecomputedVolumetricLightmapLightingPolicy
-
-	UNIFORM_MEMBER_TEXTURE(Texture3D, GlobalDistanceFieldTexture0)
-	UNIFORM_MEMBER_SAMPLER(SamplerState, GlobalDistanceFieldSampler0)
-	UNIFORM_MEMBER_TEXTURE(Texture3D, GlobalDistanceFieldTexture1)
-	UNIFORM_MEMBER_SAMPLER(SamplerState, GlobalDistanceFieldSampler1)
-	UNIFORM_MEMBER_TEXTURE(Texture3D, GlobalDistanceFieldTexture2)
-	UNIFORM_MEMBER_SAMPLER(SamplerState, GlobalDistanceFieldSampler2)
-	UNIFORM_MEMBER_TEXTURE(Texture3D, GlobalDistanceFieldTexture3)
-	UNIFORM_MEMBER_SAMPLER(SamplerState, GlobalDistanceFieldSampler3)
-
-	UNIFORM_MEMBER_TEXTURE(Texture2D, AtmosphereTransmittanceTexture)
-	UNIFORM_MEMBER_SAMPLER(SamplerState, AtmosphereTransmittanceTextureSampler)
-	UNIFORM_MEMBER_TEXTURE(Texture2D, AtmosphereIrradianceTexture)
-	UNIFORM_MEMBER_SAMPLER(SamplerState, AtmosphereIrradianceTextureSampler)
-	UNIFORM_MEMBER_TEXTURE(Texture3D, AtmosphereInscatterTexture)
-	UNIFORM_MEMBER_SAMPLER(SamplerState, AtmosphereInscatterTextureSampler)
-	UNIFORM_MEMBER_TEXTURE(Texture2D, PerlinNoiseGradientTexture)
-	UNIFORM_MEMBER_SAMPLER(SamplerState, PerlinNoiseGradientTextureSampler)
-	UNIFORM_MEMBER_TEXTURE(Texture3D, PerlinNoise3DTexture)
-	UNIFORM_MEMBER_SAMPLER(SamplerState, PerlinNoise3DTextureSampler)
-	UNIFORM_MEMBER_TEXTURE(Texture2D<uint>, SobolSamplingTexture)
-	UNIFORM_MEMBER_SAMPLER(SamplerState, SharedPointWrappedSampler)
-	UNIFORM_MEMBER_SAMPLER(SamplerState, SharedPointClampedSampler)
-	UNIFORM_MEMBER_SAMPLER(SamplerState, SharedBilinearWrappedSampler)
-	UNIFORM_MEMBER_SAMPLER(SamplerState, SharedBilinearClampedSampler)
-	UNIFORM_MEMBER_SAMPLER(SamplerState, SharedTrilinearWrappedSampler)
-	UNIFORM_MEMBER_SAMPLER(SamplerState, SharedTrilinearClampedSampler)
-	UNIFORM_MEMBER_TEXTURE(Texture2D, PreIntegratedBRDF)
-	UNIFORM_MEMBER_SAMPLER(SamplerState, PreIntegratedBRDFSampler)
-
-END_UNIFORM_BUFFER_STRUCT(FViewUniformShaderParameters)
+	SHADER_PARAMETER_SAMPLER(SamplerState, MaterialTextureBilinearWrapedSampler)
+	SHADER_PARAMETER_SAMPLER(SamplerState, MaterialTextureBilinearClampedSampler)
+
+	SHADER_PARAMETER_TEXTURE(Texture3D<uint4>, VolumetricLightmapIndirectionTexture) // FPrecomputedVolumetricLightmapLightingPolicy
+	SHADER_PARAMETER_TEXTURE(Texture3D, VolumetricLightmapBrickAmbientVector) // FPrecomputedVolumetricLightmapLightingPolicy
+	SHADER_PARAMETER_TEXTURE(Texture3D, VolumetricLightmapBrickSHCoefficients0) // FPrecomputedVolumetricLightmapLightingPolicy
+	SHADER_PARAMETER_TEXTURE(Texture3D, VolumetricLightmapBrickSHCoefficients1) // FPrecomputedVolumetricLightmapLightingPolicy
+	SHADER_PARAMETER_TEXTURE(Texture3D, VolumetricLightmapBrickSHCoefficients2) // FPrecomputedVolumetricLightmapLightingPolicy
+	SHADER_PARAMETER_TEXTURE(Texture3D, VolumetricLightmapBrickSHCoefficients3) // FPrecomputedVolumetricLightmapLightingPolicy
+	SHADER_PARAMETER_TEXTURE(Texture3D, VolumetricLightmapBrickSHCoefficients4) // FPrecomputedVolumetricLightmapLightingPolicy
+	SHADER_PARAMETER_TEXTURE(Texture3D, VolumetricLightmapBrickSHCoefficients5) // FPrecomputedVolumetricLightmapLightingPolicy
+	SHADER_PARAMETER_TEXTURE(Texture3D, SkyBentNormalBrickTexture) // FPrecomputedVolumetricLightmapLightingPolicy
+	SHADER_PARAMETER_TEXTURE(Texture3D, DirectionalLightShadowingBrickTexture) // FPrecomputedVolumetricLightmapLightingPolicy
+
+	SHADER_PARAMETER_SAMPLER(SamplerState, VolumetricLightmapBrickAmbientVectorSampler) // FPrecomputedVolumetricLightmapLightingPolicy
+	SHADER_PARAMETER_SAMPLER(SamplerState, VolumetricLightmapTextureSampler0) // FPrecomputedVolumetricLightmapLightingPolicy
+	SHADER_PARAMETER_SAMPLER(SamplerState, VolumetricLightmapTextureSampler1) // FPrecomputedVolumetricLightmapLightingPolicy
+	SHADER_PARAMETER_SAMPLER(SamplerState, VolumetricLightmapTextureSampler2) // FPrecomputedVolumetricLightmapLightingPolicy
+	SHADER_PARAMETER_SAMPLER(SamplerState, VolumetricLightmapTextureSampler3) // FPrecomputedVolumetricLightmapLightingPolicy
+	SHADER_PARAMETER_SAMPLER(SamplerState, VolumetricLightmapTextureSampler4) // FPrecomputedVolumetricLightmapLightingPolicy
+	SHADER_PARAMETER_SAMPLER(SamplerState, VolumetricLightmapTextureSampler5) // FPrecomputedVolumetricLightmapLightingPolicy
+	SHADER_PARAMETER_SAMPLER(SamplerState, SkyBentNormalTextureSampler) // FPrecomputedVolumetricLightmapLightingPolicy
+	SHADER_PARAMETER_SAMPLER(SamplerState, DirectionalLightShadowingTextureSampler) // FPrecomputedVolumetricLightmapLightingPolicy
+
+	SHADER_PARAMETER_TEXTURE(Texture3D, GlobalDistanceFieldTexture0)
+	SHADER_PARAMETER_SAMPLER(SamplerState, GlobalDistanceFieldSampler0)
+	SHADER_PARAMETER_TEXTURE(Texture3D, GlobalDistanceFieldTexture1)
+	SHADER_PARAMETER_SAMPLER(SamplerState, GlobalDistanceFieldSampler1)
+	SHADER_PARAMETER_TEXTURE(Texture3D, GlobalDistanceFieldTexture2)
+	SHADER_PARAMETER_SAMPLER(SamplerState, GlobalDistanceFieldSampler2)
+	SHADER_PARAMETER_TEXTURE(Texture3D, GlobalDistanceFieldTexture3)
+	SHADER_PARAMETER_SAMPLER(SamplerState, GlobalDistanceFieldSampler3)
+
+	SHADER_PARAMETER_TEXTURE(Texture2D, AtmosphereTransmittanceTexture)
+	SHADER_PARAMETER_SAMPLER(SamplerState, AtmosphereTransmittanceTextureSampler)
+	SHADER_PARAMETER_TEXTURE(Texture2D, AtmosphereIrradianceTexture)
+	SHADER_PARAMETER_SAMPLER(SamplerState, AtmosphereIrradianceTextureSampler)
+	SHADER_PARAMETER_TEXTURE(Texture3D, AtmosphereInscatterTexture)
+	SHADER_PARAMETER_SAMPLER(SamplerState, AtmosphereInscatterTextureSampler)
+	SHADER_PARAMETER_TEXTURE(Texture2D, PerlinNoiseGradientTexture)
+	SHADER_PARAMETER_SAMPLER(SamplerState, PerlinNoiseGradientTextureSampler)
+	SHADER_PARAMETER_TEXTURE(Texture3D, PerlinNoise3DTexture)
+	SHADER_PARAMETER_SAMPLER(SamplerState, PerlinNoise3DTextureSampler)
+	SHADER_PARAMETER_TEXTURE(Texture2D<uint>, SobolSamplingTexture)
+	SHADER_PARAMETER_SAMPLER(SamplerState, SharedPointWrappedSampler)
+	SHADER_PARAMETER_SAMPLER(SamplerState, SharedPointClampedSampler)
+	SHADER_PARAMETER_SAMPLER(SamplerState, SharedBilinearWrappedSampler)
+	SHADER_PARAMETER_SAMPLER(SamplerState, SharedBilinearClampedSampler)
+	SHADER_PARAMETER_SAMPLER(SamplerState, SharedTrilinearWrappedSampler)
+	SHADER_PARAMETER_SAMPLER(SamplerState, SharedTrilinearClampedSampler)
+	SHADER_PARAMETER_TEXTURE(Texture2D, PreIntegratedBRDF)
+	SHADER_PARAMETER_SAMPLER(SamplerState, PreIntegratedBRDFSampler)
+
+END_GLOBAL_SHADER_PARAMETER_STRUCT()
 
 /** Copy of the view uniform shader parameters associated with a view for instanced stereo. */
-BEGIN_UNIFORM_BUFFER_STRUCT_WITH_CONSTRUCTOR(FInstancedViewUniformShaderParameters, ENGINE_API)
+BEGIN_GLOBAL_SHADER_PARAMETER_STRUCT_WITH_CONSTRUCTOR(FInstancedViewUniformShaderParameters, ENGINE_API)
 	VIEW_UNIFORM_BUFFER_MEMBER_TABLE
-END_UNIFORM_BUFFER_STRUCT(FInstancedViewUniformShaderParameters)
+END_GLOBAL_SHADER_PARAMETER_STRUCT()
 
 #undef VIEW_UNIFORM_BUFFER_MEMBER_TABLE
 #undef VIEW_UNIFORM_BUFFER_MEMBER
diff --git a/Engine/Source/Runtime/Engine/Public/SpeedTreeWind.h b/Engine/Source/Runtime/Engine/Public/SpeedTreeWind.h
--- a/Engine/Source/Runtime/Engine/Public/SpeedTreeWind.h
+++ b/Engine/Source/Runtime/Engine/Public/SpeedTreeWind.h
@@ -331,44 +331,44 @@ protected:
  * Uniform buffer setup for SpeedTrees.
  */
 
-BEGIN_UNIFORM_BUFFER_STRUCT(FSpeedTreeUniformParameters, ENGINE_API)
-	UNIFORM_MEMBER(FVector4, WindVector)
-	UNIFORM_MEMBER(FVector4, WindGlobal)
-	UNIFORM_MEMBER(FVector4, WindBranch)
-	UNIFORM_MEMBER(FVector4, WindBranchTwitch)
-	UNIFORM_MEMBER(FVector4, WindBranchWhip)
-	UNIFORM_MEMBER(FVector4, WindBranchAnchor)
-	UNIFORM_MEMBER(FVector4, WindBranchAdherences)
-	UNIFORM_MEMBER(FVector4, WindTurbulences)
-	UNIFORM_MEMBER(FVector4, WindLeaf1Ripple)
-	UNIFORM_MEMBER(FVector4, WindLeaf1Tumble)
-	UNIFORM_MEMBER(FVector4, WindLeaf1Twitch)
-	UNIFORM_MEMBER(FVector4, WindLeaf2Ripple)
-	UNIFORM_MEMBER(FVector4, WindLeaf2Tumble)
-	UNIFORM_MEMBER(FVector4, WindLeaf2Twitch)
-	UNIFORM_MEMBER(FVector4, WindFrondRipple)
-	UNIFORM_MEMBER(FVector4, WindRollingBranch)
-	UNIFORM_MEMBER(FVector4, WindRollingLeafAndDirection)
-	UNIFORM_MEMBER(FVector4, WindRollingNoise)
-	UNIFORM_MEMBER(FVector4, WindAnimation)
+BEGIN_GLOBAL_SHADER_PARAMETER_STRUCT(FSpeedTreeUniformParameters, ENGINE_API)
+	SHADER_PARAMETER(FVector4, WindVector)
+	SHADER_PARAMETER(FVector4, WindGlobal)
+	SHADER_PARAMETER(FVector4, WindBranch)
+	SHADER_PARAMETER(FVector4, WindBranchTwitch)
+	SHADER_PARAMETER(FVector4, WindBranchWhip)
+	SHADER_PARAMETER(FVector4, WindBranchAnchor)
+	SHADER_PARAMETER(FVector4, WindBranchAdherences)
+	SHADER_PARAMETER(FVector4, WindTurbulences)
+	SHADER_PARAMETER(FVector4, WindLeaf1Ripple)
+	SHADER_PARAMETER(FVector4, WindLeaf1Tumble)
+	SHADER_PARAMETER(FVector4, WindLeaf1Twitch)
+	SHADER_PARAMETER(FVector4, WindLeaf2Ripple)
+	SHADER_PARAMETER(FVector4, WindLeaf2Tumble)
+	SHADER_PARAMETER(FVector4, WindLeaf2Twitch)
+	SHADER_PARAMETER(FVector4, WindFrondRipple)
+	SHADER_PARAMETER(FVector4, WindRollingBranch)
+	SHADER_PARAMETER(FVector4, WindRollingLeafAndDirection)
+	SHADER_PARAMETER(FVector4, WindRollingNoise)
+	SHADER_PARAMETER(FVector4, WindAnimation)
 	// Straight copy of the previous members for last frame's values
-	UNIFORM_MEMBER(FVector4, PrevWindVector)
-	UNIFORM_MEMBER(FVector4, PrevWindGlobal)
-	UNIFORM_MEMBER(FVector4, PrevWindBranch)
-	UNIFORM_MEMBER(FVector4, PrevWindBranchTwitch)
-	UNIFORM_MEMBER(FVector4, PrevWindBranchWhip)
-	UNIFORM_MEMBER(FVector4, PrevWindBranchAnchor)
-	UNIFORM_MEMBER(FVector4, PrevWindBranchAdherences)
-	UNIFORM_MEMBER(FVector4, PrevWindTurbulences)
-	UNIFORM_MEMBER(FVector4, PrevWindLeaf1Ripple)
-	UNIFORM_MEMBER(FVector4, PrevWindLeaf1Tumble)
-	UNIFORM_MEMBER(FVector4, PrevWindLeaf1Twitch)
-	UNIFORM_MEMBER(FVector4, PrevWindLeaf2Ripple)
-	UNIFORM_MEMBER(FVector4, PrevWindLeaf2Tumble)
-	UNIFORM_MEMBER(FVector4, PrevWindLeaf2Twitch)
-	UNIFORM_MEMBER(FVector4, PrevWindFrondRipple)
-	UNIFORM_MEMBER(FVector4, PrevWindRollingBranch)
-	UNIFORM_MEMBER(FVector4, PrevWindRollingLeafAndDirection)
-	UNIFORM_MEMBER(FVector4, PrevWindRollingNoise)
-	UNIFORM_MEMBER(FVector4, PrevWindAnimation)
-END_UNIFORM_BUFFER_STRUCT(FSpeedTreeUniformParameters)
+	SHADER_PARAMETER(FVector4, PrevWindVector)
+	SHADER_PARAMETER(FVector4, PrevWindGlobal)
+	SHADER_PARAMETER(FVector4, PrevWindBranch)
+	SHADER_PARAMETER(FVector4, PrevWindBranchTwitch)
+	SHADER_PARAMETER(FVector4, PrevWindBranchWhip)
+	SHADER_PARAMETER(FVector4, PrevWindBranchAnchor)
+	SHADER_PARAMETER(FVector4, PrevWindBranchAdherences)
+	SHADER_PARAMETER(FVector4, PrevWindTurbulences)
+	SHADER_PARAMETER(FVector4, PrevWindLeaf1Ripple)
+	SHADER_PARAMETER(FVector4, PrevWindLeaf1Tumble)
+	SHADER_PARAMETER(FVector4, PrevWindLeaf1Twitch)
+	SHADER_PARAMETER(FVector4, PrevWindLeaf2Ripple)
+	SHADER_PARAMETER(FVector4, PrevWindLeaf2Tumble)
+	SHADER_PARAMETER(FVector4, PrevWindLeaf2Twitch)
+	SHADER_PARAMETER(FVector4, PrevWindFrondRipple)
+	SHADER_PARAMETER(FVector4, PrevWindRollingBranch)
+	SHADER_PARAMETER(FVector4, PrevWindRollingLeafAndDirection)
+	SHADER_PARAMETER(FVector4, PrevWindRollingNoise)
+	SHADER_PARAMETER(FVector4, PrevWindAnimation)
+END_GLOBAL_SHADER_PARAMETER_STRUCT()
diff --git a/Engine/Source/Runtime/EyeTracker/EyeTracker.Build.cs b/Engine/Source/Runtime/EyeTracker/EyeTracker.Build.cs
--- a/Engine/Source/Runtime/EyeTracker/EyeTracker.Build.cs
+++ b/Engine/Source/Runtime/EyeTracker/EyeTracker.Build.cs
@@ -23,7 +23,6 @@ public class EyeTracker : ModuleRules
 // 				"SlateCore",
 //                 "RHI",
 //                 "Renderer",
-//                 "ShaderCore",
 //                 "RenderCore",
 //                 "UtilityShaders",
 //                 "Analytics"
diff --git a/Engine/Source/Runtime/Landscape/Private/LandscapeGrass.cpp b/Engine/Source/Runtime/Landscape/Private/LandscapeGrass.cpp
--- a/Engine/Source/Runtime/Landscape/Private/LandscapeGrass.cpp
+++ b/Engine/Source/Runtime/Landscape/Private/LandscapeGrass.cpp
@@ -217,7 +217,7 @@ protected:
 	: FMeshMaterialShader(Initializer)
 	{
 		RenderOffsetParameter.Bind(Initializer.ParameterMap, TEXT("RenderOffset"));
-		PassUniformBuffer.Bind(Initializer.ParameterMap, FSceneTexturesUniformParameters::StaticStruct.GetShaderVariableName());
+		PassUniformBuffer.Bind(Initializer.ParameterMap, FSceneTexturesUniformParameters::StaticStructMetadata.GetShaderVariableName());
 	}
 
 public:
@@ -264,7 +264,7 @@ public:
 	: FMeshMaterialShader(Initializer)
 	{
 		OutputPassParameter.Bind(Initializer.ParameterMap, TEXT("OutputPass"));
-		PassUniformBuffer.Bind(Initializer.ParameterMap, FSceneTexturesUniformParameters::StaticStruct.GetShaderVariableName());
+		PassUniformBuffer.Bind(Initializer.ParameterMap, FSceneTexturesUniformParameters::StaticStructMetadata.GetShaderVariableName());
 	}
 
 	FLandscapeGrassWeightPS()
diff --git a/Engine/Source/Runtime/Landscape/Private/LandscapeRender.cpp b/Engine/Source/Runtime/Landscape/Private/LandscapeRender.cpp
--- a/Engine/Source/Runtime/Landscape/Private/LandscapeRender.cpp
+++ b/Engine/Source/Runtime/Landscape/Private/LandscapeRender.cpp
@@ -35,7 +35,7 @@ LandscapeRender.cpp: New terrain rendering
 #include "SceneView.h"
 #include "LandscapeProxy.h"
 
-IMPLEMENT_UNIFORM_BUFFER_STRUCT(FLandscapeUniformShaderParameters, TEXT("LandscapeParameters"));
+IMPLEMENT_GLOBAL_SHADER_PARAMETER_STRUCT(FLandscapeUniformShaderParameters, "LandscapeParameters");
 
 int32 GLandscapeMeshLODBias = 0;
 FAutoConsoleVariableRef CVarLandscapeMeshLODBias(
diff --git a/Engine/Source/Runtime/Landscape/Public/LandscapeRender.h b/Engine/Source/Runtime/Landscape/Public/LandscapeRender.h
--- a/Engine/Source/Runtime/Landscape/Public/LandscapeRender.h
+++ b/Engine/Source/Runtime/Landscape/Public/LandscapeRender.h
@@ -83,16 +83,16 @@ LANDSCAPE_API extern UMaterialInterface* GLandscapeLayerUsageMaterial;
 
 
 /** The uniform shader parameters for a landscape draw call. */
-BEGIN_UNIFORM_BUFFER_STRUCT(FLandscapeUniformShaderParameters, LANDSCAPE_API)
+BEGIN_GLOBAL_SHADER_PARAMETER_STRUCT(FLandscapeUniformShaderParameters, LANDSCAPE_API)
 /** vertex shader parameters */
-UNIFORM_MEMBER(FVector4, HeightmapUVScaleBias)
-UNIFORM_MEMBER(FVector4, WeightmapUVScaleBias)
-UNIFORM_MEMBER(FVector4, LandscapeLightmapScaleBias)
-UNIFORM_MEMBER(FVector4, SubsectionSizeVertsLayerUVPan)
-UNIFORM_MEMBER(FVector4, SubsectionOffsetParams)
-UNIFORM_MEMBER(FVector4, LightmapSubsectionOffsetParams)
-UNIFORM_MEMBER(FMatrix, LocalToWorldNoScaling)
-END_UNIFORM_BUFFER_STRUCT(FLandscapeUniformShaderParameters)
+SHADER_PARAMETER(FVector4, HeightmapUVScaleBias)
+SHADER_PARAMETER(FVector4, WeightmapUVScaleBias)
+SHADER_PARAMETER(FVector4, LandscapeLightmapScaleBias)
+SHADER_PARAMETER(FVector4, SubsectionSizeVertsLayerUVPan)
+SHADER_PARAMETER(FVector4, SubsectionOffsetParams)
+SHADER_PARAMETER(FVector4, LightmapSubsectionOffsetParams)
+SHADER_PARAMETER(FMatrix, LocalToWorldNoScaling)
+END_GLOBAL_SHADER_PARAMETER_STRUCT()
 
 /* Data needed for the landscape vertex factory to set the render state for an individual batch element */
 struct FLandscapeBatchElementParams
diff --git a/Engine/Source/Runtime/Launch/Private/LaunchEngineLoop.cpp b/Engine/Source/Runtime/Launch/Private/LaunchEngineLoop.cpp
--- a/Engine/Source/Runtime/Launch/Private/LaunchEngineLoop.cpp
+++ b/Engine/Source/Runtime/Launch/Private/LaunchEngineLoop.cpp
@@ -153,7 +153,7 @@
     #include "PreLoadScreenManager.h"
 
 	#include "ShaderCodeLibrary.h"
-	#include "Rendering/ShaderPipelineStateCache.h"
+	#include "ShaderPipelineCache.h"
 
 #if !UE_BUILD_SHIPPING
 	#include "STaskGraph.h"
@@ -1892,7 +1892,7 @@ int32 FEngineLoop::PreInit(const TCHAR* CmdLine)
 	EnableEmitDrawEventsOnlyOnCommandlist();
 #endif
 
-	FUniformBufferStruct::InitializeStructs();
+	FShaderParametersMetadata::InitializeAllGlobalStructs();
 
 	// Initialize the RHI.
 	RHIInit(bHasEditorToken);
@@ -1908,9 +1908,10 @@ int32 FEngineLoop::PreInit(const TCHAR* CmdLine)
 			
 		// Initialize the pipeline cache system. Opening is deferred until the manual call to
 		// OpenPipelineFileCache below, after content pak's ShaderCodeLibraries are loaded.
-		FShaderPipelineStateCache::Initialize(GMaxRHIShaderPlatform);
+		FShaderPipelineCache::Initialize(GMaxRHIShaderPlatform);
 	}
 
+
 	FString Commandline = FCommandLine::Get();
 	bool EnableShaderCompile = !FParse::Param(*Commandline, TEXT("NoShaderCompile"));
 
@@ -2711,7 +2712,7 @@ void FEngineLoop::LoadPreInitModules()
 
 	// Initialize ShaderCore before loading or compiling any shaders,
 	// But after Renderer and any other modules which implement shader types.
-	FModuleManager::Get().LoadModule(TEXT("ShaderCore"));
+	FModuleManager::Get().LoadModule(TEXT("RenderCore"));
 
 #if WITH_EDITORONLY_DATA
 	// Load the texture compressor module before any textures load. They may
@@ -3196,7 +3197,7 @@ void FEngineLoop::Exit()
 	StopRenderingThread();
 
 	// Disable the PSO cache
-	FShaderPipelineStateCache::Shutdown();
+	FShaderPipelineCache::Shutdown();
 
 	// Close shader code map, if any
 	FShaderCodeLibrary::Shutdown();
diff --git a/Engine/Source/Runtime/OpenGLDrv/Private/OpenGLShaders.cpp b/Engine/Source/Runtime/OpenGLDrv/Private/OpenGLShaders.cpp
--- a/Engine/Source/Runtime/OpenGLDrv/Private/OpenGLShaders.cpp
+++ b/Engine/Source/Runtime/OpenGLDrv/Private/OpenGLShaders.cpp
@@ -2819,7 +2819,7 @@ static void VerifyUniformLayout(const TCHAR* UniformName, const UniformData& GLS
 
 	if(!Uniforms.Num())
 	{
-		for (TLinkedList<FUniformBufferStruct*>::TIterator StructIt(FUniformBufferStruct::GetStructList()); StructIt; StructIt.Next())
+		for (TLinkedList<FShaderParametersMetadata*>::TIterator StructIt(FShaderParametersMetadata::GetStructList()); StructIt; StructIt.Next())
 		{
 #if ENABLE_UNIFORM_BUFFER_LAYOUT_DUMP
 			UE_LOG(LogRHI, Log, TEXT("UniformBufferStruct %s %s %d"),
@@ -2828,15 +2828,15 @@ static void VerifyUniformLayout(const TCHAR* UniformName, const UniformData& GLS
 				StructIt->GetSize()
 				);
 #endif  // #if ENABLE_UNIFORM_BUFFER_LAYOUT_DUMP
-			const TArray<FUniformBufferStruct::FMember>& StructMembers = StructIt->GetMembers();
+			const TArray<FShaderParametersMetadata::FMember>& StructMembers = StructIt->GetMembers();
 			for(int32 MemberIndex = 0;MemberIndex < StructMembers.Num();++MemberIndex)
 			{
-				const FUniformBufferStruct::FMember& Member = StructMembers[MemberIndex];
+				const FShaderParametersMetadata::FMember& Member = StructMembers[MemberIndex];
 
 				FString BaseTypeName;
 				switch(Member.GetBaseType())
 				{
-					case UBMT_STRUCT:  BaseTypeName = TEXT("struct");  break;
+					case UBMT_NESTED_STRUCT:  BaseTypeName = TEXT("struct");  break;
 					case UBMT_BOOL:    BaseTypeName = TEXT("bool"); break;
 					case UBMT_INT32:   BaseTypeName = TEXT("int"); break;
 					case UBMT_UINT32:  BaseTypeName = TEXT("uint"); break;
@@ -2890,12 +2890,12 @@ static void VerifyUniformLayout(const TCHAR* UniformName, const UniformData& GLS
 
 	const UniformData* FoundUniform = Uniforms.Find(RequestedUniformName);
 
-	// MaterialTemplate uniform buffer does not have an entry in the FUniformBufferStructs list, so skipping it here
+	// MaterialTemplate uniform buffer does not have an entry in the FShaderParametersMetadatas list, so skipping it here
 	if(!(RequestedUniformName.StartsWith("Material_") || RequestedUniformName.StartsWith("MaterialCollection")))
 	{
 		if(!FoundUniform || (*FoundUniform != GLSLUniform))
 		{
-			UE_LOG(LogRHI, Fatal, TEXT("uniform buffer member %s in the GLSL source doesn't match it's declaration in it's FUniformBufferStruct"), *RequestedUniformName);
+			UE_LOG(LogRHI, Fatal, TEXT("uniform buffer member %s in the GLSL source doesn't match it's declaration in it's FShaderParametersMetadata"), *RequestedUniformName);
 		}
 	}
 }
@@ -3010,7 +3010,7 @@ static void VerifyUniformBufferLayouts(GLuint Program)
 				const UniformData GLSLUniform
 				(
 					ActiveUniformOffsets[i], 
-					ActiveUniformArrayStrides[i] > 0 ? ActiveUniformSizes[i] : 0 // GLSL has 1 as array size for non-array uniforms, but FUniformBufferStruct assumes 0
+					ActiveUniformArrayStrides[i] > 0 ? ActiveUniformSizes[i] : 0 // GLSL has 1 as array size for non-array uniforms, but FShaderParametersMetadata assumes 0
 				);
 #if ENABLE_UNIFORM_BUFFER_LAYOUT_NAME_MANGLING_CL1862097
 				VerifyUniformLayout(BlockName, ANSI_TO_TCHAR(Buffer), GLSLUniform);
diff --git a/Engine/Source/Runtime/RHI/Public/RHI.h b/Engine/Source/Runtime/RHI/Public/RHI.h
--- a/Engine/Source/Runtime/RHI/Public/RHI.h
+++ b/Engine/Source/Runtime/RHI/Public/RHI.h
@@ -18,8 +18,13 @@
 class FResourceArrayInterface;
 class FResourceBulkDataInterface;
 
-/** Uniform buffer structs must be aligned to 16-byte boundaries. */
-#define UNIFORM_BUFFER_STRUCT_ALIGNMENT 16
+
+/** Alignment of the shader parameters struct is required to be 16-byte boundaries. */
+#define SHADER_PARAMETER_STRUCT_ALIGNMENT 16
+
+/** The alignment in bytes between elements of array shader parameters. */
+#define SHADER_PARAMETER_ARRAY_ELEMENT_ALIGNMENT 16
+
 
 /** RHI Logging. */
 RHI_API DECLARE_LOG_CATEGORY_EXTERN(LogRHI,Log,VeryVerbose);
diff --git a/Engine/Source/Runtime/RHI/Public/RHICommandList.h b/Engine/Source/Runtime/RHI/Public/RHICommandList.h
--- a/Engine/Source/Runtime/RHI/Public/RHICommandList.h
+++ b/Engine/Source/Runtime/RHI/Public/RHICommandList.h
@@ -1626,7 +1626,7 @@ struct FLocalUniformBufferWorkArea
 #endif
 	{
 		check(ContentsSize);
-		Contents = InCheckCmdList->Alloc(ContentsSize, UNIFORM_BUFFER_STRUCT_ALIGNMENT);
+		Contents = InCheckCmdList->Alloc(ContentsSize, SHADER_PARAMETER_STRUCT_ALIGNMENT);
 		FMemory::Memcpy(Contents, InContents, ContentsSize);
 		ComputedUniformBuffer = new (InCheckCmdList->Alloc<FComputedUniformBuffer>()) FComputedUniformBuffer;
 	}
diff --git a/Engine/Source/Runtime/RHI/Public/RHIDefinitions.h b/Engine/Source/Runtime/RHI/Public/RHIDefinitions.h
--- a/Engine/Source/Runtime/RHI/Public/RHIDefinitions.h
+++ b/Engine/Source/Runtime/RHI/Public/RHIDefinitions.h
@@ -382,15 +382,34 @@ enum EUniformBufferUsage
 enum EUniformBufferBaseType
 {
 	UBMT_INVALID,
+
+	// Parameter types.
 	UBMT_BOOL,
 	UBMT_INT32,
 	UBMT_UINT32,
 	UBMT_FLOAT32,
-	UBMT_STRUCT,
+
+	// Untracked resources.
+	UBMT_TEXTURE,
 	UBMT_SRV,
-	UBMT_UAV,
 	UBMT_SAMPLER,
-	UBMT_TEXTURE,
+
+	// Render graph tracked resources.
+	UBMT_GRAPH_TRACKED_TEXTURE,
+	UBMT_GRAPH_TRACKED_SRV,
+	UBMT_GRAPH_TRACKED_UAV,
+
+	// Nested structure.
+	UBMT_NESTED_STRUCT,
+
+	// Structure that is nested on C++ side, but included on shader side.
+	UBMT_INCLUDED_STRUCT,
+
+	// GPU Indirection reference of struct, like is currently named Uniform buffer.
+	UBMT_REFERENCED_STRUCT,
+
+	// Structure dedicated to setup render targets for a rasterizer pass.
+	UBMT_RENDER_TARGET_BINDING_SLOTS,
 
 	EUniformBufferBaseType_Num,
 	EUniformBufferBaseType_NumBits = 4,
@@ -1070,7 +1089,14 @@ inline int32 GetFeatureLevelMaxNumberOfBones(ERHIFeatureLevel::Type FeatureLevel
 
 inline bool IsUniformBufferResourceType(EUniformBufferBaseType BaseType)
 {
-	return BaseType == UBMT_SRV || BaseType == UBMT_UAV || BaseType == UBMT_SAMPLER || BaseType == UBMT_TEXTURE;
+	return
+		BaseType == UBMT_TEXTURE ||
+		BaseType == UBMT_SRV ||
+		BaseType == UBMT_SAMPLER ||
+		BaseType == UBMT_GRAPH_TRACKED_TEXTURE ||
+		BaseType == UBMT_GRAPH_TRACKED_SRV ||
+		BaseType == UBMT_GRAPH_TRACKED_UAV ||
+		BaseType == UBMT_RENDER_TARGET_BINDING_SLOTS;
 }
 
 inline const TCHAR* GetShaderFrequencyString(EShaderFrequency Frequency, bool bIncludePrefix = true)
diff --git a/Engine/Source/Runtime/RHI/Public/RHIResources.h b/Engine/Source/Runtime/RHI/Public/RHIResources.h
--- a/Engine/Source/Runtime/RHI/Public/RHIResources.h
+++ b/Engine/Source/Runtime/RHI/Public/RHIResources.h
@@ -301,6 +301,10 @@ struct FRHIUniformBufferLayout
 
 	const FName GetDebugName() const { return Name; }
 
+	uint32 NumRenderTargets()	const { return 0; }
+	uint32 NumTextures()		const { return 0; }
+	uint32 NumUAVs()			const { return 0; }
+
 private:
 	// for debugging / error message
 	FName Name;
@@ -2247,6 +2251,5 @@ struct FRHIRenderPassInfo
 		return *this;
 	}
 
-private:
 	bool bIsMSAA = false;
 };
diff --git a/Engine/Source/Runtime/RenderCore/Private/GlobalShader.cpp b/Engine/Source/Runtime/RenderCore/Private/GlobalShader.cpp
new file mode 100644
--- /dev/null
+++ b/Engine/Source/Runtime/RenderCore/Private/GlobalShader.cpp
@@ -0,0 +1,222 @@
+// Copyright 1998-2018 Epic Games, Inc. All Rights Reserved.
+
+/*=============================================================================
+	GlobalShader.cpp: Global shader implementation.
+=============================================================================*/
+
+#include "GlobalShader.h"
+
+#include "Misc/MessageDialog.h"
+#include "HAL/FileManager.h"
+#include "Serialization/NameAsStringProxyArchive.h"
+#include "Misc/Paths.h"
+#include "Serialization/MemoryWriter.h"
+#include "Serialization/MemoryReader.h"
+#include "Misc/ScopedSlowTask.h"
+#include "Misc/App.h"
+#include "StaticBoundShaderState.h"
+
+/** The global shader map. */
+TShaderMap<FGlobalShaderType>* GGlobalShaderMap[SP_NumPlatforms] = {};
+
+IMPLEMENT_SHADER_TYPE(,FNULLPS,TEXT("/Engine/Private/NullPixelShader.usf"),TEXT("Main"),SF_Pixel);
+
+/** Used to identify the global shader map in compile queues. */
+const int32 GlobalShaderMapId = 0;
+
+FGlobalShaderMapId::FGlobalShaderMapId(EShaderPlatform Platform)
+{
+	TArray<FShaderType*> ShaderTypes;
+	TArray<const FShaderPipelineType*> ShaderPipelineTypes;
+
+	for (TLinkedList<FShaderType*>::TIterator ShaderTypeIt(FShaderType::GetTypeList()); ShaderTypeIt; ShaderTypeIt.Next())
+	{
+		FGlobalShaderType* GlobalShaderType = ShaderTypeIt->GetGlobalShaderType();
+		if (!GlobalShaderType)
+		{
+			continue;
+		}
+
+		bool bList = false;
+		for (int32 PermutationId = 0; PermutationId < GlobalShaderType->GetPermutationCount(); PermutationId++)
+		{
+			if (GlobalShaderType->ShouldCompilePermutation(Platform, PermutationId))
+			{
+				bList = true;
+				break;
+			}
+		}
+
+		if (bList)
+		{
+			ShaderTypes.Add(GlobalShaderType);
+		}
+	}
+
+	for (TLinkedList<FShaderPipelineType*>::TIterator ShaderPipelineIt(FShaderPipelineType::GetTypeList()); ShaderPipelineIt; ShaderPipelineIt.Next())
+	{
+		const FShaderPipelineType* Pipeline = *ShaderPipelineIt;
+		if (Pipeline->IsGlobalTypePipeline())
+		{
+			int32 NumStagesNeeded = 0;
+			auto& StageTypes = Pipeline->GetStages();
+			for (const FShaderType* Shader : StageTypes)
+			{
+				const FGlobalShaderType* GlobalShaderType = Shader->GetGlobalShaderType();
+				if (GlobalShaderType->ShouldCompilePermutation(Platform, /* PermutationId = */ 0))
+				{
+					++NumStagesNeeded;
+				}
+				else
+				{
+					break;
+				}
+			}
+
+			if (NumStagesNeeded == StageTypes.Num())
+			{
+				ShaderPipelineTypes.Add(Pipeline);
+			}
+		}
+	}
+
+	// Individual shader dependencies
+	ShaderTypes.Sort(FCompareShaderTypes());
+	for (int32 TypeIndex = 0; TypeIndex < ShaderTypes.Num(); TypeIndex++)
+	{
+		FShaderTypeDependency Dependency;
+		Dependency.ShaderType = ShaderTypes[TypeIndex];
+		Dependency.SourceHash = ShaderTypes[TypeIndex]->GetSourceHash(Platform);
+
+		const TCHAR* ShaderFilename = Dependency.ShaderType->GetShaderFilename();
+
+		TArray<FShaderTypeDependency>& Dependencies = ShaderFilenameToDependenciesMap.FindOrAdd(ShaderFilename);
+
+		Dependencies.Add(Dependency);
+	}
+
+	// Shader pipeline dependencies
+	ShaderPipelineTypes.Sort(FCompareShaderPipelineNameTypes());
+	for (int32 TypeIndex = 0; TypeIndex < ShaderPipelineTypes.Num(); TypeIndex++)
+	{
+		const FShaderPipelineType* Pipeline = ShaderPipelineTypes[TypeIndex];
+		FShaderPipelineTypeDependency Dependency;
+		Dependency.ShaderPipelineType = Pipeline;
+		Dependency.StagesSourceHash = Pipeline->GetSourceHash(Platform);
+		ShaderPipelineTypeDependencies.Add(Dependency);
+	}
+}
+
+void FGlobalShaderMapId::AppendKeyString(FString& KeyString, const TArray<FShaderTypeDependency>& Dependencies) const
+{
+	TMap<const TCHAR*,FCachedUniformBufferDeclaration> ReferencedUniformBuffers;
+
+	for (int32 ShaderIndex = 0; ShaderIndex < Dependencies.Num(); ShaderIndex++)
+	{
+		const FShaderTypeDependency& ShaderTypeDependency = Dependencies[ShaderIndex];
+
+		KeyString += TEXT("_");
+		KeyString += ShaderTypeDependency.ShaderType->GetName();
+		KeyString += FString::Printf(TEXT("%i"), ShaderTypeDependency.PermutationId);
+
+		// Add the type's source hash so that we can invalidate cached shaders when .usf changes are made
+		KeyString += ShaderTypeDependency.SourceHash.ToString();
+
+		// Add the serialization history to the key string so that we can detect changes to global shader serialization without a corresponding .usf change
+		ShaderTypeDependency.ShaderType->GetSerializationHistory().AppendKeyString(KeyString);
+
+		const TMap<const TCHAR*, FCachedUniformBufferDeclaration>& ReferencedUniformBufferStructsCache = ShaderTypeDependency.ShaderType->GetReferencedUniformBufferStructsCache();
+
+		// Gather referenced uniform buffers
+		for (TMap<const TCHAR*, FCachedUniformBufferDeclaration>::TConstIterator It(ReferencedUniformBufferStructsCache); It; ++It)
+		{
+			ReferencedUniformBuffers.Add(It.Key(), It.Value());
+		}
+	}
+
+	for (int32 Index = 0; Index < ShaderPipelineTypeDependencies.Num(); ++Index)
+	{
+		const FShaderPipelineTypeDependency& Dependency = ShaderPipelineTypeDependencies[Index];
+
+		KeyString += TEXT("_");
+		KeyString += Dependency.ShaderPipelineType->GetName();
+
+		// Add the type's source hash so that we can invalidate cached shaders when .usf changes are made
+		KeyString += Dependency.StagesSourceHash.ToString();
+
+		for (const FShaderType* ShaderType : Dependency.ShaderPipelineType->GetStages())
+		{
+			const TMap<const TCHAR*, FCachedUniformBufferDeclaration>& ReferencedUniformBufferStructsCache = ShaderType->GetReferencedUniformBufferStructsCache();
+
+			// Gather referenced uniform buffers
+			for (TMap<const TCHAR*, FCachedUniformBufferDeclaration>::TConstIterator It(ReferencedUniformBufferStructsCache); It; ++It)
+			{
+				ReferencedUniformBuffers.Add(It.Key(), It.Value());
+			}
+		}
+	}
+
+	{
+		TArray<uint8> TempData;
+		FSerializationHistory SerializationHistory;
+		FMemoryWriter Ar(TempData, true);
+		FShaderSaveArchive SaveArchive(Ar, SerializationHistory);
+
+		// Save uniform buffer member info so we can detect when layout has changed
+		SerializeUniformBufferInfo(SaveArchive, ReferencedUniformBuffers);
+
+		SerializationHistory.AppendKeyString(KeyString);
+	}
+}
+
+FGlobalShader::FGlobalShader(const ShaderMetaType::CompiledShaderInitializerType& Initializer)
+:	FShader(Initializer)
+{}
+
+void BackupGlobalShaderMap(FGlobalShaderBackupData& OutGlobalShaderBackup)
+{
+	for (int32 i = (int32)ERHIFeatureLevel::ES2; i < (int32)ERHIFeatureLevel::Num; ++i)
+	{
+		EShaderPlatform ShaderPlatform = GetFeatureLevelShaderPlatform((ERHIFeatureLevel::Type)i);
+		if (ShaderPlatform < EShaderPlatform::SP_NumPlatforms && GGlobalShaderMap[ShaderPlatform] != nullptr)
+		{
+			TUniquePtr<TArray<uint8>> ShaderData = MakeUnique<TArray<uint8>>();
+			FMemoryWriter Ar(*ShaderData);
+			GGlobalShaderMap[ShaderPlatform]->SerializeInline(Ar, true, true, false, nullptr);
+			GGlobalShaderMap[ShaderPlatform]->RegisterSerializedShaders(false);
+			GGlobalShaderMap[ShaderPlatform]->Empty();
+			OutGlobalShaderBackup.FeatureLevelShaderData[i] = MoveTemp(ShaderData);
+		}
+	}
+
+	// Remove cached references to global shaders
+	for (TLinkedList<FGlobalBoundShaderStateResource*>::TIterator It(FGlobalBoundShaderStateResource::GetGlobalBoundShaderStateList()); It; It.Next())
+	{
+		BeginUpdateResourceRHI(*It);
+	}
+}
+
+void RestoreGlobalShaderMap(const FGlobalShaderBackupData& GlobalShaderBackup)
+{
+	for (int32 i = (int32)ERHIFeatureLevel::ES2; i < (int32)ERHIFeatureLevel::Num; ++i)
+	{
+		EShaderPlatform ShaderPlatform = GetFeatureLevelShaderPlatform((ERHIFeatureLevel::Type)i);		
+		if (GlobalShaderBackup.FeatureLevelShaderData[i] != nullptr
+			&& ShaderPlatform < EShaderPlatform::SP_NumPlatforms
+			&& GGlobalShaderMap[ShaderPlatform] != nullptr)
+		{
+			FMemoryReader Ar(*GlobalShaderBackup.FeatureLevelShaderData[i]);
+			GGlobalShaderMap[ShaderPlatform]->SerializeInline(Ar, true, true, false, nullptr);
+			GGlobalShaderMap[ShaderPlatform]->RegisterSerializedShaders(false);
+		}
+	}
+}
+
+
+TShaderMap<FGlobalShaderType>* GetGlobalShaderMap(EShaderPlatform Platform)
+{
+	// If the global shader map hasn't been created yet
+	check(GGlobalShaderMap[Platform]);
+
+	return GGlobalShaderMap[Platform];
+}
diff --git a/Engine/Source/Runtime/RenderCore/Private/RenderCore.cpp b/Engine/Source/Runtime/RenderCore/Private/RenderCore.cpp
--- a/Engine/Source/Runtime/RenderCore/Private/RenderCore.cpp
+++ b/Engine/Source/Runtime/RenderCore/Private/RenderCore.cpp
@@ -8,8 +8,22 @@
 #include "HAL/IConsoleManager.h"
 #include "UniformBuffer.h"
 #include "Modules/ModuleManager.h"
+#include "Shader.h"
 
-IMPLEMENT_MODULE(FDefaultModuleImpl, RenderCore);
+void UpdateShaderDevelopmentMode();
+
+class FRenderCoreModule : public FDefaultModuleImpl
+{
+public:
+
+	virtual void StartupModule() override
+	{
+		// TODO(RDG): Why is this not getting called?!
+		IConsoleManager::Get().RegisterConsoleVariableSink_Handle(FConsoleCommandDelegate::CreateStatic(&UpdateShaderDevelopmentMode));
+	}
+};
+
+IMPLEMENT_MODULE(FRenderCoreModule, RenderCore);
 
 DEFINE_LOG_CATEGORY(LogRendererCore);
 
@@ -178,8 +192,6 @@ DEFINE_STAT(STAT_RenderTargetPoolSize);
 DEFINE_STAT(STAT_RenderTargetPoolUsed);
 DEFINE_STAT(STAT_RenderTargetPoolCount);
 
-static TLinkedList<FUniformBufferStruct*>* GUniformStructList = nullptr;
-
 #define EXPOSE_FORCE_LOD !(UE_BUILD_SHIPPING || UE_BUILD_TEST)
 
 #if EXPOSE_FORCE_LOD
@@ -234,30 +246,6 @@ void FInputLatencyTimer::GameThreadTick()
 #endif
 }
 
-TLinkedList<FUniformBufferStruct*>*& FUniformBufferStruct::GetStructList()
-{
-	return GUniformStructList;
-}
-
-TMap<FName, FUniformBufferStruct*>& FUniformBufferStruct::GetNameStructMap()
-{
-	static TMap<FName, FUniformBufferStruct*> GlobalNameStructMap;
-	return GlobalNameStructMap;
-}
-
-
-FUniformBufferStruct* FindUniformBufferStructByName(const TCHAR* StructName)
-{
-	FName FindByName(StructName, FNAME_Find);
-	FUniformBufferStruct* FoundStruct = FUniformBufferStruct::GetNameStructMap().FindRef(FindByName);
-	return FoundStruct;
-}
-
-FUniformBufferStruct* FindUniformBufferStructByFName(FName StructName)
-{
-	return FUniformBufferStruct::GetNameStructMap().FindRef(StructName);
-}
-
 // Can be optimized to avoid the virtual function call but it's compiled out for final release anyway
 RENDERCORE_API int32 GetCVarForceLOD()
 {
@@ -318,143 +306,3 @@ RENDERCORE_API bool IsHDRAllowed()
 
 	return false;
 }
-
-class FUniformBufferMemberAndOffset
-{
-public:
-	FUniformBufferMemberAndOffset(const FUniformBufferStruct::FMember& InMember, int32 InStructOffset) :
-		Member(InMember),
-		StructOffset(InStructOffset)
-	{}
-
-	FUniformBufferStruct::FMember Member;
-	int32 StructOffset;
-};
-
-FUniformBufferStruct::FUniformBufferStruct(const FName& InLayoutName, const TCHAR* InStructTypeName, const TCHAR* InShaderVariableName, ConstructUniformBufferParameterType InConstructRef, uint32 InSize, const TArray<FMember>& InMembers, bool bRegisterForAutoBinding)
-:	StructTypeName(InStructTypeName)
-,	ShaderVariableName(InShaderVariableName)
-,	ConstructUniformBufferParameterRef(InConstructRef)
-,	Size(InSize)
-,	bLayoutInitialized(false)
-,	Layout(InLayoutName)
-,	Members(InMembers)
-,	GlobalListLink(this)
-{
-	if (bRegisterForAutoBinding)
-	{
-		GlobalListLink.LinkHead(GetStructList());
-		FName StrutTypeFName(StructTypeName);
-		// Verify that during FName creation there's no case conversion
-		checkSlow(FCString::Strcmp(StructTypeName, *StrutTypeFName.GetPlainNameString()) == 0);
-		GetNameStructMap().Add(FName(StructTypeName), this);
-	}
-	else
-	{
-		// We cannot initialize the layout during global initialization, since we have to walk nested struct members.
-		// Structs created during global initialization will have bRegisterForAutoBinding==false, and are initialized during startup.
-		// Structs created at runtime with bRegisterForAutoBinding==true can be initialized now.
-		InitializeLayout();
-	}
-}
-
-void FUniformBufferStruct::InitializeStructs()
-{
-	for (TLinkedList<FUniformBufferStruct*>::TIterator StructIt(FUniformBufferStruct::GetStructList()); StructIt; StructIt.Next())
-	{
-		StructIt->InitializeLayout();
-	}
-}
-
-void FUniformBufferStruct::InitializeLayout()
-{
-	check(!bLayoutInitialized);
-	Layout.ConstantBufferSize = Size;
-
-	TArray<FUniformBufferMemberAndOffset> MemberStack;
-	MemberStack.Reserve(Members.Num());
-
-	for (int32 MemberIndex = 0; MemberIndex < Members.Num(); MemberIndex++)
-	{
-		MemberStack.Push(FUniformBufferMemberAndOffset(Members[MemberIndex], 0));
-	}
-
-	for (int32 i = 0; i < MemberStack.Num(); ++i)
-	{
-		const FMember& CurrentMember = MemberStack[i].Member;
-		bool bIsResource = IsUniformBufferResourceType(CurrentMember.GetBaseType());
-
-		if (bIsResource)
-		{
-			Layout.Resources.Add(CurrentMember.GetBaseType());
-			const uint32 AbsoluteMemberOffset = CurrentMember.GetOffset() + MemberStack[i].StructOffset;
-			check(AbsoluteMemberOffset < (1u << (Layout.ResourceOffsets.GetTypeSize() * 8)));
-			Layout.ResourceOffsets.Add(AbsoluteMemberOffset);
-		}
-
-		const FUniformBufferStruct* MemberStruct = CurrentMember.GetStruct();
-
-		if (MemberStruct)
-		{
-			int32 AbsoluteStructOffset = CurrentMember.GetOffset() + MemberStack[i].StructOffset;
-
-			for (int32 StructMemberIndex = 0; StructMemberIndex < MemberStruct->Members.Num(); StructMemberIndex++)
-			{
-				FMember StructMember = MemberStruct->Members[StructMemberIndex];
-				MemberStack.Insert(FUniformBufferMemberAndOffset(StructMember, AbsoluteStructOffset), i + 1 + StructMemberIndex);
-			}
-		}
-	}
-
-	Layout.ComputeHash();
-
-	bLayoutInitialized = true;
-}
-
-void FUniformBufferStruct::GetNestedStructs(TArray<const FUniformBufferStruct*>& OutNestedStructs) const
-{
-	for (int32 i = 0; i < Members.Num(); ++i)
-	{
-		const FMember& CurrentMember = Members[i];
-
-		const FUniformBufferStruct* MemberStruct = CurrentMember.GetStruct();
-
-		if (MemberStruct)
-		{
-			OutNestedStructs.Add(MemberStruct);
-			MemberStruct->GetNestedStructs(OutNestedStructs);
-		}
-	}
-}
-
-void FUniformBufferStruct::AddResourceTableEntries(TMap<FString,FResourceTableEntry>& ResourceTableMap, TMap<FString,uint32>& ResourceTableLayoutHashes) const
-{
-	uint16 ResourceIndex = 0;
-	FString Prefix = FString::Printf(TEXT("%s_"), ShaderVariableName);
-	AddResourceTableEntriesRecursive(ShaderVariableName, *Prefix, ResourceIndex, ResourceTableMap);
-	ResourceTableLayoutHashes.Add(ShaderVariableName,GetLayout().GetHash());
-}
-
-void FUniformBufferStruct::AddResourceTableEntriesRecursive(const TCHAR* UniformBufferName, const TCHAR* Prefix, uint16& ResourceIndex, TMap<FString, FResourceTableEntry>& ResourceTableMap) const
-{
-	for (int32 MemberIndex = 0; MemberIndex < Members.Num(); ++MemberIndex)
-	{
-		const FMember& Member = Members[MemberIndex];
-		if (IsUniformBufferResourceType(Member.GetBaseType()))
-		{
-			FResourceTableEntry& Entry = ResourceTableMap.FindOrAdd(FString::Printf(TEXT("%s%s"), Prefix, Member.GetName()));
-			if (Entry.UniformBufferName.IsEmpty())
-			{
-				Entry.UniformBufferName = UniformBufferName;
-				Entry.Type = Member.GetBaseType();
-				Entry.ResourceIndex = ResourceIndex++;
-			}
-		}
-		else if (Member.GetBaseType() == UBMT_STRUCT)
-		{
-			check(Member.GetStruct());
-			FString MemberPrefix = FString::Printf(TEXT("%s%s_"), Prefix, Member.GetName());
-			Member.GetStruct()->AddResourceTableEntriesRecursive(UniformBufferName, *MemberPrefix, ResourceIndex, ResourceTableMap);
-		}
-	}
-}
\ No newline at end of file
diff --git a/Engine/Source/Runtime/RenderCore/Private/RenderGraphBuilder.cpp b/Engine/Source/Runtime/RenderCore/Private/RenderGraphBuilder.cpp
new file mode 100644
--- /dev/null
+++ b/Engine/Source/Runtime/RenderCore/Private/RenderGraphBuilder.cpp
@@ -0,0 +1,646 @@
+// Copyright 1998-2018 Epic Games, Inc. All Rights Reserved.
+
+#include "RenderGraphBuilder.h"
+#include "RenderCore.h"
+#include "RenderTargetPool.h"
+
+
+#if RENDER_GRAPH_DEBUGGING
+
+static int32 GRenderGraphImmediateMode = 0;
+static FAutoConsoleVariableRef CVarImmediateMode(
+	TEXT("r.Graph.ImmediateMode"),
+	GRenderGraphImmediateMode,
+	TEXT("Executes passes as they get created. Extremely useful to have a callstack of the wiring code when crashing in the pass' lambda."),
+	ECVF_RenderThreadSafe);
+
+#else
+
+static const int32 GRenderGraphImmediateMode = 0;
+
+#endif
+
+
+void FRenderGraphBuilder::Execute()
+{
+	checkf(!bHasExecuted, TEXT("Render graph execution should only happen once."));
+
+	if (!GRenderGraphImmediateMode)
+	{
+		WalkGraphDependencies();
+
+		QUICK_SCOPE_CYCLE_COUNTER(STAT_FRenderGraphBuilder_Execute);
+		for (const FRenderGraphPass* Pass : Passes)
+		{
+			ExecutePass(Pass);
+		}
+	}
+
+	ProcessDeferredInternalResourceQueries();
+
+	DestructPasses();
+
+	#if DO_CHECK
+		bHasExecuted = true;
+	#endif
+}
+
+void FRenderGraphBuilder::DebugPass(const FRenderGraphPass* Pass)
+{
+	ValidatePass(Pass);
+
+	if (GRenderGraphImmediateMode)
+	{
+		ExecutePass(Pass);
+	}
+}
+
+void FRenderGraphBuilder::ValidatePass(const FRenderGraphPass* Pass) const
+{
+	FRenderTargetBindingSlots* RESTRICT RenderTargets = nullptr;
+	FShaderParameterStructRef ParameterStruct = Pass->GetParameters();
+
+	bool bIsCompute = Pass->IsCompute();
+	bool bCanUseUAVs = bIsCompute;
+	bool bRequiresRenderTargetSlots = !bIsCompute;
+
+	for (int ResourceIndex = 0, Num = ParameterStruct.Layout->Resources.Num(); ResourceIndex < Num; ResourceIndex++)
+	{
+		uint8  Type = ParameterStruct.Layout->Resources[ResourceIndex];
+		uint16 Offset = ParameterStruct.Layout->ResourceOffsets[ResourceIndex];
+
+		switch (Type)
+		{
+		case UBMT_GRAPH_TRACKED_UAV:
+		{
+			FGraphUAV* RESTRICT UAV = *ParameterStruct.GetMemberPtrAtOffset<FGraphUAV*>(Offset);
+			if (UAV && !bCanUseUAVs)
+			{
+				UE_LOG(LogRendererCore, Warning, TEXT("UAV can only been bound to compute shaders, therefore UAV %s is certainly useless for passs %s."), UAV->Name, Pass->GetName());
+			}
+		}
+		break;
+		case UBMT_RENDER_TARGET_BINDING_SLOTS:
+		{
+			if (RenderTargets)
+			{
+				UE_LOG(LogRendererCore, Warning, TEXT("Pass %s have duplicated render target binding slots."), Pass->GetName());
+			}
+			else
+			{
+				RenderTargets = ParameterStruct.GetMemberPtrAtOffset<FRenderTargetBindingSlots>(Offset);
+			}
+		}
+		break;
+		default:
+			break;
+		}
+	}
+
+	if (RenderTargets)
+	{
+		checkf(bRequiresRenderTargetSlots, TEXT("Render pass %s does not need render target binging slots"), Pass->GetName());
+
+		uint32 NumRenderTargets = 0;
+		for (int i = 0; i < RenderTargets->Output.Num(); i++)
+		{
+			const FRenderTargetBinding& RenderTarget = RenderTargets->Output[i];
+			if (!RenderTarget.GetTexture())
+			{
+				NumRenderTargets = i;
+				break;
+			}
+		}
+		for (int i = NumRenderTargets; i < RenderTargets->Output.Num(); i++)
+		{
+			const FRenderTargetBinding& RenderTarget = RenderTargets->Output[i];
+			checkf(RenderTarget.GetTexture() == nullptr, TEXT("Render targets must be packed. No empty spaces in the array."));
+		}
+	}
+	else
+	{
+		checkf(!bRequiresRenderTargetSlots, TEXT("Render pass %s requires render target binging slots"), Pass->GetName());
+	}
+}
+
+void FRenderGraphBuilder::WalkGraphDependencies()
+{
+	for (const FRenderGraphPass* Pass : Passes)
+	{
+		FShaderParameterStructRef ParameterStruct = Pass->GetParameters();
+
+		/** Increments all the FGraphResource::ReferenceCount. */
+		for (int ResourceIndex = 0, Num = ParameterStruct.Layout->Resources.Num(); ResourceIndex < Num; ResourceIndex++)
+		{
+			uint8  Type = ParameterStruct.Layout->Resources[ResourceIndex];
+			uint16 Offset = ParameterStruct.Layout->ResourceOffsets[ResourceIndex];
+
+			switch (Type)
+			{
+			case UBMT_GRAPH_TRACKED_TEXTURE:
+			{
+				FGraphTexture* RESTRICT Texture = *ParameterStruct.GetMemberPtrAtOffset<FGraphTexture*>(Offset);
+				if (Texture)
+				{
+					Texture->ReferenceCount++;
+				}
+			}
+			break;
+			case UBMT_GRAPH_TRACKED_SRV:
+			{
+				FGraphSRV* RESTRICT SRV = *ParameterStruct.GetMemberPtrAtOffset<FGraphSRV*>(Offset);
+				if (SRV)
+				{
+					SRV->Desc.Texture->ReferenceCount++;
+				}
+			}
+			break;
+			case UBMT_GRAPH_TRACKED_UAV:
+			{
+				FGraphUAV* RESTRICT UAV = *ParameterStruct.GetMemberPtrAtOffset<FGraphUAV*>(Offset);
+				if (UAV)
+				{
+					UAV->Desc.Texture->ReferenceCount++;
+				}
+			}
+			break;
+			case UBMT_RENDER_TARGET_BINDING_SLOTS:
+			{
+				FRenderTargetBindingSlots* RESTRICT RenderTargets = ParameterStruct.GetMemberPtrAtOffset<FRenderTargetBindingSlots>(Offset);
+
+				for (int i = 0; i < RenderTargets->Output.Num(); i++)
+				{
+					const FRenderTargetBinding& RenderTarget = RenderTargets->Output[i];
+					if (RenderTarget.GetTexture())
+					{
+						RenderTarget.GetTexture()->ReferenceCount++;
+					}
+					else
+					{
+						break;
+					}
+				}
+
+				const FDepthStencilBinding& DepthStencil = RenderTargets->DepthStencil;
+				if (DepthStencil.Texture)
+				{
+					DepthStencil.Texture->ReferenceCount++;
+				}
+			}
+			break;
+			default:
+				break;
+			}
+		}
+	} // for (const FRenderGraphPass* Pass : Passes)
+
+	// Add additional dependencies from deferred queries.
+	for (const auto& Query : DeferredInternalTextureQueries)
+	{
+		Query.Texture->ReferenceCount++;
+	}
+
+	// Release external texture that have ReferenceCount == 0 and yet are already allocated.
+	for (auto Pair : AllocatedTextures)
+	{
+		if (Pair.Key->ReferenceCount == 0)
+		{
+			Pair.Value = nullptr;
+			Pair.Key->PooledRenderTarget = nullptr;
+		}
+	}
+}
+
+void FRenderGraphBuilder::AllocateRHITextureIfNeeded(const FGraphTexture* Texture, bool bComputePass)
+{
+	check(Texture);
+
+	if (Texture->PooledRenderTarget)
+	{
+		return;
+	}
+
+	check(Texture->ReferenceCount > 0 || GRenderGraphImmediateMode);
+
+	// TODO(RDG): should avoid bDoWritableBarrier = true
+	TRefCountPtr<IPooledRenderTarget>& PooledRenderTarget = AllocatedTextures.FindOrAdd(Texture);
+	GRenderTargetPool.FindFreeElement(RHICmdList, Texture->Desc, PooledRenderTarget, Texture->Name, /* bDoWritableBarrier = */ true);
+
+	Texture->PooledRenderTarget = PooledRenderTarget;
+}
+
+static EResourceTransitionPipeline CalcTransitionPipeline(bool bCurrentCompute, bool bTargetCompute)
+{
+	// TODO(RDG) convert table to math
+	uint32 Bits;
+	Bits  = (uint32)bCurrentCompute;
+	Bits |= (uint32)bTargetCompute << 1;
+
+	EResourceTransitionPipeline Table[] = {
+		EResourceTransitionPipeline::EGfxToGfx,
+		EResourceTransitionPipeline::EComputeToGfx,
+		EResourceTransitionPipeline::EGfxToCompute,
+		EResourceTransitionPipeline::EComputeToCompute
+	};
+	
+	return static_cast< EResourceTransitionPipeline >( Table[ Bits ] );
+}
+
+void FRenderGraphBuilder::TransitionTexture( const FGraphTexture* Texture, EResourceTransitionAccess TransitionAccess, bool bRequiredCompute ) const
+{
+	const bool bRequiredWritable = TransitionAccess != EResourceTransitionAccess::EReadable;
+
+	if( Texture->bWritable != bRequiredWritable || Texture->bCompute != bRequiredCompute )
+	{
+		RHICmdList.TransitionResource( TransitionAccess, Texture->PooledRenderTarget->GetRenderTargetItem().ShaderResourceTexture );
+		Texture->bWritable = bRequiredWritable;
+		Texture->bCompute = bRequiredCompute;
+	}
+}
+
+void FRenderGraphBuilder::TransitionUAV( const FGraphUAV* UAV, EResourceTransitionAccess TransitionAccess, bool bRequiredCompute ) const
+{
+	const bool bRequiredWritable = true;
+
+	if( UAV->Desc.Texture->bWritable != bRequiredWritable || UAV->Desc.Texture->bCompute != bRequiredCompute )
+	{
+		EResourceTransitionPipeline TransitionPipeline = CalcTransitionPipeline( UAV->Desc.Texture->bCompute, bRequiredCompute );
+		RHICmdList.TransitionResource( TransitionAccess, TransitionPipeline, UAV->GetRHIUnorderedAccessView());
+		UAV->Desc.Texture->bWritable = bRequiredWritable;
+		UAV->Desc.Texture->bCompute = bRequiredCompute;
+	}
+}
+
+static bool IsBoundAsReadable( const FGraphTexture* Texture, FShaderParameterStructRef ParameterStruct )
+{
+	for( int i = 0, Num = ParameterStruct.Layout->Resources.Num(); i < Num; i++ )
+	{
+		uint8  Type   = ParameterStruct.Layout->Resources[i];
+		uint16 Offset = ParameterStruct.Layout->ResourceOffsets[i];
+
+		switch( Type )
+		{
+			case UBMT_GRAPH_TRACKED_TEXTURE:
+			{
+				const FGraphTexture* InputTexture = *ParameterStruct.GetMemberPtrAtOffset<const FGraphTexture*>(Offset);
+				if( Texture == InputTexture)
+				{
+					return true;
+				}
+			}
+			break;
+		case UBMT_GRAPH_TRACKED_SRV:
+			{
+				const FGraphSRV* InputSRV = *ParameterStruct.GetMemberPtrAtOffset<const FGraphSRV*>(Offset);
+				if (InputSRV && Texture == InputSRV->Desc.Texture)
+				{
+					return true;
+				}
+			}
+			break;
+		default:
+			break;
+		}
+	}
+
+	return false;
+}
+
+void FRenderGraphBuilder::ExecutePass( const FRenderGraphPass* Pass )
+{
+	QUICK_SCOPE_CYCLE_COUNTER(STAT_FRenderGraphBuilder_ExecutePass);
+
+	FRHIRenderPassInfo RPInfo;
+	bool bHasRenderTargets = false;
+
+	AllocateAndTransitionPassResources(Pass, &RPInfo, &bHasRenderTargets);
+
+	if( !Pass->IsCompute())
+	{
+		check(bHasRenderTargets);
+		RHICmdList.BeginRenderPass( RPInfo, Pass->GetName() );
+	}
+	else
+	{
+		SetRenderTarget( RHICmdList, nullptr, nullptr );
+	}
+	
+	Pass->Execute(RHICmdList);
+
+	if( bHasRenderTargets )
+	{
+		RHICmdList.EndRenderPass();
+	}
+
+	if (DO_CHECK)
+	{
+		WarnForUselessPassDependencies(Pass);
+	}
+
+	// Can't release resources with immediate mode, because don't know if whether they are gonna be used.
+	if (!GRenderGraphImmediateMode)
+	{
+		ReleaseUnecessaryResources(Pass);
+	}
+}
+
+void FRenderGraphBuilder::AllocateAndTransitionPassResources(const FRenderGraphPass* Pass, struct FRHIRenderPassInfo* OutRPInfo, bool* bOutHasRenderTargets)
+{
+	bool bIsCompute = Pass->IsCompute();
+	FShaderParameterStructRef ParameterStruct = Pass->GetParameters();
+
+	for (int ResourceIndex = 0, Num = ParameterStruct.Layout->Resources.Num(); ResourceIndex < Num; ResourceIndex++)
+	{
+		uint8  Type = ParameterStruct.Layout->Resources[ResourceIndex];
+		uint16 Offset = ParameterStruct.Layout->ResourceOffsets[ResourceIndex];
+
+		switch (Type)
+		{
+		case UBMT_GRAPH_TRACKED_TEXTURE:
+		{
+			FGraphTexture* RESTRICT Texture = *ParameterStruct.GetMemberPtrAtOffset<FGraphTexture*>(Offset);
+			if (Texture)
+			{
+				AllocateRHITextureIfNeeded(Texture, bIsCompute);
+				TransitionTexture(Texture, EResourceTransitionAccess::EReadable, bIsCompute);
+			}
+		}
+		break;
+		case UBMT_GRAPH_TRACKED_SRV:
+		{
+			FGraphSRV* RESTRICT SRV = *ParameterStruct.GetMemberPtrAtOffset<FGraphSRV*>(Offset);
+			if (SRV)
+			{
+				check(SRV->Desc.Texture);
+
+				const FGraphTexture* Texture = SRV->Desc.Texture;
+				AllocateRHITextureIfNeeded(Texture, bIsCompute);
+				TransitionTexture(Texture, EResourceTransitionAccess::EReadable, bIsCompute);
+			}
+		}
+		break;
+		case UBMT_GRAPH_TRACKED_UAV:
+		{
+			FGraphUAV* RESTRICT UAV = *ParameterStruct.GetMemberPtrAtOffset<FGraphUAV*>(Offset);
+			if (UAV)
+			{
+				AllocateRHITextureIfNeeded(UAV->Desc.Texture, bIsCompute);
+				TransitionUAV(UAV, EResourceTransitionAccess::EWritable, bIsCompute);
+			}
+		}
+		break;
+		case UBMT_RENDER_TARGET_BINDING_SLOTS:
+		{
+			check(!bIsCompute);
+
+			FRenderTargetBindingSlots* RESTRICT RenderTargets = ParameterStruct.GetMemberPtrAtOffset<FRenderTargetBindingSlots>(Offset);
+
+			uint32 NumRenderTargets = 0;
+			uint32 NumDepthStencilTargets = 0;
+			uint32 NumSamples = 0;
+
+			for (int i = 0; i < RenderTargets->Output.Num(); i++)
+			{
+				const FRenderTargetBinding& RenderTarget = RenderTargets->Output[i];
+				if (RenderTarget.GetTexture())
+				{
+					AllocateRHITextureIfNeeded(RenderTarget.GetTexture(), false);
+
+					// TODO(RDG): should force TargetableTexture == ShaderResourceTexture with MSAA, and instead have an explicit MSAA resolve pass.
+					OutRPInfo->ColorRenderTargets[i].RenderTarget = RenderTarget.GetTexture()->PooledRenderTarget->GetRenderTargetItem().TargetableTexture;
+					OutRPInfo->ColorRenderTargets[i].ResolveTarget = nullptr;
+					OutRPInfo->ColorRenderTargets[i].ArraySlice = -1;
+					OutRPInfo->ColorRenderTargets[i].MipIndex = RenderTarget.GetMipIndex();
+					OutRPInfo->ColorRenderTargets[i].Action = MakeRenderTargetActions(RenderTarget.GetLoadAction(), RenderTarget.GetStoreAction());
+
+					TransitionTexture(RenderTarget.GetTexture(), EResourceTransitionAccess::EWritable, false);
+
+					// TODO(RDG): There must be a better way to do this.
+					OutRPInfo->bGeneratingMips = OutRPInfo->bGeneratingMips || IsBoundAsReadable(RenderTarget.GetTexture(), ParameterStruct);
+
+					NumSamples |= OutRPInfo->ColorRenderTargets[i].RenderTarget->GetNumSamples();
+					NumRenderTargets++;
+				}
+				else
+				{
+					break;
+				}
+			}
+
+			const FDepthStencilBinding& DepthStencil = RenderTargets->DepthStencil;
+			if (DepthStencil.Texture)
+			{
+				AllocateRHITextureIfNeeded(DepthStencil.Texture, false);
+
+				OutRPInfo->DepthStencilRenderTarget.DepthStencilTarget = DepthStencil.Texture->PooledRenderTarget->GetRenderTargetItem().TargetableTexture;
+				OutRPInfo->DepthStencilRenderTarget.ResolveTarget = nullptr;
+				OutRPInfo->DepthStencilRenderTarget.Action = MakeDepthStencilTargetActions(
+					MakeRenderTargetActions(DepthStencil.DepthLoadAction, DepthStencil.DepthStoreAction),
+					MakeRenderTargetActions(DepthStencil.StencilLoadAction, DepthStencil.StencilStoreAction));
+				OutRPInfo->DepthStencilRenderTarget.ExclusiveDepthStencil = FExclusiveDepthStencil::DepthWrite_StencilWrite;
+
+				TransitionTexture(DepthStencil.Texture, EResourceTransitionAccess::EWritable, false);
+
+				NumSamples |= OutRPInfo->DepthStencilRenderTarget.DepthStencilTarget->GetNumSamples();
+				NumDepthStencilTargets++;
+			}
+
+			OutRPInfo->bIsMSAA = NumSamples > 1;
+
+			*bOutHasRenderTargets = NumRenderTargets + NumDepthStencilTargets > 0;
+		}
+		break;
+		default:
+			break;
+		}
+	}
+}
+
+// static 
+void FRenderGraphBuilder::WarnForUselessPassDependencies(const FRenderGraphPass* Pass)
+{
+	FShaderParameterStructRef ParameterStruct = Pass->GetParameters();
+
+	int32 TrackedResourceCount = 0;
+	int32 UsedResourceCount = 0;
+
+	// First pass to count resources.
+	for (int ResourceIndex = 0, Num = ParameterStruct.Layout->Resources.Num(); ResourceIndex < Num; ResourceIndex++)
+	{
+		uint8 Type = ParameterStruct.Layout->Resources[ResourceIndex];
+		uint16 Offset = ParameterStruct.Layout->ResourceOffsets[ResourceIndex];
+
+		if (Type != UBMT_GRAPH_TRACKED_TEXTURE && Type != UBMT_GRAPH_TRACKED_SRV && Type != UBMT_GRAPH_TRACKED_UAV)
+			continue;
+
+		const FGraphResource* Resource = *ParameterStruct.GetMemberPtrAtOffset<const FGraphResource*>(Offset);
+
+		if (!Resource)
+			continue;
+
+		TrackedResourceCount++;
+		UsedResourceCount += Resource->bIsActuallyUsedByPass ? 1 : 0;
+	}
+
+	if (TrackedResourceCount != UsedResourceCount)
+	{
+		UE_LOG(LogRendererCore, Warning, TEXT("%i of the %i resources of the pass %s where not actually used."), TrackedResourceCount - UsedResourceCount, TrackedResourceCount, Pass->GetName());
+
+		for (int ResourceIndex = 0, Num = ParameterStruct.Layout->Resources.Num(); ResourceIndex < Num; ResourceIndex++)
+		{
+			uint8 Type = ParameterStruct.Layout->Resources[ResourceIndex];
+			uint16 Offset = ParameterStruct.Layout->ResourceOffsets[ResourceIndex];
+
+			if (Type != UBMT_GRAPH_TRACKED_TEXTURE && Type != UBMT_GRAPH_TRACKED_SRV && Type != UBMT_GRAPH_TRACKED_UAV)
+				continue;
+
+			const FGraphResource* Resource = *ParameterStruct.GetMemberPtrAtOffset<const FGraphResource*>(Offset);
+
+			if (!Resource)
+				continue;
+
+			UE_LOG(LogRendererCore, Warning, TEXT("	%s"), Resource->Name);
+		}
+	}
+
+	// Last pass to clean the bIsActuallyUsedByPass flags.
+	for (int ResourceIndex = 0, Num = ParameterStruct.Layout->Resources.Num(); ResourceIndex < Num; ResourceIndex++)
+	{
+		uint8 Type = ParameterStruct.Layout->Resources[ResourceIndex];
+		uint16 Offset = ParameterStruct.Layout->ResourceOffsets[ResourceIndex];
+
+		if (Type != UBMT_GRAPH_TRACKED_TEXTURE && Type != UBMT_GRAPH_TRACKED_SRV && Type != UBMT_GRAPH_TRACKED_UAV)
+			continue;
+
+		const FGraphResource* Resource = *ParameterStruct.GetMemberPtrAtOffset<const FGraphResource*>(Offset);
+
+		if (!Resource)
+			continue;
+
+		Resource->bIsActuallyUsedByPass = false;
+	}
+}
+
+void FRenderGraphBuilder::ReleaseRHITextureIfPossible(const FGraphTexture* Texture)
+{
+	check(Texture->ReferenceCount > 0);
+	Texture->ReferenceCount--;
+
+	if (Texture->ReferenceCount == 0)
+	{
+		Texture->PooledRenderTarget = nullptr;
+		AllocatedTextures.FindChecked(Texture) = nullptr;
+	}
+}
+
+void FRenderGraphBuilder::ReleaseUnecessaryResources(const FRenderGraphPass* Pass)
+{
+	FShaderParameterStructRef ParameterStruct = Pass->GetParameters();
+
+	/** Increments all the FGraphResource::ReferenceCount. */
+	for (int ResourceIndex = 0, Num = ParameterStruct.Layout->Resources.Num(); ResourceIndex < Num; ResourceIndex++)
+	{
+		uint8  Type = ParameterStruct.Layout->Resources[ResourceIndex];
+		uint16 Offset = ParameterStruct.Layout->ResourceOffsets[ResourceIndex];
+
+		switch (Type)
+		{
+		case UBMT_GRAPH_TRACKED_TEXTURE:
+		{
+			FGraphTexture* RESTRICT Texture = *ParameterStruct.GetMemberPtrAtOffset<FGraphTexture*>(Offset);
+			if (Texture)
+			{
+				ReleaseRHITextureIfPossible(Texture);
+			}
+		}
+		break;
+		case UBMT_GRAPH_TRACKED_SRV:
+		{
+			FGraphSRV* RESTRICT SRV = *ParameterStruct.GetMemberPtrAtOffset<FGraphSRV*>(Offset);
+			if (SRV)
+			{
+				ReleaseRHITextureIfPossible(SRV->Desc.Texture);
+			}
+		}
+		break;
+		case UBMT_GRAPH_TRACKED_UAV:
+		{
+			FGraphUAV* RESTRICT UAV = *ParameterStruct.GetMemberPtrAtOffset<FGraphUAV*>(Offset);
+			if (UAV)
+			{
+				ReleaseRHITextureIfPossible(UAV->Desc.Texture);
+			}
+		}
+		break;
+		case UBMT_RENDER_TARGET_BINDING_SLOTS:
+		{
+			FRenderTargetBindingSlots* RESTRICT RenderTargets = ParameterStruct.GetMemberPtrAtOffset<FRenderTargetBindingSlots>(Offset);
+
+			for (int i = 0; i < RenderTargets->Output.Num(); i++)
+			{
+				const FRenderTargetBinding& RenderTarget = RenderTargets->Output[i];
+				if (RenderTarget.GetTexture())
+				{
+					ReleaseRHITextureIfPossible(RenderTarget.GetTexture());
+				}
+				else
+				{
+					break;
+				}
+			}
+			if (RenderTargets->DepthStencil.Texture)
+			{
+				ReleaseRHITextureIfPossible(RenderTargets->DepthStencil.Texture);
+			}
+		}
+		break;
+		default:
+			break;
+		}
+	}
+}
+
+void FRenderGraphBuilder::ProcessDeferredInternalResourceQueries()
+{
+	for (const auto& Query : DeferredInternalTextureQueries)
+	{
+		check(Query.Texture->PooledRenderTarget);
+
+		if (Query.bTransitionToRead)
+		{
+			RHICmdList.TransitionResource(EResourceTransitionAccess::EReadable, Query.Texture->PooledRenderTarget->GetRenderTargetItem().ShaderResourceTexture);
+		}
+
+		*Query.OutTexturePtr = AllocatedTextures.FindChecked(Query.Texture);
+
+		// No need to manually release in immediate mode, since it is done directly when emptying AllocatedTextures in DestructPasses().
+		if (!GRenderGraphImmediateMode)
+		{
+			ReleaseRHITextureIfPossible(Query.Texture);
+		}
+	}
+}
+
+void FRenderGraphBuilder::DestructPasses()
+{
+	#if DO_CHECK
+	{
+		// Make sure all resource references have been released to ensure no leaks happen.
+		for (const FGraphResource* Resource : Resources)
+		{
+			check(Resource->ReferenceCount == 0);
+		}
+		Resources.Empty();
+	}
+	#endif
+
+	// Passes are allocated on FMemStack, so need to call destructor manually.
+	for (FRenderGraphPass* Pass : Passes)
+	{
+		Pass->~FRenderGraphPass();
+	}
+
+	Passes.Empty();
+	DeferredInternalTextureQueries.Empty();
+	AllocatedTextures.Empty();
+}
\ No newline at end of file
diff --git a/Engine/Source/Runtime/RenderCore/Private/RenderGraphUtils.cpp b/Engine/Source/Runtime/RenderCore/Private/RenderGraphUtils.cpp
new file mode 100644
--- /dev/null
+++ b/Engine/Source/Runtime/RenderCore/Private/RenderGraphUtils.cpp
@@ -0,0 +1,51 @@
+// Copyright 1998-2018 Epic Games, Inc. All Rights Reserved.
+
+#include "RenderGraphUtils.h"
+
+
+void ClearUnusedGraphResourcesImpl(const FShaderParameterBindings& ShaderBindings, const FShaderParametersMetadata* ParametersMetadata, void* InoutParameters)
+{
+	int32 GraphTextureId = 0;
+	int32 GraphSRVId = 0;
+	int32 GraphUAVId = 0;
+
+	uint8* Base = reinterpret_cast<uint8*>(InoutParameters);
+
+	for (int32 ResourceIndex = 0, Num = ParametersMetadata->GetLayout().Resources.Num(); ResourceIndex < Num; ResourceIndex++)
+	{
+		EUniformBufferBaseType Type = EUniformBufferBaseType(ParametersMetadata->GetLayout().Resources[ResourceIndex]);
+		uint16 ByteOffset = ParametersMetadata->GetLayout().ResourceOffsets[ResourceIndex];
+
+		if (Type == UBMT_GRAPH_TRACKED_TEXTURE)
+		{
+			if (GraphTextureId < ShaderBindings.GraphTextures.Num() && ByteOffset == ShaderBindings.GraphTextures[GraphTextureId].ByteOffset)
+			{
+				GraphTextureId++;
+				continue;
+			}
+		}
+		else if (Type == UBMT_GRAPH_TRACKED_SRV)
+		{
+			if (GraphSRVId < ShaderBindings.GraphSRVs.Num() && ByteOffset == ShaderBindings.GraphSRVs[GraphSRVId].ByteOffset)
+			{
+				GraphSRVId++;
+				continue;
+			}
+		}
+		else if (Type == UBMT_GRAPH_TRACKED_UAV)
+		{
+			if (GraphUAVId < ShaderBindings.GraphUAVs.Num() && ByteOffset == ShaderBindings.GraphUAVs[GraphUAVId].ByteOffset)
+			{
+				GraphUAVId++;
+				continue;
+			}
+		}
+		else
+		{
+			continue;
+		}
+
+		void** ResourcePointerAddress = reinterpret_cast<void**>(Base + ByteOffset);
+		*ResourcePointerAddress = nullptr;
+	}
+}
