Merging //UE4/Dev-Rendering-Graph@4492585 to Dev-Rendering (//UE4/Dev-Rendering)

diff --git a/Engine/Source/Runtime/RenderCore/Private/RenderTargetPool.cpp b/Engine/Source/Runtime/RenderCore/Private/RenderTargetPool.cpp
new file mode 100644
--- /dev/null
+++ b/Engine/Source/Runtime/RenderCore/Private/RenderTargetPool.cpp
@@ -0,0 +1,1326 @@
+// Copyright 1998-2018 Epic Games, Inc. All Rights Reserved.
+
+/*=============================================================================
+	RenderTargetPool.cpp: Scene render target pool manager.
+=============================================================================*/
+
+#include "RenderTargetPool.h"
+#include "RHIStaticStates.h"
+
+#if !(UE_BUILD_SHIPPING || UE_BUILD_TEST) && 0
+	#include "ClearQuad.h"
+#endif
+
+/** The global render targets pool. */
+TGlobalResource<FRenderTargetPool> GRenderTargetPool;
+
+DEFINE_LOG_CATEGORY_STATIC(LogRenderTargetPool, Warning, All);
+
+static void DumpRenderTargetPoolMemory(FOutputDevice& OutputDevice)
+{
+	GRenderTargetPool.DumpMemoryUsage(OutputDevice);
+}
+static FAutoConsoleCommandWithOutputDevice GDumpRenderTargetPoolMemoryCmd(
+	TEXT("r.DumpRenderTargetPoolMemory"),
+	TEXT("Dump allocation information for the render target pool."),
+	FConsoleCommandWithOutputDeviceDelegate::CreateStatic(DumpRenderTargetPoolMemory)
+	);
+
+void RenderTargetPoolEvents(const TArray<FString>& Args)
+{
+	uint32 SizeInKBThreshold = -1;
+	if(Args.Num() && Args[0].IsNumeric())
+	{
+		SizeInKBThreshold = FCString::Atof(*Args[0]);
+	}
+
+	if(SizeInKBThreshold != -1)
+	{
+		UE_LOG(LogRenderTargetPool, Display, TEXT("r.DumpRenderTargetPoolEvents is now enabled, use r.DumpRenderTargetPoolEvents ? for help"));
+
+		GRenderTargetPool.EventRecordingSizeThreshold = SizeInKBThreshold;
+		GRenderTargetPool.bStartEventRecordingNextTick = true;
+	}
+	else
+	{
+		GRenderTargetPool.DisableEventDisplay();
+
+		UE_LOG(LogRenderTargetPool, Display, TEXT("r.DumpRenderTargetPoolEvents is now disabled, use r.DumpRenderTargetPoolEvents <SizeInKB> to enable or r.DumpRenderTargetPoolEvents ? for help"));
+	}
+}
+
+// CVars and commands
+static FAutoConsoleCommand GRenderTargetPoolEventsCmd(
+	TEXT("r.RenderTargetPool.Events"),
+	TEXT("Visualize the render target pool events over time in one frame. Optional parameter defines threshold in KB.\n")
+	TEXT("To disable the view use the command without any parameter"),
+	FConsoleCommandWithArgsDelegate::CreateStatic(RenderTargetPoolEvents)
+	);
+
+#if !(UE_BUILD_SHIPPING || UE_BUILD_TEST)
+static TAutoConsoleVariable<int32> CVarClobberRenderTarget(
+	TEXT("r.Test.ClobberRenderRarget"),
+	0,
+	TEXT("Clears the texture returned by the rendertarget pool with a special color\n")
+	TEXT("so we can see better which passes would need to clear. Doesn't work on volume textures and non rendertargets yet.\n")
+	TEXT(" 0:off (default);\n")
+	TEXT(" 1: 1000 on RGBA channels;\n")
+	TEXT(" 2: NaN on RGBA channels;\n")
+	TEXT(" 3: +INFINITY on RGBA channels.\n"),
+	ECVF_Cheat | ECVF_RenderThreadSafe);
+#endif
+
+static TAutoConsoleVariable<int32> CVarAllowMultipleAliasingDiscardsPerFrame(
+	TEXT("r.RenderTargetPool.AllowMultipleAliasingDiscardsPerFrame"),
+	0,
+	TEXT("If enabled, allows rendertargets to be discarded and reacquired in the same frame.\n")
+	TEXT("This should give better aliasing efficiency, but carries some RHIthread/GPU performance overhead\n")
+	TEXT("with some RHIs (due to additional commandlist flushes)\n")
+	TEXT(" 0:off (default), 1:on"),
+	ECVF_Cheat | ECVF_RenderThreadSafe);
+
+static TAutoConsoleVariable<int32> CVarRtPoolTransientMode(
+	TEXT("r.RenderTargetPool.TransientAliasingMode"),
+	2,
+	TEXT("Enables transient resource aliasing for rendertargets. Used only if GSupportsTransientResourceAliasing is true.\n")
+	TEXT("0 : Disabled\n")
+	TEXT("1 : enable transient resource aliasing for fastVRam rendertargets\n")
+	TEXT("2 : enable transient resource aliasing for fastVRam rendertargets and those with a Transient hint. Best for memory usage - has some GPU cost (~0.2ms)\n")
+	TEXT("3 : enable transient resource aliasing for ALL rendertargets (not recommended)\n"),
+	ECVF_RenderThreadSafe);
+
+bool FRenderTargetPool::IsEventRecordingEnabled() const
+{
+#if !(UE_BUILD_SHIPPING || UE_BUILD_TEST)
+	return bEventRecordingStarted && bEventRecordingActive; 
+#else
+	return false;
+#endif
+}
+
+IPooledRenderTarget* FRenderTargetPoolEvent::GetValidatedPointer() const
+{
+	int32 Index = GRenderTargetPool.FindIndex(Pointer);
+
+	if(Index >= 0)
+	{
+		return Pointer;
+	}
+
+	return 0;
+}
+
+bool FRenderTargetPoolEvent::NeedsDeallocEvent()
+{
+	if(GetEventType() == ERTPE_Alloc)
+	{
+		if(Pointer)
+		{
+			IPooledRenderTarget* ValidPointer = GetValidatedPointer();
+			if(!ValidPointer || ValidPointer->IsFree())
+			{
+				Pointer = 0;
+				return true;
+			}
+		}
+	}
+
+	return false;
+}
+
+
+static uint32 ComputeSizeInKB(FPooledRenderTarget& Element)
+{
+	return (Element.ComputeMemorySize() + 1023) / 1024;
+}
+
+FRenderTargetPool::FRenderTargetPool()
+	: AllocationLevelInKB(0)
+	, bCurrentlyOverBudget(false)
+	, bStartEventRecordingNextTick(false)
+	, EventRecordingSizeThreshold(0)
+	, bEventRecordingActive(false)
+	, bEventRecordingStarted(false)
+	, CurrentEventRecordingTime(0)
+{
+}
+
+// Logic for determining whether to make a rendertarget transient
+bool FRenderTargetPool::DoesTargetNeedTransienceOverride(const FPooledRenderTargetDesc& InputDesc, ERenderTargetTransience TransienceHint) const
+{
+	if (!GSupportsTransientResourceAliasing)
+	{
+		return false;
+	}
+	int32 AliasingMode = CVarRtPoolTransientMode.GetValueOnRenderThread();
+
+	// We only override transience if aliasing is supported and enabled, the format is suitable, and the target is not already transient
+	if (AliasingMode > 0 &&
+	  	(InputDesc.TargetableFlags & (TexCreate_RenderTargetable | TexCreate_DepthStencilTargetable | TexCreate_UAV)) && 
+		((InputDesc.Flags & TexCreate_Transient) == 0) )
+	{
+		if (AliasingMode == 1)
+		{
+			// Mode 1: Only make FastVRAM rendertargets transient
+			if (InputDesc.Flags & TexCreate_FastVRAM)
+			{
+				return true;
+			}
+		}
+		else if (AliasingMode == 2)
+		{
+			// Mode 2: Make fastvram and ERenderTargetTransience::Transient rendertargets transient
+			if (InputDesc.Flags & TexCreate_FastVRAM || TransienceHint == ERenderTargetTransience::Transient)
+			{
+				return true;
+			}
+		}
+		else if (AliasingMode == 3)
+		{
+			// Mode 3 : All rendertargets are transient
+			return true;
+		}
+	}
+	return false;
+}
+
+void FRenderTargetPool::TransitionTargetsWritable(FRHICommandListImmediate& RHICmdList)
+{
+	QUICK_SCOPE_CYCLE_COUNTER(STAT_RenderTargetPoolTransition);
+	check(IsInRenderingThread());
+	WaitForTransitionFence();
+	
+	TransitionTargets.Reset();	
+
+	for (int32 i = 0; i < PooledRenderTargets.Num(); ++i)
+	{
+		FPooledRenderTarget* PooledRT = PooledRenderTargets[i];
+		if (PooledRT && PooledRT->GetDesc().AutoWritable)
+		{
+			FTextureRHIParamRef RenderTarget = PooledRT->GetRenderTargetItem().TargetableTexture;
+			if (RenderTarget)
+			{				
+				TransitionTargets.Add(RenderTarget);
+			}
+		}
+	}
+
+	if (TransitionTargets.Num() > 0)
+	{
+		RHICmdList.TransitionResourceArrayNoCopy(EResourceTransitionAccess::EWritable, TransitionTargets);
+		if (IsRunningRHIInSeparateThread())
+		{
+			TransitionFence = RHICmdList.RHIThreadFence(false);
+		}
+	}
+}
+
+void FRenderTargetPool::WaitForTransitionFence()
+{
+	QUICK_SCOPE_CYCLE_COUNTER(STAT_RenderTargetPoolTransitionWait);
+	check(IsInRenderingThread());
+	if (TransitionFence)
+	{		
+		check(IsInRenderingThread());		
+		FRHICommandListExecutor::WaitOnRHIThreadFence(TransitionFence);
+		TransitionFence = nullptr;		
+	}
+	TransitionTargets.Reset();
+	DeferredDeleteArray.Reset();
+}
+
+#if 1 // TODO(RDG): Needs dependency on UtilityShaders to reimplement ClobberAllocatedRenderTarget functionality, or just move the clearing shaders in RenderCore?
+
+static void ClobberAllocatedRenderTarget(FRHICommandList& RHICmdList, const FPooledRenderTargetDesc& Desc, IPooledRenderTarget* Out)
+{
+
+}
+
+
+#elif !(UE_BUILD_SHIPPING || UE_BUILD_TEST)
+
+static void ClobberAllocatedRenderTarget(FRHICommandList& RHICmdList, const FPooledRenderTargetDesc& Desc, IPooledRenderTarget* Out)
+{
+	int32 ClearId = CVarClobberRenderTarget.GetValueOnRenderThread();
+
+	if (!ClearId)
+	{
+		return;
+	}
+	if (Desc.bIsCubemap || Desc.bIsArray)
+	{
+		return;
+	}
+
+	FLinearColor Color = FLinearColor(1000, 1000, 1000, 1000);
+
+	if (ClearId == 2)
+	{
+		Color = FLinearColor(NAN, NAN, NAN, NAN);
+	}
+	else if (ClearId == 3)
+	{
+		Color = FLinearColor(INFINITY, INFINITY, INFINITY, INFINITY);
+	}
+
+	// TODO(RDG): draw events will need to be rewritten with render graph builder instead of RHICmdList.
+	//SCOPED_DRAW_EVENT(RHICmdList, ClobberAllocatedRenderTarget);
+
+	if (Out->GetDesc().TargetableFlags & TexCreate_RenderTargetable)
+	{
+		SetRenderTarget(RHICmdList, Out->GetRenderTargetItem().TargetableTexture, FTextureRHIRef());
+		DrawClearQuad(RHICmdList, Color);
+	}
+	else if (Out->GetDesc().TargetableFlags & TexCreate_UAV)
+	{
+		ClearUAV(RHICmdList, Out->GetRenderTargetItem(), Color);
+	}
+
+	if (Desc.TargetableFlags & TexCreate_DepthStencilTargetable)
+	{
+		SetRenderTarget(RHICmdList, FTextureRHIRef(), Out->GetRenderTargetItem().TargetableTexture);
+		DrawClearQuad(RHICmdList, false, FLinearColor::Black, true, 0.0f, true, 0);
+	}
+}
+
+#endif
+
+bool FRenderTargetPool::FindFreeElement(FRHICommandList& RHICmdList, const FPooledRenderTargetDesc& InputDesc, TRefCountPtr<IPooledRenderTarget> &Out, const TCHAR* InDebugName, bool bDoWritableBarrier, ERenderTargetTransience TransienceHint)
+{
+	check(IsInRenderingThread());
+
+	if(!InputDesc.IsValid())
+	{
+		// no need to do anything
+		return true;
+	}
+
+	// Querying a render target that have no mip levels makes no sens.
+	check(InputDesc.NumMips > 0);
+
+	// Make sure if requesting a depth format that the clear value is correct
+	ensure(!IsDepthOrStencilFormat(InputDesc.Format) || (InputDesc.ClearValue.ColorBinding == EClearBinding::ENoneBound || InputDesc.ClearValue.ColorBinding == EClearBinding::EDepthStencilBound));
+
+	// If we're doing aliasing, we may need to override Transient flags, depending on the input format and mode
+	FPooledRenderTargetDesc ModifiedDesc;
+	bool bMakeTransient = DoesTargetNeedTransienceOverride(InputDesc, TransienceHint);
+	if (bMakeTransient)
+			{
+				ModifiedDesc = InputDesc;
+				ModifiedDesc.Flags |= TexCreate_Transient;
+	}
+
+	// Override the descriptor if necessary
+	const FPooledRenderTargetDesc& Desc = bMakeTransient ? ModifiedDesc : InputDesc;
+
+
+	// if we can keep the current one, do that
+	if(Out)
+	{
+		FPooledRenderTarget* Current = (FPooledRenderTarget*)Out.GetReference();
+
+		check(!Current->IsSnapshot());
+
+		const bool bExactMatch = true;
+
+		if(Out->GetDesc().Compare(Desc, bExactMatch))
+		{
+			// we can reuse the same, but the debug name might have changed
+			Current->Desc.DebugName = InDebugName;
+			RHIBindDebugLabelName(Current->GetRenderTargetItem().TargetableTexture, InDebugName);
+			check(!Out->IsFree());
+			#if !(UE_BUILD_SHIPPING || UE_BUILD_TEST)
+				ClobberAllocatedRenderTarget(RHICmdList, Desc, Out);
+			#endif
+			return true;
+		}
+		else
+		{
+			// release old reference, it might free a RT we can use
+			Out = 0;
+
+			if(Current->IsFree())
+			{
+				AllocationLevelInKB -= ComputeSizeInKB(*Current);
+
+				int32 Index = FindIndex(Current);
+
+				check(Index >= 0);
+
+				// we don't use Remove() to not shuffle around the elements for better transparency on RenderTargetPoolEvents
+				PooledRenderTargets[Index] = 0;
+
+				VerifyAllocationLevel();
+			}
+		}
+	}
+
+	int32 AliasingMode = CVarRtPoolTransientMode.GetValueOnRenderThread();
+	FPooledRenderTarget* Found = 0;
+	uint32 FoundIndex = -1;
+	bool bReusingExistingTarget = false;
+	// try to find a suitable element in the pool
+	{
+		//don't spend time doing 2 passes if the platform doesn't support fastvram
+		uint32 PassCount = 1;
+		if (AliasingMode == 0)
+		{
+			if ((Desc.Flags & TexCreate_FastVRAM) && FPlatformMemory::SupportsFastVRAMMemory() )
+			{
+				PassCount = 2;
+			}
+		}
+
+		bool bAllowMultipleDiscards = ( CVarAllowMultipleAliasingDiscardsPerFrame.GetValueOnRenderThread() != 0 );
+		// first we try exact, if that fails we try without TexCreate_FastVRAM
+		// (easily we can run out of VRam, if this search becomes a performance problem we can optimize or we should use less TexCreate_FastVRAM)
+		for(uint32 Pass = 0; Pass < PassCount; ++Pass)
+		{
+			bool bExactMatch = (Pass == 0); //-V547
+    
+			for(uint32 i = 0, Num = (uint32)PooledRenderTargets.Num(); i < Num; ++i)
+			{
+				FPooledRenderTarget* Element = PooledRenderTargets[i];
+    
+				if(Element && Element->IsFree() && Element->GetDesc().Compare(Desc, bExactMatch))
+				{
+					if ( ( Desc.Flags & TexCreate_Transient ) && bAllowMultipleDiscards == false && Element->HasBeenDiscardedThisFrame() )
+					{
+						// We can't re-use transient resources if they've already been discarded this frame
+						continue;
+					}
+					check(!Element->IsSnapshot());
+					Found = Element;
+					FoundIndex = i;
+					bReusingExistingTarget = true;
+					goto Done;
+				}
+			}
+		}		
+	}
+Done:
+
+	if(!Found)
+	{		
+		UE_LOG(LogRenderTargetPool, Display, TEXT("%d MB, NewRT %s %s"), (AllocationLevelInKB + 1023) / 1024, *Desc.GenerateInfoString(), InDebugName);
+
+		// not found in the pool, create a new element
+		Found = new FPooledRenderTarget(Desc, this);
+
+		PooledRenderTargets.Add(Found);
+		
+		// TexCreate_UAV should be used on Desc.TargetableFlags
+		check(!(Desc.Flags & TexCreate_UAV));
+		// TexCreate_FastVRAM should be used on Desc.Flags
+		ensure(!(Desc.TargetableFlags & TexCreate_FastVRAM));
+
+		FRHIResourceCreateInfo CreateInfo(Desc.ClearValue);
+		CreateInfo.DebugName = InDebugName;
+
+		if(Desc.TargetableFlags & (TexCreate_RenderTargetable | TexCreate_DepthStencilTargetable | TexCreate_UAV))
+		{
+			if(Desc.Is2DTexture())
+			{
+				if (!Desc.IsArray())
+				{
+					RHICreateTargetableShaderResource2D(
+						Desc.Extent.X,
+						Desc.Extent.Y,
+						Desc.Format,
+						Desc.NumMips,
+						Desc.Flags,
+						Desc.TargetableFlags,
+						Desc.bForceSeparateTargetAndShaderResource,
+						CreateInfo,
+						(FTexture2DRHIRef&)Found->RenderTargetItem.TargetableTexture,
+						(FTexture2DRHIRef&)Found->RenderTargetItem.ShaderResourceTexture,
+						Desc.NumSamples
+					);
+				}
+				else
+				{
+					RHICreateTargetableShaderResource2DArray(
+						Desc.Extent.X,
+						Desc.Extent.Y,
+						Desc.ArraySize,
+						Desc.Format,
+						Desc.NumMips,
+						Desc.Flags,
+						Desc.TargetableFlags,
+						CreateInfo,
+						(FTexture2DArrayRHIRef&)Found->RenderTargetItem.TargetableTexture,
+						(FTexture2DArrayRHIRef&)Found->RenderTargetItem.ShaderResourceTexture,
+						Desc.NumSamples
+					);
+				}
+
+				if (GSupportsRenderTargetWriteMask && Desc.bCreateRenderTargetWriteMask)
+				{
+					Found->RenderTargetItem.RTWriteMaskDataBufferRHI = RHICreateRTWriteMaskBuffer((FTexture2DRHIRef&)Found->RenderTargetItem.TargetableTexture);
+					Found->RenderTargetItem.RTWriteMaskBufferRHI_SRV = RHICreateShaderResourceView(Found->RenderTargetItem.RTWriteMaskDataBufferRHI);
+				}
+
+				if( Desc.NumMips > 1 )
+				{
+					Found->RenderTargetItem.MipSRVs.SetNum( Desc.NumMips );
+					for( uint16 i = 0; i < Desc.NumMips; i++ )
+					{
+						Found->RenderTargetItem.MipSRVs[i] = RHICreateShaderResourceView( (FTexture2DRHIRef&)Found->RenderTargetItem.ShaderResourceTexture, i );
+					}
+				}
+			}
+			else if(Desc.Is3DTexture())
+			{
+				Found->RenderTargetItem.ShaderResourceTexture = RHICreateTexture3D(
+					Desc.Extent.X,
+					Desc.Extent.Y,
+					Desc.Depth,
+					Desc.Format,
+					Desc.NumMips,
+					Desc.Flags | Desc.TargetableFlags,
+					CreateInfo);
+
+				// similar to RHICreateTargetableShaderResource2D
+				Found->RenderTargetItem.TargetableTexture = Found->RenderTargetItem.ShaderResourceTexture;
+			}
+			else
+			{
+				check(Desc.IsCubemap());
+				if(Desc.IsArray())
+				{
+					RHICreateTargetableShaderResourceCubeArray(
+						Desc.Extent.X,
+						Desc.ArraySize,
+						Desc.Format,
+						Desc.NumMips,
+						Desc.Flags,
+						Desc.TargetableFlags,
+						false,
+						CreateInfo,
+						(FTextureCubeRHIRef&)Found->RenderTargetItem.TargetableTexture,
+						(FTextureCubeRHIRef&)Found->RenderTargetItem.ShaderResourceTexture
+						);
+				}
+				else
+				{
+					RHICreateTargetableShaderResourceCube(
+						Desc.Extent.X,
+						Desc.Format,
+						Desc.NumMips,
+						Desc.Flags,
+						Desc.TargetableFlags,
+						false,
+						CreateInfo,
+						(FTextureCubeRHIRef&)Found->RenderTargetItem.TargetableTexture,
+						(FTextureCubeRHIRef&)Found->RenderTargetItem.ShaderResourceTexture
+						);
+
+					if( Desc.NumMips > 1 )
+					{
+						Found->RenderTargetItem.MipSRVs.SetNum( Desc.NumMips );
+						for( uint16 i = 0; i < Desc.NumMips; i++ )
+						{
+							Found->RenderTargetItem.MipSRVs[i] = RHICreateShaderResourceView( (FTextureCubeRHIRef&)Found->RenderTargetItem.ShaderResourceTexture, i );
+						}
+					}
+				}
+
+			}
+
+			RHIBindDebugLabelName(Found->RenderTargetItem.TargetableTexture, InDebugName);
+		}
+		else
+		{
+			if(Desc.Is2DTexture())
+			{
+				// this is useful to get a CPU lockable texture through the same interface
+				Found->RenderTargetItem.ShaderResourceTexture = RHICreateTexture2D(
+					Desc.Extent.X,
+					Desc.Extent.Y,
+					Desc.Format,
+					Desc.NumMips,
+					Desc.NumSamples,
+					Desc.Flags,
+					CreateInfo);
+			}
+			else if(Desc.Is3DTexture())
+			{
+				Found->RenderTargetItem.ShaderResourceTexture = RHICreateTexture3D(
+					Desc.Extent.X,
+					Desc.Extent.Y,
+					Desc.Depth,
+					Desc.Format,
+					Desc.NumMips,
+					Desc.Flags,
+					CreateInfo);
+			}
+			else 
+			{
+				check(Desc.IsCubemap());
+				if(Desc.IsArray())
+				{
+					FTextureCubeRHIRef CubeTexture = RHICreateTextureCubeArray(Desc.Extent.X,Desc.ArraySize,Desc.Format,Desc.NumMips,Desc.Flags | Desc.TargetableFlags | TexCreate_ShaderResource,CreateInfo);
+					Found->RenderTargetItem.TargetableTexture = Found->RenderTargetItem.ShaderResourceTexture = CubeTexture;
+				}
+				else
+				{
+					FTextureCubeRHIRef CubeTexture = RHICreateTextureCube(Desc.Extent.X,Desc.Format,Desc.NumMips,Desc.Flags | Desc.TargetableFlags | TexCreate_ShaderResource,CreateInfo);
+					Found->RenderTargetItem.TargetableTexture = Found->RenderTargetItem.ShaderResourceTexture = CubeTexture;
+				}
+			}
+
+			RHIBindDebugLabelName(Found->RenderTargetItem.ShaderResourceTexture, InDebugName);
+		}
+
+		if(Desc.TargetableFlags & TexCreate_UAV)
+		{
+			// The render target desc is invalid if a UAV is requested with an RHI that doesn't support the high-end feature level.
+			check(GMaxRHIFeatureLevel == ERHIFeatureLevel::SM5);
+			Found->RenderTargetItem.MipUAVs.Reserve(Desc.NumMips);
+			for (uint32 MipLevel = 0; MipLevel < Desc.NumMips; MipLevel++)
+			{
+				Found->RenderTargetItem.MipUAVs.Add(RHICreateUnorderedAccessView(Found->RenderTargetItem.TargetableTexture, MipLevel));
+			}
+
+			Found->RenderTargetItem.UAV = Found->RenderTargetItem.MipUAVs[0];
+		}
+
+		AllocationLevelInKB += ComputeSizeInKB(*Found);
+		VerifyAllocationLevel();
+
+		FoundIndex = PooledRenderTargets.Num() - 1;
+
+		Found->Desc.DebugName = InDebugName;
+	}
+
+	#if !(UE_BUILD_SHIPPING || UE_BUILD_TEST)
+		ClobberAllocatedRenderTarget(RHICmdList, Desc, Found);
+	#endif
+
+	check(Found->IsFree());
+	check(!Found->IsSnapshot());
+
+	Found->Desc.DebugName = InDebugName;
+	Found->UnusedForNFrames = 0;
+
+	AddAllocEvent(FoundIndex, Found);
+
+	uint32 OriginalNumRefs = Found->GetRefCount();
+
+	// assign to the reference counted variable
+	Out = Found;
+
+	check(!Found->IsFree());
+
+	// Only referenced by the pool, map the physical pages
+	if (OriginalNumRefs == 1 && Found->GetRenderTargetItem().TargetableTexture != nullptr)
+	{
+		RHIAcquireTransientResource(Found->GetRenderTargetItem().TargetableTexture);
+	}
+
+	if (bReusingExistingTarget)
+	{
+		if (bDoWritableBarrier)
+		{
+			RHICmdList.TransitionResource(EResourceTransitionAccess::EWritable, Found->GetRenderTargetItem().TargetableTexture);
+		}
+
+	}
+
+	// Transient RTs have to be targettable
+	check( ( Desc.Flags & TexCreate_Transient ) == 0 || Found->GetRenderTargetItem().TargetableTexture != nullptr );
+
+	return false;
+}
+
+void FRenderTargetPool::CreateUntrackedElement(const FPooledRenderTargetDesc& Desc, TRefCountPtr<IPooledRenderTarget> &Out, const FSceneRenderTargetItem& Item)
+{
+	check(IsInRenderingThread());
+
+	Out = 0;
+
+	// not found in the pool, create a new element
+	FPooledRenderTarget* Found = new FPooledRenderTarget(Desc, NULL);
+
+	Found->RenderTargetItem = Item;
+	check(!Found->IsSnapshot());
+
+	// assign to the reference counted variable
+	Out = Found;
+}
+
+IPooledRenderTarget* FRenderTargetPool::MakeSnapshot(const TRefCountPtr<IPooledRenderTarget>& In)
+{
+	check(IsInRenderingThread());
+	FPooledRenderTarget* NewSnapshot = nullptr;
+	if (In.GetReference())
+	{
+		NewSnapshot = new (FMemStack::Get()) FPooledRenderTarget(*static_cast<FPooledRenderTarget*>(In.GetReference()));
+		PooledRenderTargetSnapshots.Add(NewSnapshot);
+	}
+	return NewSnapshot;
+}
+
+
+void FRenderTargetPool::GetStats(uint32& OutWholeCount, uint32& OutWholePoolInKB, uint32& OutUsedInKB) const
+{
+	OutWholeCount = (uint32)PooledRenderTargets.Num();
+	OutUsedInKB = 0;
+	OutWholePoolInKB = 0;
+		
+	for(uint32 i = 0; i < (uint32)PooledRenderTargets.Num(); ++i)
+	{
+		FPooledRenderTarget* Element = PooledRenderTargets[i];
+
+		if(Element)
+		{
+			check(!Element->IsSnapshot());
+			uint32 SizeInKB = ComputeSizeInKB(*Element);
+
+			OutWholePoolInKB += SizeInKB;
+
+			if(!Element->IsFree())
+			{
+				OutUsedInKB += SizeInKB;
+			}
+		}
+	}
+
+	// if this triggers uncomment the code in VerifyAllocationLevel() and debug the issue, we might leak memory or not release when we could
+	ensure(AllocationLevelInKB == OutWholePoolInKB);
+}
+
+void FRenderTargetPool::AddPhaseEvent(const TCHAR *InPhaseName)
+{
+	if(IsEventRecordingEnabled())
+	{
+		AddDeallocEvents();
+
+		const FString* LastName = GetLastEventPhaseName();
+
+		if(!LastName || *LastName != InPhaseName)
+		{
+			if(CurrentEventRecordingTime)
+			{
+				// put a break to former data
+				++CurrentEventRecordingTime;
+			}
+
+			FRenderTargetPoolEvent NewEvent(InPhaseName, CurrentEventRecordingTime);
+
+			RenderTargetPoolEvents.Add(NewEvent);
+		}
+	}
+}
+
+const FString* FRenderTargetPool::GetLastEventPhaseName()
+{
+	// could be optimized but this is a debug view
+
+	// start from the end for better performance
+	for(int32 i = RenderTargetPoolEvents.Num() - 1; i >= 0; --i)
+	{
+		const FRenderTargetPoolEvent* Event = &RenderTargetPoolEvents[i];
+
+		if(Event->GetEventType() == ERTPE_Phase)
+		{
+			return &Event->GetPhaseName();
+		}
+	}
+
+	return 0;
+}
+
+FRenderTargetPool::SMemoryStats FRenderTargetPool::ComputeView()
+{
+	SMemoryStats MemoryStats;
+
+#if !(UE_BUILD_SHIPPING || UE_BUILD_TEST)
+	{
+		struct FRTPColumn
+		{
+			// index into the column, -1 if this is no valid column
+			uint32 PoolEntryId;
+			// for sorting
+			uint64 SizeInBytes;
+			// for sorting
+			bool bVRam;
+
+			// default constructor
+			FRTPColumn()
+				: PoolEntryId(-1)
+				, SizeInBytes(0)
+			{
+			}
+
+			// constructor
+			FRTPColumn(const FRenderTargetPoolEvent& Event)
+				: PoolEntryId(Event.GetPoolEntryId())
+				, bVRam((Event.GetDesc().Flags & TexCreate_FastVRAM) != 0)
+			{
+				 SizeInBytes = Event.GetSizeInBytes();
+			}
+
+			// sort criteria
+			bool operator <(const FRTPColumn& rhs) const
+			{
+				// sort VRam first (only matters on XboxOne but nice to always see it)
+// sorting only useful for XboxOne		if(bVRam != rhs.bVRam) return bVRam > rhs.bVRam;
+
+				// we want the large ones first
+				return SizeInBytes > rhs.SizeInBytes;
+			}
+		};
+
+		TArray<FRTPColumn> Colums;
+
+		// generate Colums
+		for(int32 i = 0, Num = RenderTargetPoolEvents.Num(); i < Num; i++)
+		{
+			FRenderTargetPoolEvent* Event = &RenderTargetPoolEvents[i];
+
+			if(Event->GetEventType() == ERTPE_Alloc)
+			{
+				uint32 PoolEntryId = Event->GetPoolEntryId();
+
+				if(PoolEntryId >= (uint32)Colums.Num())
+				{
+					Colums.SetNum(PoolEntryId + 1);
+				}
+
+				Colums[PoolEntryId] = FRTPColumn(*Event);
+			}
+		}
+
+		Colums.Sort();
+
+		{
+			uint32 ColumnX = 0;
+
+			for(int32 ColumnIndex = 0, ColumnsNum = Colums.Num(); ColumnIndex < ColumnsNum; ++ColumnIndex)
+			{
+				const FRTPColumn& RTPColumn = Colums[ColumnIndex];
+
+				uint32 ColumnSize = RTPColumn.SizeInBytes;
+
+				// hide columns that are too small to make a difference (e.g. <1 MB)
+				if(RTPColumn.SizeInBytes <= EventRecordingSizeThreshold * 1024)
+				{
+					ColumnSize = 0;
+				}
+				else
+				{
+					MemoryStats.DisplayedUsageInBytes += RTPColumn.SizeInBytes;
+
+					// give an entry some size to be more UI friendly (if we get mouse UI for zooming in we might not want that any more)
+					ColumnSize = FMath::Max((uint32)(1024 * 1024), ColumnSize);
+				}
+
+				MemoryStats.TotalColumnSize += ColumnSize;
+				MemoryStats.TotalUsageInBytes += RTPColumn.SizeInBytes;
+				
+				for(int32 EventIndex = 0, PoolEventsNum = RenderTargetPoolEvents.Num(); EventIndex < PoolEventsNum; EventIndex++)
+				{
+					FRenderTargetPoolEvent* Event = &RenderTargetPoolEvents[EventIndex];
+
+					if(Event->GetEventType() != ERTPE_Phase)
+					{
+						uint32 PoolEntryId = Event->GetPoolEntryId();
+
+						if(RTPColumn.PoolEntryId == PoolEntryId)
+						{
+							Event->SetColumn(ColumnIndex, ColumnX, ColumnSize);
+						}
+					}
+				}
+				ColumnX += ColumnSize;
+			}
+		}
+	}
+#endif
+
+	return MemoryStats;
+}
+
+void FRenderTargetPool::UpdateElementSize(const TRefCountPtr<IPooledRenderTarget>& Element, const uint32 OldElementSize)
+{
+	check(Element.IsValid() && FindIndex(&(*Element)) >= 0);
+	AllocationLevelInKB -= (OldElementSize + 1023) / 1024;
+	AllocationLevelInKB += (Element->ComputeMemorySize() + 1023) / 1024;
+}
+
+void FRenderTargetPool::AddDeallocEvents()
+{
+	check(IsInRenderingThread());
+
+	bool bWorkWasDone = false;
+
+	for(uint32 i = 0, Num = (uint32)RenderTargetPoolEvents.Num(); i < Num; ++i)
+	{
+		FRenderTargetPoolEvent& Event = RenderTargetPoolEvents[i];
+
+		if(Event.NeedsDeallocEvent())
+		{
+			FRenderTargetPoolEvent NewEvent(Event.GetPoolEntryId(), CurrentEventRecordingTime);
+
+			// for convenience - is actually redundant
+			NewEvent.SetDesc(Event.GetDesc());
+
+			RenderTargetPoolEvents.Add(NewEvent);
+			bWorkWasDone = true;
+		}
+	}
+
+	if(bWorkWasDone)
+	{
+		++CurrentEventRecordingTime;
+	}
+}
+
+void FRenderTargetPool::AddAllocEvent(uint32 InPoolEntryId, FPooledRenderTarget* In)
+{
+	check(In);
+
+	if(IsEventRecordingEnabled())
+	{
+		AddDeallocEvents();
+
+		check(IsInRenderingThread());
+
+		FRenderTargetPoolEvent NewEvent(InPoolEntryId, CurrentEventRecordingTime++, In);
+
+		RenderTargetPoolEvents.Add(NewEvent);
+	}
+}
+
+
+void FRenderTargetPool::AddAllocEventsFromCurrentState()
+{
+	if(!IsEventRecordingEnabled())
+	{
+		return;
+	}
+
+	check(IsInRenderingThread());
+
+	bool bWorkWasDone = false;
+
+	for(uint32 i = 0; i < (uint32)PooledRenderTargets.Num(); ++i)
+	{
+		FPooledRenderTarget* Element = PooledRenderTargets[i];
+
+		if(Element && !Element->IsFree())
+		{
+			FRenderTargetPoolEvent NewEvent(i, CurrentEventRecordingTime, Element);
+
+			RenderTargetPoolEvents.Add(NewEvent);
+			bWorkWasDone = true;
+		}
+	}
+
+	if(bWorkWasDone)
+	{
+		++CurrentEventRecordingTime;
+	}
+}
+
+void FRenderTargetPool::TickPoolElements()
+{
+	check(IsInRenderingThread());
+	WaitForTransitionFence();
+
+	if(bStartEventRecordingNextTick)
+	{
+		bStartEventRecordingNextTick = false;
+		bEventRecordingStarted = true;
+	}
+
+	uint32 MinimumPoolSizeInKB;
+	{
+		static const auto CVar = IConsoleManager::Get().FindTConsoleVariableDataInt(TEXT("r.RenderTargetPoolMin"));
+
+		MinimumPoolSizeInKB = FMath::Clamp(CVar->GetValueOnRenderThread(), 0, 2000) * 1024;
+	}
+
+	CompactPool();
+
+	for(uint32 i = 0; i < (uint32)PooledRenderTargets.Num(); ++i)
+	{
+		FPooledRenderTarget* Element = PooledRenderTargets[i];
+
+		if(Element)
+		{
+			check(!Element->IsSnapshot());
+			Element->OnFrameStart();
+		}
+	}
+
+	// we need to release something, take the oldest ones first
+	while(AllocationLevelInKB > MinimumPoolSizeInKB)
+	{
+		// -1: not set
+		int32 OldestElementIndex = -1;
+
+		// find oldest element we can remove
+		for(uint32 i = 0, Num = (uint32)PooledRenderTargets.Num(); i < Num; ++i)
+		{
+			FPooledRenderTarget* Element = PooledRenderTargets[i];
+
+			if(Element && Element->UnusedForNFrames > 2)
+			{
+				if(OldestElementIndex != -1)
+				{
+					if(PooledRenderTargets[OldestElementIndex]->UnusedForNFrames < Element->UnusedForNFrames)
+					{
+						OldestElementIndex = i;
+					}
+				}
+				else
+				{
+					OldestElementIndex = i;
+				}
+			}
+		}
+
+		if(OldestElementIndex != -1)
+		{
+			AllocationLevelInKB -= ComputeSizeInKB(*PooledRenderTargets[OldestElementIndex]);
+
+			// we assume because of reference counting the resource gets released when not needed any more
+			// we don't use Remove() to not shuffle around the elements for better transparency on RenderTargetPoolEvents
+			PooledRenderTargets[OldestElementIndex] = 0;
+
+			VerifyAllocationLevel();
+		}
+		else
+		{
+			// There is no element we can remove but we are over budget, better we log that.
+			// Options:
+			//   * Increase the pool
+			//   * Reduce rendering features or resolution
+			//   * Investigate allocations, order or reusing other render targets can help
+			//   * Ignore (editor case, might start using slow memory which can be ok)
+			if(!bCurrentlyOverBudget)
+			{
+				UE_CLOG(IsRunningClientOnly(), LogRenderTargetPool, Warning, TEXT("r.RenderTargetPoolMin exceeded %d/%d MB (ok in editor, bad on fixed memory platform)"), (AllocationLevelInKB + 1023) / 1024, MinimumPoolSizeInKB / 1024);
+				bCurrentlyOverBudget = true;
+			}
+			// at this point we need to give up
+			break;
+		}
+
+/*	
+	// confused more than it helps (often a name is used on two elements in the pool and some pool elements are not rendered to this frame)
+	else
+	{
+		// initial state of a render target (e.g. Velocity@0)
+		GVisualizeTexture.SetCheckPoint(Element);
+*/	}
+
+	if(AllocationLevelInKB <= MinimumPoolSizeInKB)
+	{
+		if(bCurrentlyOverBudget)
+		{
+			UE_LOG(LogRenderTargetPool, Display, TEXT("r.RenderTargetPoolMin resolved %d/%d MB"), (AllocationLevelInKB + 1023) / 1024, MinimumPoolSizeInKB / 1024);
+			bCurrentlyOverBudget = false;
+		}
+	}
+
+//	CompactEventArray();
+
+	AddPhaseEvent(TEXT("FromLastFrame"));
+	AddAllocEventsFromCurrentState();
+	AddPhaseEvent(TEXT("Rendering"));
+
+#if STATS
+	uint32 Count, SizeKB, UsedKB;
+	GetStats(Count, SizeKB, UsedKB);
+	SET_MEMORY_STAT(STAT_RenderTargetPoolSize, int64(SizeKB) * 1024ll);
+	SET_MEMORY_STAT(STAT_RenderTargetPoolUsed, int64(UsedKB) * 1024ll);
+	SET_DWORD_STAT(STAT_RenderTargetPoolCount, Count);
+#endif // STATS
+}
+
+
+int32 FRenderTargetPool::FindIndex(IPooledRenderTarget* In) const
+{
+	check(IsInRenderingThread());
+
+	if(In)
+	{
+		for(uint32 i = 0, Num = (uint32)PooledRenderTargets.Num(); i < Num; ++i)
+		{
+			const FPooledRenderTarget* Element = PooledRenderTargets[i];
+
+			if(Element == In)
+			{
+				check(!Element->IsSnapshot());
+				return i;
+			}
+		}
+	}
+
+	// not found
+	return -1;
+}
+
+void FRenderTargetPool::FreeUnusedResource(TRefCountPtr<IPooledRenderTarget>& In)
+{
+	check(IsInRenderingThread());
+	
+	int32 Index = FindIndex(In);
+
+	if(Index != -1)
+	{
+		FPooledRenderTarget* Element = PooledRenderTargets[Index];
+
+		if(Element)
+		{
+			check(!Element->IsSnapshot());
+			AllocationLevelInKB -= ComputeSizeInKB(*Element);
+			// we assume because of reference counting the resource gets released when not needed any more
+			// we don't use Remove() to not shuffle around the elements for better transparency on RenderTargetPoolEvents
+			DeferredDeleteArray.Add(PooledRenderTargets[Index]);
+			PooledRenderTargets[Index] = 0;
+
+			In.SafeRelease();
+
+			VerifyAllocationLevel();
+		}
+	}
+}
+
+void FRenderTargetPool::FreeUnusedResources()
+{
+	check(IsInRenderingThread());
+
+	for(uint32 i = 0, Num = (uint32)PooledRenderTargets.Num(); i < Num; ++i)
+	{
+		FPooledRenderTarget* Element = PooledRenderTargets[i];
+
+		if(Element && Element->IsFree())
+		{
+			check(!Element->IsSnapshot());
+			AllocationLevelInKB -= ComputeSizeInKB(*Element);
+			// we assume because of reference counting the resource gets released when not needed any more
+			// we don't use Remove() to not shuffle around the elements for better transparency on RenderTargetPoolEvents
+			DeferredDeleteArray.Add(PooledRenderTargets[i]);
+			PooledRenderTargets[i] = 0;
+		}
+	}
+
+	VerifyAllocationLevel();
+}
+
+
+void FRenderTargetPool::DumpMemoryUsage(FOutputDevice& OutputDevice)
+{
+	OutputDevice.Logf(TEXT("Pooled Render Targets:"));
+	for(int32 i = 0; i < PooledRenderTargets.Num(); ++i)
+	{
+		FPooledRenderTarget* Element = PooledRenderTargets[i];
+
+		if(Element)
+		{
+			check(!Element->IsSnapshot());
+			OutputDevice.Logf(
+				TEXT("  %6.3fMB %4dx%4d%s%s %2dmip(s) %s (%s) %s %s"),
+				ComputeSizeInKB(*Element) / 1024.0f,
+				Element->Desc.Extent.X,
+				Element->Desc.Extent.Y,
+				Element->Desc.Depth > 1 ? *FString::Printf(TEXT("x%3d"), Element->Desc.Depth) : (Element->Desc.IsCubemap() ? TEXT("cube") : TEXT("    ")),
+				Element->Desc.bIsArray ? *FString::Printf(TEXT("[%3d]"), Element->Desc.ArraySize) : TEXT("     "),
+				Element->Desc.NumMips,
+				Element->Desc.DebugName,
+				GPixelFormats[Element->Desc.Format].Name,
+				Element->IsTransient() ? TEXT("(transient)") : TEXT(""),
+				GSupportsTransientResourceAliasing ? *FString::Printf(TEXT("Frames since last discard: %d"), GFrameNumberRenderThread - Element->FrameNumberLastDiscard) : TEXT("")
+				);
+		}
+	}
+	uint32 NumTargets=0;
+	uint32 UsedKB=0;
+	uint32 PoolKB=0;
+	GetStats(NumTargets,PoolKB,UsedKB);
+	OutputDevice.Logf(TEXT("%.3fMB total, %.3fMB used, %d render targets"), PoolKB / 1024.f, UsedKB / 1024.f, NumTargets);
+}
+
+uint32 FPooledRenderTarget::AddRef() const
+{
+	if (!bSnapshot)
+	{
+		check(IsInRenderingThread());
+		return uint32(++NumRefs);
+	}
+	check(NumRefs == 1);
+	return 1;
+}
+
+uint32 FPooledRenderTarget::Release()
+{
+	if (!bSnapshot)
+	{
+		check(IsInRenderingThread());
+		uint32 Refs = uint32(--NumRefs);
+		if(Refs == 0)
+		{
+			RenderTargetItem.SafeRelease();
+			delete this;
+		}
+		else if (Refs == 1 && RenderTargetPool && IsTransient() )
+		{
+			// Discard the resource
+			check(GetRenderTargetItem().TargetableTexture != nullptr);
+			if (GetRenderTargetItem().TargetableTexture)
+			{
+				RHIDiscardTransientResource(GetRenderTargetItem().TargetableTexture);
+			}
+			FrameNumberLastDiscard = GFrameNumberRenderThread;
+		}
+		return Refs;
+	}
+	check(NumRefs == 1);
+	return 1;
+}
+
+uint32 FPooledRenderTarget::GetRefCount() const
+{
+	return uint32(NumRefs);
+}
+
+void FPooledRenderTarget::SetDebugName(const TCHAR *InName)
+{
+	check(InName);
+
+	Desc.DebugName = InName;
+}
+
+const FPooledRenderTargetDesc& FPooledRenderTarget::GetDesc() const
+{
+	return Desc;
+}
+
+void FRenderTargetPool::ReleaseDynamicRHI()
+{
+	check(IsInRenderingThread());
+	WaitForTransitionFence();
+
+	PooledRenderTargets.Empty();
+	if (PooledRenderTargetSnapshots.Num())
+	{
+		DestructSnapshots();
+	}
+}
+
+void FRenderTargetPool::DestructSnapshots()
+{
+	for (auto Snapshot : PooledRenderTargetSnapshots)
+	{
+		Snapshot->~FPooledRenderTarget();
+	}
+	PooledRenderTargetSnapshots.Reset();
+}
+
+// for debugging purpose
+FPooledRenderTarget* FRenderTargetPool::GetElementById(uint32 Id) const
+{
+	// is used in game and render thread
+
+	if(Id >= (uint32)PooledRenderTargets.Num())
+	{
+		return 0;
+	}
+
+	return PooledRenderTargets[Id];
+}
+
+void FRenderTargetPool::VerifyAllocationLevel() const
+{
+/*
+	// uncomment to verify internal consistency
+	uint32 OutWholeCount;
+	uint32 OutWholePoolInKB;
+	uint32 OutUsedInKB;
+
+	GetStats(OutWholeCount, OutWholePoolInKB, OutUsedInKB);
+*/
+}
+
+void FRenderTargetPool::CompactPool()
+{
+	for(uint32 i = 0, Num = (uint32)PooledRenderTargets.Num(); i < Num; ++i)
+	{
+		FPooledRenderTarget* Element = PooledRenderTargets[i];
+
+		if(!Element)
+		{
+			PooledRenderTargets.RemoveAtSwap(i);
+			--Num;
+		}
+	}
+}
+
+bool FPooledRenderTarget::OnFrameStart()
+{
+	check(IsInRenderingThread() && !bSnapshot);
+
+	// If there are any references to the pooled render target other than the pool itself, then it may not be freed.
+	if(!IsFree())
+	{
+		check(!UnusedForNFrames);
+		return false;
+	}
+
+	++UnusedForNFrames;
+
+	// this logic can be improved
+	if(UnusedForNFrames > 10)
+	{
+		// release
+		return true;
+	}
+
+	return false;
+}
+
+uint32 FPooledRenderTarget::ComputeMemorySize() const
+{
+	uint32 Size = 0;
+	if (!bSnapshot)
+	{
+		if(Desc.Is2DTexture())
+		{
+			Size += RHIComputeMemorySize((const FTexture2DRHIRef&)RenderTargetItem.TargetableTexture);
+			if(RenderTargetItem.ShaderResourceTexture != RenderTargetItem.TargetableTexture)
+			{
+				Size += RHIComputeMemorySize((const FTexture2DRHIRef&)RenderTargetItem.ShaderResourceTexture);
+			}
+		}
+		else if(Desc.Is3DTexture())
+		{
+			Size += RHIComputeMemorySize((const FTexture3DRHIRef&)RenderTargetItem.TargetableTexture);
+			if(RenderTargetItem.ShaderResourceTexture != RenderTargetItem.TargetableTexture)
+			{
+				Size += RHIComputeMemorySize((const FTexture3DRHIRef&)RenderTargetItem.ShaderResourceTexture);
+			}
+		}
+		else
+		{
+			Size += RHIComputeMemorySize((const FTextureCubeRHIRef&)RenderTargetItem.TargetableTexture);
+			if(RenderTargetItem.ShaderResourceTexture != RenderTargetItem.TargetableTexture)
+			{
+				Size += RHIComputeMemorySize((const FTextureCubeRHIRef&)RenderTargetItem.ShaderResourceTexture);
+			}
+		}
+	}
+	return Size;
+}
+
+bool FPooledRenderTarget::IsFree() const
+{
+	uint32 RefCount = GetRefCount();
+	check(RefCount >= 1);
+
+	// If the only reference to the pooled render target is from the pool, then it's unused.
+	return !bSnapshot && RefCount == 1;
+}
