Merging //UE4/Dev-Rendering-Graph@4492585 to Dev-Rendering (//UE4/Dev-Rendering)

diff --git a/Engine/Source/Runtime/RenderCore/Private/Shader.cpp b/Engine/Source/Runtime/RenderCore/Private/Shader.cpp
new file mode 100644
--- /dev/null
+++ b/Engine/Source/Runtime/RenderCore/Private/Shader.cpp
@@ -0,0 +1,2148 @@
+// Copyright 1998-2018 Epic Games, Inc. All Rights Reserved.
+
+/*=============================================================================
+	Shader.cpp: Shader implementation.
+=============================================================================*/
+
+#include "Shader.h"
+#include "Misc/CoreMisc.h"
+#include "Stats/StatsMisc.h"
+#include "Serialization/MemoryWriter.h"
+#include "VertexFactory.h"
+#include "ProfilingDebugging/DiagnosticTable.h"
+#include "Interfaces/ITargetPlatform.h"
+#include "Interfaces/ITargetPlatformManagerModule.h"
+#include "ShaderCodeLibrary.h"
+#include "ShaderCore.h"
+#include "RenderUtils.h"
+#include "Misc/ConfigCacheIni.h"
+#include "UObject/RenderingObjectVersion.h"
+#include "UObject/FortniteMainBranchObjectVersion.h"
+
+#if WITH_EDITORONLY_DATA
+#include "Interfaces/IShaderFormat.h"
+#endif
+
+DEFINE_LOG_CATEGORY(LogShaders);
+
+RENDERCORE_API bool UsePreExposure(EShaderPlatform Platform)
+{
+	// Mobile platforms are excluded because they use a different pre-exposure logic in MobileBasePassPixelShader.usf
+	static const auto CVarUsePreExposure = IConsoleManager::Get().FindTConsoleVariableDataInt(TEXT("r.UsePreExposure"));
+	return CVarUsePreExposure->GetValueOnAnyThread() != 0 && !IsMobilePlatform(Platform) && IsFeatureLevelSupported(Platform, ERHIFeatureLevel::SM5);
+}
+
+static const ECompressionFlags ShaderCompressionFlag = ECompressionFlags::COMPRESS_ZLIB;
+
+static TAutoConsoleVariable<int32> CVarUsePipelines(
+	TEXT("r.ShaderPipelines"),
+	1,
+	TEXT("Enable using Shader pipelines."));
+
+static TLinkedList<FShaderType*>*			GShaderTypeList = nullptr;
+static TLinkedList<FShaderPipelineType*>*	GShaderPipelineList = nullptr;
+static TMap<FName, FShaderType*>*			GShaderNameToTypeMap = nullptr;
+
+static FSHAHash ShaderSourceDefaultHash; //will only be read (never written) for the cooking case
+
+/**
+ * Find the shader pipeline type with the given name.
+ * @return NULL if no type matched.
+ */
+inline const FShaderPipelineType* FindShaderPipelineType(FName TypeName)
+{
+	for (TLinkedList<FShaderPipelineType*>::TIterator ShaderPipelineTypeIt(FShaderPipelineType::GetTypeList()); ShaderPipelineTypeIt; ShaderPipelineTypeIt.Next())
+	{
+		if (ShaderPipelineTypeIt->GetFName() == TypeName)
+		{
+			return *ShaderPipelineTypeIt;
+		}
+	}
+	return nullptr;
+}
+
+
+/**
+ * Serializes a reference to a shader pipeline type.
+ */
+FArchive& operator<<(FArchive& Ar, const FShaderPipelineType*& TypeRef)
+{
+	if (Ar.IsSaving())
+	{
+		FName TypeName = TypeRef ? FName(TypeRef->Name) : NAME_None;
+		Ar << TypeName;
+	}
+	else if (Ar.IsLoading())
+	{
+		FName TypeName = NAME_None;
+		Ar << TypeName;
+		TypeRef = FindShaderPipelineType(TypeName);
+	}
+	return Ar;
+}
+
+
+void FShaderParameterMap::VerifyBindingsAreComplete(const TCHAR* ShaderTypeName, FShaderTarget Target, FVertexFactoryType* InVertexFactoryType) const
+{
+#if WITH_EDITORONLY_DATA
+	// Only people working on shaders (and therefore have LogShaders unsuppressed) will want to see these errors
+	if (UE_LOG_ACTIVE(LogShaders, Warning))
+	{
+		const TCHAR* VertexFactoryName = InVertexFactoryType ? InVertexFactoryType->GetName() : TEXT("?");
+
+		bool bBindingsComplete = true;
+		FString UnBoundParameters = TEXT("");
+		for (TMap<FString,FParameterAllocation>::TConstIterator ParameterIt(ParameterMap);ParameterIt;++ParameterIt)
+		{
+			const FString& ParamName = ParameterIt.Key();
+			const FParameterAllocation& ParamValue = ParameterIt.Value();
+			if(!ParamValue.bBound)
+			{
+				// Only valid parameters should be in the shader map
+				checkSlow(ParamValue.Size > 0);
+				bBindingsComplete = bBindingsComplete && ParamValue.bBound;
+				UnBoundParameters += FString(TEXT("		Parameter ")) + ParamName + TEXT(" not bound!\n");
+			}
+		}
+		
+		if (!bBindingsComplete)
+		{
+			FString ErrorMessage = FString(TEXT("Found unbound parameters being used in shadertype ")) + ShaderTypeName + TEXT(" (VertexFactory: ") + VertexFactoryName + TEXT(")\n") + UnBoundParameters;
+
+			// There will be unbound parameters for Metal's "Hull" shader stage as it is merely a placeholder to provide binding indices to the RHI
+			if(!IsMetalPlatform((EShaderPlatform)Target.Platform) || Target.Frequency != SF_Hull)
+			{
+				// We use a non-Slate message box to avoid problem where we haven't compiled the shaders for Slate.
+				FPlatformMisc::MessageBoxExt(EAppMsgType::Ok, *ErrorMessage, TEXT("Error"));
+			}
+		}
+	}
+#endif // WITH_EDITORONLY_DATA
+}
+
+
+void FShaderParameterMap::UpdateHash(FSHA1& HashState) const
+{
+	for(TMap<FString,FParameterAllocation>::TConstIterator ParameterIt(ParameterMap);ParameterIt;++ParameterIt)
+	{
+		const FString& ParamName = ParameterIt.Key();
+		const FParameterAllocation& ParamValue = ParameterIt.Value();
+		HashState.Update((const uint8*)*ParamName, ParamName.Len() * sizeof(TCHAR));
+		HashState.Update((const uint8*)&ParamValue.BufferIndex, sizeof(ParamValue.BufferIndex));
+		HashState.Update((const uint8*)&ParamValue.BaseIndex, sizeof(ParamValue.BaseIndex));
+		HashState.Update((const uint8*)&ParamValue.Size, sizeof(ParamValue.Size));
+	}
+}
+
+bool FShaderType::bInitializedSerializationHistory = false;
+
+FShaderType::FShaderType(
+	EShaderTypeForDynamicCast InShaderTypeForDynamicCast,
+	const TCHAR* InName,
+	const TCHAR* InSourceFilename,
+	const TCHAR* InFunctionName,
+	uint32 InFrequency,
+	int32 InTotalPermutationCount,
+	ConstructSerializedType InConstructSerializedRef,
+	GetStreamOutElementsType InGetStreamOutElementsRef
+	):
+	ShaderTypeForDynamicCast(InShaderTypeForDynamicCast),
+	Name(InName),
+	TypeName(InName),
+	SourceFilename(InSourceFilename),
+	FunctionName(InFunctionName),
+	Frequency(InFrequency),
+	TotalPermutationCount(InTotalPermutationCount),
+	ConstructSerializedRef(InConstructSerializedRef),
+	GetStreamOutElementsRef(InGetStreamOutElementsRef),
+	GlobalListLink(this)
+{
+	bCachedUniformBufferStructDeclarations = false;
+
+	// This will trigger if an IMPLEMENT_SHADER_TYPE was in a module not loaded before InitializeShaderTypes
+	// Shader types need to be implemented in modules that are loaded before that
+	checkf(!bInitializedSerializationHistory, TEXT("Shader type was loaded after engine init, use ELoadingPhase::PostConfigInit on your module to cause it to load earlier."));
+
+	//make sure the name is shorter than the maximum serializable length
+	check(FCString::Strlen(InName) < NAME_SIZE);
+
+	// Make sure the format of the source file path is right.
+	check(CheckVirtualShaderFilePath(InSourceFilename));
+
+	// register this shader type
+	GlobalListLink.LinkHead(GetTypeList());
+	GetNameToTypeMap().Add(TypeName, this);
+
+	// Assign the shader type the next unassigned hash index.
+	static uint32 NextHashIndex = 0;
+	HashIndex = NextHashIndex++;
+}
+
+FShaderType::~FShaderType()
+{
+	GlobalListLink.Unlink();
+	GetNameToTypeMap().Remove(TypeName);
+}
+
+TLinkedList<FShaderType*>*& FShaderType::GetTypeList()
+{
+	return GShaderTypeList;
+}
+
+FShaderType* FShaderType::GetShaderTypeByName(const TCHAR* Name)
+{
+	for(TLinkedList<FShaderType*>::TIterator It(GetTypeList()); It; It.Next())
+	{
+		FShaderType* Type = *It;
+
+		if (FPlatformString::Strcmp(Name, Type->GetName()) == 0)
+		{
+			return Type;
+		}
+	}
+
+	return nullptr;
+}
+
+TArray<FShaderType*> FShaderType::GetShaderTypesByFilename(const TCHAR* Filename)
+{
+	TArray<FShaderType*> OutShaders;
+	for(TLinkedList<FShaderType*>::TIterator It(GetTypeList()); It; It.Next())
+	{
+		FShaderType* Type = *It;
+
+		if (FPlatformString::Strcmp(Filename, Type->GetShaderFilename()) == 0)
+		{
+			OutShaders.Add(Type);
+		}
+	}
+	return OutShaders;
+}
+
+TMap<FName, FShaderType*>& FShaderType::GetNameToTypeMap()
+{
+	if(!GShaderNameToTypeMap)
+	{
+		GShaderNameToTypeMap = new TMap<FName, FShaderType*>();
+	}
+	return *GShaderNameToTypeMap;
+}
+
+inline bool FShaderType::GetOutdatedCurrentType(TArray<FShaderType*>& OutdatedShaderTypes, TArray<const FVertexFactoryType*>& OutdatedFactoryTypes) const
+{
+	bool bOutdated = false;
+	for (TMap<FShaderId, FShader*>::TConstIterator ShaderIt(ShaderIdMap);ShaderIt;++ShaderIt)
+	{
+		FShader* Shader = ShaderIt.Value();
+		const FVertexFactoryParameterRef* VFParameterRef = Shader->GetVertexFactoryParameterRef();
+		const FSHAHash& SavedHash = Shader->GetHash();
+		const FSHAHash& CurrentHash = GetSourceHash(Shader->GetShaderPlatform());
+		const bool bOutdatedShader = SavedHash != CurrentHash;
+		const bool bOutdatedVertexFactory =
+			VFParameterRef && VFParameterRef->GetVertexFactoryType() && VFParameterRef->GetVertexFactoryType()->GetSourceHash(VFParameterRef->GetShaderPlatform()) != VFParameterRef->GetHash();
+
+		if (bOutdatedShader)
+		{
+			OutdatedShaderTypes.AddUnique(Shader->Type);
+			bOutdated = true;
+		}
+
+		if (bOutdatedVertexFactory)
+		{
+			OutdatedFactoryTypes.AddUnique(VFParameterRef->GetVertexFactoryType());
+			bOutdated = true;
+		}
+	}
+
+	return bOutdated;
+}
+
+void FShaderType::GetOutdatedTypes(TArray<FShaderType*>& OutdatedShaderTypes, TArray<const FVertexFactoryType*>& OutdatedFactoryTypes)
+{
+	for(TLinkedList<FShaderType*>::TIterator It(GetTypeList()); It; It.Next())
+	{
+		FShaderType* Type = *It;
+		Type->GetOutdatedCurrentType(OutdatedShaderTypes, OutdatedFactoryTypes);
+	}
+
+	for (int32 TypeIndex = 0; TypeIndex < OutdatedShaderTypes.Num(); TypeIndex++)
+	{
+		UE_LOG(LogShaders, Warning, TEXT("		Recompiling %s"), OutdatedShaderTypes[TypeIndex]->GetName());
+	}
+	for (int32 TypeIndex = 0; TypeIndex < OutdatedFactoryTypes.Num(); TypeIndex++)
+	{
+		UE_LOG(LogShaders, Warning, TEXT("		Recompiling %s"), OutdatedFactoryTypes[TypeIndex]->GetName());
+	}
+}
+
+
+FArchive& operator<<(FArchive& Ar,FShaderType*& Ref)
+{
+	if(Ar.IsSaving())
+	{
+		FName ShaderTypeName = Ref ? FName(Ref->Name) : NAME_None;
+		Ar << ShaderTypeName;
+	}
+	else if(Ar.IsLoading())
+	{
+		FName ShaderTypeName = NAME_None;
+		Ar << ShaderTypeName;
+		
+		Ref = NULL;
+
+		if(ShaderTypeName != NAME_None)
+		{
+			// look for the shader type in the global name to type map
+			FShaderType** ShaderType = FShaderType::GetNameToTypeMap().Find(ShaderTypeName);
+			if (ShaderType)
+			{
+				// if we found it, use it
+				Ref = *ShaderType;
+			}
+			else
+			{
+				UE_LOG(LogShaders, Warning, TEXT("ShaderType '%s' was not found!"), *ShaderTypeName.ToString());
+			}
+		}
+	}
+	return Ar;
+}
+
+
+FShader* FShaderType::FindShaderById(const FShaderId& Id)
+{
+	check(IsInGameThread());
+	FShader* Result = ShaderIdMap.FindRef(Id);
+	check(!Result || Result->GetId() == Id);
+	return Result;
+}
+
+FShader* FShaderType::ConstructForDeserialization() const
+{
+	return (*ConstructSerializedRef)();
+}
+
+const FSHAHash& FShaderType::GetSourceHash(EShaderPlatform ShaderPlatform) const
+{
+	return GetShaderFileHash(GetShaderFilename(), ShaderPlatform);
+}
+
+void FShaderType::Initialize(const TMap<FString, TArray<const TCHAR*> >& ShaderFileToUniformBufferVariables)
+{
+	//#todo-rco: Need to call this only when Initializing from a Pipeline once it's removed from the global linked list
+	if (!FPlatformProperties::RequiresCookedData())
+	{
+#if UE_BUILD_DEBUG
+		TArray<FShaderType*> UniqueShaderTypes;
+#endif
+		for(TLinkedList<FShaderType*>::TIterator It(FShaderType::GetTypeList()); It; It.Next())
+		{
+			FShaderType* Type = *It;
+#if UE_BUILD_DEBUG
+			UniqueShaderTypes.Add(Type);
+#endif
+			GenerateReferencedUniformBuffers(Type->SourceFilename, Type->Name, ShaderFileToUniformBufferVariables, Type->ReferencedUniformBufferStructsCache);
+
+			// Cache serialization history for each shader type
+			// This history is used to detect when shader serialization changes without a corresponding .usf change
+			{
+				// Construct a temporary shader, which is initialized to safe values for serialization
+				FShader* TempShader = Type->ConstructForDeserialization();
+				check(TempShader != NULL);
+				TempShader->Type = Type;
+
+				// Serialize the temp shader to memory and record the number and sizes of serializations
+				TArray<uint8> TempData;
+				FMemoryWriter Ar(TempData, true);
+				FShaderSaveArchive SaveArchive(Ar, Type->SerializationHistory);
+				TempShader->SerializeBase(SaveArchive, false, false);
+
+				// Destroy the temporary shader
+				delete TempShader;
+			}
+		}
+	
+#if UE_BUILD_DEBUG
+		// Check for duplicated shader type names
+		UniqueShaderTypes.Sort([](const FShaderType& A, const FShaderType& B) { return (SIZE_T)&A < (SIZE_T)&B; });
+		for (int32 Index = 1; Index < UniqueShaderTypes.Num(); ++Index)
+		{
+			checkf(UniqueShaderTypes[Index - 1] != UniqueShaderTypes[Index], TEXT("Duplicated FShader type name %s found, please rename one of them!"), UniqueShaderTypes[Index]->GetName());
+		}
+#endif
+	}
+
+	bInitializedSerializationHistory = true;
+}
+
+void FShaderType::Uninitialize()
+{
+	for(TLinkedList<FShaderType*>::TIterator It(FShaderType::GetTypeList()); It; It.Next())
+	{
+		FShaderType* Type = *It;
+		Type->SerializationHistory = FSerializationHistory();
+	}
+
+	bInitializedSerializationHistory = false;
+}
+
+TMap<FShaderResourceId, FShaderResource*> FShaderResource::ShaderResourceIdMap;
+
+FShaderResource::FShaderResource()
+	: SpecificType(NULL)
+	, SpecificPermutationId(0)
+	, NumRefs(0)
+	, NumInstructions(0)
+#if WITH_EDITORONLY_DATA
+	, NumTextureSamplers(0)
+#endif
+	, bCodeInSharedLocation(false)
+	, bCodeInSharedLocationRequested(false)
+{
+	INC_DWORD_STAT_BY(STAT_Shaders_NumShaderResourcesLoaded, 1);
+}
+
+
+FShaderResource::FShaderResource(const FShaderCompilerOutput& Output, FShaderType* InSpecificType, int32 InSpecificPermutationId) 
+	: SpecificType(InSpecificType)
+	, SpecificPermutationId(InSpecificPermutationId)
+	, NumRefs(0)
+	, NumInstructions(Output.NumInstructions)
+#if WITH_EDITORONLY_DATA
+	, NumTextureSamplers(Output.NumTextureSamplers)
+#endif
+	, bCodeInSharedLocation(false)
+	, bCodeInSharedLocationRequested(false)
+
+{
+	check(!(SpecificPermutationId != 0 && SpecificType == nullptr));
+
+	Target = Output.Target;
+	CompressCode(Output.ShaderCode.GetReadAccess());
+
+	check(Code.Num() > 0);
+
+	OutputHash = Output.OutputHash;
+	checkSlow(OutputHash != FSHAHash());
+
+#if WITH_EDITORONLY_DATA
+	PlatformDebugData = Output.PlatformDebugData;
+#endif
+
+	{
+		check(IsInGameThread());
+		ShaderResourceIdMap.Add(GetId(), this);
+	}
+	
+	INC_DWORD_STAT_BY_FName(GetMemoryStatType((EShaderFrequency)Target.Frequency).GetName(), Code.Num());
+	INC_DWORD_STAT_BY(STAT_Shaders_ShaderResourceMemory, GetSizeBytes());
+	INC_DWORD_STAT_BY(STAT_Shaders_NumShaderResourcesLoaded, 1);
+}
+
+
+FShaderResource::~FShaderResource()
+{
+	check(NumRefs == 0);
+
+	DEC_DWORD_STAT_BY_FName(GetMemoryStatType((EShaderFrequency)Target.Frequency).GetName(), Code.Num());
+	DEC_DWORD_STAT_BY(STAT_Shaders_ShaderResourceMemory, GetSizeBytes());
+	DEC_DWORD_STAT_BY(STAT_Shaders_NumShaderResourcesLoaded, 1);
+}
+
+
+void FShaderResource::UncompressCode(TArray<uint8>& UncompressedCode) const
+{
+	if (Code.Num() != UncompressedCodeSize)
+	{
+		UncompressedCode.SetNum(UncompressedCodeSize);
+		auto bSucceed = FCompression::UncompressMemory(ShaderCompressionFlag, UncompressedCode.GetData(), UncompressedCodeSize, Code.GetData(), Code.Num());
+		check(bSucceed);
+	}
+	else
+	{
+		UncompressedCode = Code;
+	}
+}
+
+void FShaderResource::CompressCode(const TArray<uint8>& UncompressedCode)
+{
+	UncompressedCodeSize = UncompressedCode.Num();
+	Code = UncompressedCode;
+	int32 CompressedSize = Code.Num();
+	if (FCompression::CompressMemory(ShaderCompressionFlag, Code.GetData(), CompressedSize, UncompressedCode.GetData(), UncompressedCode.Num()))
+	{
+		Code.SetNum(CompressedSize);
+	}
+	Code.Shrink();
+}
+
+void FShaderResource::Register()
+{
+	check(IsInGameThread());
+	ShaderResourceIdMap.Add(GetId(), this);
+}
+
+
+void FShaderResource::Serialize(FArchive& Ar, bool bLoadedByCookedMaterial)
+{
+	check(!(SpecificPermutationId != 0 && SpecificType == nullptr));
+
+	Ar.UsingCustomVersion(FRenderingObjectVersion::GUID);
+
+	Ar << SpecificType;
+	if (Ar.CustomVer(FRenderingObjectVersion::GUID) >= FRenderingObjectVersion::ShaderPermutationId)
+	{
+		Ar << SpecificPermutationId;
+	}
+	Ar << Target;
+
+	if (Ar.CustomVer(FRenderingObjectVersion::GUID) < FRenderingObjectVersion::ShaderResourceCodeSharing)
+	{
+		Ar << Code;
+	}
+	Ar << OutputHash;
+	Ar << NumInstructions;
+#if WITH_EDITORONLY_DATA
+	Ar << NumTextureSamplers;
+#else
+	uint32 Temp = 0;
+	Ar << Temp;
+#endif
+	
+	if (Ar.UE4Ver() >= VER_UE4_COMPRESSED_SHADER_RESOURCES)
+	{
+		Ar << UncompressedCodeSize;
+	}
+
+	if (Ar.CustomVer(FRenderingObjectVersion::GUID) >= FRenderingObjectVersion::ShaderResourceCodeSharing)
+	{
+		SerializeShaderCode(Ar);
+	}
+
+#if WITH_EDITORONLY_DATA
+	if (!bLoadedByCookedMaterial)
+	{
+		SerializePlatformDebugData(Ar);
+	}
+#endif
+
+	if (Ar.IsLoading())
+	{
+		INC_DWORD_STAT_BY_FName(GetMemoryStatType((EShaderFrequency)Target.Frequency).GetName(), (int64)Code.Num());
+		INC_DWORD_STAT_BY(STAT_Shaders_ShaderResourceMemory, GetSizeBytes());
+	}
+}
+
+void FShaderResource::SerializeShaderCode(FArchive& Ar)
+{
+	// To not pollute the DDC we don't change the state of this object in memory, just the state of the object in the serialised archive.
+	bool bCodeShared = bCodeInSharedLocation;
+	
+#if WITH_EDITOR
+	// in case shader code sharing is enabled, code will be saved outside of material asset
+	if(Ar.IsSaving() && Ar.IsCooking() && Ar.IsPersistent() && !Ar.IsObjectReferenceCollector() && !bCodeInSharedLocation)
+	{
+		bCodeShared = FShaderCodeLibrary::AddShaderCode((EShaderPlatform)Target.Platform, (EShaderFrequency)Target.Frequency, OutputHash, Code, UncompressedCodeSize);
+	}
+#endif
+	
+	Ar << bCodeShared;
+	
+	if (Ar.IsLoading())
+	{
+		bCodeInSharedLocation = bCodeShared;
+
+		if (bCodeInSharedLocation)
+		{
+			if (!GRHILazyShaderCodeLoading)
+			{
+				FShaderCodeLibrary::RequestShaderCode(OutputHash, &Ar);
+				bCodeInSharedLocationRequested = true;
+			}
+			else
+			{
+				FShaderCodeLibrary::LazyRequestShaderCode(OutputHash, &Ar);
+			}
+		}
+	}
+
+	if (!bCodeShared)
+	{
+		Ar << Code;
+	}
+}
+
+#if WITH_EDITORONLY_DATA
+void FShaderResource::SerializePlatformDebugData(FArchive& Ar)
+{
+#if WITH_ENGINE
+	if (Ar.IsCooking())
+	{
+		// Notify the platform shader format that this particular shader is being used in the cook.
+		// We discard this data in cooked builds unless Ar.CookingTarget()->HasEditorOnlyData() is true.
+		if (PlatformDebugData.Num())
+		{
+			TArray<FName> ShaderFormatNames;
+			Ar.CookingTarget()->GetAllTargetedShaderFormats(ShaderFormatNames);
+
+			for (FName FormatName : ShaderFormatNames)
+			{
+				const IShaderFormat* ShaderFormat = GetTargetPlatformManagerRef().FindShaderFormat(FormatName);
+				if (ShaderFormat)
+				{
+					ShaderFormat->NotifyShaderCooked(PlatformDebugData);
+				}
+			}
+		}
+	}
+
+	if (!Ar.IsCooking() || Ar.CookingTarget()->HasEditorOnlyData())
+#endif
+	{
+		// Always serialize if we're not cooking, the cook target requires editor only data, or we don't have the engine (i.e. we're SCW).
+		Ar << PlatformDebugData;
+	}
+}
+#endif
+
+void FShaderResource::AddRef()
+{
+	checkSlow(IsInGameThread());
+	++NumRefs;
+}
+
+
+void FShaderResource::Release()
+{
+	checkSlow(IsInGameThread());
+	check(NumRefs != 0);
+	if(--NumRefs == 0)
+	{
+		ShaderResourceIdMap.Remove(GetId());
+
+		// Send a release message to the rendering thread when the shader loses its last reference.
+		BeginReleaseResource(this);
+		BeginCleanup(this);
+
+		if (bCodeInSharedLocation)
+		{
+			if (bCodeInSharedLocationRequested)
+			{
+				FShaderCodeLibrary::ReleaseShaderCode(OutputHash);
+			}
+			else
+			{
+				FShaderCodeLibrary::LazyReleaseShaderCode(OutputHash);
+			}
+		}
+	}
+}
+
+
+FShaderResource* FShaderResource::FindShaderResourceById(const FShaderResourceId& Id)
+{
+	check(IsInGameThread());
+	FShaderResource* Result = ShaderResourceIdMap.FindRef(Id);
+	return Result;
+}
+
+
+FShaderResource* FShaderResource::FindOrCreateShaderResource(const FShaderCompilerOutput& Output, FShaderType* SpecificType, int32 SpecificPermutationId)
+{
+	const FShaderResourceId ResourceId(Output.Target, Output.OutputHash, SpecificType ? SpecificType->GetName() : nullptr, SpecificPermutationId);
+	FShaderResource* Resource = FindShaderResourceById(ResourceId);
+	if (!Resource)
+	{
+		Resource = new FShaderResource(Output, SpecificType, SpecificPermutationId);
+	}
+
+	return Resource;
+}
+
+void FShaderResource::GetAllShaderResourceId(TArray<FShaderResourceId>& Ids)
+{
+	check(IsInGameThread());
+	ShaderResourceIdMap.GetKeys(Ids);
+}
+
+bool FShaderResource::ArePlatformsCompatible(EShaderPlatform CurrentPlatform, EShaderPlatform TargetPlatform)
+{
+	bool bFeatureLevelCompatible = CurrentPlatform == TargetPlatform;
+	
+	if (!bFeatureLevelCompatible && IsPCPlatform(CurrentPlatform) && IsPCPlatform(TargetPlatform) )
+	{
+		bFeatureLevelCompatible = GetMaxSupportedFeatureLevel(CurrentPlatform) >= GetMaxSupportedFeatureLevel(TargetPlatform);
+		
+		bool const bIsTargetD3D = TargetPlatform == SP_PCD3D_SM5 ||
+		TargetPlatform == SP_PCD3D_SM4 ||
+		TargetPlatform == SP_PCD3D_ES3_1 ||
+		TargetPlatform == SP_PCD3D_ES2;
+		
+		bool const bIsCurrentPlatformD3D = CurrentPlatform == SP_PCD3D_SM5 ||
+		CurrentPlatform == SP_PCD3D_SM4 ||
+		TargetPlatform == SP_PCD3D_ES3_1 ||
+		CurrentPlatform == SP_PCD3D_ES2;
+		
+		// For Metal in Editor we can switch feature-levels, but not in cooked projects when using Metal shader librariss.
+		bool const bIsCurrentMetal = IsMetalPlatform(CurrentPlatform);
+		bool const bIsTargetMetal = IsMetalPlatform(TargetPlatform);
+		bool const bIsMetalCompatible = (bIsCurrentMetal == bIsTargetMetal) 
+#if !WITH_EDITOR	// Static analysis doesn't like (|| WITH_EDITOR)
+			&& (!IsMetalPlatform(CurrentPlatform) || (CurrentPlatform == TargetPlatform))
+#endif
+			;
+		
+		bool const bIsCurrentOpenGL = IsOpenGLPlatform(CurrentPlatform);
+		bool const bIsTargetOpenGL = IsOpenGLPlatform(TargetPlatform);
+		
+		bFeatureLevelCompatible = bFeatureLevelCompatible && (bIsCurrentPlatformD3D == bIsTargetD3D && bIsMetalCompatible && bIsCurrentOpenGL == bIsTargetOpenGL);
+	}
+
+	return bFeatureLevelCompatible;
+}
+
+FSHAHash &FShaderResource::FilterShaderSourceHashForSerialization(const FArchive& Ar, FSHAHash &HashToSerialize)
+{
+	return (!Ar.IsCooking()) ? HashToSerialize : ShaderSourceDefaultHash;
+}
+
+static void SafeAssignHash(FRHIShader* InShader, const FSHAHash& Hash)
+{
+	if (InShader)
+	{
+		InShader->SetHash(Hash);
+	}
+}
+
+void FShaderResource::InitRHI()
+{
+	checkf(bCodeInSharedLocation || Code.Num() > 0, TEXT("FShaderResource::InitRHI was called with empty bytecode, which can happen if the resource is initialized multiple times on platforms with no editor data."));
+
+	// we can't have this called on the wrong platform's shaders
+	if (!ArePlatformsCompatible(GMaxRHIShaderPlatform, (EShaderPlatform)Target.Platform))
+	{
+		if (FPlatformProperties::RequiresCookedData())
+		{
+			UE_LOG(LogShaders, Fatal, TEXT("FShaderResource::InitRHI got platform %s but it is not compatible with %s"), 
+				*LegacyShaderPlatformToShaderFormat((EShaderPlatform)Target.Platform).ToString(), *LegacyShaderPlatformToShaderFormat(GMaxRHIShaderPlatform).ToString());
+		}
+		return;
+	}
+
+	TArray<uint8> UncompressedCode;
+	if (!bCodeInSharedLocation)
+	{
+		UncompressCode(UncompressedCode);
+	}
+
+	INC_DWORD_STAT_BY(STAT_Shaders_NumShadersUsedForRendering, 1);
+	SCOPE_CYCLE_COUNTER(STAT_Shaders_RTShaderLoadTime);
+
+	if(Target.Frequency == SF_Vertex)
+	{
+		Shader = FShaderCodeLibrary::CreateVertexShader((EShaderPlatform)Target.Platform, OutputHash, UncompressedCode);
+		UE_CLOG((bCodeInSharedLocation && !IsValidRef(Shader)), LogShaders, Fatal, TEXT("FShaderResource::SerializeShaderCode can't find shader code for: [%s]"), *LegacyShaderPlatformToShaderFormat((EShaderPlatform)Target.Platform).ToString());
+	}
+	else if(Target.Frequency == SF_Pixel)
+	{
+		Shader = FShaderCodeLibrary::CreatePixelShader((EShaderPlatform)Target.Platform, OutputHash, UncompressedCode);
+		UE_CLOG((bCodeInSharedLocation && !IsValidRef(Shader)), LogShaders, Fatal, TEXT("FShaderResource::SerializeShaderCode can't find shader code for: [%s]"), *LegacyShaderPlatformToShaderFormat((EShaderPlatform)Target.Platform).ToString());
+	}
+	else if(Target.Frequency == SF_Hull)
+	{
+		Shader = FShaderCodeLibrary::CreateHullShader((EShaderPlatform)Target.Platform, OutputHash, UncompressedCode);
+		UE_CLOG((bCodeInSharedLocation && !IsValidRef(Shader)), LogShaders, Fatal, TEXT("FShaderResource::SerializeShaderCode can't find shader code for: [%s]"), *LegacyShaderPlatformToShaderFormat((EShaderPlatform)Target.Platform).ToString());
+	}
+	else if(Target.Frequency == SF_Domain)
+	{
+		Shader = FShaderCodeLibrary::CreateDomainShader((EShaderPlatform)Target.Platform, OutputHash, UncompressedCode);
+		UE_CLOG((bCodeInSharedLocation && !IsValidRef(Shader)), LogShaders, Fatal, TEXT("FShaderResource::SerializeShaderCode can't find shader code for: [%s]"), *LegacyShaderPlatformToShaderFormat((EShaderPlatform)Target.Platform).ToString());
+	}
+	else if(Target.Frequency == SF_Geometry)
+	{
+		if (SpecificType)
+		{
+			FStreamOutElementList ElementList;
+			TArray<uint32> StreamStrides;
+			int32 RasterizedStream = -1;
+			SpecificType->GetStreamOutElements(ElementList, StreamStrides, RasterizedStream);
+			checkf(ElementList.Num(), TEXT("Shader type %s was given GetStreamOutElements implementation that had no elements!"), SpecificType->GetName());
+
+			//@todo - not using the cache
+			Shader = FShaderCodeLibrary::CreateGeometryShaderWithStreamOutput((EShaderPlatform)Target.Platform, OutputHash, UncompressedCode, ElementList, StreamStrides.Num(), StreamStrides.GetData(), RasterizedStream);
+		}
+		else
+		{
+			Shader = FShaderCodeLibrary::CreateGeometryShader((EShaderPlatform)Target.Platform, OutputHash, UncompressedCode);
+		}
+		UE_CLOG((bCodeInSharedLocation && !IsValidRef(Shader)), LogShaders, Fatal, TEXT("FShaderResource::SerializeShaderCode can't find shader code for: [%s]"), *LegacyShaderPlatformToShaderFormat((EShaderPlatform)Target.Platform).ToString());
+	}
+	else if(Target.Frequency == SF_Compute)
+	{
+		Shader = FShaderCodeLibrary::CreateComputeShader((EShaderPlatform)Target.Platform, OutputHash, UncompressedCode);
+		UE_CLOG((bCodeInSharedLocation && !IsValidRef(Shader)), LogShaders, Fatal, TEXT("FShaderResource::SerializeShaderCode can't find shader code for: [%s]"), *LegacyShaderPlatformToShaderFormat((EShaderPlatform)Target.Platform).ToString());
+	}
+
+	if (Target.Frequency != SF_Geometry)
+	{
+		checkf(!SpecificType, TEXT("Only geometry shaders can use GetStreamOutElements, shader type %s"), SpecificType->GetName());
+	}
+
+	if (!FPlatformProperties::HasEditorOnlyData())
+	{
+		DEC_DWORD_STAT_BY_FName(GetMemoryStatType((EShaderFrequency)Target.Frequency).GetName(), Code.Num());
+		DEC_DWORD_STAT_BY(STAT_Shaders_ShaderResourceMemory, Code.GetAllocatedSize());
+		Code.Empty();
+		
+		if (bCodeInSharedLocation)
+		{
+			if (bCodeInSharedLocationRequested)
+			{
+				FShaderCodeLibrary::ReleaseShaderCode(OutputHash);
+			}
+			else
+			{
+				FShaderCodeLibrary::LazyReleaseShaderCode(OutputHash);
+			}
+		}
+		bCodeInSharedLocation = false;
+		bCodeInSharedLocationRequested = false;
+	}
+}
+
+
+void FShaderResource::ReleaseRHI()
+{
+	DEC_DWORD_STAT_BY(STAT_Shaders_NumShadersUsedForRendering, 1);
+
+	Shader.SafeRelease();
+}
+
+void FShaderResource::InitializeShaderRHI() 
+{ 
+	if (!IsInitialized())
+	{
+		STAT(double ShaderInitializationTime = 0);
+		{
+			SCOPE_CYCLE_COUNTER(STAT_Shaders_FrameRTShaderInitForRenderingTime);
+			SCOPE_SECONDS_COUNTER(ShaderInitializationTime);
+
+			InitResourceFromPossiblyParallelRendering();
+		}
+
+		INC_FLOAT_STAT_BY(STAT_Shaders_TotalRTShaderInitForRenderingTime,(float)ShaderInitializationTime);
+	}
+
+	checkSlow(IsInitialized());
+}
+
+FShaderResourceId FShaderResource::GetId() const
+{
+	return FShaderResourceId(Target, OutputHash, SpecificType ? SpecificType->GetName() : nullptr, SpecificPermutationId);
+}
+
+FShaderId::FShaderId(const FSHAHash& InMaterialShaderMapHash, const FShaderPipelineType* InShaderPipeline, FVertexFactoryType* InVertexFactoryType, FShaderType* InShaderType, int32 InPermutationId, FShaderTarget InTarget)
+	: MaterialShaderMapHash(InMaterialShaderMapHash)
+	, SourceHash(InShaderType->GetSourceHash(InTarget.GetPlatform()))
+	, Target(InTarget)
+	, ShaderPipeline(InShaderPipeline)
+	, ShaderType(InShaderType)
+	, PermutationId(InPermutationId)
+	, SerializationHistory(InShaderType->GetSerializationHistory())
+{
+	if (InVertexFactoryType)
+	{
+		VFSerializationHistory = InVertexFactoryType->GetSerializationHistory(InTarget.GetFrequency());
+		VertexFactoryType = InVertexFactoryType;
+		VFSourceHash = InVertexFactoryType->GetSourceHash(InTarget.GetPlatform());
+	}
+	else
+	{
+		VFSerializationHistory = nullptr;
+		VertexFactoryType = nullptr;
+	}
+}
+
+FSelfContainedShaderId::FSelfContainedShaderId() :
+	Target(FShaderTarget(SF_NumFrequencies, SP_NumPlatforms))
+{}
+
+FSelfContainedShaderId::FSelfContainedShaderId(const FShaderId& InShaderId)
+{
+	MaterialShaderMapHash = InShaderId.MaterialShaderMapHash;
+	VertexFactoryTypeName = InShaderId.VertexFactoryType ? InShaderId.VertexFactoryType->GetName() : TEXT("");
+	ShaderPipelineName = InShaderId.ShaderPipeline ? InShaderId.ShaderPipeline->GetName() : TEXT("");
+	VFSourceHash = InShaderId.VFSourceHash;
+	VFSerializationHistory = InShaderId.VFSerializationHistory ? *InShaderId.VFSerializationHistory : FSerializationHistory();
+	ShaderTypeName = InShaderId.ShaderType->GetName();
+	PermutationId = InShaderId.PermutationId;
+	SourceHash = InShaderId.SourceHash;
+	SerializationHistory = InShaderId.SerializationHistory;
+	Target = InShaderId.Target;
+}
+
+bool FSelfContainedShaderId::IsValid()
+{
+	FShaderType** TypePtr = FShaderType::GetNameToTypeMap().Find(FName(*ShaderTypeName));
+	if (TypePtr && SourceHash == (*TypePtr)->GetSourceHash(Target.GetPlatform()) && SerializationHistory == (*TypePtr)->GetSerializationHistory())
+	{
+		FVertexFactoryType* VFTypePtr = FVertexFactoryType::GetVFByName(VertexFactoryTypeName);
+
+		if (VertexFactoryTypeName == TEXT("") 
+			|| (VFTypePtr && VFSourceHash == VFTypePtr->GetSourceHash(Target.GetPlatform()) && VFSerializationHistory == *VFTypePtr->GetSerializationHistory(Target.GetFrequency())))
+		{
+			return true;
+		}
+	}
+
+	return false;
+}
+
+FArchive& operator<<(FArchive& Ar,class FSelfContainedShaderId& Ref)
+{
+	Ar.UsingCustomVersion(FRenderingObjectVersion::GUID);
+
+	Ar << Ref.MaterialShaderMapHash 
+		<< Ref.VertexFactoryTypeName
+		<< Ref.ShaderPipelineName
+		<< FShaderResource::FilterShaderSourceHashForSerialization(Ar, Ref.VFSourceHash)
+		<< Ref.VFSerializationHistory
+		<< Ref.ShaderTypeName
+		<< FShaderResource::FilterShaderSourceHashForSerialization(Ar, Ref.SourceHash)
+		<< Ref.SerializationHistory
+		<< Ref.Target;
+
+	if (Ar.CustomVer(FRenderingObjectVersion::GUID) >= FRenderingObjectVersion::ShaderPermutationId)
+	{
+		Ar << Ref.PermutationId;
+	}
+
+	return Ar;
+}
+
+/** 
+ * Used to construct a shader for deserialization.
+ * This still needs to initialize members to safe values since FShaderType::GenerateSerializationHistory uses this constructor.
+ */
+FShader::FShader()
+	: SerializedResource(nullptr)
+	, ShaderPipeline(nullptr)
+	, VFType(nullptr)
+	, Type(nullptr)
+	, PermutationId(0)
+	, NumRefs(0)
+{
+	// set to undefined (currently shared with SF_Vertex)
+	Target.Frequency = 0;
+	Target.Platform = GShaderPlatformForFeatureLevel[GMaxRHIFeatureLevel];
+}
+
+/**
+ * Construct a shader from shader compiler output.
+ */
+FShader::FShader(const CompiledShaderInitializerType& Initializer)
+	: MaterialShaderMapHash(Initializer.MaterialShaderMapHash)
+	, SerializedResource(nullptr)
+	, ShaderPipeline(Initializer.ShaderPipeline)
+	, VFType(Initializer.VertexFactoryType)
+	, Type(Initializer.Type)
+	, PermutationId(Initializer.PermutationId)
+	, Target(Initializer.Target)
+	, NumRefs(0)
+{
+	OutputHash = Initializer.OutputHash;
+	checkSlow(OutputHash != FSHAHash());
+
+	check(Type);
+
+	// Store off the source hash that this shader was compiled with
+	// This will be used as part of the shader key in order to identify when shader files have been changed and a recompile is needed
+	SourceHash = Type->GetSourceHash(Target.GetPlatform());
+
+	if (VFType)
+	{
+		// Store off the VF source hash that this shader was compiled with
+		VFSourceHash = VFType->GetSourceHash(Target.GetPlatform());
+	}
+
+	// Bind uniform buffer parameters automatically 
+	for (TLinkedList<FShaderParametersMetadata*>::TIterator StructIt(FShaderParametersMetadata::GetStructList()); StructIt; StructIt.Next())
+	{
+		if (Initializer.ParameterMap.ContainsParameterAllocation(StructIt->GetShaderVariableName()))
+		{
+			UniformBufferParameterStructs.Add(*StructIt);
+			UniformBufferParameters.Add(new FShaderUniformBufferParameter());
+			FShaderUniformBufferParameter* Parameter = UniformBufferParameters.Last();
+			Parameter->Bind(Initializer.ParameterMap, StructIt->GetShaderVariableName(), SPF_Mandatory);
+		}
+	}
+
+	SetResource(Initializer.Resource);
+
+	// Register the shader now that it is valid, so that it can be reused
+	Register(false);
+}
+
+
+FShader::~FShader()
+{
+	check(NumRefs == 0);
+
+	for (int32 StructIndex = 0; StructIndex < UniformBufferParameters.Num(); StructIndex++)
+	{
+		delete UniformBufferParameters[StructIndex];
+	}
+}
+
+
+const FSHAHash& FShader::GetHash() const 
+{ 
+	return SourceHash;
+}
+
+EShaderPlatform FShader::GetShaderPlatform() const
+{
+	return Target.GetPlatform();
+}
+
+
+bool FShader::SerializeBase(FArchive& Ar, bool bShadersInline, bool bLoadedByCookedMaterial)
+{
+	Ar.UsingCustomVersion(FFortniteMainBranchObjectVersion::GUID);
+
+	Serialize(Ar);
+
+	Ar.UsingCustomVersion(FRenderingObjectVersion::GUID);
+
+	Ar << OutputHash;
+	Ar << MaterialShaderMapHash;
+	Ar << ShaderPipeline;
+	Ar << VFType;
+	Ar << FShaderResource::FilterShaderSourceHashForSerialization(Ar, VFSourceHash);
+	Ar << Type;
+	if (Ar.CustomVer(FRenderingObjectVersion::GUID) >= FRenderingObjectVersion::ShaderPermutationId)
+	{
+		Ar << PermutationId;
+	}
+	Ar << FShaderResource::FilterShaderSourceHashForSerialization(Ar, SourceHash);
+	Ar << Target;
+
+	// TODO(RDG): Kill that once all shaders are refactored.
+	if (Ar.IsLoading())
+	{
+		int32 NumUniformParameters;
+		Ar << NumUniformParameters;
+
+		UniformBufferParameterStructs.Empty(NumUniformParameters);
+		UniformBufferParameters.Empty(NumUniformParameters);
+
+		for (int32 ParameterIndex = 0; ParameterIndex < NumUniformParameters; ParameterIndex++)
+		{
+			FShaderParametersMetadata* Struct = nullptr;
+
+			if (Ar.CustomVer(FFortniteMainBranchObjectVersion::GUID) < FFortniteMainBranchObjectVersion::MaterialInstanceSerializeOptimization_ShaderFName)
+			{
+				FString StructName;
+				Ar << StructName;
+				Struct = FindUniformBufferStructByName(*StructName);
+				checkf(Struct, TEXT("Uniform Buffer Struct %s no longer exists, which shader of type %s was compiled with.  Modify ShaderVersion.ush to invalidate old shaders."), *StructName, Type->GetName());
+			}
+			else
+			{
+				FName StructFName;
+				Ar << StructFName;
+				Struct = FindUniformBufferStructByFName(StructFName);
+				checkf(Struct, TEXT("Uniform Buffer Struct %s no longer exists, which shader of type %s was compiled with.  Modify ShaderVersion.ush to invalidate old shaders."), *StructFName.ToString(), Type->GetName());
+			}
+			
+			FShaderUniformBufferParameter* Parameter = new FShaderUniformBufferParameter();
+
+			Ar << *Parameter;
+
+			UniformBufferParameterStructs.Add(Struct);
+			UniformBufferParameters.Add(Parameter);
+		}
+	}
+	else
+	{
+		int32 NumUniformParameters = UniformBufferParameters.Num();
+		Ar << NumUniformParameters;
+
+		for (int32 StructIndex = 0; StructIndex < UniformBufferParameters.Num(); StructIndex++)
+		{
+			FString StructName(UniformBufferParameterStructs[StructIndex]->GetStructTypeName());
+
+			if (Ar.CustomVer(FFortniteMainBranchObjectVersion::GUID) < FFortniteMainBranchObjectVersion::MaterialInstanceSerializeOptimization_ShaderFName)
+			{
+				Ar << StructName;
+			}
+			else
+			{
+				FName StructFName(*StructName);
+				Ar << StructFName;
+			}
+
+			Ar << *UniformBufferParameters[StructIndex];
+		}
+	}
+
+	if (bShadersInline)
+	{
+		// Save the shader resource if we are inlining shaders
+		if (Ar.IsSaving())
+		{
+			check(Resource->Target == Target);
+			Resource->Serialize(Ar, false);
+		}
+
+		if (Ar.IsLoading())
+		{
+			// Load the inlined shader resource
+			SerializedResource = new FShaderResource();
+			SerializedResource->Serialize(Ar, bLoadedByCookedMaterial);
+		}
+	}
+	else
+	{
+		// if saving, there's nothing to, the required data is already saved above to look it up at load time
+		if (Ar.IsLoading())
+		{
+			// generate a resource id
+			FShaderResourceId ResourceId(Target, OutputHash, Type->LimitShaderResourceToThisType() ? Type->GetName() : nullptr, Type->LimitShaderResourceToThisType() ? PermutationId : 0);
+
+			// use it to look up in the registered resource map
+			FShaderResource* ExistingResource = FShaderResource::FindShaderResourceById(ResourceId);
+			SetResource(ExistingResource);
+		}
+	}
+
+	Ar << Bindings;
+
+	return false;
+}
+
+void FShader::AddRef()
+{	
+	++NumRefs;
+	if (NumRefs == 1)
+	{
+		INC_DWORD_STAT_BY(STAT_Shaders_ShaderMemory, GetSizeBytes());
+		INC_DWORD_STAT_BY(STAT_Shaders_NumShadersLoaded,1);
+	}
+}
+
+void FShader::Release()
+{
+	if(--NumRefs == 0)
+	{
+		DEC_DWORD_STAT_BY(STAT_Shaders_ShaderMemory, GetSizeBytes());
+		DEC_DWORD_STAT_BY(STAT_Shaders_NumShadersLoaded,1);
+
+		// Deregister the shader now to eliminate references to it by the type's ShaderIdMap
+		Deregister();
+		BeginCleanup(this);
+	}
+}
+
+
+void FShader::Register(bool bLoadedByCookedMaterial)
+{
+	FShaderId ShaderId = GetId();
+	check(ShaderId.MaterialShaderMapHash != FSHAHash());
+	check(ShaderId.SourceHash != FSHAHash() || FPlatformProperties::RequiresCookedData() || bLoadedByCookedMaterial);
+	check(Resource);
+	Type->AddToShaderIdMap(ShaderId, this);
+}
+
+void FShader::Deregister()
+{
+	Type->RemoveFromShaderIdMap(GetId());
+}
+
+FShaderId FShader::GetId() const
+{
+	FShaderId ShaderId(Type->GetSerializationHistory());
+	ShaderId.MaterialShaderMapHash = MaterialShaderMapHash;
+	ShaderId.ShaderPipeline = ShaderPipeline;
+	ShaderId.VertexFactoryType = VFType;
+	ShaderId.VFSourceHash = VFSourceHash;
+	ShaderId.VFSerializationHistory = VFType ? VFType->GetSerializationHistory((EShaderFrequency)GetTarget().Frequency) : NULL;
+	ShaderId.ShaderType = Type;
+	ShaderId.PermutationId = PermutationId;
+	ShaderId.SourceHash = SourceHash;
+	ShaderId.Target = Target;
+	return ShaderId;
+}
+
+void FShader::RegisterSerializedResource()
+{
+	if (SerializedResource)
+	{
+		FShaderResource* ExistingResource = FShaderResource::FindShaderResourceById(SerializedResource->GetId());
+
+		// Reuse an existing shader resource if a matching one already exists in memory
+		if (ExistingResource)
+		{
+			delete SerializedResource;
+			SerializedResource = ExistingResource;
+		}
+		else
+		{
+			// Register the newly loaded shader resource so it can be reused by other shaders
+			SerializedResource->Register();
+		}
+
+		SetResource(SerializedResource);
+	}
+}
+
+void FShader::SetResource(FShaderResource* InResource)
+{
+	check(InResource && InResource->Target == Target);
+	Resource = InResource;
+}
+
+void FShader::DumpDebugInfo()
+{
+	UE_LOG(LogConsoleResponse, Display, TEXT("      FShader  :MaterialShaderMapHash %s"), *MaterialShaderMapHash.ToString());
+	UE_LOG(LogConsoleResponse, Display, TEXT("               :Target %s"), GetShaderFrequencyString((EShaderFrequency)Target.Frequency));
+	UE_LOG(LogConsoleResponse, Display, TEXT("               :Target %s"), *LegacyShaderPlatformToShaderFormat(EShaderPlatform(Target.Platform)).ToString());
+	UE_LOG(LogConsoleResponse, Display, TEXT("               :VFType %s"), VFType ? VFType->GetName() : TEXT("null"));
+	UE_LOG(LogConsoleResponse, Display, TEXT("               :VFSourceHash %s"), *VFSourceHash.ToString());
+	UE_LOG(LogConsoleResponse, Display, TEXT("               :Type %s"), Type->GetName());
+	UE_LOG(LogConsoleResponse, Display, TEXT("               :PermutationId %d"), PermutationId);
+	UE_LOG(LogConsoleResponse, Display, TEXT("               :SourceHash %s"), *SourceHash.ToString());
+	UE_LOG(LogConsoleResponse, Display, TEXT("               :OutputHash %s"), *OutputHash.ToString());
+}
+
+void FShader::SaveShaderStableKeys(EShaderPlatform TargetShaderPlatform, const FStableShaderKeyAndValue& InSaveKeyVal)
+{
+#if WITH_EDITOR
+	if (TargetShaderPlatform == EShaderPlatform::SP_NumPlatforms || EShaderPlatform(Target.Platform) == TargetShaderPlatform)
+	{
+		FStableShaderKeyAndValue SaveKeyVal(InSaveKeyVal);
+		SaveKeyVal.TargetFrequency = FName(GetShaderFrequencyString((EShaderFrequency)Target.Frequency));
+		SaveKeyVal.TargetPlatform = FName(*LegacyShaderPlatformToShaderFormat(EShaderPlatform(Target.Platform)).ToString());
+		SaveKeyVal.VFType = FName(VFType ? VFType->GetName() : TEXT("null"));
+		SaveKeyVal.PermutationId = FName(*FString::Printf(TEXT("Perm_%d"), PermutationId));
+		SaveKeyVal.OutputHash = OutputHash;
+		if (Type)
+		{
+			Type->GetShaderStableKeyParts(SaveKeyVal);
+		}
+		FShaderCodeLibrary::AddShaderStableKeyValue(EShaderPlatform(Target.Platform), SaveKeyVal);
+	}
+#endif
+}
+
+bool FShaderPipelineType::bInitialized = false;
+
+FShaderPipelineType::FShaderPipelineType(
+	const TCHAR* InName,
+	const FShaderType* InVertexShader,
+	const FShaderType* InHullShader,
+	const FShaderType* InDomainShader,
+	const FShaderType* InGeometryShader,
+	const FShaderType* InPixelShader,
+	bool bInShouldOptimizeUnusedOutputs) :
+	Name(InName),
+	TypeName(Name),
+	GlobalListLink(this),
+	bShouldOptimizeUnusedOutputs(bInShouldOptimizeUnusedOutputs)
+{
+	checkf(Name && *Name, TEXT("Shader Pipeline Type requires a valid Name!"));
+
+	checkf(InVertexShader, TEXT("A Shader Pipeline always requires a Vertex Shader"));
+
+	checkf((InHullShader == nullptr && InDomainShader == nullptr) || (InHullShader != nullptr && InDomainShader != nullptr), TEXT("Both Hull & Domain shaders are needed for tessellation on Pipeline %s"), Name);
+
+	//make sure the name is shorter than the maximum serializable length
+	check(FCString::Strlen(InName) < NAME_SIZE);
+
+	FMemory::Memzero(AllStages);
+
+	if (InPixelShader)
+	{
+		Stages.Add(InPixelShader);
+		AllStages[SF_Pixel] = InPixelShader;
+	}
+	if (InGeometryShader)
+	{
+		Stages.Add(InGeometryShader);
+		AllStages[SF_Geometry] = InGeometryShader;
+	}
+	if (InDomainShader)
+	{
+		Stages.Add(InDomainShader);
+		AllStages[SF_Domain] = InDomainShader;
+
+		Stages.Add(InHullShader);
+		AllStages[SF_Hull] = InHullShader;
+	}
+	Stages.Add(InVertexShader);
+	AllStages[SF_Vertex] = InVertexShader;
+
+	static uint32 TypeHashCounter = 0;
+	++TypeHashCounter;
+	HashIndex = TypeHashCounter;
+
+	GlobalListLink.LinkHead(GetTypeList());
+	GetNameToTypeMap().Add(TypeName, this);
+
+	// This will trigger if an IMPLEMENT_SHADER_TYPE was in a module not loaded before InitializeShaderTypes
+	// Shader types need to be implemented in modules that are loaded before that
+	checkf(!bInitialized, TEXT("Shader Pipeline was loaded after Engine init, use ELoadingPhase::PostConfigInit on your module to cause it to load earlier."));
+}
+
+FShaderPipelineType::~FShaderPipelineType()
+{
+	GetNameToTypeMap().Remove(TypeName);
+	GlobalListLink.Unlink();
+}
+
+TMap<FName, FShaderPipelineType*>& FShaderPipelineType::GetNameToTypeMap()
+{
+	static TMap<FName, FShaderPipelineType*>* GShaderPipelineNameToTypeMap = NULL;
+	if (!GShaderPipelineNameToTypeMap)
+	{
+		GShaderPipelineNameToTypeMap = new TMap<FName, FShaderPipelineType*>();
+	}
+	return *GShaderPipelineNameToTypeMap;
+}
+
+TLinkedList<FShaderPipelineType*>*& FShaderPipelineType::GetTypeList()
+{
+	return GShaderPipelineList;
+}
+
+TArray<const FShaderPipelineType*> FShaderPipelineType::GetShaderPipelineTypesByFilename(const TCHAR* Filename)
+{
+	TArray<const FShaderPipelineType*> PipelineTypes;
+	for (TLinkedList<FShaderPipelineType*>::TIterator It(FShaderPipelineType::GetTypeList()); It; It.Next())
+	{
+		auto* PipelineType = *It;
+		for (auto* ShaderType : PipelineType->Stages)
+		{
+			if (FPlatformString::Strcmp(Filename, ShaderType->GetShaderFilename()) == 0)
+			{
+				PipelineTypes.AddUnique(PipelineType);
+				break;
+			}
+		}
+	}
+	return PipelineTypes;
+}
+
+void FShaderPipelineType::Initialize()
+{
+	check(!bInitialized);
+
+	TSet<FName> UsedNames;
+
+#if UE_BUILD_DEBUG
+	TArray<const FShaderPipelineType*> UniqueShaderPipelineTypes;
+#endif
+	for (TLinkedList<FShaderPipelineType*>::TIterator It(FShaderPipelineType::GetTypeList()); It; It.Next())
+	{
+		const auto* PipelineType = *It;
+
+#if UE_BUILD_DEBUG
+		UniqueShaderPipelineTypes.Add(PipelineType);
+#endif
+
+		// Validate stages
+		for (int32 Index = 0; Index < SF_NumFrequencies; ++Index)
+		{
+			check(!PipelineType->AllStages[Index] || PipelineType->AllStages[Index]->GetFrequency() == (EShaderFrequency)Index);
+		}
+
+		auto& Stages = PipelineType->GetStages();
+
+		// #todo-rco: Do we allow mix/match of global/mesh/material stages?
+		// Check all shaders are the same type, start from the top-most stage
+		const FGlobalShaderType* GlobalType = Stages[0]->GetGlobalShaderType();
+		const FMeshMaterialShaderType* MeshType = Stages[0]->GetMeshMaterialShaderType();
+		const FMaterialShaderType* MateriallType = Stages[0]->GetMaterialShaderType();
+		for (int32 Index = 1; Index < Stages.Num(); ++Index)
+		{
+			if (GlobalType)
+			{
+				checkf(Stages[Index]->GetGlobalShaderType(), TEXT("Invalid combination of Shader types on Pipeline %s"), PipelineType->Name);
+			}
+			else if (MeshType)
+			{
+				checkf(Stages[Index]->GetMeshMaterialShaderType(), TEXT("Invalid combination of Shader types on Pipeline %s"), PipelineType->Name);
+			}
+			else if (MateriallType)
+			{
+				checkf(Stages[Index]->GetMaterialShaderType(), TEXT("Invalid combination of Shader types on Pipeline %s"), PipelineType->Name);
+			}
+		}
+
+		FName PipelineName = PipelineType->GetFName();
+		checkf(!UsedNames.Contains(PipelineName), TEXT("Two Pipelines with the same name %s found!"), PipelineType->Name);
+		UsedNames.Add(PipelineName);
+	}
+
+#if UE_BUILD_DEBUG
+	// Check for duplicated shader pipeline type names
+	UniqueShaderPipelineTypes.Sort([](const FShaderPipelineType& A, const FShaderPipelineType& B) { return (SIZE_T)&A < (SIZE_T)&B; });
+	for (int32 Index = 1; Index < UniqueShaderPipelineTypes.Num(); ++Index)
+	{
+		checkf(UniqueShaderPipelineTypes[Index - 1] != UniqueShaderPipelineTypes[Index], TEXT("Duplicated FShaderPipeline type name %s found, please rename one of them!"), UniqueShaderPipelineTypes[Index]->GetName());
+	}
+#endif
+
+	bInitialized = true;
+}
+
+void FShaderPipelineType::Uninitialize()
+{
+	check(bInitialized);
+
+	bInitialized = false;
+}
+
+void FShaderPipelineType::GetOutdatedTypes(TArray<FShaderType*>& OutdatedShaderTypes, TArray<const FShaderPipelineType*>& OutdatedShaderPipelineTypes, TArray<const FVertexFactoryType*>& OutdatedFactoryTypes)
+{
+	for (TLinkedList<FShaderPipelineType*>::TIterator It(FShaderPipelineType::GetTypeList()); It; It.Next())
+	{
+		const auto* PipelineType = *It;
+		auto& Stages = PipelineType->GetStages();
+		bool bOutdated = false;
+		for (const FShaderType* ShaderType : Stages)
+		{
+			bOutdated = ShaderType->GetOutdatedCurrentType(OutdatedShaderTypes, OutdatedFactoryTypes) || bOutdated;
+		}
+
+		if (bOutdated)
+		{
+			OutdatedShaderPipelineTypes.AddUnique(PipelineType);
+		}
+	}
+
+	for (int32 TypeIndex = 0; TypeIndex < OutdatedShaderPipelineTypes.Num(); TypeIndex++)
+	{
+		UE_LOG(LogShaders, Warning, TEXT("		Recompiling Pipeline %s"), OutdatedShaderPipelineTypes[TypeIndex]->GetName());
+	}
+}
+
+const FShaderPipelineType* FShaderPipelineType::GetShaderPipelineTypeByName(FName Name)
+{
+	for (TLinkedList<FShaderPipelineType*>::TIterator It(GetTypeList()); It; It.Next())
+	{
+		const FShaderPipelineType* Type = *It;
+		if (Name == Type->GetFName())
+		{
+			return Type;
+		}
+	}
+
+	return nullptr;
+}
+
+const FSHAHash& FShaderPipelineType::GetSourceHash(EShaderPlatform ShaderPlatform) const
+{
+	TArray<FString> Filenames;
+	for (const FShaderType* ShaderType : Stages)
+	{
+		Filenames.Add(ShaderType->GetShaderFilename());
+	}
+	return GetShaderFilesHash(Filenames, ShaderPlatform);
+}
+
+
+FShaderPipeline::FShaderPipeline(
+	const FShaderPipelineType* InPipelineType,
+	FShader* InVertexShader,
+	FShader* InHullShader,
+	FShader* InDomainShader,
+	FShader* InGeometryShader,
+	FShader* InPixelShader) :
+	PipelineType(InPipelineType),
+	VertexShader(InVertexShader),
+	HullShader(InHullShader),
+	DomainShader(InDomainShader),
+	GeometryShader(InGeometryShader),
+	PixelShader(InPixelShader)
+{
+	check(InPipelineType);
+	Validate();
+}
+
+FShaderPipeline::FShaderPipeline(const FShaderPipelineType* InPipelineType, const TArray<FShader*>& InStages) :
+	PipelineType(InPipelineType),
+	VertexShader(nullptr),
+	HullShader(nullptr),
+	DomainShader(nullptr),
+	GeometryShader(nullptr),
+	PixelShader(nullptr)
+{
+	check(InPipelineType);
+	for (FShader* Shader : InStages)
+	{
+		if (Shader)
+		{
+			switch (Shader->GetType()->GetFrequency())
+			{
+			case SF_Vertex:
+				check(!VertexShader);
+				VertexShader = Shader;
+				break;
+			case SF_Pixel:
+				check(!PixelShader);
+				PixelShader = Shader;
+				break;
+			case SF_Hull:
+				check(!HullShader);
+				HullShader = Shader;
+				break;
+			case SF_Domain:
+				check(!DomainShader);
+				DomainShader = Shader;
+				break;
+			case SF_Geometry:
+				check(!GeometryShader);
+				GeometryShader = Shader;
+				break;
+			default:
+				checkf(0, TEXT("Invalid stage %u found!"), (uint32)Shader->GetType()->GetFrequency());
+				break;
+			}
+		}
+	}
+
+	Validate();
+}
+
+FShaderPipeline::FShaderPipeline(const FShaderPipelineType* InPipelineType, const TArray< TRefCountPtr<FShader> >& InStages) :
+	PipelineType(InPipelineType),
+	VertexShader(nullptr),
+	HullShader(nullptr),
+	DomainShader(nullptr),
+	GeometryShader(nullptr),
+	PixelShader(nullptr)
+{
+	check(InPipelineType);
+	for (FShader* Shader : InStages)
+	{
+		if (Shader)
+		{
+			switch (Shader->GetType()->GetFrequency())
+			{
+			case SF_Vertex:
+				check(!VertexShader);
+				VertexShader = Shader;
+				break;
+			case SF_Pixel:
+				check(!PixelShader);
+				PixelShader = Shader;
+				break;
+			case SF_Hull:
+				check(!HullShader);
+				HullShader = Shader;
+				break;
+			case SF_Domain:
+				check(!DomainShader);
+				DomainShader = Shader;
+				break;
+			case SF_Geometry:
+				check(!GeometryShader);
+				GeometryShader = Shader;
+				break;
+			default:
+				checkf(0, TEXT("Invalid stage %u found!"), (uint32)Shader->GetType()->GetFrequency());
+				break;
+			}
+		}
+	}
+
+	Validate();
+}
+
+FShaderPipeline::~FShaderPipeline()
+{
+	// Manually set references to nullptr, helps debugging
+	VertexShader = nullptr;
+	HullShader = nullptr;
+	DomainShader = nullptr;
+	GeometryShader = nullptr;
+	PixelShader = nullptr;
+}
+
+void FShaderPipeline::Validate()
+{
+	for (const FShaderType* Stage : PipelineType->GetStages())
+	{
+		switch (Stage->GetFrequency())
+		{
+		case SF_Vertex:
+			check(VertexShader && VertexShader->GetType() == Stage);
+			break;
+		case SF_Pixel:
+			check(PixelShader && PixelShader->GetType() == Stage);
+			break;
+		case SF_Hull:
+			check(HullShader && HullShader->GetType() == Stage);
+			break;
+		case SF_Domain:
+			check(DomainShader && DomainShader->GetType() == Stage);
+			break;
+		case SF_Geometry:
+			check(GeometryShader && GeometryShader->GetType() == Stage);
+			break;
+		default:
+			// Can never happen :)
+			break;
+		}
+	}
+}
+
+void FShaderPipeline::CookPipeline(FShaderPipeline* Pipeline)
+{
+#if WITH_EDITOR
+	FShaderCodeLibrary::AddShaderPipeline(Pipeline);
+#endif
+}
+
+void DumpShaderStats(EShaderPlatform Platform, EShaderFrequency Frequency)
+{
+#if ALLOW_DEBUG_FILES
+	FDiagnosticTableViewer ShaderTypeViewer(*FDiagnosticTableViewer::GetUniqueTemporaryFilePath(TEXT("ShaderStats")));
+
+	// Iterate over all shader types and log stats.
+	int32 TotalShaderCount		= 0;
+	int32 TotalTypeCount		= 0;
+	int32 TotalInstructionCount	= 0;
+	int32 TotalSize				= 0;
+	int32 TotalPipelineCount	= 0;
+	float TotalSizePerType		= 0;
+
+	// Write a row of headings for the table's columns.
+	ShaderTypeViewer.AddColumn(TEXT("Type"));
+	ShaderTypeViewer.AddColumn(TEXT("Instances"));
+	ShaderTypeViewer.AddColumn(TEXT("Average instructions"));
+	ShaderTypeViewer.AddColumn(TEXT("Size"));
+	ShaderTypeViewer.AddColumn(TEXT("AvgSizePerInstance"));
+	ShaderTypeViewer.AddColumn(TEXT("Pipelines"));
+	ShaderTypeViewer.AddColumn(TEXT("Shared Pipelines"));
+	ShaderTypeViewer.CycleRow();
+
+	for( TLinkedList<FShaderType*>::TIterator It(FShaderType::GetTypeList()); It; It.Next() )
+	{
+		const FShaderType* Type = *It;
+		if (Type->GetNumShaders())
+		{
+			// Calculate the average instruction count and total size of instances of this shader type.
+			float AverageNumInstructions	= 0.0f;
+			int32 NumInitializedInstructions	= 0;
+			int32 Size						= 0;
+			int32 NumShaders					= 0;
+			int32 NumPipelines = 0;
+			int32 NumSharedPipelines = 0;
+			for (TMap<FShaderId,FShader*>::TConstIterator ShaderIt(Type->ShaderIdMap);ShaderIt;++ShaderIt)
+			{
+				const FShader* Shader = ShaderIt.Value();
+				// Skip shaders that don't match frequency.
+				if( Shader->GetTarget().Frequency != Frequency && Frequency != SF_NumFrequencies )
+				{
+					continue;
+				}
+				// Skip shaders that don't match platform.
+				if( Shader->GetTarget().Platform != Platform && Platform != SP_NumPlatforms )
+				{
+					continue;
+				}
+
+				NumInitializedInstructions += Shader->GetNumInstructions();
+				Size += Shader->GetCode().Num();
+				NumShaders++;
+			}
+			AverageNumInstructions = (float)NumInitializedInstructions / (float)Type->GetNumShaders();
+			
+			for (TLinkedList<FShaderPipelineType*>::TConstIterator PipelineIt(FShaderPipelineType::GetTypeList()); PipelineIt; PipelineIt.Next())
+			{
+				const FShaderPipelineType* PipelineType = *PipelineIt;
+				bool bFound = false;
+				if (Frequency == SF_NumFrequencies)
+				{
+					if (PipelineType->GetShader(Type->GetFrequency()) == Type)
+					{
+						++NumPipelines;
+						bFound = true;
+					}
+				}
+				else
+				{
+					if (PipelineType->GetShader(Frequency) == Type)
+					{
+						++NumPipelines;
+						bFound = true;
+					}
+				}
+
+				if (!PipelineType->ShouldOptimizeUnusedOutputs() && bFound)
+				{
+					++NumSharedPipelines;
+				}
+			}
+
+			// Only add rows if there is a matching shader.
+			if( NumShaders )
+			{
+				// Write a row for the shader type.
+				ShaderTypeViewer.AddColumn(Type->GetName());
+				ShaderTypeViewer.AddColumn(TEXT("%u"),NumShaders);
+				ShaderTypeViewer.AddColumn(TEXT("%.1f"),AverageNumInstructions);
+				ShaderTypeViewer.AddColumn(TEXT("%u"),Size);
+				ShaderTypeViewer.AddColumn(TEXT("%.1f"),Size / (float)NumShaders);
+				ShaderTypeViewer.AddColumn(TEXT("%d"), NumPipelines);
+				ShaderTypeViewer.AddColumn(TEXT("%d"), NumSharedPipelines);
+				ShaderTypeViewer.CycleRow();
+
+				TotalShaderCount += NumShaders;
+				TotalPipelineCount += NumPipelines;
+				TotalInstructionCount += NumInitializedInstructions;
+				TotalTypeCount++;
+				TotalSize += Size;
+				TotalSizePerType += Size / (float)NumShaders;
+			}
+		}
+	}
+
+	// go through non shared pipelines
+
+	// Write a total row.
+	ShaderTypeViewer.AddColumn(TEXT("Total"));
+	ShaderTypeViewer.AddColumn(TEXT("%u"),TotalShaderCount);
+	ShaderTypeViewer.AddColumn(TEXT("%u"),TotalInstructionCount);
+	ShaderTypeViewer.AddColumn(TEXT("%u"),TotalSize);
+	ShaderTypeViewer.AddColumn(TEXT("0"));
+	ShaderTypeViewer.AddColumn(TEXT("%u"), TotalPipelineCount);
+	ShaderTypeViewer.AddColumn(TEXT("-"));
+	ShaderTypeViewer.CycleRow();
+
+	// Write an average row.
+	ShaderTypeViewer.AddColumn(TEXT("Average"));
+	ShaderTypeViewer.AddColumn(TEXT("%.1f"),TotalShaderCount / (float)TotalTypeCount);
+	ShaderTypeViewer.AddColumn(TEXT("%.1f"),(float)TotalInstructionCount / TotalShaderCount);
+	ShaderTypeViewer.AddColumn(TEXT("%.1f"),TotalSize / (float)TotalShaderCount);
+	ShaderTypeViewer.AddColumn(TEXT("%.1f"),TotalSizePerType / TotalTypeCount);
+	ShaderTypeViewer.AddColumn(TEXT("-"));
+	ShaderTypeViewer.AddColumn(TEXT("-"));
+	ShaderTypeViewer.CycleRow();
+#endif
+}
+
+void DumpShaderPipelineStats(EShaderPlatform Platform)
+{
+#if ALLOW_DEBUG_FILES
+	FDiagnosticTableViewer ShaderTypeViewer(*FDiagnosticTableViewer::GetUniqueTemporaryFilePath(TEXT("ShaderPipelineStats")));
+
+	int32 TotalNumPipelines = 0;
+	int32 TotalSize = 0;
+	float TotalSizePerType = 0;
+
+	// Write a row of headings for the table's columns.
+	ShaderTypeViewer.AddColumn(TEXT("Type"));
+	ShaderTypeViewer.AddColumn(TEXT("Shared/Unique"));
+
+	// Exclude compute
+	for (int32 Index = 0; Index < SF_NumFrequencies - 1; ++Index)
+	{
+		ShaderTypeViewer.AddColumn(GetShaderFrequencyString((EShaderFrequency)Index));
+	}
+	ShaderTypeViewer.CycleRow();
+
+	int32 TotalTypeCount = 0;
+	for (TLinkedList<FShaderPipelineType*>::TIterator It(FShaderPipelineType::GetTypeList()); It; It.Next())
+	{
+		const FShaderPipelineType* Type = *It;
+
+		// Write a row for the shader type.
+		ShaderTypeViewer.AddColumn(Type->GetName());
+		ShaderTypeViewer.AddColumn(Type->ShouldOptimizeUnusedOutputs() ? TEXT("U") : TEXT("S"));
+
+		for (int32 Index = 0; Index < SF_NumFrequencies - 1; ++Index)
+		{
+			const FShaderType* ShaderType = Type->GetShader((EShaderFrequency)Index);
+			ShaderTypeViewer.AddColumn(ShaderType ? ShaderType->GetName() : TEXT(""));
+		}
+
+		ShaderTypeViewer.CycleRow();
+	}
+#endif
+}
+
+FShaderType* FindShaderTypeByName(FName ShaderTypeName)
+{
+	FShaderType** FoundShader = FShaderType::GetNameToTypeMap().Find(ShaderTypeName);
+	if (FoundShader)
+	{
+		return *FoundShader;
+	}
+
+	return nullptr;
+}
+
+
+void DispatchComputeShader(
+	FRHICommandList& RHICmdList,
+	FShader* Shader,
+	uint32 ThreadGroupCountX,
+	uint32 ThreadGroupCountY,
+	uint32 ThreadGroupCountZ)
+{
+	RHICmdList.DispatchComputeShader(ThreadGroupCountX, ThreadGroupCountY, ThreadGroupCountZ);
+}
+
+void DispatchComputeShader(
+	FRHIAsyncComputeCommandListImmediate& RHICmdList,
+	FShader* Shader,
+	uint32 ThreadGroupCountX,
+	uint32 ThreadGroupCountY,
+	uint32 ThreadGroupCountZ)
+{
+	RHICmdList.DispatchComputeShader(ThreadGroupCountX, ThreadGroupCountY, ThreadGroupCountZ);
+}
+
+void DispatchIndirectComputeShader(
+	FRHICommandList& RHICmdList,
+	FShader* Shader,
+	FVertexBufferRHIParamRef ArgumentBuffer,
+	uint32 ArgumentOffset)
+{
+	RHICmdList.DispatchIndirectComputeShader(ArgumentBuffer, ArgumentOffset);
+}
+
+
+void ShaderMapAppendKeyString(EShaderPlatform Platform, FString& KeyString)
+{
+	// Globals that should cause all shaders to recompile when changed must be appended to the key here
+	// Key should be kept as short as possible while being somewhat human readable for debugging
+
+	{
+		static const auto CVar = IConsoleManager::Get().FindTConsoleVariableDataInt(TEXT("Compat.UseDXT5NormalMaps"));
+		KeyString += (CVar && CVar->GetValueOnAnyThread() != 0) ? TEXT("_DXTN") : TEXT("_BC5N");
+	}
+
+	{
+		static const auto CVar = IConsoleManager::Get().FindTConsoleVariableDataInt(TEXT("r.ClearCoatNormal"));
+		KeyString += (CVar && CVar->GetValueOnAnyThread() != 0) ? TEXT("_CCBN") : TEXT("_NoCCBN");
+	}
+
+	{
+		static const auto CVar = IConsoleManager::Get().FindTConsoleVariableDataInt(TEXT("r.IrisNormal"));
+		KeyString += (CVar && CVar->GetValueOnAnyThread() != 0) ? TEXT("_Iris") : TEXT("_NoIris");
+	}
+
+	{
+		static const auto CVar = IConsoleManager::Get().FindTConsoleVariableDataInt(TEXT("r.CompileShadersForDevelopment"));
+		KeyString += (CVar && CVar->GetValueOnAnyThread() != 0) ? TEXT("_DEV") : TEXT("_NoDEV");
+	}
+
+	{
+		static const auto CVar = IConsoleManager::Get().FindTConsoleVariableDataInt(TEXT("r.AllowStaticLighting"));
+		const bool bValue = CVar ? CVar->GetValueOnAnyThread() != 0 : true;
+		KeyString += bValue ? TEXT("_SL") : TEXT("_NoSL");
+	}
+
+	{
+		static const auto CVar = IConsoleManager::Get().FindTConsoleVariableDataInt(TEXT("r.VirtualTexturedLightmaps"));
+		KeyString += (CVar && CVar->GetValueOnAnyThread() != 0) ? TEXT("_VTLM") : TEXT("_NoVTLM");
+	}
+
+	{
+		static const auto CVar = IConsoleManager::Get().FindTConsoleVariableDataInt(TEXT("r.BasePassOutputsVelocity"));
+		if (CVar && CVar->GetValueOnGameThread() != 0)
+		{
+			KeyString += TEXT("_GV");
+		}
+	}
+
+	{
+		static const auto CVarInstancedStereo = IConsoleManager::Get().FindTConsoleVariableDataInt(TEXT("vr.InstancedStereo"));
+		static const auto CVarMultiView = IConsoleManager::Get().FindTConsoleVariableDataInt(TEXT("vr.MultiView"));
+		static const auto CVarMobileMultiView = IConsoleManager::Get().FindTConsoleVariableDataInt(TEXT("vr.MobileMultiView"));
+		static const auto CVarMonoscopicFarField = IConsoleManager::Get().FindTConsoleVariableDataInt(TEXT("vr.MonoscopicFarField"));
+		static const auto CVarODSCapture = IConsoleManager::Get().FindTConsoleVariableDataInt(TEXT("vr.ODSCapture"));
+
+		const bool bIsInstancedStereo = (RHISupportsInstancedStereo(Platform) && (CVarInstancedStereo && CVarInstancedStereo->GetValueOnGameThread() != 0));
+		const bool bIsMultiView = (RHISupportsMultiView(Platform) && (CVarMultiView && CVarMultiView->GetValueOnGameThread() != 0));
+
+		const bool bIsAndroidGLES = RHISupportsMobileMultiView(Platform);
+		const bool bIsMobileMultiView = (bIsAndroidGLES && (CVarMobileMultiView && CVarMobileMultiView->GetValueOnGameThread() != 0));
+
+		const bool bIsMonoscopicFarField = CVarMonoscopicFarField && (CVarMonoscopicFarField->GetValueOnGameThread() != 0);
+
+		const bool bIsODSCapture = CVarODSCapture && (CVarODSCapture->GetValueOnGameThread() != 0);
+
+		if (bIsInstancedStereo)
+		{
+			KeyString += TEXT("_VRIS");
+			
+			if (bIsMultiView)
+			{
+				KeyString += TEXT("_MVIEW");
+			}
+		}
+
+		if (bIsMobileMultiView)
+		{
+			KeyString += TEXT("_MMVIEW");
+		}
+
+		if (bIsMonoscopicFarField)
+		{
+			KeyString += TEXT("_MONO");
+		}
+
+		if (bIsODSCapture)
+		{
+			KeyString += TEXT("_ODSC");
+		}
+	}
+
+	{
+		static const auto CVar = IConsoleManager::Get().FindTConsoleVariableDataInt(TEXT("r.SelectiveBasePassOutputs"));
+		if (CVar && CVar->GetValueOnGameThread() != 0)
+		{
+			KeyString += TEXT("_SO");
+		}
+	}
+
+	{
+		KeyString += UsePreExposure(Platform) ? TEXT("_PreExp") : TEXT("");
+	}
+
+	{
+		KeyString += IsUsingDBuffers(Platform) ? TEXT("_DBuf") : TEXT("_NoDBuf");
+	}
+
+	{
+		static IConsoleVariable* CVar = IConsoleManager::Get().FindConsoleVariable(TEXT("r.AllowGlobalClipPlane"));
+		KeyString += (CVar && CVar->GetInt() != 0) ? TEXT("_ClipP") : TEXT("");
+	}
+
+	{
+		static const auto CVar = IConsoleManager::Get().FindConsoleVariable(TEXT("r.Shaders.KeepDebugInfo"));
+		KeyString += (CVar && CVar->GetInt() != 0) ? TEXT("_NoStrip") : TEXT("");
+	}
+
+	{
+		static const auto CVar = IConsoleManager::Get().FindConsoleVariable(TEXT("r.Shaders.Optimize"));
+		KeyString += (CVar && CVar->GetInt() != 0) ? TEXT("") : TEXT("_NoOpt");
+	}
+	
+	{
+		// Always default to fast math unless specified
+		static const auto CVar = IConsoleManager::Get().FindConsoleVariable(TEXT("r.Shaders.FastMath"));
+		KeyString += (CVar && CVar->GetInt() == 0) ? TEXT("_NoFastMath") : TEXT("");
+	}
+	
+	{
+		static const auto CVar = IConsoleManager::Get().FindConsoleVariable(TEXT("r.Shaders.FlowControlMode"));
+		if (CVar)
+		{
+			switch(CVar->GetInt())
+			{
+				case 2:
+					KeyString += TEXT("_AvoidFlow");
+					break;
+				case 1:
+					KeyString += TEXT("_PreferFlow");
+					break;
+				case 0:
+				default:
+					break;
+			}
+		}
+	}
+
+	if (IsD3DPlatform(Platform, false))
+	{
+		static const auto CVar = IConsoleManager::Get().FindConsoleVariable(TEXT("r.D3D.RemoveUnusedInterpolators"));
+		if (CVar && CVar->GetInt() != 0)
+		{
+			KeyString += TEXT("_UnInt");
+		}
+	}
+
+	if (IsMobilePlatform(Platform))
+	{
+		{
+			static IConsoleVariable* CVar = IConsoleManager::Get().FindConsoleVariable(TEXT("r.Mobile.DisableVertexFog"));
+			KeyString += (CVar && CVar->GetInt() != 0) ? TEXT("_NoVFog") : TEXT("");
+		}
+
+		{
+			static const auto* CVar = IConsoleManager::Get().FindTConsoleVariableDataInt(TEXT("r.Shadow.CSM.MaxMobileCascades"));
+			KeyString += (CVar) ? FString::Printf(TEXT("MMC%d"), CVar->GetValueOnAnyThread()) : TEXT("");
+		}	
+
+		{
+			static IConsoleVariable* CVar = IConsoleManager::Get().FindConsoleVariable(TEXT("r.Mobile.UseLegacyShadingModel"));
+			KeyString += (CVar && CVar->GetInt() != 0) ? TEXT("_legshad") : TEXT("");
+		}
+		
+		{
+			static IConsoleVariable* CVar = IConsoleManager::Get().FindConsoleVariable(TEXT("r.Mobile.ForceFullPrecisionInPS"));
+			KeyString += (CVar && CVar->GetInt() != 0) ? TEXT("_highp") : TEXT("");
+		}
+
+		{
+			static IConsoleVariable* CVar = IConsoleManager::Get().FindConsoleVariable(TEXT("r.Mobile.AllowDitheredLODTransition"));
+			KeyString += (CVar && CVar->GetInt() != 0) ? TEXT("_DLODT") : TEXT("");
+		}
+		
+		if (IsOpenGLPlatform(Platform))
+		{
+			static IConsoleVariable* CVar = IConsoleManager::Get().FindConsoleVariable(TEXT("OpenGL.UseEmulatedUBs"));
+			KeyString += (CVar && CVar->GetInt() != 0) ? TEXT("_NoUB") : TEXT("");
+		}
+	}
+
+	if (Platform == SP_PS4)
+	{
+		{
+			static const auto CVar = IConsoleManager::Get().FindTConsoleVariableDataInt(TEXT("r.PS4MixedModeShaderDebugInfo"));
+			if (CVar && CVar->GetValueOnAnyThread() != 0)
+			{
+				KeyString += TEXT("_MMDBG");
+			}
+		}
+
+		{
+			static const auto CVar = IConsoleManager::Get().FindTConsoleVariableDataInt(TEXT("r.PS4ShaderSDBMode"));
+			switch (CVar ? CVar->GetValueOnAnyThread() : 0)
+			{
+			case 1: KeyString += TEXT("_SDB1"); break;
+			case 2: KeyString += TEXT("_SDB2"); break;
+			default: break;
+			}
+		}
+
+		{
+			static const auto CVar = IConsoleManager::Get().FindTConsoleVariableDataInt(TEXT("r.PS4UseTTrace"));
+			if (CVar && CVar->GetValueOnAnyThread() > 0)
+			{
+				KeyString += FString::Printf(TEXT("TT%d"), CVar->GetValueOnAnyThread());
+			}
+		}
+	}
+	
+	// Encode the Metal standard into the shader compile options so that they recompile if the settings change.
+	if (IsMetalPlatform(Platform))
+	{
+		{
+			static const auto CVar = IConsoleManager::Get().FindConsoleVariable(TEXT("r.Shaders.ZeroInitialise"));
+			KeyString += (CVar && CVar->GetInt() != 0) ? TEXT("_ZeroInit") : TEXT("");
+		}
+		{
+			static const auto CVar = IConsoleManager::Get().FindConsoleVariable(TEXT("r.Shaders.BoundsChecking"));
+			KeyString += (CVar && CVar->GetInt() != 0) ? TEXT("_BoundsChecking") : TEXT("");
+		}
+		{
+			KeyString += RHISupportsManualVertexFetch(Platform) ? TEXT("_MVF_") : TEXT("");
+		}
+		
+		uint32 ShaderVersion = RHIGetShaderLanguageVersion(Platform);
+		KeyString += FString::Printf(TEXT("_MTLSTD%u_"), ShaderVersion);
+		
+		bool bAllowFastIntrinsics = false;
+		bool bEnableMathOptimisations = true;
+		if (IsPCPlatform(Platform))
+		{
+			GConfig->GetBool(TEXT("/Script/MacTargetPlatform.MacTargetSettings"), TEXT("UseFastIntrinsics"), bAllowFastIntrinsics, GEngineIni);
+			GConfig->GetBool(TEXT("/Script/MacTargetPlatform.MacTargetSettings"), TEXT("EnableMathOptimisations"), bEnableMathOptimisations, GEngineIni);
+		}
+		else
+		{
+			GConfig->GetBool(TEXT("/Script/IOSRuntimeSettings.IOSRuntimeSettings"), TEXT("UseFastIntrinsics"), bAllowFastIntrinsics, GEngineIni);
+			GConfig->GetBool(TEXT("/Script/IOSRuntimeSettings.IOSRuntimeSettings"), TEXT("EnableMathOptimisations"), bEnableMathOptimisations, GEngineIni);
+		}
+		
+		if (bAllowFastIntrinsics)
+		{
+			KeyString += TEXT("_MTLSL_FastIntrin");
+		}
+		
+		// Same as console-variable above, but that's global and this is per-platform, per-project
+		if (!bEnableMathOptimisations)
+		{
+			KeyString += TEXT("_NoFastMath");
+		}
+		
+		// Shaders built for archiving - for Metal that requires compiling the code in a different way so that we can strip it later
+		bool bArchive = false;
+		GConfig->GetBool(TEXT("/Script/UnrealEd.ProjectPackagingSettings"), TEXT("bSharedMaterialNativeLibraries"), bArchive, GGameIni);
+		if (bArchive)
+		{
+			KeyString += TEXT("_ARCHIVE");
+		}
+	}
+
+	if (IsFeatureLevelSupported(Platform, ERHIFeatureLevel::SM4))
+	{
+		static const auto CVar = IConsoleManager::Get().FindTConsoleVariableDataInt(TEXT("r.StencilForLODDither"));
+		if (CVar && CVar->GetValueOnAnyThread() > 0)
+		{
+			KeyString += TEXT("_SD");
+		}
+	}
+
+	{
+		ITargetPlatform* TargetPlatform = GetTargetPlatformManager()->FindTargetPlatform(ShaderPlatformToPlatformName(Platform).ToString());
+		if (TargetPlatform && TargetPlatform->UsesForwardShading())
+		{
+			KeyString += TEXT("_FS");
+		}
+	}
+
+	{
+		static const auto CVar = IConsoleManager::Get().FindTConsoleVariableDataInt(TEXT("r.PostProcessing.PropagateAlpha"));
+		if (CVar && CVar->GetValueOnAnyThread() > 0)
+		{
+			if (CVar->GetValueOnAnyThread() == 2)
+			{
+				KeyString += TEXT("_SA2");
+			}
+			else
+			{
+				KeyString += TEXT("_SA");
+			}
+		}
+	}
+
+	{
+		static const auto CVar = IConsoleManager::Get().FindTConsoleVariableDataInt(TEXT("r.VertexFoggingForOpaque"));
+		if (CVar && CVar->GetValueOnAnyThread() > 0)
+		{
+			KeyString += TEXT("_VFO");
+		}
+	}
+
+	{
+		static const auto CVar = IConsoleManager::Get().FindTConsoleVariableDataInt(TEXT("r.EarlyZPassOnlyMaterialMasking"));
+		if (CVar && CVar->GetValueOnAnyThread() > 0)
+		{
+			KeyString += TEXT("_EZPMM");
+		}
+	}
+
+	{
+		static const auto CVar = IConsoleManager::Get().FindTConsoleVariableDataInt(TEXT("r.DistanceFieldBuild.EightBit"));
+		if (CVar && CVar->GetValueOnAnyThread() > 0)
+		{
+			KeyString += TEXT("_8u");
+		}
+	}
+	
+	{
+		static const auto CVar = IConsoleManager::Get().FindTConsoleVariableDataInt(TEXT("r.GPUSkin.Limit2BoneInfluences"));
+		if (CVar && CVar->GetValueOnAnyThread() != 0)
+		{
+			KeyString += TEXT("_2bi");
+		}
+	}
+}
