Merging //UE4/Dev-Rendering-Graph@4492585 to Dev-Rendering (//UE4/Dev-Rendering)

diff --git a/Engine/Source/Runtime/RenderCore/Private/ShaderCore.cpp b/Engine/Source/Runtime/RenderCore/Private/ShaderCore.cpp
new file mode 100644
--- /dev/null
+++ b/Engine/Source/Runtime/RenderCore/Private/ShaderCore.cpp
@@ -0,0 +1,1086 @@
+// Copyright 1998-2018 Epic Games, Inc. All Rights Reserved.
+
+/*=============================================================================
+	ShaderCore.h: Shader core module implementation.
+=============================================================================*/
+
+#include "ShaderCore.h"
+#include "HAL/FileManager.h"
+#include "Misc/FileHelper.h"
+#include "Misc/Paths.h"
+#include "Misc/ScopeLock.h"
+#include "Stats/StatsMisc.h"
+#include "Misc/ScopedSlowTask.h"
+#include "Shader.h"
+#include "VertexFactory.h"
+#include "Modules/ModuleManager.h"
+#include "Interfaces/IShaderFormat.h"
+#include "Interfaces/IShaderFormatModule.h"
+#include "RHIShaderFormatDefinitions.inl"
+#if WITH_EDITOR
+#include "Misc/CoreMisc.h"
+#include "Interfaces/ITargetPlatform.h"
+#include "Interfaces/ITargetPlatformManagerModule.h"
+#endif
+
+static TAutoConsoleVariable<int32> CVarShaderDevelopmentMode(
+	TEXT("r.ShaderDevelopmentMode"),
+	0,
+	TEXT("0: Default, 1: Enable various shader development utilities, such as the ability to retry on failed shader compile, and extra logging as shaders are compiled."),
+	ECVF_Default);
+
+void UpdateShaderDevelopmentMode()
+{
+	// Keep LogShaders verbosity in sync with r.ShaderDevelopmentMode
+	// r.ShaderDevelopmentMode==1 results in all LogShaders log messages being displayed
+	bool bLogShadersUnsuppressed = UE_LOG_ACTIVE(LogShaders, Log);
+	bool bDesiredLogShadersUnsuppressed = CVarShaderDevelopmentMode.GetValueOnGameThread() == 1;
+
+	if (bLogShadersUnsuppressed != bDesiredLogShadersUnsuppressed)
+	{
+		if (bDesiredLogShadersUnsuppressed)
+		{
+			UE_SET_LOG_VERBOSITY(LogShaders, Log);
+		}
+		else
+		{
+			UE_SET_LOG_VERBOSITY(LogShaders, Error);
+		}
+	}
+}
+
+//
+// Shader stats
+//
+
+DEFINE_STAT(STAT_ShaderCompiling_NiagaraShaders);
+DEFINE_STAT(STAT_ShaderCompiling_NumTotalNiagaraShaders);
+
+DEFINE_STAT(STAT_ShaderCompiling_MaterialShaders);
+DEFINE_STAT(STAT_ShaderCompiling_GlobalShaders);
+DEFINE_STAT(STAT_ShaderCompiling_RHI);
+DEFINE_STAT(STAT_ShaderCompiling_HashingShaderFiles);
+DEFINE_STAT(STAT_ShaderCompiling_LoadingShaderFiles);
+DEFINE_STAT(STAT_ShaderCompiling_HLSLTranslation);
+DEFINE_STAT(STAT_ShaderCompiling_DDCLoading);
+DEFINE_STAT(STAT_ShaderCompiling_MaterialLoading);
+DEFINE_STAT(STAT_ShaderCompiling_MaterialCompiling);
+
+DEFINE_STAT(STAT_ShaderCompiling_NumTotalMaterialShaders);
+DEFINE_STAT(STAT_ShaderCompiling_NumSpecialMaterialShaders);
+DEFINE_STAT(STAT_ShaderCompiling_NumParticleMaterialShaders);
+DEFINE_STAT(STAT_ShaderCompiling_NumSkinnedMaterialShaders);
+DEFINE_STAT(STAT_ShaderCompiling_NumLitMaterialShaders);
+DEFINE_STAT(STAT_ShaderCompiling_NumUnlitMaterialShaders);
+DEFINE_STAT(STAT_ShaderCompiling_NumTransparentMaterialShaders);
+DEFINE_STAT(STAT_ShaderCompiling_NumOpaqueMaterialShaders);
+DEFINE_STAT(STAT_ShaderCompiling_NumMaskedMaterialShaders);
+
+
+DEFINE_STAT(STAT_Shaders_NumShadersLoaded);
+DEFINE_STAT(STAT_Shaders_NumShaderResourcesLoaded);
+DEFINE_STAT(STAT_Shaders_NumShaderMaps);
+DEFINE_STAT(STAT_Shaders_RTShaderLoadTime);
+DEFINE_STAT(STAT_Shaders_NumShadersUsedForRendering);
+DEFINE_STAT(STAT_Shaders_TotalRTShaderInitForRenderingTime);
+DEFINE_STAT(STAT_Shaders_FrameRTShaderInitForRenderingTime);
+DEFINE_STAT(STAT_Shaders_ShaderMemory);
+DEFINE_STAT(STAT_Shaders_ShaderResourceMemory);
+DEFINE_STAT(STAT_Shaders_ShaderMapMemory);
+
+/** Protects GShaderFileCache from simultaneous access by multiple threads. */
+FCriticalSection FileCacheCriticalSection;
+
+/** The shader file cache, used to minimize shader file reads */
+TMap<FString, FString> GShaderFileCache;
+
+class FShaderHashCache
+{
+public:
+	FShaderHashCache()
+		: bInitialized(false)
+	{
+	}
+
+	void Initialize()
+	{
+		const FString EmptyDirectory("");
+		for (auto& Platform : Platforms)
+		{
+			Platform.IncludeDirectory = EmptyDirectory;
+			Platform.ShaderHashCache.Reset();
+		}
+
+
+		TArray<FName> Modules;
+		FModuleManager::Get().FindModules(SHADERFORMAT_MODULE_WILDCARD, Modules);
+
+		if (!Modules.Num())
+		{
+			UE_LOG(LogShaders, Error, TEXT("No target shader formats found!"));
+		}
+
+		TArray<FName> SupportedFormats;
+
+		for (int32 ModuleIndex = 0; ModuleIndex < Modules.Num(); ++ModuleIndex)
+		{
+			IShaderFormat* ShaderFormat = FModuleManager::LoadModuleChecked<IShaderFormatModule>(Modules[ModuleIndex]).GetShaderFormat();
+			if (ShaderFormat)
+			{
+				FString IncludeDirectory = ShaderFormat->GetPlatformIncludeDirectory();
+				if (!IncludeDirectory.IsEmpty())
+				{
+					IncludeDirectory = "/" + IncludeDirectory + "/";
+				}
+
+				SupportedFormats.Reset();
+				ShaderFormat->GetSupportedFormats(SupportedFormats);
+
+				for (int32 FormatIndex = 0; FormatIndex < SupportedFormats.Num(); ++FormatIndex)
+				{
+					const EShaderPlatform ShaderPlatform = ShaderFormatNameToShaderPlatform(SupportedFormats[FormatIndex]);
+					if (ShaderPlatform != SP_NumPlatforms)
+					{
+						Platforms[ShaderPlatform].IncludeDirectory = IncludeDirectory;
+					}
+				}
+			}
+		}
+
+		bInitialized = true;
+	}
+
+	FSHAHash* FindHash(EShaderPlatform ShaderPlatform, const FString& VirtualFilePath)
+	{
+		check(ShaderPlatform < ARRAY_COUNT(Platforms));
+		checkf(bInitialized, TEXT("GShaderHashCache::Initialize needs to be called before GShaderHashCache::FindHash."));
+
+		return Platforms[ShaderPlatform].ShaderHashCache.Find(VirtualFilePath);
+	}
+
+	FSHAHash& AddHash(EShaderPlatform ShaderPlatform, const FString& VirtualFilePath)
+	{
+		check(ShaderPlatform < ARRAY_COUNT(Platforms));
+		checkf(bInitialized, TEXT("GShaderHashCache::Initialize needs to be called before GShaderHashCache::AddHash."));
+
+		return Platforms[ShaderPlatform].ShaderHashCache.Add(VirtualFilePath, FSHAHash());
+	}
+
+	bool ShouldIgnoreInclude(const FString& VirtualFilePath, EShaderPlatform ShaderPlatform)
+	{
+		// Ignore only platform specific files, which won't be used by the target platform.
+		if (VirtualFilePath.StartsWith(TEXT("/Engine/Private/Platform/"))
+			|| VirtualFilePath.StartsWith(TEXT("/Engine/Public/Platform/")))
+		{
+			const FString& PlatformIncludeDirectory = GetPlatformIncludeDirectory(ShaderPlatform);
+
+			if (PlatformIncludeDirectory.IsEmpty() || !VirtualFilePath.Contains(PlatformIncludeDirectory))
+			{
+				return true;
+			}
+		}
+
+		return false;
+	}
+
+	void Empty()
+	{
+		for (auto& Platform : Platforms)
+		{
+			Platform.ShaderHashCache.Reset();
+		}
+	}
+
+
+private:
+
+	struct FPlatform
+	{
+		/** Folder with platform specific shader files. */
+		FString IncludeDirectory;
+
+		/** The shader file hash cache, used to minimize loading and hashing shader files; it includes also hashes for multiple filenames
+		by making the key the concatenated list of filenames.
+		*/
+		TMap<FString, FSHAHash> ShaderHashCache;
+	};
+
+	FPlatform Platforms[EShaderPlatform::SP_NumPlatforms];
+	bool bInitialized;
+
+
+	const FString& GetPlatformIncludeDirectory(EShaderPlatform ShaderPlatform)
+	{
+		check(ShaderPlatform < EShaderPlatform::SP_NumPlatforms);
+		checkf(bInitialized, TEXT("GShaderHashCache::Initialize needs to be called before GShaderHashCache::GetPlatformIncludeDirectory."));
+
+		return Platforms[ShaderPlatform].IncludeDirectory;
+	}
+};
+
+FShaderHashCache GShaderHashCache;
+
+/** Global map of virtual file path to physical file paths */
+static TMap<FString, FString> GShaderSourceDirectoryMappings;
+
+static TAutoConsoleVariable<int32> CVarForceDebugViewModes(
+	TEXT("r.ForceDebugViewModes"),
+	0,
+	TEXT("0: Setting has no effect.\n")
+	TEXT("1: Forces debug view modes to be available, even on cooked builds.")
+	TEXT("2: Forces debug view modes to be unavailable, even on editor builds.  Removes many shader permutations for faster shader iteration."),
+	ECVF_RenderThreadSafe | ECVF_ReadOnly);
+
+/** Returns true if debug viewmodes are allowed for the current platform. */
+bool AllowDebugViewmodes()
+{	
+	int32 ForceDebugViewValue = CVarForceDebugViewModes.GetValueOnAnyThread();
+
+	// To use debug viewmodes on consoles, r.ForceDebugViewModes must be set to 1 in ConsoleVariables.ini
+	// And EngineDebugMaterials must be in the StartupPackages for the target platform.
+	bool bForceEnable = ForceDebugViewValue == 1;
+	bool bForceDisable = ForceDebugViewValue == 2;
+
+	//bool bEnabled = (!IsRunningCommandlet() && !FPlatformProperties::RequiresCookedData());	
+	//bEnabled |= bForceEnable;
+	//bEnabled &= !bForceDisable;
+
+	return (!bForceDisable) && (bForceEnable || (!IsRunningCommandlet() && !FPlatformProperties::RequiresCookedData()));
+}
+
+/** Returns true if debug viewmodes are allowed for the current platform. */
+bool AllowDebugViewmodes(EShaderPlatform Platform)
+{
+#if WITH_EDITOR
+	const int32 ForceDebugViewValue = CVarForceDebugViewModes.GetValueOnAnyThread();
+	bool bForceEnable = ForceDebugViewValue == 1;
+	bool bForceDisable = ForceDebugViewValue == 2;
+	ITargetPlatform* TargetPlatform = GetTargetPlatformManager()->FindTargetPlatform(ShaderPlatformToPlatformName(Platform).ToString());
+	return (!bForceDisable) && (bForceEnable || !TargetPlatform || (!IsRunningCommandlet() && !TargetPlatform->RequiresCookedData()));
+#else
+	return AllowDebugViewmodes();
+#endif
+}
+
+bool FShaderParameterMap::FindParameterAllocation(const TCHAR* ParameterName,uint16& OutBufferIndex,uint16& OutBaseIndex,uint16& OutSize) const
+{
+	const FParameterAllocation* Allocation = ParameterMap.Find(ParameterName);
+	if(Allocation)
+	{
+		OutBufferIndex = Allocation->BufferIndex;
+		OutBaseIndex = Allocation->BaseIndex;
+		OutSize = Allocation->Size;
+
+		if (Allocation->bBound)
+		{
+			// Can detect copy-paste errors in binding parameters.  Need to fix all the false positives before enabling.
+			//UE_LOG(LogShaders, Warning, TEXT("Parameter %s was bound multiple times. Code error?"), ParameterName);
+		}
+
+		Allocation->bBound = true;
+		return true;
+	}
+	else
+	{
+		return false;
+	}
+}
+
+bool FShaderParameterMap::ContainsParameterAllocation(const TCHAR* ParameterName) const
+{
+	return ParameterMap.Find(ParameterName) != NULL;
+}
+
+void FShaderParameterMap::AddParameterAllocation(const TCHAR* ParameterName,uint16 BufferIndex,uint16 BaseIndex,uint16 Size)
+{
+	FParameterAllocation Allocation;
+	Allocation.BufferIndex = BufferIndex;
+	Allocation.BaseIndex = BaseIndex;
+	Allocation.Size = Size;
+	ParameterMap.Add(ParameterName,Allocation);
+}
+
+void FShaderParameterMap::RemoveParameterAllocation(const TCHAR* ParameterName)
+{
+	ParameterMap.Remove(ParameterName);
+}
+
+void FShaderCompilerOutput::GenerateOutputHash()
+{
+	FSHA1 HashState;
+	
+	const TArray<uint8>& Code = ShaderCode.GetReadAccess();
+
+	// we don't hash the optional attachments as they would prevent sharing (e.g. many material share the save VS)
+	uint32 ShaderCodeSize = ShaderCode.GetShaderCodeSize();
+
+	HashState.Update(Code.GetData(), ShaderCodeSize * Code.GetTypeSize());
+	ParameterMap.UpdateHash(HashState);
+	HashState.Final();
+	HashState.GetHash(&OutputHash.Hash[0]);
+}
+
+static void ReportVirtualShaderFilePathError(TArray<FShaderCompilerError>* CompileErrors, FString ErrorString)
+{
+	if (CompileErrors)
+	{
+		CompileErrors->Add(FShaderCompilerError(*ErrorString));
+	}
+
+	UE_LOG(LogShaders, Error, TEXT("%s"), *ErrorString);
+}
+
+bool CheckVirtualShaderFilePath(const FString& VirtualFilePath, TArray<FShaderCompilerError>* CompileErrors /*= nullptr*/)
+{
+	bool bSuccess = true;
+
+	if (!VirtualFilePath.StartsWith(TEXT("/")))
+	{
+		FString Error = FString::Printf(TEXT("Virtual shader source file name \"%s\" should be absolute from the virtual root directory \"/\"."), *VirtualFilePath);
+		ReportVirtualShaderFilePathError(CompileErrors, Error);
+		bSuccess = false;
+	}
+
+	if (VirtualFilePath.Contains(TEXT("..")))
+	{
+		FString Error = FString::Printf(TEXT("Virtual shader source file name \"%s\" should have relative directories (\"../\") collapsed."), *VirtualFilePath);
+		ReportVirtualShaderFilePathError(CompileErrors, Error);
+		bSuccess = false;
+	}
+
+	if (VirtualFilePath.Contains(TEXT("\\")))
+	{
+		FString Error = FString::Printf(TEXT("Backslashes are not permitted in virtual shader source file name \"%s\""), *VirtualFilePath);
+		ReportVirtualShaderFilePathError(CompileErrors, Error);
+		bSuccess = false;
+	}
+
+	FString Extension = FPaths::GetExtension(VirtualFilePath);
+	if ((Extension != TEXT("usf") && Extension != TEXT("ush")) || VirtualFilePath.EndsWith(TEXT(".usf.usf")))
+	{
+		FString Error = FString::Printf(TEXT("Extension on virtual shader source file name \"%s\" is wrong. Only .usf or .ush allowed."), *VirtualFilePath);
+		ReportVirtualShaderFilePathError(CompileErrors, Error);
+		bSuccess = false;
+	}
+
+	return bSuccess;
+}
+
+/**
+* Add a new entry to the list of shader source files
+* Only unique entries which can be loaded are added as well as their #include files
+*
+* @param OutVirtualFilePaths - [out] list of shader source files to add to
+* @param ShaderFilename - shader file to add
+*/
+static void AddShaderSourceFileEntry(TArray<FString>& OutVirtualFilePaths, FString VirtualFilePath, EShaderPlatform ShaderPlatform)
+{
+	check(CheckVirtualShaderFilePath(VirtualFilePath));
+	if (!OutVirtualFilePaths.Contains(VirtualFilePath))
+	{
+		OutVirtualFilePaths.Add(VirtualFilePath);
+
+		TArray<FString> ShaderIncludes;
+		GetShaderIncludes(*VirtualFilePath, *VirtualFilePath, OutVirtualFilePaths, ShaderPlatform);
+		for( int32 IncludeIdx=0; IncludeIdx < ShaderIncludes.Num(); IncludeIdx++ )
+		{
+			OutVirtualFilePaths.AddUnique(ShaderIncludes[IncludeIdx]);
+		}
+	}
+}
+
+/**
+* Generates a list of virtual paths of all shader source that engine needs to load.
+*
+* @param OutVirtualFilePaths - [out] list of shader source files to add to
+*/
+static void GetAllVirtualShaderSourcePaths(TArray<FString>& OutVirtualFilePaths, EShaderPlatform ShaderPlatform)
+{
+	// add all shader source files for hashing
+	for( TLinkedList<FVertexFactoryType*>::TIterator FactoryIt(FVertexFactoryType::GetTypeList()); FactoryIt; FactoryIt.Next() )
+	{
+		FVertexFactoryType* VertexFactoryType = *FactoryIt;
+		if( VertexFactoryType )
+		{
+			FString ShaderFilename(VertexFactoryType->GetShaderFilename());
+			AddShaderSourceFileEntry(OutVirtualFilePaths, ShaderFilename, ShaderPlatform);
+		}
+	}
+	for( TLinkedList<FShaderType*>::TIterator ShaderIt(FShaderType::GetTypeList()); ShaderIt; ShaderIt.Next() )
+	{
+		FShaderType* ShaderType = *ShaderIt;
+		if(ShaderType)
+		{
+			FString ShaderFilename(ShaderType->GetShaderFilename());
+			AddShaderSourceFileEntry(OutVirtualFilePaths, ShaderFilename, ShaderPlatform);
+		}
+	}
+
+	//#todo-rco: No need to loop through Shader Pipeline Types (yet)
+
+	// also always add the MaterialTemplate.usf shader file
+	AddShaderSourceFileEntry(OutVirtualFilePaths, FString(TEXT("/Engine/Private/MaterialTemplate.ush")), ShaderPlatform);
+	AddShaderSourceFileEntry(OutVirtualFilePaths, FString(TEXT("/Engine/Private/Common.ush")), ShaderPlatform);
+	AddShaderSourceFileEntry(OutVirtualFilePaths, FString(TEXT("/Engine/Private/Definitions.usf")), ShaderPlatform);
+}
+
+/**
+* Kick off SHA verification for all shader source files
+*/
+void VerifyShaderSourceFiles(EShaderPlatform ShaderPlatform)
+{
+	if (!FPlatformProperties::RequiresCookedData())
+	{
+		// get the list of shader files that can be used
+		TArray<FString> VirtualShaderSourcePaths;
+		GetAllVirtualShaderSourcePaths(VirtualShaderSourcePaths, ShaderPlatform);
+		FScopedSlowTask SlowTask(VirtualShaderSourcePaths.Num());
+		for( int32 ShaderFileIdx=0; ShaderFileIdx < VirtualShaderSourcePaths.Num(); ShaderFileIdx++ )
+		{
+			SlowTask.EnterProgressFrame(1);
+			FString FileContents;
+			// load each shader source file. This will cache the shader source data after it has been verified
+			LoadShaderSourceFile(*VirtualShaderSourcePaths[ShaderFileIdx], FileContents, nullptr);
+		}
+	}
+}
+
+static void LogShaderSourceDirectoryMappings()
+{
+	for (const auto& Iter : GShaderSourceDirectoryMappings)
+	{
+		UE_LOG(LogShaders, Log, TEXT("Shader directory mapping %s -> %s"), *Iter.Key, *Iter.Value);
+	}
+}
+
+static FString GetShaderSourceFilePath(const FString& VirtualFilePath, TArray<FShaderCompilerError>* CompileErrors)
+{
+	// Make sure the .usf extension is correctly set.
+	if (!CheckVirtualShaderFilePath(VirtualFilePath, CompileErrors))
+	{
+		return FString();
+	}
+
+	// We don't cache the output of this function because only used in LoadShaderSourceFile that is cached, or when there
+	// is shader compilation errors.
+	
+	FString RealFilePath;
+
+	// Look if this virtual shader source file match any directory mapping.
+	const TMap<FString, FString>& ShaderSourceDirectoryMappings = GShaderSourceDirectoryMappings;
+	FString ParentVirtualDirectoryPath = FPaths::GetPath(VirtualFilePath);
+	FString RelativeVirtualDirectoryPath = FPaths::GetCleanFilename(VirtualFilePath);
+	while (!ParentVirtualDirectoryPath.IsEmpty())
+	{
+		if (ShaderSourceDirectoryMappings.Contains(ParentVirtualDirectoryPath))
+		{
+			RealFilePath = FPaths::Combine(
+				*ShaderSourceDirectoryMappings.Find(ParentVirtualDirectoryPath),
+				RelativeVirtualDirectoryPath);
+			break;
+		}
+
+		RelativeVirtualDirectoryPath = FPaths::GetCleanFilename(ParentVirtualDirectoryPath) / RelativeVirtualDirectoryPath;
+		ParentVirtualDirectoryPath = FPaths::GetPath(ParentVirtualDirectoryPath);
+	}
+
+	// Make sure a directory mapping has matched.
+	if (RealFilePath.IsEmpty())
+	{
+		FString Error = FString::Printf(TEXT("Can't map virtual shader source path \"%s\"."), *VirtualFilePath);
+		Error += TEXT("\nDirectory mappings are:");
+		for (const auto& Iter : ShaderSourceDirectoryMappings)
+		{
+			Error += FString::Printf(TEXT("\n  %s -> %s"), *Iter.Key, *Iter.Value);
+		}
+
+		ReportVirtualShaderFilePathError(CompileErrors, Error);
+	}
+	
+	return RealFilePath;
+}
+
+FString ParseVirtualShaderFilename(const FString& InFilename)
+{
+	FString ShaderDir = FString(FPlatformProcess::ShaderDir());
+	ShaderDir.ReplaceInline(TEXT("\\"), TEXT("/"));
+	int32 CharIndex = ShaderDir.Find(TEXT("/"), ESearchCase::CaseSensitive, ESearchDir::FromEnd, ShaderDir.Len() - 1);
+	if (CharIndex != INDEX_NONE)
+	{
+		ShaderDir = ShaderDir.Right(ShaderDir.Len() - CharIndex);
+	}
+
+	FString RelativeFilename = InFilename.Replace(TEXT("\\"), TEXT("/"));
+	// remove leading "/" because this makes path absolute on Linux (and Mac).
+	if (RelativeFilename.Len() > 0 && RelativeFilename[0] == TEXT('/'))
+	{
+		RelativeFilename = RelativeFilename.Right(RelativeFilename.Len() - 1);
+	}
+	RelativeFilename = IFileManager::Get().ConvertToRelativePath(*RelativeFilename);
+	CharIndex = RelativeFilename.Find(ShaderDir);
+	if (CharIndex != INDEX_NONE)
+	{
+		CharIndex += ShaderDir.Len();
+		if (RelativeFilename.GetCharArray()[CharIndex] == TEXT('/'))
+		{
+			CharIndex++;
+		}
+		if (RelativeFilename.Contains(TEXT("WorkingDirectory")))
+		{
+			const int32 NumDirsToSkip = 3;
+			int32 NumDirsSkipped = 0;
+			int32 NewCharIndex = CharIndex;
+
+			do
+			{
+				NewCharIndex = RelativeFilename.Find(TEXT("/"), ESearchCase::CaseSensitive, ESearchDir::FromStart, CharIndex);
+				CharIndex = (NewCharIndex == INDEX_NONE) ? CharIndex : NewCharIndex + 1;
+			}
+			while (NewCharIndex != INDEX_NONE && ++NumDirsSkipped < NumDirsToSkip);
+		}
+		RelativeFilename = RelativeFilename.Mid(CharIndex, RelativeFilename.Len() - CharIndex);
+	}
+
+	// add leading "/" to the relative filename because that's what virtual shader path expects
+	FString OutputFilename;
+	if (RelativeFilename.Len() > 0 && RelativeFilename[0] != TEXT('/'))
+	{
+		OutputFilename = TEXT("/") + RelativeFilename;
+	}
+	else
+	{
+		OutputFilename = RelativeFilename;
+	}
+	check(CheckVirtualShaderFilePath(OutputFilename));
+	return OutputFilename;
+}
+
+bool LoadShaderSourceFile(const TCHAR* VirtualFilePath, FString& OutFileContents, TArray<FShaderCompilerError>* OutCompileErrors) // TODO: const FString&
+{
+	// it's not expected that cooked platforms get here, but if they do, this is the final out
+	if (FPlatformProperties::RequiresCookedData())
+	{
+		return false;
+	}
+
+	bool bResult = false;
+
+	STAT(double ShaderFileLoadingTime = 0);
+	{
+		SCOPE_SECONDS_COUNTER(ShaderFileLoadingTime);
+
+		// Protect GShaderFileCache from simultaneous access by multiple threads
+		FScopeLock ScopeLock(&FileCacheCriticalSection);
+
+		FString* CachedFile = GShaderFileCache.Find(VirtualFilePath);
+
+		//if this file has already been loaded and cached, use that
+		if (CachedFile)
+		{
+			OutFileContents = *CachedFile;
+			bResult = true;
+		}
+		else
+		{
+			FString ShaderFilePath = GetShaderSourceFilePath(VirtualFilePath, OutCompileErrors);
+
+			// verify SHA hash of shader files on load. missing entries trigger an error
+			if (!ShaderFilePath.IsEmpty() && FFileHelper::LoadFileToString(OutFileContents, *ShaderFilePath, FFileHelper::EHashOptions::EnableVerify|FFileHelper::EHashOptions::ErrorMissingHash) )
+			{
+				//update the shader file cache
+				GShaderFileCache.Add(VirtualFilePath, *OutFileContents);
+				bResult = true;
+			}
+		}
+	}
+	INC_FLOAT_STAT_BY(STAT_ShaderCompiling_LoadingShaderFiles,(float)ShaderFileLoadingTime);
+
+	return bResult;
+}
+
+void LoadShaderSourceFileChecked(const TCHAR* VirtualFilePath, FString& OutFileContents)
+{
+	if (!LoadShaderSourceFile(VirtualFilePath, OutFileContents, nullptr))
+	{
+		UE_LOG(LogShaders, Fatal, TEXT("Couldn't find source file of virtual shader path \'%s\'"), VirtualFilePath);
+	}
+}
+
+/**
+ * Walks InStr until we find either an end-of-line or TargetChar.
+ */
+const TCHAR* SkipToCharOnCurrentLine(const TCHAR* InStr, TCHAR TargetChar)
+{
+	const TCHAR* Str = InStr;
+	if (Str)
+	{
+		while (*Str && *Str != TargetChar && *Str != TEXT('\n'))
+		{
+			++Str;
+		}
+		if (*Str != TargetChar)
+		{
+			Str = NULL;
+		}
+	}
+	return Str;
+}
+
+/**
+ * Recursively populates IncludeFilenames with the unique include filenames found in the shader file named Filename.
+ */
+static void GetShaderIncludes(const TCHAR* EntryPointVirtualFilePath, const TCHAR* VirtualFilePath, TArray<FString>& IncludeVirtualFilePaths, EShaderPlatform ShaderPlatform, uint32 DepthLimit, bool AddToIncludeFile)
+{
+	FString FileContents;
+	LoadShaderSourceFile(VirtualFilePath, FileContents, nullptr);
+
+	//avoid an infinite loop with a 0 length string
+	if (FileContents.Len() > 0)
+	{
+		if (AddToIncludeFile)
+		{
+			IncludeVirtualFilePaths.Add(VirtualFilePath);
+		}
+
+		//find the first include directive
+		const TCHAR* IncludeBegin = FCString::Strstr(*FileContents, TEXT("#include "));
+
+		uint32 SearchCount = 0;
+		const uint32 MaxSearchCount = 200;
+		//keep searching for includes as long as we are finding new ones and haven't exceeded the fixed limit
+		while (IncludeBegin != NULL && SearchCount < MaxSearchCount && DepthLimit > 0)
+		{
+			//find the first double quotation after the include directive
+			const TCHAR* IncludeFilenameBegin = SkipToCharOnCurrentLine(IncludeBegin, TEXT('\"'));
+
+			if (IncludeFilenameBegin)
+			{
+				//find the trailing double quotation
+				const TCHAR* IncludeFilenameEnd = SkipToCharOnCurrentLine(IncludeFilenameBegin + 1, TEXT('\"'));
+
+				if (IncludeFilenameEnd)
+				{
+					//construct a string between the double quotations
+					FString ExtractedIncludeFilename(FString((int32)(IncludeFilenameEnd - IncludeFilenameBegin - 1), IncludeFilenameBegin + 1));
+
+					// If the include is relative, then it must be relative to the current virtual file path.
+					if (!ExtractedIncludeFilename.StartsWith(TEXT("/")))
+					{
+						ExtractedIncludeFilename = FPaths::GetPath(VirtualFilePath) / ExtractedIncludeFilename;
+
+						// Collapse any relative directories to allow #include "../MyFile.ush"
+						FPaths::CollapseRelativeDirectories(ExtractedIncludeFilename);
+					}
+
+					//CRC the template, not the filled out version so that this shader's CRC will be independent of which material references it.
+					if (ExtractedIncludeFilename == TEXT("/Engine/Generated/Material.ush"))
+					{
+						ExtractedIncludeFilename = TEXT("/Engine/Private/MaterialTemplate.ush");
+					}
+
+					// Ignore uniform buffer, vertex factory and instanced stereo includes
+					bool bIgnoreInclude = ExtractedIncludeFilename.StartsWith(TEXT("/Engine/Generated/"));
+
+					// Check virtual.
+					bIgnoreInclude |= !CheckVirtualShaderFilePath(ExtractedIncludeFilename);
+
+					// Include only platform specific files, which will be used by the target platform.
+					bIgnoreInclude = bIgnoreInclude || GShaderHashCache.ShouldIgnoreInclude(ExtractedIncludeFilename, ShaderPlatform);
+
+					//vertex factories need to be handled separately
+					if (!bIgnoreInclude)
+					{
+						if (!IncludeVirtualFilePaths.Contains(ExtractedIncludeFilename))
+						{
+							GetShaderIncludes(EntryPointVirtualFilePath, *ExtractedIncludeFilename, IncludeVirtualFilePaths, ShaderPlatform, DepthLimit - 1, true);
+						}
+					}
+				}
+			}
+
+			// Skip to the end of the line.
+			IncludeBegin = SkipToCharOnCurrentLine(IncludeBegin, TEXT('\n'));
+		
+			//find the next include directive
+			if (IncludeBegin && *IncludeBegin != 0)
+			{
+				IncludeBegin = FCString::Strstr(IncludeBegin + 1, TEXT("#include "));
+			}
+			SearchCount++;
+		}
+
+		if (SearchCount == MaxSearchCount || DepthLimit == 0)
+		{
+			UE_LOG(LogShaders, Warning, TEXT("GetShaderIncludes parsing terminated early to avoid infinite looping!\n Entrypoint \'%s\' CurrentInclude \'%s\' SearchCount %u Depth %u"), 
+				EntryPointVirtualFilePath, 
+				VirtualFilePath,
+				SearchCount, 
+				DepthLimit);
+		}
+	}
+}
+
+void GetShaderIncludes(const TCHAR* EntryPointVirtualFilePath, const TCHAR* VirtualFilePath, TArray<FString>& IncludeVirtualFilePaths, EShaderPlatform ShaderPlatform, uint32 DepthLimit)
+{
+	GetShaderIncludes(EntryPointVirtualFilePath, VirtualFilePath, IncludeVirtualFilePaths, ShaderPlatform, DepthLimit, false);
+}
+
+static void UpdateSingleShaderFilehash(FSHA1& InOutHashState, const TCHAR* VirtualFilePath, EShaderPlatform ShaderPlatform)
+{
+	// Get the list of includes this file contains
+	TArray<FString> IncludeVirtualFilePaths;
+	GetShaderIncludes(VirtualFilePath, VirtualFilePath, IncludeVirtualFilePaths, ShaderPlatform);
+#if WITH_EDITOR &&  !(UE_BUILD_SHIPPING || UE_BUILD_TEST)
+	if ( UE_LOG_ACTIVE(LogTemp, Verbose) )
+	{
+		UE_LOG(LogTemp, Verbose, TEXT("Generating hash of file %s, "), VirtualFilePath);
+	}
+#endif
+	for (int32 IncludeIndex = 0; IncludeIndex < IncludeVirtualFilePaths.Num(); IncludeIndex++)
+	{
+		// Load the include file and hash it
+		FString IncludeFileContents;
+		LoadShaderSourceFileChecked(*IncludeVirtualFilePaths[IncludeIndex], IncludeFileContents);
+		InOutHashState.UpdateWithString(*IncludeFileContents, IncludeFileContents.Len());
+#if WITH_EDITOR &&  !(UE_BUILD_SHIPPING || UE_BUILD_TEST)
+		if (UE_LOG_ACTIVE(LogTemp, Verbose))
+		{
+			uint8 HashBytes[20];
+			FSHA1::HashBuffer(&InOutHashState, sizeof(FSHA1), HashBytes);
+			
+			UE_LOG(LogTemp, Verbose, TEXT("Processing include file for %s, %s, %s"), VirtualFilePath, *IncludeVirtualFilePaths[IncludeIndex], *BytesToHex(HashBytes, 20));
+		}
+#endif
+	}
+
+	// Load the source file and hash it
+	FString FileContents;
+	LoadShaderSourceFileChecked(VirtualFilePath, FileContents);
+	InOutHashState.UpdateWithString(*FileContents, FileContents.Len());
+}
+
+/**
+ * Calculates a Hash for the given filename and its includes if it does not already exist in the Hash cache.
+ * @param Filename - shader file to Hash
+ * @param ShaderPlatform - shader platform to Hash
+ */
+const FSHAHash& GetShaderFileHash(const TCHAR* VirtualFilePath, EShaderPlatform ShaderPlatform)
+{
+	// Make sure we are only accessing GShaderHashCache from one thread
+	//check(IsInGameThread() || IsAsyncLoading());
+	STAT(double HashTime = 0);
+	{
+		SCOPE_SECONDS_COUNTER(HashTime);
+
+		FSHAHash* CachedHash = GShaderHashCache.FindHash(ShaderPlatform, VirtualFilePath);
+
+		// If a hash for this filename has been cached, use that
+		if (CachedHash)
+		{
+			return *CachedHash;
+		}
+
+		FSHA1 HashState;
+		UpdateSingleShaderFilehash(HashState, VirtualFilePath, ShaderPlatform);
+		HashState.Final();
+
+		// Update the hash cache
+		FSHAHash& NewHash = GShaderHashCache.AddHash(ShaderPlatform, VirtualFilePath);
+		HashState.GetHash(&NewHash.Hash[0]);
+
+#if WITH_EDITOR &&  !(UE_BUILD_SHIPPING || UE_BUILD_TEST)
+		UE_LOG(LogTemp, Verbose, TEXT("Final hash for file %s, %s"), VirtualFilePath,*BytesToHex(&NewHash.Hash[0], 20));
+#endif
+		INC_FLOAT_STAT_BY(STAT_ShaderCompiling_HashingShaderFiles, (float)HashTime);
+		return NewHash;
+	}
+}
+
+/**
+ * Calculates a Hash for the given filename and its includes if it does not already exist in the Hash cache.
+ * @param Filename - shader file to Hash
+ * @param ShaderPlatform - shader platform to Hash
+ */
+const FSHAHash& GetShaderFilesHash(const TArray<FString>& VirtualFilePaths, EShaderPlatform ShaderPlatform)
+{
+	// Make sure we are only accessing GShaderHashCache from one thread
+	//check(IsInGameThread() || IsAsyncLoading());
+	STAT(double HashTime = 0);
+	{
+		SCOPE_SECONDS_COUNTER(HashTime);
+
+		FString Key;
+		for (const FString& Filename : VirtualFilePaths)
+		{
+			Key += Filename;
+		}
+
+		FSHAHash* CachedHash = GShaderHashCache.FindHash(ShaderPlatform, Key);
+
+		// If a hash for this filename has been cached, use that
+		if (CachedHash)
+		{
+			return *CachedHash;
+		}
+
+		FSHA1 HashState;
+		for (const FString& VirtualFilePath : VirtualFilePaths)
+		{
+			UpdateSingleShaderFilehash(HashState, *VirtualFilePath, ShaderPlatform);
+		}
+		HashState.Final();
+
+		// Update the hash cache
+		FSHAHash& NewHash = GShaderHashCache.AddHash(ShaderPlatform, Key);
+		HashState.GetHash(&NewHash.Hash[0]);
+
+		INC_FLOAT_STAT_BY(STAT_ShaderCompiling_HashingShaderFiles, (float)HashTime);
+		return NewHash;
+	}
+}
+
+void BuildShaderFileToUniformBufferMap(TMap<FString, TArray<const TCHAR*> >& ShaderFileToUniformBufferVariables)
+{
+	if (!FPlatformProperties::RequiresCookedData())
+	{
+		TArray<FString> ShaderSourceFiles;
+		GetAllVirtualShaderSourcePaths(ShaderSourceFiles, GMaxRHIShaderPlatform);
+
+		FScopedSlowTask SlowTask(ShaderSourceFiles.Num());
+
+		// Cache UB access strings, make it case sensitive for faster search
+		struct FShaderVariable
+		{
+			FShaderVariable(const TCHAR* ShaderVariable) :
+				OriginalShaderVariable(ShaderVariable), 
+				SearchKey(FString(ShaderVariable).ToUpper() + TEXT(".")),
+				// MCPP inserts a space after a #define replacement, make sure we detect the uniform buffer reference
+				SearchKeyWithSpace(FString(ShaderVariable).ToUpper() + TEXT(" ."))
+			{}
+
+			const TCHAR* OriginalShaderVariable;
+			FString SearchKey;
+			FString SearchKeyWithSpace;
+		};
+		// Cache each UB
+		TArray<FShaderVariable> SearchKeys;
+		for (TLinkedList<FShaderParametersMetadata*>::TIterator StructIt(FShaderParametersMetadata::GetStructList()); StructIt; StructIt.Next())
+		{
+			SearchKeys.Add(FShaderVariable(StructIt->GetShaderVariableName()));
+		}
+
+		// Find for each shader file which UBs it needs
+		for (int32 FileIndex = 0; FileIndex < ShaderSourceFiles.Num(); FileIndex++)
+		{
+			SlowTask.EnterProgressFrame(1);
+
+ 			FString ShaderFileContents;
+			LoadShaderSourceFileChecked(*ShaderSourceFiles[FileIndex], ShaderFileContents);
+
+			// To allow case sensitive search which is way faster on some platforms (no need to look up locale, etc)
+			ShaderFileContents.ToUpperInline();
+
+			TArray<const TCHAR*>& ReferencedUniformBuffers = ShaderFileToUniformBufferVariables.FindOrAdd(ShaderSourceFiles[FileIndex]);
+
+			for (int32 SearchKeyIndex = 0; SearchKeyIndex < SearchKeys.Num(); ++SearchKeyIndex)
+			{
+				// Searching for the uniform buffer shader variable being accessed with '.'
+				if (ShaderFileContents.Contains(SearchKeys[SearchKeyIndex].SearchKey, ESearchCase::CaseSensitive)
+					|| ShaderFileContents.Contains(SearchKeys[SearchKeyIndex].SearchKeyWithSpace, ESearchCase::CaseSensitive))
+				{
+					ReferencedUniformBuffers.AddUnique(SearchKeys[SearchKeyIndex].OriginalShaderVariable);
+				}
+			}
+		}
+	}
+}
+
+void InitializeShaderHashCache()
+{
+	GShaderHashCache.Initialize();
+}
+
+void CheckShaderHashCacheInclude(const FString& VirtualFilePath, EShaderPlatform ShaderPlatform)
+{
+	bool bIgnoreInclude = GShaderHashCache.ShouldIgnoreInclude(VirtualFilePath, ShaderPlatform);
+
+	checkf(!bIgnoreInclude,
+		TEXT("Shader compiler is trying to include %s, which is not located in IShaderFormat::GetPlatformIncludeDirectory for %s."),
+		*VirtualFilePath,
+		*ShaderPlatformToShaderFormatName(ShaderPlatform).ToString());
+}
+
+void InitializeShaderTypes()
+{
+	UE_LOG(LogShaders, Log, TEXT("InitializeShaderTypes() begin"));
+
+	LogShaderSourceDirectoryMappings();
+
+	TMap<FString, TArray<const TCHAR*> > ShaderFileToUniformBufferVariables;
+	BuildShaderFileToUniformBufferMap(ShaderFileToUniformBufferVariables);
+
+	FShaderType::Initialize(ShaderFileToUniformBufferVariables);
+	FVertexFactoryType::Initialize(ShaderFileToUniformBufferVariables);
+
+	FShaderPipelineType::Initialize();
+
+	UE_LOG(LogShaders, Log, TEXT("InitializeShaderTypes() end"));
+}
+
+void UninitializeShaderTypes()
+{
+	UE_LOG(LogShaders, Log, TEXT("UninitializeShaderTypes() begin"));
+
+	FShaderPipelineType::Uninitialize();
+
+	FShaderType::Uninitialize();
+	FVertexFactoryType::Uninitialize();
+
+	UE_LOG(LogShaders, Log, TEXT("UninitializeShaderTypes() end"));
+}
+
+/**
+ * Flushes the shader file and CRC cache, and regenerates the binary shader files if necessary.
+ * Allows shader source files to be re-read properly even if they've been modified since startup.
+ */
+void FlushShaderFileCache()
+{
+	UE_LOG(LogShaders, Log, TEXT("FlushShaderFileCache() begin"));
+
+	GShaderHashCache.Empty();
+	GShaderFileCache.Empty();
+
+	if (!FPlatformProperties::RequiresCookedData())
+	{
+		LogShaderSourceDirectoryMappings();
+
+		TMap<FString, TArray<const TCHAR*> > ShaderFileToUniformBufferVariables;
+		BuildShaderFileToUniformBufferMap(ShaderFileToUniformBufferVariables);
+
+		for (TLinkedList<FShaderPipelineType*>::TConstIterator It(FShaderPipelineType::GetTypeList()); It; It.Next())
+		{
+			const auto& Stages = It->GetStages();
+			for (const FShaderType* ShaderType : Stages)
+			{
+				((FShaderType*)ShaderType)->FlushShaderFileCache(ShaderFileToUniformBufferVariables);
+			}
+		}
+
+		for(TLinkedList<FShaderType*>::TIterator It(FShaderType::GetTypeList()); It; It.Next())
+		{
+			It->FlushShaderFileCache(ShaderFileToUniformBufferVariables);
+		}
+
+		for(TLinkedList<FVertexFactoryType*>::TIterator It(FVertexFactoryType::GetTypeList()); It; It.Next())
+		{
+			It->FlushShaderFileCache(ShaderFileToUniformBufferVariables);
+		}
+	}
+
+	UE_LOG(LogShaders, Log, TEXT("FlushShaderFileCache() end"));
+}
+
+void GenerateReferencedUniformBuffers(
+	const TCHAR* SourceFilename, 
+	const TCHAR* ShaderTypeName, 
+	const TMap<FString, TArray<const TCHAR*> >& ShaderFileToUniformBufferVariables,
+	TMap<const TCHAR*,FCachedUniformBufferDeclaration>& UniformBufferEntries)
+{
+	TArray<FString> FilesToSearch;
+	GetShaderIncludes(SourceFilename, SourceFilename, FilesToSearch, GMaxRHIShaderPlatform);
+	FilesToSearch.Add(SourceFilename);
+
+	for (int32 FileIndex = 0; FileIndex < FilesToSearch.Num(); FileIndex++)
+	{
+		const TArray<const TCHAR*>& FoundUniformBufferVariables = ShaderFileToUniformBufferVariables.FindChecked(FilesToSearch[FileIndex]);
+
+		for (int32 VariableIndex = 0; VariableIndex < FoundUniformBufferVariables.Num(); VariableIndex++)
+		{
+			UniformBufferEntries.Add(FoundUniformBufferVariables[VariableIndex], FCachedUniformBufferDeclaration());
+		}
+	}
+}
+
+void SerializeUniformBufferInfo(FShaderSaveArchive& Ar, const TMap<const TCHAR*,FCachedUniformBufferDeclaration>& UniformBufferEntries)
+{
+	for (TMap<const TCHAR*,FCachedUniformBufferDeclaration>::TConstIterator It(UniformBufferEntries); It; ++It)
+	{
+		for (TLinkedList<FShaderParametersMetadata*>::TIterator StructIt(FShaderParametersMetadata::GetStructList()); StructIt; StructIt.Next())
+		{
+			if (It.Key() == StructIt->GetShaderVariableName())
+			{
+				// Serialize information about the struct layout so we can detect when it changes
+				const FShaderParametersMetadata& Struct = **StructIt;
+				const TArray<FShaderParametersMetadata::FMember>& Members = Struct.GetMembers();
+
+				int32 NumMembers = Members.Num();
+				// Serializing with NULL so that FShaderSaveArchive will record the length without causing an actual data serialization
+				Ar.Serialize(NULL, NumMembers);
+
+				for (int32 MemberIndex = 0; MemberIndex < Members.Num(); MemberIndex++)
+				{
+					const FShaderParametersMetadata::FMember& Member = Members[MemberIndex];
+
+					// Note: Only comparing number of floats used by each member and type, so this can be tricked (eg. swapping two equal size and type members)
+					int32 MemberSize = Member.GetNumColumns() * Member.GetNumRows();
+					Ar.Serialize(NULL, MemberSize);
+					int32 MemberType = (int32)Member.GetBaseType();
+					Ar.Serialize(NULL, MemberType);
+				}
+				break;
+			}
+		}
+	}
+}
+
+FString FShaderCompilerError::GetShaderSourceFilePath() const
+{
+	if (IFileManager::Get().FileExists(*ErrorVirtualFilePath))
+	{
+		return ErrorVirtualFilePath;
+	}
+	else
+	{
+		return ::GetShaderSourceFilePath(ErrorVirtualFilePath, nullptr);
+	}
+}
+
+const TMap<FString, FString>& AllShaderSourceDirectoryMappings()
+{
+	return GShaderSourceDirectoryMappings;
+}
+
+void ResetAllShaderSourceDirectoryMappings()
+{
+	GShaderSourceDirectoryMappings.Reset();
+}
+
+void AddShaderSourceDirectoryMapping(const FString& VirtualShaderDirectory, const FString& RealShaderDirectory)
+{
+	check(IsInGameThread());
+
+	if (FPlatformProperties::RequiresCookedData())
+	{
+		return;
+	}
+
+	// Do sanity checks of the virtual shader directory to map.
+	check(VirtualShaderDirectory.StartsWith(TEXT("/")));
+	check(!VirtualShaderDirectory.EndsWith(TEXT("/")));
+	check(!VirtualShaderDirectory.Contains(FString(TEXT("."))));
+
+	// Detect collisions with any other mappings.
+	check(!GShaderSourceDirectoryMappings.Contains(VirtualShaderDirectory));
+
+	// Make sure the real directory to map exists.
+	check(FPaths::DirectoryExists(RealShaderDirectory));
+
+	// Make sure the Generated directory does not exist, because is reserved for C++ generated shader source
+	// by the FShaderCompilerEnvironment::IncludeVirtualPathToContentsMap member.
+	checkf(!FPaths::DirectoryExists(RealShaderDirectory / TEXT("Generated")),
+		TEXT("\"%s/Generated\" is not permitted to exist since C++ generated shader file would be mapped to this directory."), *RealShaderDirectory);
+
+	UE_LOG(LogShaders, Log, TEXT("Mapping virtual shader directory %s to %s"),
+		*VirtualShaderDirectory, *RealShaderDirectory);
+	GShaderSourceDirectoryMappings.Add(VirtualShaderDirectory, RealShaderDirectory);
+}
diff --git a/Engine/Source/Runtime/RenderCore/Private/ShaderParameterMetadata.cpp b/Engine/Source/Runtime/RenderCore/Private/ShaderParameterMetadata.cpp
new file mode 100644
--- /dev/null
+++ b/Engine/Source/Runtime/RenderCore/Private/ShaderParameterMetadata.cpp
@@ -0,0 +1,258 @@
+// Copyright 1998-2018 Epic Games, Inc. All Rights Reserved.
+
+/*=============================================================================
+	ShaderParameterMetadata.cpp: Shader parameter metadata implementations.
+=============================================================================*/
+
+#include "ShaderParameterMetadata.h"
+#include "RenderCore.h"
+
+
+static TLinkedList<FShaderParametersMetadata*>* GUniformStructList = nullptr;
+
+TLinkedList<FShaderParametersMetadata*>*& FShaderParametersMetadata::GetStructList()
+{
+	return GUniformStructList;
+}
+
+TMap<FName, FShaderParametersMetadata*>& FShaderParametersMetadata::GetNameStructMap()
+{
+	static TMap<FName, FShaderParametersMetadata*> GlobalNameStructMap;
+	return GlobalNameStructMap;
+}
+
+FShaderParametersMetadata* FindUniformBufferStructByName(const TCHAR* StructName)
+{
+	FName FindByName(StructName, FNAME_Find);
+	FShaderParametersMetadata* FoundStruct = FShaderParametersMetadata::GetNameStructMap().FindRef(FindByName);
+	return FoundStruct;
+}
+
+FShaderParametersMetadata* FindUniformBufferStructByFName(FName StructName)
+{
+	return FShaderParametersMetadata::GetNameStructMap().FindRef(StructName);
+}
+
+class FUniformBufferMemberAndOffset
+{
+public:
+	FUniformBufferMemberAndOffset(const FShaderParametersMetadata& InContainingStruct, const FShaderParametersMetadata::FMember& InMember, int32 InStructOffset) :
+		ContainingStruct(InContainingStruct),
+		Member(InMember),
+		StructOffset(InStructOffset)
+	{}
+
+	const FShaderParametersMetadata& ContainingStruct;
+	const FShaderParametersMetadata::FMember& Member;
+	int32 StructOffset;
+};
+
+FShaderParametersMetadata::FShaderParametersMetadata(
+	EUseCase InUseCase,
+	const FName& InLayoutName,
+	const TCHAR* InStructTypeName,
+	const TCHAR* InShaderVariableName,
+	uint32 InSize,
+	const TArray<FMember>& InMembers)
+	: StructTypeName(InStructTypeName)
+	, ShaderVariableName(InShaderVariableName)
+	, Size(InSize)
+	, UseCase(InUseCase)
+	, Layout(InLayoutName)
+	, Members(InMembers)
+	, GlobalListLink(this)
+	, bLayoutInitialized(false)
+{
+	if (UseCase == EUseCase::GlobalShaderParameterStruct)
+	{
+		// Register this uniform buffer struct in global list.
+		GlobalListLink.LinkHead(GetStructList());
+
+		FName StructTypeFName(StructTypeName);
+		// Verify that during FName creation there's no case conversion
+		checkSlow(FCString::Strcmp(StructTypeName, *StructTypeFName.GetPlainNameString()) == 0);
+		GetNameStructMap().Add(FName(StructTypeFName), this);
+	}
+	else
+	{
+		// We cannot initialize the layout during global initialization, since we have to walk nested struct members.
+		// Structs created during global initialization will have bRegisterForAutoBinding==false, and are initialized during startup.
+		// Structs created at runtime with bRegisterForAutoBinding==true can be initialized now.
+		InitializeLayout();
+	}
+
+	if (UseCase == EUseCase::ShaderParameterStruct)
+	{
+		check(ShaderVariableName == nullptr);
+	}
+}
+
+void FShaderParametersMetadata::InitializeAllGlobalStructs()
+{
+	for (TLinkedList<FShaderParametersMetadata*>::TIterator StructIt(FShaderParametersMetadata::GetStructList()); StructIt; StructIt.Next())
+	{
+		StructIt->InitializeLayout();
+	}
+}
+
+void FShaderParametersMetadata::InitializeLayout()
+{
+	check(!bLayoutInitialized);
+	Layout.ConstantBufferSize = Size;
+
+	TArray<FUniformBufferMemberAndOffset> MemberStack;
+	MemberStack.Reserve(Members.Num());
+
+	for (int32 MemberIndex = 0; MemberIndex < Members.Num(); MemberIndex++)
+	{
+		MemberStack.Push(FUniformBufferMemberAndOffset(*this, Members[MemberIndex], 0));
+	}
+
+	// TODO(RDG): Render graph tracked resources in GlobalShaderParameterStruct could technically be possible if the creation is deferred.
+	bool bAllowGraphResources = UseCase == EUseCase::ShaderParameterStruct;
+
+	for (int32 i = 0; i < MemberStack.Num(); ++i)
+	{
+		const FShaderParametersMetadata& CurrentStruct = MemberStack[i].ContainingStruct;
+		const FMember& CurrentMember = MemberStack[i].Member;
+		EUniformBufferBaseType BaseType = CurrentMember.GetBaseType();
+
+		const FShaderParametersMetadata* ChildStruct = CurrentMember.GetStructMetadata();
+
+		if (DO_CHECK)
+		{
+			const FString CppName = FString::Printf(TEXT("%s::%s"), CurrentStruct.GetStructTypeName(), CurrentMember.GetName());
+
+			if (!bAllowGraphResources && (
+				BaseType == UBMT_GRAPH_TRACKED_TEXTURE ||
+				BaseType == UBMT_GRAPH_TRACKED_SRV ||
+				BaseType == UBMT_GRAPH_TRACKED_UAV ||
+				BaseType == UBMT_RENDER_TARGET_BINDING_SLOTS))
+			{
+				UE_LOG(LogRendererCore, Fatal, TEXT("Shader parameter %s error: Graph resources are only allowed in shader parameter structs."), *CppName);
+			}
+
+			if (BaseType == UBMT_REFERENCED_STRUCT)
+			{
+				if (UseCase != EUseCase::ShaderParameterStruct)
+				{
+					UE_LOG(LogRendererCore, Fatal, TEXT("Shader parameter %s error: Shader parameter struct reference can only be done in shader parameter structs."), *CppName);
+				}
+
+				if (ChildStruct->GetUseCase() != FShaderParametersMetadata::EUseCase::GlobalShaderParameterStruct)
+				{
+					UE_LOG(LogRendererCore, Fatal, TEXT("Shader parameter %s error: Shader parameter struct can only reference global ones, but %s is not."), *CppName, ChildStruct->GetStructTypeName());
+				}
+			}
+
+			if (UseCase == EUseCase::ShaderParameterStruct &&
+				(BaseType == UBMT_NESTED_STRUCT || BaseType == UBMT_INCLUDED_STRUCT) &&
+				ChildStruct->GetUseCase() != EUseCase::ShaderParameterStruct)
+			{
+				UE_LOG(LogRendererCore, Fatal, TEXT("Shader parameter %s error: shader parameter structs can only nests or include shader parameter struct, but %s is not."), *CppName, ChildStruct->GetStructTypeName());
+			}
+		}
+
+		if (IsUniformBufferResourceType(BaseType))
+		{
+			Layout.Resources.Add(BaseType);
+			const uint32 AbsoluteMemberOffset = CurrentMember.GetOffset() + MemberStack[i].StructOffset;
+			check(AbsoluteMemberOffset < (1u << (Layout.ResourceOffsets.GetTypeSize() * 8)));
+			Layout.ResourceOffsets.Add(AbsoluteMemberOffset);
+		}
+
+		if (ChildStruct && BaseType != UBMT_REFERENCED_STRUCT)
+		{
+			int32 AbsoluteStructOffset = CurrentMember.GetOffset() + MemberStack[i].StructOffset;
+
+			for (int32 StructMemberIndex = 0; StructMemberIndex < ChildStruct->Members.Num(); StructMemberIndex++)
+			{
+				const FMember& StructMember = ChildStruct->Members[StructMemberIndex];
+				MemberStack.Insert(FUniformBufferMemberAndOffset(*ChildStruct, StructMember, AbsoluteStructOffset), i + 1 + StructMemberIndex);
+			}
+		}
+	}
+
+	Layout.ComputeHash();
+
+	bLayoutInitialized = true;
+}
+
+void FShaderParametersMetadata::GetNestedStructs(TArray<const FShaderParametersMetadata*>& OutNestedStructs) const
+{
+	for (int32 i = 0; i < Members.Num(); ++i)
+	{
+		const FMember& CurrentMember = Members[i];
+
+		const FShaderParametersMetadata* MemberStruct = CurrentMember.GetStructMetadata();
+
+		if (MemberStruct)
+		{
+			OutNestedStructs.Add(MemberStruct);
+			MemberStruct->GetNestedStructs(OutNestedStructs);
+		}
+	}
+}
+
+void FShaderParametersMetadata::AddResourceTableEntries(TMap<FString, FResourceTableEntry>& ResourceTableMap, TMap<FString, uint32>& ResourceTableLayoutHashes) const
+{
+	uint16 ResourceIndex = 0;
+	FString Prefix = FString::Printf(TEXT("%s_"), ShaderVariableName);
+	AddResourceTableEntriesRecursive(ShaderVariableName, *Prefix, ResourceIndex, ResourceTableMap);
+	ResourceTableLayoutHashes.Add(ShaderVariableName, GetLayout().GetHash());
+}
+
+void FShaderParametersMetadata::AddResourceTableEntriesRecursive(const TCHAR* UniformBufferName, const TCHAR* Prefix, uint16& ResourceIndex, TMap<FString, FResourceTableEntry>& ResourceTableMap) const
+{
+	for (int32 MemberIndex = 0; MemberIndex < Members.Num(); ++MemberIndex)
+	{
+		const FMember& Member = Members[MemberIndex];
+		if (IsUniformBufferResourceType(Member.GetBaseType()))
+		{
+			FResourceTableEntry& Entry = ResourceTableMap.FindOrAdd(FString::Printf(TEXT("%s%s"), Prefix, Member.GetName()));
+			if (Entry.UniformBufferName.IsEmpty())
+			{
+				Entry.UniformBufferName = UniformBufferName;
+				Entry.Type = Member.GetBaseType();
+				Entry.ResourceIndex = ResourceIndex++;
+			}
+		}
+		else if (Member.GetBaseType() == UBMT_NESTED_STRUCT)
+		{
+			check(Member.GetStructMetadata());
+			FString MemberPrefix = FString::Printf(TEXT("%s%s_"), Prefix, Member.GetName());
+			Member.GetStructMetadata()->AddResourceTableEntriesRecursive(UniformBufferName, *MemberPrefix, ResourceIndex, ResourceTableMap);
+		}
+		else if (Member.GetBaseType() == UBMT_INCLUDED_STRUCT)
+		{
+			check(Member.GetStructMetadata());
+			Member.GetStructMetadata()->AddResourceTableEntriesRecursive(UniformBufferName, Prefix, ResourceIndex, ResourceTableMap);
+		}
+	}
+}
+
+void FShaderParametersMetadata::FindMemberFromOffset(uint16 MemberOffset, const FShaderParametersMetadata** OutContainingStruct, const FShaderParametersMetadata::FMember** OutMember) const
+{
+	check(MemberOffset < GetSize());
+
+	for (const FMember& Member : Members)
+	{
+		EUniformBufferBaseType BaseType = Member.GetBaseType();
+		if (BaseType == UBMT_NESTED_STRUCT || BaseType == UBMT_INCLUDED_STRUCT)
+		{
+			const FShaderParametersMetadata* SubStruct = Member.GetStructMetadata();
+			if (MemberOffset < Member.GetOffset() + SubStruct->GetSize())
+			{
+				return SubStruct->FindMemberFromOffset(MemberOffset - Member.GetOffset(), OutContainingStruct, OutMember);
+			}
+		}
+		else if (Member.GetOffset() == MemberOffset)
+		{
+			*OutContainingStruct = this;
+			*OutMember = &Member;
+			return;
+		}
+	}
+
+	checkf(0, TEXT("Looks like this offset is invalid."));
+}
diff --git a/Engine/Source/Runtime/RenderCore/Private/ShaderParameterStruct.cpp b/Engine/Source/Runtime/RenderCore/Private/ShaderParameterStruct.cpp
new file mode 100644
--- /dev/null
+++ b/Engine/Source/Runtime/RenderCore/Private/ShaderParameterStruct.cpp
@@ -0,0 +1,189 @@
+// Copyright 1998-2018 Epic Games, Inc. All Rights Reserved.
+
+/*=============================================================================
+	ShaderParameterStruct.cpp: Shader parameter struct implementations.
+=============================================================================*/
+
+#include "ShaderParameterStruct.h"
+
+
+/** Context of binding a map. */
+struct FShaderParameterStructBindingContext
+{
+	// Bindings to bind.
+	FShaderParameterBindings* Bindings;
+
+	// The shader parameter map from the compilation.
+	const FShaderParameterMap* ParametersMap;
+
+	// Map of global shader name that were bound to C++ members.
+	TMap<FString, FString> ShaderGlobalScopeBindings;
+
+	// C++ name of the render target binding slot.
+	FString RenderTargetBindingSlotCppName;
+
+
+	void Bind(
+		const FShaderParametersMetadata& StructMetaData,
+		const TCHAR* MemberPrefix,
+		uint32 GeneralByteOffset)
+	{
+		const TArray<FShaderParametersMetadata::FMember>& StructMembers = StructMetaData.GetMembers();
+
+		for (const FShaderParametersMetadata::FMember& Member : StructMembers)
+		{
+			EUniformBufferBaseType BaseType = Member.GetBaseType();
+
+			FString CppName = FString::Printf(TEXT("%s::%s"), StructMetaData.GetStructTypeName(), Member.GetName());
+
+			// Ignore rasterizer binding slots entirely since this actually have nothing to do with a shader.
+			if (BaseType == UBMT_RENDER_TARGET_BINDING_SLOTS)
+			{
+				if (!RenderTargetBindingSlotCppName.IsEmpty())
+				{
+					UE_LOG(LogShaders, Fatal, TEXT("Render target binding slots collision: %s & %s"),
+						*RenderTargetBindingSlotCppName, *CppName);
+				}
+				RenderTargetBindingSlotCppName = CppName;
+				continue;
+			}
+
+			// Compute the shader member name to look for according to nesting.
+			FString ShaderBindingName = FString::Printf(TEXT("%s%s"), MemberPrefix, Member.GetName());
+
+			uint16 ByteOffset = uint16(GeneralByteOffset + Member.GetOffset());
+			check(uint32(ByteOffset) == GeneralByteOffset + Member.GetOffset());
+
+			if (BaseType == UBMT_NESTED_STRUCT || BaseType == UBMT_INCLUDED_STRUCT)
+			{
+				FString NewPrefix = FString::Printf(TEXT("%s%s_"), MemberPrefix, Member.GetName());
+				Bind(
+					*Member.GetStructMetadata(),
+					/* MemberPrefix = */ BaseType == UBMT_INCLUDED_STRUCT ? MemberPrefix : *NewPrefix,
+					/* GeneralByteOffset = */ ByteOffset);
+				continue;
+			}
+			else if (BaseType == UBMT_REFERENCED_STRUCT)
+			{
+				// The member name of a globally referenced struct is the not name on the struct.
+				ShaderBindingName = Member.GetStructMetadata()->GetShaderVariableName();
+			}
+
+			if (ShaderGlobalScopeBindings.Contains(ShaderBindingName))
+			{
+				UE_LOG(LogShaders, Fatal, TEXT("%s can't bind shader parameter %s, because it has already be bound by %s."), *CppName, *ShaderBindingName, **ShaderGlobalScopeBindings.Find(ShaderBindingName));
+			}
+			ShaderGlobalScopeBindings.Add(ShaderBindingName, CppName);
+
+			uint16 BufferIndex, BaseIndex, BoundSize;
+			if (!ParametersMap->FindParameterAllocation(*ShaderBindingName, BufferIndex, BaseIndex, BoundSize))
+			{
+				continue;
+			}
+
+			if (BaseType == UBMT_INT32 || BaseType == UBMT_UINT32 || BaseType == UBMT_FLOAT32)
+			{
+				uint32 ByteSize = Member.GetMemberSize();
+
+				FShaderParameterBindings::FParameter Parameter;
+				Parameter.BufferIndex = BufferIndex;
+				Parameter.BaseIndex = BaseIndex;
+				Parameter.ByteOffset = ByteOffset;
+				Parameter.ByteSize = BoundSize;
+
+				if (uint32(BoundSize) > ByteSize)
+				{
+					UE_LOG(LogShaders, Fatal, TEXT("The size required to bind shader parameter %s is %i bytes, smaller than %s's %i bytes."),
+						*ShaderBindingName, BoundSize, *CppName, ByteSize);
+				}
+
+				Bindings->Parameters.Add(Parameter);
+			}
+			else if (BaseType == UBMT_REFERENCED_STRUCT)
+			{
+				FShaderParameterBindings::FParameterStructReference Parameter;
+				Parameter.BufferIndex = BufferIndex;
+				Parameter.ByteOffset = ByteOffset;
+
+				Bindings->ParameterReferences.Add(Parameter);
+			}
+			else if (
+				BaseType == UBMT_TEXTURE ||
+				BaseType == UBMT_SRV ||
+				BaseType == UBMT_SAMPLER ||
+				BaseType == UBMT_GRAPH_TRACKED_TEXTURE ||
+				BaseType == UBMT_GRAPH_TRACKED_SRV ||
+				BaseType == UBMT_GRAPH_TRACKED_UAV)
+			{
+				FShaderParameterBindings::FResourceParameter Parameter;
+				Parameter.BaseIndex = BaseIndex;
+				Parameter.NumResources = BoundSize;
+				Parameter.ByteOffset = ByteOffset;
+
+				checkf(BoundSize <= 1, TEXT("Not enough room in shader parameter struct for the shader."));
+
+				if (BaseType == UBMT_TEXTURE)
+					Bindings->Textures.Add(Parameter);
+				else if (BaseType == UBMT_SRV)
+					Bindings->SRVs.Add(Parameter);
+				else if (BaseType == UBMT_SAMPLER)
+					Bindings->Samplers.Add(Parameter);
+				else if (BaseType == UBMT_GRAPH_TRACKED_TEXTURE)
+					Bindings->GraphTextures.Add(Parameter);
+				else if (BaseType == UBMT_GRAPH_TRACKED_SRV)
+					Bindings->GraphSRVs.Add(Parameter);
+				else if (BaseType == UBMT_GRAPH_TRACKED_UAV)
+					Bindings->GraphUAVs.Add(Parameter);
+			}
+			else
+			{
+				checkf(0, TEXT("Unexpected base type for a shader parameter struct member."));
+			}
+		} // for (const FShaderParametersMetadata::FMember& Member : StructMembers)
+	}
+}; // struct FShaderParameterStructBindingContext
+
+void FShaderParameterBindings::Bind(const FShaderParameterMap& ParametersMap, const FShaderParametersMetadata& StructMetaData)
+{
+	checkf(StructMetaData.GetSize() < (1 << (sizeof(uint16) * 8)), TEXT("Shader parameter structure can only have a size < 65536 bytes."));
+
+	FShaderParameterStructBindingContext BindingContext;
+	BindingContext.Bindings = this;
+	BindingContext.ParametersMap = &ParametersMap;
+	BindingContext.Bind(
+		StructMetaData,
+		/* MemberPrefix = */ TEXT(""),
+		/* ByteOffset = */ 0);
+
+	TArray<FString> AllParameterNames;
+	ParametersMap.GetAllParameterNames(AllParameterNames);
+	if (0 && BindingContext.ShaderGlobalScopeBindings.Num() != AllParameterNames.Num()) // TODO: enable that once we can.
+	{
+		UE_LOG(LogShaders, Error, TEXT("%i shader parameters have not been bound:"), AllParameterNames.Num() - BindingContext.ShaderGlobalScopeBindings.Num());
+		for (const FString& GlobalParameterName : AllParameterNames)
+		{
+			if (!BindingContext.ShaderGlobalScopeBindings.Contains(GlobalParameterName))
+			{
+				UE_LOG(LogShaders, Error, TEXT("  %s"), *GlobalParameterName);
+			}
+		}
+
+		// TODO: would be great to have the shader name for the error message.
+		UE_LOG(LogShaders, Fatal, TEXT("Some shader parameters have not been bound."));
+	}
+}
+
+void SetNullShaderParameterFatalError(const FShader* Shader, const FShaderParametersMetadata* ParametersMetadata, uint16 MemberOffset)
+{
+	const FShaderParametersMetadata* MemberContainingStruct = nullptr;
+	const FShaderParametersMetadata::FMember* Member = nullptr;
+	ParametersMetadata->FindMemberFromOffset(MemberOffset, &MemberContainingStruct, &Member);
+
+	const TCHAR* ShaderClassName = Shader->GetType()->GetName();
+
+	UE_LOG(LogShaders, Fatal,
+		TEXT("%s's required shader parameter %s::%s was not set."),
+		ShaderClassName,
+		MemberContainingStruct->GetStructTypeName(),
+		Member->GetName());
+}
diff --git a/Engine/Source/Runtime/RenderCore/Private/ShaderParameters.cpp b/Engine/Source/Runtime/RenderCore/Private/ShaderParameters.cpp
new file mode 100644
--- /dev/null
+++ b/Engine/Source/Runtime/RenderCore/Private/ShaderParameters.cpp
@@ -0,0 +1,515 @@
+// Copyright 1998-2018 Epic Games, Inc. All Rights Reserved.
+
+/*=============================================================================
+	ShaderParameters.cpp: Shader parameter implementation.
+=============================================================================*/
+
+#include "ShaderParameters.h"
+#include "Containers/List.h"
+#include "UniformBuffer.h"
+#include "ShaderCore.h"
+#include "Shader.h"
+#include "VertexFactory.h"
+#include "ShaderCodeLibrary.h"
+
+void FShaderParameter::Bind(const FShaderParameterMap& ParameterMap,const TCHAR* ParameterName,EShaderParameterFlags Flags)
+{
+#if UE_BUILD_DEBUG
+	bInitialized = true;
+#endif
+
+	if (!ParameterMap.FindParameterAllocation(ParameterName,BufferIndex,BaseIndex,NumBytes) && Flags == SPF_Mandatory)
+	{
+		if (!UE_LOG_ACTIVE(LogShaders, Log))
+		{
+			UE_LOG(LogShaders, Fatal,TEXT("Failure to bind non-optional shader parameter %s!  The parameter is either not present in the shader, or the shader compiler optimized it out."),ParameterName);
+		}
+		else
+		{
+			// We use a non-Slate message box to avoid problem where we haven't compiled the shaders for Slate.
+			FPlatformMisc::MessageBoxExt( EAppMsgType::Ok, *FText::Format(
+				NSLOCTEXT("UnrealEd", "Error_FailedToBindShaderParameter", "Failure to bind non-optional shader parameter {0}! The parameter is either not present in the shader, or the shader compiler optimized it out. This will be an assert with LogShaders suppressed!"),
+				FText::FromString(ParameterName)).ToString(), TEXT("Warning"));
+		}
+	}
+}
+
+FArchive& operator<<(FArchive& Ar,FShaderParameter& P)
+{
+#if UE_BUILD_DEBUG
+	if (Ar.IsLoading())
+	{
+		P.bInitialized = true;
+	}
+#endif
+
+	uint16& PBufferIndex = P.BufferIndex;
+	return Ar << P.BaseIndex << P.NumBytes << PBufferIndex;
+}
+
+void FShaderResourceParameter::Bind(const FShaderParameterMap& ParameterMap,const TCHAR* ParameterName,EShaderParameterFlags Flags)
+{
+	uint16 UnusedBufferIndex = 0;
+
+#if UE_BUILD_DEBUG
+	bInitialized = true;
+#endif
+
+	if(!ParameterMap.FindParameterAllocation(ParameterName,UnusedBufferIndex,BaseIndex,NumResources) && Flags == SPF_Mandatory)
+	{
+		if (!UE_LOG_ACTIVE(LogShaders, Log))
+		{
+			UE_LOG(LogShaders, Fatal,TEXT("Failure to bind non-optional shader resource parameter %s!  The parameter is either not present in the shader, or the shader compiler optimized it out."),ParameterName);
+		}
+		else
+		{
+			// We use a non-Slate message box to avoid problem where we haven't compiled the shaders for Slate.
+			FPlatformMisc::MessageBoxExt( EAppMsgType::Ok, *FText::Format(
+				NSLOCTEXT("UnrealEd", "Error_FailedToBindShaderParameter", "Failure to bind non-optional shader parameter {0}! The parameter is either not present in the shader, or the shader compiler optimized it out. This will be an assert with LogShaders suppressed!"),
+				FText::FromString(ParameterName)).ToString(), TEXT("Warning"));
+		}
+	}
+}
+
+FArchive& operator<<(FArchive& Ar,FShaderResourceParameter& P)
+{
+#if UE_BUILD_DEBUG
+	if (Ar.IsLoading())
+	{
+		P.bInitialized = true;
+	}
+#endif
+
+	return Ar << P.BaseIndex << P.NumResources;
+}
+
+void FShaderUniformBufferParameter::ModifyCompilationEnvironment(const TCHAR* ParameterName,const FShaderParametersMetadata& Struct,EShaderPlatform Platform,FShaderCompilerEnvironment& OutEnvironment)
+{
+	const FString IncludeName = FString::Printf(TEXT("/Engine/Generated/UniformBuffers/%s.ush"),ParameterName);
+	// Add the uniform buffer declaration to the compilation environment as an include: UniformBuffers/<ParameterName>.usf
+	FString Declaration;
+	CreateUniformBufferShaderDeclaration(ParameterName, Struct, Declaration);
+	OutEnvironment.IncludeVirtualPathToContentsMap.Add(IncludeName, Declaration);
+
+	FString& GeneratedUniformBuffersInclude = OutEnvironment.IncludeVirtualPathToContentsMap.FindOrAdd("/Engine/Generated/GeneratedUniformBuffers.ush");
+	FString Include = FString::Printf(TEXT("#include \"/Engine/Generated/UniformBuffers/%s.ush\"") LINE_TERMINATOR, ParameterName);
+
+	GeneratedUniformBuffersInclude.Append(Include);
+	Struct.AddResourceTableEntries(OutEnvironment.ResourceTableMap, OutEnvironment.ResourceTableLayoutHashes);
+}
+
+void FShaderUniformBufferParameter::Bind(const FShaderParameterMap& ParameterMap,const TCHAR* ParameterName,EShaderParameterFlags Flags)
+{
+	uint16 UnusedBaseIndex = 0;
+	uint16 UnusedNumBytes = 0;
+
+#if UE_BUILD_DEBUG
+	bInitialized = true;
+#endif
+
+	if(!ParameterMap.FindParameterAllocation(ParameterName,BaseIndex,UnusedBaseIndex,UnusedNumBytes))
+	{
+		bIsBound = false;
+		if(Flags == SPF_Mandatory)
+		{
+			if (!UE_LOG_ACTIVE(LogShaders, Log))
+			{
+				UE_LOG(LogShaders, Fatal,TEXT("Failure to bind non-optional shader resource parameter %s!  The parameter is either not present in the shader, or the shader compiler optimized it out."),ParameterName);
+			}
+			else
+			{
+				// We use a non-Slate message box to avoid problem where we haven't compiled the shaders for Slate.
+				FPlatformMisc::MessageBoxExt( EAppMsgType::Ok, *FText::Format(
+					NSLOCTEXT("UnrealEd", "Error_FailedToBindShaderParameter", "Failure to bind non-optional shader parameter {0}! The parameter is either not present in the shader, or the shader compiler optimized it out. This will be an assert with LogShaders suppressed!"),
+					FText::FromString(ParameterName)).ToString(), TEXT("Warning"));
+			}
+		}
+	}
+	else
+	{
+		bIsBound = true;
+	}
+}
+
+/** The individual bits of a uniform buffer declaration. */
+struct FUniformBufferDecl
+{
+	/** Members to place in the constant buffer. */
+	FString ConstantBufferMembers;
+	/** Members to place in the resource table. */
+	FString ResourceMembers;
+	/** Members in the struct HLSL shader code will access. */
+	FString StructMembers;
+	/** The HLSL initializer that will copy constants and resources in to the struct. */
+	FString Initializer;
+};
+
+/** Generates a HLSL struct declaration for a uniform buffer struct. */
+static void CreateHLSLUniformBufferStructMembersDeclaration(
+	const FShaderParametersMetadata& UniformBufferStruct, 
+	const FString& NamePrefix, 
+	uint32 StructOffset, 
+	FUniformBufferDecl& Decl, 
+	uint32& HLSLBaseOffset)
+{
+	const TArray<FShaderParametersMetadata::FMember>& StructMembers = UniformBufferStruct.GetMembers();
+	
+	int32 OpeningBraceLocPlusOne = Decl.Initializer.Len();
+
+	FString PreviousBaseTypeName = TEXT("float");
+	for (int32 MemberIndex = 0; MemberIndex < StructMembers.Num(); ++MemberIndex)
+	{
+		const FShaderParametersMetadata::FMember& Member = StructMembers[MemberIndex];
+		
+		FString ArrayDim;
+		if(Member.GetNumElements() > 0)
+		{
+			ArrayDim = FString::Printf(TEXT("[%u]"),Member.GetNumElements());
+		}
+
+		if(Member.GetBaseType() == UBMT_NESTED_STRUCT)
+		{
+			Decl.StructMembers += TEXT("struct {\r\n");
+			Decl.Initializer += TEXT(",{");
+			CreateHLSLUniformBufferStructMembersDeclaration(*Member.GetStructMetadata(), FString::Printf(TEXT("%s%s_"), *NamePrefix, Member.GetName()), StructOffset + Member.GetOffset(), Decl, HLSLBaseOffset);
+			Decl.Initializer += TEXT("}");
+			Decl.StructMembers += FString::Printf(TEXT("} %s%s;\r\n"),Member.GetName(),*ArrayDim);
+		}
+		else if (Member.GetBaseType() == UBMT_INCLUDED_STRUCT)
+		{
+			Decl.Initializer += TEXT(",");
+			CreateHLSLUniformBufferStructMembersDeclaration(*Member.GetStructMetadata(), NamePrefix, StructOffset + Member.GetOffset(), Decl, HLSLBaseOffset);
+		}
+		else if (IsUniformBufferResourceType(Member.GetBaseType()))
+		{
+			// Skip resources, they will be replaced with padding by the next member in the constant buffer.
+			// This padding will cause gaps in the constant buffer.  Alternatively we could compact the constant buffer during RHICreateUniformBuffer.
+			continue;
+		}
+		else 
+		{
+			// Generate the base type name.
+			FString BaseTypeName;
+			switch(Member.GetBaseType())
+			{
+			case UBMT_BOOL:    BaseTypeName = TEXT("bool"); break;
+			case UBMT_INT32:   BaseTypeName = TEXT("int"); break;
+			case UBMT_UINT32:  BaseTypeName = TEXT("uint"); break;
+			case UBMT_FLOAT32: 
+				if (Member.GetPrecision() == EShaderPrecisionModifier::Float)
+				{
+					BaseTypeName = TEXT("float"); 
+				}
+				else if (Member.GetPrecision() == EShaderPrecisionModifier::Half)
+				{
+					BaseTypeName = TEXT("half"); 
+				}
+				else if (Member.GetPrecision() == EShaderPrecisionModifier::Fixed)
+				{
+					BaseTypeName = TEXT("fixed"); 
+				}
+				break;
+			default:           UE_LOG(LogShaders, Fatal,TEXT("Unrecognized uniform buffer struct member base type."));
+			};
+
+			// Generate the type dimensions for vectors and matrices.
+			FString TypeDim;
+			uint32 HLSLMemberSize = 4;
+			if(Member.GetNumRows() > 1)
+			{
+				TypeDim = FString::Printf(TEXT("%ux%u"),Member.GetNumRows(),Member.GetNumColumns());
+
+				// Each row of a matrix is 16 byte aligned.
+				HLSLMemberSize = (Member.GetNumRows() - 1) * 16 + Member.GetNumColumns() * 4;
+			}
+			else if(Member.GetNumColumns() > 1)
+			{
+				TypeDim = FString::Printf(TEXT("%u"),Member.GetNumColumns());
+				HLSLMemberSize = Member.GetNumColumns() * 4;
+			}
+
+			// Array elements are 16 byte aligned.
+			if(Member.GetNumElements() > 0)
+			{
+				HLSLMemberSize = (Member.GetNumElements() - 1) * Align(HLSLMemberSize,16) + HLSLMemberSize;
+			}
+
+			const uint32 AbsoluteMemberOffset = StructOffset + Member.GetOffset();
+
+			// If the HLSL offset doesn't match the C++ offset, generate padding to fix it.
+			if(HLSLBaseOffset != AbsoluteMemberOffset)
+			{
+				check(HLSLBaseOffset < AbsoluteMemberOffset);
+				while(HLSLBaseOffset < AbsoluteMemberOffset)
+				{
+					Decl.ConstantBufferMembers += FString::Printf(TEXT("\t%s PrePadding_%s%u;\r\n"), *PreviousBaseTypeName, *NamePrefix, HLSLBaseOffset);
+					HLSLBaseOffset += 4;
+				};
+				check(HLSLBaseOffset == AbsoluteMemberOffset);
+			}
+			PreviousBaseTypeName = BaseTypeName;
+			HLSLBaseOffset = AbsoluteMemberOffset + HLSLMemberSize;
+
+			// Generate the member declaration.
+			FString ParameterName = FString::Printf(TEXT("%s%s"),*NamePrefix,Member.GetName());
+			Decl.ConstantBufferMembers += FString::Printf(TEXT("\t%s%s %s%s;\r\n"),*BaseTypeName,*TypeDim,*ParameterName,*ArrayDim);
+			Decl.StructMembers += FString::Printf(TEXT("\t%s%s %s%s;\r\n"),*BaseTypeName,*TypeDim,Member.GetName(),*ArrayDim);
+			Decl.Initializer += FString::Printf(TEXT(",%s"),*ParameterName);
+		}
+	}
+
+	for (int32 MemberIndex = 0; MemberIndex < StructMembers.Num(); ++MemberIndex)
+	{
+		const FShaderParametersMetadata::FMember& Member = StructMembers[MemberIndex];
+
+		if (IsUniformBufferResourceType(Member.GetBaseType()))
+		{
+			check(Member.GetBaseType() != UBMT_GRAPH_TRACKED_SRV || Member.GetBaseType() != UBMT_GRAPH_TRACKED_UAV);
+			if (Member.GetBaseType() == UBMT_SRV)
+			{
+				// TODO: handle arrays?
+				FString ParameterName = FString::Printf(TEXT("%s%s"),*NamePrefix,Member.GetName());
+				Decl.ResourceMembers += FString::Printf(TEXT("PLATFORM_SUPPORTS_SRV_UB_MACRO( %s %s; ) \r\n"), Member.GetShaderType(), *ParameterName);
+				Decl.StructMembers += FString::Printf(TEXT("\tPLATFORM_SUPPORTS_SRV_UB_MACRO( %s %s; ) \r\n"), Member.GetShaderType(), Member.GetName());
+				Decl.Initializer += FString::Printf(TEXT(" PLATFORM_SUPPORTS_SRV_UB_MACRO( ,%s ) "), *ParameterName);
+			}
+			else
+			{
+				// TODO: handle arrays?
+				FString ParameterName = FString::Printf(TEXT("%s%s"),*NamePrefix,Member.GetName());
+				Decl.ResourceMembers += FString::Printf(TEXT("%s %s;\r\n"), Member.GetShaderType(), *ParameterName);
+				Decl.StructMembers += FString::Printf(TEXT("\t%s %s;\r\n"), Member.GetShaderType(), Member.GetName());
+				Decl.Initializer += FString::Printf(TEXT(",%s"), *ParameterName);
+			}
+		}
+	}
+
+	if (Decl.Initializer[OpeningBraceLocPlusOne] == TEXT(','))
+	{
+		Decl.Initializer[OpeningBraceLocPlusOne] = TEXT(' ');
+	}
+}
+
+/** Creates a HLSL declaration of a uniform buffer with the given structure. */
+static FString CreateHLSLUniformBufferDeclaration(const TCHAR* Name,const FShaderParametersMetadata& UniformBufferStruct)
+{
+	// If the uniform buffer has no members, we don't want to write out anything.  Shader compilers throw errors when faced with empty cbuffers and structs.
+	if (UniformBufferStruct.GetMembers().Num() > 0)
+	{
+		FString NamePrefix(FString(Name) + FString(TEXT("_")));
+		FUniformBufferDecl Decl;
+		uint32 HLSLBaseOffset = 0;
+		CreateHLSLUniformBufferStructMembersDeclaration(UniformBufferStruct, NamePrefix, 0, Decl, HLSLBaseOffset);
+
+		return FString::Printf(
+			TEXT("#ifndef __UniformBuffer_%s_Definition__\r\n")
+			TEXT("#define __UniformBuffer_%s_Definition__\r\n")
+			TEXT("cbuffer %s\r\n")
+			TEXT("{\r\n")
+			TEXT("%s")
+			TEXT("}\r\n")
+				TEXT("%s")
+				TEXT("static const struct\r\n")
+				TEXT("{\r\n")
+				TEXT("%s")
+				TEXT("} %s = {%s};\r\n")
+			TEXT("#endif\r\n"),
+			Name,
+			Name,
+			Name,
+			*Decl.ConstantBufferMembers,
+			*Decl.ResourceMembers,
+			*Decl.StructMembers,
+			Name,
+			*Decl.Initializer
+			);
+	}
+
+	return FString(TEXT("\n"));
+}
+
+RENDERCORE_API void CreateUniformBufferShaderDeclaration(const TCHAR* Name,const FShaderParametersMetadata& UniformBufferStruct, FString& OutDeclaration)
+{
+	OutDeclaration = CreateHLSLUniformBufferDeclaration(Name, UniformBufferStruct);
+}
+
+RENDERCORE_API void CacheUniformBufferIncludes(TMap<const TCHAR*,FCachedUniformBufferDeclaration>& Cache)
+{
+	for (TMap<const TCHAR*,FCachedUniformBufferDeclaration>::TIterator It(Cache); It; ++It)
+	{
+		FCachedUniformBufferDeclaration& BufferDeclaration = It.Value();
+		check(BufferDeclaration.Declaration.Get() == NULL);
+
+		for (TLinkedList<FShaderParametersMetadata*>::TIterator StructIt(FShaderParametersMetadata::GetStructList()); StructIt; StructIt.Next())
+		{
+			if (It.Key() == StructIt->GetShaderVariableName())
+			{
+				FString* NewDeclaration = new FString();
+				CreateUniformBufferShaderDeclaration(StructIt->GetShaderVariableName(), **StructIt, *NewDeclaration);
+				check(!NewDeclaration->IsEmpty());
+				BufferDeclaration.Declaration = MakeShareable(NewDeclaration);
+				break;
+			}
+		}
+	}
+}
+
+void FShaderType::AddReferencedUniformBufferIncludes(FShaderCompilerEnvironment& OutEnvironment, FString& OutSourceFilePrefix, EShaderPlatform Platform)
+{
+	// Cache uniform buffer struct declarations referenced by this shader type's files
+	if (!bCachedUniformBufferStructDeclarations)
+	{
+		CacheUniformBufferIncludes(ReferencedUniformBufferStructsCache);
+		bCachedUniformBufferStructDeclarations = true;
+	}
+
+	FString UniformBufferIncludes;
+
+	for (TMap<const TCHAR*,FCachedUniformBufferDeclaration>::TIterator It(ReferencedUniformBufferStructsCache); It; ++It)
+	{
+		check(It.Value().Declaration.Get() != NULL);
+		check(!It.Value().Declaration.Get()->IsEmpty());
+		UniformBufferIncludes += FString::Printf(TEXT("#include \"/Engine/Generated/UniformBuffers/%s.ush\"") LINE_TERMINATOR, It.Key());
+		OutEnvironment.IncludeVirtualPathToExternalContentsMap.Add(
+			*FString::Printf(TEXT("/Engine/Generated/UniformBuffers/%s.ush"),It.Key()),
+			It.Value().Declaration
+			);
+
+		for (TLinkedList<FShaderParametersMetadata*>::TIterator StructIt(FShaderParametersMetadata::GetStructList()); StructIt; StructIt.Next())
+		{
+			if (It.Key() == StructIt->GetShaderVariableName())
+			{
+				StructIt->AddResourceTableEntries(OutEnvironment.ResourceTableMap, OutEnvironment.ResourceTableLayoutHashes);
+			}
+		}
+	}
+
+	FString& GeneratedUniformBuffersInclude = OutEnvironment.IncludeVirtualPathToContentsMap.FindOrAdd("/Engine/Generated/GeneratedUniformBuffers.ush");
+	GeneratedUniformBuffersInclude.Append(UniformBufferIncludes);
+
+	ERHIFeatureLevel::Type MaxFeatureLevel = GetMaxSupportedFeatureLevel(Platform);
+	if (MaxFeatureLevel >= ERHIFeatureLevel::SM4)
+	{
+		OutEnvironment.SetDefine(TEXT("PLATFORM_SUPPORTS_SRV_UB"), TEXT("1"));
+	}
+}
+
+void FShaderType::DumpDebugInfo()
+{
+	UE_LOG(LogConsoleResponse, Display, TEXT("----------------------------- GLobalShader %s"), GetName());
+	UE_LOG(LogConsoleResponse, Display, TEXT("               :Target %s"), GetShaderFrequencyString(GetFrequency()));
+	UE_LOG(LogConsoleResponse, Display, TEXT("               :TotalPermutationCount %d"), TotalPermutationCount);
+#if WITH_EDITOR
+	UE_LOG(LogConsoleResponse, Display, TEXT("               :SourceHash %s"), *GetSourceHash(GMaxRHIShaderPlatform).ToString());
+#endif
+	switch (ShaderTypeForDynamicCast)
+	{
+	case EShaderTypeForDynamicCast::Global:
+		UE_LOG(LogConsoleResponse, Display, TEXT("               :ShaderType Global"));
+		break;
+	case EShaderTypeForDynamicCast::Material:
+		UE_LOG(LogConsoleResponse, Display, TEXT("               :ShaderType Material"));
+		break;
+	case EShaderTypeForDynamicCast::MeshMaterial:
+		UE_LOG(LogConsoleResponse, Display, TEXT("               :ShaderType MeshMaterial"));
+		break;
+	case EShaderTypeForDynamicCast::Niagara:
+		UE_LOG(LogConsoleResponse, Display, TEXT("               :ShaderType Niagara"));
+		break;
+	}
+
+	UE_LOG(LogConsoleResponse, Display, TEXT("  --- %d shaders"), ShaderIdMap.Num());
+	int32 Index = 0;
+	for (auto& KeyValue : ShaderIdMap)
+	{
+		UE_LOG(LogConsoleResponse, Display, TEXT("    --- shader %d"), Index);
+		FShader* Shader = KeyValue.Value;
+		Shader->DumpDebugInfo();
+		Index++;
+	}
+
+}
+
+void FShaderType::GetShaderStableKeyParts(FStableShaderKeyAndValue& SaveKeyVal)
+{
+#if WITH_EDITOR
+	static FName NAME_Global(TEXT("Global"));
+	static FName NAME_Material(TEXT("Material"));
+	static FName NAME_MeshMaterial(TEXT("MeshMaterial"));
+	static FName NAME_Niagara(TEXT("Niagara"));
+	switch (ShaderTypeForDynamicCast)
+	{
+	case EShaderTypeForDynamicCast::Global:
+		SaveKeyVal.ShaderClass = NAME_Global;
+		break;
+	case EShaderTypeForDynamicCast::Material:
+		SaveKeyVal.ShaderClass = NAME_Material;
+		break;
+	case EShaderTypeForDynamicCast::MeshMaterial:
+		SaveKeyVal.ShaderClass = NAME_MeshMaterial;
+		break;
+	case EShaderTypeForDynamicCast::Niagara:
+		SaveKeyVal.ShaderClass = NAME_Niagara;
+		break;
+	}
+	SaveKeyVal.ShaderType = FName(GetName() ? GetName() : TEXT("null"));
+#endif
+}
+
+void FShaderType::SaveShaderStableKeys(EShaderPlatform TargetShaderPlatform)
+{
+#if WITH_EDITOR
+	FStableShaderKeyAndValue SaveKeyVal;
+	if (ShaderIdMap.Num())
+	{
+		for (auto& KeyValue : ShaderIdMap)
+		{
+			FShader* Shader = KeyValue.Value;
+			check(Shader);
+			check(Shader->Type == this);
+			Shader->SaveShaderStableKeys(TargetShaderPlatform, SaveKeyVal);
+		}
+	}
+#endif
+}
+
+
+void FVertexFactoryType::AddReferencedUniformBufferIncludes(FShaderCompilerEnvironment& OutEnvironment, FString& OutSourceFilePrefix, EShaderPlatform Platform)
+{
+	// Cache uniform buffer struct declarations referenced by this shader type's files
+	if (!bCachedUniformBufferStructDeclarations)
+	{
+		CacheUniformBufferIncludes(ReferencedUniformBufferStructsCache);
+		bCachedUniformBufferStructDeclarations = true;
+	}
+
+	FString UniformBufferIncludes;
+
+	for (TMap<const TCHAR*,FCachedUniformBufferDeclaration>::TIterator It(ReferencedUniformBufferStructsCache); It; ++It)
+	{
+		check(It.Value().Declaration.Get() != NULL);
+		check(!It.Value().Declaration.Get()->IsEmpty());
+		UniformBufferIncludes += FString::Printf(TEXT("#include \"/Engine/Generated/UniformBuffers/%s.ush\"") LINE_TERMINATOR, It.Key());
+		OutEnvironment.IncludeVirtualPathToExternalContentsMap.Add(
+			*FString::Printf(TEXT("/Engine/Generated/UniformBuffers/%s.ush"),It.Key()),
+			It.Value().Declaration
+		);
+
+		for (TLinkedList<FShaderParametersMetadata*>::TIterator StructIt(FShaderParametersMetadata::GetStructList()); StructIt; StructIt.Next())
+		{
+			if (It.Key() == StructIt->GetShaderVariableName())
+			{
+				StructIt->AddResourceTableEntries(OutEnvironment.ResourceTableMap, OutEnvironment.ResourceTableLayoutHashes);
+			}
+		}
+	}
+
+	FString& GeneratedUniformBuffersInclude = OutEnvironment.IncludeVirtualPathToContentsMap.FindOrAdd("/Engine/Generated/GeneratedUniformBuffers.ush");
+	GeneratedUniformBuffersInclude.Append(UniformBufferIncludes);
+
+	ERHIFeatureLevel::Type MaxFeatureLevel = GetMaxSupportedFeatureLevel(Platform);
+	if (MaxFeatureLevel >= ERHIFeatureLevel::SM4)
+	{
+		OutEnvironment.SetDefine(TEXT("PLATFORM_SUPPORTS_SRV_UB"), TEXT("1"));
+	}
+}
