Merging //UE4/Dev-Rendering-Graph@4492585 to Dev-Rendering (//UE4/Dev-Rendering)

diff --git a/Engine/Source/Runtime/RenderCore/Private/ShaderPipelineCache.cpp b/Engine/Source/Runtime/RenderCore/Private/ShaderPipelineCache.cpp
new file mode 100644
--- /dev/null
+++ b/Engine/Source/Runtime/RenderCore/Private/ShaderPipelineCache.cpp
@@ -0,0 +1,1366 @@
+// Copyright 1998-2018 Epic Games, Inc. All Rights Reserved.
+
+/**
+ * Shader Pipeline Precompilation Cache
+ * Precompilation half of the shader pipeline cache, which builds on the runtime RHI pipeline cache.
+ */
+ 
+#include "ShaderPipelineCache.h"
+#include "RenderingThread.h"
+#include "HAL/FileManager.h"
+#include "Misc/Paths.h"
+#include "Serialization/MemoryWriter.h"
+#include "Serialization/ArchiveSaveCompressedProxy.h"
+#include "Serialization/ArchiveLoadCompressedProxy.h"
+#include "Serialization/CustomVersion.h"
+#include "Serialization/MemoryReader.h"
+#include "Shader.h"
+#include "Misc/EngineVersion.h"
+#include "PipelineStateCache.h"
+#include "PipelineFileCache.h"
+#include "Misc/ScopeRWLock.h"
+#include "Misc/CoreDelegates.h"
+#include "ShaderCodeLibrary.h"
+#include "TickableObjectRenderThread.h"
+#include "Misc/ConfigCacheIni.h"
+#include "Async/AsyncFileHandle.h"
+
+DECLARE_DWORD_ACCUMULATOR_STAT(TEXT("Outstanding Tasks"), STAT_ShaderPipelineTaskCount, STATGROUP_PipelineStateCache );
+DECLARE_DWORD_ACCUMULATOR_STAT(TEXT("Waiting Tasks"), STAT_ShaderPipelineWaitingTaskCount, STATGROUP_PipelineStateCache );
+DECLARE_DWORD_ACCUMULATOR_STAT(TEXT("Active Tasks"), STAT_ShaderPipelineActiveTaskCount, STATGROUP_PipelineStateCache );
+DECLARE_MEMORY_STAT(TEXT("Pre-Compile Memory"), STAT_PreCompileMemory, STATGROUP_PipelineStateCache);
+DECLARE_CYCLE_STAT(TEXT("Pre-Compile Time"),STAT_PreCompileTime,STATGROUP_PipelineStateCache);
+DECLARE_FLOAT_ACCUMULATOR_STAT(TEXT("Total Pre-Compile Time"),STAT_PreCompileTotalTime,STATGROUP_PipelineStateCache);
+DECLARE_DWORD_ACCUMULATOR_STAT(TEXT("Total Pipelines Pre-Compiled"), STAT_PreCompileShadersTotal, STATGROUP_PipelineStateCache);
+DECLARE_DWORD_COUNTER_STAT(TEXT("# Pipelines Pre-Compiled"), STAT_PreCompileShadersNum, STATGROUP_PipelineStateCache);
+DECLARE_DWORD_ACCUMULATOR_STAT(TEXT("Total Batches Pre-Compiled"), STAT_PreCompileBatchTotal, STATGROUP_PipelineStateCache);
+DECLARE_DWORD_COUNTER_STAT(TEXT("# Batches Pre-Compiled"), STAT_PreCompileBatchNum, STATGROUP_PipelineStateCache);
+
+namespace FShaderPipelineCacheConstants
+{
+	static TCHAR const* SectionHeading = TEXT("ShaderPipelineCache.CacheFile");
+	static TCHAR const* LastOpenedKey = TEXT("LastOpened");
+	static TCHAR const* SortOrderKey = TEXT("SortOrder");
+	static TCHAR const* GameVersionKey = TEXT("GameVersion");
+}
+
+static TAutoConsoleVariable<int32> CVarPSOFileCacheBackgroundBatchSize(
+														  TEXT("r.ShaderPipelineCache.BackgroundBatchSize"),
+														  1,
+														  TEXT("Set the number of PipelineStateObjects to compile in a single batch operation when compiling in the background. Defaults to a maximum of 1 per frame, due to async. file IO it is less in practice."),
+														  ECVF_Default | ECVF_RenderThreadSafe
+														  );
+static TAutoConsoleVariable<int32> CVarPSOFileCacheBatchSize(
+														   TEXT("r.ShaderPipelineCache.BatchSize"),
+#if PLATFORM_MAC
+														   16, // On Mac, where we have many more PSOs to preload due to different video settings 16 works better than 50
+#else
+														   50,
+#endif
+														   TEXT("Set the number of PipelineStateObjects to compile in a single batch operation when compiling takes priority. Defaults to a maximum of 50 per frame, due to async. file IO it is less in practice."),
+														   ECVF_Default | ECVF_RenderThreadSafe
+														   );
+static TAutoConsoleVariable<float> CVarPSOFileCacheBackgroundBatchTime(
+														  TEXT("r.ShaderPipelineCache.BackgroundBatchTime"),
+														  0.0f,
+														  TEXT("The target time (in ms) to spend precompiling each frame when in the background or 0.0 to disable. When precompiling is faster the batch size will grow and when slower will shrink to attempt to occupy the full amount. Defaults to 0.0 (off)."),
+														  ECVF_Default | ECVF_RenderThreadSafe
+														  );
+static TAutoConsoleVariable<float> CVarPSOFileCacheBatchTime(
+														   TEXT("r.ShaderPipelineCache.BatchTime"),
+														   16.0f,
+														   TEXT("The target time (in ms) to spend precompiling each frame when compiling takes priority or 0.0 to disable. When precompiling is faster the batch size will grow and when slower will shrink to attempt to occupy the full amount. Defaults to 16.0 (max. ms per-frame of precompilation)."),
+														   ECVF_Default | ECVF_RenderThreadSafe
+														   );
+static TAutoConsoleVariable<int32> CVarPSOFileCacheSaveAfterPSOsLogged(
+														   TEXT("r.ShaderPipelineCache.SaveAfterPSOsLogged"),
+#if !UE_BUILD_SHIPPING
+														   100,
+#else
+														   0,
+#endif
+														   TEXT("Set the number of PipelineStateObjects to log before automatically saving. 0 will disable automatic saving. Shipping defaults to 0, otherwise default is 100."),
+														   ECVF_Default | ECVF_RenderThreadSafe
+														   );
+ 
+static TAutoConsoleVariable<int32> CVarPSOFileCacheAutoSaveTime(
+															TEXT("r.ShaderPipelineCache.AutoSaveTime"),
+															30,
+															TEXT("Set the time where any logged PSO's will be saved if the number is < r.ShaderPipelineCache.SaveAfterPSOsLogged. Disabled when r.ShaderPipelineCache.SaveAfterPSOsLogged is 0"),
+															ECVF_Default | ECVF_RenderThreadSafe
+														);
+
+static TAutoConsoleVariable<int32> CVarPSOFileCacheAutoSaveTimeBoundPSO(
+	TEXT("r.ShaderPipelineCache.AutoSaveTimeBoundPSO"),
+	10,
+	TEXT("Set the time where any logged PSO's will be saved when -logpso is on th ecommand line."),
+	ECVF_Default | ECVF_RenderThreadSafe
+);
+
+static TAutoConsoleVariable<int32> CVarPSOFileCacheSaveBoundPSOLog(
+																 TEXT("r.ShaderPipelineCache.SaveBoundPSOLog"),
+																 (int32)0,
+																 TEXT("If > 0 then a log of all bound PSOs for this run of the program will be saved to a writable user cache file. Defaults to 0 but is forced on with -logpso."),
+																 ECVF_Default | ECVF_RenderThreadSafe
+																 );
+
+static bool GetShaderPipelineCacheSaveBoundPSOLog()
+{
+	static bool bOnce = false;
+	static bool bCmdLineForce = false;
+	if (!bOnce)
+	{
+		bOnce = true;
+		bCmdLineForce = FParse::Param(FCommandLine::Get(), TEXT("logpso"));
+	}
+	return (bCmdLineForce || CVarPSOFileCacheSaveBoundPSOLog.GetValueOnAnyThread() == 1);
+}
+
+static bool GetPSOFileCacheSaveUserCache()
+{
+	static const auto CVarPSOFileCacheSaveUserCache = IConsoleManager::Get().FindConsoleVariable(TEXT("r.ShaderPipelineCache.SaveUserCache"));
+	return (CVarPSOFileCacheSaveUserCache && CVarPSOFileCacheSaveUserCache->GetInt() > 0);
+}
+
+void ConsoleCommandLoadPipelineFileCache(const TArray< FString >& Args)
+{
+	FShaderPipelineCache::ClosePipelineFileCache();
+	FString Name = FApp::GetProjectName();
+	if (Args.Num() > 0)
+	{
+		Name = Args[0];
+	}
+	FShaderPipelineCache::OpenPipelineFileCache(Name, GMaxRHIShaderPlatform);
+}
+
+void ConsoleCommandSavePipelineFileCache()
+{
+	if (GetShaderPipelineCacheSaveBoundPSOLog())
+	{
+		FShaderPipelineCache::SavePipelineFileCache(FPipelineFileCache::SaveMode::BoundPSOsOnly);
+	}
+	if (GetPSOFileCacheSaveUserCache())
+	{
+		FShaderPipelineCache::SavePipelineFileCache(FPipelineFileCache::SaveMode::SortedBoundPSOs);
+	}
+}
+
+void ConsoleCommandClosePipelineFileCache()
+{
+	FShaderPipelineCache::ClosePipelineFileCache();
+}
+
+void ConsoleCommandSwitchModePipelineCacheCmd(const TArray< FString >& Args)
+{
+    if (Args.Num() > 0)
+    {
+        FString Mode = Args[0];
+        if (Mode == TEXT("Pause"))
+        {
+            FShaderPipelineCache::PauseBatching();
+        }
+        else if (Mode == TEXT("Background"))
+        {
+            FShaderPipelineCache::SetBatchMode(FShaderPipelineCache::BatchMode::Background);
+            FShaderPipelineCache::ResumeBatching();
+        }
+        else if (Mode == TEXT("Fast"))
+        {
+            FShaderPipelineCache::SetBatchMode(FShaderPipelineCache::BatchMode::Fast);
+            FShaderPipelineCache::ResumeBatching();
+        }
+    }
+}
+
+static FAutoConsoleCommand LoadPipelineCacheCmd(
+												TEXT("r.ShaderPipelineCache.Open"),
+												TEXT("Takes the desired filename to open and then loads the pipeline file cache."),
+												FConsoleCommandWithArgsDelegate::CreateStatic(ConsoleCommandLoadPipelineFileCache)
+												);
+
+static FAutoConsoleCommand SavePipelineCacheCmd(
+										   		TEXT("r.ShaderPipelineCache.Save"),
+										   		TEXT("Save the current pipeline file cache."),
+										   		FConsoleCommandDelegate::CreateStatic(ConsoleCommandSavePipelineFileCache)
+										   		);
+
+static FAutoConsoleCommand ClosePipelineCacheCmd(
+												TEXT("r.ShaderPipelineCache.Close"),
+												TEXT("Close the current pipeline file cache."),
+												FConsoleCommandDelegate::CreateStatic(ConsoleCommandClosePipelineFileCache)
+												);
+
+static FAutoConsoleCommand SwitchModePipelineCacheCmd(
+                                                TEXT("r.ShaderPipelineCache.SetBatchMode"),
+                                                TEXT("Sets the compilation batch mode, which should be one of:\n\tPause: Suspend precompilation.\n\tBackground: Low priority precompilation.\n\tFast: High priority precompilation."),
+                                                FConsoleCommandWithArgsDelegate::CreateStatic(ConsoleCommandSwitchModePipelineCacheCmd)
+                                                );
+ 
+ 
+class FShaderPipelineCacheArchive final : public FArchive
+{
+public:
+	FShaderPipelineCacheArchive()
+	{
+	}
+	virtual ~FShaderPipelineCacheArchive()
+	{
+	}
+
+	virtual bool AttachExternalReadDependency(FExternalReadCallback& ReadCallback) override final
+	{
+		ExternalReadDependencies.Add(ReadCallback);
+		return true;
+	}
+
+	bool PollExternalReadDependencies()
+	{
+		for (uint32 i = 0; i < (uint32)ExternalReadDependencies.Num(); )
+		{
+			FExternalReadCallback& ReadCallback = ExternalReadDependencies[i];
+			bool bFinished = ReadCallback(-1.0);
+			if (bFinished)
+			{
+				ExternalReadDependencies.RemoveAt(i);
+			}
+			else
+			{
+				++i;
+			}
+		}
+		return (ExternalReadDependencies.Num() == 0);
+	}
+	
+	void BlockingWaitComplete()
+	{
+		for (uint32 i = 0; i < (uint32)ExternalReadDependencies.Num(); ++i)
+		{
+			FExternalReadCallback& ReadCallback = ExternalReadDependencies[i];
+			ReadCallback(0.0);
+		}
+	}
+
+private:
+	/**
+	*  List of external read dependencies that must be finished to load this package
+	*/
+	TArray<FExternalReadCallback> ExternalReadDependencies;
+};
+
+FShaderPipelineCache* FShaderPipelineCache::ShaderPipelineCache = nullptr;
+
+FShaderPipelineCache::FShaderCacheOpenedDelegate FShaderPipelineCache::OnCachedOpened;
+FShaderPipelineCache::FShaderCacheClosedDelegate FShaderPipelineCache::OnCachedClosed;
+FShaderPipelineCache::FShaderPrecompilationBeginDelegate FShaderPipelineCache::OnPrecompilationBegin;
+FShaderPipelineCache::FShaderPrecompilationCompleteDelegate FShaderPipelineCache::OnPrecompilationComplete;
+
+static void PipelineStateCacheOnAppDeactivate()
+{
+	if (GetShaderPipelineCacheSaveBoundPSOLog())
+	{
+		FShaderPipelineCache::SavePipelineFileCache(FPipelineFileCache::SaveMode::BoundPSOsOnly);
+	}
+	if (GetPSOFileCacheSaveUserCache())
+	{
+		FShaderPipelineCache::SavePipelineFileCache(FPipelineFileCache::SaveMode::Incremental);
+	}
+}
+
+int32 FShaderPipelineCache::GetGameVersionForPSOFileCache()
+{
+	int32 GameVersion = (int32)FEngineVersion::Current().GetChangelist();
+	GConfig->GetInt(FShaderPipelineCacheConstants::SectionHeading, FShaderPipelineCacheConstants::GameVersionKey, GameVersion, *GGameIni);
+	return GameVersion;
+}
+
+
+void FShaderPipelineCache::Initialize(EShaderPlatform Platform)
+{
+	check(ShaderPipelineCache == nullptr);
+	
+	if(FShaderCodeLibrary::IsEnabled())
+	{
+		FPipelineFileCache::Initialize(GetGameVersionForPSOFileCache());
+		ShaderPipelineCache = new FShaderPipelineCache(Platform);
+	}
+}
+
+void FShaderPipelineCache::Shutdown(void)
+{
+	if (ShaderPipelineCache)
+	{
+		delete ShaderPipelineCache;
+		ShaderPipelineCache = nullptr;
+	}
+}
+
+void FShaderPipelineCache::PauseBatching()
+{
+	if (ShaderPipelineCache)
+		ShaderPipelineCache->bPaused = true;
+}
+
+void FShaderPipelineCache::SetBatchMode(BatchMode Mode)
+{
+	if (ShaderPipelineCache)
+	{
+		switch (Mode)
+		{
+			case BatchMode::Fast:
+			{
+				ShaderPipelineCache->BatchSize = CVarPSOFileCacheBatchSize.GetValueOnAnyThread();
+				ShaderPipelineCache->BatchTime = CVarPSOFileCacheBatchTime.GetValueOnAnyThread();
+				break;
+			}
+			case BatchMode::Background:
+			default:
+			{
+				ShaderPipelineCache->BatchSize = CVarPSOFileCacheBackgroundBatchSize.GetValueOnAnyThread();
+				ShaderPipelineCache->BatchTime = CVarPSOFileCacheBackgroundBatchTime.GetValueOnAnyThread();
+				break;
+			}
+		}
+	}
+}
+
+void FShaderPipelineCache::ResumeBatching()
+{
+	if (ShaderPipelineCache)
+		ShaderPipelineCache->bPaused = false;
+}
+
+uint32 FShaderPipelineCache::NumPrecompilesRemaining()
+{
+	uint32 NumRemaining = 0;
+	
+	if (ShaderPipelineCache)
+	{
+        int64 NumActiveTasksRemaining = FMath::Max(0ll, FPlatformAtomics::AtomicRead(&ShaderPipelineCache->TotalActiveTasks));
+        int64 NumWaitingTasksRemaining = FMath::Max(0ll, FPlatformAtomics::AtomicRead(&ShaderPipelineCache->TotalWaitingTasks));
+		NumRemaining = (uint32)(NumActiveTasksRemaining + NumWaitingTasksRemaining);
+	}
+
+	return NumRemaining;
+}
+
+uint32 FShaderPipelineCache::NumPrecompilesActive()
+{
+    uint32 NumRemaining = 0;
+    
+    if (ShaderPipelineCache)
+    {
+        int64 NumTasksRemaining = FPlatformAtomics::AtomicRead(&ShaderPipelineCache->TotalActiveTasks);
+        if(NumTasksRemaining > 0)
+        {
+            NumRemaining = (uint32) NumTasksRemaining;
+        }
+    }
+    
+    return NumRemaining;
+}
+
+bool FShaderPipelineCache::OpenPipelineFileCache(EShaderPlatform Platform)
+{
+	bool bFileOpen = false;
+	if (GConfig)
+	{
+		FString LastOpenedName;
+		if ((GConfig->GetString(FShaderPipelineCacheConstants::SectionHeading, FShaderPipelineCacheConstants::LastOpenedKey, LastOpenedName, *GGameUserSettingsIni) || GConfig->GetString(FShaderPipelineCacheConstants::SectionHeading, FShaderPipelineCacheConstants::LastOpenedKey, LastOpenedName, *GGameIni)) && LastOpenedName.Len())
+		{
+			bFileOpen = OpenPipelineFileCache(LastOpenedName, Platform);
+		}
+	}
+
+	if (!bFileOpen)
+	{
+		bFileOpen = OpenPipelineFileCache(FApp::GetProjectName(), Platform);
+	}
+
+	return bFileOpen;
+}
+
+bool FShaderPipelineCache::OpenPipelineFileCache(FString const& Name, EShaderPlatform Platform)
+{
+	if (ShaderPipelineCache)
+		return ShaderPipelineCache->Open(Name, Platform);
+	else
+		return false;
+}
+
+bool FShaderPipelineCache::SavePipelineFileCache(FPipelineFileCache::SaveMode Mode)
+{
+	if (ShaderPipelineCache)
+		return ShaderPipelineCache->Save(Mode);
+	else
+		return false;
+}
+
+void FShaderPipelineCache::ClosePipelineFileCache()
+{
+	if (ShaderPipelineCache)
+		ShaderPipelineCache->Close();
+}
+
+void FShaderPipelineCache::ShaderLibraryStateChanged(ELibraryState State, EShaderPlatform Platform, FString const& Name)
+{
+    if (ShaderPipelineCache)
+        ShaderPipelineCache->OnShaderLibraryStateChanged(State, Platform, Name);
+}
+
+bool FShaderPipelineCache::Precompile(FRHICommandListImmediate& RHICmdList, EShaderPlatform Platform, FPipelineCacheFileFormatPSO const& PSO)
+{
+	INC_DWORD_STAT(STAT_PreCompileShadersTotal);
+	INC_DWORD_STAT(STAT_PreCompileShadersNum);
+    
+    uint64 StartTime = FPlatformTime::Cycles64();
+		
+	bool bOk = false;
+	
+	TArray<uint8> DummyCode;
+
+	TArray<FSHAHash> DeferredReleaseHashes;
+	bool bDeferRelease = GRHILazyShaderCodeLoading && IsRunningRHIInSeparateThread();
+
+	auto ReleaseShaderCode = [bDeferRelease, &DeferredReleaseHashes](const FSHAHash& DeferHash)
+	{
+		if (bDeferRelease)
+		{
+			DeferredReleaseHashes.Add(DeferHash);
+		}
+		else
+		{
+			FShaderCodeLibrary::ReleaseShaderCode(DeferHash);
+		}
+	};
+
+	
+	if(FPipelineCacheFileFormatPSO::DescriptorType::Graphics == PSO.Type)
+	{
+		FGraphicsPipelineStateInitializer GraphicsInitializer;
+		
+		auto VertexDesc = RHICmdList.CreateVertexDeclaration(PSO.GraphicsDesc.VertexDescriptor);
+		GraphicsInitializer.BoundShaderState.VertexDeclarationRHI = VertexDesc;
+		
+		FVertexShaderRHIRef VertexShader;
+		if (PSO.GraphicsDesc.VertexShader != FSHAHash())
+		{
+			VertexShader = FShaderCodeLibrary::CreateVertexShader(Platform, PSO.GraphicsDesc.VertexShader, DummyCode);
+			GraphicsInitializer.BoundShaderState.VertexShaderRHI = VertexShader;
+            
+            ReleaseShaderCode(PSO.GraphicsDesc.VertexShader);
+		}
+
+		FHullShaderRHIRef HullShader;
+		if (PSO.GraphicsDesc.HullShader != FSHAHash())
+		{
+			HullShader = FShaderCodeLibrary::CreateHullShader(Platform, PSO.GraphicsDesc.HullShader, DummyCode);
+			GraphicsInitializer.BoundShaderState.HullShaderRHI = HullShader;
+            
+            ReleaseShaderCode(PSO.GraphicsDesc.HullShader);
+		}
+
+		FDomainShaderRHIRef DomainShader;
+		if (PSO.GraphicsDesc.DomainShader != FSHAHash())
+		{
+			DomainShader = FShaderCodeLibrary::CreateDomainShader(Platform, PSO.GraphicsDesc.DomainShader, DummyCode);
+			GraphicsInitializer.BoundShaderState.DomainShaderRHI = DomainShader;
+            
+            ReleaseShaderCode(PSO.GraphicsDesc.DomainShader);
+		}
+
+		FPixelShaderRHIRef FragmentShader;
+		if (PSO.GraphicsDesc.FragmentShader != FSHAHash())
+		{
+			FragmentShader = FShaderCodeLibrary::CreatePixelShader(Platform, PSO.GraphicsDesc.FragmentShader, DummyCode);
+			GraphicsInitializer.BoundShaderState.PixelShaderRHI = FragmentShader;
+            
+            ReleaseShaderCode(PSO.GraphicsDesc.FragmentShader);
+		}
+
+		FGeometryShaderRHIRef GeometryShader;
+		if (PSO.GraphicsDesc.GeometryShader != FSHAHash())
+		{
+			GeometryShader = FShaderCodeLibrary::CreateGeometryShader(Platform, PSO.GraphicsDesc.GeometryShader, DummyCode);
+			GraphicsInitializer.BoundShaderState.GeometryShaderRHI = GeometryShader;
+            
+            ReleaseShaderCode(PSO.GraphicsDesc.GeometryShader);
+		}
+		
+		auto BlendState = RHICmdList.CreateBlendState(PSO.GraphicsDesc.BlendState);
+		GraphicsInitializer.BlendState = BlendState;
+		
+		auto RasterState = RHICmdList.CreateRasterizerState(PSO.GraphicsDesc.RasterizerState);
+		GraphicsInitializer.RasterizerState = RasterState;
+		
+		auto DepthState = RHICmdList.CreateDepthStencilState(PSO.GraphicsDesc.DepthStencilState);
+		GraphicsInitializer.DepthStencilState = DepthState;
+		
+		for (uint32 i = 0; i < MaxSimultaneousRenderTargets; ++i)
+		{
+			GraphicsInitializer.RenderTargetFormats[i] = PSO.GraphicsDesc.RenderTargetFormats[i];
+			GraphicsInitializer.RenderTargetFlags[i] = PSO.GraphicsDesc.RenderTargetFlags[i];
+		}
+		
+		GraphicsInitializer.RenderTargetsEnabled = PSO.GraphicsDesc.RenderTargetsActive;
+		GraphicsInitializer.NumSamples = PSO.GraphicsDesc.MSAASamples;
+		
+		GraphicsInitializer.DepthStencilTargetFormat = PSO.GraphicsDesc.DepthStencilFormat;
+		GraphicsInitializer.DepthStencilTargetFlag = PSO.GraphicsDesc.DepthStencilFlags;
+		GraphicsInitializer.DepthTargetLoadAction = PSO.GraphicsDesc.DepthLoad;
+		GraphicsInitializer.StencilTargetLoadAction = PSO.GraphicsDesc.StencilLoad;
+		GraphicsInitializer.DepthTargetStoreAction = PSO.GraphicsDesc.DepthStore;
+		GraphicsInitializer.StencilTargetStoreAction = PSO.GraphicsDesc.StencilStore;
+		
+		GraphicsInitializer.PrimitiveType = PSO.GraphicsDesc.PrimitiveType;
+		GraphicsInitializer.bFromPSOFileCache = true;
+		
+		// This indicates we do not want a fatal error if this compilation fails
+		// (ie, if this entry in the file cache is bad)
+		GraphicsInitializer.bFromPSOFileCache = 1;
+		
+		// Use SetGraphicsPipelineState to call down into PipelineStateCache and also handle the fallback case used by OpenGL.
+		SetGraphicsPipelineState(RHICmdList, GraphicsInitializer, EApplyRendertargetOption::DoNothing);
+		bOk = true;
+	}
+	else if(FPipelineCacheFileFormatPSO::DescriptorType::Compute == PSO.Type)
+	{
+		FComputeShaderRHIRef ComputeInitializer = FShaderCodeLibrary::CreateComputeShader(GMaxRHIShaderPlatform, PSO.ComputeDesc.ComputeShader, DummyCode);
+		FComputePipelineState* ComputeResult = PipelineStateCache::GetAndOrCreateComputePipelineState(RHICmdList, ComputeInitializer);
+		bOk = ComputeResult != nullptr;
+        
+        ReleaseShaderCode(PSO.ComputeDesc.ComputeShader);
+	}
+	else
+	{
+		check(false);
+	}
+    
+    if (bOk)
+    {
+        uint64 TimeDelta = FPlatformTime::Cycles64() - StartTime;
+        FPlatformAtomics::InterlockedIncrement(&TotalCompleteTasks);
+        FPlatformAtomics::InterlockedAdd(&TotalPrecompileTime, TimeDelta);
+    }
+	if (DeferredReleaseHashes.Num())
+	{
+		RHICmdList.EnqueueLambda([DeferredReleaseHashes](FRHICommandListImmediate&)
+		{
+			for (const FSHAHash& DelHash : DeferredReleaseHashes)
+			{
+				FShaderCodeLibrary::ReleaseShaderCode(DelHash);
+			}
+		});
+	}
+	return bOk;
+}
+
+void FShaderPipelineCache::PreparePipelineBatch(TArray<FPipelineCacheFileFormatPSORead*>& PipelineBatch)
+{
+	for(auto BatchIterator = PipelineBatch.CreateIterator();BatchIterator;++BatchIterator)
+	{
+		FPipelineCacheFileFormatPSORead* PSORead = *BatchIterator;
+		
+		FShaderPipelineCacheArchive* Archive = (FShaderPipelineCacheArchive*)PSORead->Ar;
+		
+		if (PSORead->bValid && ((Archive && Archive->PollExternalReadDependencies()) || PSORead->bReadCompleted))
+		{
+			check(PSORead->bReadCompleted);
+		
+			FPipelineCacheFileFormatPSO PSO;
+			
+			FMemoryReader Ar(PSORead->Data);
+			Ar.SetGameNetVer(FPipelineCacheFileFormatCurrentVersion);
+			Ar << PSO;
+			
+			// Create an FArchive that *only* registers read dependencies for tracking when the shaders are available
+			FShaderPipelineCacheArchive* ReadArchive = new FShaderPipelineCacheArchive;
+	
+			// Assume that the shader is present and the PSO can be compiled by default,
+			bool bOK = true;
+	
+            // Shaders required.
+            TSet<FSHAHash> Shaders;
+            
+            static FSHAHash EmptySHA;
+            
+            if (PSO.Type == FPipelineCacheFileFormatPSO::DescriptorType::Graphics)
+			{
+                // See if the shaders exist in the current code libraries, before trying to load the shader data
+                if (PSO.GraphicsDesc.VertexShader != EmptySHA)
+                {
+                    Shaders.Add(PSO.GraphicsDesc.VertexShader);
+                    bOK &= FShaderCodeLibrary::ContainsShaderCode(PSO.GraphicsDesc.VertexShader);
+                    UE_CLOG(!bOK, LogRHI, Verbose, TEXT("Failed to find VertexShader shader: %s"), *(PSO.GraphicsDesc.VertexShader.ToString()));
+					
+					if (PSO.GraphicsDesc.HullShader != EmptySHA)
+					{
+						Shaders.Add(PSO.GraphicsDesc.HullShader);
+						bOK &= FShaderCodeLibrary::ContainsShaderCode(PSO.GraphicsDesc.HullShader);
+						UE_CLOG(!bOK, LogRHI, Verbose, TEXT("Failed to find HullShader shader: %s"), *(PSO.GraphicsDesc.HullShader.ToString()));
+					}
+					if (PSO.GraphicsDesc.DomainShader != EmptySHA)
+					{
+						Shaders.Add(PSO.GraphicsDesc.DomainShader);
+						bOK &= FShaderCodeLibrary::ContainsShaderCode(PSO.GraphicsDesc.DomainShader);
+						UE_CLOG(!bOK, LogRHI, Verbose, TEXT("Failed to find DomainShader shader: %s"), *(PSO.GraphicsDesc.DomainShader.ToString()));
+					}
+					if (PSO.GraphicsDesc.FragmentShader != EmptySHA)
+					{
+						Shaders.Add(PSO.GraphicsDesc.FragmentShader);
+						bOK &= FShaderCodeLibrary::ContainsShaderCode(PSO.GraphicsDesc.FragmentShader);
+						UE_CLOG(!bOK, LogRHI, Verbose, TEXT("Failed to find FragmentShader shader: %s"), *(PSO.GraphicsDesc.FragmentShader.ToString()));
+					}
+					if (PSO.GraphicsDesc.GeometryShader != EmptySHA)
+					{
+						Shaders.Add(PSO.GraphicsDesc.GeometryShader);
+						bOK &= FShaderCodeLibrary::ContainsShaderCode(PSO.GraphicsDesc.GeometryShader);
+						UE_CLOG(!bOK, LogRHI, Verbose, TEXT("Failed to find GeometryShader shader: %s"), *(PSO.GraphicsDesc.GeometryShader.ToString()));
+					}
+				}
+				else
+				{
+					// if we don't have a vertex shader then we won't bother to add any shaders to the list of outstanding shaders to load
+					// Later on this PSO will be killed forever because it is truly bogus.
+					UE_LOG(LogRHI, Error, TEXT("PSO Entry has no vertex shader: %u this is an invalid entry!"), PSORead->Hash);
+					bOK = false;
+				}
+                
+                // If everything is OK then we can issue reads of the actual shader code
+				if (bOK && PSO.GraphicsDesc.VertexShader != FSHAHash())
+				{
+                    bOK &= FShaderCodeLibrary::RequestShaderCode(PSO.GraphicsDesc.VertexShader, ReadArchive);
+                    UE_CLOG(!bOK, LogRHI, Verbose, TEXT("Failed to read VertexShader shader: %s"), *(PSO.GraphicsDesc.VertexShader.ToString()));
+				}
+		
+				if (bOK && PSO.GraphicsDesc.HullShader != EmptySHA)
+				{
+					bOK &= FShaderCodeLibrary::RequestShaderCode(PSO.GraphicsDesc.HullShader, ReadArchive);
+                    UE_CLOG(!bOK, LogRHI, Verbose, TEXT("Failed to read HullShader shader: %s"), *(PSO.GraphicsDesc.HullShader.ToString()));
+				}
+		
+				if (bOK && PSO.GraphicsDesc.DomainShader != EmptySHA)
+				{
+					bOK &= FShaderCodeLibrary::RequestShaderCode(PSO.GraphicsDesc.DomainShader, ReadArchive);
+                    UE_CLOG(!bOK, LogRHI, Verbose, TEXT("Failed to read DomainShader shader: %s"), *(PSO.GraphicsDesc.DomainShader.ToString()));
+				}
+		
+				if (bOK && PSO.GraphicsDesc.FragmentShader != EmptySHA)
+				{
+					bOK &= FShaderCodeLibrary::RequestShaderCode(PSO.GraphicsDesc.FragmentShader, ReadArchive);
+                    UE_CLOG(!bOK, LogRHI, Verbose, TEXT("Failed to read FragmentShader shader: %s"), *(PSO.GraphicsDesc.FragmentShader.ToString()));
+				}
+		
+				if (bOK && PSO.GraphicsDesc.GeometryShader != EmptySHA)
+				{
+					bOK &= FShaderCodeLibrary::RequestShaderCode(PSO.GraphicsDesc.GeometryShader, ReadArchive);
+                    UE_CLOG(!bOK, LogRHI, Verbose, TEXT("Failed to read GeometryShader shader: %s"), *(PSO.GraphicsDesc.GeometryShader.ToString()));
+				}
+			}
+			else if (PSO.Type == FPipelineCacheFileFormatPSO::DescriptorType::Compute)
+			{
+				if (PSO.ComputeDesc.ComputeShader != EmptySHA)
+				{
+                    Shaders.Add(PSO.ComputeDesc.ComputeShader);
+					bOK &= FShaderCodeLibrary::RequestShaderCode(PSO.ComputeDesc.ComputeShader, ReadArchive);
+                    UE_CLOG(!bOK, LogRHI, Verbose, TEXT("Failed to find ComputeShader shader: %s"), *(PSO.ComputeDesc.ComputeShader.ToString()));
+				}
+				else
+				{
+					bOK = false;
+					UE_LOG(LogRHI, Error, TEXT("Invalid PSO entry in pipeline cache!"));
+				}
+			}
+			else
+			{
+				bOK = false;
+				UE_LOG(LogRHI, Error, TEXT("Invalid PSO entry in pipeline cache!"));
+			}
+		
+			// Then if and only if all shaders can be found do we schedule a compile job
+			if (bOK)
+			{
+				// Add Async Read task
+				CompileJob AsyncJob;
+				
+				AsyncJob.PSO = PSO;
+				AsyncJob.ReadRequests = ReadArchive;
+				
+				ReadTasks.Add(AsyncJob);
+			
+				delete PSORead;
+			}
+			else if (Shaders.Num())
+			{
+				// Re-add to the OrderedCompile tasks and process later
+                // We can never know when this PSO might become valid so we can't ever drop it.
+				FPipelineCachePSOHeader Hdr;
+				Hdr.Hash = PSORead->Hash;
+				Hdr.Shaders = Shaders;
+				OrderedCompileTasks.Insert(Hdr, 0);
+                FPlatformAtomics::InterlockedDecrement(&TotalActiveTasks);
+                
+                delete PSORead;
+			}
+            else
+            {
+                UE_LOG(LogRHI, Error, TEXT("Invalid PSO entry in pipeline cache: %u!"), PSORead->Hash);
+                
+                // Invalid PSOs can be deleted
+                FPlatformAtomics::InterlockedDecrement(&TotalActiveTasks);
+                
+                delete PSORead;
+            }
+			
+			BatchIterator.RemoveCurrent();
+		}
+        else if (!PSORead->bValid)
+        {
+            UE_LOG(LogRHI, Error, TEXT("Invalid PSO entry in pipeline cache: %u!"), PSORead->Hash);
+            
+            // Invalid PSOs can be deleted
+            FPlatformAtomics::InterlockedDecrement(&TotalActiveTasks);
+            
+            delete PSORead;
+            BatchIterator.RemoveCurrent();
+        }
+	}
+}
+
+bool FShaderPipelineCache::ReadyForPrecompile()
+{
+	for(uint32 i = 0; i < (uint32)ReadTasks.Num(); )
+	{
+		check(ReadTasks[i].ReadRequests);
+		if (ReadTasks[i].ReadRequests->PollExternalReadDependencies())
+		{
+			CompileTasks.Add(ReadTasks[i]);
+			ReadTasks.RemoveAt(i);
+		}
+		else
+		{
+			++i;
+		}
+	}
+	if (LastPrecompileRHIFence.GetReference() && LastPrecompileRHIFence->IsComplete())
+	{
+		LastPrecompileRHIFence = nullptr;
+	}
+
+	return (CompileTasks.Num() != 0) && !LastPrecompileRHIFence.GetReference();
+}
+
+void FShaderPipelineCache::PrecompilePipelineBatch()
+{
+	INC_DWORD_STAT(STAT_PreCompileBatchTotal);
+	INC_DWORD_STAT(STAT_PreCompileBatchNum);
+	
+    int32 NumToPrecompile = FMath::Min<int32>(CompileTasks.Num(), BatchSize);
+    
+	for(uint32 i = 0; i < (uint32)NumToPrecompile; i++)
+	{
+		check(CompileTasks[i].ReadRequests && CompileTasks[i].ReadRequests->PollExternalReadDependencies());
+		Precompile(GRHICommandList.GetImmediateCommandList(), GMaxRHIShaderPlatform, CompileTasks[i].PSO);
+		CompiledHashes.Add(GetTypeHash(CompileTasks[i].PSO));
+#if STATS
+		switch(CompileTasks[i].PSO.Type)
+		{
+			case FPipelineCacheFileFormatPSO::DescriptorType::Compute:
+			{
+				INC_DWORD_STAT(STAT_TotalComputePipelineStateCount);
+				break;
+			}
+			case FPipelineCacheFileFormatPSO::DescriptorType::Graphics:
+			{
+				INC_DWORD_STAT(STAT_TotalGraphicsPipelineStateCount);
+				break;
+			}
+			default:
+			{
+				check(false);
+				break;
+			}
+		}
+#endif
+		delete CompileTasks[i].ReadRequests;
+	}
+	
+    FPlatformAtomics::InterlockedAdd(&TotalActiveTasks, -NumToPrecompile);
+
+#if PLATFORM_ANDROID
+	if (NumToPrecompile > 0 && IsRunningRHIInSeparateThread())
+	{
+		LastPrecompileRHIFence = FRHICommandListExecutor::GetImmediateCommandList().RHIThreadFence(false);
+	}
+#endif
+
+	CompileTasks.RemoveAt(0, NumToPrecompile);
+}
+
+bool FShaderPipelineCache::ReadyForNextBatch() const
+{
+	return ReadTasks.Num() == 0;
+}
+
+bool FShaderPipelineCache::ReadyForAutoSave() const
+{
+	bool bAutoSave = false;
+	uint32 SaveAfterNum = CVarPSOFileCacheSaveAfterPSOsLogged.GetValueOnAnyThread();
+	uint32 NumLogged = FPipelineFileCache::NumPSOsLogged();
+
+	const float TimeSinceSave = FPlatformTime::Seconds() - LastAutoSaveTime;
+
+	// autosave if its enabled, and we have more than the desired number, or it's been a while since our last save
+	if (SaveAfterNum > 0 && 
+			(NumLogged >= SaveAfterNum || (NumLogged > 0 && TimeSinceSave >= CVarPSOFileCacheAutoSaveTime.GetValueOnAnyThread()))
+		)
+	{
+		bAutoSave = true;
+	}
+	return bAutoSave;
+}
+
+void FShaderPipelineCache::PollShutdownItems()
+{
+	int64 RemovedTaskCount = 0;
+	
+	const int64 InitialReadTaskCount = ShutdownReadTasks.Num();
+	if(ShutdownReadTasks.Num() > 0)
+	{
+		for(uint32 i = 0; i < (uint32)ShutdownReadTasks.Num(); )
+		{
+			check(ShutdownReadTasks[i].ReadRequests);
+			if (ShutdownReadTasks[i].ReadRequests->PollExternalReadDependencies())
+			{
+                switch(ShutdownReadTasks[i].PSO.Type)
+                {
+                    case FPipelineCacheFileFormatPSO::DescriptorType::Compute:
+                    {
+                        FShaderCodeLibrary::ReleaseShaderCode(ShutdownReadTasks[i].PSO.ComputeDesc.ComputeShader);
+                        break;
+                    }
+                    case FPipelineCacheFileFormatPSO::DescriptorType::Graphics:
+                    {
+                        if (ShutdownReadTasks[i].PSO.GraphicsDesc.VertexShader != FSHAHash())
+                        {
+                            FShaderCodeLibrary::ReleaseShaderCode(ShutdownReadTasks[i].PSO.GraphicsDesc.VertexShader);
+                        }
+                        if (ShutdownReadTasks[i].PSO.GraphicsDesc.GeometryShader != FSHAHash())
+                        {
+                            FShaderCodeLibrary::ReleaseShaderCode(ShutdownReadTasks[i].PSO.GraphicsDesc.GeometryShader);
+                        }
+                        if (ShutdownReadTasks[i].PSO.GraphicsDesc.HullShader != FSHAHash())
+                        {
+                            FShaderCodeLibrary::ReleaseShaderCode(ShutdownReadTasks[i].PSO.GraphicsDesc.HullShader);
+                        }
+                        if (ShutdownReadTasks[i].PSO.GraphicsDesc.DomainShader != FSHAHash())
+                        {
+                            FShaderCodeLibrary::ReleaseShaderCode(ShutdownReadTasks[i].PSO.GraphicsDesc.DomainShader);
+                        }
+                        if (ShutdownReadTasks[i].PSO.GraphicsDesc.FragmentShader != FSHAHash())
+                        {
+                            FShaderCodeLibrary::ReleaseShaderCode(ShutdownReadTasks[i].PSO.GraphicsDesc.FragmentShader);
+                        }
+                        break;
+                    }
+                    default:
+                    {
+                        check(false);
+                        break;
+                    }
+                }
+                
+                delete ShutdownReadTasks[i].ReadRequests;
+                
+				ShutdownReadTasks[i] = ShutdownReadTasks.Last();
+				ShutdownReadTasks.Pop(false);
+				
+				++RemovedTaskCount;
+			}
+			else
+			{
+				++i;
+			}
+		}
+		
+		if(ShutdownFetchTasks.Num() == 0)
+		{
+			ShutdownFetchTasks.Shrink();
+		}
+	}
+	
+	if(ShutdownFetchTasks.Num() > 0)
+	{
+		for(uint32 i = 0; i < (uint32)ShutdownFetchTasks.Num(); )
+		{
+			FPipelineCacheFileFormatPSORead* PSORead = ShutdownFetchTasks[i];
+			
+			FShaderPipelineCacheArchive* Archive = (FShaderPipelineCacheArchive*)PSORead->Ar;
+			check(Archive);
+			if ((Archive && Archive->PollExternalReadDependencies()) || PSORead->bReadCompleted)
+			{
+				delete PSORead;
+				ShutdownFetchTasks[i] = ShutdownFetchTasks.Last();
+				ShutdownFetchTasks.Pop(false);
+				
+				++RemovedTaskCount;
+			}
+			else
+			{
+				++i;
+			}
+		}
+		
+		if(ShutdownFetchTasks.Num() == 0)
+		{
+			ShutdownFetchTasks.Shrink();
+		}
+	}
+	
+	if(RemovedTaskCount > 0)
+    {
+        FPlatformAtomics::InterlockedAdd(&TotalActiveTasks, - RemovedTaskCount);
+	}
+}
+
+void FShaderPipelineCache::Flush()
+{
+	FScopeLock Lock(&Mutex);
+	
+	// reset everything
+	// Abandon all the existing work.
+	// This must be done on the render-thread to avoid locks.
+	CompileTasks.Empty();
+	OrderedCompileTasks.Empty();
+	CompiledHashes.Empty();
+	
+	// Marshall the current read tasks into shutdown
+	for(CompileJob& Entry: ReadTasks )
+	{
+		ShutdownReadTasks.Add(Entry);
+	}
+	ReadTasks.Empty();
+	
+	// Marshall the current fetch tasks into shutdown
+	for (FPipelineCacheFileFormatPSORead* Entry : FetchTasks)
+	{
+		check(Entry->ReadRequest.IsValid());
+		Entry->ReadRequest->Cancel();
+		ShutdownFetchTasks.Add(Entry);
+	}
+	FetchTasks.Empty();
+	
+	int64 StartTaskCount = OrderedCompileTasks.Num() + ShutdownReadTasks.Num() + ShutdownFetchTasks.Num();;
+    FPlatformAtomics::InterlockedExchange(&TotalWaitingTasks, 0);
+}
+
+FShaderPipelineCache::FShaderPipelineCache(EShaderPlatform Platform)
+: FTickableObjectRenderThread(true, false) // (RegisterNow, HighFrequency)
+, BatchSize(0)
+, BatchTime(0.0f)
+, bPaused(false)
+, TotalActiveTasks(0)
+, TotalWaitingTasks(0)
+, TotalCompleteTasks(0)
+, TotalPrecompileTime(0)
+, PrecompileStartTime(0.0)
+, LastAutoSaveTime(0)
+, LastAutoSaveTimeLogBoundPSO(0)
+, LastAutoSaveNum(-1)
+{
+	SET_DWORD_STAT(STAT_ShaderPipelineTaskCount, 0);
+    SET_DWORD_STAT(STAT_ShaderPipelineWaitingTaskCount, 0);
+    SET_DWORD_STAT(STAT_ShaderPipelineActiveTaskCount, 0);
+	
+	BatchSize = CVarPSOFileCacheBatchSize.GetValueOnAnyThread();
+	BatchTime = CVarPSOFileCacheBatchTime.GetValueOnAnyThread();
+	
+	FCoreDelegates::ApplicationWillDeactivateDelegate.AddStatic(&PipelineStateCacheOnAppDeactivate);
+}
+
+FShaderPipelineCache::~FShaderPipelineCache()
+{
+	if (GetShaderPipelineCacheSaveBoundPSOLog())
+	{
+		FShaderPipelineCache::SavePipelineFileCache(FPipelineFileCache::SaveMode::BoundPSOsOnly);
+	}
+	if (GetPSOFileCacheSaveUserCache())
+	{
+		FShaderPipelineCache::SavePipelineFileCache(FPipelineFileCache::SaveMode::Incremental);
+	}
+	
+	Close();
+	
+	// The render thread tick should be dead now and we are safe to destroy everything that needs to wait or manual destruction
+
+	for(CompileJob& Entry: ReadTasks )
+	{
+		check(Entry.ReadRequests);
+		Entry.ReadRequests->BlockingWaitComplete();
+	}
+	
+	for(CompileJob& Entry: ShutdownReadTasks )
+	{
+		check(Entry.ReadRequests);
+		Entry.ReadRequests->BlockingWaitComplete();
+	}
+	
+	for (FPipelineCacheFileFormatPSORead* Entry : FetchTasks)
+	{
+		if(Entry->ReadRequest.IsValid())
+		{
+			Entry->ReadRequest->WaitCompletion(0.0);
+		}
+		delete Entry;
+	}
+	
+	for (FPipelineCacheFileFormatPSORead* Entry : ShutdownFetchTasks)
+	{
+		if(Entry->ReadRequest.IsValid())
+		{
+			Entry->ReadRequest->WaitCompletion(0.0);
+		}
+		delete Entry;
+	}	
+}
+
+bool FShaderPipelineCache::IsTickable() const
+{
+	return FPlatformProperties::RequiresCookedData() && !bPaused &&
+		(FPlatformAtomics::AtomicRead(&TotalActiveTasks) != 0 ||
+			FPlatformAtomics::AtomicRead(&TotalWaitingTasks) != 0 ||
+			FPlatformAtomics::AtomicRead(&TotalCompleteTasks) != 0 ||
+			ReadyForAutoSave()||
+			GetShaderPipelineCacheSaveBoundPSOLog());
+}
+
+void FShaderPipelineCache::Tick( float DeltaTime )
+{
+	FScopeLock Lock(&Mutex);
+    
+	if (LastPrecompileRHIFence.GetReference() && LastPrecompileRHIFence->IsComplete())
+	{
+		LastPrecompileRHIFence = nullptr;
+	}
+
+    if (FPlatformAtomics::AtomicRead(&TotalWaitingTasks) == 0 && FPlatformAtomics::AtomicRead(&TotalActiveTasks) == 0 && FPlatformAtomics::AtomicRead(&TotalCompleteTasks) != 0 && !LastPrecompileRHIFence.GetReference())
+    {
+		float WallTime = 0.0f;
+		if (PrecompileStartTime > 0.0)
+		{
+			WallTime = float(FPlatformTime::Seconds() - PrecompileStartTime);
+		}
+		UE_LOG(LogRHI, Warning, TEXT("FShaderPipelineCache completed %u tasks in %.2fs (%.2fs wall time since intial open)."), (uint32)TotalCompleteTasks, FPlatformTime::ToSeconds64(TotalPrecompileTime), WallTime);
+        if (OnPrecompilationComplete.IsBound())
+        {
+            OnPrecompilationComplete.Broadcast((uint32)TotalCompleteTasks, FPlatformTime::ToSeconds64(TotalPrecompileTime), ShaderCachePrecompileContext);
+        }
+        FPlatformAtomics::InterlockedExchange(&TotalCompleteTasks, 0);
+        FPlatformAtomics::InterlockedExchange(&TotalPrecompileTime, 0);
+    }
+	
+	if (ReadyForAutoSave())
+	{
+		if (GetPSOFileCacheSaveUserCache())
+		{
+			Save(FPipelineFileCache::SaveMode::Incremental);
+		}
+	}
+	if (GetShaderPipelineCacheSaveBoundPSOLog())
+	{
+		if (LastAutoSaveNum < int32(FPipelineFileCache::NumPSOsLogged()))
+		{
+			const float TimeSinceSave = FPlatformTime::Seconds() - LastAutoSaveTimeLogBoundPSO;
+
+			if (TimeSinceSave >= CVarPSOFileCacheAutoSaveTimeBoundPSO.GetValueOnAnyThread())
+			{
+				Save(FPipelineFileCache::SaveMode::BoundPSOsOnly);
+				LastAutoSaveTimeLogBoundPSO = FPlatformTime::Seconds();
+				LastAutoSaveNum = FPipelineFileCache::NumPSOsLogged();
+			}
+		}
+	}
+
+	PollShutdownItems();
+
+	if (PrecompileStartTime == 0.0 && (PreFetchedTasks.Num() || FetchTasks.Num() || OrderedCompileTasks.Num()))
+	{
+		PrecompileStartTime = FPlatformTime::Seconds();
+	}
+	
+	// Copy any new items over to our 'internal' safe array
+	if (PreFetchedTasks.Num())
+	{
+		OrderedCompileTasks = PreFetchedTasks;
+		PreFetchedTasks.Empty();
+	}
+	
+	if (ReadyForPrecompile())
+	{
+		SCOPE_SECONDS_ACCUMULATOR(STAT_PreCompileTotalTime);
+		SCOPE_CYCLE_COUNTER(STAT_PreCompileTime);
+		
+		uint32 Start = FPlatformTime::Cycles();
+
+		PrecompilePipelineBatch();
+
+		uint32 End = FPlatformTime::Cycles();
+
+		if (BatchTime > 0.0f)
+		{
+			float ElapsedMs = FPlatformTime::ToMilliseconds(End - Start);
+			if (ElapsedMs < BatchTime)
+			{
+				BatchSize++;
+			}
+			else if (ElapsedMs > BatchTime)
+			{
+				if (BatchSize > 1)
+				{
+					BatchSize--;
+				}
+				else
+				{
+					UE_LOG(LogRHI, Warning, TEXT("FShaderPipelineCache: Cannot reduce BatchSize below 1 to meet target of %f ms, elapsed time was %f ms)"), BatchTime, ElapsedMs);
+				}
+			}
+		}
+	}
+	
+	if (ReadyForNextBatch() && (OrderedCompileTasks.Num() || FetchTasks.Num()))
+	{
+        uint32 Num = 0;
+        if (BatchSize > static_cast<uint32>(FetchTasks.Num()))
+        {
+            Num = BatchSize - FetchTasks.Num();
+        }
+        Num = FMath::Min(Num, (uint32)OrderedCompileTasks.Num());
+            
+		if (FetchTasks.Num() < (int32)Num)
+		{
+			TArray<FPipelineCacheFileFormatPSORead*> NewBatch;
+		
+			Num -= FetchTasks.Num();
+            for (auto Iterator = OrderedCompileTasks.CreateIterator(); Iterator && Num > 0; ++Iterator)
+            {
+                bool bHasShaders = true;
+                for (FSHAHash const& Hash : Iterator->Shaders)
+                {
+                    bHasShaders &= FShaderCodeLibrary::ContainsShaderCode(Hash);
+                }
+            	if (bHasShaders)
+            	{
+            		FPipelineCacheFileFormatPSORead* Entry = new FPipelineCacheFileFormatPSORead;
+					Entry->Hash = Iterator->Hash;
+					Entry->Ar = new FShaderPipelineCacheArchive;
+					NewBatch.Add(Entry);
+	            	Iterator.RemoveCurrent();
+                    FPlatformAtomics::InterlockedIncrement(&TotalActiveTasks);
+                    FPlatformAtomics::InterlockedDecrement(&TotalWaitingTasks);
+                    --Num;
+            	}
+            }
+			
+			FPipelineFileCache::FetchPSODescriptors(NewBatch);
+			
+			FetchTasks.Append(NewBatch);
+		}		
+
+        if (static_cast<uint32>(FetchTasks.Num()) > BatchSize)
+        {
+            UE_LOG(LogRHI, Warning, TEXT("FShaderPipelineCache: Attempting to pre-compile more jobs (%d) than the batch size (%d)"), FetchTasks.Num(), BatchSize);
+        }
+        
+		PreparePipelineBatch(FetchTasks);
+	}
+	
+#if STATS
+	int64 ActiveTaskCount = FMath::Max((int64)0, FPlatformAtomics::AtomicRead(&TotalActiveTasks));
+    int64 WaitingTaskCount = FMath::Max((int64)0, FPlatformAtomics::AtomicRead(&TotalWaitingTasks));
+	SET_DWORD_STAT(STAT_ShaderPipelineTaskCount, ActiveTaskCount+WaitingTaskCount);
+    SET_DWORD_STAT(STAT_ShaderPipelineWaitingTaskCount, WaitingTaskCount);
+    SET_DWORD_STAT(STAT_ShaderPipelineActiveTaskCount, ActiveTaskCount);
+	
+	// Calc in one place otherwise this computation will be splattered all over the place - this will not be exact but counts the expensive bits
+	int64 InUseMemory = OrderedCompileTasks.GetAllocatedSize() + 
+						CompiledHashes.GetAllocatedSize() + 
+						ReadTasks.GetAllocatedSize() + 
+						CompileTasks.GetAllocatedSize() + 
+						ShutdownReadTasks.GetAllocatedSize() +
+						FetchTasks.GetAllocatedSize() +
+						ShutdownFetchTasks.GetAllocatedSize();
+	if(ActiveTaskCount+WaitingTaskCount > 0)
+	{
+		InUseMemory += (ReadTasks.Num() + CompileTasks.Num() + ShutdownReadTasks.Num()) * (sizeof(FShaderPipelineCacheArchive));
+		InUseMemory += (FetchTasks.Num() + ShutdownFetchTasks.Num()) * (sizeof(FPipelineCacheFileFormatPSORead));
+		for(int32 i = 0;i < FetchTasks.Num();++i)
+		{
+			InUseMemory += FetchTasks[i]->Data.Num();
+		}
+		for(int32 i = 0;i < ShutdownFetchTasks.Num();++i)
+		{
+			InUseMemory += ShutdownFetchTasks[i]->Data.Num();
+		}
+	}
+	SET_MEMORY_STAT(STAT_PreCompileMemory, InUseMemory);
+#endif
+}
+
+bool FShaderPipelineCache::NeedsRenderingResumedForRenderingThreadTick() const
+{
+	return true;
+}
+
+TStatId FShaderPipelineCache::GetStatId() const
+{
+	RETURN_QUICK_DECLARE_CYCLE_STAT(FShaderPipelineBatchCompiler, STATGROUP_Tickables);
+}
+
+bool FShaderPipelineCache::Open(FString const& Name, EShaderPlatform Platform)
+{
+	FileName = Name;
+    CurrentPlatform = Platform;
+	FGuid CacheFileGuid;
+	bool bOK = FPipelineFileCache::OpenPipelineFileCache(Name, Platform, CacheFileGuid);
+	if (bOK)
+	{
+		FScopeLock Lock(&Mutex);
+		
+		Flush();
+
+		int32 Order = (int32)FPipelineFileCache::PSOOrder::Default;
+		
+		if(!GConfig->GetInt(FShaderPipelineCacheConstants::SectionHeading, FShaderPipelineCacheConstants::SortOrderKey, Order, *GGameUserSettingsIni))
+		{
+			GConfig->GetInt(FShaderPipelineCacheConstants::SectionHeading, FShaderPipelineCacheConstants::SortOrderKey, Order, *GGameIni);
+		}
+
+		TArray<FPipelineCachePSOHeader> LocalPreFetchedTasks;
+		FPipelineFileCache::GetOrderedPSOHashes(LocalPreFetchedTasks, (FPipelineFileCache::PSOOrder)Order);
+        // Iterate over all the tasks we haven't yet begun to read data for - these are the 'waiting' tasks
+        int64 Count = 0;
+        for (auto const& Task : LocalPreFetchedTasks)
+        {
+            bool bHasShaders = true;
+            for (FSHAHash const& Hash : Task.Shaders)
+            {
+                bHasShaders &= FShaderCodeLibrary::ContainsShaderCode(Hash);
+            }
+            if (bHasShaders)
+            {
+                Count++;
+            }
+        }
+
+		FPlatformAtomics::InterlockedAdd(&TotalWaitingTasks, Count);
+
+		if (OnCachedOpened.IsBound())
+		{
+			OnCachedOpened.Broadcast(Name, Platform, LocalPreFetchedTasks.Num(), CacheFileGuid, ShaderCachePrecompileContext);
+		}
+
+		PreFetchedTasks = LocalPreFetchedTasks;
+
+		UE_LOG(LogRHI, Display, TEXT("Opened pipeline cache and enqueued %d of %d tasks for precompile."), Count, PreFetchedTasks.Num());
+	}
+
+	UE_CLOG(!bOK, LogRHI, Display, TEXT("Failed to open default shader pipeline cache for %s using shader platform %d."), *Name, (uint32)Platform);
+	
+	bOpened = bOK;
+
+	// OnPrecompilationBegin can bring up a modal loading screen, so we call it outside the scope lock to prevent deadlocks
+	if (bOK && OnPrecompilationBegin.IsBound())
+	{
+		OnPrecompilationBegin.Broadcast(PreFetchedTasks.Num(), ShaderCachePrecompileContext);
+	}
+
+	return bOK;
+}
+
+bool FShaderPipelineCache::Save(FPipelineFileCache::SaveMode Mode)
+{
+	FScopeLock Lock(&Mutex);
+		
+	bool bOK = FPipelineFileCache::SavePipelineFileCache(FileName, Mode);
+	UE_CLOG(!bOK, LogRHI, Warning, TEXT("Failed to save shader pipeline cache for %s using save mode %d."), *FileName, (uint32)Mode);
+
+	LastAutoSaveTime = FPlatformTime::Seconds();
+
+	return bOK;
+}
+
+void FShaderPipelineCache::Close()
+{
+	FScopeLock Lock(&Mutex);
+		
+	if(GConfig)
+	{
+		GConfig->SetString(FShaderPipelineCacheConstants::SectionHeading, FShaderPipelineCacheConstants::LastOpenedKey, *FileName, *GGameUserSettingsIni);
+		GConfig->Flush(false, *GGameUserSettingsIni);
+	}
+	
+	// Log all bound PSOs
+	if (GetShaderPipelineCacheSaveBoundPSOLog())
+	{
+		Save(FPipelineFileCache::SaveMode::BoundPSOsOnly);
+	}
+	
+	// Force a fast save, just in case
+	if (GetPSOFileCacheSaveUserCache())
+	{
+		Save(FPipelineFileCache::SaveMode::Incremental);
+	}
+	
+	// Signal flush of outstanding work to allow restarting for a new cache file
+	Flush();
+    
+    if (OnCachedClosed.IsBound())
+    {
+        OnCachedClosed.Broadcast(FileName, CurrentPlatform);
+    }
+	
+	bOpened = false;
+
+	FPipelineFileCache::ClosePipelineFileCache();
+}
+
+void FShaderPipelineCache::OnShaderLibraryStateChanged(ELibraryState State, EShaderPlatform Platform, FString const& Name)
+{
+    FScopeLock Lock(&Mutex);
+	
+	if (State == FShaderPipelineCache::Opened && Name == FApp::GetProjectName() && Platform == CurrentPlatform && !bOpened)
+	{
+		Close();
+		FString LastOpenedName;
+		if ((!GConfig->GetString(FShaderPipelineCacheConstants::SectionHeading, FShaderPipelineCacheConstants::LastOpenedKey, LastOpenedName, *GGameUserSettingsIni) && !GConfig->GetString(FShaderPipelineCacheConstants::SectionHeading, FShaderPipelineCacheConstants::LastOpenedKey, LastOpenedName, *GGameIni)) && !LastOpenedName.Len())
+		{
+			LastOpenedName = FApp::GetProjectName();
+		}
+		Open(LastOpenedName, Platform);
+	}
+	
+    // Copy any new items over to our 'internal' safe array
+    if (PreFetchedTasks.Num())
+    {
+        OrderedCompileTasks = PreFetchedTasks;
+        PreFetchedTasks.Empty();
+    }
+    
+    // Iterate over all the tasks we haven't yet begun to read data for - these are the 'waiting' tasks
+    int64 Count = 0;
+    for (auto const& Task : OrderedCompileTasks)
+    {
+        bool bHasShaders = true;
+        for (FSHAHash const& Hash : Task.Shaders)
+        {
+            bHasShaders &= FShaderCodeLibrary::ContainsShaderCode(Hash);
+        }
+        if (bHasShaders)
+        {
+            Count++;
+        }
+    }
+    
+    // Set the new waiting count that we can actually process.
+    FPlatformAtomics::InterlockedExchange(&TotalWaitingTasks, Count);
+	UE_LOG(LogRHI, Display, TEXT("Opened pipeline cache after state change and enqueued %d of %d tasks for precompile."), Count, OrderedCompileTasks.Num());
+}
diff --git a/Engine/Source/Runtime/RenderCore/Private/StaticBoundShaderState.cpp b/Engine/Source/Runtime/RenderCore/Private/StaticBoundShaderState.cpp
new file mode 100644
--- /dev/null
+++ b/Engine/Source/Runtime/RenderCore/Private/StaticBoundShaderState.cpp
@@ -0,0 +1,110 @@
+// Copyright 1998-2018 Epic Games, Inc. All Rights Reserved.
+
+/*=============================================================================
+	StaticBoundShaderState.cpp: Static bound shader state implementation.
+=============================================================================*/
+
+#include "StaticBoundShaderState.h"
+#include "RenderingThread.h"
+#include "Shader.h"
+
+PRAGMA_DISABLE_DEPRECATION_WARNINGS
+
+TLinkedList<FGlobalBoundShaderStateResource*>*& FGlobalBoundShaderStateResource::GetGlobalBoundShaderStateList()
+{
+	static TLinkedList<FGlobalBoundShaderStateResource*>* List = NULL;
+	return List;
+}
+
+FGlobalBoundShaderStateResource::FGlobalBoundShaderStateResource()
+	: GlobalListLink(this)
+#if DO_CHECK
+	, BoundVertexDeclaration(nullptr)
+	, BoundVertexShader(nullptr)
+	, BoundPixelShader(nullptr)
+	, BoundGeometryShader(nullptr)
+#endif 
+{
+	// Add this resource to the global list in the rendering thread.
+	if(IsInRenderingThread())
+	{
+		GlobalListLink.LinkHead(GetGlobalBoundShaderStateList());
+	}
+	else
+	{
+		ENQUEUE_UNIQUE_RENDER_COMMAND_ONEPARAMETER(
+			LinkGlobalBoundShaderStateResource,FGlobalBoundShaderStateResource*,Resource,this,
+			{
+				Resource->GlobalListLink.LinkHead(GetGlobalBoundShaderStateList());
+			});
+	}
+}
+
+FGlobalBoundShaderStateResource::~FGlobalBoundShaderStateResource()
+{
+	// Remove this resource from the global list.
+	GlobalListLink.Unlink();
+}
+
+/**
+ * Initializes a global bound shader state with a vanilla bound shader state and required information.
+ */
+FBoundShaderStateRHIParamRef FGlobalBoundShaderStateResource::GetInitializedRHI(
+	FVertexDeclarationRHIParamRef VertexDeclaration,
+	FVertexShaderRHIParamRef VertexShader, 
+	FPixelShaderRHIParamRef PixelShader,
+	FGeometryShaderRHIParamRef GeometryShader
+	)
+{
+	check(IsInitialized());
+
+	// This should only be called in the rendering thread after the RHI has been initialized.
+	check(GIsRHIInitialized);
+	check(IsInRenderingThread());
+
+	// Create the bound shader state if it hasn't been cached yet.
+	if(!IsValidRef(BoundShaderState))
+	{
+#if DO_CHECK
+		BoundVertexDeclaration = VertexDeclaration;
+		BoundVertexShader = VertexShader;
+		BoundPixelShader = PixelShader;
+		BoundGeometryShader = GeometryShader;
+#endif 
+		BoundShaderState = 
+			RHICreateBoundShaderState(
+			VertexDeclaration,
+			VertexShader,
+			FHullShaderRHIRef(), 
+			FDomainShaderRHIRef(),
+			PixelShader,
+			GeometryShader);
+	}
+#if DO_CHECK
+	// Verify that the passed in shaders will actually be used
+	// This will catch cases where one bound shader state is being used with more than one combination of shaders
+	// Otherwise setting the shader will just silently fail once the bound shader state has been initialized with a different shader 
+	check(BoundVertexDeclaration == VertexDeclaration &&
+		BoundVertexShader == VertexShader &&
+		BoundPixelShader == PixelShader &&
+		BoundGeometryShader == GeometryShader);
+#endif 
+
+	return BoundShaderState;
+}
+
+FBoundShaderStateRHIParamRef FGlobalBoundShaderStateResource::GetPreinitializedRHI()
+{
+	check(IsInitialized());
+	check(GIsRHIInitialized);
+	check(IsInParallelRenderingThread());
+	return BoundShaderState;
+}
+
+void FGlobalBoundShaderStateResource::ReleaseRHI()
+{
+	// Release the cached bound shader state.
+	BoundShaderState.SafeRelease();
+}
+
+PRAGMA_ENABLE_DEPRECATION_WARNINGS
diff --git a/Engine/Source/Runtime/RenderCore/Private/VertexFactory.cpp b/Engine/Source/Runtime/RenderCore/Private/VertexFactory.cpp
new file mode 100644
--- /dev/null
+++ b/Engine/Source/Runtime/RenderCore/Private/VertexFactory.cpp
@@ -0,0 +1,411 @@
+// Copyright 1998-2018 Epic Games, Inc. All Rights Reserved.
+
+/*=============================================================================
+	VertexFactory.cpp: Vertex factory implementation
+=============================================================================*/
+
+#include "VertexFactory.h"
+#include "Serialization/MemoryWriter.h"
+#include "UObject/DebugSerializationFlags.h"
+
+uint32 FVertexFactoryType::NextHashIndex = 0;
+bool FVertexFactoryType::bInitializedSerializationHistory = false;
+
+static TLinkedList<FVertexFactoryType*>* GVFTypeList = nullptr;
+
+/**
+ * @return The global shader factory list.
+ */
+TLinkedList<FVertexFactoryType*>*& FVertexFactoryType::GetTypeList()
+{
+	return GVFTypeList;
+}
+
+/**
+ * Finds a FVertexFactoryType by name.
+ */
+FVertexFactoryType* FVertexFactoryType::GetVFByName(const FString& VFName)
+{
+	for(TLinkedList<FVertexFactoryType*>::TIterator It(GetTypeList()); It; It.Next())
+	{
+		FString CurrentVFName = FString(It->GetName());
+		if (CurrentVFName == VFName)
+		{
+			return *It;
+		}
+	}
+	return NULL;
+}
+
+void FVertexFactoryType::Initialize(const TMap<FString, TArray<const TCHAR*> >& ShaderFileToUniformBufferVariables)
+{
+	if (!FPlatformProperties::RequiresCookedData())
+	{
+		// Cache serialization history for each VF type
+		// This history is used to detect when shader serialization changes without a corresponding .usf change
+		for(TLinkedList<FVertexFactoryType*>::TIterator It(FVertexFactoryType::GetTypeList()); It; It.Next())
+		{
+			FVertexFactoryType* Type = *It;
+			GenerateReferencedUniformBuffers(Type->ShaderFilename, Type->Name, ShaderFileToUniformBufferVariables, Type->ReferencedUniformBufferStructsCache);
+
+			for (int32 Frequency = 0; Frequency < SF_NumFrequencies; Frequency++)
+			{
+				// Construct a temporary shader parameter instance, which is initialized to safe values for serialization
+				FVertexFactoryShaderParameters* Parameters = Type->CreateShaderParameters((EShaderFrequency)Frequency);
+
+				if (Parameters)
+				{
+					// Serialize the temp shader to memory and record the number and sizes of serializations
+					TArray<uint8> TempData;
+					FMemoryWriter Ar(TempData, true);
+					FShaderSaveArchive SaveArchive(Ar, Type->SerializationHistory[Frequency]);
+					Parameters->Serialize(SaveArchive);
+					delete Parameters;
+				}
+			}
+		}
+	}
+
+	bInitializedSerializationHistory = true;
+}
+
+void FVertexFactoryType::Uninitialize()
+{
+	for(TLinkedList<FVertexFactoryType*>::TIterator It(FVertexFactoryType::GetTypeList()); It; It.Next())
+	{
+		FVertexFactoryType* Type = *It;
+
+		for (int32 Frequency = 0; Frequency < SF_NumFrequencies; Frequency++)
+		{
+			Type->SerializationHistory[Frequency] = FSerializationHistory();
+		}
+	}
+
+	bInitializedSerializationHistory = false;
+}
+
+FVertexFactoryType::FVertexFactoryType(
+	const TCHAR* InName,
+	const TCHAR* InShaderFilename,
+	bool bInUsedWithMaterials,
+	bool bInSupportsStaticLighting,
+	bool bInSupportsDynamicLighting,
+	bool bInSupportsPrecisePrevWorldPos,
+	bool bInSupportsPositionOnly,
+	ConstructParametersType InConstructParameters,
+	ShouldCacheType InShouldCache,
+	ModifyCompilationEnvironmentType InModifyCompilationEnvironment,
+	SupportsTessellationShadersType InSupportsTessellationShaders
+	):
+	Name(InName),
+	ShaderFilename(InShaderFilename),
+	TypeName(InName),
+	bUsedWithMaterials(bInUsedWithMaterials),
+	bSupportsStaticLighting(bInSupportsStaticLighting),
+	bSupportsDynamicLighting(bInSupportsDynamicLighting),
+	bSupportsPrecisePrevWorldPos(bInSupportsPrecisePrevWorldPos),
+	bSupportsPositionOnly(bInSupportsPositionOnly),
+	ConstructParameters(InConstructParameters),
+	ShouldCacheRef(InShouldCache),
+	ModifyCompilationEnvironmentRef(InModifyCompilationEnvironment),
+	SupportsTessellationShadersRef(InSupportsTessellationShaders),
+	GlobalListLink(this)
+{
+	// Make sure the format of the source file path is right.
+	check(CheckVirtualShaderFilePath(InShaderFilename));
+
+	checkf(FPaths::GetExtension(InShaderFilename) == TEXT("ush"),
+		TEXT("Incorrect virtual shader path extension for vertex factory shader header '%s': Only .ush files should be included."),
+		InShaderFilename);
+
+	bCachedUniformBufferStructDeclarations = false;
+
+	// This will trigger if an IMPLEMENT_VERTEX_FACTORY_TYPE was in a module not loaded before InitializeShaderTypes
+	// Vertex factory types need to be implemented in modules that are loaded before that
+	checkf(!bInitializedSerializationHistory, TEXT("VF type was loaded after engine init, use ELoadingPhase::PostConfigInit on your module to cause it to load earlier."));
+
+	// Add this vertex factory type to the global list.
+	GlobalListLink.LinkHead(GetTypeList());
+
+	// Assign the vertex factory type the next unassigned hash index.
+	HashIndex = NextHashIndex++;
+}
+
+FVertexFactoryType::~FVertexFactoryType()
+{
+	GlobalListLink.Unlink();
+}
+
+/** Calculates a Hash based on this vertex factory type's source code and includes */
+const FSHAHash& FVertexFactoryType::GetSourceHash(EShaderPlatform ShaderPlatform) const
+{
+	return GetShaderFileHash(GetShaderFilename(), ShaderPlatform);
+}
+
+FArchive& operator<<(FArchive& Ar,FVertexFactoryType*& TypeRef)
+{
+	if(Ar.IsSaving())
+	{
+		FName TypeName = TypeRef ? FName(TypeRef->GetName()) : NAME_None;
+		Ar << TypeName;
+	}
+	else if(Ar.IsLoading())
+	{
+		FName TypeName = NAME_None;
+		Ar << TypeName;
+		TypeRef = FindVertexFactoryType(TypeName);
+	}
+	return Ar;
+}
+
+FVertexFactoryType* FindVertexFactoryType(FName TypeName)
+{
+	// Search the global vertex factory list for a type with a matching name.
+	for(TLinkedList<FVertexFactoryType*>::TIterator VertexFactoryTypeIt(FVertexFactoryType::GetTypeList());VertexFactoryTypeIt;VertexFactoryTypeIt.Next())
+	{
+		if(VertexFactoryTypeIt->GetFName() == TypeName)
+		{
+			return *VertexFactoryTypeIt;
+		}
+	}
+	return NULL;
+}
+
+void FVertexFactory::SetStreams(ERHIFeatureLevel::Type InFeatureLevel, FRHICommandList& RHICmdList) const
+{
+	bool bSupportsVertexFetch = SupportsManualVertexFetch(InFeatureLevel);
+	check(IsInitialized());
+	for(int32 StreamIndex = 0;StreamIndex < Streams.Num();StreamIndex++)
+	{
+		const FVertexStream& Stream = Streams[StreamIndex];
+		if (!(EnumHasAnyFlags(EVertexStreamUsage::ManualFetch, Stream.VertexStreamUsage) && bSupportsVertexFetch))
+		{
+			if (!Stream.VertexBuffer)
+			{
+				RHICmdList.SetStreamSource(StreamIndex, nullptr, 0);
+			}
+			else
+			{
+				if (EnumHasAnyFlags(EVertexStreamUsage::Overridden, Stream.VertexStreamUsage) && !Stream.VertexBuffer->IsInitialized())
+				{
+					RHICmdList.SetStreamSource(StreamIndex, nullptr, 0);
+				}
+				else
+				{
+					checkf(Stream.VertexBuffer->IsInitialized(), TEXT("Vertex buffer was not initialized! Stream %u, Stride %u, Name %s"), StreamIndex, Stream.Stride, *Stream.VertexBuffer->GetFriendlyName());
+					RHICmdList.SetStreamSource(StreamIndex, Stream.VertexBuffer->VertexBufferRHI, Stream.Offset);
+				}
+			}
+		}
+	}
+}
+
+void FVertexFactory::OffsetInstanceStreams(FRHICommandList& RHICmdList, uint32 FirstVertex) const
+{
+	for(int32 StreamIndex = 0;StreamIndex < Streams.Num();StreamIndex++)
+	{
+		const FVertexStream& Stream = Streams[StreamIndex];
+		if (EnumHasAnyFlags(EVertexStreamUsage::Instancing, Stream.VertexStreamUsage))
+		{
+			RHICmdList.SetStreamSource( StreamIndex, Stream.VertexBuffer->VertexBufferRHI, Stream.Offset + Stream.Stride * FirstVertex);
+		}
+	}
+}
+
+void FVertexFactory::SetPositionStream(FRHICommandList& RHICmdList) const
+{
+	check(IsInitialized());
+	// Set the predefined vertex streams.
+	for(int32 StreamIndex = 0;StreamIndex < PositionStream.Num();StreamIndex++)
+	{
+		const FVertexStream& Stream = PositionStream[StreamIndex];
+		check(Stream.VertexBuffer->IsInitialized());
+		RHICmdList.SetStreamSource( StreamIndex, Stream.VertexBuffer->VertexBufferRHI, Stream.Offset);
+	}
+}
+
+void FVertexFactory::OffsetPositionInstanceStreams(FRHICommandList& RHICmdList, uint32 FirstVertex) const
+{
+	for(int32 StreamIndex = 0;StreamIndex < PositionStream.Num();StreamIndex++)
+	{
+		const FVertexStream& Stream = PositionStream[StreamIndex];
+		if (EnumHasAnyFlags(EVertexStreamUsage::Instancing, Stream.VertexStreamUsage))
+		{
+			RHICmdList.SetStreamSource( StreamIndex, Stream.VertexBuffer->VertexBufferRHI, Stream.Offset + Stream.Stride * FirstVertex);
+		}
+	}
+}
+
+void FVertexFactory::ReleaseRHI()
+{
+	Declaration.SafeRelease();
+	PositionDeclaration.SafeRelease();
+	Streams.Empty();
+	PositionStream.Empty();
+}
+
+/**
+* Fill in array of strides from this factory's vertex streams without shadow/light maps
+* @param OutStreamStrides - output array of # MaxVertexElementCount stream strides to fill
+*/
+int32 FVertexFactory::GetStreamStrides(uint32 *OutStreamStrides, bool bPadWithZeroes) const
+{
+	int32 StreamIndex;
+	for(StreamIndex = 0;StreamIndex < Streams.Num();++StreamIndex)
+	{
+		OutStreamStrides[StreamIndex] = Streams[StreamIndex].Stride;
+	}
+	if (bPadWithZeroes)
+	{
+		// Pad stream strides with 0's to be safe (they can be used in hashes elsewhere)
+		for (;StreamIndex < MaxVertexElementCount;++StreamIndex)
+		{
+			OutStreamStrides[StreamIndex] = 0;
+		}
+	}
+	return StreamIndex;
+}
+
+/**
+* Fill in array of strides from this factory's position only vertex streams
+* @param OutStreamStrides - output array of # MaxVertexElementCount stream strides to fill
+*/
+void FVertexFactory::GetPositionStreamStride(uint32 *OutStreamStrides) const
+{
+	int32 StreamIndex;
+	for(StreamIndex = 0;StreamIndex < PositionStream.Num();++StreamIndex)
+	{
+		OutStreamStrides[StreamIndex] = PositionStream[StreamIndex].Stride;
+	}
+	// Pad stream strides with 0's to be safe (they can be used in hashes elsewhere)
+	for (;StreamIndex < MaxVertexElementCount;++StreamIndex)
+	{
+		OutStreamStrides[StreamIndex] = 0;
+	}
+}
+
+FVertexElement FVertexFactory::AccessStreamComponent(const FVertexStreamComponent& Component,uint8 AttributeIndex)
+{
+	FVertexStream VertexStream;
+	VertexStream.VertexBuffer = Component.VertexBuffer;
+	VertexStream.Stride = Component.Stride;
+	VertexStream.Offset = Component.StreamOffset;
+	VertexStream.VertexStreamUsage = Component.VertexStreamUsage;
+
+	return FVertexElement(Streams.AddUnique(VertexStream),Component.Offset,Component.Type,AttributeIndex,VertexStream.Stride, EnumHasAnyFlags(EVertexStreamUsage::Instancing, VertexStream.VertexStreamUsage));
+}
+
+FVertexElement FVertexFactory::AccessPositionStreamComponent(const FVertexStreamComponent& Component,uint8 AttributeIndex)
+{
+	FVertexStream VertexStream;
+	VertexStream.VertexBuffer = Component.VertexBuffer;
+	VertexStream.Stride = Component.Stride;
+	VertexStream.Offset = Component.StreamOffset;
+	VertexStream.VertexStreamUsage = Component.VertexStreamUsage;
+
+	return FVertexElement(PositionStream.AddUnique(VertexStream),Component.Offset,Component.Type,AttributeIndex,VertexStream.Stride, EnumHasAnyFlags(EVertexStreamUsage::Instancing, VertexStream.VertexStreamUsage));
+}
+
+void FVertexFactory::InitDeclaration(FVertexDeclarationElementList& Elements)
+{
+	// Create the vertex declaration for rendering the factory normally.
+	Declaration = RHICreateVertexDeclaration(Elements);
+}
+
+void FVertexFactory::InitPositionDeclaration(const FVertexDeclarationElementList& Elements)
+{
+	PositionDeclaration = RHICreateVertexDeclaration(Elements);
+}
+
+FVertexFactoryParameterRef::FVertexFactoryParameterRef(FVertexFactoryType* InVertexFactoryType,const FShaderParameterMap& ParameterMap, EShaderFrequency InShaderFrequency, EShaderPlatform InShaderPlatform)
+: Parameters(NULL)
+, VertexFactoryType(InVertexFactoryType)
+, ShaderFrequency(InShaderFrequency)
+, ShaderPlatform(InShaderPlatform)
+{
+	Parameters = VertexFactoryType->CreateShaderParameters(InShaderFrequency);
+	VFHash = GetShaderFileHash(VertexFactoryType->GetShaderFilename(), InShaderPlatform);
+
+	if(Parameters)
+	{
+		Parameters->Bind(ParameterMap);
+	}
+}
+
+bool operator<<(FArchive& Ar,FVertexFactoryParameterRef& Ref)
+{
+	bool bShaderHasOutdatedParameters = false;
+
+	Ar << Ref.VertexFactoryType;
+
+	uint8 ShaderFrequencyByte = Ref.ShaderFrequency;
+	Ar << ShaderFrequencyByte;
+	if(Ar.IsLoading())
+	{
+		Ref.ShaderFrequency = (EShaderFrequency)ShaderFrequencyByte;
+	}
+
+	uint8 ShaderPlatformByte = Ref.ShaderPlatform;
+	Ar << ShaderPlatformByte;
+	if (Ar.IsLoading())
+	{
+		Ref.ShaderPlatform = (EShaderPlatform)ShaderPlatformByte;
+	}
+
+	Ar << FShaderResource::FilterShaderSourceHashForSerialization(Ar, Ref.VFHash);
+
+
+	if (Ar.IsLoading())
+	{
+		delete Ref.Parameters;
+		if (Ref.VertexFactoryType)
+		{
+			Ref.Parameters = Ref.VertexFactoryType->CreateShaderParameters(Ref.ShaderFrequency);
+		}
+		else
+		{
+			bShaderHasOutdatedParameters = true;
+			Ref.Parameters = NULL;
+		}
+	}
+
+	// Need to be able to skip over parameters for no longer existing vertex factories.
+	int64 SkipOffset = Ar.Tell();
+	{
+		FArchive::FScopeSetDebugSerializationFlags S(Ar, DSF_IgnoreDiff);
+		// Write placeholder.
+		Ar << SkipOffset;
+	}
+
+
+	if(Ref.Parameters)
+	{
+		Ref.Parameters->Serialize(Ar);
+	}
+	else if(Ar.IsLoading())
+	{
+		Ar.Seek( SkipOffset );
+	}
+
+	if( Ar.IsSaving() )
+	{
+		int64 EndOffset = Ar.Tell();
+		Ar.Seek( SkipOffset );
+		Ar << EndOffset;
+		Ar.Seek( EndOffset );
+	}
+
+	return bShaderHasOutdatedParameters;
+}
+
+/** Returns the hash of the vertex factory shader file that this shader was compiled with. */
+const FSHAHash& FVertexFactoryParameterRef::GetHash() const 
+{ 
+	return VFHash;
+}
+
+/** Returns the shader platform that this shader was compiled with. */
+EShaderPlatform FVertexFactoryParameterRef::GetShaderPlatform() const
+{
+	return ShaderPlatform;
+}
diff --git a/Engine/Source/Runtime/RenderCore/Public/GlobalShader.h b/Engine/Source/Runtime/RenderCore/Public/GlobalShader.h
new file mode 100644
--- /dev/null
+++ b/Engine/Source/Runtime/RenderCore/Public/GlobalShader.h
@@ -0,0 +1,286 @@
+// Copyright 1998-2018 Epic Games, Inc. All Rights Reserved.
+
+/*=============================================================================
+	GlobalShader.h: Shader manager definitions.
+=============================================================================*/
+
+#pragma once
+
+#include "CoreMinimal.h"
+#include "Templates/ScopedPointer.h"
+#include "Shader.h"
+#include "Templates/UniquePtr.h"
+#include "RHIResources.h"
+#include "Containers/Map.h"
+
+class FShaderCommonCompileJob;
+class FShaderCompileJob;
+
+/** Used to identify the global shader map in compile queues. */
+extern RENDERCORE_API const int32 GlobalShaderMapId;
+
+/** Class that encapsulates logic to create a DDC key for the global shader map. */
+class FGlobalShaderMapId
+{
+public:
+
+	/** Create a global shader map Id for the given platform. */
+	RENDERCORE_API FGlobalShaderMapId(EShaderPlatform Platform);
+
+	/** Append to a string that will be used as a DDC key. */
+	RENDERCORE_API void AppendKeyString(FString& KeyString, const TArray<FShaderTypeDependency>& Dependencies) const;
+
+	RENDERCORE_API const TMap<FString, TArray<FShaderTypeDependency>>& GetShaderFilenameToDependeciesMap() const { return ShaderFilenameToDependenciesMap; }
+
+private:
+	/** Shader types that this shader map is dependent on and their stored state. Mapped by shader filename, so every filename can have it's own DDC key. */
+	TMap<FString, TArray<FShaderTypeDependency>> ShaderFilenameToDependenciesMap;
+
+	/** Shader pipeline types that this shader map is dependent on and their stored state. */
+	TArray<FShaderPipelineTypeDependency> ShaderPipelineTypeDependencies;
+};
+
+struct FGlobalShaderPermutationParameters
+{
+	// Shader platform to compile to.
+	const EShaderPlatform Platform;
+
+	/** Unique permutation identifier of the global shader type. */
+	const int32 PermutationId;
+
+	FGlobalShaderPermutationParameters(EShaderPlatform InPlatform, int32 InPermutationId)
+		: Platform(InPlatform)
+		, PermutationId(InPermutationId)
+	{ }
+};
+
+/**
+ * A shader meta type for the simplest shaders; shaders which are not material or vertex factory linked.
+ * There should only a single instance of each simple shader type.
+ */
+class FGlobalShaderType : public FShaderType
+{
+	friend class FGlobalShaderTypeCompiler;
+public:
+
+	typedef FShader::CompiledShaderInitializerType CompiledShaderInitializerType;
+	typedef FShader* (*ConstructCompiledType)(const CompiledShaderInitializerType&);
+	typedef bool (*ShouldCompilePermutationType)(const FGlobalShaderPermutationParameters&);
+	typedef bool(*ValidateCompiledResultType)(EShaderPlatform, const FShaderParameterMap&, TArray<FString>&);
+	typedef void (*ModifyCompilationEnvironmentType)(const FGlobalShaderPermutationParameters&, FShaderCompilerEnvironment&);
+
+	FGlobalShaderType(
+		const TCHAR* InName,
+		const TCHAR* InSourceFilename,
+		const TCHAR* InFunctionName,
+		uint32 InFrequency,
+		int32 InTotalPermutationCount,
+		ConstructSerializedType InConstructSerializedRef,
+		ConstructCompiledType InConstructCompiledRef,
+		ModifyCompilationEnvironmentType InModifyCompilationEnvironmentRef,
+		ShouldCompilePermutationType InShouldCompilePermutationRef,
+		ValidateCompiledResultType InValidateCompiledResultRef,
+		GetStreamOutElementsType InGetStreamOutElementsRef
+		):
+		FShaderType(EShaderTypeForDynamicCast::Global, InName, InSourceFilename, InFunctionName, InFrequency, InTotalPermutationCount, InConstructSerializedRef, InGetStreamOutElementsRef),
+		ConstructCompiledRef(InConstructCompiledRef),
+		ShouldCompilePermutationRef(InShouldCompilePermutationRef),
+		ValidateCompiledResultRef(InValidateCompiledResultRef),
+		ModifyCompilationEnvironmentRef(InModifyCompilationEnvironmentRef)
+	{
+		checkf(FPaths::GetExtension(InSourceFilename) == TEXT("usf"),
+			TEXT("Incorrect virtual shader path extension for global shader '%s': Only .usf files should be "
+			     "compiled."),
+			InSourceFilename);
+	}
+
+	/**
+	 * Checks if the shader type should be cached for a particular platform.
+	 * @param Platform - The platform to check.
+	 * @return True if this shader type should be cached.
+	 */
+	bool ShouldCompilePermutation(EShaderPlatform Platform, int32 PermutationId) const
+	{
+		return (*ShouldCompilePermutationRef)(FGlobalShaderPermutationParameters(Platform, PermutationId));
+	}
+
+	/**
+	 * Sets up the environment used to compile an instance of this shader type.
+	 * @param Platform - Platform to compile for.
+	 * @param Environment - The shader compile environment that the function modifies.
+	 */
+	void SetupCompileEnvironment(EShaderPlatform Platform, int32 PermutationId, FShaderCompilerEnvironment& Environment)
+	{
+		// Allow the shader type to modify its compile environment.
+		(*ModifyCompilationEnvironmentRef)(FGlobalShaderPermutationParameters(Platform, PermutationId), Environment);
+	}
+
+	/**
+	* Checks if the shader type should pass compilation for a particular set of parameters.
+	* @param Platform - Platform to validate.
+	* @param ParameterMap - Shader parameters to validate.
+	* @param OutError - List for appending validation errors.
+	*/
+	bool ValidateCompiledResult(EShaderPlatform Platform, const FShaderParameterMap& ParameterMap, TArray<FString>& OutError) const
+	{
+		return (*ValidateCompiledResultRef)(Platform, ParameterMap, OutError);
+	}
+
+private:
+	ConstructCompiledType ConstructCompiledRef;
+	ShouldCompilePermutationType ShouldCompilePermutationRef;
+	ValidateCompiledResultType ValidateCompiledResultRef;
+	ModifyCompilationEnvironmentType ModifyCompilationEnvironmentRef;
+};
+
+extern RENDERCORE_API TShaderMap<FGlobalShaderType>* GGlobalShaderMap[SP_NumPlatforms];
+
+
+/**
+ * FGlobalShader
+ * 
+ * Global shaders derive from this class to set their default recompile group as a global one
+ */
+class FGlobalShader : public FShader
+{
+	DECLARE_SHADER_TYPE(FGlobalShader,Global);
+public:
+
+	FGlobalShader() : FShader() {}
+
+	RENDERCORE_API FGlobalShader(const ShaderMetaType::CompiledShaderInitializerType& Initializer);
+	
+	template<typename TViewUniformShaderParameters, typename ShaderRHIParamRef, typename TRHICmdList>
+	inline void SetParameters(TRHICmdList& RHICmdList, const ShaderRHIParamRef ShaderRHI, const FUniformBufferRHIParamRef ViewUniformBuffer)
+	{
+		const auto& ViewUniformBufferParameter = static_cast<const FShaderUniformBufferParameter&>(GetUniformBufferParameter<TViewUniformShaderParameters>());
+		SetUniformBufferParameter(RHICmdList, ShaderRHI, ViewUniformBufferParameter, ViewUniformBuffer);
+	}
+
+	typedef void (*ModifyCompilationEnvironmentType)(EShaderPlatform, FShaderCompilerEnvironment&);
+
+	// FShader interface.
+	RENDERCORE_API static void ModifyCompilationEnvironment(const FGlobalShaderPermutationParameters& Parameters, FShaderCompilerEnvironment& OutEnvironment) {}
+	RENDERCORE_API static bool ValidateCompiledResult(EShaderPlatform Platform, const FShaderParameterMap& ParameterMap, TArray<FString>& OutErrors) { return true; }
+};
+
+/**
+ * An internal dummy pixel shader to use when the user calls RHISetPixelShader(NULL).
+ */
+class FNULLPS : public FGlobalShader
+{
+	DECLARE_EXPORTED_SHADER_TYPE(FNULLPS,Global, RENDERCORE_API);
+public:
+
+	static bool ShouldCompilePermutation(const FGlobalShaderPermutationParameters& Parameters)
+	{
+		return true;
+	}
+
+	FNULLPS( )	{ }
+	FNULLPS(const ShaderMetaType::CompiledShaderInitializerType& Initializer):
+		FGlobalShader(Initializer)
+	{
+	}
+};
+
+/**
+* Container for Backup/RestoreGlobalShaderMap functions.
+* Includes shader data from any populated feature levels.
+*/
+struct FGlobalShaderBackupData
+{
+	TUniquePtr<TArray<uint8>> FeatureLevelShaderData[ERHIFeatureLevel::Num];
+};
+
+/** Backs up all global shaders to memory through serialization, and removes all references to FShaders from the global shader map. */
+extern RENDERCORE_API void BackupGlobalShaderMap(FGlobalShaderBackupData& OutGlobalShaderBackup);
+
+/** Recreates shaders in the global shader map from the serialized memory. */
+extern RENDERCORE_API void RestoreGlobalShaderMap(const FGlobalShaderBackupData& GlobalShaderData);
+
+/**
+ * Accesses the global shader map.  This is a global TShaderMap<FGlobalShaderType> which contains an instance of each global shader type.
+ *
+ * @param Platform Which platform's global shader map to use
+ * @param bRefreshShaderMap If true, the existing global shader map will be tossed first
+ * @return A reference to the global shader map.
+ */
+extern RENDERCORE_API TShaderMap<FGlobalShaderType>* GetGlobalShaderMap(EShaderPlatform Platform);
+
+/**
+  * Overload for the above GetGlobalShaderMap which takes a feature level and translates to the appropriate shader platform
+  *
+  * @param FeatureLevel - Which feature levels shader map to use
+  * @param bRefreshShaderMap If true, the existing global shader map will be tossed first
+  * @return A reference to the global shader map.
+  *
+  **/
+inline TShaderMap<FGlobalShaderType>* GetGlobalShaderMap(ERHIFeatureLevel::Type FeatureLevel) 
+{ 
+	return GetGlobalShaderMap(GShaderPlatformForFeatureLevel[FeatureLevel]); 
+}
+
+
+/** DECLARE_GLOBAL_SHADER and IMPLEMENT_GLOBAL_SHADER setup a global shader class's boiler plate. They are meant to be used like so:
+ *
+ * class FMyGlobalShaderPS : public FGlobalShader
+ * {
+ *		// Setup the shader's boiler plate.
+ *		DECLARE_GLOBAL_SHADER(FMyGlobalShaderPS);
+ *
+ *		// Setup the shader's permutation domain. If no dimensions, can do FPermutationDomain = FShaderPermutationNone.
+ *		using FPermutationDomain = TShaderPermutationDomain<DIMENSIONS...>;
+ *
+ *		// ...
+ * };
+ *
+ * // Instantiates global shader's global variable that will take care of compilation process of the shader. This needs imperatively to be
+ * done in a .cpp file regardless of whether FMyGlobalShaderPS is in a header or not.
+ * IMPLEMENT_GLOBAL_SHADER(FMyGlobalShaderPS, "/Engine/Private/MyShaderFile.usf", "MainPS", SF_Pixel);
+ */
+#define DECLARE_GLOBAL_SHADER(ShaderClass) \
+	public: \
+	\
+	using ShaderMetaType = FGlobalShaderType; \
+	\
+	static ShaderMetaType StaticType; \
+	\
+	static FShader* ConstructSerializedInstance() { return new ShaderClass(); } \
+	static FShader* ConstructCompiledInstance(const ShaderMetaType::CompiledShaderInitializerType& Initializer) \
+	{ return new ShaderClass(Initializer); } \
+	\
+	virtual uint32 GetTypeSize() const override { return sizeof(*this); } \
+	\
+	static bool ShouldCompilePermutationImpl(const FGlobalShaderPermutationParameters& Parameters) \
+	{ \
+		return ShaderClass::ShouldCompilePermutation(Parameters); \
+	} \
+	\
+	static bool ValidateCompiledResultImpl(EShaderPlatform Platform, const FShaderParameterMap& ParameterMap, TArray<FString>& OutErrors) \
+	{ \
+		return ShaderClass::ValidateCompiledResult(Platform, ParameterMap, OutErrors); \
+	} \
+	static void ModifyCompilationEnvironmentImpl( \
+		const FGlobalShaderPermutationParameters& Parameters, \
+		FShaderCompilerEnvironment& OutEnvironment) \
+	{ \
+		FPermutationDomain PermutationVector(Parameters.PermutationId); \
+		PermutationVector.ModifyCompilationEnvironment(OutEnvironment); \
+		ShaderClass::ModifyCompilationEnvironment(Parameters, OutEnvironment); \
+	}
+
+#define IMPLEMENT_GLOBAL_SHADER(ShaderClass,SourceFilename,FunctionName,Frequency) \
+	ShaderClass::ShaderMetaType ShaderClass::StaticType( \
+		TEXT(#ShaderClass), \
+		TEXT(SourceFilename), \
+		TEXT(FunctionName), \
+		Frequency, \
+		ShaderClass::FPermutationDomain::PermutationCount, \
+		ShaderClass::ConstructSerializedInstance, \
+		ShaderClass::ConstructCompiledInstance, \
+		ShaderClass::ModifyCompilationEnvironmentImpl, \
+		ShaderClass::ShouldCompilePermutationImpl, \
+		ShaderClass::ValidateCompiledResultImpl, \
+		ShaderClass::GetStreamOutElements \
+		)
diff --git a/Engine/Source/Runtime/RenderCore/Public/RenderGraph.h b/Engine/Source/Runtime/RenderCore/Public/RenderGraph.h
new file mode 100644
--- /dev/null
+++ b/Engine/Source/Runtime/RenderCore/Public/RenderGraph.h
@@ -0,0 +1,9 @@
+// Copyright 1998-2018 Epic Games, Inc. All Rights Reserved.
+
+#pragma once
+
+/** Includes everything you need to use the render graph framework */
+
+#include "RenderGraphResources.h"
+#include "RenderGraphBuilder.h"
+#include "RenderGraphUtils.h"
diff --git a/Engine/Source/Runtime/RenderCore/Public/RenderGraphBuilder.h b/Engine/Source/Runtime/RenderCore/Public/RenderGraphBuilder.h
new file mode 100644
--- /dev/null
+++ b/Engine/Source/Runtime/RenderCore/Public/RenderGraphBuilder.h
@@ -0,0 +1,269 @@
+// Copyright 1998-2018 Epic Games, Inc. All Rights Reserved.
+
+#pragma once
+
+#include "RenderGraphResources.h"
+#include "ShaderParameterMacros.h"
+#include "RendererInterface.h"
+
+
+/** Whether render graph debugging is compiled. */
+#define RENDER_GRAPH_DEBUGGING (DO_CHECK)
+
+
+/** Flags to anotate passes. */
+enum class ERenderGraphPassFlags
+{
+	None = 0,
+
+	/** Pass uses compute only */
+	Compute = 1 << 0,
+};
+
+ENUM_CLASS_FLAGS(ERenderGraphPassFlags)
+
+
+// TODO(RDG): remove from global scope?
+struct RENDERCORE_API FShaderParameterStructRef
+{
+	const void*						Contents;
+	const FRHIUniformBufferLayout*	Layout;
+
+	template<typename MemberType>
+	MemberType* GetMemberPtrAtOffset(uint16 Offset) const
+	{
+		return reinterpret_cast<MemberType*>(((uint8*)Contents) + Offset);
+	}
+};
+
+
+/** 
+ * Base class of a render graph pass
+ */
+struct RENDERCORE_API FRenderGraphPass
+{
+	FRenderGraphPass(const TCHAR* InName, FShaderParameterStructRef InParameterStruct, ERenderGraphPassFlags InPassFlags)
+		: Name(InName)
+		, ParameterStruct(InParameterStruct)
+		, PassFlags(InPassFlags)
+	{
+		if (IsCompute())
+		{
+			ensureMsgf(ParameterStruct.Layout->NumRenderTargets() == 0, TEXT("Pass %s was declared as ERenderGraphPassFlags::Compute yet has RenderTargets in its ResourceTable"), GetName());
+		}
+	}
+
+	virtual ~FRenderGraphPass() {}
+
+	virtual void Execute(FRHICommandListImmediate& RHICmdList) const = 0;
+
+	const TCHAR*			GetName() const { return Name; }
+	ERenderGraphPassFlags	GetFlags() const { return PassFlags; }
+	bool					IsCompute() const { return (PassFlags & ERenderGraphPassFlags::Compute) == ERenderGraphPassFlags::Compute; }
+	
+	FShaderParameterStructRef GetParameters() const { return ParameterStruct; }
+
+protected:
+	const TCHAR* const Name;
+	FShaderParameterStructRef ParameterStruct;
+	const ERenderGraphPassFlags PassFlags;
+};
+
+/** 
+ * Render graph pass with lambda execute function
+ */
+template <typename ParameterStructType, typename ExecuteLambdaType>
+struct TLambdaRenderPass final : public FRenderGraphPass
+{
+	TLambdaRenderPass(const TCHAR* InName, FShaderParameterStructRef InParameterStruct, ERenderGraphPassFlags InPassFlags, ExecuteLambdaType&& InExecuteLambda)
+		: FRenderGraphPass(InName, InParameterStruct, InPassFlags)
+		, ExecuteLambda(static_cast<ExecuteLambdaType&&>(InExecuteLambda))
+	{
+	}
+
+	~TLambdaRenderPass()
+	{
+		// Manually call the destructor of the pass parameter, to make sure RHI references are released since the pass parameters are allocated on FMemStack.
+		// TODO(RDG): this may lead to RHI resource leaks if a struct allocated in FMemStack does not actually get used through FRenderGraphBuilder::AddPass().
+		ParameterStructType* Struct = reinterpret_cast<ParameterStructType*>(const_cast<void*>(FRenderGraphPass::ParameterStruct.Contents));
+		Struct->~ParameterStructType();
+	}
+
+	virtual void Execute(FRHICommandListImmediate& RHICmdList) const override
+	{
+		ExecuteLambda(RHICmdList);
+	}
+
+	ExecuteLambdaType ExecuteLambda;
+};
+
+/** 
+ * Builds the per-frame render graph.
+ * Resources must be created from the builder before they can be bound to Pass ResourceTables.
+ * These resources are descriptors only until the graph is executed, where RHI resources are allocated as needed.
+ */
+class RENDERCORE_API FRenderGraphBuilder
+{
+public:
+	/** A RHI cmd list is required, if using the immediate mode. */
+	FRenderGraphBuilder(FRHICommandListImmediate& InRHICmdList)
+		: RHICmdList(InRHICmdList)
+	{ }
+
+	~FRenderGraphBuilder()
+	{
+		DestructPasses();
+	}
+
+	/** Register a external texture to be tracked by the render graph. */
+	inline const FGraphTexture* RegisterExternalTexture(const TRefCountPtr<IPooledRenderTarget>& ExternalPooledTexture, const TCHAR* Name = TEXT("External"))
+	{
+		ensureMsgf(ExternalPooledTexture.IsValid(), TEXT("Attempted to register NULL external texture: %s"), Name);
+		FGraphTexture* OutTexture = new(FMemStack::Get()) FGraphTexture(Name, ExternalPooledTexture->GetDesc());
+		OutTexture->PooledRenderTarget = ExternalPooledTexture;
+		AllocatedTextures.Add(OutTexture, ExternalPooledTexture);
+		#if DO_CHECK
+			Resources.Add(OutTexture);
+		#endif
+		return OutTexture;
+	}
+
+	/** Create graph tracked resource from a descriptor with a debug name. */
+	inline const FGraphTexture* CreateTexture(const FPooledRenderTargetDesc& Desc, const TCHAR* DebugName)
+	{
+		ensureMsgf(!bHasExecuted, TEXT("Render graph texture %s needs to be created before the builder execution."), DebugName);
+		FGraphTexture* Texture = new(FMemStack::Get()) FGraphTexture(DebugName, Desc);
+		#if DO_CHECK
+			Resources.Add(Texture);
+		#endif
+		return Texture;
+	}
+
+	/** Create graph tracked SRV from a descriptor. */
+	inline const FGraphSRV* CreateSRV(const FGraphSRVDesc& Desc)
+	{
+		check(Desc.Texture);
+		ensureMsgf(!bHasExecuted, TEXT("Render graph SRV %s needs to be created before the builder execution."), Desc.Texture->Name);
+		ensureMsgf(Desc.Texture->Desc.TargetableFlags & TexCreate_ShaderResource, TEXT("Attempted to create SRV from texture %s which was not created with TexCreate_ShaderResource"), Desc.Texture->Name);
+		
+		FGraphSRV* SRV = new(FMemStack::Get()) FGraphSRV(Desc.Texture->Name, Desc);
+		#if DO_CHECK
+			Resources.Add(SRV);
+		#endif
+		return SRV;
+	}
+
+	/** Create graph tracked UAV from a descriptor. */
+	inline const FGraphUAV* CreateUAV(const FGraphUAVDesc& Desc)
+	{
+		check(Desc.Texture);
+		ensureMsgf(!bHasExecuted, TEXT("Render graph UAV %s needs to be created before the builder execution."), Desc.Texture->Name);
+		ensureMsgf(Desc.Texture->Desc.TargetableFlags & TexCreate_UAV, TEXT("Attempted to create UAV from texture %s which was not created with TexCreate_UAV"), Desc.Texture->Name);
+		
+		FGraphUAV* UAV = new(FMemStack::Get()) FGraphUAV(Desc.Texture->Name, Desc);
+		#if DO_CHECK
+			Resources.Add(UAV);
+		#endif
+		return UAV;
+	}
+
+	/** Allocates parameter struct specifically to survive through the life time of the render graph. */
+	template< typename ParameterStructType >
+	inline void CreateParameters(ParameterStructType** OutParameterPtr) const
+	{
+		// Check because destructor called by the pass's destructor.
+		checkf(!bHasExecuted, TEXT("Render graph allocated shader parameters %s needs to be allocated before the builder execution."), ParameterStructType::FTypeInfo::GetStructMetadata()->GetStructTypeName());
+		*OutParameterPtr = new(FMemStack::Get()) ParameterStructType;
+		FMemory::Memzero( *OutParameterPtr, sizeof( ParameterStructType ) );
+	}
+
+	/** 
+	 * Adds a lambda pass to the graph.
+	 */
+	template<typename ParameterStructType, typename ExecuteLambdaType>
+	void AddPass(
+		const TCHAR* Name, 
+		ParameterStructType* ParameterStruct,
+		ERenderGraphPassFlags Flags,
+		ExecuteLambdaType&& ExecuteLambda)
+	{
+		checkf(!bHasExecuted, TEXT("Render graph pass %s needs to be added before the builder execution."), Name);
+		auto NewPass = new(FMemStack::Get()) TLambdaRenderPass<ParameterStructType, ExecuteLambdaType>(
+			Name,
+			{ ParameterStruct, &ParameterStructType::FTypeInfo::GetStructMetadata()->GetLayout() },
+			Flags,
+			static_cast<ExecuteLambdaType&&>(ExecuteLambda) );
+		Passes.Emplace(NewPass);
+		if (DO_CHECK)
+		{
+			DebugPass(NewPass);
+		}
+	}
+
+	/**
+	 * Extracts an internal texture by handle.  Must be called before Execute.
+	 */
+	inline void GetInternalTexture(const FGraphTexture* Texture, TRefCountPtr<IPooledRenderTarget>* OutTexturePtr, bool bTransitionToRead = true)
+	{
+		check(Texture);
+		check(OutTexturePtr);
+		checkf(!bHasExecuted, TEXT("Accessing render graph internal texture %s needs to happen before the builder's execution."), Texture->Name);
+		FDeferredInternalTextureQuery Query;
+		Query.Texture = Texture;
+		Query.OutTexturePtr = OutTexturePtr;
+		Query.bTransitionToRead = bTransitionToRead;
+		DeferredInternalTextureQueries.Emplace(Query);
+	}
+
+	/** 
+	 * Executes the queued passes, managing setting of render targets (RHI RenderPasses) and resource transitions.
+	 */
+	void Execute();
+
+private:
+	/** The RHI command list to use. */
+	FRHICommandListImmediate& RHICmdList;
+
+	/** Array of all pass created */
+	TArray<FRenderGraphPass*, SceneRenderingAllocator> Passes;
+
+	/** Keep the references over the pooled render target, since FGraphTexture is allocated on FMemStack. */
+	TMap<const FGraphTexture*, TRefCountPtr<IPooledRenderTarget>, SceneRenderingSetAllocator> AllocatedTextures;
+
+	/** Array of all deferred access to internal textures. */
+	struct FDeferredInternalTextureQuery
+	{
+		const FGraphTexture* Texture;
+		TRefCountPtr<IPooledRenderTarget>* OutTexturePtr;
+		bool bTransitionToRead;
+	};
+	TArray<FDeferredInternalTextureQuery, SceneRenderingAllocator> DeferredInternalTextureQueries;
+
+
+	#if DO_CHECK
+		/** Whether the Execute() has already been called. */
+		bool bHasExecuted = false;
+
+		/** Lists of all created resources */
+		TArray<const FGraphResource*, SceneRenderingAllocator> Resources;
+	#endif
+
+	void DebugPass(const FRenderGraphPass* Pass);
+	void ValidatePass(const FRenderGraphPass* Pass) const;
+
+	void WalkGraphDependencies();
+
+	void AllocateRHITextureIfNeeded(const FGraphTexture* Texture, bool bComputePass);
+
+	void TransitionTexture( const FGraphTexture* Texture, EResourceTransitionAccess TransitionAccess, bool bRequiredCompute ) const;
+	void TransitionUAV( const FGraphUAV* UAV, EResourceTransitionAccess TransitionAccess, bool bRequiredCompute ) const;
+
+	void ExecutePass( const FRenderGraphPass* Pass );
+	void AllocateAndTransitionPassResources(const FRenderGraphPass* Pass, struct FRHIRenderPassInfo* OutRPInfo, bool* bOutHasRenderTargets);
+	static void WarnForUselessPassDependencies(const FRenderGraphPass* Pass);
+	void ReleaseRHITextureIfPossible(const FGraphTexture* Texture);
+	void ReleaseUnecessaryResources(const FRenderGraphPass* Pass);
+
+	void ProcessDeferredInternalResourceQueries();
+	void DestructPasses();
+}; // class FRenderGraphBuilder
diff --git a/Engine/Source/Runtime/RenderCore/Public/RenderGraphResources.h b/Engine/Source/Runtime/RenderCore/Public/RenderGraphResources.h
new file mode 100644
--- /dev/null
+++ b/Engine/Source/Runtime/RenderCore/Public/RenderGraphResources.h
@@ -0,0 +1,175 @@
+// Copyright 1998-2018 Epic Games, Inc. All Rights Reserved.
+
+#pragma once
+
+#include "CoreMinimal.h"
+#include "RHI.h"
+#include "RendererInterface.h"
+
+
+/** Generic graph resource. */
+class RENDERCORE_API FGraphResource
+{
+public:
+	// Name of the resource for debugging purpose.
+	const TCHAR* const Name = nullptr;
+
+	FGraphResource(const TCHAR* InDebugName)
+		: Name(InDebugName)
+	{ }
+
+	FGraphResource() = delete;
+	FGraphResource(const FGraphResource&) = delete;
+	void operator = (const FGraphResource&) = delete;
+
+private:
+	/** Number of references in passes and deferred queries. */
+	mutable int32 ReferenceCount = 0;
+
+	// Used for tracking resource state during execution
+	mutable bool bWritable = false;
+	mutable bool bCompute = false;
+
+	/** Boolean to tracked whether a ressource is actually used by the lambda of a pass or not. */
+	mutable bool bIsActuallyUsedByPass = false;
+
+	friend class FRenderGraphBuilder;
+
+
+	/** Friendship over parameter settings for bIsActuallyUsedByPass. It means only this parameter settings can be
+	 * used for render graph resource, to force the useless dependency tracking to just work.
+	 */
+	template<typename TRHICmdList, typename TShaderClass, typename TShaderRHI>
+	friend void SetShaderParameters(TRHICmdList&, const TShaderClass*, TShaderRHI*, const typename TShaderClass::FParameters&);
+
+	template<typename TRHICmdList, typename TShaderClass>
+	friend void SetShaderUAVs(TRHICmdList&, const TShaderClass*, FComputeShaderRHIParamRef, const typename TShaderClass::FParameters&);
+};
+
+/** Render graph tracked Texture. */
+class RENDERCORE_API FGraphTexture : public FGraphResource
+{
+public:
+	//TODO(RDG): using FDesc = FPooledRenderTargetDesc;
+
+	/** Descriptor of the graph tracked texture. */
+	const FPooledRenderTargetDesc Desc;
+
+private:
+	/** This is not a TRefCountPtr<> because FGraphTexture is allocated on the FMemStack
+	 * FGraphBuilder::AllocatedTextures is actually keeping the reference.
+	 */
+	mutable IPooledRenderTarget* PooledRenderTarget = nullptr;
+
+	FGraphTexture(const TCHAR* DebugName, const FPooledRenderTargetDesc& InDesc)
+		: FGraphResource(DebugName)
+		, Desc(InDesc)
+	{ }
+
+	/** Returns the allocated RHI texture. */
+	inline FTextureRHIParamRef GetRHITexture() const
+	{
+		check(PooledRenderTarget);
+		return PooledRenderTarget->GetRenderTargetItem().ShaderResourceTexture;
+	}
+
+	friend class FRenderGraphBuilder;
+	friend class FGraphSRV;
+	friend class FGraphUAV;
+
+	template<typename TRHICmdList, typename TShaderClass, typename TShaderRHI>
+	friend void SetShaderParameters(TRHICmdList&, const TShaderClass*, TShaderRHI*, const typename TShaderClass::FParameters&);
+
+	template<typename TRHICmdList, typename TShaderClass>
+	friend void SetShaderUAVs(TRHICmdList&, const TShaderClass*, FComputeShaderRHIParamRef, const typename TShaderClass::FParameters&);
+};
+
+/** Decsriptor for render graph tracked SRV. */
+class RENDERCORE_API FGraphSRVDesc
+{
+public:
+	const FGraphTexture* Texture;
+	uint8 MipLevel = 0;
+
+	FGraphSRVDesc() {}
+
+	FGraphSRVDesc(const FGraphTexture* InTexture, uint8 InMipLevel) :
+		Texture(InTexture),
+		MipLevel(InMipLevel)
+	{}
+};
+
+/** Render graph tracked SRV. */
+class RENDERCORE_API FGraphSRV : public FGraphResource
+{
+public:
+	/** Descriptor of the graph tracked SRV. */
+	const FGraphSRVDesc Desc;
+
+private:
+	FGraphSRV(const TCHAR* DebugName, const FGraphSRVDesc& InDesc)
+		: FGraphResource(DebugName)
+		, Desc(InDesc)
+	{ }
+
+	/** Returns the allocated  shader resource view. */
+	inline FShaderResourceViewRHIParamRef GetRHIShaderResourceView() const
+	{
+		check(Desc.Texture);
+		check(Desc.Texture->PooledRenderTarget);
+		return Desc.Texture->PooledRenderTarget->GetRenderTargetItem().MipSRVs[Desc.MipLevel];
+	}
+
+	friend class FRenderGraphBuilder;
+
+	template<typename TRHICmdList, typename TShaderClass, typename TShaderRHI>
+	friend void SetShaderParameters(TRHICmdList&, const TShaderClass*, TShaderRHI*, const typename TShaderClass::FParameters&);
+
+	template<typename TRHICmdList, typename TShaderClass>
+	friend void SetShaderUAVs(TRHICmdList&, const TShaderClass*, FComputeShaderRHIParamRef, const typename TShaderClass::FParameters&);
+};
+
+/** Decsriptor for render graph tracked UAV. */
+class RENDERCORE_API FGraphUAVDesc
+{
+public:
+	const FGraphTexture* Texture;
+	uint8 MipLevel = 0;
+
+	FGraphUAVDesc() {}
+
+	FGraphUAVDesc(const FGraphTexture* InTexture, uint8 InMipLevel = 0) :
+		Texture(InTexture),
+		MipLevel(InMipLevel)
+	{}
+};
+
+/** Render graph tracked UAV. */
+class RENDERCORE_API FGraphUAV : public FGraphResource
+{
+public:
+	/** Descriptor of the graph tracked UAV. */
+	const FGraphUAVDesc Desc;
+
+private:
+	FGraphUAV(const TCHAR* DebugName, const FGraphUAVDesc& InDesc)
+		: FGraphResource(DebugName)
+		, Desc(InDesc)
+	{ }
+
+	/** Returns the allocated unordered access view. */
+	inline FUnorderedAccessViewRHIParamRef GetRHIUnorderedAccessView() const
+	{
+		check(Desc.Texture);
+		check(Desc.Texture->PooledRenderTarget);
+		return Desc.Texture->PooledRenderTarget->GetRenderTargetItem().MipUAVs[Desc.MipLevel];
+	}
+
+	friend class FRenderGraphBuilder;
+
+	template<typename TRHICmdList, typename TShaderClass, typename TShaderRHI>
+	friend void SetShaderParameters(TRHICmdList&, const TShaderClass*, TShaderRHI*, const typename TShaderClass::FParameters&);
+
+	template<typename TRHICmdList, typename TShaderClass>
+	friend void SetShaderUAVs(TRHICmdList&, const TShaderClass*, FComputeShaderRHIParamRef, const typename TShaderClass::FParameters&);
+};
diff --git a/Engine/Source/Runtime/RenderCore/Public/RenderGraphUtils.h b/Engine/Source/Runtime/RenderCore/Public/RenderGraphUtils.h
new file mode 100644
--- /dev/null
+++ b/Engine/Source/Runtime/RenderCore/Public/RenderGraphUtils.h
@@ -0,0 +1,28 @@
+// Copyright 1998-2018 Epic Games, Inc. All Rights Reserved.
+
+#pragma once
+
+#include "RenderGraphResources.h"
+#include "Shader.h"
+#include "UniformBuffer.h"
+
+
+/** Useful parameter struct that only have render targets.
+ *
+ *	FRenderTargetParameters PassParameters;
+ *	PassParameters.RenderTargets.DepthStencil = ... ;
+ *	PassParameters.RenderTargets[0] = ... ;
+ */
+BEGIN_SHADER_PARAMETER_STRUCT(FRenderTargetParameters, RENDERCORE_API)
+	RENDER_TARGET_BINDING_SLOTS()
+END_SHADER_PARAMETER_STRUCT()
+
+
+/** Clears all render graph tracked resources that are not used by a shader bindings. */
+extern RENDERCORE_API void ClearUnusedGraphResourcesImpl(const FShaderParameterBindings& ShaderBindings, const FShaderParametersMetadata* ParametersMetadata, void* InoutParameters);
+
+template<typename TShaderClass>
+inline void ClearUnusedGraphResources(const TShaderClass* Shader, typename TShaderClass::FParameters* InoutParameters)
+{
+	return ClearUnusedGraphResourcesImpl(Shader->Bindings, TShaderClass::FParameters::FTypeInfo::GetStructMetadata(), InoutParameters);
+}
