Merging //UE4/Dev-Rendering-Graph@4492585 to Dev-Rendering (//UE4/Dev-Rendering)

diff --git a/Engine/Source/Runtime/RenderCore/Public/RenderTargetPool.h b/Engine/Source/Runtime/RenderCore/Public/RenderTargetPool.h
new file mode 100644
--- /dev/null
+++ b/Engine/Source/Runtime/RenderCore/Public/RenderTargetPool.h
@@ -0,0 +1,366 @@
+// Copyright 1998-2018 Epic Games, Inc. All Rights Reserved.
+
+/*=============================================================================
+	RenderTargetPool.h: Scene render target pool manager.
+=============================================================================*/
+
+#pragma once
+
+#include "CoreMinimal.h"
+#include "RHI.h"
+#include "RenderResource.h"
+#include "RendererInterface.h"
+
+
+/** The reference to a pooled render target, use like this: TRefCountPtr<IPooledRenderTarget> */
+struct RENDERCORE_API FPooledRenderTarget : public IPooledRenderTarget
+{
+	FPooledRenderTarget(const FPooledRenderTargetDesc& InDesc, class FRenderTargetPool* InRenderTargetPool) 
+		: NumRefs(0)
+		, UnusedForNFrames(0)
+		, Desc(InDesc)
+		, bSnapshot(false)
+		, RenderTargetPool(InRenderTargetPool)
+		, FrameNumberLastDiscard(-1)
+	{
+	}
+	/* Constructor that makes a snapshot */
+	FPooledRenderTarget(const FPooledRenderTarget& SnaphotSource)
+		: NumRefs(1)
+		, UnusedForNFrames(0)
+		, Desc(SnaphotSource.Desc)
+		, bSnapshot(true)
+		, RenderTargetPool(SnaphotSource.RenderTargetPool)
+		, FrameNumberLastDiscard(-1)
+	{
+		check(IsInRenderingThread());
+		RenderTargetItem = SnaphotSource.RenderTargetItem;
+	}
+
+	virtual ~FPooledRenderTarget()
+	{
+		check(!NumRefs || (bSnapshot && NumRefs == 1));
+		RenderTargetItem.SafeRelease();
+	}
+
+	bool IsSnapshot() const 
+	{ 
+		return bSnapshot;
+	}
+
+	uint32 GetUnusedForNFrames() const 
+	{ 
+		check(!bSnapshot);
+		return UnusedForNFrames; 
+	}
+
+	// interface IPooledRenderTarget --------------
+
+	virtual uint32 AddRef() const override final;
+	virtual uint32 Release() override final;
+	virtual uint32 GetRefCount() const override final;
+	virtual bool IsFree() const override final;
+	virtual uint32 HasBeenDiscardedThisFrame() const
+	{
+		return GFrameNumberRenderThread == FrameNumberLastDiscard;
+	}
+	bool IsTransient() const
+	{
+		return !!(Desc.Flags & TexCreate_Transient);
+	}
+	virtual void SetDebugName(const TCHAR *InName);
+	virtual const FPooledRenderTargetDesc& GetDesc() const;
+	virtual uint32 ComputeMemorySize() const;
+
+// todo private:
+	FVRamAllocation VRamAllocation;
+
+private:
+
+	/** For pool management (only if NumRef == 0 the element can be reused) */
+	mutable int32 NumRefs;
+	/** Allows to defer the release to save performance on some hardware (DirectX) */
+	uint32 UnusedForNFrames;
+	/** All necessary data to create the render target */
+	FPooledRenderTargetDesc Desc;
+	/** Snapshots are sortof fake pooled render targets, they don't own anything and can outlive the things that created them. These are for threaded rendering. */
+	bool bSnapshot;
+
+	/** Pointer back to the pool for render targets which are actually pooled, otherwise NULL. */
+	FRenderTargetPool* RenderTargetPool;
+
+	/** Keeps track of the last frame we unmapped physical memory for this resource. We can't map again in the same frame if we did that */
+	uint32 FrameNumberLastDiscard;
+
+	/** @return true:release this one, false otherwise */
+	bool OnFrameStart();
+
+	friend class FRenderTargetPool;
+};
+
+enum ERenderTargetPoolEventType
+{
+	ERTPE_Alloc,
+	ERTPE_Dealloc,
+	ERTPE_Phase
+};
+
+struct RENDERCORE_API FRenderTargetPoolEvent
+{
+	// constructor for ERTPE_Alloc
+	FRenderTargetPoolEvent(uint32 InPoolEntryId, uint32 InTimeStep, FPooledRenderTarget* InPointer)
+		: PoolEntryId(InPoolEntryId)
+		, TimeStep(InTimeStep)
+		, Pointer(InPointer)
+		, EventType(ERTPE_Alloc)
+		, ColumnIndex(-1)
+		, ColumnX(0)
+		, ColumnSize(0)
+	{
+		Desc = Pointer->GetDesc();
+		SizeInBytes = Pointer->ComputeMemorySize();
+		VRamAllocation = Pointer->VRamAllocation;
+	}
+
+	// constructor for ERTPE_Alloc
+	FRenderTargetPoolEvent(uint32 InPoolEntryId, uint32 InTimeStep)
+		: PoolEntryId(InPoolEntryId)
+		, TimeStep(InTimeStep)
+		, Pointer(0)
+		, SizeInBytes(0)
+		, EventType(ERTPE_Dealloc)
+		, ColumnIndex(-1)
+		, ColumnX(0)
+		, ColumnSize(0)
+	{
+	}
+
+	// constructor for ERTPE_Alloc
+	// @param PhaseName must not be 0
+	FRenderTargetPoolEvent(const FString& InPhaseName, uint32 InTimeStep)
+		: PoolEntryId(-1)
+		, TimeStep(InTimeStep)
+		, Pointer(0)
+		, PhaseName(InPhaseName)
+		, SizeInBytes(0)
+		, EventType(ERTPE_Phase)
+		, ColumnIndex(-1)
+		, ColumnX(0)
+		, ColumnSize(0)
+	{
+	}
+
+	// @return Pointer if the object is still in the pool
+	IPooledRenderTarget* GetValidatedPointer() const;
+
+	ERenderTargetPoolEventType GetEventType() const { return EventType; }
+
+	const int32 GetPoolEntryId() const { check(EventType == ERTPE_Alloc || EventType == ERTPE_Dealloc); return PoolEntryId; }
+	const FString& GetPhaseName() const { check(EventType == ERTPE_Phase); return PhaseName; }
+	const FPooledRenderTargetDesc& GetDesc() const { check(EventType == ERTPE_Alloc || EventType == ERTPE_Dealloc); return Desc; }
+	const uint32 GetTimeStep() const { return TimeStep; }
+	const uint64 GetSizeInBytes() const { check(EventType == ERTPE_Alloc); return SizeInBytes; }
+	const void SetPoolEntryId(uint32 InPoolEntryId) { PoolEntryId = InPoolEntryId; }
+	const void SetColumn(uint32 InColumnIndex, uint32 InColumnX, uint32 InColumnSize) { check(EventType == ERTPE_Alloc || EventType == ERTPE_Dealloc); ColumnIndex = InColumnIndex; ColumnX = InColumnX; ColumnSize = InColumnSize; }
+	const uint32 GetColumnX() const { check(EventType == ERTPE_Alloc || EventType == ERTPE_Dealloc); return ColumnX; }
+	const uint32 GetColumnSize() const { check(EventType == ERTPE_Alloc || EventType == ERTPE_Dealloc); return ColumnSize; }
+	bool IsVisible() const { return EventType == ERTPE_Phase || ColumnSize > 0; }
+
+	void SetDesc(const FPooledRenderTargetDesc &InDesc) { Desc = InDesc; }
+
+	bool NeedsDeallocEvent();
+
+private:
+
+	// valid if EventType==ERTPE_Alloc || EventType==ERTPE_Dealloc, -1 if not set, was index into PooledRenderTargets[]
+	uint32 PoolEntryId;
+	//
+	uint32 TimeStep;
+	// valid EventType==ERTPE_Alloc, 0 if not set
+	FPooledRenderTarget* Pointer;
+	//
+	FVRamAllocation VRamAllocation;
+	// valid if EventType==ERTPE_Phase TEXT("") if not set
+	FString PhaseName;
+	// valid if EventType==ERTPE_Alloc || EventType==ERTPE_Dealloc
+	FPooledRenderTargetDesc Desc;
+	// valid if EventType==ERTPE_Alloc 0 if unknown
+	uint64 SizeInBytes;
+	// e.g. ERTPE_Alloc
+	ERenderTargetPoolEventType EventType;
+
+	// for display, computed by ComputeView()
+
+	// valid if EventType==ERTPE_Alloc || EventType==ERTPE_Dealloc, -1 if not defined yet
+	uint32 ColumnIndex;
+	//
+	uint32 ColumnX;
+	//
+	uint32 ColumnSize;
+};
+
+
+enum class ERenderTargetTransience : uint8
+{
+	NonTransient,
+	Transient,
+};
+
+/**
+ * Encapsulates the render targets pools that allows easy sharing (mostly used on the render thread side)
+ */
+class RENDERCORE_API FRenderTargetPool : public FRenderResource
+{
+public:
+	FRenderTargetPool();
+
+	/** Transitions all targets in the pool to writable. */
+	void TransitionTargetsWritable(FRHICommandListImmediate& RHICmdList);
+
+	/**
+	 * @param DebugName must not be 0, we only store the pointer
+	 * @param Out is not the return argument to avoid double allocation because of wrong reference counting
+	 * call from RenderThread only
+	 * @return true if the old element was still valid, false if a new one was assigned
+	 */
+	bool FindFreeElement(FRHICommandList& RHICmdList, const FPooledRenderTargetDesc& Desc, TRefCountPtr<IPooledRenderTarget>& Out, const TCHAR* InDebugName, bool bDoWritableBarrier = true, ERenderTargetTransience TransienceHint = ERenderTargetTransience::Transient );
+
+	void CreateUntrackedElement(const FPooledRenderTargetDesc& Desc, TRefCountPtr<IPooledRenderTarget>& Out, const FSceneRenderTargetItem& Item);
+
+	/** Destruct all snapshots, this must be done after all outstanding async tasks are done. It is important because they hold ref counted texture pointers etc **/
+	IPooledRenderTarget* MakeSnapshot(const TRefCountPtr<IPooledRenderTarget>& In);
+
+	/** Destruct all snapshots, this must be done after all outstanding async tasks are done. It is important because they hold ref counted texture pointers etc **/
+	void DestructSnapshots();
+
+	void OnRenderTargetUnreferenced(IPooledRenderTarget* RenderTarget);
+
+	/** Only to get statistics on usage and free elements. Normally only called in renderthread or if FlushRenderingCommands was called() */
+	void GetStats(uint32& OutWholeCount, uint32& OutWholePoolInKB, uint32& OutUsedInKB) const;
+	/**
+	 * Can release RT, should be called once per frame.
+	 * call from RenderThread only
+	 */
+	void TickPoolElements();
+	/** Free renderer resources */
+	void ReleaseDynamicRHI();
+
+	/** Allows to remove a resource so it cannot be shared and gets released immediately instead a/some frame[s] later. */
+	void FreeUnusedResource(TRefCountPtr<IPooledRenderTarget>& In);
+
+	/** Good to call between levels or before memory intense operations. */
+	void FreeUnusedResources();
+
+	// for debugging purpose, assumes you call FlushRenderingCommands() be
+	// @return can be 0, that doesn't mean iteration is done
+	FPooledRenderTarget* GetElementById(uint32 Id) const;
+
+	uint32 GetElementCount() const { return PooledRenderTargets.Num(); }
+
+	// @return -1 if not found
+	int32 FindIndex(IPooledRenderTarget* In) const;
+
+	// Render visualization
+	void RenderVisualizeTexture(class FDeferredShadingSceneRenderer& Scene);
+
+	void SetObserveTarget(const FString& InObservedDebugName, uint32 InObservedDebugNameReusedGoal = 0xffffffff);
+
+	// Logs out usage information.
+	void DumpMemoryUsage(FOutputDevice& OutputDevice);
+
+	// to not have event recording for some time during rendering (e.g. thumbnail rendering)
+	void SetEventRecordingActive(bool bValue) { bEventRecordingActive = bValue; }
+
+	//
+	void DisableEventDisplay() { RenderTargetPoolEvents.Empty(); bEventRecordingStarted = false; }
+	//
+	bool IsEventRecordingEnabled() const;
+
+	void AddPhaseEvent(const TCHAR* InPhaseName);
+
+	void UpdateElementSize(const TRefCountPtr<IPooledRenderTarget>& Element, const uint32 OldSize);
+
+private:
+
+	bool DoesTargetNeedTransienceOverride(const FPooledRenderTargetDesc& InputDesc, ERenderTargetTransience TransienceHint) const;
+
+	friend void RenderTargetPoolEvents(const TArray<FString>& Args);
+
+	/** Elements can be 0, we compact the buffer later. */
+	TArray< TRefCountPtr<FPooledRenderTarget> > PooledRenderTargets;
+	TArray< TRefCountPtr<FPooledRenderTarget> > DeferredDeleteArray;
+	TArray< FTextureRHIParamRef > TransitionTargets;	
+
+	/** These are snapshots, have odd life times, live in the scene allocator, and don't contribute to any accounting or other management. */
+	TArray<FPooledRenderTarget*> PooledRenderTargetSnapshots;
+
+	// redundant, can always be computed with GetStats(), to debug "out of memory" situations and used for r.RenderTargetPoolMin
+	uint32 AllocationLevelInKB;
+
+	FGraphEventRef TransitionFence;
+
+	// to avoid log spam
+	bool bCurrentlyOverBudget;
+
+	// for debugging purpose
+	void VerifyAllocationLevel() const;
+
+	// could be done on the fly but that makes the RenderTargetPoolEvents harder to read
+	void CompactPool();
+
+	void WaitForTransitionFence();
+
+	// the following is used for Event recording --------------------------------
+
+	struct SMemoryStats
+	{
+		SMemoryStats()
+			: DisplayedUsageInBytes(0)
+			, TotalUsageInBytes(0)
+		{
+		}
+		// for statistics
+		uint64 DisplayedUsageInBytes;
+		// for statistics
+		uint64 TotalUsageInBytes;
+		// for display purposes, to normalize the view width (Initialize to 1 to avoid a division by zero when compiled out)
+		uint64 TotalColumnSize = 1;
+	};
+
+	// if next frame we want to run with bEventRecording=true
+	bool bStartEventRecordingNextTick;
+	// in KB, e.g. 1MB = 1024, 0 to display all
+	uint32 EventRecordingSizeThreshold;
+	// true if active, to not have the event recording for some time during rendering (e.g. thumbnail rendering)
+	bool bEventRecordingActive;
+	// true meaning someone used r.RenderTargetPool.Events to start it
+	bool bEventRecordingStarted;
+	// only used if bEventRecording
+	TArray<FRenderTargetPoolEvent> RenderTargetPoolEvents;
+	//
+	uint32 CurrentEventRecordingTime;
+
+	//
+	void AddDeallocEvents();
+
+	// @param In must no be 0
+	void AddAllocEvent(uint32 InPoolEntryId, FPooledRenderTarget* In);
+
+	//
+	void AddAllocEventsFromCurrentState();
+
+	// @return 0 if none was found
+	const FString* GetLastEventPhaseName();
+
+	// sorted by size
+	SMemoryStats ComputeView();
+
+	// can be optimized to not be generated each time
+	void GenerateVRamAllocationUsage(TArray<FVRamAllocation>& Out);
+
+	friend struct FPooledRenderTarget;
+	friend class FVisualizeTexture;
+};
+
+/** The global render targets for easy shading. */
+extern RENDERCORE_API TGlobalResource<FRenderTargetPool> GRenderTargetPool;
diff --git a/Engine/Source/Runtime/RenderCore/Public/RendererInterface.h b/Engine/Source/Runtime/RenderCore/Public/RendererInterface.h
--- a/Engine/Source/Runtime/RenderCore/Public/RendererInterface.h
+++ b/Engine/Source/Runtime/RenderCore/Public/RendererInterface.h
@@ -55,6 +55,21 @@ ENUM_CLASS_FLAGS(EShowMaterialDrawEventTypes)
 // Shortcut for the allocator used by scene rendering.
 typedef TMemStackAllocator<> SceneRenderingAllocator;
 
+class SceneRenderingBitArrayAllocator
+	: public TInlineAllocator<4,SceneRenderingAllocator>
+{
+};
+
+class SceneRenderingSparseArrayAllocator
+	: public TSparseArrayAllocator<SceneRenderingAllocator,SceneRenderingBitArrayAllocator>
+{
+};
+
+class SceneRenderingSetAllocator
+	: public TSetAllocator<SceneRenderingSparseArrayAllocator,TInlineAllocator<1,SceneRenderingAllocator> >
+{
+};
+
 /** All necessary data to create a render target from the pooled render targets. */
 struct FPooledRenderTargetDesc
 {
diff --git a/Engine/Source/Runtime/RenderCore/Public/Shader.h b/Engine/Source/Runtime/RenderCore/Public/Shader.h
new file mode 100644
--- /dev/null
+++ b/Engine/Source/Runtime/RenderCore/Public/Shader.h
@@ -0,0 +1,2434 @@
+// Copyright 1998-2018 Epic Games, Inc. All Rights Reserved.
+
+/*=============================================================================
+	Shader.h: Shader definitions.
+=============================================================================*/
+
+#pragma once
+
+#include "CoreMinimal.h"
+#include "Containers/List.h"
+#include "Misc/SecureHash.h"
+#include "RenderResource.h"
+#include "ShaderParameters.h"
+#include "RenderingThread.h"
+#include "ShaderCore.h"
+#include "ShaderPermutation.h"
+#include "Serialization/ArchiveProxy.h"
+#include "UObject/RenderingObjectVersion.h"
+
+// For FShaderUniformBufferParameter
+
+#if WITH_EDITOR
+#include "UObject/DebugSerializationFlags.h"
+#endif
+
+class FGlobalShaderType;
+class FMaterialShaderType;
+class FNiagaraShaderType;
+class FMeshMaterialShaderType;
+class FShaderPipelineType;
+class FShaderType;
+class FVertexFactoryParameterRef;
+class FVertexFactoryType;
+
+
+/** Define a shader permutation uniquely according to its type, and permutation id.*/
+template<typename MetaShaderType>
+struct TShaderTypePermutation
+{
+	MetaShaderType* const Type;
+	const int32 PermutationId;
+
+	TShaderTypePermutation(MetaShaderType* InType, int32 InPermutationId)
+		: Type(InType)
+		, PermutationId(InPermutationId)
+	{
+	}
+
+	FORCEINLINE bool operator==(const TShaderTypePermutation& Other)const
+	{
+		return Type == Other.Type && PermutationId == Other.PermutationId;
+	}
+
+	FORCEINLINE bool operator!=(const TShaderTypePermutation& Other)const
+	{
+		return !(*this == Other);
+	}
+};
+
+using FShaderPermutation = TShaderTypePermutation<FShaderType>;
+
+template<typename MetaShaderType>
+FORCEINLINE uint32 GetTypeHash(const TShaderTypePermutation<MetaShaderType>& Var)
+{
+	return HashCombine(GetTypeHash(Var.Type), (uint32)Var.PermutationId);
+}
+
+/** Used to compare order shader types permutation deterministically. */
+template<typename MetaShaderType>
+class TCompareShaderTypePermutation
+{
+public:
+	FORCEINLINE bool operator()(const TShaderTypePermutation<MetaShaderType>& A, const TShaderTypePermutation<MetaShaderType>& B) const
+	{
+		int32 AL = FCString::Strlen(A.Type->GetName());
+		int32 BL = FCString::Strlen(B.Type->GetName());
+		if (AL == BL)
+		{
+			int32 StrCmp = FCString::Strncmp(A.Type->GetName(), B.Type->GetName(), AL);
+			if (StrCmp != 0)
+			{
+				return StrCmp > 0;
+			}
+			return A.PermutationId > B.PermutationId;
+		}
+		return AL > BL;
+	}
+};
+
+
+/** 
+ * Uniquely identifies an FShaderResource.  
+ * Used to link FShaders to FShaderResources on load. 
+ */
+class FShaderResourceId
+{
+public:
+
+	FShaderResourceId() {}
+
+	FShaderResourceId(const FShaderTarget& InTarget, const FSHAHash& InOutputHash, const TCHAR* InSpecificShaderTypeName, int32 InSpecificPermutationId) :
+		OutputHash(InOutputHash),
+		Target(InTarget),
+		SpecificShaderTypeName(InSpecificShaderTypeName),
+		SpecificPermutationId(InSpecificPermutationId)
+	{
+		check(!(SpecificShaderTypeName == nullptr && InSpecificPermutationId != 0));
+	}
+
+	friend inline uint32 GetTypeHash( const FShaderResourceId& Id )
+	{
+		return FCrc::MemCrc_DEPRECATED((const void*)&Id.OutputHash, sizeof(Id.OutputHash));
+	}
+
+	friend bool operator==(const FShaderResourceId& X, const FShaderResourceId& Y)
+	{
+		return X.Target == Y.Target 
+			&& X.OutputHash == Y.OutputHash 
+			&& X.SpecificPermutationId == Y.SpecificPermutationId
+			&& ((X.SpecificShaderTypeName == NULL && Y.SpecificShaderTypeName == NULL)
+				|| (FCString::Strcmp(X.SpecificShaderTypeName, Y.SpecificShaderTypeName) == 0));
+	}
+
+	friend bool operator!=(const FShaderResourceId& X, const FShaderResourceId& Y)
+	{
+		return !(X == Y);
+	}
+
+	friend FArchive& operator<<(FArchive& Ar, FShaderResourceId& Id)
+	{
+		Ar.UsingCustomVersion(FRenderingObjectVersion::GUID);
+
+		Ar << Id.Target << Id.OutputHash;
+
+		if (Ar.IsSaving())
+		{
+			Id.SpecificShaderTypeStorage = Id.SpecificShaderTypeName ? Id.SpecificShaderTypeName : TEXT("");
+		}
+
+		Ar << Id.SpecificShaderTypeStorage;
+
+		if (Ar.CustomVer(FRenderingObjectVersion::GUID) >= FRenderingObjectVersion::ShaderPermutationId)
+		{
+			Ar << Id.SpecificPermutationId;
+		}
+
+		if (Ar.IsLoading())
+		{
+			Id.SpecificShaderTypeName = *Id.SpecificShaderTypeStorage;
+
+			if (FCString::Strcmp(Id.SpecificShaderTypeName, TEXT("")) == 0)
+			{
+				// Store NULL for empty string to be consistent with FShaderResourceId's created at compile time
+				Id.SpecificShaderTypeName = NULL;
+			}
+		}
+
+		return Ar;
+	}
+
+	/** Hash of the compiled shader output, which is used to create the FShaderResource. */
+	FSHAHash OutputHash;
+
+	/** Target platform and frequency. */
+	FShaderTarget Target;
+
+	/** Stores the memory for SpecificShaderTypeName if this is a standalone Id, otherwise is empty and SpecificShaderTypeName points to an FShaderType name. */
+	FString SpecificShaderTypeStorage;
+
+	/** NULL if type doesn't matter, otherwise the name of the type that this was created specifically for, which is used with geometry shader stream out. */
+	const TCHAR* SpecificShaderTypeName;
+
+	/** Specific permutation identifier of the shader when SpecificShaderTypeName is non null, ignored otherwise. */
+	int32 SpecificPermutationId;
+};
+
+/** 
+ * Compiled shader bytecode and its corresponding RHI resource. 
+ * This can be shared by multiple FShaders with identical compiled output.
+ */
+class FShaderResource : public FRenderResource, public FDeferredCleanupInterface
+{
+	friend class FShader;
+public:
+
+	/** Constructor used for deserialization. */
+	RENDERCORE_API FShaderResource();
+
+	/** Constructor used when creating a new shader resource from compiled output. */
+	FShaderResource(const FShaderCompilerOutput& Output, FShaderType* InSpecificType, int32 InSpecificPermutationId);
+
+	~FShaderResource();
+
+	RENDERCORE_API void Serialize(FArchive& Ar, bool bLoadedByCookedMaterial);
+
+	// Reference counting.
+	RENDERCORE_API void AddRef();
+	RENDERCORE_API void Release();
+
+	RENDERCORE_API void Register();
+
+	/** @return the shader's vertex shader */
+	FORCEINLINE const FVertexShaderRHIParamRef GetVertexShader()
+	{
+		checkSlow(Target.Frequency == SF_Vertex);
+		if (!IsInitialized())
+		{
+			InitializeShaderRHI();
+		}
+		return (FRHIVertexShader*)Shader.GetReference();
+	}
+	/** @return the shader's pixel shader */
+	FORCEINLINE const FPixelShaderRHIParamRef GetPixelShader()
+	{
+		checkSlow(Target.Frequency == SF_Pixel);
+		if (!IsInitialized())
+		{
+			InitializeShaderRHI();
+		}
+		return (FRHIPixelShader*)Shader.GetReference();
+	}
+	/** @return the shader's hull shader */
+	FORCEINLINE const FHullShaderRHIParamRef GetHullShader()
+	{
+		checkSlow(Target.Frequency == SF_Hull);
+		if (!IsInitialized())
+		{
+			InitializeShaderRHI();
+		}
+		return (FRHIHullShader*)Shader.GetReference();
+	}
+	/** @return the shader's domain shader */
+	FORCEINLINE const FDomainShaderRHIParamRef GetDomainShader()
+	{
+		checkSlow(Target.Frequency == SF_Domain);
+		if (!IsInitialized())
+		{
+			InitializeShaderRHI();
+		}
+		return (FRHIDomainShader*)Shader.GetReference();
+	}
+	/** @return the shader's geometry shader */
+	FORCEINLINE const FGeometryShaderRHIParamRef GetGeometryShader()
+	{
+		checkSlow(Target.Frequency == SF_Geometry);
+		if (!IsInitialized())
+		{
+			InitializeShaderRHI();
+		}
+		return (FRHIGeometryShader*)Shader.GetReference();
+	}
+	/** @return the shader's compute shader */
+	FORCEINLINE const FComputeShaderRHIParamRef GetComputeShader()
+	{
+		checkSlow(Target.Frequency == SF_Compute);
+		if (!IsInitialized())
+		{
+			InitializeShaderRHI();
+		}
+		return (FRHIComputeShader*)Shader.GetReference();
+	}
+
+	RENDERCORE_API FShaderResourceId GetId() const;
+
+	uint32 GetSizeBytes() const
+	{
+		return Code.GetAllocatedSize() + sizeof(FShaderResource);
+	}
+
+	// FRenderResource interface.
+	virtual void InitRHI();
+	virtual void ReleaseRHI();
+
+	/** Finds a matching shader resource in memory if possible. */
+	RENDERCORE_API static FShaderResource* FindShaderResourceById(const FShaderResourceId& Id);
+
+	/** 
+	 * Finds a matching shader resource in memory or creates a new one with the given compiler output.  
+	 * SpecificType can be NULL
+	 */
+	RENDERCORE_API static FShaderResource* FindOrCreateShaderResource(const FShaderCompilerOutput& Output, class FShaderType* SpecificType, int32 SpecificPermutationId);
+
+	/** Return a list of all shader Ids currently known */
+	RENDERCORE_API static void GetAllShaderResourceId(TArray<FShaderResourceId>& Ids);
+
+	/** Returns true if and only if TargetPlatform is compatible for use with CurrentPlatform. */
+	RENDERCORE_API static bool ArePlatformsCompatible(EShaderPlatform CurrentPlatform, EShaderPlatform TargetPlatform);
+
+	void GetShaderCode(TArray<uint8>& OutCode) const
+	{
+		UncompressCode(OutCode);
+	}
+
+	/**
+	* Passes back a zeroed out hash to serialize when saving out cooked data.
+	* The goal here is to ensure that source hash changes do not cause widespread binary differences in cooked data, resulting in bloated patch diffs.
+	*/
+	RENDERCORE_API static FSHAHash &FilterShaderSourceHashForSerialization(const FArchive& Ar, FSHAHash &HashToSerialize);
+
+private:
+	// compression functions
+	void UncompressCode(TArray<uint8>& UncompressedCode) const;
+	void CompressCode(const TArray<uint8>& UncompressedCode);
+	
+
+	/** Conditionally serialize shader code. */
+	void SerializeShaderCode(FArchive& Ar);
+
+#if WITH_EDITORONLY_DATA
+	/** Conditionally serialize platform debug data. */
+	void SerializePlatformDebugData(FArchive& Ar);
+#endif
+
+	/**
+	 * Hash of the compiled bytecode and the generated parameter map.
+	 * This is used to find existing shader resources in memory or the DDC.
+	 */
+	FSHAHash OutputHash;
+
+	/** Compiled bytecode. */
+	TArray<uint8> Code;
+
+	/** Target platform and frequency. */
+	FShaderTarget Target;
+
+	/** Reference to the RHI shader. References the matching shader type of Target.Frequency. */
+	TRefCountPtr<FRHIShader> Shader;
+
+#if WITH_EDITORONLY_DATA
+	/** Platform specific debug data output by the shader compiler. Discarded in cooked builds. */
+	TArray<uint8> PlatformDebugData;
+#endif
+
+	/** Original bytecode size, before compression */
+	uint32 UncompressedCodeSize = 0;
+
+	/** If not NULL, the shader type this resource must be used with. */
+	class FShaderType* SpecificType;
+
+	/** Specific permutation identifier of the shader when SpecificType is non null, ignored otherwise. */
+	int32 SpecificPermutationId;
+
+	/** The number of references to this shader. */
+	mutable uint32 NumRefs;
+
+	/** The number of instructions the shader takes to execute. */
+	uint32 NumInstructions;
+
+#if WITH_EDITORONLY_DATA
+	/** Number of texture samplers the shader uses. */
+	uint32 NumTextureSamplers;
+#endif
+
+	/** Whether the shader code is stored in a shader library. */
+	bool bCodeInSharedLocation;
+	/** Whether the shader code was requested (and hence if we need to drop the ref later). */
+	bool bCodeInSharedLocationRequested;
+
+	/** Initialize the shader RHI resources. */
+	RENDERCORE_API void InitializeShaderRHI();
+
+	/** Tracks loaded shader resources by id. */
+	static TMap<FShaderResourceId, FShaderResource*> ShaderResourceIdMap;
+	/** Critical section for ShaderResourceIdMap. */
+	static FCriticalSection ShaderResourceIdMapCritical;
+};
+
+/** Encapsulates information about a shader's serialization behavior, used to detect when C++ serialization changes to auto-recompile. */
+class FSerializationHistory
+{
+public: 
+
+	/** Token stream stored as uint32's.  Each token is 4 bits, with a 0 meaning there's an associated 32 bit value in FullLengths. */
+	TArray<uint32> TokenBits;
+
+	/** Number of tokens in TokenBits. */
+	int32 NumTokens;
+
+	/** Full size length entries. One of these are used for every token with a value of 0. */
+	TArray<uint32> FullLengths;
+
+	FSerializationHistory() :
+		NumTokens(0)
+	{}
+
+	void AddValue(uint32 InValue)
+	{
+		const int32 UIntIndex = NumTokens / 8; 
+
+		if (UIntIndex >= TokenBits.Num())
+		{
+			// Add another uint32 if needed
+			TokenBits.AddZeroed();
+		}
+
+		uint8 Token = InValue;
+
+		// Anything that does not fit in 4 bits needs to go into FullLengths, with a special token value of 0
+		// InValue == 0 also should go into FullLengths, because its Token value is also 0
+		if (InValue > 7 || InValue == 0)
+		{
+			Token = 0;
+			FullLengths.Add(InValue);
+		}
+
+		const uint32 Shift = (NumTokens % 8) * 4;
+		// Add the new token bits into the existing uint32
+		TokenBits[UIntIndex] = TokenBits[UIntIndex] | (Token << Shift);
+		NumTokens++;
+	}
+
+	uint8 GetToken(int32 Index) const
+	{
+		check(Index < NumTokens);
+		const int32 UIntIndex = Index / 8; 
+		check(UIntIndex < TokenBits.Num());
+		const uint32 Shift = (Index % 8) * 4;
+		const uint8 Token = (TokenBits[UIntIndex] >> Shift) & 0xF;
+		return Token;
+	}
+
+	void AppendKeyString(FString& KeyString) const
+	{
+		KeyString += FString::FromInt(NumTokens);
+		KeyString += BytesToHex((uint8*)TokenBits.GetData(), TokenBits.Num() * TokenBits.GetTypeSize());
+		KeyString += BytesToHex((uint8*)FullLengths.GetData(), FullLengths.Num() * FullLengths.GetTypeSize());
+	}
+
+	inline bool operator==(const FSerializationHistory& Other) const
+	{
+		return TokenBits == Other.TokenBits && NumTokens == Other.NumTokens && FullLengths == Other.FullLengths;
+	}
+
+	friend FArchive& operator<<(FArchive& Ar,class FSerializationHistory& Ref)
+	{
+		Ar << Ref.TokenBits << Ref.NumTokens << Ref.FullLengths;
+		return Ar;
+	}
+};
+
+/** 
+ * Uniquely identifies an FShader instance.  
+ * Used to link FMaterialShaderMaps and FShaders on load. 
+ */
+class FShaderId
+{
+public:
+
+	/** 
+	 * Hash of the material shader map Id, since this shader depends on the generated material code from that shader map.
+	 * A hash is used instead of the full shader map Id to shorten the key length, even though this will result in a hash being hashed when we make a DDC key. 
+	 */ 
+	FSHAHash MaterialShaderMapHash;
+
+	/** Used to detect changes to the vertex factory source files. */
+	FSHAHash VFSourceHash;
+
+	/** Used to detect changes to the shader source files. */
+	FSHAHash SourceHash;
+
+	/** Shader platform and frequency. */
+	FShaderTarget Target;
+
+	/** Shader Pipeline linked to this shader, needed since a single shader might be used on different Pipelines. */
+	const FShaderPipelineType* ShaderPipeline;
+
+	/** 
+	 * Vertex factory type that the shader was created for, 
+	 * This is needed in the Id since a single shader type will be compiled for multiple vertex factories within a material shader map.
+	 * Will be NULL for global shaders.
+	 */
+	FVertexFactoryType* VertexFactoryType;
+
+	/** 
+	 * Used to detect changes to the vertex factory parameter class serialization, or NULL for global shaders. 
+	 * Note: This is referencing memory in the VF Type, since it is the same for all shaders using that VF Type.
+	 */
+	const FSerializationHistory* VFSerializationHistory;
+
+	/** Shader type */
+	FShaderType* ShaderType;
+
+	/** Unique permutation identifier within the ShaderType. */
+	int32 PermutationId;
+
+	/** Used to detect changes to the shader serialization.  Note: this is referencing memory in the FShaderType. */
+	const FSerializationHistory& SerializationHistory;
+
+	/** Create a minimally initialized Id.  Members will have to be assigned individually. */
+	FShaderId(const FSerializationHistory& InSerializationHistory)
+		: PermutationId(0)
+		, SerializationHistory(InSerializationHistory)
+	{}
+
+	/** Creates an Id for the given material, vertex factory, shader type and target. */
+	RENDERCORE_API FShaderId(const FSHAHash& InMaterialShaderMapHash, const FShaderPipelineType* InShaderPipeline, FVertexFactoryType* InVertexFactoryType, FShaderType* InShaderType, int32 PermutationId, FShaderTarget InTarget);
+
+	friend inline uint32 GetTypeHash( const FShaderId& Id )
+	{
+		return FCrc::MemCrc_DEPRECATED((const void*)&Id.MaterialShaderMapHash, sizeof(Id.MaterialShaderMapHash));
+	}
+
+	friend bool operator==(const FShaderId& X, const FShaderId& Y)
+	{
+		return X.MaterialShaderMapHash == Y.MaterialShaderMapHash
+			&& X.ShaderPipeline == Y.ShaderPipeline
+			&& X.VertexFactoryType == Y.VertexFactoryType
+			&& X.VFSourceHash == Y.VFSourceHash
+			&& ((X.VFSerializationHistory == NULL && Y.VFSerializationHistory == NULL)
+				|| (X.VFSerializationHistory != NULL && Y.VFSerializationHistory != NULL &&
+					*X.VFSerializationHistory == *Y.VFSerializationHistory))
+			&& X.ShaderType == Y.ShaderType
+			&& X.PermutationId == Y.PermutationId 
+			&& X.SourceHash == Y.SourceHash 
+			&& X.SerializationHistory == Y.SerializationHistory
+			&& X.Target == Y.Target;
+	}
+
+	friend bool operator!=(const FShaderId& X, const FShaderId& Y)
+	{
+		return !(X == Y);
+	}
+};
+
+/** Self contained version of FShaderId, which is useful for serializing. */
+class FSelfContainedShaderId
+{
+public:
+
+	/** 
+	 * Hash of the material shader map Id, since this shader depends on the generated material code from that shader map.
+	 * A hash is used instead of the full shader map Id to shorten the key length, even though this will result in a hash being hashed when we make a DDC key. 
+	 */ 
+	FSHAHash MaterialShaderMapHash;
+
+	/** Used to detect changes to the vertex factory source files. */
+	FSHAHash VFSourceHash;
+
+	/** Used to detect changes to the shader source files. */
+	FSHAHash SourceHash;
+
+	/** 
+	 * Name of the vertex factory type that the shader was created for, 
+	 * This is needed in the Id since a single shader type will be compiled for multiple vertex factories within a material shader map.
+	 * Will be the empty string for global shaders.
+	 */
+	FString VertexFactoryTypeName;
+
+	// Required to differentiate amongst unique shaders in the global map per Type
+	FString ShaderPipelineName;
+
+	/** Used to detect changes to the vertex factory parameter class serialization, or empty for global shaders. */
+	FSerializationHistory VFSerializationHistory;
+
+	/** Shader type name */
+	FString ShaderTypeName;
+
+	/** Unique permutation identifier within the ShaderType. */
+	int32 PermutationId;
+
+	/** Used to detect changes to the shader serialization. */
+	FSerializationHistory SerializationHistory;
+
+	/** Shader platform and frequency. */
+	FShaderTarget Target;
+
+	RENDERCORE_API FSelfContainedShaderId();
+
+	RENDERCORE_API FSelfContainedShaderId(const FShaderId& InShaderId);
+
+	RENDERCORE_API bool IsValid();
+
+	RENDERCORE_API friend FArchive& operator<<(FArchive& Ar,class FSelfContainedShaderId& Ref);
+};
+
+class FMaterial;
+
+
+/**
+ * Stores all shader parameter bindings and their corresponding offset and size in the shader's parameters struct.
+ */
+class RENDERCORE_API FShaderParameterBindings
+{
+public:
+	struct FParameter
+	{
+		uint16 BufferIndex;
+		uint16 BaseIndex;
+		uint16 ByteOffset;
+		uint16 ByteSize;
+
+		friend FArchive& operator<<(FArchive& Ar, FParameter& ParameterBindingData)
+		{
+			Ar << ParameterBindingData.BufferIndex << ParameterBindingData.BaseIndex << ParameterBindingData.ByteOffset << ParameterBindingData.ByteSize;
+			return Ar;
+		}
+	};
+
+	struct FResourceParameter
+	{
+		uint16 BaseIndex;
+		uint16 NumResources;
+		uint16 ByteOffset;
+
+		friend FArchive& operator<<(FArchive& Ar, FResourceParameter& ParameterBindingData)
+		{
+			Ar << ParameterBindingData.BaseIndex << ParameterBindingData.NumResources << ParameterBindingData.ByteOffset;
+			return Ar;
+		}
+	};
+
+	struct FParameterStructReference
+	{
+		uint16 BufferIndex;
+		uint16 ByteOffset;
+
+		friend FArchive& operator<<(FArchive& Ar, FParameterStructReference& ParameterBindingData)
+		{
+			Ar << ParameterBindingData.BufferIndex << ParameterBindingData.ByteOffset;
+			return Ar;
+		}
+	};
+
+	TArray<FParameter> Parameters;
+	TArray<FResourceParameter> Textures;
+	TArray<FResourceParameter> SRVs;
+	TArray<FResourceParameter> Samplers;
+	TArray<FResourceParameter> GraphTextures;
+	TArray<FResourceParameter> GraphSRVs;
+	TArray<FResourceParameter> GraphUAVs;
+	TArray<FParameterStructReference> ParameterReferences;
+
+	friend FArchive& operator<<(FArchive& Ar, FShaderParameterBindings& ParametersBindingData)
+	{
+		Ar << ParametersBindingData.Parameters;
+		Ar << ParametersBindingData.Textures;
+		Ar << ParametersBindingData.SRVs;
+		Ar << ParametersBindingData.Samplers;
+		Ar << ParametersBindingData.GraphTextures;
+		Ar << ParametersBindingData.GraphSRVs;
+		Ar << ParametersBindingData.GraphUAVs;
+		Ar << ParametersBindingData.ParameterReferences;
+		return Ar;
+	}
+
+	void Bind(const FShaderParameterMap& ParameterMaps, const FShaderParametersMetadata& StructMetaData);
+}; // FShaderParameterBindings
+
+
+/** A compiled shader and its parameter bindings. */
+class RENDERCORE_API FShader : public FDeferredCleanupInterface
+{
+	friend class FShaderType;
+public:
+
+	struct CompiledShaderInitializerType
+	{
+		FShaderType* Type;
+		FShaderTarget Target;
+		const TArray<uint8>& Code;
+		const FShaderParameterMap& ParameterMap;
+		const FSHAHash& OutputHash;
+		FShaderResource* Resource;
+		FSHAHash MaterialShaderMapHash;
+		const FShaderPipelineType* ShaderPipeline;
+		FVertexFactoryType* VertexFactoryType;
+		int32 PermutationId;
+
+		CompiledShaderInitializerType(
+			FShaderType* InType,
+			int32 InPermutationId,
+			const FShaderCompilerOutput& CompilerOutput,
+			FShaderResource* InResource,
+			const FSHAHash& InMaterialShaderMapHash,
+			const FShaderPipelineType* InShaderPipeline,
+			FVertexFactoryType* InVertexFactoryType
+			):
+			Type(InType),
+			Target(CompilerOutput.Target),
+			Code(CompilerOutput.ShaderCode.GetReadAccess()),
+			ParameterMap(CompilerOutput.ParameterMap),
+			OutputHash(CompilerOutput.OutputHash),
+			Resource(InResource),
+			MaterialShaderMapHash(InMaterialShaderMapHash),
+			ShaderPipeline(InShaderPipeline),
+			VertexFactoryType(InVertexFactoryType),
+			PermutationId(InPermutationId)
+		{}
+	};
+
+	/** 
+	 * Used to construct a shader for deserialization.
+	 * This still needs to initialize members to safe values since FShaderType::GenerateSerializationHistory uses this constructor.
+	 */
+	FShader();
+
+	/**
+	 * Construct a shader from shader compiler output.
+	 */
+	FShader(const CompiledShaderInitializerType& Initializer);
+
+	virtual ~FShader();
+
+/*
+	/ ** Can be overridden by FShader subclasses to modify their compile environment just before compilation occurs. * /
+	static void ModifyCompilationEnvironment(EShaderPlatform Platform, const FMaterial* Material, FShaderCompilerEnvironment& OutEnvironment){}
+*/
+
+	/** Serializes the shader. */
+	bool SerializeBase(FArchive& Ar, bool bShadersInline, bool bLoadedByCookedMaterial);
+
+	virtual bool Serialize(FArchive& Ar) { return false; }
+
+	// Reference counting.
+	void AddRef();
+	void Release();
+
+	/** Registers this shader for lookup by ID. */
+	void Register(bool bLoadedByCookedMaterial);
+
+	/** Removes this shader from the ID lookup map. */
+	void Deregister();
+
+	/** Returns the hash of the shader file that this shader was compiled with. */
+	const FSHAHash& GetHash() const;
+
+	/** Returns the shader platform of the shader file that this shader was compiled with. */
+	EShaderPlatform GetShaderPlatform() const;
+	
+	/** @return If the shader is linked with a vertex factory, returns the vertex factory's parameter object. */
+	virtual const FVertexFactoryParameterRef* GetVertexFactoryParameterRef() const { return NULL; }
+
+	/** @return the shader's vertex shader */
+	inline const FVertexShaderRHIParamRef GetVertexShader() const
+	{
+		return Resource->GetVertexShader();
+	}
+	/** @return the shader's pixel shader */
+	inline const FPixelShaderRHIParamRef GetPixelShader() const
+	{
+		return Resource->GetPixelShader();
+	}
+	/** @return the shader's hull shader */
+	inline const FHullShaderRHIParamRef GetHullShader() const
+	{
+		return Resource->GetHullShader();
+	}
+	/** @return the shader's domain shader */
+	inline const FDomainShaderRHIParamRef GetDomainShader() const
+	{
+		return Resource->GetDomainShader();
+	}
+	/** @return the shader's geometry shader */
+	inline const FGeometryShaderRHIParamRef GetGeometryShader() const
+	{
+		return Resource->GetGeometryShader();
+	}
+	/** @return the shader's compute shader */
+	inline const FComputeShaderRHIParamRef GetComputeShader() const
+	{
+		return Resource->GetComputeShader();
+	}
+
+	// Accessors.
+	inline FShaderType* GetType() const { return Type; }
+	inline int32 GetPermutationId() const { return PermutationId; }
+	inline uint32 GetNumInstructions() const { return Resource->NumInstructions; }
+	inline void SetNumInstructions(uint32 Num) { Resource->NumInstructions = Num; }
+#if WITH_EDITOR
+	inline uint32 GetNumTextureSamplers() const { return Resource->NumTextureSamplers; }
+#endif
+	inline const TArray<uint8>& GetCode() const { return Resource->Code; }
+	inline const FShaderTarget GetTarget() const { return Target; }
+	inline FSHAHash GetOutputHash() const { return OutputHash; }
+	FShaderId GetId() const;
+	inline FVertexFactoryType* GetVertexFactoryType() const { return VFType; }
+	inline int32 GetNumRefs() const { return NumRefs; }
+
+	inline FShaderResourceId GetResourceId() const
+	{
+		return Resource->GetId();
+	}
+
+	inline uint32 GetSizeBytes() const
+	{
+		return GetTypeSize() + GetAllocatedSize();
+	}
+	
+	/** Returns the size of the concrete type of this shader. */
+	virtual uint32 GetTypeSize() const
+	{
+		return sizeof(*this);
+	}
+
+	/** Returns the size of all allocations owned by this shader, e.g. TArrays. */
+	virtual uint32 GetAllocatedSize() const
+	{
+		return UniformBufferParameters.GetAllocatedSize() + UniformBufferParameterStructs.GetAllocatedSize();
+	}
+
+	uint32 GetResourceSizeBytes() const
+	{
+		return Resource->GetSizeBytes();
+	}
+
+	void SetResource(FShaderResource* InResource);
+
+	/** Called from the main thread to register and set the serialized resource */
+	void RegisterSerializedResource();
+
+	/** Implement for geometry shaders that want to use stream out. */
+	static void GetStreamOutElements(FStreamOutElementList& ElementList, TArray<uint32>& StreamStrides, int32& RasterizedStream) {}
+
+	void BeginInitializeResources()
+	{
+		BeginInitResource(Resource);
+	}
+
+	/** Finds an automatically bound uniform buffer matching the given uniform buffer type if one exists, or returns an unbound parameter. */
+	template<typename UniformBufferStructType>
+	FORCEINLINE_DEBUGGABLE const TShaderUniformBufferParameter<UniformBufferStructType>& GetUniformBufferParameter() const
+	{
+		const FShaderParametersMetadata* SearchStruct = &UniformBufferStructType::StaticStructMetadata;
+		int32 FoundIndex = INDEX_NONE;
+
+		for (int32 StructIndex = 0, Count = UniformBufferParameterStructs.Num(); StructIndex < Count; StructIndex++)
+		{
+			if (UniformBufferParameterStructs[StructIndex] == SearchStruct)
+			{
+				FoundIndex = StructIndex;
+				break;
+			}
+		}
+
+		if (FoundIndex != INDEX_NONE)
+		{
+			const TShaderUniformBufferParameter<UniformBufferStructType>& FoundParameter = (const TShaderUniformBufferParameter<UniformBufferStructType>&)*UniformBufferParameters[FoundIndex];
+			return FoundParameter;
+		}
+		else
+		{
+			// This can happen if the uniform buffer was not bound
+			// There's no good way to distinguish not being bound due to temporary debugging / compiler optimizations or an actual code bug,
+			// Hence failing silently instead of an error message
+			static TShaderUniformBufferParameter<UniformBufferStructType> UnboundParameter;
+			UnboundParameter.SetInitialized();
+			return UnboundParameter;
+		}
+	}
+
+	/** Finds an automatically bound uniform buffer matching the given uniform buffer struct if one exists, or returns an unbound parameter. */
+	const FShaderUniformBufferParameter& GetUniformBufferParameter(const FShaderParametersMetadata* SearchStruct) const
+	{
+		int32 FoundIndex = INDEX_NONE;
+
+		for (int32 StructIndex = 0, Count = UniformBufferParameterStructs.Num(); StructIndex < Count; StructIndex++)
+		{
+			if (UniformBufferParameterStructs[StructIndex] == SearchStruct)
+			{
+				FoundIndex = StructIndex;
+				break;
+			}
+		}
+
+		if (FoundIndex != INDEX_NONE)
+		{
+			const FShaderUniformBufferParameter& FoundParameter = *UniformBufferParameters[FoundIndex];
+			return FoundParameter;
+		}
+		else
+		{
+			static FShaderUniformBufferParameter UnboundParameter;
+			UnboundParameter.SetInitialized();
+			return UnboundParameter;
+		}
+	}
+
+	/** Gets the shader. */
+	inline FShader* GetShader()
+	{
+		return this;
+	}
+
+	/** Discards the serialized resource, used when the engine is using NullRHI */
+	void DiscardSerializedResource()
+	{
+		delete SerializedResource;
+		SerializedResource = nullptr;
+	}
+
+	void DumpDebugInfo();
+	void SaveShaderStableKeys(EShaderPlatform TargetShaderPlatform, const struct FStableShaderKeyAndValue& SaveKeyVal);
+
+protected:
+
+	/** Indexed the same as UniformBufferParameters.  Packed densely for coherent traversal. */
+	TArray<FShaderParametersMetadata*> UniformBufferParameterStructs;
+	TArray<FShaderUniformBufferParameter*> UniformBufferParameters;
+
+private:
+
+	/** 
+	 * Hash of the compiled output from this shader and the resulting parameter map.  
+	 * This is used to find a matching resource.
+	 */
+	FSHAHash OutputHash;
+
+	/** Hash of the material shader map this shader belongs to, stored so that an FShaderId can be constructed from this shader. */
+	FSHAHash MaterialShaderMapHash;
+
+	/** Vertex factory source hash, stored so that an FShaderId can be constructed from this shader. */
+	FSHAHash VFSourceHash;
+
+	/** Hash of this shader's source files generated at compile time, and stored to allow creating an FShaderId. */
+	FSHAHash SourceHash;
+
+	/** Reference to the shader resource, which stores the compiled bytecode and the RHI shader resource. */
+	TRefCountPtr<FShaderResource> Resource;
+
+	/** Pointer to the shader resource that has been serialized from disk, to be registered on the main thread later. */
+	FShaderResource* SerializedResource;
+
+	/** Shader pipeline this shader belongs to, stored so that an FShaderId can be constructed from this shader. */
+	const FShaderPipelineType* ShaderPipeline;
+
+	/** Vertex factory type this shader was created for, stored so that an FShaderId can be constructed from this shader. */
+	FVertexFactoryType* VFType;
+
+	/** Shader Type metadata for this shader. */
+	FShaderType* Type;
+
+	/** Unique permutation identifier of the shader in the shader type. */
+	int32 PermutationId;
+
+	/** Target platform and frequency. */
+	FShaderTarget Target;
+
+	/** The number of references to this shader. */
+	mutable uint32 NumRefs;
+
+public:
+	/** Shader parameter bindings. */
+	FShaderParameterBindings Bindings;
+};
+
+/**
+ * An object which is used to serialize/deserialize, compile, and cache a particular shader class.
+ *
+ * A shader type can manage multiple instance of FShader across mutiple dimensions such as EShaderPlatform, or permutation id.
+ * The number of permutation of a shader type is simply given by GetPermutationCount().
+ */
+class RENDERCORE_API FShaderType
+{
+public:
+	enum class EShaderTypeForDynamicCast : uint32
+	{
+		Global,
+		Material,
+		MeshMaterial,
+		Niagara
+	};
+
+	typedef class FShader* (*ConstructSerializedType)();
+	typedef void (*GetStreamOutElementsType)(FStreamOutElementList& ElementList, TArray<uint32>& StreamStrides, int32& RasterizedStream);
+
+	/** @return The global shader factory list. */
+	static TLinkedList<FShaderType*>*& GetTypeList();
+
+	static FShaderType* GetShaderTypeByName(const TCHAR* Name);
+	static TArray<FShaderType*> GetShaderTypesByFilename(const TCHAR* Filename);
+
+	/** @return The global shader name to type map */
+	static TMap<FName, FShaderType*>& GetNameToTypeMap();
+
+	/** Gets a list of FShaderTypes whose source file no longer matches what that type was compiled with */
+	static void GetOutdatedTypes(TArray<FShaderType*>& OutdatedShaderTypes, TArray<const FVertexFactoryType*>& OutdatedFactoryTypes);
+
+	/** Returns true if the source file no longer matches what that type was compiled with */
+	bool GetOutdatedCurrentType(TArray<FShaderType*>& OutdatedShaderTypes, TArray<const FVertexFactoryType*>& OutdatedFactoryTypes) const;
+
+	/** Initialize FShaderType static members, this must be called before any shader types are created. */
+	static void Initialize(const TMap<FString, TArray<const TCHAR*> >& ShaderFileToUniformBufferVariables);
+
+	/** Uninitializes FShaderType cached data. */
+	static void Uninitialize();
+
+	/** Minimal initialization constructor. */
+	FShaderType(
+		EShaderTypeForDynamicCast InShaderTypeForDynamicCast,
+		const TCHAR* InName,
+		const TCHAR* InSourceFilename,
+		const TCHAR* InFunctionName,
+		uint32 InFrequency,
+		int32 TotalPermutationCount,
+		ConstructSerializedType InConstructSerializedRef,
+		GetStreamOutElementsType InGetStreamOutElementsRef);
+
+	virtual ~FShaderType();
+
+	/**
+	 * Finds a shader of this type by ID.
+	 * @return NULL if no shader with the specified ID was found.
+	 */
+	FShader* FindShaderById(const FShaderId& Id);
+
+	/** Constructs a new instance of the shader type for deserialization. */
+	FShader* ConstructForDeserialization() const;
+
+	/** Calculates a Hash based on this shader type's source code and includes */
+	const FSHAHash& GetSourceHash(EShaderPlatform ShaderPlatform) const;
+
+	/** Serializes a shader type reference by name. */
+	RENDERCORE_API friend FArchive& operator<<(FArchive& Ar,FShaderType*& Ref);
+	
+	/** Hashes a pointer to a shader type. */
+	friend uint32 GetTypeHash(FShaderType* Ref)
+	{
+		return Ref ? Ref->HashIndex : 0;
+	}
+
+	// Dynamic casts.
+	FORCEINLINE FGlobalShaderType* GetGlobalShaderType() 
+	{ 
+		return (ShaderTypeForDynamicCast == EShaderTypeForDynamicCast::Global) ? reinterpret_cast<FGlobalShaderType*>(this) : nullptr;
+	}
+	FORCEINLINE const FGlobalShaderType* GetGlobalShaderType() const
+	{
+		return (ShaderTypeForDynamicCast == EShaderTypeForDynamicCast::Global) ? reinterpret_cast<const FGlobalShaderType*>(this) : nullptr;
+	}
+	FORCEINLINE FMaterialShaderType* GetMaterialShaderType()
+	{
+		return (ShaderTypeForDynamicCast == EShaderTypeForDynamicCast::Material) ? reinterpret_cast<FMaterialShaderType*>(this) : nullptr;
+	}
+	FORCEINLINE const FMaterialShaderType* GetMaterialShaderType() const
+	{
+		return (ShaderTypeForDynamicCast == EShaderTypeForDynamicCast::Material) ? reinterpret_cast<const FMaterialShaderType*>(this) : nullptr;
+	}
+	FORCEINLINE FMeshMaterialShaderType* GetMeshMaterialShaderType()
+	{
+		return (ShaderTypeForDynamicCast == EShaderTypeForDynamicCast::MeshMaterial) ? reinterpret_cast<FMeshMaterialShaderType*>(this) : nullptr;
+	}
+	FORCEINLINE const FMeshMaterialShaderType* GetMeshMaterialShaderType() const
+	{
+		return (ShaderTypeForDynamicCast == EShaderTypeForDynamicCast::MeshMaterial) ? reinterpret_cast<const FMeshMaterialShaderType*>(this) : nullptr;
+	}
+	FORCEINLINE const FNiagaraShaderType* GetNiagaraShaderType() const
+	{
+		return (ShaderTypeForDynamicCast == EShaderTypeForDynamicCast::Niagara) ? reinterpret_cast<const FNiagaraShaderType*>(this) : nullptr;
+	}
+	FORCEINLINE FNiagaraShaderType* GetNiagaraShaderType()
+	{
+		return (ShaderTypeForDynamicCast == EShaderTypeForDynamicCast::Niagara) ? reinterpret_cast<FNiagaraShaderType*>(this) : nullptr;
+	}
+
+	// Accessors.
+	inline EShaderFrequency GetFrequency() const
+	{ 
+		return (EShaderFrequency)Frequency; 
+	}
+	inline const TCHAR* GetName() const
+	{ 
+		return Name; 
+	}
+	inline const FName& GetFName() const
+	{
+		return TypeName;
+	}
+	inline const TCHAR* GetShaderFilename() const
+	{ 
+		return SourceFilename; 
+	}
+	inline const TCHAR* GetFunctionName() const
+	{
+		return FunctionName;
+	}
+	inline int32 GetNumShaders() const
+	{
+		return ShaderIdMap.Num();
+	}
+
+	inline int32 GetPermutationCount() const
+	{
+		return TotalPermutationCount;
+	}
+
+	inline const FSerializationHistory& GetSerializationHistory() const
+	{
+		return SerializationHistory;
+	}
+
+	inline const TMap<const TCHAR*, FCachedUniformBufferDeclaration>& GetReferencedUniformBufferStructsCache() const
+	{
+		return ReferencedUniformBufferStructsCache;
+	}
+
+	/** Adds include statements for uniform buffers that this shader type references, and builds a prefix for the shader file with the include statements. */
+	void AddReferencedUniformBufferIncludes(FShaderCompilerEnvironment& OutEnvironment, FString& OutSourceFilePrefix, EShaderPlatform Platform);
+
+	void FlushShaderFileCache(const TMap<FString, TArray<const TCHAR*> >& ShaderFileToUniformBufferVariables)
+	{
+		ReferencedUniformBufferStructsCache.Empty();
+		GenerateReferencedUniformBuffers(SourceFilename, Name, ShaderFileToUniformBufferVariables, ReferencedUniformBufferStructsCache);
+		bCachedUniformBufferStructDeclarations = false;
+	}
+
+	void AddToShaderIdMap(FShaderId Id, FShader* Shader)
+	{
+		check(IsInGameThread());
+		ShaderIdMap.Add(Id, Shader);
+	}
+
+	inline void RemoveFromShaderIdMap(FShaderId Id)
+	{
+		check(IsInGameThread());
+		ShaderIdMap.Remove(Id);
+	}
+
+	bool LimitShaderResourceToThisType() const
+	{
+		return GetStreamOutElementsRef != &FShader::GetStreamOutElements;
+	}
+
+	void GetStreamOutElements(FStreamOutElementList& ElementList, TArray<uint32>& StreamStrides, int32& RasterizedStream) 
+	{
+		(*GetStreamOutElementsRef)(ElementList, StreamStrides, RasterizedStream);
+	}
+
+	void DumpDebugInfo();
+	void SaveShaderStableKeys(EShaderPlatform TargetShaderPlatform);
+	void GetShaderStableKeyParts(struct FStableShaderKeyAndValue& SaveKeyVal);
+private:
+	EShaderTypeForDynamicCast ShaderTypeForDynamicCast;
+	uint32 HashIndex;
+	const TCHAR* Name;
+	FName TypeName;
+	const TCHAR* SourceFilename;
+	const TCHAR* FunctionName;
+	uint32 Frequency;
+	int32 TotalPermutationCount;
+
+	ConstructSerializedType ConstructSerializedRef;
+	GetStreamOutElementsType GetStreamOutElementsRef;
+
+	/** A map from shader ID to shader.  A shader will be removed from it when deleted, so this doesn't need to use a TRefCountPtr. */
+	TMap<FShaderId,FShader*> ShaderIdMap;
+
+	TLinkedList<FShaderType*> GlobalListLink;
+
+	// DumpShaderStats needs to access ShaderIdMap.
+	friend void RENDERCORE_API DumpShaderStats( EShaderPlatform Platform, EShaderFrequency Frequency );
+
+
+	/** 
+	 * Stores a history of serialization sizes for this shader type. 
+	 * This is used to invalidate shaders when serialization changes.
+	 */
+	FSerializationHistory SerializationHistory;
+
+	/** Tracks whether serialization history for all shader types has been initialized. */
+	static bool bInitializedSerializationHistory;
+
+protected:
+	/** Tracks what platforms ReferencedUniformBufferStructsCache has had declarations cached for. */
+	bool bCachedUniformBufferStructDeclarations;
+
+	/**
+	* Cache of referenced uniform buffer includes.
+	* These are derived from source files so they need to be flushed when editing and recompiling shaders on the fly.
+	* FShaderType::Initialize will add an entry for each referenced uniform buffer, but the declarations are added on demand as shaders are compiled.
+	*/
+	TMap<const TCHAR*, FCachedUniformBufferDeclaration> ReferencedUniformBufferStructsCache;
+
+};
+
+/**
+ * A macro to declare a new shader type.  This should be called in the class body of the new shader type.
+ * @param ShaderClass - The name of the class representing an instance of the shader type.
+ * @param ShaderMetaTypeShortcut - The shortcut for the shader meta type: simple, material, meshmaterial, etc.  The shader meta type
+ *	controls 
+ */
+#define DECLARE_EXPORTED_SHADER_TYPE(ShaderClass,ShaderMetaTypeShortcut,RequiredAPI, ...) \
+	public: \
+	using FPermutationDomain = FShaderPermutationNone; \
+	using ShaderMetaType = F##ShaderMetaTypeShortcut##ShaderType; \
+	\
+	static RequiredAPI ShaderMetaType StaticType; \
+	\
+	static FShader* ConstructSerializedInstance() { return new ShaderClass(); } \
+	static FShader* ConstructCompiledInstance(const ShaderMetaType::CompiledShaderInitializerType& Initializer) \
+	{ return new ShaderClass(Initializer); } \
+	\
+	virtual uint32 GetTypeSize() const override { return sizeof(*this); }
+
+#define DECLARE_SHADER_TYPE(ShaderClass,ShaderMetaTypeShortcut,...) \
+	DECLARE_EXPORTED_SHADER_TYPE(ShaderClass,ShaderMetaTypeShortcut,, ##__VA_ARGS__)
+
+#if !UE_BUILD_DOCS
+/** A macro to implement a shader type. */
+#define IMPLEMENT_SHADER_TYPE(TemplatePrefix,ShaderClass,SourceFilename,FunctionName,Frequency) \
+	TemplatePrefix \
+	ShaderClass::ShaderMetaType ShaderClass::StaticType( \
+		TEXT(#ShaderClass), \
+		SourceFilename, \
+		FunctionName, \
+		Frequency, \
+		1, \
+		ShaderClass::ConstructSerializedInstance, \
+		ShaderClass::ConstructCompiledInstance, \
+		ShaderClass::ModifyCompilationEnvironment, \
+		ShaderClass::ShouldCompilePermutation, \
+		ShaderClass::ValidateCompiledResult, \
+		ShaderClass::GetStreamOutElements \
+		);
+
+/** A macro to implement a shader type. Shader name is got from GetDebugName(), which is helpful for templated shaders. */
+#define IMPLEMENT_SHADER_TYPE_WITH_DEBUG_NAME(TemplatePrefix,ShaderClass,SourceFilename,FunctionName,Frequency) \
+	TemplatePrefix \
+	typename ShaderClass::ShaderMetaType ShaderClass::StaticType( \
+		ShaderClass::GetDebugName(), \
+		SourceFilename, \
+		FunctionName, \
+		Frequency, \
+		1, \
+		ShaderClass::ConstructSerializedInstance, \
+		ShaderClass::ConstructCompiledInstance, \
+		ShaderClass::ModifyCompilationEnvironment, \
+		ShaderClass::ShouldCompilePermutation, \
+		ShaderClass::ValidateCompiledResult, \
+		ShaderClass::GetStreamOutElements \
+		);
+
+/** A macro to implement a templated shader type, the function name and the source filename comes from the class. */
+#define IMPLEMENT_SHADER_TYPE2_WITH_TEMPLATE_PREFIX(TemplatePrefix,ShaderClass,Frequency) \
+	TemplatePrefix \
+	ShaderClass::ShaderMetaType ShaderClass::StaticType( \
+	TEXT(#ShaderClass), \
+	ShaderClass::GetSourceFilename(), \
+	ShaderClass::GetFunctionName(), \
+	Frequency, \
+	1, \
+	ShaderClass::ConstructSerializedInstance, \
+	ShaderClass::ConstructCompiledInstance, \
+	ShaderClass::ModifyCompilationEnvironment, \
+	ShaderClass::ShouldCompilePermutation, \
+	ShaderClass::ValidateCompiledResult, \
+	ShaderClass::GetStreamOutElements \
+	);
+
+#define IMPLEMENT_SHADER_TYPE2(ShaderClass,Frequency) \
+	IMPLEMENT_SHADER_TYPE2_WITH_TEMPLATE_PREFIX(template<>, ShaderClass, Frequency)
+
+/** todo: this should replace IMPLEMENT_SHADER_TYPE */
+#define IMPLEMENT_SHADER_TYPE3(ShaderClass,Frequency) \
+	ShaderClass::ShaderMetaType ShaderClass::StaticType( \
+	TEXT(#ShaderClass), \
+	ShaderClass::GetSourceFilename(), \
+	ShaderClass::GetFunctionName(), \
+	Frequency, \
+	1, \
+	ShaderClass::ConstructSerializedInstance, \
+	ShaderClass::ConstructCompiledInstance, \
+	ShaderClass::ModifyCompilationEnvironment, \
+	ShaderClass::ShouldCompilePermutation, \
+	ShaderClass::ValidateCompiledResult, \
+	ShaderClass::GetStreamOutElements \
+	);
+#endif
+
+
+// Binding of a set of shader stages in a single pipeline
+class RENDERCORE_API FShaderPipelineType
+{
+public:
+	// Set bShouldOptimizeUnusedOutputs to true if we want unique FShaders for each shader pipeline
+	// Set bShouldOptimizeUnusedOutputs to false if the FShaders will point to the individual shaders in the map
+	FShaderPipelineType(
+		const TCHAR* InName,
+		const FShaderType* InVertexShader,
+		const FShaderType* InHullShader,
+		const FShaderType* InDomainShader,
+		const FShaderType* InGeometryShader,
+		const FShaderType* InPixelShader,
+		bool bInShouldOptimizeUnusedOutputs);
+	~FShaderPipelineType();
+
+	FORCEINLINE bool HasTessellation() const { return AllStages[SF_Domain] != nullptr; }
+	FORCEINLINE bool HasGeometry() const { return AllStages[SF_Geometry] != nullptr; }
+	FORCEINLINE bool HasPixelShader() const { return AllStages[SF_Pixel] != nullptr; }
+
+	FORCEINLINE const FShaderType* GetShader(EShaderFrequency Frequency) const
+	{
+		check(Frequency < SF_NumFrequencies);
+		return AllStages[Frequency];
+	}
+
+	FORCEINLINE FName GetFName() const { return TypeName; }
+	FORCEINLINE TCHAR const* GetName() const { return Name; }
+
+	// Returns an array of valid stages, sorted from PS->GS->DS->HS->VS, no gaps if missing stages
+	FORCEINLINE const TArray<const FShaderType*>& GetStages() const { return Stages; }
+
+	static TLinkedList<FShaderPipelineType*>*& GetTypeList();
+
+	/** @return The global shader pipeline name to type map */
+	static TMap<FName, FShaderPipelineType*>& GetNameToTypeMap();
+	static const FShaderPipelineType* GetShaderPipelineTypeByName(FName Name);
+
+	/** Initialize static members, this must be called before any shader types are created. */
+	static void Initialize();
+	static void Uninitialize();
+
+	static TArray<const FShaderPipelineType*> GetShaderPipelineTypesByFilename(const TCHAR* Filename);
+
+	/** Serializes a shader type reference by name. */
+	RENDERCORE_API friend FArchive& operator<<(FArchive& Ar, const FShaderPipelineType*& Ref);
+
+	/** Hashes a pointer to a shader type. */
+	friend uint32 GetTypeHash(FShaderPipelineType* Ref) { return Ref ? Ref->HashIndex : 0; }
+	friend uint32 GetTypeHash(const FShaderPipelineType* Ref) { return Ref ? Ref->HashIndex : 0; }
+
+	// Check if this pipeline is built of specific types
+	bool IsGlobalTypePipeline() const { return Stages[0]->GetGlobalShaderType() != nullptr; }
+	bool IsMaterialTypePipeline() const { return Stages[0]->GetMaterialShaderType() != nullptr; }
+	bool IsMeshMaterialTypePipeline() const { return Stages[0]->GetMeshMaterialShaderType() != nullptr; }
+	bool ShouldOptimizeUnusedOutputs() const { return bShouldOptimizeUnusedOutputs; }
+
+	/** Gets a list of FShaderTypes & PipelineTypes whose source file no longer matches what that type was compiled with */
+	static void GetOutdatedTypes(TArray<FShaderType*>& OutdatedShaderTypes, TArray<const FShaderPipelineType*>& ShaderPipelineTypesToFlush, TArray<const FVertexFactoryType*>& OutdatedFactoryTypes);
+
+	/** Calculates a Hash based on this shader pipeline type stages' source code and includes */
+	const FSHAHash& GetSourceHash(EShaderPlatform ShaderPlatform) const;
+
+protected:
+	const TCHAR* const Name;
+	FName TypeName;
+
+	// Pipeline Stages, ordered from lowest (usually PS) to highest (VS). Guaranteed at least one stage (for VS).
+	TArray<const FShaderType*> Stages;
+
+	const FShaderType* AllStages[SF_NumFrequencies];
+
+	TLinkedList<FShaderPipelineType*> GlobalListLink;
+
+	uint32 HashIndex;
+	bool bShouldOptimizeUnusedOutputs;
+
+	static bool bInitialized;
+};
+
+#if !UE_BUILD_DOCS
+// Vertex+Pixel
+#define IMPLEMENT_SHADERPIPELINE_TYPE_VSPS(PipelineName, VertexShaderType, PixelShaderType, bRemoveUnused)	\
+	static FShaderPipelineType PipelineName(TEXT(PREPROCESSOR_TO_STRING(PipelineName)), &VertexShaderType::StaticType, nullptr, nullptr, nullptr, &PixelShaderType::StaticType, bRemoveUnused);
+// Only VS
+#define IMPLEMENT_SHADERPIPELINE_TYPE_VS(PipelineName, VertexShaderType, bRemoveUnused)	\
+	static FShaderPipelineType PipelineName(TEXT(PREPROCESSOR_TO_STRING(PipelineName)), &VertexShaderType::StaticType, nullptr, nullptr, nullptr, nullptr, bRemoveUnused);
+// Vertex+Geometry+Pixel
+#define IMPLEMENT_SHADERPIPELINE_TYPE_VSGSPS(PipelineName, VertexShaderType, GeometryShaderType, PixelShaderType, bRemoveUnused)	\
+	static FShaderPipelineType PipelineName(TEXT(PREPROCESSOR_TO_STRING(PipelineName)), &VertexShaderType::StaticType, nullptr, nullptr, &GeometryShaderType::StaticType, &PixelShaderType::StaticType, bRemoveUnused);
+// Vertex+Geometry
+#define IMPLEMENT_SHADERPIPELINE_TYPE_VSGS(PipelineName, VertexShaderType, GeometryShaderType, bRemoveUnused)	\
+	static FShaderPipelineType PipelineName(TEXT(PREPROCESSOR_TO_STRING(PipelineName)), &VertexShaderType::StaticType, nullptr, nullptr, &GeometryShaderType::StaticType, nullptr, bRemoveUnused);
+// Vertex+Hull+Domain+Pixel
+#define IMPLEMENT_SHADERPIPELINE_TYPE_VSHSDSPS(PipelineName, VertexShaderType, HullShaderType, DomainShaderType, PixelShaderType, bRemoveUnused)	\
+	static FShaderPipelineType PipelineName(TEXT(PREPROCESSOR_TO_STRING(PipelineName)), &VertexShaderType::StaticType, &HullShaderType::StaticType, &DomainShaderType::StaticType, nullptr, &PixelShaderType::StaticType, bRemoveUnused);
+// Vertex+Hull+Domain+Geometry+Pixel
+#define IMPLEMENT_SHADERPIPELINE_TYPE_VSHSDSGSPS(PipelineName, VertexShaderType, HullShaderType, DomainShaderType, GeometryShaderType, PixelShaderType, bRemoveUnused)	\
+	static FShaderPipelineType PipelineName(TEXT(PREPROCESSOR_TO_STRING(PipelineName)), &VertexShaderType::StaticType, &HullShaderType::StaticType, &DomainShaderType::StaticType, &GeometryShaderType::StaticType, &PixelShaderType::StaticType, bRemoveUnused);
+// Vertex+Hull+Domain
+#define IMPLEMENT_SHADERPIPELINE_TYPE_VSHSDS(PipelineName, VertexShaderType, HullShaderType, DomainShaderType, bRemoveUnused)	\
+	static FShaderPipelineType PipelineName(TEXT(PREPROCESSOR_TO_STRING(PipelineName)), &VertexShaderType::StaticType, &HullShaderType::StaticType, &DomainShaderType::StaticType, nullptr, nullptr, bRemoveUnused);
+// Vertex+Hull+Domain+Geometry
+#define IMPLEMENT_SHADERPIPELINE_TYPE_VSHSDSGS(PipelineName, VertexShaderType, HullShaderType, DomainShaderType, GeometryShaderType, bRemoveUnused)	\
+	static FShaderPipelineType PipelineName(TEXT(PREPROCESSOR_TO_STRING(PipelineName)), &VertexShaderType::StaticType, &HullShaderType::StaticType, &DomainShaderType::StaticType, &GeometryShaderType::StaticType, nullptr, bRemoveUnused);
+#endif
+
+/** Encapsulates a dependency on a shader type and saved state from that shader type. */
+class FShaderTypeDependency
+{
+public:
+
+	FShaderTypeDependency()
+		: ShaderType(nullptr)
+		, PermutationId(0)
+	{}
+
+	/** Shader type */
+	FShaderType* ShaderType;
+
+	/** Unique permutation identifier of the global shader type. */
+	int32 PermutationId;
+
+	/** Used to detect changes to the shader source files. */
+	FSHAHash SourceHash;
+
+	friend FArchive& operator<<(FArchive& Ar,class FShaderTypeDependency& Ref)
+	{
+		Ar.UsingCustomVersion(FRenderingObjectVersion::GUID);
+
+		Ar << Ref.ShaderType;
+		Ar << FShaderResource::FilterShaderSourceHashForSerialization(Ar, Ref.SourceHash);
+
+		if (Ar.CustomVer(FRenderingObjectVersion::GUID) >= FRenderingObjectVersion::ShaderPermutationId)
+		{
+			Ar << Ref.PermutationId;
+		}
+
+		return Ar;
+	}
+
+	bool operator==(const FShaderTypeDependency& Reference) const
+	{
+		return ShaderType == Reference.ShaderType && PermutationId == Reference.PermutationId && SourceHash == Reference.SourceHash;
+	}
+};
+
+
+class FShaderPipelineTypeDependency
+{
+public:
+	FShaderPipelineTypeDependency() :
+		ShaderPipelineType(nullptr)
+	{}
+
+	/** Shader Pipeline type */
+	const FShaderPipelineType* ShaderPipelineType;
+
+	/** Used to detect changes to the shader source files. */
+	FSHAHash StagesSourceHash;
+
+	friend FArchive& operator<<(FArchive& Ar, class FShaderPipelineTypeDependency& Ref)
+	{
+		Ar << Ref.ShaderPipelineType;
+		Ar << FShaderResource::FilterShaderSourceHashForSerialization(Ar, Ref.StagesSourceHash);
+		return Ar;
+	}
+
+	bool operator==(const FShaderPipelineTypeDependency& Reference) const
+	{
+		return ShaderPipelineType == Reference.ShaderPipelineType && StagesSourceHash == Reference.StagesSourceHash;
+	}
+};
+
+/** Used to compare two shader types by name. */
+class FCompareShaderTypes
+{																				
+public:
+	FORCEINLINE bool operator()(const FShaderType& A, const FShaderType& B ) const
+	{
+		int32 AL = FCString::Strlen(A.GetName());
+		int32 BL = FCString::Strlen(B.GetName());
+		if ( AL == BL )
+		{
+			return FCString::Strncmp(A.GetName(), B.GetName(), AL) > 0;
+		}
+		return AL > BL;
+	}
+};
+
+
+/** Used to compare two shader pipeline types by name. */
+class FCompareShaderPipelineNameTypes
+{
+public:
+	/*FORCEINLINE*/ bool operator()(const FShaderPipelineType& A, const FShaderPipelineType& B) const
+	{
+		//#todo-rco: Avoid this by adding an FNullShaderPipelineType
+		bool bNullA = &A == nullptr;
+		bool bNullB = &B == nullptr;
+		if (bNullA && bNullB)
+		{
+			return false;
+		}
+		else if (bNullA)
+		{
+			return true;
+		}
+		else if (bNullB)
+		{
+			return false;
+		}
+
+
+		int32 AL = FCString::Strlen(A.GetName());
+		int32 BL = FCString::Strlen(B.GetName());
+		if (AL == BL)
+		{
+			return FCString::Strncmp(A.GetName(), B.GetName(), AL) > 0;
+		}
+		return AL > BL;
+	}
+};
+
+// A Shader Pipeline instance with compiled stages
+class RENDERCORE_API FShaderPipeline
+{
+public:
+	const FShaderPipelineType* PipelineType;
+	TRefCountPtr<FShader> VertexShader;
+	TRefCountPtr<FShader> HullShader;
+	TRefCountPtr<FShader> DomainShader;
+	TRefCountPtr<FShader> GeometryShader;
+	TRefCountPtr<FShader> PixelShader;
+
+	FShaderPipeline(
+		const FShaderPipelineType* InPipelineType,
+		FShader* InVertexShader,
+		FShader* InHullShader,
+		FShader* InDomainShader,
+		FShader* InGeometryShader,
+		FShader* InPixelShader);
+
+	FShaderPipeline(const FShaderPipelineType* InPipelineType, const TArray<FShader*>& InStages);
+	FShaderPipeline(const FShaderPipelineType* InPipelineType, const TArray< TRefCountPtr<FShader> >& InStages);
+
+	~FShaderPipeline();
+
+	// Find a shader inside the pipeline
+	template<typename ShaderType>
+	ShaderType* GetShader()
+	{
+		if (PixelShader && PixelShader->GetType() == &ShaderType::StaticType)
+		{
+			return (ShaderType*)PixelShader.GetReference();
+		}
+		else if (VertexShader && VertexShader->GetType() == &ShaderType::StaticType)
+		{
+			return (ShaderType*)VertexShader.GetReference();
+		}
+		else if (GeometryShader && GeometryShader->GetType() == &ShaderType::StaticType)
+		{
+			return (ShaderType*)GeometryShader.GetReference();
+		}
+		else if (HullShader)
+		{
+			if (HullShader->GetType() == &ShaderType::StaticType)
+			{
+				return (ShaderType*)HullShader.GetReference();
+			}
+			else if (DomainShader && DomainShader->GetType() == &ShaderType::StaticType)
+			{
+				return (ShaderType*)DomainShader.GetReference();
+			}
+		}
+
+		return nullptr;
+	}
+
+	FShader* GetShader(EShaderFrequency Frequency)
+	{
+		switch (Frequency)
+		{
+		case SF_Vertex: return VertexShader.GetReference();
+		case SF_Domain: return DomainShader.GetReference();
+		case SF_Hull: return HullShader.GetReference();
+		case SF_Geometry: return GeometryShader.GetReference();
+		case SF_Pixel: return PixelShader.GetReference();
+		default: check(0);
+		}
+
+		return nullptr;
+	}
+
+	const FShader* GetShader(EShaderFrequency Frequency) const
+	{
+		switch (Frequency)
+		{
+		case SF_Vertex: return VertexShader.GetReference();
+		case SF_Domain: return DomainShader.GetReference();
+		case SF_Hull: return HullShader.GetReference();
+		case SF_Geometry: return GeometryShader.GetReference();
+		case SF_Pixel: return PixelShader.GetReference();
+		default: check(0);
+		}
+
+		return nullptr;
+	}
+
+	inline TArray<FShader*> GetShaders() const
+	{
+		TArray<FShader*> Shaders;
+
+		if (PixelShader)
+		{
+			Shaders.Add(PixelShader.GetReference());
+		}
+		if (GeometryShader)
+		{
+			Shaders.Add(GeometryShader.GetReference());
+		}
+		if (HullShader)
+		{
+			Shaders.Add(DomainShader.GetReference());
+			Shaders.Add(HullShader.GetReference());
+		}
+
+		Shaders.Add(VertexShader.GetReference());
+
+		return Shaders;
+	}
+
+	inline uint32 GetSizeBytes() const
+	{
+		return sizeof(*this);
+	}
+
+	void Validate();
+
+	enum EFilter
+	{
+		EAll,			// All pipelines
+		EOnlyShared,	// Only pipelines with shared shaders
+		EOnlyUnique,	// Only pipelines with unique shaders
+	};
+	
+	static void CookPipeline(FShaderPipeline* Pipeline);
+};
+
+inline bool operator<(const FShaderPipeline& Lhs, const FShaderPipeline& Rhs)
+{
+	FCompareShaderPipelineNameTypes Comparator;
+	return Comparator(*Lhs.PipelineType, *Rhs.PipelineType);
+}
+
+/** A collection of shaders of different types, but the same meta type. */
+template<typename ShaderMetaType>
+class TShaderMap
+{
+	/** Container for serialized shader pipeline stages to be registered on the game thread */
+	struct FSerializedShaderPipeline
+	{
+		const FShaderPipelineType* ShaderPipelineType;
+		TArray< TRefCountPtr<FShader> > ShaderStages;
+		FSerializedShaderPipeline()
+			: ShaderPipelineType(nullptr)
+		{
+		}
+	};
+
+	/** List of serialized shaders to be processed and registered on the game thread */
+	TArray<FShader*> SerializedShaders;
+	/** List of serialized shader pipeline stages to be processed and registered on the game thread */
+	TArray<FSerializedShaderPipeline*> SerializedShaderPipelines;
+protected:
+	/** The platform this shader map was compiled with */
+	EShaderPlatform Platform;
+private:
+	/** Flag that makes sure this shader map isn't used until all shaders have been registerd */
+	bool bHasBeenRegistered;
+
+public:
+
+	using FShaderPrimaryKey = TShaderTypePermutation<FShaderType>;
+
+	/** Used to compare two shader types by name. */
+	class FCompareShaderPrimaryKey
+	{
+	public:
+		FORCEINLINE bool operator()(const FShaderPrimaryKey& A, const FShaderPrimaryKey& B) const
+		{
+			int32 AL = FCString::Strlen(A.Type->GetName());
+			int32 BL = FCString::Strlen(B.Type->GetName());
+			if (AL == BL)
+			{
+				return FCString::Strncmp(A.Type->GetName(), B.Type->GetName(), AL) > 0 || A.PermutationId > B.PermutationId;
+			}
+			return AL > BL;
+		}
+	};
+
+	/** Default constructor. */
+	TShaderMap(EShaderPlatform InPlatform)
+		: Platform(InPlatform)
+		, bHasBeenRegistered(true)
+	{}
+
+	/** Destructor ensures pipelines cleared up. */
+	virtual ~TShaderMap()
+	{
+		EmptyShaderPipelines();
+	}
+
+	EShaderPlatform GetShaderPlatform() const { return Platform; }
+
+	/** Finds the shader with the given type.  Asserts on failure. */
+	template<typename ShaderType>
+	ShaderType* GetShader(int32 PermutationId = 0) const
+	{
+		check(bHasBeenRegistered);
+		const TRefCountPtr<FShader>* ShaderRef = Shaders.Find(FShaderPrimaryKey(&ShaderType::StaticType, PermutationId));
+		checkf(ShaderRef != NULL && *ShaderRef != nullptr, TEXT("Failed to find shader type %s in Platform %s"), ShaderType::StaticType.GetName(), *LegacyShaderPlatformToShaderFormat(Platform).ToString());
+		return (ShaderType*)((*ShaderRef)->GetShader());
+	}
+
+	/** Finds the shader with the given type.  Asserts on failure. */
+	template<typename ShaderType>
+	ShaderType* GetShader( const typename ShaderType::FPermutationDomain& PermutationVector ) const
+	{
+		return GetShader<ShaderType>( PermutationVector.ToDimensionValueId() );
+	}
+
+	/** Finds the shader with the given type.  May return NULL. */
+	FShader* GetShader(FShaderType* ShaderType, int32 PermutationId = 0) const
+	{
+		check(bHasBeenRegistered);
+		const TRefCountPtr<FShader>* ShaderRef = Shaders.Find(FShaderPrimaryKey(ShaderType, PermutationId));
+		return ShaderRef ? (*ShaderRef)->GetShader() : nullptr;
+	}
+
+	/** Finds the shader with the given type. */
+	bool HasShader(FShaderType* Type, int32 PermutationId) const
+	{
+		check(bHasBeenRegistered);
+		const TRefCountPtr<FShader>* ShaderRef = Shaders.Find(FShaderPrimaryKey(Type, PermutationId));
+		return ShaderRef != nullptr && *ShaderRef != nullptr;
+	}
+
+	inline const TMap<FShaderPrimaryKey,TRefCountPtr<FShader> >& GetShaders() const
+	{
+		check(bHasBeenRegistered);
+		return Shaders;
+	}
+
+	void AddShader(FShaderType* Type, int32 PermutationId, FShader* Shader)
+	{
+		check(Type);
+		Shaders.Add(FShaderPrimaryKey(Type, PermutationId), Shader);
+	}
+
+	/**
+	 * Removes the shader of the given type from the shader map
+	 * @param Type Shader type to remove the entry for 
+	 */
+	void RemoveShaderTypePermutaion(FShaderType* Type, int32 PermutationId)
+	{
+		Shaders.Remove(FShaderPrimaryKey(Type, PermutationId));
+	}
+
+
+	void RemoveShaderPipelineType(const FShaderPipelineType* ShaderPipelineType)
+	{
+		FShaderPipeline** Found = ShaderPipelines.Find(ShaderPipelineType);
+		if (Found)
+		{
+			if (*Found)
+			{
+				delete *Found;
+			}
+			ShaderPipelines.Remove(ShaderPipelineType);
+		}
+	}
+
+	/** Builds a list of the shaders in a shader map. */
+	void GetShaderList(TMap<FShaderId, FShader*>& OutShaders) const
+	{
+		check(bHasBeenRegistered);
+		for (typename TMap<FShaderPrimaryKey, TRefCountPtr<FShader>>::TConstIterator ShaderIt(Shaders); ShaderIt; ++ShaderIt)
+		{
+			if (ShaderIt.Value())
+			{
+				OutShaders.Add(ShaderIt.Value()->GetId(), ShaderIt.Value());
+			}
+		}
+	}
+
+	/** Builds a list of the shaders in a shader map. Key is FShaderType::TypeName */
+	void GetShaderList(TMap<FName, FShader*>& OutShaders) const
+	{
+		check(bHasBeenRegistered);
+		for (TMap<FShaderPrimaryKey, TRefCountPtr<FShader> >::TConstIterator ShaderIt(Shaders); ShaderIt; ++ShaderIt)
+		{
+			if (ShaderIt.Value())
+			{
+				OutShaders.Add(ShaderIt.Value()->GetType()->GetFName(), ShaderIt.Value());
+			}
+		}
+	}
+
+	/** Builds a list of the shader pipelines in a shader map. */
+	void GetShaderPipelineList(TArray<FShaderPipeline*>& OutShaderPipelines, FShaderPipeline::EFilter Filter) const
+	{
+		check(bHasBeenRegistered);
+		for (auto Pair : ShaderPipelines)
+		{
+			FShaderPipeline* Pipeline = Pair.Value;
+			if (Pipeline->PipelineType->ShouldOptimizeUnusedOutputs() && Filter == FShaderPipeline::EOnlyShared)
+			{
+				continue;
+			}
+			else if (!Pipeline->PipelineType->ShouldOptimizeUnusedOutputs() && Filter == FShaderPipeline::EOnlyUnique)
+			{
+				continue;
+			}
+			OutShaderPipelines.Add(Pipeline);
+		}
+	}
+
+#if WITH_EDITOR
+	uint32 GetMaxTextureSamplersShaderMap() const
+	{
+		check(bHasBeenRegistered);
+		uint32 MaxTextureSamplers = 0;
+
+		for (typename TMap<FShaderPrimaryKey,TRefCountPtr<FShader> >::TConstIterator ShaderIt(Shaders);ShaderIt;++ShaderIt)
+		{
+			if (ShaderIt.Value())
+			{
+				MaxTextureSamplers = FMath::Max(MaxTextureSamplers, ShaderIt.Value()->GetNumTextureSamplers());
+			}
+		}
+
+		for (auto Pair : ShaderPipelines)
+		{
+			const FShaderPipeline* Pipeline = Pair.Value;
+			for (const FShaderType* ShaderType : Pair.Key->GetStages())
+			{
+				MaxTextureSamplers = FMath::Max(MaxTextureSamplers, Pipeline->GetShader(ShaderType->GetFrequency())->GetNumTextureSamplers());
+			}
+		}
+
+		return MaxTextureSamplers;
+	}
+#endif // WITH_EDITOR
+
+	inline void SerializeShaderForSaving(FShader* CurrentShader, FArchive& Ar, bool bHandleShaderKeyChanges, bool bInlineShaderResource)
+	{
+		int64 SkipOffset = Ar.Tell();
+
+		{
+#if WITH_EDITOR
+			FArchive::FScopeSetDebugSerializationFlags S(Ar, DSF_IgnoreDiff);
+#endif
+			// Serialize a placeholder value, we will overwrite this with an offset to the end of the shader
+			Ar << SkipOffset;
+		}
+
+		if (bHandleShaderKeyChanges)
+		{
+			FSelfContainedShaderId SelfContainedKey = CurrentShader->GetId();
+			Ar << SelfContainedKey;
+		}
+
+		CurrentShader->SerializeBase(Ar, bInlineShaderResource, false);
+
+		// Get the offset to the end of the shader's serialized data
+		int64 EndOffset = Ar.Tell();
+		// Seek back to the placeholder and write the end offset
+		// This allows us to skip over the shader's serialized data at load time without knowing how to deserialize it
+		// Which can happen with shaders that were available at cook time, but not on the target platform (shaders in editor module for example)
+		Ar.Seek(SkipOffset);
+		Ar << EndOffset;
+		Ar.Seek(EndOffset);
+	}
+
+	inline FShader* SerializeShaderForLoad(FShaderType* Type, FArchive& Ar, bool bHandleShaderKeyChanges, bool bInlineShaderResource, bool bLoadedByCookedMaterial)
+	{
+		int64 EndOffset = 0;
+		Ar << EndOffset;
+
+		FSelfContainedShaderId SelfContainedKey;
+
+		if (bHandleShaderKeyChanges)
+		{
+			Ar << SelfContainedKey;
+		}
+
+		FShader* Shader = nullptr;
+		if (Type
+			// If we are handling shader key changes, only create the shader if the serialized key matches the key the shader would have if created
+			// This allows serialization changes between the save and load to be safely handled
+			&& (!bHandleShaderKeyChanges || SelfContainedKey.IsValid()))
+		{
+			Shader = Type->ConstructForDeserialization();
+			check(Shader != nullptr);
+			Shader->SerializeBase(Ar, bInlineShaderResource, bLoadedByCookedMaterial);
+		}
+		else
+		{
+			// Skip over this shader's serialized data if the type doesn't exist
+			// This can happen with shader types in modules that were loaded during cooking but not at run time (editor)
+			Ar.Seek(EndOffset);
+		}
+		return Shader;
+	}
+
+	/** 
+	 * Used to serialize a shader map inline in a material in a package. 
+	 * @param bInlineShaderResource - whether to inline the shader resource's serializations
+	 * @param bHandleShaderKeyChanges - whether to serialize the data necessary to detect and gracefully handle shader key changes between saving and loading
+	 * @param DependenciesToSave - array of specific ShaderTypeDepencies which should be saved.
+	 */
+	void SerializeInline(FArchive& Ar, bool bInlineShaderResource, bool bHandleShaderKeyChanges, bool bLoadedByCookedMaterial, const TArray<FShaderPrimaryKey>* ShaderKeysToSave = nullptr)
+	{
+		if (Ar.IsSaving())
+		{
+			TArray<FShaderPrimaryKey> SortedShaderKeys;
+
+			if (ShaderKeysToSave)
+			{
+				SortedShaderKeys = *ShaderKeysToSave;
+			}
+			else
+			{
+				Shaders.GenerateKeyArray(SortedShaderKeys);
+			}
+
+			int32 NumShaders = SortedShaderKeys.Num();
+			Ar << NumShaders;
+
+			// Sort the shaders by type name before saving, to make sure the saved result is binary equivalent to what is generated on other machines, 
+			// Which is a requirement of the Derived Data Cache.
+			SortedShaderKeys.Sort(FCompareShaderPrimaryKey());
+
+			for (FShaderPrimaryKey Key : SortedShaderKeys)
+			{
+				FShaderType* Type = Key.Type;
+				check(Type);
+				checkSlow(FName(Type->GetName()) != NAME_None);
+				Ar << Type;
+				FShader* CurrentShader = Shaders.FindChecked(Key);
+				SerializeShaderForSaving(CurrentShader, Ar, bHandleShaderKeyChanges, bInlineShaderResource);
+			}
+
+			TArray<FShaderPipeline*> SortedPipelines;
+			GetShaderPipelineList(SortedPipelines, FShaderPipeline::EAll);
+			int32 NumPipelines = SortedPipelines.Num();
+			Ar << NumPipelines;
+
+			checkf(!ShaderKeysToSave || NumPipelines == 0, TEXT("ShaderPipelines currently not supported for specific list of shader keys."));
+
+			// Sort the shader pipelines by type name before saving, to make sure the saved result is binary equivalent to what is generated on other machines, Which is a requirement of the Derived Data Cache.
+			SortedPipelines.Sort();
+			for (FShaderPipeline* CurrentPipeline : SortedPipelines)
+			{
+				const FShaderPipelineType* PipelineType = CurrentPipeline->PipelineType;
+				Ar << PipelineType;
+
+				auto& PipelineStages = PipelineType->GetStages();
+				int32 NumStages = PipelineStages.Num();
+				Ar << NumStages;
+				for (int32 Index = 0; Index < NumStages; ++Index)
+				{
+					auto* Shader = CurrentPipeline->GetShader(PipelineStages[Index]->GetFrequency());
+					FShaderType* Type = Shader->GetType();
+					Ar << Type;
+					SerializeShaderForSaving(Shader, Ar, bHandleShaderKeyChanges, bInlineShaderResource);
+				}
+#if WITH_EDITORONLY_DATA
+				if(Ar.IsCooking())
+				{
+					FShaderPipeline::CookPipeline(CurrentPipeline);
+				}
+#endif
+			}
+		}
+
+		if (Ar.IsLoading())
+		{
+			// Mark as unregistered - about to load new shaders that need to be registered later 
+			// on the game thread.
+			bHasBeenRegistered = false;
+
+			int32 NumShaders = 0;
+			Ar << NumShaders;
+
+			SerializedShaders.Reserve(NumShaders);
+			for (int32 ShaderIndex = 0; ShaderIndex < NumShaders; ShaderIndex++)
+			{
+				FShaderType* Type = nullptr;
+				Ar << Type;
+
+				FShader* Shader = SerializeShaderForLoad(Type, Ar, bHandleShaderKeyChanges, bInlineShaderResource, bLoadedByCookedMaterial);
+				if (Shader)
+				{
+					SerializedShaders.Add(Shader);
+				}
+			}
+
+			int32 NumPipelines = 0;
+			Ar << NumPipelines;
+			for (int32 PipelineIndex = 0; PipelineIndex < NumPipelines; ++PipelineIndex)
+			{
+				const FShaderPipelineType* ShaderPipelineType = nullptr;
+				Ar << ShaderPipelineType;
+				int32 NumStages = 0;
+				Ar << NumStages;
+				// Make a list of references so they can be deleted when going out of scope if needed
+				TArray< TRefCountPtr<FShader> > ShaderStages;
+				for (int32 Index = 0; Index < NumStages; ++Index)
+				{
+					FShaderType* Type = nullptr;
+					Ar << Type;
+					FShader* Shader = SerializeShaderForLoad(Type, Ar, bHandleShaderKeyChanges, bInlineShaderResource, bLoadedByCookedMaterial);
+					if (Shader)
+					{
+						ShaderStages.Add(Shader);
+					}
+				}
+
+				// ShaderPipelineType can be nullptr if the pipeline existed but now is gone!
+				if (ShaderPipelineType && ShaderStages.Num() == ShaderPipelineType->GetStages().Num())
+				{
+					FSerializedShaderPipeline* SerializedPipeline = new FSerializedShaderPipeline();
+					SerializedPipeline->ShaderPipelineType = ShaderPipelineType;
+					SerializedPipeline->ShaderStages = MoveTemp(ShaderStages);
+					SerializedShaderPipelines.Add(SerializedPipeline);
+				}
+			}
+		}
+	}
+
+	/** Registered all shaders that have been serialized (maybe) on another thread */
+	virtual void RegisterSerializedShaders(bool bCookedMaterial)
+	{
+		bHasBeenRegistered = true;
+		check(IsInGameThread());
+		for (FShader* Shader : SerializedShaders)
+		{
+			Shader->RegisterSerializedResource();
+
+			FShaderType* Type = Shader->GetType();
+			FShader* ExistingShader = Type->FindShaderById(Shader->GetId());
+
+			if (ExistingShader != nullptr)
+			{
+				delete Shader;
+				Shader = ExistingShader;
+			}
+			else
+			{
+				// Register the shader now that it is valid, so that it can be reused
+				Shader->Register(bCookedMaterial);
+			}
+			AddShader(Shader->GetType(), Shader->GetPermutationId(), Shader);
+		}
+		SerializedShaders.Empty();
+
+		for (FSerializedShaderPipeline* SerializedPipeline : SerializedShaderPipelines)
+		{
+			for (TRefCountPtr<FShader> Shader : SerializedPipeline->ShaderStages)
+			{
+				Shader->RegisterSerializedResource();
+			}
+			FShaderPipeline* ShaderPipeline = new FShaderPipeline(SerializedPipeline->ShaderPipelineType, SerializedPipeline->ShaderStages);
+			AddShaderPipeline(SerializedPipeline->ShaderPipelineType, ShaderPipeline);
+
+			delete SerializedPipeline;
+		}
+		SerializedShaderPipelines.Empty();
+	}
+
+	/** Discards serialized shaders when they are not going to be used for anything (NullRHI) */
+	virtual void DiscardSerializedShaders()
+	{
+		for (FShader* Shader : SerializedShaders)
+		{
+			if (Shader)
+			{
+				Shader->DiscardSerializedResource();
+			}
+			delete Shader;
+		}
+		SerializedShaders.Empty();
+
+		for (FSerializedShaderPipeline* SerializedPipeline : SerializedShaderPipelines)
+		{
+			for (TRefCountPtr<FShader> Shader : SerializedPipeline->ShaderStages)
+			{
+				Shader->DiscardSerializedResource();
+			}
+			delete SerializedPipeline;
+		}
+		SerializedShaderPipelines.Empty();
+	}
+
+	/** @return true if the map is empty */
+	inline bool IsEmpty() const
+	{
+		check(bHasBeenRegistered);
+		return Shaders.Num() == 0;
+	}
+
+	/** @return The number of shaders in the map. */
+	inline uint32 GetNumShaders() const
+	{
+		check(bHasBeenRegistered);
+		return Shaders.Num();
+	}
+
+	/** @return The number of shader pipelines in the map. */
+	inline uint32 GetNumShaderPipelines() const
+	{
+		check(bHasBeenRegistered);
+		return ShaderPipelines.Num();
+	}
+
+	/** clears out all shaders and deletes shader pipelines held in the map */
+	void Empty()
+	{
+		Shaders.Empty();
+		EmptyShaderPipelines();
+	}
+
+	inline FShaderPipeline* GetShaderPipeline(const FShaderPipelineType* PipelineType)
+	{
+		check(bHasBeenRegistered);
+		FShaderPipeline** Found = ShaderPipelines.Find(PipelineType);
+		return Found ? *Found : nullptr;
+	}
+
+	inline FShaderPipeline* GetShaderPipeline(const FShaderPipelineType* PipelineType) const
+	{
+		check(bHasBeenRegistered);
+		FShaderPipeline* const* Found = ShaderPipelines.Find(PipelineType);
+		return Found ? *Found : nullptr;
+	}
+
+	// Returns nullptr if not found
+	inline bool HasShaderPipeline(const FShaderPipelineType* PipelineType) const
+	{
+		check(bHasBeenRegistered);
+		return (GetShaderPipeline(PipelineType) != nullptr);
+	}
+
+	inline void AddShaderPipeline(const FShaderPipelineType* Type, FShaderPipeline* ShaderPipeline)
+	{
+		check(bHasBeenRegistered);
+		check(Type);
+		check(!ShaderPipeline || ShaderPipeline->PipelineType == Type);
+		ShaderPipelines.Add(Type, ShaderPipeline);
+	}
+
+	uint32 GetMaxNumInstructionsForShader(FShaderType* ShaderType) const
+	{
+		check(bHasBeenRegistered);
+		uint32 MaxNumInstructions = 0;
+		const TRefCountPtr<FShader>* FoundShader = Shaders.Find(FShaderPrimaryKey(ShaderType, 0));
+		if (FoundShader && *FoundShader)
+		{
+			MaxNumInstructions = FMath::Max(MaxNumInstructions, (*FoundShader)->GetNumInstructions());
+		}
+
+		for (auto& Pair : ShaderPipelines)
+		{
+			FShaderPipeline* Pipeline = Pair.Value;
+			auto* Shader = Pipeline->GetShader(ShaderType->GetFrequency());
+			if (Shader)
+			{
+				MaxNumInstructions = FMath::Max(MaxNumInstructions, Shader->GetNumInstructions());
+			}
+		}
+
+		return MaxNumInstructions;
+	}
+
+protected:
+	inline void EmptyShaderPipelines()
+	{
+		for (auto& Pair : ShaderPipelines)
+		{
+			if (FShaderPipeline* Pipeline = Pair.Value)
+			{
+				delete Pipeline;
+				Pipeline = nullptr;
+			}
+		}
+		ShaderPipelines.Empty();
+	}
+
+	TMap<FShaderPrimaryKey, TRefCountPtr<FShader> > Shaders;
+	TMap<const FShaderPipelineType*, FShaderPipeline*> ShaderPipelines;
+};
+
+/** A reference which is initialized with the requested shader type from a shader map. */
+template<typename ShaderType>
+class TShaderMapRef
+{
+public:
+	TShaderMapRef(const TShaderMap<typename ShaderType::ShaderMetaType>* ShaderIndex)
+		: Shader(ShaderIndex->template GetShader<ShaderType>(/* PermutationId = */ 0)) // gcc3 needs the template quantifier so it knows the < is not a less-than
+	{
+		static_assert(
+			TIsSame<typename ShaderType::FPermutationDomain, FShaderPermutationNone>::Value,
+			"Missing permutation vector argument for shader that have a permutation domain.");
+	}
+
+	TShaderMapRef(
+		const TShaderMap<typename ShaderType::ShaderMetaType>* ShaderIndex,
+		const typename ShaderType::FPermutationDomain& PermutationVector)
+		: Shader(ShaderIndex->template GetShader<ShaderType>(PermutationVector.ToDimensionValueId())) // gcc3 needs the template quantifier so it knows the < is not a less-than
+	{ }
+
+	FORCEINLINE ShaderType* operator->() const
+	{
+		return Shader;
+	}
+
+	FORCEINLINE ShaderType* operator*() const
+	{
+		return Shader;
+	}
+
+private:
+	ShaderType* Shader;
+};
+
+/** A reference to an optional shader, initialized with a shader type from a shader map if it is available or nullptr if it is not. */
+template<typename ShaderType>
+class TOptionalShaderMapRef
+{
+public:
+	TOptionalShaderMapRef(const TShaderMap<typename ShaderType::ShaderMetaType>* ShaderIndex):
+	Shader((ShaderType*)ShaderIndex->GetShader(&ShaderType::StaticType)) // gcc3 needs the template quantifier so it knows the < is not a less-than
+	{}
+	FORCEINLINE bool IsValid() const
+	{
+		return Shader != nullptr;
+	}
+	FORCEINLINE ShaderType* operator->() const
+	{
+		return Shader;
+	}
+	FORCEINLINE ShaderType* operator*() const
+	{
+		return Shader;
+	}
+private:
+	ShaderType* Shader;
+};
+
+/** Tracks state when traversing a FSerializationHistory. */
+class FSerializationHistoryTraversalState
+{
+public:
+
+	const FSerializationHistory& History;
+	int32 NextTokenIndex;
+	int32 NextFullLengthIndex;
+
+	FSerializationHistoryTraversalState(const FSerializationHistory& InHistory) :
+		History(InHistory),
+		NextTokenIndex(0),
+		NextFullLengthIndex(0)
+	{}
+
+	/** Gets the length value from NextTokenIndex + Offset into history. */
+	uint32 GetValue(int32 Offset)
+	{
+		int32 CurrentOffset = Offset;
+
+		// Move to the desired offset
+		while (CurrentOffset > 0)
+		{
+			StepForward();
+			CurrentOffset--;
+		}
+
+		while (CurrentOffset < 0)
+		{
+			StepBackward();
+			CurrentOffset++;
+		}
+		check(CurrentOffset == 0);
+
+		// Decode
+		const int8 Token = History.GetToken(NextTokenIndex);
+		const uint32 Value = Token == 0 ? History.FullLengths[NextFullLengthIndex] : (int32)Token;
+
+		// Restore state
+		while (CurrentOffset < Offset)
+		{
+			StepBackward();
+			CurrentOffset++;
+		}
+
+		while (CurrentOffset > Offset)
+		{
+			StepForward();
+			CurrentOffset--;
+		}
+		check(CurrentOffset == Offset);
+
+		return Value;
+	}
+
+	void StepForward()
+	{
+		const int8 Token = History.GetToken(NextTokenIndex);
+
+		if (Token == 0)
+		{
+			check(NextFullLengthIndex - 1 < History.FullLengths.Num());
+			NextFullLengthIndex++;
+		}
+
+		// Not supporting seeking past the front most serialization in the history
+		check(NextTokenIndex - 1 < History.NumTokens);
+		NextTokenIndex++;
+	}
+
+	void StepBackward()
+	{
+		// Not supporting seeking outside of the history tracked
+		check(NextTokenIndex > 0);
+		NextTokenIndex--;
+
+		const int8 Token = History.GetToken(NextTokenIndex);
+
+		if (Token == 0)
+		{
+			check(NextFullLengthIndex > 0);
+			NextFullLengthIndex--;
+		}
+	}
+};
+
+/** Archive used when saving shaders, which generates data used to detect serialization mismatches on load. */
+class FShaderSaveArchive : public FArchiveProxy
+{
+public:
+
+	FShaderSaveArchive(FArchive& Archive, FSerializationHistory& InHistory) : 
+		FArchiveProxy(Archive),
+		HistoryTraversalState(InHistory),
+		History(InHistory)
+	{
+		OriginalPosition = Archive.Tell();
+	}
+
+	virtual ~FShaderSaveArchive()
+	{
+		// Seek back to the original archive position so we can undo any serializations that went through this archive
+		InnerArchive.Seek(OriginalPosition);
+	}
+
+	virtual void Serialize( void* V, int64 Length )
+	{
+		if (HistoryTraversalState.NextTokenIndex < HistoryTraversalState.History.NumTokens)
+		{
+			// We are no longer appending (due to a seek), make sure writes match up in size with what's already been written
+			check(Length == HistoryTraversalState.GetValue(0));
+		}
+		else
+		{
+			// Appending to the archive, track the size of this serialization
+			History.AddValue(Length);
+		}
+		HistoryTraversalState.StepForward();
+		
+		if (V)
+		{
+			FArchiveProxy::Serialize(V, Length);
+		}
+	}
+
+	virtual void Seek( int64 InPos )
+	{
+		int64 Offset = InPos - Tell();
+		if (Offset <= 0)
+		{
+			// We're seeking backward, walk backward through the serialization history while updating NextSerialization
+			while (Offset < 0)
+			{
+				Offset += HistoryTraversalState.GetValue(-1);
+				HistoryTraversalState.StepBackward();
+			}
+		}
+		else
+		{
+			// We're seeking forward, walk forward through the serialization history while updating NextSerialization
+			while (Offset > 0)
+			{
+				Offset -= HistoryTraversalState.GetValue(-1);
+				HistoryTraversalState.StepForward();
+			}
+			HistoryTraversalState.StepForward();
+		}
+		check(Offset == 0);
+		
+		FArchiveProxy::Seek(InPos);
+	}
+
+	FSerializationHistoryTraversalState HistoryTraversalState;
+	FSerializationHistory& History;
+
+private:
+	/** Stored off position of the original archive we are wrapping. */
+	int64 OriginalPosition;
+};
+
+/**
+ * Dumps shader stats to the log. Will also print some shader pipeline information.
+ * @param Platform  - Platform to dump shader info for, use SP_NumPlatforms for all
+ * @param Frequency - Whether to dump PS or VS info, use SF_NumFrequencies to dump both
+ */
+extern RENDERCORE_API void DumpShaderStats( EShaderPlatform Platform, EShaderFrequency Frequency );
+
+/**
+ * Dumps shader pipeline stats to the log. Does not include material (eg shader pipeline instance) information.
+ * @param Platform  - Platform to dump shader info for, use SP_NumPlatforms for all
+ */
+extern RENDERCORE_API void DumpShaderPipelineStats(EShaderPlatform Platform);
+
+/**
+ * Finds the shader type with a given name.
+ * @param ShaderTypeName - The name of the shader type to find.
+ * @return The shader type, or NULL if none matched.
+ */
+extern RENDERCORE_API FShaderType* FindShaderTypeByName(FName ShaderTypeName);
+
+/** Helper function to dispatch a compute shader while checking that parameters have been set correctly. */
+extern RENDERCORE_API void DispatchComputeShader(
+	FRHICommandList& RHICmdList,
+	FShader* Shader,
+	uint32 ThreadGroupCountX,
+	uint32 ThreadGroupCountY,
+	uint32 ThreadGroupCountZ);
+
+extern RENDERCORE_API void DispatchComputeShader(
+	FRHIAsyncComputeCommandListImmediate& RHICmdList,
+	FShader* Shader,
+	uint32 ThreadGroupCountX,
+	uint32 ThreadGroupCountY,
+	uint32 ThreadGroupCountZ);
+
+/** Helper function to dispatch a compute shader indirectly while checking that parameters have been set correctly. */
+extern RENDERCORE_API void DispatchIndirectComputeShader(
+	FRHICommandList& RHICmdList,
+	FShader* Shader,
+	FVertexBufferRHIParamRef ArgumentBuffer,
+	uint32 ArgumentOffset);
+
+/** Appends to KeyString for all shaders. */
+extern RENDERCORE_API void ShaderMapAppendKeyString(EShaderPlatform Platform, FString& KeyString);
