Merging //UE4/Dev-Rendering-Graph@4492585 to Dev-Rendering (//UE4/Dev-Rendering)

diff --git a/Engine/Source/Runtime/RenderCore/Public/ShaderCodeLibrary.h b/Engine/Source/Runtime/RenderCore/Public/ShaderCodeLibrary.h
new file mode 100644
--- /dev/null
+++ b/Engine/Source/Runtime/RenderCore/Public/ShaderCodeLibrary.h
@@ -0,0 +1,251 @@
+// Copyright 1998-2018 Epic Games, Inc. All Rights Reserved.
+
+/*=============================================================================
+	ShaderCodeLibrary.h: 
+=============================================================================*/
+
+#pragma once
+
+#include "CoreMinimal.h"
+#include "RHI.h"
+
+DECLARE_LOG_CATEGORY_EXTERN(LogShaderLibrary, Log, All);
+
+class FShaderPipeline;
+
+struct RENDERCORE_API FShaderCodeLibraryPipeline
+{
+	FSHAHash VertexShader;
+	FSHAHash PixelShader;
+	FSHAHash GeometryShader;
+	FSHAHash HullShader;
+	FSHAHash DomainShader;
+	mutable uint32 Hash;
+	
+	FShaderCodeLibraryPipeline() : Hash(0) {}
+	
+	friend bool operator ==(const FShaderCodeLibraryPipeline& A,const FShaderCodeLibraryPipeline& B)
+	{
+		return A.VertexShader == B.VertexShader && A.PixelShader == B.PixelShader && A.GeometryShader == B.GeometryShader && A.HullShader == B.HullShader && A.DomainShader == B.DomainShader;
+	}
+	
+	friend uint32 GetTypeHash(const FShaderCodeLibraryPipeline &Key)
+	{
+		if(!Key.Hash)
+		{
+			Key.Hash = FCrc::MemCrc32(Key.VertexShader.Hash, sizeof(Key.VertexShader.Hash));
+			Key.Hash = FCrc::MemCrc32(Key.PixelShader.Hash, sizeof(Key.PixelShader.Hash), Key.Hash);
+			Key.Hash = FCrc::MemCrc32(Key.GeometryShader.Hash, sizeof(Key.GeometryShader.Hash), Key.Hash);
+			Key.Hash = FCrc::MemCrc32(Key.HullShader.Hash, sizeof(Key.HullShader.Hash), Key.Hash);
+			Key.Hash = FCrc::MemCrc32(Key.DomainShader.Hash, sizeof(Key.DomainShader.Hash), Key.Hash);
+		}
+		return Key.Hash;
+	}
+	
+	friend FArchive& operator<<( FArchive& Ar, FShaderCodeLibraryPipeline& Info )
+	{
+		return Ar << Info.VertexShader << Info.PixelShader << Info.GeometryShader << Info.HullShader << Info.DomainShader << Info.Hash;
+	}
+};
+
+struct RENDERCORE_API FCompactFullName
+{
+	TArray<FName> ObjectClassAndPath;
+
+	bool operator==(const FCompactFullName& Other) const
+	{
+		return ObjectClassAndPath == Other.ObjectClassAndPath;
+	}
+
+	FString ToString() const;
+	void ParseFromString(const FString& Src);
+	friend RENDERCORE_API uint32 GetTypeHash(const FCompactFullName& A);
+};
+
+
+struct RENDERCORE_API FStableShaderKeyAndValue
+{
+	FCompactFullName ClassNameAndObjectPath;
+	FName ShaderType;
+	FName ShaderClass;
+	FName MaterialDomain;
+	FName FeatureLevel;
+	FName QualityLevel;
+	FName TargetFrequency;
+	FName TargetPlatform;
+	FName VFType;
+	FName PermutationId;
+
+	uint32 KeyHash;
+
+	FSHAHash OutputHash;
+
+	FStableShaderKeyAndValue()
+		: KeyHash(0)
+	{
+	}
+
+	void ComputeKeyHash();
+	void ParseFromString(const FString& Src);
+	FString ToString() const;
+	static FString HeaderLine();
+
+	friend bool operator ==(const FStableShaderKeyAndValue& A, const FStableShaderKeyAndValue& B)
+	{
+		return
+			A.ClassNameAndObjectPath == B.ClassNameAndObjectPath &&
+			A.ShaderType == B.ShaderType &&
+			A.ShaderClass == B.ShaderClass &&
+			A.MaterialDomain == B.MaterialDomain &&
+			A.FeatureLevel == B.FeatureLevel &&
+			A.QualityLevel == B.QualityLevel &&
+			A.TargetFrequency == B.TargetFrequency &&
+			A.TargetPlatform == B.TargetPlatform &&
+			A.VFType == B.VFType &&
+			A.PermutationId == B.PermutationId;
+	}
+
+	friend uint32 GetTypeHash(const FStableShaderKeyAndValue &Key)
+	{
+		return Key.KeyHash;
+	}
+
+};
+
+class FShaderFactoryInterface : public FRHIShaderLibrary
+{
+public:
+	FShaderFactoryInterface(EShaderPlatform InPlatform, FString const& Name) : FRHIShaderLibrary(InPlatform, Name) {}
+	
+	virtual bool IsNativeLibrary() const override final {return false;}
+	
+	virtual FPixelShaderRHIRef CreatePixelShader(const FSHAHash& Hash) = 0; 
+	virtual FVertexShaderRHIRef CreateVertexShader(const FSHAHash& Hash) = 0;
+	virtual FHullShaderRHIRef CreateHullShader(const FSHAHash& Hash) = 0;
+	virtual FDomainShaderRHIRef CreateDomainShader(const FSHAHash& Hash) = 0;
+	virtual FGeometryShaderRHIRef CreateGeometryShader(const FSHAHash& Hash) = 0;
+	virtual FGeometryShaderRHIRef CreateGeometryShaderWithStreamOutput(const FSHAHash& Hash, const FStreamOutElementList& ElementList, uint32 NumStrides, const uint32* Strides, int32 RasterizedStream) = 0;
+	virtual FComputeShaderRHIRef CreateComputeShader(const FSHAHash& Hash) = 0;
+};
+
+DECLARE_MULTICAST_DELEGATE_TwoParams(FSharedShaderCodeRequest, const FSHAHash&, FArchive*);
+DECLARE_MULTICAST_DELEGATE_OneParam(FSharedShaderCodeRelease, const FSHAHash&);
+
+// Collection of unique shader code
+// Populated at cook time
+struct RENDERCORE_API FShaderCodeLibrary
+{
+	static void InitForRuntime(EShaderPlatform ShaderPlatform);
+	static void Shutdown();
+	
+	static bool IsEnabled();
+	
+	// Open a named library.
+	// For cooking this will place all added shaders & pipelines into the library file with this name.
+	// At runtime this will open the shader library with this name.
+	static void OpenLibrary(FString const& Name, FString const& Directory);
+    
+	// Close a named library.
+	// For cooking, after this point any AddShaderCode/AddShaderPipeline calls will be invalid until OpenLibrary is called again.
+	// At runtime this will release the library data and further requests for shaders from this library will fail.
+	static void CloseLibrary(FString const& Name);
+	
+	/** Instantiate or retrieve a vertex shader from the cache for the provided code & hash. */
+	static FVertexShaderRHIRef CreateVertexShader(EShaderPlatform Platform, FSHAHash Hash, TArray<uint8> const& Code);
+	/** Instantiate or retrieve a pixel shader from the cache for the provided code & hash. */
+	static FPixelShaderRHIRef CreatePixelShader(EShaderPlatform Platform, FSHAHash Hash, TArray<uint8> const& Code);
+	/** Instantiate or retrieve a geometry shader from the cache for the provided code & hash. */
+	static FGeometryShaderRHIRef CreateGeometryShader(EShaderPlatform Platform, FSHAHash Hash, TArray<uint8> const& Code);
+	/** Instantiate or retrieve a geometry shader from the cache for the provided code & hash. */
+	static FGeometryShaderRHIRef CreateGeometryShaderWithStreamOutput(EShaderPlatform Platform, FSHAHash Hash, const TArray<uint8>& Code, const FStreamOutElementList& ElementList, uint32 NumStrides, const uint32* Strides, int32 RasterizedStream);
+	/** Instantiate or retrieve a hull shader from the cache for the provided code & hash. */
+	static FHullShaderRHIRef CreateHullShader(EShaderPlatform Platform, FSHAHash Hash, TArray<uint8> const& Code);
+	/** Instantiate or retrieve a domain shader from the cache for the provided code & hash. */
+	static FDomainShaderRHIRef CreateDomainShader(EShaderPlatform Platform, FSHAHash Hash, TArray<uint8> const& Code);
+	/** Instantiate or retrieve a compute shader from the cache for the provided code & hash. */
+	static FComputeShaderRHIRef CreateComputeShader(EShaderPlatform Platform, FSHAHash Hash, TArray<uint8> const& Code);
+
+    static bool ContainsShaderCode(const FSHAHash& Hash);
+    
+	// Place a request to preload shader code
+	// Blocking call if no Archive is provided or Archive is not a type of FLinkerLoad
+	// Shader code preload will be finished before owning UObject PostLoad call
+	static bool RequestShaderCode(const FSHAHash& Hash, FArchive* Ar);
+
+	// Note that we skipped preloading shader code.
+	// All this does is call the delegate so that other folks are aware that it was lazy
+	static bool LazyRequestShaderCode(const FSHAHash& Hash, FArchive* Ar);
+
+	// Get the raw payload synchronously
+	// This does NOT require a ReleaseShaderCode; because it is synchronous
+	// This also does not fire any delegates...which are used to load binary programs (we are calling this because we failed to find a binary program)
+	static bool RequestShaderCode(const FSHAHash& Hash, TArray<uint8>& OutRaw);
+
+	// Request to release shader code
+	// Must match RequestShaderCode call
+	// Invalid to call before owning UObject PostLoad call
+	static void ReleaseShaderCode(const FSHAHash& Hash);
+
+	// Request to release shader code that we lazy loaded
+	// Must match LazyRequestShaderCode call
+	// All this does is call the delegate so that other folks are aware that it was lazy
+	static void LazyReleaseShaderCode(const FSHAHash& Hash);
+
+	// Create an iterator over all the shaders in the library
+	static TRefCountPtr<FRHIShaderLibrary::FShaderLibraryIterator> CreateIterator(void);
+	
+	// Total number of shader entries in the library
+	static uint32 GetShaderCount(void);
+	
+	// The shader platform that the library manages - at runtime this will only be one
+	static EShaderPlatform GetRuntimeShaderPlatform(void);
+	
+	// Get the shader pipelines in the library - only ever valid for OpenGL which can link without full PSO state
+	static TSet<FShaderCodeLibraryPipeline> const* GetShaderPipelines(EShaderPlatform Platform);
+
+#if WITH_EDITOR
+	// Initialize the library cooker
+	static void InitForCooking(bool bNativeFormat);
+	
+	// Clean the cook directories
+	static void CleanDirectories(TArray<FName> const& ShaderFormats);
+    
+    // Specify the shader formats to cook
+    static void CookShaderFormats(TArray<FName> const& ShaderFormats);
+	
+	// At cook time, add shader code to collection
+	static bool AddShaderCode(EShaderPlatform ShaderPlatform, EShaderFrequency Frequency, const FSHAHash& Hash, const TArray<uint8>& InCode, uint32 const UncompressedSize);
+
+	// We check this early in the callstack to avoid creating a bunch of FName and keys and things we will never save anyway
+	static bool NeedsShaderStableKeys();
+
+	// At cook time, add the human readable key value information
+	static void AddShaderStableKeyValue(EShaderPlatform ShaderPlatform, FStableShaderKeyAndValue& StableKeyValue);
+
+	// At cook time, add shader pipeline to collection
+	static bool AddShaderPipeline(FShaderPipeline* Pipeline);
+	
+	// Save collected shader code to a file for each specified shader platform, collating all child cooker results.
+	static bool SaveShaderCodeMaster(const FString& OutputDir, const FString& MetaOutputDir, const TArray<FName>& ShaderFormats, TArray<FString>& OutSCLCSVPath);
+	
+	// Save collected shader code to a file for each specified shader platform, handles only this instances intermediate results.
+	static bool SaveShaderCodeChild(const FString& OutputDir, const FString& MetaOutputDir, const TArray<FName>& ShaderFormats);
+	
+	// Package the separate shader bytecode files into a single native shader library. Must be called by the master process.
+	static bool PackageNativeShaderLibrary(const FString& ShaderCodeDir, const TArray<FName>& ShaderFormats);
+	
+	// Dump collected stats for each shader platform
+	static void DumpShaderCodeStats();
+#endif
+	
+	// Safely assign the hash to a shader object
+	static void SafeAssignHash(FRHIShader* InShader, const FSHAHash& Hash);
+
+	// Delegate called whenever shader code is requested.
+	static FDelegateHandle RegisterSharedShaderCodeRequestDelegate_Handle(const FSharedShaderCodeRequest::FDelegate& Delegate);
+	static void UnregisterSharedShaderCodeRequestDelegate_Handle(FDelegateHandle Handle);
+
+	// Delegate called whenever shader code is released.
+	static FDelegateHandle RegisterSharedShaderCodeReleaseDelegate_Handle(const FSharedShaderCodeRelease::FDelegate& Delegate);
+	static void UnregisterSharedShaderCodeReleaseDelegate_Handle(FDelegateHandle Handle);
+};
diff --git a/Engine/Source/Runtime/RenderCore/Public/ShaderCore.h b/Engine/Source/Runtime/RenderCore/Public/ShaderCore.h
new file mode 100644
--- /dev/null
+++ b/Engine/Source/Runtime/RenderCore/Public/ShaderCore.h
@@ -0,0 +1,1080 @@
+// Copyright 1998-2018 Epic Games, Inc. All Rights Reserved.
+
+/*=============================================================================
+	ShaderCore.h: Shader core module definitions.
+=============================================================================*/
+
+#pragma once
+
+#include "CoreMinimal.h"
+#include "Stats/Stats.h"
+#include "Templates/RefCounting.h"
+#include "Misc/SecureHash.h"
+#include "Misc/Paths.h"
+#include "Misc/CoreStats.h"
+#include "UniformBuffer.h"
+
+class Error;
+
+/**
+ * Controls whether shader related logs are visible.
+ * Note: The runtime verbosity is driven by the console variable 'r.ShaderDevelopmentMode'
+ */
+#if UE_BUILD_DEBUG && (PLATFORM_UNIX)
+RENDERCORE_API DECLARE_LOG_CATEGORY_EXTERN(LogShaders, Log, All);
+#else
+RENDERCORE_API DECLARE_LOG_CATEGORY_EXTERN(LogShaders, Error, All);
+#endif
+
+DECLARE_DWORD_ACCUMULATOR_STAT_EXTERN(TEXT("Num Total Niagara Shaders"), STAT_ShaderCompiling_NumTotalNiagaraShaders, STATGROUP_ShaderCompiling, RENDERCORE_API);
+DECLARE_FLOAT_ACCUMULATOR_STAT_EXTERN(TEXT("Total Niagara Shader Compiling Time"), STAT_ShaderCompiling_NiagaraShaders, STATGROUP_ShaderCompiling, RENDERCORE_API);
+
+DECLARE_FLOAT_ACCUMULATOR_STAT_EXTERN(TEXT("Total Material Shader Compiling Time"),STAT_ShaderCompiling_MaterialShaders,STATGROUP_ShaderCompiling, RENDERCORE_API);
+DECLARE_FLOAT_ACCUMULATOR_STAT_EXTERN(TEXT("Total Global Shader Compiling Time"),STAT_ShaderCompiling_GlobalShaders,STATGROUP_ShaderCompiling, RENDERCORE_API);
+DECLARE_FLOAT_ACCUMULATOR_STAT_EXTERN(TEXT("RHI Compile Time"),STAT_ShaderCompiling_RHI,STATGROUP_ShaderCompiling, RENDERCORE_API);
+DECLARE_FLOAT_ACCUMULATOR_STAT_EXTERN(TEXT("Loading Shader Files"),STAT_ShaderCompiling_LoadingShaderFiles,STATGROUP_ShaderCompiling, RENDERCORE_API);
+DECLARE_FLOAT_ACCUMULATOR_STAT_EXTERN(TEXT("CRCing Shader Files"),STAT_ShaderCompiling_HashingShaderFiles,STATGROUP_ShaderCompiling, RENDERCORE_API);
+DECLARE_FLOAT_ACCUMULATOR_STAT_EXTERN(TEXT("HLSL Translation"),STAT_ShaderCompiling_HLSLTranslation,STATGROUP_ShaderCompiling, RENDERCORE_API);
+DECLARE_FLOAT_ACCUMULATOR_STAT_EXTERN(TEXT("DDC Loading"),STAT_ShaderCompiling_DDCLoading,STATGROUP_ShaderCompiling, RENDERCORE_API);
+DECLARE_FLOAT_ACCUMULATOR_STAT_EXTERN(TEXT("Material Loading"),STAT_ShaderCompiling_MaterialLoading,STATGROUP_ShaderCompiling, RENDERCORE_API);
+DECLARE_FLOAT_ACCUMULATOR_STAT_EXTERN(TEXT("Material Compiling"),STAT_ShaderCompiling_MaterialCompiling,STATGROUP_ShaderCompiling, RENDERCORE_API);
+DECLARE_DWORD_ACCUMULATOR_STAT_EXTERN(TEXT("Num Total Material Shaders"),STAT_ShaderCompiling_NumTotalMaterialShaders,STATGROUP_ShaderCompiling, RENDERCORE_API);
+DECLARE_DWORD_ACCUMULATOR_STAT_EXTERN(TEXT("Num Special Material Shaders"),STAT_ShaderCompiling_NumSpecialMaterialShaders,STATGROUP_ShaderCompiling, RENDERCORE_API);
+DECLARE_DWORD_ACCUMULATOR_STAT_EXTERN(TEXT("Num Particle Material Shaders"),STAT_ShaderCompiling_NumParticleMaterialShaders,STATGROUP_ShaderCompiling, RENDERCORE_API);
+DECLARE_DWORD_ACCUMULATOR_STAT_EXTERN(TEXT("Num Skinned Material Shaders"),STAT_ShaderCompiling_NumSkinnedMaterialShaders,STATGROUP_ShaderCompiling, RENDERCORE_API);
+DECLARE_DWORD_ACCUMULATOR_STAT_EXTERN(TEXT("Num Lit Material Shaders"),STAT_ShaderCompiling_NumLitMaterialShaders,STATGROUP_ShaderCompiling, RENDERCORE_API);
+DECLARE_DWORD_ACCUMULATOR_STAT_EXTERN(TEXT("Num Unlit Material Shaders"),STAT_ShaderCompiling_NumUnlitMaterialShaders,STATGROUP_ShaderCompiling, RENDERCORE_API);
+DECLARE_DWORD_ACCUMULATOR_STAT_EXTERN(TEXT("Num Transparent Material Shaders"),STAT_ShaderCompiling_NumTransparentMaterialShaders,STATGROUP_ShaderCompiling, RENDERCORE_API);
+DECLARE_DWORD_ACCUMULATOR_STAT_EXTERN(TEXT("Num Opaque Material Shaders"),STAT_ShaderCompiling_NumOpaqueMaterialShaders,STATGROUP_ShaderCompiling, RENDERCORE_API);
+DECLARE_DWORD_ACCUMULATOR_STAT_EXTERN(TEXT("Num Masked Material Shaders"),STAT_ShaderCompiling_NumMaskedMaterialShaders,STATGROUP_ShaderCompiling, RENDERCORE_API);
+
+DECLARE_DWORD_ACCUMULATOR_STAT_EXTERN(TEXT("Num Shaders Loaded"),STAT_Shaders_NumShadersLoaded,STATGROUP_Shaders, RENDERCORE_API);
+DECLARE_DWORD_ACCUMULATOR_STAT_EXTERN(TEXT("Num Shader Resources Loaded"),STAT_Shaders_NumShaderResourcesLoaded,STATGROUP_Shaders, RENDERCORE_API);
+DECLARE_DWORD_ACCUMULATOR_STAT_EXTERN(TEXT("Num Shader Maps Registered"),STAT_Shaders_NumShaderMaps,STATGROUP_Shaders, RENDERCORE_API);
+DECLARE_CYCLE_STAT_EXTERN(TEXT("RT Shader Load Time"),STAT_Shaders_RTShaderLoadTime,STATGROUP_Shaders, RENDERCORE_API);
+DECLARE_DWORD_ACCUMULATOR_STAT_EXTERN(TEXT("Num Shaders Used"),STAT_Shaders_NumShadersUsedForRendering,STATGROUP_Shaders, RENDERCORE_API);
+DECLARE_FLOAT_ACCUMULATOR_STAT_EXTERN(TEXT("Total RT Shader Init Time"),STAT_Shaders_TotalRTShaderInitForRenderingTime,STATGROUP_Shaders, RENDERCORE_API);
+DECLARE_CYCLE_STAT_EXTERN(TEXT("Frame RT Shader Init Time"),STAT_Shaders_FrameRTShaderInitForRenderingTime,STATGROUP_Shaders, RENDERCORE_API);
+DECLARE_MEMORY_STAT_EXTERN(TEXT("Shader Memory"),STAT_Shaders_ShaderMemory,STATGROUP_Shaders, RENDERCORE_API);
+DECLARE_MEMORY_STAT_EXTERN(TEXT("Shader Resource Mem"),STAT_Shaders_ShaderResourceMemory,STATGROUP_Shaders, RENDERCORE_API);
+DECLARE_MEMORY_STAT_EXTERN(TEXT("Shader MapMemory"),STAT_Shaders_ShaderMapMemory,STATGROUP_Shaders, RENDERCORE_API);
+
+inline TStatId GetMemoryStatType(EShaderFrequency ShaderFrequency)
+{
+	static_assert(6 == SF_NumFrequencies, "EShaderFrequency has a bad size.");
+	switch(ShaderFrequency)
+	{
+		case SF_Pixel:		return GET_STATID(STAT_PixelShaderMemory);
+		case SF_Compute:	return GET_STATID(STAT_PixelShaderMemory);
+	}
+	return GET_STATID(STAT_VertexShaderMemory);
+}
+
+/** Initializes shader hash cache from IShaderFormatModules. This must be called before reading any shader include. */
+extern RENDERCORE_API void InitializeShaderHashCache();
+
+/** Checks if shader include isn't skipped by a shader hash cache. */
+extern RENDERCORE_API void CheckShaderHashCacheInclude(const FString& VirtualFilePath, EShaderPlatform ShaderPlatform);
+
+/** Initializes cached shader type data.  This must be called before creating any FShaderType. */
+extern RENDERCORE_API void InitializeShaderTypes();
+
+/** Uninitializes cached shader type data.  This is needed before unloading modules that contain FShaderTypes. */
+extern RENDERCORE_API void UninitializeShaderTypes();
+
+/** Returns true if debug viewmodes are allowed for the current platform. */
+extern RENDERCORE_API bool AllowDebugViewmodes();
+
+/** Returns true if debug viewmodes are allowed for the given platform. */
+extern RENDERCORE_API bool AllowDebugViewmodes(EShaderPlatform Platform);
+
+struct FShaderTarget
+{
+	uint32 Frequency : SF_NumBits;
+	uint32 Platform : SP_NumBits;
+
+	FShaderTarget()
+	{}
+
+	FShaderTarget(EShaderFrequency InFrequency,EShaderPlatform InPlatform)
+	:	Frequency(InFrequency)
+	,	Platform(InPlatform)
+	{}
+
+	friend bool operator==(const FShaderTarget& X, const FShaderTarget& Y)
+	{
+		return X.Frequency == Y.Frequency && X.Platform == Y.Platform;
+	}
+
+	friend FArchive& operator<<(FArchive& Ar,FShaderTarget& Target)
+	{
+		uint32 TargetFrequency = Target.Frequency;
+		uint32 TargetPlatform = Target.Platform;
+		Ar << TargetFrequency << TargetPlatform;
+		if (Ar.IsLoading())
+		{
+			Target.Frequency = TargetFrequency;
+			Target.Platform = TargetPlatform;
+		}
+		return Ar;
+	}
+
+	EShaderPlatform GetPlatform() const
+	{
+		return (EShaderPlatform)Platform;
+	}
+
+	EShaderFrequency GetFrequency() const
+	{
+		return (EShaderFrequency)Frequency;
+	}
+};
+
+enum ECompilerFlags
+{
+	CFLAG_PreferFlowControl = 0,
+	CFLAG_Debug,
+	CFLAG_AvoidFlowControl,
+	/** Disable shader validation */
+	CFLAG_SkipValidation,
+	/** Only allows standard optimizations, not the longest compile times. */
+	CFLAG_StandardOptimization,
+	/** Shader should use on chip memory instead of main memory ring buffer memory. */
+	CFLAG_OnChip,
+	CFLAG_KeepDebugInfo,
+	CFLAG_NoFastMath,
+	/** Explicitly enforce zero initialisation on shader platforms that may omit it. */
+	CFLAG_ZeroInitialise,
+	/** Explicitly enforce bounds checking on shader platforms that may omit it. */
+	CFLAG_BoundsChecking,
+	// Compile ES2 with ES3.1 features
+	CFLAG_FeatureLevelES31,
+	// Force removing unused interpolators for platforms that can opt out
+	CFLAG_ForceRemoveUnusedInterpolators,
+	// Set default precision to highp in a pixel shader (default is mediump on ES2 platforms)
+	CFLAG_UseFullPrecisionInPS,
+	// Hint that its a vertex to geometry shader
+	CFLAG_VertexToGeometryShader,
+	// Prepare the shader for archiving in the native binary shader cache format
+	CFLAG_Archive,
+	// Shaders uses external texture so may need special runtime handling
+	CFLAG_UsesExternalTexture,
+	// Use emulated uniform buffers on supported platforms
+	CFLAG_UseEmulatedUB
+};
+
+/**
+ * A map of shader parameter names to registers allocated to that parameter.
+ */
+class FShaderParameterMap
+{
+public:
+
+	FShaderParameterMap()
+	{}
+
+	RENDERCORE_API bool FindParameterAllocation(const TCHAR* ParameterName,uint16& OutBufferIndex,uint16& OutBaseIndex,uint16& OutSize) const;
+	RENDERCORE_API bool ContainsParameterAllocation(const TCHAR* ParameterName) const;
+	RENDERCORE_API void AddParameterAllocation(const TCHAR* ParameterName, uint16 BufferIndex, uint16 BaseIndex, uint16 Size);
+	RENDERCORE_API void RemoveParameterAllocation(const TCHAR* ParameterName);
+	/** Checks that all parameters are bound and asserts if any aren't in a debug build
+	* @param InVertexFactoryType can be 0
+	*/
+	RENDERCORE_API void VerifyBindingsAreComplete(const TCHAR* ShaderTypeName, FShaderTarget Target, class FVertexFactoryType* InVertexFactoryType) const;
+
+	/** Updates the hash state with the contents of this parameter map. */
+	void UpdateHash(FSHA1& HashState) const;
+
+	friend FArchive& operator<<(FArchive& Ar,FShaderParameterMap& InParameterMap)
+	{
+		// Note: this serialize is used to pass between UE4 and the shader compile worker, recompile both when modifying
+		return Ar << InParameterMap.ParameterMap;
+	}
+
+	inline void GetAllParameterNames(TArray<FString>& OutNames) const
+	{
+		ParameterMap.GenerateKeyArray(OutNames);
+	}
+
+private:
+	struct FParameterAllocation
+	{
+		uint16 BufferIndex;
+		uint16 BaseIndex;
+		uint16 Size;
+		mutable bool bBound;
+
+		FParameterAllocation() :
+			bBound(false)
+		{}
+
+		friend FArchive& operator<<(FArchive& Ar, FParameterAllocation& Allocation)
+		{
+			return Ar << Allocation.BufferIndex << Allocation.BaseIndex << Allocation.Size << Allocation.bBound;
+		}
+	};
+
+	TMap<FString,FParameterAllocation> ParameterMap;
+};
+
+/** Container for shader compiler definitions. */
+class FShaderCompilerDefinitions
+{
+public:
+
+	FShaderCompilerDefinitions()
+	{
+		// Presize to reduce re-hashing while building shader jobs
+		Definitions.Empty(50);
+	}
+
+	/**
+	 * Works for TCHAR
+	 * e.g. SetDefine(TEXT("NUM_SAMPLES"), TEXT("1"));
+	 */
+	void SetDefine(const TCHAR* Name, const TCHAR* Value)
+	{
+		Definitions.Add(Name, Value);
+	}
+
+	/**
+	 * Works for uint32 and bool
+	 * e.g. OutEnvironment.SetDefine(TEXT("REALLY"), bReally);
+	 * e.g. OutEnvironment.SetDefine(TEXT("NUM_SAMPLES"), NumSamples);
+	 */
+	void SetDefine(const TCHAR* Name, uint32 Value)
+	{
+		// can be optimized
+		Definitions.Add(Name, *FString::Printf(TEXT("%u"), Value));
+	}
+
+	void SetDefine(const TCHAR* Name, int32 Value)
+	{
+		// can be optimized
+		Definitions.Add(Name, *FString::Printf(TEXT("%d"), Value));
+	}
+
+	/**
+	 * Works for float
+	 */
+	void SetFloatDefine(const TCHAR* Name, float Value)
+	{
+		// can be optimized
+		Definitions.Add(Name, *FString::Printf(TEXT("%f"), Value));
+	}
+
+	const TMap<FString,FString>& GetDefinitionMap() const
+	{
+		return Definitions;
+	}
+
+	friend FArchive& operator<<(FArchive& Ar,FShaderCompilerDefinitions& Defs)
+	{
+		return Ar << Defs.Definitions;
+	}
+
+	void Merge(const FShaderCompilerDefinitions& Other)
+	{
+		Definitions.Append(Other.Definitions);
+	}
+
+private:
+
+	/** Map: definition -> value. */
+	TMap<FString,FString> Definitions;
+};
+
+struct FBaseShaderResourceTable
+{
+	/** Bits indicating which resource tables contain resources bound to this shader. */
+	uint32 ResourceTableBits;
+
+	/** Mapping of bound SRVs to their location in resource tables. */
+	TArray<uint32> ShaderResourceViewMap;
+
+	/** Mapping of bound sampler states to their location in resource tables. */
+	TArray<uint32> SamplerMap;
+
+	/** Mapping of bound UAVs to their location in resource tables. */
+	TArray<uint32> UnorderedAccessViewMap;
+
+	/** Hash of the layouts of resource tables at compile time, used for runtime validation. */
+	TArray<uint32> ResourceTableLayoutHashes;
+
+	FBaseShaderResourceTable() :
+		ResourceTableBits(0)
+	{
+	}
+
+	friend bool operator==(const FBaseShaderResourceTable &A, const FBaseShaderResourceTable& B)
+	{
+		bool bEqual = true;
+		bEqual &= (A.ResourceTableBits == B.ResourceTableBits);
+		bEqual &= (A.ShaderResourceViewMap.Num() == B.ShaderResourceViewMap.Num());
+		bEqual &= (A.SamplerMap.Num() == B.SamplerMap.Num());
+		bEqual &= (A.UnorderedAccessViewMap.Num() == B.UnorderedAccessViewMap.Num());
+		bEqual &= (A.ResourceTableLayoutHashes.Num() == B.ResourceTableLayoutHashes.Num());
+		if (!bEqual)
+		{
+			return false;
+		}
+		bEqual &= (FMemory::Memcmp(A.ShaderResourceViewMap.GetData(), B.ShaderResourceViewMap.GetData(), A.ShaderResourceViewMap.GetTypeSize()*A.ShaderResourceViewMap.Num()) == 0);
+		bEqual &= (FMemory::Memcmp(A.SamplerMap.GetData(), B.SamplerMap.GetData(), A.SamplerMap.GetTypeSize()*A.SamplerMap.Num()) == 0);
+		bEqual &= (FMemory::Memcmp(A.UnorderedAccessViewMap.GetData(), B.UnorderedAccessViewMap.GetData(), A.UnorderedAccessViewMap.GetTypeSize()*A.UnorderedAccessViewMap.Num()) == 0);
+		bEqual &= (FMemory::Memcmp(A.ResourceTableLayoutHashes.GetData(), B.ResourceTableLayoutHashes.GetData(), A.ResourceTableLayoutHashes.GetTypeSize()*A.ResourceTableLayoutHashes.Num()) == 0);
+		return bEqual;
+	}
+};
+
+inline FArchive& operator<<(FArchive& Ar, FBaseShaderResourceTable& SRT)
+{
+	Ar << SRT.ResourceTableBits;
+	Ar << SRT.ShaderResourceViewMap;
+	Ar << SRT.SamplerMap;
+	Ar << SRT.UnorderedAccessViewMap;
+	Ar << SRT.ResourceTableLayoutHashes;
+
+	return Ar;
+}
+
+struct FShaderCompilerResourceTable
+{
+	/** Bits indicating which resource tables contain resources bound to this shader. */
+	uint32 ResourceTableBits;
+
+	/** The max index of a uniform buffer from which resources are bound. */
+	uint32 MaxBoundResourceTable;
+
+	/** Mapping of bound Textures to their location in resource tables. */
+	TArray<uint32> TextureMap;
+
+	/** Mapping of bound SRVs to their location in resource tables. */
+	TArray<uint32> ShaderResourceViewMap;
+
+	/** Mapping of bound sampler states to their location in resource tables. */
+	TArray<uint32> SamplerMap;
+
+	/** Mapping of bound UAVs to their location in resource tables. */
+	TArray<uint32> UnorderedAccessViewMap;
+
+	/** Hash of the layouts of resource tables at compile time, used for runtime validation. */
+	TArray<uint32> ResourceTableLayoutHashes;
+
+	FShaderCompilerResourceTable()
+		: ResourceTableBits(0)
+		, MaxBoundResourceTable(0)
+	{
+	}
+};
+
+inline FArchive& operator<<(FArchive& Ar, FResourceTableEntry& Entry)
+{
+	Ar << Entry.UniformBufferName;
+	Ar << Entry.Type;
+	Ar << Entry.ResourceIndex;
+	return Ar;
+}
+
+/** Additional compilation settings that can be configured by each FMaterial instance before compilation */
+struct FExtraShaderCompilerSettings
+{
+	bool bExtractShaderSource = false;
+	FString OfflineCompilerPath;
+
+	friend FArchive& operator<<(FArchive& Ar, FExtraShaderCompilerSettings& StatsSettings)
+	{
+		// Note: this serialize is used to pass between UE4 and the shader compile worker, recompile both when modifying
+		return Ar << StatsSettings.bExtractShaderSource << StatsSettings.OfflineCompilerPath;
+	}
+};
+
+/** The environment used to compile a shader. */
+struct FShaderCompilerEnvironment : public FRefCountedObject
+{
+	// Map of the virtual file path -> content.
+	// The virtual file paths are the ones that USF files query through the #include "<The Virtual Path of the file>"
+	TMap<FString,FString> IncludeVirtualPathToContentsMap;
+	
+	TMap<FString,TSharedPtr<FString>> IncludeVirtualPathToExternalContentsMap;
+
+	TArray<uint32> CompilerFlags;
+	TMap<uint32,uint8> RenderTargetOutputFormatsMap;
+	TMap<FString,FResourceTableEntry> ResourceTableMap;
+	TMap<FString,uint32> ResourceTableLayoutHashes;
+	TMap<FString, FString> RemoteServerData;
+	TMap<FString, FString> ShaderFormatCVars;
+
+	/** Default constructor. */
+	FShaderCompilerEnvironment()
+	{
+		// Presize to reduce re-hashing while building shader jobs
+		IncludeVirtualPathToContentsMap.Empty(15);
+	}
+
+	/** Initialization constructor. */
+	explicit FShaderCompilerEnvironment(const FShaderCompilerDefinitions& InDefinitions)
+		: Definitions(InDefinitions)
+	{
+	}
+
+	/**
+	 * Works for TCHAR
+	 * e.g. SetDefine(TEXT("NAME"), TEXT("Test"));
+	 * e.g. SetDefine(TEXT("NUM_SAMPLES"), 1);
+	 * e.g. SetDefine(TEXT("DOIT"), true);
+	 */
+	void SetDefine(const TCHAR* Name, const TCHAR* Value)	{ Definitions.SetDefine(Name, Value); }
+	void SetDefine(const TCHAR* Name, uint32 Value)			{ Definitions.SetDefine(Name, Value); }
+	void SetDefine(const TCHAR* Name, int32 Value)			{ Definitions.SetDefine(Name, Value); }
+	void SetDefine(const TCHAR* Name, bool Value)			{ Definitions.SetDefine(Name, (uint32)Value); }
+	void SetDefine(const TCHAR* Name, float Value)			{ Definitions.SetFloatDefine(Name, Value); }
+
+	const TMap<FString,FString>& GetDefinitions() const
+	{
+		return Definitions.GetDefinitionMap();
+	}
+
+	void SetRenderTargetOutputFormat(uint32 RenderTargetIndex, EPixelFormat PixelFormat)
+	{
+		RenderTargetOutputFormatsMap.Add(RenderTargetIndex, PixelFormat);
+	}
+
+	friend FArchive& operator<<(FArchive& Ar,FShaderCompilerEnvironment& Environment)
+	{
+		// Note: this serialize is used to pass between UE4 and the shader compile worker, recompile both when modifying
+		Ar << Environment.IncludeVirtualPathToContentsMap;
+
+		// Note: skipping Environment.IncludeVirtualPathToExternalContentsMap, which is handled by FShaderCompileUtilities::DoWriteTasks in order to maintain sharing
+
+		Ar << Environment.Definitions;
+		Ar << Environment.CompilerFlags;
+		Ar << Environment.RenderTargetOutputFormatsMap;
+		Ar << Environment.ResourceTableMap;
+		Ar << Environment.ResourceTableLayoutHashes;
+		Ar << Environment.RemoteServerData;
+		Ar << Environment.ShaderFormatCVars;
+
+		return Ar;
+	}
+	
+	void Merge(const FShaderCompilerEnvironment& Other)
+	{
+		// Merge the include maps
+		// Merge the values of any existing keys
+		for (TMap<FString,FString>::TConstIterator It(Other.IncludeVirtualPathToContentsMap); It; ++It )
+		{
+			FString* ExistingContents = IncludeVirtualPathToContentsMap.Find(It.Key());
+
+			if (ExistingContents)
+			{
+				ExistingContents->Append(It.Value());
+			}
+			else
+			{
+				IncludeVirtualPathToContentsMap.Add(It.Key(), It.Value());
+			}
+		}
+
+		check(Other.IncludeVirtualPathToExternalContentsMap.Num() == 0);
+
+		CompilerFlags.Append(Other.CompilerFlags);
+		ResourceTableMap.Append(Other.ResourceTableMap);
+		ResourceTableLayoutHashes.Append(Other.ResourceTableLayoutHashes);
+		Definitions.Merge(Other.Definitions);
+		RenderTargetOutputFormatsMap.Append(Other.RenderTargetOutputFormatsMap);
+		RemoteServerData.Append(Other.RemoteServerData);
+		ShaderFormatCVars.Append(Other.ShaderFormatCVars);
+	}
+
+private:
+
+	FShaderCompilerDefinitions Definitions;
+};
+
+/** Struct that gathers all readonly inputs needed for the compilation of a single shader. */
+struct FShaderCompilerInput
+{
+	FShaderTarget Target;
+	FName ShaderFormat;
+	FString SourceFilePrefix;
+	FString VirtualSourceFilePath;
+	FString EntryPointName;
+
+	// Skips the preprocessor and instead loads the usf file directly
+	bool bSkipPreprocessedCache;
+
+	bool bGenerateDirectCompileFile;
+
+	// Shader pipeline information
+	bool bCompilingForShaderPipeline;
+	bool bIncludeUsedOutputs;
+	TArray<FString> UsedOutputs;
+
+	// Dump debug path (up to platform) e.g. "D:/MMittring-Z3941-A/UE4-Orion/OrionGame/Saved/ShaderDebugInfo/PCD3D_SM5"
+	FString DumpDebugInfoRootPath;
+	// only used if enabled by r.DumpShaderDebugInfo (platform/groupname) e.g. ""
+	FString DumpDebugInfoPath;
+	// materialname or "Global" "for debugging and better error messages
+	FString DebugGroupName;
+
+	// Compilation Environment
+	FShaderCompilerEnvironment Environment;
+	TRefCountPtr<FShaderCompilerEnvironment> SharedEnvironment;
+
+	// Additional compilation settings that can be filled by FMaterial::SetupExtaCompilationSettings
+	// FMaterial::SetupExtaCompilationSettings is usually called by each (*)MaterialShaderType::BeginCompileShader() function
+	FExtraShaderCompilerSettings ExtraSettings;
+
+	FShaderCompilerInput() :
+		bSkipPreprocessedCache(false),
+		bGenerateDirectCompileFile(false),
+		bCompilingForShaderPipeline(false),
+		bIncludeUsedOutputs(false)
+	{
+	}
+
+	// generate human readable name for debugging
+	FString GenerateShaderName() const
+	{
+		FString Name;
+
+		if(DebugGroupName == TEXT("Global"))
+		{
+			Name = VirtualSourceFilePath + TEXT("|") + EntryPointName;
+		}
+		else
+		{
+			// we skip EntryPointName as it's usually not useful
+			Name = DebugGroupName + TEXT(":") + VirtualSourceFilePath;
+		}
+
+		return Name;
+	}
+
+	FString GetSourceFilename() const
+	{
+		return FPaths::GetCleanFilename(VirtualSourceFilePath);
+	}
+
+	void GatherSharedInputs(TMap<FString,FString>& ExternalIncludes, TArray<FShaderCompilerEnvironment*>& SharedEnvironments)
+	{
+		check(!SharedEnvironment || SharedEnvironment->IncludeVirtualPathToExternalContentsMap.Num() == 0);
+
+		for (TMap<FString, TSharedPtr<FString>>::TConstIterator It(Environment.IncludeVirtualPathToExternalContentsMap); It; ++It)
+		{
+			FString* FoundEntry = ExternalIncludes.Find(It.Key());
+
+			if (!FoundEntry)
+			{
+				ExternalIncludes.Add(It.Key(), *It.Value());
+			}
+		}
+
+		if (SharedEnvironment)
+		{
+			SharedEnvironments.AddUnique(SharedEnvironment.GetReference());
+		}
+	}
+
+	void SerializeSharedInputs(FArchive& Ar, const TArray<FShaderCompilerEnvironment*>& SharedEnvironments)
+	{
+		check(Ar.IsSaving());
+
+		TArray<FString> ReferencedExternalIncludes;
+		ReferencedExternalIncludes.Empty(Environment.IncludeVirtualPathToExternalContentsMap.Num());
+
+		for (TMap<FString, TSharedPtr<FString>>::TConstIterator It(Environment.IncludeVirtualPathToExternalContentsMap); It; ++It)
+		{
+			ReferencedExternalIncludes.Add(It.Key());
+		}
+
+		Ar << ReferencedExternalIncludes;
+
+		int32 SharedEnvironmentIndex = SharedEnvironments.Find(SharedEnvironment.GetReference());
+		Ar << SharedEnvironmentIndex;
+	}
+
+	void DeserializeSharedInputs(FArchive& Ar, const TMap<FString,TSharedPtr<FString>>& ExternalIncludes, const TArray<FShaderCompilerEnvironment>& SharedEnvironments)
+	{
+		check(Ar.IsLoading());
+
+		TArray<FString> ReferencedExternalIncludes;
+		Ar << ReferencedExternalIncludes;
+
+		Environment.IncludeVirtualPathToExternalContentsMap.Reserve(ReferencedExternalIncludes.Num());
+
+		for (int32 i = 0; i < ReferencedExternalIncludes.Num(); i++)
+		{
+			Environment.IncludeVirtualPathToExternalContentsMap.Add(ReferencedExternalIncludes[i], ExternalIncludes.FindChecked(ReferencedExternalIncludes[i]));
+		}
+
+		int32 SharedEnvironmentIndex = 0;
+		Ar << SharedEnvironmentIndex;
+
+		if (SharedEnvironments.IsValidIndex(SharedEnvironmentIndex))
+		{
+			Environment.Merge(SharedEnvironments[SharedEnvironmentIndex]);
+		}
+	}
+
+	friend FArchive& operator<<(FArchive& Ar,FShaderCompilerInput& Input)
+	{
+		// Note: this serialize is used to pass between UE4 and the shader compile worker, recompile both when modifying
+		Ar << Input.Target;
+		{
+			FString ShaderFormatString(Input.ShaderFormat.ToString());
+			Ar << ShaderFormatString;
+			Input.ShaderFormat = FName(*ShaderFormatString);
+		}
+		Ar << Input.SourceFilePrefix;
+		Ar << Input.VirtualSourceFilePath;
+		Ar << Input.EntryPointName;
+		Ar << Input.bSkipPreprocessedCache;
+		Ar << Input.bCompilingForShaderPipeline;
+		Ar << Input.bGenerateDirectCompileFile;
+		Ar << Input.bIncludeUsedOutputs;
+		Ar << Input.UsedOutputs;
+		Ar << Input.DumpDebugInfoRootPath;
+		Ar << Input.DumpDebugInfoPath;
+		Ar << Input.DebugGroupName;
+		Ar << Input.Environment;
+		Ar << Input.ExtraSettings;
+
+		// Note: skipping Input.SharedEnvironment, which is handled by FShaderCompileUtilities::DoWriteTasks in order to maintain sharing
+
+		return Ar;
+	}
+};
+
+/** A shader compiler error or warning. */
+struct FShaderCompilerError
+{
+	FShaderCompilerError(const TCHAR* InStrippedErrorMessage = TEXT(""))
+	:	ErrorVirtualFilePath(TEXT(""))
+	,	ErrorLineString(TEXT(""))
+	,	StrippedErrorMessage(InStrippedErrorMessage)
+	{}
+
+	FString ErrorVirtualFilePath;
+	FString ErrorLineString;
+	FString StrippedErrorMessage;
+
+	FString GetErrorString() const
+	{
+		return ErrorVirtualFilePath + TEXT("(") + ErrorLineString + TEXT("): ") + StrippedErrorMessage; // TODO
+	}
+
+	/** Returns the path of the underlying source file relative to the process base dir. */
+	FString RENDERCORE_API GetShaderSourceFilePath() const;
+
+	friend FArchive& operator<<(FArchive& Ar,FShaderCompilerError& Error)
+	{
+		return Ar << Error.ErrorVirtualFilePath << Error.ErrorLineString << Error.StrippedErrorMessage;
+	}
+};
+
+// if this changes you need to make sure all D3D11 shaders get invalidated
+struct FShaderCodePackedResourceCounts
+{
+	// for FindOptionalData() and AddOptionalData()
+	static const uint8 Key = 'p';
+
+	bool bGlobalUniformBufferUsed;
+	uint8 NumSamplers;
+	uint8 NumSRVs;
+	uint8 NumCBs;
+	uint8 NumUAVs;
+};
+
+#ifdef __EMSCRIPTEN__
+// Emscripten asm.js is strict and doesn't support unaligned memory load or stores.
+// When such an unaligned memory access occurs, the compiler needs to know about it,
+// so that it can generate the appropriate unalignment-aware memory load/store instruction.
+typedef int32 __attribute__((aligned(1))) unaligned_int32;
+typedef uint32 __attribute__((aligned(1))) unaligned_uint32;
+#else
+// On x86 etc. unaligned memory accesses are supported by the CPU, so no need to
+// behave specially for them.
+typedef int32 unaligned_int32;
+typedef uint32 unaligned_uint32;
+#endif
+
+// later we can transform that to the actual class passed around at the RHI level
+class FShaderCodeReader
+{
+	const TArray<uint8>& ShaderCode;
+
+public:
+	FShaderCodeReader(const TArray<uint8>& InShaderCode)
+		: ShaderCode(InShaderCode)
+	{
+		check(ShaderCode.Num());
+	}
+
+	uint32 GetActualShaderCodeSize() const
+	{
+		return ShaderCode.Num() - GetOptionalDataSize();
+	}
+
+	// for convenience
+	template <class T>
+	const T* FindOptionalData() const
+	{
+		return (const T*)FindOptionalData(T::Key, sizeof(T));
+	}
+
+
+	// @param InKey e.g. FShaderCodePackedResourceCounts::Key
+	// @return 0 if not found
+	const uint8* FindOptionalData(uint8 InKey, uint8 ValueSize) const
+	{
+		check(ValueSize);
+
+		const uint8* End = &ShaderCode[0] + ShaderCode.Num();
+
+		int32 LocalOptionalDataSize = GetOptionalDataSize();
+
+		const uint8* Start = End - LocalOptionalDataSize;
+		// while searching don't include the optional data size
+		End = End - sizeof(LocalOptionalDataSize);
+		const uint8* Current = Start;
+
+		while(Current < End)
+		{
+			uint8 Key = *Current++;
+			uint32 Size = *((const unaligned_uint32*)Current);
+			Current += sizeof(Size);
+
+			if(Key == InKey && Size == ValueSize)
+			{
+				return Current;
+			}
+
+			Current += Size;
+		}
+
+		return 0;
+	}
+
+	const ANSICHAR* FindOptionalData(uint8 InKey) const
+	{
+		check(ShaderCode.Num() >= 4);
+
+		const uint8* End = &ShaderCode[0] + ShaderCode.Num();
+
+		int32 LocalOptionalDataSize = GetOptionalDataSize();
+
+		const uint8* Start = End - LocalOptionalDataSize;
+		// while searching don't include the optional data size
+		End = End - sizeof(LocalOptionalDataSize);
+		const uint8* Current = Start;
+
+		while(Current < End)
+		{
+			uint8 Key = *Current++;
+			uint32 Size = *((const unaligned_uint32*)Current);
+			Current += sizeof(Size);
+
+			if(Key == InKey)
+			{
+				return (ANSICHAR*)Current;
+			}
+
+			Current += Size;
+		}
+
+		return 0;
+	}
+
+	// Returns nullptr and Size -1 if not key was not found
+	const uint8* FindOptionalDataAndSize(uint8 InKey, int32& OutSize) const
+	{
+		check(ShaderCode.Num() >= 4);
+
+		const uint8* End = &ShaderCode[0] + ShaderCode.Num();
+
+		int32 LocalOptionalDataSize = GetOptionalDataSize();
+
+		const uint8* Start = End - LocalOptionalDataSize;
+		// while searching don't include the optional data size
+		End = End - sizeof(LocalOptionalDataSize);
+		const uint8* Current = Start;
+
+		while (Current < End)
+		{
+			uint8 Key = *Current++;
+			uint32 Size = *((const unaligned_uint32*)Current);
+			Current += sizeof(Size);
+
+			if (Key == InKey)
+			{
+				OutSize = Size;
+				return Current;
+			}
+
+			Current += Size;
+		}
+
+		OutSize = -1;
+		return nullptr;
+	}
+
+	int32 GetOptionalDataSize() const
+	{
+		if(ShaderCode.Num() < sizeof(int32))
+		{
+			return 0;
+		}
+
+		const uint8* End = &ShaderCode[0] + ShaderCode.Num();
+
+		int32 LocalOptionalDataSize = ((const unaligned_int32*)End)[-1];
+
+		check(LocalOptionalDataSize >= 0);
+		check(ShaderCode.Num() >= LocalOptionalDataSize);
+
+		return LocalOptionalDataSize;
+	}
+
+	int32 GetShaderCodeSize() const
+	{
+		return ShaderCode.Num() - GetOptionalDataSize();
+	}
+};
+
+class FShaderCode
+{
+	// -1 if ShaderData was finalized
+	mutable int32 OptionalDataSize;
+	// access through class methods
+	mutable TArray<uint8> ShaderCodeWithOptionalData;
+
+public:
+
+	FShaderCode()
+	: OptionalDataSize(0)
+	{
+	}
+
+	// adds CustomData or does nothing if that was already done before
+	void FinalizeShaderCode() const
+	{
+		if(OptionalDataSize != -1)
+		{
+			OptionalDataSize += sizeof(OptionalDataSize);
+			ShaderCodeWithOptionalData.Append((const uint8*)&OptionalDataSize, sizeof(OptionalDataSize));
+			OptionalDataSize = -1;
+		}
+	}
+
+	// for write access
+	TArray<uint8>& GetWriteAccess()
+	{
+		return ShaderCodeWithOptionalData;
+	}
+
+	int32 GetShaderCodeSize() const
+	{
+		FinalizeShaderCode();
+
+		FShaderCodeReader Wrapper(ShaderCodeWithOptionalData);
+
+		return Wrapper.GetShaderCodeSize();
+	}
+
+	// inefficient, will/should be replaced by GetShaderCodeToRead()
+	void GetShaderCodeLegacy(TArray<uint8>& Out) const
+	{
+		Out.Empty();
+
+		Out.AddUninitialized(GetShaderCodeSize());
+		FMemory::Memcpy(Out.GetData(), GetReadAccess().GetData(), ShaderCodeWithOptionalData.Num());
+	}
+
+	// for read access, can have additional data attached to the end
+	const TArray<uint8>& GetReadAccess() const
+	{
+		FinalizeShaderCode();
+
+		return ShaderCodeWithOptionalData;
+	}
+
+	// for convenience
+	template <class T>
+	void AddOptionalData(const T &In)
+	{
+		AddOptionalData(T::Key, (uint8*)&In, sizeof(T));
+	}
+
+	// Note: we don't hash the optional attachments in GenerateOutputHash() as they would prevent sharing (e.g. many material share the save VS)
+	// can be called after the non optional data was stored in ShaderData
+	// @param Key uint8 to save memory so max 255, e.g. FShaderCodePackedResourceCounts::Key
+	// @param Size >0, only restriction is that sum of all optional data values must be < 4GB
+	void AddOptionalData(uint8 Key, const uint8* ValuePtr, uint32 ValueSize)
+	{
+		check(ValuePtr);
+
+		// don't add after Finalize happened
+		check(OptionalDataSize >= 0);
+
+		ShaderCodeWithOptionalData.Add(Key);
+		ShaderCodeWithOptionalData.Append((const uint8*)&ValueSize, sizeof(ValueSize));
+		ShaderCodeWithOptionalData.Append(ValuePtr, ValueSize);
+		OptionalDataSize += sizeof(uint8) + sizeof(ValueSize) + (uint32)ValueSize;
+	}
+
+	// Note: we don't hash the optional attachments in GenerateOutputHash() as they would prevent sharing (e.g. many material share the save VS)
+	// convenience, silently drops the data if string is too long
+	// @param e.g. 'n' for the ShaderSourceFileName
+	void AddOptionalData(uint8 Key, const ANSICHAR* InString)
+	{
+		uint32 Size = FCStringAnsi::Strlen(InString) + 1;
+		AddOptionalData(Key, (uint8*)InString, Size);
+	}
+
+	friend FArchive& operator<<(FArchive& Ar, FShaderCode& Output)
+	{
+		if(Ar.IsLoading())
+		{
+			Output.OptionalDataSize = -1;
+		}
+		else
+		{
+			Output.FinalizeShaderCode();
+		}
+
+		// Note: this serialize is used to pass between UE4 and the shader compile worker, recompile both when modifying
+		Ar << Output.ShaderCodeWithOptionalData;
+		return Ar;
+	}
+};
+
+/** The output of the shader compiler. */
+struct FShaderCompilerOutput
+{
+	FShaderCompilerOutput()
+	:	NumInstructions(0)
+	,	NumTextureSamplers(0)
+	,	bSucceeded(false)
+	,	bFailedRemovingUnused(false)
+	,	bSupportsQueryingUsedAttributes(false)
+	{
+	}
+
+	FShaderParameterMap ParameterMap;
+	TArray<FShaderCompilerError> Errors;
+	TArray<FString> PragmaDirectives;
+	FShaderTarget Target;
+	FShaderCode ShaderCode;
+	FSHAHash OutputHash;
+	uint32 NumInstructions;
+	uint32 NumTextureSamplers;
+	bool bSucceeded;
+	bool bFailedRemovingUnused;
+	bool bSupportsQueryingUsedAttributes;
+	TArray<FString> UsedAttributes;
+
+	FString OptionalFinalShaderSource;
+
+	TArray<uint8> PlatformDebugData;
+
+	/** Generates OutputHash from the compiler output. */
+	RENDERCORE_API void GenerateOutputHash();
+
+	friend FArchive& operator<<(FArchive& Ar, FShaderCompilerOutput& Output)
+	{
+		// Note: this serialize is used to pass between UE4 and the shader compile worker, recompile both when modifying
+		Ar << Output.ParameterMap << Output.Errors << Output.Target << Output.ShaderCode << Output.NumInstructions << Output.NumTextureSamplers << Output.bSucceeded;
+		Ar << Output.bFailedRemovingUnused << Output.bSupportsQueryingUsedAttributes << Output.UsedAttributes;
+		Ar << Output.OptionalFinalShaderSource;
+		Ar << Output.PlatformDebugData;
+
+		return Ar;
+	}
+};
+
+/** 
+ * Validates the format of a virtual shader file path.
+ * Meant to be use as such: check(CheckVirtualShaderFilePath(VirtualFilePath));
+ * CompileErrors output array is optional.
+ */
+extern RENDERCORE_API bool CheckVirtualShaderFilePath(const FString& VirtualPath, TArray<FShaderCompilerError>* CompileErrors = nullptr);
+
+/**
+ * Converts an absolute or relative shader filename to a filename relative to
+ * the shader directory.
+ * @param InFilename - The shader filename.
+ * @returns a filename relative to the shaders directory.
+ */
+extern RENDERCORE_API FString ParseVirtualShaderFilename(const FString& InFilename);
+
+/**
+ * Loads the shader file with the given name.
+ * @param VirtualFilePath - The virtual path of shader file to load.
+ * @param OutFileContents - If true is returned, will contain the contents of the shader file.
+ * @return True if the file was successfully loaded.
+ */
+extern RENDERCORE_API bool LoadShaderSourceFile(const TCHAR* VirtualFilePath, FString& OutFileContents, TArray<FShaderCompilerError>* OutCompileErrors);
+
+/** Loads the shader file with the given name.  If the shader file couldn't be loaded, throws a fatal error. */
+extern RENDERCORE_API void LoadShaderSourceFileChecked(const TCHAR* VirtualFilePath, FString& OutFileContents);
+
+/**
+ * Recursively populates IncludeFilenames with the include filenames from Filename
+ */
+extern RENDERCORE_API void GetShaderIncludes(const TCHAR* EntryPointVirtualFilePath, const TCHAR* VirtualFilePath, TArray<FString>& IncludeVirtualFilePaths, EShaderPlatform ShaderPlatform, uint32 DepthLimit=100);
+
+/**
+ * Calculates a Hash for the given filename if it does not already exist in the Hash cache.
+ * @param Filename - shader file to Hash
+ * @param ShaderPlatform - shader platform to Hash
+ */
+extern RENDERCORE_API const class FSHAHash& GetShaderFileHash(const TCHAR* VirtualFilePath, EShaderPlatform ShaderPlatform);
+
+/**
+ * Calculates a Hash for the list of filenames if it does not already exist in the Hash cache.
+ */
+extern RENDERCORE_API const class FSHAHash& GetShaderFilesHash(const TArray<FString>& VirtualFilePaths, EShaderPlatform ShaderPlatform);
+
+extern void BuildShaderFileToUniformBufferMap(TMap<FString, TArray<const TCHAR*> >& ShaderFileToUniformBufferVariables);
+
+/**
+ * Flushes the shader file and CRC cache, and regenerates the binary shader files if necessary.
+ * Allows shader source files to be re-read properly even if they've been modified since startup.
+ */
+extern RENDERCORE_API void FlushShaderFileCache();
+
+extern RENDERCORE_API void VerifyShaderSourceFiles(EShaderPlatform ShaderPlatform);
+
+struct FCachedUniformBufferDeclaration
+{
+	// Using SharedPtr so we can hand off lifetime ownership to FShaderCompilerEnvironment::IncludeVirtualPathToExternalContentsMap when invalidating this cache
+	TSharedPtr<FString> Declaration;
+};
+
+/** Parses the given source file and its includes for references of uniform buffers, which are then stored in UniformBufferEntries. */
+extern void GenerateReferencedUniformBuffers(
+	const TCHAR* SourceFilename,
+	const TCHAR* ShaderTypeName,
+	const TMap<FString, TArray<const TCHAR*> >& ShaderFileToUniformBufferVariables,
+	TMap<const TCHAR*,FCachedUniformBufferDeclaration>& UniformBufferEntries);
+
+/** Records information about all the uniform buffer layouts referenced by UniformBufferEntries. */
+extern RENDERCORE_API void SerializeUniformBufferInfo(class FShaderSaveArchive& Ar, const TMap<const TCHAR*,FCachedUniformBufferDeclaration>& UniformBufferEntries);
+
+
+
+/**
+ * Returns the map virtual shader directory path -> real shader directory path.
+ */
+extern RENDERCORE_API const TMap<FString, FString>& AllShaderSourceDirectoryMappings();
+
+/** Hook for shader compile worker to reset the directory mappings. */
+extern RENDERCORE_API void ResetAllShaderSourceDirectoryMappings();
+
+/**
+ * Maps a real shader directory existing on disk to a virtual shader directory.
+ * @param VirtualShaderDirectory Unique absolute path of the virtual shader directory (ex: /Project).
+ * @param RealShaderDirectory FPlatformProcess::BaseDir() relative path of the directory map.
+ */
+extern RENDERCORE_API void AddShaderSourceDirectoryMapping(const FString& VirtualShaderDirectory, const FString& RealShaderDirectory);
diff --git a/Engine/Source/Runtime/RenderCore/Public/ShaderParameterMacros.h b/Engine/Source/Runtime/RenderCore/Public/ShaderParameterMacros.h
new file mode 100644
--- /dev/null
+++ b/Engine/Source/Runtime/RenderCore/Public/ShaderParameterMacros.h
@@ -0,0 +1,845 @@
+// Copyright 1998-2018 Epic Games, Inc. All Rights Reserved.
+
+/*=============================================================================
+	ShaderParameterMacros.h: Macros to builds shader parameter structures and
+		their metadata.
+=============================================================================*/
+
+#pragma once
+
+#include "ShaderParameterMetadata.h"
+
+
+class FGraphTexture;
+class FGraphSRV;
+class FGraphUAV;
+
+
+// RHICreateUniformBuffer assumes C++ constant layout matches the shader layout when extracting float constants, yet the C++ struct contains pointers.  
+// Enforce a min size of 64 bits on pointer types in uniform buffer structs to guarantee layout matching between languages.
+#define SHADER_PARAMETER_POINTER_ALIGNMENT sizeof(uint64)
+static_assert(sizeof(void*) <= SHADER_PARAMETER_POINTER_ALIGNMENT, "The alignment of pointer needs to match the largest pointer.");
+
+
+/** Alignements tools because alignas() does not work on type in clang. */
+template<typename T, int32 Alignment>
+class TAlignedTypedef;
+
+#define IMPLEMENT_ALIGNED_TYPE(Alignment) \
+	template<typename T> \
+	class TAlignedTypedef<T,Alignment> \
+	{ \
+	public: \
+		typedef MS_ALIGN(Alignment) T Type GCC_ALIGN(Alignment); \
+	};
+
+IMPLEMENT_ALIGNED_TYPE(1);
+IMPLEMENT_ALIGNED_TYPE(2);
+IMPLEMENT_ALIGNED_TYPE(4);
+IMPLEMENT_ALIGNED_TYPE(8);
+IMPLEMENT_ALIGNED_TYPE(16);
+#undef IMPLEMENT_ALIGNED_TYPE
+
+
+#if PLATFORM_64BITS
+
+/** Fixed 8bytes sized pointer for shader parameters. */
+template<typename PtrType>
+using TNonAlignedShaderParameterPtr = PtrType;
+
+/** Fixed 8bytes sized and aligned pointer for shader parameters. */
+template<typename PtrType>
+using TAlignedShaderParameterPtr = typename TAlignedTypedef<PtrType, SHADER_PARAMETER_POINTER_ALIGNMENT>::Type;
+
+static_assert(sizeof(void*) == 8, "Wrong PLATFORM_64BITS settings.");
+
+#else //!PLATFORM_64BITS
+
+/** Fixed 8bytes sized pointer for shader parameters. */
+template<typename PtrType>
+class TNonAlignedShaderParameterPtr
+{
+public:
+	FORCEINLINE void operator=(const PtrType& Other)
+	{
+		Ref = Other;
+	}
+
+	FORCEINLINE operator PtrType&()
+	{
+		return Ref;
+	}
+
+	FORCEINLINE operator const PtrType&() const
+	{
+		return Ref;
+	}
+
+	FORCEINLINE PtrType& operator->() const
+	{
+		return Ref;
+	}
+
+private:
+	PtrType Ref;
+	uint32 _Padding;
+
+	static_assert(sizeof(PtrType) == sizeof(void*), "T should be a pointer.");
+};
+
+/** Fixed 8bytes sized and aligned pointer for shader parameters. */
+template<typename PtrType>
+class alignas(SHADER_PARAMETER_POINTER_ALIGNMENT) TAlignedShaderParameterPtr : public TNonAlignedShaderParameterPtr
+{ };
+
+static_assert(sizeof(void*) == 4, "Wrong PLATFORM_64BITS settings.");
+
+#endif // !PLATFORM_64BITS
+
+
+/** A reference to a uniform buffer RHI resource with a specific structure. */
+template<typename TBufferStruct>
+class TUniformBufferRef : public FUniformBufferRHIRef
+{
+public:
+
+	/** Initializes the reference to null. */
+	TUniformBufferRef()
+	{}
+
+	/** Creates a uniform buffer with the given value, and returns a structured reference to it. */
+	static TUniformBufferRef<TBufferStruct> CreateUniformBufferImmediate(const TBufferStruct& Value, EUniformBufferUsage Usage)
+	{
+		check(IsInRenderingThread() || IsInRHIThread());
+		return TUniformBufferRef<TBufferStruct>(RHICreateUniformBuffer(&Value,TBufferStruct::StaticStructMetadata.GetLayout(),Usage));
+	}
+	/** Creates a uniform buffer with the given value, and returns a structured reference to it. */
+	static FLocalUniformBuffer CreateLocalUniformBuffer(FRHICommandList& RHICmdList, const TBufferStruct& Value, EUniformBufferUsage Usage)
+	{
+		return RHICmdList.BuildLocalUniformBuffer(&Value, sizeof(TBufferStruct), TBufferStruct::StaticStructMetadata.GetLayout());
+	}
+
+private:
+
+	/** A private constructor used to coerce an arbitrary RHI uniform buffer reference to a structured reference. */
+	TUniformBufferRef(FUniformBufferRHIParamRef InRHIRef)
+	: FUniformBufferRHIRef(InRHIRef)
+	{}
+
+	template<typename TBufferStruct2>
+	friend class TUniformBuffer;
+};
+
+
+/** Render graph information about how to bind a render target. */
+struct FRenderTargetBinding
+{
+	FRenderTargetBinding() = default;
+	FRenderTargetBinding(const FRenderTargetBinding&) = default;
+
+	/** Creates a render target binding informations.
+	 *
+	 * Notes: Load and store action are on purpose without default values, to force the user to not forget one of these.
+	 */
+	FRenderTargetBinding(const FGraphTexture* InTexture, ERenderTargetLoadAction InLoadAction, ERenderTargetStoreAction InStoreAction, uint8 InMipIndex = 0)
+		: Texture(InTexture)
+		, LoadAction(InLoadAction)
+		, StoreAction(InStoreAction)
+		, MipIndex(InMipIndex)
+	{}
+
+	inline const FGraphTexture* GetTexture() const {
+		return Texture;
+	};
+	inline ERenderTargetLoadAction GetLoadAction() const {
+		return LoadAction;
+	};
+	inline ERenderTargetStoreAction GetStoreAction() const {
+		return StoreAction;
+	};
+	inline uint8 GetMipIndex() const {
+		return MipIndex;
+	};
+
+private:
+	/** All parameters required to bind a render target deferred. This are purposefully private to
+	 * force the user to call FRenderTargetBinding() constructor, forcing him to specify the load and store action.
+	 */
+	TNonAlignedShaderParameterPtr<const FGraphTexture*> Texture = nullptr;
+	ERenderTargetLoadAction		LoadAction		= ERenderTargetLoadAction::ENoAction;
+	ERenderTargetStoreAction	StoreAction		= ERenderTargetStoreAction::ENoAction;
+	uint8						MipIndex		= 0;
+};
+
+
+/** Render graph information about how to bind a depth-stencil render target. */
+struct FDepthStencilBinding
+{
+	TNonAlignedShaderParameterPtr<const FGraphTexture*> Texture = nullptr;
+	ERenderTargetLoadAction		DepthLoadAction		= ERenderTargetLoadAction::ENoAction;
+	ERenderTargetStoreAction	DepthStoreAction	= ERenderTargetStoreAction::ENoAction;
+	ERenderTargetLoadAction		StencilLoadAction	= ERenderTargetLoadAction::ENoAction;
+	ERenderTargetStoreAction	StencilStoreAction	= ERenderTargetStoreAction::ENoAction;
+	FExclusiveDepthStencil		DepthStencilAccess	= FExclusiveDepthStencil::DepthNop_StencilNop;
+};
+
+
+/** Special shader parameters type for a pass parameter to setup render targets. */
+struct alignas(SHADER_PARAMETER_STRUCT_ALIGNMENT) FRenderTargetBindingSlots
+{
+	TStaticArray<FRenderTargetBinding, MaxSimultaneousRenderTargets> Output;
+	FDepthStencilBinding DepthStencil;
+
+	/** Accessors for regular output to simplify the syntax to:
+	 *
+	 *	FRenderTargetParameters PassParameters;
+	 *	PassParameters.RenderTargets.DepthStencil = ... ;
+	 *	PassParameters.RenderTargets[0] = ... ;
+	 */
+	FRenderTargetBinding& operator[](uint32 Index)
+	{
+		return Output[Index];
+	}
+
+	const FRenderTargetBinding& operator[](uint32 Index) const
+	{
+		return Output[Index];
+	}
+
+	struct FTypeInfo {
+		static constexpr int32 NumRows = 1;
+		static constexpr int32 NumColumns = 1;
+		static constexpr int32 NumElements = 0;
+		static constexpr int32 Alignment = SHADER_PARAMETER_STRUCT_ALIGNMENT;
+		static constexpr bool bIsStoredInConstantBuffer = false;
+
+		using TAlignedType = FRenderTargetBindingSlots;
+
+		static inline const FShaderParametersMetadata* GetStructMetadata() { return nullptr; }
+	};
+};
+
+static_assert(sizeof(FRenderTargetBindingSlots) == 144, "FRenderTargetBindingSlots needs to be same size on all platforms.");
+
+
+/** Template to transcode some meta data information for a type <TypeParameter> not specific to shader parameters API. */
+template<typename TypeParameter>
+struct TShaderParameterTypeInfo
+{
+	/** Defines what the type actually is. */
+	static constexpr EUniformBufferBaseType BaseType = UBMT_INVALID;
+
+	/** Defines the number rows and columns for vector or matrix based data typed. */
+	static constexpr int32 NumRows = 1;
+	static constexpr int32 NumColumns = 1;
+
+	/** Defines the number of elements in an array fashion. 0 means this is not a TStaticArray,
+	 * which therefor means there is 1 element.
+	 */
+	static constexpr int32 NumElements = 0;
+
+	/** Defines the alignment of the elements in bytes. */
+	static constexpr int32 Alignment = alignof(TypeParameter);
+
+	/** Defines whether this element is stored in constant buffer or not.
+	 * This informations is usefull to ensure at compile time everything in the
+	 * structure get defined at the end of the structure, to reduce as much as possible
+	 * the size of the constant buffer.
+	 */
+	static constexpr bool bIsStoredInConstantBuffer = true;
+
+	/** Type that is actually alligned. */
+	using TAlignedType = TypeParameter;
+
+	static const FShaderParametersMetadata* GetStructMetadata() { return &TypeParameter::StaticStructMetadata; }
+};
+
+template<>
+struct TShaderParameterTypeInfo<bool>
+{
+	static constexpr EUniformBufferBaseType BaseType = UBMT_BOOL;
+	static constexpr int32 NumRows = 1;
+	static constexpr int32 NumColumns = 1;
+	static constexpr int32 NumElements = 0;
+	static constexpr int32 Alignment = 4;
+	static constexpr bool bIsStoredInConstantBuffer = true;
+	
+	using TAlignedType = TAlignedTypedef<bool, Alignment>::Type;
+
+	static const FShaderParametersMetadata* GetStructMetadata() { return nullptr; }
+};
+
+template<>
+struct TShaderParameterTypeInfo<uint32>
+{
+	static constexpr EUniformBufferBaseType BaseType = UBMT_UINT32;
+	static constexpr int32 NumRows = 1;
+	static constexpr int32 NumColumns = 1;
+	static constexpr int32 NumElements = 0;
+	static constexpr int32 Alignment = 4;
+	static constexpr bool bIsStoredInConstantBuffer = true;
+
+	using TAlignedType = TAlignedTypedef<uint32, Alignment>::Type;
+
+	static const FShaderParametersMetadata* GetStructMetadata() { return nullptr; }
+};
+	
+
+template<>
+struct TShaderParameterTypeInfo<int32>
+{
+	static constexpr EUniformBufferBaseType BaseType = UBMT_INT32;
+	static constexpr int32 NumRows = 1;
+	static constexpr int32 NumColumns = 1;
+	static constexpr int32 NumElements = 0;
+	static constexpr int32 Alignment = 4;
+	static constexpr bool bIsStoredInConstantBuffer = true;
+
+	using TAlignedType = TAlignedTypedef<int32, Alignment>::Type;
+
+	static const FShaderParametersMetadata* GetStructMetadata() { return nullptr; }
+};
+
+template<>
+struct TShaderParameterTypeInfo<float>
+{
+	static constexpr EUniformBufferBaseType BaseType = UBMT_FLOAT32;
+	static constexpr int32 NumRows = 1;
+	static constexpr int32 NumColumns = 1;
+	static constexpr int32 NumElements = 0;
+	static constexpr int32 Alignment = 4;
+	static constexpr bool bIsStoredInConstantBuffer = true;
+
+	using TAlignedType = TAlignedTypedef<float, Alignment>::Type;
+
+	static const FShaderParametersMetadata* GetStructMetadata() { return nullptr; }
+};
+
+template<>
+struct TShaderParameterTypeInfo<FVector2D>
+{
+	static constexpr EUniformBufferBaseType BaseType = UBMT_FLOAT32;
+	static constexpr int32 NumRows = 1;
+	static constexpr int32 NumColumns = 2;
+	static constexpr int32 NumElements = 0;
+	static constexpr int32 Alignment = 8;
+	static constexpr bool bIsStoredInConstantBuffer = true;
+
+	using TAlignedType = TAlignedTypedef<FVector2D, Alignment>::Type;
+
+	static const FShaderParametersMetadata* GetStructMetadata() { return nullptr; }
+};
+
+template<>
+struct TShaderParameterTypeInfo<FVector>
+{
+	static constexpr EUniformBufferBaseType BaseType = UBMT_FLOAT32;
+	static constexpr int32 NumRows = 1;
+	static constexpr int32 NumColumns = 3;
+	static constexpr int32 NumElements = 0;
+	static constexpr int32 Alignment = 16;
+	static constexpr bool bIsStoredInConstantBuffer = true;
+
+	using TAlignedType = TAlignedTypedef<FVector, Alignment>::Type;
+
+	static const FShaderParametersMetadata* GetStructMetadata() { return nullptr; }
+};
+
+template<>
+struct TShaderParameterTypeInfo<FVector4>
+{
+	static constexpr EUniformBufferBaseType BaseType = UBMT_FLOAT32;
+	static constexpr int32 NumRows = 1;
+	static constexpr int32 NumColumns = 4;
+	static constexpr int32 NumElements = 0;
+	static constexpr int32 Alignment = 16;
+	static constexpr bool bIsStoredInConstantBuffer = true;
+
+	using TAlignedType = TAlignedTypedef<FVector4, Alignment>::Type;
+
+	static const FShaderParametersMetadata* GetStructMetadata() { return nullptr; }
+};
+	
+template<>
+struct TShaderParameterTypeInfo<FLinearColor>
+{
+	static constexpr EUniformBufferBaseType BaseType = UBMT_FLOAT32;
+	static constexpr int32 NumRows = 1;
+	static constexpr int32 NumColumns = 4;
+	static constexpr int32 NumElements = 0;
+	static constexpr int32 Alignment = 16;
+	static constexpr bool bIsStoredInConstantBuffer = true;
+
+	using TAlignedType = TAlignedTypedef<FLinearColor, Alignment>::Type;
+
+	static const FShaderParametersMetadata* GetStructMetadata() { return nullptr; }
+};
+	
+template<>
+struct TShaderParameterTypeInfo<FIntPoint>
+{
+	static constexpr EUniformBufferBaseType BaseType = UBMT_INT32;
+	static constexpr int32 NumRows = 1;
+	static constexpr int32 NumColumns = 2;
+	static constexpr int32 NumElements = 0;
+	static constexpr int32 Alignment = 8;
+	static constexpr bool bIsStoredInConstantBuffer = true;
+
+	using TAlignedType = TAlignedTypedef<FIntPoint, Alignment>::Type;
+
+	static const FShaderParametersMetadata* GetStructMetadata() { return nullptr; }
+};
+
+template<>
+struct TShaderParameterTypeInfo<FIntVector>
+{
+	static constexpr EUniformBufferBaseType BaseType = UBMT_INT32;
+	static constexpr int32 NumRows = 1;
+	static constexpr int32 NumColumns = 3;
+	static constexpr int32 NumElements = 0;
+	static constexpr int32 Alignment = 16;
+	static constexpr bool bIsStoredInConstantBuffer = true;
+
+	using TAlignedType = TAlignedTypedef<FIntVector, Alignment>::Type;
+
+	static const FShaderParametersMetadata* GetStructMetadata() { return nullptr; }
+};
+
+template<>
+struct TShaderParameterTypeInfo<FIntRect>
+{
+	static constexpr EUniformBufferBaseType BaseType = UBMT_INT32;
+	static constexpr int32 NumRows = 1;
+	static constexpr int32 NumColumns = 4;
+	static constexpr int32 NumElements = 0;
+	static constexpr int32 Alignment = 16;
+	static constexpr bool bIsStoredInConstantBuffer = true;
+
+	using TAlignedType = TAlignedTypedef<FIntRect, Alignment>::Type;
+
+	static const FShaderParametersMetadata* GetStructMetadata() { return nullptr; }
+};
+
+template<>
+struct TShaderParameterTypeInfo<FMatrix>
+{
+	static constexpr EUniformBufferBaseType BaseType = UBMT_FLOAT32;
+	static constexpr int32 NumRows = 4;
+	static constexpr int32 NumColumns = 4;
+	static constexpr int32 NumElements = 0;
+	static constexpr int32 Alignment = 16;
+	static constexpr bool bIsStoredInConstantBuffer = true;
+
+	using TAlignedType = TAlignedTypedef<FMatrix, Alignment>::Type;
+
+	static const FShaderParametersMetadata* GetStructMetadata() { return nullptr; }
+};
+
+template<typename T, size_t InNumElements>
+struct TShaderParameterTypeInfo<T[InNumElements]>
+{
+	static constexpr EUniformBufferBaseType BaseType = TShaderParameterTypeInfo<T>::BaseType;
+	static constexpr int32 NumRows = TShaderParameterTypeInfo<T>::NumRows;
+	static constexpr int32 NumColumns = TShaderParameterTypeInfo<T>::NumColumns;
+	static constexpr int32 NumElements = InNumElements;
+	static constexpr int32 Alignment = SHADER_PARAMETER_ARRAY_ELEMENT_ALIGNMENT;
+	static constexpr bool bIsStoredInConstantBuffer = TShaderParameterTypeInfo<T>::bIsStoredInConstantBuffer;
+
+	using TAlignedType = TStaticArray<T, InNumElements, Alignment>;
+
+	static const FShaderParametersMetadata* GetStructMetadata() { return TShaderParameterTypeInfo<T>::GetStructMetadata(); }
+};
+	
+template<typename T,size_t InNumElements,uint32 IgnoredAlignment>
+struct TShaderParameterTypeInfo<TStaticArray<T,InNumElements,IgnoredAlignment>>
+{
+	static constexpr EUniformBufferBaseType BaseType = TShaderParameterTypeInfo<T>::BaseType;
+	static constexpr int32 NumRows = TShaderParameterTypeInfo<T>::NumRows;
+	static constexpr int32 NumColumns = TShaderParameterTypeInfo<T>::NumColumns;
+	static constexpr int32 NumElements = InNumElements;
+	static constexpr int32 Alignment = SHADER_PARAMETER_ARRAY_ELEMENT_ALIGNMENT;
+	static constexpr bool bIsStoredInConstantBuffer = TShaderParameterTypeInfo<T>::bIsStoredInConstantBuffer;
+
+	using TAlignedType = TStaticArray<T, InNumElements, Alignment>;
+
+	static const FShaderParametersMetadata* GetStructMetadata() { return TShaderParameterTypeInfo<T>::GetStructMetadata(); }
+};
+
+template<>
+struct TShaderParameterTypeInfo<FShaderResourceViewRHIParamRef>
+{
+	static constexpr int32 NumRows = 1;
+	static constexpr int32 NumColumns = 1;
+	static constexpr int32 NumElements = 0;
+	static constexpr int32 Alignment = SHADER_PARAMETER_POINTER_ALIGNMENT;
+	static constexpr bool bIsStoredInConstantBuffer = false;
+
+	using TAlignedType = TAlignedShaderParameterPtr<FShaderResourceViewRHIParamRef>;
+
+	static const FShaderParametersMetadata* GetStructMetadata() { return nullptr; }
+};
+static_assert(sizeof(TShaderParameterTypeInfo<FShaderResourceViewRHIParamRef>::TAlignedType) == SHADER_PARAMETER_POINTER_ALIGNMENT, "Uniform buffer layout must not be platform dependent.");
+
+template<>
+struct TShaderParameterTypeInfo<FSamplerStateRHIParamRef>
+{
+	static constexpr int32 NumRows = 1;
+	static constexpr int32 NumColumns = 1;
+	static constexpr int32 NumElements = 0;
+	static constexpr int32 Alignment = SHADER_PARAMETER_POINTER_ALIGNMENT;
+	static constexpr bool bIsStoredInConstantBuffer = false;
+
+	using TAlignedType = TAlignedShaderParameterPtr<FSamplerStateRHIParamRef>;
+
+	static const FShaderParametersMetadata* GetStructMetadata() { return nullptr; }
+};
+static_assert(sizeof(TShaderParameterTypeInfo<FSamplerStateRHIParamRef>::TAlignedType) == SHADER_PARAMETER_POINTER_ALIGNMENT, "Uniform buffer layout must not be platform dependent.");
+
+template<>
+struct TShaderParameterTypeInfo<FTextureRHIParamRef>
+{
+	static constexpr int32 NumRows = 1;
+	static constexpr int32 NumColumns = 1;
+	static constexpr int32 NumElements = 0;
+	static constexpr int32 Alignment = SHADER_PARAMETER_POINTER_ALIGNMENT;
+	static constexpr bool bIsStoredInConstantBuffer = false;
+
+	using TAlignedType = TAlignedShaderParameterPtr<FTextureRHIParamRef>;
+
+	static const FShaderParametersMetadata* GetStructMetadata() { return nullptr; }
+};
+static_assert(sizeof(TShaderParameterTypeInfo<FTextureRHIParamRef>::TAlignedType) == SHADER_PARAMETER_POINTER_ALIGNMENT, "Uniform buffer layout must not be platform dependent.");
+
+template<>
+struct TShaderParameterTypeInfo<const FGraphTexture*>
+{
+	static constexpr int32 NumRows = 1;
+	static constexpr int32 NumColumns = 1;
+	static constexpr int32 NumElements = 0;
+	static constexpr int32 Alignment = SHADER_PARAMETER_POINTER_ALIGNMENT;
+	static constexpr bool bIsStoredInConstantBuffer = false;
+
+	using TAlignedType = TAlignedShaderParameterPtr<const FGraphTexture*>;
+
+	static const FShaderParametersMetadata* GetStructMetadata() { return nullptr; }
+};
+static_assert(sizeof(TShaderParameterTypeInfo<const FGraphTexture*>::TAlignedType) == SHADER_PARAMETER_POINTER_ALIGNMENT, "Uniform buffer layout must not be platform dependent.");
+
+template<>
+struct TShaderParameterTypeInfo<const FGraphSRV*>
+{
+	static constexpr int32 NumRows = 1;
+	static constexpr int32 NumColumns = 1;
+	static constexpr int32 NumElements = 0;
+	static constexpr int32 Alignment = SHADER_PARAMETER_POINTER_ALIGNMENT;
+	static constexpr bool bIsStoredInConstantBuffer = false;
+	
+	using TAlignedType = TAlignedShaderParameterPtr<const FGraphSRV*>;
+	
+	static const FShaderParametersMetadata* GetStructMetadata() { return nullptr; }
+};
+static_assert(sizeof(TShaderParameterTypeInfo<const FGraphSRV*>::TAlignedType) == SHADER_PARAMETER_POINTER_ALIGNMENT, "Uniform buffer layout must not be platform dependent.");
+
+template<>
+struct TShaderParameterTypeInfo<const FGraphUAV*>
+{
+	static constexpr int32 NumRows = 1;
+	static constexpr int32 NumColumns = 1;
+	static constexpr int32 NumElements = 0;
+	static constexpr int32 Alignment = SHADER_PARAMETER_POINTER_ALIGNMENT;
+	static constexpr bool bIsStoredInConstantBuffer = false;
+	
+	using TAlignedType = TAlignedShaderParameterPtr<const FGraphUAV*>;
+	
+	static const FShaderParametersMetadata* GetStructMetadata() { return nullptr; }
+};
+static_assert(sizeof(TShaderParameterTypeInfo<const FGraphUAV*>::TAlignedType) == SHADER_PARAMETER_POINTER_ALIGNMENT, "Uniform buffer layout must not be platform dependent.");
+
+template<class UniformBufferStructType>
+struct TShaderParameterTypeInfo<TUniformBufferRef<UniformBufferStructType>>
+{
+	static constexpr int32 NumRows = 1;
+	static constexpr int32 NumColumns = 1;
+	static constexpr int32 NumElements = 0;
+	static constexpr int32 Alignment = SHADER_PARAMETER_POINTER_ALIGNMENT;
+	static constexpr bool bIsStoredInConstantBuffer = false;
+	
+	using TAlignedType = TAlignedShaderParameterPtr<TUniformBufferRef<UniformBufferStructType>>;
+	
+	static const FShaderParametersMetadata* GetStructMetadata() { return &UniformBufferStructType::StaticStructMetadata; }
+};
+
+
+#define INTERNAL_BEGIN_GLOBAL_SHADER_PARAMETER_STRUCT \
+	static FShaderParametersMetadata StaticStructMetadata; \
+
+#define INTERNAL_GLOBAL_SHADER_PARAMETER_GET_STRUCT_METADATA(StructTypeName) \
+	return &StructTypeName::StaticStructMetadata;
+
+#define INTERNAL_LOCAL_SHADER_PARAMETER_GET_STRUCT_METADATA(StructTypeName) \
+	static FShaderParametersMetadata StaticStructMetadata(\
+		FShaderParametersMetadata::EUseCase::ShaderParameterStruct, \
+		FName(TEXT(#StructTypeName)), \
+		nullptr, \
+		nullptr, \
+		sizeof(StructTypeName), \
+		StructTypeName::zzGetMembers()); \
+	return &StaticStructMetadata;
+
+
+/** Begins a uniform buffer struct declaration. */
+#define INTERNAL_SHADER_PARAMETER_STRUCT_BEGIN(StructTypeName,PrefixKeywords,ConstructorSuffix,GetStructMetadataScope) \
+	class alignas(SHADER_PARAMETER_STRUCT_ALIGNMENT) PrefixKeywords StructTypeName \
+	{ \
+	public: \
+		StructTypeName () ConstructorSuffix \
+		struct FTypeInfo { \
+			static constexpr int32 NumRows = 1; \
+			static constexpr int32 NumColumns = 1; \
+			static constexpr int32 NumElements = 0; \
+			static constexpr int32 Alignment = SHADER_PARAMETER_STRUCT_ALIGNMENT; \
+			static constexpr bool bIsStoredInConstantBuffer = true; \
+			using TAlignedType = StructTypeName; \
+			static inline const FShaderParametersMetadata* GetStructMetadata() { GetStructMetadataScope } \
+		}; \
+	private: \
+		typedef StructTypeName zzTThisStruct; \
+		struct zzFirstMemberId { enum { HasDeclaredResource = 0 }; }; \
+		static TArray<FShaderParametersMetadata::FMember> zzGetMembersBefore(zzFirstMemberId) \
+		{ \
+			return TArray<FShaderParametersMetadata::FMember>(); \
+		} \
+		typedef zzFirstMemberId
+
+/** Declares a member of a uniform buffer struct. */
+#define DECLARE_UNIFORM_BUFFER_STRUCT_MEMBER_EXPLICIT(BaseType,TypeInfo,MemberType,MemberName,ArrayDecl,Precision,OptionalShaderType,IsMemberStruct) \
+		zzMemberId##MemberName; \
+	public: \
+		typedef MemberType zzA##MemberName ArrayDecl; \
+		typedef TypeInfo::TAlignedType zzT##MemberName; \
+		zzT##MemberName MemberName; \
+		static_assert(BaseType != UBMT_INVALID, "Invalid type " #MemberType " of member " #MemberName "."); \
+	private: \
+		struct zzNextMemberId##MemberName { enum { HasDeclaredResource = zzMemberId##MemberName::HasDeclaredResource || !TypeInfo::bIsStoredInConstantBuffer }; }; \
+		static TArray<FShaderParametersMetadata::FMember> zzGetMembersBefore(zzNextMemberId##MemberName) \
+		{ \
+			static_assert(!TypeInfo::bIsStoredInConstantBuffer || zzMemberId##MemberName::HasDeclaredResource == 0, "Shader resources must be declared after " #MemberName "."); \
+			static_assert(TypeInfo::bIsStoredInConstantBuffer || TIsArrayOrRefOfType<decltype(OptionalShaderType), TCHAR>::Value, "No shader type for " #MemberName "."); \
+			/* Route the member enumeration on to the function for the member following this. */ \
+			TArray<FShaderParametersMetadata::FMember> OutMembers = zzGetMembersBefore(zzMemberId##MemberName()); \
+			/* Add this member. */ \
+			OutMembers.Add(FShaderParametersMetadata::FMember( \
+				TEXT(#MemberName), \
+				OptionalShaderType, \
+				STRUCT_OFFSET(zzTThisStruct,MemberName), \
+				EUniformBufferBaseType(BaseType), \
+				Precision, \
+				TypeInfo::NumRows, \
+				TypeInfo::NumColumns, \
+				TypeInfo::NumElements, \
+				TypeInfo::GetStructMetadata() \
+				)); \
+			static_assert( \
+				(STRUCT_OFFSET(zzTThisStruct,MemberName) & (TypeInfo::Alignment - 1)) == 0, \
+				"Misaligned uniform buffer struct member " #MemberName "."); \
+			return OutMembers; \
+		} \
+		typedef zzNextMemberId##MemberName
+
+/** Finds the FShaderParametersMetadata corresponding to the given name, or NULL if not found. */
+extern RENDERCORE_API FShaderParametersMetadata* FindUniformBufferStructByName(const TCHAR* StructName);
+extern RENDERCORE_API FShaderParametersMetadata* FindUniformBufferStructByFName(FName StructName);
+
+
+/** Begins & ends a shader parameter structure.
+ *
+ * Example:
+ *	BEGIN_SHADER_PARAMETER_STRUCT(FMyParameterStruct, RENDERER_API)
+ *	END_SHADER_PARAMETER_STRUCT()
+ */
+#define BEGIN_SHADER_PARAMETER_STRUCT(StructTypeName, PrefixKeywords) \
+	INTERNAL_SHADER_PARAMETER_STRUCT_BEGIN(StructTypeName, PrefixKeywords, {}, INTERNAL_LOCAL_SHADER_PARAMETER_GET_STRUCT_METADATA(StructTypeName))
+
+#define END_SHADER_PARAMETER_STRUCT() \
+		zzLastMemberId; \
+		static TArray<FShaderParametersMetadata::FMember> zzGetMembers() { return zzGetMembersBefore(zzLastMemberId()); } \
+	};
+
+/** Begins & ends a shader global parameter structure.
+ *
+ * Example:
+ *	// header
+ *	BEGIN_GLOBAL_SHADER_PARAMETER_STRUCT(FMyParameterStruct, RENDERER_API)
+ *	END_GLOBAL_SHADER_PARAMETER_STRUCT()
+ *
+ *	// C++ file
+ *	IMPLEMENT_GLOBAL_SHADER_PARAMETER_STRUCT(FMyParameterStruct, "MyShaderBindingName");
+ */
+#define BEGIN_GLOBAL_SHADER_PARAMETER_STRUCT(StructTypeName, PrefixKeywords) \
+	INTERNAL_SHADER_PARAMETER_STRUCT_BEGIN(StructTypeName,PrefixKeywords,{} INTERNAL_BEGIN_GLOBAL_SHADER_PARAMETER_STRUCT, INTERNAL_GLOBAL_SHADER_PARAMETER_GET_STRUCT_METADATA(StructTypeName))
+
+#define BEGIN_GLOBAL_SHADER_PARAMETER_STRUCT_WITH_CONSTRUCTOR(StructTypeName, PrefixKeywords) \
+	INTERNAL_SHADER_PARAMETER_STRUCT_BEGIN(StructTypeName,PrefixKeywords,; INTERNAL_BEGIN_GLOBAL_SHADER_PARAMETER_STRUCT, INTERNAL_GLOBAL_SHADER_PARAMETER_GET_STRUCT_METADATA(StructTypeName))
+
+#define END_GLOBAL_SHADER_PARAMETER_STRUCT() \
+	END_SHADER_PARAMETER_STRUCT()
+
+#define IMPLEMENT_GLOBAL_SHADER_PARAMETER_STRUCT(StructTypeName,ShaderVariableName) \
+	FShaderParametersMetadata StructTypeName::StaticStructMetadata( \
+	FShaderParametersMetadata::EUseCase::GlobalShaderParameterStruct, \
+	FName(TEXT(#StructTypeName)), \
+	TEXT(#StructTypeName), \
+	TEXT(ShaderVariableName), \
+	sizeof(StructTypeName), \
+	StructTypeName::zzGetMembers())
+
+
+/** Adds a constant-buffer stored value.
+ *
+ * Example:
+ *	SHADER_PARAMETER(float, MyScalar)
+ *	SHADER_PARAMETER(FMatrix, MyMatrix)
+ */
+#define SHADER_PARAMETER(MemberType,MemberName) \
+	SHADER_PARAMETER_EX(MemberType,MemberName,EShaderPrecisionModifier::Float)
+
+#define SHADER_PARAMETER_EX(MemberType,MemberName,Precision) \
+	DECLARE_UNIFORM_BUFFER_STRUCT_MEMBER_EXPLICIT(TShaderParameterTypeInfo<MemberType>::BaseType, TShaderParameterTypeInfo<MemberType>, MemberType,MemberName,,Precision,TEXT(""),false)
+
+/** Adds a constant-buffer stored array of values.
+ *
+ * Example:
+ *	SHADER_PARAMETER_ARRAY(float, MyScalarArray, [8])
+ *	SHADER_PARAMETER_ARRAY(FMatrix, MyMatrixArray)
+ */
+#define SHADER_PARAMETER_ARRAY(MemberType,MemberName,ArrayDecl) \
+	SHADER_PARAMETER_ARRAY_EX(MemberType,MemberName,ArrayDecl,EShaderPrecisionModifier::Float)
+
+#define SHADER_PARAMETER_ARRAY_EX(MemberType,MemberName,ArrayDecl,Precision) \
+	DECLARE_UNIFORM_BUFFER_STRUCT_MEMBER_EXPLICIT(TShaderParameterTypeInfo<MemberType ArrayDecl>::BaseType, TShaderParameterTypeInfo<MemberType ArrayDecl>, MemberType,MemberName,ArrayDecl,Precision,TEXT(""),false)
+
+/** Adds a texture.
+ *
+ * Example:
+ *	SHADER_PARAMETER_TEXTURE(Texture2D, MyTexture)
+ */
+#define SHADER_PARAMETER_TEXTURE(ShaderType,MemberName) \
+	DECLARE_UNIFORM_BUFFER_STRUCT_MEMBER_EXPLICIT(UBMT_TEXTURE, TShaderParameterTypeInfo<FTextureRHIParamRef>, FTextureRHIParamRef,MemberName,,EShaderPrecisionModifier::Float,TEXT(#ShaderType),false)
+
+/** Adds a shader resource view.
+ *
+ * Example:
+ *	SHADER_PARAMETER_SRV(Texture2D, MySRV)
+ */
+#define SHADER_PARAMETER_SRV(ShaderType,MemberName) \
+	DECLARE_UNIFORM_BUFFER_STRUCT_MEMBER_EXPLICIT(UBMT_SRV, TShaderParameterTypeInfo<FShaderResourceViewRHIParamRef>, FShaderResourceViewRHIParamRef,MemberName,,EShaderPrecisionModifier::Float,TEXT(#ShaderType),false)
+
+/** Adds a sampler.
+ *
+ * Example:
+ *	SHADER_PARAMETER_SAMPLER(SamplerState, MySampler)
+ */
+#define SHADER_PARAMETER_SAMPLER(ShaderType,MemberName) \
+	DECLARE_UNIFORM_BUFFER_STRUCT_MEMBER_EXPLICIT(UBMT_SAMPLER, TShaderParameterTypeInfo<FSamplerStateRHIParamRef>, FSamplerStateRHIParamRef,MemberName,,EShaderPrecisionModifier::Float,TEXT(#ShaderType),false)
+
+/** Adds a render graph tracked texture.
+ *
+ * Example:
+ *	SHADER_PARAMETER_GRAPH_TEXTURE(Texture2D, MyTexture)
+ */
+#define SHADER_PARAMETER_GRAPH_TEXTURE(ShaderType,MemberName) \
+	DECLARE_UNIFORM_BUFFER_STRUCT_MEMBER_EXPLICIT(UBMT_GRAPH_TRACKED_TEXTURE, TShaderParameterTypeInfo<const FGraphTexture*>, const FGraphTexture*,MemberName,,EShaderPrecisionModifier::Float,TEXT(#ShaderType),false)
+
+/** Adds a render graph tracked shader resource view.
+ *
+ * Example:
+ *	SHADER_PARAMETER_GRAPH_SRV(Texture2D, MySRV)
+ */
+#define SHADER_PARAMETER_GRAPH_SRV(ShaderType,MemberName) \
+	DECLARE_UNIFORM_BUFFER_STRUCT_MEMBER_EXPLICIT(UBMT_GRAPH_TRACKED_SRV, TShaderParameterTypeInfo<const FGraphSRV*>, const FGraphSRV*,MemberName,,EShaderPrecisionModifier::Float,TEXT(#ShaderType),false)
+
+/** Adds a render graph tracked unordered access view.
+ *
+ * Example:
+ *	SHADER_PARAMETER_GRAPH_UAV(RWTexture2D, MyUAV)
+ */
+#define SHADER_PARAMETER_GRAPH_UAV(ShaderType,MemberName) \
+	DECLARE_UNIFORM_BUFFER_STRUCT_MEMBER_EXPLICIT(UBMT_GRAPH_TRACKED_UAV, TShaderParameterTypeInfo<const FGraphUAV*>, const FGraphUAV*,MemberName,,EShaderPrecisionModifier::Float,TEXT(#ShaderType),false)
+
+/** Nests a shader parameter structure into another one, in C++ and shader code.
+ *
+ * Example:
+ *	BEGIN_SHADER_PARAMETER_STRUCT(FMyNestedStruct,)
+ *		SHADER_PARAMETER(float, MyScalar)
+ *		// ...
+ *	END_SHADER_PARAMETER_STRUCT()
+ *
+ *	BEGIN_SHADER_PARAMETER_STRUCT(FOtherStruct)
+ *		SHADER_PARAMETER_STRUCT(FMyNestedStruct, MyStruct)
+ *
+ * C++ use case:
+ *	FOtherStruct Parameters;
+ *	Parameters.MyStruct.MyScalar = 1.0f;
+ *
+ * Shader code for a globally named shader parameter struct:
+ *	float MyScalar = MyGlobalShaderBindingName.MyStruct.MyScalar;
+ *
+ * Shader code for a unnamed shader parameter struct:
+ *	float MyScalar = MyStruct_MyScalar;
+ */
+#define SHADER_PARAMETER_STRUCT(StructType,MemberName) \
+	DECLARE_UNIFORM_BUFFER_STRUCT_MEMBER_EXPLICIT(UBMT_NESTED_STRUCT, StructType::FTypeInfo, StructType, MemberName,,EShaderPrecisionModifier::Float,TEXT(#StructType),true)
+
+/** Include a shader parameter structure into another one in shader code.
+ *
+ * Example:
+ *	BEGIN_SHADER_PARAMETER_STRUCT(FMyNestedStruct,)
+ *		SHADER_PARAMETER(float, MyScalar)
+ *		// ...
+ *	END_SHADER_PARAMETER_STRUCT()
+ *
+ *	BEGIN_SHADER_PARAMETER_STRUCT(FOtherStruct)
+ *		SHADER_PARAMETER_STRUCT_INCLUDE(FMyNestedStruct, MyStruct)
+ *
+ * C++ use case:
+ *	FOtherStruct Parameters;
+ *	Parameters.MyStruct.MyScalar = 1.0f;
+ *
+ * Shader code for a globally named shader parameter struct:
+ *	float MyScalar = MyGlobalShaderBindingName.MyScalar;
+ *
+ * Shader code for a unnamed shader parameter struct:
+ *	float MyScalarValue = MyScalar;
+ */
+#define SHADER_PARAMETER_STRUCT_INCLUDE(StructType,MemberName) \
+	DECLARE_UNIFORM_BUFFER_STRUCT_MEMBER_EXPLICIT(UBMT_INCLUDED_STRUCT, StructType::FTypeInfo, StructType, MemberName,,EShaderPrecisionModifier::Float,TEXT(#StructType),true)
+
+/** Include a binding slot for a globally named shader parameter structure
+ *
+ * Example:
+ *	BEGIN_GLOBAL_SHADER_PARAMETER_STRUCT(FGlobalViewParameters,)
+ *		SHADER_PARAMETER(FVector4, ViewSizeAndInvSize)
+ *		// ...
+ *	END_GLOBAL_SHADER_PARAMETER_STRUCT()
+ *
+ *	BEGIN_SHADER_PARAMETER_STRUCT(FOtherStruct)
+ *		SHADER_PARAMETER_STRUCT_REF(FMyNestedStruct, MyStruct)
+ */
+#define SHADER_PARAMETER_STRUCT_REF(StructType,MemberName) \
+	DECLARE_UNIFORM_BUFFER_STRUCT_MEMBER_EXPLICIT(UBMT_REFERENCED_STRUCT, TShaderParameterTypeInfo<TUniformBufferRef<StructType>>, TUniformBufferRef<StructType>,MemberName,,EShaderPrecisionModifier::Float,TEXT(#StructType),false)
+
+/** Adds bindings slots for render targets on the structure. This is important for rasterizer based pass bind the
+ * render target at the RHI pass creation. The name of the struct member will forced to RenderTargets, and
+ * its type is a FRenderTargetBindingSlots.
+ *
+ * Example:
+ *	BEGIN_SHADER_PARAMETER_STRUCT(FGlobalViewParameters,)
+ *		RENDER_TARGET_BINDING_SLOTS()
+ *		// ...
+ *	END_SHADER_PARAMETER_STRUCT()
+ *
+ *	FGlobalViewParameters Parameters;
+ *	Parameters.RenderTargets.DepthStencil = FDepthStencilBinding( //... );
+ */
+#define RENDER_TARGET_BINDING_SLOTS() \
+	DECLARE_UNIFORM_BUFFER_STRUCT_MEMBER_EXPLICIT(UBMT_RENDER_TARGET_BINDING_SLOTS, FRenderTargetBindingSlots::FTypeInfo, FRenderTargetBindingSlots,RenderTargets,,EShaderPrecisionModifier::Float,TEXT(""),false)
diff --git a/Engine/Source/Runtime/RenderCore/Public/ShaderParameterMetadata.h b/Engine/Source/Runtime/RenderCore/Public/ShaderParameterMetadata.h
new file mode 100644
--- /dev/null
+++ b/Engine/Source/Runtime/RenderCore/Public/ShaderParameterMetadata.h
@@ -0,0 +1,187 @@
+// Copyright 1998-2018 Epic Games, Inc. All Rights Reserved.
+
+/*=============================================================================
+	ShaderParameterMetadata.h: Meta data about shader parameter structures
+=============================================================================*/
+
+#pragma once
+
+#include "CoreMinimal.h"
+#include "Containers/List.h"
+#include "Containers/StaticArray.h"
+#include "RHI.h"
+
+
+namespace EShaderPrecisionModifier
+{
+	enum Type
+	{
+		Float,
+		Half,
+		Fixed
+	};
+};
+
+/** Each entry in a resource table is provided to the shader compiler for creating mappings. */
+struct FResourceTableEntry
+{
+	/** The name of the uniform buffer in which this resource exists. */
+	FString UniformBufferName;
+	/** The type of the resource (EUniformBufferBaseType). */
+	uint16 Type;
+	/** The index of the resource in the table. */
+	uint16 ResourceIndex;
+};
+
+/** A uniform buffer struct. */
+class RENDERCORE_API FShaderParametersMetadata
+{
+public:
+	/** The use case of the uniform buffer structures. */
+	enum class EUseCase
+	{
+		/** Stand alone shader parameter struct used for render passes and shader parameters. */
+		ShaderParameterStruct,
+
+		/** Globally named shader parameter struct to be stored in uniform buffer. */
+		GlobalShaderParameterStruct,
+
+		/** Shader parameter struct generated from assets, such as material parameter collection or Niagara. */
+		DataDrivenShaderParameterStruct,
+	};
+	
+	/** A member of a shader parameter structure. */
+	class FMember
+	{
+	public:
+
+		/** Initialization constructor. */
+		FMember(
+			const TCHAR* InName,
+			const TCHAR* InShaderType,
+			uint32 InOffset,
+			EUniformBufferBaseType InBaseType,
+			EShaderPrecisionModifier::Type InPrecision,
+			uint32 InNumRows,
+			uint32 InNumColumns,
+			uint32 InNumElements,
+			const FShaderParametersMetadata* InStruct
+			)
+		:	Name(InName)
+		,	ShaderType(InShaderType)
+		,	Offset(InOffset)
+		,	BaseType(InBaseType)
+		,	Precision(InPrecision)
+		,	NumRows(InNumRows)
+		,	NumColumns(InNumColumns)
+		,	NumElements(InNumElements)
+		,	Struct(InStruct)
+		{}
+	
+		const TCHAR* GetName() const { return Name; }
+		const TCHAR* GetShaderType() const { return ShaderType; }
+		uint32 GetOffset() const { return Offset; }
+		EUniformBufferBaseType GetBaseType() const { return BaseType; }
+		EShaderPrecisionModifier::Type GetPrecision() const { return Precision; }
+		uint32 GetNumRows() const { return NumRows; }
+		uint32 GetNumColumns() const { return NumColumns; }
+		uint32 GetNumElements() const { return NumElements; }
+		const FShaderParametersMetadata* GetStructMetadata() const { return Struct; }
+
+		/** Returns the size of the member. */
+		inline uint32 GetMemberSize() const
+		{
+			check(BaseType == UBMT_BOOL || BaseType == UBMT_FLOAT32 || BaseType == UBMT_INT32 || BaseType == UBMT_UINT32);
+			uint32 ElementSize = sizeof(uint32) * NumRows * NumColumns;
+
+			/** If this an array, the alignment of the element are changed. */
+			if (NumElements > 0)
+			{
+				return Align(ElementSize, SHADER_PARAMETER_ARRAY_ELEMENT_ALIGNMENT) * NumElements;
+			}
+			return ElementSize;
+		}
+
+	private:
+
+		const TCHAR* Name;
+		const TCHAR* ShaderType;
+		uint32 Offset;
+		EUniformBufferBaseType BaseType;
+		EShaderPrecisionModifier::Type Precision;
+		uint32 NumRows;
+		uint32 NumColumns;
+		uint32 NumElements;
+		const FShaderParametersMetadata* Struct;
+	};
+
+	/** Initialization constructor. */
+	FShaderParametersMetadata(
+		EUseCase UseCase,
+		const FName& InLayoutName,
+		const TCHAR* InStructTypeName,
+		const TCHAR* InShaderVariableName,
+		uint32 InSize,
+		const TArray<FMember>& InMembers);
+
+	virtual ~FShaderParametersMetadata()
+	{
+		GlobalListLink.Unlink();
+		GetNameStructMap().Remove(FName(StructTypeName, FNAME_Find));
+	}
+
+	void GetNestedStructs(TArray<const FShaderParametersMetadata*>& OutNestedStructs) const;
+
+	void AddResourceTableEntries(TMap<FString, FResourceTableEntry>& ResourceTableMap, TMap<FString, uint32>& ResourceTableLayoutHashes) const;
+
+	const TCHAR* GetStructTypeName() const { return StructTypeName; }
+	const TCHAR* GetShaderVariableName() const { return ShaderVariableName; }
+	uint32 GetSize() const { return Size; }
+	EUseCase GetUseCase() const { return UseCase; }
+	const FRHIUniformBufferLayout& GetLayout() const 
+	{ 
+		check(bLayoutInitialized);
+		return Layout; 
+	}
+	const TArray<FMember>& GetMembers() const { return Members; }
+
+	/** Find a member for a given offset. */
+	void FindMemberFromOffset(uint16 MemberOffset, const FShaderParametersMetadata** OutContainingStruct, const FShaderParametersMetadata::FMember** OutMember) const;
+
+	static TLinkedList<FShaderParametersMetadata*>*& GetStructList();
+	/** Speed up finding the uniform buffer by its name */
+	static TMap<FName, FShaderParametersMetadata*>& GetNameStructMap();
+
+	/** Initialize all the global shader parameter structs. */
+	static void InitializeAllGlobalStructs();
+
+private:
+	/** Name of the structure type in C++ and shader code. */
+	const TCHAR* const StructTypeName;
+
+	/** Name of the shader variable name for global shader parameter structs. */
+	const TCHAR* const ShaderVariableName;
+
+	/** Size of the entire struct in bytes. */
+	const uint32 Size;
+
+	/** The use case of this shader parameter struct. */
+	const EUseCase UseCase;
+
+	/** Layout of all the resources in the shader parameter struct. */
+	FRHIUniformBufferLayout Layout;
+	
+	/** List of all members. */
+	TArray<FMember> Members;
+
+	/** Shackle elements in global link list of globally named shader parameters. */
+	TLinkedList<FShaderParametersMetadata*> GlobalListLink;
+
+	/** Whether the layout is actually initialized yet or not. */
+	uint32 bLayoutInitialized : 1;
+
+
+	void InitializeLayout();
+
+	void AddResourceTableEntriesRecursive(const TCHAR* UniformBufferName, const TCHAR* Prefix, uint16& ResourceIndex, TMap<FString, FResourceTableEntry>& ResourceTableMap) const;
+}; // class FShaderParametersMetadata
diff --git a/Engine/Source/Runtime/RenderCore/Public/ShaderParameterStruct.h b/Engine/Source/Runtime/RenderCore/Public/ShaderParameterStruct.h
new file mode 100644
--- /dev/null
+++ b/Engine/Source/Runtime/RenderCore/Public/ShaderParameterStruct.h
@@ -0,0 +1,201 @@
+// Copyright 1998-2018 Epic Games, Inc. All Rights Reserved.
+
+/*=============================================================================
+	ShaderParameterStruct.h: API to submit all shader parameters in single function call.
+=============================================================================*/
+
+#pragma once
+
+#include "CoreMinimal.h"
+#include "Shader.h"
+#include "RHI.h"
+
+
+/** Instrument a shader class to use strutured shader parameters API.
+ *
+ * class FMyShaderClassCS : public FGlobalShader
+ * {
+ *		DECLARE_GLOBAL_SHADER(FMyShaderClassCS);
+ *		SHADER_USE_PARAMETER_STRUCT(FMyShaderClassCS, FGlobalShader);
+ *
+ *		BEGIN_SHADER_PARAMETER_STRUCT(FParameters)
+ *			SHADER_PARAMETER(FMatrix, ViewToClip)
+ *			//...
+ *		END_SHADER_PARAMETER_STRUCT()
+ * };
+ *
+ * Notes: Long therm, this macro will be no longer be needed.
+ */
+ // TODO(RDG): would not even need ShaderParentClass anymore. And in fact should not so Bindings.Bind() is not being called twice.
+#define SHADER_USE_PARAMETER_STRUCT(ShaderClass, ShaderParentClass) \
+	ShaderClass(const ShaderMetaType::CompiledShaderInitializerType& Initializer) \
+		: ShaderParentClass(Initializer) \
+	{ \
+		this->Bindings.Bind(Initializer.ParameterMap, *FParameters::FTypeInfo::GetStructMetadata()); \
+	} \
+	\
+	ShaderClass() \
+	{ } \
+
+
+/** Raise fatal error when a required shader parameter has not been set. */
+extern RENDERCORE_API void SetNullShaderParameterFatalError(const FShader* Shader, const FShaderParametersMetadata* ParametersMetadata, uint16 MemberOffset);
+
+
+/** Set compute shader UAVs. */
+template<typename TRHICmdList, typename TShaderClass, typename TShaderRHI>
+inline void SetShaderUAVs(TRHICmdList& RHICmdList, const TShaderClass* Shader, TShaderRHI* ShadeRHI, const typename TShaderClass::FParameters& Parameters)
+{}
+
+template<typename TRHICmdList, typename TShaderClass>
+inline void SetShaderUAVs(TRHICmdList& RHICmdList, const TShaderClass* Shader, FComputeShaderRHIParamRef ShadeRHI, const typename TShaderClass::FParameters& Parameters)
+{
+	const FShaderParameterBindings& Bindings = Shader->Bindings;
+
+	const typename TShaderClass::FParameters* ParametersPtr = &Parameters;
+	const uint8* Base = reinterpret_cast<const uint8*>(ParametersPtr);
+
+	// Graph UAVs
+	for (const FShaderParameterBindings::FResourceParameter& ParameterBinding : Bindings.GraphUAVs)
+	{
+		auto GraphUAV = *reinterpret_cast<const FGraphUAV* const*>(Base + ParameterBinding.ByteOffset);
+
+		if (DO_CHECK)
+		{
+			if (!GraphUAV)
+			{
+				SetNullShaderParameterFatalError(Shader, TShaderClass::FParameters::FTypeInfo::GetStructMetadata(), ParameterBinding.ByteOffset);
+			}
+
+			// Mark this resource as used by the pass for inefficient pass resource dependency debugging purpose.
+			GraphUAV->bIsActuallyUsedByPass = true;
+		}
+
+		RHICmdList.SetUAVParameter(ShadeRHI, ParameterBinding.BaseIndex, GraphUAV->GetRHIUnorderedAccessView());
+	}
+}
+
+
+/** Unset compute shader UAVs. */
+template<typename TRHICmdList, typename TShaderClass>
+inline void UnsetShaderUAVs(TRHICmdList& RHICmdList, const TShaderClass* Shader, FComputeShaderRHIParamRef ShadeRHI)
+{
+	// TODO(RDG): Once all shader sets their parameter through this, can refactor RHI so all UAVs of a shader get unset through a single RHI function call.
+	const FShaderParameterBindings& Bindings = Shader->Bindings;
+	for (const FShaderParameterBindings::FResourceParameter& ParameterBinding : Bindings.GraphUAVs)
+	{
+		RHICmdList.SetUAVParameter(ShadeRHI, ParameterBinding.BaseIndex, nullptr);
+	}
+}
+
+
+/** Set shader's parameters from its parameters struct. */
+template<typename TRHICmdList, typename TShaderClass, typename TShaderRHI>
+inline void SetShaderParameters(TRHICmdList& RHICmdList, const TShaderClass* Shader, TShaderRHI* ShadeRHI, const typename TShaderClass::FParameters& Parameters)
+{
+	// TODO(RDG): Once all shader sets their parameter through this, can refactor RHI so all shader parameters get sets through a single RHI function call.
+	const FShaderParameterBindings& Bindings = Shader->Bindings;
+
+	const typename TShaderClass::FParameters* ParametersPtr = &Parameters;
+	const uint8* Base = reinterpret_cast<const uint8*>(ParametersPtr);
+
+	// Parameters
+	for (const FShaderParameterBindings::FParameter& ParameterBinding : Bindings.Parameters)
+	{
+		const void* DataPtr = reinterpret_cast<const char*>(&Parameters) + ParameterBinding.ByteOffset;
+		RHICmdList.SetShaderParameter(ShadeRHI, ParameterBinding.BufferIndex, ParameterBinding.BaseIndex, ParameterBinding.ByteSize, DataPtr);
+	}
+
+	// Textures
+	for (const FShaderParameterBindings::FResourceParameter& ParameterBinding : Bindings.Textures)
+	{
+		auto ShaderParameterRef = *reinterpret_cast<const FTextureRHIParamRef*>(Base + ParameterBinding.ByteOffset);
+
+		if (DO_CHECK && !ShaderParameterRef)
+		{
+			SetNullShaderParameterFatalError(Shader, TShaderClass::FParameters::FTypeInfo::GetStructMetadata(), ParameterBinding.ByteOffset);
+		}
+
+		RHICmdList.SetShaderTexture(ShadeRHI, ParameterBinding.BaseIndex, ShaderParameterRef);
+	}
+
+	// SRVs
+	for (const FShaderParameterBindings::FResourceParameter& ParameterBinding : Bindings.SRVs)
+	{
+		auto ShaderParameterRef = *reinterpret_cast<const FShaderResourceViewRHIParamRef*>(Base + ParameterBinding.ByteOffset);
+
+		if (DO_CHECK && !ShaderParameterRef)
+		{
+			SetNullShaderParameterFatalError(Shader, TShaderClass::FParameters::FTypeInfo::GetStructMetadata(), ParameterBinding.ByteOffset);
+		}
+
+		RHICmdList.SetShaderResourceViewParameter(ShadeRHI, ParameterBinding.BaseIndex, ShaderParameterRef);
+	}
+
+	// Samplers
+	for (const FShaderParameterBindings::FResourceParameter& ParameterBinding : Bindings.Samplers)
+	{
+		auto ShaderParameterRef = *reinterpret_cast<const FSamplerStateRHIParamRef*>(Base + ParameterBinding.ByteOffset);
+
+		if (DO_CHECK && !ShaderParameterRef)
+		{
+			SetNullShaderParameterFatalError(Shader, TShaderClass::FParameters::FTypeInfo::GetStructMetadata(), ParameterBinding.ByteOffset);
+		}
+
+		RHICmdList.SetShaderSampler(ShadeRHI, ParameterBinding.BaseIndex, ShaderParameterRef);
+	}
+
+	// Graph Textures
+	for (const FShaderParameterBindings::FResourceParameter& ParameterBinding : Bindings.GraphTextures)
+	{
+		auto GraphTexture = *reinterpret_cast<const FGraphTexture* const*>(Base + ParameterBinding.ByteOffset);
+
+		if (DO_CHECK)
+		{
+			if (!GraphTexture)
+			{
+				SetNullShaderParameterFatalError(Shader, TShaderClass::FParameters::FTypeInfo::GetStructMetadata(), ParameterBinding.ByteOffset);
+			}
+
+			// Mark this resource as used by the pass for inefficient pass resource dependency debugging purpose.
+			GraphTexture->bIsActuallyUsedByPass = true;
+		}
+
+		RHICmdList.SetShaderTexture(ShadeRHI, ParameterBinding.BaseIndex, GraphTexture->GetRHITexture());
+	}
+
+	// Graph SRVs
+	for (const FShaderParameterBindings::FResourceParameter& ParameterBinding : Bindings.GraphSRVs)
+	{
+		auto GraphSRV = *reinterpret_cast<const FGraphSRV* const*>(Base + ParameterBinding.ByteOffset);
+
+		if (DO_CHECK)
+		{
+			if (!GraphSRV)
+			{
+				SetNullShaderParameterFatalError(Shader, TShaderClass::FParameters::FTypeInfo::GetStructMetadata(), ParameterBinding.ByteOffset);
+			}
+
+			// Mark this resource as used by the pass for inefficient pass resource dependency debugging purpose.
+			GraphSRV->bIsActuallyUsedByPass = true;
+		}
+
+		RHICmdList.SetShaderResourceViewParameter(ShadeRHI, ParameterBinding.BaseIndex, GraphSRV->GetRHIShaderResourceView());
+	}
+
+	// Graph UAVs for compute shaders
+	SetShaderUAVs(RHICmdList, Shader, ShadeRHI, Parameters);
+
+	// Reference structures
+	for (const FShaderParameterBindings::FParameterStructReference& ParameterBinding : Bindings.ParameterReferences)
+	{
+		const TRefCountPtr<FRHIUniformBuffer>& ShaderParameterRef = *reinterpret_cast<const TRefCountPtr<FRHIUniformBuffer>*>(Base + ParameterBinding.ByteOffset);
+
+		if (DO_CHECK && !ShaderParameterRef.IsValid())
+		{
+			SetNullShaderParameterFatalError(Shader, TShaderClass::FParameters::FTypeInfo::GetStructMetadata(), ParameterBinding.ByteOffset);
+		}
+
+		RHICmdList.SetShaderUniformBuffer(ShadeRHI, ParameterBinding.BufferIndex, ShaderParameterRef);
+	}
+}
