Merging //UE4/Dev-Rendering-Graph@4492585 to Dev-Rendering (//UE4/Dev-Rendering)

diff --git a/Engine/Source/Runtime/RenderCore/Public/ShaderParameterUtils.h b/Engine/Source/Runtime/RenderCore/Public/ShaderParameterUtils.h
new file mode 100644
--- /dev/null
+++ b/Engine/Source/Runtime/RenderCore/Public/ShaderParameterUtils.h
@@ -0,0 +1,539 @@
+// Copyright 1998-2018 Epic Games, Inc. All Rights Reserved.
+
+/*=============================================================================
+	ShaderParameters.h: Shader parameter inline definitions.
+=============================================================================*/
+
+#pragma once
+
+#include "CoreMinimal.h"
+#include "RHI.h"
+#include "ShaderParameters.h"
+#include "ShaderCore.h"
+#include "Misc/App.h"
+
+template<typename TBufferStruct> class TUniformBuffer;
+template<typename TBufferStruct> class TUniformBufferRef;
+
+/**
+ * Sets the value of a  shader parameter.  Template'd on shader type
+ * A template parameter specified the type of the parameter value.
+ * NOTE: Shader should be the param ref type, NOT the param type, since Shader is passed by value. 
+ * Otherwise AddRef/ReleaseRef will be called many times.
+ */
+template<typename ShaderRHIParamRef, class ParameterType, typename TRHICmdList>
+void SetShaderValue(
+	TRHICmdList& RHICmdList,
+	ShaderRHIParamRef Shader,
+	const FShaderParameter& Parameter,
+	const ParameterType& Value,
+	uint32 ElementIndex = 0
+	)
+{
+	static_assert(!TIsPointer<ParameterType>::Value, "Passing by value is not valid.");
+
+	const uint32 AlignedTypeSize = Align(sizeof(ParameterType), SHADER_PARAMETER_ARRAY_ELEMENT_ALIGNMENT);
+	const int32 NumBytesToSet = FMath::Min<int32>(sizeof(ParameterType),Parameter.GetNumBytes() - ElementIndex * AlignedTypeSize);
+
+	// This will trigger if the parameter was not serialized
+	checkSlow(Parameter.IsInitialized());
+
+	if(NumBytesToSet > 0)
+	{
+		RHICmdList.SetShaderParameter(
+			Shader,
+			Parameter.GetBufferIndex(),
+			Parameter.GetBaseIndex() + ElementIndex * AlignedTypeSize,
+			(uint32)NumBytesToSet,
+			&Value
+			);
+	}
+}
+
+template<typename ShaderRHIParamRef, class ParameterType>
+void SetShaderValueOnContext(
+	IRHICommandContext& RHICmdListContext,
+	ShaderRHIParamRef Shader,
+	const FShaderParameter& Parameter,
+	const ParameterType& Value,
+	uint32 ElementIndex = 0
+	)
+{
+	static_assert(!TIsPointer<ParameterType>::Value, "Passing by value is not valid.");
+
+	const uint32 AlignedTypeSize = Align(sizeof(ParameterType), SHADER_PARAMETER_ARRAY_ELEMENT_ALIGNMENT);
+	const int32 NumBytesToSet = FMath::Min<int32>(sizeof(ParameterType), Parameter.GetNumBytes() - ElementIndex * AlignedTypeSize);
+
+	// This will trigger if the parameter was not serialized
+	checkSlow(Parameter.IsInitialized());
+
+	if (NumBytesToSet > 0)
+	{
+		RHICmdListContext.RHISetShaderParameter(
+			Shader,
+			Parameter.GetBufferIndex(),
+			Parameter.GetBaseIndex() + ElementIndex * AlignedTypeSize,
+			(uint32)NumBytesToSet,
+			&Value
+			);
+	}
+}
+
+
+
+/** Specialization of the above for C++ bool type. */
+template<typename ShaderRHIParamRef>
+void SetShaderValue(
+	FRHICommandList& RHICmdList, 
+	ShaderRHIParamRef Shader,
+	const FShaderParameter& Parameter,
+	const bool& Value,
+	uint32 ElementIndex = 0
+	)
+{
+	const uint32 BoolValue = Value;
+	SetShaderValue(RHICmdList, Shader, Parameter, BoolValue, ElementIndex);
+}
+
+/** Specialization of the above for C++ bool type. */
+template<typename ShaderRHIParamRef>
+void SetShaderValue(
+	FRHIAsyncComputeCommandList& RHICmdList,
+	ShaderRHIParamRef Shader,
+	const FShaderParameter& Parameter,
+	const bool& Value,
+	uint32 ElementIndex = 0
+	)
+{
+	const uint32 BoolValue = Value;
+	SetShaderValue(RHICmdList, Shader, Parameter, BoolValue, ElementIndex);
+}
+
+/**
+ * Sets the value of a shader parameter array.  Template'd on shader type
+ * A template parameter specified the type of the parameter value.
+ * NOTE: Shader should be the param ref type, NOT the param type, since Shader is passed by value. 
+ * Otherwise AddRef/ReleaseRef will be called many times.
+ */
+template<typename ShaderRHIParamRef,class ParameterType, typename TRHICmdList>
+void SetShaderValueArray(
+	TRHICmdList& RHICmdList,
+	ShaderRHIParamRef Shader,
+	const FShaderParameter& Parameter,
+	const ParameterType* Values,
+	uint32 NumElements,
+	uint32 BaseElementIndex = 0
+	)
+{
+	const uint32 AlignedTypeSize = Align(sizeof(ParameterType), SHADER_PARAMETER_ARRAY_ELEMENT_ALIGNMENT);
+	const int32 NumBytesToSet = FMath::Min<int32>(NumElements * AlignedTypeSize,Parameter.GetNumBytes() - BaseElementIndex * AlignedTypeSize);
+
+	// This will trigger if the parameter was not serialized
+	checkSlow(Parameter.IsInitialized());
+
+	if(NumBytesToSet > 0)
+	{
+		RHICmdList.SetShaderParameter(
+			Shader,
+			Parameter.GetBufferIndex(),
+			Parameter.GetBaseIndex() + BaseElementIndex * AlignedTypeSize,
+			(uint32)NumBytesToSet,
+			Values
+			);
+	}
+}
+
+/** Specialization of the above for C++ bool type. */
+template<typename ShaderRHIParamRef, typename TRHICmdList>
+void SetShaderValueArray(
+	TRHICmdList& RHICmdList,
+	ShaderRHIParamRef Shader,
+	const FShaderParameter& Parameter,
+	const bool* Values,
+	uint32 NumElements,
+	uint32 BaseElementIndex = 0
+	)
+{
+	UE_LOG(LogShaders, Fatal, TEXT("SetShaderValueArray does not support bool arrays."));
+}
+
+/**
+ * Sets the value of a pixel shader bool parameter.
+ */
+inline void SetPixelShaderBool(
+	FRHICommandList& RHICmdList, 
+	FPixelShaderRHIParamRef PixelShader,
+	const FShaderParameter& Parameter,
+	bool Value
+	)
+{
+	// This will trigger if the parameter was not serialized
+	checkSlow(Parameter.IsInitialized());
+
+	if (Parameter.GetNumBytes() > 0)
+	{
+		// Convert to uint32 before passing to RHI
+		uint32 BoolValue = Value;
+		RHICmdList.SetShaderParameter(
+			PixelShader,
+			Parameter.GetBufferIndex(),
+			Parameter.GetBaseIndex(),
+			sizeof(BoolValue),
+			&BoolValue
+			);
+	}
+}
+
+/**
+ * Sets the value of a shader texture parameter.  Template'd on shader type
+ */
+template<typename ShaderTypeRHIParamRef, typename TRHICmdList>
+FORCEINLINE void SetTextureParameter(
+	TRHICmdList& RHICmdList,
+	ShaderTypeRHIParamRef Shader,
+	const FShaderResourceParameter& TextureParameter,
+	const FShaderResourceParameter& SamplerParameter,
+	const FTexture* Texture,
+	uint32 ElementIndex = 0
+	)
+{
+	// This will trigger if the parameter was not serialized
+	checkSlow(TextureParameter.IsInitialized());
+	checkSlow(SamplerParameter.IsInitialized());
+	if(TextureParameter.IsBound())
+	{
+		Texture->LastRenderTime = FApp::GetCurrentTime();
+
+		if (ElementIndex < TextureParameter.GetNumResources())
+		{
+			RHICmdList.SetShaderTexture( Shader, TextureParameter.GetBaseIndex() + ElementIndex, Texture->TextureRHI);
+		}
+	}
+	
+	// @todo ue4 samplerstate Should we maybe pass in two separate values? SamplerElement and TextureElement? Or never allow an array of samplers? Unsure best
+	// if there is a matching sampler for this texture array index (ElementIndex), then set it. This will help with this case:
+	//			Texture2D LightMapTextures[NUM_LIGHTMAP_COEFFICIENTS];
+	//			SamplerState LightMapTexturesSampler;
+	// In this case, we only set LightMapTexturesSampler when ElementIndex is 0, we don't set the sampler state for all 4 textures
+	// This assumes that the all textures want to use the same sampler state
+	if(SamplerParameter.IsBound())
+	{
+		if (ElementIndex < SamplerParameter.GetNumResources())
+		{
+			RHICmdList.SetShaderSampler( Shader, SamplerParameter.GetBaseIndex() + ElementIndex, Texture->SamplerStateRHI);
+		}
+	}
+}
+
+/**
+ * Sets the value of a shader texture parameter. Template'd on shader type.
+ */
+template<typename ShaderTypeRHIParamRef, typename TRHICmdList>
+FORCEINLINE void SetTextureParameter(
+	TRHICmdList& RHICmdList,
+	ShaderTypeRHIParamRef Shader,
+	const FShaderResourceParameter& TextureParameter,
+	const FShaderResourceParameter& SamplerParameter,
+	FSamplerStateRHIParamRef SamplerStateRHI,
+	FTextureRHIParamRef TextureRHI,
+	uint32 ElementIndex = 0
+	)
+{
+	// This will trigger if the parameter was not serialized
+	checkSlow(TextureParameter.IsInitialized());
+	checkSlow(SamplerParameter.IsInitialized());
+	if(TextureParameter.IsBound())
+	{
+		if (ElementIndex < TextureParameter.GetNumResources())
+		{
+			RHICmdList.SetShaderTexture( Shader, TextureParameter.GetBaseIndex() + ElementIndex, TextureRHI);
+		}
+	}
+	// @todo ue4 samplerstate Should we maybe pass in two separate values? SamplerElement and TextureElement? Or never allow an array of samplers? Unsure best
+	// if there is a matching sampler for this texture array index (ElementIndex), then set it. This will help with this case:
+	//			Texture2D LightMapTextures[NUM_LIGHTMAP_COEFFICIENTS];
+	//			SamplerState LightMapTexturesSampler;
+	// In this case, we only set LightMapTexturesSampler when ElementIndex is 0, we don't set the sampler state for all 4 textures
+	// This assumes that the all textures want to use the same sampler state
+	if(SamplerParameter.IsBound())
+	{
+		if (ElementIndex < SamplerParameter.GetNumResources())
+		{
+			RHICmdList.SetShaderSampler( Shader, SamplerParameter.GetBaseIndex() + ElementIndex, SamplerStateRHI);
+		}
+	}
+}
+
+/**
+ * Sets the value of a shader surface parameter (e.g. to access MSAA samples).
+ * Template'd on shader type (e.g. pixel shader or compute shader).
+ */
+template<typename ShaderTypeRHIParamRef, typename TRHICmdList>
+FORCEINLINE void SetTextureParameter(
+	TRHICmdList& RHICmdList,
+	ShaderTypeRHIParamRef Shader,
+	const FShaderResourceParameter& Parameter,
+	FTextureRHIParamRef NewTextureRHI
+	)
+{
+	if(Parameter.IsBound())
+	{
+		RHICmdList.SetShaderTexture(
+			Shader,
+			Parameter.GetBaseIndex(),
+			NewTextureRHI
+			);
+	}
+}
+
+/**
+ * Sets the value of a shader sampler parameter. Template'd on shader type.
+ */
+template<typename ShaderTypeRHIParamRef, typename TRHICmdList>
+FORCEINLINE void SetSamplerParameter(
+	TRHICmdList& RHICmdList,
+	ShaderTypeRHIParamRef Shader,
+	const FShaderResourceParameter& Parameter,
+	FSamplerStateRHIParamRef SamplerStateRHI
+	)
+{
+	if(Parameter.IsBound())
+	{
+		RHICmdList.SetShaderSampler(
+			Shader,
+			Parameter.GetBaseIndex(),
+			SamplerStateRHI
+			);
+	}
+}
+
+/**
+ * Sets the value of a shader resource view parameter
+ * Template'd on shader type (e.g. pixel shader or compute shader).
+ */
+template<typename ShaderTypeRHIParamRef, typename TRHICmdList>
+FORCEINLINE void SetSRVParameter(
+	TRHICmdList& RHICmdList,
+	ShaderTypeRHIParamRef Shader,
+	const FShaderResourceParameter& Parameter,
+	FShaderResourceViewRHIParamRef NewShaderResourceViewRHI
+	)
+{
+	if(Parameter.IsBound())
+	{
+		RHICmdList.SetShaderResourceViewParameter(
+			Shader,
+			Parameter.GetBaseIndex(),
+			NewShaderResourceViewRHI
+			);
+	}
+}
+
+/**
+ * Sets the value of a unordered access view parameter
+ */
+template<typename TRHICmdList>
+FORCEINLINE void SetUAVParameter(
+TRHICmdList& RHICmdList,
+	FComputeShaderRHIParamRef ComputeShader,
+	const FShaderResourceParameter& Parameter,
+	FUnorderedAccessViewRHIParamRef NewUnorderedAccessViewRHI
+	)
+{
+	if(Parameter.IsBound())
+	{
+		RHICmdList.SetUAVParameter(
+			ComputeShader,
+			Parameter.GetBaseIndex(),
+			NewUnorderedAccessViewRHI
+			);
+	}
+}
+
+
+template<typename TRHICmdList>
+inline bool SetUAVParameterIfCS(TRHICmdList& RHICmdList, const FVertexShaderRHIParamRef Shader, const FShaderResourceParameter& UAVParameter, FUnorderedAccessViewRHIParamRef UAV)
+{
+	return false;
+}
+
+template<typename TRHICmdList>
+inline bool SetUAVParameterIfCS(TRHICmdList& RHICmdList, const FPixelShaderRHIParamRef Shader, const FShaderResourceParameter& UAVParameter, FUnorderedAccessViewRHIParamRef UAV)
+{
+	return false;
+}
+
+template<typename TRHICmdList>
+inline bool SetUAVParameterIfCS(TRHICmdList& RHICmdList, const FHullShaderRHIParamRef Shader, const FShaderResourceParameter& UAVParameter, FUnorderedAccessViewRHIParamRef UAV)
+{
+	return false;
+}
+
+template<typename TRHICmdList>
+inline bool SetUAVParameterIfCS(TRHICmdList& RHICmdList, const FDomainShaderRHIParamRef Shader, const FShaderResourceParameter& UAVParameter, FUnorderedAccessViewRHIParamRef UAV)
+{
+	return false;
+}
+
+template<typename TRHICmdList>
+inline bool SetUAVParameterIfCS(TRHICmdList& RHICmdList, const FGeometryShaderRHIParamRef Shader, const FShaderResourceParameter& UAVParameter, FUnorderedAccessViewRHIParamRef UAV)
+{
+	return false;
+}
+
+template<typename TRHICmdList>
+inline bool SetUAVParameterIfCS(TRHICmdList& RHICmdList, const FComputeShaderRHIParamRef Shader, const FShaderResourceParameter& UAVParameter, FUnorderedAccessViewRHIParamRef UAV)
+{
+	SetUAVParameter(RHICmdList, Shader, UAVParameter, UAV);
+	return UAVParameter.IsBound();
+}
+
+template<typename TShaderRHIRef, typename TRHICmdList>
+inline void FRWShaderParameter::SetBuffer(TRHICmdList& RHICmdList, TShaderRHIRef Shader, const FRWBuffer& RWBuffer) const
+{
+	if (!SetUAVParameterIfCS(RHICmdList, Shader, UAVParameter, RWBuffer.UAV))
+	{
+		SetSRVParameter(RHICmdList, Shader, SRVParameter, RWBuffer.SRV);
+	}
+}
+
+template<typename TShaderRHIRef, typename TRHICmdList>
+inline void FRWShaderParameter::SetBuffer(TRHICmdList& RHICmdList, TShaderRHIRef Shader, const FRWBufferStructured& RWBuffer) const
+{
+	if (!SetUAVParameterIfCS(RHICmdList, Shader, UAVParameter, RWBuffer.UAV))
+	{
+		SetSRVParameter(RHICmdList, Shader, SRVParameter, RWBuffer.SRV);
+	}
+}
+
+template<typename TShaderRHIRef, typename TRHICmdList>
+inline void FRWShaderParameter::SetTexture(TRHICmdList& RHICmdList, TShaderRHIRef Shader, const FTextureRHIParamRef Texture, FUnorderedAccessViewRHIParamRef UAV) const
+{
+	if (!SetUAVParameterIfCS(RHICmdList, Shader, UAVParameter, UAV))
+	{
+		SetTextureParameter(RHICmdList, Shader, SRVParameter, Texture);
+	}
+}
+
+template<typename TRHICmdList>
+inline void FRWShaderParameter::UnsetUAV(TRHICmdList& RHICmdList, FComputeShaderRHIParamRef ComputeShader) const
+{
+	SetUAVParameter(RHICmdList, ComputeShader,UAVParameter,FUnorderedAccessViewRHIRef());
+}
+
+
+/** Sets the value of a shader uniform buffer parameter to a uniform buffer containing the struct. */
+template<typename TShaderRHIRef>
+inline void SetLocalUniformBufferParameter(
+	FRHICommandList& RHICmdList,
+	TShaderRHIRef Shader,
+	const FShaderUniformBufferParameter& Parameter,
+	const FLocalUniformBuffer& LocalUniformBuffer
+	)
+{
+	// This will trigger if the parameter was not serialized
+	checkSlow(Parameter.IsInitialized());
+	if(Parameter.IsBound())
+	{
+		RHICmdList.SetLocalShaderUniformBuffer(Shader, Parameter.GetBaseIndex(), LocalUniformBuffer);
+	}
+}
+
+/** Sets the value of a shader uniform buffer parameter to a uniform buffer containing the struct. */
+template<typename TShaderRHIRef, typename TRHICmdList>
+inline void SetUniformBufferParameter(
+	TRHICmdList& RHICmdList,
+	TShaderRHIRef Shader,
+	const FShaderUniformBufferParameter& Parameter,
+	FUniformBufferRHIParamRef UniformBufferRHI
+	)
+{
+	// This will trigger if the parameter was not serialized
+	checkSlow(Parameter.IsInitialized());
+	// If it is bound, we must set it so something valid
+	checkSlow(!Parameter.IsBound() || UniformBufferRHI);
+	if(Parameter.IsBound())
+	{
+		RHICmdList.SetShaderUniformBuffer( Shader, Parameter.GetBaseIndex(), UniformBufferRHI);
+	}
+}
+
+/** Sets the value of a shader uniform buffer parameter to a uniform buffer containing the struct. */
+template<typename TShaderRHIRef, typename TBufferStruct, typename TRHICmdList>
+inline void SetUniformBufferParameter(
+	TRHICmdList& RHICmdList,
+	TShaderRHIRef Shader,
+	const TShaderUniformBufferParameter<TBufferStruct>& Parameter,
+	const TUniformBufferRef<TBufferStruct>& UniformBufferRef
+	)
+{
+	// This will trigger if the parameter was not serialized
+	checkSlow(Parameter.IsInitialized());
+	// If it is bound, we must set it so something valid
+	checkSlow(!Parameter.IsBound() || IsValidRef(UniformBufferRef));
+	if(Parameter.IsBound())
+	{
+		RHICmdList.SetShaderUniformBuffer( Shader, Parameter.GetBaseIndex(), UniformBufferRef);
+	}
+}
+
+/** Sets the value of a shader uniform buffer parameter to a uniform buffer containing the struct. */
+template<typename TShaderRHIRef, typename TBufferStruct, typename TRHICmdList>
+inline void SetUniformBufferParameter(
+	TRHICmdList& RHICmdList,
+	TShaderRHIRef Shader,
+	const TShaderUniformBufferParameter<TBufferStruct>& Parameter,
+	const TUniformBuffer<TBufferStruct>& UniformBuffer
+	)
+{
+	// This will trigger if the parameter was not serialized
+	checkSlow(Parameter.IsInitialized());
+	// If it is bound, we must set it so something valid
+	checkSlow(!Parameter.IsBound() || UniformBuffer.GetUniformBufferRHI());
+	if (Parameter.IsBound())
+	{
+		RHICmdList.SetShaderUniformBuffer( Shader, Parameter.GetBaseIndex(), UniformBuffer.GetUniformBufferRHI());
+	}
+}
+
+/** Sets the value of a shader uniform buffer parameter to a value of the struct. */
+template<typename TShaderRHIRef,typename TBufferStruct>
+inline void SetUniformBufferParameterImmediate(
+	FRHICommandList& RHICmdList,
+	TShaderRHIRef Shader,
+	const TShaderUniformBufferParameter<TBufferStruct>& Parameter,
+	const TBufferStruct& UniformBufferValue
+	)
+{
+	// This will trigger if the parameter was not serialized
+	checkSlow(Parameter.IsInitialized());
+	if(Parameter.IsBound())
+	{
+		RHICmdList.SetShaderUniformBuffer(
+			Shader,
+			Parameter.GetBaseIndex(),
+			RHICreateUniformBuffer(&UniformBufferValue,TBufferStruct::StaticStructMetadata.GetLayout(),UniformBuffer_SingleDraw)
+			);
+	}
+}
+
+/** Sets the value of a shader uniform buffer parameter to a value of the struct. */
+template<typename TShaderRHIRef,typename TBufferStruct, typename TRHICmdList>
+inline void SetUniformBufferParameterImmediate(
+	TRHICmdList& RHICmdList,
+	TShaderRHIRef Shader,
+	const TShaderUniformBufferParameter<TBufferStruct>& Parameter,
+	const TBufferStruct& UniformBufferValue
+	)
+{
+	// This will trigger if the parameter was not serialized
+	checkSlow(Parameter.IsInitialized());
+	if(Parameter.IsBound())
+	{
+		RHICmdList.SetShaderUniformBuffer(
+			Shader,
+			Parameter.GetBaseIndex(),
+			RHICreateUniformBuffer(&UniformBufferValue,TBufferStruct::StaticStructMetadata.GetLayout(),UniformBuffer_SingleDraw)
+			);
+	}
+}
diff --git a/Engine/Source/Runtime/RenderCore/Public/ShaderParameters.h b/Engine/Source/Runtime/RenderCore/Public/ShaderParameters.h
new file mode 100644
--- /dev/null
+++ b/Engine/Source/Runtime/RenderCore/Public/ShaderParameters.h
@@ -0,0 +1,238 @@
+// Copyright 1998-2018 Epic Games, Inc. All Rights Reserved.
+
+/*=============================================================================
+	ShaderParameters.h: Shader parameter definitions.
+=============================================================================*/
+
+#pragma once
+
+#include "CoreMinimal.h"
+#include "RHI.h"
+
+class FShaderParameterMap;
+class FShaderParametersMetadata;
+struct FShaderCompilerEnvironment;
+
+RENDERCORE_API void CacheUniformBufferIncludes(TMap<const TCHAR*, struct FCachedUniformBufferDeclaration>& Cache, EShaderPlatform Platform);
+
+
+enum EShaderParameterFlags
+{
+	// no shader error if the parameter is not used
+	SPF_Optional,
+	// shader error if the parameter is not used
+	SPF_Mandatory
+};
+
+/** A shader parameter's register binding. e.g. float1/2/3/4, can be an array, UAV */
+class FShaderParameter
+{
+public:
+	FShaderParameter()
+	:	BufferIndex(0)
+	,	BaseIndex(0)
+	,	NumBytes(0)
+#if UE_BUILD_DEBUG
+	,	bInitialized(false)
+#endif
+	{}
+
+	RENDERCORE_API void Bind(const FShaderParameterMap& ParameterMap,const TCHAR* ParameterName, EShaderParameterFlags Flags = SPF_Optional);
+	friend RENDERCORE_API FArchive& operator<<(FArchive& Ar,FShaderParameter& P);
+	bool IsBound() const { return NumBytes > 0; }
+
+	inline bool IsInitialized() const 
+	{ 
+#if UE_BUILD_DEBUG
+		return bInitialized; 
+#else 
+		return true;
+#endif
+	}
+
+	uint32 GetBufferIndex() const { return BufferIndex; }
+	uint32 GetBaseIndex() const { return BaseIndex; }
+	uint32 GetNumBytes() const { return NumBytes; }
+
+private:
+	uint16 BufferIndex;
+	uint16 BaseIndex;
+	// 0 if the parameter wasn't bound
+	uint16 NumBytes;
+
+#if UE_BUILD_DEBUG
+	bool bInitialized;
+#endif
+};
+
+/** A shader resource binding (textures or samplerstates). */
+class FShaderResourceParameter
+{
+public:
+	FShaderResourceParameter()
+	:	BaseIndex(0)
+	,	NumResources(0) 
+#if UE_BUILD_DEBUG
+	,	bInitialized(false)
+#endif
+	{}
+
+	RENDERCORE_API void Bind(const FShaderParameterMap& ParameterMap,const TCHAR* ParameterName,EShaderParameterFlags Flags = SPF_Optional);
+	friend RENDERCORE_API FArchive& operator<<(FArchive& Ar,FShaderResourceParameter& P);
+	bool IsBound() const { return NumResources > 0; }
+
+	inline bool IsInitialized() const 
+	{ 
+#if UE_BUILD_DEBUG
+		return bInitialized; 
+#else 
+		return true;
+#endif
+	}
+
+	uint32 GetBaseIndex() const { return BaseIndex; }
+	uint32 GetNumResources() const { return NumResources; }
+private:
+	uint16 BaseIndex;
+	uint16 NumResources;
+
+#if UE_BUILD_DEBUG
+	bool bInitialized;
+#endif
+};
+
+/** A class that binds either a UAV or SRV of a resource. */
+class FRWShaderParameter
+{
+public:
+
+	void Bind(const FShaderParameterMap& ParameterMap,const TCHAR* BaseName)
+	{
+		SRVParameter.Bind(ParameterMap,BaseName);
+
+		// If the shader wants to bind the parameter as a UAV, the parameter name must start with "RW"
+		FString UAVName = FString(TEXT("RW")) + BaseName;
+		UAVParameter.Bind(ParameterMap,*UAVName);
+
+		// Verify that only one of the UAV or SRV parameters is accessed by the shader.
+		checkf(!(SRVParameter.GetNumResources() && UAVParameter.GetNumResources()),TEXT("Shader binds SRV and UAV of the same resource: %s"),BaseName);
+	}
+
+	bool IsBound() const
+	{
+		return SRVParameter.IsBound() || UAVParameter.IsBound();
+	}
+
+	bool IsUAVBound() const
+	{
+		return UAVParameter.IsBound();
+	}
+
+	uint32 GetUAVIndex() const
+	{
+		return UAVParameter.GetBaseIndex();
+	}
+
+	friend FArchive& operator<<(FArchive& Ar,FRWShaderParameter& Parameter)
+	{
+		return Ar << Parameter.SRVParameter << Parameter.UAVParameter;
+	}
+
+	template<typename TShaderRHIRef, typename TRHICmdList>
+	inline void SetBuffer(TRHICmdList& RHICmdList, TShaderRHIRef Shader, const FRWBuffer& RWBuffer) const;
+
+	template<typename TShaderRHIRef, typename TRHICmdList>
+	inline void SetBuffer(TRHICmdList& RHICmdList, TShaderRHIRef Shader, const FRWBufferStructured& RWBuffer) const;
+
+	template<typename TShaderRHIRef, typename TRHICmdList>
+	inline void SetTexture(TRHICmdList& RHICmdList, TShaderRHIRef Shader, const FTextureRHIParamRef Texture, FUnorderedAccessViewRHIParamRef UAV) const;
+
+	template<typename TRHICmdList>
+	inline void UnsetUAV(TRHICmdList& RHICmdList, FComputeShaderRHIParamRef ComputeShader) const;
+
+private:
+
+	FShaderResourceParameter SRVParameter;
+	FShaderResourceParameter UAVParameter;
+};
+
+/** Creates a shader code declaration of this struct for the given shader platform. */
+extern RENDERCORE_API void CreateUniformBufferShaderDeclaration(const TCHAR* Name,const FShaderParametersMetadata& UniformBufferStruct, FString& OutDeclaration);
+
+class FShaderUniformBufferParameter
+{
+public:
+	FShaderUniformBufferParameter()
+	:	BaseIndex(0)
+	,	bIsBound(false) 
+#if UE_BUILD_DEBUG
+	,	bInitialized(false)
+#endif
+	{}
+
+	static RENDERCORE_API void ModifyCompilationEnvironment(const TCHAR* ParameterName,const FShaderParametersMetadata& Struct,EShaderPlatform Platform,FShaderCompilerEnvironment& OutEnvironment);
+
+	RENDERCORE_API void Bind(const FShaderParameterMap& ParameterMap,const TCHAR* ParameterName,EShaderParameterFlags Flags = SPF_Optional);
+
+	friend FArchive& operator<<(FArchive& Ar,FShaderUniformBufferParameter& P)
+	{
+		P.Serialize(Ar);
+		return Ar;
+	}
+
+	bool IsBound() const { return bIsBound; }
+
+	void Serialize(FArchive& Ar)
+	{
+#if UE_BUILD_DEBUG
+		if (Ar.IsLoading())
+		{
+			bInitialized = true;
+		}
+#endif
+		Ar << BaseIndex << bIsBound;
+	}
+
+	inline bool IsInitialized() const 
+	{ 
+#if UE_BUILD_DEBUG
+		return bInitialized; 
+#else 
+		return true;
+#endif
+	}
+
+	inline void SetInitialized() 
+	{ 
+#if UE_BUILD_DEBUG
+		bInitialized = true; 
+#endif
+	}
+
+	uint32 GetBaseIndex() const { return BaseIndex; }
+
+private:
+	uint16 BaseIndex;
+	bool bIsBound;
+
+#if UE_BUILD_DEBUG
+	bool bInitialized;
+#endif
+};
+
+/** A shader uniform buffer binding with a specific structure. */
+template<typename TBufferStruct>
+class TShaderUniformBufferParameter : public FShaderUniformBufferParameter
+{
+public:
+	static void ModifyCompilationEnvironment(const TCHAR* ParameterName,EShaderPlatform Platform, FShaderCompilerEnvironment& OutEnvironment)
+	{
+		FShaderUniformBufferParameter::ModifyCompilationEnvironment(ParameterName,TBufferStruct::StaticStruct,Platform,OutEnvironment);
+	}
+
+	friend FArchive& operator<<(FArchive& Ar,TShaderUniformBufferParameter& P)
+	{
+		P.Serialize(Ar);
+		return Ar;
+	}
+};
diff --git a/Engine/Source/Runtime/RenderCore/Public/UniformBuffer.h b/Engine/Source/Runtime/RenderCore/Public/UniformBuffer.h
--- a/Engine/Source/Runtime/RenderCore/Public/UniformBuffer.h
+++ b/Engine/Source/Runtime/RenderCore/Public/UniformBuffer.h
@@ -6,10 +6,7 @@
 
 #pragma once
 
-#include "CoreMinimal.h"
-#include "Containers/List.h"
-#include "Containers/StaticArray.h"
-#include "RHI.h"
+#include "ShaderParameterMacros.h"
 #include "RenderingThread.h"
 #include "RenderResource.h"
 #include "Templates/IsArrayOrRefOfType.h"
@@ -17,15 +14,14 @@
 class FShaderUniformBufferParameter;
 template<typename TBufferStruct> class TShaderUniformBufferParameter;
 
-namespace EShaderPrecisionModifier
+/** Creates a 
+uniform buffer with the given value, and returns a structured reference to it. */
+template<typename TBufferStruct>
+TUniformBufferRef<TBufferStruct> CreateUniformBufferImmediate(const TBufferStruct& Value, EUniformBufferUsage Usage)
 {
-	enum Type
-	{
-		Float,
-		Half,
-		Fixed
-	};
-};
+	return TUniformBufferRef<TBufferStruct>::CreateUniformBufferImmediate(Value, Usage);
+}
+
 
 /** A uniform buffer resource. */
 template<typename TBufferStruct>
@@ -58,7 +54,7 @@ public:
 	{
 		if (!Contents)
 		{
-			Contents = (uint8*)FMemory::Malloc(sizeof(TBufferStruct), UNIFORM_BUFFER_STRUCT_ALIGNMENT);
+			Contents = (uint8*)FMemory::Malloc(sizeof(TBufferStruct), SHADER_PARAMETER_STRUCT_ALIGNMENT);
 		}
 		FMemory::Memzero(Contents, sizeof(TBufferStruct));
 		UpdateRHI();
@@ -71,7 +67,7 @@ public:
 		UniformBufferRHI.SafeRelease();
 		if (Contents)
 		{
-			UniformBufferRHI = RHICreateUniformBuffer(Contents,TBufferStruct::StaticStruct.GetLayout(),BufferUsage);
+			UniformBufferRHI = CreateUniformBufferImmediate<TBufferStruct>(*((const TBufferStruct*)Contents), BufferUsage);
 		}
 	}
 	virtual void ReleaseDynamicRHI() override
@@ -86,6 +82,12 @@ public:
 		return UniformBufferRHI; 
 	}
 
+	const TUniformBufferRef<TBufferStruct>& GetUniformBufferRef() const
+	{
+		check(UniformBufferRHI.GetReference()); // you are trying to use a UB that was never filled with anything
+		return UniformBufferRHI;
+	}
+
 	EUniformBufferUsage BufferUsage;
 
 protected:
@@ -96,52 +98,17 @@ protected:
 		check(IsInRenderingThread());
 		if (!Contents)
 		{
-			Contents = (uint8*)FMemory::Malloc(sizeof(TBufferStruct),UNIFORM_BUFFER_STRUCT_ALIGNMENT);
+			Contents = (uint8*)FMemory::Malloc(sizeof(TBufferStruct), SHADER_PARAMETER_STRUCT_ALIGNMENT);
 		}
 		FMemory::Memcpy(Contents,&NewContents,sizeof(TBufferStruct));
 	}
 
 private:
 
-	FUniformBufferRHIRef UniformBufferRHI;
+	TUniformBufferRef<TBufferStruct> UniformBufferRHI;
 	uint8* Contents;
 };
 
-/** A reference to a uniform buffer RHI resource with a specific structure. */
-template<typename TBufferStruct>
-class TUniformBufferRef : public FUniformBufferRHIRef
-{
-public:
-
-	/** Initializes the reference to null. */
-	TUniformBufferRef()
-	{}
-
-	/** Initializes the reference to point to a buffer. */
-	TUniformBufferRef(const TUniformBuffer<TBufferStruct>& InBuffer)
-	: FUniformBufferRHIRef(InBuffer.GetUniformBufferRHI())
-	{}
-
-	/** Creates a uniform buffer with the given value, and returns a structured reference to it. */
-	static TUniformBufferRef<TBufferStruct> CreateUniformBufferImmediate(const TBufferStruct& Value, EUniformBufferUsage Usage)
-	{
-		check(IsInRenderingThread() || IsInRHIThread());
-		return TUniformBufferRef<TBufferStruct>(RHICreateUniformBuffer(&Value,TBufferStruct::StaticStruct.GetLayout(),Usage));
-	}
-	/** Creates a uniform buffer with the given value, and returns a structured reference to it. */
-	static FLocalUniformBuffer CreateLocalUniformBuffer(FRHICommandList& RHICmdList, const TBufferStruct& Value, EUniformBufferUsage Usage)
-	{
-		return RHICmdList.BuildLocalUniformBuffer(&Value, sizeof(TBufferStruct), TBufferStruct::StaticStruct.GetLayout());
-	}
-
-private:
-
-	/** A private constructor used to coerce an arbitrary RHI uniform buffer reference to a structured reference. */
-	TUniformBufferRef(FUniformBufferRHIParamRef InRHIRef)
-	: FUniformBufferRHIRef(InRHIRef)
-	{}
-};
-
 ENQUEUE_UNIQUE_RENDER_COMMAND_TWOPARAMETER_DECLARE_TEMPLATE(
 	SetUniformBufferContents,TBufferStruct,
 	TUniformBuffer<TBufferStruct>*,UniformBuffer,&UniformBuffer,
@@ -162,567 +129,3 @@ void BeginSetUniformBufferContents(
 		TUniformBuffer<TBufferStruct>*,&UniformBuffer,
 		TBufferStruct,Struct);
 }
-
-/** Each entry in a resource table is provided to the shader compiler for creating mappings. */
-struct FResourceTableEntry
-{
-	/** The name of the uniform buffer in which this resource exists. */
-	FString UniformBufferName;
-	/** The type of the resource (EUniformBufferBaseType). */
-	uint16 Type;
-	/** The index of the resource in the table. */
-	uint16 ResourceIndex;
-};
-
-/** A uniform buffer struct. */
-class RENDERCORE_API FUniformBufferStruct
-{
-public:
-	
-	/** A member of a uniform buffer type. */
-	class FMember
-	{
-	public:
-
-		/** Initialization constructor. */
-		FMember(
-			const TCHAR* InName,
-			const TCHAR* InShaderType,
-			uint32 InOffset,
-			EUniformBufferBaseType InBaseType,
-			EShaderPrecisionModifier::Type InPrecision,
-			uint32 InNumRows,
-			uint32 InNumColumns,
-			uint32 InNumElements,
-			const FUniformBufferStruct* InStruct
-			)
-		:	Name(InName)
-		,	ShaderType(InShaderType)
-		,	Offset(InOffset)
-		,	BaseType(InBaseType)
-		,	Precision(InPrecision)
-		,	NumRows(InNumRows)
-		,	NumColumns(InNumColumns)
-		,	NumElements(InNumElements)
-		,	Struct(InStruct)
-		{}
-	
-		const TCHAR* GetName() const { return Name; }
-		const TCHAR* GetShaderType() const { return ShaderType; }
-		uint32 GetOffset() const { return Offset; }
-		EUniformBufferBaseType GetBaseType() const { return BaseType; }
-		EShaderPrecisionModifier::Type GetPrecision() const { return Precision; }
-		uint32 GetNumRows() const { return NumRows; }
-		uint32 GetNumColumns() const { return NumColumns; }
-		uint32 GetNumElements() const { return NumElements; }
-		const FUniformBufferStruct* GetStruct() const { return Struct; }
-
-	private:
-
-		const TCHAR* Name;
-		const TCHAR* ShaderType;
-		uint32 Offset;
-		EUniformBufferBaseType BaseType;
-		EShaderPrecisionModifier::Type Precision;
-		uint32 NumRows;
-		uint32 NumColumns;
-		uint32 NumElements;
-		const FUniformBufferStruct* Struct;
-	};
-
-	typedef class FShaderUniformBufferParameter* (*ConstructUniformBufferParameterType)();
-
-	/** Initialization constructor. */
-	FUniformBufferStruct(const FName& InLayoutName, const TCHAR* InStructTypeName, const TCHAR* InShaderVariableName, ConstructUniformBufferParameterType InConstructRef, uint32 InSize, const TArray<FMember>& InMembers, bool bRegisterForAutoBinding);
-
-	virtual ~FUniformBufferStruct()
-	{
-		GlobalListLink.Unlink();
-		GetNameStructMap().Remove(FName(StructTypeName, FNAME_Find));
-	}
-
-	void InitializeLayout();
-
-	void GetNestedStructs(TArray<const FUniformBufferStruct*>& OutNestedStructs) const;
-
-	void AddResourceTableEntries(TMap<FString, FResourceTableEntry>& ResourceTableMap, TMap<FString, uint32>& ResourceTableLayoutHashes) const;
-
-	const TCHAR* GetStructTypeName() const { return StructTypeName; }
-	const TCHAR* GetShaderVariableName() const { return ShaderVariableName; }
-	const uint32 GetSize() const { return Size; }
-	const FRHIUniformBufferLayout& GetLayout() const 
-	{ 
-		check(bLayoutInitialized);
-		return Layout; 
-	}
-	const TArray<FMember>& GetMembers() const { return Members; }
-	FShaderUniformBufferParameter* ConstructTypedParameter() const { return (*ConstructUniformBufferParameterRef)(); }
-
-	static TLinkedList<FUniformBufferStruct*>*& GetStructList();
-	/** Speed up finding the uniform buffer by its name */
-	static TMap<FName, FUniformBufferStruct*>& GetNameStructMap();
-
-	static void InitializeStructs();
-
-private:
-	const TCHAR* StructTypeName;
-	const TCHAR* ShaderVariableName;
-	ConstructUniformBufferParameterType ConstructUniformBufferParameterRef;
-	uint32 Size;
-	bool bLayoutInitialized;
-	FRHIUniformBufferLayout Layout;
-	TArray<FMember> Members;
-	TLinkedList<FUniformBufferStruct*> GlobalListLink;
-
-	void AddResourceTableEntriesRecursive(const TCHAR* UniformBufferName, const TCHAR* Prefix, uint16& ResourceIndex, TMap<FString, FResourceTableEntry>& ResourceTableMap) const;
-};
-
-/** TFixedSizePointerWrapper acts like a pointer of the given type but enforces a fixed 64 bit size. */
-template<class T>
-class TFixedSizePointerWrapper
-{
-public:
-	void operator=(T Other)
-	{
-		Member = Other;
-	}
-
-	operator T&()
-	{
-		return Member;
-	}
-
-private:
-
-	union
-	{
-		T Member;
-		uint64 _MinSizeDummy;
-	};
-	
-	static_assert(sizeof(T) == sizeof(void*), "T should be a pointer.");
-	static_assert(sizeof(T) <= sizeof(uint64), "Assuming max 64 bit platform");
-};
-
-//
-// Uniform buffer alignment tools (should only be used by the uniform buffer type infos below).
-//
-
-	template<typename T,uint32 Alignment>
-	class TAlignedTypedef;
-
-	#define IMPLEMENT_ALIGNED_TYPE(Alignment) \
-	template<typename T> \
-	class TAlignedTypedef<T,Alignment> \
-	{ \
-	public: \
-		typedef MS_ALIGN(Alignment) T TAlignedType GCC_ALIGN(Alignment); \
-	};
-
-	IMPLEMENT_ALIGNED_TYPE(1);
-	IMPLEMENT_ALIGNED_TYPE(2);
-	IMPLEMENT_ALIGNED_TYPE(4);
-	IMPLEMENT_ALIGNED_TYPE(8);
-	IMPLEMENT_ALIGNED_TYPE(16);
-
-//
-// Template specializations used to map C++ types to uniform buffer member types.
-//
-
-	template<typename, bool>
-	class TUniformBufferTypeInfo
-	{
-	public:
-		enum { BaseType = UBMT_INVALID };
-		enum { NumRows = 0 };
-		enum { NumColumns = 0 };
-		enum { NumElements = 0 };
-		enum { Alignment = 1 };
-		enum { IsResource = 0 };
-		static const FUniformBufferStruct* GetStruct() { return NULL; }
-	};
-
-	template<>
-	class TUniformBufferTypeInfo<bool, false>
-	{
-	public:
-		enum { BaseType = UBMT_BOOL };
-		enum { NumRows = 1 };
-		enum { NumColumns = 1 };
-		enum { NumElements = 0 };
-		enum { Alignment = 4 };
-		enum { IsResource = 0 };
-		typedef TAlignedTypedef<bool,Alignment>::TAlignedType TAlignedType;
-		static const FUniformBufferStruct* GetStruct() { return NULL; }
-	};
-
-	template<>
-	class TUniformBufferTypeInfo<uint32, false>
-	{
-	public:
-		enum { BaseType = UBMT_UINT32 };
-		enum { NumRows = 1 };
-		enum { NumColumns = 1 };
-		enum { NumElements = 0 };
-		enum { Alignment = 4 };
-		enum { IsResource = 0 };
-		typedef uint32 TAlignedType;
-		static const FUniformBufferStruct* GetStruct() { return NULL; }
-	};
-	
-
-	template<>
-	class TUniformBufferTypeInfo<int32, false>
-	{
-	public:
-		enum { BaseType = UBMT_INT32 };
-		enum { NumRows = 1 };
-		enum { NumColumns = 1 };
-		enum { NumElements = 0 };
-		enum { Alignment = 4 };
-		enum { IsResource = 0 };
-		typedef int32 TAlignedType;
-		static const FUniformBufferStruct* GetStruct() { return NULL; }
-	};
-
-	template<>
-	class TUniformBufferTypeInfo<float, false>
-	{
-	public:
-		enum { BaseType = UBMT_FLOAT32 };
-		enum { NumRows = 1 };
-		enum { NumColumns = 1 };
-		enum { NumElements = 0 };
-		enum { Alignment = 4 };
-		enum { IsResource = 0 };
-		typedef float TAlignedType;
-		static const FUniformBufferStruct* GetStruct() { return NULL; }
-	};
-
-	template<>
-	class TUniformBufferTypeInfo<FVector2D, false>
-	{
-	public:
-		enum { BaseType = UBMT_FLOAT32 };
-		enum { NumRows = 1 };
-		enum { NumColumns = 2 };
-		enum { NumElements = 0 };
-		enum { Alignment = 8 };
-		enum { IsResource = 0 };
-		typedef TAlignedTypedef<FVector2D,Alignment>::TAlignedType TAlignedType;
-		static const FUniformBufferStruct* GetStruct() { return NULL; }
-	};
-
-	template<>
-	class TUniformBufferTypeInfo<FVector, false>
-	{
-	public:
-		enum { BaseType = UBMT_FLOAT32 };
-		enum { NumRows = 1 };
-		enum { NumColumns = 3 };
-		enum { NumElements = 0 };
-		enum { Alignment = 16 };
-		enum { IsResource = 0 };
-		typedef TAlignedTypedef<FVector,Alignment>::TAlignedType TAlignedType;
-		static const FUniformBufferStruct* GetStruct() { return NULL; }
-	};
-
-	template<>
-	class TUniformBufferTypeInfo<FVector4, false>
-	{
-	public:
-		enum { BaseType = UBMT_FLOAT32 };
-		enum { NumRows = 1 };
-		enum { NumColumns = 4 };
-		enum { NumElements = 0 };
-		enum { Alignment = 16 };
-		enum { IsResource = 0 };
-		typedef TAlignedTypedef<FVector4,Alignment>::TAlignedType TAlignedType;
-		static const FUniformBufferStruct* GetStruct() { return NULL; }
-	};
-	
-	template<>
-	class TUniformBufferTypeInfo<FLinearColor, false>
-	{
-	public:
-		enum { BaseType = UBMT_FLOAT32 };
-		enum { NumRows = 1 };
-		enum { NumColumns = 4 };
-		enum { NumElements = 0 };
-		enum { Alignment = 16 };
-		enum { IsResource = 0 };
-		typedef TAlignedTypedef<FLinearColor,Alignment>::TAlignedType TAlignedType;
-		static const FUniformBufferStruct* GetStruct() { return NULL; }
-	};
-	
-	template<>
-	class TUniformBufferTypeInfo<FIntPoint, false>
-	{
-	public:
-		enum { BaseType = UBMT_INT32 };
-		enum { NumRows = 1 };
-		enum { NumColumns = 2 };
-		enum { NumElements = 0 };
-		enum { Alignment = 8 };
-		enum { IsResource = 0 };
-		typedef TAlignedTypedef<FIntPoint,Alignment>::TAlignedType TAlignedType;
-		static const FUniformBufferStruct* GetStruct() { return NULL; }
-	};
-
-	template<>
-	class TUniformBufferTypeInfo<FIntVector, false>
-	{
-	public:
-		enum { BaseType = UBMT_INT32 };
-		enum { NumRows = 1 };
-		enum { NumColumns = 3 };
-		enum { NumElements = 0 };
-		enum { Alignment = 16 };
-		enum { IsResource = 0 };
-		typedef TAlignedTypedef<FIntVector,Alignment>::TAlignedType TAlignedType;
-		static const FUniformBufferStruct* GetStruct() { return NULL; }
-	};
-
-	template<>
-	class TUniformBufferTypeInfo<FIntRect, false>
-	{
-	public:
-		enum { BaseType = UBMT_INT32 };
-		enum { NumRows = 1 };
-		enum { NumColumns = 4 };
-		enum { NumElements = 0 };
-		enum { Alignment = 16 };
-		enum { IsResource = 0 };
-		typedef TAlignedTypedef<FIntRect,Alignment>::TAlignedType TAlignedType;
-		static const FUniformBufferStruct* GetStruct() { return NULL; }
-	};
-
-	template<>
-	class TUniformBufferTypeInfo<FMatrix, false>
-	{
-	public:
-		enum { BaseType = UBMT_FLOAT32 };
-		enum { NumRows = 4 };
-		enum { NumColumns = 4 };
-		enum { NumElements = 0 };
-		enum { Alignment = 16 };
-		enum { IsResource = 0 };
-		typedef TAlignedTypedef<FMatrix,Alignment>::TAlignedType TAlignedType;
-		static const FUniformBufferStruct* GetStruct() { return NULL; }
-	};
-	
-	template<typename T,size_t InNumElements>
-	class TUniformBufferTypeInfo<T[InNumElements], false>
-	{
-	public:
-		enum { BaseType = TUniformBufferTypeInfo<T, false>::BaseType };
-		enum { NumRows = TUniformBufferTypeInfo<T, false>::NumRows };
-		enum { NumColumns = TUniformBufferTypeInfo<T, false>::NumColumns };
-		enum { NumElements = InNumElements };
-		enum { Alignment = TUniformBufferTypeInfo<T, false>::Alignment };
-		enum { IsResource = TUniformBufferTypeInfo<T, false>::IsResource };
-		typedef TStaticArray<T,InNumElements,16> TAlignedType;
-		static const FUniformBufferStruct* GetStruct() { return TUniformBufferTypeInfo<T, false>::GetStruct(); }
-	};
-	
-	template<typename T,size_t InNumElements,uint32 IgnoredAlignment>
-	class TUniformBufferTypeInfo<TStaticArray<T,InNumElements,IgnoredAlignment>, false >
-	{
-	public:
-		enum { BaseType = TUniformBufferTypeInfo<T, false>::BaseType };
-		enum { NumRows = TUniformBufferTypeInfo<T, false>::NumRows };
-		enum { NumColumns = TUniformBufferTypeInfo<T, false>::NumColumns };
-		enum { NumElements = InNumElements };
-		enum { Alignment = TUniformBufferTypeInfo<T, false>::Alignment };
-		enum { IsResource = TUniformBufferTypeInfo<T, false>::IsResource };
-		typedef TStaticArray<T,InNumElements,16> TAlignedType;
-		static const FUniformBufferStruct* GetStruct() { return TUniformBufferTypeInfo<T, false>::GetStruct(); }
-	};
-
-	template<>
-	class TUniformBufferTypeInfo<TFixedSizePointerWrapper<FShaderResourceViewRHIParamRef>, false>
-	{
-	public:
-		enum { BaseType = UBMT_SRV };
-		enum { NumRows = 1 };
-		enum { NumColumns = 1 };
-		enum { NumElements = 0 };
-		enum { Alignment = sizeof(TFixedSizePointerWrapper<FShaderResourceViewRHIParamRef>) };
-		enum { IsResource = 1 };
-		typedef TAlignedTypedef<TFixedSizePointerWrapper<FShaderResourceViewRHIParamRef>,Alignment>::TAlignedType TAlignedType;
-		static const FUniformBufferStruct* GetStruct() { return NULL; }
-	};
-	// RHICreateUniformBuffer assumes C++ constant layout matches the shader layout when extracting float constants, yet the C++ struct contains pointers.  
-	// Enforce a min size of 64 bits on pointer types in uniform buffer structs to guarantee layout matching between languages.
-	static_assert(sizeof(TFixedSizePointerWrapper<FShaderResourceViewRHIParamRef>) == sizeof(uint64), "Uniform buffer layout must not be platform dependent.");
-
-	template<>
-	class TUniformBufferTypeInfo<TFixedSizePointerWrapper<FUnorderedAccessViewRHIParamRef>, false>
-	{
-	public:
-		enum { BaseType = UBMT_UAV };
-		enum { NumRows = 1 };
-		enum { NumColumns = 1 };
-		enum { NumElements = 0 };
-		enum { Alignment = sizeof(TFixedSizePointerWrapper<FUnorderedAccessViewRHIParamRef>) };
-		enum { IsResource = 1 };
-		typedef TAlignedTypedef<TFixedSizePointerWrapper<FUnorderedAccessViewRHIParamRef>,Alignment>::TAlignedType TAlignedType;
-		static const FUniformBufferStruct* GetStruct() { return NULL; }
-	};
-	static_assert(sizeof(TFixedSizePointerWrapper<FUnorderedAccessViewRHIParamRef>) == sizeof(uint64), "Uniform buffer layout must not be platform dependent.");
-
-	template<>
-	class TUniformBufferTypeInfo<TFixedSizePointerWrapper<FSamplerStateRHIParamRef>, false>
-	{
-	public:
-		enum { BaseType = UBMT_SAMPLER };
-		enum { NumRows = 1 };
-		enum { NumColumns = 1 };
-		enum { NumElements = 0 };
-		enum { Alignment = sizeof(TFixedSizePointerWrapper<FSamplerStateRHIParamRef>) };
-		enum { IsResource = 1 };
-		typedef TAlignedTypedef<TFixedSizePointerWrapper<FSamplerStateRHIParamRef>,Alignment>::TAlignedType TAlignedType;
-		static const FUniformBufferStruct* GetStruct() { return NULL; }
-	};
-	static_assert(sizeof(TFixedSizePointerWrapper<FSamplerStateRHIParamRef>) == sizeof(uint64), "Uniform buffer layout must not be platform dependent.");
-
-	template<>
-	class TUniformBufferTypeInfo<TFixedSizePointerWrapper<FTextureRHIParamRef>, false>
-	{
-	public:
-		enum { BaseType = UBMT_TEXTURE };
-		enum { NumRows = 1 };
-		enum { NumColumns = 1 };
-		enum { NumElements = 0 };
-		enum { Alignment = sizeof(TFixedSizePointerWrapper<FTextureRHIParamRef>) };
-		enum { IsResource = 1 };
-		typedef TAlignedTypedef<TFixedSizePointerWrapper<FTextureRHIParamRef>,Alignment>::TAlignedType TAlignedType;
-		static const FUniformBufferStruct* GetStruct() { return NULL; }
-	};
-	static_assert(sizeof(TFixedSizePointerWrapper<FTextureRHIParamRef>) == sizeof(uint64), "Uniform buffer layout must not be platform dependent.");
-
-	template<class StructType>
-	class TUniformBufferTypeInfo<StructType, true> 
-	{ 
-	public: 
-		enum { BaseType = UBMT_STRUCT }; 
-		enum { NumRows = 1 }; 
-		enum { NumColumns = 1 };
-		enum { NumElements = 0 };
-		enum { Alignment = UNIFORM_BUFFER_STRUCT_ALIGNMENT }; 
-		enum { IsResource = 0 };
-		typedef StructType TAlignedType;
-		static const FUniformBufferStruct* GetStruct() { return &StructType::StaticStruct; } 
-	};
-
-//
-// Macros for declaring uniform buffer structures.
-//
-
-#define IMPLEMENT_UNIFORM_BUFFER_STRUCT(StructTypeName,ShaderVariableName) \
-	FUniformBufferStruct StructTypeName::StaticStruct( \
-	FName(TEXT(#StructTypeName)), \
-	TEXT(#StructTypeName), \
-	ShaderVariableName, \
-	StructTypeName::ConstructUniformBufferParameter, \
-	sizeof(StructTypeName), \
-	StructTypeName::zzGetMembers(), \
-	true);
-
-/** Begins a uniform buffer struct declaration. */
-#define BEGIN_UNIFORM_BUFFER_STRUCT_EX(StructTypeName,PrefixKeywords,ConstructorSuffix) \
-	MS_ALIGN(UNIFORM_BUFFER_STRUCT_ALIGNMENT) class PrefixKeywords StructTypeName \
-	{ \
-	public: \
-		StructTypeName () ConstructorSuffix \
-		static FUniformBufferStruct StaticStruct; \
-		static FShaderUniformBufferParameter* ConstructUniformBufferParameter() { return new TShaderUniformBufferParameter<StructTypeName>(); } \
-		static FUniformBufferRHIRef CreateUniformBuffer(const StructTypeName& InContents, EUniformBufferUsage InUsage) \
-		{ \
-			return RHICreateUniformBuffer(&InContents,StaticStruct.GetLayout(),InUsage); \
-		} \
-	private: \
-		typedef StructTypeName zzTThisStruct; \
-		struct zzFirstMemberId { enum { HasDeclaredResource = 0 }; }; \
-		static TArray<FUniformBufferStruct::FMember> zzGetMembersBefore(zzFirstMemberId) \
-		{ \
-			return TArray<FUniformBufferStruct::FMember>(); \
-		} \
-		typedef zzFirstMemberId
-
-/** Declares a member of a uniform buffer struct. */
-#define DECLARE_UNIFORM_BUFFER_STRUCT_MEMBER_EXPLICIT(MemberType,MemberName,ArrayDecl,Precision,OptionalShaderType,IsMemberStruct) \
-		zzMemberId##MemberName; \
-	public: \
-		typedef MemberType zzA##MemberName ArrayDecl; \
-		typedef TUniformBufferTypeInfo<zzA##MemberName, IsMemberStruct>::TAlignedType zzT##MemberName; \
-		zzT##MemberName MemberName; \
-		static_assert(TUniformBufferTypeInfo<zzA##MemberName, IsMemberStruct>::BaseType != UBMT_INVALID, "Invalid type " #MemberType " of member " #MemberName "."); \
-		static_assert(TUniformBufferTypeInfo<zzA##MemberName, IsMemberStruct>::BaseType != UBMT_UAV, "UAV is not yet supported in resource tables for " #MemberName " of type " #MemberType "."); \
-	private: \
-		struct zzNextMemberId##MemberName { enum { HasDeclaredResource = zzMemberId##MemberName::HasDeclaredResource || TUniformBufferTypeInfo<zzT##MemberName, IsMemberStruct>::IsResource }; }; \
-		static TArray<FUniformBufferStruct::FMember> zzGetMembersBefore(zzNextMemberId##MemberName) \
-		{ \
-			static_assert(TUniformBufferTypeInfo<zzT##MemberName, IsMemberStruct>::IsResource == 1 || zzMemberId##MemberName::HasDeclaredResource == 0, "All resources must be declared last for " #MemberName "."); \
-			static_assert(TUniformBufferTypeInfo<zzT##MemberName, IsMemberStruct>::IsResource == 0 || TIsArrayOrRefOfType<decltype(OptionalShaderType), TCHAR>::Value, "No shader type for " #MemberName "."); \
-			/* Route the member enumeration on to the function for the member following this. */ \
-			TArray<FUniformBufferStruct::FMember> OutMembers = zzGetMembersBefore(zzMemberId##MemberName()); \
-			/* Add this member. */ \
-			OutMembers.Add(FUniformBufferStruct::FMember( \
-				TEXT(#MemberName), \
-				OptionalShaderType, \
-				STRUCT_OFFSET(zzTThisStruct,MemberName), \
-				(EUniformBufferBaseType)TUniformBufferTypeInfo<zzA##MemberName, IsMemberStruct>::BaseType, \
-				Precision, \
-				TUniformBufferTypeInfo<zzA##MemberName, IsMemberStruct>::NumRows, \
-				TUniformBufferTypeInfo<zzA##MemberName, IsMemberStruct>::NumColumns, \
-				TUniformBufferTypeInfo<zzA##MemberName, IsMemberStruct>::NumElements, \
-				TUniformBufferTypeInfo<zzA##MemberName, IsMemberStruct>::GetStruct() \
-				)); \
-			static_assert( \
-				(STRUCT_OFFSET(zzTThisStruct,MemberName) & (TUniformBufferTypeInfo<zzA##MemberName, IsMemberStruct>::Alignment - 1)) == 0, \
-				"Misaligned uniform buffer struct member " #MemberName "."); \
-			return OutMembers; \
-		} \
-		typedef zzNextMemberId##MemberName
-
-#define UNIFORM_MEMBER_ARRAY_EX(MemberType,MemberName,ArrayDecl,Precision) DECLARE_UNIFORM_BUFFER_STRUCT_MEMBER_EXPLICIT(MemberType,MemberName,ArrayDecl,Precision,TEXT(""),false)
-
-#define UNIFORM_MEMBER_ARRAY(MemberType,MemberName,ArrayDecl) DECLARE_UNIFORM_BUFFER_STRUCT_MEMBER_EXPLICIT(MemberType,MemberName,ArrayDecl,EShaderPrecisionModifier::Float,TEXT(""),false)
-
-#define UNIFORM_MEMBER(MemberType,MemberName) DECLARE_UNIFORM_BUFFER_STRUCT_MEMBER_EXPLICIT(MemberType,MemberName,,EShaderPrecisionModifier::Float,TEXT(""),false)
-
-#define UNIFORM_MEMBER_EX(MemberType,MemberName,Precision) DECLARE_UNIFORM_BUFFER_STRUCT_MEMBER_EXPLICIT(MemberType,MemberName,,Precision,TEXT(""),false)
-
-#define UNIFORM_MEMBER_SRV(ShaderType,MemberName) DECLARE_UNIFORM_BUFFER_STRUCT_MEMBER_EXPLICIT(TFixedSizePointerWrapper<FShaderResourceViewRHIParamRef>,MemberName,,EShaderPrecisionModifier::Float,TEXT(#ShaderType),false)
-
-// NOT SUPPORTED YET
-//#define DECLARE_UNIFORM_BUFFER_STRUCT_MEMBER_UAV(ShaderType,MemberName) DECLARE_UNIFORM_BUFFER_STRUCT_MEMBER_EXPLICIT(TFixedSizePointerWrapper<FUnorderedAccessViewRHIParamRef>,MemberName,,EShaderPrecisionModifier::Float,TEXT(#ShaderType))
-
-#define UNIFORM_MEMBER_SAMPLER(ShaderType,MemberName) DECLARE_UNIFORM_BUFFER_STRUCT_MEMBER_EXPLICIT(TFixedSizePointerWrapper<FSamplerStateRHIParamRef>,MemberName,,EShaderPrecisionModifier::Float,TEXT(#ShaderType),false)
-
-#define UNIFORM_MEMBER_TEXTURE(ShaderType,MemberName) DECLARE_UNIFORM_BUFFER_STRUCT_MEMBER_EXPLICIT(TFixedSizePointerWrapper<FTextureRHIParamRef>,MemberName,,EShaderPrecisionModifier::Float,TEXT(#ShaderType),false)
-
-#define UNIFORM_MEMBER_STRUCT(StructType,MemberName) DECLARE_UNIFORM_BUFFER_STRUCT_MEMBER_EXPLICIT(StructType,MemberName,,EShaderPrecisionModifier::Float,TEXT(#StructType),true)
-
-/** Ends a uniform buffer struct declaration. */
-#define END_UNIFORM_BUFFER_STRUCT(Name) \
-			zzLastMemberId; \
-		static TArray<FUniformBufferStruct::FMember> zzGetMembers() { return zzGetMembersBefore(zzLastMemberId()); } \
-	} GCC_ALIGN(UNIFORM_BUFFER_STRUCT_ALIGNMENT); \
-	template<> class TUniformBufferTypeInfo<Name, false> \
-	{ \
-	public: \
-		enum { BaseType = UBMT_STRUCT }; \
-		enum { NumRows = 1 }; \
-		enum { NumColumns = 1 }; \
-		enum { Alignment = UNIFORM_BUFFER_STRUCT_ALIGNMENT }; \
-		static const FUniformBufferStruct* GetStruct() { return &Name::StaticStruct; } \
-	};
-
-#define BEGIN_UNIFORM_BUFFER_STRUCT(StructTypeName,PrefixKeywords) BEGIN_UNIFORM_BUFFER_STRUCT_EX(StructTypeName,PrefixKeywords,{})
-#define BEGIN_UNIFORM_BUFFER_STRUCT_WITH_CONSTRUCTOR(StructTypeName,PrefixKeywords) BEGIN_UNIFORM_BUFFER_STRUCT_EX(StructTypeName,PrefixKeywords,;)
-
-/** Finds the FUniformBufferStruct corresponding to the given name, or NULL if not found. */
-extern RENDERCORE_API FUniformBufferStruct* FindUniformBufferStructByName(const TCHAR* StructName);
-extern RENDERCORE_API FUniformBufferStruct* FindUniformBufferStructByFName(FName StructName);
diff --git a/Engine/Source/Runtime/RenderCore/Public/VertexFactory.h b/Engine/Source/Runtime/RenderCore/Public/VertexFactory.h
new file mode 100644
--- /dev/null
+++ b/Engine/Source/Runtime/RenderCore/Public/VertexFactory.h
@@ -0,0 +1,557 @@
+// Copyright 1998-2018 Epic Games, Inc. All Rights Reserved.
+
+/*=============================================================================
+	VertexFactory.h: Vertex factory definitions.
+=============================================================================*/
+
+#pragma once
+
+#include "CoreMinimal.h"
+#include "Containers/List.h"
+#include "Misc/SecureHash.h"
+#include "RHI.h"
+#include "RenderResource.h"
+#include "ShaderCore.h"
+#include "Shader.h"
+#include "Misc/EnumClassFlags.h"
+
+class FMaterial;
+
+enum class EVertexStreamUsage : uint8
+{
+	Default			= 0 << 0,
+	Instancing		= 1 << 0,
+	Overridden		= 1 << 1,
+	ManualFetch		= 1 << 2
+};
+
+ENUM_CLASS_FLAGS(EVertexStreamUsage);
+/**
+ * A typed data source for a vertex factory which streams data from a vertex buffer.
+ */
+struct FVertexStreamComponent
+{
+	/** The vertex buffer to stream data from.  If null, no data can be read from this stream. */
+	const FVertexBuffer* VertexBuffer = nullptr;
+
+	/** The offset to the start of the vertex buffer fetch. */
+	uint32 StreamOffset = 0;
+
+	/** The offset of the data, relative to the beginning of each element in the vertex buffer. */
+	uint8 Offset = 0;
+
+	/** The stride of the data. */
+	uint8 Stride = 0;
+
+	/** The type of the data read from this stream. */
+	TEnumAsByte<EVertexElementType> Type = VET_None;
+
+	EVertexStreamUsage VertexStreamUsage = EVertexStreamUsage::Default;
+
+	/**
+	 * Initializes the data stream to null.
+	 */
+	FVertexStreamComponent()
+	{}
+
+	/**
+	 * Minimal initialization constructor.
+	 */
+	FVertexStreamComponent(const FVertexBuffer* InVertexBuffer, uint32 InOffset, uint32 InStride, EVertexElementType InType, EVertexStreamUsage Usage = EVertexStreamUsage::Default) :
+		VertexBuffer(InVertexBuffer),
+		StreamOffset(0),
+		Offset(InOffset),
+		Stride(InStride),
+		Type(InType),
+		VertexStreamUsage(Usage)
+	{
+		check(InStride <= 0xFF);
+		check(InOffset <= 0xFF);
+	}
+
+	FVertexStreamComponent(const FVertexBuffer* InVertexBuffer, uint32 InStreamOffset, uint32 InOffset, uint32 InStride, EVertexElementType InType, EVertexStreamUsage Usage = EVertexStreamUsage::Default) :
+		VertexBuffer(InVertexBuffer),
+		StreamOffset(InStreamOffset),
+		Offset(InOffset),
+		Stride(InStride),
+		Type(InType),
+		VertexStreamUsage(Usage)
+	{
+		check(InStride <= 0xFF);
+		check(InOffset <= 0xFF);
+	}
+};
+
+/**
+ * A macro which initializes a FVertexStreamComponent to read a member from a struct.
+ */
+#define STRUCTMEMBER_VERTEXSTREAMCOMPONENT(VertexBuffer,VertexType,Member,MemberType) \
+	FVertexStreamComponent(VertexBuffer,STRUCT_OFFSET(VertexType,Member),sizeof(VertexType),MemberType)
+
+/**
+ * An interface to the parameter bindings for the vertex factory used by a shader.
+ */
+class RENDERCORE_API FVertexFactoryShaderParameters
+{
+public:
+	virtual ~FVertexFactoryShaderParameters() {}
+	virtual void Bind(const class FShaderParameterMap& ParameterMap) = 0;
+	virtual void Serialize(FArchive& Ar) = 0;
+	virtual void SetMesh(FRHICommandList& RHICmdList, FShader* VertexShader,const class FVertexFactory* VertexFactory,const class FSceneView& View,const struct FMeshBatchElement& BatchElement,uint32 DataFlags) const = 0;
+	virtual uint32 GetSize() const { return sizeof(*this); }
+};
+
+/**
+ * An object used to represent the type of a vertex factory.
+ */
+class FVertexFactoryType
+{
+public:
+
+	typedef FVertexFactoryShaderParameters* (*ConstructParametersType)(EShaderFrequency ShaderFrequency);
+	typedef bool (*ShouldCacheType)(EShaderPlatform, const class FMaterial*, const class FShaderType*);
+	typedef void (*ModifyCompilationEnvironmentType)(EShaderPlatform, const class FMaterial*, FShaderCompilerEnvironment&);
+	typedef bool (*SupportsTessellationShadersType)();
+
+	/**
+	 * @return The global shader factory list.
+	 */
+	static RENDERCORE_API TLinkedList<FVertexFactoryType*>*& GetTypeList();
+
+	/**
+	 * Finds a FVertexFactoryType by name.
+	 */
+	static RENDERCORE_API FVertexFactoryType* GetVFByName(const FString& VFName);
+
+	/** Initialize FVertexFactoryType static members, this must be called before any VF types are created. */
+	static void Initialize(const TMap<FString, TArray<const TCHAR*> >& ShaderFileToUniformBufferVariables);
+
+	/** Uninitializes FVertexFactoryType cached data. */
+	static void Uninitialize();
+
+	/**
+	 * Minimal initialization constructor.
+	 * @param bInUsedWithMaterials - True if the vertex factory will be rendered in combination with a material.
+	 * @param bInSupportsStaticLighting - True if the vertex factory will be rendered with static lighting.
+	 */
+	RENDERCORE_API FVertexFactoryType(
+		const TCHAR* InName,
+		const TCHAR* InShaderFilename,
+		bool bInUsedWithMaterials,
+		bool bInSupportsStaticLighting,
+		bool bInSupportsDynamicLighting,
+		bool bInSupportsPrecisePrevWorldPos,
+		bool bInSupportsPositionOnly,
+		ConstructParametersType InConstructParameters,
+		ShouldCacheType InShouldCache,
+		ModifyCompilationEnvironmentType InModifyCompilationEnvironment,
+		SupportsTessellationShadersType InSupportsTessellationShaders
+		);
+
+	RENDERCORE_API virtual ~FVertexFactoryType();
+
+	// Accessors.
+	const TCHAR* GetName() const { return Name; }
+	FName GetFName() const { return TypeName; }
+	const TCHAR* GetShaderFilename() const { return ShaderFilename; }
+	FVertexFactoryShaderParameters* CreateShaderParameters(EShaderFrequency ShaderFrequency) const { return (*ConstructParameters)(ShaderFrequency); }
+	bool IsUsedWithMaterials() const { return bUsedWithMaterials; }
+	bool SupportsStaticLighting() const { return bSupportsStaticLighting; }
+	bool SupportsDynamicLighting() const { return bSupportsDynamicLighting; }
+	bool SupportsPrecisePrevWorldPos() const { return bSupportsPrecisePrevWorldPos; }
+	bool SupportsPositionOnly() const { return bSupportsPositionOnly; }
+	/** Returns an int32 specific to this vertex factory type. */
+	inline int32 GetId() const { return HashIndex; }
+	static int32 GetNumVertexFactoryTypes() { return NextHashIndex; }
+
+	const FSerializationHistory* GetSerializationHistory(EShaderFrequency Frequency) const
+	{
+		return &SerializationHistory[Frequency];
+	}
+
+	// Hash function.
+	friend uint32 GetTypeHash(FVertexFactoryType* Type)
+	{ 
+		return Type ? Type->HashIndex : 0;
+	}
+
+	/** Calculates a Hash based on this vertex factory type's source code and includes */
+	RENDERCORE_API const FSHAHash& GetSourceHash(EShaderPlatform ShaderPlatform) const;
+
+	/**
+	 * Should we cache the material's shadertype on this platform with this vertex factory? 
+	 */
+	bool ShouldCache(EShaderPlatform Platform, const class FMaterial* Material, const class FShaderType* ShaderType)
+	{
+		return (*ShouldCacheRef)(Platform, Material, ShaderType); 
+	}
+
+	/**
+	* Calls the function ptr for the shader type on the given environment
+	* @param Environment - shader compile environment to modify
+	*/
+	void ModifyCompilationEnvironment( EShaderPlatform Platform, const FMaterial* Material, FShaderCompilerEnvironment& OutEnvironment )
+	{
+		// Set up the mapping from VertexFactory.usf to the vertex factory type's source code.
+		FString VertexFactoryIncludeString = FString::Printf( TEXT("#include \"%s\""), GetShaderFilename() );
+		OutEnvironment.IncludeVirtualPathToContentsMap.Add(TEXT("/Engine/Generated/VertexFactory.ush"), VertexFactoryIncludeString);
+
+		OutEnvironment.SetDefine(TEXT("HAS_PRIMITIVE_UNIFORM_BUFFER"), 1);
+
+		(*ModifyCompilationEnvironmentRef)(Platform, Material, OutEnvironment);
+	}
+
+	/**
+	 * Does this vertex factory support tessellation shaders?
+	 */
+	bool SupportsTessellationShaders() const
+	{
+		return (*SupportsTessellationShadersRef)(); 
+	}
+
+	/** Adds include statements for uniform buffers that this shader type references, and builds a prefix for the shader file with the include statements. */
+	RENDERCORE_API void AddReferencedUniformBufferIncludes(FShaderCompilerEnvironment& OutEnvironment, FString& OutSourceFilePrefix, EShaderPlatform Platform);
+
+	void FlushShaderFileCache(const TMap<FString, TArray<const TCHAR*> >& ShaderFileToUniformBufferVariables)
+	{
+		ReferencedUniformBufferStructsCache.Empty();
+		GenerateReferencedUniformBuffers(ShaderFilename, Name, ShaderFileToUniformBufferVariables, ReferencedUniformBufferStructsCache);
+		bCachedUniformBufferStructDeclarations = false;
+	}
+
+	const TMap<const TCHAR*, FCachedUniformBufferDeclaration>& GetReferencedUniformBufferStructsCache() const
+	{
+		return ReferencedUniformBufferStructsCache;
+	}
+
+private:
+	static RENDERCORE_API uint32 NextHashIndex;
+
+	/** Tracks whether serialization history for all shader types has been initialized. */
+	static bool bInitializedSerializationHistory;
+
+	uint32 HashIndex;
+	const TCHAR* Name;
+	const TCHAR* ShaderFilename;
+	FName TypeName;
+	uint32 bUsedWithMaterials : 1;
+	uint32 bSupportsStaticLighting : 1;
+	uint32 bSupportsDynamicLighting : 1;
+	uint32 bSupportsPrecisePrevWorldPos : 1;
+	uint32 bSupportsPositionOnly : 1;
+	ConstructParametersType ConstructParameters;
+	ShouldCacheType ShouldCacheRef;
+	ModifyCompilationEnvironmentType ModifyCompilationEnvironmentRef;
+	SupportsTessellationShadersType SupportsTessellationShadersRef;
+
+	TLinkedList<FVertexFactoryType*> GlobalListLink;
+
+	/** 
+	 * Cache of referenced uniform buffer includes.  
+	 * These are derived from source files so they need to be flushed when editing and recompiling shaders on the fly. 
+	 * FVertexFactoryType::Initialize will add an entry for each referenced uniform buffer, but the declarations are added on demand as shaders are compiled.
+	 */
+	TMap<const TCHAR*, FCachedUniformBufferDeclaration> ReferencedUniformBufferStructsCache;
+
+	/** Tracks what platforms ReferencedUniformBufferStructsCache has had declarations cached for. */
+	bool bCachedUniformBufferStructDeclarations;
+
+	/** 
+	 * Stores a history of serialization sizes for this vertex factory's shader parameter class. 
+	 * This is used to invalidate shaders when serialization changes.
+	 */
+	FSerializationHistory SerializationHistory[SF_NumFrequencies];
+};
+
+/**
+ * Serializes a reference to a vertex factory type.
+ */
+extern RENDERCORE_API FArchive& operator<<(FArchive& Ar,FVertexFactoryType*& TypeRef);
+
+/**
+ * Find the vertex factory type with the given name.
+ * @return NULL if no vertex factory type matched, otherwise a vertex factory type with a matching name.
+ */
+extern RENDERCORE_API FVertexFactoryType* FindVertexFactoryType(FName TypeName);
+
+/**
+ * A macro for declaring a new vertex factory type, for use in the vertex factory class's definition body.
+ */
+#define DECLARE_VERTEX_FACTORY_TYPE(FactoryClass) \
+	public: \
+	static FVertexFactoryType StaticType; \
+	virtual FVertexFactoryType* GetType() const override;
+
+/**
+ * A macro for implementing the static vertex factory type object, and specifying parameters used by the type.
+ * @param bUsedWithMaterials - True if the vertex factory will be rendered in combination with a material.
+ * @param bSupportsStaticLighting - True if the vertex factory will be rendered with static lighting.
+ */
+#define IMPLEMENT_VERTEX_FACTORY_TYPE(FactoryClass,ShaderFilename,bUsedWithMaterials,bSupportsStaticLighting,bSupportsDynamicLighting,bPrecisePrevWorldPos,bSupportsPositionOnly) \
+	FVertexFactoryShaderParameters* Construct##FactoryClass##ShaderParameters(EShaderFrequency ShaderFrequency) { return FactoryClass::ConstructShaderParameters(ShaderFrequency); } \
+	FVertexFactoryType FactoryClass::StaticType( \
+		TEXT(#FactoryClass), \
+		TEXT(ShaderFilename), \
+		bUsedWithMaterials, \
+		bSupportsStaticLighting, \
+		bSupportsDynamicLighting, \
+		bPrecisePrevWorldPos, \
+		bSupportsPositionOnly, \
+		Construct##FactoryClass##ShaderParameters, \
+		FactoryClass::ShouldCompilePermutation, \
+		FactoryClass::ModifyCompilationEnvironment, \
+		FactoryClass::SupportsTessellationShaders \
+		); \
+		FVertexFactoryType* FactoryClass::GetType() const { return &StaticType; }
+
+/** Encapsulates a dependency on a vertex factory type and saved state from that vertex factory type. */
+class FVertexFactoryTypeDependency
+{
+public:
+
+	FVertexFactoryTypeDependency() :
+		VertexFactoryType(NULL)
+	{}
+
+	FVertexFactoryType* VertexFactoryType;
+
+	/** Used to detect changes to the vertex factory source files. */
+	FSHAHash VFSourceHash;
+
+	friend FArchive& operator<<(FArchive& Ar,class FVertexFactoryTypeDependency& Ref)
+	{
+		Ar << Ref.VertexFactoryType << FShaderResource::FilterShaderSourceHashForSerialization(Ar, Ref.VFSourceHash);
+		return Ar;
+	}
+
+	bool operator==(const FVertexFactoryTypeDependency& Reference) const
+	{
+		return VertexFactoryType == Reference.VertexFactoryType && VFSourceHash == Reference.VFSourceHash;
+	}
+};
+
+/** Used to compare two Vertex Factory types by name. */
+class FCompareVertexFactoryTypes											
+{																				
+public:		
+	FORCEINLINE bool operator()(const FVertexFactoryType& A, const FVertexFactoryType& B ) const
+	{
+		int32 AL = FCString::Strlen(A.GetName());
+		int32 BL = FCString::Strlen(B.GetName());
+		if ( AL == BL )
+		{
+			return FCString::Strncmp(A.GetName(), B.GetName(), AL) > 0;
+		}
+		return AL > BL;
+	}
+};
+
+/**
+ * Encapsulates a vertex data source which can be linked into a vertex shader.
+ */
+class RENDERCORE_API FVertexFactory : public FRenderResource
+{
+public:
+	FVertexFactory(ERHIFeatureLevel::Type InFeatureLevel) 
+		: FRenderResource(InFeatureLevel)
+	{
+	}
+
+	virtual FVertexFactoryType* GetType() const { return NULL; }
+
+	/**
+	 * Activates the vertex factory.
+	 */
+	void SetStreams(ERHIFeatureLevel::Type InFeatureLevel, FRHICommandList& RHICmdList) const;
+
+	/**
+	 * Call SetStreamSource on instance streams to offset the read pointer
+	 */
+	void OffsetInstanceStreams(FRHICommandList& RHICmdList, uint32 FirstVertex) const;
+
+	/**
+	* Sets the position stream as the current stream source.
+	*/
+	void SetPositionStream(FRHICommandList& RHICmdList) const;
+
+	/**
+	* Call SetStreamSource on instance streams to offset the read pointer
+	 */
+	void OffsetPositionInstanceStreams(FRHICommandList& RHICmdList, uint32 FirstVertex) const;
+
+	/**
+	* Can be overridden by FVertexFactory subclasses to modify their compile environment just before compilation occurs.
+	*/
+	static void ModifyCompilationEnvironment( EShaderPlatform Platform, const FMaterial* Material, FShaderCompilerEnvironment& OutEnvironment ) {}
+
+	/**
+	* Can be overridden by FVertexFactory subclasses to enable HS/DS in D3D11
+	*/
+	static bool SupportsTessellationShaders() { return false; }
+
+	// FRenderResource interface.
+	virtual void ReleaseRHI();
+
+	// Accessors.
+	FVertexDeclarationRHIRef& GetDeclaration() { return Declaration; }
+	void SetDeclaration(FVertexDeclarationRHIRef& NewDeclaration) { Declaration = NewDeclaration; }
+
+	const FVertexDeclarationRHIRef& GetDeclaration() const { return Declaration; }
+	const FVertexDeclarationRHIRef& GetPositionDeclaration() const { return PositionDeclaration; }
+
+	virtual bool IsGPUSkinned() const { return false; }
+
+	/** Indicates whether the vertex factory supports a position-only stream. */
+	virtual bool SupportsPositionOnlyStream() const { return !!PositionStream.Num(); }
+
+	/** Indicates whether the vertex factory supports a null pixel shader. */
+	virtual bool SupportsNullPixelShader() const { return true; }
+
+	virtual bool RendersPrimitivesAsCameraFacingSprites() const { return false; }
+
+	/**
+	* Fill in array of strides from this factory's vertex streams without shadow/light maps
+	* @param OutStreamStrides - output array of # MaxVertexElementCount stream strides to fill
+	*/
+	int32 GetStreamStrides(uint32 *OutStreamStrides, bool bPadWithZeroes=true) const;
+	/**
+	* Fill in array of strides from this factory's position only vertex streams
+	* @param OutStreamStrides - output array of # MaxVertexElementCount stream strides to fill
+	*/
+	void GetPositionStreamStride(uint32 *OutStreamStrides) const;
+
+	/**
+	 * Get a bitmask representing the visibility of each FMeshBatch element.
+	 * FMeshBatch.bRequiresPerElementVisibility must be set for this to be called.
+	 */
+	virtual uint64 GetStaticBatchElementVisibility(const class FSceneView& View, const struct FMeshBatch* Batch, const void* InViewCustomData = nullptr) const { return 1; }
+
+	bool NeedsDeclaration() const { return bNeedsDeclaration; }
+	bool SupportsManualVertexFetch(ERHIFeatureLevel::Type InFeatureLevel) const 
+	{ 
+		check(InFeatureLevel != ERHIFeatureLevel::Num);
+		return bSupportsManualVertexFetch && (InFeatureLevel > ERHIFeatureLevel::ES3_1) && RHISupportsManualVertexFetch(GMaxRHIShaderPlatform);
+	}
+
+protected:
+
+	/**
+	 * Creates a vertex element for a vertex stream components.  Adds a unique stream index for the vertex buffer used by the component.
+	 * @param Component - The vertex stream component.
+	 * @param AttributeIndex - The attribute index to which the stream component is bound.
+	 * @return The vertex element which corresponds to Component.
+	 */
+	FVertexElement AccessStreamComponent(const FVertexStreamComponent& Component,uint8 AttributeIndex);
+
+	/**
+	 * Creates a vertex element for a position vertex stream component.  Adds a unique position stream index for the vertex buffer used by the component.
+	 * @param Component - The position vertex stream component.
+	 * @param Usage - The vertex element usage semantics.
+	 * @param AttributeIndex - The attribute index to which the stream component is bound.
+	 * @return The vertex element which corresponds to Component.
+	 */
+	FVertexElement AccessPositionStreamComponent(const FVertexStreamComponent& Component,uint8 AttributeIndex);
+
+	/**
+	 * Initializes the vertex declaration.
+	 * @param Elements - The elements of the vertex declaration.
+	 */
+	void InitDeclaration(FVertexDeclarationElementList& Elements);
+
+	/**
+	 * Initializes the position-only vertex declaration.
+	 * @param Elements - The elements of the vertex declaration.
+	 */
+	void InitPositionDeclaration(const FVertexDeclarationElementList& Elements);
+
+	/**
+	 * Information needed to set a vertex stream.
+	 */
+	struct FVertexStream
+	{
+		const FVertexBuffer* VertexBuffer = nullptr;
+		uint32 Stride = 0;
+		uint32 Offset = 0;
+		EVertexStreamUsage VertexStreamUsage = EVertexStreamUsage::Default;
+
+		friend bool operator==(const FVertexStream& A,const FVertexStream& B)
+		{
+			return A.VertexBuffer == B.VertexBuffer && A.Stride == B.Stride && A.Offset == B.Offset && A.VertexStreamUsage == B.VertexStreamUsage;
+		}
+
+		FVertexStream()
+		{
+		}
+	};
+
+	/** The vertex streams used to render the factory. */
+	TArray<FVertexStream,TFixedAllocator<MaxVertexElementCount> > Streams;
+
+	/* VF can explicitly set this to false to avoid errors without decls; this is for VFs that fetch from buffers directly (e.g. Niagara) */
+	bool bNeedsDeclaration = true;
+	
+	bool bSupportsManualVertexFetch = false;
+
+private:
+
+	/** The position only vertex stream used to render the factory during depth only passes. */
+	TArray<FVertexStream,TFixedAllocator<MaxVertexElementCount> > PositionStream;
+
+	/** The RHI vertex declaration used to render the factory normally. */
+	FVertexDeclarationRHIRef Declaration;
+
+	/** The RHI vertex declaration used to render the factory during depth only passes. */
+	FVertexDeclarationRHIRef PositionDeclaration;
+};
+
+/**
+ * An encapsulation of the vertex factory parameters for a shader.
+ * Handles serialization and binding of the vertex factory parameters for the shader's vertex factory type.
+ */
+class RENDERCORE_API FVertexFactoryParameterRef
+{
+public:
+	FVertexFactoryParameterRef(FVertexFactoryType* InVertexFactoryType,const FShaderParameterMap& ParameterMap, EShaderFrequency InShaderFrequency, EShaderPlatform InShaderPlatform);
+
+	FVertexFactoryParameterRef():
+		Parameters(NULL),
+		VertexFactoryType(NULL)
+	{}
+
+	~FVertexFactoryParameterRef()
+	{
+		delete Parameters;
+	}
+
+	void SetMesh(FRHICommandList& RHICmdList, FShader* Shader,const FVertexFactory* VertexFactory,const class FSceneView& View,const struct FMeshBatchElement& BatchElement,uint32 DataFlags) const
+	{
+		if(Parameters)
+		{
+			Parameters->SetMesh(RHICmdList, Shader,VertexFactory,View,BatchElement,DataFlags);
+		}
+	}
+
+	const FVertexFactoryType* GetVertexFactoryType() const { return VertexFactoryType; }
+
+	/** Returns the hash of the vertex factory shader file that this shader was compiled with. */
+	const FSHAHash& GetHash() const;
+
+	/** Returns the shader platform that this shader was compiled with. */
+	EShaderPlatform GetShaderPlatform() const;
+
+	friend RENDERCORE_API bool operator<<(FArchive& Ar,FVertexFactoryParameterRef& Ref);
+
+	uint32 GetAllocatedSize() const
+	{
+		return Parameters ? Parameters->GetSize() : 0;
+	}
+
+private:
+	FVertexFactoryShaderParameters* Parameters;
+	FVertexFactoryType* VertexFactoryType;
+	EShaderFrequency ShaderFrequency;
+	EShaderPlatform ShaderPlatform;
+
+	// Hash of the vertex factory's source file at shader compile time, used by the automatic versioning system to detect changes
+	FSHAHash VFHash;
+};
diff --git a/Engine/Source/Runtime/RenderCore/RenderCore.Build.cs b/Engine/Source/Runtime/RenderCore/RenderCore.Build.cs
--- a/Engine/Source/Runtime/RenderCore/RenderCore.Build.cs
+++ b/Engine/Source/Runtime/RenderCore/RenderCore.Build.cs
@@ -5,12 +5,22 @@ using UnrealBuildTool;
 public class RenderCore : ModuleRules
 {
 	public RenderCore(ReadOnlyTargetRules Target) : base(Target)
-	{
+    {
+        PublicDependencyModuleNames.AddRange(new string[] { "RHI" });
+		
         if (Target.bBuildEditor == true)
         {
             PrivateDependencyModuleNames.Add("TargetPlatform");
         }
+		else
+        {
+
+            PrivateIncludePathModuleNames.AddRange(new string[] { "TargetPlatform" });
+
+        }
+
+        PrivateDependencyModuleNames.AddRange(new string[] { "Core", "Projects", "RHI", "ApplicationCore" });
 
-        PrivateDependencyModuleNames.AddRange(new string[] { "Core", "RHI", "ApplicationCore" });
-	}
+        PrivateIncludePathModuleNames.AddRange(new string[] { "DerivedDataCache" });
+    }
 }
diff --git a/Engine/Source/Runtime/Renderer/Private/AtmosphereTextures.cpp b/Engine/Source/Runtime/Renderer/Private/AtmosphereTextures.cpp
--- a/Engine/Source/Runtime/Renderer/Private/AtmosphereTextures.cpp
+++ b/Engine/Source/Runtime/Renderer/Private/AtmosphereTextures.cpp
@@ -6,7 +6,7 @@
 
 #include "AtmosphereTextures.h"
 #include "Atmosphere/AtmosphericFogComponent.h"
-#include "PostProcess/RenderTargetPool.h"
+#include "RenderTargetPool.h"
 #include "AtmosphereTextureParameters.h"
 #include "ShaderParameterUtils.h"
 
diff --git a/Engine/Source/Runtime/Renderer/Private/BasePassRendering.cpp b/Engine/Source/Runtime/Renderer/Private/BasePassRendering.cpp
--- a/Engine/Source/Runtime/Renderer/Private/BasePassRendering.cpp
+++ b/Engine/Source/Runtime/Renderer/Private/BasePassRendering.cpp
@@ -90,9 +90,9 @@ static TAutoConsoleVariable<int32> CVarSupportAllShaderPermutations(
 /** Whether to replace lightmap textures with solid colors to visualize the mip-levels. */
 bool GVisualizeMipLevels = false;
 
-IMPLEMENT_UNIFORM_BUFFER_STRUCT(FSharedBasePassUniformParameters, TEXT("BasePass"));
-IMPLEMENT_UNIFORM_BUFFER_STRUCT(FOpaqueBasePassUniformParameters, TEXT("OpaqueBasePass"));
-IMPLEMENT_UNIFORM_BUFFER_STRUCT(FTranslucentBasePassUniformParameters, TEXT("TranslucentBasePass"));
+IMPLEMENT_GLOBAL_SHADER_PARAMETER_STRUCT(FSharedBasePassUniformParameters, "BasePass");
+IMPLEMENT_GLOBAL_SHADER_PARAMETER_STRUCT(FOpaqueBasePassUniformParameters, "OpaqueBasePass");
+IMPLEMENT_GLOBAL_SHADER_PARAMETER_STRUCT(FTranslucentBasePassUniformParameters, "TranslucentBasePass");
 
 // Typedef is necessary because the C preprocessor thinks the comma in the template parameter list is a comma in the macro parameter list.
 // BasePass Vertex Shader needs to include hull and domain shaders for tessellation, these only compile for D3D11
diff --git a/Engine/Source/Runtime/Renderer/Private/BasePassRendering.h b/Engine/Source/Runtime/Renderer/Private/BasePassRendering.h
--- a/Engine/Source/Runtime/Renderer/Private/BasePassRendering.h
+++ b/Engine/Source/Runtime/Renderer/Private/BasePassRendering.h
@@ -50,54 +50,54 @@ public:
 	FVector4 LightTangentAndSoftSourceRadius;
 };
 
-BEGIN_UNIFORM_BUFFER_STRUCT(FSharedBasePassUniformParameters,)
-	UNIFORM_MEMBER_STRUCT(FForwardLightData, Forward)
-	UNIFORM_MEMBER_STRUCT(FForwardLightData, ForwardISR)
-	UNIFORM_MEMBER_STRUCT(FReflectionUniformParameters, Reflection)
-	UNIFORM_MEMBER_STRUCT(FFogUniformParameters, Fog)
-	UNIFORM_MEMBER_TEXTURE(Texture2D, SSProfilesTexture)
-END_UNIFORM_BUFFER_STRUCT(FSharedBasePassUniformParameters)
-
-BEGIN_UNIFORM_BUFFER_STRUCT(FOpaqueBasePassUniformParameters,)
-	UNIFORM_MEMBER_STRUCT(FSharedBasePassUniformParameters, Shared)
+BEGIN_GLOBAL_SHADER_PARAMETER_STRUCT(FSharedBasePassUniformParameters,)
+	SHADER_PARAMETER_STRUCT(FForwardLightData, Forward)
+	SHADER_PARAMETER_STRUCT(FForwardLightData, ForwardISR)
+	SHADER_PARAMETER_STRUCT(FReflectionUniformParameters, Reflection)
+	SHADER_PARAMETER_STRUCT(FFogUniformParameters, Fog)
+	SHADER_PARAMETER_TEXTURE(Texture2D, SSProfilesTexture)
+END_GLOBAL_SHADER_PARAMETER_STRUCT()
+
+BEGIN_GLOBAL_SHADER_PARAMETER_STRUCT(FOpaqueBasePassUniformParameters,)
+	SHADER_PARAMETER_STRUCT(FSharedBasePassUniformParameters, Shared)
 	// Forward shading 
-	UNIFORM_MEMBER(uint32, UseForwardScreenSpaceShadowMask)
-	UNIFORM_MEMBER_TEXTURE(Texture2D, ForwardScreenSpaceShadowMaskTexture)
-	UNIFORM_MEMBER_TEXTURE(Texture2D, IndirectOcclusionTexture)
-	UNIFORM_MEMBER_TEXTURE(Texture2D, ResolvedSceneDepthTexture)
+	SHADER_PARAMETER(int32, UseForwardScreenSpaceShadowMask)
+	SHADER_PARAMETER_TEXTURE(Texture2D, ForwardScreenSpaceShadowMaskTexture)
+	SHADER_PARAMETER_TEXTURE(Texture2D, IndirectOcclusionTexture)
+	SHADER_PARAMETER_TEXTURE(Texture2D, ResolvedSceneDepthTexture)
 	// DBuffer decals
-	UNIFORM_MEMBER_TEXTURE(Texture2D, DBufferATexture)
-	UNIFORM_MEMBER_SAMPLER(SamplerState, DBufferATextureSampler)
-	UNIFORM_MEMBER_TEXTURE(Texture2D, DBufferBTexture)
-	UNIFORM_MEMBER_SAMPLER(SamplerState, DBufferBTextureSampler)
-	UNIFORM_MEMBER_TEXTURE(Texture2D, DBufferCTexture)
-	UNIFORM_MEMBER_SAMPLER(SamplerState, DBufferCTextureSampler)
-	UNIFORM_MEMBER_TEXTURE(Texture2D<uint>, DBufferRenderMask)
+	SHADER_PARAMETER_TEXTURE(Texture2D, DBufferATexture)
+	SHADER_PARAMETER_SAMPLER(SamplerState, DBufferATextureSampler)
+	SHADER_PARAMETER_TEXTURE(Texture2D, DBufferBTexture)
+	SHADER_PARAMETER_SAMPLER(SamplerState, DBufferBTextureSampler)
+	SHADER_PARAMETER_TEXTURE(Texture2D, DBufferCTexture)
+	SHADER_PARAMETER_SAMPLER(SamplerState, DBufferCTextureSampler)
+	SHADER_PARAMETER_TEXTURE(Texture2D<uint>, DBufferRenderMask)
 	// Misc
-	UNIFORM_MEMBER_TEXTURE(Texture2D, EyeAdaptation)
-END_UNIFORM_BUFFER_STRUCT(FOpaqueBasePassUniformParameters)
+	SHADER_PARAMETER_TEXTURE(Texture2D, EyeAdaptation)
+END_GLOBAL_SHADER_PARAMETER_STRUCT()
 
-BEGIN_UNIFORM_BUFFER_STRUCT(FTranslucentBasePassUniformParameters,)
-	UNIFORM_MEMBER_STRUCT(FSharedBasePassUniformParameters, Shared)
-	UNIFORM_MEMBER_STRUCT(FSceneTexturesUniformParameters, SceneTextures)
+BEGIN_GLOBAL_SHADER_PARAMETER_STRUCT(FTranslucentBasePassUniformParameters,)
+	SHADER_PARAMETER_STRUCT(FSharedBasePassUniformParameters, Shared)
+	SHADER_PARAMETER_STRUCT(FSceneTexturesUniformParameters, SceneTextures)
 	// Material SSR
-	UNIFORM_MEMBER(FVector4, HZBUvFactorAndInvFactor)
-	UNIFORM_MEMBER(FVector4, PrevScreenPositionScaleBias)
-	UNIFORM_MEMBER(float, PrevSceneColorPreExposureInv)
-	UNIFORM_MEMBER_TEXTURE(Texture2D, HZBTexture)
-	UNIFORM_MEMBER_SAMPLER(SamplerState, HZBSampler)
-	UNIFORM_MEMBER_TEXTURE(Texture2D, PrevSceneColor)
-	UNIFORM_MEMBER_SAMPLER(SamplerState, PrevSceneColorSampler)
+	SHADER_PARAMETER(FVector4, HZBUvFactorAndInvFactor)
+	SHADER_PARAMETER(FVector4, PrevScreenPositionScaleBias)
+	SHADER_PARAMETER(float, PrevSceneColorPreExposureInv)
+	SHADER_PARAMETER_TEXTURE(Texture2D, HZBTexture)
+	SHADER_PARAMETER_SAMPLER(SamplerState, HZBSampler)
+	SHADER_PARAMETER_TEXTURE(Texture2D, PrevSceneColor)
+	SHADER_PARAMETER_SAMPLER(SamplerState, PrevSceneColorSampler)
 	// Translucency Lighting Volume
-	UNIFORM_MEMBER_TEXTURE(Texture3D, TranslucencyLightingVolumeAmbientInner)
-	UNIFORM_MEMBER_SAMPLER(SamplerState, TranslucencyLightingVolumeAmbientInnerSampler)
-	UNIFORM_MEMBER_TEXTURE(Texture3D, TranslucencyLightingVolumeAmbientOuter)
-	UNIFORM_MEMBER_SAMPLER(SamplerState, TranslucencyLightingVolumeAmbientOuterSampler)
-	UNIFORM_MEMBER_TEXTURE(Texture3D, TranslucencyLightingVolumeDirectionalInner)
-	UNIFORM_MEMBER_SAMPLER(SamplerState, TranslucencyLightingVolumeDirectionalInnerSampler)
-	UNIFORM_MEMBER_TEXTURE(Texture3D, TranslucencyLightingVolumeDirectionalOuter)
-	UNIFORM_MEMBER_SAMPLER(SamplerState, TranslucencyLightingVolumeDirectionalOuterSampler)
-END_UNIFORM_BUFFER_STRUCT(FTranslucentBasePassUniformParameters)
+	SHADER_PARAMETER_TEXTURE(Texture3D, TranslucencyLightingVolumeAmbientInner)
+	SHADER_PARAMETER_SAMPLER(SamplerState, TranslucencyLightingVolumeAmbientInnerSampler)
+	SHADER_PARAMETER_TEXTURE(Texture3D, TranslucencyLightingVolumeAmbientOuter)
+	SHADER_PARAMETER_SAMPLER(SamplerState, TranslucencyLightingVolumeAmbientOuterSampler)
+	SHADER_PARAMETER_TEXTURE(Texture3D, TranslucencyLightingVolumeDirectionalInner)
+	SHADER_PARAMETER_SAMPLER(SamplerState, TranslucencyLightingVolumeDirectionalInnerSampler)
+	SHADER_PARAMETER_TEXTURE(Texture3D, TranslucencyLightingVolumeDirectionalOuter)
+	SHADER_PARAMETER_SAMPLER(SamplerState, TranslucencyLightingVolumeDirectionalOuterSampler)
+END_GLOBAL_SHADER_PARAMETER_STRUCT()
 
 extern FTextureRHIRef& GetEyeAdaptation(const FViewInfo& View);
 
@@ -138,9 +138,9 @@ public:
 
 inline void BindBasePassUniformBuffer(const FShaderParameterMap& ParameterMap, FShaderUniformBufferParameter& BasePassUniformBuffer)
 {
-	TArray<const FUniformBufferStruct*> NestedStructs;
-	FOpaqueBasePassUniformParameters::StaticStruct.GetNestedStructs(NestedStructs);
-	FTranslucentBasePassUniformParameters::StaticStruct.GetNestedStructs(NestedStructs);
+	TArray<const FShaderParametersMetadata*> NestedStructs;
+	FOpaqueBasePassUniformParameters::StaticStructMetadata.GetNestedStructs(NestedStructs);
+	FTranslucentBasePassUniformParameters::StaticStructMetadata.GetNestedStructs(NestedStructs);
 
 	for (int32 StructIndex = 0; StructIndex < NestedStructs.Num(); StructIndex++)
 	{
@@ -148,16 +148,16 @@ inline void BindBasePassUniformBuffer(const FShaderParameterMap& ParameterMap, F
 		checkfSlow(!ParameterMap.ContainsParameterAllocation(StructVariableName), TEXT("%s found bound in the base pass.  Base Pass uniform buffer nested structs should not be bound separately"), StructVariableName);
 	}
 	
-	const bool bNeedsOpaqueBasePass = ParameterMap.ContainsParameterAllocation(FOpaqueBasePassUniformParameters::StaticStruct.GetShaderVariableName());
-	const bool bNeedsTransparentBasePass = ParameterMap.ContainsParameterAllocation(FTranslucentBasePassUniformParameters::StaticStruct.GetShaderVariableName());
+	const bool bNeedsOpaqueBasePass = ParameterMap.ContainsParameterAllocation(FOpaqueBasePassUniformParameters::StaticStructMetadata.GetShaderVariableName());
+	const bool bNeedsTransparentBasePass = ParameterMap.ContainsParameterAllocation(FTranslucentBasePassUniformParameters::StaticStructMetadata.GetShaderVariableName());
 
 	checkSlow(!(bNeedsOpaqueBasePass && bNeedsTransparentBasePass));
 
-	BasePassUniformBuffer.Bind(ParameterMap, FOpaqueBasePassUniformParameters::StaticStruct.GetShaderVariableName());
+	BasePassUniformBuffer.Bind(ParameterMap, FOpaqueBasePassUniformParameters::StaticStructMetadata.GetShaderVariableName());
 
 	if (!BasePassUniformBuffer.IsBound())
 	{
-		BasePassUniformBuffer.Bind(ParameterMap, FTranslucentBasePassUniformParameters::StaticStruct.GetShaderVariableName());
+		BasePassUniformBuffer.Bind(ParameterMap, FTranslucentBasePassUniformParameters::StaticStructMetadata.GetShaderVariableName());
 	}
 }
 
@@ -458,7 +458,7 @@ public:
 
 	static bool ValidateCompiledResult(EShaderPlatform Platform, const TArray<FMaterial*>& Materials, const FVertexFactoryType* VertexFactoryType, const FShaderParameterMap& ParameterMap, TArray<FString>& OutError)
 	{
-		if (ParameterMap.ContainsParameterAllocation(FSceneTexturesUniformParameters::StaticStruct.GetShaderVariableName()))
+		if (ParameterMap.ContainsParameterAllocation(FSceneTexturesUniformParameters::StaticStructMetadata.GetShaderVariableName()))
 		{
 			OutError.Add(TEXT("Base pass shaders cannot read from the SceneTexturesStruct."));
 			return false;
@@ -477,8 +477,8 @@ public:
 		ReflectionCaptureBuffer.Bind(Initializer.ParameterMap, TEXT("ReflectionCapture"));
 
 		// These parameters should only be used nested in the base pass uniform buffer
-		check(!Initializer.ParameterMap.ContainsParameterAllocation(FFogUniformParameters::StaticStruct.GetShaderVariableName()));
-		check(!Initializer.ParameterMap.ContainsParameterAllocation(FReflectionUniformParameters::StaticStruct.GetShaderVariableName()));
+		check(!Initializer.ParameterMap.ContainsParameterAllocation(FFogUniformParameters::StaticStructMetadata.GetShaderVariableName()));
+		check(!Initializer.ParameterMap.ContainsParameterAllocation(FReflectionUniformParameters::StaticStructMetadata.GetShaderVariableName()));
 	}
 	TBasePassPixelShaderPolicyParamType() {}
 
diff --git a/Engine/Source/Runtime/Renderer/Private/CompositionLighting/CompositionLighting.cpp b/Engine/Source/Runtime/Renderer/Private/CompositionLighting/CompositionLighting.cpp
--- a/Engine/Source/Runtime/Renderer/Private/CompositionLighting/CompositionLighting.cpp
+++ b/Engine/Source/Runtime/Renderer/Private/CompositionLighting/CompositionLighting.cpp
@@ -274,14 +274,14 @@ void FCompositionLighting::ProcessAfterBasePass(FRHICommandListImmediate& RHICmd
 	SceneContext.GetSceneColor()->SetDebugName(TEXT("SceneColor"));
 	// to be able to observe results with VisualizeTexture
 
-	GRenderTargetPool.VisualizeTexture.SetCheckPoint(RHICmdList, SceneContext.GetSceneColor());
-	GRenderTargetPool.VisualizeTexture.SetCheckPoint(RHICmdList, SceneContext.GBufferA);
-	GRenderTargetPool.VisualizeTexture.SetCheckPoint(RHICmdList, SceneContext.GBufferB);
-	GRenderTargetPool.VisualizeTexture.SetCheckPoint(RHICmdList, SceneContext.GBufferC);
-	GRenderTargetPool.VisualizeTexture.SetCheckPoint(RHICmdList, SceneContext.GBufferD);
-	GRenderTargetPool.VisualizeTexture.SetCheckPoint(RHICmdList, SceneContext.GBufferE);
-	GRenderTargetPool.VisualizeTexture.SetCheckPoint(RHICmdList, SceneContext.GBufferVelocity);
-	GRenderTargetPool.VisualizeTexture.SetCheckPoint(RHICmdList, SceneContext.ScreenSpaceAO);
+	GVisualizeTexture.SetCheckPoint(RHICmdList, SceneContext.GetSceneColor());
+	GVisualizeTexture.SetCheckPoint(RHICmdList, SceneContext.GBufferA);
+	GVisualizeTexture.SetCheckPoint(RHICmdList, SceneContext.GBufferB);
+	GVisualizeTexture.SetCheckPoint(RHICmdList, SceneContext.GBufferC);
+	GVisualizeTexture.SetCheckPoint(RHICmdList, SceneContext.GBufferD);
+	GVisualizeTexture.SetCheckPoint(RHICmdList, SceneContext.GBufferE);
+	GVisualizeTexture.SetCheckPoint(RHICmdList, SceneContext.GBufferVelocity);
+	GVisualizeTexture.SetCheckPoint(RHICmdList, SceneContext.ScreenSpaceAO);
 	
 	// so that the passes can register themselves to the graph
 	{
diff --git a/Engine/Source/Runtime/Renderer/Private/CompositionLighting/PostProcessDeferredDecals.cpp b/Engine/Source/Runtime/Renderer/Private/CompositionLighting/PostProcessDeferredDecals.cpp
--- a/Engine/Source/Runtime/Renderer/Private/CompositionLighting/PostProcessDeferredDecals.cpp
+++ b/Engine/Source/Runtime/Renderer/Private/CompositionLighting/PostProcessDeferredDecals.cpp
@@ -972,7 +972,7 @@ void FRCPassPostProcessDeferredDecals::Process(FRenderingCompositePassContext& C
 
 					if (SceneContext.DBufferMask)
 					{
-						GRenderTargetPool.VisualizeTexture.SetCheckPoint(RHICmdList, SceneContext.DBufferMask);
+						GVisualizeTexture.SetCheckPoint(RHICmdList, SceneContext.DBufferMask);
 						bHasValidDBufferMask = true;
 					}
 				}
@@ -993,9 +993,9 @@ void FRCPassPostProcessDeferredDecals::Process(FRenderingCompositePassContext& C
 		if (CurrentStage == DRS_BeforeBasePass && bNeedsDBufferTargets)
 		{
 			// before BasePass
-			GRenderTargetPool.VisualizeTexture.SetCheckPoint(RHICmdList, SceneContext.DBufferA);
-			GRenderTargetPool.VisualizeTexture.SetCheckPoint(RHICmdList, SceneContext.DBufferB);
-			GRenderTargetPool.VisualizeTexture.SetCheckPoint(RHICmdList, SceneContext.DBufferC);
+			GVisualizeTexture.SetCheckPoint(RHICmdList, SceneContext.DBufferA);
+			GVisualizeTexture.SetCheckPoint(RHICmdList, SceneContext.DBufferB);
+			GVisualizeTexture.SetCheckPoint(RHICmdList, SceneContext.DBufferC);
 		}
 	}
 
diff --git a/Engine/Source/Runtime/Renderer/Private/CompositionLighting/PostProcessLpvIndirect.cpp b/Engine/Source/Runtime/Renderer/Private/CompositionLighting/PostProcessLpvIndirect.cpp
--- a/Engine/Source/Runtime/Renderer/Private/CompositionLighting/PostProcessLpvIndirect.cpp
+++ b/Engine/Source/Runtime/Renderer/Private/CompositionLighting/PostProcessLpvIndirect.cpp
@@ -24,7 +24,7 @@
 #include "LightPropagationVolumeSettings.h"
 #include "PipelineStateCache.h"
 
-IMPLEMENT_UNIFORM_BUFFER_STRUCT(FLpvReadUniformBufferParameters,TEXT("LpvRead"));
+IMPLEMENT_GLOBAL_SHADER_PARAMETER_STRUCT(FLpvReadUniformBufferParameters,"LpvRead");
 typedef TUniformBufferRef<FLpvReadUniformBufferParameters> FLpvReadUniformBufferRef;
 
 
@@ -349,7 +349,7 @@ void FRCPassPostProcessLpvIndirect::Process(FRenderingCompositePassContext& Cont
 
 	if ( LPVSettings.LPVDirectionalOcclusionIntensity > 0.0001f )
 	{
-		GRenderTargetPool.VisualizeTexture.SetCheckPoint(Context.RHICmdList, SceneContext.DirectionalOcclusion);
+		GVisualizeTexture.SetCheckPoint(Context.RHICmdList, SceneContext.DirectionalOcclusion);
 	}
 }
 
diff --git a/Engine/Source/Runtime/Renderer/Private/CompositionLighting/PostProcessMeshDecals.cpp b/Engine/Source/Runtime/Renderer/Private/CompositionLighting/PostProcessMeshDecals.cpp
--- a/Engine/Source/Runtime/Renderer/Private/CompositionLighting/PostProcessMeshDecals.cpp
+++ b/Engine/Source/Runtime/Renderer/Private/CompositionLighting/PostProcessMeshDecals.cpp
@@ -49,7 +49,7 @@ protected:
 	FMeshDecalVS(const ShaderMetaType::CompiledShaderInitializerType& Initializer)
 		:	FMeshMaterialShader(Initializer)
 	{
-		PassUniformBuffer.Bind(Initializer.ParameterMap, FSceneTexturesUniformParameters::StaticStruct.GetShaderVariableName());
+		PassUniformBuffer.Bind(Initializer.ParameterMap, FSceneTexturesUniformParameters::StaticStructMetadata.GetShaderVariableName());
 	}
 
 	FMeshDecalVS()
@@ -87,7 +87,7 @@ protected:
 	FMeshDecalHS(const ShaderMetaType::CompiledShaderInitializerType& Initializer)
 		: FBaseHS(Initializer)
 	{
-		PassUniformBuffer.Bind(Initializer.ParameterMap, FSceneTexturesUniformParameters::StaticStruct.GetShaderVariableName());
+		PassUniformBuffer.Bind(Initializer.ParameterMap, FSceneTexturesUniformParameters::StaticStructMetadata.GetShaderVariableName());
 	}
 
 	FMeshDecalHS() {}
@@ -147,7 +147,7 @@ public:
 	FMeshDecalsPS(const ShaderMetaType::CompiledShaderInitializerType& Initializer)
 		:	FMeshMaterialShader(Initializer)
 	{
-		PassUniformBuffer.Bind(Initializer.ParameterMap, FSceneTexturesUniformParameters::StaticStruct.GetShaderVariableName());
+		PassUniformBuffer.Bind(Initializer.ParameterMap, FSceneTexturesUniformParameters::StaticStructMetadata.GetShaderVariableName());
 	}
 
 	FMeshDecalsPS()
diff --git a/Engine/Source/Runtime/Renderer/Private/ConvertToUniformMesh.cpp b/Engine/Source/Runtime/Renderer/Private/ConvertToUniformMesh.cpp
--- a/Engine/Source/Runtime/Renderer/Private/ConvertToUniformMesh.cpp
+++ b/Engine/Source/Runtime/Renderer/Private/ConvertToUniformMesh.cpp
@@ -25,7 +25,7 @@ protected:
 	FConvertToUniformMeshVS(const ShaderMetaType::CompiledShaderInitializerType& Initializer)
 		:	FMeshMaterialShader(Initializer)
 	{
-		PassUniformBuffer.Bind(Initializer.ParameterMap, FSceneTexturesUniformParameters::StaticStruct.GetShaderVariableName());
+		PassUniformBuffer.Bind(Initializer.ParameterMap, FSceneTexturesUniformParameters::StaticStructMetadata.GetShaderVariableName());
 	}
 
 	FConvertToUniformMeshVS()
@@ -110,7 +110,7 @@ protected:
 	FConvertToUniformMeshGS(const ShaderMetaType::CompiledShaderInitializerType& Initializer)
 		:	FMeshMaterialShader(Initializer)
 	{
-		PassUniformBuffer.Bind(Initializer.ParameterMap, FSceneTexturesUniformParameters::StaticStruct.GetShaderVariableName());
+		PassUniformBuffer.Bind(Initializer.ParameterMap, FSceneTexturesUniformParameters::StaticStructMetadata.GetShaderVariableName());
 	}
 
 	FConvertToUniformMeshGS()
@@ -382,9 +382,14 @@ int32 FUniformMeshConverter::Convert(
 				for (int32 BatchElementIndex = 0; BatchElementIndex < Mesh.Elements.Num(); BatchElementIndex++)
 				{
 					//@todo - fix
-					OutPrimitiveUniformBuffer = IsValidRef(Mesh.Elements[BatchElementIndex].PrimitiveUniformBuffer) 
-						? Mesh.Elements[BatchElementIndex].PrimitiveUniformBuffer
-						: *Mesh.Elements[BatchElementIndex].PrimitiveUniformBufferResource;
+					if (Mesh.Elements[BatchElementIndex].PrimitiveUniformBuffer.IsValid())
+					{
+						OutPrimitiveUniformBuffer = Mesh.Elements[BatchElementIndex].PrimitiveUniformBuffer;
+					}
+					else
+					{
+						OutPrimitiveUniformBuffer = Mesh.Elements[BatchElementIndex].PrimitiveUniformBufferResource->GetUniformBufferRHI();
+					}
 
 					DrawingPolicy.SetMeshRenderState(RHICmdList, View, PrimitiveSceneProxy, Mesh, BatchElementIndex, DrawRenderState, FConvertToUniformMeshDrawingPolicy::ElementDataType(), FConvertToUniformMeshDrawingPolicy::ContextDataType());
 					DrawingPolicy.DrawMesh(RHICmdList, View, Mesh, BatchElementIndex);
diff --git a/Engine/Source/Runtime/Renderer/Private/DebugViewModeRendering.h b/Engine/Source/Runtime/Renderer/Private/DebugViewModeRendering.h
--- a/Engine/Source/Runtime/Renderer/Private/DebugViewModeRendering.h
+++ b/Engine/Source/Runtime/Renderer/Private/DebugViewModeRendering.h
@@ -34,7 +34,7 @@ protected:
 	{
 		IsInstancedStereoParameter.Bind(Initializer.ParameterMap, TEXT("bIsInstancedStereo"));
 		InstancedEyeIndexParameter.Bind(Initializer.ParameterMap, TEXT("InstancedEyeIndex"));
-		PassUniformBuffer.Bind(Initializer.ParameterMap, FSceneTexturesUniformParameters::StaticStruct.GetShaderVariableName());
+		PassUniformBuffer.Bind(Initializer.ParameterMap, FSceneTexturesUniformParameters::StaticStructMetadata.GetShaderVariableName());
 	}
 
 	FDebugViewModeVS() {}
diff --git a/Engine/Source/Runtime/Renderer/Private/DeferredShadingRenderer.cpp b/Engine/Source/Runtime/Renderer/Private/DeferredShadingRenderer.cpp
--- a/Engine/Source/Runtime/Renderer/Private/DeferredShadingRenderer.cpp
+++ b/Engine/Source/Runtime/Renderer/Private/DeferredShadingRenderer.cpp
@@ -1136,7 +1136,7 @@ void FDeferredShadingSceneRenderer::Render(FRHICommandListImmediate& RHICmdList)
 		Scene->FXSystem->PostRenderOpaque(
 			RHICmdList,
 			Views[0].ViewUniformBuffer,
-			&FSceneTexturesUniformParameters::StaticStruct,
+			&FSceneTexturesUniformParameters::StaticStructMetadata,
 			SceneTextureUniformBuffer.GetReference()
 			);
 		ServiceLocalQueue();
diff --git a/Engine/Source/Runtime/Renderer/Private/DeferredShadingRenderer.h b/Engine/Source/Runtime/Renderer/Private/DeferredShadingRenderer.h
--- a/Engine/Source/Runtime/Renderer/Private/DeferredShadingRenderer.h
+++ b/Engine/Source/Runtime/Renderer/Private/DeferredShadingRenderer.h
@@ -395,7 +395,7 @@ private:
 	void PostInitViewCustomData(FGraphEventArray& OutUpdateEvents);
 
 	void RenderLocalLightsForVolumetricFog(
-		FRHICommandListImmediate& RHICmdList,
+		FRenderGraphBuilder& GraphBuilder,
 		FViewInfo& View,
 		bool bUseTemporalReprojection,
 		const struct FVolumetricFogIntegrationParameterData& IntegrationData,
@@ -403,19 +403,19 @@ private:
 		FIntVector VolumetricFogGridSize,
 		FVector GridZParams,
 		const FPooledRenderTargetDesc& VolumeDesc,
-		TRefCountPtr<IPooledRenderTarget>& OutLocalShadowedLightScattering);
+		const FGraphTexture*& OutLocalShadowedLightScattering);
 
 	void RenderLightFunctionForVolumetricFog(
-		FRHICommandListImmediate& RHICmdList,
+		FRenderGraphBuilder& GraphBuilder,
 		FViewInfo& View,
 		FIntVector VolumetricFogGridSize,
 		float VolumetricFogMaxDistance,
 		FMatrix& OutLightFunctionWorldToShadow,
-		TRefCountPtr<IPooledRenderTarget>& OutLightFunctionTexture,
+		const FGraphTexture*& OutLightFunctionTexture,
 		bool& bOutUseDirectionalLightShadowing);
 
 	void VoxelizeFogVolumePrimitives(
-		FRHICommandListImmediate& RHICmdList,
+		FRenderGraphBuilder& GraphBuilder,
 		const FViewInfo& View,
 		const FVolumetricFogIntegrationParameterData& IntegrationData,
 		FIntVector VolumetricFogGridSize,
diff --git a/Engine/Source/Runtime/Renderer/Private/DistanceFieldAmbientOcclusion.cpp b/Engine/Source/Runtime/Renderer/Private/DistanceFieldAmbientOcclusion.cpp
--- a/Engine/Source/Runtime/Renderer/Private/DistanceFieldAmbientOcclusion.cpp
+++ b/Engine/Source/Runtime/Renderer/Private/DistanceFieldAmbientOcclusion.cpp
@@ -133,7 +133,7 @@ int32 GDistanceFieldAOTileSizeY = 16;
 
 DEFINE_LOG_CATEGORY(LogDistanceField);
 
-IMPLEMENT_UNIFORM_BUFFER_STRUCT(FAOSampleData2,TEXT("AOSamples2"));
+IMPLEMENT_GLOBAL_SHADER_PARAMETER_STRUCT(FAOSampleData2,"AOSamples2");
 
 FDistanceFieldAOParameters::FDistanceFieldAOParameters(float InOcclusionMaxDistance, float InContrast)
 {
@@ -850,7 +850,7 @@ bool FDeferredShadingSceneRenderer::RenderDistanceFieldLighting(
 				BuildTileObjectLists(RHICmdList, Scene, Views, DistanceFieldNormal->GetRenderTargetItem(), Parameters);
 			}
 
-			GRenderTargetPool.VisualizeTexture.SetCheckPoint(RHICmdList, DistanceFieldNormal);
+			GVisualizeTexture.SetCheckPoint(RHICmdList, DistanceFieldNormal);
 
 			TRefCountPtr<IPooledRenderTarget> BentNormalOutput;
 
@@ -872,7 +872,7 @@ bool FDeferredShadingSceneRenderer::RenderDistanceFieldLighting(
 
 			RenderCapsuleShadowsForMovableSkylight(RHICmdList, BentNormalOutput);
 
-			GRenderTargetPool.VisualizeTexture.SetCheckPoint(RHICmdList, BentNormalOutput);
+			GVisualizeTexture.SetCheckPoint(RHICmdList, BentNormalOutput);
 
 			if (bVisualizeAmbientOcclusion)
 			{
diff --git a/Engine/Source/Runtime/Renderer/Private/DistanceFieldAmbientOcclusion.h b/Engine/Source/Runtime/Renderer/Private/DistanceFieldAmbientOcclusion.h
--- a/Engine/Source/Runtime/Renderer/Private/DistanceFieldAmbientOcclusion.h
+++ b/Engine/Source/Runtime/Renderer/Private/DistanceFieldAmbientOcclusion.h
@@ -280,9 +280,9 @@ public:
 
 extern void GetSpacedVectors(uint32 FrameNumber, TArray<FVector, TInlineAllocator<9> >& OutVectors);
 
-BEGIN_UNIFORM_BUFFER_STRUCT(FAOSampleData2,)
-	UNIFORM_MEMBER_ARRAY(FVector4,SampleDirections,[NumConeSampleDirections])
-END_UNIFORM_BUFFER_STRUCT(FAOSampleData2)
+BEGIN_GLOBAL_SHADER_PARAMETER_STRUCT(FAOSampleData2,)
+	SHADER_PARAMETER_ARRAY(FVector4,SampleDirections,[NumConeSampleDirections])
+END_GLOBAL_SHADER_PARAMETER_STRUCT()
 
 class FAOParameters
 {
diff --git a/Engine/Source/Runtime/Renderer/Private/DistanceFieldScreenGridLighting.cpp b/Engine/Source/Runtime/Renderer/Private/DistanceFieldScreenGridLighting.cpp
--- a/Engine/Source/Runtime/Renderer/Private/DistanceFieldScreenGridLighting.cpp
+++ b/Engine/Source/Runtime/Renderer/Private/DistanceFieldScreenGridLighting.cpp
@@ -695,7 +695,7 @@ void FDeferredShadingSceneRenderer::RenderDistanceFieldAOScreenGrid(
 		ScreenGridResources->DiscardTransientResource();
 	}
 
-	GRenderTargetPool.VisualizeTexture.SetCheckPoint(RHICmdList, DownsampledBentNormal);
+	GVisualizeTexture.SetCheckPoint(RHICmdList, DownsampledBentNormal);
 
 	PostProcessBentNormalAOScreenGrid(
 		RHICmdList, 
