Speculative fix for linux build error "declaration of anonymous class must be a definition"

diff --git a/Engine/Source/Runtime/RenderCore/Public/ShaderParameterMacros.h b/Engine/Source/Runtime/RenderCore/Public/ShaderParameterMacros.h
--- a/Engine/Source/Runtime/RenderCore/Public/ShaderParameterMacros.h
+++ b/Engine/Source/Runtime/RenderCore/Public/ShaderParameterMacros.h
@@ -589,7 +589,7 @@ struct TShaderParameterTypeInfo<TUniformBufferRef<UniformBufferStructType>>
 
 /** Begins a uniform buffer struct declaration. */
 #define INTERNAL_SHADER_PARAMETER_STRUCT_BEGIN(StructTypeName,PrefixKeywords,ConstructorSuffix,GetStructMetadataScope) \
-	class alignas(SHADER_PARAMETER_STRUCT_ALIGNMENT) PrefixKeywords StructTypeName \
+	MS_ALIGN(SHADER_PARAMETER_STRUCT_ALIGNMENT) class PrefixKeywords StructTypeName \
 	{ \
 	public: \
 		StructTypeName () ConstructorSuffix \
@@ -663,7 +663,7 @@ extern RENDERCORE_API FShaderParametersMetadata* FindUniformBufferStructByFName(
 #define END_SHADER_PARAMETER_STRUCT() \
 		zzLastMemberId; \
 		static TArray<FShaderParametersMetadata::FMember> zzGetMembers() { return zzGetMembersBefore(zzLastMemberId()); } \
-	};
+	} GCC_ALIGN(SHADER_PARAMETER_STRUCT_ALIGNMENT);
 
 /** Begins & ends a shader global parameter structure.
  *
