Fixes missing include in ShaderParameterStruct.h that turn out to be important for linux build

diff --git a/Engine/Source/Runtime/RenderCore/Public/ShaderParameterStruct.h b/Engine/Source/Runtime/RenderCore/Public/ShaderParameterStruct.h
--- a/Engine/Source/Runtime/RenderCore/Public/ShaderParameterStruct.h
+++ b/Engine/Source/Runtime/RenderCore/Public/ShaderParameterStruct.h
@@ -9,6 +9,7 @@
 #include "CoreMinimal.h"
 #include "Shader.h"
 #include "RHI.h"
+#include "RenderGraphResources.h"
 
 
 /** Instrument a shader class to use strutured shader parameters API.
