Fixes some compile errors caused by render graph.

diff --git a/Engine/Source/Runtime/RenderCore/Private/RenderGraphBuilder.cpp b/Engine/Source/Runtime/RenderCore/Private/RenderGraphBuilder.cpp
--- a/Engine/Source/Runtime/RenderCore/Private/RenderGraphBuilder.cpp
+++ b/Engine/Source/Runtime/RenderCore/Private/RenderGraphBuilder.cpp
@@ -23,7 +23,11 @@ static const int32 GRenderGraphImmediateMode = 0;
 
 void FRenderGraphBuilder::Execute()
 {
-	checkf(!bHasExecuted, TEXT("Render graph execution should only happen once."));
+	#if RENDER_GRAPH_DEBUGGING
+	{
+		checkf(!bHasExecuted, TEXT("Render graph execution should only happen once."));
+	}
+	#endif
 
 	if (!GRenderGraphImmediateMode)
 	{
@@ -40,8 +44,10 @@ void FRenderGraphBuilder::Execute()
 
 	DestructPasses();
 
-	#if DO_CHECK
+	#if RENDER_GRAPH_DEBUGGING
+	{
 		bHasExecuted = true;
+	}
 	#endif
 }
 
@@ -332,7 +338,7 @@ void FRenderGraphBuilder::ExecutePass( const FRenderGraphPass* Pass )
 		RHICmdList.EndRenderPass();
 	}
 
-	if (DO_CHECK)
+	if (RENDER_GRAPH_DEBUGGING)
 	{
 		WarnForUselessPassDependencies(Pass);
 	}
@@ -623,7 +629,7 @@ void FRenderGraphBuilder::ProcessDeferredInternalResourceQueries()
 
 void FRenderGraphBuilder::DestructPasses()
 {
-	#if DO_CHECK
+	#if RENDER_GRAPH_DEBUGGING
 	{
 		// Make sure all resource references have been released to ensure no leaks happen.
 		for (const FGraphResource* Resource : Resources)
diff --git a/Engine/Source/Runtime/RenderCore/Public/RenderGraphBuilder.h b/Engine/Source/Runtime/RenderCore/Public/RenderGraphBuilder.h
--- a/Engine/Source/Runtime/RenderCore/Public/RenderGraphBuilder.h
+++ b/Engine/Source/Runtime/RenderCore/Public/RenderGraphBuilder.h
@@ -118,7 +118,11 @@ public:
 	/** Register a external texture to be tracked by the render graph. */
 	inline const FGraphTexture* RegisterExternalTexture(const TRefCountPtr<IPooledRenderTarget>& ExternalPooledTexture, const TCHAR* Name = TEXT("External"))
 	{
-		ensureMsgf(ExternalPooledTexture.IsValid(), TEXT("Attempted to register NULL external texture: %s"), Name);
+		#if RENDER_GRAPH_DEBUGGING
+		{
+			ensureMsgf(ExternalPooledTexture.IsValid(), TEXT("Attempted to register NULL external texture: %s"), Name);
+		}
+		#endif
 		FGraphTexture* OutTexture = new(FMemStack::Get()) FGraphTexture(Name, ExternalPooledTexture->GetDesc());
 		OutTexture->PooledRenderTarget = ExternalPooledTexture;
 		AllocatedTextures.Add(OutTexture, ExternalPooledTexture);
@@ -131,7 +135,11 @@ public:
 	/** Create graph tracked resource from a descriptor with a debug name. */
 	inline const FGraphTexture* CreateTexture(const FPooledRenderTargetDesc& Desc, const TCHAR* DebugName)
 	{
-		ensureMsgf(!bHasExecuted, TEXT("Render graph texture %s needs to be created before the builder execution."), DebugName);
+		#if RENDER_GRAPH_DEBUGGING
+		{
+			ensureMsgf(!bHasExecuted, TEXT("Render graph texture %s needs to be created before the builder execution."), DebugName);
+		}
+		#endif
 		FGraphTexture* Texture = new(FMemStack::Get()) FGraphTexture(DebugName, Desc);
 		#if DO_CHECK
 			Resources.Add(Texture);
@@ -143,8 +151,12 @@ public:
 	inline const FGraphSRV* CreateSRV(const FGraphSRVDesc& Desc)
 	{
 		check(Desc.Texture);
-		ensureMsgf(!bHasExecuted, TEXT("Render graph SRV %s needs to be created before the builder execution."), Desc.Texture->Name);
-		ensureMsgf(Desc.Texture->Desc.TargetableFlags & TexCreate_ShaderResource, TEXT("Attempted to create SRV from texture %s which was not created with TexCreate_ShaderResource"), Desc.Texture->Name);
+		#if RENDER_GRAPH_DEBUGGING
+		{
+			ensureMsgf(!bHasExecuted, TEXT("Render graph SRV %s needs to be created before the builder execution."), Desc.Texture->Name);
+			ensureMsgf(Desc.Texture->Desc.TargetableFlags & TexCreate_ShaderResource, TEXT("Attempted to create SRV from texture %s which was not created with TexCreate_ShaderResource"), Desc.Texture->Name);
+		}
+		#endif
 		
 		FGraphSRV* SRV = new(FMemStack::Get()) FGraphSRV(Desc.Texture->Name, Desc);
 		#if DO_CHECK
@@ -157,8 +169,12 @@ public:
 	inline const FGraphUAV* CreateUAV(const FGraphUAVDesc& Desc)
 	{
 		check(Desc.Texture);
-		ensureMsgf(!bHasExecuted, TEXT("Render graph UAV %s needs to be created before the builder execution."), Desc.Texture->Name);
-		ensureMsgf(Desc.Texture->Desc.TargetableFlags & TexCreate_UAV, TEXT("Attempted to create UAV from texture %s which was not created with TexCreate_UAV"), Desc.Texture->Name);
+		#if RENDER_GRAPH_DEBUGGING
+		{
+			ensureMsgf(!bHasExecuted, TEXT("Render graph UAV %s needs to be created before the builder execution."), Desc.Texture->Name);
+			ensureMsgf(Desc.Texture->Desc.TargetableFlags & TexCreate_UAV, TEXT("Attempted to create UAV from texture %s which was not created with TexCreate_UAV"), Desc.Texture->Name);
+		}
+		#endif
 		
 		FGraphUAV* UAV = new(FMemStack::Get()) FGraphUAV(Desc.Texture->Name, Desc);
 		#if DO_CHECK
@@ -172,7 +188,11 @@ public:
 	inline void CreateParameters(ParameterStructType** OutParameterPtr) const
 	{
 		// Check because destructor called by the pass's destructor.
-		checkf(!bHasExecuted, TEXT("Render graph allocated shader parameters %s needs to be allocated before the builder execution."), ParameterStructType::FTypeInfo::GetStructMetadata()->GetStructTypeName());
+		#if RENDER_GRAPH_DEBUGGING
+		{
+			checkf(!bHasExecuted, TEXT("Render graph allocated shader parameters %s needs to be allocated before the builder execution."), ParameterStructType::FTypeInfo::GetStructMetadata()->GetStructTypeName());
+		}
+		#endif
 		*OutParameterPtr = new(FMemStack::Get()) ParameterStructType;
 		FMemory::Memzero( *OutParameterPtr, sizeof( ParameterStructType ) );
 	}
@@ -187,7 +207,11 @@ public:
 		ERenderGraphPassFlags Flags,
 		ExecuteLambdaType&& ExecuteLambda)
 	{
-		checkf(!bHasExecuted, TEXT("Render graph pass %s needs to be added before the builder execution."), Name);
+		#if RENDER_GRAPH_DEBUGGING
+		{
+			checkf(!bHasExecuted, TEXT("Render graph pass %s needs to be added before the builder execution."), Name);
+		}
+		#endif
 		auto NewPass = new(FMemStack::Get()) TLambdaRenderPass<ParameterStructType, ExecuteLambdaType>(
 			Name,
 			{ ParameterStruct, &ParameterStructType::FTypeInfo::GetStructMetadata()->GetLayout() },
@@ -207,7 +231,11 @@ public:
 	{
 		check(Texture);
 		check(OutTexturePtr);
-		checkf(!bHasExecuted, TEXT("Accessing render graph internal texture %s needs to happen before the builder's execution."), Texture->Name);
+		#if RENDER_GRAPH_DEBUGGING
+		{
+			checkf(!bHasExecuted, TEXT("Accessing render graph internal texture %s needs to happen before the builder's execution."), Texture->Name);
+		}
+		#endif
 		FDeferredInternalTextureQuery Query;
 		Query.Texture = Texture;
 		Query.OutTexturePtr = OutTexturePtr;
@@ -240,7 +268,7 @@ private:
 	TArray<FDeferredInternalTextureQuery, SceneRenderingAllocator> DeferredInternalTextureQueries;
 
 
-	#if DO_CHECK
+	#if RENDER_GRAPH_DEBUGGING
 		/** Whether the Execute() has already been called. */
 		bool bHasExecuted = false;
 
diff --git a/Engine/Source/Runtime/RenderCore/Public/ShaderParameterMacros.h b/Engine/Source/Runtime/RenderCore/Public/ShaderParameterMacros.h
--- a/Engine/Source/Runtime/RenderCore/Public/ShaderParameterMacros.h
+++ b/Engine/Source/Runtime/RenderCore/Public/ShaderParameterMacros.h
@@ -43,10 +43,6 @@ IMPLEMENT_ALIGNED_TYPE(16);
 
 #if PLATFORM_64BITS
 
-/** Fixed 8bytes sized pointer for shader parameters. */
-template<typename PtrType>
-using TNonAlignedShaderParameterPtr = PtrType;
-
 /** Fixed 8bytes sized and aligned pointer for shader parameters. */
 template<typename PtrType>
 using TAlignedShaderParameterPtr = typename TAlignedTypedef<PtrType, SHADER_PARAMETER_POINTER_ALIGNMENT>::Type;
@@ -57,14 +53,34 @@ static_assert(sizeof(void*) == 8, "Wrong PLATFORM_64BITS settings.");
 
 /** Fixed 8bytes sized pointer for shader parameters. */
 template<typename PtrType>
-class TNonAlignedShaderParameterPtr
+class TAlignedShaderParameterPtr
 {
 public:
+	TAlignedShaderParameterPtr()
+	{ }
+
+	TAlignedShaderParameterPtr(const PtrType& Other)
+		: Ref(Other)
+	{ }
+
+	TAlignedShaderParameterPtr(nullptr_t Other)
+		: Ref(Other)
+	{ }
+
+	TAlignedShaderParameterPtr(const TAlignedShaderParameterPtr<PtrType>& Other)
+		: Ref(Other.Ref)
+	{ }
+
 	FORCEINLINE void operator=(const PtrType& Other)
 	{
 		Ref = Other;
 	}
 
+	FORCEINLINE void operator=(nullptr_t Other)
+	{
+		Ref = Other;
+	}
+
 	FORCEINLINE operator PtrType&()
 	{
 		return Ref;
@@ -75,25 +91,21 @@ public:
 		return Ref;
 	}
 
-	FORCEINLINE PtrType& operator->() const
+	FORCEINLINE const PtrType& operator->() const
 	{
 		return Ref;
 	}
 
 private:
 	PtrType Ref;
-	uint32 _Padding;
+	#if !PLATFORM_64BITS
+		uint32 _Padding;
+		static_assert(sizeof(void*) == 4, "Wrong PLATFORM_64BITS settings.");
+	#endif
 
 	static_assert(sizeof(PtrType) == sizeof(void*), "T should be a pointer.");
 };
 
-/** Fixed 8bytes sized and aligned pointer for shader parameters. */
-template<typename PtrType>
-class alignas(SHADER_PARAMETER_POINTER_ALIGNMENT) TAlignedShaderParameterPtr : public TNonAlignedShaderParameterPtr
-{ };
-
-static_assert(sizeof(void*) == 4, "Wrong PLATFORM_64BITS settings.");
-
 #endif // !PLATFORM_64BITS
 
 
@@ -132,9 +144,11 @@ private:
 
 
 /** Render graph information about how to bind a render target. */
-struct FRenderTargetBinding
+struct alignas(SHADER_PARAMETER_STRUCT_ALIGNMENT) FRenderTargetBinding
 {
-	FRenderTargetBinding() = default;
+	FRenderTargetBinding()
+		: Texture(nullptr)
+	{ }
 	FRenderTargetBinding(const FRenderTargetBinding&) = default;
 
 	/** Creates a render target binding informations.
@@ -165,7 +179,7 @@ private:
 	/** All parameters required to bind a render target deferred. This are purposefully private to
 	 * force the user to call FRenderTargetBinding() constructor, forcing him to specify the load and store action.
 	 */
-	TNonAlignedShaderParameterPtr<const FGraphTexture*> Texture = nullptr;
+	TAlignedShaderParameterPtr<const FGraphTexture*> Texture;
 	ERenderTargetLoadAction		LoadAction		= ERenderTargetLoadAction::ENoAction;
 	ERenderTargetStoreAction	StoreAction		= ERenderTargetStoreAction::ENoAction;
 	uint8						MipIndex		= 0;
@@ -173,9 +187,14 @@ private:
 
 
 /** Render graph information about how to bind a depth-stencil render target. */
-struct FDepthStencilBinding
+struct alignas(SHADER_PARAMETER_STRUCT_ALIGNMENT) FDepthStencilBinding
 {
-	TNonAlignedShaderParameterPtr<const FGraphTexture*> Texture = nullptr;
+	FDepthStencilBinding()
+		: Texture(nullptr)
+	{ }
+	FDepthStencilBinding(const FDepthStencilBinding&) = default;
+
+	TAlignedShaderParameterPtr<const FGraphTexture*> Texture;
 	ERenderTargetLoadAction		DepthLoadAction		= ERenderTargetLoadAction::ENoAction;
 	ERenderTargetStoreAction	DepthStoreAction	= ERenderTargetStoreAction::ENoAction;
 	ERenderTargetLoadAction		StencilLoadAction	= ERenderTargetLoadAction::ENoAction;
