Fixes shader pointer misalignements on 32bits builds.

diff --git a/Engine/Source/Runtime/RenderCore/Public/ShaderParameterMacros.h b/Engine/Source/Runtime/RenderCore/Public/ShaderParameterMacros.h
--- a/Engine/Source/Runtime/RenderCore/Public/ShaderParameterMacros.h
+++ b/Engine/Source/Runtime/RenderCore/Public/ShaderParameterMacros.h
@@ -53,7 +53,7 @@ static_assert(sizeof(void*) == 8, "Wrong PLATFORM_64BITS settings.");
 
 /** Fixed 8bytes sized pointer for shader parameters. */
 template<typename PtrType>
-class TAlignedShaderParameterPtr
+class alignas(SHADER_PARAMETER_POINTER_ALIGNMENT) TAlignedShaderParameterPtr
 {
 public:
 	TAlignedShaderParameterPtr()
@@ -63,10 +63,6 @@ public:
 		: Ref(Other)
 	{ }
 
-	TAlignedShaderParameterPtr(nullptr_t Other)
-		: Ref(Other)
-	{ }
-
 	TAlignedShaderParameterPtr(const TAlignedShaderParameterPtr<PtrType>& Other)
 		: Ref(Other.Ref)
 	{ }
@@ -76,11 +72,6 @@ public:
 		Ref = Other;
 	}
 
-	FORCEINLINE void operator=(nullptr_t Other)
-	{
-		Ref = Other;
-	}
-
 	FORCEINLINE operator PtrType&()
 	{
 		return Ref;
