Renames FGraphTexture to FRDGTexture and FRenderGraphBuilder to FRDGBuilder

diff --git a/Engine/Source/Runtime/RenderCore/Private/RenderGraphBuilder.cpp b/Engine/Source/Runtime/RenderCore/Private/RenderGraphBuilder.cpp
--- a/Engine/Source/Runtime/RenderCore/Private/RenderGraphBuilder.cpp
+++ b/Engine/Source/Runtime/RenderCore/Private/RenderGraphBuilder.cpp
@@ -21,7 +21,7 @@ static const int32 GRenderGraphImmediateMode = 0;
 #endif
 
 
-void FRenderGraphBuilder::Execute()
+void FRDGBuilder::Execute()
 {
 	#if RENDER_GRAPH_DEBUGGING
 	{
@@ -33,7 +33,7 @@ void FRenderGraphBuilder::Execute()
 	{
 		WalkGraphDependencies();
 
-		QUICK_SCOPE_CYCLE_COUNTER(STAT_FRenderGraphBuilder_Execute);
+		QUICK_SCOPE_CYCLE_COUNTER(STAT_FRDGBuilder_Execute);
 		for (const FRenderGraphPass* Pass : Passes)
 		{
 			ExecutePass(Pass);
@@ -51,7 +51,7 @@ void FRenderGraphBuilder::Execute()
 	#endif
 }
 
-void FRenderGraphBuilder::DebugPass(const FRenderGraphPass* Pass)
+void FRDGBuilder::DebugPass(const FRenderGraphPass* Pass)
 {
 	ValidatePass(Pass);
 
@@ -61,7 +61,7 @@ void FRenderGraphBuilder::DebugPass(const FRenderGraphPass* Pass)
 	}
 }
 
-void FRenderGraphBuilder::ValidatePass(const FRenderGraphPass* Pass) const
+void FRDGBuilder::ValidatePass(const FRenderGraphPass* Pass) const
 {
 	FRenderTargetBindingSlots* RESTRICT RenderTargets = nullptr;
 	FShaderParameterStructRef ParameterStruct = Pass->GetParameters();
@@ -79,7 +79,7 @@ void FRenderGraphBuilder::ValidatePass(const FRenderGraphPass* Pass) const
 		{
 		case UBMT_GRAPH_TRACKED_UAV:
 		{
-			FGraphUAV* RESTRICT UAV = *ParameterStruct.GetMemberPtrAtOffset<FGraphUAV*>(Offset);
+			FRDGTextureUAV* RESTRICT UAV = *ParameterStruct.GetMemberPtrAtOffset<FRDGTextureUAV*>(Offset);
 			if (UAV && !bCanUseUAVs)
 			{
 				UE_LOG(LogRendererCore, Warning, TEXT("UAV can only been bound to compute shaders, therefore UAV %s is certainly useless for passs %s."), UAV->Name, Pass->GetName());
@@ -129,13 +129,13 @@ void FRenderGraphBuilder::ValidatePass(const FRenderGraphPass* Pass) const
 	}
 }
 
-void FRenderGraphBuilder::WalkGraphDependencies()
+void FRDGBuilder::WalkGraphDependencies()
 {
 	for (const FRenderGraphPass* Pass : Passes)
 	{
 		FShaderParameterStructRef ParameterStruct = Pass->GetParameters();
 
-		/** Increments all the FGraphResource::ReferenceCount. */
+		/** Increments all the FRDGResource::ReferenceCount. */
 		for (int ResourceIndex = 0, Num = ParameterStruct.Layout->Resources.Num(); ResourceIndex < Num; ResourceIndex++)
 		{
 			uint8  Type = ParameterStruct.Layout->Resources[ResourceIndex];
@@ -145,7 +145,7 @@ void FRenderGraphBuilder::WalkGraphDependencies()
 			{
 			case UBMT_GRAPH_TRACKED_TEXTURE:
 			{
-				FGraphTexture* RESTRICT Texture = *ParameterStruct.GetMemberPtrAtOffset<FGraphTexture*>(Offset);
+				FRDGTexture* RESTRICT Texture = *ParameterStruct.GetMemberPtrAtOffset<FRDGTexture*>(Offset);
 				if (Texture)
 				{
 					Texture->ReferenceCount++;
@@ -154,7 +154,7 @@ void FRenderGraphBuilder::WalkGraphDependencies()
 			break;
 			case UBMT_GRAPH_TRACKED_SRV:
 			{
-				FGraphSRV* RESTRICT SRV = *ParameterStruct.GetMemberPtrAtOffset<FGraphSRV*>(Offset);
+				FRDGTextureSRV* RESTRICT SRV = *ParameterStruct.GetMemberPtrAtOffset<FRDGTextureSRV*>(Offset);
 				if (SRV)
 				{
 					SRV->Desc.Texture->ReferenceCount++;
@@ -163,7 +163,7 @@ void FRenderGraphBuilder::WalkGraphDependencies()
 			break;
 			case UBMT_GRAPH_TRACKED_UAV:
 			{
-				FGraphUAV* RESTRICT UAV = *ParameterStruct.GetMemberPtrAtOffset<FGraphUAV*>(Offset);
+				FRDGTextureUAV* RESTRICT UAV = *ParameterStruct.GetMemberPtrAtOffset<FRDGTextureUAV*>(Offset);
 				if (UAV)
 				{
 					UAV->Desc.Texture->ReferenceCount++;
@@ -217,7 +217,7 @@ void FRenderGraphBuilder::WalkGraphDependencies()
 	}
 }
 
-void FRenderGraphBuilder::AllocateRHITextureIfNeeded(const FGraphTexture* Texture, bool bComputePass)
+void FRDGBuilder::AllocateRHITextureIfNeeded(const FRDGTexture* Texture, bool bComputePass)
 {
 	check(Texture);
 
@@ -252,7 +252,7 @@ static EResourceTransitionPipeline CalcTransitionPipeline(bool bCurrentCompute,
 	return static_cast< EResourceTransitionPipeline >( Table[ Bits ] );
 }
 
-void FRenderGraphBuilder::TransitionTexture( const FGraphTexture* Texture, EResourceTransitionAccess TransitionAccess, bool bRequiredCompute ) const
+void FRDGBuilder::TransitionTexture( const FRDGTexture* Texture, EResourceTransitionAccess TransitionAccess, bool bRequiredCompute ) const
 {
 	const bool bRequiredWritable = TransitionAccess != EResourceTransitionAccess::EReadable;
 
@@ -264,7 +264,7 @@ void FRenderGraphBuilder::TransitionTexture( const FGraphTexture* Texture, EReso
 	}
 }
 
-void FRenderGraphBuilder::TransitionUAV( const FGraphUAV* UAV, EResourceTransitionAccess TransitionAccess, bool bRequiredCompute ) const
+void FRDGBuilder::TransitionUAV( const FRDGTextureUAV* UAV, EResourceTransitionAccess TransitionAccess, bool bRequiredCompute ) const
 {
 	const bool bRequiredWritable = true;
 
@@ -277,7 +277,7 @@ void FRenderGraphBuilder::TransitionUAV( const FGraphUAV* UAV, EResourceTransiti
 	}
 }
 
-static bool IsBoundAsReadable( const FGraphTexture* Texture, FShaderParameterStructRef ParameterStruct )
+static bool IsBoundAsReadable( const FRDGTexture* Texture, FShaderParameterStructRef ParameterStruct )
 {
 	for( int i = 0, Num = ParameterStruct.Layout->Resources.Num(); i < Num; i++ )
 	{
@@ -288,7 +288,7 @@ static bool IsBoundAsReadable( const FGraphTexture* Texture, FShaderParameterStr
 		{
 			case UBMT_GRAPH_TRACKED_TEXTURE:
 			{
-				const FGraphTexture* InputTexture = *ParameterStruct.GetMemberPtrAtOffset<const FGraphTexture*>(Offset);
+				const FRDGTexture* InputTexture = *ParameterStruct.GetMemberPtrAtOffset<const FRDGTexture*>(Offset);
 				if( Texture == InputTexture)
 				{
 					return true;
@@ -297,7 +297,7 @@ static bool IsBoundAsReadable( const FGraphTexture* Texture, FShaderParameterStr
 			break;
 		case UBMT_GRAPH_TRACKED_SRV:
 			{
-				const FGraphSRV* InputSRV = *ParameterStruct.GetMemberPtrAtOffset<const FGraphSRV*>(Offset);
+				const FRDGTextureSRV* InputSRV = *ParameterStruct.GetMemberPtrAtOffset<const FRDGTextureSRV*>(Offset);
 				if (InputSRV && Texture == InputSRV->Desc.Texture)
 				{
 					return true;
@@ -312,9 +312,9 @@ static bool IsBoundAsReadable( const FGraphTexture* Texture, FShaderParameterStr
 	return false;
 }
 
-void FRenderGraphBuilder::ExecutePass( const FRenderGraphPass* Pass )
+void FRDGBuilder::ExecutePass( const FRenderGraphPass* Pass )
 {
-	QUICK_SCOPE_CYCLE_COUNTER(STAT_FRenderGraphBuilder_ExecutePass);
+	QUICK_SCOPE_CYCLE_COUNTER(STAT_FRDGBuilder_ExecutePass);
 
 	FRHIRenderPassInfo RPInfo;
 	bool bHasRenderTargets = false;
@@ -350,7 +350,7 @@ void FRenderGraphBuilder::ExecutePass( const FRenderGraphPass* Pass )
 	}
 }
 
-void FRenderGraphBuilder::AllocateAndTransitionPassResources(const FRenderGraphPass* Pass, struct FRHIRenderPassInfo* OutRPInfo, bool* bOutHasRenderTargets)
+void FRDGBuilder::AllocateAndTransitionPassResources(const FRenderGraphPass* Pass, struct FRHIRenderPassInfo* OutRPInfo, bool* bOutHasRenderTargets)
 {
 	bool bIsCompute = Pass->IsCompute();
 	FShaderParameterStructRef ParameterStruct = Pass->GetParameters();
@@ -364,7 +364,7 @@ void FRenderGraphBuilder::AllocateAndTransitionPassResources(const FRenderGraphP
 		{
 		case UBMT_GRAPH_TRACKED_TEXTURE:
 		{
-			FGraphTexture* RESTRICT Texture = *ParameterStruct.GetMemberPtrAtOffset<FGraphTexture*>(Offset);
+			FRDGTexture* RESTRICT Texture = *ParameterStruct.GetMemberPtrAtOffset<FRDGTexture*>(Offset);
 			if (Texture)
 			{
 				AllocateRHITextureIfNeeded(Texture, bIsCompute);
@@ -374,12 +374,12 @@ void FRenderGraphBuilder::AllocateAndTransitionPassResources(const FRenderGraphP
 		break;
 		case UBMT_GRAPH_TRACKED_SRV:
 		{
-			FGraphSRV* RESTRICT SRV = *ParameterStruct.GetMemberPtrAtOffset<FGraphSRV*>(Offset);
+			FRDGTextureSRV* RESTRICT SRV = *ParameterStruct.GetMemberPtrAtOffset<FRDGTextureSRV*>(Offset);
 			if (SRV)
 			{
 				check(SRV->Desc.Texture);
 
-				const FGraphTexture* Texture = SRV->Desc.Texture;
+				const FRDGTexture* Texture = SRV->Desc.Texture;
 				AllocateRHITextureIfNeeded(Texture, bIsCompute);
 				TransitionTexture(Texture, EResourceTransitionAccess::EReadable, bIsCompute);
 			}
@@ -387,7 +387,7 @@ void FRenderGraphBuilder::AllocateAndTransitionPassResources(const FRenderGraphP
 		break;
 		case UBMT_GRAPH_TRACKED_UAV:
 		{
-			FGraphUAV* RESTRICT UAV = *ParameterStruct.GetMemberPtrAtOffset<FGraphUAV*>(Offset);
+			FRDGTextureUAV* RESTRICT UAV = *ParameterStruct.GetMemberPtrAtOffset<FRDGTextureUAV*>(Offset);
 			if (UAV)
 			{
 				AllocateRHITextureIfNeeded(UAV->Desc.Texture, bIsCompute);
@@ -463,7 +463,7 @@ void FRenderGraphBuilder::AllocateAndTransitionPassResources(const FRenderGraphP
 }
 
 // static 
-void FRenderGraphBuilder::WarnForUselessPassDependencies(const FRenderGraphPass* Pass)
+void FRDGBuilder::WarnForUselessPassDependencies(const FRenderGraphPass* Pass)
 {
 	FShaderParameterStructRef ParameterStruct = Pass->GetParameters();
 
@@ -479,7 +479,7 @@ void FRenderGraphBuilder::WarnForUselessPassDependencies(const FRenderGraphPass*
 		if (Type != UBMT_GRAPH_TRACKED_TEXTURE && Type != UBMT_GRAPH_TRACKED_SRV && Type != UBMT_GRAPH_TRACKED_UAV)
 			continue;
 
-		const FGraphResource* Resource = *ParameterStruct.GetMemberPtrAtOffset<const FGraphResource*>(Offset);
+		const FRDGResource* Resource = *ParameterStruct.GetMemberPtrAtOffset<const FRDGResource*>(Offset);
 
 		if (!Resource)
 			continue;
@@ -500,7 +500,7 @@ void FRenderGraphBuilder::WarnForUselessPassDependencies(const FRenderGraphPass*
 			if (Type != UBMT_GRAPH_TRACKED_TEXTURE && Type != UBMT_GRAPH_TRACKED_SRV && Type != UBMT_GRAPH_TRACKED_UAV)
 				continue;
 
-			const FGraphResource* Resource = *ParameterStruct.GetMemberPtrAtOffset<const FGraphResource*>(Offset);
+			const FRDGResource* Resource = *ParameterStruct.GetMemberPtrAtOffset<const FRDGResource*>(Offset);
 
 			if (!Resource)
 				continue;
@@ -518,7 +518,7 @@ void FRenderGraphBuilder::WarnForUselessPassDependencies(const FRenderGraphPass*
 		if (Type != UBMT_GRAPH_TRACKED_TEXTURE && Type != UBMT_GRAPH_TRACKED_SRV && Type != UBMT_GRAPH_TRACKED_UAV)
 			continue;
 
-		const FGraphResource* Resource = *ParameterStruct.GetMemberPtrAtOffset<const FGraphResource*>(Offset);
+		const FRDGResource* Resource = *ParameterStruct.GetMemberPtrAtOffset<const FRDGResource*>(Offset);
 
 		if (!Resource)
 			continue;
@@ -527,7 +527,7 @@ void FRenderGraphBuilder::WarnForUselessPassDependencies(const FRenderGraphPass*
 	}
 }
 
-void FRenderGraphBuilder::ReleaseRHITextureIfPossible(const FGraphTexture* Texture)
+void FRDGBuilder::ReleaseRHITextureIfPossible(const FRDGTexture* Texture)
 {
 	check(Texture->ReferenceCount > 0);
 	Texture->ReferenceCount--;
@@ -539,11 +539,11 @@ void FRenderGraphBuilder::ReleaseRHITextureIfPossible(const FGraphTexture* Textu
 	}
 }
 
-void FRenderGraphBuilder::ReleaseUnecessaryResources(const FRenderGraphPass* Pass)
+void FRDGBuilder::ReleaseUnecessaryResources(const FRenderGraphPass* Pass)
 {
 	FShaderParameterStructRef ParameterStruct = Pass->GetParameters();
 
-	/** Increments all the FGraphResource::ReferenceCount. */
+	/** Increments all the FRDGResource::ReferenceCount. */
 	for (int ResourceIndex = 0, Num = ParameterStruct.Layout->Resources.Num(); ResourceIndex < Num; ResourceIndex++)
 	{
 		uint8  Type = ParameterStruct.Layout->Resources[ResourceIndex];
@@ -553,7 +553,7 @@ void FRenderGraphBuilder::ReleaseUnecessaryResources(const FRenderGraphPass* Pas
 		{
 		case UBMT_GRAPH_TRACKED_TEXTURE:
 		{
-			FGraphTexture* RESTRICT Texture = *ParameterStruct.GetMemberPtrAtOffset<FGraphTexture*>(Offset);
+			FRDGTexture* RESTRICT Texture = *ParameterStruct.GetMemberPtrAtOffset<FRDGTexture*>(Offset);
 			if (Texture)
 			{
 				ReleaseRHITextureIfPossible(Texture);
@@ -562,7 +562,7 @@ void FRenderGraphBuilder::ReleaseUnecessaryResources(const FRenderGraphPass* Pas
 		break;
 		case UBMT_GRAPH_TRACKED_SRV:
 		{
-			FGraphSRV* RESTRICT SRV = *ParameterStruct.GetMemberPtrAtOffset<FGraphSRV*>(Offset);
+			FRDGTextureSRV* RESTRICT SRV = *ParameterStruct.GetMemberPtrAtOffset<FRDGTextureSRV*>(Offset);
 			if (SRV)
 			{
 				ReleaseRHITextureIfPossible(SRV->Desc.Texture);
@@ -571,7 +571,7 @@ void FRenderGraphBuilder::ReleaseUnecessaryResources(const FRenderGraphPass* Pas
 		break;
 		case UBMT_GRAPH_TRACKED_UAV:
 		{
-			FGraphUAV* RESTRICT UAV = *ParameterStruct.GetMemberPtrAtOffset<FGraphUAV*>(Offset);
+			FRDGTextureUAV* RESTRICT UAV = *ParameterStruct.GetMemberPtrAtOffset<FRDGTextureUAV*>(Offset);
 			if (UAV)
 			{
 				ReleaseRHITextureIfPossible(UAV->Desc.Texture);
@@ -606,7 +606,7 @@ void FRenderGraphBuilder::ReleaseUnecessaryResources(const FRenderGraphPass* Pas
 	}
 }
 
-void FRenderGraphBuilder::ProcessDeferredInternalResourceQueries()
+void FRDGBuilder::ProcessDeferredInternalResourceQueries()
 {
 	for (const auto& Query : DeferredInternalTextureQueries)
 	{
@@ -627,12 +627,12 @@ void FRenderGraphBuilder::ProcessDeferredInternalResourceQueries()
 	}
 }
 
-void FRenderGraphBuilder::DestructPasses()
+void FRDGBuilder::DestructPasses()
 {
 	#if RENDER_GRAPH_DEBUGGING
 	{
 		// Make sure all resource references have been released to ensure no leaks happen.
-		for (const FGraphResource* Resource : Resources)
+		for (const FRDGResource* Resource : Resources)
 		{
 			check(Resource->ReferenceCount == 0);
 		}
diff --git a/Engine/Source/Runtime/RenderCore/Public/RenderGraphBuilder.h b/Engine/Source/Runtime/RenderCore/Public/RenderGraphBuilder.h
--- a/Engine/Source/Runtime/RenderCore/Public/RenderGraphBuilder.h
+++ b/Engine/Source/Runtime/RenderCore/Public/RenderGraphBuilder.h
@@ -84,7 +84,7 @@ struct TLambdaRenderPass final : public FRenderGraphPass
 	~TLambdaRenderPass()
 	{
 		// Manually call the destructor of the pass parameter, to make sure RHI references are released since the pass parameters are allocated on FMemStack.
-		// TODO(RDG): this may lead to RHI resource leaks if a struct allocated in FMemStack does not actually get used through FRenderGraphBuilder::AddPass().
+		// TODO(RDG): this may lead to RHI resource leaks if a struct allocated in FMemStack does not actually get used through FRDGBuilder::AddPass().
 		ParameterStructType* Struct = reinterpret_cast<ParameterStructType*>(const_cast<void*>(FRenderGraphPass::ParameterStruct.Contents));
 		Struct->~ParameterStructType();
 	}
@@ -102,28 +102,28 @@ struct TLambdaRenderPass final : public FRenderGraphPass
  * Resources must be created from the builder before they can be bound to Pass ResourceTables.
  * These resources are descriptors only until the graph is executed, where RHI resources are allocated as needed.
  */
-class RENDERCORE_API FRenderGraphBuilder
+class RENDERCORE_API FRDGBuilder
 {
 public:
 	/** A RHI cmd list is required, if using the immediate mode. */
-	FRenderGraphBuilder(FRHICommandListImmediate& InRHICmdList)
+	FRDGBuilder(FRHICommandListImmediate& InRHICmdList)
 		: RHICmdList(InRHICmdList)
 	{ }
 
-	~FRenderGraphBuilder()
+	~FRDGBuilder()
 	{
 		DestructPasses();
 	}
 
 	/** Register a external texture to be tracked by the render graph. */
-	inline const FGraphTexture* RegisterExternalTexture(const TRefCountPtr<IPooledRenderTarget>& ExternalPooledTexture, const TCHAR* Name = TEXT("External"))
+	inline const FRDGTexture* RegisterExternalTexture(const TRefCountPtr<IPooledRenderTarget>& ExternalPooledTexture, const TCHAR* Name = TEXT("External"))
 	{
 		#if RENDER_GRAPH_DEBUGGING
 		{
 			ensureMsgf(ExternalPooledTexture.IsValid(), TEXT("Attempted to register NULL external texture: %s"), Name);
 		}
 		#endif
-		FGraphTexture* OutTexture = new(FMemStack::Get()) FGraphTexture(Name, ExternalPooledTexture->GetDesc());
+		FRDGTexture* OutTexture = new(FMemStack::Get()) FRDGTexture(Name, ExternalPooledTexture->GetDesc());
 		OutTexture->PooledRenderTarget = ExternalPooledTexture;
 		AllocatedTextures.Add(OutTexture, ExternalPooledTexture);
 		#if DO_CHECK
@@ -133,14 +133,14 @@ public:
 	}
 
 	/** Create graph tracked resource from a descriptor with a debug name. */
-	inline const FGraphTexture* CreateTexture(const FPooledRenderTargetDesc& Desc, const TCHAR* DebugName)
+	inline const FRDGTexture* CreateTexture(const FPooledRenderTargetDesc& Desc, const TCHAR* DebugName)
 	{
 		#if RENDER_GRAPH_DEBUGGING
 		{
 			ensureMsgf(!bHasExecuted, TEXT("Render graph texture %s needs to be created before the builder execution."), DebugName);
 		}
 		#endif
-		FGraphTexture* Texture = new(FMemStack::Get()) FGraphTexture(DebugName, Desc);
+		FRDGTexture* Texture = new(FMemStack::Get()) FRDGTexture(DebugName, Desc);
 		#if DO_CHECK
 			Resources.Add(Texture);
 		#endif
@@ -148,7 +148,7 @@ public:
 	}
 
 	/** Create graph tracked SRV from a descriptor. */
-	inline const FGraphSRV* CreateSRV(const FGraphSRVDesc& Desc)
+	inline const FRDGTextureSRV* CreateSRV(const FRDGTextureSRVDesc& Desc)
 	{
 		check(Desc.Texture);
 		#if RENDER_GRAPH_DEBUGGING
@@ -158,7 +158,7 @@ public:
 		}
 		#endif
 		
-		FGraphSRV* SRV = new(FMemStack::Get()) FGraphSRV(Desc.Texture->Name, Desc);
+		FRDGTextureSRV* SRV = new(FMemStack::Get()) FRDGTextureSRV(Desc.Texture->Name, Desc);
 		#if DO_CHECK
 			Resources.Add(SRV);
 		#endif
@@ -166,7 +166,7 @@ public:
 	}
 
 	/** Create graph tracked UAV from a descriptor. */
-	inline const FGraphUAV* CreateUAV(const FGraphUAVDesc& Desc)
+	inline const FRDGTextureUAV* CreateUAV(const FRDGTextureUAVDesc& Desc)
 	{
 		check(Desc.Texture);
 		#if RENDER_GRAPH_DEBUGGING
@@ -176,7 +176,7 @@ public:
 		}
 		#endif
 		
-		FGraphUAV* UAV = new(FMemStack::Get()) FGraphUAV(Desc.Texture->Name, Desc);
+		FRDGTextureUAV* UAV = new(FMemStack::Get()) FRDGTextureUAV(Desc.Texture->Name, Desc);
 		#if DO_CHECK
 			Resources.Add(UAV);
 		#endif
@@ -227,7 +227,7 @@ public:
 	/**
 	 * Extracts an internal texture by handle.  Must be called before Execute.
 	 */
-	inline void GetInternalTexture(const FGraphTexture* Texture, TRefCountPtr<IPooledRenderTarget>* OutTexturePtr, bool bTransitionToRead = true)
+	inline void GetInternalTexture(const FRDGTexture* Texture, TRefCountPtr<IPooledRenderTarget>* OutTexturePtr, bool bTransitionToRead = true)
 	{
 		check(Texture);
 		check(OutTexturePtr);
@@ -255,13 +255,13 @@ private:
 	/** Array of all pass created */
 	TArray<FRenderGraphPass*, SceneRenderingAllocator> Passes;
 
-	/** Keep the references over the pooled render target, since FGraphTexture is allocated on FMemStack. */
-	TMap<const FGraphTexture*, TRefCountPtr<IPooledRenderTarget>, SceneRenderingSetAllocator> AllocatedTextures;
+	/** Keep the references over the pooled render target, since FRDGTexture is allocated on FMemStack. */
+	TMap<const FRDGTexture*, TRefCountPtr<IPooledRenderTarget>, SceneRenderingSetAllocator> AllocatedTextures;
 
 	/** Array of all deferred access to internal textures. */
 	struct FDeferredInternalTextureQuery
 	{
-		const FGraphTexture* Texture;
+		const FRDGTexture* Texture;
 		TRefCountPtr<IPooledRenderTarget>* OutTexturePtr;
 		bool bTransitionToRead;
 	};
@@ -273,7 +273,7 @@ private:
 		bool bHasExecuted = false;
 
 		/** Lists of all created resources */
-		TArray<const FGraphResource*, SceneRenderingAllocator> Resources;
+		TArray<const FRDGResource*, SceneRenderingAllocator> Resources;
 	#endif
 
 	void DebugPass(const FRenderGraphPass* Pass);
@@ -281,17 +281,17 @@ private:
 
 	void WalkGraphDependencies();
 
-	void AllocateRHITextureIfNeeded(const FGraphTexture* Texture, bool bComputePass);
+	void AllocateRHITextureIfNeeded(const FRDGTexture* Texture, bool bComputePass);
 
-	void TransitionTexture( const FGraphTexture* Texture, EResourceTransitionAccess TransitionAccess, bool bRequiredCompute ) const;
-	void TransitionUAV( const FGraphUAV* UAV, EResourceTransitionAccess TransitionAccess, bool bRequiredCompute ) const;
+	void TransitionTexture( const FRDGTexture* Texture, EResourceTransitionAccess TransitionAccess, bool bRequiredCompute ) const;
+	void TransitionUAV( const FRDGTextureUAV* UAV, EResourceTransitionAccess TransitionAccess, bool bRequiredCompute ) const;
 
 	void ExecutePass( const FRenderGraphPass* Pass );
 	void AllocateAndTransitionPassResources(const FRenderGraphPass* Pass, struct FRHIRenderPassInfo* OutRPInfo, bool* bOutHasRenderTargets);
 	static void WarnForUselessPassDependencies(const FRenderGraphPass* Pass);
-	void ReleaseRHITextureIfPossible(const FGraphTexture* Texture);
+	void ReleaseRHITextureIfPossible(const FRDGTexture* Texture);
 	void ReleaseUnecessaryResources(const FRenderGraphPass* Pass);
 
 	void ProcessDeferredInternalResourceQueries();
 	void DestructPasses();
-}; // class FRenderGraphBuilder
+}; // class FRDGBuilder
diff --git a/Engine/Source/Runtime/RenderCore/Public/RenderGraphResources.h b/Engine/Source/Runtime/RenderCore/Public/RenderGraphResources.h
--- a/Engine/Source/Runtime/RenderCore/Public/RenderGraphResources.h
+++ b/Engine/Source/Runtime/RenderCore/Public/RenderGraphResources.h
@@ -8,19 +8,19 @@
 
 
 /** Generic graph resource. */
-class RENDERCORE_API FGraphResource
+class RENDERCORE_API FRDGResource
 {
 public:
 	// Name of the resource for debugging purpose.
 	const TCHAR* const Name = nullptr;
 
-	FGraphResource(const TCHAR* InDebugName)
+	FRDGResource(const TCHAR* InDebugName)
 		: Name(InDebugName)
 	{ }
 
-	FGraphResource() = delete;
-	FGraphResource(const FGraphResource&) = delete;
-	void operator = (const FGraphResource&) = delete;
+	FRDGResource() = delete;
+	FRDGResource(const FRDGResource&) = delete;
+	void operator = (const FRDGResource&) = delete;
 
 private:
 	/** Number of references in passes and deferred queries. */
@@ -33,7 +33,7 @@ private:
 	/** Boolean to tracked whether a ressource is actually used by the lambda of a pass or not. */
 	mutable bool bIsActuallyUsedByPass = false;
 
-	friend class FRenderGraphBuilder;
+	friend class FRDGBuilder;
 
 
 	/** Friendship over parameter settings for bIsActuallyUsedByPass. It means only this parameter settings can be
@@ -47,7 +47,7 @@ private:
 };
 
 /** Render graph tracked Texture. */
-class RENDERCORE_API FGraphTexture : public FGraphResource
+class RENDERCORE_API FRDGTexture : public FRDGResource
 {
 public:
 	//TODO(RDG): using FDesc = FPooledRenderTargetDesc;
@@ -56,13 +56,13 @@ public:
 	const FPooledRenderTargetDesc Desc;
 
 private:
-	/** This is not a TRefCountPtr<> because FGraphTexture is allocated on the FMemStack
+	/** This is not a TRefCountPtr<> because FRDGTexture is allocated on the FMemStack
 	 * FGraphBuilder::AllocatedTextures is actually keeping the reference.
 	 */
 	mutable IPooledRenderTarget* PooledRenderTarget = nullptr;
 
-	FGraphTexture(const TCHAR* DebugName, const FPooledRenderTargetDesc& InDesc)
-		: FGraphResource(DebugName)
+	FRDGTexture(const TCHAR* DebugName, const FPooledRenderTargetDesc& InDesc)
+		: FRDGResource(DebugName)
 		, Desc(InDesc)
 	{ }
 
@@ -73,9 +73,9 @@ private:
 		return PooledRenderTarget->GetRenderTargetItem().ShaderResourceTexture;
 	}
 
-	friend class FRenderGraphBuilder;
-	friend class FGraphSRV;
-	friend class FGraphUAV;
+	friend class FRDGBuilder;
+	friend class FRDGTextureSRV;
+	friend class FRDGTextureUAV;
 
 	template<typename TRHICmdList, typename TShaderClass, typename TShaderRHI>
 	friend void SetShaderParameters(TRHICmdList&, const TShaderClass*, TShaderRHI*, const typename TShaderClass::FParameters&);
@@ -85,30 +85,30 @@ private:
 };
 
 /** Decsriptor for render graph tracked SRV. */
-class RENDERCORE_API FGraphSRVDesc
+class RENDERCORE_API FRDGTextureSRVDesc
 {
 public:
-	const FGraphTexture* Texture;
+	const FRDGTexture* Texture;
 	uint8 MipLevel = 0;
 
-	FGraphSRVDesc() {}
+	FRDGTextureSRVDesc() {}
 
-	FGraphSRVDesc(const FGraphTexture* InTexture, uint8 InMipLevel) :
+	FRDGTextureSRVDesc(const FRDGTexture* InTexture, uint8 InMipLevel) :
 		Texture(InTexture),
 		MipLevel(InMipLevel)
 	{}
 };
 
 /** Render graph tracked SRV. */
-class RENDERCORE_API FGraphSRV : public FGraphResource
+class RENDERCORE_API FRDGTextureSRV : public FRDGResource
 {
 public:
 	/** Descriptor of the graph tracked SRV. */
-	const FGraphSRVDesc Desc;
+	const FRDGTextureSRVDesc Desc;
 
 private:
-	FGraphSRV(const TCHAR* DebugName, const FGraphSRVDesc& InDesc)
-		: FGraphResource(DebugName)
+	FRDGTextureSRV(const TCHAR* DebugName, const FRDGTextureSRVDesc& InDesc)
+		: FRDGResource(DebugName)
 		, Desc(InDesc)
 	{ }
 
@@ -120,7 +120,7 @@ private:
 		return Desc.Texture->PooledRenderTarget->GetRenderTargetItem().MipSRVs[Desc.MipLevel];
 	}
 
-	friend class FRenderGraphBuilder;
+	friend class FRDGBuilder;
 
 	template<typename TRHICmdList, typename TShaderClass, typename TShaderRHI>
 	friend void SetShaderParameters(TRHICmdList&, const TShaderClass*, TShaderRHI*, const typename TShaderClass::FParameters&);
@@ -130,30 +130,30 @@ private:
 };
 
 /** Decsriptor for render graph tracked UAV. */
-class RENDERCORE_API FGraphUAVDesc
+class RENDERCORE_API FRDGTextureUAVDesc
 {
 public:
-	const FGraphTexture* Texture;
+	const FRDGTexture* Texture;
 	uint8 MipLevel = 0;
 
-	FGraphUAVDesc() {}
+	FRDGTextureUAVDesc() {}
 
-	FGraphUAVDesc(const FGraphTexture* InTexture, uint8 InMipLevel = 0) :
+	FRDGTextureUAVDesc(const FRDGTexture* InTexture, uint8 InMipLevel = 0) :
 		Texture(InTexture),
 		MipLevel(InMipLevel)
 	{}
 };
 
 /** Render graph tracked UAV. */
-class RENDERCORE_API FGraphUAV : public FGraphResource
+class RENDERCORE_API FRDGTextureUAV : public FRDGResource
 {
 public:
 	/** Descriptor of the graph tracked UAV. */
-	const FGraphUAVDesc Desc;
+	const FRDGTextureUAVDesc Desc;
 
 private:
-	FGraphUAV(const TCHAR* DebugName, const FGraphUAVDesc& InDesc)
-		: FGraphResource(DebugName)
+	FRDGTextureUAV(const TCHAR* DebugName, const FRDGTextureUAVDesc& InDesc)
+		: FRDGResource(DebugName)
 		, Desc(InDesc)
 	{ }
 
@@ -165,7 +165,7 @@ private:
 		return Desc.Texture->PooledRenderTarget->GetRenderTargetItem().MipUAVs[Desc.MipLevel];
 	}
 
-	friend class FRenderGraphBuilder;
+	friend class FRDGBuilder;
 
 	template<typename TRHICmdList, typename TShaderClass, typename TShaderRHI>
 	friend void SetShaderParameters(TRHICmdList&, const TShaderClass*, TShaderRHI*, const typename TShaderClass::FParameters&);
diff --git a/Engine/Source/Runtime/RenderCore/Public/ShaderParameterMacros.h b/Engine/Source/Runtime/RenderCore/Public/ShaderParameterMacros.h
--- a/Engine/Source/Runtime/RenderCore/Public/ShaderParameterMacros.h
+++ b/Engine/Source/Runtime/RenderCore/Public/ShaderParameterMacros.h
@@ -10,9 +10,9 @@
 #include "ShaderParameterMetadata.h"
 
 
-class FGraphTexture;
-class FGraphSRV;
-class FGraphUAV;
+class FRDGTexture;
+class FRDGTextureSRV;
+class FRDGTextureUAV;
 
 
 // RHICreateUniformBuffer assumes C++ constant layout matches the shader layout when extracting float constants, yet the C++ struct contains pointers.  
@@ -146,14 +146,14 @@ struct alignas(SHADER_PARAMETER_STRUCT_ALIGNMENT) FRenderTargetBinding
 	 *
 	 * Notes: Load and store action are on purpose without default values, to force the user to not forget one of these.
 	 */
-	FRenderTargetBinding(const FGraphTexture* InTexture, ERenderTargetLoadAction InLoadAction, ERenderTargetStoreAction InStoreAction, uint8 InMipIndex = 0)
+	FRenderTargetBinding(const FRDGTexture* InTexture, ERenderTargetLoadAction InLoadAction, ERenderTargetStoreAction InStoreAction, uint8 InMipIndex = 0)
 		: Texture(InTexture)
 		, LoadAction(InLoadAction)
 		, StoreAction(InStoreAction)
 		, MipIndex(InMipIndex)
 	{}
 
-	inline const FGraphTexture* GetTexture() const {
+	inline const FRDGTexture* GetTexture() const {
 		return Texture;
 	};
 	inline ERenderTargetLoadAction GetLoadAction() const {
@@ -170,7 +170,7 @@ private:
 	/** All parameters required to bind a render target deferred. This are purposefully private to
 	 * force the user to call FRenderTargetBinding() constructor, forcing him to specify the load and store action.
 	 */
-	TAlignedShaderParameterPtr<const FGraphTexture*> Texture;
+	TAlignedShaderParameterPtr<const FRDGTexture*> Texture;
 	ERenderTargetLoadAction		LoadAction		= ERenderTargetLoadAction::ENoAction;
 	ERenderTargetStoreAction	StoreAction		= ERenderTargetStoreAction::ENoAction;
 	uint8						MipIndex		= 0;
@@ -185,7 +185,7 @@ struct alignas(SHADER_PARAMETER_STRUCT_ALIGNMENT) FDepthStencilBinding
 	{ }
 	FDepthStencilBinding(const FDepthStencilBinding&) = default;
 
-	TAlignedShaderParameterPtr<const FGraphTexture*> Texture;
+	TAlignedShaderParameterPtr<const FRDGTexture*> Texture;
 	ERenderTargetLoadAction		DepthLoadAction		= ERenderTargetLoadAction::ENoAction;
 	ERenderTargetStoreAction	DepthStoreAction	= ERenderTargetStoreAction::ENoAction;
 	ERenderTargetLoadAction		StencilLoadAction	= ERenderTargetLoadAction::ENoAction;
@@ -521,7 +521,7 @@ struct TShaderParameterTypeInfo<FTextureRHIParamRef>
 static_assert(sizeof(TShaderParameterTypeInfo<FTextureRHIParamRef>::TAlignedType) == SHADER_PARAMETER_POINTER_ALIGNMENT, "Uniform buffer layout must not be platform dependent.");
 
 template<>
-struct TShaderParameterTypeInfo<const FGraphTexture*>
+struct TShaderParameterTypeInfo<const FRDGTexture*>
 {
 	static constexpr int32 NumRows = 1;
 	static constexpr int32 NumColumns = 1;
@@ -529,14 +529,14 @@ struct TShaderParameterTypeInfo<const FGraphTexture*>
 	static constexpr int32 Alignment = SHADER_PARAMETER_POINTER_ALIGNMENT;
 	static constexpr bool bIsStoredInConstantBuffer = false;
 
-	using TAlignedType = TAlignedShaderParameterPtr<const FGraphTexture*>;
+	using TAlignedType = TAlignedShaderParameterPtr<const FRDGTexture*>;
 
 	static const FShaderParametersMetadata* GetStructMetadata() { return nullptr; }
 };
-static_assert(sizeof(TShaderParameterTypeInfo<const FGraphTexture*>::TAlignedType) == SHADER_PARAMETER_POINTER_ALIGNMENT, "Uniform buffer layout must not be platform dependent.");
+static_assert(sizeof(TShaderParameterTypeInfo<const FRDGTexture*>::TAlignedType) == SHADER_PARAMETER_POINTER_ALIGNMENT, "Uniform buffer layout must not be platform dependent.");
 
 template<>
-struct TShaderParameterTypeInfo<const FGraphSRV*>
+struct TShaderParameterTypeInfo<const FRDGTextureSRV*>
 {
 	static constexpr int32 NumRows = 1;
 	static constexpr int32 NumColumns = 1;
@@ -544,14 +544,14 @@ struct TShaderParameterTypeInfo<const FGraphSRV*>
 	static constexpr int32 Alignment = SHADER_PARAMETER_POINTER_ALIGNMENT;
 	static constexpr bool bIsStoredInConstantBuffer = false;
 	
-	using TAlignedType = TAlignedShaderParameterPtr<const FGraphSRV*>;
+	using TAlignedType = TAlignedShaderParameterPtr<const FRDGTextureSRV*>;
 	
 	static const FShaderParametersMetadata* GetStructMetadata() { return nullptr; }
 };
-static_assert(sizeof(TShaderParameterTypeInfo<const FGraphSRV*>::TAlignedType) == SHADER_PARAMETER_POINTER_ALIGNMENT, "Uniform buffer layout must not be platform dependent.");
+static_assert(sizeof(TShaderParameterTypeInfo<const FRDGTextureSRV*>::TAlignedType) == SHADER_PARAMETER_POINTER_ALIGNMENT, "Uniform buffer layout must not be platform dependent.");
 
 template<>
-struct TShaderParameterTypeInfo<const FGraphUAV*>
+struct TShaderParameterTypeInfo<const FRDGTextureUAV*>
 {
 	static constexpr int32 NumRows = 1;
 	static constexpr int32 NumColumns = 1;
@@ -559,11 +559,11 @@ struct TShaderParameterTypeInfo<const FGraphUAV*>
 	static constexpr int32 Alignment = SHADER_PARAMETER_POINTER_ALIGNMENT;
 	static constexpr bool bIsStoredInConstantBuffer = false;
 	
-	using TAlignedType = TAlignedShaderParameterPtr<const FGraphUAV*>;
+	using TAlignedType = TAlignedShaderParameterPtr<const FRDGTextureUAV*>;
 	
 	static const FShaderParametersMetadata* GetStructMetadata() { return nullptr; }
 };
-static_assert(sizeof(TShaderParameterTypeInfo<const FGraphUAV*>::TAlignedType) == SHADER_PARAMETER_POINTER_ALIGNMENT, "Uniform buffer layout must not be platform dependent.");
+static_assert(sizeof(TShaderParameterTypeInfo<const FRDGTextureUAV*>::TAlignedType) == SHADER_PARAMETER_POINTER_ALIGNMENT, "Uniform buffer layout must not be platform dependent.");
 
 template<class UniformBufferStructType>
 struct TShaderParameterTypeInfo<TUniformBufferRef<UniformBufferStructType>>
@@ -758,7 +758,7 @@ extern RENDERCORE_API FShaderParametersMetadata* FindUniformBufferStructByFName(
  *	SHADER_PARAMETER_GRAPH_TEXTURE(Texture2D, MyTexture)
  */
 #define SHADER_PARAMETER_GRAPH_TEXTURE(ShaderType,MemberName) \
-	DECLARE_UNIFORM_BUFFER_STRUCT_MEMBER_EXPLICIT(UBMT_GRAPH_TRACKED_TEXTURE, TShaderParameterTypeInfo<const FGraphTexture*>, const FGraphTexture*,MemberName,,EShaderPrecisionModifier::Float,TEXT(#ShaderType),false)
+	DECLARE_UNIFORM_BUFFER_STRUCT_MEMBER_EXPLICIT(UBMT_GRAPH_TRACKED_TEXTURE, TShaderParameterTypeInfo<const FRDGTexture*>, const FRDGTexture*,MemberName,,EShaderPrecisionModifier::Float,TEXT(#ShaderType),false)
 
 /** Adds a render graph tracked shader resource view.
  *
@@ -766,7 +766,7 @@ extern RENDERCORE_API FShaderParametersMetadata* FindUniformBufferStructByFName(
  *	SHADER_PARAMETER_GRAPH_SRV(Texture2D, MySRV)
  */
 #define SHADER_PARAMETER_GRAPH_SRV(ShaderType,MemberName) \
-	DECLARE_UNIFORM_BUFFER_STRUCT_MEMBER_EXPLICIT(UBMT_GRAPH_TRACKED_SRV, TShaderParameterTypeInfo<const FGraphSRV*>, const FGraphSRV*,MemberName,,EShaderPrecisionModifier::Float,TEXT(#ShaderType),false)
+	DECLARE_UNIFORM_BUFFER_STRUCT_MEMBER_EXPLICIT(UBMT_GRAPH_TRACKED_SRV, TShaderParameterTypeInfo<const FRDGTextureSRV*>, const FRDGTextureSRV*,MemberName,,EShaderPrecisionModifier::Float,TEXT(#ShaderType),false)
 
 /** Adds a render graph tracked unordered access view.
  *
@@ -774,7 +774,7 @@ extern RENDERCORE_API FShaderParametersMetadata* FindUniformBufferStructByFName(
  *	SHADER_PARAMETER_GRAPH_UAV(RWTexture2D, MyUAV)
  */
 #define SHADER_PARAMETER_GRAPH_UAV(ShaderType,MemberName) \
-	DECLARE_UNIFORM_BUFFER_STRUCT_MEMBER_EXPLICIT(UBMT_GRAPH_TRACKED_UAV, TShaderParameterTypeInfo<const FGraphUAV*>, const FGraphUAV*,MemberName,,EShaderPrecisionModifier::Float,TEXT(#ShaderType),false)
+	DECLARE_UNIFORM_BUFFER_STRUCT_MEMBER_EXPLICIT(UBMT_GRAPH_TRACKED_UAV, TShaderParameterTypeInfo<const FRDGTextureUAV*>, const FRDGTextureUAV*,MemberName,,EShaderPrecisionModifier::Float,TEXT(#ShaderType),false)
 
 /** Nests a shader parameter structure into another one, in C++ and shader code.
  *
diff --git a/Engine/Source/Runtime/RenderCore/Public/ShaderParameterStruct.h b/Engine/Source/Runtime/RenderCore/Public/ShaderParameterStruct.h
--- a/Engine/Source/Runtime/RenderCore/Public/ShaderParameterStruct.h
+++ b/Engine/Source/Runtime/RenderCore/Public/ShaderParameterStruct.h
@@ -59,7 +59,7 @@ inline void SetShaderUAVs(TRHICmdList& RHICmdList, const TShaderClass* Shader, F
 	// Graph UAVs
 	for (const FShaderParameterBindings::FResourceParameter& ParameterBinding : Bindings.GraphUAVs)
 	{
-		auto GraphUAV = *reinterpret_cast<const FGraphUAV* const*>(Base + ParameterBinding.ByteOffset);
+		auto GraphUAV = *reinterpret_cast<const FRDGTextureUAV* const*>(Base + ParameterBinding.ByteOffset);
 
 		if (DO_CHECK)
 		{
@@ -149,7 +149,7 @@ inline void SetShaderParameters(TRHICmdList& RHICmdList, const TShaderClass* Sha
 	// Graph Textures
 	for (const FShaderParameterBindings::FResourceParameter& ParameterBinding : Bindings.GraphTextures)
 	{
-		auto GraphTexture = *reinterpret_cast<const FGraphTexture* const*>(Base + ParameterBinding.ByteOffset);
+		auto GraphTexture = *reinterpret_cast<const FRDGTexture* const*>(Base + ParameterBinding.ByteOffset);
 
 		if (DO_CHECK)
 		{
@@ -168,7 +168,7 @@ inline void SetShaderParameters(TRHICmdList& RHICmdList, const TShaderClass* Sha
 	// Graph SRVs
 	for (const FShaderParameterBindings::FResourceParameter& ParameterBinding : Bindings.GraphSRVs)
 	{
-		auto GraphSRV = *reinterpret_cast<const FGraphSRV* const*>(Base + ParameterBinding.ByteOffset);
+		auto GraphSRV = *reinterpret_cast<const FRDGTextureSRV* const*>(Base + ParameterBinding.ByteOffset);
 
 		if (DO_CHECK)
 		{
diff --git a/Engine/Source/Runtime/Renderer/Private/DeferredShadingRenderer.h b/Engine/Source/Runtime/Renderer/Private/DeferredShadingRenderer.h
--- a/Engine/Source/Runtime/Renderer/Private/DeferredShadingRenderer.h
+++ b/Engine/Source/Runtime/Renderer/Private/DeferredShadingRenderer.h
@@ -395,7 +395,7 @@ private:
 	void PostInitViewCustomData(FGraphEventArray& OutUpdateEvents);
 
 	void RenderLocalLightsForVolumetricFog(
-		FRenderGraphBuilder& GraphBuilder,
+		FRDGBuilder& GraphBuilder,
 		FViewInfo& View,
 		bool bUseTemporalReprojection,
 		const struct FVolumetricFogIntegrationParameterData& IntegrationData,
@@ -403,19 +403,19 @@ private:
 		FIntVector VolumetricFogGridSize,
 		FVector GridZParams,
 		const FPooledRenderTargetDesc& VolumeDesc,
-		const FGraphTexture*& OutLocalShadowedLightScattering);
+		const FRDGTexture*& OutLocalShadowedLightScattering);
 
 	void RenderLightFunctionForVolumetricFog(
-		FRenderGraphBuilder& GraphBuilder,
+		FRDGBuilder& GraphBuilder,
 		FViewInfo& View,
 		FIntVector VolumetricFogGridSize,
 		float VolumetricFogMaxDistance,
 		FMatrix& OutLightFunctionWorldToShadow,
-		const FGraphTexture*& OutLightFunctionTexture,
+		const FRDGTexture*& OutLightFunctionTexture,
 		bool& bOutUseDirectionalLightShadowing);
 
 	void VoxelizeFogVolumePrimitives(
-		FRenderGraphBuilder& GraphBuilder,
+		FRDGBuilder& GraphBuilder,
 		const FViewInfo& View,
 		const FVolumetricFogIntegrationParameterData& IntegrationData,
 		FIntVector VolumetricFogGridSize,
diff --git a/Engine/Source/Runtime/Renderer/Private/SceneOcclusion.cpp b/Engine/Source/Runtime/Renderer/Private/SceneOcclusion.cpp
--- a/Engine/Source/Runtime/Renderer/Private/SceneOcclusion.cpp
+++ b/Engine/Source/Runtime/Renderer/Private/SceneOcclusion.cpp
@@ -1093,7 +1093,7 @@ void BuildHZB(FRHICommandListImmediate& RHICmdList, FViewInfo& View)
 {
 	QUICK_SCOPE_CYCLE_COUNTER(STAT_BuildHZB);
 	
-	FRenderGraphBuilder GraphBuilder(RHICmdList);
+	FRDGBuilder GraphBuilder(RHICmdList);
 
 	// View.ViewRect.{Width,Height}() are most likely to be < 2^24, so the float
 	// conversion won't loss any precision (assuming float have 23bits for mantissa)
@@ -1111,7 +1111,7 @@ void BuildHZB(FRHICommandListImmediate& RHICmdList, FViewInfo& View)
 
 	//@DW: Explicit creation of graph resource handles - full support for everything the RHI supports
 	//@DW: Now that we've created a resource handle, it will have to be passed around to other passes for manual wiring or put into a Blackboard structure for automatic wiring
-	const FGraphTexture* HZBTexture = GraphBuilder.CreateTexture(HZBDesc, TEXT("HZB"));
+	const FRDGTexture* HZBTexture = GraphBuilder.CreateTexture(HZBDesc, TEXT("HZB"));
 
 	{
 		FHZBBuildPassParameters* PassParameters;
@@ -1194,8 +1194,8 @@ void BuildHZB(FRHICommandListImmediate& RHICmdList, FViewInfo& View)
 		DstSize.Y = FMath::Max(DstSize.Y, 1);
 
 		//@DW: Explicit creation of SRV, full configuration of SRV supported
-		FGraphSRVDesc Desc(HZBTexture, MipIndex - 1);
-		const FGraphSRV* ParentMipSRV = GraphBuilder.CreateSRV(Desc);
+		FRDGTextureSRVDesc Desc(HZBTexture, MipIndex - 1);
+		const FRDGTextureSRV* ParentMipSRV = GraphBuilder.CreateSRV(Desc);
 
 		FHZBBuildPassParameters* PassParameters;
 		GraphBuilder.CreateParameters(&PassParameters);
diff --git a/Engine/Source/Runtime/Renderer/Private/VolumetricFog.cpp b/Engine/Source/Runtime/Renderer/Private/VolumetricFog.cpp
--- a/Engine/Source/Runtime/Renderer/Private/VolumetricFog.cpp
+++ b/Engine/Source/Runtime/Renderer/Private/VolumetricFog.cpp
@@ -463,7 +463,7 @@ public:
 TGlobalResource<FCircleRasterizeIndexBuffer> GCircleRasterizeIndexBuffer;
 
 void FDeferredShadingSceneRenderer::RenderLocalLightsForVolumetricFog(
-	FRenderGraphBuilder& GraphBuilder,
+	FRDGBuilder& GraphBuilder,
 	FViewInfo& View,
 	bool bUseTemporalReprojection,
 	const FVolumetricFogIntegrationParameterData& IntegrationData,
@@ -471,7 +471,7 @@ void FDeferredShadingSceneRenderer::RenderLocalLightsForVolumetricFog(
 	FIntVector VolumetricFogGridSize, 
 	FVector GridZParams,
 	const FPooledRenderTargetDesc& VolumeDesc,
-	const FGraphTexture*& OutLocalShadowedLightScattering)
+	const FRDGTexture*& OutLocalShadowedLightScattering)
 {
 	TArray<const FLightSceneInfo*, SceneRenderingAllocator> LightsToInject;
 
@@ -1000,10 +1000,10 @@ void FDeferredShadingSceneRenderer::ComputeVolumetricFog(FRHICommandListImmediat
 
 			FMatrix LightFunctionWorldToShadow;
 
-			FRenderGraphBuilder GraphBuilder(RHICmdList);
+			FRDGBuilder GraphBuilder(RHICmdList);
 
 			//@DW - register GWhiteTexture as a graph external for when there's no light function - later a shader is going to bind it whether we rendered to it or not
-			const FGraphTexture* LightFunctionTexture = GraphBuilder.RegisterExternalTexture(GSystemTextures.WhiteDummy);
+			const FRDGTexture* LightFunctionTexture = GraphBuilder.RegisterExternalTexture(GSystemTextures.WhiteDummy);
 			bool bUseDirectionalLightShadowing;
 
 			RenderLightFunctionForVolumetricFog(
@@ -1024,8 +1024,8 @@ void FDeferredShadingSceneRenderer::ComputeVolumetricFog(FRHICommandListImmediat
 			//@DW - Passing these around in a struct to ease manual wiring
 			IntegrationData.VBufferA = GraphBuilder.CreateTexture(VolumeDescFastVRAM, TEXT("VBufferA"));
 			IntegrationData.VBufferB = GraphBuilder.CreateTexture(VolumeDescFastVRAM, TEXT("VBufferB"));
-			IntegrationData.VBufferA_UAV = GraphBuilder.CreateUAV(FGraphUAVDesc(IntegrationData.VBufferA));
-			IntegrationData.VBufferB_UAV = GraphBuilder.CreateUAV(FGraphUAVDesc(IntegrationData.VBufferB));
+			IntegrationData.VBufferA_UAV = GraphBuilder.CreateUAV(FRDGTextureUAVDesc(IntegrationData.VBufferA));
+			IntegrationData.VBufferB_UAV = GraphBuilder.CreateUAV(FRDGTextureUAVDesc(IntegrationData.VBufferB));
 
 			{
 				FVolumetricFogMaterialSetupCS::FParameters* PassParameters;
@@ -1072,11 +1072,11 @@ void FDeferredShadingSceneRenderer::ComputeVolumetricFog(FRHICommandListImmediat
 					FogInfo.VolumetricFogDistance);
 			}
 
-			const FGraphTexture* LocalShadowedLightScattering = GraphBuilder.RegisterExternalTexture(GSystemTextures.VolumetricBlackDummy);
+			const FRDGTexture* LocalShadowedLightScattering = GraphBuilder.RegisterExternalTexture(GSystemTextures.VolumetricBlackDummy);
 			RenderLocalLightsForVolumetricFog(GraphBuilder, View, bUseTemporalReprojection, IntegrationData, FogInfo, VolumetricFogGridSize, GridZParams, VolumeDescFastVRAM, LocalShadowedLightScattering);
 
 			IntegrationData.LightScattering = GraphBuilder.CreateTexture(VolumeDesc, TEXT("LightScattering"));
-			IntegrationData.LightScatteringUAV = GraphBuilder.CreateUAV(FGraphUAVDesc(IntegrationData.LightScattering));
+			IntegrationData.LightScatteringUAV = GraphBuilder.CreateUAV(FRDGTextureUAVDesc(IntegrationData.LightScattering));
 
 			{
 				TVolumetricFogLightScatteringCS::FParameters* PassParameters;
@@ -1137,8 +1137,8 @@ void FDeferredShadingSceneRenderer::ComputeVolumetricFog(FRHICommandListImmediat
 					});
 			}
 
-			const FGraphTexture* IntegratedLightScattering = GraphBuilder.CreateTexture(VolumeDesc, TEXT("IntegratedLightScattering"));
-			const FGraphUAV* IntegratedLightScatteringUAV = GraphBuilder.CreateUAV(FGraphUAVDesc(IntegratedLightScattering));
+			const FRDGTexture* IntegratedLightScattering = GraphBuilder.CreateTexture(VolumeDesc, TEXT("IntegratedLightScattering"));
+			const FRDGTextureUAV* IntegratedLightScatteringUAV = GraphBuilder.CreateUAV(FRDGTextureUAVDesc(IntegratedLightScattering));
 
 			{
 				FVolumetricFogFinalIntegrationCS::FParameters* PassParameters;
diff --git a/Engine/Source/Runtime/Renderer/Private/VolumetricFogLightFunction.cpp b/Engine/Source/Runtime/Renderer/Private/VolumetricFogLightFunction.cpp
--- a/Engine/Source/Runtime/Renderer/Private/VolumetricFogLightFunction.cpp
+++ b/Engine/Source/Runtime/Renderer/Private/VolumetricFogLightFunction.cpp
@@ -97,12 +97,12 @@ private:
 IMPLEMENT_MATERIAL_SHADER_TYPE(,FVolumetricFogLightFunctionPS,TEXT("/Engine/Private/VolumetricFogLightFunction.usf"),TEXT("Main"),SF_Pixel);
 
 void FDeferredShadingSceneRenderer::RenderLightFunctionForVolumetricFog(
-	FRenderGraphBuilder& GraphBuilder,
+	FRDGBuilder& GraphBuilder,
 	FViewInfo& View, 
 	FIntVector VolumetricFogGridSize,
 	float VolumetricFogMaxDistance,
 	FMatrix& OutLightFunctionWorldToShadow,
-	const FGraphTexture*& OutLightFunctionTexture,
+	const FRDGTexture*& OutLightFunctionTexture,
 	bool& bOutUseDirectionalLightShadowing)
 {
 	OutLightFunctionWorldToShadow = FMatrix::Identity;
diff --git a/Engine/Source/Runtime/Renderer/Private/VolumetricFogShared.h b/Engine/Source/Runtime/Renderer/Private/VolumetricFogShared.h
--- a/Engine/Source/Runtime/Renderer/Private/VolumetricFogShared.h
+++ b/Engine/Source/Runtime/Renderer/Private/VolumetricFogShared.h
@@ -20,13 +20,13 @@ struct FVolumetricFogIntegrationParameterData
 
 	bool bTemporalHistoryIsValid;
 	TArray<FVector4, TInlineAllocator<16>> FrameJitterOffsetValues;
-	const FGraphTexture* VBufferA;
-	const FGraphTexture* VBufferB;
-	const FGraphUAV* VBufferA_UAV;
-	const FGraphUAV* VBufferB_UAV;
+	const FRDGTexture* VBufferA;
+	const FRDGTexture* VBufferB;
+	const FRDGTextureUAV* VBufferA_UAV;
+	const FRDGTextureUAV* VBufferB_UAV;
 
-	const FGraphTexture* LightScattering;
-	const FGraphUAV* LightScatteringUAV;
+	const FRDGTexture* LightScattering;
+	const FRDGTextureUAV* LightScatteringUAV;
 };
 
 /**  */
diff --git a/Engine/Source/Runtime/Renderer/Private/VolumetricFogVoxelization.cpp b/Engine/Source/Runtime/Renderer/Private/VolumetricFogVoxelization.cpp
--- a/Engine/Source/Runtime/Renderer/Private/VolumetricFogVoxelization.cpp
+++ b/Engine/Source/Runtime/Renderer/Private/VolumetricFogVoxelization.cpp
@@ -681,7 +681,7 @@ void VoxelizeVolumePrimitive(
 }
 
 void FDeferredShadingSceneRenderer::VoxelizeFogVolumePrimitives(
-	FRenderGraphBuilder& GraphBuilder,
+	FRDGBuilder& GraphBuilder,
 	const FViewInfo& View,
 	const FVolumetricFogIntegrationParameterData& IntegrationData,
 	FIntVector VolumetricFogGridSize,
