Makes static analysis happy in render graph.

diff --git a/Engine/Source/Runtime/RenderCore/Private/ShaderParameterMetadata.cpp b/Engine/Source/Runtime/RenderCore/Private/ShaderParameterMetadata.cpp
--- a/Engine/Source/Runtime/RenderCore/Private/ShaderParameterMetadata.cpp
+++ b/Engine/Source/Runtime/RenderCore/Private/ShaderParameterMetadata.cpp
@@ -145,6 +145,11 @@ void FShaderParametersMetadata::InitializeLayout()
 				}
 			}
 
+			if (BaseType == UBMT_NESTED_STRUCT || BaseType == UBMT_INCLUDED_STRUCT)
+			{
+				check(ChildStruct);
+			}
+
 			if (UseCase == EUseCase::ShaderParameterStruct &&
 				(BaseType == UBMT_NESTED_STRUCT || BaseType == UBMT_INCLUDED_STRUCT) &&
 				ChildStruct->GetUseCase() != EUseCase::ShaderParameterStruct)
diff --git a/Engine/Source/Runtime/RenderCore/Private/ShaderParameterStruct.cpp b/Engine/Source/Runtime/RenderCore/Private/ShaderParameterStruct.cpp
--- a/Engine/Source/Runtime/RenderCore/Private/ShaderParameterStruct.cpp
+++ b/Engine/Source/Runtime/RenderCore/Private/ShaderParameterStruct.cpp
@@ -132,7 +132,7 @@ struct FShaderParameterStructBindingContext
 					Bindings->GraphTextures.Add(Parameter);
 				else if (BaseType == UBMT_GRAPH_TRACKED_SRV)
 					Bindings->GraphSRVs.Add(Parameter);
-				else if (BaseType == UBMT_GRAPH_TRACKED_UAV)
+				else // if (BaseType == UBMT_GRAPH_TRACKED_UAV)
 					Bindings->GraphUAVs.Add(Parameter);
 			}
 			else
diff --git a/Engine/Source/Runtime/RenderCore/Private/ShaderParameters.cpp b/Engine/Source/Runtime/RenderCore/Private/ShaderParameters.cpp
--- a/Engine/Source/Runtime/RenderCore/Private/ShaderParameters.cpp
+++ b/Engine/Source/Runtime/RenderCore/Private/ShaderParameters.cpp
@@ -264,7 +264,7 @@ static void CreateHLSLUniformBufferStructMembersDeclaration(
 
 		if (IsUniformBufferResourceType(Member.GetBaseType()))
 		{
-			check(Member.GetBaseType() != UBMT_GRAPH_TRACKED_SRV || Member.GetBaseType() != UBMT_GRAPH_TRACKED_UAV);
+			check(Member.GetBaseType() != UBMT_GRAPH_TRACKED_SRV && Member.GetBaseType() != UBMT_GRAPH_TRACKED_UAV);
 			if (Member.GetBaseType() == UBMT_SRV)
 			{
 				// TODO: handle arrays?
diff --git a/Engine/Source/Runtime/RenderCore/Public/ShaderParameterStruct.h b/Engine/Source/Runtime/RenderCore/Public/ShaderParameterStruct.h
--- a/Engine/Source/Runtime/RenderCore/Public/ShaderParameterStruct.h
+++ b/Engine/Source/Runtime/RenderCore/Public/ShaderParameterStruct.h
@@ -67,9 +67,11 @@ inline void SetShaderUAVs(TRHICmdList& RHICmdList, const TShaderClass* Shader, F
 			{
 				SetNullShaderParameterFatalError(Shader, TShaderClass::FParameters::FTypeInfo::GetStructMetadata(), ParameterBinding.ByteOffset);
 			}
-
-			// Mark this resource as used by the pass for inefficient pass resource dependency debugging purpose.
-			GraphUAV->bIsActuallyUsedByPass = true;
+			else
+			{
+				// Mark this resource as used by the pass for inefficient pass resource dependency debugging purpose.
+				GraphUAV->bIsActuallyUsedByPass = true;
+			}
 		}
 
 		RHICmdList.SetUAVParameter(ShadeRHI, ParameterBinding.BaseIndex, GraphUAV->GetRHIUnorderedAccessView());
@@ -157,9 +159,11 @@ inline void SetShaderParameters(TRHICmdList& RHICmdList, const TShaderClass* Sha
 			{
 				SetNullShaderParameterFatalError(Shader, TShaderClass::FParameters::FTypeInfo::GetStructMetadata(), ParameterBinding.ByteOffset);
 			}
-
-			// Mark this resource as used by the pass for inefficient pass resource dependency debugging purpose.
-			GraphTexture->bIsActuallyUsedByPass = true;
+			else
+			{
+				// Mark this resource as used by the pass for inefficient pass resource dependency debugging purpose.
+				GraphTexture->bIsActuallyUsedByPass = true;
+			}
 		}
 
 		RHICmdList.SetShaderTexture(ShadeRHI, ParameterBinding.BaseIndex, GraphTexture->GetRHITexture());
@@ -176,9 +180,11 @@ inline void SetShaderParameters(TRHICmdList& RHICmdList, const TShaderClass* Sha
 			{
 				SetNullShaderParameterFatalError(Shader, TShaderClass::FParameters::FTypeInfo::GetStructMetadata(), ParameterBinding.ByteOffset);
 			}
-
-			// Mark this resource as used by the pass for inefficient pass resource dependency debugging purpose.
-			GraphSRV->bIsActuallyUsedByPass = true;
+			else
+			{
+				// Mark this resource as used by the pass for inefficient pass resource dependency debugging purpose.
+				GraphSRV->bIsActuallyUsedByPass = true;
+			}
 		}
 
 		RHICmdList.SetShaderResourceViewParameter(ShadeRHI, ParameterBinding.BaseIndex, GraphSRV->GetRHIShaderResourceView());
