Implements FRDGBuffer, FRDGBufferSRV & FRDGBufferUAV

diff --git a/Engine/Source/Developer/ShaderCompilerCommon/Private/ShaderCompilerCommon.cpp b/Engine/Source/Developer/ShaderCompilerCommon/Private/ShaderCompilerCommon.cpp
--- a/Engine/Source/Developer/ShaderCompilerCommon/Private/ShaderCompilerCommon.cpp
+++ b/Engine/Source/Developer/ShaderCompilerCommon/Private/ShaderCompilerCommon.cpp
@@ -139,9 +139,11 @@ bool BuildResourceTableMapping(
 				break;
 			case UBMT_SRV:
 			case UBMT_GRAPH_TRACKED_SRV:
+			case UBMT_GRAPH_TRACKED_BUFFER_SRV:
 				OutSRT.ShaderResourceViewMap.Add(ResourceMap);
 				break;
 			case UBMT_GRAPH_TRACKED_UAV:
+			case UBMT_GRAPH_TRACKED_BUFFER_UAV:
 				OutSRT.UnorderedAccessViewMap.Add(ResourceMap);
 				break;
 			default:
diff --git a/Engine/Source/Runtime/RHI/Public/RHIDefinitions.h b/Engine/Source/Runtime/RHI/Public/RHIDefinitions.h
--- a/Engine/Source/Runtime/RHI/Public/RHIDefinitions.h
+++ b/Engine/Source/Runtime/RHI/Public/RHIDefinitions.h
@@ -398,6 +398,9 @@ enum EUniformBufferBaseType
 	UBMT_GRAPH_TRACKED_TEXTURE,
 	UBMT_GRAPH_TRACKED_SRV,
 	UBMT_GRAPH_TRACKED_UAV,
+	UBMT_GRAPH_TRACKED_BUFFER,
+	UBMT_GRAPH_TRACKED_BUFFER_SRV,
+	UBMT_GRAPH_TRACKED_BUFFER_UAV,
 
 	// Nested structure.
 	UBMT_NESTED_STRUCT,
@@ -412,7 +415,7 @@ enum EUniformBufferBaseType
 	UBMT_RENDER_TARGET_BINDING_SLOTS,
 
 	EUniformBufferBaseType_Num,
-	EUniformBufferBaseType_NumBits = 4,
+	EUniformBufferBaseType_NumBits = 5,
 };
 static_assert(EUniformBufferBaseType_Num <= (1 << EUniformBufferBaseType_NumBits), "EUniformBufferBaseType_Num will not fit on EUniformBufferBaseType_NumBits");
 
@@ -1096,6 +1099,9 @@ inline bool IsUniformBufferResourceType(EUniformBufferBaseType BaseType)
 		BaseType == UBMT_GRAPH_TRACKED_TEXTURE ||
 		BaseType == UBMT_GRAPH_TRACKED_SRV ||
 		BaseType == UBMT_GRAPH_TRACKED_UAV ||
+		BaseType == UBMT_GRAPH_TRACKED_BUFFER ||
+		BaseType == UBMT_GRAPH_TRACKED_BUFFER_SRV ||
+		BaseType == UBMT_GRAPH_TRACKED_BUFFER_UAV ||
 		BaseType == UBMT_RENDER_TARGET_BINDING_SLOTS;
 }
 
diff --git a/Engine/Source/Runtime/RenderCore/Private/RenderGraphBuilder.cpp b/Engine/Source/Runtime/RenderCore/Private/RenderGraphBuilder.cpp
--- a/Engine/Source/Runtime/RenderCore/Private/RenderGraphBuilder.cpp
+++ b/Engine/Source/Runtime/RenderCore/Private/RenderGraphBuilder.cpp
@@ -3,6 +3,7 @@
 #include "RenderGraphBuilder.h"
 #include "RenderCore.h"
 #include "RenderTargetPool.h"
+#include "RenderGraphResourcePool.h"
 
 
 #if RENDER_GRAPH_DEBUGGING
@@ -78,8 +79,9 @@ void FRDGBuilder::ValidatePass(const FRenderGraphPass* Pass) const
 		switch (Type)
 		{
 		case UBMT_GRAPH_TRACKED_UAV:
+		case UBMT_GRAPH_TRACKED_BUFFER_UAV:
 		{
-			FRDGTextureUAV* RESTRICT UAV = *ParameterStruct.GetMemberPtrAtOffset<FRDGTextureUAV*>(Offset);
+			FRDGResource* RESTRICT UAV = *ParameterStruct.GetMemberPtrAtOffset<FRDGResource*>(Offset);
 			if (UAV && !bCanUseUAVs)
 			{
 				UE_LOG(LogRendererCore, Warning, TEXT("UAV can only been bound to compute shaders, therefore UAV %s is certainly useless for passs %s."), UAV->Name, Pass->GetName());
@@ -144,11 +146,12 @@ void FRDGBuilder::WalkGraphDependencies()
 			switch (Type)
 			{
 			case UBMT_GRAPH_TRACKED_TEXTURE:
+			case UBMT_GRAPH_TRACKED_BUFFER:
 			{
-				FRDGTexture* RESTRICT Texture = *ParameterStruct.GetMemberPtrAtOffset<FRDGTexture*>(Offset);
-				if (Texture)
+				FRDGResource* RESTRICT Resource = *ParameterStruct.GetMemberPtrAtOffset<FRDGResource*>(Offset);
+				if (Resource)
 				{
-					Texture->ReferenceCount++;
+					Resource->ReferenceCount++;
 				}
 			}
 			break;
@@ -170,6 +173,24 @@ void FRDGBuilder::WalkGraphDependencies()
 				}
 			}
 			break;
+			case UBMT_GRAPH_TRACKED_BUFFER_SRV:
+			{
+				FRDGBufferSRV* RESTRICT SRV = *ParameterStruct.GetMemberPtrAtOffset<FRDGBufferSRV*>(Offset);
+				if (SRV)
+				{
+					SRV->Desc.Buffer->ReferenceCount++;
+				}
+			}
+			break;
+			case UBMT_GRAPH_TRACKED_BUFFER_UAV:
+			{
+				FRDGBufferUAV* RESTRICT UAV = *ParameterStruct.GetMemberPtrAtOffset<FRDGBufferUAV*>(Offset);
+				if (UAV)
+				{
+					UAV->Desc.Buffer->ReferenceCount++;
+				}
+			}
+			break;
 			case UBMT_RENDER_TARGET_BINDING_SLOTS:
 			{
 				FRenderTargetBindingSlots* RESTRICT RenderTargets = ParameterStruct.GetMemberPtrAtOffset<FRenderTargetBindingSlots>(Offset);
@@ -235,6 +256,126 @@ void FRDGBuilder::AllocateRHITextureIfNeeded(const FRDGTexture* Texture, bool bC
 	Texture->PooledRenderTarget = PooledRenderTarget;
 }
 
+void FRDGBuilder::AllocateRHITextureSRVIfNeeded(const FRDGTextureSRV* SRV, bool bComputePass)
+{
+	check(SRV);
+
+	if (SRV->CachedRHI.SRV)
+	{
+		return;
+	}
+
+	AllocateRHITextureIfNeeded(SRV->Desc.Texture, bComputePass);
+
+	SRV->CachedRHI.SRV = SRV->Desc.Texture->PooledRenderTarget->GetRenderTargetItem().MipSRVs[SRV->Desc.MipLevel];
+}
+
+void FRDGBuilder::AllocateRHITextureUAVIfNeeded(const FRDGTextureUAV* UAV, bool bComputePass)
+{
+	check(UAV);
+
+	if (UAV->CachedRHI.UAV)
+	{
+		return;
+	}
+
+	AllocateRHITextureIfNeeded(UAV->Desc.Texture, bComputePass);
+
+	UAV->CachedRHI.UAV = UAV->Desc.Texture->PooledRenderTarget->GetRenderTargetItem().MipUAVs[UAV->Desc.MipLevel];
+}
+
+void FRDGBuilder::AllocateRHIBufferIfNeeded(const FRDGBuffer* Buffer, bool bComputePass)
+{
+	check(Buffer);
+
+	if (Buffer->PooledBuffer)
+	{
+		return;
+	}
+
+	check(Buffer->ReferenceCount > 0 || GRenderGraphImmediateMode);
+
+	TRefCountPtr<FPooledRDGBuffer>& AllocatedBuffer = AllocatedBuffers.FindOrAdd(Buffer);
+	GRenderGraphResourcePool.FindFreeBuffer(RHICmdList, Buffer->Desc, AllocatedBuffer, Buffer->Name);
+
+	Buffer->PooledBuffer = AllocatedBuffer;
+}
+
+void FRDGBuilder::AllocateRHIBufferSRVIfNeeded(const FRDGBufferSRV* SRV, bool bComputePass)
+{
+	check(SRV);
+
+	if (SRV->CachedRHI.SRV)
+	{
+		return;
+	}
+
+	AllocateRHIBufferIfNeeded(SRV->Desc.Buffer, bComputePass);
+
+	if (SRV->Desc.Buffer->PooledBuffer->SRVs.Contains(SRV->Desc))
+	{
+		SRV->CachedRHI.SRV = SRV->Desc.Buffer->PooledBuffer->SRVs[SRV->Desc];
+		return;
+	}
+
+	FShaderResourceViewRHIRef RHIShaderResourceView;
+
+	if (SRV->Desc.Buffer->Desc.UnderlyingType == FRDGBufferDesc::EUnderlyingType::VertexBuffer)
+	{
+		RHIShaderResourceView = RHICreateShaderResourceView(SRV->Desc.Buffer->PooledBuffer->VertexBuffer, SRV->Desc.BytesPerElement, SRV->Desc.Format);
+	}
+	else if (SRV->Desc.Buffer->Desc.UnderlyingType == FRDGBufferDesc::EUnderlyingType::VertexBuffer)
+	{
+		RHIShaderResourceView = RHICreateShaderResourceView(SRV->Desc.Buffer->PooledBuffer->StructuredBuffer);
+	}
+	else
+	{
+		check(0);
+	}
+
+	SRV->CachedRHI.SRV = RHIShaderResourceView;
+	SRV->Desc.Buffer->PooledBuffer->SRVs.Add(SRV->Desc, RHIShaderResourceView);
+}
+
+void FRDGBuilder::AllocateRHIBufferUAVIfNeeded(const FRDGBufferUAV* UAV, bool bComputePass)
+{
+	check(UAV);
+
+	if (UAV->CachedRHI.UAV)
+	{
+		return;
+	}
+
+	AllocateRHIBufferIfNeeded(UAV->Desc.Buffer, bComputePass);
+
+	if (UAV->Desc.Buffer->PooledBuffer->UAVs.Contains(UAV->Desc))
+	{
+		UAV->CachedRHI.UAV = UAV->Desc.Buffer->PooledBuffer->UAVs[UAV->Desc];
+		return;
+	}
+
+	// Hack to make sure only one UAVs is arround.
+	UAV->Desc.Buffer->PooledBuffer->UAVs.Empty();
+
+	FUnorderedAccessViewRHIRef RHIUnorderedAccessView;
+
+	if (UAV->Desc.Buffer->Desc.UnderlyingType == FRDGBufferDesc::EUnderlyingType::VertexBuffer)
+	{
+		RHIUnorderedAccessView = RHICreateUnorderedAccessView(UAV->Desc.Buffer->PooledBuffer->VertexBuffer, UAV->Desc.Format);
+	}
+	else if (UAV->Desc.Buffer->Desc.UnderlyingType == FRDGBufferDesc::EUnderlyingType::VertexBuffer)
+	{
+		RHIUnorderedAccessView = RHICreateUnorderedAccessView(UAV->Desc.Buffer->PooledBuffer->StructuredBuffer, UAV->Desc.bSupportsAtomicCounter, UAV->Desc.bSupportsAppendBuffer);
+	}
+	else
+	{
+		check(0);
+	}
+
+	UAV->CachedRHI.UAV = RHIUnorderedAccessView;
+	UAV->Desc.Buffer->PooledBuffer->UAVs.Add(UAV->Desc, RHIUnorderedAccessView);
+}
+
 static EResourceTransitionPipeline CalcTransitionPipeline(bool bCurrentCompute, bool bTargetCompute)
 {
 	// TODO(RDG) convert table to math
@@ -264,16 +405,16 @@ void FRDGBuilder::TransitionTexture( const FRDGTexture* Texture, EResourceTransi
 	}
 }
 
-void FRDGBuilder::TransitionUAV( const FRDGTextureUAV* UAV, EResourceTransitionAccess TransitionAccess, bool bRequiredCompute ) const
+void FRDGBuilder::TransitionUAV(FUnorderedAccessViewRHIParamRef UAV, const FRDGResource* UnderlyingResource, EResourceTransitionAccess TransitionAccess, bool bRequiredCompute ) const
 {
 	const bool bRequiredWritable = true;
 
-	if( UAV->Desc.Texture->bWritable != bRequiredWritable || UAV->Desc.Texture->bCompute != bRequiredCompute )
+	if(UnderlyingResource->bWritable != bRequiredWritable || UnderlyingResource->bCompute != bRequiredCompute )
 	{
-		EResourceTransitionPipeline TransitionPipeline = CalcTransitionPipeline( UAV->Desc.Texture->bCompute, bRequiredCompute );
-		RHICmdList.TransitionResource( TransitionAccess, TransitionPipeline, UAV->GetRHIUnorderedAccessView());
-		UAV->Desc.Texture->bWritable = bRequiredWritable;
-		UAV->Desc.Texture->bCompute = bRequiredCompute;
+		EResourceTransitionPipeline TransitionPipeline = CalcTransitionPipeline(UnderlyingResource->bCompute, bRequiredCompute );
+		RHICmdList.TransitionResource( TransitionAccess, TransitionPipeline, UAV);
+		UnderlyingResource->bWritable = bRequiredWritable;
+		UnderlyingResource->bCompute = bRequiredCompute;
 	}
 }
 
@@ -379,9 +520,8 @@ void FRDGBuilder::AllocateAndTransitionPassResources(const FRenderGraphPass* Pas
 			{
 				check(SRV->Desc.Texture);
 
-				const FRDGTexture* Texture = SRV->Desc.Texture;
-				AllocateRHITextureIfNeeded(Texture, bIsCompute);
-				TransitionTexture(Texture, EResourceTransitionAccess::EReadable, bIsCompute);
+				AllocateRHITextureSRVIfNeeded(SRV, bIsCompute);
+				TransitionTexture(SRV->Desc.Texture, EResourceTransitionAccess::EReadable, bIsCompute);
 			}
 		}
 		break;
@@ -390,8 +530,46 @@ void FRDGBuilder::AllocateAndTransitionPassResources(const FRenderGraphPass* Pas
 			FRDGTextureUAV* RESTRICT UAV = *ParameterStruct.GetMemberPtrAtOffset<FRDGTextureUAV*>(Offset);
 			if (UAV)
 			{
-				AllocateRHITextureIfNeeded(UAV->Desc.Texture, bIsCompute);
-				TransitionUAV(UAV, EResourceTransitionAccess::EWritable, bIsCompute);
+				AllocateRHITextureUAVIfNeeded(UAV, bIsCompute);
+				TransitionUAV(UAV->CachedRHI.UAV, UAV->Desc.Texture, EResourceTransitionAccess::EWritable, bIsCompute);
+			}
+		}
+		break;
+		case UBMT_GRAPH_TRACKED_BUFFER:
+		{
+			FRDGBuffer* RESTRICT Buffer = *ParameterStruct.GetMemberPtrAtOffset<FRDGBuffer*>(Offset);
+			if (Buffer)
+			{
+				AllocateRHIBufferIfNeeded(Buffer, bIsCompute);
+
+				// TODO(RDG): supper hacky, find the UAV and transition it. Hopefully there is one...
+				check(Buffer->PooledBuffer->UAVs.Num() == 1);
+				FUnorderedAccessViewRHIParamRef UAV = Buffer->PooledBuffer->UAVs.CreateIterator().Value();
+				TransitionUAV(UAV, Buffer, EResourceTransitionAccess::EReadable, bIsCompute);
+			}
+		}
+		break;
+		case UBMT_GRAPH_TRACKED_BUFFER_SRV:
+		{
+			FRDGBufferSRV* RESTRICT SRV = *ParameterStruct.GetMemberPtrAtOffset<FRDGBufferSRV*>(Offset);
+			if (SRV)
+			{
+				AllocateRHIBufferSRVIfNeeded(SRV, bIsCompute);
+
+				// TODO(RDG): supper hacky, find the UAV and transition it. Hopefully there is one...
+				check(SRV->Desc.Buffer->PooledBuffer->UAVs.Num() == 1);
+				FUnorderedAccessViewRHIParamRef UAV = SRV->Desc.Buffer->PooledBuffer->UAVs.CreateIterator().Value();
+				TransitionUAV(UAV, SRV->Desc.Buffer, EResourceTransitionAccess::EReadable, bIsCompute);
+			}
+		}
+		break;
+		case UBMT_GRAPH_TRACKED_BUFFER_UAV:
+		{
+			FRDGBufferUAV* RESTRICT UAV = *ParameterStruct.GetMemberPtrAtOffset<FRDGBufferUAV*>(Offset);
+			if (UAV)
+			{
+				AllocateRHIBufferUAVIfNeeded(UAV, bIsCompute);
+				TransitionUAV(UAV->CachedRHI.UAV, UAV->Desc.Buffer, EResourceTransitionAccess::EWritable, bIsCompute);
 			}
 		}
 		break;
@@ -476,7 +654,8 @@ void FRDGBuilder::WarnForUselessPassDependencies(const FRenderGraphPass* Pass)
 		uint8 Type = ParameterStruct.Layout->Resources[ResourceIndex];
 		uint16 Offset = ParameterStruct.Layout->ResourceOffsets[ResourceIndex];
 
-		if (Type != UBMT_GRAPH_TRACKED_TEXTURE && Type != UBMT_GRAPH_TRACKED_SRV && Type != UBMT_GRAPH_TRACKED_UAV)
+		if (Type != UBMT_GRAPH_TRACKED_TEXTURE && Type != UBMT_GRAPH_TRACKED_SRV && Type != UBMT_GRAPH_TRACKED_UAV &&
+			Type != UBMT_GRAPH_TRACKED_BUFFER && Type != UBMT_GRAPH_TRACKED_BUFFER_SRV && Type != UBMT_GRAPH_TRACKED_BUFFER_UAV)
 			continue;
 
 		const FRDGResource* Resource = *ParameterStruct.GetMemberPtrAtOffset<const FRDGResource*>(Offset);
@@ -497,7 +676,8 @@ void FRDGBuilder::WarnForUselessPassDependencies(const FRenderGraphPass* Pass)
 			uint8 Type = ParameterStruct.Layout->Resources[ResourceIndex];
 			uint16 Offset = ParameterStruct.Layout->ResourceOffsets[ResourceIndex];
 
-			if (Type != UBMT_GRAPH_TRACKED_TEXTURE && Type != UBMT_GRAPH_TRACKED_SRV && Type != UBMT_GRAPH_TRACKED_UAV)
+			if (Type != UBMT_GRAPH_TRACKED_TEXTURE && Type != UBMT_GRAPH_TRACKED_SRV && Type != UBMT_GRAPH_TRACKED_UAV &&
+				Type != UBMT_GRAPH_TRACKED_BUFFER && Type != UBMT_GRAPH_TRACKED_BUFFER_SRV && Type != UBMT_GRAPH_TRACKED_BUFFER_UAV)
 				continue;
 
 			const FRDGResource* Resource = *ParameterStruct.GetMemberPtrAtOffset<const FRDGResource*>(Offset);
@@ -505,7 +685,8 @@ void FRDGBuilder::WarnForUselessPassDependencies(const FRenderGraphPass* Pass)
 			if (!Resource)
 				continue;
 
-			UE_LOG(LogRendererCore, Warning, TEXT("	%s"), Resource->Name);
+			if (!Resource->bIsActuallyUsedByPass)
+				UE_LOG(LogRendererCore, Warning, TEXT("	%s"), Resource->Name);
 		}
 	}
 
@@ -515,7 +696,8 @@ void FRDGBuilder::WarnForUselessPassDependencies(const FRenderGraphPass* Pass)
 		uint8 Type = ParameterStruct.Layout->Resources[ResourceIndex];
 		uint16 Offset = ParameterStruct.Layout->ResourceOffsets[ResourceIndex];
 
-		if (Type != UBMT_GRAPH_TRACKED_TEXTURE && Type != UBMT_GRAPH_TRACKED_SRV && Type != UBMT_GRAPH_TRACKED_UAV)
+		if (Type != UBMT_GRAPH_TRACKED_TEXTURE && Type != UBMT_GRAPH_TRACKED_SRV && Type != UBMT_GRAPH_TRACKED_UAV &&
+			Type != UBMT_GRAPH_TRACKED_BUFFER && Type != UBMT_GRAPH_TRACKED_BUFFER_SRV && Type != UBMT_GRAPH_TRACKED_BUFFER_UAV)
 			continue;
 
 		const FRDGResource* Resource = *ParameterStruct.GetMemberPtrAtOffset<const FRDGResource*>(Offset);
@@ -539,11 +721,24 @@ void FRDGBuilder::ReleaseRHITextureIfPossible(const FRDGTexture* Texture)
 	}
 }
 
+void FRDGBuilder::ReleaseRHIBufferIfPossible(const FRDGBuffer* Buffer)
+{
+	check(Buffer->ReferenceCount > 0);
+	Buffer->ReferenceCount--;
+
+	if (Buffer->ReferenceCount == 0)
+	{
+		Buffer->PooledBuffer = nullptr;
+		AllocatedBuffers.FindChecked(Buffer) = nullptr;
+	}
+}
+
 void FRDGBuilder::ReleaseUnecessaryResources(const FRenderGraphPass* Pass)
 {
 	FShaderParameterStructRef ParameterStruct = Pass->GetParameters();
 
 	/** Increments all the FRDGResource::ReferenceCount. */
+	// TODO(RDG): Investigate the cost of branch miss-prediction.
 	for (int ResourceIndex = 0, Num = ParameterStruct.Layout->Resources.Num(); ResourceIndex < Num; ResourceIndex++)
 	{
 		uint8  Type = ParameterStruct.Layout->Resources[ResourceIndex];
@@ -578,6 +773,33 @@ void FRDGBuilder::ReleaseUnecessaryResources(const FRenderGraphPass* Pass)
 			}
 		}
 		break;
+		case UBMT_GRAPH_TRACKED_BUFFER:
+		{
+			FRDGBuffer* RESTRICT Buffer = *ParameterStruct.GetMemberPtrAtOffset<FRDGBuffer*>(Offset);
+			if (Buffer)
+			{
+				ReleaseRHIBufferIfPossible(Buffer);
+			}
+		}
+		break;
+		case UBMT_GRAPH_TRACKED_BUFFER_SRV:
+		{
+			FRDGBufferSRV* RESTRICT SRV = *ParameterStruct.GetMemberPtrAtOffset<FRDGBufferSRV*>(Offset);
+			if (SRV)
+			{
+				ReleaseRHIBufferIfPossible(SRV->Desc.Buffer);
+			}
+		}
+		break;
+		case UBMT_GRAPH_TRACKED_BUFFER_UAV:
+		{
+			FRDGBufferUAV* RESTRICT UAV = *ParameterStruct.GetMemberPtrAtOffset<FRDGBufferUAV*>(Offset);
+			if (UAV)
+			{
+				ReleaseRHIBufferIfPossible(UAV->Desc.Buffer);
+			}
+		}
+		break;
 		case UBMT_RENDER_TARGET_BINDING_SLOTS:
 		{
 			FRenderTargetBindingSlots* RESTRICT RenderTargets = ParameterStruct.GetMemberPtrAtOffset<FRenderTargetBindingSlots>(Offset);
diff --git a/Engine/Source/Runtime/RenderCore/Private/RenderGraphResourcePool.cpp b/Engine/Source/Runtime/RenderCore/Private/RenderGraphResourcePool.cpp
new file mode 100644
--- /dev/null
+++ b/Engine/Source/Runtime/RenderCore/Private/RenderGraphResourcePool.cpp
@@ -0,0 +1,99 @@
+// Copyright 1998-2018 Epic Games, Inc. All Rights Reserved.
+
+/*=============================================================================
+	RenderTargetPool.cpp: Scene render target pool manager.
+=============================================================================*/
+
+#include "RenderGraphResourcePool.h"
+#include "RenderGraphResources.h"
+
+
+uint32 FPooledRDGBuffer::Release()
+{
+	return RefCount--;
+
+	if (RefCount == 0)
+	{
+		VertexBuffer.SafeRelease();
+		IndexBuffer.SafeRelease();
+		StructuredBuffer.SafeRelease();
+		UAVs.Empty();
+		SRVs.Empty();
+	}
+
+	return RefCount;
+}
+
+
+FRenderGraphResourcePool::FRenderGraphResourcePool()
+{ }
+
+
+void FRenderGraphResourcePool::FindFreeBuffer(
+	FRHICommandList& RHICmdList,
+	const FRDGBufferDesc& Desc,
+	TRefCountPtr<FPooledRDGBuffer>& Out,
+	const TCHAR* InDebugName)
+{
+	// First find if available.
+	for (auto& PooledBuffer : AllocatedBuffers)
+	{
+		// Still being used outside the pool.
+		if (PooledBuffer->GetRefCount() > 1)
+		{
+			continue;
+		}
+
+		if (PooledBuffer->Desc == Desc)
+		{
+			Out = PooledBuffer;
+			// TODO(RDG): assign name on RHI.
+			return;
+		}
+	}
+
+	// Allocate new one
+	{
+		Out = new FPooledRDGBuffer;
+		AllocatedBuffers.Add(Out);
+		check(Out->GetRefCount() == 2);
+
+		Out->Desc = Desc;
+
+		uint32 NumBytes = Desc.GetTotalNumBytes();
+
+		FRHIResourceCreateInfo CreateInfo;
+		CreateInfo.DebugName = InDebugName;
+
+		if (Desc.UnderlyingType == FRDGBufferDesc::EUnderlyingType::VertexBuffer)
+		{
+			Out->VertexBuffer = RHICreateVertexBuffer(NumBytes, Desc.Usage, CreateInfo);
+		}
+		else if (Desc.UnderlyingType == FRDGBufferDesc::EUnderlyingType::StructuredBuffer)
+		{
+			Out->StructuredBuffer = RHICreateStructuredBuffer(Desc.BytesPerElement, NumBytes, Desc.Usage, CreateInfo);
+		}
+		else
+		{
+			check(0);
+		}
+	}
+}
+
+void FRenderGraphResourcePool::ReleaseDynamicRHI()
+{
+	AllocatedBuffers.Empty();
+}
+
+void FRenderGraphResourcePool::FreeUnusedResources()
+{
+	for (auto& PooledBuffer : AllocatedBuffers)
+	{
+		if (PooledBuffer.GetRefCount() == 1)
+		{
+			PooledBuffer = nullptr;
+		}
+	}
+}
+
+TGlobalResource<FRenderGraphResourcePool> GRenderGraphResourcePool;
diff --git a/Engine/Source/Runtime/RenderCore/Private/RenderGraphResourcePool.h b/Engine/Source/Runtime/RenderCore/Private/RenderGraphResourcePool.h
new file mode 100644
--- /dev/null
+++ b/Engine/Source/Runtime/RenderCore/Private/RenderGraphResourcePool.h
@@ -0,0 +1,43 @@
+// Copyright 1998-2018 Epic Games, Inc. All Rights Reserved.
+
+/*=============================================================================
+	RenderTargetPool.h: Scene render target pool manager.
+=============================================================================*/
+
+#pragma once
+
+#include "CoreMinimal.h"
+#include "RHI.h"
+#include "RenderResource.h"
+#include "RendererInterface.h"
+#include "RenderGraphResources.h"
+
+
+/**
+ * Pools all resources for the render graph.
+ */
+class RENDERCORE_API FRenderGraphResourcePool : public FRenderResource
+{
+public:
+	FRenderGraphResourcePool();
+
+	/** Allocate a buffer from a given descriptor. */
+	void FindFreeBuffer(
+		FRHICommandList& RHICmdList,
+		const FRDGBufferDesc& Desc,
+		TRefCountPtr<FPooledRDGBuffer>& Out,
+		const TCHAR* InDebugName);
+
+	/** Free renderer resources */
+	virtual void ReleaseDynamicRHI() override;
+
+	/** Good to call between levels or before memory intense operations. */
+	void FreeUnusedResources();
+
+private:
+	/** Elements can be 0, we compact the buffer later. */
+	TArray< TRefCountPtr<FPooledRDGBuffer> > AllocatedBuffers;
+};
+
+/** The global render targets for easy shading. */
+extern RENDERCORE_API TGlobalResource<FRenderGraphResourcePool> GRenderGraphResourcePool;
diff --git a/Engine/Source/Runtime/RenderCore/Private/RenderGraphUtils.cpp b/Engine/Source/Runtime/RenderCore/Private/RenderGraphUtils.cpp
--- a/Engine/Source/Runtime/RenderCore/Private/RenderGraphUtils.cpp
+++ b/Engine/Source/Runtime/RenderCore/Private/RenderGraphUtils.cpp
@@ -24,7 +24,7 @@ void ClearUnusedGraphResourcesImpl(const FShaderParameterBindings& ShaderBinding
 				continue;
 			}
 		}
-		else if (Type == UBMT_GRAPH_TRACKED_SRV)
+		else if (Type == UBMT_GRAPH_TRACKED_SRV || Type == UBMT_GRAPH_TRACKED_BUFFER_SRV)
 		{
 			if (GraphSRVId < ShaderBindings.GraphSRVs.Num() && ByteOffset == ShaderBindings.GraphSRVs[GraphSRVId].ByteOffset)
 			{
@@ -32,7 +32,7 @@ void ClearUnusedGraphResourcesImpl(const FShaderParameterBindings& ShaderBinding
 				continue;
 			}
 		}
-		else if (Type == UBMT_GRAPH_TRACKED_UAV)
+		else if (Type == UBMT_GRAPH_TRACKED_UAV || Type == UBMT_GRAPH_TRACKED_BUFFER_UAV)
 		{
 			if (GraphUAVId < ShaderBindings.GraphUAVs.Num() && ByteOffset == ShaderBindings.GraphUAVs[GraphUAVId].ByteOffset)
 			{
diff --git a/Engine/Source/Runtime/RenderCore/Private/ShaderParameterMetadata.cpp b/Engine/Source/Runtime/RenderCore/Private/ShaderParameterMetadata.cpp
--- a/Engine/Source/Runtime/RenderCore/Private/ShaderParameterMetadata.cpp
+++ b/Engine/Source/Runtime/RenderCore/Private/ShaderParameterMetadata.cpp
@@ -127,6 +127,9 @@ void FShaderParametersMetadata::InitializeLayout()
 				BaseType == UBMT_GRAPH_TRACKED_TEXTURE ||
 				BaseType == UBMT_GRAPH_TRACKED_SRV ||
 				BaseType == UBMT_GRAPH_TRACKED_UAV ||
+				BaseType == UBMT_GRAPH_TRACKED_BUFFER ||
+				BaseType == UBMT_GRAPH_TRACKED_BUFFER_SRV ||
+				BaseType == UBMT_GRAPH_TRACKED_BUFFER_UAV ||
 				BaseType == UBMT_RENDER_TARGET_BINDING_SLOTS))
 			{
 				UE_LOG(LogRendererCore, Fatal, TEXT("Shader parameter %s error: Graph resources are only allowed in shader parameter structs."), *CppName);
diff --git a/Engine/Source/Runtime/RenderCore/Private/ShaderParameterStruct.cpp b/Engine/Source/Runtime/RenderCore/Private/ShaderParameterStruct.cpp
--- a/Engine/Source/Runtime/RenderCore/Private/ShaderParameterStruct.cpp
+++ b/Engine/Source/Runtime/RenderCore/Private/ShaderParameterStruct.cpp
@@ -47,6 +47,10 @@ struct FShaderParameterStructBindingContext
 				RenderTargetBindingSlotCppName = CppName;
 				continue;
 			}
+			else if (BaseType == UBMT_GRAPH_TRACKED_BUFFER)
+			{
+				continue;
+			}
 
 			// Compute the shader member name to look for according to nesting.
 			FString ShaderBindingName = FString::Printf(TEXT("%s%s"), MemberPrefix, Member.GetName());
@@ -113,7 +117,9 @@ struct FShaderParameterStructBindingContext
 				BaseType == UBMT_SAMPLER ||
 				BaseType == UBMT_GRAPH_TRACKED_TEXTURE ||
 				BaseType == UBMT_GRAPH_TRACKED_SRV ||
-				BaseType == UBMT_GRAPH_TRACKED_UAV)
+				BaseType == UBMT_GRAPH_TRACKED_UAV ||
+				BaseType == UBMT_GRAPH_TRACKED_BUFFER_SRV ||
+				BaseType == UBMT_GRAPH_TRACKED_BUFFER_UAV)
 			{
 				FShaderParameterBindings::FResourceParameter Parameter;
 				Parameter.BaseIndex = BaseIndex;
@@ -130,9 +136,9 @@ struct FShaderParameterStructBindingContext
 					Bindings->Samplers.Add(Parameter);
 				else if (BaseType == UBMT_GRAPH_TRACKED_TEXTURE)
 					Bindings->GraphTextures.Add(Parameter);
-				else if (BaseType == UBMT_GRAPH_TRACKED_SRV)
+				else if (BaseType == UBMT_GRAPH_TRACKED_SRV || BaseType == UBMT_GRAPH_TRACKED_BUFFER_SRV)
 					Bindings->GraphSRVs.Add(Parameter);
-				else // if (BaseType == UBMT_GRAPH_TRACKED_UAV)
+				else // if (BaseType == UBMT_GRAPH_TRACKED_UAV || BaseType == UBMT_GRAPH_TRACKED_BUFFER_UAV)
 					Bindings->GraphUAVs.Add(Parameter);
 			}
 			else
diff --git a/Engine/Source/Runtime/RenderCore/Public/RenderGraphBuilder.h b/Engine/Source/Runtime/RenderCore/Public/RenderGraphBuilder.h
--- a/Engine/Source/Runtime/RenderCore/Public/RenderGraphBuilder.h
+++ b/Engine/Source/Runtime/RenderCore/Public/RenderGraphBuilder.h
@@ -126,7 +126,7 @@ public:
 		FRDGTexture* OutTexture = new(FMemStack::Get()) FRDGTexture(Name, ExternalPooledTexture->GetDesc());
 		OutTexture->PooledRenderTarget = ExternalPooledTexture;
 		AllocatedTextures.Add(OutTexture, ExternalPooledTexture);
-		#if DO_CHECK
+		#if RENDER_GRAPH_DEBUGGING
 			Resources.Add(OutTexture);
 		#endif
 		return OutTexture;
@@ -141,13 +141,28 @@ public:
 		}
 		#endif
 		FRDGTexture* Texture = new(FMemStack::Get()) FRDGTexture(DebugName, Desc);
-		#if DO_CHECK
+		#if RENDER_GRAPH_DEBUGGING
 			Resources.Add(Texture);
 		#endif
 		return Texture;
 	}
 
-	/** Create graph tracked SRV from a descriptor. */
+	/** Create graph tracked resource from a descriptor with a debug name. */
+	inline const FRDGBuffer* CreateBuffer(const FRDGBufferDesc& Desc, const TCHAR* DebugName)
+	{
+		#if RENDER_GRAPH_DEBUGGING
+		{
+			ensureMsgf(!bHasExecuted, TEXT("Render graph buffer %s needs to be created before the builder execution."), DebugName);
+		}
+		#endif
+		FRDGBuffer* Buffer = new(FMemStack::Get()) FRDGBuffer(DebugName, Desc);
+		#if RENDER_GRAPH_DEBUGGING
+			Resources.Add(Buffer);
+		#endif
+		return Buffer;
+	}
+
+	/** Create graph tracked SRV for a texture from a descriptor. */
 	inline const FRDGTextureSRV* CreateSRV(const FRDGTextureSRVDesc& Desc)
 	{
 		check(Desc.Texture);
@@ -159,13 +174,30 @@ public:
 		#endif
 		
 		FRDGTextureSRV* SRV = new(FMemStack::Get()) FRDGTextureSRV(Desc.Texture->Name, Desc);
-		#if DO_CHECK
+		#if RENDER_GRAPH_DEBUGGING
+			Resources.Add(SRV);
+		#endif
+		return SRV;
+	}
+
+	/** Create graph tracked SRV for a buffer from a descriptor. */
+	inline const FRDGBufferSRV* CreateSRV(const FRDGBufferSRVDesc& Desc)
+	{
+		check(Desc.Buffer);
+		#if RENDER_GRAPH_DEBUGGING
+		{
+			ensureMsgf(!bHasExecuted, TEXT("Render graph SRV %s needs to be created before the builder execution."), Desc.Buffer->Name);
+		}
+		#endif
+		
+		FRDGBufferSRV* SRV = new(FMemStack::Get()) FRDGBufferSRV(Desc.Buffer->Name, Desc);
+		#if RENDER_GRAPH_DEBUGGING
 			Resources.Add(SRV);
 		#endif
 		return SRV;
 	}
 
-	/** Create graph tracked UAV from a descriptor. */
+	/** Create graph tracked UAV for a texture from a descriptor. */
 	inline const FRDGTextureUAV* CreateUAV(const FRDGTextureUAVDesc& Desc)
 	{
 		check(Desc.Texture);
@@ -177,7 +209,24 @@ public:
 		#endif
 		
 		FRDGTextureUAV* UAV = new(FMemStack::Get()) FRDGTextureUAV(Desc.Texture->Name, Desc);
-		#if DO_CHECK
+		#if RENDER_GRAPH_DEBUGGING
+			Resources.Add(UAV);
+		#endif
+		return UAV;
+	}
+
+	/** Create graph tracked UAV for a buffer from a descriptor. */
+	inline const FRDGBufferUAV* CreateUAV(const FRDGBufferUAVDesc& Desc)
+	{
+		check(Desc.Buffer);
+		#if RENDER_GRAPH_DEBUGGING
+		{
+			ensureMsgf(!bHasExecuted, TEXT("Render graph UAV %s needs to be created before the builder execution."), Desc.Buffer->Name);
+		}
+		#endif
+		
+		FRDGBufferUAV* UAV = new(FMemStack::Get()) FRDGBufferUAV(Desc.Buffer->Name, Desc);
+		#if RENDER_GRAPH_DEBUGGING
 			Resources.Add(UAV);
 		#endif
 		return UAV;
@@ -258,6 +307,9 @@ private:
 	/** Keep the references over the pooled render target, since FRDGTexture is allocated on FMemStack. */
 	TMap<const FRDGTexture*, TRefCountPtr<IPooledRenderTarget>, SceneRenderingSetAllocator> AllocatedTextures;
 
+	/** Keep the references over the pooled render target, since FRDGTexture is allocated on FMemStack. */
+	TMap<const FRDGBuffer*, TRefCountPtr<FPooledRDGBuffer>, SceneRenderingSetAllocator> AllocatedBuffers;
+
 	/** Array of all deferred access to internal textures. */
 	struct FDeferredInternalTextureQuery
 	{
@@ -282,14 +334,22 @@ private:
 	void WalkGraphDependencies();
 
 	void AllocateRHITextureIfNeeded(const FRDGTexture* Texture, bool bComputePass);
+	void AllocateRHITextureSRVIfNeeded(const FRDGTextureSRV* SRV, bool bComputePass);
+	void AllocateRHITextureUAVIfNeeded(const FRDGTextureUAV* UAV, bool bComputePass);
+	void AllocateRHIBufferIfNeeded(const FRDGBuffer* Texture, bool bComputePass);
+	void AllocateRHIBufferSRVIfNeeded(const FRDGBufferSRV* SRV, bool bComputePass);
+	void AllocateRHIBufferUAVIfNeeded(const FRDGBufferUAV* UAV, bool bComputePass);
+
 
 	void TransitionTexture( const FRDGTexture* Texture, EResourceTransitionAccess TransitionAccess, bool bRequiredCompute ) const;
-	void TransitionUAV( const FRDGTextureUAV* UAV, EResourceTransitionAccess TransitionAccess, bool bRequiredCompute ) const;
+	void TransitionUAV(FUnorderedAccessViewRHIParamRef UAV, const FRDGResource* UnderlyingResource, EResourceTransitionAccess TransitionAccess, bool bRequiredCompute ) const;
 
 	void ExecutePass( const FRenderGraphPass* Pass );
 	void AllocateAndTransitionPassResources(const FRenderGraphPass* Pass, struct FRHIRenderPassInfo* OutRPInfo, bool* bOutHasRenderTargets);
 	static void WarnForUselessPassDependencies(const FRenderGraphPass* Pass);
+
 	void ReleaseRHITextureIfPossible(const FRDGTexture* Texture);
+	void ReleaseRHIBufferIfPossible(const FRDGBuffer* Buffer);
 	void ReleaseUnecessaryResources(const FRenderGraphPass* Pass);
 
 	void ProcessDeferredInternalResourceQueries();
diff --git a/Engine/Source/Runtime/RenderCore/Public/RenderGraphResources.h b/Engine/Source/Runtime/RenderCore/Public/RenderGraphResources.h
--- a/Engine/Source/Runtime/RenderCore/Public/RenderGraphResources.h
+++ b/Engine/Source/Runtime/RenderCore/Public/RenderGraphResources.h
@@ -5,6 +5,19 @@
 #include "CoreMinimal.h"
 #include "RHI.h"
 #include "RendererInterface.h"
+#include "RendererInterface.h"
+
+
+class FRDGResource;
+class FRDGTexture;
+class FRDGTextureSRV;
+class FRDGTextureUAV;
+
+class FRDGBuffer;
+class FRDGBufferSRV;
+class FRDGBufferUAV;
+
+struct FPooledRDGBuffer;
 
 
 /** Generic graph resource. */
@@ -16,13 +29,23 @@ public:
 
 	FRDGResource(const TCHAR* InDebugName)
 		: Name(InDebugName)
-	{ }
+	{
+		CachedRHI.Resource = nullptr;
+	}
 
 	FRDGResource() = delete;
 	FRDGResource(const FRDGResource&) = delete;
 	void operator = (const FRDGResource&) = delete;
 
 private:
+	// RHI resource once allocated.
+	union
+	{
+		mutable FRHIResource* Resource;
+		mutable FShaderResourceViewRHIParamRef SRV;
+		mutable FUnorderedAccessViewRHIParamRef UAV;
+	} CachedRHI;
+
 	/** Number of references in passes and deferred queries. */
 	mutable int32 ReferenceCount = 0;
 
@@ -74,14 +97,9 @@ private:
 	}
 
 	friend class FRDGBuilder;
-	friend class FRDGTextureSRV;
-	friend class FRDGTextureUAV;
 
 	template<typename TRHICmdList, typename TShaderClass, typename TShaderRHI>
 	friend void SetShaderParameters(TRHICmdList&, const TShaderClass*, TShaderRHI*, const typename TShaderClass::FParameters&);
-
-	template<typename TRHICmdList, typename TShaderClass>
-	friend void SetShaderUAVs(TRHICmdList&, const TShaderClass*, FComputeShaderRHIParamRef, const typename TShaderClass::FParameters&);
 };
 
 /** Decsriptor for render graph tracked SRV. */
@@ -112,21 +130,7 @@ private:
 		, Desc(InDesc)
 	{ }
 
-	/** Returns the allocated  shader resource view. */
-	inline FShaderResourceViewRHIParamRef GetRHIShaderResourceView() const
-	{
-		check(Desc.Texture);
-		check(Desc.Texture->PooledRenderTarget);
-		return Desc.Texture->PooledRenderTarget->GetRenderTargetItem().MipSRVs[Desc.MipLevel];
-	}
-
 	friend class FRDGBuilder;
-
-	template<typename TRHICmdList, typename TShaderClass, typename TShaderRHI>
-	friend void SetShaderParameters(TRHICmdList&, const TShaderClass*, TShaderRHI*, const typename TShaderClass::FParameters&);
-
-	template<typename TRHICmdList, typename TShaderClass>
-	friend void SetShaderUAVs(TRHICmdList&, const TShaderClass*, FComputeShaderRHIParamRef, const typename TShaderClass::FParameters&);
 };
 
 /** Decsriptor for render graph tracked UAV. */
@@ -157,19 +161,244 @@ private:
 		, Desc(InDesc)
 	{ }
 
-	/** Returns the allocated unordered access view. */
-	inline FUnorderedAccessViewRHIParamRef GetRHIUnorderedAccessView() const
+	friend class FRDGBuilder;
+};
+
+
+/** Descriptor for render graph tracked Buffer. */
+struct RENDERCORE_API FRDGBufferDesc
+{
+	// Type of buffers to the RHI
+	// TODO(RDG): refactor RHI to only have one FRHIBuffer.
+	enum class EUnderlyingType
 	{
-		check(Desc.Texture);
-		check(Desc.Texture->PooledRenderTarget);
-		return Desc.Texture->PooledRenderTarget->GetRenderTargetItem().MipUAVs[Desc.MipLevel];
+		VertexBuffer,
+		IndexBuffer, // not implemented yet.
+		StructuredBuffer,
+	};
+
+	/** Stride in bytes for index and structured buffers. */
+	uint32 BytesPerElement = 1;
+
+	/** Number of elements. */
+	uint32 NumElements = 1;
+
+	/** Bitfields describing the uses of that buffer. */
+	EBufferUsageFlags Usage = BUF_None;
+
+	/** The underlying RHI type to use. A bit of a work arround because RHI still have 3 different objects. */
+	EUnderlyingType UnderlyingType = EUnderlyingType::VertexBuffer;
+
+	/** Returns the total number of bytes allocated for a such buffer. */
+	uint32 GetTotalNumBytes() const
+	{
+		return BytesPerElement * NumElements;
 	}
 
+
+	bool operator == (const FRDGBufferDesc& Other) const
+	{
+		return (
+			BytesPerElement == Other.BytesPerElement &&
+			NumElements == Other.NumElements &&
+			Usage == Other.Usage &&
+			UnderlyingType == Other.UnderlyingType);
+	}
+
+	bool operator != (const FRDGBufferDesc& Other) const
+	{
+		return !(*this == Other);
+	}
+
+	/** Create the descriptor for an indirect RHI call.
+	 *
+	 * Note, IndirectParameterStruct should be one of the:
+	 *		struct FRHIDispatchIndirectParameters
+	 *		struct FRHIDrawIndirectParameters
+	 *		struct FRHIDrawIndexedIndirectParameters
+	 */
+	template<typename IndirectParameterStruct>
+	static inline FRDGBufferDesc CreateIndirectDesc(uint32 NumElements = 1)
+	{
+		FRDGBufferDesc Desc;
+		Desc.UnderlyingType = EUnderlyingType::VertexBuffer;
+		Desc.Usage = EBufferUsageFlags(BUF_Static | BUF_DrawIndirect | BUF_UnorderedAccess | BUF_ShaderResource);
+		Desc.BytesPerElement = sizeof(IndirectParameterStruct);
+		Desc.NumElements = NumElements;
+		return Desc;
+	}
+};
+
+
+struct RENDERCORE_API FRDGBufferSRVDesc
+{
+	const FRDGBuffer* Buffer = nullptr;
+
+	/** Nymber of bytes per element (used for vertex buffer). */
+	uint32 BytesPerElement = 1;
+
+	/** Encoding format for the element (used for vertex buffer). */
+	EPixelFormat Format = PF_Unknown;
+};
+
+struct RENDERCORE_API FRDGBufferUAVDesc
+{
+	const FRDGBuffer* Buffer = nullptr;
+
+	/** Nymber of bytes per element (used for vertex buffer). */
+	EPixelFormat Format = PF_Unknown;
+
+	/** Whether the uav supports atomic counter or append buffer operations (used for structured buffers) */
+	bool bSupportsAtomicCounter = false;
+	bool bSupportsAppendBuffer = false;
+
+	/** Create descriptor for a the UAV of a buffer meant to be use for indirect dral call parameters. */
+	static inline FRDGBufferUAVDesc CreateIndirectDesc(const FRDGBuffer* Buffer);
+};
+
+
+/** Defines how the map's pairs are hashed. */
+template<typename KeyType, typename ValueType>
+struct TMapRDGBufferSRVFuncs : BaseKeyFuncs<TPair<KeyType, ValueType>, KeyType, /* bInAllowDuplicateKeys = */ false>
+{
+	typedef typename TTypeTraits<KeyType>::ConstPointerType KeyInitType;
+	typedef const TPairInitializer<typename TTypeTraits<KeyType>::ConstInitType, typename TTypeTraits<ValueType>::ConstInitType>& ElementInitType;
+
+	static FORCEINLINE KeyInitType GetSetKey(ElementInitType Element)
+	{
+		return Element.Key;
+	}
+	static FORCEINLINE bool Matches(KeyInitType A, KeyInitType B)
+	{
+		return A.BytesPerElement == B.BytesPerElement && A.Format == B.Format;
+	}
+	static FORCEINLINE uint32 GetKeyHash(KeyInitType Key)
+	{
+		return HashCombine(uint32(Key.BytesPerElement), uint32(Key.Format));
+	}
+};
+
+/** Defines how the map's pairs are hashed. */
+template<typename KeyType, typename ValueType>
+struct TMapRDGBufferUAVFuncs : BaseKeyFuncs<TPair<KeyType, ValueType>, KeyType, /* bInAllowDuplicateKeys = */ false>
+{
+	typedef typename TTypeTraits<KeyType>::ConstPointerType KeyInitType;
+	typedef const TPairInitializer<typename TTypeTraits<KeyType>::ConstInitType, typename TTypeTraits<ValueType>::ConstInitType>& ElementInitType;
+
+	static FORCEINLINE KeyInitType GetSetKey(ElementInitType Element)
+	{
+		return Element.Key;
+	}
+	static FORCEINLINE bool Matches(KeyInitType A, KeyInitType B)
+	{
+		return A.Format == B.Format && A.bSupportsAtomicCounter == B.bSupportsAtomicCounter && A.bSupportsAppendBuffer == B.bSupportsAppendBuffer;
+	}
+	static FORCEINLINE uint32 GetKeyHash(KeyInitType Key)
+	{
+		return (uint32(Key.bSupportsAtomicCounter) << 8) | (uint32(Key.bSupportsAppendBuffer) << 9) | uint32(Key.Format);
+	}
+};
+
+/** Render graph tracked Buffer. Only as opose to vertex, index and structured buffer at RHI level, because there is already
+ * plans to refactor them. */
+struct FPooledRDGBuffer
+{
+	FVertexBufferRHIRef VertexBuffer;
+	FIndexBufferRHIRef IndexBuffer;
+	FStructuredBufferRHIRef StructuredBuffer;
+	TMap<FRDGBufferUAVDesc, FUnorderedAccessViewRHIRef, FDefaultSetAllocator, TMapRDGBufferUAVFuncs<FRDGBufferUAVDesc, FUnorderedAccessViewRHIRef>> UAVs;
+	TMap<FRDGBufferSRVDesc, FShaderResourceViewRHIRef, FDefaultSetAllocator, TMapRDGBufferSRVFuncs<FRDGBufferSRVDesc, FShaderResourceViewRHIRef>> SRVs;
+
+	/** Descriptor. */
+	FRDGBufferDesc Desc;
+
+
+	// Refcounting
+	inline uint32 AddRef()
+	{
+		return ++RefCount;
+	}
+
+	uint32 Release();
+
+	inline uint32 GetRefCount()
+	{
+		return RefCount;
+	}
+
+private:
+	uint32 RefCount = 0;
+};
+
+
+/** Render graph tracked buffers. */
+class RENDERCORE_API FRDGBuffer : public FRDGResource
+{
+public:
+	/** Descriptor of the graph */
+	const FRDGBufferDesc Desc;
+
+	/** Returns the buffer to use for indirect RHI calls. */
+	FVertexBufferRHIParamRef GetIndirectRHICallBuffer() const
+	{
+		check(PooledBuffer);
+		checkf(Desc.UnderlyingType == FRDGBufferDesc::EUnderlyingType::VertexBuffer, TEXT("Indirect buffers needs to be underlying vertex buffer."));
+		check(PooledBuffer->VertexBuffer.IsValid());
+		return PooledBuffer->VertexBuffer;
+	}
+
+private:
+	/** This is not a TRefCountPtr<> because FRDGBuffer is allocated on the FMemStack
+	 * FGraphBuilder::AllocatedBuffers is actually keeping the reference.
+	 */
+	mutable FPooledRDGBuffer* PooledBuffer = nullptr;
+
+	FRDGBuffer(const TCHAR* DebugName, const FRDGBufferDesc& InDesc)
+		: FRDGResource(DebugName)
+		, Desc(InDesc)
+	{ }
+
 	friend class FRDGBuilder;
+};
 
-	template<typename TRHICmdList, typename TShaderClass, typename TShaderRHI>
-	friend void SetShaderParameters(TRHICmdList&, const TShaderClass*, TShaderRHI*, const typename TShaderClass::FParameters&);
+/** Render graph tracked SRV. */
+class RENDERCORE_API FRDGBufferSRV : public FRDGResource
+{
+public:
+	/** Descriptor of the graph tracked SRV. */
+	const FRDGBufferSRVDesc Desc;
 
-	template<typename TRHICmdList, typename TShaderClass>
-	friend void SetShaderUAVs(TRHICmdList&, const TShaderClass*, FComputeShaderRHIParamRef, const typename TShaderClass::FParameters&);
+private:
+	FRDGBufferSRV(const TCHAR* DebugName, const FRDGBufferSRVDesc& InDesc)
+		: FRDGResource(DebugName)
+		, Desc(InDesc)
+	{ }
+
+	friend class FRDGBuilder;
 };
+
+/** Render graph tracked UAV. */
+class RENDERCORE_API FRDGBufferUAV : public FRDGResource
+{
+public:
+	/** Descriptor of the graph tracked UAV. */
+	const FRDGBufferUAVDesc Desc;
+
+private:
+	FRDGBufferUAV(const TCHAR* DebugName, const FRDGBufferUAVDesc& InDesc)
+		: FRDGResource(DebugName)
+		, Desc(InDesc)
+	{ }
+
+	friend class FRDGBuilder;
+};
+
+
+FRDGBufferUAVDesc FRDGBufferUAVDesc::CreateIndirectDesc(const FRDGBuffer* Buffer)
+{
+	checkf(Buffer->Desc.Usage & BUF_DrawIndirect, TEXT("Buffer %s has not been created with a BUF_DrawIndirect flag."), Buffer->Name);
+	FRDGBufferUAVDesc Desc;
+	Desc.Buffer = Buffer;
+	Desc.Format = PF_R32_UINT;
+	return Desc;
+}
diff --git a/Engine/Source/Runtime/RenderCore/Public/ShaderParameterMacros.h b/Engine/Source/Runtime/RenderCore/Public/ShaderParameterMacros.h
--- a/Engine/Source/Runtime/RenderCore/Public/ShaderParameterMacros.h
+++ b/Engine/Source/Runtime/RenderCore/Public/ShaderParameterMacros.h
@@ -13,6 +13,9 @@
 class FRDGTexture;
 class FRDGTextureSRV;
 class FRDGTextureUAV;
+class FRDGBuffer;
+class FRDGBufferSRV;
+class FRDGBufferUAV;
 
 
 // RHICreateUniformBuffer assumes C++ constant layout matches the shader layout when extracting float constants, yet the C++ struct contains pointers.  
@@ -565,6 +568,51 @@ struct TShaderParameterTypeInfo<const FRDGTextureUAV*>
 };
 static_assert(sizeof(TShaderParameterTypeInfo<const FRDGTextureUAV*>::TAlignedType) == SHADER_PARAMETER_POINTER_ALIGNMENT, "Uniform buffer layout must not be platform dependent.");
 
+template<>
+struct TShaderParameterTypeInfo<const FRDGBuffer*>
+{
+	static constexpr int32 NumRows = 1;
+	static constexpr int32 NumColumns = 1;
+	static constexpr int32 NumElements = 0;
+	static constexpr int32 Alignment = SHADER_PARAMETER_POINTER_ALIGNMENT;
+	static constexpr bool bIsStoredInConstantBuffer = false;
+
+	using TAlignedType = TAlignedShaderParameterPtr<const FRDGBuffer*>;
+
+	static const FShaderParametersMetadata* GetStructMetadata() { return nullptr; }
+};
+static_assert(sizeof(TShaderParameterTypeInfo<const FRDGBuffer*>::TAlignedType) == SHADER_PARAMETER_POINTER_ALIGNMENT, "Uniform buffer layout must not be platform dependent.");
+
+template<>
+struct TShaderParameterTypeInfo<const FRDGBufferSRV*>
+{
+	static constexpr int32 NumRows = 1;
+	static constexpr int32 NumColumns = 1;
+	static constexpr int32 NumElements = 0;
+	static constexpr int32 Alignment = SHADER_PARAMETER_POINTER_ALIGNMENT;
+	static constexpr bool bIsStoredInConstantBuffer = false;
+
+	using TAlignedType = TAlignedShaderParameterPtr<const FRDGBufferSRV*>;
+
+	static const FShaderParametersMetadata* GetStructMetadata() { return nullptr; }
+};
+static_assert(sizeof(TShaderParameterTypeInfo<const FRDGBufferSRV*>::TAlignedType) == SHADER_PARAMETER_POINTER_ALIGNMENT, "Uniform buffer layout must not be platform dependent.");
+
+template<>
+struct TShaderParameterTypeInfo<const FRDGBufferUAV*>
+{
+	static constexpr int32 NumRows = 1;
+	static constexpr int32 NumColumns = 1;
+	static constexpr int32 NumElements = 0;
+	static constexpr int32 Alignment = SHADER_PARAMETER_POINTER_ALIGNMENT;
+	static constexpr bool bIsStoredInConstantBuffer = false;
+
+	using TAlignedType = TAlignedShaderParameterPtr<const FRDGBufferUAV*>;
+
+	static const FShaderParametersMetadata* GetStructMetadata() { return nullptr; }
+};
+static_assert(sizeof(TShaderParameterTypeInfo<const FRDGBufferUAV*>::TAlignedType) == SHADER_PARAMETER_POINTER_ALIGNMENT, "Uniform buffer layout must not be platform dependent.");
+
 template<class UniformBufferStructType>
 struct TShaderParameterTypeInfo<TUniformBufferRef<UniformBufferStructType>>
 {
@@ -760,7 +808,7 @@ extern RENDERCORE_API FShaderParametersMetadata* FindUniformBufferStructByFName(
 #define SHADER_PARAMETER_GRAPH_TEXTURE(ShaderType,MemberName) \
 	DECLARE_UNIFORM_BUFFER_STRUCT_MEMBER_EXPLICIT(UBMT_GRAPH_TRACKED_TEXTURE, TShaderParameterTypeInfo<const FRDGTexture*>, const FRDGTexture*,MemberName,,EShaderPrecisionModifier::Float,TEXT(#ShaderType),false)
 
-/** Adds a render graph tracked shader resource view.
+/** Adds a shader resource view for a render graph tracked texture.
  *
  * Example:
  *	SHADER_PARAMETER_GRAPH_SRV(Texture2D, MySRV)
@@ -768,7 +816,7 @@ extern RENDERCORE_API FShaderParametersMetadata* FindUniformBufferStructByFName(
 #define SHADER_PARAMETER_GRAPH_SRV(ShaderType,MemberName) \
 	DECLARE_UNIFORM_BUFFER_STRUCT_MEMBER_EXPLICIT(UBMT_GRAPH_TRACKED_SRV, TShaderParameterTypeInfo<const FRDGTextureSRV*>, const FRDGTextureSRV*,MemberName,,EShaderPrecisionModifier::Float,TEXT(#ShaderType),false)
 
-/** Adds a render graph tracked unordered access view.
+/** Adds a unordered access view for a render graph tracked texture.
  *
  * Example:
  *	SHADER_PARAMETER_GRAPH_UAV(RWTexture2D, MyUAV)
@@ -776,6 +824,30 @@ extern RENDERCORE_API FShaderParametersMetadata* FindUniformBufferStructByFName(
 #define SHADER_PARAMETER_GRAPH_UAV(ShaderType,MemberName) \
 	DECLARE_UNIFORM_BUFFER_STRUCT_MEMBER_EXPLICIT(UBMT_GRAPH_TRACKED_UAV, TShaderParameterTypeInfo<const FRDGTextureUAV*>, const FRDGTextureUAV*,MemberName,,EShaderPrecisionModifier::Float,TEXT(#ShaderType),false)
 
+/** Adds a render graph tracked buffer.
+ *
+ * Example:
+ *	SHADER_PARAMETER_GRAPH_BUFFER(Texture2D, MyTexture)
+ */
+#define SHADER_PARAMETER_GRAPH_BUFFER(ShaderType,MemberName) \
+	DECLARE_UNIFORM_BUFFER_STRUCT_MEMBER_EXPLICIT(UBMT_GRAPH_TRACKED_BUFFER, TShaderParameterTypeInfo<const FRDGBuffer*>, const FRDGBuffer*,MemberName,,EShaderPrecisionModifier::Float,TEXT(#ShaderType),false)
+
+/** Adds a shader resource view for a render graph tracked buffer.
+ *
+ * Example:
+ *	SHADER_PARAMETER_GRAPH_BUFFER_SRV(Buffer<float4>, MySRV)
+ */
+#define SHADER_PARAMETER_GRAPH_BUFFER_SRV(ShaderType,MemberName) \
+	DECLARE_UNIFORM_BUFFER_STRUCT_MEMBER_EXPLICIT(UBMT_GRAPH_TRACKED_BUFFER_SRV, TShaderParameterTypeInfo<const FRDGBufferSRV*>, const FRDGBufferSRV*,MemberName,,EShaderPrecisionModifier::Float,TEXT(#ShaderType),false)
+
+/** Adds a unordered access view for a render graph tracked buffer.
+ *
+ * Example:
+ *	SHADER_PARAMETER_GRAPH_BUFFER_UAV(RWBuffer<float4>, MyUAV)
+ */
+#define SHADER_PARAMETER_GRAPH_BUFFER_UAV(ShaderType,MemberName) \
+	DECLARE_UNIFORM_BUFFER_STRUCT_MEMBER_EXPLICIT(UBMT_GRAPH_TRACKED_BUFFER_UAV, TShaderParameterTypeInfo<const FRDGBufferUAV*>, const FRDGBufferUAV*,MemberName,,EShaderPrecisionModifier::Float,TEXT(#ShaderType),false)
+
 /** Nests a shader parameter structure into another one, in C++ and shader code.
  *
  * Example:
diff --git a/Engine/Source/Runtime/RenderCore/Public/ShaderParameterStruct.h b/Engine/Source/Runtime/RenderCore/Public/ShaderParameterStruct.h
--- a/Engine/Source/Runtime/RenderCore/Public/ShaderParameterStruct.h
+++ b/Engine/Source/Runtime/RenderCore/Public/ShaderParameterStruct.h
@@ -74,7 +74,7 @@ inline void SetShaderUAVs(TRHICmdList& RHICmdList, const TShaderClass* Shader, F
 			}
 		}
 
-		RHICmdList.SetUAVParameter(ShadeRHI, ParameterBinding.BaseIndex, GraphUAV->GetRHIUnorderedAccessView());
+		RHICmdList.SetUAVParameter(ShadeRHI, ParameterBinding.BaseIndex, GraphUAV->CachedRHI.UAV);
 	}
 }
 
@@ -187,7 +187,7 @@ inline void SetShaderParameters(TRHICmdList& RHICmdList, const TShaderClass* Sha
 			}
 		}
 
-		RHICmdList.SetShaderResourceViewParameter(ShadeRHI, ParameterBinding.BaseIndex, GraphSRV->GetRHIShaderResourceView());
+		RHICmdList.SetShaderResourceViewParameter(ShadeRHI, ParameterBinding.BaseIndex, GraphSRV->CachedRHI.SRV);
 	}
 
 	// Graph UAVs for compute shaders
diff --git a/Engine/Source/Runtime/Windows/D3D11RHI/Private/D3D11UniformBuffer.cpp b/Engine/Source/Runtime/Windows/D3D11RHI/Private/D3D11UniformBuffer.cpp
--- a/Engine/Source/Runtime/Windows/D3D11RHI/Private/D3D11UniformBuffer.cpp
+++ b/Engine/Source/Runtime/Windows/D3D11RHI/Private/D3D11UniformBuffer.cpp
@@ -213,7 +213,7 @@ FUniformBufferRHIRef FD3D11DynamicRHI::RHICreateUniformBuffer(const void* Conten
 			FRHIResource* Resource = *(FRHIResource**)((uint8*)Contents + Layout.ResourceOffsets[i]);
 
 			// Allow null SRV's in uniform buffers for feature levels that don't support SRV's in shaders
-			if (!(GMaxRHIFeatureLevel <= ERHIFeatureLevel::ES3_1 && (Layout.Resources[i] == UBMT_SRV || Layout.Resources[i] == UBMT_GRAPH_TRACKED_SRV)))
+			if (!(GMaxRHIFeatureLevel <= ERHIFeatureLevel::ES3_1 && (Layout.Resources[i] == UBMT_SRV || Layout.Resources[i] == UBMT_GRAPH_TRACKED_SRV || Layout.Resources[i] == UBMT_GRAPH_TRACKED_BUFFER_SRV)))
 			{
 				checkf(Resource, TEXT("Invalid resource entry creating uniform buffer, %s.Resources[%u], ResourceType 0x%x."), *Layout.GetDebugName().ToString(), i, Layout.Resources[i]);
 			}
