Renames shader parameter macros for render graph managed resource to final name

diff --git a/Engine/Source/Runtime/RenderCore/Public/ShaderParameterMacros.h b/Engine/Source/Runtime/RenderCore/Public/ShaderParameterMacros.h
--- a/Engine/Source/Runtime/RenderCore/Public/ShaderParameterMacros.h
+++ b/Engine/Source/Runtime/RenderCore/Public/ShaderParameterMacros.h
@@ -803,49 +803,49 @@ extern RENDERCORE_API FShaderParametersMetadata* FindUniformBufferStructByFName(
 /** Adds a render graph tracked texture.
  *
  * Example:
- *	SHADER_PARAMETER_GRAPH_TEXTURE(Texture2D, MyTexture)
+ *	SHADER_PARAMETER_RDG_TEXTURE(Texture2D, MyTexture)
  */
-#define SHADER_PARAMETER_GRAPH_TEXTURE(ShaderType,MemberName) \
+#define SHADER_PARAMETER_RDG_TEXTURE(ShaderType,MemberName) \
 	DECLARE_UNIFORM_BUFFER_STRUCT_MEMBER_EXPLICIT(UBMT_GRAPH_TRACKED_TEXTURE, TShaderParameterTypeInfo<const FRDGTexture*>, const FRDGTexture*,MemberName,,EShaderPrecisionModifier::Float,TEXT(#ShaderType),false)
 
 /** Adds a shader resource view for a render graph tracked texture.
  *
  * Example:
- *	SHADER_PARAMETER_GRAPH_SRV(Texture2D, MySRV)
+ *	SHADER_PARAMETER_RDG_TEXTURE_SRV(Texture2D, MySRV)
  */
-#define SHADER_PARAMETER_GRAPH_SRV(ShaderType,MemberName) \
+#define SHADER_PARAMETER_RDG_TEXTURE_SRV(ShaderType,MemberName) \
 	DECLARE_UNIFORM_BUFFER_STRUCT_MEMBER_EXPLICIT(UBMT_GRAPH_TRACKED_SRV, TShaderParameterTypeInfo<const FRDGTextureSRV*>, const FRDGTextureSRV*,MemberName,,EShaderPrecisionModifier::Float,TEXT(#ShaderType),false)
 
 /** Adds a unordered access view for a render graph tracked texture.
  *
  * Example:
- *	SHADER_PARAMETER_GRAPH_UAV(RWTexture2D, MyUAV)
+ *	SHADER_PARAMETER_RDG_TEXTURE_UAV(RWTexture2D, MyUAV)
  */
-#define SHADER_PARAMETER_GRAPH_UAV(ShaderType,MemberName) \
+#define SHADER_PARAMETER_RDG_TEXTURE_UAV(ShaderType,MemberName) \
 	DECLARE_UNIFORM_BUFFER_STRUCT_MEMBER_EXPLICIT(UBMT_GRAPH_TRACKED_UAV, TShaderParameterTypeInfo<const FRDGTextureUAV*>, const FRDGTextureUAV*,MemberName,,EShaderPrecisionModifier::Float,TEXT(#ShaderType),false)
 
 /** Adds a render graph tracked buffer.
  *
  * Example:
- *	SHADER_PARAMETER_GRAPH_BUFFER(Texture2D, MyTexture)
+ *	SHADER_PARAMETER_RDG_BUFFER(Texture2D, MyTexture)
  */
-#define SHADER_PARAMETER_GRAPH_BUFFER(ShaderType,MemberName) \
+#define SHADER_PARAMETER_RDG_BUFFER(ShaderType,MemberName) \
 	DECLARE_UNIFORM_BUFFER_STRUCT_MEMBER_EXPLICIT(UBMT_GRAPH_TRACKED_BUFFER, TShaderParameterTypeInfo<const FRDGBuffer*>, const FRDGBuffer*,MemberName,,EShaderPrecisionModifier::Float,TEXT(#ShaderType),false)
 
 /** Adds a shader resource view for a render graph tracked buffer.
  *
  * Example:
- *	SHADER_PARAMETER_GRAPH_BUFFER_SRV(Buffer<float4>, MySRV)
+ *	SHADER_PARAMETER_RDG_BUFFER_SRV(Buffer<float4>, MySRV)
  */
-#define SHADER_PARAMETER_GRAPH_BUFFER_SRV(ShaderType,MemberName) \
+#define SHADER_PARAMETER_RDG_BUFFER_SRV(ShaderType,MemberName) \
 	DECLARE_UNIFORM_BUFFER_STRUCT_MEMBER_EXPLICIT(UBMT_GRAPH_TRACKED_BUFFER_SRV, TShaderParameterTypeInfo<const FRDGBufferSRV*>, const FRDGBufferSRV*,MemberName,,EShaderPrecisionModifier::Float,TEXT(#ShaderType),false)
 
 /** Adds a unordered access view for a render graph tracked buffer.
  *
  * Example:
- *	SHADER_PARAMETER_GRAPH_BUFFER_UAV(RWBuffer<float4>, MyUAV)
+ *	SHADER_PARAMETER_RDG_BUFFER_UAV(RWBuffer<float4>, MyUAV)
  */
-#define SHADER_PARAMETER_GRAPH_BUFFER_UAV(ShaderType,MemberName) \
+#define SHADER_PARAMETER_RDG_BUFFER_UAV(ShaderType,MemberName) \
 	DECLARE_UNIFORM_BUFFER_STRUCT_MEMBER_EXPLICIT(UBMT_GRAPH_TRACKED_BUFFER_UAV, TShaderParameterTypeInfo<const FRDGBufferUAV*>, const FRDGBufferUAV*,MemberName,,EShaderPrecisionModifier::Float,TEXT(#ShaderType),false)
 
 /** Nests a shader parameter structure into another one, in C++ and shader code.
diff --git a/Engine/Source/Runtime/Renderer/Private/SceneOcclusion.cpp b/Engine/Source/Runtime/Renderer/Private/SceneOcclusion.cpp
--- a/Engine/Source/Runtime/Renderer/Private/SceneOcclusion.cpp
+++ b/Engine/Source/Runtime/Renderer/Private/SceneOcclusion.cpp
@@ -1053,7 +1053,7 @@ void FHZBOcclusionTester::Submit(FRHICommandListImmediate& RHICmdList, const FVi
 
 BEGIN_SHADER_PARAMETER_STRUCT(FHZBBuildPassParameters, )
 	RENDER_TARGET_BINDING_SLOTS()
-	SHADER_PARAMETER_GRAPH_SRV(Texture2D, Texture)
+	SHADER_PARAMETER_RDG_TEXTURE_SRV(Texture2D, Texture)
 	SHADER_PARAMETER_SAMPLER(SamplerState, TextureSampler)
 END_SHADER_PARAMETER_STRUCT()
 
diff --git a/Engine/Source/Runtime/Renderer/Private/VolumetricFog.cpp b/Engine/Source/Runtime/Renderer/Private/VolumetricFog.cpp
--- a/Engine/Source/Runtime/Renderer/Private/VolumetricFog.cpp
+++ b/Engine/Source/Runtime/Renderer/Private/VolumetricFog.cpp
@@ -147,8 +147,8 @@ class FVolumetricFogMaterialSetupCS : public FGlobalShader
 		SHADER_PARAMETER_STRUCT_REF( FFogUniformParameters,			FogUniformParameters )
 		SHADER_PARAMETER_STRUCT_REF( FViewUniformShaderParameters,	View )
 		
-		SHADER_PARAMETER_GRAPH_UAV(	RWTexture3D<float4>, RWVBufferA )
-		SHADER_PARAMETER_GRAPH_UAV(	RWTexture3D<float4>, RWVBufferB )
+		SHADER_PARAMETER_RDG_TEXTURE_UAV(	RWTexture3D<float4>, RWVBufferA )
+		SHADER_PARAMETER_RDG_TEXTURE_UAV(	RWTexture3D<float4>, RWVBufferB )
 	END_SHADER_PARAMETER_STRUCT()
 
 public:
@@ -574,12 +574,12 @@ class TVolumetricFogLightScatteringCS : public FGlobalShader
 		FDistanceFieldSkyOcclusion >;
 
 	BEGIN_SHADER_PARAMETER_STRUCT( FParameters, )
-		SHADER_PARAMETER_GRAPH_TEXTURE( Texture2D, VBufferA )
-		SHADER_PARAMETER_GRAPH_TEXTURE( Texture2D, VBufferB )
-		SHADER_PARAMETER_GRAPH_TEXTURE( Texture2D, LocalShadowedLightScattering )
-		SHADER_PARAMETER_GRAPH_TEXTURE( Texture2D, LightFunctionTexture )
+		SHADER_PARAMETER_RDG_TEXTURE( Texture2D, VBufferA )
+		SHADER_PARAMETER_RDG_TEXTURE( Texture2D, VBufferB )
+		SHADER_PARAMETER_RDG_TEXTURE( Texture2D, LocalShadowedLightScattering )
+		SHADER_PARAMETER_RDG_TEXTURE( Texture2D, LightFunctionTexture )
 
-		SHADER_PARAMETER_GRAPH_UAV( RWTexture2D, RWLightScattering )
+		SHADER_PARAMETER_RDG_TEXTURE_UAV( RWTexture2D, RWLightScattering )
 	END_SHADER_PARAMETER_STRUCT()
 
 public:
@@ -761,8 +761,8 @@ class FVolumetricFogFinalIntegrationCS : public FGlobalShader
 	DECLARE_SHADER_TYPE(FVolumetricFogFinalIntegrationCS,Global)
 
 	BEGIN_SHADER_PARAMETER_STRUCT( FParameters, )
-		SHADER_PARAMETER_GRAPH_TEXTURE(	Texture3D<float4>,		LightScattering )
-		SHADER_PARAMETER_GRAPH_UAV(		RWTexture3D<float4>,	RWIntegratedLightScattering )
+		SHADER_PARAMETER_RDG_TEXTURE(	Texture3D<float4>,		LightScattering )
+		SHADER_PARAMETER_RDG_TEXTURE_UAV(		RWTexture3D<float4>,	RWIntegratedLightScattering )
 	END_SHADER_PARAMETER_STRUCT()
 
 	static bool ShouldCompilePermutation(const FGlobalShaderPermutationParameters& Parameters)
