Fixes PVS-Studio static analysis

diff --git a/Engine/Source/Runtime/RenderCore/Private/RenderGraphBuilder.cpp b/Engine/Source/Runtime/RenderCore/Private/RenderGraphBuilder.cpp
--- a/Engine/Source/Runtime/RenderCore/Private/RenderGraphBuilder.cpp
+++ b/Engine/Source/Runtime/RenderCore/Private/RenderGraphBuilder.cpp
@@ -324,7 +324,7 @@ void FRDGBuilder::AllocateRHIBufferSRVIfNeeded(const FRDGBufferSRV* SRV, bool bC
 	{
 		RHIShaderResourceView = RHICreateShaderResourceView(SRV->Desc.Buffer->PooledBuffer->VertexBuffer, SRV->Desc.BytesPerElement, SRV->Desc.Format);
 	}
-	else if (SRV->Desc.Buffer->Desc.UnderlyingType == FRDGBufferDesc::EUnderlyingType::VertexBuffer)
+	else if (SRV->Desc.Buffer->Desc.UnderlyingType == FRDGBufferDesc::EUnderlyingType::StructuredBuffer)
 	{
 		RHIShaderResourceView = RHICreateShaderResourceView(SRV->Desc.Buffer->PooledBuffer->StructuredBuffer);
 	}
@@ -363,7 +363,7 @@ void FRDGBuilder::AllocateRHIBufferUAVIfNeeded(const FRDGBufferUAV* UAV, bool bC
 	{
 		RHIUnorderedAccessView = RHICreateUnorderedAccessView(UAV->Desc.Buffer->PooledBuffer->VertexBuffer, UAV->Desc.Format);
 	}
-	else if (UAV->Desc.Buffer->Desc.UnderlyingType == FRDGBufferDesc::EUnderlyingType::VertexBuffer)
+	else if (UAV->Desc.Buffer->Desc.UnderlyingType == FRDGBufferDesc::EUnderlyingType::StructuredBuffer)
 	{
 		RHIUnorderedAccessView = RHICreateUnorderedAccessView(UAV->Desc.Buffer->PooledBuffer->StructuredBuffer, UAV->Desc.bSupportsAtomicCounter, UAV->Desc.bSupportsAppendBuffer);
 	}
