Fixes a crash in render graph on mac.

diff --git a/Engine/Source/Runtime/RenderCore/Private/ShaderParameterMetadata.cpp b/Engine/Source/Runtime/RenderCore/Private/ShaderParameterMetadata.cpp
--- a/Engine/Source/Runtime/RenderCore/Private/ShaderParameterMetadata.cpp
+++ b/Engine/Source/Runtime/RenderCore/Private/ShaderParameterMetadata.cpp
@@ -63,6 +63,16 @@ FShaderParametersMetadata::FShaderParametersMetadata(
 	, GlobalListLink(this)
 	, bLayoutInitialized(false)
 {
+	check(StructTypeName);
+	if (UseCase == EUseCase::ShaderParameterStruct)
+	{
+		check(ShaderVariableName == nullptr);
+	}
+	else
+	{
+		check(ShaderVariableName);
+	}
+
 	if (UseCase == EUseCase::GlobalShaderParameterStruct)
 	{
 		// Register this uniform buffer struct in global list.
@@ -80,11 +90,6 @@ FShaderParametersMetadata::FShaderParametersMetadata(
 		// Structs created at runtime with bRegisterForAutoBinding==true can be initialized now.
 		InitializeLayout();
 	}
-
-	if (UseCase == EUseCase::ShaderParameterStruct)
-	{
-		check(ShaderVariableName == nullptr);
-	}
 }
 
 void FShaderParametersMetadata::InitializeAllGlobalStructs()
diff --git a/Engine/Source/Runtime/RenderCore/Public/ShaderParameterMacros.h b/Engine/Source/Runtime/RenderCore/Public/ShaderParameterMacros.h
--- a/Engine/Source/Runtime/RenderCore/Public/ShaderParameterMacros.h
+++ b/Engine/Source/Runtime/RenderCore/Public/ShaderParameterMacros.h
@@ -638,7 +638,7 @@ struct TShaderParameterTypeInfo<TUniformBufferRef<UniformBufferStructType>>
 	static FShaderParametersMetadata StaticStructMetadata(\
 		FShaderParametersMetadata::EUseCase::ShaderParameterStruct, \
 		FName(TEXT(#StructTypeName)), \
-		nullptr, \
+		TEXT(#StructTypeName), \
 		nullptr, \
 		sizeof(StructTypeName), \
 		StructTypeName::zzGetMembers()); \
