PR #5242: Fix typos in RenderCore.h comments. (Contributed by [REDACTED_PERSON])

diff --git a/Engine/Source/Runtime/RenderCore/Public/RenderCore.h b/Engine/Source/Runtime/RenderCore/Public/RenderCore.h
--- a/Engine/Source/Runtime/RenderCore/Public/RenderCore.h
+++ b/Engine/Source/Runtime/RenderCore/Public/RenderCore.h
@@ -235,13 +235,13 @@ struct FInputLatencyTimer
 		return FPlatformTime::ToSeconds(DeltaTime);
 	}
 
-	/** Weather GInputLatencyTimer is initialized or not. */
+	/** Whether GInputLatencyTimer is initialized or not. */
 	bool	bInitialized;
 
-	/** Weather a measurement has been triggered on the gamethread. */
+	/** Whether a measurement has been triggered on the gamethread. */
 	bool	GameThreadTrigger;
 
-	/** Weather a measurement has been triggered on the renderthread. */
+	/** Whether a measurement has been triggered on the renderthread. */
 	bool	RenderThreadTrigger;
 
 	/** Start time (in FPlatformTime::Cycles). */
