UE4 - Static analysis fix

diff --git a/Engine/Source/Runtime/RenderCore/Private/RenderGraphBuilder.cpp b/Engine/Source/Runtime/RenderCore/Private/RenderGraphBuilder.cpp
--- a/Engine/Source/Runtime/RenderCore/Private/RenderGraphBuilder.cpp
+++ b/Engine/Source/Runtime/RenderCore/Private/RenderGraphBuilder.cpp
@@ -624,9 +624,12 @@ void FRDGBuilder::ExecutePass( const FRenderGraphPass* Pass )
 
 	AllocateAndTransitionPassResources(Pass, &RPInfo, &bHasRenderTargets);
 
-	if (RENDER_GRAPH_DRAW_EVENTS && GetEmitDrawEvents())
+	if (RENDER_GRAPH_DRAW_EVENTS)
 	{
-		PushDrawEventStack(Pass);
+		if (GetEmitDrawEvents())
+		{
+			PushDrawEventStack(Pass);
+		}
 	}
 
 	if( !Pass->IsCompute())
@@ -649,9 +652,12 @@ void FRDGBuilder::ExecutePass( const FRenderGraphPass* Pass )
 		RHICmdList.EndRenderPass();
 	}
 
-	if (RENDER_GRAPH_DRAW_EVENTS && GetEmitDrawEvents())
+	if (RENDER_GRAPH_DRAW_EVENTS)
 	{
-		RHICmdList.PopEvent();
+		if (GetEmitDrawEvents())
+		{
+			RHICmdList.PopEvent();
+		}
 	}
 
 	if (RENDER_GRAPH_DEBUGGING)
diff --git a/Engine/Source/Runtime/RenderCore/Private/Shader.cpp b/Engine/Source/Runtime/RenderCore/Private/Shader.cpp
--- a/Engine/Source/Runtime/RenderCore/Private/Shader.cpp
+++ b/Engine/Source/Runtime/RenderCore/Private/Shader.cpp
@@ -2104,8 +2104,8 @@ void ShaderMapAppendKeyString(EShaderPlatform Platform, FString& KeyString)
 		else
 		{
 			// shader platform doesn't match a specific target platform, use cvar setting for forward shading
-			static IConsoleVariable* CVarForwardShading = IConsoleManager::Get().FindConsoleVariable(TEXT("r.ForwardShading"));
-			bForwardShading = CVarForwardShading ? (CVarForwardShading->GetInt() != 0) : false;
+			static IConsoleVariable* CVarForwardShadingLocal = IConsoleManager::Get().FindConsoleVariable(TEXT("r.ForwardShading"));
+			bForwardShading = CVarForwardShadingLocal ? (CVarForwardShadingLocal->GetInt() != 0) : false;
 		}
 
 		if (bForwardShading)
diff --git a/Engine/Source/Runtime/RenderCore/Public/PixelShaderUtils.h b/Engine/Source/Runtime/RenderCore/Public/PixelShaderUtils.h
--- a/Engine/Source/Runtime/RenderCore/Public/PixelShaderUtils.h
+++ b/Engine/Source/Runtime/RenderCore/Public/PixelShaderUtils.h
@@ -60,6 +60,7 @@ struct RENDERCORE_API FPixelShaderUtils
 		const TShaderClass* PixelShader,
 		typename TShaderClass::FParameters* Parameters)
 	{
+		CA_ASSUME(PixelShader != nullptr);
 		ClearUnusedGraphResources(PixelShader, Parameters);
 
 		GraphBuilder.AddPass(
