UE - Another attempt to fix static analysis

diff --git a/Engine/Source/Runtime/RenderCore/Public/PixelShaderUtils.h b/Engine/Source/Runtime/RenderCore/Public/PixelShaderUtils.h
--- a/Engine/Source/Runtime/RenderCore/Public/PixelShaderUtils.h
+++ b/Engine/Source/Runtime/RenderCore/Public/PixelShaderUtils.h
@@ -32,6 +32,7 @@ struct RENDERCORE_API FPixelShaderUtils
 		const TShaderClass* PixelShader,
 		const typename TShaderClass::FParameters& Parameters)
 	{
+		CA_ASSUME(PixelShader);
 		TShaderMapRef<FVisualizeTextureVS> VertexShader(GlobalShaderMap);
 
 		FGraphicsPipelineStateInitializer GraphicsPSOInit;
