Fixing deprecation warning

diff --git a/Engine/Source/Runtime/RenderCore/Private/RenderGraphBuilder.cpp b/Engine/Source/Runtime/RenderCore/Private/RenderGraphBuilder.cpp
--- a/Engine/Source/Runtime/RenderCore/Private/RenderGraphBuilder.cpp
+++ b/Engine/Source/Runtime/RenderCore/Private/RenderGraphBuilder.cpp
@@ -33,7 +33,7 @@ FRDGEventName::FRDGEventName(const TCHAR* EventFormat, ...)
 		va_start(ptr, EventFormat);
 		TCHAR TempStr[256];
 		// Build the string in the temp buffer
-		FCString::GetVarArgs(TempStr, ARRAY_COUNT(TempStr), ARRAY_COUNT(TempStr) - 1, EventFormat, ptr);
+		FCString::GetVarArgs(TempStr, ARRAY_COUNT(TempStr), EventFormat, ptr);
 
 		EventName = TempStr;
 	}
