Fix for missing va_end().

diff --git a/Engine/Source/Runtime/RenderCore/Private/RenderGraphBuilder.cpp b/Engine/Source/Runtime/RenderCore/Private/RenderGraphBuilder.cpp
--- a/Engine/Source/Runtime/RenderCore/Private/RenderGraphBuilder.cpp
+++ b/Engine/Source/Runtime/RenderCore/Private/RenderGraphBuilder.cpp
@@ -34,6 +34,7 @@ FRDGEventName::FRDGEventName(const TCHAR* EventFormat, ...)
 		TCHAR TempStr[256];
 		// Build the string in the temp buffer
 		FCString::GetVarArgs(TempStr, ARRAY_COUNT(TempStr), EventFormat, ptr);
+		va_end(ptr);
 
 		EventName = TempStr;
 	}
