When generating shader map key, need to fallback to forward shading cvar in the case where shader platform doesn't match any target platforms
This matches an exisiting fix in ShaderCompiler.cpp, when deciding to append the 'FORWARD_SHADING' macro

diff --git a/Engine/Shaders/Public/ShaderVersion.ush b/Engine/Shaders/Public/ShaderVersion.ush
--- a/Engine/Shaders/Public/ShaderVersion.ush
+++ b/Engine/Shaders/Public/ShaderVersion.ush
@@ -2,4 +2,4 @@
 // Each time the console command is executed it generates a new GUID. As this file is included
 // in Platform.ush (which should be included in any shader) it allows to invalidate the shader DDC.
 // 
-// GUID = 9E1B9149057548A6A7BC0CEE4B6B886E
+// GUID = A7CB4E7E6E3347189A794F837FDB1A6D
diff --git a/Engine/Source/Runtime/RenderCore/Private/Shader.cpp b/Engine/Source/Runtime/RenderCore/Private/Shader.cpp
--- a/Engine/Source/Runtime/RenderCore/Private/Shader.cpp
+++ b/Engine/Source/Runtime/RenderCore/Private/Shader.cpp
@@ -2092,8 +2092,21 @@ void ShaderMapAppendKeyString(EShaderPlatform Platform, FString& KeyString)
 	}
 
 	{
+		bool bForwardShading = false;
 		ITargetPlatform* TargetPlatform = GetTargetPlatformManager()->FindTargetPlatform(ShaderPlatformToPlatformName(Platform).ToString());
-		if (TargetPlatform && TargetPlatform->UsesForwardShading())
+		if (TargetPlatform)
+		{
+			// if there is a specific target platform that matches our shader platform, use that to drive forward shading
+			bForwardShading = TargetPlatform->UsesForwardShading();
+		}
+		else
+		{
+			// shader platform doesn't match a specific target platform, use cvar setting for forward shading
+			static IConsoleVariable* CVarForwardShading = IConsoleManager::Get().FindConsoleVariable(TEXT("r.ForwardShading"));
+			bForwardShading = CVarForwardShading ? (CVarForwardShading->GetInt() != 0) : false;
+		}
+
+		if (bForwardShading)
 		{
 			KeyString += TEXT("_FS");
 		}
