Cherry-pick 4510508: Implemenets some utilities to reduce boiler plate in render graph user code.

diff --git a/Engine/Source/Runtime/RenderCore/Public/RenderGraph.h b/Engine/Source/Runtime/RenderCore/Public/RenderGraph.h
--- a/Engine/Source/Runtime/RenderCore/Public/RenderGraph.h
+++ b/Engine/Source/Runtime/RenderCore/Public/RenderGraph.h
@@ -7,3 +7,5 @@
 #include "RenderGraphResources.h"
 #include "RenderGraphBuilder.h"
 #include "RenderGraphUtils.h"
+#include "ShaderParameterStruct.h"
+#include "ShaderParameterMacros.h"
diff --git a/Engine/Source/Runtime/RenderCore/Public/RenderGraphBuilder.h b/Engine/Source/Runtime/RenderCore/Public/RenderGraphBuilder.h
--- a/Engine/Source/Runtime/RenderCore/Public/RenderGraphBuilder.h
+++ b/Engine/Source/Runtime/RenderCore/Public/RenderGraphBuilder.h
@@ -116,7 +116,7 @@ public:
 	}
 
 	/** Register a external texture to be tracked by the render graph. */
-	inline const FRDGTexture* RegisterExternalTexture(const TRefCountPtr<IPooledRenderTarget>& ExternalPooledTexture, const TCHAR* Name = TEXT("External"))
+	inline FRDGTextureRef RegisterExternalTexture(const TRefCountPtr<IPooledRenderTarget>& ExternalPooledTexture, const TCHAR* Name = TEXT("External"))
 	{
 		#if RENDER_GRAPH_DEBUGGING
 		{
@@ -133,7 +133,7 @@ public:
 	}
 
 	/** Create graph tracked resource from a descriptor with a debug name. */
-	inline const FRDGTexture* CreateTexture(const FPooledRenderTargetDesc& Desc, const TCHAR* DebugName)
+	inline FRDGTextureRef CreateTexture(const FPooledRenderTargetDesc& Desc, const TCHAR* DebugName)
 	{
 		#if RENDER_GRAPH_DEBUGGING
 		{
@@ -148,7 +148,7 @@ public:
 	}
 
 	/** Create graph tracked resource from a descriptor with a debug name. */
-	inline const FRDGBuffer* CreateBuffer(const FRDGBufferDesc& Desc, const TCHAR* DebugName)
+	inline FRDGBufferRef CreateBuffer(const FRDGBufferDesc& Desc, const TCHAR* DebugName)
 	{
 		#if RENDER_GRAPH_DEBUGGING
 		{
@@ -163,7 +163,7 @@ public:
 	}
 
 	/** Create graph tracked SRV for a texture from a descriptor. */
-	inline const FRDGTextureSRV* CreateSRV(const FRDGTextureSRVDesc& Desc)
+	inline FRDGTextureSRVRef CreateSRV(const FRDGTextureSRVDesc& Desc)
 	{
 		check(Desc.Texture);
 		#if RENDER_GRAPH_DEBUGGING
@@ -181,7 +181,7 @@ public:
 	}
 
 	/** Create graph tracked SRV for a buffer from a descriptor. */
-	inline const FRDGBufferSRV* CreateSRV(const FRDGBufferSRVDesc& Desc)
+	inline FRDGBufferSRVRef CreateSRV(const FRDGBufferSRVDesc& Desc)
 	{
 		check(Desc.Buffer);
 		#if RENDER_GRAPH_DEBUGGING
@@ -216,7 +216,7 @@ public:
 	}
 
 	/** Create graph tracked UAV for a buffer from a descriptor. */
-	inline const FRDGBufferUAV* CreateUAV(const FRDGBufferUAVDesc& Desc)
+	inline FRDGBufferUAVRef CreateUAV(const FRDGBufferUAVDesc& Desc)
 	{
 		check(Desc.Buffer);
 		#if RENDER_GRAPH_DEBUGGING
@@ -276,7 +276,7 @@ public:
 	/**
 	 * Extracts an internal texture by handle.  Must be called before Execute.
 	 */
-	inline void GetInternalTexture(const FRDGTexture* Texture, TRefCountPtr<IPooledRenderTarget>* OutTexturePtr, bool bTransitionToRead = true)
+	inline void GetInternalTexture(FRDGTextureRef Texture, TRefCountPtr<IPooledRenderTarget>* OutTexturePtr, bool bTransitionToRead = true)
 	{
 		check(Texture);
 		check(OutTexturePtr);
@@ -297,10 +297,13 @@ public:
 	 */
 	void Execute();
 
-private:
-	/** The RHI command list to use. */
+
+public:
+	/** The RHI command list used for the render graph. */
 	FRHICommandListImmediate& RHICmdList;
 
+
+private:
 	/** Array of all pass created */
 	TArray<FRenderGraphPass*, SceneRenderingAllocator> Passes;
 
diff --git a/Engine/Source/Runtime/RenderCore/Public/RenderGraphResources.h b/Engine/Source/Runtime/RenderCore/Public/RenderGraphResources.h
--- a/Engine/Source/Runtime/RenderCore/Public/RenderGraphResources.h
+++ b/Engine/Source/Runtime/RenderCore/Public/RenderGraphResources.h
@@ -69,6 +69,9 @@ private:
 	friend void SetShaderUAVs(TRHICmdList&, const TShaderClass*, FComputeShaderRHIParamRef, const typename TShaderClass::FParameters&);
 };
 
+/** Descriptor of a graph tracked texture. */
+using FRDGTextureDesc = FPooledRenderTargetDesc;
+
 /** Render graph tracked Texture. */
 class RENDERCORE_API FRDGTexture : public FRDGResource
 {
@@ -402,3 +405,13 @@ FRDGBufferUAVDesc FRDGBufferUAVDesc::CreateIndirectDesc(const FRDGBuffer* Buffer
 	Desc.Format = PF_R32_UINT;
 	return Desc;
 }
+
+
+/** Defines the RDG resource references for user code not forgetting the const every time. */
+using FRDGResourceRef = const FRDGResource*;
+using FRDGTextureRef = const FRDGTexture*;
+using FRDGTextureSRVRef = const FRDGTextureSRV*;
+using FRDGTextureUAVRef = const FRDGTextureUAV*;
+using FRDGBufferRef = const  FRDGBuffer*;
+using FRDGBufferSRVRef = const FRDGBufferSRV*;
+using FRDGBufferUAVRef = const FRDGBufferUAV*;
diff --git a/Engine/Source/Runtime/RenderCore/Public/RenderGraphUtils.h b/Engine/Source/Runtime/RenderCore/Public/RenderGraphUtils.h
--- a/Engine/Source/Runtime/RenderCore/Public/RenderGraphUtils.h
+++ b/Engine/Source/Runtime/RenderCore/Public/RenderGraphUtils.h
@@ -3,8 +3,10 @@
 #pragma once
 
 #include "RenderGraphResources.h"
+#include "RenderGraphBuilder.h"
 #include "Shader.h"
-#include "UniformBuffer.h"
+#include "ShaderParameterStruct.h"
+#include "ShaderParameterMacros.h"
 
 
 /** Useful parameter struct that only have render targets.
@@ -26,3 +28,61 @@ inline void ClearUnusedGraphResources(const TShaderClass* Shader, typename TShad
 {
 	return ClearUnusedGraphResourcesImpl(Shader->Bindings, TShaderClass::FParameters::FTypeInfo::GetStructMetadata(), InoutParameters);
 }
+
+
+/** All utils for compute shaders.
+ */
+struct RENDERCORE_API FComputeShaderUtils
+{
+	/** Ideal size of group size 8x8 to occupy at least an entire wave on GCN, two warp on Nvidia. */
+	static constexpr int32 kGolden2DGroupSize = 8;
+
+	/** Compute the number of group to dispatch. */
+	static FIntVector GetGroupCount(const FIntPoint& ThreadCount, const FIntPoint& GroupSize)
+	{
+		return FIntVector(
+			FMath::DivideAndRoundUp(ThreadCount.X, GroupSize.X),
+			FMath::DivideAndRoundUp(ThreadCount.Y, GroupSize.Y),
+			1);
+	}
+	static FIntVector GetGroupCount(const FIntPoint& ThreadCount, const int32 GroupSize)
+	{
+		return FIntVector(
+			FMath::DivideAndRoundUp(ThreadCount.X, GroupSize),
+			FMath::DivideAndRoundUp(ThreadCount.Y, GroupSize),
+			1);
+	}
+
+
+	/** Dispatch a compute shader to rhi command list with its parameters. */
+	template<typename TShaderClass>
+	static inline void Dispatch(FRHICommandList& RHICmdList, const TShaderClass* ComputeShader, const typename TShaderClass::FParameters& Parameters, FIntVector GroupCount)
+	{
+		FRHIComputeShader* ShaderRHI = ComputeShader->GetComputeShader();
+		RHICmdList.SetComputeShader(ShaderRHI);
+		SetShaderParameters(RHICmdList, ComputeShader, ShaderRHI, Parameters);
+		RHICmdList.DispatchComputeShader(GroupCount.X, GroupCount.Y, GroupCount.Z);
+		UnsetShaderUAVs(RHICmdList, ComputeShader, ShaderRHI);
+	}
+
+	/** Dispatch a compute shader to render graph builder with its parameters. */
+	template<typename TShaderClass>
+	static inline void AddPass(
+		FRDGBuilder& GraphBuilder,
+		const TCHAR* PassName,
+		const TShaderClass* ComputeShader,
+		typename TShaderClass::FParameters* Parameters,
+		FIntVector GroupCount)
+	{
+		ClearUnusedGraphResources(ComputeShader, Parameters);
+
+		GraphBuilder.AddPass(
+			PassName,
+			Parameters,
+			ERenderGraphPassFlags::Compute,
+			[Parameters, ComputeShader, GroupCount](FRHICommandList& RHICmdList)
+		{
+			FComputeShaderUtils::Dispatch(RHICmdList, ComputeShader, *Parameters, GroupCount);
+		});
+	}
+};
diff --git a/Engine/Source/Runtime/Renderer/Private/SceneOcclusion.cpp b/Engine/Source/Runtime/Renderer/Private/SceneOcclusion.cpp
--- a/Engine/Source/Runtime/Renderer/Private/SceneOcclusion.cpp
+++ b/Engine/Source/Runtime/Renderer/Private/SceneOcclusion.cpp
@@ -1109,12 +1109,12 @@ void BuildHZB(FRHICommandListImmediate& RHICmdListImmediate, FViewInfo& View)
 	View.HZBMipmap0Size = HZBSize;
 
 	//@DW: Configure texture creation
-	FPooledRenderTargetDesc HZBDesc = FPooledRenderTargetDesc::Create2DDesc(HZBSize, PF_R16F, FClearValueBinding::None, TexCreate_None, TexCreate_RenderTargetable | TexCreate_ShaderResource | TexCreate_NoFastClear, false, NumMips);
+	FRDGTextureDesc HZBDesc = FRDGTextureDesc::Create2DDesc(HZBSize, PF_R16F, FClearValueBinding::None, TexCreate_None, TexCreate_RenderTargetable | TexCreate_ShaderResource | TexCreate_NoFastClear, false, NumMips);
 	HZBDesc.Flags |= GFastVRamConfig.HZB;
 
 	//@DW: Explicit creation of graph resource handles - full support for everything the RHI supports
 	//@DW: Now that we've created a resource handle, it will have to be passed around to other passes for manual wiring or put into a Blackboard structure for automatic wiring
-	const FRDGTexture* HZBTexture = GraphBuilder.CreateTexture(HZBDesc, TEXT("HZB"));
+	FRDGTextureRef HZBTexture = GraphBuilder.CreateTexture(HZBDesc, TEXT("HZB"));
 
 	{
 		FHZBBuildPassParameters* PassParameters;
@@ -1198,7 +1198,7 @@ void BuildHZB(FRHICommandListImmediate& RHICmdListImmediate, FViewInfo& View)
 
 		//@DW: Explicit creation of SRV, full configuration of SRV supported
 		FRDGTextureSRVDesc Desc(HZBTexture, MipIndex - 1);
-		const FRDGTextureSRV* ParentMipSRV = GraphBuilder.CreateSRV(Desc);
+		FRDGTextureSRVRef ParentMipSRV = GraphBuilder.CreateSRV(Desc);
 
 		FHZBBuildPassParameters* PassParameters;
 		GraphBuilder.CreateParameters(&PassParameters);
