Cherry-pick 4522720: Fixes badly initialized FShaderParameterBindings::RootParameterBufferIndex

diff --git a/Engine/Source/Runtime/RenderCore/Private/ShaderParameterStruct.cpp b/Engine/Source/Runtime/RenderCore/Private/ShaderParameterStruct.cpp
--- a/Engine/Source/Runtime/RenderCore/Private/ShaderParameterStruct.cpp
+++ b/Engine/Source/Runtime/RenderCore/Private/ShaderParameterStruct.cpp
@@ -212,6 +212,10 @@ void FShaderParameterBindings::BindForRootShaderParameters(const FShaderParamete
 			BindingContext.ShaderGlobalScopeBindings.Add(ShaderBindingName, ShaderBindingName);
 			RootParameterBufferIndex = BufferIndex;
 		}
+		else
+		{
+			check(RootParameterBufferIndex == FShaderParameterBindings::kInvalidBufferIndex);
+		}
 	}
 
 	TArray<FString> AllParameterNames;
diff --git a/Engine/Source/Runtime/RenderCore/Public/Shader.h b/Engine/Source/Runtime/RenderCore/Public/Shader.h
--- a/Engine/Source/Runtime/RenderCore/Public/Shader.h
+++ b/Engine/Source/Runtime/RenderCore/Public/Shader.h
@@ -637,7 +637,7 @@ public:
 	TArray<FParameterStructReference> ParameterReferences;
 
 	// Buffer index of FShaderParametersMetadata::kRootUniformBufferBindingName
-	uint16 RootParameterBufferIndex;
+	uint16 RootParameterBufferIndex = FShaderParameterBindings::kInvalidBufferIndex;
 
 	friend FArchive& operator<<(FArchive& Ar, FShaderParameterBindings& ParametersBindingData)
 	{
