Cherry-pick 4530038: Implements render graph event names and scopes to allow more precise pass name when debugging and hierarching pass names in GPU capture.

diff --git a/Engine/Source/Runtime/RenderCore/Private/RenderGraphBuilder.cpp b/Engine/Source/Runtime/RenderCore/Private/RenderGraphBuilder.cpp
--- a/Engine/Source/Runtime/RenderCore/Private/RenderGraphBuilder.cpp
+++ b/Engine/Source/Runtime/RenderCore/Private/RenderGraphBuilder.cpp
@@ -23,10 +23,31 @@ static const int32 GRenderGraphImmediateMode = 0;
 #endif
 
 
+#if RENDER_GRAPH_DRAW_EVENTS == 2
+
+FRDGEventName::FRDGEventName(const TCHAR* EventFormat, ...)
+{
+	if (GetEmitDrawEvents())
+	{
+		va_list ptr;
+		va_start(ptr, EventFormat);
+		TCHAR TempStr[256];
+		// Build the string in the temp buffer
+		FCString::GetVarArgs(TempStr, ARRAY_COUNT(TempStr), ARRAY_COUNT(TempStr) - 1, EventFormat, ptr);
+
+		EventName = TempStr;
+	}
+}
+
+#endif
+
+
 void FRDGBuilder::Execute()
 {
 	#if RENDER_GRAPH_DEBUGGING
 	{
+		checkf(CurrentScope == nullptr, TEXT("Render graph needs to have all scopes ended to execute."));
+
 		checkf(!bHasExecuted, TEXT("Render graph execution should only happen once."));
 	}
 	#endif
@@ -42,6 +63,16 @@ void FRDGBuilder::Execute()
 		}
 	}
 
+	// Pops remaining scopes
+	{
+		for (int32 i = 0; i < kMaxScopeCount; i++)
+		{
+			if (!ScopesStack[i])
+				break;
+			RHICmdList.PopEvent();
+		}
+	}
+
 	ProcessDeferredInternalResourceQueries();
 
 	DestructPasses();
@@ -514,6 +545,76 @@ static bool IsBoundAsReadable( const FRDGTexture* Texture, FShaderParameterStruc
 	return false;
 }
 
+void FRDGBuilder::PushDrawEventStack(const FRenderGraphPass* Pass)
+{
+	// Push the scope event.
+	{
+		// Find out how many scope events needs to be poped.
+		TStaticArray<const FRDGEventScope*, kMaxScopeCount> TraversedScopes;
+		int32 CommonScopeId = -1;
+		int32 TraversedScopeCount = 0;
+		const FRDGEventScope* PassParentScope = Pass->ParentScope;
+		while (PassParentScope)
+		{
+			TraversedScopes[TraversedScopeCount] = PassParentScope;
+
+			for (int32 i = 0; i < ScopesStack.Num(); i++)
+			{
+				if (ScopesStack[i] == PassParentScope)
+				{
+					CommonScopeId = i;
+					break;
+				}
+			}
+
+			if (CommonScopeId != -1)
+			{
+				break;
+			}
+
+			TraversedScopeCount++;
+			PassParentScope = PassParentScope->ParentScope;
+		}
+
+		// Pop no longer used scopes
+		for (int32 i = CommonScopeId + 1; CommonScopeId >= 0 && i < kMaxScopeCount; i++)
+		{
+			if (!ScopesStack[i])
+				break;
+
+			RHICmdList.PopEvent();
+			ScopesStack[i] = nullptr;
+		}
+
+		// Push new scopes
+		const FColor ScopeColor(0);
+		for (int32 i = TraversedScopeCount - 1; i >= 0; i--)
+		{
+			RHICmdList.PushEvent(TraversedScopes[i]->Name.GetTCHAR(), ScopeColor);
+			CommonScopeId++;
+			ScopesStack[CommonScopeId] = TraversedScopes[i];
+		}
+	}
+
+	// Push the pass's event with some color.
+	{
+		FColor Color(0, 0, 0);
+
+		if (Pass->IsCompute())
+		{
+			// Green for compute.
+			Color = FColor(128, 255, 128);
+		}
+		else
+		{
+			// Ref for rasterizer.
+			Color = FColor(255, 128, 128);
+		}
+
+		RHICmdList.PushEvent(Pass->GetName(), Color);
+	}
+}
+
 void FRDGBuilder::ExecutePass( const FRenderGraphPass* Pass )
 {
 	QUICK_SCOPE_CYCLE_COUNTER(STAT_FRDGBuilder_ExecutePass);
@@ -523,6 +624,11 @@ void FRDGBuilder::ExecutePass( const FRenderGraphPass* Pass )
 
 	AllocateAndTransitionPassResources(Pass, &RPInfo, &bHasRenderTargets);
 
+	if (RENDER_GRAPH_DRAW_EVENTS != 0 && GetEmitDrawEvents())
+	{
+		PushDrawEventStack(Pass);
+	}
+
 	if( !Pass->IsCompute())
 	{
 		check(bHasRenderTargets);
@@ -540,6 +646,11 @@ void FRDGBuilder::ExecutePass( const FRenderGraphPass* Pass )
 		RHICmdList.EndRenderPass();
 	}
 
+	if (RENDER_GRAPH_DRAW_EVENTS != 0 && GetEmitDrawEvents())
+	{
+		RHICmdList.PopEvent();
+	}
+
 	if (RENDER_GRAPH_DEBUGGING)
 	{
 		WarnForUselessPassDependencies(Pass);
@@ -912,6 +1023,17 @@ void FRDGBuilder::ProcessDeferredInternalResourceQueries()
 
 void FRDGBuilder::DestructPasses()
 {
+	#if RENDER_GRAPH_DRAW_EVENTS == 2
+	{
+		// Event scopes are allocated on FMemStack, so need to call their destructor because have a FString within them.
+		for (FRDGEventScope* EventScope : EventScopes)
+		{
+			EventScope->~FRDGEventScope();
+		}
+		EventScopes.Empty();
+	}
+	#endif
+
 	#if RENDER_GRAPH_DEBUGGING
 	{
 		// Make sure all resource references have been released to ensure no leaks happen.
diff --git a/Engine/Source/Runtime/RenderCore/Private/VisualizeTexture.cpp b/Engine/Source/Runtime/RenderCore/Private/VisualizeTexture.cpp
--- a/Engine/Source/Runtime/RenderCore/Private/VisualizeTexture.cpp
+++ b/Engine/Source/Runtime/RenderCore/Private/VisualizeTexture.cpp
@@ -273,7 +273,7 @@ void FVisualizeTexture::CreateContentCapturePass(FRDGBuilder& GraphBuilder, cons
 
 	FPixelShaderUtils::AddFullscreenPass(
 		GraphBuilder, ShaderMap,
-		TEXT("VisualizeTexture"),
+		RDG_EVENT_NAME("VisualizeTextureCapture(%s)", SrcTexture->Name),
 		*PixelShader, PassParameters);
 
 	// Save the copied texture and descriptor about original informations.
diff --git a/Engine/Source/Runtime/RenderCore/Public/PixelShaderUtils.h b/Engine/Source/Runtime/RenderCore/Public/PixelShaderUtils.h
--- a/Engine/Source/Runtime/RenderCore/Public/PixelShaderUtils.h
+++ b/Engine/Source/Runtime/RenderCore/Public/PixelShaderUtils.h
@@ -52,14 +52,14 @@ struct RENDERCORE_API FPixelShaderUtils
 	static inline void AddFullscreenPass(
 		FRDGBuilder& GraphBuilder,
 		const TShaderMap<FGlobalShaderType>* GlobalShaderMap,
-		const TCHAR* PassName,
+		FRDGEventName&& PassName,
 		const TShaderClass* PixelShader,
 		typename TShaderClass::FParameters* Parameters)
 	{
 		ClearUnusedGraphResources(PixelShader, Parameters);
 
 		GraphBuilder.AddPass(
-			PassName,
+			Forward<FRDGEventName>(PassName),
 			Parameters,
 			ERenderGraphPassFlags::None,
 			[Parameters, GlobalShaderMap, PixelShader](FRHICommandList& RHICmdList)
diff --git a/Engine/Source/Runtime/RenderCore/Public/RenderGraphBuilder.h b/Engine/Source/Runtime/RenderCore/Public/RenderGraphBuilder.h
--- a/Engine/Source/Runtime/RenderCore/Public/RenderGraphBuilder.h
+++ b/Engine/Source/Runtime/RenderCore/Public/RenderGraphBuilder.h
@@ -10,6 +10,77 @@
 /** Whether render graph debugging is compiled. */
 #define RENDER_GRAPH_DEBUGGING (DO_CHECK)
 
+/** Whether render graph should support draw events or not.
+ * RENDER_GRAPH_DRAW_EVENTS == 0 means there is no string processing at all.
+ * RENDER_GRAPH_DRAW_EVENTS == 1 means const TCHAR* is only passdown.
+ * RENDER_GRAPH_DRAW_EVENTS == 2 means complex formated FString is passdown.
+ */
+#if WITH_PROFILEGPU
+	#define RENDER_GRAPH_DRAW_EVENTS 2
+#else
+	#define RENDER_GRAPH_DRAW_EVENTS 0
+#endif
+
+
+/** Opaque object to store a draw event. */
+class RENDERCORE_API FRDGEventName
+{
+public:
+	inline FRDGEventName() 
+	{ }
+
+	#if RENDER_GRAPH_DRAW_EVENTS == 2
+
+	explicit FRDGEventName(const TCHAR* EventFormat, ...);
+
+	#elif RENDER_GRAPH_DRAW_EVENTS == 1
+
+	explicit inline FRDGEventName(const TCHAR* EventFormat)
+		: EventName(EventFormat)
+	{ }
+
+	#endif
+
+	const TCHAR* GetTCHAR() const
+	{
+		#if RENDER_GRAPH_DRAW_EVENTS == 2
+			return *EventName;
+		#elif RENDER_GRAPH_DRAW_EVENTS == 1
+			return EventName;
+		#else
+			return TEXT("UnknownRDVEvent");
+		#endif
+	}
+
+private:
+	#if RENDER_GRAPH_DRAW_EVENTS == 2
+		FString EventName;
+	#elif RENDER_GRAPH_DRAW_EVENTS == 1
+		const TCHAR* EventName;
+	#endif
+};
+
+
+/** Hierarchical scope for draw events of passes. */
+class RENDERCORE_API FRDGEventScope
+{
+private:
+	// Pointer towards this one is contained in.
+	const FRDGEventScope* const ParentScope;
+
+	// Name of the event.
+	const FRDGEventName Name;
+
+
+	FRDGEventScope(const FRDGEventScope* InParentScope, FRDGEventName&& InName)
+		: ParentScope(InParentScope), Name(InName)
+	{ }
+
+
+	friend class FRDGBuilder;
+	friend class FStackRDGEventScopeRef;
+};
+
 
 /** Flags to anotate passes. */
 enum class ERenderGraphPassFlags
@@ -42,8 +113,9 @@ struct RENDERCORE_API FShaderParameterStructRef
  */
 struct RENDERCORE_API FRenderGraphPass
 {
-	FRenderGraphPass(const TCHAR* InName, FShaderParameterStructRef InParameterStruct, ERenderGraphPassFlags InPassFlags)
-		: Name(InName)
+	FRenderGraphPass(FRDGEventName&& InName, const FRDGEventScope* InParentScope, FShaderParameterStructRef InParameterStruct, ERenderGraphPassFlags InPassFlags)
+		: Name(static_cast<FRDGEventName&&>(InName))
+		, ParentScope(InParentScope)
 		, ParameterStruct(InParameterStruct)
 		, PassFlags(InPassFlags)
 	{
@@ -57,16 +129,28 @@ struct RENDERCORE_API FRenderGraphPass
 
 	virtual void Execute(FRHICommandListImmediate& RHICmdList) const = 0;
 
-	const TCHAR*			GetName() const { return Name; }
-	ERenderGraphPassFlags	GetFlags() const { return PassFlags; }
-	bool					IsCompute() const { return (PassFlags & ERenderGraphPassFlags::Compute) == ERenderGraphPassFlags::Compute; }
+	const TCHAR* GetName() const {
+		return Name.GetTCHAR();
+	}
+
+	const ERenderGraphPassFlags& GetFlags() const
+	{
+		return PassFlags;
+	}
+
+	bool IsCompute() const {
+		return (PassFlags & ERenderGraphPassFlags::Compute) == ERenderGraphPassFlags::Compute;
+	}
 	
 	FShaderParameterStructRef GetParameters() const { return ParameterStruct; }
 
 protected:
-	const TCHAR* const Name;
-	FShaderParameterStructRef ParameterStruct;
+	const FRDGEventName Name;
+	const FRDGEventScope* const ParentScope;
+	const FShaderParameterStructRef ParameterStruct;
 	const ERenderGraphPassFlags PassFlags;
+
+	friend class FRDGBuilder;
 };
 
 /** 
@@ -75,11 +159,10 @@ protected:
 template <typename ParameterStructType, typename ExecuteLambdaType>
 struct TLambdaRenderPass final : public FRenderGraphPass
 {
-	TLambdaRenderPass(const TCHAR* InName, FShaderParameterStructRef InParameterStruct, ERenderGraphPassFlags InPassFlags, ExecuteLambdaType&& InExecuteLambda)
-		: FRenderGraphPass(InName, InParameterStruct, InPassFlags)
+	TLambdaRenderPass(FRDGEventName&& InName, const FRDGEventScope* InParentScope, FShaderParameterStructRef InParameterStruct, ERenderGraphPassFlags InPassFlags, ExecuteLambdaType&& InExecuteLambda)
+		: FRenderGraphPass(static_cast<FRDGEventName&&>(InName), InParentScope, InParameterStruct, InPassFlags)
 		, ExecuteLambda(static_cast<ExecuteLambdaType&&>(InExecuteLambda))
-	{
-	}
+	{ }
 
 	~TLambdaRenderPass()
 	{
@@ -108,7 +191,10 @@ public:
 	/** A RHI cmd list is required, if using the immediate mode. */
 	FRDGBuilder(FRHICommandListImmediate& InRHICmdList)
 		: RHICmdList(InRHICmdList)
-	{ }
+	{
+		for (int32 i = 0; i < kMaxScopeCount; i++)
+			ScopesStack[i] = nullptr;
+	}
 
 	~FRDGBuilder()
 	{
@@ -251,18 +337,18 @@ public:
 	 */
 	template<typename ParameterStructType, typename ExecuteLambdaType>
 	void AddPass(
-		const TCHAR* Name, 
+		FRDGEventName&& Name, 
 		ParameterStructType* ParameterStruct,
 		ERenderGraphPassFlags Flags,
 		ExecuteLambdaType&& ExecuteLambda)
 	{
 		#if RENDER_GRAPH_DEBUGGING
 		{
-			checkf(!bHasExecuted, TEXT("Render graph pass %s needs to be added before the builder execution."), Name);
+			checkf(!bHasExecuted, TEXT("Render graph pass %s needs to be added before the builder execution."), Name.GetTCHAR());
 		}
 		#endif
 		auto NewPass = new(FMemStack::Get()) TLambdaRenderPass<ParameterStructType, ExecuteLambdaType>(
-			Name,
+			static_cast<FRDGEventName&&>(Name), CurrentScope,
 			{ ParameterStruct, &ParameterStructType::FTypeInfo::GetStructMetadata()->GetLayout() },
 			Flags,
 			static_cast<ExecuteLambdaType&&>(ExecuteLambda) );
@@ -304,6 +390,8 @@ public:
 
 
 private:
+	static constexpr int32 kMaxScopeCount = 8;
+
 	/** Array of all pass created */
 	TArray<FRenderGraphPass*, SceneRenderingAllocator> Passes;
 
@@ -322,6 +410,16 @@ private:
 	};
 	TArray<FDeferredInternalTextureQuery, SceneRenderingAllocator> DeferredInternalTextureQueries;
 
+	#if RENDER_GRAPH_DRAW_EVENTS == 2
+		/** All scopes allocated that needs to be arround to call destructors. */
+		TArray<FRDGEventScope*, SceneRenderingAllocator> EventScopes;
+	#endif
+
+	/** The current event scope as creating passes. */
+	const FRDGEventScope* CurrentScope = nullptr;
+
+	/** Stacks of scopes pushed to the RHI command list. */
+	TStaticArray<const FRDGEventScope*, kMaxScopeCount> ScopesStack;
 
 	#if RENDER_GRAPH_DEBUGGING
 		/** Whether the Execute() has already been called. */
@@ -348,6 +446,7 @@ private:
 	void TransitionTexture( const FRDGTexture* Texture, EResourceTransitionAccess TransitionAccess, bool bRequiredCompute ) const;
 	void TransitionUAV(FUnorderedAccessViewRHIParamRef UAV, const FRDGResource* UnderlyingResource, EResourceTransitionAccess TransitionAccess, bool bRequiredCompute ) const;
 
+	void PushDrawEventStack(const FRenderGraphPass* Pass);
 	void ExecutePass( const FRenderGraphPass* Pass );
 	void AllocateAndTransitionPassResources(const FRenderGraphPass* Pass, struct FRHIRenderPassInfo* OutRPInfo, bool* bOutHasRenderTargets);
 	static void WarnForUselessPassDependencies(const FRenderGraphPass* Pass);
@@ -358,4 +457,72 @@ private:
 
 	void ProcessDeferredInternalResourceQueries();
 	void DestructPasses();
+
+	friend class FStackRDGEventScopeRef;
 }; // class FRDGBuilder
+
+
+#if RENDER_GRAPH_DRAW_EVENTS
+
+/** Stack reference of render graph scope. */
+class RENDERCORE_API FStackRDGEventScopeRef
+{
+public:
+	FStackRDGEventScopeRef() = delete;
+	FStackRDGEventScopeRef(const FStackRDGEventScopeRef&) = delete;
+	FStackRDGEventScopeRef(FStackRDGEventScopeRef&&) = delete;
+	void operator = (const FStackRDGEventScopeRef&) = delete;
+
+	inline FStackRDGEventScopeRef(FRDGBuilder& InGraphBuilder, FRDGEventName&& ScopeName)
+		: GraphBuilder(InGraphBuilder)
+	{
+		checkf(!GraphBuilder.bHasExecuted, TEXT("Render graph bulider has already been executed."));
+
+		auto NewScope = new(FMemStack::Get()) FRDGEventScope(GraphBuilder.CurrentScope, Forward<FRDGEventName>(ScopeName));
+
+		#if RENDER_GRAPH_DRAW_EVENTS == 2
+		{
+			GraphBuilder.EventScopes.Add(NewScope);
+		}
+		#endif
+
+		GraphBuilder.CurrentScope = NewScope;
+	}
+
+	inline ~FStackRDGEventScopeRef()
+	{
+		check(GraphBuilder.CurrentScope != nullptr);
+		GraphBuilder.CurrentScope = GraphBuilder.CurrentScope->ParentScope;
+	}
+
+private:
+	FRDGBuilder& GraphBuilder;
+};
+
+#endif // RENDER_GRAPH_DRAW_EVENTS
+
+
+/** Macros for create render graph event names and scopes.
+ *
+ *		FRDGEventName Name = RDG_EVENT_NAME("MyPass %sx%s", ViewRect.Width(), ViewRect.Height());
+ *
+ *		RDG_EVENT_SCOPE(GraphBuilder, "MyProcessing %sx%s", ViewRect.Width(), ViewRect.Height());
+ */
+#if RENDER_GRAPH_DRAW_EVENTS == 2
+
+#define RDG_EVENT_NAME(Format, ...) FRDGEventName(TEXT(Format), ##__VA_ARGS__)
+#define RDG_EVENT_SCOPE(GraphBuilder, Format, ...) \
+	FStackRDGEventScopeRef PREPROCESSOR_JOIN(__RDG_ScopeRef_,__LINE__) ((GraphBuilder), RDG_EVENT_NAME(Format, ##__VA_ARGS__))
+
+#elif RENDER_GRAPH_DRAW_EVENTS == 1
+
+#define RDG_EVENT_NAME(Format, ...) FRDGEventName(TEXT(Format))
+#define RDG_EVENT_SCOPE(GraphBuilder, Format, ...) \
+	FStackRDGEventScopeRef PREPROCESSOR_JOIN(__RDG_ScopeRef_,__LINE__) ((GraphBuilder), RDG_EVENT_NAME(Format, ##__VA_ARGS__))
+
+#else // !RENDER_GRAPH_DRAW_EVENTS
+
+#define RDG_EVENT_NAME(Format, ...) FRDGEventName()
+#define RDG_EVENT_SCOPE(GraphBuilder, Format, ...) 
+
+#endif // !RENDER_GRAPH_DRAW_EVENTS
diff --git a/Engine/Source/Runtime/RenderCore/Public/RenderGraphUtils.h b/Engine/Source/Runtime/RenderCore/Public/RenderGraphUtils.h
--- a/Engine/Source/Runtime/RenderCore/Public/RenderGraphUtils.h
+++ b/Engine/Source/Runtime/RenderCore/Public/RenderGraphUtils.h
@@ -69,7 +69,7 @@ struct RENDERCORE_API FComputeShaderUtils
 	template<typename TShaderClass>
 	static inline void AddPass(
 		FRDGBuilder& GraphBuilder,
-		const TCHAR* PassName,
+		FRDGEventName&& PassName,
 		const TShaderClass* ComputeShader,
 		typename TShaderClass::FParameters* Parameters,
 		FIntVector GroupCount)
@@ -77,7 +77,7 @@ struct RENDERCORE_API FComputeShaderUtils
 		ClearUnusedGraphResources(ComputeShader, Parameters);
 
 		GraphBuilder.AddPass(
-			PassName,
+			Forward<FRDGEventName>(PassName),
 			Parameters,
 			ERenderGraphPassFlags::Compute,
 			[Parameters, ComputeShader, GroupCount](FRHICommandList& RHICmdList)
diff --git a/Engine/Source/Runtime/Renderer/Private/DeferredShadingRenderer.cpp b/Engine/Source/Runtime/Renderer/Private/DeferredShadingRenderer.cpp
--- a/Engine/Source/Runtime/Renderer/Private/DeferredShadingRenderer.cpp
+++ b/Engine/Source/Runtime/Renderer/Private/DeferredShadingRenderer.cpp
@@ -398,7 +398,7 @@ void FDeferredShadingSceneRenderer::RenderFinish(FRHICommandListImmediate& RHICm
 	FSceneRenderTargets::Get(RHICmdList).SetLightAttenuation(0);
 }
 
-void BuildHZB( FRHICommandListImmediate& RHICmdList, FViewInfo& View );
+void BuildHZB( FRDGBuilder& GraphBuilder, FViewInfo& View );
 
 /** 
 * Renders the view family. 
@@ -453,7 +453,12 @@ bool FDeferredShadingSceneRenderer::RenderHzb(FRHICommandListImmediate& RHICmdLi
 
 		if (bSSAO || bHZBOcclusion || bSSR)
 		{
-			BuildHZB(RHICmdList, Views[ViewIndex]);
+			FRDGBuilder GraphBuilder(RHICmdList);
+			{
+				RDG_EVENT_SCOPE(GraphBuilder, "BuildHZB(ViewId=%d)", ViewIndex);
+				BuildHZB(GraphBuilder, Views[ViewIndex]);
+			}
+			GraphBuilder.Execute();
 		}
 
 		if (bHZBOcclusion && ViewState && ViewState->HZBOcclusionTests.GetNum() != 0)
diff --git a/Engine/Source/Runtime/Renderer/Private/SceneOcclusion.cpp b/Engine/Source/Runtime/Renderer/Private/SceneOcclusion.cpp
--- a/Engine/Source/Runtime/Renderer/Private/SceneOcclusion.cpp
+++ b/Engine/Source/Runtime/Renderer/Private/SceneOcclusion.cpp
@@ -1093,12 +1093,10 @@ class FHZBBuildPS : public FGlobalShader
 IMPLEMENT_GLOBAL_SHADER(FHZBBuildPS, "/Engine/Private/HZBOcclusion.usf", "HZBBuildPS", SF_Pixel);
 
 
-void BuildHZB(FRHICommandListImmediate& RHICmdListImmediate, FViewInfo& View)
+void BuildHZB(FRDGBuilder& GraphBuilder, FViewInfo& View)
 {
 	QUICK_SCOPE_CYCLE_COUNTER(STAT_BuildHZB);
 	
-	FRDGBuilder GraphBuilder(RHICmdListImmediate);
-
 	// View.ViewRect.{Width,Height}() are most likely to be < 2^24, so the float
 	// conversion won't loss any precision (assuming float have 23bits for mantissa)
 	const int32 NumMipsX = FMath::Max(FPlatformMath::CeilToInt(FMath::Log2(float(View.ViewRect.Width()))) - 1, 1);
@@ -1123,7 +1121,8 @@ void BuildHZB(FRHICommandListImmediate& RHICmdListImmediate, FViewInfo& View)
 		PassParameters->RenderTargets[0] = FRenderTargetBinding( HZBTexture, ERenderTargetLoadAction::ENoAction, ERenderTargetStoreAction::EStore);
 
 		//@DW - this pass only reads external textures, we don't have any graph inputs
-		GraphBuilder.AddPass(TEXT("HZB_SetupMip0"),
+		GraphBuilder.AddPass(
+			RDG_EVENT_NAME("HZB(mip=0) %dx%d", HZBSize.X, HZBSize.Y),
 			PassParameters,
 			ERenderGraphPassFlags::None,
 			[PassParameters, &View, HZBSize](FRHICommandListImmediate& RHICmdList)
@@ -1133,9 +1132,6 @@ void BuildHZB(FRHICommandListImmediate& RHICmdListImmediate, FViewInfo& View)
 				GraphicsPSOInit.RasterizerState = TStaticRasterizerState<>::GetRHI();
 				GraphicsPSOInit.DepthStencilState = TStaticDepthStencilState<false, CF_Always>::GetRHI();
 
-				// Mip 0
-				SCOPED_DRAW_EVENTF(RHICmdList, BuildHZB, TEXT("HZB SetupMip 0 %dx%d"), HZBSize.X, HZBSize.Y);
-
 				RHICmdList.ApplyCachedRenderTargets(GraphicsPSOInit);
 				GraphicsPSOInit.PrimitiveType = PT_TriangleList;
 
@@ -1207,13 +1203,12 @@ void BuildHZB(FRHICommandListImmediate& RHICmdListImmediate, FViewInfo& View)
 		PassParameters->Texture = ParentMipSRV;
 		PassParameters->TextureSampler = TStaticSamplerState<SF_Point, AM_Clamp, AM_Clamp, AM_Clamp>::GetRHI();
 
-		GraphBuilder.AddPass(TEXT("HZB_Mip"),
+		GraphBuilder.AddPass(
+			RDG_EVENT_NAME("HZB(mip=%d) %dx%d", MipIndex, DstSize.X, DstSize.Y),
 			PassParameters,
 			ERenderGraphPassFlags::None,
 			[PassParameters, SrcSize, DstSize, &View](FRHICommandListImmediate& RHICmdList)
 			{
-				SCOPED_DRAW_EVENTF(RHICmdList, BuildHZB, TEXT("HZB MipLevel %dx%d"), DstSize.X, DstSize.Y);
-
 				FHZBBuildPS::FPermutationDomain PermutationVector;
 				PermutationVector.Set<FHZBBuildPS::FStageDim>(true);
 
@@ -1260,7 +1255,6 @@ void BuildHZB(FRHICommandListImmediate& RHICmdListImmediate, FViewInfo& View)
 	}
 
 	GraphBuilder.GetInternalTexture(HZBTexture, &View.HZB);
-	GraphBuilder.Execute();
 }
 
 struct FViewOcclusionQueries
diff --git a/Engine/Source/Runtime/Renderer/Private/VolumetricFog.cpp b/Engine/Source/Runtime/Renderer/Private/VolumetricFog.cpp
--- a/Engine/Source/Runtime/Renderer/Private/VolumetricFog.cpp
+++ b/Engine/Source/Runtime/Renderer/Private/VolumetricFog.cpp
@@ -502,12 +502,12 @@ void FDeferredShadingSceneRenderer::RenderLocalLightsForVolumetricFog(
 		GraphBuilder.CreateParameters(&PassParameters);
 		PassParameters->RenderTargets[0] = FRenderTargetBinding( OutLocalShadowedLightScattering, ERenderTargetLoadAction::EClear, ERenderTargetStoreAction::ENoAction );
 
-		GraphBuilder.AddPass(TEXT("LocalShadowedLightScattering"),
+		GraphBuilder.AddPass(
+			RDG_EVENT_NAME("ShadowedLights"),
 			PassParameters,
 			ERenderGraphPassFlags::None,
 			[PassParameters, &View, this, LightsToInject, VolumetricFogGridSize, GridZParams, bUseTemporalReprojection, IntegrationData, FogInfo](FRHICommandListImmediate& RHICmdList)
 			{
-				SCOPED_DRAW_EVENT(RHICmdList, ShadowedLights);
 				for (int32 LightIndex = 0; LightIndex < LightsToInject.Num(); LightIndex++)
 				{
 					const FLightSceneInfo* LightSceneInfo = LightsToInject[LightIndex];
@@ -1048,13 +1048,13 @@ void FDeferredShadingSceneRenderer::ComputeVolumetricFog(FRHICommandListImmediat
 				ClearUnusedGraphResources(ComputeShader, PassParameters);
 
 				//@DW - this pass only reads external textures, we don't have any graph inputs
-				GraphBuilder.AddPass(TEXT("VolumetricFogMaterialSetup"),
+				GraphBuilder.AddPass(
+					RDG_EVENT_NAME("InitializeVolumeAttributes"),
 					PassParameters,
 					ERenderGraphPassFlags::Compute,
 					[PassParameters, &View, VolumetricFogGridSize, IntegrationData, ComputeShader](FRHICommandListImmediate& RHICmdList)
 					{
 						const FIntVector NumGroups = FIntVector::DivideAndRoundUp(VolumetricFogGridSize, VolumetricFogGridInjectionGroupSize);
-						SCOPED_DRAW_EVENT(RHICmdList, InitializeVolumeAttributes);
 
 						RHICmdList.SetComputeShader(ComputeShader->GetComputeShader());
 
@@ -1110,7 +1110,13 @@ void FDeferredShadingSceneRenderer::ComputeVolumetricFog(FRHICommandListImmediat
 				auto ComputeShader = View.ShaderMap->GetShader< TVolumetricFogLightScatteringCS >(PermutationVector);
 				ClearUnusedGraphResources(ComputeShader, PassParameters);
 
-				GraphBuilder.AddPass(TEXT("LightScattering"),
+				GraphBuilder.AddPass(
+					RDG_EVENT_NAME("LightScattering %dx%dx%d %s %s",
+						VolumetricFogGridSize.X,
+						VolumetricFogGridSize.Y,
+						VolumetricFogGridSize.Z,
+						bUseDistanceFieldSkyOcclusion ? TEXT("DFAO") : TEXT(""),
+						PassParameters->LightFunctionTexture ? TEXT("LF") : TEXT("")),
 					PassParameters,
 					ERenderGraphPassFlags::Compute,
 					[PassParameters, ComputeShader, &View, this, FogInfo, bUseTemporalReprojection, VolumetricFogGridSize, IntegrationData, bUseDirectionalLightShadowing, bUseDistanceFieldSkyOcclusion, LightFunctionWorldToShadow](FRHICommandListImmediate& RHICmdList)
@@ -1118,13 +1124,6 @@ void FDeferredShadingSceneRenderer::ComputeVolumetricFog(FRHICommandListImmediat
 						UnbindRenderTargets(RHICmdList);
 						const FIntVector NumGroups = FIntVector::DivideAndRoundUp(VolumetricFogGridSize, VolumetricFogGridInjectionGroupSize);
 
-						SCOPED_DRAW_EVENTF(RHICmdList, LightScattering, TEXT("LightScattering %dx%dx%d %s %s"), 
-							VolumetricFogGridSize.X, 
-							VolumetricFogGridSize.Y, 
-							VolumetricFogGridSize.Z,
-							bUseDistanceFieldSkyOcclusion ? TEXT("DFAO") : TEXT(""),
-							PassParameters->LightFunctionTexture ? TEXT("LF") : TEXT(""));
-
 						RHICmdList.SetComputeShader(ComputeShader->GetComputeShader());
 
 						FTextureRHIParamRef LightScatteringHistoryTexture = View.ViewState->LightScatteringHistory.IsValid()
@@ -1148,13 +1147,12 @@ void FDeferredShadingSceneRenderer::ComputeVolumetricFog(FRHICommandListImmediat
 				PassParameters->LightScattering = IntegrationData.LightScattering;
 				PassParameters->RWIntegratedLightScattering = IntegratedLightScatteringUAV;
 
-				GraphBuilder.AddPass(TEXT("FinalIntegration"),
+				GraphBuilder.AddPass(
+					RDG_EVENT_NAME("FinalIntegration"),
 					PassParameters,
 					ERenderGraphPassFlags::Compute,
 					[PassParameters, &View, VolumetricFogGridSize, IntegrationData, this](FRHICommandListImmediate& RHICmdList)
 					{
-						SCOPED_DRAW_EVENT(RHICmdList, FinalIntegration);
-
 						const FIntVector NumGroups = FIntVector::DivideAndRoundUp(VolumetricFogGridSize, VolumetricFogIntegrationGroupSize);
 
 						auto ComputeShader = View.ShaderMap->GetShader< FVolumetricFogFinalIntegrationCS >();
diff --git a/Engine/Source/Runtime/Renderer/Private/VolumetricFogLightFunction.cpp b/Engine/Source/Runtime/Renderer/Private/VolumetricFogLightFunction.cpp
--- a/Engine/Source/Runtime/Renderer/Private/VolumetricFogLightFunction.cpp
+++ b/Engine/Source/Runtime/Renderer/Private/VolumetricFogLightFunction.cpp
@@ -202,7 +202,8 @@ void FDeferredShadingSceneRenderer::RenderLightFunctionForVolumetricFog(
 			GraphBuilder.CreateParameters(&PassParameters);
 			PassParameters->RenderTargets[0] = FRenderTargetBinding( OutLightFunctionTexture, ERenderTargetLoadAction::ENoAction, ERenderTargetStoreAction::EStore );
 
-			GraphBuilder.AddPass(TEXT("VolumetricFogLightFunction"),
+			GraphBuilder.AddPass(
+				RDG_EVENT_NAME("LightFunction"),
 				PassParameters,
 				ERenderGraphPassFlags::None,
 				[PassParameters, &View, MaterialProxy, LightFunctionResolution, DirectionalLightSceneInfo, WorldToShadowValue, this](FRHICommandListImmediate& RHICmdList)
diff --git a/Engine/Source/Runtime/Renderer/Private/VolumetricFogVoxelization.cpp b/Engine/Source/Runtime/Renderer/Private/VolumetricFogVoxelization.cpp
--- a/Engine/Source/Runtime/Renderer/Private/VolumetricFogVoxelization.cpp
+++ b/Engine/Source/Runtime/Renderer/Private/VolumetricFogVoxelization.cpp
@@ -696,13 +696,12 @@ void FDeferredShadingSceneRenderer::VoxelizeFogVolumePrimitives(
 		PassParameters->RenderTargets[0] = FRenderTargetBinding( IntegrationData.VBufferA, ERenderTargetLoadAction::ENoAction, ERenderTargetStoreAction::EStore );
 		PassParameters->RenderTargets[1] = FRenderTargetBinding( IntegrationData.VBufferB, ERenderTargetLoadAction::ENoAction, ERenderTargetStoreAction::EStore );
 
-		GraphBuilder.AddPass(TEXT("VoxelizeVolumePrimitives"),
+		GraphBuilder.AddPass(
+			RDG_EVENT_NAME("VoxelizeVolumePrimitives"),
 			PassParameters,
 			ERenderGraphPassFlags::None,
 			[PassParameters, &View, VolumetricFogGridSize, IntegrationData, VolumetricFogDistance, GridZParams](FRHICommandListImmediate& RHICmdList)
 			{
-				SCOPED_DRAW_EVENT(RHICmdList, VoxelizeVolumePrimitives);
-
 				FViewUniformShaderParameters VoxelizeParameters = *View.CachedViewUniformShaderParameters;
 
 				// Update the parts of VoxelizeParameters which are dependent on the buffer size and view rect
