Cherry-pick 4544349: Adds early shader parameters validation in ClearUnusedGraphResources() to crash even before creating a lambda shader pass when high level have a bug.

diff --git a/Engine/Shaders/Public/ShaderVersion.ush b/Engine/Shaders/Public/ShaderVersion.ush
--- a/Engine/Shaders/Public/ShaderVersion.ush
+++ b/Engine/Shaders/Public/ShaderVersion.ush
@@ -2,4 +2,4 @@
 // Each time the console command is executed it generates a new GUID. As this file is included
 // in Platform.ush (which should be included in any shader) it allows to invalidate the shader DDC.
 // 
-// GUID = 8301458A46E1B310E41BA7B2A499BFC1aaaaa
\ No newline at end of file
+// GUID = 2312D6114FB4AD900F22DB942A5FE652
\ No newline at end of file
diff --git a/Engine/Source/Runtime/RenderCore/Private/RenderGraphBuilder.cpp b/Engine/Source/Runtime/RenderCore/Private/RenderGraphBuilder.cpp
--- a/Engine/Source/Runtime/RenderCore/Private/RenderGraphBuilder.cpp
+++ b/Engine/Source/Runtime/RenderCore/Private/RenderGraphBuilder.cpp
@@ -639,6 +639,9 @@ void FRDGBuilder::ExecutePass( const FRenderGraphPass* Pass )
 		UnbindRenderTargets(RHICmdList);
 	}
 	
+	// The name of the pass just for debuging convenience when crashing in the Execute().
+	const TCHAR* PassName = Pass->GetName();
+
 	Pass->Execute(RHICmdList);
 
 	if( bHasRenderTargets )
diff --git a/Engine/Source/Runtime/RenderCore/Private/ShaderParameterStruct.cpp b/Engine/Source/Runtime/RenderCore/Private/ShaderParameterStruct.cpp
--- a/Engine/Source/Runtime/RenderCore/Private/ShaderParameterStruct.cpp
+++ b/Engine/Source/Runtime/RenderCore/Private/ShaderParameterStruct.cpp
@@ -157,7 +157,7 @@ struct FShaderParameterStructBindingContext
 	}
 }; // struct FShaderParameterStructBindingContext
 
-void FShaderParameterBindings::BindForLegacyShaderParameters(const FShaderParameterMap& ParametersMap, const FShaderParametersMetadata& StructMetaData)
+void FShaderParameterBindings::BindForLegacyShaderParameters(const FShaderParameterMap& ParametersMap, const FShaderParametersMetadata& StructMetaData, bool bShouldBindEverything)
 {
 	checkf(StructMetaData.GetSize() < (1 << (sizeof(uint16) * 8)), TEXT("Shader parameter structure can only have a size < 65536 bytes."));
 
@@ -174,7 +174,7 @@ void FShaderParameterBindings::BindForLegacyShaderParameters(const FShaderParame
 
 	TArray<FString> AllParameterNames;
 	ParametersMap.GetAllParameterNames(AllParameterNames);
-	if (0 && BindingContext.ShaderGlobalScopeBindings.Num() != AllParameterNames.Num()) // TODO: enable that once we can.
+	if (bShouldBindEverything && BindingContext.ShaderGlobalScopeBindings.Num() != AllParameterNames.Num())
 	{
 		UE_LOG(LogShaders, Error, TEXT("%i shader parameters have not been bound:"), AllParameterNames.Num() - BindingContext.ShaderGlobalScopeBindings.Num());
 		for (const FString& GlobalParameterName : AllParameterNames)
@@ -236,7 +236,7 @@ void FShaderParameterBindings::BindForRootShaderParameters(const FShaderParamete
 	}
 }
 
-void SetNullShaderParameterFatalError(const FShader* Shader, const FShaderParametersMetadata* ParametersMetadata, uint16 MemberOffset)
+void EmitNullShaderParameterFatalError(const FShader* Shader, const FShaderParametersMetadata* ParametersMetadata, uint16 MemberOffset)
 {
 	const FShaderParametersMetadata* MemberContainingStruct = nullptr;
 	const FShaderParametersMetadata::FMember* Member = nullptr;
@@ -250,3 +250,84 @@ void SetNullShaderParameterFatalError(const FShader* Shader, const FShaderParame
 		MemberContainingStruct->GetStructTypeName(),
 		Member->GetName());
 }
+
+#if DO_CHECK
+
+void ValidateShaderParameters(const FShader* Shader, const FShaderParametersMetadata* ParametersMetadata, const void* Parameters)
+{
+	const FShaderParameterBindings& Bindings = Shader->Bindings;
+	const uint8* Base = reinterpret_cast<const uint8*>(Parameters);
+
+	// Textures
+	for (const FShaderParameterBindings::FResourceParameter& ParameterBinding : Bindings.Textures)
+	{
+		auto ShaderParameterRef = *reinterpret_cast<const FTextureRHIParamRef*>(Base + ParameterBinding.ByteOffset);
+		if (!ShaderParameterRef)
+		{
+			EmitNullShaderParameterFatalError(Shader, ParametersMetadata, ParameterBinding.ByteOffset);
+		}
+	}
+
+	// SRVs
+	for (const FShaderParameterBindings::FResourceParameter& ParameterBinding : Bindings.SRVs)
+	{
+		auto ShaderParameterRef = *reinterpret_cast<const FShaderResourceViewRHIParamRef*>(Base + ParameterBinding.ByteOffset);
+		if (!ShaderParameterRef)
+		{
+			EmitNullShaderParameterFatalError(Shader, ParametersMetadata, ParameterBinding.ByteOffset);
+		}
+	}
+
+	// Samplers
+	for (const FShaderParameterBindings::FResourceParameter& ParameterBinding : Bindings.Samplers)
+	{
+		auto ShaderParameterRef = *reinterpret_cast<const FSamplerStateRHIParamRef*>(Base + ParameterBinding.ByteOffset);
+		if (!ShaderParameterRef)
+		{
+			EmitNullShaderParameterFatalError(Shader, ParametersMetadata, ParameterBinding.ByteOffset);
+		}
+	}
+
+	// Graph Textures
+	for (const FShaderParameterBindings::FResourceParameter& ParameterBinding : Bindings.GraphTextures)
+	{
+		auto GraphTexture = *reinterpret_cast<const FRDGTexture* const*>(Base + ParameterBinding.ByteOffset);
+		if (!GraphTexture)
+		{
+			EmitNullShaderParameterFatalError(Shader, ParametersMetadata, ParameterBinding.ByteOffset);
+		}
+	}
+
+	// Graph SRVs
+	for (const FShaderParameterBindings::FResourceParameter& ParameterBinding : Bindings.GraphSRVs)
+	{
+		auto GraphSRV = *reinterpret_cast<const FRDGTextureSRV* const*>(Base + ParameterBinding.ByteOffset);
+		if (!GraphSRV)
+		{
+			EmitNullShaderParameterFatalError(Shader, ParametersMetadata, ParameterBinding.ByteOffset);
+		}
+	}
+
+	// Graph UAVs for compute shaders	
+	for (const FShaderParameterBindings::FResourceParameter& ParameterBinding : Bindings.GraphUAVs)
+	{
+		auto GraphUAV = *reinterpret_cast<const FRDGTextureUAV* const*>(Base + ParameterBinding.ByteOffset);
+		if (!GraphUAV)
+		{
+			EmitNullShaderParameterFatalError(Shader, ParametersMetadata, ParameterBinding.ByteOffset);
+		}
+	}
+
+	// Reference structures
+	for (const FShaderParameterBindings::FParameterStructReference& ParameterBinding : Bindings.ParameterReferences)
+	{
+		const TRefCountPtr<FRHIUniformBuffer>& ShaderParameterRef = *reinterpret_cast<const TRefCountPtr<FRHIUniformBuffer>*>(Base + ParameterBinding.ByteOffset);
+
+		if (!ShaderParameterRef.IsValid())
+		{
+			EmitNullShaderParameterFatalError(Shader, ParametersMetadata, ParameterBinding.ByteOffset);
+		}
+	}
+}
+
+#endif // DO_CHECK
diff --git a/Engine/Source/Runtime/RenderCore/Public/RenderGraphUtils.h b/Engine/Source/Runtime/RenderCore/Public/RenderGraphUtils.h
--- a/Engine/Source/Runtime/RenderCore/Public/RenderGraphUtils.h
+++ b/Engine/Source/Runtime/RenderCore/Public/RenderGraphUtils.h
@@ -26,7 +26,15 @@ extern RENDERCORE_API void ClearUnusedGraphResourcesImpl(const FShaderParameterB
 template<typename TShaderClass>
 inline void ClearUnusedGraphResources(const TShaderClass* Shader, typename TShaderClass::FParameters* InoutParameters)
 {
-	return ClearUnusedGraphResourcesImpl(Shader->Bindings, TShaderClass::FParameters::FTypeInfo::GetStructMetadata(), InoutParameters);
+	const FShaderParametersMetadata* ParametersMetadata = TShaderClass::FParameters::FTypeInfo::GetStructMetadata();
+
+	// Verify the shader have all the parameters it needs. This is done before the
+	// ClearUnusedGraphResourcesImpl() to not misslead user on why some resource are misseing
+	// when debugging a validation failure.
+	ValidateShaderParameters(Shader, ParametersMetadata, InoutParameters);
+
+	// Clear the resources the shader won't need.
+	return ClearUnusedGraphResourcesImpl(Shader->Bindings, ParametersMetadata, InoutParameters);
 }
 
 
diff --git a/Engine/Source/Runtime/RenderCore/Public/Shader.h b/Engine/Source/Runtime/RenderCore/Public/Shader.h
--- a/Engine/Source/Runtime/RenderCore/Public/Shader.h
+++ b/Engine/Source/Runtime/RenderCore/Public/Shader.h
@@ -653,7 +653,7 @@ public:
 		return Ar;
 	}
 
-	void BindForLegacyShaderParameters(const FShaderParameterMap& ParameterMaps, const FShaderParametersMetadata& StructMetaData);
+	void BindForLegacyShaderParameters(const FShaderParameterMap& ParameterMaps, const FShaderParametersMetadata& StructMetaData, bool bShouldBindEverything = false);
 	void BindForRootShaderParameters(const FShaderParameterMap& ParameterMaps, const FShaderParametersMetadata& StructMetaData);
 }; // FShaderParameterBindings
 
diff --git a/Engine/Source/Runtime/RenderCore/Public/ShaderParameterStruct.h b/Engine/Source/Runtime/RenderCore/Public/ShaderParameterStruct.h
--- a/Engine/Source/Runtime/RenderCore/Public/ShaderParameterStruct.h
+++ b/Engine/Source/Runtime/RenderCore/Public/ShaderParameterStruct.h
@@ -32,7 +32,10 @@
 	ShaderClass(const ShaderMetaType::CompiledShaderInitializerType& Initializer) \
 		: ShaderParentClass(Initializer) \
 	{ \
-		this->Bindings.BindForLegacyShaderParameters(Initializer.ParameterMap, *FParameters::FTypeInfo::GetStructMetadata()); \
+		this->Bindings.BindForLegacyShaderParameters( \
+			Initializer.ParameterMap, \
+			*FParameters::FTypeInfo::GetStructMetadata(), \
+			true); \
 	} \
 	\
 	ShaderClass() \
@@ -52,7 +55,18 @@
 
 
 /** Raise fatal error when a required shader parameter has not been set. */
-extern RENDERCORE_API void SetNullShaderParameterFatalError(const FShader* Shader, const FShaderParametersMetadata* ParametersMetadata, uint16 MemberOffset);
+extern RENDERCORE_API void EmitNullShaderParameterFatalError(const FShader* Shader, const FShaderParametersMetadata* ParametersMetadata, uint16 MemberOffset);
+
+
+/** Validates that all resource parameters of a shader are set. */
+#if DO_CHECK
+extern RENDERCORE_API void ValidateShaderParameters(const FShader* Shader, const FShaderParametersMetadata* ParametersMetadata, const void* Parameters);
+
+#else // !DO_CHECK
+FORCE_INLINE void ValidateShaderParameters(const FShader* Shader, const FShaderParametersMetadata* ParametersMetadata, const void* Parameters)
+{ }
+
+#endif // !DO_CHECK
 
 
 /** Set compute shader UAVs. */
@@ -77,7 +91,7 @@ inline void SetShaderUAVs(TRHICmdList& RHICmdList, const TShaderClass* Shader, F
 		{
 			if (!GraphUAV)
 			{
-				SetNullShaderParameterFatalError(Shader, TShaderClass::FParameters::FTypeInfo::GetStructMetadata(), ParameterBinding.ByteOffset);
+				EmitNullShaderParameterFatalError(Shader, TShaderClass::FParameters::FTypeInfo::GetStructMetadata(), ParameterBinding.ByteOffset);
 			}
 			else
 			{
@@ -134,7 +148,7 @@ inline void SetShaderParameters(TRHICmdList& RHICmdList, const TShaderClass* Sha
 
 		if (DO_CHECK && !ShaderParameterRef)
 		{
-			SetNullShaderParameterFatalError(Shader, TShaderClass::FParameters::FTypeInfo::GetStructMetadata(), ParameterBinding.ByteOffset);
+			EmitNullShaderParameterFatalError(Shader, TShaderClass::FParameters::FTypeInfo::GetStructMetadata(), ParameterBinding.ByteOffset);
 		}
 
 		RHICmdList.SetShaderTexture(ShadeRHI, ParameterBinding.BaseIndex, ShaderParameterRef);
@@ -147,7 +161,7 @@ inline void SetShaderParameters(TRHICmdList& RHICmdList, const TShaderClass* Sha
 
 		if (DO_CHECK && !ShaderParameterRef)
 		{
-			SetNullShaderParameterFatalError(Shader, TShaderClass::FParameters::FTypeInfo::GetStructMetadata(), ParameterBinding.ByteOffset);
+			EmitNullShaderParameterFatalError(Shader, TShaderClass::FParameters::FTypeInfo::GetStructMetadata(), ParameterBinding.ByteOffset);
 		}
 
 		RHICmdList.SetShaderResourceViewParameter(ShadeRHI, ParameterBinding.BaseIndex, ShaderParameterRef);
@@ -160,7 +174,7 @@ inline void SetShaderParameters(TRHICmdList& RHICmdList, const TShaderClass* Sha
 
 		if (DO_CHECK && !ShaderParameterRef)
 		{
-			SetNullShaderParameterFatalError(Shader, TShaderClass::FParameters::FTypeInfo::GetStructMetadata(), ParameterBinding.ByteOffset);
+			EmitNullShaderParameterFatalError(Shader, TShaderClass::FParameters::FTypeInfo::GetStructMetadata(), ParameterBinding.ByteOffset);
 		}
 
 		RHICmdList.SetShaderSampler(ShadeRHI, ParameterBinding.BaseIndex, ShaderParameterRef);
@@ -175,7 +189,7 @@ inline void SetShaderParameters(TRHICmdList& RHICmdList, const TShaderClass* Sha
 		{
 			if (!GraphTexture)
 			{
-				SetNullShaderParameterFatalError(Shader, TShaderClass::FParameters::FTypeInfo::GetStructMetadata(), ParameterBinding.ByteOffset);
+				EmitNullShaderParameterFatalError(Shader, TShaderClass::FParameters::FTypeInfo::GetStructMetadata(), ParameterBinding.ByteOffset);
 			}
 			else
 			{
@@ -197,7 +211,7 @@ inline void SetShaderParameters(TRHICmdList& RHICmdList, const TShaderClass* Sha
 		{
 			if (!GraphSRV)
 			{
-				SetNullShaderParameterFatalError(Shader, TShaderClass::FParameters::FTypeInfo::GetStructMetadata(), ParameterBinding.ByteOffset);
+				EmitNullShaderParameterFatalError(Shader, TShaderClass::FParameters::FTypeInfo::GetStructMetadata(), ParameterBinding.ByteOffset);
 			}
 			else
 			{
@@ -220,7 +234,7 @@ inline void SetShaderParameters(TRHICmdList& RHICmdList, const TShaderClass* Sha
 
 		if (DO_CHECK && !ShaderParameterRef.IsValid())
 		{
-			SetNullShaderParameterFatalError(Shader, TShaderClass::FParameters::FTypeInfo::GetStructMetadata(), ParameterBinding.ByteOffset);
+			EmitNullShaderParameterFatalError(Shader, TShaderClass::FParameters::FTypeInfo::GetStructMetadata(), ParameterBinding.ByteOffset);
 		}
 
 		RHICmdList.SetShaderUniformBuffer(ShadeRHI, ParameterBinding.BufferIndex, ShaderParameterRef);
