Cherry-pick 4550134: Simplifies user facing functions for render graph with user facing render graph type.

diff --git a/Engine/Source/Runtime/RenderCore/Public/RenderGraphBuilder.h b/Engine/Source/Runtime/RenderCore/Public/RenderGraphBuilder.h
--- a/Engine/Source/Runtime/RenderCore/Public/RenderGraphBuilder.h
+++ b/Engine/Source/Runtime/RenderCore/Public/RenderGraphBuilder.h
@@ -284,7 +284,7 @@ public:
 	}
 
 	/** Create graph tracked UAV for a texture from a descriptor. */
-	inline const FRDGTextureUAV* CreateUAV(const FRDGTextureUAVDesc& Desc)
+	inline FRDGTextureUAVRef CreateUAV(const FRDGTextureUAVDesc& Desc)
 	{
 		check(Desc.Texture);
 		#if RENDER_GRAPH_DEBUGGING
diff --git a/Engine/Source/Runtime/RenderCore/Public/RenderGraphResources.h b/Engine/Source/Runtime/RenderCore/Public/RenderGraphResources.h
--- a/Engine/Source/Runtime/RenderCore/Public/RenderGraphResources.h
+++ b/Engine/Source/Runtime/RenderCore/Public/RenderGraphResources.h
@@ -20,6 +20,17 @@ class FRDGBufferUAV;
 struct FPooledRDGBuffer;
 struct FRaytracingShaderBindingsWriter;
 
+
+/** Defines the RDG resource references for user code not forgetting the const every time. */
+using FRDGResourceRef = const FRDGResource*;
+using FRDGTextureRef = const FRDGTexture*;
+using FRDGTextureSRVRef = const FRDGTextureSRV*;
+using FRDGTextureUAVRef = const FRDGTextureUAV*;
+using FRDGBufferRef = const  FRDGBuffer*;
+using FRDGBufferSRVRef = const FRDGBufferSRV*;
+using FRDGBufferUAVRef = const FRDGBufferUAV*;
+
+
 /** Generic graph resource. */
 class RENDERCORE_API FRDGResource
 {
@@ -110,12 +121,12 @@ private:
 class RENDERCORE_API FRDGTextureSRVDesc
 {
 public:
-	const FRDGTexture* Texture;
+	FRDGTextureRef Texture;
 	uint8 MipLevel = 0;
 
 	FRDGTextureSRVDesc() {}
 
-	FRDGTextureSRVDesc(const FRDGTexture* InTexture, uint8 InMipLevel) :
+	FRDGTextureSRVDesc(FRDGTextureRef InTexture, uint8 InMipLevel) :
 		Texture(InTexture),
 		MipLevel(InMipLevel)
 	{}
@@ -141,12 +152,12 @@ private:
 class RENDERCORE_API FRDGTextureUAVDesc
 {
 public:
-	const FRDGTexture* Texture;
+	FRDGTextureRef Texture;
 	uint8 MipLevel = 0;
 
 	FRDGTextureUAVDesc() {}
 
-	FRDGTextureUAVDesc(const FRDGTexture* InTexture, uint8 InMipLevel = 0) :
+	FRDGTextureUAVDesc(FRDGTextureRef InTexture, uint8 InMipLevel = 0) :
 		Texture(InTexture),
 		MipLevel(InMipLevel)
 	{}
@@ -236,7 +247,7 @@ struct RENDERCORE_API FRDGBufferDesc
 
 struct RENDERCORE_API FRDGBufferSRVDesc
 {
-	const FRDGBuffer* Buffer = nullptr;
+	FRDGBufferRef Buffer = nullptr;
 
 	/** Nymber of bytes per element (used for vertex buffer). */
 	uint32 BytesPerElement = 1;
@@ -247,7 +258,7 @@ struct RENDERCORE_API FRDGBufferSRVDesc
 
 struct RENDERCORE_API FRDGBufferUAVDesc
 {
-	const FRDGBuffer* Buffer = nullptr;
+	FRDGBufferRef Buffer = nullptr;
 
 	/** Nymber of bytes per element (used for vertex buffer). */
 	EPixelFormat Format = PF_Unknown;
@@ -257,7 +268,7 @@ struct RENDERCORE_API FRDGBufferUAVDesc
 	bool bSupportsAppendBuffer = false;
 
 	/** Create descriptor for a the UAV of a buffer meant to be use for indirect dral call parameters. */
-	static inline FRDGBufferUAVDesc CreateIndirectDesc(const FRDGBuffer* Buffer);
+	static inline FRDGBufferUAVDesc CreateIndirectDesc(FRDGBufferRef Buffer);
 };
 
 
@@ -398,7 +409,7 @@ private:
 };
 
 
-FRDGBufferUAVDesc FRDGBufferUAVDesc::CreateIndirectDesc(const FRDGBuffer* Buffer)
+FRDGBufferUAVDesc FRDGBufferUAVDesc::CreateIndirectDesc(FRDGBufferRef Buffer)
 {
 	checkf(Buffer->Desc.Usage & BUF_DrawIndirect, TEXT("Buffer %s has not been created with a BUF_DrawIndirect flag."), Buffer->Name);
 	FRDGBufferUAVDesc Desc;
@@ -406,13 +417,3 @@ FRDGBufferUAVDesc FRDGBufferUAVDesc::CreateIndirectDesc(const FRDGBuffer* Buffer
 	Desc.Format = PF_R32_UINT;
 	return Desc;
 }
-
-
-/** Defines the RDG resource references for user code not forgetting the const every time. */
-using FRDGResourceRef = const FRDGResource*;
-using FRDGTextureRef = const FRDGTexture*;
-using FRDGTextureSRVRef = const FRDGTextureSRV*;
-using FRDGTextureUAVRef = const FRDGTextureUAV*;
-using FRDGBufferRef = const  FRDGBuffer*;
-using FRDGBufferSRVRef = const FRDGBufferSRV*;
-using FRDGBufferUAVRef = const FRDGBufferUAV*;
diff --git a/Engine/Source/Runtime/RenderCore/Public/ShaderParameterMacros.h b/Engine/Source/Runtime/RenderCore/Public/ShaderParameterMacros.h
--- a/Engine/Source/Runtime/RenderCore/Public/ShaderParameterMacros.h
+++ b/Engine/Source/Runtime/RenderCore/Public/ShaderParameterMacros.h
@@ -493,113 +493,8 @@ struct TShaderParameterTypeInfo<TStaticArray<T,InNumElements,IgnoredAlignment>>
 	static const FShaderParametersMetadata* GetStructMetadata() { return TShaderParameterTypeInfo<T>::GetStructMetadata(); }
 };
 
-template<>
-struct TShaderParameterTypeInfo<FShaderResourceViewRHIParamRef>
-{
-	static constexpr int32 NumRows = 1;
-	static constexpr int32 NumColumns = 1;
-	static constexpr int32 NumElements = 0;
-	static constexpr int32 Alignment = SHADER_PARAMETER_POINTER_ALIGNMENT;
-	static constexpr bool bIsStoredInConstantBuffer = false;
-
-	using TAlignedType = TAlignedShaderParameterPtr<FShaderResourceViewRHIParamRef>;
-
-	static const FShaderParametersMetadata* GetStructMetadata() { return nullptr; }
-};
-static_assert(sizeof(TShaderParameterTypeInfo<FShaderResourceViewRHIParamRef>::TAlignedType) == SHADER_PARAMETER_POINTER_ALIGNMENT, "Uniform buffer layout must not be platform dependent.");
-
-template<>
-struct TShaderParameterTypeInfo<FSamplerStateRHIParamRef>
-{
-	static constexpr int32 NumRows = 1;
-	static constexpr int32 NumColumns = 1;
-	static constexpr int32 NumElements = 0;
-	static constexpr int32 Alignment = SHADER_PARAMETER_POINTER_ALIGNMENT;
-	static constexpr bool bIsStoredInConstantBuffer = false;
-
-	using TAlignedType = TAlignedShaderParameterPtr<FSamplerStateRHIParamRef>;
-
-	static const FShaderParametersMetadata* GetStructMetadata() { return nullptr; }
-};
-static_assert(sizeof(TShaderParameterTypeInfo<FSamplerStateRHIParamRef>::TAlignedType) == SHADER_PARAMETER_POINTER_ALIGNMENT, "Uniform buffer layout must not be platform dependent.");
-
-template<>
-struct TShaderParameterTypeInfo<FTextureRHIParamRef>
-{
-	static constexpr int32 NumRows = 1;
-	static constexpr int32 NumColumns = 1;
-	static constexpr int32 NumElements = 0;
-	static constexpr int32 Alignment = SHADER_PARAMETER_POINTER_ALIGNMENT;
-	static constexpr bool bIsStoredInConstantBuffer = false;
-
-	using TAlignedType = TAlignedShaderParameterPtr<FTextureRHIParamRef>;
-
-	static const FShaderParametersMetadata* GetStructMetadata() { return nullptr; }
-};
-static_assert(sizeof(TShaderParameterTypeInfo<FTextureRHIParamRef>::TAlignedType) == SHADER_PARAMETER_POINTER_ALIGNMENT, "Uniform buffer layout must not be platform dependent.");
-
-template<>
-struct TShaderParameterTypeInfo<const FRDGTexture*>
-{
-	static constexpr int32 NumRows = 1;
-	static constexpr int32 NumColumns = 1;
-	static constexpr int32 NumElements = 0;
-	static constexpr int32 Alignment = SHADER_PARAMETER_POINTER_ALIGNMENT;
-	static constexpr bool bIsStoredInConstantBuffer = false;
-
-	using TAlignedType = TAlignedShaderParameterPtr<const FRDGTexture*>;
-
-	static const FShaderParametersMetadata* GetStructMetadata() { return nullptr; }
-};
-static_assert(sizeof(TShaderParameterTypeInfo<const FRDGTexture*>::TAlignedType) == SHADER_PARAMETER_POINTER_ALIGNMENT, "Uniform buffer layout must not be platform dependent.");
-
-template<>
-struct TShaderParameterTypeInfo<const FRDGTextureSRV*>
-{
-	static constexpr int32 NumRows = 1;
-	static constexpr int32 NumColumns = 1;
-	static constexpr int32 NumElements = 0;
-	static constexpr int32 Alignment = SHADER_PARAMETER_POINTER_ALIGNMENT;
-	static constexpr bool bIsStoredInConstantBuffer = false;
-	
-	using TAlignedType = TAlignedShaderParameterPtr<const FRDGTextureSRV*>;
-	
-	static const FShaderParametersMetadata* GetStructMetadata() { return nullptr; }
-};
-static_assert(sizeof(TShaderParameterTypeInfo<const FRDGTextureSRV*>::TAlignedType) == SHADER_PARAMETER_POINTER_ALIGNMENT, "Uniform buffer layout must not be platform dependent.");
-
-template<>
-struct TShaderParameterTypeInfo<const FRDGTextureUAV*>
-{
-	static constexpr int32 NumRows = 1;
-	static constexpr int32 NumColumns = 1;
-	static constexpr int32 NumElements = 0;
-	static constexpr int32 Alignment = SHADER_PARAMETER_POINTER_ALIGNMENT;
-	static constexpr bool bIsStoredInConstantBuffer = false;
-	
-	using TAlignedType = TAlignedShaderParameterPtr<const FRDGTextureUAV*>;
-	
-	static const FShaderParametersMetadata* GetStructMetadata() { return nullptr; }
-};
-static_assert(sizeof(TShaderParameterTypeInfo<const FRDGTextureUAV*>::TAlignedType) == SHADER_PARAMETER_POINTER_ALIGNMENT, "Uniform buffer layout must not be platform dependent.");
-
-template<>
-struct TShaderParameterTypeInfo<const FRDGBuffer*>
-{
-	static constexpr int32 NumRows = 1;
-	static constexpr int32 NumColumns = 1;
-	static constexpr int32 NumElements = 0;
-	static constexpr int32 Alignment = SHADER_PARAMETER_POINTER_ALIGNMENT;
-	static constexpr bool bIsStoredInConstantBuffer = false;
-
-	using TAlignedType = TAlignedShaderParameterPtr<const FRDGBuffer*>;
-
-	static const FShaderParametersMetadata* GetStructMetadata() { return nullptr; }
-};
-static_assert(sizeof(TShaderParameterTypeInfo<const FRDGBuffer*>::TAlignedType) == SHADER_PARAMETER_POINTER_ALIGNMENT, "Uniform buffer layout must not be platform dependent.");
-
-template<>
-struct TShaderParameterTypeInfo<const FRDGBufferSRV*>
+template<typename ShaderResourceType>
+struct TShaderResourceParameterTypeInfo
 {
 	static constexpr int32 NumRows = 1;
 	static constexpr int32 NumColumns = 1;
@@ -607,26 +502,12 @@ struct TShaderParameterTypeInfo<const FRDGBufferSRV*>
 	static constexpr int32 Alignment = SHADER_PARAMETER_POINTER_ALIGNMENT;
 	static constexpr bool bIsStoredInConstantBuffer = false;
 
-	using TAlignedType = TAlignedShaderParameterPtr<const FRDGBufferSRV*>;
+	using TAlignedType = TAlignedShaderParameterPtr<ShaderResourceType>;
 
 	static const FShaderParametersMetadata* GetStructMetadata() { return nullptr; }
-};
-static_assert(sizeof(TShaderParameterTypeInfo<const FRDGBufferSRV*>::TAlignedType) == SHADER_PARAMETER_POINTER_ALIGNMENT, "Uniform buffer layout must not be platform dependent.");
-
-template<>
-struct TShaderParameterTypeInfo<const FRDGBufferUAV*>
-{
-	static constexpr int32 NumRows = 1;
-	static constexpr int32 NumColumns = 1;
-	static constexpr int32 NumElements = 0;
-	static constexpr int32 Alignment = SHADER_PARAMETER_POINTER_ALIGNMENT;
-	static constexpr bool bIsStoredInConstantBuffer = false;
-
-	using TAlignedType = TAlignedShaderParameterPtr<const FRDGBufferUAV*>;
 
-	static const FShaderParametersMetadata* GetStructMetadata() { return nullptr; }
+	static_assert(sizeof(TAlignedType) == SHADER_PARAMETER_POINTER_ALIGNMENT, "Uniform buffer layout must not be platform dependent.");
 };
-static_assert(sizeof(TShaderParameterTypeInfo<const FRDGBufferUAV*>::TAlignedType) == SHADER_PARAMETER_POINTER_ALIGNMENT, "Uniform buffer layout must not be platform dependent.");
 
 template<class UniformBufferStructType>
 struct TShaderParameterTypeInfo<TUniformBufferRef<UniformBufferStructType>>
@@ -688,9 +569,7 @@ struct TShaderParameterTypeInfo<TUniformBufferRef<UniformBufferStructType>>
 #define DECLARE_UNIFORM_BUFFER_STRUCT_MEMBER_EXPLICIT(BaseType,TypeInfo,MemberType,MemberName,ArrayDecl,Precision,OptionalShaderType,IsMemberStruct) \
 		zzMemberId##MemberName; \
 	public: \
-		typedef MemberType zzA##MemberName ArrayDecl; \
-		typedef TypeInfo::TAlignedType zzT##MemberName; \
-		zzT##MemberName MemberName; \
+		TypeInfo::TAlignedType MemberName; \
 		static_assert(BaseType != UBMT_INVALID, "Invalid type " #MemberType " of member " #MemberName "."); \
 	private: \
 		struct zzNextMemberId##MemberName { enum { HasDeclaredResource = zzMemberId##MemberName::HasDeclaredResource || !TypeInfo::bIsStoredInConstantBuffer }; }; \
@@ -797,7 +676,7 @@ extern RENDERCORE_API FShaderParametersMetadata* FindUniformBufferStructByFName(
  *	SHADER_PARAMETER_TEXTURE(Texture2D, MyTexture)
  */
 #define SHADER_PARAMETER_TEXTURE(ShaderType,MemberName) \
-	DECLARE_UNIFORM_BUFFER_STRUCT_MEMBER_EXPLICIT(UBMT_TEXTURE, TShaderParameterTypeInfo<FTextureRHIParamRef>, FTextureRHIParamRef,MemberName,,EShaderPrecisionModifier::Float,TEXT(#ShaderType),false)
+	DECLARE_UNIFORM_BUFFER_STRUCT_MEMBER_EXPLICIT(UBMT_TEXTURE, TShaderResourceParameterTypeInfo<FTextureRHIParamRef>, FTextureRHIParamRef,MemberName,,EShaderPrecisionModifier::Float,TEXT(#ShaderType),false)
 
 /** Adds a shader resource view.
  *
@@ -805,7 +684,7 @@ extern RENDERCORE_API FShaderParametersMetadata* FindUniformBufferStructByFName(
  *	SHADER_PARAMETER_SRV(Texture2D, MySRV)
  */
 #define SHADER_PARAMETER_SRV(ShaderType,MemberName) \
-	DECLARE_UNIFORM_BUFFER_STRUCT_MEMBER_EXPLICIT(UBMT_SRV, TShaderParameterTypeInfo<FShaderResourceViewRHIParamRef>, FShaderResourceViewRHIParamRef,MemberName,,EShaderPrecisionModifier::Float,TEXT(#ShaderType),false)
+	DECLARE_UNIFORM_BUFFER_STRUCT_MEMBER_EXPLICIT(UBMT_SRV, TShaderResourceParameterTypeInfo<FShaderResourceViewRHIParamRef>, FShaderResourceViewRHIParamRef,MemberName,,EShaderPrecisionModifier::Float,TEXT(#ShaderType),false)
 
 /** Adds a sampler.
  *
@@ -813,7 +692,7 @@ extern RENDERCORE_API FShaderParametersMetadata* FindUniformBufferStructByFName(
  *	SHADER_PARAMETER_SAMPLER(SamplerState, MySampler)
  */
 #define SHADER_PARAMETER_SAMPLER(ShaderType,MemberName) \
-	DECLARE_UNIFORM_BUFFER_STRUCT_MEMBER_EXPLICIT(UBMT_SAMPLER, TShaderParameterTypeInfo<FSamplerStateRHIParamRef>, FSamplerStateRHIParamRef,MemberName,,EShaderPrecisionModifier::Float,TEXT(#ShaderType),false)
+	DECLARE_UNIFORM_BUFFER_STRUCT_MEMBER_EXPLICIT(UBMT_SAMPLER, TShaderResourceParameterTypeInfo<FSamplerStateRHIParamRef>, FSamplerStateRHIParamRef,MemberName,,EShaderPrecisionModifier::Float,TEXT(#ShaderType),false)
 
 /** Adds a render graph tracked texture.
  *
@@ -821,7 +700,7 @@ extern RENDERCORE_API FShaderParametersMetadata* FindUniformBufferStructByFName(
  *	SHADER_PARAMETER_RDG_TEXTURE(Texture2D, MyTexture)
  */
 #define SHADER_PARAMETER_RDG_TEXTURE(ShaderType,MemberName) \
-	DECLARE_UNIFORM_BUFFER_STRUCT_MEMBER_EXPLICIT(UBMT_GRAPH_TRACKED_TEXTURE, TShaderParameterTypeInfo<const FRDGTexture*>, const FRDGTexture*,MemberName,,EShaderPrecisionModifier::Float,TEXT(#ShaderType),false)
+	DECLARE_UNIFORM_BUFFER_STRUCT_MEMBER_EXPLICIT(UBMT_GRAPH_TRACKED_TEXTURE, TShaderResourceParameterTypeInfo<FRDGTextureRef>, FRDGTextureRef,MemberName,,EShaderPrecisionModifier::Float,TEXT(#ShaderType),false)
 
 /** Adds a shader resource view for a render graph tracked texture.
  *
@@ -829,7 +708,7 @@ extern RENDERCORE_API FShaderParametersMetadata* FindUniformBufferStructByFName(
  *	SHADER_PARAMETER_RDG_TEXTURE_SRV(Texture2D, MySRV)
  */
 #define SHADER_PARAMETER_RDG_TEXTURE_SRV(ShaderType,MemberName) \
-	DECLARE_UNIFORM_BUFFER_STRUCT_MEMBER_EXPLICIT(UBMT_GRAPH_TRACKED_SRV, TShaderParameterTypeInfo<const FRDGTextureSRV*>, const FRDGTextureSRV*,MemberName,,EShaderPrecisionModifier::Float,TEXT(#ShaderType),false)
+	DECLARE_UNIFORM_BUFFER_STRUCT_MEMBER_EXPLICIT(UBMT_GRAPH_TRACKED_SRV, TShaderResourceParameterTypeInfo<FRDGTextureSRVRef>, FRDGTextureSRVRef,MemberName,,EShaderPrecisionModifier::Float,TEXT(#ShaderType),false)
 
 /** Adds a unordered access view for a render graph tracked texture.
  *
@@ -837,7 +716,7 @@ extern RENDERCORE_API FShaderParametersMetadata* FindUniformBufferStructByFName(
  *	SHADER_PARAMETER_RDG_TEXTURE_UAV(RWTexture2D, MyUAV)
  */
 #define SHADER_PARAMETER_RDG_TEXTURE_UAV(ShaderType,MemberName) \
-	DECLARE_UNIFORM_BUFFER_STRUCT_MEMBER_EXPLICIT(UBMT_GRAPH_TRACKED_UAV, TShaderParameterTypeInfo<const FRDGTextureUAV*>, const FRDGTextureUAV*,MemberName,,EShaderPrecisionModifier::Float,TEXT(#ShaderType),false)
+	DECLARE_UNIFORM_BUFFER_STRUCT_MEMBER_EXPLICIT(UBMT_GRAPH_TRACKED_UAV, TShaderResourceParameterTypeInfo<FRDGTextureUAVRef>, FRDGTextureUAVRef,MemberName,,EShaderPrecisionModifier::Float,TEXT(#ShaderType),false)
 
 /** Adds a render graph tracked buffer.
  *
@@ -845,7 +724,7 @@ extern RENDERCORE_API FShaderParametersMetadata* FindUniformBufferStructByFName(
  *	SHADER_PARAMETER_RDG_BUFFER(Texture2D, MyTexture)
  */
 #define SHADER_PARAMETER_RDG_BUFFER(ShaderType,MemberName) \
-	DECLARE_UNIFORM_BUFFER_STRUCT_MEMBER_EXPLICIT(UBMT_GRAPH_TRACKED_BUFFER, TShaderParameterTypeInfo<const FRDGBuffer*>, const FRDGBuffer*,MemberName,,EShaderPrecisionModifier::Float,TEXT(#ShaderType),false)
+	DECLARE_UNIFORM_BUFFER_STRUCT_MEMBER_EXPLICIT(UBMT_GRAPH_TRACKED_BUFFER, TShaderResourceParameterTypeInfo<FRDGBufferRef>, FRDGBufferRef,MemberName,,EShaderPrecisionModifier::Float,TEXT(#ShaderType),false)
 
 /** Adds a shader resource view for a render graph tracked buffer.
  *
@@ -853,7 +732,7 @@ extern RENDERCORE_API FShaderParametersMetadata* FindUniformBufferStructByFName(
  *	SHADER_PARAMETER_RDG_BUFFER_SRV(Buffer<float4>, MySRV)
  */
 #define SHADER_PARAMETER_RDG_BUFFER_SRV(ShaderType,MemberName) \
-	DECLARE_UNIFORM_BUFFER_STRUCT_MEMBER_EXPLICIT(UBMT_GRAPH_TRACKED_BUFFER_SRV, TShaderParameterTypeInfo<const FRDGBufferSRV*>, const FRDGBufferSRV*,MemberName,,EShaderPrecisionModifier::Float,TEXT(#ShaderType),false)
+	DECLARE_UNIFORM_BUFFER_STRUCT_MEMBER_EXPLICIT(UBMT_GRAPH_TRACKED_BUFFER_SRV, TShaderResourceParameterTypeInfo<FRDGBufferSRVRef>, FRDGBufferSRVRef,MemberName,,EShaderPrecisionModifier::Float,TEXT(#ShaderType),false)
 
 /** Adds a unordered access view for a render graph tracked buffer.
  *
@@ -861,7 +740,7 @@ extern RENDERCORE_API FShaderParametersMetadata* FindUniformBufferStructByFName(
  *	SHADER_PARAMETER_RDG_BUFFER_UAV(RWBuffer<float4>, MyUAV)
  */
 #define SHADER_PARAMETER_RDG_BUFFER_UAV(ShaderType,MemberName) \
-	DECLARE_UNIFORM_BUFFER_STRUCT_MEMBER_EXPLICIT(UBMT_GRAPH_TRACKED_BUFFER_UAV, TShaderParameterTypeInfo<const FRDGBufferUAV*>, const FRDGBufferUAV*,MemberName,,EShaderPrecisionModifier::Float,TEXT(#ShaderType),false)
+	DECLARE_UNIFORM_BUFFER_STRUCT_MEMBER_EXPLICIT(UBMT_GRAPH_TRACKED_BUFFER_UAV, TShaderResourceParameterTypeInfo<FRDGBufferUAVRef>, FRDGBufferUAVRef,MemberName,,EShaderPrecisionModifier::Float,TEXT(#ShaderType),false)
 
 /** Nests a shader parameter structure into another one, in C++ and shader code.
  *
