Cherry-pick 4588285: Renames render graph builder's GetInternalTexture to QueueTextureExtraction

diff --git a/Engine/Source/Runtime/RenderCore/Private/VisualizeTexture.cpp b/Engine/Source/Runtime/RenderCore/Private/VisualizeTexture.cpp
--- a/Engine/Source/Runtime/RenderCore/Private/VisualizeTexture.cpp
+++ b/Engine/Source/Runtime/RenderCore/Private/VisualizeTexture.cpp
@@ -279,7 +279,7 @@ void FVisualizeTexture::CreateContentCapturePass(FRDGBuilder& GraphBuilder, cons
 	{
 		VisualizeTextureDesc = SrcDesc;
 		VisualizeTextureContent = nullptr;
-		GraphBuilder.GetInternalTexture(CopyTexture, &VisualizeTextureContent);
+		GraphBuilder.QueueTextureExtraction(CopyTexture, &VisualizeTextureContent);
 	}
 
 #if 0 // TODO(RDG): requires some kind of CPU readback pass.
diff --git a/Engine/Source/Runtime/RenderCore/Public/RenderGraphBuilder.h b/Engine/Source/Runtime/RenderCore/Public/RenderGraphBuilder.h
--- a/Engine/Source/Runtime/RenderCore/Public/RenderGraphBuilder.h
+++ b/Engine/Source/Runtime/RenderCore/Public/RenderGraphBuilder.h
@@ -354,16 +354,20 @@ public:
 		}
 	}
 
-	/**
-	 * Extracts an internal texture by handle.  Must be called before Execute.
+	/** Queue a texture extraction. This will set *OutTexturePtr with the internal pooled render target at the Execute().
+	 *
+	 * Note: even when the render graph uses the immediate debugging mode (executing passes as they get added), the texture extrations
+	 * will still happen in the Execute(), to ensure there is no bug caused in code outside the render graph on whether this mode is used or not.
 	 */
-	inline void GetInternalTexture(FRDGTextureRef Texture, TRefCountPtr<IPooledRenderTarget>* OutTexturePtr, bool bTransitionToRead = true)
+	inline void QueueTextureExtraction(FRDGTextureRef Texture, TRefCountPtr<IPooledRenderTarget>* OutTexturePtr, bool bTransitionToRead = true)
 	{
 		check(Texture);
 		check(OutTexturePtr);
 		#if RENDER_GRAPH_DEBUGGING
 		{
-			checkf(!bHasExecuted, TEXT("Accessing render graph internal texture %s needs to happen before the builder's execution."), Texture->Name);
+			checkf(!bHasExecuted,
+				TEXT("Accessing render graph internal texture %s with QueueTextureExtraction() needs to happen before the builder's execution."),
+				Texture->Name);
 		}
 		#endif
 		FDeferredInternalTextureQuery Query;
@@ -374,7 +378,7 @@ public:
 	}
 
 	/** 
-	 * Executes the queued passes, managing setting of render targets (RHI RenderPasses) and resource transitions.
+	 * Executes the queued passes, managing setting of render targets (RHI RenderPasses), resource transitions and queued texture extraction.
 	 */
 	void Execute();
 
diff --git a/Engine/Source/Runtime/Renderer/Private/SceneOcclusion.cpp b/Engine/Source/Runtime/Renderer/Private/SceneOcclusion.cpp
--- a/Engine/Source/Runtime/Renderer/Private/SceneOcclusion.cpp
+++ b/Engine/Source/Runtime/Renderer/Private/SceneOcclusion.cpp
@@ -1252,7 +1252,7 @@ void BuildHZB(FRDGBuilder& GraphBuilder, FViewInfo& View)
 		DstSize /= 2;
 	}
 
-	GraphBuilder.GetInternalTexture(HZBTexture, &View.HZB);
+	GraphBuilder.QueueTextureExtraction(HZBTexture, &View.HZB);
 }
 
 struct FViewOcclusionQueries
diff --git a/Engine/Source/Runtime/Renderer/Private/VolumetricFog.cpp b/Engine/Source/Runtime/Renderer/Private/VolumetricFog.cpp
--- a/Engine/Source/Runtime/Renderer/Private/VolumetricFog.cpp
+++ b/Engine/Source/Runtime/Renderer/Private/VolumetricFog.cpp
@@ -1161,11 +1161,11 @@ void FDeferredShadingSceneRenderer::ComputeVolumetricFog(FRHICommandListImmediat
 					});
 			}
 
-			GraphBuilder.GetInternalTexture(IntegratedLightScattering, &View.VolumetricFogResources.IntegratedLightScattering);
+			GraphBuilder.QueueTextureExtraction(IntegratedLightScattering, &View.VolumetricFogResources.IntegratedLightScattering);
 
 			if (bUseTemporalReprojection)
 			{
-				GraphBuilder.GetInternalTexture(IntegrationData.LightScattering, &View.ViewState->LightScatteringHistory);
+				GraphBuilder.QueueTextureExtraction(IntegrationData.LightScattering, &View.ViewState->LightScatteringHistory);
 			}
 			else if (View.ViewState)
 			{
