Clone of CL 4589926 for 4.21.1:

Fix for shader code library iterative cooking missing previous cook library data.  Don't clean at the end of a cook as we might be iterative next time - it's done at cook startup anyway based on the iterate flag.  Fix previous cooked cache file name searches - formats seem to have changed.

diff --git a/Engine/Source/Editor/UnrealEd/Private/CookOnTheFlyServer.cpp b/Engine/Source/Editor/UnrealEd/Private/CookOnTheFlyServer.cpp
--- a/Engine/Source/Editor/UnrealEd/Private/CookOnTheFlyServer.cpp
+++ b/Engine/Source/Editor/UnrealEd/Private/CookOnTheFlyServer.cpp
@@ -6246,7 +6246,9 @@ void UCookOnTheFlyServer::CookByTheBookFinished()
 			// Save shader code map
 			FString LibraryName = !IsCookingDLC() ? FApp::GetProjectName() : CookByTheBookOptions->DlcName;
 			SaveShaderCodeLibrary(LibraryName);
-            CleanShaderCodeLibraries();
+			
+			// Don't clean Saved/Shaders/<LibraryPlatform(s)>/ at the end as we might iterate next time - Next cook at startup will decide if clean on iterate flag
+            // /*CleanShaderCodeLibraries();*/
 			ProcessShaderCodeLibraries(LibraryName);
             
 			FShaderCodeLibrary::Shutdown();
diff --git a/Engine/Source/Runtime/RenderCore/Private/ShaderCodeLibrary.cpp b/Engine/Source/Runtime/RenderCore/Private/ShaderCodeLibrary.cpp
--- a/Engine/Source/Runtime/RenderCore/Private/ShaderCodeLibrary.cpp
+++ b/Engine/Source/Runtime/RenderCore/Private/ShaderCodeLibrary.cpp
@@ -883,14 +883,14 @@ struct FEditorShaderCodeArchive
 	{
 		check(LibraryName.Len() > 0);
 
-		const FString ShaderIntermediateLocation = FPaths::ProjectSavedDir() / TEXT("Shaders") / FormatName.ToString() / LibraryName;
+		const FString ShaderIntermediateLocation = FPaths::ProjectSavedDir() / TEXT("Shaders") / FormatName.ToString();
 
 		TArray<FString> ShaderFiles;
 		IFileManager::Get().FindFiles(ShaderFiles, *ShaderIntermediateLocation, *ShaderExtension);
 
 		for (const FString& ShaderFileName : ShaderFiles)
 		{
-			if (ShaderFileName.Contains(LibraryName + TEXT("_") + FormatName.ToString() + TEXT(".")))
+			if (ShaderFileName.Contains(LibraryName + TEXT("-") + FormatName.ToString() + TEXT(".")))
 			{
 				FArchive* PrevCookedAr = IFileManager::Get().CreateFileReader(*GetCodeArchiveFilename(OutputDir, LibraryName, FormatName));
 
@@ -940,7 +940,7 @@ struct FEditorShaderCodeArchive
 
 		for (const FString& ShaderFileName : PipelineFiles)
 		{
-			if (ShaderFileName.Contains(LibraryName + TEXT("_") + FormatName.ToString() + TEXT(".")))
+			if (ShaderFileName.Contains(LibraryName + TEXT("-") + FormatName.ToString() + TEXT(".")))
 			{
 				FArchive* PrevCookedAr = IFileManager::Get().CreateFileReader(*GetPipelinesArchiveFilename(OutputDir, LibraryName, FormatName));
 
@@ -1206,14 +1206,14 @@ struct FEditorShaderStableInfo
 	{
 		check(LibraryName.Len() > 0);
 
-		const FString ShaderIntermediateLocation = FPaths::ProjectSavedDir() / TEXT("Shaders") / FormatName.ToString() / LibraryName;
+		const FString ShaderIntermediateLocation = FPaths::ProjectSavedDir() / TEXT("Shaders") / FormatName.ToString();
 
 		TArray<FString> ShaderFiles;
 		IFileManager::Get().FindFiles(ShaderFiles, *ShaderIntermediateLocation, *ShaderExtension);
 
 		for (const FString& ShaderFileName : ShaderFiles)
 		{
-			if (ShaderFileName.Contains(LibraryName + TEXT("_") + FormatName.ToString() + TEXT(".")))
+			if (ShaderFileName.Contains(LibraryName + TEXT("-") + FormatName.ToString() + TEXT(".")))
 			{
 				TArray<FString> SourceFileContents;
 				if (FFileHelper::LoadFileToStringArray(SourceFileContents, *GetStableInfoArchiveFilename(OutputDir, LibraryName, FormatName)))
