Prints the name of the shader class in the log when their is shader parameter binding failure.

diff --git a/Engine/Source/Runtime/RenderCore/Private/RenderGraphBuilder.cpp b/Engine/Source/Runtime/RenderCore/Private/RenderGraphBuilder.cpp
--- a/Engine/Source/Runtime/RenderCore/Private/RenderGraphBuilder.cpp
+++ b/Engine/Source/Runtime/RenderCore/Private/RenderGraphBuilder.cpp
@@ -624,7 +624,7 @@ void FRDGBuilder::ExecutePass( const FRenderGraphPass* Pass )
 
 	AllocateAndTransitionPassResources(Pass, &RPInfo, &bHasRenderTargets);
 
-	if (RENDER_GRAPH_DRAW_EVENTS != 0 && GetEmitDrawEvents())
+	if (RENDER_GRAPH_DRAW_EVENTS && GetEmitDrawEvents())
 	{
 		PushDrawEventStack(Pass);
 	}
@@ -649,7 +649,7 @@ void FRDGBuilder::ExecutePass( const FRenderGraphPass* Pass )
 		RHICmdList.EndRenderPass();
 	}
 
-	if (RENDER_GRAPH_DRAW_EVENTS != 0 && GetEmitDrawEvents())
+	if (RENDER_GRAPH_DRAW_EVENTS && GetEmitDrawEvents())
 	{
 		RHICmdList.PopEvent();
 	}
diff --git a/Engine/Source/Runtime/RenderCore/Private/ShaderParameterStruct.cpp b/Engine/Source/Runtime/RenderCore/Private/ShaderParameterStruct.cpp
--- a/Engine/Source/Runtime/RenderCore/Private/ShaderParameterStruct.cpp
+++ b/Engine/Source/Runtime/RenderCore/Private/ShaderParameterStruct.cpp
@@ -10,6 +10,9 @@
 /** Context of binding a map. */
 struct FShaderParameterStructBindingContext
 {
+	// Shader having its parameter bound.
+	const FShader* Shader;
+
 	// Bindings to bind.
 	FShaderParameterBindings* Bindings;
 
@@ -105,7 +108,8 @@ struct FShaderParameterStructBindingContext
 
 				if (uint32(BoundSize) > ByteSize)
 				{
-					UE_LOG(LogShaders, Fatal, TEXT("The size required to bind shader parameter %s is %i bytes, smaller than %s's %i bytes."),
+					UE_LOG(LogShaders, Fatal, TEXT("The size required to bind shader %s's parameter %s is %i bytes, smaller than %s's %i bytes."),
+						Shader->GetType()->GetName(),
 						*ShaderBindingName, BoundSize, *CppName, ByteSize);
 				}
 
@@ -157,11 +161,13 @@ struct FShaderParameterStructBindingContext
 	}
 }; // struct FShaderParameterStructBindingContext
 
-void FShaderParameterBindings::BindForLegacyShaderParameters(const FShaderParameterMap& ParametersMap, const FShaderParametersMetadata& StructMetaData, bool bShouldBindEverything)
+void FShaderParameterBindings::BindForLegacyShaderParameters(const FShader* Shader, const FShaderParameterMap& ParametersMap, const FShaderParametersMetadata& StructMetaData, bool bShouldBindEverything)
 {
 	checkf(StructMetaData.GetSize() < (1 << (sizeof(uint16) * 8)), TEXT("Shader parameter structure can only have a size < 65536 bytes."));
+	check(this == &Shader->Bindings);
 
 	FShaderParameterStructBindingContext BindingContext;
+	BindingContext.Shader = Shader;
 	BindingContext.Bindings = this;
 	BindingContext.ParametersMap = &ParametersMap;
 	BindingContext.bUseRootShaderParameters = false;
@@ -176,7 +182,10 @@ void FShaderParameterBindings::BindForLegacyShaderParameters(const FShaderParame
 	ParametersMap.GetAllParameterNames(AllParameterNames);
 	if (bShouldBindEverything && BindingContext.ShaderGlobalScopeBindings.Num() != AllParameterNames.Num())
 	{
-		UE_LOG(LogShaders, Error, TEXT("%i shader parameters have not been bound:"), AllParameterNames.Num() - BindingContext.ShaderGlobalScopeBindings.Num());
+		UE_LOG(LogShaders, Error, TEXT("%i shader parameters have not been bound for %s:"),
+			AllParameterNames.Num() - BindingContext.ShaderGlobalScopeBindings.Num(),
+			Shader->GetType()->GetName());
+
 		for (const FString& GlobalParameterName : AllParameterNames)
 		{
 			if (!BindingContext.ShaderGlobalScopeBindings.Contains(GlobalParameterName))
@@ -185,16 +194,21 @@ void FShaderParameterBindings::BindForLegacyShaderParameters(const FShaderParame
 			}
 		}
 
-		// TODO: would be great to have the shader name for the error message.
-		UE_LOG(LogShaders, Fatal, TEXT("Some shader parameters have not been bound."));
+		UE_LOG(LogShaders, Fatal, TEXT("Unable to bind all shader parameters of %s."),
+			Shader->GetType()->GetName());
 	}
 }
 
-void FShaderParameterBindings::BindForRootShaderParameters(const FShaderParameterMap& ParametersMap, const FShaderParametersMetadata& StructMetaData)
+void FShaderParameterBindings::BindForRootShaderParameters(const FShader* Shader, const FShaderParameterMap& ParametersMap)
 {
+	check(this == &Shader->Bindings);
+	check(Shader->GetType()->GetRootParametersMetadata());
+
+	const FShaderParametersMetadata& StructMetaData = *Shader->GetType()->GetRootParametersMetadata();
 	checkf(StructMetaData.GetSize() < (1 << (sizeof(uint16) * 8)), TEXT("Shader parameter structure can only have a size < 65536 bytes."));
 
 	FShaderParameterStructBindingContext BindingContext;
+	BindingContext.Shader = Shader;
 	BindingContext.Bindings = this;
 	BindingContext.ParametersMap = &ParametersMap;
 	BindingContext.bUseRootShaderParameters = true;
@@ -222,7 +236,10 @@ void FShaderParameterBindings::BindForRootShaderParameters(const FShaderParamete
 	ParametersMap.GetAllParameterNames(AllParameterNames);
 	if (BindingContext.ShaderGlobalScopeBindings.Num() != AllParameterNames.Num())
 	{
-		UE_LOG(LogShaders, Error, TEXT("%i shader parameters have not been bound:"), AllParameterNames.Num() - BindingContext.ShaderGlobalScopeBindings.Num());
+		UE_LOG(LogShaders, Error, TEXT("%i shader parameters have not been bound for %s:"),
+			AllParameterNames.Num() - BindingContext.ShaderGlobalScopeBindings.Num(),
+			Shader->GetType()->GetName());
+
 		for (const FString& GlobalParameterName : AllParameterNames)
 		{
 			if (!BindingContext.ShaderGlobalScopeBindings.Contains(GlobalParameterName))
@@ -231,8 +248,8 @@ void FShaderParameterBindings::BindForRootShaderParameters(const FShaderParamete
 			}
 		}
 
-		// TODO: would be great to have the shader name for the error message.
-		UE_LOG(LogShaders, Fatal, TEXT("Some shader parameters have not been bound."));
+		UE_LOG(LogShaders, Fatal, TEXT("Unable to bind all shader parameters of %s."),
+			Shader->GetType()->GetName());
 	}
 }
 
diff --git a/Engine/Source/Runtime/RenderCore/Public/Shader.h b/Engine/Source/Runtime/RenderCore/Public/Shader.h
--- a/Engine/Source/Runtime/RenderCore/Public/Shader.h
+++ b/Engine/Source/Runtime/RenderCore/Public/Shader.h
@@ -653,8 +653,8 @@ public:
 		return Ar;
 	}
 
-	void BindForLegacyShaderParameters(const FShaderParameterMap& ParameterMaps, const FShaderParametersMetadata& StructMetaData, bool bShouldBindEverything = false);
-	void BindForRootShaderParameters(const FShaderParameterMap& ParameterMaps, const FShaderParametersMetadata& StructMetaData);
+	void BindForLegacyShaderParameters(const FShader* Shader, const FShaderParameterMap& ParameterMaps, const FShaderParametersMetadata& StructMetaData, bool bShouldBindEverything = false);
+	void BindForRootShaderParameters(const FShader* Shader, const FShaderParameterMap& ParameterMaps);
 }; // FShaderParameterBindings
 
 
diff --git a/Engine/Source/Runtime/RenderCore/Public/ShaderParameterStruct.h b/Engine/Source/Runtime/RenderCore/Public/ShaderParameterStruct.h
--- a/Engine/Source/Runtime/RenderCore/Public/ShaderParameterStruct.h
+++ b/Engine/Source/Runtime/RenderCore/Public/ShaderParameterStruct.h
@@ -33,6 +33,7 @@
 		: ShaderParentClass(Initializer) \
 	{ \
 		this->Bindings.BindForLegacyShaderParameters( \
+			this, \
 			Initializer.ParameterMap, \
 			*FParameters::FTypeInfo::GetStructMetadata(), \
 			true); \
@@ -47,7 +48,7 @@
 	ShaderClass(const ShaderMetaType::CompiledShaderInitializerType& Initializer) \
 		: ShaderParentClass(Initializer) \
 	{ \
-		this->Bindings.BindForRootShaderParameters(Initializer.ParameterMap, *FParameters::FTypeInfo::GetStructMetadata()); \
+		this->Bindings.BindForRootShaderParameters(this, Initializer.ParameterMap); \
 	} \
 	\
 	ShaderClass() \
diff --git a/Engine/Source/Runtime/Renderer/Private/VolumetricFog.cpp b/Engine/Source/Runtime/Renderer/Private/VolumetricFog.cpp
--- a/Engine/Source/Runtime/Renderer/Private/VolumetricFog.cpp
+++ b/Engine/Source/Runtime/Renderer/Private/VolumetricFog.cpp
@@ -168,7 +168,7 @@ public:
 	FVolumetricFogMaterialSetupCS(const ShaderMetaType::CompiledShaderInitializerType& Initializer)
 		: FGlobalShader(Initializer)
 	{
-		Bindings.BindForLegacyShaderParameters(Initializer.ParameterMap, *FParameters::FTypeInfo::GetStructMetadata());
+		Bindings.BindForLegacyShaderParameters(this, Initializer.ParameterMap, *FParameters::FTypeInfo::GetStructMetadata());
 
 		VolumetricFogParameters.Bind(Initializer.ParameterMap);
 	} 
@@ -599,7 +599,7 @@ public:
 	TVolumetricFogLightScatteringCS(const ShaderMetaType::CompiledShaderInitializerType& Initializer)
 		: FGlobalShader(Initializer)
 	{
-		Bindings.BindForLegacyShaderParameters(Initializer.ParameterMap, *FParameters::FTypeInfo::GetStructMetadata());
+		Bindings.BindForLegacyShaderParameters(this, Initializer.ParameterMap, *FParameters::FTypeInfo::GetStructMetadata());
 
 		LocalShadowedLightScattering.Bind(Initializer.ParameterMap, TEXT("LocalShadowedLightScattering"));
 		LightScatteringHistory.Bind(Initializer.ParameterMap, TEXT("LightScatteringHistory"));
@@ -779,7 +779,7 @@ class FVolumetricFogFinalIntegrationCS : public FGlobalShader
 	FVolumetricFogFinalIntegrationCS(const ShaderMetaType::CompiledShaderInitializerType& Initializer)
 		: FGlobalShader(Initializer)
 	{
-		Bindings.BindForLegacyShaderParameters(Initializer.ParameterMap, *FParameters::FTypeInfo::GetStructMetadata());
+		Bindings.BindForLegacyShaderParameters(this, Initializer.ParameterMap, *FParameters::FTypeInfo::GetStructMetadata());
 
 		VolumetricFogParameters.Bind(Initializer.ParameterMap);
 	}
