Temporally disables the binding size of shader parameter on HLSLCC compiled shader to unblock everyone.

diff --git a/Engine/Source/Runtime/RenderCore/Private/ShaderParameterStruct.cpp b/Engine/Source/Runtime/RenderCore/Private/ShaderParameterStruct.cpp
--- a/Engine/Source/Runtime/RenderCore/Private/ShaderParameterStruct.cpp
+++ b/Engine/Source/Runtime/RenderCore/Private/ShaderParameterStruct.cpp
@@ -106,7 +106,7 @@ struct FShaderParameterStructBindingContext
 				Parameter.ByteOffset = ByteOffset;
 				Parameter.ByteSize = BoundSize;
 
-				if (uint32(BoundSize) > ByteSize)
+				if (uint32(BoundSize) > ByteSize && !IsHlslccShaderPlatform(Shader->GetShaderPlatform())) // TODO(UE-66830): Should check on all platforms.
 				{
 					UE_LOG(LogShaders, Fatal, TEXT("The size required to bind shader %s's parameter %s is %i bytes, smaller than %s's %i bytes."),
 						Shader->GetType()->GetName(),
