Fixes compilation failure on clang.

diff --git a/Engine/Source/Runtime/RenderCore/Public/Shader.h b/Engine/Source/Runtime/RenderCore/Public/Shader.h
--- a/Engine/Source/Runtime/RenderCore/Public/Shader.h
+++ b/Engine/Source/Runtime/RenderCore/Public/Shader.h
@@ -27,6 +27,7 @@ class FGlobalShaderType;
 class FMaterialShaderType;
 class FNiagaraShaderType;
 class FMeshMaterialShaderType;
+class FShader;
 class FShaderPipelineType;
 class FShaderType;
 class FVertexFactoryParameterRef;
