Renames DummyRenderResources.h to CommonRenderResources.h

diff --git a/Engine/Plugins/Experimental/PixelStreaming/Source/PixelStreaming/Private/NvVideoEncoder.cpp b/Engine/Plugins/Experimental/PixelStreaming/Source/PixelStreaming/Private/NvVideoEncoder.cpp
--- a/Engine/Plugins/Experimental/PixelStreaming/Source/PixelStreaming/Private/NvVideoEncoder.cpp
+++ b/Engine/Plugins/Experimental/PixelStreaming/Source/PixelStreaming/Private/NvVideoEncoder.cpp
@@ -15,7 +15,7 @@
 #include "HAL/ThreadSafeCounter.h"
 #include "Misc/CoreDelegates.h"
 #include "Modules/ModuleManager.h"
-#include "DummyRenderResources.h"
+#include "CommonRenderResources.h"
 
 #if defined PLATFORM_WINDOWS
 // Disable macro redefinition warning for compatibility with Windows SDK 8+
diff --git a/Engine/Plugins/Experimental/RemoteSession/Source/RemoteSession/Private/Channels/RemoteSessionARCameraChannel.cpp b/Engine/Plugins/Experimental/RemoteSession/Source/RemoteSession/Private/Channels/RemoteSessionARCameraChannel.cpp
--- a/Engine/Plugins/Experimental/RemoteSession/Source/RemoteSession/Private/Channels/RemoteSessionARCameraChannel.cpp
+++ b/Engine/Plugins/Experimental/RemoteSession/Source/RemoteSession/Private/Channels/RemoteSessionARCameraChannel.cpp
@@ -38,7 +38,7 @@
 
 #include "IImageWrapper.h"
 #include "IImageWrapperModule.h"
-#include "DummyRenderResources.h"
+#include "CommonRenderResources.h"
 
 #define CAMERA_MESSAGE_ADDRESS TEXT("/ARCamera")
 
diff --git a/Engine/Plugins/Runtime/AppleARKit/Source/AppleARKit/Private/AppleARKitVideoOverlay.cpp b/Engine/Plugins/Runtime/AppleARKit/Source/AppleARKit/Private/AppleARKitVideoOverlay.cpp
--- a/Engine/Plugins/Runtime/AppleARKit/Source/AppleARKit/Private/AppleARKitVideoOverlay.cpp
+++ b/Engine/Plugins/Runtime/AppleARKit/Source/AppleARKit/Private/AppleARKitVideoOverlay.cpp
@@ -20,7 +20,7 @@
 #include "Containers/DynamicRHIResourceArray.h"
 #include "PostProcess/SceneFilterRendering.h"
 #include "PostProcessParameters.h"
-#include "DummyRenderResources.h"
+#include "CommonRenderResources.h"
 #if SUPPORTS_ARKIT_1_0
 	#include "IOSAppDelegate.h"
 #endif
diff --git a/Engine/Plugins/Runtime/GoogleARCore/Source/GoogleARCoreRendering/Private/GoogleARCorePassthroughCameraRenderer.cpp b/Engine/Plugins/Runtime/GoogleARCore/Source/GoogleARCoreRendering/Private/GoogleARCorePassthroughCameraRenderer.cpp
--- a/Engine/Plugins/Runtime/GoogleARCore/Source/GoogleARCoreRendering/Private/GoogleARCorePassthroughCameraRenderer.cpp
+++ b/Engine/Plugins/Runtime/GoogleARCore/Source/GoogleARCoreRendering/Private/GoogleARCorePassthroughCameraRenderer.cpp
@@ -18,7 +18,7 @@
 #include "ExternalTexture.h"
 #include "GoogleARCorePassthroughCameraExternalTextureGuid.h"
 #include "GoogleARCoreAndroidHelper.h"
-#include "DummyRenderResources.h"
+#include "CommonRenderResources.h"
 
 FGoogleARCorePassthroughCameraRenderer::FGoogleARCorePassthroughCameraRenderer()
 	: OverlayQuadUVs{ 0.0f, 0.0f, 0.0f, 1.0f, 1.0f, 0.0f, 1.0f, 1.0f }
diff --git a/Engine/Plugins/Runtime/Oculus/OculusVR/Source/OculusHMD/Private/OculusHMD_CustomPresent.cpp b/Engine/Plugins/Runtime/Oculus/OculusVR/Source/OculusHMD/Private/OculusHMD_CustomPresent.cpp
--- a/Engine/Plugins/Runtime/Oculus/OculusVR/Source/OculusHMD/Private/OculusHMD_CustomPresent.cpp
+++ b/Engine/Plugins/Runtime/Oculus/OculusVR/Source/OculusHMD/Private/OculusHMD_CustomPresent.cpp
@@ -8,7 +8,7 @@
 #include "PipelineStateCache.h"
 #include "ClearQuad.h"
 #include "OculusShaders.h"
-#include "DummyRenderResources.h"
+#include "CommonRenderResources.h"
 
 #if PLATFORM_ANDROID
 #include "Android/AndroidJNI.h"
diff --git a/Engine/Source/Editor/MovieSceneTools/Private/TrackEditorThumbnail/TrackEditorThumbnail.cpp b/Engine/Source/Editor/MovieSceneTools/Private/TrackEditorThumbnail/TrackEditorThumbnail.cpp
--- a/Engine/Source/Editor/MovieSceneTools/Private/TrackEditorThumbnail/TrackEditorThumbnail.cpp
+++ b/Engine/Source/Editor/MovieSceneTools/Private/TrackEditorThumbnail/TrackEditorThumbnail.cpp
@@ -17,7 +17,7 @@
 #include "ScreenRendering.h"
 #include "TrackEditorThumbnail/TrackEditorThumbnailPool.h"
 #include "PipelineStateCache.h"
-#include "DummyRenderResources.h"
+#include "CommonRenderResources.h"
 
 namespace TrackEditorThumbnailConstants
 {
diff --git a/Engine/Source/Runtime/HeadMountedDisplay/Private/DefaultStereoLayers.cpp b/Engine/Source/Runtime/HeadMountedDisplay/Private/DefaultStereoLayers.cpp
--- a/Engine/Source/Runtime/HeadMountedDisplay/Private/DefaultStereoLayers.cpp
+++ b/Engine/Source/Runtime/HeadMountedDisplay/Private/DefaultStereoLayers.cpp
@@ -13,7 +13,7 @@
 #include "PipelineStateCache.h"
 #include "ClearQuad.h"
 #include "SceneView.h"
-#include "DummyRenderResources.h"
+#include "CommonRenderResources.h"
 
 namespace 
 {
diff --git a/Engine/Source/Runtime/HeadMountedDisplay/Private/HeadMountedDisplayTypes.cpp b/Engine/Source/Runtime/HeadMountedDisplay/Private/HeadMountedDisplayTypes.cpp
--- a/Engine/Source/Runtime/HeadMountedDisplay/Private/HeadMountedDisplayTypes.cpp
+++ b/Engine/Source/Runtime/HeadMountedDisplay/Private/HeadMountedDisplayTypes.cpp
@@ -2,7 +2,7 @@
 
 #include "HeadMountedDisplayTypes.h"
 #include "RendererInterface.h"
-#include "DummyRenderResources.h"
+#include "CommonRenderResources.h"
 
 DEFINE_LOG_CATEGORY(LogHMD);
 DEFINE_LOG_CATEGORY(LogLoadingSplash);
diff --git a/Engine/Source/Runtime/MovieSceneCapture/Private/FrameGrabber.cpp b/Engine/Source/Runtime/MovieSceneCapture/Private/FrameGrabber.cpp
--- a/Engine/Source/Runtime/MovieSceneCapture/Private/FrameGrabber.cpp
+++ b/Engine/Source/Runtime/MovieSceneCapture/Private/FrameGrabber.cpp
@@ -15,7 +15,7 @@
 #include "GlobalShader.h"
 #include "ScreenRendering.h"
 #include "PipelineStateCache.h"
-#include "DummyRenderResources.h"
+#include "CommonRenderResources.h"
 
 int32 GFrameGrabberFrameLatency = 0;
 static FAutoConsoleVariableRef CVarFrameGrabberFrameLatency(
diff --git a/Engine/Source/Runtime/RenderCore/Private/CommonRenderResources.cpp b/Engine/Source/Runtime/RenderCore/Private/CommonRenderResources.cpp
new file mode 100644
--- /dev/null
+++ b/Engine/Source/Runtime/RenderCore/Private/CommonRenderResources.cpp
@@ -0,0 +1,62 @@
+// Copyright 1998-2018 Epic Games, Inc. All Rights Reserved.
+
+/*=============================================================================
+	DummyRenderResources.cpp: Implementations of frequently used render resources.
+=============================================================================*/
+
+#include "CommonRenderResources.h"
+#include "Containers/DynamicRHIResourceArray.h"
+
+
+TGlobalResource<FFilterVertexDeclaration> GFilterVertexDeclaration;
+TGlobalResource<FEmptyVertexDeclaration> GEmptyVertexDeclaration;
+
+TGlobalResource<FScreenRectangleVertexBuffer> GScreenRectangleVertexBuffer;
+TGlobalResource<FScreenRectangleIndexBuffer> GScreenRectangleIndexBuffer;
+
+IMPLEMENT_GLOBAL_SHADER(FVisualizeTextureVS, "/Engine/Private/Tools/FullscreenVertexShader.usf", "MainVS", SF_Vertex);
+
+
+void FScreenRectangleVertexBuffer::InitRHI()
+{
+	TResourceArray<FFilterVertex, VERTEXBUFFER_ALIGNMENT> Vertices;
+	Vertices.SetNumUninitialized(6);
+
+	Vertices[0].Position = FVector4(1, 1, 0, 1);
+	Vertices[0].UV = FVector2D(1, 1);
+
+	Vertices[1].Position = FVector4(0, 1, 0, 1);
+	Vertices[1].UV = FVector2D(0, 1);
+
+	Vertices[2].Position = FVector4(1, 0, 0, 1);
+	Vertices[2].UV = FVector2D(1, 0);
+
+	Vertices[3].Position = FVector4(0, 0, 0, 1);
+	Vertices[3].UV = FVector2D(0, 0);
+
+	//The final two vertices are used for the triangle optimization (a single triangle spans the entire viewport )
+	Vertices[4].Position = FVector4(-1, 1, 0, 1);
+	Vertices[4].UV = FVector2D(-1, 1);
+
+	Vertices[5].Position = FVector4(1, -1, 0, 1);
+	Vertices[5].UV = FVector2D(1, -1);
+
+	// Create vertex buffer. Fill buffer with initial data upon creation
+	FRHIResourceCreateInfo CreateInfo(&Vertices);
+	VertexBufferRHI = RHICreateVertexBuffer(Vertices.GetResourceDataSize(), BUF_Static, CreateInfo);
+}
+
+void FScreenRectangleIndexBuffer::InitRHI()
+{
+	// Indices 0 - 5 are used for rendering a quad. Indices 6 - 8 are used for triangle optimization.
+	const uint16 Indices[] = { 0, 1, 2, 2, 1, 3, 0, 4, 5 };
+
+	TResourceArray<uint16, INDEXBUFFER_ALIGNMENT> IndexBuffer;
+	uint32 NumIndices = ARRAY_COUNT(Indices);
+	IndexBuffer.AddUninitialized(NumIndices);
+	FMemory::Memcpy(IndexBuffer.GetData(), Indices, NumIndices * sizeof(uint16));
+
+	// Create index buffer. Fill buffer with initial data upon creation
+	FRHIResourceCreateInfo CreateInfo(&IndexBuffer);
+	IndexBufferRHI = RHICreateIndexBuffer(sizeof(uint16), IndexBuffer.GetResourceDataSize(), BUF_Static, CreateInfo);
+}
diff --git a/Engine/Source/Runtime/RenderCore/Private/DummyRenderResources.cpp b/Engine/Source/Runtime/RenderCore/Private/DummyRenderResources.cpp
deleted file mode 100644
diff --git a/Engine/Source/Runtime/RenderCore/Private/PixelShaderUtils.cpp b/Engine/Source/Runtime/RenderCore/Private/PixelShaderUtils.cpp
--- a/Engine/Source/Runtime/RenderCore/Private/PixelShaderUtils.cpp
+++ b/Engine/Source/Runtime/RenderCore/Private/PixelShaderUtils.cpp
@@ -5,7 +5,7 @@
 =============================================================================*/
 
 #include "PixelShaderUtils.h"
-#include "DummyRenderResources.h"
+#include "CommonRenderResources.h"
 #include "RenderGraph.h"
 
 
diff --git a/Engine/Source/Runtime/RenderCore/Private/VisualizeTexture.cpp b/Engine/Source/Runtime/RenderCore/Private/VisualizeTexture.cpp
--- a/Engine/Source/Runtime/RenderCore/Private/VisualizeTexture.cpp
+++ b/Engine/Source/Runtime/RenderCore/Private/VisualizeTexture.cpp
@@ -16,7 +16,7 @@
 #include "RenderTargetPool.h"
 #include "GlobalShader.h"
 #include "PipelineStateCache.h"
-#include "DummyRenderResources.h"
+#include "CommonRenderResources.h"
 #include "PixelShaderUtils.h"
 
 
diff --git a/Engine/Source/Runtime/RenderCore/Public/CommonRenderResources.h b/Engine/Source/Runtime/RenderCore/Public/CommonRenderResources.h
new file mode 100644
--- /dev/null
+++ b/Engine/Source/Runtime/RenderCore/Public/CommonRenderResources.h
@@ -0,0 +1,105 @@
+// Copyright 1998-2018 Epic Games, Inc. All Rights Reserved.
+
+/*=============================================================================
+	DummyRenderResource.h: Frequently used rendering resources
+=============================================================================*/
+
+#pragma once
+
+#include "RenderResource.h"
+#include "GlobalShader.h"
+#include "ShaderParameterStruct.h"
+
+
+/** The vertex data used to filter a texture. */
+struct FFilterVertex
+{
+	FVector4 Position;
+	FVector2D UV;
+};
+
+/** The filter vertex declaration resource type. */
+class FFilterVertexDeclaration : public FRenderResource
+{
+public:
+	FVertexDeclarationRHIRef VertexDeclarationRHI;
+
+	/** Destructor. */
+	virtual ~FFilterVertexDeclaration() {}
+
+	virtual void InitRHI()
+	{
+		FVertexDeclarationElementList Elements;
+		uint32 Stride = sizeof(FFilterVertex);
+		Elements.Add(FVertexElement(0, STRUCT_OFFSET(FFilterVertex, Position), VET_Float4, 0, Stride));
+		Elements.Add(FVertexElement(0, STRUCT_OFFSET(FFilterVertex, UV), VET_Float2, 1, Stride));
+		VertexDeclarationRHI = RHICreateVertexDeclaration(Elements);
+	}
+
+	virtual void ReleaseRHI()
+	{
+		VertexDeclarationRHI.SafeRelease();
+	}
+};
+
+extern RENDERCORE_API TGlobalResource<FFilterVertexDeclaration> GFilterVertexDeclaration;
+
+/** The empty vertex declaration resource type. */
+class FEmptyVertexDeclaration : public FRenderResource
+{
+public:
+	FVertexDeclarationRHIRef VertexDeclarationRHI;
+
+	/** Destructor. */
+	virtual ~FEmptyVertexDeclaration() {}
+
+	virtual void InitRHI()
+	{
+		FVertexDeclarationElementList Elements;
+		VertexDeclarationRHI = RHICreateVertexDeclaration(Elements);
+	}
+
+	virtual void ReleaseRHI()
+	{
+		VertexDeclarationRHI.SafeRelease();
+	}
+};
+
+extern RENDERCORE_API TGlobalResource<FEmptyVertexDeclaration> GEmptyVertexDeclaration;
+
+/**
+ * Static vertex and index buffer used for 2D screen rectangles.
+ */
+class FScreenRectangleVertexBuffer : public FVertexBuffer
+{
+public:
+	/** Initialize the RHI for this rendering resource */
+	void InitRHI() override;
+};
+
+extern RENDERCORE_API TGlobalResource<FScreenRectangleVertexBuffer> GScreenRectangleVertexBuffer;
+
+
+class FScreenRectangleIndexBuffer : public FIndexBuffer
+{
+public:
+	/** Initialize the RHI for this rendering resource */
+	void InitRHI() override;
+};
+
+extern RENDERCORE_API TGlobalResource<FScreenRectangleIndexBuffer> GScreenRectangleIndexBuffer;
+
+
+/** Vertex shader to draw a full screen quad that works on all platforms. */
+class RENDERCORE_API FVisualizeTextureVS : public FGlobalShader
+{
+	DECLARE_GLOBAL_SHADER(FVisualizeTextureVS);
+	SHADER_USE_PARAMETER_STRUCT(FVisualizeTextureVS, FGlobalShader);
+
+	static bool ShouldCompilePermutation(const FGlobalShaderPermutationParameters& Parameters) {
+		return true;
+	}
+
+	BEGIN_SHADER_PARAMETER_STRUCT(FParameters, )
+	END_SHADER_PARAMETER_STRUCT()
+};
diff --git a/Engine/Source/Runtime/RenderCore/Public/DummyRenderResources.h b/Engine/Source/Runtime/RenderCore/Public/DummyRenderResources.h
deleted file mode 100644
diff --git a/Engine/Source/Runtime/RenderCore/Public/PixelShaderUtils.h b/Engine/Source/Runtime/RenderCore/Public/PixelShaderUtils.h
--- a/Engine/Source/Runtime/RenderCore/Public/PixelShaderUtils.h
+++ b/Engine/Source/Runtime/RenderCore/Public/PixelShaderUtils.h
@@ -8,7 +8,7 @@
 
 #include "RenderResource.h"
 #include "RenderGraphUtils.h"
-#include "DummyRenderResources.h"
+#include "CommonRenderResources.h"
 #include "RHI.h"
 #include "RHIStaticStates.h"
 #include "PipelineStateCache.h"
diff --git a/Engine/Source/Runtime/Renderer/Private/GPUBenchmark.cpp b/Engine/Source/Runtime/Renderer/Private/GPUBenchmark.cpp
--- a/Engine/Source/Runtime/Renderer/Private/GPUBenchmark.cpp
+++ b/Engine/Source/Runtime/Renderer/Private/GPUBenchmark.cpp
@@ -22,7 +22,7 @@
 #include "PipelineStateCache.h"
 #include "LongGPUTask.h"
 #include "VisualizeTexture.h"
-#include "DummyRenderResources.h"
+#include "CommonRenderResources.h"
 
 static const uint32 GBenchmarkResolution = 512;
 static const uint32 GBenchmarkPrimitives = 200000;
diff --git a/Engine/Source/Runtime/Renderer/Private/GammaCorrection.cpp b/Engine/Source/Runtime/Renderer/Private/GammaCorrection.cpp
--- a/Engine/Source/Runtime/Renderer/Private/GammaCorrection.cpp
+++ b/Engine/Source/Runtime/Renderer/Private/GammaCorrection.cpp
@@ -16,7 +16,7 @@
 #include "PostProcess/SceneFilterRendering.h"
 #include "PipelineStateCache.h"
 #include "ClearQuad.h"
-#include "DummyRenderResources.h"
+#include "CommonRenderResources.h"
 
 /** Encapsulates the gamma correction pixel shader. */
 class FGammaCorrectionPS : public FGlobalShader
diff --git a/Engine/Source/Runtime/Renderer/Private/MobileSceneCaptureRendering.cpp b/Engine/Source/Runtime/Renderer/Private/MobileSceneCaptureRendering.cpp
--- a/Engine/Source/Runtime/Renderer/Private/MobileSceneCaptureRendering.cpp
+++ b/Engine/Source/Runtime/Renderer/Private/MobileSceneCaptureRendering.cpp
@@ -26,7 +26,7 @@ MobileSceneCaptureRendering.cpp - Mobile specific scene capture code.
 #include "ScreenRendering.h"
 #include "ClearQuad.h"
 #include "PipelineStateCache.h"
-#include "DummyRenderResources.h"
+#include "CommonRenderResources.h"
 
 /**
 * Shader set for the copy of scene color to capture target, decoding mosaic or RGBE encoded HDR image as part of a
diff --git a/Engine/Source/Runtime/Renderer/Private/PostProcess/SceneFilterRendering.cpp b/Engine/Source/Runtime/Renderer/Private/PostProcess/SceneFilterRendering.cpp
--- a/Engine/Source/Runtime/Renderer/Private/PostProcess/SceneFilterRendering.cpp
+++ b/Engine/Source/Runtime/Renderer/Private/PostProcess/SceneFilterRendering.cpp
@@ -11,7 +11,7 @@
 #include "IHeadMountedDisplay.h"
 #include "IXRTrackingSystem.h"
 #include "PixelShaderUtils.h"
-#include "DummyRenderResources.h"
+#include "CommonRenderResources.h"
 
 void FTesselatedScreenRectangleIndexBuffer::InitRHI()
 {
diff --git a/Engine/Source/Runtime/Renderer/Private/ReflectionEnvironmentDiffuseIrradiance.cpp b/Engine/Source/Runtime/Renderer/Private/ReflectionEnvironmentDiffuseIrradiance.cpp
--- a/Engine/Source/Runtime/Renderer/Private/ReflectionEnvironmentDiffuseIrradiance.cpp
+++ b/Engine/Source/Runtime/Renderer/Private/ReflectionEnvironmentDiffuseIrradiance.cpp
@@ -17,7 +17,7 @@
 #include "ScreenRendering.h"
 #include "PipelineStateCache.h"
 #include "VisualizeTexture.h"
-#include "DummyRenderResources.h"
+#include "CommonRenderResources.h"
 
 extern int32 GDiffuseIrradianceCubemapSize;
 
diff --git a/Engine/Source/Runtime/Renderer/Private/ScenePrivate.h b/Engine/Source/Runtime/Renderer/Private/ScenePrivate.h
--- a/Engine/Source/Runtime/Renderer/Private/ScenePrivate.h
+++ b/Engine/Source/Runtime/Renderer/Private/ScenePrivate.h
@@ -41,7 +41,7 @@
 #include "MobileBasePassRendering.h"
 #include "VolumeRendering.h"
 #include "SceneSoftwareOcclusion.h"
-#include "DummyRenderResources.h"
+#include "CommonRenderResources.h"
 
 /** Factor by which to grow occlusion tests **/
 #define OCCLUSION_SLOP (1.0f)
diff --git a/Engine/Source/Runtime/Renderer/Private/VT/VirtualTextureSpace.cpp b/Engine/Source/Runtime/Renderer/Private/VT/VirtualTextureSpace.cpp
--- a/Engine/Source/Runtime/Renderer/Private/VT/VirtualTextureSpace.cpp
+++ b/Engine/Source/Runtime/Renderer/Private/VT/VirtualTextureSpace.cpp
@@ -11,7 +11,7 @@
 #include "HAL/IConsoleManager.h"
 #include "SceneUtils.h"
 #include "VisualizeTexture.h"
-#include "DummyRenderResources.h"
+#include "CommonRenderResources.h"
 
 static TAutoConsoleVariable<int32> CVarVTRefreshEntirePageTable(
 	TEXT("r.VT.RefreshEntirePageTable"),
diff --git a/Engine/Source/Runtime/SlateRHIRenderer/Private/SlatePostProcessor.cpp b/Engine/Source/Runtime/SlateRHIRenderer/Private/SlatePostProcessor.cpp
--- a/Engine/Source/Runtime/SlateRHIRenderer/Private/SlatePostProcessor.cpp
+++ b/Engine/Source/Runtime/SlateRHIRenderer/Private/SlatePostProcessor.cpp
@@ -8,7 +8,7 @@
 #include "RendererInterface.h"
 #include "StaticBoundShaderState.h"
 #include "PipelineStateCache.h"
-#include "DummyRenderResources.h"
+#include "CommonRenderResources.h"
 
 DECLARE_CYCLE_STAT(TEXT("Slate PostProcessing RT"), STAT_SlatePostProcessingRTTime, STATGROUP_Slate);
 DECLARE_CYCLE_STAT(TEXT("Slate ColorDeficiency RT"), STAT_SlateColorDeficiencyRTTime, STATGROUP_Slate);
diff --git a/Engine/Source/Runtime/SlateRHIRenderer/Private/SlateRHIRenderer.cpp b/Engine/Source/Runtime/SlateRHIRenderer/Private/SlateRHIRenderer.cpp
--- a/Engine/Source/Runtime/SlateRHIRenderer/Private/SlateRHIRenderer.cpp
+++ b/Engine/Source/Runtime/SlateRHIRenderer/Private/SlateRHIRenderer.cpp
@@ -27,7 +27,7 @@
 #include "EngineModule.h"
 #include "Interfaces/ISlate3DRenderer.h"
 #include "Slate/SlateTextureAtlasInterface.h"
-#include "DummyRenderResources.h"
+#include "CommonRenderResources.h"
 
 DECLARE_CYCLE_STAT(TEXT("Slate RT: Rendering"), STAT_SlateRenderingRTTime, STATGROUP_Slate);
 DECLARE_CYCLE_STAT(TEXT("Slate RT: Create Batches"), STAT_SlateRTCreateBatches, STATGROUP_Slate);
