DR - Undo 4590384
Some float3's have to be converted to float4's in the new resource tables due to hlslcc limitations
#rnx

diff --git a/Engine/Source/Runtime/RenderCore/Private/ShaderParameterStruct.cpp b/Engine/Source/Runtime/RenderCore/Private/ShaderParameterStruct.cpp
--- a/Engine/Source/Runtime/RenderCore/Private/ShaderParameterStruct.cpp
+++ b/Engine/Source/Runtime/RenderCore/Private/ShaderParameterStruct.cpp
@@ -106,11 +106,19 @@ struct FShaderParameterStructBindingContext
 				Parameter.ByteOffset = ByteOffset;
 				Parameter.ByteSize = BoundSize;
 
-				if (uint32(BoundSize) > ByteSize && !IsHlslccShaderPlatform(Shader->GetShaderPlatform())) // TODO(UE-66830): Should check on all platforms.
+				if (uint32(BoundSize) > ByteSize)
 				{
-					UE_LOG(LogShaders, Fatal, TEXT("The size required to bind shader %s's parameter %s is %i bytes, smaller than %s's %i bytes."),
-						Shader->GetType()->GetName(),
-						*ShaderBindingName, BoundSize, *CppName, ByteSize);
+					if (IsHlslccShaderPlatform(Shader->GetShaderPlatform()) && BoundSize == 16 && ByteSize == 12)
+					{
+						UE_LOG(LogShaders, Fatal, TEXT("Shader %s's (Permutation Id %d) struct %s parameter %s is a float3; hlslcc will pad float3 to float4, so please adjust the parameter from a float3 to a float4"),
+							Shader->GetType()->GetName(), Shader->GetPermutationId(), StructMetaData.GetStructTypeName(), *ShaderBindingName);
+					}
+					else
+					{
+						UE_LOG(LogShaders, Fatal, TEXT("The size required to bind shader %s's (Permutation Id %d) struct %s parameter %s is %i bytes, smaller than %s's %i bytes."),
+							Shader->GetType()->GetName(), Shader->GetPermutationId(), StructMetaData.GetStructTypeName(),
+							*ShaderBindingName, BoundSize, *CppName, ByteSize);
+					}
 				}
 
 				Bindings->Parameters.Add(Parameter);
diff --git a/Engine/Source/Runtime/RenderCore/Private/VisualizeTexture.cpp b/Engine/Source/Runtime/RenderCore/Private/VisualizeTexture.cpp
--- a/Engine/Source/Runtime/RenderCore/Private/VisualizeTexture.cpp
+++ b/Engine/Source/Runtime/RenderCore/Private/VisualizeTexture.cpp
@@ -55,9 +55,12 @@ enum class EVisualisePSType
 	MAX
 };
 
+#endif
+
 TGlobalResource<FVisualizeTexture> GVisualizeTexture;
 
 
+#if WITH_ENGINE
 /** A pixel shader which filters a texture. */
 // @param TextureType 0:Cube, 1:1D(not yet supported), 2:2D no MSAA, 3:3D, 4:Cube[], 5:2D MSAA, 6:2D DepthStencil no MSAA (needed to avoid D3DDebug error)
 class FVisualizeTexturePS : public FGlobalShader
@@ -76,7 +79,7 @@ class FVisualizeTexturePS : public FGlobalShader
 	}
 
 	BEGIN_SHADER_PARAMETER_STRUCT(FParameters, )
-		SHADER_PARAMETER(FVector, TextureExtent)
+		SHADER_PARAMETER(FVector4, TextureExtent)
 		SHADER_PARAMETER_ARRAY(FVector4, VisualizeParam, [3])
 
 		SHADER_PARAMETER_RDG_TEXTURE(Texture2D, VisualizeTexture2D)
@@ -227,7 +230,7 @@ void FVisualizeTexture::CreateContentCapturePass(FRDGBuilder& GraphBuilder, cons
 
 	FVisualizeTexturePS::FParameters* PassParameters = GraphBuilder.AllocParameters<FVisualizeTexturePS::FParameters>();
 	{
-		PassParameters->TextureExtent = FVector(SrcDesc.Extent.X, SrcDesc.Extent.Y, SrcDesc.Depth);
+		PassParameters->TextureExtent = FVector4(SrcDesc.Extent.X, SrcDesc.Extent.Y, SrcDesc.Depth, 0.0f);
 
 		{
 			// alternates between 0 and 1 with a short pause
