Fix cis

diff --git a/Engine/Source/Runtime/RenderCore/Public/RenderingThread.h b/Engine/Source/Runtime/RenderCore/Public/RenderingThread.h
--- a/Engine/Source/Runtime/RenderCore/Public/RenderingThread.h
+++ b/Engine/Source/Runtime/RenderCore/Public/RenderingThread.h
@@ -669,7 +669,7 @@ public:
 		ENQUEUE_RENDER_COMMAND(DispatchScopeCommands)(
 			[RenderCommandArray](FRHICommandListImmediate& RHICmdList)
 		{
-			for(uint32 Index = 0; Index < RenderCommandArray->Num(); Index++)
+			for(int32 Index = 0; Index < RenderCommandArray->Num(); Index++)
 			{
 				(*RenderCommandArray)[Index](RHICmdList);
 			}
