DR - Copy from 4629036
Static analysis fix
#rnx

diff --git a/Engine/Source/Runtime/RenderCore/Private/RenderGraphBuilder.cpp b/Engine/Source/Runtime/RenderCore/Private/RenderGraphBuilder.cpp
--- a/Engine/Source/Runtime/RenderCore/Private/RenderGraphBuilder.cpp
+++ b/Engine/Source/Runtime/RenderCore/Private/RenderGraphBuilder.cpp
@@ -624,9 +624,12 @@ void FRDGBuilder::ExecutePass( const FRenderGraphPass* Pass )
 
 	AllocateAndTransitionPassResources(Pass, &RPInfo, &bHasRenderTargets);
 
-	if (RENDER_GRAPH_DRAW_EVENTS && GetEmitDrawEvents())
+	if (RENDER_GRAPH_DRAW_EVENTS)
 	{
-		PushDrawEventStack(Pass);
+		if (GetEmitDrawEvents())
+		{
+			PushDrawEventStack(Pass);
+		}
 	}
 
 	if( !Pass->IsCompute())
@@ -649,9 +652,12 @@ void FRDGBuilder::ExecutePass( const FRenderGraphPass* Pass )
 		RHICmdList.EndRenderPass();
 	}
 
-	if (RENDER_GRAPH_DRAW_EVENTS && GetEmitDrawEvents())
+	if (RENDER_GRAPH_DRAW_EVENTS)
 	{
-		RHICmdList.PopEvent();
+		if (GetEmitDrawEvents())
+		{
+			RHICmdList.PopEvent();
+		}
 	}
 
 	if (RENDER_GRAPH_DEBUGGING)
