DR - Copy from 4629036
Static analysis fix
#rnx

diff --git a/Engine/Source/Runtime/RenderCore/Private/Shader.cpp b/Engine/Source/Runtime/RenderCore/Private/Shader.cpp
--- a/Engine/Source/Runtime/RenderCore/Private/Shader.cpp
+++ b/Engine/Source/Runtime/RenderCore/Private/Shader.cpp
@@ -2104,8 +2104,8 @@ void ShaderMapAppendKeyString(EShaderPlatform Platform, FString& KeyString)
 		else
 		{
 			// shader platform doesn't match a specific target platform, use cvar setting for forward shading
-			static IConsoleVariable* CVarForwardShading = IConsoleManager::Get().FindConsoleVariable(TEXT("r.ForwardShading"));
-			bForwardShading = CVarForwardShading ? (CVarForwardShading->GetInt() != 0) : false;
+			static IConsoleVariable* CVarForwardShadingLocal = IConsoleManager::Get().FindConsoleVariable(TEXT("r.ForwardShading"));
+			bForwardShading = CVarForwardShadingLocal ? (CVarForwardShadingLocal->GetInt() != 0) : false;
 		}
 
 		if (bForwardShading)
diff --git a/Engine/Source/Runtime/RenderCore/Public/PixelShaderUtils.h b/Engine/Source/Runtime/RenderCore/Public/PixelShaderUtils.h
--- a/Engine/Source/Runtime/RenderCore/Public/PixelShaderUtils.h
+++ b/Engine/Source/Runtime/RenderCore/Public/PixelShaderUtils.h
@@ -32,6 +32,7 @@ struct RENDERCORE_API FPixelShaderUtils
 		const TShaderClass* PixelShader,
 		const typename TShaderClass::FParameters& Parameters)
 	{
+		CA_ASSUME(PixelShader);
 		TShaderMapRef<FVisualizeTextureVS> VertexShader(GlobalShaderMap);
 
 		FGraphicsPipelineStateInitializer GraphicsPSOInit;
@@ -60,6 +61,7 @@ struct RENDERCORE_API FPixelShaderUtils
 		const TShaderClass* PixelShader,
 		typename TShaderClass::FParameters* Parameters)
 	{
+		CA_ASSUME(PixelShader != nullptr);
 		ClearUnusedGraphResources(PixelShader, Parameters);
 
 		GraphBuilder.AddPass(
