DR - Generate mips fixes
#rnx

diff --git a/Engine/Source/Runtime/RenderCore/Private/RenderGraphBuilder.cpp b/Engine/Source/Runtime/RenderCore/Private/RenderGraphBuilder.cpp
--- a/Engine/Source/Runtime/RenderCore/Private/RenderGraphBuilder.cpp
+++ b/Engine/Source/Runtime/RenderCore/Private/RenderGraphBuilder.cpp
@@ -23,6 +23,43 @@ static const int32 GRenderGraphImmediateMode = 0;
 #endif
 
 
+
+static bool IsBoundAsReadable(const FRDGTexture* Texture, FShaderParameterStructRef ParameterStruct)
+{
+	for (int i = 0, Num = ParameterStruct.Layout->Resources.Num(); i < Num; i++)
+	{
+		uint8  Type = ParameterStruct.Layout->Resources[i];
+		uint16 Offset = ParameterStruct.Layout->ResourceOffsets[i];
+
+		switch (Type)
+		{
+		case UBMT_GRAPH_TRACKED_TEXTURE:
+		{
+			const FRDGTexture* InputTexture = *ParameterStruct.GetMemberPtrAtOffset<const FRDGTexture*>(Offset);
+			if (Texture == InputTexture)
+			{
+				return true;
+			}
+		}
+		break;
+		case UBMT_GRAPH_TRACKED_SRV:
+		{
+			const FRDGTextureSRV* InputSRV = *ParameterStruct.GetMemberPtrAtOffset<const FRDGTextureSRV*>(Offset);
+			if (InputSRV && Texture == InputSRV->Desc.Texture)
+			{
+				return true;
+			}
+		}
+		break;
+		default:
+			break;
+		}
+	}
+
+	return false;
+}
+
+
 #if RENDER_GRAPH_DRAW_EVENTS == 2
 
 FRDGEventName::FRDGEventName(const TCHAR* EventFormat, ...)
@@ -149,6 +186,9 @@ void FRDGBuilder::ValidatePass(const FRenderGraphPass* Pass) const
 	{
 		checkf(bRequiresRenderTargetSlots, TEXT("Render pass %s does not need render target binging slots"), Pass->GetName());
 
+		const bool bGeneratingMips = (Pass->GetFlags() & ERenderGraphPassFlags::GenerateMips) == ERenderGraphPassFlags::GenerateMips;
+		bool bFoundRTBound = false;
+
 		uint32 NumRenderTargets = 0;
 		for (int i = 0; i < RenderTargets->Output.Num(); i++)
 		{
@@ -158,12 +198,14 @@ void FRDGBuilder::ValidatePass(const FRenderGraphPass* Pass) const
 				NumRenderTargets = i;
 				break;
 			}
+			bFoundRTBound = bFoundRTBound || IsBoundAsReadable(RenderTarget.GetTexture(), ParameterStruct);
 		}
 		for (int i = NumRenderTargets; i < RenderTargets->Output.Num(); i++)
 		{
 			const FRenderTargetBinding& RenderTarget = RenderTargets->Output[i];
 			checkf(RenderTarget.GetTexture() == nullptr, TEXT("Render targets must be packed. No empty spaces in the array."));
 		}
+		ensureMsgf(!bGeneratingMips || bFoundRTBound, TEXT("GenerateMips enabled but no RT found as source!"));
 	}
 	else
 	{
@@ -510,41 +552,6 @@ void FRDGBuilder::TransitionUAV(FUnorderedAccessViewRHIParamRef UAV, const FRDGR
 	}
 }
 
-static bool IsBoundAsReadable( const FRDGTexture* Texture, FShaderParameterStructRef ParameterStruct )
-{
-	for( int i = 0, Num = ParameterStruct.Layout->Resources.Num(); i < Num; i++ )
-	{
-		uint8  Type   = ParameterStruct.Layout->Resources[i];
-		uint16 Offset = ParameterStruct.Layout->ResourceOffsets[i];
-
-		switch( Type )
-		{
-			case UBMT_GRAPH_TRACKED_TEXTURE:
-			{
-				const FRDGTexture* InputTexture = *ParameterStruct.GetMemberPtrAtOffset<const FRDGTexture*>(Offset);
-				if( Texture == InputTexture)
-				{
-					return true;
-				}
-			}
-			break;
-		case UBMT_GRAPH_TRACKED_SRV:
-			{
-				const FRDGTextureSRV* InputSRV = *ParameterStruct.GetMemberPtrAtOffset<const FRDGTextureSRV*>(Offset);
-				if (InputSRV && Texture == InputSRV->Desc.Texture)
-				{
-					return true;
-				}
-			}
-			break;
-		default:
-			break;
-		}
-	}
-
-	return false;
-}
-
 void FRDGBuilder::PushDrawEventStack(const FRenderGraphPass* Pass)
 {
 	// Push the scope event.
@@ -677,6 +684,7 @@ void FRDGBuilder::AllocateAndTransitionPassResources(const FRenderGraphPass* Pas
 	bool bIsCompute = Pass->IsCompute();
 	FShaderParameterStructRef ParameterStruct = Pass->GetParameters();
 
+	const bool bGeneratingMips = (Pass->GetFlags() & ERenderGraphPassFlags::GenerateMips) == ERenderGraphPassFlags::GenerateMips;
 	for (int ResourceIndex = 0, Num = ParameterStruct.Layout->Resources.Num(); ResourceIndex < Num; ResourceIndex++)
 	{
 		uint8  Type = ParameterStruct.Layout->Resources[ResourceIndex];
@@ -778,10 +786,11 @@ void FRDGBuilder::AllocateAndTransitionPassResources(const FRenderGraphPass* Pas
 					OutRPInfo->ColorRenderTargets[i].MipIndex = RenderTarget.GetMipIndex();
 					OutRPInfo->ColorRenderTargets[i].Action = MakeRenderTargetActions(RenderTarget.GetLoadAction(), RenderTarget.GetStoreAction());
 
-					TransitionTexture(RenderTarget.GetTexture(), EResourceTransitionAccess::EWritable, false);
-
-					// TODO(RDG): There must be a better way to do this.
-					OutRPInfo->bGeneratingMips = OutRPInfo->bGeneratingMips || IsBoundAsReadable(RenderTarget.GetTexture(), ParameterStruct);
+					if (!bGeneratingMips)
+					{
+						// Implicit assurance the RHI will do the correct transitions
+						TransitionTexture(RenderTarget.GetTexture(), EResourceTransitionAccess::EWritable, false);
+					}
 
 					NumSamples |= OutRPInfo->ColorRenderTargets[i].RenderTarget->GetNumSamples();
 					NumRenderTargets++;
@@ -819,6 +828,8 @@ void FRDGBuilder::AllocateAndTransitionPassResources(const FRenderGraphPass* Pas
 			break;
 		}
 	}
+
+	OutRPInfo->bGeneratingMips = bGeneratingMips;
 }
 
 // static 
diff --git a/Engine/Source/Runtime/RenderCore/Public/RenderGraphBuilder.h b/Engine/Source/Runtime/RenderCore/Public/RenderGraphBuilder.h
--- a/Engine/Source/Runtime/RenderCore/Public/RenderGraphBuilder.h
+++ b/Engine/Source/Runtime/RenderCore/Public/RenderGraphBuilder.h
@@ -89,6 +89,10 @@ enum class ERenderGraphPassFlags
 
 	/** Pass uses compute only */
 	Compute = 1 << 0,
+
+	//#todo-rco: Remove this when we can do split/per mip layout transitions.
+	/** Hint to some RHIs this pass will be generating mips to optimize transitions. */
+	GenerateMips = 1 << 1,
 };
 
 ENUM_CLASS_FLAGS(ERenderGraphPassFlags)
diff --git a/Engine/Source/Runtime/Renderer/Private/SceneOcclusion.cpp b/Engine/Source/Runtime/Renderer/Private/SceneOcclusion.cpp
--- a/Engine/Source/Runtime/Renderer/Private/SceneOcclusion.cpp
+++ b/Engine/Source/Runtime/Renderer/Private/SceneOcclusion.cpp
@@ -1204,7 +1204,7 @@ void BuildHZB(FRDGBuilder& GraphBuilder, FViewInfo& View)
 		GraphBuilder.AddPass(
 			RDG_EVENT_NAME("HZB(mip=%d) %dx%d", MipIndex, DstSize.X, DstSize.Y),
 			PassParameters,
-			ERenderGraphPassFlags::None,
+			ERenderGraphPassFlags::GenerateMips,
 			[PassParameters, SrcSize, DstSize, &View](FRHICommandListImmediate& RHICmdList)
 			{
 				FHZBBuildPS::FPermutationDomain PermutationVector;
