Add a CVar, r.ShaderPipelineCache.StartupMode  to control what mode the PipelineCache starts in so we can control whether it starts paused (0), in 'fast' mode (1 - default), or in 'background' mode (2).

diff --git a/Engine/Source/Runtime/RenderCore/Private/ShaderPipelineCache.cpp b/Engine/Source/Runtime/RenderCore/Private/ShaderPipelineCache.cpp
--- a/Engine/Source/Runtime/RenderCore/Private/ShaderPipelineCache.cpp
+++ b/Engine/Source/Runtime/RenderCore/Private/ShaderPipelineCache.cpp
@@ -44,6 +44,18 @@ namespace FShaderPipelineCacheConstants
 	static TCHAR const* GameVersionKey = TEXT("GameVersion");
 }
 
+
+static TAutoConsoleVariable<int32> CVarPSOFileCacheStartupMode(
+														  TEXT("r.ShaderPipelineCache.StartupMode"),
+														  1,
+														  TEXT("Sets the startup mode for the PSO cache, determining what the cache does after initialisation:\n")
+														  TEXT("\t0: Precompilation is paused and nothing will compile until a call to ResumeBatching().\n")
+														  TEXT("\t1: Precompilation is enabled in the 'Fast' mode.\n")
+														  TEXT("\t2: Precompilation is enabled in the 'Background' mode.\n")
+														  TEXT("Default is 1."),
+														  ECVF_Default | ECVF_RenderThreadSafe
+														  );
+
 static TAutoConsoleVariable<int32> CVarPSOFileCacheBackgroundBatchSize(
 														  TEXT("r.ShaderPipelineCache.BackgroundBatchSize"),
 														  1,
@@ -1029,6 +1041,25 @@ FShaderPipelineCache::FShaderPipelineCache(EShaderPlatform Platform)
     SET_DWORD_STAT(STAT_ShaderPipelineWaitingTaskCount, 0);
     SET_DWORD_STAT(STAT_ShaderPipelineActiveTaskCount, 0);
 	
+	int32 Mode = CVarPSOFileCacheStartupMode.GetValueOnAnyThread();
+	switch (Mode)
+	{
+		case 0:
+			BatchSize = CVarPSOFileCacheBatchSize.GetValueOnAnyThread();
+			BatchTime = CVarPSOFileCacheBatchTime.GetValueOnAnyThread();
+			bPaused = true;
+			break;
+		case 2:
+			BatchSize = CVarPSOFileCacheBackgroundBatchSize.GetValueOnAnyThread();
+			BatchTime = CVarPSOFileCacheBackgroundBatchTime.GetValueOnAnyThread();
+			break;
+		case 1:
+		default:
+			BatchSize = CVarPSOFileCacheBatchSize.GetValueOnAnyThread();
+			BatchTime = CVarPSOFileCacheBatchTime.GetValueOnAnyThread();
+			break;
+	}
+	
 	BatchSize = CVarPSOFileCacheBatchSize.GetValueOnAnyThread();
 	BatchTime = CVarPSOFileCacheBatchTime.GetValueOnAnyThread();
 	
