DR - Minor fix
#rnx

diff --git a/Engine/Source/Runtime/RenderCore/Public/ShaderParameterMacros.h b/Engine/Source/Runtime/RenderCore/Public/ShaderParameterMacros.h
--- a/Engine/Source/Runtime/RenderCore/Public/ShaderParameterMacros.h
+++ b/Engine/Source/Runtime/RenderCore/Public/ShaderParameterMacros.h
@@ -162,18 +162,22 @@ struct alignas(SHADER_PARAMETER_STRUCT_ALIGNMENT) FRenderTargetBinding
 		, MipIndex(InMipIndex)
 	{}
 
-	inline const FRDGTexture* GetTexture() const {
+	inline const FRDGTexture* GetTexture() const
+	{
 		return Texture;
-	};
-	inline ERenderTargetLoadAction GetLoadAction() const {
+	}
+	inline ERenderTargetLoadAction GetLoadAction() const
+	{
 		return LoadAction;
-	};
-	inline ERenderTargetStoreAction GetStoreAction() const {
+	}
+	inline ERenderTargetStoreAction GetStoreAction() const
+	{
 		return StoreAction;
-	};
-	inline uint8 GetMipIndex() const {
+	}
+	inline uint8 GetMipIndex() const
+	{
 		return MipIndex;
-	};
+	}
 
 private:
 	/** All parameters required to bind a render target deferred. This are purposefully private to
@@ -225,7 +229,8 @@ struct alignas(SHADER_PARAMETER_STRUCT_ALIGNMENT) FRenderTargetBindingSlots
 		return Output[Index];
 	}
 
-	struct FTypeInfo {
+	struct FTypeInfo
+	{
 		static constexpr int32 NumRows = 1;
 		static constexpr int32 NumColumns = 1;
 		static constexpr int32 NumElements = 0;
