Fixed blocker UE-67751. VS2015 did not like converting from initializer list.

diff --git a/Engine/Source/Runtime/RenderCore/Private/Shader.cpp b/Engine/Source/Runtime/RenderCore/Private/Shader.cpp
--- a/Engine/Source/Runtime/RenderCore/Private/Shader.cpp
+++ b/Engine/Source/Runtime/RenderCore/Private/Shader.cpp
@@ -404,7 +404,9 @@ void FShaderResource::GetRayTracingMaterialLibrary(TArray<FRayTracingHitGroupIni
 
 	for (const auto Entry : GlobalRayTracingMaterialLibrary)
 	{
-		RayTracingMaterials[Entry.Key] = FRayTracingHitGroupInitializer{ Entry.Value };
+		FRayTracingHitGroupInitializer HitGroupInitializer;
+		HitGroupInitializer.ShaderRHI = Entry.Value;
+		RayTracingMaterials[Entry.Key] = HitGroupInitializer;
 	}
 
 	for (uint32 Index : GlobalUnusedIndicies)
