Added STAT_ProcessVisibleMeshDrawCommands.

diff --git a/Engine/Source/Runtime/RenderCore/Private/RenderCore.cpp b/Engine/Source/Runtime/RenderCore/Private/RenderCore.cpp
--- a/Engine/Source/Runtime/RenderCore/Private/RenderCore.cpp
+++ b/Engine/Source/Runtime/RenderCore/Private/RenderCore.cpp
@@ -108,6 +108,7 @@ DEFINE_STAT(STAT_CreateWholeSceneProjectedShadow);
 DEFINE_STAT(STAT_AddViewDependentWholeSceneShadowsForView);
 DEFINE_STAT(STAT_SetupInteractionShadows);
 DEFINE_STAT(STAT_GetDynamicMeshElements);
+DEFINE_STAT(STAT_ProcessVisibleMeshDrawCommands);
 DEFINE_STAT(STAT_UpdateStaticMeshesTime);
 DEFINE_STAT(STAT_StaticRelevance);
 DEFINE_STAT(STAT_ViewRelevance);
diff --git a/Engine/Source/Runtime/RenderCore/Public/RenderCore.h b/Engine/Source/Runtime/RenderCore/Public/RenderCore.h
--- a/Engine/Source/Runtime/RenderCore/Public/RenderCore.h
+++ b/Engine/Source/Runtime/RenderCore/Public/RenderCore.h
@@ -59,6 +59,7 @@ DECLARE_CYCLE_STAT_EXTERN(TEXT("Compute View Relevance"),STAT_ComputeViewRelevan
 DECLARE_CYCLE_STAT_EXTERN(TEXT("Static Mesh Relevance"),STAT_StaticRelevance,STATGROUP_InitViews, RENDERCORE_API);
 DECLARE_CYCLE_STAT_EXTERN(TEXT("UpdateStaticMeshes"),STAT_UpdateStaticMeshesTime,STATGROUP_InitViews, RENDERCORE_API);
 DECLARE_CYCLE_STAT_EXTERN(TEXT("GetDynamicMeshElements"),STAT_GetDynamicMeshElements,STATGROUP_InitViews, RENDERCORE_API);
+DECLARE_CYCLE_STAT_EXTERN(TEXT("ProcessVisibleMeshDrawCommands"), STAT_ProcessVisibleMeshDrawCommands, STATGROUP_InitViews, RENDERCORE_API);
 DECLARE_CYCLE_STAT_EXTERN(TEXT("Init dynamic shadows"),STAT_InitDynamicShadowsTime,STATGROUP_InitViews, RENDERCORE_API);
 DECLARE_CYCLE_STAT_EXTERN(TEXT("Update Preshadow Cache"),STAT_UpdatePreshadowCache,STATGROUP_InitViews, RENDERCORE_API);
 DECLARE_CYCLE_STAT_EXTERN(TEXT("Init Projected Shadow"),STAT_InitProjectedShadowVisibility,STATGROUP_InitViews, RENDERCORE_API);
diff --git a/Engine/Source/Runtime/Renderer/Private/SceneRendering.cpp b/Engine/Source/Runtime/Renderer/Private/SceneRendering.cpp
--- a/Engine/Source/Runtime/Renderer/Private/SceneRendering.cpp
+++ b/Engine/Source/Runtime/Renderer/Private/SceneRendering.cpp
@@ -3104,7 +3104,7 @@ private:
 
 void FSceneRenderer::SortMeshDrawCommands()
 {
-	QUICK_SCOPE_CYCLE_COUNTER(STAT_ProcessVisibleMeshDrawCommands);
+	SCOPE_CYCLE_COUNTER(STAT_ProcessVisibleMeshDrawCommands);
 
 	const EShadingPath ShadingPath = Scene->GetShadingPath();
 	const bool bUseGPUScene = UseGPUScene(GMaxRHIShaderPlatform, FeatureLevel);
