Fixed implementation of IsRayTracingTierSupported()

diff --git a/Engine/Source/Runtime/RenderCore/Public/RenderResource.h b/Engine/Source/Runtime/RenderCore/Public/RenderResource.h
--- a/Engine/Source/Runtime/RenderCore/Public/RenderResource.h
+++ b/Engine/Source/Runtime/RenderCore/Public/RenderResource.h
@@ -520,8 +520,8 @@ FORCEINLINE bool ShouldCompileRayTracingShadersForProject(EShaderPlatform Shader
 // This function is a runtime only function!
 FORCEINLINE bool IsRayTracingTierSupported(int32 RequiredTier)
 {
-	ensure(RequiredTier >= 0);
-	return GRHIRayTracingSupportTier > RequiredTier;
+	ensure(RequiredTier >= 1);
+	return GRHIRayTracingSupportTier >= RequiredTier;
 }
 
 #if RHI_RAYTRACING
