Replaced IsRayTracingTierSupported() with IsRayTracingEnabled(), which now also checks if RT support is enabled for the project. Removed ray tracing tiers (making support binary).

diff --git a/Engine/Source/Editor/LevelEditor/Private/LevelEditorToolBar.cpp b/Engine/Source/Editor/LevelEditor/Private/LevelEditorToolBar.cpp
--- a/Engine/Source/Editor/LevelEditor/Private/LevelEditorToolBar.cpp
+++ b/Engine/Source/Editor/LevelEditor/Private/LevelEditorToolBar.cpp
@@ -1715,7 +1715,7 @@ static void MakeShaderModelPreviewMenu(FMenuBuilder& MenuBuilder)
 		for (int32 i = GMaxRHIFeatureLevel; i >= 0; --i)
 		{
 #if RHI_RAYTRACING
-			if (IsRayTracingTierSupported(2))
+			if (IsRayTracingEnabled())
 			{
 				if (i != ERHIFeatureLevel::SM5)
 				{
diff --git a/Engine/Source/Editor/UnrealEd/Private/EditorEngine.cpp b/Engine/Source/Editor/UnrealEd/Private/EditorEngine.cpp
--- a/Engine/Source/Editor/UnrealEd/Private/EditorEngine.cpp
+++ b/Engine/Source/Editor/UnrealEd/Private/EditorEngine.cpp
@@ -7478,7 +7478,7 @@ void UEditorEngine::AllMaterialsCacheResourceShadersForRendering()
 void UEditorEngine::SetPreviewPlatform(const FName MaterialQualityPlatform, ERHIFeatureLevel::Type PreviewFeatureLevel, const bool bSaveSettings/* = true*/)
 {
 #if RHI_RAYTRACING
-	if (IsRayTracingTierSupported(2))
+	if (IsRayTracingEnabled())
 	{
 		if (PreviewFeatureLevel != ERHIFeatureLevel::SM5)
 		{
diff --git a/Engine/Source/Editor/UnrealEd/Private/SEditorViewport.cpp b/Engine/Source/Editor/UnrealEd/Private/SEditorViewport.cpp
--- a/Engine/Source/Editor/UnrealEd/Private/SEditorViewport.cpp
+++ b/Engine/Source/Editor/UnrealEd/Private/SEditorViewport.cpp
@@ -359,7 +359,7 @@ void SEditorViewport::BindCommands()
 	MAP_VIEWMODE_ACTION( Commands.UnlitMode, VMI_Unlit );
 	MAP_VIEWMODE_ACTION( Commands.LitMode, VMI_Lit );
 #if RHI_RAYTRACING
-	if (IsRayTracingTierSupported(2))
+	if (IsRayTracingEnabled())
 	{
 		MAP_VIEWMODE_ACTION(Commands.PathTracingMode, VMI_PathTracing);
 		MAP_VIEWMODE_ACTION(Commands.RayTracingDebugMode, VMI_RayTracingDebug);
diff --git a/Engine/Source/Editor/UnrealEd/Private/SEditorViewportViewMenu.cpp b/Engine/Source/Editor/UnrealEd/Private/SEditorViewportViewMenu.cpp
--- a/Engine/Source/Editor/UnrealEd/Private/SEditorViewportViewMenu.cpp
+++ b/Engine/Source/Editor/UnrealEd/Private/SEditorViewportViewMenu.cpp
@@ -302,7 +302,7 @@ TSharedRef<SWidget> SEditorViewportViewMenu::GenerateViewMenuContent() const
 			}
 
 #if RHI_RAYTRACING
-			if (IsRayTracingTierSupported(2))
+			if (IsRayTracingEnabled())
 			{
 				struct Local
 				{
diff --git a/Engine/Source/Runtime/D3D12RHI/Private/Windows/WindowsD3D12Device.cpp b/Engine/Source/Runtime/D3D12RHI/Private/Windows/WindowsD3D12Device.cpp
--- a/Engine/Source/Runtime/D3D12RHI/Private/Windows/WindowsD3D12Device.cpp
+++ b/Engine/Source/Runtime/D3D12RHI/Private/Windows/WindowsD3D12Device.cpp
@@ -664,7 +664,7 @@ void FD3D12DynamicRHI::Init()
 	GRHIForceNoDeletionLatencyForStreamingTextures = !!PLATFORM_WINDOWS;
 
 #if D3D12_RHI_RAYTRACING
-	GRHIRayTracingSupportTier = GetAdapter().GetD3DRayTracingDevice() ? 2 : 0;
+	GRHISupportsRayTracing = GetAdapter().GetD3DRayTracingDevice() != nullptr;
 #endif
 
 	// Set the RHI initialized flag.
@@ -673,7 +673,7 @@ void FD3D12DynamicRHI::Init()
 
 void FD3D12DynamicRHI::PostInit()
 {
-	if (IsRayTracingTierSupported(2))
+	if (GRHISupportsRayTracing)
 	{
 		for (FD3D12Adapter*& Adapter : ChosenAdapters)
 		{
diff --git a/Engine/Source/Runtime/Engine/Private/GameViewportClient.cpp b/Engine/Source/Runtime/Engine/Private/GameViewportClient.cpp
--- a/Engine/Source/Runtime/Engine/Private/GameViewportClient.cpp
+++ b/Engine/Source/Runtime/Engine/Private/GameViewportClient.cpp
@@ -3020,7 +3020,7 @@ bool UGameViewportClient::HandleViewModeCommand( const TCHAR* Cmd, FOutputDevice
 	}
 
 #if RHI_RAYTRACING
-	if (GRHIRayTracingSupportTier == 0)
+	if (!GRHISupportsRayTracing)
 	{
 		if (ViewModeIndex == VMI_PathTracing)
 		{
diff --git a/Engine/Source/Runtime/Engine/Private/SceneView.cpp b/Engine/Source/Runtime/Engine/Private/SceneView.cpp
--- a/Engine/Source/Runtime/Engine/Private/SceneView.cpp
+++ b/Engine/Source/Runtime/Engine/Private/SceneView.cpp
@@ -2428,7 +2428,7 @@ void FSceneView::SetupRayTracedRendering()
 {
 	RayTracingRenderMode = ERayTracingRenderMode::Disabled;
 
-	if (!IsRayTracingTierSupported(2))
+	if (!IsRayTracingEnabled())
 	{
 		return;
 	}
diff --git a/Engine/Source/Runtime/Engine/Private/ShaderCompiler/ShaderCompiler.cpp b/Engine/Source/Runtime/Engine/Private/ShaderCompiler/ShaderCompiler.cpp
--- a/Engine/Source/Runtime/Engine/Private/ShaderCompiler/ShaderCompiler.cpp
+++ b/Engine/Source/Runtime/Engine/Private/ShaderCompiler/ShaderCompiler.cpp
@@ -3893,7 +3893,7 @@ void VerifyGlobalShaders(EShaderPlatform Platform, bool bLoadedFromCacheFile)
 			!IsOpenGLPlatform(GMaxRHIShaderPlatform) && !IsVulkanPlatform(GMaxRHIShaderPlatform) &&
 			!IsMetalPlatform(GMaxRHIShaderPlatform) && !IsSwitchPlatform(GMaxRHIShaderPlatform) &&
 #if RHI_RAYTRACING
-			!IsRayTracingTierSupported(2) && //This is here because DXR is caching its BuiltIn Tier1 Shaders in PostInit (see FD3D12Device::InitRayTracing)
+			!GRHISupportsRayTracing && //This is here because DXR is caching its BuiltIn Shaders in PostInit (see FD3D12Device::InitRayTracing)
 #endif // RHI_RAYTRACING
 			GShaderCompilingManager->AllowAsynchronousShaderCompiling();
 
diff --git a/Engine/Source/Runtime/Engine/Private/SkeletalRenderCPUSkin.cpp b/Engine/Source/Runtime/Engine/Private/SkeletalRenderCPUSkin.cpp
--- a/Engine/Source/Runtime/Engine/Private/SkeletalRenderCPUSkin.cpp
+++ b/Engine/Source/Runtime/Engine/Private/SkeletalRenderCPUSkin.cpp
@@ -333,7 +333,7 @@ void FSkeletalMeshObjectCPUSkin::FSkeletalMeshObjectLOD::InitResources(FSkelMesh
 		});
 
 #if RHI_RAYTRACING
-	if (IsRayTracingTierSupported(2))
+	if (IsRayTracingEnabled())
 	{
 		check(SkelMeshRenderData);
 		check(SkelMeshRenderData->LODRenderData.IsValidIndex(LODIndex));
diff --git a/Engine/Source/Runtime/Engine/Private/SkeletalRenderGPUSkin.cpp b/Engine/Source/Runtime/Engine/Private/SkeletalRenderGPUSkin.cpp
--- a/Engine/Source/Runtime/Engine/Private/SkeletalRenderGPUSkin.cpp
+++ b/Engine/Source/Runtime/Engine/Private/SkeletalRenderGPUSkin.cpp
@@ -326,7 +326,7 @@ void FSkeletalMeshObjectGPUSkin::UpdateDynamicData_RenderThread(FGPUSkinCache* G
 	}
 
 #if RHI_RAYTRACING
-	if (IsRayTracingTierSupported(2))
+	if (IsRayTracingEnabled())
 	{
 		if (GEnableGPUSkinCache && SkinCacheEntry)
 		{
diff --git a/Engine/Source/Runtime/Engine/Private/SkeletalRenderStatic.cpp b/Engine/Source/Runtime/Engine/Private/SkeletalRenderStatic.cpp
--- a/Engine/Source/Runtime/Engine/Private/SkeletalRenderStatic.cpp
+++ b/Engine/Source/Runtime/Engine/Private/SkeletalRenderStatic.cpp
@@ -122,7 +122,7 @@ void FSkeletalMeshObjectStatic::FSkeletalMeshObjectLOD::InitResources(FSkelMeshC
 
 
 #if RHI_RAYTRACING
-	if (IsRayTracingTierSupported(2))
+	if (IsRayTracingEnabled())
 	{
 		check(SkelMeshRenderData);
 		check(SkelMeshRenderData->LODRenderData.IsValidIndex(LODIndex));
diff --git a/Engine/Source/Runtime/Engine/Private/StaticMesh.cpp b/Engine/Source/Runtime/Engine/Private/StaticMesh.cpp
--- a/Engine/Source/Runtime/Engine/Private/StaticMesh.cpp
+++ b/Engine/Source/Runtime/Engine/Private/StaticMesh.cpp
@@ -683,7 +683,7 @@ void FStaticMeshLODResources::InitResources(UStaticMesh* Parent)
 	}
 
 #if RHI_RAYTRACING
-	if (IsRayTracingTierSupported(2))
+	if (IsRayTracingEnabled())
 	{
 		ENQUEUE_RENDER_COMMAND(InitStaticMeshRayTracingGeometry)(
 			[this](FRHICommandListImmediate& RHICmdList)
diff --git a/Engine/Source/Runtime/RHI/Private/RHI.cpp b/Engine/Source/Runtime/RHI/Private/RHI.cpp
--- a/Engine/Source/Runtime/RHI/Private/RHI.cpp
+++ b/Engine/Source/Runtime/RHI/Private/RHI.cpp
@@ -498,7 +498,7 @@ bool GRHISupportsBaseVertexIndex = true;
 TRHIGlobal<bool> GRHISupportsInstancing(true);
 bool GRHISupportsFirstInstance = false;
 bool GRHISupportsDynamicResolution = false;
-int32 GRHIRayTracingSupportTier = 0;
+bool GRHISupportsRayTracing = false;
 bool GRHISupportsRHIThread = false;
 bool GRHISupportsRHIOnTaskThread = false;
 bool GRHISupportsParallelRHIExecute = false;
diff --git a/Engine/Source/Runtime/RHI/Public/RHI.h b/Engine/Source/Runtime/RHI/Public/RHI.h
--- a/Engine/Source/Runtime/RHI/Public/RHI.h
+++ b/Engine/Source/Runtime/RHI/Public/RHI.h
@@ -195,7 +195,7 @@ inline bool RHISupportsAbsoluteVertexID(EShaderPlatform InShaderPlatform)
 }
 
 /** Can this platform compile ray tracing shaders (regardless of project settings).
- *  To use at runtime, also check GRHIRayTracingSupportTier.
+ *  To use at runtime, also check GRHISupportsRayTracing and r.RayTracing CVar (see IsRayTracingEnabled() helper).
  **/
 inline RHI_API bool RHISupportsRayTracingShaders(EShaderPlatform Platform)
 {
@@ -466,12 +466,8 @@ extern RHI_API bool GRHISupportsFirstInstance;
 /** Whether or not the RHI can handle dynamic resolution or not. */
 extern RHI_API bool GRHISupportsDynamicResolution;
 
-/** Ray tracing tier level:
- *   0: Ray tracing is not supported.
- *   1: Basic support for acceleration structure construction and ray queries. Ray tracing shader types are not supported.
- *   2: Full ray tracing support, including ray tracing shader types, shader binding tables, etc.
- **/
-extern RHI_API int32 GRHIRayTracingSupportTier;
+/** Whether or not the RHI supports ray tracing on current hardware (acceleration structure building and new ray tracing-specific shader types). */
+extern RHI_API bool GRHISupportsRayTracing;
 
 /** Whether or not the RHI supports an RHI thread.
 Requirements for RHI thread
diff --git a/Engine/Source/Runtime/RenderCore/Public/RenderResource.h b/Engine/Source/Runtime/RenderCore/Public/RenderResource.h
--- a/Engine/Source/Runtime/RenderCore/Public/RenderResource.h
+++ b/Engine/Source/Runtime/RenderCore/Public/RenderResource.h
@@ -517,11 +517,20 @@ FORCEINLINE bool ShouldCompileRayTracingShadersForProject(EShaderPlatform Shader
 	return false;
 }
 
+// Returns `true` when running on RT-capable machine and RT support is enabled for the project.
 // This function is a runtime only function!
-FORCEINLINE bool IsRayTracingTierSupported(int32 RequiredTier)
+FORCEINLINE bool IsRayTracingEnabled()
 {
-	ensure(RequiredTier >= 1);
-	return GRHIRayTracingSupportTier >= RequiredTier;
+	if (GRHISupportsRayTracing)
+	{
+		// r.RayTracing is a read-only CVar. UE needs to be restarted to effectively change it
+		static const bool bRayTracingEnabled = IConsoleManager::Get().FindConsoleVariable(TEXT("r.RayTracing"))->GetInt() > 0;
+		return bRayTracingEnabled;
+	}
+	else
+	{
+		return false;
+	}
 }
 
 #if RHI_RAYTRACING
@@ -555,7 +564,7 @@ public:
 
 	virtual void InitRHI() override
 	{
-		if (Initializer.IndexBuffer && Initializer.PositionVertexBuffer && IsRayTracingTierSupported(2))
+		if (Initializer.IndexBuffer && Initializer.PositionVertexBuffer && IsRayTracingEnabled())
 		{
 			RayTracingGeometryRHI = RHICreateRayTracingGeometry(Initializer);
 			FRHICommandListExecutor::GetImmediateCommandList().BuildAccelerationStructure(RayTracingGeometryRHI);
diff --git a/Engine/Source/Runtime/Renderer/Private/CompositionLighting/CompositionLighting.cpp b/Engine/Source/Runtime/Renderer/Private/CompositionLighting/CompositionLighting.cpp
--- a/Engine/Source/Runtime/Renderer/Private/CompositionLighting/CompositionLighting.cpp
+++ b/Engine/Source/Runtime/Renderer/Private/CompositionLighting/CompositionLighting.cpp
@@ -125,7 +125,7 @@ bool ShouldRenderScreenSpaceAmbientOcclusion(const FViewInfo& View)
 
 	// #dxr_todo: hard-coded FindConsoleVariable is a bad hack. The current variable is hidden behind DeferredShadingRenderer
 	static IConsoleVariable* UseRayTracingAmbientOcclusionCVar = IConsoleManager::Get().FindConsoleVariable(TEXT("r.RayTracing.AmbientOcclusion"));
-	bool bUseRayTracingAmbientOcclusion = IsRayTracingTierSupported(2) && UseRayTracingAmbientOcclusionCVar && UseRayTracingAmbientOcclusionCVar->GetInt() > 0;
+	bool bUseRayTracingAmbientOcclusion = IsRayTracingEnabled() && UseRayTracingAmbientOcclusionCVar && UseRayTracingAmbientOcclusionCVar->GetInt() > 0;
 	return bEnabled && !bUseRayTracingAmbientOcclusion;
 }
 
diff --git a/Engine/Source/Runtime/Renderer/Private/DeferredShadingRenderer.cpp b/Engine/Source/Runtime/Renderer/Private/DeferredShadingRenderer.cpp
--- a/Engine/Source/Runtime/Renderer/Private/DeferredShadingRenderer.cpp
+++ b/Engine/Source/Runtime/Renderer/Private/DeferredShadingRenderer.cpp
@@ -581,7 +581,7 @@ void FDeferredShadingSceneRenderer::PrepareDistanceFieldScene(FRHICommandListImm
 
 bool FDeferredShadingSceneRenderer::GatherRayTracingWorldInstances(FRHICommandListImmediate& RHICmdList)
 {
-	if (!IsRayTracingTierSupported(2))
+	if (!IsRayTracingEnabled())
 	{
 		return false;
 	}
@@ -874,7 +874,7 @@ bool FDeferredShadingSceneRenderer::GatherRayTracingWorldInstances(FRHICommandLi
 
 bool FDeferredShadingSceneRenderer::DispatchRayTracingWorldUpdates(FRHICommandListImmediate& RHICmdList)
 {
-	if (!IsRayTracingTierSupported(2))
+	if (!IsRayTracingEnabled())
 	{
 		return false;
 	}
@@ -1685,7 +1685,7 @@ void FDeferredShadingSceneRenderer::Render(FRHICommandListImmediate& RHICmdList)
 
 	checkSlow(RHICmdList.IsOutsideRenderPass());
 #if RHI_RAYTRACING
-	if (IsRayTracingTierSupported(2))
+	if (IsRayTracingEnabled())
 	{
 		for (int32 ViewIndex = 0; ViewIndex < Views.Num(); ViewIndex++)
 		{
diff --git a/Engine/Source/Runtime/Renderer/Private/LightRendering.cpp b/Engine/Source/Runtime/Renderer/Private/LightRendering.cpp
--- a/Engine/Source/Runtime/Renderer/Private/LightRendering.cpp
+++ b/Engine/Source/Runtime/Renderer/Private/LightRendering.cpp
@@ -764,7 +764,7 @@ void FDeferredShadingSceneRenderer::RenderLights(FRHICommandListImmediate& RHICm
 						RenderRayTracingOcclusionForRectLight(RHICmdList, LightSceneInfo, ScreenShadowMaskTexture);
 					}
 					// #dxr_todo: ShouldRenderRayTracingOcclusion()
-					else if (IsRayTracingTierSupported(2) && GRayTracingOcclusion == 1)
+					else if (IsRayTracingEnabled() && GRayTracingOcclusion == 1)
 					{
 						RenderRayTracingOcclusion(RHICmdList, &LightSceneInfo, ScreenShadowMaskTexture);
 					}
diff --git a/Engine/Source/Runtime/Renderer/Private/PrimitiveSceneInfo.cpp b/Engine/Source/Runtime/Renderer/Private/PrimitiveSceneInfo.cpp
--- a/Engine/Source/Runtime/Renderer/Private/PrimitiveSceneInfo.cpp
+++ b/Engine/Source/Runtime/Renderer/Private/PrimitiveSceneInfo.cpp
@@ -326,7 +326,7 @@ void FPrimitiveSceneInfo::AddStaticMeshes(FRHICommandListImmediate& RHICmdList,
 	CacheMeshDrawCommands(RHICmdList);
 
 #if RHI_RAYTRACING
-	if (IsRayTracingTierSupported(1))
+	if (IsRayTracingEnabled())
 	{
 		int MaxLOD = -1;
 
diff --git a/Engine/Source/Runtime/Renderer/Private/RayTracing/RayTracingMaterialHitShaders.cpp b/Engine/Source/Runtime/Renderer/Private/RayTracing/RayTracingMaterialHitShaders.cpp
--- a/Engine/Source/Runtime/Renderer/Private/RayTracing/RayTracingMaterialHitShaders.cpp
+++ b/Engine/Source/Runtime/Renderer/Private/RayTracing/RayTracingMaterialHitShaders.cpp
@@ -127,7 +127,7 @@ IMPLEMENT_MATERIALCHS_TYPE(TUniformLightMapPolicy<LMP_DISTANCE_FIELD_SHADOWS_AND
 
 static FMeshPassProcessor* CreateRayTracingProcessor(const FScene* Scene, const FSceneView* InViewIfDynamicMeshCommand, FMeshPassDrawListContext& InDrawListContext)
 {
-	return IsRayTracingTierSupported(2) ? new(FMemStack::Get()) FRayTracingMeshProcessor(Scene, Scene->GetFeatureLevel(), InViewIfDynamicMeshCommand, InDrawListContext) : nullptr;
+	return IsRayTracingEnabled() ? new(FMemStack::Get()) FRayTracingMeshProcessor(Scene, Scene->GetFeatureLevel(), InViewIfDynamicMeshCommand, InDrawListContext) : nullptr;
 }
 
 FRegisterPassProcessorCreateFunction RegisterRayTracing(&CreateRayTracingProcessor, EShadingPath::Deferred, EMeshPass::RayTracing, EMeshPassFlags::CachedMeshCommands);
diff --git a/Engine/Source/Runtime/Renderer/Private/RayTracing/RayTracingShadowsForRectLight.cpp b/Engine/Source/Runtime/Renderer/Private/RayTracing/RayTracingShadowsForRectLight.cpp
--- a/Engine/Source/Runtime/Renderer/Private/RayTracing/RayTracingShadowsForRectLight.cpp
+++ b/Engine/Source/Runtime/Renderer/Private/RayTracing/RayTracingShadowsForRectLight.cpp
@@ -39,7 +39,7 @@ static FAutoConsoleVariableRef CVarRayTracingRecLightIsTextureImportanceSampling
 
 bool IsRayTracingRectLightSelected()
 {
-	return IsRayTracingTierSupported(2) && GRayTracingRectLight == 1;
+	return IsRayTracingEnabled() && GRayTracingRectLight == 1;
 }
 
 bool ShouldRenderRayTracingStaticOrStationaryRectLight(const FLightSceneInfo& LightSceneInfo)
diff --git a/Engine/Source/Runtime/Renderer/Private/RayTracing/RayTracingTestbed.cpp b/Engine/Source/Runtime/Renderer/Private/RayTracing/RayTracingTestbed.cpp
--- a/Engine/Source/Runtime/Renderer/Private/RayTracing/RayTracingTestbed.cpp
+++ b/Engine/Source/Runtime/Renderer/Private/RayTracing/RayTracingTestbed.cpp
@@ -10,7 +10,7 @@
 
 void TestBasicRayTracing(bool bValidateResults)
 {
-	if (GRHIRayTracingSupportTier == 0)
+	if (!GRHISupportsRayTracing)
 	{
 		return;
 	}
diff --git a/Engine/Source/Runtime/Renderer/Private/RayTracing/RaytracingSkylight.cpp b/Engine/Source/Runtime/Renderer/Private/RayTracing/RaytracingSkylight.cpp
--- a/Engine/Source/Runtime/Renderer/Private/RayTracing/RaytracingSkylight.cpp
+++ b/Engine/Source/Runtime/Renderer/Private/RayTracing/RaytracingSkylight.cpp
@@ -10,7 +10,7 @@ static int32 GRayTracingSkyLight = 0;
 bool IsRayTracingSkyLightSelected()
 {
 #if RHI_RAYTRACING
-	return IsRayTracingTierSupported(2) && GRayTracingSkyLight > 0;
+	return IsRayTracingEnabled() && GRayTracingSkyLight > 0;
 #else
 	return false;
 #endif
diff --git a/Engine/Source/Runtime/Renderer/Private/ReflectionEnvironment.cpp b/Engine/Source/Runtime/Renderer/Private/ReflectionEnvironment.cpp
--- a/Engine/Source/Runtime/Renderer/Private/ReflectionEnvironment.cpp
+++ b/Engine/Source/Runtime/Renderer/Private/ReflectionEnvironment.cpp
@@ -688,7 +688,7 @@ void FDeferredShadingSceneRenderer::RenderDeferredReflectionsAndSkyLighting(FRHI
 		bReflectionCapture = bReflectionCapture || View.bIsReflectionCapture;
 	}
 
-	bool bRayTracedReflections = IsRayTracingTierSupported(2) && GRayTracingReflections == 1;
+	bool bRayTracedReflections = IsRayTracingEnabled() && GRayTracingReflections == 1;
 
 	const bool bSkyLight = Scene->SkyLight
 		&& Scene->SkyLight->ProcessedTexture
diff --git a/Engine/Source/Runtime/Renderer/Private/RendererScene.cpp b/Engine/Source/Runtime/Renderer/Private/RendererScene.cpp
--- a/Engine/Source/Runtime/Renderer/Private/RendererScene.cpp
+++ b/Engine/Source/Runtime/Renderer/Private/RendererScene.cpp
@@ -891,7 +891,7 @@ FScene::FScene(UWorld* InWorld, bool bInRequiresHitProxies, bool bInIsEditorScen
 	}
 
 #if RHI_RAYTRACING
-	if (IsRayTracingTierSupported(2))
+	if (IsRayTracingEnabled())
 	{
 		RayTracingDynamicGeometryCollection = new FRayTracingDynamicGeometryCollection();
 	}
