s/UBMT_GRAPH_TRACKED_/UBMT_RDG_ for better naming consistency.

diff --git a/Engine/Source/Developer/ShaderCompilerCommon/Private/ShaderCompilerCommon.cpp b/Engine/Source/Developer/ShaderCompilerCommon/Private/ShaderCompilerCommon.cpp
--- a/Engine/Source/Developer/ShaderCompilerCommon/Private/ShaderCompilerCommon.cpp
+++ b/Engine/Source/Developer/ShaderCompilerCommon/Private/ShaderCompilerCommon.cpp
@@ -131,19 +131,19 @@ bool BuildResourceTableMapping(
 			switch( Entry.Type )
 			{
 			case UBMT_TEXTURE:
-			case UBMT_GRAPH_TRACKED_TEXTURE:
+			case UBMT_RDG_TEXTURE:
 				OutSRT.TextureMap.Add(ResourceMap);
 				break;
 			case UBMT_SAMPLER:
 				OutSRT.SamplerMap.Add(ResourceMap);
 				break;
 			case UBMT_SRV:
-			case UBMT_GRAPH_TRACKED_SRV:
-			case UBMT_GRAPH_TRACKED_BUFFER_SRV:
+			case UBMT_RDG_TEXTURE_SRV:
+			case UBMT_RDG_BUFFER_SRV:
 				OutSRT.ShaderResourceViewMap.Add(ResourceMap);
 				break;
-			case UBMT_GRAPH_TRACKED_UAV:
-			case UBMT_GRAPH_TRACKED_BUFFER_UAV:
+			case UBMT_RDG_TEXTURE_UAV:
+			case UBMT_RDG_BUFFER_UAV:
 				OutSRT.UnorderedAccessViewMap.Add(ResourceMap);
 				break;
 			default:
diff --git a/Engine/Source/Runtime/Apple/MetalRHI/Private/MetalUniformBuffer.cpp b/Engine/Source/Runtime/Apple/MetalRHI/Private/MetalUniformBuffer.cpp
--- a/Engine/Source/Runtime/Apple/MetalRHI/Private/MetalUniformBuffer.cpp
+++ b/Engine/Source/Runtime/Apple/MetalRHI/Private/MetalUniformBuffer.cpp
@@ -247,8 +247,8 @@ void FMetalUniformBuffer::InitIAB()
 			
 			switch(GetLayout().Resources[i].MemberType)
 			{
-				case UBMT_GRAPH_TRACKED_SRV:
-				case UBMT_GRAPH_TRACKED_BUFFER_SRV:
+				case UBMT_RDG_TEXTURE_SRV:
+				case UBMT_RDG_BUFFER_SRV:
 				case UBMT_SRV:
 				{
 					FMetalArgumentDesc& Desc = Arguments.Emplace_GetRef();
@@ -311,8 +311,8 @@ void FMetalUniformBuffer::InitIAB()
 					}
 					break;
 				}
-				case UBMT_GRAPH_TRACKED_UAV:
-				case UBMT_GRAPH_TRACKED_BUFFER_UAV:
+				case UBMT_RDG_TEXTURE_UAV:
+				case UBMT_RDG_BUFFER_UAV:
 				{
 					FMetalArgumentDesc& Desc = Arguments.Emplace_GetRef();
 					Desc.SetIndex(Index);
@@ -390,7 +390,7 @@ void FMetalUniformBuffer::InitIAB()
 					NewIAB->IndirectArgumentResources.Add(Argument(Sampler->State));
 					break;
 				}
-				case UBMT_GRAPH_TRACKED_TEXTURE:
+				case UBMT_RDG_TEXTURE:
 				case UBMT_TEXTURE:
 				{
 					FMetalArgumentDesc& Desc = Arguments.Emplace_GetRef();
@@ -534,8 +534,8 @@ void FMetalUniformBuffer::Update(const void* Contents, EUniformBufferValidation
             {
                 switch(Layout.Resources[i].MemberType)
                 {
-                    case UBMT_GRAPH_TRACKED_SRV:
-                    case UBMT_GRAPH_TRACKED_BUFFER_SRV:
+                    case UBMT_RDG_TEXTURE_SRV:
+                    case UBMT_RDG_BUFFER_SRV:
                     case UBMT_SRV:
                     {
                         FMetalShaderResourceView* SRV = (FMetalShaderResourceView*)Resource;
@@ -546,8 +546,8 @@ void FMetalUniformBuffer::Update(const void* Contents, EUniformBufferValidation
                         }
                         break;
                     }
-                    case UBMT_GRAPH_TRACKED_UAV:
-                    case UBMT_GRAPH_TRACKED_BUFFER_UAV:
+                    case UBMT_RDG_TEXTURE_UAV:
+                    case UBMT_RDG_BUFFER_UAV:
                     {
                         FMetalUnorderedAccessView* UAV = (FMetalUnorderedAccessView*)Resource;
                         FRHITexture* Texture = UAV->SourceView->SourceTexture.GetReference();
@@ -557,7 +557,7 @@ void FMetalUniformBuffer::Update(const void* Contents, EUniformBufferValidation
                         }
                         break;
                     }
-                    case UBMT_GRAPH_TRACKED_TEXTURE:
+                    case UBMT_RDG_TEXTURE:
                     case UBMT_TEXTURE:
                     {
                         FRHITexture* Texture = (FRHITexture*)Resource;
diff --git a/Engine/Source/Runtime/D3D12RHI/Private/D3D12UniformBuffer.cpp b/Engine/Source/Runtime/D3D12RHI/Private/D3D12UniformBuffer.cpp
--- a/Engine/Source/Runtime/D3D12RHI/Private/D3D12UniformBuffer.cpp
+++ b/Engine/Source/Runtime/D3D12RHI/Private/D3D12UniformBuffer.cpp
@@ -101,7 +101,7 @@ FUniformBufferRHIRef FD3D12DynamicRHI::RHICreateUniformBuffer(const void* Conten
 				}
 
 				// Allow null SRV's in uniform buffers for feature levels that don't support SRV's in shaders
-				if (!(GMaxRHIFeatureLevel <= ERHIFeatureLevel::ES3_1 && (ResourceType == UBMT_SRV || ResourceType == UBMT_GRAPH_TRACKED_SRV)) && Validation == EUniformBufferValidation::ValidateResources)
+				if (!(GMaxRHIFeatureLevel <= ERHIFeatureLevel::ES3_1 && (ResourceType == UBMT_SRV || ResourceType == UBMT_RDG_TEXTURE_SRV)) && Validation == EUniformBufferValidation::ValidateResources)
 				{
 					check(Resource);
 				}
diff --git a/Engine/Source/Runtime/Engine/Private/ShaderCompiler/ShaderCompiler.cpp b/Engine/Source/Runtime/Engine/Private/ShaderCompiler/ShaderCompiler.cpp
--- a/Engine/Source/Runtime/Engine/Private/ShaderCompiler/ShaderCompiler.cpp
+++ b/Engine/Source/Runtime/Engine/Private/ShaderCompiler/ShaderCompiler.cpp
@@ -2783,8 +2783,8 @@ static void PullRootShaderParametersLayout(FShaderCompilerInput& CompileInput, c
 			// RHI don't need to care about render target bindings slot anyway.
 		}
 		else if (
-			BaseType == UBMT_GRAPH_TRACKED_BUFFER_UAV ||
-			BaseType == UBMT_GRAPH_TRACKED_UAV)
+			BaseType == UBMT_RDG_BUFFER_UAV ||
+			BaseType == UBMT_RDG_TEXTURE_UAV)
 		{
 			// UAV are ignored on purpose because not supported in uniform buffers.
 		}
diff --git a/Engine/Source/Runtime/RHI/Public/RHIDefinitions.h b/Engine/Source/Runtime/RHI/Public/RHIDefinitions.h
--- a/Engine/Source/Runtime/RHI/Public/RHIDefinitions.h
+++ b/Engine/Source/Runtime/RHI/Public/RHIDefinitions.h
@@ -421,18 +421,18 @@ enum EUniformBufferBaseType : uint8
 	UBMT_UINT32,
 	UBMT_FLOAT32,
 
-	// Untracked resources.
+	// RHI resources not tracked by render graph.
 	UBMT_TEXTURE,
 	UBMT_SRV,
 	UBMT_SAMPLER,
 
-	// Render graph tracked resources.
-	UBMT_GRAPH_TRACKED_TEXTURE,
-	UBMT_GRAPH_TRACKED_SRV,
-	UBMT_GRAPH_TRACKED_UAV,
-	UBMT_GRAPH_TRACKED_BUFFER,
-	UBMT_GRAPH_TRACKED_BUFFER_SRV,
-	UBMT_GRAPH_TRACKED_BUFFER_UAV,
+	// Resources tracked by render graph.
+	UBMT_RDG_TEXTURE,
+	UBMT_RDG_TEXTURE_SRV,
+	UBMT_RDG_TEXTURE_UAV,
+	UBMT_RDG_BUFFER,
+	UBMT_RDG_BUFFER_SRV,
+	UBMT_RDG_BUFFER_UAV,
 
 	// Nested structure.
 	UBMT_NESTED_STRUCT,
@@ -1150,12 +1150,12 @@ inline int32 GetFeatureLevelMaxNumberOfBones(ERHIFeatureLevel::Type FeatureLevel
 inline bool IsRDGResourceReferenceShaderParameterType(EUniformBufferBaseType BaseType)
 {
 	return
-		BaseType == UBMT_GRAPH_TRACKED_TEXTURE ||
-		BaseType == UBMT_GRAPH_TRACKED_SRV ||
-		BaseType == UBMT_GRAPH_TRACKED_UAV ||
-		BaseType == UBMT_GRAPH_TRACKED_BUFFER ||
-		BaseType == UBMT_GRAPH_TRACKED_BUFFER_SRV ||
-		BaseType == UBMT_GRAPH_TRACKED_BUFFER_UAV;
+		BaseType == UBMT_RDG_TEXTURE ||
+		BaseType == UBMT_RDG_TEXTURE_SRV ||
+		BaseType == UBMT_RDG_TEXTURE_UAV ||
+		BaseType == UBMT_RDG_BUFFER ||
+		BaseType == UBMT_RDG_BUFFER_SRV ||
+		BaseType == UBMT_RDG_BUFFER_UAV;
 }
 
 /** Returns whether the shader parameter type needs to be passdown to RHI through FRHIUniformBufferLayout when creating an uniform buffer. */
diff --git a/Engine/Source/Runtime/RenderCore/Private/RenderGraphBuilder.cpp b/Engine/Source/Runtime/RenderCore/Private/RenderGraphBuilder.cpp
--- a/Engine/Source/Runtime/RenderCore/Private/RenderGraphBuilder.cpp
+++ b/Engine/Source/Runtime/RenderCore/Private/RenderGraphBuilder.cpp
@@ -51,7 +51,7 @@ static bool IsBoundAsReadable(const FRDGTexture* Texture, FShaderParameterStruct
 
 		switch (Type)
 		{
-		case UBMT_GRAPH_TRACKED_TEXTURE:
+		case UBMT_RDG_TEXTURE:
 		{
 			const FRDGTexture* InputTexture = *ParameterStruct.GetMemberPtrAtOffset<const FRDGTexture*>(Offset);
 			if (Texture == InputTexture)
@@ -60,7 +60,7 @@ static bool IsBoundAsReadable(const FRDGTexture* Texture, FShaderParameterStruct
 			}
 		}
 		break;
-		case UBMT_GRAPH_TRACKED_SRV:
+		case UBMT_RDG_TEXTURE_SRV:
 		{
 			const FRDGTextureSRV* InputSRV = *ParameterStruct.GetMemberPtrAtOffset<const FRDGTextureSRV*>(Offset);
 			if (InputSRV && Texture == InputSRV->Desc.Texture)
@@ -181,8 +181,8 @@ void FRDGBuilder::ValidatePass(const FRenderGraphPass* Pass) const
 
 		switch (Type)
 		{
-		case UBMT_GRAPH_TRACKED_UAV:
-		case UBMT_GRAPH_TRACKED_BUFFER_UAV:
+		case UBMT_RDG_TEXTURE_UAV:
+		case UBMT_RDG_BUFFER_UAV:
 		{
 			FRDGResource* RESTRICT UAV = *ParameterStruct.GetMemberPtrAtOffset<FRDGResource*>(Offset);
 			if (UAV && !bCanUseUAVs && GRenderGraphEmitWarnings)
@@ -251,7 +251,7 @@ void FRDGBuilder::CaptureAnyInterestingPassOutput(const FRenderGraphPass* Pass)
 
 		switch (Type)
 		{
-		case UBMT_GRAPH_TRACKED_UAV:
+		case UBMT_RDG_TEXTURE_UAV:
 		{
 			FRDGTextureUAV* RESTRICT UAV = *ParameterStruct.GetMemberPtrAtOffset<FRDGTextureUAV*>(Offset);
 			if (UAV && GVisualizeTexture.ShouldCapture(UAV->Desc.Texture->Name))
@@ -306,8 +306,8 @@ void FRDGBuilder::WalkGraphDependencies()
 
 			switch (Type)
 			{
-			case UBMT_GRAPH_TRACKED_TEXTURE:
-			case UBMT_GRAPH_TRACKED_BUFFER:
+			case UBMT_RDG_TEXTURE:
+			case UBMT_RDG_BUFFER:
 			{
 				FRDGResource* RESTRICT Resource = *ParameterStruct.GetMemberPtrAtOffset<FRDGResource*>(Offset);
 				if (Resource)
@@ -316,7 +316,7 @@ void FRDGBuilder::WalkGraphDependencies()
 				}
 			}
 			break;
-			case UBMT_GRAPH_TRACKED_SRV:
+			case UBMT_RDG_TEXTURE_SRV:
 			{
 				FRDGTextureSRV* RESTRICT SRV = *ParameterStruct.GetMemberPtrAtOffset<FRDGTextureSRV*>(Offset);
 				if (SRV)
@@ -325,7 +325,7 @@ void FRDGBuilder::WalkGraphDependencies()
 				}
 			}
 			break;
-			case UBMT_GRAPH_TRACKED_UAV:
+			case UBMT_RDG_TEXTURE_UAV:
 			{
 				FRDGTextureUAV* RESTRICT UAV = *ParameterStruct.GetMemberPtrAtOffset<FRDGTextureUAV*>(Offset);
 				if (UAV)
@@ -334,7 +334,7 @@ void FRDGBuilder::WalkGraphDependencies()
 				}
 			}
 			break;
-			case UBMT_GRAPH_TRACKED_BUFFER_SRV:
+			case UBMT_RDG_BUFFER_SRV:
 			{
 				FRDGBufferSRV* RESTRICT SRV = *ParameterStruct.GetMemberPtrAtOffset<FRDGBufferSRV*>(Offset);
 				if (SRV)
@@ -343,7 +343,7 @@ void FRDGBuilder::WalkGraphDependencies()
 				}
 			}
 			break;
-			case UBMT_GRAPH_TRACKED_BUFFER_UAV:
+			case UBMT_RDG_BUFFER_UAV:
 			{
 				FRDGBufferUAV* RESTRICT UAV = *ParameterStruct.GetMemberPtrAtOffset<FRDGBufferUAV*>(Offset);
 				if (UAV)
@@ -720,7 +720,7 @@ void FRDGBuilder::AllocateAndTransitionPassResources(const FRenderGraphPass* Pas
 
 		switch (Type)
 		{
-		case UBMT_GRAPH_TRACKED_TEXTURE:
+		case UBMT_RDG_TEXTURE:
 		{
 			FRDGTexture* RESTRICT Texture = *ParameterStruct.GetMemberPtrAtOffset<FRDGTexture*>(Offset);
 			if (Texture)
@@ -730,7 +730,7 @@ void FRDGBuilder::AllocateAndTransitionPassResources(const FRenderGraphPass* Pas
 			}
 		}
 		break;
-		case UBMT_GRAPH_TRACKED_SRV:
+		case UBMT_RDG_TEXTURE_SRV:
 		{
 			FRDGTextureSRV* RESTRICT SRV = *ParameterStruct.GetMemberPtrAtOffset<FRDGTextureSRV*>(Offset);
 			if (SRV)
@@ -742,7 +742,7 @@ void FRDGBuilder::AllocateAndTransitionPassResources(const FRenderGraphPass* Pas
 			}
 		}
 		break;
-		case UBMT_GRAPH_TRACKED_UAV:
+		case UBMT_RDG_TEXTURE_UAV:
 		{
 			FRDGTextureUAV* RESTRICT UAV = *ParameterStruct.GetMemberPtrAtOffset<FRDGTextureUAV*>(Offset);
 			if (UAV)
@@ -752,7 +752,7 @@ void FRDGBuilder::AllocateAndTransitionPassResources(const FRenderGraphPass* Pas
 			}
 		}
 		break;
-		case UBMT_GRAPH_TRACKED_BUFFER:
+		case UBMT_RDG_BUFFER:
 		{
 			FRDGBuffer* RESTRICT Buffer = *ParameterStruct.GetMemberPtrAtOffset<FRDGBuffer*>(Offset);
 			if (Buffer)
@@ -766,7 +766,7 @@ void FRDGBuilder::AllocateAndTransitionPassResources(const FRenderGraphPass* Pas
 			}
 		}
 		break;
-		case UBMT_GRAPH_TRACKED_BUFFER_SRV:
+		case UBMT_RDG_BUFFER_SRV:
 		{
 			FRDGBufferSRV* RESTRICT SRV = *ParameterStruct.GetMemberPtrAtOffset<FRDGBufferSRV*>(Offset);
 			if (SRV)
@@ -780,7 +780,7 @@ void FRDGBuilder::AllocateAndTransitionPassResources(const FRenderGraphPass* Pas
 			}
 		}
 		break;
-		case UBMT_GRAPH_TRACKED_BUFFER_UAV:
+		case UBMT_RDG_BUFFER_UAV:
 		{
 			FRDGBufferUAV* RESTRICT UAV = *ParameterStruct.GetMemberPtrAtOffset<FRDGBufferUAV*>(Offset);
 			if (UAV)
@@ -879,8 +879,8 @@ void FRDGBuilder::WarnForUselessPassDependencies(const FRenderGraphPass* Pass)
 		EUniformBufferBaseType Type = ParameterStruct.Layout->Resources[ResourceIndex].MemberType;
 		uint16 Offset = ParameterStruct.Layout->Resources[ResourceIndex].MemberOffset;
 
-		if (Type != UBMT_GRAPH_TRACKED_TEXTURE && Type != UBMT_GRAPH_TRACKED_SRV && Type != UBMT_GRAPH_TRACKED_UAV &&
-			Type != UBMT_GRAPH_TRACKED_BUFFER && Type != UBMT_GRAPH_TRACKED_BUFFER_SRV && Type != UBMT_GRAPH_TRACKED_BUFFER_UAV)
+		if (Type != UBMT_RDG_TEXTURE && Type != UBMT_RDG_TEXTURE_SRV && Type != UBMT_RDG_TEXTURE_UAV &&
+			Type != UBMT_RDG_BUFFER && Type != UBMT_RDG_BUFFER_SRV && Type != UBMT_RDG_BUFFER_UAV)
 			continue;
 
 		const FRDGResource* Resource = *ParameterStruct.GetMemberPtrAtOffset<const FRDGResource*>(Offset);
@@ -901,8 +901,8 @@ void FRDGBuilder::WarnForUselessPassDependencies(const FRenderGraphPass* Pass)
 			EUniformBufferBaseType Type = ParameterStruct.Layout->Resources[ResourceIndex].MemberType;
 			uint16 Offset = ParameterStruct.Layout->Resources[ResourceIndex].MemberOffset;
 
-			if (Type != UBMT_GRAPH_TRACKED_TEXTURE && Type != UBMT_GRAPH_TRACKED_SRV && Type != UBMT_GRAPH_TRACKED_UAV &&
-				Type != UBMT_GRAPH_TRACKED_BUFFER && Type != UBMT_GRAPH_TRACKED_BUFFER_SRV && Type != UBMT_GRAPH_TRACKED_BUFFER_UAV)
+			if (Type != UBMT_RDG_TEXTURE && Type != UBMT_RDG_TEXTURE_SRV && Type != UBMT_RDG_TEXTURE_UAV &&
+				Type != UBMT_RDG_BUFFER && Type != UBMT_RDG_BUFFER_SRV && Type != UBMT_RDG_BUFFER_UAV)
 				continue;
 
 			const FRDGResource* Resource = *ParameterStruct.GetMemberPtrAtOffset<const FRDGResource*>(Offset);
@@ -921,8 +921,8 @@ void FRDGBuilder::WarnForUselessPassDependencies(const FRenderGraphPass* Pass)
 		EUniformBufferBaseType Type = ParameterStruct.Layout->Resources[ResourceIndex].MemberType;
 		uint16 Offset = ParameterStruct.Layout->Resources[ResourceIndex].MemberOffset;
 
-		if (Type != UBMT_GRAPH_TRACKED_TEXTURE && Type != UBMT_GRAPH_TRACKED_SRV && Type != UBMT_GRAPH_TRACKED_UAV &&
-			Type != UBMT_GRAPH_TRACKED_BUFFER && Type != UBMT_GRAPH_TRACKED_BUFFER_SRV && Type != UBMT_GRAPH_TRACKED_BUFFER_UAV)
+		if (Type != UBMT_RDG_TEXTURE && Type != UBMT_RDG_TEXTURE_SRV && Type != UBMT_RDG_TEXTURE_UAV &&
+			Type != UBMT_RDG_BUFFER && Type != UBMT_RDG_BUFFER_SRV && Type != UBMT_RDG_BUFFER_UAV)
 			continue;
 
 		const FRDGResource* Resource = *ParameterStruct.GetMemberPtrAtOffset<const FRDGResource*>(Offset);
@@ -971,7 +971,7 @@ void FRDGBuilder::ReleaseUnecessaryResources(const FRenderGraphPass* Pass)
 
 		switch (Type)
 		{
-		case UBMT_GRAPH_TRACKED_TEXTURE:
+		case UBMT_RDG_TEXTURE:
 		{
 			FRDGTexture* RESTRICT Texture = *ParameterStruct.GetMemberPtrAtOffset<FRDGTexture*>(Offset);
 			if (Texture)
@@ -980,7 +980,7 @@ void FRDGBuilder::ReleaseUnecessaryResources(const FRenderGraphPass* Pass)
 			}
 		}
 		break;
-		case UBMT_GRAPH_TRACKED_SRV:
+		case UBMT_RDG_TEXTURE_SRV:
 		{
 			FRDGTextureSRV* RESTRICT SRV = *ParameterStruct.GetMemberPtrAtOffset<FRDGTextureSRV*>(Offset);
 			if (SRV)
@@ -989,7 +989,7 @@ void FRDGBuilder::ReleaseUnecessaryResources(const FRenderGraphPass* Pass)
 			}
 		}
 		break;
-		case UBMT_GRAPH_TRACKED_UAV:
+		case UBMT_RDG_TEXTURE_UAV:
 		{
 			FRDGTextureUAV* RESTRICT UAV = *ParameterStruct.GetMemberPtrAtOffset<FRDGTextureUAV*>(Offset);
 			if (UAV)
@@ -998,7 +998,7 @@ void FRDGBuilder::ReleaseUnecessaryResources(const FRenderGraphPass* Pass)
 			}
 		}
 		break;
-		case UBMT_GRAPH_TRACKED_BUFFER:
+		case UBMT_RDG_BUFFER:
 		{
 			FRDGBuffer* RESTRICT Buffer = *ParameterStruct.GetMemberPtrAtOffset<FRDGBuffer*>(Offset);
 			if (Buffer)
@@ -1007,7 +1007,7 @@ void FRDGBuilder::ReleaseUnecessaryResources(const FRenderGraphPass* Pass)
 			}
 		}
 		break;
-		case UBMT_GRAPH_TRACKED_BUFFER_SRV:
+		case UBMT_RDG_BUFFER_SRV:
 		{
 			FRDGBufferSRV* RESTRICT SRV = *ParameterStruct.GetMemberPtrAtOffset<FRDGBufferSRV*>(Offset);
 			if (SRV)
@@ -1016,7 +1016,7 @@ void FRDGBuilder::ReleaseUnecessaryResources(const FRenderGraphPass* Pass)
 			}
 		}
 		break;
-		case UBMT_GRAPH_TRACKED_BUFFER_UAV:
+		case UBMT_RDG_BUFFER_UAV:
 		{
 			FRDGBufferUAV* RESTRICT UAV = *ParameterStruct.GetMemberPtrAtOffset<FRDGBufferUAV*>(Offset);
 			if (UAV)
diff --git a/Engine/Source/Runtime/RenderCore/Private/RenderGraphUtils.cpp b/Engine/Source/Runtime/RenderCore/Private/RenderGraphUtils.cpp
--- a/Engine/Source/Runtime/RenderCore/Private/RenderGraphUtils.cpp
+++ b/Engine/Source/Runtime/RenderCore/Private/RenderGraphUtils.cpp
@@ -16,7 +16,7 @@ void ClearUnusedGraphResourcesImpl(const FShaderParameterBindings& ShaderBinding
 		EUniformBufferBaseType Type = ParametersMetadata->GetLayout().Resources[ResourceIndex].MemberType;
 		uint16 ByteOffset = ParametersMetadata->GetLayout().Resources[ResourceIndex].MemberOffset;
 
-		if (Type == UBMT_GRAPH_TRACKED_TEXTURE)
+		if (Type == UBMT_RDG_TEXTURE)
 		{
 			if (GraphTextureId < ShaderBindings.GraphTextures.Num() && ByteOffset == ShaderBindings.GraphTextures[GraphTextureId].ByteOffset)
 			{
@@ -24,7 +24,7 @@ void ClearUnusedGraphResourcesImpl(const FShaderParameterBindings& ShaderBinding
 				continue;
 			}
 		}
-		else if (Type == UBMT_GRAPH_TRACKED_SRV || Type == UBMT_GRAPH_TRACKED_BUFFER_SRV)
+		else if (Type == UBMT_RDG_TEXTURE_SRV || Type == UBMT_RDG_BUFFER_SRV)
 		{
 			if (GraphSRVId < ShaderBindings.GraphSRVs.Num() && ByteOffset == ShaderBindings.GraphSRVs[GraphSRVId].ByteOffset)
 			{
@@ -32,7 +32,7 @@ void ClearUnusedGraphResourcesImpl(const FShaderParameterBindings& ShaderBinding
 				continue;
 			}
 		}
-		else if (Type == UBMT_GRAPH_TRACKED_UAV || Type == UBMT_GRAPH_TRACKED_BUFFER_UAV)
+		else if (Type == UBMT_RDG_TEXTURE_UAV || Type == UBMT_RDG_BUFFER_UAV)
 		{
 			if (GraphUAVId < ShaderBindings.GraphUAVs.Num() && ByteOffset == ShaderBindings.GraphUAVs[GraphUAVId].ByteOffset)
 			{
diff --git a/Engine/Source/Runtime/RenderCore/Private/ShaderParameterStruct.cpp b/Engine/Source/Runtime/RenderCore/Private/ShaderParameterStruct.cpp
--- a/Engine/Source/Runtime/RenderCore/Private/ShaderParameterStruct.cpp
+++ b/Engine/Source/Runtime/RenderCore/Private/ShaderParameterStruct.cpp
@@ -53,7 +53,7 @@ struct FShaderParameterStructBindingContext
 				RenderTargetBindingSlotCppName = CppName;
 				continue;
 			}
-			else if (BaseType == UBMT_GRAPH_TRACKED_BUFFER)
+			else if (BaseType == UBMT_RDG_BUFFER)
 			{
 				continue;
 			}
@@ -127,11 +127,11 @@ struct FShaderParameterStructBindingContext
 				BaseType == UBMT_TEXTURE ||
 				BaseType == UBMT_SRV ||
 				BaseType == UBMT_SAMPLER ||
-				BaseType == UBMT_GRAPH_TRACKED_TEXTURE ||
-				BaseType == UBMT_GRAPH_TRACKED_SRV ||
-				BaseType == UBMT_GRAPH_TRACKED_UAV ||
-				BaseType == UBMT_GRAPH_TRACKED_BUFFER_SRV ||
-				BaseType == UBMT_GRAPH_TRACKED_BUFFER_UAV)
+				BaseType == UBMT_RDG_TEXTURE ||
+				BaseType == UBMT_RDG_TEXTURE_SRV ||
+				BaseType == UBMT_RDG_TEXTURE_UAV ||
+				BaseType == UBMT_RDG_BUFFER_SRV ||
+				BaseType == UBMT_RDG_BUFFER_UAV)
 			{
 				FShaderParameterBindings::FResourceParameter Parameter;
 				Parameter.BaseIndex = BaseIndex;
@@ -146,11 +146,11 @@ struct FShaderParameterStructBindingContext
 					Bindings->SRVs.Add(Parameter);
 				else if (BaseType == UBMT_SAMPLER)
 					Bindings->Samplers.Add(Parameter);
-				else if (BaseType == UBMT_GRAPH_TRACKED_TEXTURE)
+				else if (BaseType == UBMT_RDG_TEXTURE)
 					Bindings->GraphTextures.Add(Parameter);
-				else if (BaseType == UBMT_GRAPH_TRACKED_SRV || BaseType == UBMT_GRAPH_TRACKED_BUFFER_SRV)
+				else if (BaseType == UBMT_RDG_TEXTURE_SRV || BaseType == UBMT_RDG_BUFFER_SRV)
 					Bindings->GraphSRVs.Add(Parameter);
-				else // if (BaseType == UBMT_GRAPH_TRACKED_UAV || BaseType == UBMT_GRAPH_TRACKED_BUFFER_UAV)
+				else // if (BaseType == UBMT_RDG_TEXTURE_UAV || BaseType == UBMT_RDG_BUFFER_UAV)
 					Bindings->GraphUAVs.Add(Parameter);
 			}
 			else
diff --git a/Engine/Source/Runtime/RenderCore/Private/ShaderParameters.cpp b/Engine/Source/Runtime/RenderCore/Private/ShaderParameters.cpp
--- a/Engine/Source/Runtime/RenderCore/Private/ShaderParameters.cpp
+++ b/Engine/Source/Runtime/RenderCore/Private/ShaderParameters.cpp
@@ -264,7 +264,7 @@ static void CreateHLSLUniformBufferStructMembersDeclaration(
 
 		if (IsShaderParameterTypeForUniformBufferLayout(Member.GetBaseType()))
 		{
-			check(Member.GetBaseType() != UBMT_GRAPH_TRACKED_SRV && Member.GetBaseType() != UBMT_GRAPH_TRACKED_UAV);
+			check(Member.GetBaseType() != UBMT_RDG_TEXTURE_SRV && Member.GetBaseType() != UBMT_RDG_TEXTURE_UAV);
 			if (Member.GetBaseType() == UBMT_SRV)
 			{
 				// TODO: handle arrays?
diff --git a/Engine/Source/Runtime/RenderCore/Public/ShaderParameterMacros.h b/Engine/Source/Runtime/RenderCore/Public/ShaderParameterMacros.h
--- a/Engine/Source/Runtime/RenderCore/Public/ShaderParameterMacros.h
+++ b/Engine/Source/Runtime/RenderCore/Public/ShaderParameterMacros.h
@@ -712,7 +712,7 @@ extern RENDERCORE_API FShaderParametersMetadata* FindUniformBufferStructByFName(
  *	SHADER_PARAMETER_RDG_TEXTURE(Texture2D, MyTexture)
  */
 #define SHADER_PARAMETER_RDG_TEXTURE(ShaderType,MemberName) \
-	INTERNAL_SHADER_PARAMETER_EXPLICIT(UBMT_GRAPH_TRACKED_TEXTURE, TShaderResourceParameterTypeInfo<FRDGTextureRef>, FRDGTextureRef,MemberName,, = nullptr,EShaderPrecisionModifier::Float,TEXT(#ShaderType),false)
+	INTERNAL_SHADER_PARAMETER_EXPLICIT(UBMT_RDG_TEXTURE, TShaderResourceParameterTypeInfo<FRDGTextureRef>, FRDGTextureRef,MemberName,, = nullptr,EShaderPrecisionModifier::Float,TEXT(#ShaderType),false)
 
 /** Adds a shader resource view for a render graph tracked texture.
  *
@@ -720,7 +720,7 @@ extern RENDERCORE_API FShaderParametersMetadata* FindUniformBufferStructByFName(
  *	SHADER_PARAMETER_RDG_TEXTURE_SRV(Texture2D, MySRV)
  */
 #define SHADER_PARAMETER_RDG_TEXTURE_SRV(ShaderType,MemberName) \
-	INTERNAL_SHADER_PARAMETER_EXPLICIT(UBMT_GRAPH_TRACKED_SRV, TShaderResourceParameterTypeInfo<FRDGTextureSRVRef>, FRDGTextureSRVRef,MemberName,, = nullptr,EShaderPrecisionModifier::Float,TEXT(#ShaderType),false)
+	INTERNAL_SHADER_PARAMETER_EXPLICIT(UBMT_RDG_TEXTURE_SRV, TShaderResourceParameterTypeInfo<FRDGTextureSRVRef>, FRDGTextureSRVRef,MemberName,, = nullptr,EShaderPrecisionModifier::Float,TEXT(#ShaderType),false)
 
 /** Adds a unordered access view for a render graph tracked texture.
  *
@@ -728,7 +728,7 @@ extern RENDERCORE_API FShaderParametersMetadata* FindUniformBufferStructByFName(
  *	SHADER_PARAMETER_RDG_TEXTURE_UAV(RWTexture2D, MyUAV)
  */
 #define SHADER_PARAMETER_RDG_TEXTURE_UAV(ShaderType,MemberName) \
-	INTERNAL_SHADER_PARAMETER_EXPLICIT(UBMT_GRAPH_TRACKED_UAV, TShaderResourceParameterTypeInfo<FRDGTextureUAVRef>, FRDGTextureUAVRef,MemberName,, = nullptr,EShaderPrecisionModifier::Float,TEXT(#ShaderType),false)
+	INTERNAL_SHADER_PARAMETER_EXPLICIT(UBMT_RDG_TEXTURE_UAV, TShaderResourceParameterTypeInfo<FRDGTextureUAVRef>, FRDGTextureUAVRef,MemberName,, = nullptr,EShaderPrecisionModifier::Float,TEXT(#ShaderType),false)
 
 /** Adds a render graph tracked buffer.
  *
@@ -736,7 +736,7 @@ extern RENDERCORE_API FShaderParametersMetadata* FindUniformBufferStructByFName(
  *	SHADER_PARAMETER_RDG_BUFFER(Texture2D, MyTexture)
  */
 #define SHADER_PARAMETER_RDG_BUFFER(ShaderType,MemberName) \
-	INTERNAL_SHADER_PARAMETER_EXPLICIT(UBMT_GRAPH_TRACKED_BUFFER, TShaderResourceParameterTypeInfo<FRDGBufferRef>, FRDGBufferRef,MemberName,, = nullptr,EShaderPrecisionModifier::Float,TEXT(#ShaderType),false)
+	INTERNAL_SHADER_PARAMETER_EXPLICIT(UBMT_RDG_BUFFER, TShaderResourceParameterTypeInfo<FRDGBufferRef>, FRDGBufferRef,MemberName,, = nullptr,EShaderPrecisionModifier::Float,TEXT(#ShaderType),false)
 
 /** Adds a shader resource view for a render graph tracked buffer.
  *
@@ -744,7 +744,7 @@ extern RENDERCORE_API FShaderParametersMetadata* FindUniformBufferStructByFName(
  *	SHADER_PARAMETER_RDG_BUFFER_SRV(Buffer<float4>, MySRV)
  */
 #define SHADER_PARAMETER_RDG_BUFFER_SRV(ShaderType,MemberName) \
-	INTERNAL_SHADER_PARAMETER_EXPLICIT(UBMT_GRAPH_TRACKED_BUFFER_SRV, TShaderResourceParameterTypeInfo<FRDGBufferSRVRef>, FRDGBufferSRVRef,MemberName,, = nullptr,EShaderPrecisionModifier::Float,TEXT(#ShaderType),false)
+	INTERNAL_SHADER_PARAMETER_EXPLICIT(UBMT_RDG_BUFFER_SRV, TShaderResourceParameterTypeInfo<FRDGBufferSRVRef>, FRDGBufferSRVRef,MemberName,, = nullptr,EShaderPrecisionModifier::Float,TEXT(#ShaderType),false)
 
 /** Adds a unordered access view for a render graph tracked buffer.
  *
@@ -752,7 +752,7 @@ extern RENDERCORE_API FShaderParametersMetadata* FindUniformBufferStructByFName(
  *	SHADER_PARAMETER_RDG_BUFFER_UAV(RWBuffer<float4>, MyUAV)
  */
 #define SHADER_PARAMETER_RDG_BUFFER_UAV(ShaderType,MemberName) \
-	INTERNAL_SHADER_PARAMETER_EXPLICIT(UBMT_GRAPH_TRACKED_BUFFER_UAV, TShaderResourceParameterTypeInfo<FRDGBufferUAVRef>, FRDGBufferUAVRef,MemberName,, = nullptr,EShaderPrecisionModifier::Float,TEXT(#ShaderType),false)
+	INTERNAL_SHADER_PARAMETER_EXPLICIT(UBMT_RDG_BUFFER_UAV, TShaderResourceParameterTypeInfo<FRDGBufferUAVRef>, FRDGBufferUAVRef,MemberName,, = nullptr,EShaderPrecisionModifier::Float,TEXT(#ShaderType),false)
 
 /** Nests a shader parameter structure into another one, in C++ and shader code.
  *
diff --git a/Engine/Source/Runtime/Windows/D3D11RHI/Private/D3D11UniformBuffer.cpp b/Engine/Source/Runtime/Windows/D3D11RHI/Private/D3D11UniformBuffer.cpp
--- a/Engine/Source/Runtime/Windows/D3D11RHI/Private/D3D11UniformBuffer.cpp
+++ b/Engine/Source/Runtime/Windows/D3D11RHI/Private/D3D11UniformBuffer.cpp
@@ -212,7 +212,7 @@ FUniformBufferRHIRef FD3D11DynamicRHI::RHICreateUniformBuffer(const void* Conten
 
 			// Allow null SRV's in uniform buffers for feature levels that don't support SRV's in shaders
 			if (!(GMaxRHIFeatureLevel <= ERHIFeatureLevel::ES3_1 
-				&& (Layout.Resources[i].MemberType == UBMT_SRV || Layout.Resources[i].MemberType == UBMT_GRAPH_TRACKED_SRV || Layout.Resources[i].MemberType == UBMT_GRAPH_TRACKED_BUFFER_SRV))
+				&& (Layout.Resources[i].MemberType == UBMT_SRV || Layout.Resources[i].MemberType == UBMT_RDG_TEXTURE_SRV || Layout.Resources[i].MemberType == UBMT_RDG_BUFFER_SRV))
 				&& Validation == EUniformBufferValidation::ValidateResources)
 			{
 				checkf(Resource, TEXT("Invalid resource entry creating uniform buffer, %s.Resources[%u], ResourceType 0x%x."), *Layout.GetDebugName().ToString(), i, Layout.Resources[i].MemberType);
