Implements early validation of resource production in render graph to catch common user mistakes more conveniently.

diff --git a/Engine/Source/Runtime/RenderCore/Private/RenderGraphBuilder.cpp b/Engine/Source/Runtime/RenderCore/Private/RenderGraphBuilder.cpp
--- a/Engine/Source/Runtime/RenderCore/Private/RenderGraphBuilder.cpp
+++ b/Engine/Source/Runtime/RenderCore/Private/RenderGraphBuilder.cpp
@@ -181,13 +181,67 @@ void FRDGBuilder::ValidatePass(const FRenderGraphPass* Pass) const
 
 		switch (Type)
 		{
+		case UBMT_RDG_TEXTURE:
+		{
+			FRDGTexture* RESTRICT Texture = *ParameterStruct.GetMemberPtrAtOffset<FRDGTexture*>(Offset);
+			checkf(!Texture || Texture->bHasEverBeenProduced,
+				TEXT("Pass %s has a dependency over the texture %s that has never been produced."),
+				Pass->GetName(), Texture->Name);
+		}
+		break;
+		case UBMT_RDG_TEXTURE_SRV:
+		{
+			FRDGTextureSRV* RESTRICT SRV = *ParameterStruct.GetMemberPtrAtOffset<FRDGTextureSRV*>(Offset);
+			if (SRV)
+			{
+				checkf(SRV->Desc.Texture->bHasEverBeenProduced,
+					TEXT("Pass %s has a dependency over the texture %s that has never been produced."),
+					Pass->GetName(), SRV->Desc.Texture->Name);
+			}
+		}
+		break;
 		case UBMT_RDG_TEXTURE_UAV:
+		{
+			FRDGTextureUAV* RESTRICT UAV = *ParameterStruct.GetMemberPtrAtOffset<FRDGTextureUAV*>(Offset);
+			if (UAV)
+			{
+				UAV->Desc.Texture->bHasEverBeenProduced = true;
+				if (!bCanUseUAVs && GRenderGraphEmitWarnings)
+				{
+					UE_LOG(LogRendererCore, Warning, TEXT("UAV can only been bound to compute shaders, therefore UAV %s is certainly useless for pass %s."), UAV->Name, Pass->GetName());
+				}
+			}
+		}
+		break;
+		case UBMT_RDG_BUFFER:
+		{
+			FRDGBuffer* RESTRICT Buffer = *ParameterStruct.GetMemberPtrAtOffset<FRDGBuffer*>(Offset);
+			checkf(!Buffer || Buffer->bHasEverBeenProduced,
+				TEXT("Pass %s has a dependency over the buffer %s that has never been produced."),
+				Pass->GetName(), Buffer->Name);
+		}
+		break;
+		case UBMT_RDG_BUFFER_SRV:
+		{
+			FRDGBufferSRV* RESTRICT SRV = *ParameterStruct.GetMemberPtrAtOffset<FRDGBufferSRV*>(Offset);
+			if (SRV)
+			{
+				checkf(SRV->Desc.Buffer->bHasEverBeenProduced,
+					TEXT("Pass %s has a dependency over the buffer %s that has never been produced."),
+					Pass->GetName(), SRV->Desc.Buffer->Name);
+			}
+		}
+		break;
 		case UBMT_RDG_BUFFER_UAV:
 		{
-			FRDGResource* RESTRICT UAV = *ParameterStruct.GetMemberPtrAtOffset<FRDGResource*>(Offset);
-			if (UAV && !bCanUseUAVs && GRenderGraphEmitWarnings)
+			FRDGBufferUAV* RESTRICT UAV = *ParameterStruct.GetMemberPtrAtOffset<FRDGBufferUAV*>(Offset);
+			if (UAV)
 			{
-				UE_LOG(LogRendererCore, Warning, TEXT("UAV can only been bound to compute shaders, therefore UAV %s is certainly useless for passs %s."), UAV->Name, Pass->GetName());
+				UAV->Desc.Buffer->bHasEverBeenProduced = true;
+				if (!bCanUseUAVs && GRenderGraphEmitWarnings)
+				{
+					UE_LOG(LogRendererCore, Warning, TEXT("UAV can only been bound to compute shaders, therefore UAV %s is certainly useless for pass %s."), UAV->Name, Pass->GetName());
+				}
 			}
 		}
 		break;
@@ -215,23 +269,48 @@ void FRDGBuilder::ValidatePass(const FRenderGraphPass* Pass) const
 		const bool bGeneratingMips = (Pass->GetFlags() & ERenderGraphPassFlags::GenerateMips) == ERenderGraphPassFlags::GenerateMips;
 		bool bFoundRTBound = false;
 
-		uint32 NumRenderTargets = 0;
-		for (int i = 0; i < RenderTargets->Output.Num(); i++)
+		int32 NumRenderTargets = 0;
+		for (int32 i = 0; i < RenderTargets->Output.Num(); i++)
 		{
 			const FRenderTargetBinding& RenderTarget = RenderTargets->Output[i];
-			if (!RenderTarget.GetTexture())
+			const FRDGTexture* Texture = RenderTarget.GetTexture();
+
+			if (!Texture)
 			{
 				NumRenderTargets = i;
 				break;
 			}
+
+			if (!Texture->bHasEverBeenProduced)
+			{
+				checkf(RenderTarget.GetLoadAction() != ERenderTargetLoadAction::ELoad,
+					TEXT("Can't load a render target %s that has never been produced."),
+					RenderTarget.GetTexture()->Name);
+
+				// TODO(RDG): should only be done when there is a store action.
+				Texture->bHasEverBeenProduced = true;
+			}
+
 			bFoundRTBound = bFoundRTBound || IsBoundAsReadable(RenderTarget.GetTexture(), ParameterStruct);
 		}
-		for (int i = NumRenderTargets; i < RenderTargets->Output.Num(); i++)
+		for (int32 i = NumRenderTargets; i < RenderTargets->Output.Num(); i++)
 		{
 			const FRenderTargetBinding& RenderTarget = RenderTargets->Output[i];
 			checkf(RenderTarget.GetTexture() == nullptr, TEXT("Render targets must be packed. No empty spaces in the array."));
 		}
 		ensureMsgf(!bGeneratingMips || bFoundRTBound, TEXT("GenerateMips enabled but no RT found as source!"));
+
+		const FRDGTexture* Texture = RenderTargets->DepthStencil.Texture;
+		if (Texture && !Texture->bHasEverBeenProduced)
+		{
+			checkf(RenderTargets->DepthStencil.DepthLoadAction != ERenderTargetLoadAction::ELoad,
+				TEXT("Can't load depth from a render target that has never been produced."));
+			checkf(RenderTargets->DepthStencil.StencilLoadAction != ERenderTargetLoadAction::ELoad,
+				TEXT("Can't load stencil from a render target that has never been produced."));
+
+			// TODO(RDG): should only be done when there is a store action.
+			Texture->bHasEverBeenProduced = true;
+		}
 	}
 	else
 	{
@@ -269,7 +348,7 @@ void FRDGBuilder::CaptureAnyInterestingPassOutput(const FRenderGraphPass* Pass)
 			{
 				GVisualizeTexture.CreateContentCapturePass(*this, RenderTargets->DepthStencil.Texture);
 			}
-			for (int i = 0; i < RenderTargets->Output.Num(); i++)
+			for (int32 i = 0; i < RenderTargets->Output.Num(); i++)
 			{
 				const FRenderTargetBinding& RenderTarget = RenderTargets->Output[i];
 				if (RenderTarget.GetTexture() &&
@@ -356,7 +435,7 @@ void FRDGBuilder::WalkGraphDependencies()
 			{
 				FRenderTargetBindingSlots* RESTRICT RenderTargets = ParameterStruct.GetMemberPtrAtOffset<FRenderTargetBindingSlots>(Offset);
 
-				for (int i = 0; i < RenderTargets->Output.Num(); i++)
+				for (int32 i = 0; i < RenderTargets->Output.Num(); i++)
 				{
 					const FRenderTargetBinding& RenderTarget = RenderTargets->Output[i];
 					if (RenderTarget.GetTexture())
@@ -800,7 +879,7 @@ void FRDGBuilder::AllocateAndTransitionPassResources(const FRenderGraphPass* Pas
 			uint32 NumDepthStencilTargets = 0;
 			uint32 NumSamples = 0;
 
-			for (int i = 0; i < RenderTargets->Output.Num(); i++)
+			for (int32 i = 0; i < RenderTargets->Output.Num(); i++)
 			{
 				const FRenderTargetBinding& RenderTarget = RenderTargets->Output[i];
 				if (RenderTarget.GetTexture())
@@ -1029,7 +1108,7 @@ void FRDGBuilder::ReleaseUnecessaryResources(const FRenderGraphPass* Pass)
 		{
 			FRenderTargetBindingSlots* RESTRICT RenderTargets = ParameterStruct.GetMemberPtrAtOffset<FRenderTargetBindingSlots>(Offset);
 
-			for (int i = 0; i < RenderTargets->Output.Num(); i++)
+			for (int32 i = 0; i < RenderTargets->Output.Num(); i++)
 			{
 				const FRenderTargetBinding& RenderTarget = RenderTargets->Output[i];
 				if (RenderTarget.GetTexture())
diff --git a/Engine/Source/Runtime/RenderCore/Public/RenderGraphBuilder.h b/Engine/Source/Runtime/RenderCore/Public/RenderGraphBuilder.h
--- a/Engine/Source/Runtime/RenderCore/Public/RenderGraphBuilder.h
+++ b/Engine/Source/Runtime/RenderCore/Public/RenderGraphBuilder.h
@@ -4,17 +4,12 @@
 
 #include "RenderGraphResources.h"
 #include "ShaderParameterMacros.h"
-#include "RendererInterface.h"
 
 
 /** Whether visualize texture tool is supported. */
 #define SUPPORTS_VISUALIZE_TEXTURE (WITH_ENGINE && !(UE_BUILD_SHIPPING || UE_BUILD_TEST))
 
 
-/** Whether render graph debugging is compiled. */
-#define RENDER_GRAPH_DEBUGGING (!UE_BUILD_SHIPPING)
-
-
 /** Whether render graph should support draw events or not.
  * RENDER_GRAPH_DRAW_EVENTS == 0 means there is no string processing at all.
  * RENDER_GRAPH_DRAW_EVENTS == 1 means const TCHAR* is only passdown.
@@ -222,7 +217,10 @@ public:
 		OutTexture->PooledRenderTarget = ExternalPooledTexture;
 		AllocatedTextures.Add(OutTexture, ExternalPooledTexture);
 		#if RENDER_GRAPH_DEBUGGING
+		{
+			OutTexture->bHasEverBeenProduced = true;
 			Resources.Add(OutTexture);
+		}
 		#endif
 		return OutTexture;
 	}
@@ -383,6 +381,10 @@ public:
 			checkf(!bHasExecuted,
 				TEXT("Accessing render graph internal texture %s with QueueTextureExtraction() needs to happen before the builder's execution."),
 				Texture->Name);
+
+			checkf(Texture->bHasEverBeenProduced,
+				TEXT("Unable to queue the extraction of the texture %s because it has not been produced by any pass."),
+				Texture->Name);
 		}
 		#endif
 		FDeferredInternalTextureQuery Query;
diff --git a/Engine/Source/Runtime/RenderCore/Public/RenderGraphResources.h b/Engine/Source/Runtime/RenderCore/Public/RenderGraphResources.h
--- a/Engine/Source/Runtime/RenderCore/Public/RenderGraphResources.h
+++ b/Engine/Source/Runtime/RenderCore/Public/RenderGraphResources.h
@@ -5,7 +5,10 @@
 #include "CoreMinimal.h"
 #include "RHI.h"
 #include "RendererInterface.h"
-#include "RendererInterface.h"
+
+
+/** Whether render graph debugging is compiled. */
+#define RENDER_GRAPH_DEBUGGING (!UE_BUILD_SHIPPING)
 
 
 class FRDGResource;
@@ -66,9 +69,14 @@ private:
 	mutable bool bWritable = false;
 	mutable bool bCompute = false;
 
-	/** Boolean to tracked whether a ressource is actually used by the lambda of a pass or not. */
+	/** Boolean to track at runtime whether a ressource is actually used by the lambda of a pass or not, to detect unnecessary resource dependencies on passes. */
 	mutable bool bIsActuallyUsedByPass = false;
 
+#if RENDER_GRAPH_DEBUGGING
+	/** Boolean to track at wiring time if a resource has ever been produced by a pass, to error out early if accessing a resource that has not been produced. */
+	mutable bool bHasEverBeenProduced = false;
+#endif
+
 	friend class FRDGBuilder;
 
 
