Fixed crash on startup when using r.CreateShadersOnLoad=1 and D3D11

GRHISupportsRayTracing is now checked before attempting to create a ray tracing shader.

diff --git a/Engine/Source/Runtime/RenderCore/Private/Shader.cpp b/Engine/Source/Runtime/RenderCore/Private/Shader.cpp
--- a/Engine/Source/Runtime/RenderCore/Private/Shader.cpp
+++ b/Engine/Source/Runtime/RenderCore/Private/Shader.cpp
@@ -937,8 +937,11 @@ void FShaderResource::InitRHI()
 #if RHI_RAYTRACING
 	else if (Target.Frequency == SF_RayGen || Target.Frequency == SF_RayMiss || Target.Frequency == SF_RayHitGroup)
 	{
-		RayTracingShader = RHICreateRayTracingShader(UncompressedCode, Target.GetFrequency());
-		UE_CLOG((bCodeInSharedLocation && !IsValidRef(RayTracingShader)), LogShaders, Fatal, TEXT("FShaderResource::SerializeShaderCode can't find shader code for: [%s]"), *LegacyShaderPlatformToShaderFormat((EShaderPlatform)Target.Platform).ToString());
+		if (GRHISupportsRayTracing)
+		{
+			RayTracingShader = RHICreateRayTracingShader(UncompressedCode, Target.GetFrequency());
+			UE_CLOG((bCodeInSharedLocation && !IsValidRef(RayTracingShader)), LogShaders, Fatal, TEXT("FShaderResource::SerializeShaderCode can't find shader code for: [%s]"), *LegacyShaderPlatformToShaderFormat((EShaderPlatform)Target.Platform).ToString());
+		}
 	}
 #endif // RHI_RAYTRACING
 	else
