Fixed crash on startup CreateShadersOnLoad is enabled

diff --git a/Engine/Source/Runtime/RenderCore/Private/Shader.cpp b/Engine/Source/Runtime/RenderCore/Private/Shader.cpp
--- a/Engine/Source/Runtime/RenderCore/Private/Shader.cpp
+++ b/Engine/Source/Runtime/RenderCore/Private/Shader.cpp
@@ -941,6 +941,11 @@ void FShaderResource::InitRHI()
 		{
 			RayTracingShader = RHICreateRayTracingShader(UncompressedCode, Target.GetFrequency());
 			UE_CLOG((bCodeInSharedLocation && !IsValidRef(RayTracingShader)), LogShaders, Fatal, TEXT("FShaderResource::SerializeShaderCode can't find shader code for: [%s]"), *LegacyShaderPlatformToShaderFormat((EShaderPlatform)Target.Platform).ToString());
+
+			if (Target.Frequency == SF_RayHitGroup)
+			{
+				RayTracingMaterialLibraryIndex = AddToRayTracingLibrary(RayTracingShader);
+			}
 		}
 	}
 #endif // RHI_RAYTRACING
@@ -1009,13 +1014,6 @@ void FShaderResource::InitializeShaderRHI()
 		}
 
 		INC_FLOAT_STAT_BY(STAT_Shaders_TotalRTShaderInitForRenderingTime,(float)ShaderInitializationTime);
-
-#if RHI_RAYTRACING
-		if (Target.Frequency == SF_RayHitGroup)
-		{
-			RayTracingMaterialLibraryIndex = AddToRayTracingLibrary(GetRayTracingShader());
-		}
-#endif
 	}
 
 	checkSlow(IsInitialized());
