Impoves implementation of FRDGBuilder::WarnForUselessPassDependencies()

diff --git a/Engine/Source/Runtime/RenderCore/Private/RenderGraphBuilder.cpp b/Engine/Source/Runtime/RenderCore/Private/RenderGraphBuilder.cpp
--- a/Engine/Source/Runtime/RenderCore/Private/RenderGraphBuilder.cpp
+++ b/Engine/Source/Runtime/RenderCore/Private/RenderGraphBuilder.cpp
@@ -958,8 +958,7 @@ void FRDGBuilder::WarnForUselessPassDependencies(const FRenderGraphPass* Pass)
 		EUniformBufferBaseType Type = ParameterStruct.Layout->Resources[ResourceIndex].MemberType;
 		uint16 Offset = ParameterStruct.Layout->Resources[ResourceIndex].MemberOffset;
 
-		if (Type != UBMT_RDG_TEXTURE && Type != UBMT_RDG_TEXTURE_SRV && Type != UBMT_RDG_TEXTURE_UAV &&
-			Type != UBMT_RDG_BUFFER && Type != UBMT_RDG_BUFFER_SRV && Type != UBMT_RDG_BUFFER_UAV)
+		if (!IsRDGResourceReferenceShaderParameterType(Type))
 			continue;
 
 		const FRDGResource* Resource = *ParameterStruct.GetMemberPtrAtOffset<const FRDGResource*>(Offset);
@@ -980,8 +979,7 @@ void FRDGBuilder::WarnForUselessPassDependencies(const FRenderGraphPass* Pass)
 			EUniformBufferBaseType Type = ParameterStruct.Layout->Resources[ResourceIndex].MemberType;
 			uint16 Offset = ParameterStruct.Layout->Resources[ResourceIndex].MemberOffset;
 
-			if (Type != UBMT_RDG_TEXTURE && Type != UBMT_RDG_TEXTURE_SRV && Type != UBMT_RDG_TEXTURE_UAV &&
-				Type != UBMT_RDG_BUFFER && Type != UBMT_RDG_BUFFER_SRV && Type != UBMT_RDG_BUFFER_UAV)
+			if (!IsRDGResourceReferenceShaderParameterType(Type))
 				continue;
 
 			const FRDGResource* Resource = *ParameterStruct.GetMemberPtrAtOffset<const FRDGResource*>(Offset);
@@ -1000,8 +998,7 @@ void FRDGBuilder::WarnForUselessPassDependencies(const FRenderGraphPass* Pass)
 		EUniformBufferBaseType Type = ParameterStruct.Layout->Resources[ResourceIndex].MemberType;
 		uint16 Offset = ParameterStruct.Layout->Resources[ResourceIndex].MemberOffset;
 
-		if (Type != UBMT_RDG_TEXTURE && Type != UBMT_RDG_TEXTURE_SRV && Type != UBMT_RDG_TEXTURE_UAV &&
-			Type != UBMT_RDG_BUFFER && Type != UBMT_RDG_BUFFER_SRV && Type != UBMT_RDG_BUFFER_UAV)
+		if (!IsRDGResourceReferenceShaderParameterType(Type))
 			continue;
 
 		const FRDGResource* Resource = *ParameterStruct.GetMemberPtrAtOffset<const FRDGResource*>(Offset);
