Allocates render graph resources with lifetime for the RHI

diff --git a/Engine/Source/Runtime/D3D12RHI/Private/D3D12UniformBuffer.cpp b/Engine/Source/Runtime/D3D12RHI/Private/D3D12UniformBuffer.cpp
--- a/Engine/Source/Runtime/D3D12RHI/Private/D3D12UniformBuffer.cpp
+++ b/Engine/Source/Runtime/D3D12RHI/Private/D3D12UniformBuffer.cpp
@@ -90,8 +90,6 @@ FUniformBufferRHIRef FD3D12DynamicRHI::RHICreateUniformBuffer(const void* Conten
 				}
 				else if (IsRDGResourceReferenceShaderParameterType(ResourceType))
 				{
-					// TODO(RDG): need to allocate FRDGResource with a lifetime for RHI.
-					check(IsInRenderingThread());
 					FRHIResource** ResourcePtr = *(FRHIResource***)((uint8*)Contents + Layout.Resources[i].MemberOffset);
 					Resource = ResourcePtr ? *ResourcePtr : nullptr;
 				}
@@ -216,8 +214,6 @@ void FD3D12DynamicRHI::RHIUpdateUniformBuffer(FUniformBufferRHIParamRef UniformB
 			}
 			else if (IsRDGResourceReferenceShaderParameterType(ResourceType))
 			{
-				// TODO(RDG): need to allocate FRDGResource with a lifetime for RHI.
-				check(IsInRenderingThread());
 				FRHIResource** ResourcePtr = *(FRHIResource***)((uint8*)Contents + Layout.Resources[ResourceIndex].MemberOffset);
 				Resource = ResourcePtr ? *ResourcePtr : nullptr;
 			}
diff --git a/Engine/Source/Runtime/RenderCore/Public/RenderGraphBuilder.h b/Engine/Source/Runtime/RenderCore/Public/RenderGraphBuilder.h
--- a/Engine/Source/Runtime/RenderCore/Public/RenderGraphBuilder.h
+++ b/Engine/Source/Runtime/RenderCore/Public/RenderGraphBuilder.h
@@ -213,7 +213,7 @@ public:
 			ensureMsgf(ExternalPooledTexture.IsValid(), TEXT("Attempted to register NULL external texture: %s"), Name);
 		}
 		#endif
-		FRDGTexture* OutTexture = new(FMemStack::Get()) FRDGTexture(Name, ExternalPooledTexture->GetDesc());
+		FRDGTexture* OutTexture = AllocateForRHILifeTime<FRDGTexture>(Name, ExternalPooledTexture->GetDesc());
 		OutTexture->PooledRenderTarget = ExternalPooledTexture;
 		AllocatedTextures.Add(OutTexture, ExternalPooledTexture);
 		#if RENDER_GRAPH_DEBUGGING
@@ -233,7 +233,7 @@ public:
 			ensureMsgf(!bHasExecuted, TEXT("Render graph texture %s needs to be created before the builder execution."), DebugName);
 		}
 		#endif
-		FRDGTexture* Texture = new(FMemStack::Get()) FRDGTexture(DebugName, Desc);
+		FRDGTexture* Texture = AllocateForRHILifeTime<FRDGTexture>(DebugName, Desc);
 		#if RENDER_GRAPH_DEBUGGING
 			Resources.Add(Texture);
 		#endif
@@ -248,7 +248,7 @@ public:
 			ensureMsgf(!bHasExecuted, TEXT("Render graph buffer %s needs to be created before the builder execution."), DebugName);
 		}
 		#endif
-		FRDGBuffer* Buffer = new(FMemStack::Get()) FRDGBuffer(DebugName, Desc);
+		FRDGBuffer* Buffer = AllocateForRHILifeTime<FRDGBuffer>(DebugName, Desc);
 		#if RENDER_GRAPH_DEBUGGING
 			Resources.Add(Buffer);
 		#endif
@@ -266,7 +266,7 @@ public:
 		}
 		#endif
 		
-		FRDGTextureSRV* SRV = new(FMemStack::Get()) FRDGTextureSRV(Desc.Texture->Name, Desc);
+		FRDGTextureSRV* SRV = AllocateForRHILifeTime<FRDGTextureSRV>(Desc.Texture->Name, Desc);
 		#if RENDER_GRAPH_DEBUGGING
 			Resources.Add(SRV);
 		#endif
@@ -283,7 +283,7 @@ public:
 		}
 		#endif
 		
-		FRDGBufferSRV* SRV = new(FMemStack::Get()) FRDGBufferSRV(Desc.Buffer->Name, Desc);
+		FRDGBufferSRV* SRV = AllocateForRHILifeTime<FRDGBufferSRV>(Desc.Buffer->Name, Desc);
 		#if RENDER_GRAPH_DEBUGGING
 			Resources.Add(SRV);
 		#endif
@@ -301,7 +301,7 @@ public:
 		}
 		#endif
 		
-		FRDGTextureUAV* UAV = new(FMemStack::Get()) FRDGTextureUAV(Desc.Texture->Name, Desc);
+		FRDGTextureUAV* UAV = AllocateForRHILifeTime<FRDGTextureUAV>(Desc.Texture->Name, Desc);
 		#if RENDER_GRAPH_DEBUGGING
 			Resources.Add(UAV);
 		#endif
@@ -318,7 +318,7 @@ public:
 		}
 		#endif
 		
-		FRDGBufferUAV* UAV = new(FMemStack::Get()) FRDGBufferUAV(Desc.Buffer->Name, Desc);
+		FRDGBufferUAV* UAV = AllocateForRHILifeTime<FRDGBufferUAV>(Desc.Buffer->Name, Desc);
 		#if RENDER_GRAPH_DEBUGGING
 			Resources.Add(UAV);
 		#endif
@@ -329,6 +329,7 @@ public:
 	template< typename ParameterStructType >
 	inline ParameterStructType* AllocParameters() const
 	{
+		// TODO(RDG): could allocate using AllocateForRHILifeTime() to avoid the copy done when using FRHICommandList::BuildLocalUniformBuffer()
 		ParameterStructType* OutParameterPtr = new(FMemStack::Get()) ParameterStructType;
 		FMemory::Memzero(OutParameterPtr, sizeof(ParameterStructType));
 		return OutParameterPtr;
@@ -450,6 +451,23 @@ private:
 	void CaptureAnyInterestingPassOutput(const FRenderGraphPass* Pass);
 
 	void WalkGraphDependencies();
+	
+	template<class Type, class ...ConstructorParameterTypes>
+	Type* AllocateForRHILifeTime(ConstructorParameterTypes&&... ConstructorParameters)
+	{
+		check(IsInRenderingThread());
+		// When bypassing the RHI command queuing, can allocate directly on render thread memory stack allocator, otherwise allocate
+		// on the RHI's stack allocator so RHICreateUniformBuffer() can dereference render graph resources.
+		if (RHICmdList.Bypass())
+		{
+			return new (FMemStack::Get()) Type(Forward<ConstructorParameterTypes>(ConstructorParameters)...);
+		}
+		else
+		{
+			void* UnitializedType = RHICmdList.Alloc<Type>();
+			return new (UnitializedType) Type(Forward<ConstructorParameterTypes>(ConstructorParameters)...);
+		}
+	}
 
 	void AllocateRHITextureIfNeeded(const FRDGTexture* Texture, bool bComputePass);
 	void AllocateRHITextureSRVIfNeeded(const FRDGTextureSRV* SRV, bool bComputePass);
