Adds a warning in render graph when a resource is produced by a pass, but never used after.

diff --git a/Engine/Shaders/Private/ScreenSpaceDenoise/SSDTemporalAccumulation.usf b/Engine/Shaders/Private/ScreenSpaceDenoise/SSDTemporalAccumulation.usf
--- a/Engine/Shaders/Private/ScreenSpaceDenoise/SSDTemporalAccumulation.usf
+++ b/Engine/Shaders/Private/ScreenSpaceDenoise/SSDTemporalAccumulation.usf
@@ -846,11 +846,11 @@ void TemporallyAccumulate(
 				#error Unsupported signal rejection.
 			#endif
 			
-			DebugOutput[DispatchThreadId] = float4(
-				RejectionFactor,
-				0,
-				0,
-				0);
+			//DebugOutput[DispatchThreadId] = float4(
+			//	RejectionFactor,
+			//	0,
+			//	0,
+			//	0);
 
 			float FinalHistoryWeight = HistorySample.AccumulatedSampleCount * saturate(1 - RejectionFactor);
 
diff --git a/Engine/Source/Runtime/RenderCore/Private/RenderGraphBuilder.cpp b/Engine/Source/Runtime/RenderCore/Private/RenderGraphBuilder.cpp
--- a/Engine/Source/Runtime/RenderCore/Private/RenderGraphBuilder.cpp
+++ b/Engine/Source/Runtime/RenderCore/Private/RenderGraphBuilder.cpp
@@ -205,7 +205,12 @@ void FRDGBuilder::ValidatePass(const FRenderGraphPass* Pass) const
 			FRDGTextureUAV* RESTRICT UAV = *ParameterStruct.GetMemberPtrAtOffset<FRDGTextureUAV*>(Offset);
 			if (UAV)
 			{
-				UAV->Desc.Texture->bHasEverBeenProduced = true;
+				if (!UAV->Desc.Texture->bHasEverBeenProduced)
+				{
+					UAV->Desc.Texture->bHasEverBeenProduced = true;
+					UAV->Desc.Texture->DebugFirstProducer = Pass;
+				}
+
 				if (!bCanUseUAVs && GRenderGraphEmitWarnings)
 				{
 					UE_LOG(LogRendererCore, Warning, TEXT("UAV can only been bound to compute shaders, therefore UAV %s is certainly useless for pass %s."), UAV->Name, Pass->GetName());
@@ -237,7 +242,12 @@ void FRDGBuilder::ValidatePass(const FRenderGraphPass* Pass) const
 			FRDGBufferUAV* RESTRICT UAV = *ParameterStruct.GetMemberPtrAtOffset<FRDGBufferUAV*>(Offset);
 			if (UAV)
 			{
-				UAV->Desc.Buffer->bHasEverBeenProduced = true;
+				if (!UAV->Desc.Buffer->bHasEverBeenProduced)
+				{
+					UAV->Desc.Buffer->bHasEverBeenProduced = true;
+					UAV->Desc.Buffer->DebugFirstProducer = Pass;
+				}
+
 				if (!bCanUseUAVs && GRenderGraphEmitWarnings)
 				{
 					UE_LOG(LogRendererCore, Warning, TEXT("UAV can only been bound to compute shaders, therefore UAV %s is certainly useless for pass %s."), UAV->Name, Pass->GetName());
@@ -288,7 +298,11 @@ void FRDGBuilder::ValidatePass(const FRenderGraphPass* Pass) const
 					RenderTarget.GetTexture()->Name);
 
 				// TODO(RDG): should only be done when there is a store action.
-				Texture->bHasEverBeenProduced = true;
+				if (!Texture->bHasEverBeenProduced)
+				{
+					Texture->bHasEverBeenProduced = true;
+					Texture->DebugFirstProducer = Pass;
+				}
 			}
 
 			bFoundRTBound = bFoundRTBound || IsBoundAsReadable(RenderTarget.GetTexture(), ParameterStruct);
@@ -309,7 +323,11 @@ void FRDGBuilder::ValidatePass(const FRenderGraphPass* Pass) const
 				TEXT("Can't load stencil from a render target that has never been produced."));
 
 			// TODO(RDG): should only be done when there is a store action.
-			Texture->bHasEverBeenProduced = true;
+			if (!Texture->bHasEverBeenProduced)
+			{
+				Texture->bHasEverBeenProduced = true;
+				Texture->DebugFirstProducer = Pass;
+			}
 		}
 	}
 	else
@@ -474,6 +492,7 @@ void FRDGBuilder::WalkGraphDependencies()
 		{
 			Pair.Value = nullptr;
 			Pair.Key->PooledRenderTarget = nullptr;
+			Pair.Key->CachedRHI.Resource = nullptr;
 		}
 	}
 }
@@ -495,20 +514,7 @@ void FRDGBuilder::AllocateRHITextureIfNeeded(const FRDGTexture* Texture, bool bC
 
 	Texture->PooledRenderTarget = PooledRenderTarget;
 	Texture->CachedRHI.Texture = PooledRenderTarget->GetRenderTargetItem().ShaderResourceTexture;
-}
-
-void FRDGBuilder::AllocateRHITextureSRVIfNeeded(const FRDGTextureSRV* SRV, bool bComputePass)
-{
-	check(SRV);
-
-	if (SRV->CachedRHI.SRV)
-	{
-		return;
-	}
-
-	AllocateRHITextureIfNeeded(SRV->Desc.Texture, bComputePass);
-
-	SRV->CachedRHI.SRV = SRV->Desc.Texture->PooledRenderTarget->GetRenderTargetItem().MipSRVs[SRV->Desc.MipLevel];
+	check(Texture->CachedRHI.Resource);
 }
 
 void FRDGBuilder::AllocateRHITextureUAVIfNeeded(const FRDGTextureUAV* UAV, bool bComputePass)
@@ -525,23 +531,6 @@ void FRDGBuilder::AllocateRHITextureUAVIfNeeded(const FRDGTextureUAV* UAV, bool
 	UAV->CachedRHI.UAV = UAV->Desc.Texture->PooledRenderTarget->GetRenderTargetItem().MipUAVs[UAV->Desc.MipLevel];
 }
 
-void FRDGBuilder::AllocateRHIBufferIfNeeded(const FRDGBuffer* Buffer, bool bComputePass)
-{
-	check(Buffer);
-
-	if (Buffer->PooledBuffer)
-	{
-		return;
-	}
-
-	check(Buffer->ReferenceCount > 0 || GRenderGraphImmediateMode);
-
-	TRefCountPtr<FPooledRDGBuffer>& AllocatedBuffer = AllocatedBuffers.FindOrAdd(Buffer);
-	GRenderGraphResourcePool.FindFreeBuffer(RHICmdList, Buffer->Desc, AllocatedBuffer, Buffer->Name);
-
-	Buffer->PooledBuffer = AllocatedBuffer;
-}
-
 void FRDGBuilder::AllocateRHIBufferSRVIfNeeded(const FRDGBufferSRV* SRV, bool bComputePass)
 {
 	check(SRV);
@@ -550,8 +539,10 @@ void FRDGBuilder::AllocateRHIBufferSRVIfNeeded(const FRDGBufferSRV* SRV, bool bC
 	{
 		return;
 	}
-
-	AllocateRHIBufferIfNeeded(SRV->Desc.Buffer, bComputePass);
+	
+	// The underlying buffer have already been allocated by a prior pass through AllocateRHIBufferUAVIfNeeded().
+	check(SRV->Desc.Buffer->bHasEverBeenProduced);
+	check(SRV->Desc.Buffer->PooledBuffer);
 
 	if (SRV->Desc.Buffer->PooledBuffer->SRVs.Contains(SRV->Desc))
 	{
@@ -586,27 +577,38 @@ void FRDGBuilder::AllocateRHIBufferUAVIfNeeded(const FRDGBufferUAV* UAV, bool bC
 	{
 		return;
 	}
+	
+	FRDGBufferRef Buffer = UAV->Desc.Buffer;
+
+	// Allocate a buffer resource.
+	if (!Buffer->PooledBuffer)
+	{
+		check(Buffer->ReferenceCount > 0 || GRenderGraphImmediateMode);
 
-	AllocateRHIBufferIfNeeded(UAV->Desc.Buffer, bComputePass);
+		TRefCountPtr<FPooledRDGBuffer>& AllocatedBuffer = AllocatedBuffers.FindOrAdd(Buffer);
+		GRenderGraphResourcePool.FindFreeBuffer(RHICmdList, Buffer->Desc, AllocatedBuffer, Buffer->Name);
+
+		Buffer->PooledBuffer = AllocatedBuffer;
+	}
 
-	if (UAV->Desc.Buffer->PooledBuffer->UAVs.Contains(UAV->Desc))
+	if (Buffer->PooledBuffer->UAVs.Contains(UAV->Desc))
 	{
-		UAV->CachedRHI.UAV = UAV->Desc.Buffer->PooledBuffer->UAVs[UAV->Desc];
+		UAV->CachedRHI.UAV = Buffer->PooledBuffer->UAVs[UAV->Desc];
 		return;
 	}
 
 	// Hack to make sure only one UAVs is arround.
-	UAV->Desc.Buffer->PooledBuffer->UAVs.Empty();
+	Buffer->PooledBuffer->UAVs.Empty();
 
 	FUnorderedAccessViewRHIRef RHIUnorderedAccessView;
 
-	if (UAV->Desc.Buffer->Desc.UnderlyingType == FRDGBufferDesc::EUnderlyingType::VertexBuffer)
+	if (Buffer->Desc.UnderlyingType == FRDGBufferDesc::EUnderlyingType::VertexBuffer)
 	{
-		RHIUnorderedAccessView = RHICreateUnorderedAccessView(UAV->Desc.Buffer->PooledBuffer->VertexBuffer, UAV->Desc.Format);
+		RHIUnorderedAccessView = RHICreateUnorderedAccessView(Buffer->PooledBuffer->VertexBuffer, UAV->Desc.Format);
 	}
-	else if (UAV->Desc.Buffer->Desc.UnderlyingType == FRDGBufferDesc::EUnderlyingType::StructuredBuffer)
+	else if (Buffer->Desc.UnderlyingType == FRDGBufferDesc::EUnderlyingType::StructuredBuffer)
 	{
-		RHIUnorderedAccessView = RHICreateUnorderedAccessView(UAV->Desc.Buffer->PooledBuffer->StructuredBuffer, UAV->Desc.bSupportsAtomicCounter, UAV->Desc.bSupportsAppendBuffer);
+		RHIUnorderedAccessView = RHICreateUnorderedAccessView(Buffer->PooledBuffer->StructuredBuffer, UAV->Desc.bSupportsAtomicCounter, UAV->Desc.bSupportsAppendBuffer);
 	}
 	else
 	{
@@ -614,7 +616,7 @@ void FRDGBuilder::AllocateRHIBufferUAVIfNeeded(const FRDGBufferUAV* UAV, bool bC
 	}
 
 	UAV->CachedRHI.UAV = RHIUnorderedAccessView;
-	UAV->Desc.Buffer->PooledBuffer->UAVs.Add(UAV->Desc, RHIUnorderedAccessView);
+	Buffer->PooledBuffer->UAVs.Add(UAV->Desc, RHIUnorderedAccessView);
 }
 
 static EResourceTransitionPipeline CalcTransitionPipeline(bool bCurrentCompute, bool bTargetCompute)
@@ -804,8 +806,17 @@ void FRDGBuilder::AllocateAndTransitionPassResources(const FRenderGraphPass* Pas
 			FRDGTexture* RESTRICT Texture = *ParameterStruct.GetMemberPtrAtOffset<FRDGTexture*>(Offset);
 			if (Texture)
 			{
-				AllocateRHITextureIfNeeded(Texture, bIsCompute);
+				// The underlying texture have already been allocated by a prior pass.
+				check(Texture->bHasEverBeenProduced);
+				check(Texture->PooledRenderTarget);
+				check(Texture->CachedRHI.Resource);
 				TransitionTexture(Texture, EResourceTransitionAccess::EReadable, bIsCompute);
+
+				#if RENDER_GRAPH_DEBUGGING
+				{
+					Texture->DebugPassAccessCount++;
+				}
+				#endif
 			}
 		}
 		break;
@@ -814,10 +825,24 @@ void FRDGBuilder::AllocateAndTransitionPassResources(const FRenderGraphPass* Pas
 			FRDGTextureSRV* RESTRICT SRV = *ParameterStruct.GetMemberPtrAtOffset<FRDGTextureSRV*>(Offset);
 			if (SRV)
 			{
+				// The underlying texture have already been allocated by a prior pass.
 				check(SRV->Desc.Texture);
+				check(SRV->Desc.Texture->bHasEverBeenProduced);
+				check(SRV->Desc.Texture->PooledRenderTarget);
+
+				// Might be the first time using this render graph SRV, so need to setup the cached rhi resource.
+				if (!SRV->CachedRHI.SRV)
+				{
+					SRV->CachedRHI.SRV = SRV->Desc.Texture->PooledRenderTarget->GetRenderTargetItem().MipSRVs[SRV->Desc.MipLevel];
+				}
 
-				AllocateRHITextureSRVIfNeeded(SRV, bIsCompute);
 				TransitionTexture(SRV->Desc.Texture, EResourceTransitionAccess::EReadable, bIsCompute);
+
+				#if RENDER_GRAPH_DEBUGGING
+				{
+					SRV->Desc.Texture->DebugPassAccessCount++;
+				}
+				#endif
 			}
 		}
 		break;
@@ -828,6 +853,12 @@ void FRDGBuilder::AllocateAndTransitionPassResources(const FRenderGraphPass* Pas
 			{
 				AllocateRHITextureUAVIfNeeded(UAV, bIsCompute);
 				TransitionUAV(UAV->CachedRHI.UAV, UAV->Desc.Texture, EResourceTransitionAccess::EWritable, bIsCompute);
+
+				#if RENDER_GRAPH_DEBUGGING
+				{
+					UAV->Desc.Texture->DebugPassAccessCount++;
+				}
+				#endif
 			}
 		}
 		break;
@@ -836,12 +867,20 @@ void FRDGBuilder::AllocateAndTransitionPassResources(const FRenderGraphPass* Pas
 			FRDGBuffer* RESTRICT Buffer = *ParameterStruct.GetMemberPtrAtOffset<FRDGBuffer*>(Offset);
 			if (Buffer)
 			{
-				AllocateRHIBufferIfNeeded(Buffer, bIsCompute);
+				// The underlying buffer have already been allocated by a prior pass through AllocateRHIBufferUAVIfNeeded().
+				check(Buffer->bHasEverBeenProduced);
+				check(Buffer->PooledBuffer);
 
 				// TODO(RDG): supper hacky, find the UAV and transition it. Hopefully there is one...
 				check(Buffer->PooledBuffer->UAVs.Num() == 1);
 				FUnorderedAccessViewRHIParamRef UAV = Buffer->PooledBuffer->UAVs.CreateIterator().Value();
 				TransitionUAV(UAV, Buffer, EResourceTransitionAccess::EReadable, bIsCompute);
+
+				#if RENDER_GRAPH_DEBUGGING
+				{
+					Buffer->DebugPassAccessCount++;
+				}
+				#endif
 			}
 		}
 		break;
@@ -850,12 +889,23 @@ void FRDGBuilder::AllocateAndTransitionPassResources(const FRenderGraphPass* Pas
 			FRDGBufferSRV* RESTRICT SRV = *ParameterStruct.GetMemberPtrAtOffset<FRDGBufferSRV*>(Offset);
 			if (SRV)
 			{
+				// The underlying buffer have already been allocated by a prior pass through AllocateRHIBufferUAVIfNeeded().
+				check(SRV->Desc.Buffer);
+				check(SRV->Desc.Buffer->bHasEverBeenProduced);
+				check(SRV->Desc.Buffer->PooledBuffer);
+				
 				AllocateRHIBufferSRVIfNeeded(SRV, bIsCompute);
 
 				// TODO(RDG): supper hacky, find the UAV and transition it. Hopefully there is one...
 				check(SRV->Desc.Buffer->PooledBuffer->UAVs.Num() == 1);
 				FUnorderedAccessViewRHIParamRef UAV = SRV->Desc.Buffer->PooledBuffer->UAVs.CreateIterator().Value();
 				TransitionUAV(UAV, SRV->Desc.Buffer, EResourceTransitionAccess::EReadable, bIsCompute);
+
+				#if RENDER_GRAPH_DEBUGGING
+				{
+					SRV->Desc.Buffer->DebugPassAccessCount++;
+				}
+				#endif
 			}
 		}
 		break;
@@ -866,6 +916,12 @@ void FRDGBuilder::AllocateAndTransitionPassResources(const FRenderGraphPass* Pas
 			{
 				AllocateRHIBufferUAVIfNeeded(UAV, bIsCompute);
 				TransitionUAV(UAV->CachedRHI.UAV, UAV->Desc.Buffer, EResourceTransitionAccess::EWritable, bIsCompute);
+
+				#if RENDER_GRAPH_DEBUGGING
+				{
+					UAV->Desc.Buffer->DebugPassAccessCount++;
+				}
+				#endif
 			}
 		}
 		break;
@@ -901,6 +957,12 @@ void FRDGBuilder::AllocateAndTransitionPassResources(const FRenderGraphPass* Pas
 
 					NumSamples |= OutRPInfo->ColorRenderTargets[i].RenderTarget->GetNumSamples();
 					NumRenderTargets++;
+					
+					#if RENDER_GRAPH_DEBUGGING
+					{
+						RenderTarget.GetTexture()->DebugPassAccessCount++;
+					}
+					#endif
 				}
 				else
 				{
@@ -924,6 +986,12 @@ void FRDGBuilder::AllocateAndTransitionPassResources(const FRenderGraphPass* Pas
 
 				NumSamples |= OutRPInfo->DepthStencilRenderTarget.DepthStencilTarget->GetNumSamples();
 				NumDepthStencilTargets++;
+					
+				#if RENDER_GRAPH_DEBUGGING
+				{
+					DepthStencil.Texture->DebugPassAccessCount++;
+				}
+				#endif
 			}
 
 			OutRPInfo->bIsMSAA = NumSamples > 1;
@@ -1018,6 +1086,7 @@ void FRDGBuilder::ReleaseRHITextureIfPossible(const FRDGTexture* Texture)
 	if (Texture->ReferenceCount == 0)
 	{
 		Texture->PooledRenderTarget = nullptr;
+		Texture->CachedRHI.Resource = nullptr;
 		AllocatedTextures.FindChecked(Texture) = nullptr;
 	}
 }
@@ -1030,6 +1099,7 @@ void FRDGBuilder::ReleaseRHIBufferIfPossible(const FRDGBuffer* Buffer)
 	if (Buffer->ReferenceCount == 0)
 	{
 		Buffer->PooledBuffer = nullptr;
+		Buffer->CachedRHI.Resource = nullptr;
 		AllocatedBuffers.FindChecked(Buffer) = nullptr;
 	}
 }
@@ -1141,6 +1211,14 @@ void FRDGBuilder::ProcessDeferredInternalResourceQueries()
 		}
 
 		*Query.OutTexturePtr = AllocatedTextures.FindChecked(Query.Texture);
+		
+		#if RENDER_GRAPH_DEBUGGING
+		{
+			// Increment the number of time the texture has been accessed to avoid warning on produced but never used resources that were produced
+			// only to be extracted for the graph.
+			Query.Texture->DebugPassAccessCount += 1;
+		}
+		#endif
 
 		// No need to manually release in immediate mode, since it is done directly when emptying AllocatedTextures in DestructPasses().
 		if (!GRenderGraphImmediateMode)
@@ -1165,10 +1243,20 @@ void FRDGBuilder::DestructPasses()
 
 	#if RENDER_GRAPH_DEBUGGING
 	{
-		// Make sure all resource references have been released to ensure no leaks happen.
+		// Make sure all resource references have been released to ensure no leaks happen,
+		// and emit warning if produced resource has not been used.
 		for (const FRDGResource* Resource : Resources)
 		{
 			check(Resource->ReferenceCount == 0);
+
+			if (GRenderGraphEmitWarnings && Resource->DebugPassAccessCount == 1 && Resource->DebugFirstProducer)
+			{
+				check(Resource->bHasEverBeenProduced);
+
+				UE_LOG(LogRendererCore, Warning,
+					TEXT("Resources %s has been produced by the pass %s, but never used by another pass."),
+					Resource->Name, Resource->DebugFirstProducer->GetName());
+			}
 		}
 		Resources.Empty();
 	}
diff --git a/Engine/Source/Runtime/RenderCore/Public/RenderGraphBuilder.h b/Engine/Source/Runtime/RenderCore/Public/RenderGraphBuilder.h
--- a/Engine/Source/Runtime/RenderCore/Public/RenderGraphBuilder.h
+++ b/Engine/Source/Runtime/RenderCore/Public/RenderGraphBuilder.h
@@ -215,6 +215,7 @@ public:
 		#endif
 		FRDGTexture* OutTexture = AllocateForRHILifeTime<FRDGTexture>(Name, ExternalPooledTexture->GetDesc());
 		OutTexture->PooledRenderTarget = ExternalPooledTexture;
+		OutTexture->CachedRHI.Texture = ExternalPooledTexture->GetRenderTargetItem().ShaderResourceTexture;
 		AllocatedTextures.Add(OutTexture, ExternalPooledTexture);
 		#if RENDER_GRAPH_DEBUGGING
 		{
@@ -395,6 +396,23 @@ public:
 		DeferredInternalTextureQueries.Emplace(Query);
 	}
 
+	/** Flag a texture that is only produced by only 1 pass, but never used or extracted, to avoid generating a warning at runtime. */
+	FORCEINLINE_DEBUGGABLE void RemoveUnusedTextureWarning(FRDGTextureRef Texture)
+	{
+		check(Texture);
+		#if RENDER_GRAPH_DEBUGGING
+		{
+			checkf(!bHasExecuted,
+				TEXT("Flaging texture %s with FlagUnusedTexture() needs to happen before the builder's execution."),
+				Texture->Name);
+			
+			// Increment the number of time the texture has been accessed to avoid warning on produced but never used resources that were produced
+			// only to be extracted for the graph.
+			Texture->DebugPassAccessCount += 1;
+		}
+		#endif
+	}
+
 	/** 
 	 * Executes the queued passes, managing setting of render targets (RHI RenderPasses), resource transitions and queued texture extraction.
 	 */
@@ -470,9 +488,7 @@ private:
 	}
 
 	void AllocateRHITextureIfNeeded(const FRDGTexture* Texture, bool bComputePass);
-	void AllocateRHITextureSRVIfNeeded(const FRDGTextureSRV* SRV, bool bComputePass);
 	void AllocateRHITextureUAVIfNeeded(const FRDGTextureUAV* UAV, bool bComputePass);
-	void AllocateRHIBufferIfNeeded(const FRDGBuffer* Texture, bool bComputePass);
 	void AllocateRHIBufferSRVIfNeeded(const FRDGBufferSRV* SRV, bool bComputePass);
 	void AllocateRHIBufferUAVIfNeeded(const FRDGBufferUAV* UAV, bool bComputePass);
 
diff --git a/Engine/Source/Runtime/RenderCore/Public/RenderGraphResources.h b/Engine/Source/Runtime/RenderCore/Public/RenderGraphResources.h
--- a/Engine/Source/Runtime/RenderCore/Public/RenderGraphResources.h
+++ b/Engine/Source/Runtime/RenderCore/Public/RenderGraphResources.h
@@ -10,6 +10,7 @@
 /** Whether render graph debugging is compiled. */
 #define RENDER_GRAPH_DEBUGGING (!UE_BUILD_SHIPPING)
 
+struct FRenderGraphPass;
 
 class FRDGResource;
 class FRDGTexture;
@@ -75,6 +76,12 @@ private:
 #if RENDER_GRAPH_DEBUGGING
 	/** Boolean to track at wiring time if a resource has ever been produced by a pass, to error out early if accessing a resource that has not been produced. */
 	mutable bool bHasEverBeenProduced = false;
+
+	/** Pointer towards the pass that is the first to produce it, for even more convenient error message. */
+	mutable const FRenderGraphPass* DebugFirstProducer = nullptr;
+
+	/** Count the number of times it has been used by a pass. */
+	mutable int32 DebugPassAccessCount = 0;
 #endif
 
 	friend class FRDGBuilder;
diff --git a/Engine/Source/Runtime/Renderer/Private/ScreenSpaceDenoise.cpp b/Engine/Source/Runtime/Renderer/Private/ScreenSpaceDenoise.cpp
--- a/Engine/Source/Runtime/Renderer/Private/ScreenSpaceDenoise.cpp
+++ b/Engine/Source/Runtime/Renderer/Private/ScreenSpaceDenoise.cpp
@@ -788,7 +788,7 @@ IScreenSpaceDenoiser::FReflectionOutputs DenoiseReflections(
 		// The SignalHistory1 is always generated for temporal history, but will endup useless if there is no view state,
 		// in witch case we do not extract any textures. Don't support a shader permutation that does not produce it, because
 		// it is already a not ideal case for the denoiser.
-		//GraphBuilder.FlagUsedOnceTexture(SignalHistory1);
+		GraphBuilder.RemoveUnusedTextureWarning(SignalHistory1);
 	}
 
 	IScreenSpaceDenoiser::FReflectionOutputs Outputs;
