Ensures render graph emitted warnings are only emitted once.

diff --git a/Engine/Source/Runtime/RenderCore/Private/RenderGraphBuilder.cpp b/Engine/Source/Runtime/RenderCore/Private/RenderGraphBuilder.cpp
--- a/Engine/Source/Runtime/RenderCore/Private/RenderGraphBuilder.cpp
+++ b/Engine/Source/Runtime/RenderCore/Private/RenderGraphBuilder.cpp
@@ -41,6 +41,20 @@ void InitRenderGraph()
 #endif
 }
 
+static void EmitRenderGraphWarning(const FString& WarningMessage)
+{
+	check(GRenderGraphEmitWarnings);
+	
+	static TSet<FString> GAlreadyEmittedWarnings;
+	if (!GAlreadyEmittedWarnings.Contains(WarningMessage))
+	{
+		GAlreadyEmittedWarnings.Add(WarningMessage);
+		UE_LOG(LogRendererCore, Warning, TEXT("%s"), *WarningMessage);
+	}
+}
+
+#define EmitRenderGraphWarningf(WarningMessageFormat, ...) \
+	EmitRenderGraphWarning(FString::Printf(WarningMessageFormat, ##__VA_ARGS__));
 
 static bool IsBoundAsReadable(const FRDGTexture* Texture, FShaderParameterStructRef ParameterStruct)
 {
@@ -213,7 +227,9 @@ void FRDGBuilder::ValidatePass(const FRenderGraphPass* Pass) const
 
 				if (!bCanUseUAVs && GRenderGraphEmitWarnings)
 				{
-					UE_LOG(LogRendererCore, Warning, TEXT("UAV can only been bound to compute shaders, therefore UAV %s is certainly useless for pass %s."), UAV->Name, Pass->GetName());
+					EmitRenderGraphWarningf(
+						TEXT("UAV can only been bound to compute shaders, therefore UAV %s is certainly useless for pass %s."),
+						UAV->Name, Pass->GetName());
 				}
 			}
 		}
@@ -250,7 +266,9 @@ void FRDGBuilder::ValidatePass(const FRenderGraphPass* Pass) const
 
 				if (!bCanUseUAVs && GRenderGraphEmitWarnings)
 				{
-					UE_LOG(LogRendererCore, Warning, TEXT("UAV can only been bound to compute shaders, therefore UAV %s is certainly useless for pass %s."), UAV->Name, Pass->GetName());
+					EmitRenderGraphWarningf(
+						TEXT("UAV can only been bound to compute shaders, therefore UAV %s is certainly useless for pass %s."),
+						UAV->Name, Pass->GetName());
 				}
 			}
 		}
@@ -263,7 +281,9 @@ void FRDGBuilder::ValidatePass(const FRenderGraphPass* Pass) const
 			}
 			else if (GRenderGraphEmitWarnings)
 			{
-				UE_LOG(LogRendererCore, Warning, TEXT("Pass %s have duplicated render target binding slots."), Pass->GetName());
+				EmitRenderGraphWarningf(
+					TEXT("Pass %s have duplicated render target binding slots."),
+					Pass->GetName());
 			}
 		}
 		break;
@@ -1040,7 +1060,9 @@ void FRDGBuilder::WarnForUselessPassDependencies(const FRenderGraphPass* Pass)
 
 	if (TrackedResourceCount != UsedResourceCount)
 	{
-		UE_LOG(LogRendererCore, Warning, TEXT("%i of the %i resources of the pass %s where not actually used."), TrackedResourceCount - UsedResourceCount, TrackedResourceCount, Pass->GetName());
+		FString WarningMessage = FString::Printf(
+			TEXT("%i of the %i resources of the pass %s where not actually used."),
+			TrackedResourceCount - UsedResourceCount, TrackedResourceCount, Pass->GetName());
 
 		for (int ResourceIndex = 0, Num = ParameterStruct.Layout->Resources.Num(); ResourceIndex < Num; ResourceIndex++)
 		{
@@ -1056,8 +1078,12 @@ void FRDGBuilder::WarnForUselessPassDependencies(const FRenderGraphPass* Pass)
 				continue;
 
 			if (!Resource->bIsActuallyUsedByPass)
-				UE_LOG(LogRendererCore, Warning, TEXT("	%s"), Resource->Name);
+			{
+				WarningMessage += FString::Printf(TEXT("\n    %s"), Resource->Name);
+			}
 		}
+
+		EmitRenderGraphWarning(WarningMessage);
 	}
 
 	// Last pass to clean the bIsActuallyUsedByPass flags.
@@ -1253,7 +1279,7 @@ void FRDGBuilder::DestructPasses()
 			{
 				check(Resource->bHasEverBeenProduced);
 
-				UE_LOG(LogRendererCore, Warning,
+				EmitRenderGraphWarningf(
 					TEXT("Resources %s has been produced by the pass %s, but never used by another pass."),
 					Resource->Name, Resource->DebugFirstProducer->GetName());
 			}
