Enable GPUScene shaders on other platforms.

diff --git a/Engine/Source/Runtime/RenderCore/Public/RenderUtils.h b/Engine/Source/Runtime/RenderCore/Public/RenderUtils.h
--- a/Engine/Source/Runtime/RenderCore/Public/RenderUtils.h
+++ b/Engine/Source/Runtime/RenderCore/Public/RenderUtils.h
@@ -475,9 +475,7 @@ inline bool UseGPUScene(EShaderPlatform Platform, ERHIFeatureLevel::Type Feature
 	return FeatureLevel >= ERHIFeatureLevel::SM5 
 		//@todo - support GPU Scene management compute shaders on these platforms to get dynamic instancing speedups on the Rendering Thread and RHI Thread
 		&& Platform != SP_OPENGL_SM5
-		&& Platform != SP_VULKAN_SM5
-		&& Platform != SP_SWITCH
-		&& Platform != SP_SWITCH_FORWARD;
+		&& Platform != SP_VULKAN_SM5;
 }
 
 /** Unit cube vertex buffer (VertexDeclarationFVector4) */
