Fixed FShaderParameterInfo sizes being too small

diff --git a/Engine/Plugins/FX/Niagara/Source/NiagaraShader/Private/NiagaraShaderDerivedDataVersion.h b/Engine/Plugins/FX/Niagara/Source/NiagaraShader/Private/NiagaraShaderDerivedDataVersion.h
--- a/Engine/Plugins/FX/Niagara/Source/NiagaraShader/Private/NiagaraShaderDerivedDataVersion.h
+++ b/Engine/Plugins/FX/Niagara/Source/NiagaraShader/Private/NiagaraShaderDerivedDataVersion.h
@@ -11,5 +11,5 @@ NiagaraShaderDerivedDataVersion.h: Shader derived data version for Niagara.
 // In case of merge conflicts with DDC versions, you MUST generate a new GUID and set this new
 // guid as version
 
-#define NIAGARASHADERMAP_DERIVEDDATA_VER		TEXT("183CB3659FE7414CB5DC9CB93D9A13D7")
+#define NIAGARASHADERMAP_DERIVEDDATA_VER		TEXT("184CB3659FE7414CB5DC9CB93D9A13D7")
 													  
diff --git a/Engine/Shaders/Public/ShaderVersion.ush b/Engine/Shaders/Public/ShaderVersion.ush
--- a/Engine/Shaders/Public/ShaderVersion.ush
+++ b/Engine/Shaders/Public/ShaderVersion.ush
@@ -3,4 +3,4 @@
 // in Platform.ush (which should be included in any shader) it allows to invalidate the shader DDC.
 // 
 // If you are merging streams and there is a conflict with this GUID you should make a new GUID rather than taking one or the other.
-// GUID = 2ED93466-374C-4558-B573-BC368C3718D8
+// GUID = 3ED98468-474C-4558-B573-BC368C3718D8
diff --git a/Engine/Source/Runtime/RenderCore/Private/Shader.cpp b/Engine/Source/Runtime/RenderCore/Private/Shader.cpp
--- a/Engine/Source/Runtime/RenderCore/Private/Shader.cpp
+++ b/Engine/Source/Runtime/RenderCore/Private/Shader.cpp
@@ -527,9 +527,7 @@ void FShaderResource::BuildParameterMapInfo(const TMap<FString, FParameterAlloca
 
 						if (LooseParameterBufferInfo.BufferIndex == ParamValue.BufferIndex)
 						{
-							FShaderParameterInfo ParameterInfo;
-							ParameterInfo.BaseIndex = ParamValue.BaseIndex;
-							ParameterInfo.Size = ParamValue.Size;
+							FShaderParameterInfo ParameterInfo(ParamValue.BaseIndex, ParamValue.Size);
 							LooseParameterBufferInfo.Parameters.Add(ParameterInfo);
 							LooseParameterBufferInfo.BufferSize += ParamValue.Size;
 							bAddedToExistingBuffer = true;
@@ -538,13 +536,9 @@ void FShaderResource::BuildParameterMapInfo(const TMap<FString, FParameterAlloca
 
 					if (!bAddedToExistingBuffer)
 					{
-						FShaderLooseParameterBufferInfo NewParameterBufferInfo;
-						NewParameterBufferInfo.BufferIndex = ParamValue.BufferIndex;
-						NewParameterBufferInfo.BufferSize = ParamValue.Size;
+						FShaderLooseParameterBufferInfo NewParameterBufferInfo(ParamValue.BufferIndex, ParamValue.Size);
 
-						FShaderParameterInfo ParameterInfo;
-						ParameterInfo.BaseIndex = ParamValue.BaseIndex;
-						ParameterInfo.Size = ParamValue.Size;
+						FShaderParameterInfo ParameterInfo(ParamValue.BaseIndex, ParamValue.Size);
 						NewParameterBufferInfo.Parameters.Add(ParameterInfo);
 
 						ParameterMapInfo.LooseParameterBuffers.Add(NewParameterBufferInfo);
@@ -589,9 +583,8 @@ void FShaderResource::BuildParameterMapInfo(const TMap<FString, FParameterAlloca
 
 				if (ParamValue.Type == CurrentParameterType)
 				{
-					FShaderParameterInfo ParameterInfo;
-					ParameterInfo.BaseIndex = CurrentParameterType == EShaderParameterType::UniformBuffer ? ParamValue.BufferIndex : ParamValue.BaseIndex;
-					ParameterInfo.Size = ParamValue.Size;
+					const uint16 BaseIndex = CurrentParameterType == EShaderParameterType::UniformBuffer ? ParamValue.BufferIndex : ParamValue.BaseIndex;
+					FShaderParameterInfo ParameterInfo(BaseIndex, ParamValue.Size);
 					ParameterInfoArray->Add(ParameterInfo);
 				}
 			}
diff --git a/Engine/Source/Runtime/RenderCore/Public/Shader.h b/Engine/Source/Runtime/RenderCore/Public/Shader.h
--- a/Engine/Source/Runtime/RenderCore/Public/Shader.h
+++ b/Engine/Source/Runtime/RenderCore/Public/Shader.h
@@ -180,8 +180,17 @@ public:
 class FShaderParameterInfo
 {
 public:
-	uint8 BaseIndex;
-	uint8 Size;
+	uint16 BaseIndex;
+	uint16 Size;
+
+	FShaderParameterInfo() {}
+
+	FShaderParameterInfo(uint16 InBaseIndex, uint16 InSize)
+	{
+		BaseIndex = InBaseIndex;
+		Size = InSize;
+		checkf(BaseIndex == InBaseIndex && Size == InSize, TEXT("Tweak FShaderParameterInfo type sizes"));
+	}
 
 	friend FArchive& operator<<(FArchive& Ar,FShaderParameterInfo& Info)
 	{
@@ -200,10 +209,19 @@ public:
 class FShaderLooseParameterBufferInfo
 {
 public:
-	uint8 BufferIndex;
+	uint16 BufferIndex;
 	uint16 BufferSize;
 	TArray<FShaderParameterInfo> Parameters;
 
+	FShaderLooseParameterBufferInfo() {}
+
+	FShaderLooseParameterBufferInfo(uint16 InBufferIndex, uint16 InBufferSize)
+	{
+		BufferIndex = InBufferIndex;
+		BufferSize = InBufferSize;
+		checkf(BufferIndex == InBufferIndex, TEXT("Tweak FShaderLooseParameterBufferInfo type sizes"));
+	}
+
 	friend FArchive& operator<<(FArchive& Ar,FShaderLooseParameterBufferInfo& Info)
 	{
 		Ar << Info.BufferIndex;
