Always add contents of ShaderVersion.ush to the shader virtual source paths (shadermap DDC key), so if shader forgets to include it, it will still won't break DDC.

diff --git a/Engine/Source/Runtime/RenderCore/Private/ShaderCore.cpp b/Engine/Source/Runtime/RenderCore/Private/ShaderCore.cpp
--- a/Engine/Source/Runtime/RenderCore/Private/ShaderCore.cpp
+++ b/Engine/Source/Runtime/RenderCore/Private/ShaderCore.cpp
@@ -423,7 +423,8 @@ static void GetAllVirtualShaderSourcePaths(TArray<FString>& OutVirtualFilePaths,
 
 	//#todo-rco: No need to loop through Shader Pipeline Types (yet)
 
-	// also always add the MaterialTemplate.usf shader file
+	// Always add ShaderVersion.ush, so if shader forgets to include it, it will still won't break DDC.
+	AddShaderSourceFileEntry(OutVirtualFilePaths, FString(TEXT("/Engine/Public/ShaderVersion.ush")), ShaderPlatform);
 	AddShaderSourceFileEntry(OutVirtualFilePaths, FString(TEXT("/Engine/Private/MaterialTemplate.ush")), ShaderPlatform);
 	AddShaderSourceFileEntry(OutVirtualFilePaths, FString(TEXT("/Engine/Private/Common.ush")), ShaderPlatform);
 	AddShaderSourceFileEntry(OutVirtualFilePaths, FString(TEXT("/Engine/Private/Definitions.usf")), ShaderPlatform);
