Added missing FRDGBuilder forward declaration. This fixes Clang compile errors.

diff --git a/Engine/Source/Runtime/RenderCore/Public/RenderGraphResources.h b/Engine/Source/Runtime/RenderCore/Public/RenderGraphResources.h
--- a/Engine/Source/Runtime/RenderCore/Public/RenderGraphResources.h
+++ b/Engine/Source/Runtime/RenderCore/Public/RenderGraphResources.h
@@ -12,6 +12,8 @@
 
 struct FRenderGraphPass;
 
+class FRDGBuilder;
+
 class FRDGResource;
 class FRDGTexture;
 class FRDGTextureSRV;
