DR - Disable GPU scene for Vulkan
#rnx

diff --git a/Engine/Source/Runtime/RenderCore/Public/RenderUtils.h b/Engine/Source/Runtime/RenderCore/Public/RenderUtils.h
--- a/Engine/Source/Runtime/RenderCore/Public/RenderUtils.h
+++ b/Engine/Source/Runtime/RenderCore/Public/RenderUtils.h
@@ -474,8 +474,8 @@ inline bool UseGPUScene(EShaderPlatform Platform, ERHIFeatureLevel::Type Feature
 	// GPU Scene management uses compute shaders
 	return FeatureLevel >= ERHIFeatureLevel::SM5 
 		//@todo - support GPU Scene management compute shaders on these platforms to get dynamic instancing speedups on the Rendering Thread and RHI Thread
-		&& Platform != SP_OPENGL_SM5
-		&& Platform != SP_VULKAN_SM5;
+		&& !IsOpenGLPlatform(Platform)
+		&& !IsVulkanPlatform(Platform);
 }
 
 /** Unit cube vertex buffer (VertexDeclarationFVector4) */
