Fixes compilation failures in render graph.

diff --git a/Engine/Source/Runtime/RenderCore/Private/RenderGraphBuilder.cpp b/Engine/Source/Runtime/RenderCore/Private/RenderGraphBuilder.cpp
--- a/Engine/Source/Runtime/RenderCore/Private/RenderGraphBuilder.cpp
+++ b/Engine/Source/Runtime/RenderCore/Private/RenderGraphBuilder.cpp
@@ -362,7 +362,7 @@ void FRDGBuilder::ValidatePass(const FRenderGraphPass* Pass) const
 	{
 		checkf(!bRequiresRenderTargetSlots, TEXT("Render pass %s requires render target binging slots"), Pass->GetName());
 	}
-#endif
+#endif // RENDER_GRAPH_DEBUGGING
 }
 
 void FRDGBuilder::CaptureAnyInterestingPassOutput(const FRenderGraphPass* Pass)
@@ -569,7 +569,11 @@ void FRDGBuilder::AllocateRHIBufferSRVIfNeeded(const FRDGBufferSRV* SRV, bool bC
 	}
 	
 	// The underlying buffer have already been allocated by a prior pass through AllocateRHIBufferUAVIfNeeded().
-	check(SRV->Desc.Buffer->bHasEverBeenProduced);
+	#if RENDER_GRAPH_DEBUGGING
+	{
+		check(SRV->Desc.Buffer->bHasEverBeenProduced);
+	}	
+	#endif
 	check(SRV->Desc.Buffer->PooledBuffer);
 
 	if (SRV->Desc.Buffer->PooledBuffer->SRVs.Contains(SRV->Desc))
@@ -835,7 +839,11 @@ void FRDGBuilder::AllocateAndTransitionPassResources(const FRenderGraphPass* Pas
 			if (Texture)
 			{
 				// The underlying texture have already been allocated by a prior pass.
-				check(Texture->bHasEverBeenProduced);
+				#if RENDER_GRAPH_DEBUGGING
+				{
+					check(Texture->bHasEverBeenProduced);
+				}	
+				#endif
 				check(Texture->PooledRenderTarget);
 				check(Texture->CachedRHI.Resource);
 				TransitionTexture(Texture, EResourceTransitionAccess::EReadable, bIsCompute);
@@ -855,7 +863,11 @@ void FRDGBuilder::AllocateAndTransitionPassResources(const FRenderGraphPass* Pas
 			{
 				// The underlying texture have already been allocated by a prior pass.
 				check(SRV->Desc.Texture);
-				check(SRV->Desc.Texture->bHasEverBeenProduced);
+				#if RENDER_GRAPH_DEBUGGING
+				{
+					check(SRV->Desc.Texture->bHasEverBeenProduced);
+				}	
+				#endif
 				check(SRV->Desc.Texture->PooledRenderTarget);
 
 				// Might be the first time using this render graph SRV, so need to setup the cached rhi resource.
@@ -896,7 +908,11 @@ void FRDGBuilder::AllocateAndTransitionPassResources(const FRenderGraphPass* Pas
 			if (Buffer)
 			{
 				// The underlying buffer have already been allocated by a prior pass through AllocateRHIBufferUAVIfNeeded().
-				check(Buffer->bHasEverBeenProduced);
+				#if RENDER_GRAPH_DEBUGGING
+				{
+					check(Buffer->bHasEverBeenProduced);
+				}	
+				#endif
 				check(Buffer->PooledBuffer);
 
 				// TODO(RDG): supper hacky, find the UAV and transition it. Hopefully there is one...
@@ -919,7 +935,11 @@ void FRDGBuilder::AllocateAndTransitionPassResources(const FRenderGraphPass* Pas
 			{
 				// The underlying buffer have already been allocated by a prior pass through AllocateRHIBufferUAVIfNeeded().
 				check(SRV->Desc.Buffer);
-				check(SRV->Desc.Buffer->bHasEverBeenProduced);
+				#if RENDER_GRAPH_DEBUGGING
+				{
+					check(SRV->Desc.Buffer->bHasEverBeenProduced);
+				}	
+				#endif
 				check(SRV->Desc.Buffer->PooledBuffer);
 				
 				AllocateRHIBufferSRVIfNeeded(SRV, bIsCompute);
