Makes the errors checking render graph executions more precise.

diff --git a/Engine/Source/Runtime/RenderCore/Private/RenderGraphBuilder.cpp b/Engine/Source/Runtime/RenderCore/Private/RenderGraphBuilder.cpp
--- a/Engine/Source/Runtime/RenderCore/Private/RenderGraphBuilder.cpp
+++ b/Engine/Source/Runtime/RenderCore/Private/RenderGraphBuilder.cpp
@@ -123,9 +123,20 @@ void FRDGBuilder::Execute()
 {
 	#if RENDER_GRAPH_DEBUGGING
 	{
+		/** The usage need RDG_EVENT_SCOPE() needs to happen in inner scope of the one containing FRDGBuilder because of
+		 *  FStackRDGEventScopeRef's destructor modifying this FRDGBuilder instance.
+		 * 
+		 *
+		 *  FRDGBuilder GraphBuilder(RHICmdList);
+		 *  {
+		 *  	RDG_EVENT_SCOPE(GraphBuilder, "MyEventScope");
+		 *  	// ...
+		 *  }
+		 *  GraphBuilder.Execute();
+		 */
 		checkf(CurrentScope == nullptr, TEXT("Render graph needs to have all scopes ended to execute."));
 
-		checkf(!bHasExecuted, TEXT("Render graph execution should only happen once."));
+		checkf(!bHasExecuted, TEXT("Render graph execution should only happen once to ensure consistency with immediate mode."));
 	}
 	#endif
 
@@ -1325,4 +1336,13 @@ void FRDGBuilder::DestructPasses()
 	Passes.Empty();
 	DeferredInternalTextureQueries.Empty();
 	AllocatedTextures.Empty();
-}
\ No newline at end of file
+}
+
+FRDGBuilder::~FRDGBuilder()
+{
+	#if RENDER_GRAPH_DEBUGGING
+	{
+		checkf(bHasExecuted, TEXT("Render graph execution si required to ensure consistency with immediate mode."));
+	}
+	#endif
+}
diff --git a/Engine/Source/Runtime/RenderCore/Public/RenderGraphBuilder.h b/Engine/Source/Runtime/RenderCore/Public/RenderGraphBuilder.h
--- a/Engine/Source/Runtime/RenderCore/Public/RenderGraphBuilder.h
+++ b/Engine/Source/Runtime/RenderCore/Public/RenderGraphBuilder.h
@@ -200,10 +200,7 @@ public:
 			ScopesStack[i] = nullptr;
 	}
 
-	~FRDGBuilder()
-	{
-		DestructPasses();
-	}
+	~FRDGBuilder();
 
 	/** Register a external texture to be tracked by the render graph. */
 	inline FRDGTextureRef RegisterExternalTexture(const TRefCountPtr<IPooledRenderTarget>& ExternalPooledTexture, const TCHAR* Name = TEXT("External"))
