DR - More removal of old macros
#rnx

diff --git a/Engine/Source/Runtime/Engine/Private/Components/DirectionalLightComponent.cpp b/Engine/Source/Runtime/Engine/Private/Components/DirectionalLightComponent.cpp
--- a/Engine/Source/Runtime/Engine/Private/Components/DirectionalLightComponent.cpp
+++ b/Engine/Source/Runtime/Engine/Private/Components/DirectionalLightComponent.cpp
@@ -153,13 +153,12 @@ public:
 	{
 		FVector NewLightShaftOverrideDirection = Component->LightShaftOverrideDirection;
 		NewLightShaftOverrideDirection.Normalize();
-		ENQUEUE_UNIQUE_RENDER_COMMAND_TWOPARAMETER(
-			FUpdateLightShaftOverrideDirectionCommand,
-			FDirectionalLightSceneProxy*,Proxy,this,
-			FVector,NewLightShaftOverrideDirection,NewLightShaftOverrideDirection,
-		{
-			Proxy->UpdateLightShaftOverrideDirection_RenderThread(NewLightShaftOverrideDirection);
-		});
+		FDirectionalLightSceneProxy* Proxy = this;
+		ENQUEUE_RENDER_COMMAND(FUpdateLightShaftOverrideDirectionCommand)(
+			[Proxy, NewLightShaftOverrideDirection](FRHICommandList& RHICmdList)
+			{
+				Proxy->UpdateLightShaftOverrideDirection_RenderThread(NewLightShaftOverrideDirection);
+			});
 	}
 
 	/** Accesses parameters needed for rendering the light. */
diff --git a/Engine/Source/Runtime/Engine/Private/Components/LightComponent.cpp b/Engine/Source/Runtime/Engine/Private/Components/LightComponent.cpp
--- a/Engine/Source/Runtime/Engine/Private/Components/LightComponent.cpp
+++ b/Engine/Source/Runtime/Engine/Private/Components/LightComponent.cpp
@@ -1168,14 +1168,12 @@ void ULightComponent::InitializeStaticShadowDepthMap()
 			DepthMapData = &MapBuildData->DepthMap;
 		}
 
-		ENQUEUE_UNIQUE_RENDER_COMMAND_TWOPARAMETER(
-			SetDepthMapData,
-			FStaticShadowDepthMap*, DepthMap, &StaticShadowDepthMap,
-			const FStaticShadowDepthMapData*, DepthMapData, DepthMapData,
+		FStaticShadowDepthMap* DepthMap = &StaticShadowDepthMap;
+		ENQUEUE_RENDER_COMMAND(SetDepthMapData)(
+			[DepthMap, DepthMapData](FRHICommandList& RHICmdList)
 			{
 				DepthMap->Data = DepthMapData;
-			}
-		);
+			});
 
 		BeginInitResource(&StaticShadowDepthMap);
 	}
diff --git a/Engine/Source/Runtime/Engine/Private/Components/PointLightComponent.cpp b/Engine/Source/Runtime/Engine/Private/Components/PointLightComponent.cpp
--- a/Engine/Source/Runtime/Engine/Private/Components/PointLightComponent.cpp
+++ b/Engine/Source/Runtime/Engine/Private/Components/PointLightComponent.cpp
@@ -20,13 +20,12 @@ static FAutoConsoleVariableRef CVarAllowPointLightCubemapShadows(
 
 void FLocalLightSceneProxy::UpdateRadius_GameThread(float ComponentRadius)
 {
-	ENQUEUE_UNIQUE_RENDER_COMMAND_TWOPARAMETER(
-		UpdateRadius,
-		FLocalLightSceneProxy*,LightSceneInfo,this,
-		float,ComponentRadius,ComponentRadius,
-	{
-		LightSceneInfo->UpdateRadius(ComponentRadius);
-	});
+	FLocalLightSceneProxy* InLightSceneInfo = this;
+	ENQUEUE_RENDER_COMMAND(UpdateRadius)(
+		[InLightSceneInfo, ComponentRadius](FRHICommandList& RHICmdList)
+		{
+			InLightSceneInfo->UpdateRadius(ComponentRadius);
+		});
 }
 
 /** Accesses parameters needed for rendering the light. */
diff --git a/Engine/Source/Runtime/Engine/Private/Components/PrimitiveComponent.cpp b/Engine/Source/Runtime/Engine/Private/Components/PrimitiveComponent.cpp
--- a/Engine/Source/Runtime/Engine/Private/Components/PrimitiveComponent.cpp
+++ b/Engine/Source/Runtime/Engine/Private/Components/PrimitiveComponent.cpp
@@ -834,11 +834,13 @@ void UPrimitiveComponent::SendRenderDebugPhysics(FPrimitiveSceneProxy* OverrideS
 			}
 		}
 
-		ENQUEUE_UNIQUE_RENDER_COMMAND_TWOPARAMETER(
-			PrimitiveComponent_SendRenderDebugPhysics, FPrimitiveSceneProxy*, PassedSceneProxy, UseSceneProxy, TArray<FPrimitiveSceneProxy::FDebugMassData>, UseDebugMassData, DebugMassData,
-		{
-				PassedSceneProxy->SetDebugMassData(UseDebugMassData);
-		});
+		FPrimitiveSceneProxy* PassedSceneProxy = UseSceneProxy;
+		TArray<FPrimitiveSceneProxy::FDebugMassData> UseDebugMassData = DebugMassData;
+		ENQUEUE_RENDER_COMMAND(PrimitiveComponent_SendRenderDebugPhysics)(
+			[UseSceneProxy, DebugMassData](FRHICommandList& RHICmdList)
+			{
+					UseSceneProxy->SetDebugMassData(DebugMassData);
+			});
 	}
 }
 #endif
diff --git a/Engine/Source/Runtime/Engine/Private/GPUSort.cpp b/Engine/Source/Runtime/Engine/Private/GPUSort.cpp
--- a/Engine/Source/Runtime/Engine/Private/GPUSort.cpp
+++ b/Engine/Source/Runtime/Engine/Private/GPUSort.cpp
@@ -1077,11 +1077,9 @@ static bool TestGPUSort_RenderThread(FRHICommandListImmediate& RHICmdList, EGPUS
  */
 void TestGPUSort(EGPUSortTest TestToRun, ERHIFeatureLevel::Type FeatureLevel)
 {
-	ENQUEUE_UNIQUE_RENDER_COMMAND_TWOPARAMETER(
-		FTestGPUSortCommand,
-		EGPUSortTest, TestToRun, TestToRun,
-		ERHIFeatureLevel::Type, FeatureLevel, FeatureLevel,
-	{
-		TestGPUSort_RenderThread(RHICmdList, TestToRun, FeatureLevel);
-	});
+	ENQUEUE_RENDER_COMMAND(FTestGPUSortCommand)(
+		[TestToRun, FeatureLevel](FRHICommandListImmediate& RHICmdList)
+		{
+			TestGPUSort_RenderThread(RHICmdList, TestToRun, FeatureLevel);
+		});
 }
diff --git a/Engine/Source/Runtime/RenderCore/Private/RenderingThread.cpp b/Engine/Source/Runtime/RenderCore/Private/RenderingThread.cpp
--- a/Engine/Source/Runtime/RenderCore/Private/RenderingThread.cpp
+++ b/Engine/Source/Runtime/RenderCore/Private/RenderingThread.cpp
@@ -370,13 +370,10 @@ static void AdvanceRenderingThreadStats(int64 StatsFrame, int32 MasterDisableCha
  */
 void AdvanceRenderingThreadStatsGT( bool bDiscardCallstack, int64 StatsFrame, int32 MasterDisableChangeTagStartFrame )
 {
-	ENQUEUE_UNIQUE_RENDER_COMMAND_TWOPARAMETER
-	(
-		RenderingThreadTickCommand,
-		int64, SentStatsFrame, StatsFrame,
-		int32, SentMasterDisableChangeTagStartFrame, MasterDisableChangeTagStartFrame,
+	ENQUEUE_RENDER_COMMAND(RenderingThreadTickCommand)(
+		[StatsFrame, MasterDisableChangeTagStartFrame](FRHICommandList& RHICmdList)
 		{
-			AdvanceRenderingThreadStats( SentStatsFrame, SentMasterDisableChangeTagStartFrame );
+			AdvanceRenderingThreadStats(StatsFrame, MasterDisableChangeTagStartFrame );
 		}
 	);
 	if( bDiscardCallstack )
@@ -968,22 +965,21 @@ void FRenderCommandFence::BeginFence(bool bSyncToRHIAndGPU)
 			// Create a task graph event which we can pass to the render or RHI threads.
 			CompletionEvent = FGraphEvent::CreateGraphEvent();
 
-			ENQUEUE_UNIQUE_RENDER_COMMAND_TWOPARAMETER(
-				FSyncFrameCommand,
-				FGraphEventRef, CompletionEvent, CompletionEvent,
-				int32, GTSyncType, GTSyncType,
-			{
-				if (GRHIThread_InternalUseOnly)
+			FGraphEventRef InCompletionEvent = CompletionEvent;
+			ENQUEUE_RENDER_COMMAND(FSyncFrameCommand)(
+				[InCompletionEvent, GTSyncType](FRHICommandListImmediate& RHICmdList)
 				{
-					ALLOC_COMMAND_CL(RHICmdList, FRHISyncFrameCommand)(CompletionEvent, GTSyncType);
-					RHICmdList.ImmediateFlush(EImmediateFlushType::DispatchToRHIThread);
-				}
-				else
-				{
-					FRHISyncFrameCommand Command(CompletionEvent, GTSyncType);
-					Command.Execute(RHICmdList);
-				}
-			});
+					if (GRHIThread_InternalUseOnly)
+					{
+						ALLOC_COMMAND_CL(RHICmdList, FRHISyncFrameCommand)(InCompletionEvent, GTSyncType);
+						RHICmdList.ImmediateFlush(EImmediateFlushType::DispatchToRHIThread);
+					}
+					else
+					{
+						FRHISyncFrameCommand Command(InCompletionEvent, GTSyncType);
+						Command.Execute(RHICmdList);
+					}
+				});
 		}
 		else
 		{
