Removed RHI_RAYTRACING define check from ShouldCompileRayTracingShadersForProject. It is now always compiled, but presence of `r.RayTracing` CVar is checked before accessing it.

diff --git a/Engine/Source/Runtime/RenderCore/Public/RenderResource.h b/Engine/Source/Runtime/RenderCore/Public/RenderResource.h
--- a/Engine/Source/Runtime/RenderCore/Public/RenderResource.h
+++ b/Engine/Source/Runtime/RenderCore/Public/RenderResource.h
@@ -506,15 +506,21 @@ public:
 
 FORCEINLINE bool ShouldCompileRayTracingShadersForProject(EShaderPlatform ShaderPlatform)
 {
-#if RHI_RAYTRACING
 	if (RHISupportsRayTracingShaders(ShaderPlatform))
 	{
 		// r.RayTracing is a read-only CVar. UE needs to be restarted to effectively change it
-		static const bool bRayTracingEnabled = IConsoleManager::Get().FindConsoleVariable(TEXT("r.RayTracing"))->GetInt() > 0;
+		auto GetRayTracingCVarValue = []()
+		{
+			auto CVar = IConsoleManager::Get().FindConsoleVariable(TEXT("r.RayTracing"));
+			return CVar && CVar->GetInt() > 0;
+		};
+		static const bool bRayTracingEnabled = GetRayTracingCVarValue();
 		return bRayTracingEnabled;
 	}
-#endif
-	return false;
+	else
+	{
+		return false;
+	}
 }
 
 // Returns `true` when running on RT-capable machine and RT support is enabled for the project.
@@ -524,7 +530,12 @@ FORCEINLINE bool IsRayTracingEnabled()
 	if (GRHISupportsRayTracing)
 	{
 		// r.RayTracing is a read-only CVar. UE needs to be restarted to effectively change it
-		static const bool bRayTracingEnabled = IConsoleManager::Get().FindConsoleVariable(TEXT("r.RayTracing"))->GetInt() > 0;
+		auto GetRayTracingCVarValue = []()
+		{
+			auto CVar = IConsoleManager::Get().FindConsoleVariable(TEXT("r.RayTracing"));
+			return CVar && CVar->GetInt() > 0;
+		};
+		static const bool bRayTracingEnabled = GetRayTracingCVarValue();
 		return bRayTracingEnabled;
 	}
 	else
