Disable gpuscene on Switch.

diff --git a/Engine/Source/Runtime/RenderCore/Public/RenderUtils.h b/Engine/Source/Runtime/RenderCore/Public/RenderUtils.h
--- a/Engine/Source/Runtime/RenderCore/Public/RenderUtils.h
+++ b/Engine/Source/Runtime/RenderCore/Public/RenderUtils.h
@@ -475,7 +475,8 @@ inline bool UseGPUScene(EShaderPlatform Platform, ERHIFeatureLevel::Type Feature
 	return FeatureLevel >= ERHIFeatureLevel::SM5 
 		//@todo - support GPU Scene management compute shaders on these platforms to get dynamic instancing speedups on the Rendering Thread and RHI Thread
 		&& !IsOpenGLPlatform(Platform)
-		&& !IsVulkanPlatform(Platform);
+		&& !IsVulkanPlatform(Platform)
+		&& !IsSwitchPlatform(Platform);
 }
 
 /** Unit cube vertex buffer (VertexDeclarationFVector4) */
