Resolving Robomerge conflict from shelf 4866583.

diff --git a/Engine/Source/Runtime/Engine/Private/Materials/MaterialShader.cpp b/Engine/Source/Runtime/Engine/Private/Materials/MaterialShader.cpp
--- a/Engine/Source/Runtime/Engine/Private/Materials/MaterialShader.cpp
+++ b/Engine/Source/Runtime/Engine/Private/Materials/MaterialShader.cpp
@@ -634,9 +634,8 @@ bool FMaterialShaderMapId::operator==(const FMaterialShaderMapId& ReferenceSet)
 	for (int32 ShaderIndex = 0; ShaderIndex < ShaderTypeDependencies.Num(); ShaderIndex++)
 	{
 		const FShaderTypeDependency& ShaderTypeDependency = ShaderTypeDependencies[ShaderIndex];
-		
-		if (ShaderTypeDependency.ShaderType != ReferenceSet.ShaderTypeDependencies[ShaderIndex].ShaderType
-			|| ShaderTypeDependency.SourceHash != ReferenceSet.ShaderTypeDependencies[ShaderIndex].SourceHash)
+
+		if (ShaderTypeDependency != ReferenceSet.ShaderTypeDependencies[ShaderIndex])
 		{
 			return false;
 		}
@@ -646,8 +645,7 @@ bool FMaterialShaderMapId::operator==(const FMaterialShaderMapId& ReferenceSet)
 	{
 		const FShaderPipelineTypeDependency& ShaderPipelineTypeDependency = ShaderPipelineTypeDependencies[ShaderPipelineIndex];
 
-		if (ShaderPipelineTypeDependency.ShaderPipelineType != ReferenceSet.ShaderPipelineTypeDependencies[ShaderPipelineIndex].ShaderPipelineType
-			|| ShaderPipelineTypeDependency.StagesSourceHash != ReferenceSet.ShaderPipelineTypeDependencies[ShaderPipelineIndex].StagesSourceHash)
+		if (ShaderPipelineTypeDependency != ReferenceSet.ShaderPipelineTypeDependencies[ShaderPipelineIndex])
 		{
 			return false;
 		}
@@ -656,9 +654,8 @@ bool FMaterialShaderMapId::operator==(const FMaterialShaderMapId& ReferenceSet)
 	for (int32 VFIndex = 0; VFIndex < VertexFactoryTypeDependencies.Num(); VFIndex++)
 	{
 		const FVertexFactoryTypeDependency& VFDependency = VertexFactoryTypeDependencies[VFIndex];
-		
-		if (VFDependency.VertexFactoryType != ReferenceSet.VertexFactoryTypeDependencies[VFIndex].VertexFactoryType
-			|| VFDependency.VFSourceHash != ReferenceSet.VertexFactoryTypeDependencies[VFIndex].VFSourceHash)
+
+		if (VFDependency != ReferenceSet.VertexFactoryTypeDependencies[VFIndex])
 		{
 			return false;
 		}
@@ -710,6 +707,7 @@ void FMaterialShaderMapId::UpdateParameterSet(const FStaticParameterSet& StaticP
 
 void FMaterialShaderMapId::AppendKeyString(FString& KeyString) const
 {
+#if WITH_EDITOR
 	KeyString += BaseMaterialId.ToString();
 	KeyString += TEXT("_");
 
@@ -817,10 +815,12 @@ void FMaterialShaderMapId::AppendKeyString(FString& KeyString) const
 	KeyString += BytesToHex(&TextureReferencesHash.Hash[0], sizeof(TextureReferencesHash.Hash));
 
 	KeyString += BytesToHex(&BasePropertyOverridesHash.Hash[0], sizeof(BasePropertyOverridesHash.Hash));
+#endif // WITH_EDITOR
 }
 
 void FMaterialShaderMapId::SetShaderDependencies(const TArray<FShaderType*>& ShaderTypes, const TArray<const FShaderPipelineType*>& ShaderPipelineTypes, const TArray<FVertexFactoryType*>& VFTypes, EShaderPlatform ShaderPlatform)
 {
+#if WITH_EDITOR
 	if (!FPlatformProperties::RequiresCookedData())
 	{
 		for (int32 ShaderTypeIndex = 0; ShaderTypeIndex < ShaderTypes.Num(); ShaderTypeIndex++)
@@ -848,6 +848,7 @@ void FMaterialShaderMapId::SetShaderDependencies(const TArray<FShaderType*>& Sha
 			ShaderPipelineTypeDependencies.Add(Dependency);
 		}
 	}
+#endif // WITH_EDITOR
 }
 
 /**
@@ -1235,6 +1236,7 @@ void FMaterialShaderMap::SaveForRemoteRecompile(FArchive& Ar, const TMap<FString
 
 void FMaterialShaderMap::LoadForRemoteRecompile(FArchive& Ar, EShaderPlatform ShaderPlatform, const TArray<FString>& MaterialsForShaderMaps)
 {
+#if WITH_EDITOR
 	int32 NumResources;
 	Ar << NumResources;
 
@@ -1336,6 +1338,7 @@ void FMaterialShaderMap::LoadForRemoteRecompile(FArchive& Ar, EShaderPlatform Sh
 			}
 		}
 	}
+#endif // WITH_EDITOR
 }
 
 
diff --git a/Engine/Source/Runtime/RenderCore/Private/GlobalShader.cpp b/Engine/Source/Runtime/RenderCore/Private/GlobalShader.cpp
--- a/Engine/Source/Runtime/RenderCore/Private/GlobalShader.cpp
+++ b/Engine/Source/Runtime/RenderCore/Private/GlobalShader.cpp
@@ -84,9 +84,8 @@ FGlobalShaderMapId::FGlobalShaderMapId(EShaderPlatform Platform)
 	ShaderTypes.Sort(FCompareShaderTypes());
 	for (int32 TypeIndex = 0; TypeIndex < ShaderTypes.Num(); TypeIndex++)
 	{
-		FShaderTypeDependency Dependency;
-		Dependency.ShaderType = ShaderTypes[TypeIndex];
-		Dependency.SourceHash = ShaderTypes[TypeIndex]->GetSourceHash(Platform);
+		FShaderType* ShaderType = ShaderTypes[TypeIndex];
+		FShaderTypeDependency Dependency(ShaderType, Platform);
 
 		const TCHAR* ShaderFilename = Dependency.ShaderType->GetShaderFilename();
 
@@ -100,15 +99,14 @@ FGlobalShaderMapId::FGlobalShaderMapId(EShaderPlatform Platform)
 	for (int32 TypeIndex = 0; TypeIndex < ShaderPipelineTypes.Num(); TypeIndex++)
 	{
 		const FShaderPipelineType* Pipeline = ShaderPipelineTypes[TypeIndex];
-		FShaderPipelineTypeDependency Dependency;
-		Dependency.ShaderPipelineType = Pipeline;
-		Dependency.StagesSourceHash = Pipeline->GetSourceHash(Platform);
+		FShaderPipelineTypeDependency Dependency(Pipeline, Platform);
 		ShaderPipelineTypeDependencies.Add(Dependency);
 	}
 }
 
 void FGlobalShaderMapId::AppendKeyString(FString& KeyString, const TArray<FShaderTypeDependency>& Dependencies) const
 {
+#if WITH_EDITOR
 	TMap<const TCHAR*,FCachedUniformBufferDeclaration> ReferencedUniformBuffers;
 
 	for (int32 ShaderIndex = 0; ShaderIndex < Dependencies.Num(); ShaderIndex++)
@@ -167,6 +165,7 @@ void FGlobalShaderMapId::AppendKeyString(FString& KeyString, const TArray<FShade
 
 		SerializationHistory.AppendKeyString(KeyString);
 	}
+#endif // WITH_EDITOR
 }
 
 FGlobalShader::FGlobalShader(const ShaderMetaType::CompiledShaderInitializerType& Initializer)
diff --git a/Engine/Source/Runtime/RenderCore/Private/Shader.cpp b/Engine/Source/Runtime/RenderCore/Private/Shader.cpp
--- a/Engine/Source/Runtime/RenderCore/Private/Shader.cpp
+++ b/Engine/Source/Runtime/RenderCore/Private/Shader.cpp
@@ -233,6 +233,7 @@ TMap<FName, FShaderType*>& FShaderType::GetNameToTypeMap()
 inline bool FShaderType::GetOutdatedCurrentType(TArray<FShaderType*>& OutdatedShaderTypes, TArray<const FVertexFactoryType*>& OutdatedFactoryTypes) const
 {
 	bool bOutdated = false;
+#if WITH_EDITOR
 	for (TMap<FShaderId, FShader*>::TConstIterator ShaderIt(ShaderIdMap);ShaderIt;++ShaderIt)
 	{
 		FShader* Shader = ShaderIt.Value();
@@ -255,7 +256,7 @@ inline bool FShaderType::GetOutdatedCurrentType(TArray<FShaderType*>& OutdatedSh
 			bOutdated = true;
 		}
 	}
-
+#endif // WITH_EDITOR
 	return bOutdated;
 }
 
@@ -848,7 +849,11 @@ bool FShaderResource::ArePlatformsCompatible(EShaderPlatform CurrentPlatform, ES
 
 FSHAHash &FShaderResource::FilterShaderSourceHashForSerialization(const FArchive& Ar, FSHAHash &HashToSerialize)
 {
+#if KEEP_SHADER_SOURCE_HASHES
 	return (!Ar.IsCooking()) ? HashToSerialize : ShaderSourceDefaultHash;
+#else
+	return ShaderSourceDefaultHash;
+#endif
 }
 
 static void SafeAssignHash(FRHIShader* InShader, const FSHAHash& Hash)
@@ -1019,7 +1024,9 @@ FShaderResourceId FShaderResource::GetId() const
 
 FShaderId::FShaderId(const FSHAHash& InMaterialShaderMapHash, const FShaderPipelineType* InShaderPipeline, FVertexFactoryType* InVertexFactoryType, FShaderType* InShaderType, int32 InPermutationId, FShaderTarget InTarget)
 	: MaterialShaderMapHash(InMaterialShaderMapHash)
+#if KEEP_SHADER_SOURCE_HASHES
 	, SourceHash(InShaderType->GetSourceHash(InTarget.GetPlatform()))
+#endif
 	, Target(InTarget)
 	, ShaderPipeline(InShaderPipeline)
 	, ShaderType(InShaderType)
@@ -1030,7 +1037,9 @@ FShaderId::FShaderId(const FSHAHash& InMaterialShaderMapHash, const FShaderPipel
 	{
 		VFSerializationHistory = InVertexFactoryType->GetSerializationHistory(InTarget.GetFrequency());
 		VertexFactoryType = InVertexFactoryType;
+#if KEEP_SHADER_SOURCE_HASHES
 		VFSourceHash = InVertexFactoryType->GetSourceHash(InTarget.GetPlatform());
+#endif
 	}
 	else
 	{
@@ -1048,11 +1057,13 @@ FSelfContainedShaderId::FSelfContainedShaderId(const FShaderId& InShaderId)
 	MaterialShaderMapHash = InShaderId.MaterialShaderMapHash;
 	VertexFactoryTypeName = InShaderId.VertexFactoryType ? InShaderId.VertexFactoryType->GetName() : TEXT("");
 	ShaderPipelineName = InShaderId.ShaderPipeline ? InShaderId.ShaderPipeline->GetName() : TEXT("");
-	VFSourceHash = InShaderId.VFSourceHash;
 	VFSerializationHistory = InShaderId.VFSerializationHistory ? *InShaderId.VFSerializationHistory : FSerializationHistory();
 	ShaderTypeName = InShaderId.ShaderType->GetName();
 	PermutationId = InShaderId.PermutationId;
+#if KEEP_SHADER_SOURCE_HASHES
 	SourceHash = InShaderId.SourceHash;
+	VFSourceHash = InShaderId.VFSourceHash;
+#endif
 	SerializationHistory = InShaderId.SerializationHistory;
 	Target = InShaderId.Target;
 }
@@ -1060,12 +1071,20 @@ FSelfContainedShaderId::FSelfContainedShaderId(const FShaderId& InShaderId)
 bool FSelfContainedShaderId::IsValid()
 {
 	FShaderType** TypePtr = FShaderType::GetNameToTypeMap().Find(FName(*ShaderTypeName));
-	if (TypePtr && SourceHash == (*TypePtr)->GetSourceHash(Target.GetPlatform()) && SerializationHistory == (*TypePtr)->GetSerializationHistory())
+	if (TypePtr
+#if KEEP_SHADER_SOURCE_HASHES
+		&& SourceHash == (*TypePtr)->GetSourceHash(Target.GetPlatform())
+#endif
+		&& SerializationHistory == (*TypePtr)->GetSerializationHistory())
 	{
 		FVertexFactoryType* VFTypePtr = FVertexFactoryType::GetVFByName(VertexFactoryTypeName);
 
 		if (VertexFactoryTypeName == TEXT("") 
-			|| (VFTypePtr && VFSourceHash == VFTypePtr->GetSourceHash(Target.GetPlatform()) && VFSerializationHistory == *VFTypePtr->GetSerializationHistory(Target.GetFrequency())))
+			|| (VFTypePtr
+#if KEEP_SHADER_SOURCE_HASHES
+				&& VFSourceHash == VFTypePtr->GetSourceHash(Target.GetPlatform())
+#endif
+				&& VFSerializationHistory == *VFTypePtr->GetSerializationHistory(Target.GetFrequency())))
 		{
 			return true;
 		}
@@ -1078,13 +1097,20 @@ FArchive& operator<<(FArchive& Ar,class FSelfContainedShaderId& Ref)
 {
 	Ar.UsingCustomVersion(FRenderingObjectVersion::GUID);
 
+#if KEEP_SHADER_SOURCE_HASHES
+	FSHAHash& VFHash = Ref.VFSourceHash;
+	FSHAHash& Hash = Ref.SourceHash;
+#else
+	FSHAHash VFHash, Hash;
+#endif
+
 	Ar << Ref.MaterialShaderMapHash 
 		<< Ref.VertexFactoryTypeName
 		<< Ref.ShaderPipelineName
-		<< FShaderResource::FilterShaderSourceHashForSerialization(Ar, Ref.VFSourceHash)
+		<< FShaderResource::FilterShaderSourceHashForSerialization(Ar, VFHash)
 		<< Ref.VFSerializationHistory
 		<< Ref.ShaderTypeName
-		<< FShaderResource::FilterShaderSourceHashForSerialization(Ar, Ref.SourceHash)
+		<< FShaderResource::FilterShaderSourceHashForSerialization(Ar, Hash)
 		<< Ref.SerializationHistory
 		<< Ref.Target;
 
@@ -1131,6 +1157,7 @@ FShader::FShader(const CompiledShaderInitializerType& Initializer)
 
 	check(Type);
 
+#if KEEP_SHADER_SOURCE_HASHES
 	// Store off the source hash that this shader was compiled with
 	// This will be used as part of the shader key in order to identify when shader files have been changed and a recompile is needed
 	SourceHash = Type->GetSourceHash(Target.GetPlatform());
@@ -1140,6 +1167,7 @@ FShader::FShader(const CompiledShaderInitializerType& Initializer)
 		// Store off the VF source hash that this shader was compiled with
 		VFSourceHash = VFType->GetSourceHash(Target.GetPlatform());
 	}
+#endif
 
 	// Bind uniform buffer parameters automatically 
 	for (TLinkedList<FShaderParametersMetadata*>::TIterator StructIt(FShaderParametersMetadata::GetStructList()); StructIt; StructIt.Next())
@@ -1173,7 +1201,11 @@ FShader::~FShader()
 
 const FSHAHash& FShader::GetHash() const 
 { 
+#if KEEP_SHADER_SOURCE_HASHES
 	return SourceHash;
+#else
+	return ShaderSourceDefaultHash;
+#endif
 }
 
 EShaderPlatform FShader::GetShaderPlatform() const
@@ -1190,17 +1222,24 @@ bool FShader::SerializeBase(FArchive& Ar, bool bShadersInline, bool bLoadedByCoo
 
 	Ar.UsingCustomVersion(FRenderingObjectVersion::GUID);
 
+#if KEEP_SHADER_SOURCE_HASHES
+	FSHAHash& VFHash = VFSourceHash;
+	FSHAHash& Hash = SourceHash;
+#else
+	FSHAHash VFHash, Hash;
+#endif
+
 	Ar << OutputHash;
 	Ar << MaterialShaderMapHash;
 	Ar << ShaderPipeline;
 	Ar << VFType;
-	Ar << FShaderResource::FilterShaderSourceHashForSerialization(Ar, VFSourceHash);
+	Ar << FShaderResource::FilterShaderSourceHashForSerialization(Ar, VFHash);
 	Ar << Type;
 	if (Ar.CustomVer(FRenderingObjectVersion::GUID) >= FRenderingObjectVersion::ShaderPermutationId)
 	{
 		Ar << PermutationId;
 	}
-	Ar << FShaderResource::FilterShaderSourceHashForSerialization(Ar, SourceHash);
+	Ar << FShaderResource::FilterShaderSourceHashForSerialization(Ar, Hash);
 	Ar << Target;
 
 	// TODO(RDG): Kill that once all shaders are refactored.
@@ -1325,7 +1364,9 @@ void FShader::Register(bool bLoadedByCookedMaterial)
 {
 	FShaderId ShaderId = GetId();
 	check(ShaderId.MaterialShaderMapHash != FSHAHash());
+#if KEEP_SHADER_SOURCE_HASHES
 	check(ShaderId.SourceHash != FSHAHash() || FPlatformProperties::RequiresCookedData() || bLoadedByCookedMaterial);
+#endif
 	check(Resource);
 	Type->AddToShaderIdMap(ShaderId, this);
 }
@@ -1341,11 +1382,13 @@ FShaderId FShader::GetId() const
 	ShaderId.MaterialShaderMapHash = MaterialShaderMapHash;
 	ShaderId.ShaderPipeline = ShaderPipeline;
 	ShaderId.VertexFactoryType = VFType;
-	ShaderId.VFSourceHash = VFSourceHash;
 	ShaderId.VFSerializationHistory = VFType ? VFType->GetSerializationHistory((EShaderFrequency)GetTarget().Frequency) : NULL;
 	ShaderId.ShaderType = Type;
 	ShaderId.PermutationId = PermutationId;
+#if KEEP_SHADER_SOURCE_HASHES
 	ShaderId.SourceHash = SourceHash;
+	ShaderId.VFSourceHash = VFSourceHash;
+#endif
 	ShaderId.Target = Target;
 	return ShaderId;
 }
@@ -1384,10 +1427,12 @@ void FShader::DumpDebugInfo()
 	UE_LOG(LogConsoleResponse, Display, TEXT("               :Target %s"), GetShaderFrequencyString((EShaderFrequency)Target.Frequency));
 	UE_LOG(LogConsoleResponse, Display, TEXT("               :Target %s"), *LegacyShaderPlatformToShaderFormat(EShaderPlatform(Target.Platform)).ToString());
 	UE_LOG(LogConsoleResponse, Display, TEXT("               :VFType %s"), VFType ? VFType->GetName() : TEXT("null"));
-	UE_LOG(LogConsoleResponse, Display, TEXT("               :VFSourceHash %s"), *VFSourceHash.ToString());
 	UE_LOG(LogConsoleResponse, Display, TEXT("               :Type %s"), Type->GetName());
 	UE_LOG(LogConsoleResponse, Display, TEXT("               :PermutationId %d"), PermutationId);
+#if KEEP_SHADER_SOURCE_HASHES
 	UE_LOG(LogConsoleResponse, Display, TEXT("               :SourceHash %s"), *SourceHash.ToString());
+	UE_LOG(LogConsoleResponse, Display, TEXT("               :VFSourceHash %s"), *VFSourceHash.ToString());
+#endif
 	UE_LOG(LogConsoleResponse, Display, TEXT("               :OutputHash %s"), *OutputHash.ToString());
 }
 
diff --git a/Engine/Source/Runtime/RenderCore/Private/VertexFactory.cpp b/Engine/Source/Runtime/RenderCore/Private/VertexFactory.cpp
--- a/Engine/Source/Runtime/RenderCore/Private/VertexFactory.cpp
+++ b/Engine/Source/Runtime/RenderCore/Private/VertexFactory.cpp
@@ -289,7 +289,9 @@ FVertexFactoryParameterRef::FVertexFactoryParameterRef(FVertexFactoryType* InVer
 , ShaderPlatform(InShaderPlatform)
 {
 	Parameters = VertexFactoryType->CreateShaderParameters(InShaderFrequency);
+#if KEEP_SHADER_SOURCE_HASHES
 	VFHash = GetShaderFileHash(VertexFactoryType->GetShaderFilename(), InShaderPlatform);
+#endif
 
 	if(Parameters)
 	{
@@ -317,7 +319,12 @@ bool operator<<(FArchive& Ar,FVertexFactoryParameterRef& Ref)
 		Ref.ShaderPlatform = (EShaderPlatform)ShaderPlatformByte;
 	}
 
-	Ar << FShaderResource::FilterShaderSourceHashForSerialization(Ar, Ref.VFHash);
+#if KEEP_SHADER_SOURCE_HASHES
+	FSHAHash& VFHash = Ref.VFHash;
+#else
+	FSHAHash VFHash;
+#endif
+	Ar << FShaderResource::FilterShaderSourceHashForSerialization(Ar, VFHash);
 
 
 	if (Ar.IsLoading())
@@ -363,11 +370,13 @@ bool operator<<(FArchive& Ar,FVertexFactoryParameterRef& Ref)
 	return bShaderHasOutdatedParameters;
 }
 
+#if KEEP_SHADER_SOURCE_HASHES
 /** Returns the hash of the vertex factory shader file that this shader was compiled with. */
 const FSHAHash& FVertexFactoryParameterRef::GetHash() const 
 { 
 	return VFHash;
 }
+#endif
 
 /** Returns the shader platform that this shader was compiled with. */
 EShaderPlatform FVertexFactoryParameterRef::GetShaderPlatform() const
diff --git a/Engine/Source/Runtime/RenderCore/Public/Shader.h b/Engine/Source/Runtime/RenderCore/Public/Shader.h
--- a/Engine/Source/Runtime/RenderCore/Public/Shader.h
+++ b/Engine/Source/Runtime/RenderCore/Public/Shader.h
@@ -35,6 +35,9 @@ class FVertexFactoryParameterRef;
 class FVertexFactoryType;
 class FShaderParametersMetadata;
 
+/** By default most shader source hashes are stripped at cook time so can be discarded
+	to save memory. See implementation of FilterShaderSourceHashForSerialization. */
+#define KEEP_SHADER_SOURCE_HASHES	(WITH_EDITOR)
 
 /** Define a shader permutation uniquely according to its type, and permutation id.*/
 template<typename MetaShaderType>
@@ -591,11 +594,13 @@ public:
 	 */ 
 	FSHAHash MaterialShaderMapHash;
 
+#if KEEP_SHADER_SOURCE_HASHES
 	/** Used to detect changes to the vertex factory source files. */
 	FSHAHash VFSourceHash;
 
 	/** Used to detect changes to the shader source files. */
 	FSHAHash SourceHash;
+#endif
 
 	/** Shader platform and frequency. */
 	FShaderTarget Target;
@@ -647,13 +652,15 @@ public:
 		return X.MaterialShaderMapHash == Y.MaterialShaderMapHash
 			&& X.ShaderPipeline == Y.ShaderPipeline
 			&& X.VertexFactoryType == Y.VertexFactoryType
-			&& X.VFSourceHash == Y.VFSourceHash
 			&& ((X.VFSerializationHistory == NULL && Y.VFSerializationHistory == NULL)
 				|| (X.VFSerializationHistory != NULL && Y.VFSerializationHistory != NULL &&
 					*X.VFSerializationHistory == *Y.VFSerializationHistory))
 			&& X.ShaderType == Y.ShaderType
 			&& X.PermutationId == Y.PermutationId 
+#if KEEP_SHADER_SOURCE_HASHES
 			&& X.SourceHash == Y.SourceHash 
+			&& X.VFSourceHash == Y.VFSourceHash
+#endif
 			&& X.SerializationHistory == Y.SerializationHistory
 			&& X.Target == Y.Target;
 	}
@@ -675,11 +682,13 @@ public:
 	 */ 
 	FSHAHash MaterialShaderMapHash;
 
+#if KEEP_SHADER_SOURCE_HASHES
 	/** Used to detect changes to the vertex factory source files. */
 	FSHAHash VFSourceHash;
 
 	/** Used to detect changes to the shader source files. */
 	FSHAHash SourceHash;
+#endif
 
 	/** 
 	 * Name of the vertex factory type that the shader was created for, 
@@ -1582,21 +1591,41 @@ public:
 		, PermutationId(0)
 	{}
 
+	FShaderTypeDependency(FShaderType* InShaderType, EShaderPlatform ShaderPlatform)
+		: ShaderType(InShaderType)
+		, PermutationId(0)
+	{
+#if KEEP_SHADER_SOURCE_HASHES
+		if (ShaderType)
+		{
+			SourceHash = ShaderType->GetSourceHash(ShaderPlatform);
+		}
+#endif
+	}
+
 	/** Shader type */
 	FShaderType* ShaderType;
 
 	/** Unique permutation identifier of the global shader type. */
 	int32 PermutationId;
 
+#if KEEP_SHADER_SOURCE_HASHES
 	/** Used to detect changes to the shader source files. */
 	FSHAHash SourceHash;
+#endif
 
 	friend FArchive& operator<<(FArchive& Ar,class FShaderTypeDependency& Ref)
 	{
 		Ar.UsingCustomVersion(FRenderingObjectVersion::GUID);
 
 		Ar << Ref.ShaderType;
-		Ar << FShaderResource::FilterShaderSourceHashForSerialization(Ar, Ref.SourceHash);
+
+#if KEEP_SHADER_SOURCE_HASHES
+		FSHAHash& Hash = Ref.SourceHash;
+#else
+		FSHAHash Hash;
+#endif
+		Ar << FShaderResource::FilterShaderSourceHashForSerialization(Ar, Hash);
 
 		if (Ar.CustomVer(FRenderingObjectVersion::GUID) >= FRenderingObjectVersion::ShaderPermutationId)
 		{
@@ -1608,7 +1637,16 @@ public:
 
 	bool operator==(const FShaderTypeDependency& Reference) const
 	{
+#if KEEP_SHADER_SOURCE_HASHES
 		return ShaderType == Reference.ShaderType && PermutationId == Reference.PermutationId && SourceHash == Reference.SourceHash;
+#else
+		return ShaderType == Reference.ShaderType && PermutationId == Reference.PermutationId;
+#endif
+	}
+
+	bool operator!=(const FShaderTypeDependency& Reference) const
+	{
+		return !(*this == Reference);
 	}
 };
 
@@ -1620,22 +1658,50 @@ public:
 		ShaderPipelineType(nullptr)
 	{}
 
+	FShaderPipelineTypeDependency(const FShaderPipelineType* InShaderPipelineType, EShaderPlatform ShaderPlatform) :
+		ShaderPipelineType(InShaderPipelineType)
+	{
+#if KEEP_SHADER_SOURCE_HASHES
+		if (ShaderPipelineType)
+		{
+			StagesSourceHash = ShaderPipelineType->GetSourceHash(ShaderPlatform);
+		}
+#endif
+	}
+
 	/** Shader Pipeline type */
 	const FShaderPipelineType* ShaderPipelineType;
 
+#if KEEP_SHADER_SOURCE_HASHES
 	/** Used to detect changes to the shader source files. */
 	FSHAHash StagesSourceHash;
+#endif
 
 	friend FArchive& operator<<(FArchive& Ar, class FShaderPipelineTypeDependency& Ref)
 	{
 		Ar << Ref.ShaderPipelineType;
-		Ar << FShaderResource::FilterShaderSourceHashForSerialization(Ar, Ref.StagesSourceHash);
+
+#if KEEP_SHADER_SOURCE_HASHES
+		FSHAHash& Hash = Ref.StagesSourceHash;
+#else
+		FSHAHash Hash;
+#endif
+		Ar << FShaderResource::FilterShaderSourceHashForSerialization(Ar, Hash);
 		return Ar;
 	}
 
 	bool operator==(const FShaderPipelineTypeDependency& Reference) const
 	{
+#if KEEP_SHADER_SOURCE_HASHES	
 		return ShaderPipelineType == Reference.ShaderPipelineType && StagesSourceHash == Reference.StagesSourceHash;
+#else
+		return ShaderPipelineType == Reference.ShaderPipelineType;
+#endif
+	}
+
+	bool operator!=(const FShaderPipelineTypeDependency& Reference) const
+	{
+		return !(*this == Reference);
 	}
 };
 
diff --git a/Engine/Source/Runtime/RenderCore/Public/VertexFactory.h b/Engine/Source/Runtime/RenderCore/Public/VertexFactory.h
--- a/Engine/Source/Runtime/RenderCore/Public/VertexFactory.h
+++ b/Engine/Source/Runtime/RenderCore/Public/VertexFactory.h
@@ -419,18 +419,34 @@ public:
 
 	FVertexFactoryType* VertexFactoryType;
 
+#if KEEP_SHADER_SOURCE_HASHES
 	/** Used to detect changes to the vertex factory source files. */
 	FSHAHash VFSourceHash;
+#endif
 
 	friend FArchive& operator<<(FArchive& Ar,class FVertexFactoryTypeDependency& Ref)
 	{
-		Ar << Ref.VertexFactoryType << FShaderResource::FilterShaderSourceHashForSerialization(Ar, Ref.VFSourceHash);
+#if KEEP_SHADER_SOURCE_HASHES
+		FSHAHash& Hash = Ref.VFSourceHash;
+#else
+		FSHAHash Hash;
+#endif
+		Ar << Ref.VertexFactoryType << FShaderResource::FilterShaderSourceHashForSerialization(Ar, Hash);
 		return Ar;
 	}
 
 	bool operator==(const FVertexFactoryTypeDependency& Reference) const
 	{
+#if KEEP_SHADER_SOURCE_HASHES
 		return VertexFactoryType == Reference.VertexFactoryType && VFSourceHash == Reference.VFSourceHash;
+#else
+		return VertexFactoryType == Reference.VertexFactoryType;
+#endif
+	}
+
+	bool operator!=(const FVertexFactoryTypeDependency& Reference) const
+	{
+		return !(*this == Reference);
 	}
 };
 
@@ -637,8 +653,10 @@ public:
 
 	const FVertexFactoryType* GetVertexFactoryType() const { return VertexFactoryType; }
 
+#if KEEP_SHADER_SOURCE_HASHES
 	/** Returns the hash of the vertex factory shader file that this shader was compiled with. */
 	const FSHAHash& GetHash() const;
+#endif
 
 	/** Returns the shader platform that this shader was compiled with. */
 	EShaderPlatform GetShaderPlatform() const;
@@ -656,8 +674,10 @@ private:
 	EShaderFrequency ShaderFrequency;
 	EShaderPlatform ShaderPlatform;
 
+#if KEEP_SHADER_SOURCE_HASHES
 	// Hash of the vertex factory's source file at shader compile time, used by the automatic versioning system to detect changes
 	FSHAHash VFHash;
+#endif
 };
 
 /**
