Add support for enabling/disabling kepping shader debug info per platform.

- This can be overriden per platform in the [ShaderCompiler] section of the platform's Engine.ini by adding "r.Shaders.KeepDebugInfo"
- Also add in support for exporting the debug info on platforms where it can be stripped and stored as separate files ("r.Shaders.ExportDebugInfo") and this is also per platform.
- Add "r.Shaders.ExportDebugInfoMode" to control how the exported files are stored, 0 for loose or 1 for an uncompressed archive.

diff --git a/Engine/Source/Runtime/Engine/Private/ShaderCompiler/ShaderCompiler.cpp b/Engine/Source/Runtime/Engine/Private/ShaderCompiler/ShaderCompiler.cpp
--- a/Engine/Source/Runtime/Engine/Private/ShaderCompiler/ShaderCompiler.cpp
+++ b/Engine/Source/Runtime/Engine/Private/ShaderCompiler/ShaderCompiler.cpp
@@ -173,7 +173,23 @@ static FAutoConsoleVariableRef CVarForceAllCoresForShaderCompiling(
 static TAutoConsoleVariable<int32> CVarKeepShaderDebugData(
 	TEXT("r.Shaders.KeepDebugInfo"),
 	0,
-	TEXT("Whether to keep shader reflection and debug data from shader bytecode, default is to strip.  When using graphical debuggers like Nsight it can be useful to enable this on startup."),
+	TEXT("Whether to keep shader reflection and debug data from shader bytecode, default is to strip.  When using graphical debuggers like Nsight it can be useful to enable this on startup.")
+	TEXT("For some platforms this cvar can be overriden in the Engine.ini, under the [ShaderCompiler] section."),
+	ECVF_ReadOnly);
+
+static TAutoConsoleVariable<int32> CVarExportShaderDebugData(
+	TEXT("r.Shaders.ExportDebugInfo"),
+	0,
+	TEXT("Whether to export the shader reflection and debug data from shader bytecode as separate files.")
+	TEXT("r.Shaders.KeepDebugInfo must be enabled and r.DumpShaderDebugInfo will enable this cvar.")
+	TEXT("For some platforms this cvar can be overriden in the Engine.ini, under the [ShaderCompiler] section."),
+	ECVF_ReadOnly);
+
+static TAutoConsoleVariable<int32> CVarExportShaderDebugDataMode(
+	TEXT("r.Shaders.ExportDebugInfoMode"),
+	0,
+	TEXT(" 0: Export as loose files.\n")
+	TEXT(" 1: Export as an uncompressed archive.\n"),
 	ECVF_ReadOnly);
 
 static TAutoConsoleVariable<int32> CVarOptimizeShaders(
@@ -3032,17 +3048,14 @@ void GlobalBeginCompileShader(
 
 	{
 		static const auto CVar = IConsoleManager::Get().FindConsoleVariable(TEXT("r.Shaders.Optimize"));
-
-		if (CVar->GetInt() == 0)
+		if (CVar && CVar->GetInt() == 0)
 		{
 			Input.Environment.CompilerFlags.Add(CFLAG_Debug);
 		}
 	}
 
 	{
-		static const auto CVar = IConsoleManager::Get().FindConsoleVariable(TEXT("r.Shaders.KeepDebugInfo"));
-
-		if (CVar->GetInt() != 0)
+		if (ShouldKeepShaderDebugInfo((EShaderPlatform)Target.Platform))
 		{
 			Input.Environment.CompilerFlags.Add(CFLAG_KeepDebugInfo);
 		}
diff --git a/Engine/Source/Runtime/RenderCore/Private/Shader.cpp b/Engine/Source/Runtime/RenderCore/Private/Shader.cpp
--- a/Engine/Source/Runtime/RenderCore/Private/Shader.cpp
+++ b/Engine/Source/Runtime/RenderCore/Private/Shader.cpp
@@ -2138,8 +2138,7 @@ void ShaderMapAppendKeyString(EShaderPlatform Platform, FString& KeyString)
 	}
 
 	{
-		static const auto CVar = IConsoleManager::Get().FindConsoleVariable(TEXT("r.Shaders.KeepDebugInfo"));
-		KeyString += (CVar && CVar->GetInt() != 0) ? TEXT("_NoStrip") : TEXT("");
+		KeyString += ShouldKeepShaderDebugInfo(Platform) ? TEXT("_NoStrip") : TEXT("");
 	}
 
 	{
diff --git a/Engine/Source/Runtime/RenderCore/Private/ShaderCore.cpp b/Engine/Source/Runtime/RenderCore/Private/ShaderCore.cpp
--- a/Engine/Source/Runtime/RenderCore/Private/ShaderCore.cpp
+++ b/Engine/Source/Runtime/RenderCore/Private/ShaderCore.cpp
@@ -21,6 +21,7 @@
 #include "Misc/CoreMisc.h"
 #include "Interfaces/ITargetPlatform.h"
 #include "Interfaces/ITargetPlatformManagerModule.h"
+#include "Misc/ConfigCacheIni.h"
 #endif
 
 static TAutoConsoleVariable<int32> CVarShaderDevelopmentMode(
@@ -265,6 +266,91 @@ bool AllowDebugViewmodes(EShaderPlatform Platform)
 #endif
 }
 
+#if WITH_EDITOR
+static void GetShaderCompilerPlatformConfigs(const TCHAR* Key, uint64& OutPlatformFlags)
+{
+	for (uint32 ShaderPlatformIndex = 0; ShaderPlatformIndex < SP_NumPlatforms; ++ShaderPlatformIndex)
+	{
+		EShaderPlatform ShaderPlatform = EShaderPlatform(ShaderPlatformIndex);
+		FName PlatformName = ShaderPlatformToPlatformName(ShaderPlatform);
+		if (!PlatformName.IsNone())
+		{
+			FConfigFile EngineSettings;
+			FConfigCacheIni::LoadLocalIniFile(EngineSettings, TEXT("Engine"), true, *PlatformName.ToString());
+
+			bool bEnabled = false;
+			if (EngineSettings.GetBool(TEXT("ShaderCompiler"), Key, bEnabled))
+			{
+				uint64 Mask = (uint64)1 << ShaderPlatformIndex;
+				if (bEnabled)
+				{
+					OutPlatformFlags |= Mask;
+				}
+				else
+				{
+					OutPlatformFlags &= ~Mask;
+				}
+			}
+		}
+	}
+}
+#endif
+
+static uint64 GetKeepShaderDebugInfoPlatforms()
+{
+	uint64 KeepDebugInfoPlatforms = 0;
+
+	// First check the global cvars
+	static IConsoleVariable* CVarKeepDebugInfo = IConsoleManager::Get().FindConsoleVariable(TEXT("r.Shaders.KeepDebugInfo"));
+	if (CVarKeepDebugInfo && CVarKeepDebugInfo->GetInt())
+	{
+		KeepDebugInfoPlatforms = ~(uint64)0;
+	}
+
+#if WITH_EDITOR
+	// Then load the per platform settings.
+	GetShaderCompilerPlatformConfigs(TEXT("r.Shaders.KeepDebugInfo"), KeepDebugInfoPlatforms);
+#endif
+
+	return KeepDebugInfoPlatforms;
+}
+
+bool ShouldKeepShaderDebugInfo(EShaderPlatform Platform)
+{
+	static uint64 KeepShaderDebugInfoPlatforms = GetKeepShaderDebugInfoPlatforms();
+	return (KeepShaderDebugInfoPlatforms & (uint64(1) << Platform)) != 0;
+}
+
+static uint64 GetExportShaderDebugInfoPlatforms()
+{
+	uint64 ExportDebugInfoPlatforms = 0;
+
+	// First check the global cvars
+
+	// r.DumpShaderDebugInfo should also turn on ExportShaderDebugInfo
+	// The difference is that r.DumpShaderDebugInfo will also output engine debug files such as converted hlsl or SCW helper files.
+	// Where as r.Shader.ExportDebugInfo is purely to export the graphics debugging tool's debug info files.
+	static IConsoleVariable* CVarExportDebugInfo = IConsoleManager::Get().FindConsoleVariable(TEXT("r.Shaders.ExportDebugInfo"));
+	static IConsoleVariable* CVarDumpDebugInfo = IConsoleManager::Get().FindConsoleVariable(TEXT("r.DumpShaderDebugInfo"));
+	if ((CVarExportDebugInfo && CVarExportDebugInfo->GetInt()) || (CVarDumpDebugInfo && CVarDumpDebugInfo->GetInt()))
+	{
+		ExportDebugInfoPlatforms = ~(uint64)0;
+	}
+
+#if WITH_EDITOR
+	// Then load the per platform settings.
+	GetShaderCompilerPlatformConfigs(TEXT("r.Shaders.ExportDebugInfo"), ExportDebugInfoPlatforms);
+#endif
+
+	return ExportDebugInfoPlatforms;
+}
+
+bool ShouldExportShaderDebugInfo(EShaderPlatform Platform)
+{
+	static uint64 GExportDebugInfoPlatforms = GetExportShaderDebugInfoPlatforms();
+	return (GExportDebugInfoPlatforms & (uint64(1) << Platform)) != 0;
+}
+
 bool FShaderParameterMap::FindParameterAllocation(const TCHAR* ParameterName,uint16& OutBufferIndex,uint16& OutBaseIndex,uint16& OutSize) const
 {
 	const FParameterAllocation* Allocation = ParameterMap.Find(ParameterName);
diff --git a/Engine/Source/Runtime/RenderCore/Public/ShaderCore.h b/Engine/Source/Runtime/RenderCore/Public/ShaderCore.h
--- a/Engine/Source/Runtime/RenderCore/Public/ShaderCore.h
+++ b/Engine/Source/Runtime/RenderCore/Public/ShaderCore.h
@@ -95,6 +95,12 @@ extern RENDERCORE_API bool AllowDebugViewmodes();
 /** Returns true if debug viewmodes are allowed for the given platform. */
 extern RENDERCORE_API bool AllowDebugViewmodes(EShaderPlatform Platform);
 
+/** Returns true if debug information should be kept for a given platform. */
+extern SHADERCORE_API bool ShouldKeepShaderDebugInfo(EShaderPlatform Platform);
+
+/** Returns true if debug information should be exported to separate files for a given platform . */
+extern SHADERCORE_API bool ShouldExportShaderDebugInfo(EShaderPlatform Platform);
+
 struct FShaderTarget
 {
 	uint32 Frequency : SF_NumBits;
