Fix UE4 Main build
#rnx

diff --git a/Engine/Source/Runtime/RenderCore/Public/ShaderCore.h b/Engine/Source/Runtime/RenderCore/Public/ShaderCore.h
--- a/Engine/Source/Runtime/RenderCore/Public/ShaderCore.h
+++ b/Engine/Source/Runtime/RenderCore/Public/ShaderCore.h
@@ -96,10 +96,10 @@ extern RENDERCORE_API bool AllowDebugViewmodes();
 extern RENDERCORE_API bool AllowDebugViewmodes(EShaderPlatform Platform);
 
 /** Returns true if debug information should be kept for a given platform. */
-extern SHADERCORE_API bool ShouldKeepShaderDebugInfo(EShaderPlatform Platform);
+extern RENDERCORE_API bool ShouldKeepShaderDebugInfo(EShaderPlatform Platform);
 
 /** Returns true if debug information should be exported to separate files for a given platform . */
-extern SHADERCORE_API bool ShouldExportShaderDebugInfo(EShaderPlatform Platform);
+extern RENDERCORE_API bool ShouldExportShaderDebugInfo(EShaderPlatform Platform);
 
 struct FShaderTarget
 {
