Disabled render graph building warnings inside test build, and made it off by default, as it takes considerable amount of render thread time.

diff --git a/Engine/Source/Runtime/RenderCore/Private/RenderGraphBuilder.cpp b/Engine/Source/Runtime/RenderCore/Private/RenderGraphBuilder.cpp
--- a/Engine/Source/Runtime/RenderCore/Private/RenderGraphBuilder.cpp
+++ b/Engine/Source/Runtime/RenderCore/Private/RenderGraphBuilder.cpp
@@ -16,7 +16,7 @@ static FAutoConsoleVariableRef CVarImmediateMode(
 	TEXT("Executes passes as they get created. Useful to have a callstack of the wiring code when crashing in the pass' lambda."),
 	ECVF_RenderThreadSafe);
 
-static int32 GRenderGraphEmitWarnings = 1;
+static int32 GRenderGraphEmitWarnings = 0;
 static FAutoConsoleVariableRef CVarEmitWarnings(
 	TEXT("r.RDG.EmitWarnings"),
 	GRenderGraphEmitWarnings,
diff --git a/Engine/Source/Runtime/RenderCore/Public/RenderGraphResources.h b/Engine/Source/Runtime/RenderCore/Public/RenderGraphResources.h
--- a/Engine/Source/Runtime/RenderCore/Public/RenderGraphResources.h
+++ b/Engine/Source/Runtime/RenderCore/Public/RenderGraphResources.h
@@ -8,7 +8,7 @@
 
 
 /** Whether render graph debugging is compiled. */
-#define RENDER_GRAPH_DEBUGGING (!UE_BUILD_SHIPPING)
+#define RENDER_GRAPH_DEBUGGING (!UE_BUILD_SHIPPING && !UE_BUILD_TEST)
 
 struct FRenderGraphPass;
 
