Remove source hashes from FShader in cooked builds. LLM showed ~5 MB reduction in ELLMTag::Shaders on Win64

diff --git a/Engine/Source/Runtime/RenderCore/Public/Shader.h b/Engine/Source/Runtime/RenderCore/Public/Shader.h
--- a/Engine/Source/Runtime/RenderCore/Public/Shader.h
+++ b/Engine/Source/Runtime/RenderCore/Public/Shader.h
@@ -1098,11 +1098,13 @@ private:
 	/** Hash of the material shader map this shader belongs to, stored so that an FShaderId can be constructed from this shader. */
 	FSHAHash MaterialShaderMapHash;
 
+#if KEEP_SHADER_SOURCE_HASHES
 	/** Vertex factory source hash, stored so that an FShaderId can be constructed from this shader. */
 	FSHAHash VFSourceHash;
 
 	/** Hash of this shader's source files generated at compile time, and stored to allow creating an FShaderId. */
 	FSHAHash SourceHash;
+#endif
 
 	/** Reference to the shader resource, which stores the compiled bytecode and the RHI shader resource. */
 	TRefCountPtr<FShaderResource> Resource;
