Fixes vis SceneDepthZ by hacking the SRV creating in the pass lambda

This hack is motivated because the render target pool and render graph requires to much work for 4.22 to supporting arbitrary SRV creation for a FRDGTexture properly.

diff --git a/Engine/Source/Runtime/RenderCore/Private/VisualizeTexture.cpp b/Engine/Source/Runtime/RenderCore/Private/VisualizeTexture.cpp
--- a/Engine/Source/Runtime/RenderCore/Private/VisualizeTexture.cpp
+++ b/Engine/Source/Runtime/RenderCore/Private/VisualizeTexture.cpp
@@ -90,7 +90,7 @@ class FVisualizeTexturePS : public FGlobalShader
 		SHADER_PARAMETER_SAMPLER(SamplerState, VisualizeTextureCubeSampler)
 		SHADER_PARAMETER_RDG_TEXTURE(TextureCubeArray, VisualizeTextureCubeArray)
 		SHADER_PARAMETER_SAMPLER(SamplerState, VisualizeTextureCubeArraySampler)
-		SHADER_PARAMETER_RDG_TEXTURE(Texture2D<uint4>, VisualizeDepthStencilTexture)
+		SHADER_PARAMETER_SRV(Texture2D<uint4>, VisualizeDepthStencilTexture)
 		SHADER_PARAMETER_RDG_TEXTURE(Texture2DMS<float4>, VisualizeTexture2DMS)
 		SHADER_PARAMETER_RDG_TEXTURE(Texture2D<uint>, VisualizeUINT8Texture2D)
 
@@ -194,33 +194,6 @@ void FVisualizeTexture::CreateContentCapturePass(FRDGBuilder& GraphBuilder, cons
 		LocalVisualizeTextureInputMapping = 1;
 	}
 
-#if 0 // TODO(RDG): requires changes in FRDGTextureSRVDesc
-	bool bIsDefault = StencilSRVSrc == GBlackTexture->TextureRHI;
-	bool bDepthStencil = SrcDesc.Is2DTexture() && SrcDesc.Format == PF_DepthStencil;
-
-	//clear if this is a new different Stencil buffer, or it's not a stencil buffer and we haven't switched to the default yet.
-	bool bNeedsClear = bDepthStencil && (StencilSRVSrc != RenderTargetItem.TargetableTexture);
-	bNeedsClear |= !bDepthStencil && !bIsDefault;
-	if (bNeedsClear)
-	{
-		StencilSRVSrc = nullptr;
-		StencilSRV.SafeRelease();
-	}
-
-	//always set something into the StencilSRV slot for platforms that require a full resource binding, even if
-	//dynamic branching will cause them not to be used.	
-	if (bDepthStencil && !StencilSRVSrc)
-	{
-		StencilSRVSrc = RenderTargetItem.TargetableTexture;
-		StencilSRV = RHICreateShaderResourceView((FTexture2DRHIRef&)RenderTargetItem.TargetableTexture, 0, 1, PF_X24_G8);
-	}
-	else if (!StencilSRVSrc)
-	{
-		StencilSRVSrc = GBlackTexture->TextureRHI;
-		StencilSRV = RHICreateShaderResourceView((FTexture2DRHIRef&)GBlackTexture->TextureRHI, 0, 1, PF_B8G8R8A8);
-	}
-#endif
-
 	// distinguish between standard depth and shadow depth to produce more reasonable default value mapping in the pixel shader.
 	const bool bDepthTexture = (SrcDesc.TargetableFlags & TexCreate_DepthStencilTargetable) != 0;
 	const bool bShadowDepth = (SrcDesc.Format == PF_ShadowDepth);
@@ -273,10 +246,66 @@ void FVisualizeTexture::CreateContentCapturePass(FRDGBuilder& GraphBuilder, cons
 
 	TShaderMapRef<FVisualizeTexturePS> PixelShader(ShaderMap, PermutationVector);
 
-	FPixelShaderUtils::AddFullscreenPass(
-		GraphBuilder, ShaderMap,
+	GraphBuilder.AddPass(
 		RDG_EVENT_NAME("VisualizeTextureCapture(%s)", SrcTexture->Name),
-		*PixelShader, PassParameters);
+		PassParameters,
+		ERenderGraphPassFlags::None,
+		[this, PassParameters, ShaderMap, PixelShader, RTExtent](FRHICommandList& RHICmdList)
+	{
+		FVisualizeTexturePS::FParameters ShaderParameter = *PassParameters;
+
+		// TODO(RDG): technically could use FPixelShaderUtils::AddPass(), but there is lot of work to support creating arbitrary number of SRV for a FRDGTexture,
+		// so the VisualizeDepthStencilTexture has to be hacked in the lambda...
+		{
+			// Some RHI might be unhappy with RHICreateShaderResourceView() inside renderpass.
+			check(RHICmdList.IsInsideRenderPass());
+			RHICmdList.EndRenderPass();
+			check(RHICmdList.IsOutsideRenderPass());
+
+			const FRDGTextureDesc& SrcDesc = PassParameters->VisualizeTexture2D->Desc;
+			FSceneRenderTargetItem& RenderTargetItem = PassParameters->VisualizeTexture2D->GetPooledRenderTarget()->GetRenderTargetItem();
+
+			bool bIsDefault = this->StencilSRVSrc == GBlackTexture->TextureRHI;
+			bool bDepthStencil = SrcDesc.Is2DTexture() && SrcDesc.Format == PF_DepthStencil;
+
+			//clear if this is a new different Stencil buffer, or it's not a stencil buffer and we haven't switched to the default yet.
+			bool bNeedsClear = bDepthStencil && (this->StencilSRVSrc != RenderTargetItem.TargetableTexture);
+			bNeedsClear |= !bDepthStencil && !bIsDefault;
+			if (bNeedsClear)
+			{
+				this->StencilSRVSrc = nullptr;
+				this->StencilSRV.SafeRelease();
+			}
+
+			//always set something into the StencilSRV slot for platforms that require a full resource binding, even if
+			//dynamic branching will cause them not to be used.	
+			if (bDepthStencil && !GVisualizeTexture.StencilSRVSrc)
+			{
+				this->StencilSRVSrc = RenderTargetItem.TargetableTexture;
+				this->StencilSRV = RHICreateShaderResourceView((FTexture2DRHIRef&)RenderTargetItem.TargetableTexture, 0, 1, PF_X24_G8);
+			}
+			else if (!GVisualizeTexture.StencilSRVSrc)
+			{
+				this->StencilSRVSrc = GBlackTexture->TextureRHI;
+				this->StencilSRV = RHICreateShaderResourceView((FTexture2DRHIRef&)GBlackTexture->TextureRHI, 0, 1, PF_B8G8R8A8);
+			}
+
+			ShaderParameter.VisualizeDepthStencilTexture = this->StencilSRV;
+
+			// Rebind the render targets.
+			FRHIRenderPassInfo RPInfo;
+			RPInfo.ColorRenderTargets[0].RenderTarget = PassParameters->RenderTargets[0].GetTexture()->GetPooledRenderTarget()->GetRenderTargetItem().TargetableTexture;
+			RPInfo.ColorRenderTargets[0].ResolveTarget = nullptr;
+			RPInfo.ColorRenderTargets[0].ArraySlice = -1;
+			RPInfo.ColorRenderTargets[0].MipIndex = 0;
+			RPInfo.ColorRenderTargets[0].Action = MakeRenderTargetActions(ERenderTargetLoadAction::EClear, ERenderTargetStoreAction::EStore);
+
+			RHICmdList.BeginRenderPass(RPInfo, TEXT("VisualizeTextureCapture"));
+			check(RHICmdList.IsInsideRenderPass());
+		}
+
+		FPixelShaderUtils::DrawFullscreenPixelShader(RHICmdList, ShaderMap, *PixelShader, ShaderParameter, FIntRect(0, 0, RTExtent.X, RTExtent.Y));
+	});
 
 	// Save the copied texture and descriptor about original informations.
 	{
diff --git a/Engine/Source/Runtime/RenderCore/Public/PixelShaderUtils.h b/Engine/Source/Runtime/RenderCore/Public/PixelShaderUtils.h
--- a/Engine/Source/Runtime/RenderCore/Public/PixelShaderUtils.h
+++ b/Engine/Source/Runtime/RenderCore/Public/PixelShaderUtils.h
@@ -30,7 +30,8 @@ struct RENDERCORE_API FPixelShaderUtils
 		FRHICommandList& RHICmdList, 
 		const TShaderMap<FGlobalShaderType>* GlobalShaderMap,
 		const TShaderClass* PixelShader,
-		const typename TShaderClass::FParameters& Parameters)
+		const typename TShaderClass::FParameters& Parameters,
+		const FIntRect& Viewport)
 	{
 		CA_ASSUME(PixelShader);
 		TShaderMapRef<FVisualizeTextureVS> VertexShader(GlobalShaderMap);
@@ -47,6 +48,8 @@ struct RENDERCORE_API FPixelShaderUtils
 		GraphicsPSOInit.PrimitiveType = PT_TriangleList;
 		SetGraphicsPipelineState(RHICmdList, GraphicsPSOInit);
 
+		RHICmdList.SetViewport(Viewport.Min.X, Viewport.Min.Y, 0.0f, Viewport.Max.X, Viewport.Max.Y, 1.0f);
+
 		SetShaderParameters(RHICmdList, PixelShader, PixelShader->GetPixelShader(), Parameters);
 
 		DrawFullscreenTriangle(RHICmdList);
@@ -59,7 +62,8 @@ struct RENDERCORE_API FPixelShaderUtils
 		const TShaderMap<FGlobalShaderType>* GlobalShaderMap,
 		FRDGEventName&& PassName,
 		const TShaderClass* PixelShader,
-		typename TShaderClass::FParameters* Parameters)
+		typename TShaderClass::FParameters* Parameters,
+		const FIntRect& Viewport)
 	{
 		CA_ASSUME(PixelShader != nullptr);
 		ClearUnusedGraphResources(PixelShader, Parameters);
@@ -68,9 +72,9 @@ struct RENDERCORE_API FPixelShaderUtils
 			Forward<FRDGEventName>(PassName),
 			Parameters,
 			ERenderGraphPassFlags::None,
-			[Parameters, GlobalShaderMap, PixelShader](FRHICommandList& RHICmdList)
+			[Parameters, GlobalShaderMap, PixelShader, Viewport](FRHICommandList& RHICmdList)
 		{
-			FPixelShaderUtils::DrawFullscreenPixelShader(RHICmdList, GlobalShaderMap, PixelShader, *Parameters);
+			FPixelShaderUtils::DrawFullscreenPixelShader(RHICmdList, GlobalShaderMap, PixelShader, *Parameters, Viewport);
 		});
 	}
 };
diff --git a/Engine/Source/Runtime/RenderCore/Public/RenderGraphResources.h b/Engine/Source/Runtime/RenderCore/Public/RenderGraphResources.h
--- a/Engine/Source/Runtime/RenderCore/Public/RenderGraphResources.h
+++ b/Engine/Source/Runtime/RenderCore/Public/RenderGraphResources.h
@@ -122,6 +122,13 @@ public:
 	/** Descriptor of the graph tracked texture. */
 	const FPooledRenderTargetDesc Desc;
 
+	/** Returns the allocated pooled render target. */
+	inline IPooledRenderTarget* GetPooledRenderTarget() const
+	{
+		check(PooledRenderTarget);
+		return PooledRenderTarget;
+	}
+
 private:
 	/** This is not a TRefCountPtr<> because FRDGTexture is allocated on the FMemStack
 	 * FGraphBuilder::AllocatedTextures is actually keeping the reference.
