Fixes a shadow variable warning in visualize texture.

diff --git a/Engine/Source/Runtime/RenderCore/Private/VisualizeTexture.cpp b/Engine/Source/Runtime/RenderCore/Private/VisualizeTexture.cpp
--- a/Engine/Source/Runtime/RenderCore/Private/VisualizeTexture.cpp
+++ b/Engine/Source/Runtime/RenderCore/Private/VisualizeTexture.cpp
@@ -262,11 +262,11 @@ void FVisualizeTexture::CreateContentCapturePass(FRDGBuilder& GraphBuilder, cons
 			RHICmdList.EndRenderPass();
 			check(RHICmdList.IsOutsideRenderPass());
 
-			const FRDGTextureDesc& SrcDesc = PassParameters->VisualizeTexture2D->Desc;
+			const FRDGTextureDesc& SrcDesc2 = PassParameters->VisualizeTexture2D->Desc;
 			FSceneRenderTargetItem& RenderTargetItem = PassParameters->VisualizeTexture2D->GetPooledRenderTarget()->GetRenderTargetItem();
 
 			bool bIsDefault = this->StencilSRVSrc == GBlackTexture->TextureRHI;
-			bool bDepthStencil = SrcDesc.Is2DTexture() && SrcDesc.Format == PF_DepthStencil;
+			bool bDepthStencil = SrcDesc2.Is2DTexture() && SrcDesc2.Format == PF_DepthStencil;
 
 			//clear if this is a new different Stencil buffer, or it's not a stencil buffer and we haven't switched to the default yet.
 			bool bNeedsClear = bDepthStencil && (this->StencilSRVSrc != RenderTargetItem.TargetableTexture);
