Save 24 bytes from FRenderResource for platforms that set PLATFORM_NEEDS_RHIRESOURCELIST to 0

diff --git a/Engine/Source/Runtime/Apple/MetalRHI/Private/MetalRHI.cpp b/Engine/Source/Runtime/Apple/MetalRHI/Private/MetalRHI.cpp
--- a/Engine/Source/Runtime/Apple/MetalRHI/Private/MetalRHI.cpp
+++ b/Engine/Source/Runtime/Apple/MetalRHI/Private/MetalRHI.cpp
@@ -856,18 +856,9 @@ FMetalDynamicRHI::FMetalDynamicRHI(ERHIFeatureLevel::Type RequestedFeatureLevel)
 	if (ImmediateContext.Profiler)
 		ImmediateContext.Profiler->BeginFrame();
 #endif
-		
-		
+	
 	// Notify all initialized FRenderResources that there's a valid RHI device to create their RHI resources for now.
-	for(TLinkedList<FRenderResource*>::TIterator ResourceIt(FRenderResource::GetResourceList());ResourceIt;ResourceIt.Next())
-	{
-		ResourceIt->InitRHI();
-	}
-	// Dynamic resources can have dependencies on static resources (with uniform buffers) and must initialized last!
-	for(TLinkedList<FRenderResource*>::TIterator ResourceIt(FRenderResource::GetResourceList());ResourceIt;ResourceIt.Next())
-	{
-		ResourceIt->InitDynamicRHI();
-	}
+	FRenderResource::InitPreRHIResources();	
 	
 	AsyncComputeContext = GSupportsEfficientAsyncCompute ? new FMetalRHIComputeContext(ImmediateContext.Profiler, new FMetalContext(ImmediateContext.Context->GetDevice(), ImmediateContext.Context->GetCommandQueue(), true)) : nullptr;
 
diff --git a/Engine/Source/Runtime/Core/Public/HAL/Platform.h b/Engine/Source/Runtime/Core/Public/HAL/Platform.h
--- a/Engine/Source/Runtime/Core/Public/HAL/Platform.h
+++ b/Engine/Source/Runtime/Core/Public/HAL/Platform.h
@@ -461,6 +461,10 @@
 	#define PLATFORM_HAS_CRC_INTRINSICS							0
 #endif
 
+#ifndef PLATFORM_NEEDS_RHIRESOURCELIST
+	#define PLATFORM_NEEDS_RHIRESOURCELIST 1
+#endif
+
 // deprecated, do not use
 #define PLATFORM_HAS_THREADSAFE_RHIGetRenderQueryResult	#
 #define PLATFORM_SUPPORTS_RHI_THREAD #
diff --git a/Engine/Source/Runtime/Core/Public/IOS/IOSPlatform.h b/Engine/Source/Runtime/Core/Public/IOS/IOSPlatform.h
--- a/Engine/Source/Runtime/Core/Public/IOS/IOSPlatform.h
+++ b/Engine/Source/Runtime/Core/Public/IOS/IOSPlatform.h
@@ -63,6 +63,7 @@ typedef FIOSPlatformTypes FPlatformTypes;
 #define PLATFORM_UI_NEEDS_TOOLTIPS						0
 #define PLATFORM_UI_NEEDS_FOCUS_OUTLINES				0
 
+#define PLATFORM_NEEDS_RHIRESOURCELIST					0
 #if WITH_SIMULATOR
 	#define PLATFORM_BREAK()							__asm__("int $3")
 #elif PLATFORM_64BITS
diff --git a/Engine/Source/Runtime/RenderCore/Private/RenderResource.cpp b/Engine/Source/Runtime/RenderCore/Private/RenderResource.cpp
--- a/Engine/Source/Runtime/RenderCore/Private/RenderResource.cpp
+++ b/Engine/Source/Runtime/RenderCore/Private/RenderResource.cpp
@@ -26,6 +26,30 @@ TLinkedList<FRenderResource*>*& FRenderResource::GetResourceList()
 	return FirstResourceLink;
 }
 
+/** Initialize all resources initialized before the RHI was initialized */
+void FRenderResource::InitPreRHIResources()
+{	
+	// Notify all initialized FRenderResources that there's a valid RHI device to create their RHI resources for now.
+	for (TLinkedList<FRenderResource*>::TIterator ResourceIt(FRenderResource::GetResourceList()); ResourceIt; ResourceIt.Next())
+	{
+		ResourceIt->InitRHI();
+	}
+	// Dynamic resources can have dependencies on static resources (with uniform buffers) and must initialized last!
+	for (TLinkedList<FRenderResource*>::TIterator ResourceIt(FRenderResource::GetResourceList()); ResourceIt; ResourceIt.Next())
+	{
+		ResourceIt->InitDynamicRHI();
+	}
+
+#if !PLATFORM_NEEDS_RHIRESOURCELIST
+	while (GetResourceList())
+	{
+		TLinkedList<FRenderResource*>* CurrentHead = GetResourceList();
+		CurrentHead->Unlink();
+		delete CurrentHead;
+	}
+#endif
+}
+
 void FRenderResource::ChangeFeatureLevel(ERHIFeatureLevel::Type NewFeatureLevel)
 {
 	ENQUEUE_RENDER_COMMAND(FRenderResourceChangeFeatureLevel)(
@@ -53,13 +77,22 @@ void FRenderResource::InitResource()
 	check(IsInRenderingThread());
 	if(!bInitialized)
 	{
+#if PLATFORM_NEEDS_RHIRESOURCELIST
 		ResourceLink = TLinkedList<FRenderResource*>(this);
 		ResourceLink.LinkHead(GetResourceList());
+#endif
 		if(GIsRHIInitialized)
 		{
 			InitDynamicRHI();
 			InitRHI();
 		}
+		else
+		{
+#if !PLATFORM_NEEDS_RHIRESOURCELIST
+		TLinkedList<FRenderResource*>* ListEntry = new TLinkedList<FRenderResource*>(this);
+		ListEntry->LinkHead(GetResourceList());
+#endif
+		}
 		FPlatformMisc::MemoryBarrier(); // there are some multithreaded reads of bInitialized
 		bInitialized = true;
 	}
@@ -77,7 +110,9 @@ void FRenderResource::ReleaseResource()
 				ReleaseRHI();
 				ReleaseDynamicRHI();
 			}
+#if PLATFORM_NEEDS_RHIRESOURCELIST
 			ResourceLink.Unlink();
+#endif
 			bInitialized = false;
 		}
 	}
diff --git a/Engine/Source/Runtime/RenderCore/Public/RenderResource.h b/Engine/Source/Runtime/RenderCore/Public/RenderResource.h
--- a/Engine/Source/Runtime/RenderCore/Public/RenderResource.h
+++ b/Engine/Source/Runtime/RenderCore/Public/RenderResource.h
@@ -95,6 +95,9 @@ public:
 	// Accessors.
 	FORCEINLINE bool IsInitialized() const { return bInitialized; }
 
+	/** Initialize all resources initialized before the RHI was initialized */
+	static void InitPreRHIResources();
+
 protected:
 	// This is used during mobile editor preview refactor, this will eventually be replaced with a parameter to InitRHI() etc..
 	ERHIFeatureLevel::Type GetFeatureLevel() const { return FeatureLevel == ERHIFeatureLevel::Num ? GMaxRHIFeatureLevel : FeatureLevel; }
@@ -103,12 +106,12 @@ protected:
 	ERHIFeatureLevel::Type FeatureLevel;
 
 private:
+	#if PLATFORM_NEEDS_RHIRESOURCELIST
+	TLinkedList<FRenderResource*> ResourceLink;
+	#endif
 
 	/** True if the resource has been initialized. */
 	bool bInitialized;
-
-	/** This resource's link in the global resource list. */
-	TLinkedList<FRenderResource*> ResourceLink;
 };
 
 /**
