Shave some memory off of FVertexFactory
FLocalVertexFactory goes from 1072 bytes -> 480

diff --git a/Engine/Source/Runtime/RenderCore/Public/VertexFactory.h b/Engine/Source/Runtime/RenderCore/Public/VertexFactory.h
--- a/Engine/Source/Runtime/RenderCore/Public/VertexFactory.h
+++ b/Engine/Source/Runtime/RenderCore/Public/VertexFactory.h
@@ -575,9 +575,10 @@ protected:
 	struct FVertexStream
 	{
 		const FVertexBuffer* VertexBuffer = nullptr;
-		uint32 Stride = 0;
 		uint32 Offset = 0;
+		uint16 Stride = 0;
 		EVertexStreamUsage VertexStreamUsage = EVertexStreamUsage::Default;
+		uint8 Padding = 0;
 
 		friend bool operator==(const FVertexStream& A,const FVertexStream& B)
 		{
@@ -590,7 +591,7 @@ protected:
 	};
 
 	/** The vertex streams used to render the factory. */
-	TArray<FVertexStream,TFixedAllocator<MaxVertexElementCount> > Streams;
+	TArray<FVertexStream,TInlineAllocator<8> > Streams;
 
 	/* VF can explicitly set this to false to avoid errors without decls; this is for VFs that fetch from buffers directly (e.g. Niagara) */
 	bool bNeedsDeclaration = true;
@@ -603,7 +604,7 @@ protected:
 private:
 
 	/** The position only vertex stream used to render the factory during depth only passes. */
-	TArray<FVertexStream,TFixedAllocator<MaxVertexElementCount> > PositionStream;
+	TArray<FVertexStream,TInlineAllocator<2> > PositionStream;
 
 	/** The RHI vertex declaration used to render the factory normally. */
 	FVertexDeclarationRHIRef Declaration;
