UE4 - Fix compile error
#rnx

diff --git a/Engine/Source/Runtime/RenderCore/Private/VertexFactory.cpp b/Engine/Source/Runtime/RenderCore/Private/VertexFactory.cpp
--- a/Engine/Source/Runtime/RenderCore/Private/VertexFactory.cpp
+++ b/Engine/Source/Runtime/RenderCore/Private/VertexFactory.cpp
@@ -210,8 +210,7 @@ void FVertexFactory::GetStreams(ERHIFeatureLevel::Type InFeatureLevel, FVertexIn
 
 void FVertexFactory::OffsetInstanceStreams(uint32 InstanceOffset, bool bOperateOnPositionOnly, FVertexInputStreamArray& VertexStreams) const
 {
-	const TArray<FVertexStream,TFixedAllocator<MaxVertexElementCount> >& StreamArray = bOperateOnPositionOnly ? PositionStream : Streams;
-
+	const TArrayView<const FVertexStream>& StreamArray = bOperateOnPositionOnly ? MakeArrayView(PositionStream) : MakeArrayView(Streams);
 	for (int32 StreamIndex = 0; StreamIndex < StreamArray.Num(); StreamIndex++)
 	{
 		const FVertexStream& Stream = StreamArray[StreamIndex];
