Remove FShader::OutputHash in non-editor build since it is a copy of FShaderResource::OutputHash. It is not used in non-editor build so cache miss isn't  a concern.

diff --git a/Engine/Source/Runtime/RenderCore/Private/Shader.cpp b/Engine/Source/Runtime/RenderCore/Private/Shader.cpp
--- a/Engine/Source/Runtime/RenderCore/Private/Shader.cpp
+++ b/Engine/Source/Runtime/RenderCore/Private/Shader.cpp
@@ -1152,12 +1152,12 @@ FShader::FShader(const CompiledShaderInitializerType& Initializer)
 	, Target(Initializer.Target)
 	, NumRefs(0)
 {
-	OutputHash = Initializer.OutputHash;
-	checkSlow(OutputHash != FSHAHash());
-
 	check(Type);
 
 #if KEEP_SHADER_SOURCE_HASHES
+	OutputHash = Initializer.OutputHash;
+	checkSlow(OutputHash != FSHAHash());
+
 	// Store off the source hash that this shader was compiled with
 	// This will be used as part of the shader key in order to identify when shader files have been changed and a recompile is needed
 	SourceHash = Type->GetSourceHash(Target.GetPlatform());
@@ -1226,7 +1226,7 @@ bool FShader::SerializeBase(FArchive& Ar, bool bShadersInline, bool bLoadedByCoo
 	FSHAHash& VFHash = VFSourceHash;
 	FSHAHash& Hash = SourceHash;
 #else
-	FSHAHash VFHash, Hash;
+	FSHAHash VFHash, Hash, OutputHash;
 #endif
 
 	Ar << OutputHash;
@@ -1315,6 +1315,7 @@ bool FShader::SerializeBase(FArchive& Ar, bool bShadersInline, bool bLoadedByCoo
 			// Load the inlined shader resource
 			SerializedResource = new FShaderResource();
 			SerializedResource->Serialize(Ar, bLoadedByCookedMaterial);
+			checkSlow(OutputHash == SerializedResource->OutputHash);
 		}
 	}
 	else
@@ -1432,8 +1433,8 @@ void FShader::DumpDebugInfo()
 #if KEEP_SHADER_SOURCE_HASHES
 	UE_LOG(LogConsoleResponse, Display, TEXT("               :SourceHash %s"), *SourceHash.ToString());
 	UE_LOG(LogConsoleResponse, Display, TEXT("               :VFSourceHash %s"), *VFSourceHash.ToString());
-#endif
 	UE_LOG(LogConsoleResponse, Display, TEXT("               :OutputHash %s"), *OutputHash.ToString());
+#endif
 }
 
 void FShader::SaveShaderStableKeys(EShaderPlatform TargetShaderPlatform, const FStableShaderKeyAndValue& InSaveKeyVal)
diff --git a/Engine/Source/Runtime/RenderCore/Public/Shader.h b/Engine/Source/Runtime/RenderCore/Public/Shader.h
--- a/Engine/Source/Runtime/RenderCore/Public/Shader.h
+++ b/Engine/Source/Runtime/RenderCore/Public/Shader.h
@@ -941,7 +941,15 @@ public:
 #endif
 	inline const TArray<uint8>& GetCode() const { return Resource->Code; }
 	inline const FShaderTarget GetTarget() const { return Target; }
-	inline FSHAHash GetOutputHash() const { return OutputHash; }
+	inline FSHAHash GetOutputHash() const
+	{
+#if KEEP_SHADER_SOURCE_HASHES
+		return OutputHash;
+#else
+		check(Resource);
+		return Resource->OutputHash;
+#endif
+	}
 	FShaderId GetId() const;
 	inline FVertexFactoryType* GetVertexFactoryType() const { return VFType; }
 	inline int32 GetNumRefs() const { return NumRefs; }
@@ -1088,17 +1096,16 @@ protected:
 	TArray<FShaderUniformBufferParameter*> UniformBufferParameters;
 
 private:
-
-	/** 
-	 * Hash of the compiled output from this shader and the resulting parameter map.  
-	 * This is used to find a matching resource.
-	 */
-	FSHAHash OutputHash;
-
 	/** Hash of the material shader map this shader belongs to, stored so that an FShaderId can be constructed from this shader. */
 	FSHAHash MaterialShaderMapHash;
 
 #if KEEP_SHADER_SOURCE_HASHES
+	/**
+	* Hash of the compiled output from this shader and the resulting parameter map.
+	* This is used to find a matching resource.
+	*/
+	FSHAHash OutputHash;
+
 	/** Vertex factory source hash, stored so that an FShaderId can be constructed from this shader. */
 	FSHAHash VFSourceHash;
 
