UE4 - Fix clang compile errors
#rnx

diff --git a/Engine/Source/Runtime/OpenGLDrv/Private/OpenGLUniformBuffer.cpp b/Engine/Source/Runtime/OpenGLDrv/Private/OpenGLUniformBuffer.cpp
--- a/Engine/Source/Runtime/OpenGLDrv/Private/OpenGLUniformBuffer.cpp
+++ b/Engine/Source/Runtime/OpenGLDrv/Private/OpenGLUniformBuffer.cpp
@@ -11,6 +11,7 @@
 #include "RHI.h"
 #include "OpenGLDrv.h"
 #include "OpenGLDrvPrivate.h"
+#include "Misc/ScopeLock.h"
 
 
 namespace OpenGLConsoleVariables
diff --git a/Engine/Source/Runtime/RHI/Private/RHIUtilities.cpp b/Engine/Source/Runtime/RHI/Private/RHIUtilities.cpp
--- a/Engine/Source/Runtime/RHI/Private/RHIUtilities.cpp
+++ b/Engine/Source/Runtime/RHI/Private/RHIUtilities.cpp
@@ -10,6 +10,7 @@ RHIUtilities.cpp:
 #include "RHI.h"
 #include "HAL/Runnable.h"
 #include "HAL/RunnableThread.h"
+#include "Misc/ScopeLock.h"
 
 #define USE_FRAME_OFFSET_THREAD 1
 
diff --git a/Engine/Source/Runtime/RenderCore/Private/DynamicBufferAllocator.cpp b/Engine/Source/Runtime/RenderCore/Private/DynamicBufferAllocator.cpp
--- a/Engine/Source/Runtime/RenderCore/Private/DynamicBufferAllocator.cpp
+++ b/Engine/Source/Runtime/RenderCore/Private/DynamicBufferAllocator.cpp
@@ -6,6 +6,7 @@ DynamicBufferAllocator.cpp: Classes for allocating transient rendering data.
 
 #include "DynamicBufferAllocator.h"
 #include "RenderResource.h"
+#include "Misc/ScopeLock.h"
 
 int32 GMaxReadBufferRenderingBytesAllocatedPerFrame = 32 * 1024 * 1024;
 
diff --git a/Engine/Source/Runtime/RenderCore/Private/ShaderPipelineCache.cpp b/Engine/Source/Runtime/RenderCore/Private/ShaderPipelineCache.cpp
--- a/Engine/Source/Runtime/RenderCore/Private/ShaderPipelineCache.cpp
+++ b/Engine/Source/Runtime/RenderCore/Private/ShaderPipelineCache.cpp
@@ -24,6 +24,7 @@
 #include "TickableObjectRenderThread.h"
 #include "Misc/ConfigCacheIni.h"
 #include "Async/AsyncFileHandle.h"
+#include "Misc/ScopeLock.h"
 
 DECLARE_DWORD_ACCUMULATOR_STAT(TEXT("Outstanding Tasks"), STAT_ShaderPipelineTaskCount, STATGROUP_PipelineStateCache );
 DECLARE_DWORD_ACCUMULATOR_STAT(TEXT("Waiting Tasks"), STAT_ShaderPipelineWaitingTaskCount, STATGROUP_PipelineStateCache );
