Do not check for debugger unless already timed out.

- This check is somewhat expensive under Linux (opening/reading files etc).

diff --git a/Engine/Source/Runtime/RenderCore/Private/RenderingThread.cpp b/Engine/Source/Runtime/RenderCore/Private/RenderingThread.cpp
--- a/Engine/Source/Runtime/RenderCore/Private/RenderingThread.cpp
+++ b/Engine/Source/Runtime/RenderCore/Private/RenderingThread.cpp
@@ -1129,9 +1129,9 @@ static void GameThreadWaitForTask(const FGraphEventRef& Task, ENamedThreads::Typ
 #if !WITH_EDITOR
 #if !PLATFORM_IOS && !PLATFORM_MAC // @todo MetalMRT: Timeout isn't long enough...
 				// editor threads can block for quite a while... 
-				if (!bDone && !bRenderThreadEnsured && !FPlatformMisc::IsDebuggerPresent())
+				if (!bDone && !bRenderThreadEnsured)
 				{
-					if (bOverdue && !bDisabled)
+					if (bOverdue && !bDisabled && !FPlatformMisc::IsDebuggerPresent())
 					{
 						UE_LOG(LogRendererCore, Fatal, TEXT("GameThread timed out waiting for RenderThread after %.02f secs"), RenderThreadTimeoutClock.Seconds() - StartTime);
 					}
