RL - Remove Tmap for Tarray in RTLib because functionallity of Tmap was not needed.

diff --git a/Engine/Source/Runtime/RenderCore/Private/Shader.cpp b/Engine/Source/Runtime/RenderCore/Private/Shader.cpp
--- a/Engine/Source/Runtime/RenderCore/Private/Shader.cpp
+++ b/Engine/Source/Runtime/RenderCore/Private/Shader.cpp
@@ -392,21 +392,14 @@ void FShaderType::Uninitialize()
 TMap<FShaderResourceId, FShaderResource*> FShaderResource::ShaderResourceIdMap;
 
 #if RHI_RAYTRACING
-uint32 FShaderResource::GlobalMaxIndex = 0;
 TArray<uint32> FShaderResource::GlobalUnusedIndicies;
-TMap<uint32, FRHIRayTracingShader*> FShaderResource::GlobalRayTracingMaterialLibrary;
+TArray<FRHIRayTracingShader*> FShaderResource::GlobalRayTracingMaterialLibrary;
 FCriticalSection FShaderResource::GlobalRayTracingMaterialLibraryCS;
 
 void FShaderResource::GetRayTracingMaterialLibrary(TArray<FRayTracingShaderRHIParamRef>& RayTracingMaterials, FRayTracingShaderRHIParamRef DefaultShader)
 {
 	FScopeLock Lock(&GlobalRayTracingMaterialLibraryCS);
-	RayTracingMaterials.Reset();
-	RayTracingMaterials.AddUninitialized(GlobalMaxIndex);
-
-	for (const auto Entry : GlobalRayTracingMaterialLibrary)
-	{
-		RayTracingMaterials[Entry.Key] = Entry.Value;
-	}
+	RayTracingMaterials = GlobalRayTracingMaterialLibrary;
 
 	for (uint32 Index : GlobalUnusedIndicies)
 	{
@@ -421,16 +414,14 @@ uint32 FShaderResource::AddToRayTracingLibrary(FRayTracingShaderRHIParamRef Shad
 	if (GlobalUnusedIndicies.Num() != 0)
 	{
 		uint32 Index = GlobalUnusedIndicies.Pop(false);
-		checkSlow(GlobalRayTracingMaterialLibrary.Find(Index) == nullptr);
-		GlobalRayTracingMaterialLibrary.Add(Index, Shader);
+		checkSlow(GlobalRayTracingMaterialLibrary[Index] == nullptr);
+		GlobalRayTracingMaterialLibrary[Index] = Shader;
 		return Index;
 	}
 	else
 	{
-		uint32 Index = GlobalMaxIndex++;
-		checkSlow(GlobalRayTracingMaterialLibrary.Find(Index) == nullptr);
-		GlobalRayTracingMaterialLibrary.Add(Index, Shader);
-		return Index;
+		GlobalRayTracingMaterialLibrary.Add(Shader);
+		return GlobalRayTracingMaterialLibrary.Num() - 1;
 	}
 }
 
@@ -438,7 +429,7 @@ void FShaderResource::RemoveFromRayTracingLibrary(uint32 Index)
 {
 	FScopeLock Lock(&GlobalRayTracingMaterialLibraryCS);
 	GlobalUnusedIndicies.Push(Index);
-	GlobalRayTracingMaterialLibrary.Remove(Index);
+	GlobalRayTracingMaterialLibrary[Index] = nullptr;
 }
 #endif
 
diff --git a/Engine/Source/Runtime/RenderCore/Public/Shader.h b/Engine/Source/Runtime/RenderCore/Public/Shader.h
--- a/Engine/Source/Runtime/RenderCore/Public/Shader.h
+++ b/Engine/Source/Runtime/RenderCore/Public/Shader.h
@@ -386,9 +386,8 @@ private:
 	RENDERCORE_API static uint32 AddToRayTracingLibrary(FRayTracingShaderRHIParamRef Shader);
 	RENDERCORE_API static void RemoveFromRayTracingLibrary(uint32 Index);
 
-	static uint32 GlobalMaxIndex;
 	static TArray<uint32> GlobalUnusedIndicies;
-	static TMap<uint32, FRayTracingShaderRHIParamRef> GlobalRayTracingMaterialLibrary;
+	static TArray<FRHIRayTracingShader*> GlobalRayTracingMaterialLibrary;
 	static FCriticalSection GlobalRayTracingMaterialLibraryCS;
 
 public:
