Fix crash on exit in InfiltratorDemo

diff --git a/Engine/Source/Runtime/RenderCore/Private/Shader.cpp b/Engine/Source/Runtime/RenderCore/Private/Shader.cpp
--- a/Engine/Source/Runtime/RenderCore/Private/Shader.cpp
+++ b/Engine/Source/Runtime/RenderCore/Private/Shader.cpp
@@ -976,7 +976,7 @@ void FShaderResource::ReleaseRHI()
 	DEC_DWORD_STAT_BY(STAT_Shaders_NumShadersUsedForRendering, 1);
 
 #if RHI_RAYTRACING
-	if (IsInitialized() && Target.Frequency == SF_RayHitGroup)
+	if (IsInitialized() && RayTracingMaterialLibraryIndex != UINT_MAX)
 	{
 		RemoveFromRayTracingLibrary(RayTracingMaterialLibraryIndex);
 		RayTracingMaterialLibraryIndex = UINT_MAX;
