Temporary change to quantize size used by compute shaders from 8 to 4

diff --git a/Engine/Source/Runtime/RenderCore/Private/RenderUtils.cpp b/Engine/Source/Runtime/RenderCore/Private/RenderUtils.cpp
--- a/Engine/Source/Runtime/RenderCore/Private/RenderUtils.cpp
+++ b/Engine/Source/Runtime/RenderCore/Private/RenderUtils.cpp
@@ -1040,7 +1040,7 @@ RENDERCORE_API FIndexBufferRHIRef& GetUnitCubeIndexBuffer()
 RENDERCORE_API void QuantizeSceneBufferSize(const FIntPoint& InBufferSize, FIntPoint& OutBufferSize)
 {
 	// Ensure sizes are dividable by the ideal group size for 2d tiles to make it more convenient.
-	const uint32 DividableBy = FComputeShaderUtils::kGolden2DGroupSize;
+	const uint32 DividableBy = 4;
 
 	check(DividableBy % 4 == 0); // A lot of graphic algorithms where previously assuming DividableBy == 4.
 
