UE4 - Compile fix
#rnx

diff --git a/Engine/Source/Runtime/RenderCore/Private/Shader.cpp b/Engine/Source/Runtime/RenderCore/Private/Shader.cpp
--- a/Engine/Source/Runtime/RenderCore/Private/Shader.cpp
+++ b/Engine/Source/Runtime/RenderCore/Private/Shader.cpp
@@ -18,6 +18,7 @@
 #include "Misc/ConfigCacheIni.h"
 #include "UObject/RenderingObjectVersion.h"
 #include "UObject/FortniteMainBranchObjectVersion.h"
+#include "Misc/ScopeLock.h"
 
 #if WITH_EDITORONLY_DATA
 #include "Interfaces/IShaderFormat.h"
