Makes some changes to the reflection denoiser interface and render graph for third party denoiser plugins.

diff --git a/Engine/Shaders/Private/RayTracing/RayTracingReflections.usf b/Engine/Shaders/Private/RayTracing/RayTracingReflections.usf
--- a/Engine/Shaders/Private/RayTracing/RayTracingReflections.usf
+++ b/Engine/Shaders/Private/RayTracing/RayTracingReflections.usf
@@ -67,6 +67,7 @@ RaytracingAccelerationStructure TLAS;
 
 RWTexture2D<float4> ColorOutput;
 RWTexture2D<float> RayHitDistanceOutput;
+RWTexture2D<float> RayImaginaryDepthOutput;
 
 
 // Material buffer and tile size are used for sorted materials.
@@ -434,6 +435,18 @@ void RayTracingReflectionsRGS()
 
 		ColorOutput[DispatchThreadId] = ClampToHalfFloatRange(ReflectedColor);
 		RayHitDistanceOutput[DispatchThreadId] = HitDistance;
+
+		// Imaginary depth computation for Nvidia's denoiser.
+		// TODO: shader permutation?
+		{
+			float3 OriginalWorldPos = ReconstructWorldPositionFromDepth(UV, Depth);
+			float3 CameraToPixel = normalize(OriginalWorldPos - View.WorldCameraOrigin.xyz);
+
+			float4 ImaginaryWorldPos = float4(OriginalWorldPos + CameraToPixel * HitDistance, 1.0f);
+			float4 ImaginaryClipPos = mul(ImaginaryWorldPos, View.WorldToClip);
+			float ImaginarySVDepth = saturate(ImaginaryClipPos.z / ImaginaryClipPos.w);
+			RayImaginaryDepthOutput[DispatchThreadId] = ImaginarySVDepth;
+		}
 	}
 }
 
diff --git a/Engine/Source/Runtime/Engine/Private/SceneView.cpp b/Engine/Source/Runtime/Engine/Private/SceneView.cpp
--- a/Engine/Source/Runtime/Engine/Private/SceneView.cpp
+++ b/Engine/Source/Runtime/Engine/Private/SceneView.cpp
@@ -2245,6 +2245,7 @@ void FSceneView::SetupCommonViewUniformBufferParameters(
 	ViewUniformShaderParameters.ViewToTranslatedWorld = InViewMatrices.GetOverriddenInvTranslatedViewMatrix();
 	ViewUniformShaderParameters.TranslatedWorldToClip = InViewMatrices.GetTranslatedViewProjectionMatrix();
 	ViewUniformShaderParameters.WorldToClip = InViewMatrices.GetViewProjectionMatrix();
+	ViewUniformShaderParameters.ClipToWorld = InViewMatrices.GetInvViewProjectionMatrix();
 	ViewUniformShaderParameters.TranslatedWorldToView = InViewMatrices.GetOverriddenTranslatedViewMatrix();
 	ViewUniformShaderParameters.TranslatedWorldToCameraView = InViewMatrices.GetTranslatedViewMatrix();
 	ViewUniformShaderParameters.CameraViewToTranslatedWorld = InViewMatrices.GetInvTranslatedViewMatrix();
diff --git a/Engine/Source/Runtime/Engine/Public/SceneView.h b/Engine/Source/Runtime/Engine/Public/SceneView.h
--- a/Engine/Source/Runtime/Engine/Public/SceneView.h
+++ b/Engine/Source/Runtime/Engine/Public/SceneView.h
@@ -564,6 +564,7 @@ enum ETranslucencyVolumeCascade
 #define VIEW_UNIFORM_BUFFER_MEMBER_TABLE \
 	VIEW_UNIFORM_BUFFER_MEMBER(FMatrix, TranslatedWorldToClip) \
 	VIEW_UNIFORM_BUFFER_MEMBER(FMatrix, WorldToClip) \
+	VIEW_UNIFORM_BUFFER_MEMBER(FMatrix, ClipToWorld)  \
 	VIEW_UNIFORM_BUFFER_MEMBER(FMatrix, TranslatedWorldToView) \
 	VIEW_UNIFORM_BUFFER_MEMBER(FMatrix, ViewToTranslatedWorld) \
 	VIEW_UNIFORM_BUFFER_MEMBER(FMatrix, TranslatedWorldToCameraView) \
diff --git a/Engine/Source/Runtime/RenderCore/Public/RenderGraphBuilder.h b/Engine/Source/Runtime/RenderCore/Public/RenderGraphBuilder.h
--- a/Engine/Source/Runtime/RenderCore/Public/RenderGraphBuilder.h
+++ b/Engine/Source/Runtime/RenderCore/Public/RenderGraphBuilder.h
@@ -353,7 +353,7 @@ public:
 		return OutParameterPtr;
 	}
 
-	/** Adds a lambda pass to the graph.
+	/** Adds a hard coded lambda pass to the graph.
 	 *
 	 * The Name of the pass should be generated with enough information to identify it's purpose and GPU cost, to be clear
 	 * for GPU profiling tools.
@@ -402,6 +402,31 @@ public:
 		#endif
 	}
 
+	/** Adds a procedurally created pass to the render graph.
+	 *
+	 * Note: You want to use this only when the layout of the pass might be procedurally generated from data driven, as opose to AddPass() that have,
+	 * constant hard coded pass layout.
+	 *
+	 * Caution: You are on your own to have correct memory lifetime of the FRenderGraphPass.
+	 */
+	void AddProcedurallyCreatedPass(FRenderGraphPass* NewPass)
+	{
+		#if RENDER_GRAPH_DEBUGGING
+		{
+			checkf(!bHasExecuted, TEXT("Render graph pass %s needs to be added before the builder execution."), NewPass->GetName());
+		}
+		#endif
+
+		// TODO(RDG): perhaps the CurrentScope could be set here instead of GetCurrentScope(), to not allow user code to start adding pass to random scopes. 
+		Passes.Emplace(NewPass);
+
+		#if RENDER_GRAPH_DEBUGGING || SUPPORTS_VISUALIZE_TEXTURE
+		{
+			DebugPass(NewPass);
+		}
+		#endif
+	}
+
 	/** Queue a texture extraction. This will set *OutTexturePtr with the internal pooled render target at the Execute().
 	 *
 	 * Note: even when the render graph uses the immediate debugging mode (executing passes as they get added), the texture extrations
@@ -451,6 +476,11 @@ public:
 	 */
 	void Execute();
 
+	/** Returns the draw event scope, where passes are currently being added in. */
+	const FRDGEventScope* GetCurrentScope() const
+	{
+		return CurrentScope;
+	}
 
 public:
 	/** The RHI command list used for the render graph. */
diff --git a/Engine/Source/Runtime/RenderCore/Public/RenderGraphResources.h b/Engine/Source/Runtime/RenderCore/Public/RenderGraphResources.h
--- a/Engine/Source/Runtime/RenderCore/Public/RenderGraphResources.h
+++ b/Engine/Source/Runtime/RenderCore/Public/RenderGraphResources.h
@@ -66,6 +66,9 @@ public:
 	FRDGResource(const FRDGResource&) = delete;
 	void operator = (const FRDGResource&) = delete;
 
+	/** Boolean to track at runtime whether a ressource is actually used by the lambda of a pass or not, to detect unnecessary resource dependencies on passes. */
+	mutable bool bIsActuallyUsedByPass = false;
+
 private:
 	/** Number of references in passes and deferred queries. */
 	mutable int32 ReferenceCount = 0;
@@ -74,9 +77,6 @@ private:
 	mutable bool bWritable = false;
 	mutable bool bCompute = false;
 
-	/** Boolean to track at runtime whether a ressource is actually used by the lambda of a pass or not, to detect unnecessary resource dependencies on passes. */
-	mutable bool bIsActuallyUsedByPass = false;
-
 #if RENDER_GRAPH_DEBUGGING
 	/** Boolean to track at wiring time if a resource has ever been produced by a pass, to error out early if accessing a resource that has not been produced. */
 	mutable bool bHasEverBeenProduced = false;
@@ -122,13 +122,20 @@ public:
 	/** Descriptor of the graph tracked texture. */
 	const FPooledRenderTargetDesc Desc;
 
-	/** Returns the allocated pooled render target. */
+	/** Returns the allocated pooled render target. Must only be called within a pass's lambda. */
 	inline IPooledRenderTarget* GetPooledRenderTarget() const
 	{
 		check(PooledRenderTarget);
 		return PooledRenderTarget;
 	}
 
+	/** Returns the allocated RHI texture. Must only be called within a pass's lambda. */
+	inline FTextureRHIParamRef GetRHITexture() const
+	{
+		check(PooledRenderTarget);
+		return PooledRenderTarget->GetRenderTargetItem().ShaderResourceTexture;
+	}
+
 private:
 	/** This is not a TRefCountPtr<> because FRDGTexture is allocated on the FMemStack
 	 * FGraphBuilder::AllocatedTextures is actually keeping the reference.
@@ -140,13 +147,6 @@ private:
 		, Desc(InDesc)
 	{ }
 
-	/** Returns the allocated RHI texture. */
-	inline FTextureRHIParamRef GetRHITexture() const
-	{
-		check(PooledRenderTarget);
-		return PooledRenderTarget->GetRenderTargetItem().ShaderResourceTexture;
-	}
-
 	friend class FRDGBuilder;
 
 	template<typename TRHICmdList, typename TShaderClass, typename TShaderRHI>
diff --git a/Engine/Source/Runtime/Renderer/Private/DeferredShadingRenderer.h b/Engine/Source/Runtime/Renderer/Private/DeferredShadingRenderer.h
--- a/Engine/Source/Runtime/Renderer/Private/DeferredShadingRenderer.h
+++ b/Engine/Source/Runtime/Renderer/Private/DeferredShadingRenderer.h
@@ -445,6 +445,7 @@ private:
 		const FViewInfo& View,
 		FRDGTextureRef* OutColorTexture,
 		FRDGTextureRef* OutRayHitDistanceTexture,
+		FRDGTextureRef* OutRayImaginaryDepthTexture,
 		int32 SamplePerPixel,
 		int32 HeightFog,
 		float ResolutionFraction);
diff --git a/Engine/Source/Runtime/Renderer/Private/RayTracing/RayTracingReflections.cpp b/Engine/Source/Runtime/Renderer/Private/RayTracing/RayTracingReflections.cpp
--- a/Engine/Source/Runtime/Renderer/Private/RayTracing/RayTracingReflections.cpp
+++ b/Engine/Source/Runtime/Renderer/Private/RayTracing/RayTracingReflections.cpp
@@ -145,6 +145,7 @@ class FRayTracingReflectionsRGS : public FGlobalShader
 
 		SHADER_PARAMETER_RDG_TEXTURE_UAV(RWTexture2D<float4>, ColorOutput)
 		SHADER_PARAMETER_RDG_TEXTURE_UAV(RWTexture2D<float>, RayHitDistanceOutput)
+		SHADER_PARAMETER_RDG_TEXTURE_UAV(RWTexture2D<float>, RayImaginaryDepthOutput)
 	END_SHADER_PARAMETER_STRUCT()
 
 	static bool ShouldCompilePermutation(const FGlobalShaderPermutationParameters& Parameters)
@@ -206,6 +207,7 @@ void FDeferredShadingSceneRenderer::RenderRayTracingReflections(
 	const FViewInfo& View,
 	FRDGTextureRef* OutColorTexture,
 	FRDGTextureRef* OutRayHitDistanceTexture,
+	FRDGTextureRef* OutRayImaginaryDepthTexture,
 	int32 SamplePerPixel,
 	int32 HeightFog,
 	float ResolutionFraction)
@@ -232,6 +234,7 @@ void FDeferredShadingSceneRenderer::RenderRayTracingReflections(
 		
 		Desc.Format = PF_R16F;
 		*OutRayHitDistanceTexture = GraphBuilder.CreateTexture(Desc, TEXT("RayTracingReflectionsHitDistance"));
+		*OutRayImaginaryDepthTexture = GraphBuilder.CreateTexture(Desc, TEXT("RayTracingReflectionsImaginaryDepth"));
 	}
 
 	// When deferred materials are used, we need to dispatch the reflection shader twice:
@@ -280,6 +283,7 @@ void FDeferredShadingSceneRenderer::RenderRayTracingReflections(
 	CommonParameters.IESLightProfileParameters = CreateIESLightProfilesUniformBuffer(View, EUniformBufferUsage::UniformBuffer_SingleFrame);
 	CommonParameters.ColorOutput = GraphBuilder.CreateUAV(*OutColorTexture);
 	CommonParameters.RayHitDistanceOutput = GraphBuilder.CreateUAV(*OutRayHitDistanceTexture);
+	CommonParameters.RayImaginaryDepthOutput = GraphBuilder.CreateUAV(*OutRayImaginaryDepthTexture);
 	CommonParameters.SortTileSize = SortTileSize;
 
 	for (uint32 PassIndex = 0; PassIndex < NumPasses; ++PassIndex)
diff --git a/Engine/Source/Runtime/Renderer/Private/ReflectionEnvironment.cpp b/Engine/Source/Runtime/Renderer/Private/ReflectionEnvironment.cpp
--- a/Engine/Source/Runtime/Renderer/Private/ReflectionEnvironment.cpp
+++ b/Engine/Source/Runtime/Renderer/Private/ReflectionEnvironment.cpp
@@ -791,7 +791,7 @@ void FDeferredShadingSceneRenderer::RenderDeferredReflectionsAndSkyLighting(FRHI
 			IScreenSpaceDenoiser::FReflectionsInputs DenoiserInputs;
 			RenderRayTracingReflections(
 				GraphBuilder,
-				View, &DenoiserInputs.Color, &DenoiserInputs.RayHitDistance,
+				View, &DenoiserInputs.Color, &DenoiserInputs.RayHitDistance, &DenoiserInputs.RayImaginaryDepth,
 				RayTracingReflectionsSPP, GRayTracingReflectionsHeightFog, RayTracingConfig.ResolutionFraction);
 
 
diff --git a/Engine/Source/Runtime/Renderer/Private/SceneRendering.h b/Engine/Source/Runtime/Renderer/Private/SceneRendering.h
--- a/Engine/Source/Runtime/Renderer/Private/SceneRendering.h
+++ b/Engine/Source/Runtime/Renderer/Private/SceneRendering.h
@@ -641,6 +641,7 @@ struct FPreviousViewInfo
 	TRefCountPtr<IPooledRenderTarget> DepthBuffer;
 	TRefCountPtr<IPooledRenderTarget> GBufferA;
 	TRefCountPtr<IPooledRenderTarget> GBufferB;
+	TRefCountPtr<IPooledRenderTarget> GBufferC;
 
 	// Temporal AA result of last frame
 	FTemporalAAHistory TemporalAAHistory;
@@ -675,6 +676,7 @@ struct FPreviousViewInfo
 		DepthBuffer.SafeRelease();
 		GBufferA.SafeRelease();
 		GBufferB.SafeRelease();
+		GBufferC.SafeRelease();
 		TemporalAAHistory.SafeRelease();
 		DOFPreGatherHistory.SafeRelease();
 		DOFPostGatherForegroundHistory.SafeRelease();
diff --git a/Engine/Source/Runtime/Renderer/Private/ScreenSpaceDenoise.h b/Engine/Source/Runtime/Renderer/Private/ScreenSpaceDenoise.h
--- a/Engine/Source/Runtime/Renderer/Private/ScreenSpaceDenoise.h
+++ b/Engine/Source/Runtime/Renderer/Private/ScreenSpaceDenoise.h
@@ -114,6 +114,7 @@ public:
 	BEGIN_SHADER_PARAMETER_STRUCT(FReflectionsInputs, )
 		SHADER_PARAMETER_RDG_TEXTURE(Texture2D, Color)
 		SHADER_PARAMETER_RDG_TEXTURE(Texture2D, RayHitDistance)
+		SHADER_PARAMETER_RDG_TEXTURE(Texture2D, RayImaginaryDepth)
 	END_SHADER_PARAMETER_STRUCT()
 
 	/** All the outputs of the reflection denoiser may generate. */
diff --git a/Engine/Source/Runtime/Renderer/Private/SystemTextures.h b/Engine/Source/Runtime/Renderer/Private/SystemTextures.h
--- a/Engine/Source/Runtime/Renderer/Private/SystemTextures.h
+++ b/Engine/Source/Runtime/Renderer/Private/SystemTextures.h
@@ -96,4 +96,4 @@ protected:
 };
 
 /** The global system textures used for scene rendering. */
-extern TGlobalResource<FSystemTextures> GSystemTextures;
+RENDERER_API extern TGlobalResource<FSystemTextures> GSystemTextures;
