Change the ShaderCodeLibrary to only initialise when the project settings require it and make a failure to open the 'Global' shader library fatal for non-Editor builds.

diff --git a/Engine/Source/Runtime/RenderCore/Private/ShaderCodeLibrary.cpp b/Engine/Source/Runtime/RenderCore/Private/ShaderCodeLibrary.cpp
--- a/Engine/Source/Runtime/RenderCore/Private/ShaderCodeLibrary.cpp
+++ b/Engine/Source/Runtime/RenderCore/Private/ShaderCodeLibrary.cpp
@@ -20,6 +20,7 @@ ShaderCodeLibrary.cpp: Bound shader state cache implementation.
 #include "Interfaces/IShaderFormatArchive.h"
 #include "ShaderPipelineCache.h"
 #include "Misc/FileHelper.h"
+#include "Misc/ConfigCacheIni.h"
 
 #if WITH_EDITORONLY_DATA
 #include "Modules/ModuleManager.h"
@@ -1918,7 +1919,10 @@ void FShaderCodeLibrary::InitForRuntime(EShaderPlatform ShaderPlatform)
 	}
 
 	// Cannot be enabled by the server, pointless if we can't ever render and not compatible with cook-on-the-fly
-	bool bEnable = !FPlatformProperties::IsServerOnly() && FApp::CanEverRender();
+	bool bArchive = false;
+	GConfig->GetBool(TEXT("/Script/UnrealEd.ProjectPackagingSettings"), TEXT("bShareMaterialShaderCode"), bArchive, GGameIni);
+
+	bool bEnable = !FPlatformProperties::IsServerOnly() && FApp::CanEverRender() && bArchive;
 #if !UE_BUILD_SHIPPING
 	FString FileHostIP;
 	const bool bCookOnTheFly = FParse::Value(FCommandLine::Get(), TEXT("filehostip"), FileHostIP);
@@ -1946,6 +1950,9 @@ void FShaderCodeLibrary::InitForRuntime(EShaderPlatform ShaderPlatform)
 		}
 		else
 		{
+#if !WITH_EDITOR
+            UE_LOG(LogShaderLibrary, Fatal, TEXT("Failed to initialise ShaderCodeLibrary required by the project because part of the Global shader library is missing from %s."), *FPaths::ProjectContentDir());
+#endif
 			Shutdown();
 		}
 	}
