Generating stable shader info is now configurable.

Shader code library was always emitting "<shadermodel>.scl.csv" files while cooking. This output data is only needed for users (or automated jobs) that update the PSO cache. With this change the configuration value "NeedsShaderStableKeys" needs to be set to true for platforms that use PSO (e.g. AndroidEngine.ini and  IOSEngine.ini).

This saves diskspace and around 90 seconds on a iterative cook on a large project.

diff --git a/Engine/Source/Editor/UnrealEd/Private/CookOnTheFlyServer.cpp b/Engine/Source/Editor/UnrealEd/Private/CookOnTheFlyServer.cpp
--- a/Engine/Source/Editor/UnrealEd/Private/CookOnTheFlyServer.cpp
+++ b/Engine/Source/Editor/UnrealEd/Private/CookOnTheFlyServer.cpp
@@ -6052,16 +6052,28 @@ void UCookOnTheFlyServer::InitShaderCodeLibrary(void)
         {
             FString TargetPlatformNameString = TargetPlatformName.ToString();
             const ITargetPlatform* TargetPlatform = TPM.FindTargetPlatform(TargetPlatformNameString);
+
+			// Find out if this platform requires stable shader keys, by reading the platform setting file.
+			bool bNeedShaderStableKeys = false;
+			FConfigFile PlatformIniFile;
+			FConfigCacheIni::LoadLocalIniFile(PlatformIniFile, TEXT("Engine"), true, *TargetPlatform->IniPlatformName());
+			PlatformIniFile.GetBool(TEXT("DevOptions.Shaders"), TEXT("NeedsShaderStableKeys"), bNeedShaderStableKeys);
             
             TArray<FName> ShaderFormats;
             TargetPlatform->GetAllTargetedShaderFormats(ShaderFormats);
+			TArray<TPair<FName, bool>> ShaderFormatsWithStableKeys;
+			for (FName& Format : ShaderFormats)
+			{
+				ShaderFormatsWithStableKeys.Push(MakeTuple(Format, bNeedShaderStableKeys));
+			}
+
             if (ShaderFormats.Num() > 0)
 			{
 				if (!IsCookFlagSet(ECookInitializationFlags::Iterative))
 				{
 					FShaderCodeLibrary::CleanDirectories(ShaderFormats);
 				}
-				FShaderCodeLibrary::CookShaderFormats(ShaderFormats);
+				FShaderCodeLibrary::CookShaderFormats(ShaderFormatsWithStableKeys);
 			}
         }
     }
diff --git a/Engine/Source/Runtime/Engine/Private/Materials/Material.cpp b/Engine/Source/Runtime/Engine/Private/Materials/Material.cpp
--- a/Engine/Source/Runtime/Engine/Private/Materials/Material.cpp
+++ b/Engine/Source/Runtime/Engine/Private/Materials/Material.cpp
@@ -3512,7 +3512,7 @@ void UMaterial::Serialize(FArchive& Ar)
 		}
 	}
 #if WITH_EDITOR
-	if (Ar.IsSaving() && Ar.IsCooking() && Ar.IsPersistent() && !Ar.IsObjectReferenceCollector() && FShaderCodeLibrary::NeedsShaderStableKeys())
+	if (Ar.IsSaving() && Ar.IsCooking() && Ar.IsPersistent() && !Ar.IsObjectReferenceCollector() && FShaderCodeLibrary::NeedsShaderStableKeys(EShaderPlatform::SP_NumPlatforms))
 	{
 		SaveShaderStableKeys(Ar.CookingTarget());
 	}
diff --git a/Engine/Source/Runtime/Engine/Private/Materials/MaterialInstance.cpp b/Engine/Source/Runtime/Engine/Private/Materials/MaterialInstance.cpp
--- a/Engine/Source/Runtime/Engine/Private/Materials/MaterialInstance.cpp
+++ b/Engine/Source/Runtime/Engine/Private/Materials/MaterialInstance.cpp
@@ -3132,7 +3132,7 @@ void UMaterialInstance::Serialize(FArchive& Ar)
 #endif
 	}
 #if WITH_EDITOR
-	if (Ar.IsSaving() && Ar.IsCooking() && Ar.IsPersistent() && !Ar.IsObjectReferenceCollector() && FShaderCodeLibrary::NeedsShaderStableKeys())
+	if (Ar.IsSaving() && Ar.IsCooking() && Ar.IsPersistent() && !Ar.IsObjectReferenceCollector() && FShaderCodeLibrary::NeedsShaderStableKeys(EShaderPlatform::SP_NumPlatforms))
 	{
 		SaveShaderStableKeys(Ar.CookingTarget());
 	}
diff --git a/Engine/Source/Runtime/Engine/Private/ShaderCompiler/ShaderCompiler.cpp b/Engine/Source/Runtime/Engine/Private/ShaderCompiler/ShaderCompiler.cpp
--- a/Engine/Source/Runtime/Engine/Private/ShaderCompiler/ShaderCompiler.cpp
+++ b/Engine/Source/Runtime/Engine/Private/ShaderCompiler/ShaderCompiler.cpp
@@ -4042,7 +4042,7 @@ FString SaveGlobalShaderFile(EShaderPlatform Platform, FString SavePath, class I
 		UE_LOG(LogShaders, Fatal, TEXT("Could not save global shader file to '%s'"), *FullPath);
 	}
 #if WITH_EDITOR
-	if (FShaderCodeLibrary::NeedsShaderStableKeys())
+	if (FShaderCodeLibrary::NeedsShaderStableKeys(Platform))
 	{
 		for (TLinkedList<FShaderType*>::TIterator ShaderTypeIt(FShaderType::GetTypeList()); ShaderTypeIt; ShaderTypeIt.Next())
 		{
diff --git a/Engine/Source/Runtime/RenderCore/Private/Shader.cpp b/Engine/Source/Runtime/RenderCore/Private/Shader.cpp
--- a/Engine/Source/Runtime/RenderCore/Private/Shader.cpp
+++ b/Engine/Source/Runtime/RenderCore/Private/Shader.cpp
@@ -1435,7 +1435,8 @@ void FShader::DumpDebugInfo()
 void FShader::SaveShaderStableKeys(EShaderPlatform TargetShaderPlatform, const FStableShaderKeyAndValue& InSaveKeyVal)
 {
 #if WITH_EDITOR
-	if (TargetShaderPlatform == EShaderPlatform::SP_NumPlatforms || EShaderPlatform(Target.Platform) == TargetShaderPlatform)
+	if ((TargetShaderPlatform == EShaderPlatform::SP_NumPlatforms || EShaderPlatform(Target.Platform) == TargetShaderPlatform) 
+		&& FShaderCodeLibrary::NeedsShaderStableKeys(TargetShaderPlatform))
 	{
 		FStableShaderKeyAndValue SaveKeyVal(InSaveKeyVal);
 		SaveKeyVal.TargetFrequency = FName(GetShaderFrequencyString((EShaderFrequency)Target.Frequency));
diff --git a/Engine/Source/Runtime/RenderCore/Private/ShaderCodeLibrary.cpp b/Engine/Source/Runtime/RenderCore/Private/ShaderCodeLibrary.cpp
--- a/Engine/Source/Runtime/RenderCore/Private/ShaderCodeLibrary.cpp
+++ b/Engine/Source/Runtime/RenderCore/Private/ShaderCodeLibrary.cpp
@@ -1527,6 +1527,8 @@ class FShaderCodeLibraryImpl
 	FEditorShaderCodeArchive* EditorShaderCodeArchive[EShaderPlatform::SP_NumPlatforms];
 	// At cook time, shader code collection for each shader platform
 	FEditorShaderStableInfo* EditorShaderStableInfo[EShaderPlatform::SP_NumPlatforms];
+	// Cached bit field for shader formats that require stable keys
+	uint64_t bShaderFormatsThatNeedStableKeys = 0;
 	// At cook time, shader stats for each shader platform
 	FShaderCodeStats EditorShaderCodeStats[EShaderPlatform::SP_NumPlatforms];
 	// At cook time, whether the shader archive supports pipelines (only OpenGL should)
@@ -1941,10 +1943,12 @@ public:
 		}
 	}
 
-	void CookShaderFormats(TArray<FName> const& ShaderFormats)
+	void CookShaderFormats(TArray<TTuple<FName,bool>> const& ShaderFormats)
 	{
-		for (FName const& Format : ShaderFormats)
+		for (auto Pair : ShaderFormats)
 		{
+			FName const& Format = Pair.Get<0>();
+
 			EShaderPlatform Platform = ShaderFormatToLegacyShaderPlatform(Format);
 			FName PossiblyAdjustedFormat = LegacyShaderPlatformToShaderFormat(Platform);	// Vulkan and GL switch between name variants depending on CVars (e.g. see r.Vulkan.UseRealUBs)
 			FEditorShaderCodeArchive* CodeArchive = EditorShaderCodeArchive[Platform];
@@ -1956,20 +1960,33 @@ public:
 			}
 			check(CodeArchive);
 		}
-		for (FName const& Format : ShaderFormats)
+		for (auto Pair : ShaderFormats)
 		{
+			FName const& Format = Pair.Get<0>();
+			bool bUseStableKeys = Pair.Get<1>();
+
 			EShaderPlatform Platform = ShaderFormatToLegacyShaderPlatform(Format);
 			FName PossiblyAdjustedFormat = LegacyShaderPlatformToShaderFormat(Platform);	// Vulkan and GL switch between name variants depending on CVars (e.g. see r.Vulkan.UseRealUBs)
 			FEditorShaderStableInfo* StableArchive = EditorShaderStableInfo[Platform];
-			if (!StableArchive)
+			if (!StableArchive && bUseStableKeys)
 			{
 				StableArchive = new FEditorShaderStableInfo(PossiblyAdjustedFormat);
 				EditorShaderStableInfo[Platform] = StableArchive;
+				bShaderFormatsThatNeedStableKeys |= (uint64_t(1u) << (uint32_t)Platform);
+				static_assert(SP_NumPlatforms < 64u, "ShaderPlatform will no longer fit into bitfield.");
 			}
-			check(StableArchive);
 		}
 	}
 
+	bool NeedsShaderStableKeys(EShaderPlatform Platform) 
+	{
+		if (Platform == EShaderPlatform::SP_NumPlatforms)
+		{
+			return bShaderFormatsThatNeedStableKeys != 0;
+		}
+		return (bShaderFormatsThatNeedStableKeys & (uint64_t(1u) << (uint32_t) Platform)) != 0;
+	}
+
 	void AddShaderCode(EShaderPlatform Platform, EShaderFrequency Frequency, const FSHAHash& Hash, const TArray<uint8>& InCode, uint32 const UncompressedSize)
 	{
 		FScopeLock ScopeLock(&ShaderCodeCS);
@@ -1989,13 +2006,15 @@ public:
 
 	void AddShaderStableKeyValue(EShaderPlatform InShaderPlatform, FStableShaderKeyAndValue& StableKeyValue)
 	{
+		FEditorShaderStableInfo* StableArchive = EditorShaderStableInfo[InShaderPlatform];
+		if (!StableArchive)
+		{
+			return;
+		}
+
 		FScopeLock ScopeLock(&ShaderCodeCS);
 
 		StableKeyValue.ComputeKeyHash();
-
-		FEditorShaderStableInfo* StableArchive = EditorShaderStableInfo[InShaderPlatform];
-		check(StableArchive);
-
 		StableArchive->AddShader(StableKeyValue);
 	}
 
@@ -2430,7 +2449,7 @@ void FShaderCodeLibrary::CleanDirectories(TArray<FName> const& ShaderFormats)
 	}
 }
 
-void FShaderCodeLibrary::CookShaderFormats(TArray<FName> const& ShaderFormats)
+void FShaderCodeLibrary::CookShaderFormats(TArray<TTuple<FName,bool>> const& ShaderFormats)
 {
 	if (FShaderCodeLibraryImpl::Impl)
 	{
@@ -2451,12 +2470,12 @@ bool FShaderCodeLibrary::AddShaderCode(EShaderPlatform ShaderPlatform, EShaderFr
 	return false;
 }
 
-bool FShaderCodeLibrary::NeedsShaderStableKeys()
+bool FShaderCodeLibrary::NeedsShaderStableKeys(EShaderPlatform ShaderPlatform)
 {
 #if WITH_EDITOR
 	if (FShaderCodeLibraryImpl::Impl)
 	{
-		return true;
+		return FShaderCodeLibraryImpl::Impl->NeedsShaderStableKeys(ShaderPlatform);
 	}
 #endif// WITH_EDITOR
 	return false;
diff --git a/Engine/Source/Runtime/RenderCore/Public/ShaderCodeLibrary.h b/Engine/Source/Runtime/RenderCore/Public/ShaderCodeLibrary.h
--- a/Engine/Source/Runtime/RenderCore/Public/ShaderCodeLibrary.h
+++ b/Engine/Source/Runtime/RenderCore/Public/ShaderCodeLibrary.h
@@ -212,14 +212,17 @@ struct RENDERCORE_API FShaderCodeLibrary
 	// Clean the cook directories
 	static void CleanDirectories(TArray<FName> const& ShaderFormats);
     
-    // Specify the shader formats to cook
-    static void CookShaderFormats(TArray<FName> const& ShaderFormats);
+    // Specify the shader formats to cook and which ones needs stable keys. Provide an array of tuples
+	// with names and whether the format needs stable keys.
+    static void CookShaderFormats(TArray<TTuple<FName,bool>> const& ShaderFormats);
 	
 	// At cook time, add shader code to collection
 	static bool AddShaderCode(EShaderPlatform ShaderPlatform, EShaderFrequency Frequency, const FSHAHash& Hash, const TArray<uint8>& InCode, uint32 const UncompressedSize);
 
-	// We check this early in the callstack to avoid creating a bunch of FName and keys and things we will never save anyway
-	static bool NeedsShaderStableKeys();
+	// We check this early in the callstack to avoid creating a bunch of FName and keys and things we will never save anyway. 
+	// Pass the shader platform to check or EShaderPlatform::SP_NumPlatforms to check if any of the registered types require
+	// stable keys.
+	static bool NeedsShaderStableKeys(EShaderPlatform ShaderPlatform);
 
 	// At cook time, add the human readable key value information
 	static void AddShaderStableKeyValue(EShaderPlatform ShaderPlatform, FStableShaderKeyAndValue& StableKeyValue);
