Optimizations to ShaderCodeLibrary cooking

* Use a local cache for FNames when parsing FStableShaderKeyAndValue. This saves around 29 seconds when cooking a large data set.
* Write directly to CSV file instead of using string arrays. This saves around 2.3 seconds.
* Reserve string size when parsing FCompactFullName. This saves around 0.5 seconds.

diff --git a/Engine/Source/Runtime/RenderCore/Private/ShaderCodeLibrary.cpp b/Engine/Source/Runtime/RenderCore/Private/ShaderCodeLibrary.cpp
--- a/Engine/Source/Runtime/RenderCore/Private/ShaderCodeLibrary.cpp
+++ b/Engine/Source/Runtime/RenderCore/Private/ShaderCodeLibrary.cpp
@@ -16,6 +16,7 @@ ShaderCodeLibrary.cpp: Bound shader state cache implementation.
 #include "Async/AsyncFileHandle.h"
 #include "PipelineFileCache.h"
 #include "Interfaces/IPluginManager.h"
+#include "Hash/CityHash.h"
 
 #include "Interfaces/IShaderFormatArchive.h"
 #include "ShaderPipelineCache.h"
@@ -99,9 +100,24 @@ static void FShaderLibraryHelperCompressCode(EShaderPlatform Platform, const TAr
 		CompressedCode.Shrink();
 }
 
+
+FORCEINLINE FName ParseFNameCached(const FString& Src, TMap<uint32,FName>& NameCache)
+{
+	uint32 SrcHash = CityHash32((char*)Src.GetCharArray().GetData(), Src.GetCharArray().GetAllocatedSize());
+	if (FName* Name = NameCache.Find(SrcHash))
+	{
+		return *Name;
+	}
+	else
+	{
+		return NameCache.Emplace(SrcHash, FName(*Src));
+	}
+}
+
 FString FCompactFullName::ToString() const
 {
 	FString RetString;
+	RetString.Reserve(256);
 	if (!ObjectClassAndPath.Num())
 	{
 		RetString += TEXT("empty");
@@ -160,6 +176,28 @@ uint32 GetTypeHash(const FCompactFullName& A)
 	return Hash;
 }
 
+void FixupUnsanitizedNames(const FString& Src, TArray<FString>& OutFields) 
+{
+	FString NewSrc(Src);
+
+	int32 ParenOpen = -1;
+	int32 ParenClose = -1;
+
+	if (NewSrc.FindChar(TCHAR('('), ParenOpen) && NewSrc.FindChar(TCHAR(')'), ParenClose) && ParenOpen < ParenClose && ParenOpen >= 0 && ParenClose >= 0)
+	{
+		for (int32 Index = ParenOpen + 1; Index < ParenClose; Index++)
+		{
+			if (NewSrc[Index] == TCHAR(','))
+			{
+				NewSrc[Index] = ' ';
+			}
+		}
+		OutFields.Empty();
+		NewSrc.TrimStartAndEnd().ParseIntoArray(OutFields, TEXT(","), false);
+		check(OutFields.Num() == 11);
+	}
+}
+
 void FStableShaderKeyAndValue::ComputeKeyHash()
 {
 	KeyHash = GetTypeHash(ClassNameAndObjectPath);
@@ -184,28 +222,9 @@ void FStableShaderKeyAndValue::ParseFromString(const FString& Src)
 	if (Fields.Num() > 11)
 	{
 		// hack fix for unsanitized names, should not occur anymore.
-
-		FString NewSrc(Src);
-
-		int32 ParenOpen = -1;
-		int32 ParenClose = -1;
-
-		if (NewSrc.FindChar(TCHAR('('), ParenOpen) && NewSrc.FindChar(TCHAR(')'), ParenClose) && ParenOpen < ParenClose && ParenOpen >= 0 && ParenClose >= 0)
-		{
-			for (int32 Index = ParenOpen + 1; Index < ParenClose; Index++)
-			{
-				if (NewSrc[Index] == TCHAR(','))
-				{
-					NewSrc[Index] = ' ';
-				}
-			}
-			Fields.Empty();
-			NewSrc.TrimStartAndEnd().ParseIntoArray(Fields, TEXT(","), false);
-			check(Fields.Num() == 11);
-		}
+		FixupUnsanitizedNames(Src, Fields);
 	}
 
-
 	check(Fields.Num() == 11);
 
 	int32 Index = 0;
@@ -230,6 +249,44 @@ void FStableShaderKeyAndValue::ParseFromString(const FString& Src)
 	ComputeKeyHash();
 }
 
+
+void FStableShaderKeyAndValue::ParseFromStringCached(const FString& Src, TMap<uint32, FName>& NameCache)
+{
+	TArray<FString> Fields;
+	Src.TrimStartAndEnd().ParseIntoArray(Fields, TEXT(","), false);
+
+	if (Fields.Num() > 11)
+	{
+		// hack fix for unsanitized names, should not occur anymore.
+		FixupUnsanitizedNames(Src, Fields);
+	}
+	
+	check(Fields.Num() == 11);
+
+	int32 Index = 0;
+	ClassNameAndObjectPath.ParseFromString(Fields[Index++]);
+
+	// There is a high level of uniformity on the following FNames, use
+	// the local name cache to accelerate lookup
+	ShaderType = ParseFNameCached(Fields[Index++], NameCache);
+	ShaderClass = ParseFNameCached(Fields[Index++], NameCache);
+	MaterialDomain = ParseFNameCached(Fields[Index++], NameCache);
+	FeatureLevel = ParseFNameCached(Fields[Index++], NameCache);
+
+	QualityLevel = ParseFNameCached(Fields[Index++], NameCache);
+	TargetFrequency = ParseFNameCached(Fields[Index++], NameCache);
+	TargetPlatform = ParseFNameCached(Fields[Index++], NameCache);
+
+	VFType = ParseFNameCached(Fields[Index++], NameCache);
+	PermutationId = ParseFNameCached(Fields[Index++], NameCache);
+
+	OutputHash.FromString(Fields[Index++]);
+
+	check(Index == 11);
+
+	ComputeKeyHash();
+}
+
 FString FStableShaderKeyAndValue::ToString() const
 {
 	FString Result;
@@ -520,7 +577,6 @@ public:
 			FScopeLock ScopeLock(&ReadRequestLock);
 
 			Entry->NumRefs--;
-			check(Entry->NumRefs >= 0);
 			if (Entry->NumRefs == 0)
 			{
 				// should not attempt to release shader code while it's loading
@@ -1211,6 +1267,9 @@ struct FEditorShaderStableInfo
 	{
 		check(LibraryName.Len() > 0);
 
+		TMap<uint32, FName> NameCache;
+		NameCache.Reserve(2048);
+
 		const FString ShaderIntermediateLocation = FPaths::ProjectSavedDir() / TEXT("Shaders") / FormatName.ToString();
 
 		TArray<FString> ShaderFiles;
@@ -1226,7 +1285,7 @@ struct FEditorShaderStableInfo
 					for (int32 Index = 1; Index < SourceFileContents.Num(); Index++)
 					{
 						FStableShaderKeyAndValue Item;
-						Item.ParseFromString(SourceFileContents[Index]);
+						Item.ParseFromStringCached(SourceFileContents[Index], NameCache);
 						AddShader(Item);
 					}
 				}
@@ -1254,17 +1313,23 @@ struct FEditorShaderStableInfo
 			// Write to a intermediate file
 			FString IntermediateFormatPath = GetStableInfoArchiveFilename(FPaths::ProjectSavedDir() / TEXT("Shaders") / FormatName.ToString(), LibraryName, FormatName);
 
-			TArray<FString> FileContents;
-			FileContents.Reserve(StableMap.Num() + 1);
+			// Write directly to the file
+			TUniquePtr<FArchive> IntermediateFormatAr(IFileManager::Get().CreateFileWriter(*IntermediateFormatPath));
+
+			const FString HeaderText = FStableShaderKeyAndValue::HeaderLine();
+			auto HeaderSrc = StringCast<ANSICHAR>(*HeaderText, HeaderText.Len());
+
+			IntermediateFormatAr->Serialize((ANSICHAR*)HeaderSrc.Get(), HeaderSrc.Length() * sizeof(ANSICHAR));
 
-			FileContents.Add(FStableShaderKeyAndValue::HeaderLine());
+			FString LineBuffer;
+			LineBuffer.Reserve(512);
 
 			for (const FStableShaderKeyAndValue& Item : StableMap)
 			{
-				FString& LineBuffer = FileContents.Emplace_GetRef();
 				Item.ToString(LineBuffer);
+				auto LineConverted = StringCast<ANSICHAR>(*LineBuffer, LineBuffer.Len());
+				IntermediateFormatAr->Serialize((ANSICHAR*)LineConverted.Get(), LineConverted.Length() * sizeof(ANSICHAR));
 			}
-			FFileHelper::SaveStringArrayToFile(FileContents, *IntermediateFormatPath);
 
 			// Only the master cooker needs to write to the output directory, child cookers only write to the Saved directory
 			if (bMasterCooker)
diff --git a/Engine/Source/Runtime/RenderCore/Public/ShaderCodeLibrary.h b/Engine/Source/Runtime/RenderCore/Public/ShaderCodeLibrary.h
--- a/Engine/Source/Runtime/RenderCore/Public/ShaderCodeLibrary.h
+++ b/Engine/Source/Runtime/RenderCore/Public/ShaderCodeLibrary.h
@@ -87,6 +87,7 @@ struct RENDERCORE_API FStableShaderKeyAndValue
 
 	void ComputeKeyHash();
 	void ParseFromString(const FString& Src);
+	void ParseFromStringCached(const FString& Src, class TMap<uint32, FName>& NameCache);
 	FString ToString() const;
 	void ToString(FString& OutResult) const;
 	static FString HeaderLine();
