Adds missing viewport FIntRect on FPixelShaderUtils::AddFullscreenPass()

diff --git a/Engine/Source/Runtime/RenderCore/Private/VisualizeTexture.cpp b/Engine/Source/Runtime/RenderCore/Private/VisualizeTexture.cpp
--- a/Engine/Source/Runtime/RenderCore/Private/VisualizeTexture.cpp
+++ b/Engine/Source/Runtime/RenderCore/Private/VisualizeTexture.cpp
@@ -276,7 +276,7 @@ void FVisualizeTexture::CreateContentCapturePass(FRDGBuilder& GraphBuilder, cons
 	FPixelShaderUtils::AddFullscreenPass(
 		GraphBuilder, ShaderMap,
 		RDG_EVENT_NAME("VisualizeTextureCapture(%s)", SrcTexture->Name),
-		*PixelShader, PassParameters);
+		*PixelShader, PassParameters, FIntRect(0, 0, RTExtent.X, RTExtent.Y));
 
 	// Save the copied texture and descriptor about original informations.
 	{
diff --git a/Engine/Source/Runtime/RenderCore/Public/PixelShaderUtils.h b/Engine/Source/Runtime/RenderCore/Public/PixelShaderUtils.h
--- a/Engine/Source/Runtime/RenderCore/Public/PixelShaderUtils.h
+++ b/Engine/Source/Runtime/RenderCore/Public/PixelShaderUtils.h
@@ -30,7 +30,8 @@ struct RENDERCORE_API FPixelShaderUtils
 		FRHICommandList& RHICmdList, 
 		const TShaderMap<FGlobalShaderType>* GlobalShaderMap,
 		const TShaderClass* PixelShader,
-		const typename TShaderClass::FParameters& Parameters)
+		const typename TShaderClass::FParameters& Parameters,
+		const FIntRect& Viewport)
 	{
 		CA_ASSUME(PixelShader);
 		TShaderMapRef<FVisualizeTextureVS> VertexShader(GlobalShaderMap);
@@ -47,6 +48,8 @@ struct RENDERCORE_API FPixelShaderUtils
 		GraphicsPSOInit.PrimitiveType = PT_TriangleList;
 		SetGraphicsPipelineState(RHICmdList, GraphicsPSOInit);
 
+		RHICmdList.SetViewport(Viewport.Min.X, Viewport.Min.Y, 0.0f, Viewport.Max.X, Viewport.Max.Y, 1.0f);
+
 		SetShaderParameters(RHICmdList, PixelShader, PixelShader->GetPixelShader(), Parameters);
 
 		DrawFullscreenTriangle(RHICmdList);
@@ -59,7 +62,8 @@ struct RENDERCORE_API FPixelShaderUtils
 		const TShaderMap<FGlobalShaderType>* GlobalShaderMap,
 		FRDGEventName&& PassName,
 		const TShaderClass* PixelShader,
-		typename TShaderClass::FParameters* Parameters)
+		typename TShaderClass::FParameters* Parameters,
+		const FIntRect& Viewport)
 	{
 		CA_ASSUME(PixelShader != nullptr);
 		ClearUnusedGraphResources(PixelShader, Parameters);
@@ -68,9 +72,9 @@ struct RENDERCORE_API FPixelShaderUtils
 			Forward<FRDGEventName>(PassName),
 			Parameters,
 			ERenderGraphPassFlags::None,
-			[Parameters, GlobalShaderMap, PixelShader](FRHICommandList& RHICmdList)
+			[Parameters, GlobalShaderMap, PixelShader, Viewport](FRHICommandList& RHICmdList)
 		{
-			FPixelShaderUtils::DrawFullscreenPixelShader(RHICmdList, GlobalShaderMap, PixelShader, *Parameters);
+			FPixelShaderUtils::DrawFullscreenPixelShader(RHICmdList, GlobalShaderMap, PixelShader, *Parameters, Viewport);
 		});
 	}
 };
