Added UAV to shader parameter struct

diff --git a/Engine/Shaders/Public/ShaderVersion.ush b/Engine/Shaders/Public/ShaderVersion.ush
--- a/Engine/Shaders/Public/ShaderVersion.ush
+++ b/Engine/Shaders/Public/ShaderVersion.ush
@@ -3,4 +3,4 @@
 // in Platform.ush (which should be included in any shader) it allows to invalidate the shader DDC.
 // 
 // If you are merging streams and there is a conflict with this GUID you should make a new GUID rather than taking one or the other.
-// GUID = A54A0FDE-D1D0-4543-9A07-4EBA26888610
+// GUID = BE84FF0E-76C0-4946-9F4B-2FC63610027B
diff --git a/Engine/Source/Developer/ShaderCompilerCommon/Private/ShaderCompilerCommon.cpp b/Engine/Source/Developer/ShaderCompilerCommon/Private/ShaderCompilerCommon.cpp
--- a/Engine/Source/Developer/ShaderCompilerCommon/Private/ShaderCompilerCommon.cpp
+++ b/Engine/Source/Developer/ShaderCompilerCommon/Private/ShaderCompilerCommon.cpp
@@ -142,6 +142,7 @@ bool BuildResourceTableMapping(
 			case UBMT_RDG_BUFFER_SRV:
 				OutSRT.ShaderResourceViewMap.Add(ResourceMap);
 				break;
+			case UBMT_UAV:
 			case UBMT_RDG_TEXTURE_UAV:
 			case UBMT_RDG_BUFFER_UAV:
 				OutSRT.UnorderedAccessViewMap.Add(ResourceMap);
diff --git a/Engine/Source/Runtime/RHI/Public/RHIDefinitions.h b/Engine/Source/Runtime/RHI/Public/RHIDefinitions.h
--- a/Engine/Source/Runtime/RHI/Public/RHIDefinitions.h
+++ b/Engine/Source/Runtime/RHI/Public/RHIDefinitions.h
@@ -415,6 +415,7 @@ enum EUniformBufferBaseType : uint8
 	// RHI resources not tracked by render graph.
 	UBMT_TEXTURE,
 	UBMT_SRV,
+	UBMT_UAV,
 	UBMT_SAMPLER,
 
 	// Resources tracked by render graph.
diff --git a/Engine/Source/Runtime/RenderCore/Private/ShaderParameterStruct.cpp b/Engine/Source/Runtime/RenderCore/Private/ShaderParameterStruct.cpp
--- a/Engine/Source/Runtime/RenderCore/Private/ShaderParameterStruct.cpp
+++ b/Engine/Source/Runtime/RenderCore/Private/ShaderParameterStruct.cpp
@@ -130,6 +130,7 @@ struct FShaderParameterStructBindingContext
 			else if (
 				BaseType == UBMT_TEXTURE ||
 				BaseType == UBMT_SRV ||
+				BaseType == UBMT_UAV ||
 				BaseType == UBMT_SAMPLER ||
 				BaseType == UBMT_RDG_TEXTURE ||
 				BaseType == UBMT_RDG_TEXTURE_SRV ||
@@ -148,6 +149,8 @@ struct FShaderParameterStructBindingContext
 					Bindings->Textures.Add(Parameter);
 				else if (BaseType == UBMT_SRV)
 					Bindings->SRVs.Add(Parameter);
+				else if (BaseType == UBMT_UAV)
+					Bindings->UAVs.Add(Parameter);
 				else if (BaseType == UBMT_SAMPLER)
 					Bindings->Samplers.Add(Parameter);
 				else if (BaseType == UBMT_RDG_TEXTURE)
diff --git a/Engine/Source/Runtime/RenderCore/Public/Shader.h b/Engine/Source/Runtime/RenderCore/Public/Shader.h
--- a/Engine/Source/Runtime/RenderCore/Public/Shader.h
+++ b/Engine/Source/Runtime/RenderCore/Public/Shader.h
@@ -777,6 +777,7 @@ public:
 	TArray<FParameter> Parameters;
 	TArray<FResourceParameter> Textures;
 	TArray<FResourceParameter> SRVs;
+	TArray<FResourceParameter> UAVs;
 	TArray<FResourceParameter> Samplers;
 	TArray<FResourceParameter> GraphTextures;
 	TArray<FResourceParameter> GraphSRVs;
@@ -791,6 +792,7 @@ public:
 		Ar << ParametersBindingData.Parameters;
 		Ar << ParametersBindingData.Textures;
 		Ar << ParametersBindingData.SRVs;
+		Ar << ParametersBindingData.UAVs;
 		Ar << ParametersBindingData.Samplers;
 		Ar << ParametersBindingData.GraphTextures;
 		Ar << ParametersBindingData.GraphSRVs;
diff --git a/Engine/Source/Runtime/RenderCore/Public/ShaderParameterMacros.h b/Engine/Source/Runtime/RenderCore/Public/ShaderParameterMacros.h
--- a/Engine/Source/Runtime/RenderCore/Public/ShaderParameterMacros.h
+++ b/Engine/Source/Runtime/RenderCore/Public/ShaderParameterMacros.h
@@ -698,6 +698,14 @@ extern RENDERCORE_API FShaderParametersMetadata* FindUniformBufferStructByFName(
 #define SHADER_PARAMETER_SRV(ShaderType,MemberName) \
 	INTERNAL_SHADER_PARAMETER_EXPLICIT(UBMT_SRV, TShaderResourceParameterTypeInfo<FShaderResourceViewRHIParamRef>, FShaderResourceViewRHIParamRef,MemberName,, = nullptr,EShaderPrecisionModifier::Float,TEXT(#ShaderType),false)
 
+/** Adds an unordered access view.
+ *
+ * Example:
+ *	SHADER_PARAMETER_UAV(Texture2D, MyUAV)
+ */
+#define SHADER_PARAMETER_UAV(ShaderType,MemberName) \
+	INTERNAL_SHADER_PARAMETER_EXPLICIT(UBMT_UAV, TShaderResourceParameterTypeInfo<FUnorderedAccessViewRHIParamRef>, FUnorderedAccessViewRHIParamRef,MemberName,, = nullptr,EShaderPrecisionModifier::Float,TEXT(#ShaderType),false)
+
 /** Adds a sampler.
  *
  * Example:
diff --git a/Engine/Source/Runtime/RenderCore/Public/ShaderParameterStruct.h b/Engine/Source/Runtime/RenderCore/Public/ShaderParameterStruct.h
--- a/Engine/Source/Runtime/RenderCore/Public/ShaderParameterStruct.h
+++ b/Engine/Source/Runtime/RenderCore/Public/ShaderParameterStruct.h
@@ -83,6 +83,19 @@ inline void SetShaderUAVs(TRHICmdList& RHICmdList, const TShaderClass* Shader, F
 	const typename TShaderClass::FParameters* ParametersPtr = &Parameters;
 	const uint8* Base = reinterpret_cast<const uint8*>(ParametersPtr);
 
+	// UAVs
+	for (const FShaderParameterBindings::FResourceParameter& ParameterBinding : Bindings.UAVs)
+	{
+		auto ShaderParameterRef = *reinterpret_cast<const FUnorderedAccessViewRHIParamRef*>(Base + ParameterBinding.ByteOffset);
+
+		if (DO_CHECK && !ShaderParameterRef)
+		{
+			EmitNullShaderParameterFatalError(Shader, TShaderClass::FParameters::FTypeInfo::GetStructMetadata(), ParameterBinding.ByteOffset);
+		}
+
+		RHICmdList.SetUAVParameter(ShadeRHI, ParameterBinding.BaseIndex, ShaderParameterRef);
+	}
+
 	// Graph UAVs
 	for (const FShaderParameterBindings::FResourceParameter& ParameterBinding : Bindings.GraphUAVs)
 	{
@@ -116,6 +129,11 @@ inline void UnsetShaderUAVs(TRHICmdList& RHICmdList, const TShaderClass* Shader,
 
 	checkf(Bindings.RootParameterBufferIndex == FShaderParameterBindings::kInvalidBufferIndex, TEXT("Can't use UnsetShaderUAVs() for root parameter buffer index."));
 
+	for (const FShaderParameterBindings::FResourceParameter& ParameterBinding : Bindings.UAVs)
+	{
+		RHICmdList.SetUAVParameter(ShadeRHI, ParameterBinding.BaseIndex, nullptr);
+	}
+
 	for (const FShaderParameterBindings::FResourceParameter& ParameterBinding : Bindings.GraphUAVs)
 	{
 		RHICmdList.SetUAVParameter(ShadeRHI, ParameterBinding.BaseIndex, nullptr);
@@ -225,7 +243,7 @@ inline void SetShaderParameters(TRHICmdList& RHICmdList, const TShaderClass* Sha
 		RHICmdList.SetShaderResourceViewParameter(ShadeRHI, ParameterBinding.BaseIndex, GraphSRV->CachedRHI.SRV);
 	}
 
-	// Graph UAVs for compute shaders
+	// UAVs for compute shaders
 	SetShaderUAVs(RHICmdList, Shader, ShadeRHI, Parameters);
 
 	// Reference structures
@@ -282,6 +300,19 @@ void SetShaderParameters(FRayTracingShaderBindingsWriter& RTBindingsWriter, cons
 		RTBindingsWriter.SetSRV(ParameterBinding.BaseIndex, ShaderParameterRef);
 	}
 
+	// UAVs
+	for (const FShaderParameterBindings::FResourceParameter& ParameterBinding : Bindings.UAVs)
+	{
+		auto ShaderParameterRef = *reinterpret_cast<const FUnorderedAccessViewRHIParamRef*>(Base + ParameterBinding.ByteOffset);
+
+		if (DO_CHECK && !ShaderParameterRef)
+		{
+			EmitNullShaderParameterFatalError(Shader, TShaderClass::FParameters::FTypeInfo::GetStructMetadata(), ParameterBinding.ByteOffset);
+		}
+
+		RTBindingsWriter.SetUAV(ParameterBinding.BaseIndex, ShaderParameterRef);
+	}
+
 	// Samplers
 	for (const FShaderParameterBindings::FResourceParameter& ParameterBinding : Bindings.Samplers)
 	{
