Integrated fix for render graph validation error.

diff --git a/Engine/Source/Runtime/RenderCore/Private/RenderGraphValidation.cpp b/Engine/Source/Runtime/RenderCore/Private/RenderGraphValidation.cpp
--- a/Engine/Source/Runtime/RenderCore/Private/RenderGraphValidation.cpp
+++ b/Engine/Source/Runtime/RenderCore/Private/RenderGraphValidation.cpp
@@ -428,19 +428,42 @@ void FRDGUserValidation::ValidateAddPass(const FRDGPass* Pass)
 				{
 					PassResourceValidator.Write(Texture, RenderTarget.GetMipIndex());
 
-					/** Validate that load action is correct. */
 					const bool bIsLoadAction = RenderTarget.GetLoadAction() == ERenderTargetLoadAction::ELoad;
-					const bool bIsLoadActionInvalid = bIsLoadAction && !Texture->HasBeenProduced();
-					checkf(
-						!bIsLoadActionInvalid,
-						TEXT("Can't load a render target '%s' that has never been produced."),
-						Texture->Name);
-					
-					const bool bOverwritesResourceEntirely = !bIsLoadAction && Texture->HasBeenProduced() && Texture->Desc.NumMips == 1;
-					ensureMsgf(!bOverwritesResourceEntirely,
-						TEXT("Clobbering render target %s result that has been produced previously by another pass. Should instead render to a new FRDGTexture to reduce memory pressure."),
-						Texture->Name);
-					
+
+					ensureMsgf(
+						Texture->Desc.TargetableFlags & TexCreate_RenderTargetable,
+						TEXT("Pass '%s' attempted to bind texture '%s' as a render target, but the texture has not been created with TexCreate_RenderTargetable."),
+						PassName, Texture->Name);
+
+					/** Validate that load action is correct. We can only load contents if a pass previously produced something. */
+					{
+						const bool bIsLoadActionInvalid = bIsLoadAction && !Texture->HasBeenProduced();
+						checkf(
+							!bIsLoadActionInvalid,
+							TEXT("Pass '%s' attempted to bind texture '%s' as a render target with the 'Load' action specified, but the texture has not been produced yet. The render target must use either 'Clear' or 'NoAction' action instead."),
+							PassName,
+							Texture->Name);
+					}
+
+					/** Validate that any previously produced texture contents are loaded. This occurs if the user failed to specify a load action
+					 *  on a texture that was produced by a previous pass, effectively losing that data. This can also happen if the user 're-uses'
+					 *  a texture for some other purpose. The latter is considered bad practice, since it increases memory pressure on the render
+					 *  target pool. Instead, the user should create a new texture instance. An exception to this rule are untracked render targets,
+					 *  which are not actually managed by the render target pool and likely represent the frame buffer.
+					 */
+					{
+						// We only validate single-mip textures since we don't track production at the subresource level.
+						const bool bFailedToLoadProducedContent = !bIsLoadAction && Texture->HasBeenProduced() && Texture->Desc.NumMips == 1;
+
+						// Untracked render targets aren't actually managed by the render target pool.
+						const bool bIsUntrackedRenderTarget = Texture->PooledRenderTarget && !Texture->PooledRenderTarget->IsTracked();
+
+						ensureMsgf(!bFailedToLoadProducedContent || bIsUntrackedRenderTarget,
+							TEXT("Pass '%s' attempted to bind texture '%s' as a render target without the 'Load' action specified, despite a prior pass having produced it. It's invalid to completely clobber the contents of a resource. Create a new texture instance instead."),
+							PassName,
+							Texture->Name);
+					}
+
 					/** Mark the pass as a producer for render targets with a store action. */
 					{
 						const bool bIsStoreAction = RenderTarget.GetStoreAction() != ERenderTargetStoreAction::ENoAction;
diff --git a/Engine/Source/Runtime/RenderCore/Public/RenderGraphResources.h b/Engine/Source/Runtime/RenderCore/Public/RenderGraphResources.h
--- a/Engine/Source/Runtime/RenderCore/Public/RenderGraphResources.h
+++ b/Engine/Source/Runtime/RenderCore/Public/RenderGraphResources.h
@@ -271,6 +271,7 @@ private:
 	IPooledRenderTarget* PooledRenderTarget = nullptr;
 
 	friend class FRDGBuilder;
+	friend class FRDGUserValidation;
 	friend class FRDGBarrierBatcher;
 };
 
diff --git a/Engine/Source/Runtime/RenderCore/Public/RenderTargetPool.h b/Engine/Source/Runtime/RenderCore/Public/RenderTargetPool.h
--- a/Engine/Source/Runtime/RenderCore/Public/RenderTargetPool.h
+++ b/Engine/Source/Runtime/RenderCore/Public/RenderTargetPool.h
@@ -68,6 +68,10 @@ struct RENDERCORE_API FPooledRenderTarget : public IPooledRenderTarget
 	{
 		return !!(Desc.Flags & TexCreate_Transient);
 	}
+	bool IsTracked() const override
+	{
+		return RenderTargetPool != nullptr;
+	}
 	virtual void SetDebugName(const TCHAR *InName);
 	virtual const FPooledRenderTargetDesc& GetDesc() const;
 	virtual uint32 ComputeMemorySize() const;
diff --git a/Engine/Source/Runtime/RenderCore/Public/RendererInterface.h b/Engine/Source/Runtime/RenderCore/Public/RendererInterface.h
--- a/Engine/Source/Runtime/RenderCore/Public/RendererInterface.h
+++ b/Engine/Source/Runtime/RenderCore/Public/RendererInterface.h
@@ -448,6 +448,8 @@ struct IPooledRenderTarget
 	inline FSceneRenderTargetItem& GetRenderTargetItem() { return RenderTargetItem; }
 	/** Get the low level internals (texture/surface) */
 	inline const FSceneRenderTargetItem& GetRenderTargetItem() const { return RenderTargetItem; }
+	/** Returns if the render target is tracked by a pool. */
+	virtual bool IsTracked() const = 0;
 
 	// Refcounting
 	virtual uint32 AddRef() const = 0;
