- Adjust "r.ReadBuffer.MinSize" to a more reasonable default and add "r.ReadBuffer.AlignSize" to so that new buffers are rounded to a reasonable size
- This fixes massive memory bloat with a Niagara system spawning 10k particles, as each frame we add more particles we were allocating a new buffer.  In the end I saw 94 buffers allocated of varying sizes, the rough memory savings here (at least from a stat POV) is ~70mb.

diff --git a/Engine/Source/Runtime/RenderCore/Private/DynamicBufferAllocator.cpp b/Engine/Source/Runtime/RenderCore/Private/DynamicBufferAllocator.cpp
--- a/Engine/Source/Runtime/RenderCore/Private/DynamicBufferAllocator.cpp
+++ b/Engine/Source/Runtime/RenderCore/Private/DynamicBufferAllocator.cpp
@@ -15,11 +15,20 @@ FAutoConsoleVariableRef CVarMaxReadBufferRenderingBytesAllocatedPerFrame(
 	GMaxReadBufferRenderingBytesAllocatedPerFrame,
 	TEXT("The maximum number of transient rendering read buffer bytes to allocate before we start panic logging who is doing the allocations"));
 
-int32 GMinReadBufferRenderingBufferSize = 8 * 1024;
+// The allocator works by looking for the first free buffer that contains the required number of elements.  There is currently no trim so buffers stay in memory.
+// To avoid increasing allocation sizes over multiple frames causing severe memory bloat (i.e. 100 elements, 1001 elements) we first align the required
+// number of elements to GMinReadBufferRenderingBufferSize, we then take the max(aligned num, GMinReadBufferRenderingBufferSize)
+int32 GMinReadBufferRenderingBufferSize = 256 * 1024;
 FAutoConsoleVariableRef CVarMinReadBufferSize(
 	TEXT("r.ReadBuffer.MinSize"),
 	GMinReadBufferRenderingBufferSize,
-	TEXT("The minimum size (in instances) to allocate in blocks for rendering read buffers."));
+	TEXT("The minimum size (in instances) to allocate in blocks for rendering read buffers. i.e. 256*1024 = 1mb for a float buffer"));
+
+int32 GAlignReadBufferRenderingBufferSize = 64 * 1024;
+FAutoConsoleVariableRef CVarAlignReadBufferSize(
+	TEXT("r.ReadBuffer.AlignSize"),
+	GAlignReadBufferRenderingBufferSize,
+	TEXT("The alignment size (in instances) to allocate in blocks for rendering read buffers. i.e. 64*1024 = 256k for a float buffer"));
 
 struct FDynamicReadBufferPool
 {
@@ -117,7 +126,8 @@ FGlobalDynamicReadBuffer::FAllocation FGlobalDynamicReadBuffer::AllocateFloat(ui
 		// Create a new vertex buffer if needed.
 		if (Buffer == NULL)
 		{
-			uint32 NewBufferSize = FMath::Max(Num, (uint32)GMinReadBufferRenderingBufferSize);
+			const uint32 AlignedNum = FMath::DivideAndRoundUp(Num, (uint32)GAlignReadBufferRenderingBufferSize) * GAlignReadBufferRenderingBufferSize;
+			const uint32 NewBufferSize = FMath::Max(AlignedNum, (uint32)GMinReadBufferRenderingBufferSize);
 			Buffer = new FDynamicAllocReadBuffer();
 			FloatBufferPool->Buffers.Add(Buffer);
 			Buffer->Initialize(sizeof(float), NewBufferSize, PF_R32_FLOAT, BUF_Dynamic);
@@ -172,7 +182,8 @@ FGlobalDynamicReadBuffer::FAllocation FGlobalDynamicReadBuffer::AllocateInt32(ui
 		// Create a new vertex buffer if needed.
 		if (Buffer == NULL)
 		{
-			uint32 NewBufferSize = FMath::Max(Num, (uint32)GMinReadBufferRenderingBufferSize);
+			const uint32 AlignedNum = FMath::DivideAndRoundUp(Num, (uint32)GAlignReadBufferRenderingBufferSize) * GAlignReadBufferRenderingBufferSize;
+			const uint32 NewBufferSize = FMath::Max(AlignedNum, (uint32)GMinReadBufferRenderingBufferSize);
 			Buffer = new FDynamicAllocReadBuffer();
 			Int32BufferPool->Buffers.Add(Buffer);
 			Buffer->Initialize(sizeof(int32), NewBufferSize, PF_R32_SINT, BUF_Dynamic);
