Don't consider VT enabled if only 'r.VirtualTexturedLightmaps' is set (but 'r.VirtualTexturing' is not)
This cvar config is not valid, but still possible to manually set the cvars to this state

diff --git a/Engine/Source/Runtime/RenderCore/Private/RenderUtils.cpp b/Engine/Source/Runtime/RenderCore/Private/RenderUtils.cpp
--- a/Engine/Source/Runtime/RenderCore/Private/RenderUtils.cpp
+++ b/Engine/Source/Runtime/RenderCore/Private/RenderUtils.cpp
@@ -1171,11 +1171,9 @@ RENDERCORE_API bool UseVirtualTexturing(ERHIFeatureLevel::Type InFeatureLevel, c
 #endif
 
 		// does the project has it enabled ?
-		static const auto CVarVirtualTextureLightmaps = IConsoleManager::Get().FindTConsoleVariableDataInt(TEXT("r.VirtualTexturedLightmaps"));
 		static const auto CVarVirtualTexture = IConsoleManager::Get().FindTConsoleVariableDataInt(TEXT("r.VirtualTextures"));
-		check(CVarVirtualTextureLightmaps);
 		check(CVarVirtualTexture);
-		if (CVarVirtualTexture->GetValueOnAnyThread() == 0 && CVarVirtualTextureLightmaps->GetValueOnAnyThread() == 0)
+		if (CVarVirtualTexture->GetValueOnAnyThread() == 0)
 		{
 			return false;
 		}		
