Added some metrics to shrink buffer size when Niagara particles go down significantly.
This affects FGlobalDynamicReadBuffer, FGlobalDynamicVertexBuffer and FGlobalDynamicIndexBuffer.

   "r.NumFramesUnusedBeforeReleasingGlobalResourceBuffer"
   Number of frames after which unused global resource allocations will be discarded. Set 0 to ignore. (default=30)

Also added shrinking logic to the FNiagaraDataSet buffers:

   "fx.NiagaraDataBufferMinSize"
   Niagara data buffer minimum allocation size in bytes. (Default=512)

   "fx.NiagaraDataBufferShrinkFactor"
   Niagara data buffer size threshold for shrinking. (Default=3)
   The buffer will be reallocated when the used size becomes 1/F of the allocated size.

diff --git a/Engine/Plugins/FX/Niagara/Source/Niagara/Private/NiagaraDataSet.cpp b/Engine/Plugins/FX/Niagara/Source/Niagara/Private/NiagaraDataSet.cpp
--- a/Engine/Plugins/FX/Niagara/Source/Niagara/Private/NiagaraDataSet.cpp
+++ b/Engine/Plugins/FX/Niagara/Source/Niagara/Private/NiagaraDataSet.cpp
@@ -49,19 +49,20 @@ void FNiagaraSharedObject::FlushDeletionList()
 }
 
 //////////////////////////////////////////////////////////////////////////
-static int32 GRenderDataBlockSize = 4096;
+static int32 GNiagaraDataBufferMinSize = 512;
 static FAutoConsoleVariableRef CVarRenderDataBlockSize(
-	TEXT("fx.RenderDataBlockSize"),
-	GRenderDataBlockSize,
-	TEXT("Size of alloaction blocks for Niagara render data. \n"),
+	TEXT("fx.NiagaraDataBufferMinSize"),
+	GNiagaraDataBufferMinSize,
+	TEXT("Niagara data buffer minimum allocation size in bytes (Default=512). \n"),
 	ECVF_Default
 );
 
-static float GGPUBufferShrinkFactor = .5f;
+static int32 GNiagaraDataBufferShrinkFactor = 3;
 static FAutoConsoleVariableRef CVarNiagaraRenderBufferShrinkFactor(
-	TEXT("fx.RenderBufferShrinkFactor"),
-	GGPUBufferShrinkFactor,
-	TEXT("What factor should the render buffers grow by when they need to grow. \n"),
+	TEXT("fx.NiagaraDataBufferShrinkFactor"),
+	GNiagaraDataBufferShrinkFactor,
+	TEXT("Niagara data buffer size threshold for shrinking. (Default=3) \n")
+	TEXT("The buffer will be reallocated when the used size becomes 1/F of the allocated size. \n"),
 	ECVF_Default
 );
 
@@ -570,10 +571,18 @@ void FNiagaraDataBuffer::Allocate(uint32 InNumInstances, bool bMaintainExisting)
 	}
 
 	FloatStride = GetSafeComponentBufferSize(NumInstancesAllocated * sizeof(float));
-	FloatData.SetNum(FloatStride * Owner->GetNumFloatComponents(), false);
+	{
+		const int32 NewNum = FloatStride * Owner->GetNumFloatComponents();
+		const bool bAllowShrink = GNiagaraDataBufferShrinkFactor * FMath::Max(GNiagaraDataBufferMinSize, NewNum) < FloatData.Max() || !NewNum;
+		FloatData.SetNum(NewNum, bAllowShrink);
+	}
 
 	Int32Stride = GetSafeComponentBufferSize(NumInstancesAllocated * sizeof(int32));
-	Int32Data.SetNum(Int32Stride * Owner->GetNumInt32Components(), false);
+	{
+		const int32 NewNum = Int32Stride * Owner->GetNumInt32Components();
+		const bool bAllowShrink = GNiagaraDataBufferShrinkFactor * FMath::Max(GNiagaraDataBufferMinSize, NewNum) < Int32Data.Max() || !NewNum;
+		Int32Data.SetNum(NewNum, bAllowShrink);
+	}
 
 	INC_MEMORY_STAT_BY(STAT_NiagaraParticleMemory, FloatData.GetAllocatedSize() + Int32Data.GetAllocatedSize());
 
diff --git a/Engine/Source/Runtime/RenderCore/Private/DynamicBufferAllocator.cpp b/Engine/Source/Runtime/RenderCore/Private/DynamicBufferAllocator.cpp
--- a/Engine/Source/Runtime/RenderCore/Private/DynamicBufferAllocator.cpp
+++ b/Engine/Source/Runtime/RenderCore/Private/DynamicBufferAllocator.cpp
@@ -223,8 +223,21 @@ void FGlobalDynamicReadBuffer::Commit()
 		{
 			Buffer.Unlock();
 		}
+		else if (GGlobalBufferNumFramesUnusedThresold && !Buffer.AllocatedByteCount)
+		{
+			++Buffer.NumFramesUnused;
+			if (Buffer.NumFramesUnused >= GGlobalBufferNumFramesUnusedThresold )
+			{
+				// Remove the buffer, assumes they are unordered.
+				Buffer.Release();
+				FloatBufferPool->Buffers.RemoveAtSwap(BufferIndex);
+				--BufferIndex;
+				--NumBuffers;
+			}
+		}
 	}
 	FloatBufferPool->CurrentBuffer = NULL;
+
 	for (int32 BufferIndex = 0, NumBuffers = Int32BufferPool->Buffers.Num(); BufferIndex < NumBuffers; ++BufferIndex)
 	{
 		FDynamicAllocReadBuffer& Buffer = Int32BufferPool->Buffers[BufferIndex];
@@ -232,6 +245,18 @@ void FGlobalDynamicReadBuffer::Commit()
 		{
 			Buffer.Unlock();
 		}
+		else if (GGlobalBufferNumFramesUnusedThresold  && !Buffer.AllocatedByteCount)
+		{
+			++Buffer.NumFramesUnused;
+			if (Buffer.NumFramesUnused >= GGlobalBufferNumFramesUnusedThresold )
+			{
+				// Remove the buffer, assumes they are unordered.
+				Buffer.Release();
+				Int32BufferPool->Buffers.RemoveAtSwap(BufferIndex);
+				--BufferIndex;
+				--NumBuffers;
+			}
+		}
 	}
 	Int32BufferPool->CurrentBuffer = NULL;
 	TotalAllocatedSinceLastCommit = 0;
diff --git a/Engine/Source/Runtime/RenderCore/Private/RenderResource.cpp b/Engine/Source/Runtime/RenderCore/Private/RenderResource.cpp
--- a/Engine/Source/Runtime/RenderCore/Private/RenderResource.cpp
+++ b/Engine/Source/Runtime/RenderCore/Private/RenderResource.cpp
@@ -20,6 +20,12 @@ FAutoConsoleVariableRef CVarMaxVertexBytesAllocatedPerFrame(
 	GMaxVertexBytesAllocatedPerFrame,
 	TEXT("The maximum number of transient vertex buffer bytes to allocate before we start panic logging who is doing the allocations"));
 
+int32 GGlobalBufferNumFramesUnusedThresold = 30;
+FAutoConsoleVariableRef CVarReadBufferNumFramesUnusedThresold(
+	TEXT("r.NumFramesUnusedBeforeReleasingGlobalResourceBuffers"),
+	GGlobalBufferNumFramesUnusedThresold ,
+	TEXT("Number of frames after which unused global resource allocations will be discarded. Set 0 to ignore. (default=30)"));
+
 TLinkedList<FRenderResource*>*& FRenderResource::GetResourceList()
 {
 	static TLinkedList<FRenderResource*>* FirstResourceLink = NULL;
@@ -362,6 +368,8 @@ public:
 	uint32 BufferSize;
 	/** Number of bytes currently allocated from the buffer. */
 	uint32 AllocatedByteCount;
+	/** Number of successive frames for which AllocatedByteCount == 0. Used as a metric to decide when to free the allocation. */
+	int32 NumFramesUnused = 0;
 
 	/** Default constructor. */
 	explicit FDynamicVertexBuffer(uint32 InMinBufferSize)
@@ -392,6 +400,7 @@ public:
 		RHIUnlockVertexBuffer(VertexBufferRHI);
 		MappedBuffer = NULL;
 		AllocatedByteCount = 0;
+		NumFramesUnused = 0;
 	}
 
 	// FRenderResource interface.
@@ -523,6 +532,18 @@ void FGlobalDynamicVertexBuffer::Commit()
 		{
 			VertexBuffer.Unlock();
 		}
+		else if (GGlobalBufferNumFramesUnusedThresold && !VertexBuffer.AllocatedByteCount)
+		{
+			++VertexBuffer.NumFramesUnused;
+			if (VertexBuffer.NumFramesUnused >= GGlobalBufferNumFramesUnusedThresold)
+			{
+				// Remove the buffer, assumes they are unordered.
+				VertexBuffer.ReleaseResource();
+				Pool->VertexBuffers.RemoveAtSwap(BufferIndex);
+				--BufferIndex;
+				--NumBuffers;
+			}
+		}
 	}
 	Pool->CurrentVertexBuffer = NULL;
 	TotalAllocatedSinceLastCommit = 0;
@@ -551,6 +572,8 @@ public:
 	uint32 AllocatedByteCount;
 	/** Stride of the buffer in bytes. */
 	uint32 Stride;
+	/** Number of successive frames for which AllocatedByteCount == 0. Used as a metric to decide when to free the allocation. */
+	int32 NumFramesUnused = 0;
 
 	/** Initialization constructor. */
 	explicit FDynamicIndexBuffer(uint32 InMinBufferSize, uint32 InStride)
@@ -582,6 +605,7 @@ public:
 		RHIUnlockIndexBuffer(IndexBufferRHI);
 		MappedBuffer = NULL;
 		AllocatedByteCount = 0;
+		NumFramesUnused = 0;
 	}
 
 	// FRenderResource interface.
@@ -719,6 +743,18 @@ void FGlobalDynamicIndexBuffer::Commit()
 			{
 				IndexBuffer.Unlock();
 			}
+			else if (GGlobalBufferNumFramesUnusedThresold && !IndexBuffer.AllocatedByteCount)
+			{
+				++IndexBuffer.NumFramesUnused;
+				if (IndexBuffer.NumFramesUnused >= GGlobalBufferNumFramesUnusedThresold)
+				{
+					// Remove the buffer, assumes they are unordered.
+					IndexBuffer.ReleaseResource();
+					Pool->IndexBuffers.RemoveAtSwap(BufferIndex);
+					--BufferIndex;
+					--NumBuffers;
+				}
+			}
 		}
 		Pool->CurrentIndexBuffer = NULL;
 	}
diff --git a/Engine/Source/Runtime/RenderCore/Public/DynamicBufferAllocator.h b/Engine/Source/Runtime/RenderCore/Public/DynamicBufferAllocator.h
--- a/Engine/Source/Runtime/RenderCore/Public/DynamicBufferAllocator.h
+++ b/Engine/Source/Runtime/RenderCore/Public/DynamicBufferAllocator.h
@@ -12,12 +12,9 @@ struct FDynamicReadBufferPool;
 
 struct FDynamicAllocReadBuffer : public FDynamicReadBuffer
 {
-	FDynamicAllocReadBuffer()
-		: FDynamicReadBuffer()
-		, AllocatedByteCount(0)
-	{}
-
-	int32 AllocatedByteCount;
+	int32 AllocatedByteCount = 0;
+	/** Number of successive frames for which AllocatedByteCount == 0. Used as a metric to decide when to free the allocation. */
+	int32 NumFramesUnused = 0;
 
 	/**
 	* Unocks the buffer so the GPU may read from it.
@@ -26,6 +23,7 @@ struct FDynamicAllocReadBuffer : public FDynamicReadBuffer
 	{
 		FDynamicReadBuffer::Unlock();
 		AllocatedByteCount = 0;
+		NumFramesUnused = 0;
 	}
 };
 
diff --git a/Engine/Source/Runtime/RenderCore/Public/RenderResource.h b/Engine/Source/Runtime/RenderCore/Public/RenderResource.h
--- a/Engine/Source/Runtime/RenderCore/Public/RenderResource.h
+++ b/Engine/Source/Runtime/RenderCore/Public/RenderResource.h
@@ -11,6 +11,9 @@
 #include "RHI.h"
 #include "RenderCore.h"
 
+/** Number of frames after which unused global resource allocations will be discarded. */
+extern int32 GGlobalBufferNumFramesUnusedThresold;
+
 /**
  * A rendering resource which is owned by the rendering thread.
  */
