Redo CL 8374078 to resolve merge conflict - Reduce FRealtimeGPUProfiler overhead

diff --git a/Engine/Source/Runtime/RenderCore/Private/ProfilingDebugging/RealtimeGPUProfiler.cpp b/Engine/Source/Runtime/RenderCore/Private/ProfilingDebugging/RealtimeGPUProfiler.cpp
--- a/Engine/Source/Runtime/RenderCore/Private/ProfilingDebugging/RealtimeGPUProfiler.cpp
+++ b/Engine/Source/Runtime/RenderCore/Private/ProfilingDebugging/RealtimeGPUProfiler.cpp
@@ -113,134 +113,91 @@ public:
 	static const uint64 InvalidQueryResult = 0xFFFFFFFFFFFFFFFFull;
 
 public:
-	FRealtimeGPUProfilerEvent(const FName& InName, const FName& InStatName, FRHIRenderQueryPool* RenderQueryPool, uint32& QueryCount)
-		: StartResultMicroseconds(InvalidQueryResult)
+	FRealtimeGPUProfilerEvent(FRHIRenderQueryPool& RenderQueryPool)
+		: StartQuery(RenderQueryPool.AllocateQuery())
+		, EndQuery(RenderQueryPool.AllocateQuery())
+		, StartResultMicroseconds(InvalidQueryResult)
 		, EndResultMicroseconds(InvalidQueryResult)
+#if TRACING_PROFILER
 		, FrameNumber(-1)
+#endif
+#if DO_CHECK
 		, bInsideQuery(false)
-		, bBeginQueryInFlight(false)
-		, bEndQueryInFlight(false)
-	{
-#if STATS
-		StatName = InStatName;
 #endif
-		Name = InName;
-
-		const int MaxGPUQueries = CVarGPUStatsMaxQueriesPerFrame.GetValueOnRenderThread();
-		if ( MaxGPUQueries == -1 || QueryCount < uint32(MaxGPUQueries) )
-		{
-			StartQuery = RenderQueryPool->AllocateQuery();
-			check(StartQuery.GetQuery() != nullptr);
-			EndQuery = RenderQueryPool->AllocateQuery();
-			check(EndQuery.GetQuery() != nullptr);
-			QueryCount += 2;
-		}
-	}
-
-	bool HasQueriesAllocated() const 
-	{ 
-		return StartQuery.GetQuery() != nullptr;
-	}
-
-	void ReleaseQueries(FRHIRenderQueryPool* RenderQueryPool, FRHICommandListImmediate* RHICmdListPtr, uint32& QueryCount)
 	{
-		if ( HasQueriesAllocated() )
-		{
-			if (RHICmdListPtr)
-			{
-				// If we have queries in flight then get results before releasing back to the pool to avoid an ensure fail in the gnm RHI
-				uint64 Temp;
-				if (bBeginQueryInFlight)
-				{
-					RHICmdListPtr->GetRenderQueryResult(StartQuery.GetQuery(), Temp, true);
-				}
-
-				if (bEndQueryInFlight)
-				{
-					RHICmdListPtr->GetRenderQueryResult(EndQuery.GetQuery(), Temp, true);
-				}
-			}
-			StartQuery.ReleaseQuery();
-			EndQuery.ReleaseQuery();
-			QueryCount -= 2;
-		}
+		check(StartQuery.IsValid() && EndQuery.IsValid());
 	}
 
-	void Begin(FRHICommandListImmediate& RHICmdList)
+	void Begin(FRHICommandListImmediate& RHICmdList, const FName& NewName, const FName& NewStatName)
 	{
 		check(IsInRenderingThread());
-		check(!bInsideQuery);
+		check(!bInsideQuery && StartQuery.IsValid());
+#if DO_CHECK
 		bInsideQuery = true;
+#endif
+		RHICmdList.EndRenderQuery(StartQuery.GetQuery());
 
-		if (StartQuery.GetQuery() != nullptr)
-		{
-			RHICmdList.EndRenderQuery(StartQuery.GetQuery());
-			bBeginQueryInFlight = true;
-		}
+		Name = NewName;
+		STAT(StatName = NewStatName;)
 		StartResultMicroseconds = InvalidQueryResult;
 		EndResultMicroseconds = InvalidQueryResult;
+#if TRACING_PROFILER
 		FrameNumber = GFrameNumberRenderThread;
+#endif
 	}
 
 	void End(FRHICommandListImmediate& RHICmdList)
 	{
 		check(IsInRenderingThread());
-		check(bInsideQuery);
+		check(bInsideQuery && EndQuery.IsValid());
+#if DO_CHECK
 		bInsideQuery = false;
-
-		if ( HasQueriesAllocated() )
-		{
-			RHICmdList.EndRenderQuery(EndQuery.GetQuery());
-			bEndQueryInFlight = true;
-		}
+#endif
+		RHICmdList.EndRenderQuery(EndQuery.GetQuery());
 	}
 
 	bool GatherQueryResults(FRHICommandListImmediate& RHICmdList)
 	{
-		QUICK_SCOPE_CYCLE_COUNTER(STAT_SceneUtils_GatherQueryResults);
+		//QUICK_SCOPE_CYCLE_COUNTER(STAT_SceneUtils_GatherQueryResults);
+
 		// Get the query results which are still outstanding
 		check(GFrameNumberRenderThread != FrameNumber);
-		if ( HasQueriesAllocated() )
+		check(StartQuery.IsValid() && EndQuery.IsValid());
+
+		if (StartResultMicroseconds == InvalidQueryResult)
 		{
-			if (StartResultMicroseconds == InvalidQueryResult)
-			{
-				if (!RHICmdList.GetRenderQueryResult(StartQuery.GetQuery(), StartResultMicroseconds, false))
-				{
-					StartResultMicroseconds = InvalidQueryResult;
-				}
-				bBeginQueryInFlight = false;
-			}
-			if (EndResultMicroseconds == InvalidQueryResult)
+			if (!RHICmdList.GetRenderQueryResult(StartQuery.GetQuery(), StartResultMicroseconds, false))
 			{
-				if (!RHICmdList.GetRenderQueryResult(EndQuery.GetQuery(), EndResultMicroseconds, false))
-				{
-					EndResultMicroseconds = InvalidQueryResult;
-				}
-				bEndQueryInFlight = false;
+				StartResultMicroseconds = InvalidQueryResult;
 			}
 		}
-		else
+
+		if (EndResultMicroseconds == InvalidQueryResult)
 		{
-			// If we don't have a query allocated, just set the results to zero
-			EndResultMicroseconds = StartResultMicroseconds = 0;
+			if (!RHICmdList.GetRenderQueryResult(EndQuery.GetQuery(), EndResultMicroseconds, false))
+			{
+				EndResultMicroseconds = InvalidQueryResult;
+			}
 		}
+
 		return HasValidResult();
 	}
 
-	float GetResultMS() const
+	uint64 GetResultUs() const
 	{
 		check(HasValidResult());
-		if (EndResultMicroseconds < StartResultMicroseconds)
+
+		if (StartResultMicroseconds > EndResultMicroseconds)
 		{
-			// This should never happen...
-			return 0.0f;
+			return 0llu;
 		}
-		return float(EndResultMicroseconds - StartResultMicroseconds) / 1000.0f;
+
+		return EndResultMicroseconds - StartResultMicroseconds;
 	}
 
 	bool HasValidResult() const
 	{
-		return StartResultMicroseconds != FRealtimeGPUProfilerEvent::InvalidQueryResult && EndResultMicroseconds != FRealtimeGPUProfilerEvent::InvalidQueryResult;
+		return StartResultMicroseconds != InvalidQueryResult && EndResultMicroseconds != InvalidQueryResult;
 	}
 
 #if STATS
@@ -249,6 +206,7 @@ public:
 		return StatName;
 	}
 #endif
+
 	const FName& GetName() const
 	{
 		return Name;
@@ -264,25 +222,35 @@ public:
 		return EndResultMicroseconds;
 	}
 
+#if TRACING_PROFILER
 	uint32 GetFrameNumber() const
 	{
 		return FrameNumber;
 	}
+#endif
+
+	static constexpr uint32 GetNumRHIQueriesPerEvent()
+	{
+		return 2u;
+	}
 
 private:
 	FRHIPooledRenderQuery StartQuery;
 	FRHIPooledRenderQuery EndQuery;
-#if STATS
-	FName StatName;
-#endif
+
 	FName Name;
+	STAT(FName StatName;)
+
 	uint64 StartResultMicroseconds;
 	uint64 EndResultMicroseconds;
+
+#if TRACING_PROFILER
 	uint32 FrameNumber;
+#endif
 
+#if DO_CHECK
 	bool bInsideQuery;
-	bool bBeginQueryInFlight;
-	bool bEndQueryInFlight;
+#endif
 };
 
 /*-----------------------------------------------------------------------------
@@ -291,237 +259,316 @@ Container for a single frame's GPU stats
 -----------------------------------------------------------------------------*/
 class FRealtimeGPUProfilerFrame
 {
-	uint32& QueryCount;
 public:
 	FRealtimeGPUProfilerFrame(FRenderQueryPoolRHIRef InRenderQueryPool, uint32& InQueryCount)
-		: QueryCount(InQueryCount)
-		, FrameNumber(-1)
+		: NextEventIdx(1)
+		, OverflowEventCount(0)
+		, NextResultPendingEventIdx(1)
+		, QueryCount(InQueryCount)
 		, RenderQueryPool(InRenderQueryPool)
-	{}
+	{
+		GpuProfilerEvents.Empty(GPredictedMaxNumEvents);
+		GpuProfilerEvents.AddUninitialized(GPredictedMaxNumEvents);
+		FMemory::Memset(&GpuProfilerEvents[0], 0, sizeof(FRealtimeGPUProfilerEvent));
+
+		for (uint32 Idx = 1u; Idx < GPredictedMaxNumEvents; ++Idx)
+		{
+			new (&GpuProfilerEvents[Idx]) FRealtimeGPUProfilerEvent(*RenderQueryPool);
+		}
+
+		QueryCount += (GPredictedMaxNumEvents - 1u) * FRealtimeGPUProfilerEvent::GetNumRHIQueriesPerEvent();
+
+		GpuProfilerEventParentIndices.Empty(GPredictedMaxNumEvents);
+		GpuProfilerEventParentIndices.AddUninitialized();
+
+		EventStack.Empty(GPredictedMaxStackDepth);
+		EventStack.Add(0);
+
+		EventAggregates.Empty(GPredictedMaxNumEvents);
+		EventAggregates.AddUninitialized();
+	}
 
 	~FRealtimeGPUProfilerFrame()
 	{
-		Clear(nullptr);
+		QueryCount -= (GpuProfilerEvents.Num() - 1) * FRealtimeGPUProfilerEvent::GetNumRHIQueriesPerEvent();
 	}
 
-	void PushEvent(FRHICommandListImmediate& RHICmdList, const FName& Name, const FName& StatName)
+	void Clear(void* Dummy)
 	{
-		// TODO: this should really use a pool / free list
-		FRealtimeGPUProfilerEvent* Event = new FRealtimeGPUProfilerEvent(Name, StatName, RenderQueryPool, QueryCount);
-		const int32 EventIndex = GpuProfilerEvents.Num();
+		check(!OverflowEventCount);
 
-		GpuProfilerEvents.Add(Event);
+		NextEventIdx = 1;
+		NextResultPendingEventIdx = 1;
 
-		FRealtimeGPUProfilerTimelineEvent TimelineEvent = {};
-		TimelineEvent.Type = FRealtimeGPUProfilerTimelineEvent::EType::PushEvent;
-		TimelineEvent.EventIndex = EventIndex;
+		GpuProfilerEventParentIndices.Reset();
+		GpuProfilerEventParentIndices.AddUninitialized();
 
-		GpuProfilerTimelineEvents.Push(TimelineEvent);
+		EventStack.Reset();
+		EventStack.Add(0);
 
-		EventStack.Push(EventIndex);
-		Event->Begin(RHICmdList);
+		EventAggregates.Reset();
+		EventAggregates.AddUninitialized();
 	}
 
-	void PopEvent(FRHICommandListImmediate& RHICmdList)
+	void PushEvent(FRHICommandListImmediate& RHICmdList, const FName& Name, const FName& StatName)
 	{
-		int32 Index = EventStack.Pop();
+		if (NextEventIdx >= GpuProfilerEvents.Num())
+		{
+			const int32 MaxNumQueries = CVarGPUStatsMaxQueriesPerFrame.GetValueOnRenderThread();
 
-		FRealtimeGPUProfilerTimelineEvent TimelineEvent = {};
-		TimelineEvent.Type = FRealtimeGPUProfilerTimelineEvent::EType::PopEvent;
-		TimelineEvent.EventIndex = Index;
-		GpuProfilerTimelineEvents.Push(TimelineEvent);
+			if (MaxNumQueries < 0 || QueryCount < (uint32)MaxNumQueries)
+			{
+				new (GpuProfilerEvents) FRealtimeGPUProfilerEvent(*RenderQueryPool);
+				QueryCount += FRealtimeGPUProfilerEvent::GetNumRHIQueriesPerEvent();
+			}
+			else
+			{
+				++OverflowEventCount;
+				return;
+			}
+		}
+
+		const int32 EventIdx = NextEventIdx++;
 
-		GpuProfilerEvents[Index]->End(RHICmdList);
+		GpuProfilerEventParentIndices.Add(EventStack.Last());
+		EventStack.Push(EventIdx);
+		GpuProfilerEvents[EventIdx].Begin(RHICmdList, Name, StatName);
 	}
 
-	void Clear(FRHICommandListImmediate* RHICommandListPtr)
+	void PopEvent(FRHICommandListImmediate& RHICmdList)
 	{
-		EventStack.Empty();
-
-		for (int Index = 0; Index < GpuProfilerEvents.Num(); Index++)
+		if (OverflowEventCount)
 		{
-			if (GpuProfilerEvents[Index])
-			{
-				GpuProfilerEvents[Index]->ReleaseQueries(RenderQueryPool, RHICommandListPtr, QueryCount);
-				delete GpuProfilerEvents[Index];
-			}
+			--OverflowEventCount;
+			return;
 		}
-		GpuProfilerEvents.Empty();
-		GpuProfilerTimelineEvents.Empty();
-		EventAggregates.Empty();
+
+		const int32 EventIdx = EventStack.Pop(false);
+
+		GpuProfilerEvents[EventIdx].End(RHICmdList);
 	}
 
 	bool UpdateStats(FRHICommandListImmediate& RHICmdList)
 	{
-		const bool bCsvStatsEnabled = !!CVarGPUCsvStatsEnabled.GetValueOnRenderThread();
-		const bool bTracingStatsEnabled = !!CVarGPUTracingStatsEnabled.GetValueOnRenderThread();
-
 		// Gather any remaining results and check all the results are ready
-		bool bAllQueriesAllocated = true;
-		bool bAnyEventFailed = false;
-		for (int Index = 0; Index < GpuProfilerEvents.Num(); Index++)
+		const int32 NumEventsThisFramePlusOne = NextEventIdx;
+
+		for (; NextResultPendingEventIdx < NumEventsThisFramePlusOne; ++NextResultPendingEventIdx)
 		{
-			FRealtimeGPUProfilerEvent* Event = GpuProfilerEvents[Index];
-			check(Event != nullptr);
-			if (!Event->HasValidResult())
-			{
-				Event->GatherQueryResults(RHICmdList);
-			}
-			if (!Event->HasValidResult())
+			FRealtimeGPUProfilerEvent& Event = GpuProfilerEvents[NextResultPendingEventIdx];
+
+			if (!Event.GatherQueryResults(RHICmdList))
 			{
-#if UE_BUILD_DEBUG
-				UE_LOG(LogRendererCore, Warning, TEXT("Query '%s' not ready."), *Event->GetName().ToString());
+#if !(UE_BUILD_SHIPPING || UE_BUILD_TEST)
+				UE_LOG(LogRendererCore, Warning, TEXT("Query '%s' not ready."), *Event.GetName().ToString());
 #endif
 				// The frame isn't ready yet. Don't update stats - we'll try again next frame. 
-				bAnyEventFailed = true;
-				continue;
-			}
-			if (!Event->HasQueriesAllocated())
-			{
-				bAllQueriesAllocated = false;
+				return false;
 			}
-		}
-
-		if (bAnyEventFailed)
-		{
-			return false;
-		}
-
-		if (!bAllQueriesAllocated)
-		{
-			static bool bWarned = false;
 
-			if (!bWarned)
-			{
-				bWarned = true;
-				UE_LOG(LogRendererCore, Warning, TEXT("Ran out of GPU queries! Results for this frame will be incomplete"));
-			}
+			FGPUEventTimeAggregate Aggregate;
+			Aggregate.InclusiveTimeUs = (uint32)Event.GetResultUs();
+			Aggregate.ExclusiveTimeUs = Aggregate.InclusiveTimeUs;
+			EventAggregates.Add(Aggregate);
 		}
 
 		// Calculate inclusive and exclusive time for all events
-
-		EventAggregates.Reserve(GpuProfilerEvents.Num());
-		TArray<int32, TInlineAllocator<32>> TimelineEventStack;
-
-		for (const FRealtimeGPUProfilerEvent* Event : GpuProfilerEvents)
+		for (int32 EventIdx = 1; EventIdx < GpuProfilerEventParentIndices.Num(); ++EventIdx)
 		{
-			FGPUEventTimeAggregate Aggregate;
-			Aggregate.InclusiveTime = Event->GetResultMS();
-			Aggregate.ExclusiveTime = Aggregate.InclusiveTime;
-			EventAggregates.Push(Aggregate);
-		}
+			const int32 ParentIdx = GpuProfilerEventParentIndices[EventIdx];
 
-		for (const FRealtimeGPUProfilerTimelineEvent& TimelineEvent : GpuProfilerTimelineEvents)
-		{
-			if (TimelineEvent.Type == FRealtimeGPUProfilerTimelineEvent::EType::PushEvent)
-			{
-				if (TimelineEventStack.Num() != 0)
-				{
-					EventAggregates[TimelineEventStack.Last()].ExclusiveTime -= EventAggregates[TimelineEvent.EventIndex].InclusiveTime;
-				}
-				TimelineEventStack.Push(TimelineEvent.EventIndex);
-			}
-			else
-			{
-				TimelineEventStack.Pop();
-			}
+			EventAggregates[ParentIdx].ExclusiveTimeUs -= EventAggregates[EventIdx].InclusiveTimeUs;
 		}
 
 		// Update the stats
-
+#if CSV_PROFILER
+		const bool bCsvStatsEnabled = !!CVarGPUCsvStatsEnabled.GetValueOnRenderThread();
+		FCsvProfiler* CsvProfiler = bCsvStatsEnabled ? FCsvProfiler::Get() : nullptr;
+#endif
 		const bool GPUStatsChildTimesIncluded = !!CVarGPUStatsChildTimesIncluded.GetValueOnRenderThread();
-		float TotalMS = 0.0f;
+		uint64 TotalUs = 0llu;
+		FNameSet StatSeenSet;
 
-		TMap<FName, bool> StatSeenMap;
-		for (int Index = 0; Index < GpuProfilerEvents.Num(); Index++)
+		for (int32 Idx = 1; Idx < NumEventsThisFramePlusOne; ++Idx)
 		{
-			FRealtimeGPUProfilerEvent* Event = GpuProfilerEvents[Index];
-			check(Event != nullptr);
-			check(Event->HasValidResult());
-			const FName& StatName = Event->GetName();
+			const FRealtimeGPUProfilerEvent& Event = GpuProfilerEvents[Idx];
+			const FGPUEventTimeAggregate IncExcTime = EventAggregates[Idx];
 
 			// Check if we've seen this stat yet 
-			bool bIsNew = false;
-			if (StatSeenMap.Find(StatName) == nullptr)
-			{
-				StatSeenMap.Add(StatName, true);
-				bIsNew = true;
-			}
+			const bool bKnownStat = StatSeenSet.Add(Event.GetName());
 
-			const float EventTime = GPUStatsChildTimesIncluded
-				? EventAggregates[Index].InclusiveTime
-				: EventAggregates[Index].ExclusiveTime;
+			const uint32 EventTimeUs = GPUStatsChildTimesIncluded ? IncExcTime.InclusiveTimeUs : IncExcTime.ExclusiveTimeUs;
+			TotalUs += IncExcTime.ExclusiveTimeUs;
 
 #if STATS
-			EStatOperation::Type StatOp = bIsNew ? EStatOperation::Set : EStatOperation::Add;
-			FThreadStats::AddMessage(Event->GetStatName(), StatOp, double(EventTime));
+			const EStatOperation::Type StatOp = bKnownStat ? EStatOperation::Add : EStatOperation::Set;
+			FThreadStats::AddMessage(Event.GetStatName(), StatOp, EventTimeUs / 1000.);
 #endif
 
 #if CSV_PROFILER
-			if (bCsvStatsEnabled)
+			if (CsvProfiler)
 			{
-				ECsvCustomStatOp CsvStatOp = bIsNew ? ECsvCustomStatOp::Set : ECsvCustomStatOp::Accumulate;
-				FCsvProfiler::Get()->RecordCustomStat(Event->GetName(), CSV_CATEGORY_INDEX(GPU), EventTime, CsvStatOp);
+				const ECsvCustomStatOp CsvStatOp = bKnownStat ? ECsvCustomStatOp::Accumulate : ECsvCustomStatOp::Set;
+				CsvProfiler->RecordCustomStat(Event.GetName(), CSV_CATEGORY_INDEX(GPU), EventTimeUs / 1000.f, CsvStatOp);
 			}
 #endif
 
 #if TRACING_PROFILER
+			const bool bTracingStatsEnabled = !!CVarGPUTracingStatsEnabled.GetValueOnRenderThread();
 			if (bTracingStatsEnabled)
 			{
 				ANSICHAR EventName[NAME_SIZE];
-				Event->GetName().GetPlainANSIString(EventName);
+				Event.GetName().GetPlainANSIString(EventName);
 
 				const uint32 GPUIndex = 0;
 				FTracingProfiler::Get()->AddGPUEvent(
 					EventName,
-					Event->GetStartResultMicroseconds(),
-					Event->GetEndResultMicroseconds(),
+					Event.GetStartResultMicroseconds(),
+					Event.GetEndResultMicroseconds(),
 					GPUIndex,
-					Event->GetFrameNumber());
+					Event.GetFrameNumber());
 			}
 #endif //TRACING_PROFILER
-
-			TotalMS += EventAggregates[Index].ExclusiveTime;
 		}
 
 #if STATS
-		FThreadStats::AddMessage( GET_STATFNAME(Stat_GPU_Total), EStatOperation::Set, double(TotalMS) );
+		FThreadStats::AddMessage(GET_STATFNAME(Stat_GPU_Total), EStatOperation::Set, TotalUs / 1000.);
 #endif 
 
 #if CSV_PROFILER
-		if (bCsvStatsEnabled)
+		if (CsvProfiler)
 		{
-			FCsvProfiler::Get()->RecordCustomStat(CSV_STAT_FNAME(Total), CSV_CATEGORY_INDEX(GPU), TotalMS, ECsvCustomStatOp::Set);
+			CsvProfiler->RecordCustomStat(CSV_STAT_FNAME(Total), CSV_CATEGORY_INDEX(GPU), TotalUs / 1000.f, ECsvCustomStatOp::Set);
 		}
 #endif
 		return true;
 	}
 
 private:
-
-	TArray<FRealtimeGPUProfilerEvent*> GpuProfilerEvents;
-	TArray<int32> EventStack;
-
-	struct FRealtimeGPUProfilerTimelineEvent
+	struct FGPUEventTimeAggregate
 	{
-		enum class EType { PushEvent, PopEvent };
-		EType Type;
-		int32 EventIndex;
+		uint32 ExclusiveTimeUs;
+		uint32 InclusiveTimeUs;
 	};
 
-	// All profiler push and pop events are recorded to calculate inclusive and exclusive timing
-	// while maintaining hierarchy and not splitting events unnecessarily.
-	TArray<FRealtimeGPUProfilerTimelineEvent> GpuProfilerTimelineEvents;
+	static constexpr uint32 GPredictedMaxNumEvents = 100u;
+	static constexpr uint32 GPredictedMaxNumEventsUpPow2 = 128u;
+	static constexpr uint32 GPredictedMaxStackDepth = 32u;
 
-	struct FGPUEventTimeAggregate
+	class FNameSet
 	{
-		float ExclusiveTime;
-		float InclusiveTime;
+	public:
+		FNameSet()
+			: NumElements(0)
+			, Capacity(GInitialCapacity)
+			, SecondaryStore(nullptr)
+		{
+			FMemory::Memset(InlineStore, 0, GInitialCapacity * sizeof(FName));
+		}
+
+		~FNameSet()
+		{
+			if (SecondaryStore)
+			{
+				FMemory::Free(SecondaryStore);
+				SecondaryStore = nullptr;
+			}
+		}
+
+		// @return Whether Name is already in set
+		bool Add(const FName& Name)
+		{
+			check(Name != NAME_None);
+
+			if (NumElements * GResizeDivFactor > Capacity)
+			{
+				uint32 NewCapacity = Capacity;
+
+				do
+				{
+					NewCapacity *= 2u;
+				} while (NumElements * GResizeDivFactor > NewCapacity);
+
+				Resize(NewCapacity);
+			}
+
+			FName* NameStore = GetNameStore();
+			const uint32 NameHash = GetTypeHash(Name);
+			const uint32 Mask = Capacity - 1;
+			uint32 Idx = NameHash & Mask;
+			uint32 Probe = 1;
+			const FName NameNone = NAME_None;
+
+			while (NameNone != NameStore[Idx] && Name != NameStore[Idx])
+			{
+				Idx = (Idx + Probe) & Mask;
+				Probe *= 2u;
+			}
+
+			if (NameNone != NameStore[Idx])
+			{
+				return true;
+			}
+			else
+			{
+				NameStore[Idx] = Name;
+				++NumElements;
+				return false;
+			}
+		}
+
+	private:
+		void Resize(uint32 NewCapacity)
+		{
+			const bool bNeedFree = !!SecondaryStore;
+			FName* OldStore = bNeedFree ? SecondaryStore : InlineStore;
+
+			SecondaryStore = (FName*)FMemory::Malloc(NewCapacity * sizeof(FName));
+			FMemory::Memset(SecondaryStore, 0, NewCapacity * sizeof(FName));
+
+			const uint32 OldCapacity = Capacity;
+			Capacity = NewCapacity;
+
+			for (uint32 Idx = 0; Idx < OldCapacity; ++Idx)
+			{
+				const FName& Name = OldStore[Idx];
+				if (Name != NAME_None)
+				{
+					Add(Name);
+				}
+			}
+
+			if (bNeedFree)
+			{
+				FMemory::Free(OldStore);
+			}
+		}
+
+		FName* GetNameStore()
+		{
+			return SecondaryStore ? SecondaryStore : InlineStore;
+		}
+
+		static constexpr uint32 GResizeDivFactor = 2u;
+		static constexpr uint32 GInitialCapacity = GPredictedMaxNumEventsUpPow2 * GResizeDivFactor;
+
+		uint32 NumElements;
+		uint32 Capacity;
+		FName InlineStore[GInitialCapacity];
+		FName* SecondaryStore;
 	};
-	TArray<FGPUEventTimeAggregate> EventAggregates;
 
-	uint32 FrameNumber;
+	int32 NextEventIdx;
+	int32 OverflowEventCount;
+	int32 NextResultPendingEventIdx;
+
+	uint32& QueryCount;
 	FRenderQueryPoolRHIRef RenderQueryPool;
+
+	TArray<FRealtimeGPUProfilerEvent, TInlineAllocator<GPredictedMaxNumEvents>> GpuProfilerEvents;
+	TArray<int32, TInlineAllocator<GPredictedMaxNumEvents>> GpuProfilerEventParentIndices;
+	TArray<int32, TInlineAllocator<GPredictedMaxStackDepth>> EventStack;
+	TArray<FGPUEventTimeAggregate, TInlineAllocator<GPredictedMaxNumEvents>> EventAggregates;
 };
 
 /*-----------------------------------------------------------------------------
diff --git a/Engine/Source/Runtime/RenderCore/Public/ProfilingDebugging/TracingProfiler.h b/Engine/Source/Runtime/RenderCore/Public/ProfilingDebugging/TracingProfiler.h
--- a/Engine/Source/Runtime/RenderCore/Public/ProfilingDebugging/TracingProfiler.h
+++ b/Engine/Source/Runtime/RenderCore/Public/ProfilingDebugging/TracingProfiler.h
@@ -16,7 +16,11 @@
 #include "CoreTypes.h"
 #include "CoreMinimal.h"
 
-#define TRACING_PROFILER (WITH_ENGINE && !UE_BUILD_SHIPPING)
+#ifndef USE_TRACING_PROFILER_IN_TEST_BUILD
+#define USE_TRACING_PROFILER_IN_TEST_BUILD 1
+#endif
+
+#define TRACING_PROFILER (WITH_ENGINE && !UE_BUILD_SHIPPING && (!UE_BUILD_TEST || USE_TRACING_PROFILER_IN_TEST_BUILD))
 
 #if TRACING_PROFILER
 
