Compile fix for tools build

diff --git a/Engine/Source/Runtime/RenderCore/Private/ProfilingDebugging/RealtimeGPUProfiler.cpp b/Engine/Source/Runtime/RenderCore/Private/ProfilingDebugging/RealtimeGPUProfiler.cpp
--- a/Engine/Source/Runtime/RenderCore/Private/ProfilingDebugging/RealtimeGPUProfiler.cpp
+++ b/Engine/Source/Runtime/RenderCore/Private/ProfilingDebugging/RealtimeGPUProfiler.cpp
@@ -161,7 +161,9 @@ public:
 		//QUICK_SCOPE_CYCLE_COUNTER(STAT_SceneUtils_GatherQueryResults);
 
 		// Get the query results which are still outstanding
+#if TRACING_PROFILER
 		check(GFrameNumberRenderThread != FrameNumber);
+#endif
 		check(StartQuery.IsValid() && EndQuery.IsValid());
 
 		if (StartResultMicroseconds == InvalidQueryResult)
