Fix compile errors

diff --git a/Engine/Source/Runtime/RenderCore/Private/ProfilingDebugging/RealtimeGPUProfiler.cpp b/Engine/Source/Runtime/RenderCore/Private/ProfilingDebugging/RealtimeGPUProfiler.cpp
--- a/Engine/Source/Runtime/RenderCore/Private/ProfilingDebugging/RealtimeGPUProfiler.cpp
+++ b/Engine/Source/Runtime/RenderCore/Private/ProfilingDebugging/RealtimeGPUProfiler.cpp
@@ -8,6 +8,8 @@
 // Only exposed for debugging. Disabling this carries a severe performance penalty
 #define RENDER_QUERY_POOLING_ENABLED 1
 
+#define REALTIME_GPU_PROFILER_EVENT_TRACK_FRAME_NUMBER (TRACING_PROFILER || DO_CHECK)
+
 #if HAS_GPU_STATS 
 
 CSV_DEFINE_CATEGORY_MODULE(RENDERCORE_API, GPU, true);
@@ -118,7 +120,7 @@ public:
 		, EndQuery(RenderQueryPool.AllocateQuery())
 		, StartResultMicroseconds(InvalidQueryResult)
 		, EndResultMicroseconds(InvalidQueryResult)
-#if TRACING_PROFILER
+#if REALTIME_GPU_PROFILER_EVENT_TRACK_FRAME_NUMBER
 		, FrameNumber(-1)
 #endif
 #if DO_CHECK
@@ -141,7 +143,7 @@ public:
 		STAT(StatName = NewStatName;)
 		StartResultMicroseconds = InvalidQueryResult;
 		EndResultMicroseconds = InvalidQueryResult;
-#if TRACING_PROFILER
+#if REALTIME_GPU_PROFILER_EVENT_TRACK_FRAME_NUMBER
 		FrameNumber = GFrameNumberRenderThread;
 #endif
 	}
@@ -161,9 +163,7 @@ public:
 		//QUICK_SCOPE_CYCLE_COUNTER(STAT_SceneUtils_GatherQueryResults);
 
 		// Get the query results which are still outstanding
-#if TRACING_PROFILER
 		check(GFrameNumberRenderThread != FrameNumber);
-#endif
 		check(StartQuery.IsValid() && EndQuery.IsValid());
 
 		if (StartResultMicroseconds == InvalidQueryResult)
@@ -246,7 +246,7 @@ private:
 	uint64 StartResultMicroseconds;
 	uint64 EndResultMicroseconds;
 
-#if TRACING_PROFILER
+#if REALTIME_GPU_PROFILER_EVENT_TRACK_FRAME_NUMBER
 	uint32 FrameNumber;
 #endif
 
