Do not allocate timestamp queries on unsupported platforms

diff --git a/Engine/Source/Runtime/RenderCore/Private/ProfilingDebugging/RealtimeGPUProfiler.cpp b/Engine/Source/Runtime/RenderCore/Private/ProfilingDebugging/RealtimeGPUProfiler.cpp
--- a/Engine/Source/Runtime/RenderCore/Private/ProfilingDebugging/RealtimeGPUProfiler.cpp
+++ b/Engine/Source/Runtime/RenderCore/Private/ProfilingDebugging/RealtimeGPUProfiler.cpp
@@ -603,11 +603,14 @@ FRealtimeGPUProfiler::FRealtimeGPUProfiler()
 	, bStatGatheringPaused(false)
 	, bInBeginEndBlock(false)
 {
-	const int MaxGPUQueries = CVarGPUStatsMaxQueriesPerFrame.GetValueOnRenderThread();
-	RenderQueryPool = RHICreateRenderQueryPool(RQT_AbsoluteTime, (MaxGPUQueries > 0) ? MaxGPUQueries * 2 : UINT32_MAX);
-	for (int Index = 0; Index < NumGPUProfilerBufferedFrames; Index++)
+	if (GSupportsTimestampRenderQueries)
 	{
-		Frames.Add(new FRealtimeGPUProfilerFrame(RenderQueryPool, QueryCount));
+		const int MaxGPUQueries = CVarGPUStatsMaxQueriesPerFrame.GetValueOnRenderThread();
+		RenderQueryPool = RHICreateRenderQueryPool(RQT_AbsoluteTime, (MaxGPUQueries > 0) ? MaxGPUQueries * 2 : UINT32_MAX);
+		for (int Index = 0; Index < NumGPUProfilerBufferedFrames; Index++)
+		{
+			Frames.Add(new FRealtimeGPUProfilerFrame(RenderQueryPool, QueryCount));
+		}
 	}
 }
 
@@ -664,7 +667,7 @@ void FRealtimeGPUProfiler::EndFrame(FRHICommandListImmediate& RHICmdList)
 	// This is called at the end of the renderthread frame. Note that the RHI thread may still be processing commands for the frame at this point, however
 	// The read buffer index is always 3 frames beind the write buffer index in order to prevent us reading from the frame the RHI thread is still processing. 
 	// This should also ensure the GPU is done with the queries before we try to read them
-	check(Frames.Num() > 0);
+	check(!GSupportsTimestampRenderQueries || Frames.Num() > 0);
 	check(IsInRenderingThread());
 	check(bInBeginEndBlock == true);
 	bInBeginEndBlock = false;
