Vulkan: prevent crash on projects with VT enabled (UE-80091).

- Side effect of CL 8366113

diff --git a/Engine/Source/Runtime/RenderCore/Private/RenderUtils.cpp b/Engine/Source/Runtime/RenderCore/Private/RenderUtils.cpp
--- a/Engine/Source/Runtime/RenderCore/Private/RenderUtils.cpp
+++ b/Engine/Source/Runtime/RenderCore/Private/RenderUtils.cpp
@@ -1170,6 +1170,13 @@ RENDERCORE_API bool UseVirtualTexturing(ERHIFeatureLevel::Type InFeatureLevel, c
 		}
 #endif
 
+		// Remove this when UE-80097 is implemented
+		static bool bIsVulkanRHI = FCString::Strcmp(GDynamicRHI->GetName(), TEXT("Vulkan")) == 0;
+		if (bIsVulkanRHI)
+		{
+			return false;
+		}
+
 		// does the project has it enabled ?
 		static const auto CVarVirtualTexture = IConsoleManager::Get().FindTConsoleVariableDataInt(TEXT("r.VirtualTextures"));
 		check(CVarVirtualTexture);
