Copying

//UE4/Main/Engine/Source/Runtime/RenderCore/...
//UE4/Main/Engine/Source/Runtime/UtilityShaders/...

to //UE4/Dev-Build/Engine/Source/Runtime/...

#rnx

diff --git a/Engine/Source/Runtime/RenderCore/Private/BuiltInRayTracingShaders.cpp b/Engine/Source/Runtime/RenderCore/Private/BuiltInRayTracingShaders.cpp
new file mode 100644
--- /dev/null
+++ b/Engine/Source/Runtime/RenderCore/Private/BuiltInRayTracingShaders.cpp
@@ -0,0 +1,13 @@
+// Copyright 1998-2019 Epic Games, Inc. All Rights Reserved.
+
+#include "BuiltInRayTracingShaders.h"
+
+#if RHI_RAYTRACING
+
+IMPLEMENT_GLOBAL_SHADER( FOcclusionMainRG,		"/Engine/Private/RayTracing/RayTracingBuiltInShaders.usf",		"OcclusionMainRG",				SF_RayGen);
+IMPLEMENT_GLOBAL_SHADER( FIntersectionMainRG,	"/Engine/Private/RayTracing/RayTracingBuiltInShaders.usf",		"IntersectionMainRG",			SF_RayGen);
+IMPLEMENT_SHADER_TYPE(, FIntersectionMainCHS,	TEXT("/Engine/Private/RayTracing/RayTracingBuiltInShaders.usf"), TEXT("IntersectionMainCHS"),	SF_RayHitGroup);
+IMPLEMENT_SHADER_TYPE(, FDefaultMainCHS,		TEXT("/Engine/Private/RayTracing/RayTracingBuiltInShaders.usf"), TEXT("DefaultMainCHS"),		SF_RayHitGroup);
+IMPLEMENT_SHADER_TYPE(, FDefaultMainMS,			TEXT("/Engine/Private/RayTracing/RayTracingBuiltInShaders.usf"), TEXT("DefaultMainMS"),			SF_RayMiss);
+
+#endif // RHI_RAYTRACING
diff --git a/Engine/Source/Runtime/RenderCore/Private/ClearQuad.cpp b/Engine/Source/Runtime/RenderCore/Private/ClearQuad.cpp
new file mode 100644
--- /dev/null
+++ b/Engine/Source/Runtime/RenderCore/Private/ClearQuad.cpp
@@ -0,0 +1,377 @@
+// Copyright 1998-2019 Epic Games, Inc. All Rights Reserved.
+
+#include "ClearQuad.h"
+#include "Shader.h"
+#include "RHIStaticStates.h"
+#include "OneColorShader.h"
+#include "PipelineStateCache.h"
+#include "ClearReplacementShaders.h"
+#include "RendererInterface.h"
+#include "Logging/LogMacros.h"
+
+TGlobalResource<FClearVertexBuffer> GClearVertexBuffer;
+
+static TAutoConsoleVariable<int32> CVarFastClearUAVMaxSize(
+	TEXT("r.RHI.FastClearUAVMaxSize"),
+	0,
+	TEXT("Max size in bytes to fast clear tiny UAV. 0 by default or when feature is not supported by the platform"),
+	ECVF_RenderThreadSafe
+);
+
+DEFINE_LOG_CATEGORY_STATIC(LogClearQuad, Log, Log)
+
+static void ClearQuadSetup( FRHICommandList& RHICmdList, bool bClearColor, int32 NumClearColors, const FLinearColor* ClearColorArray, bool bClearDepth, float Depth, bool bClearStencil, uint32 Stencil, TFunction<void(FGraphicsPipelineStateInitializer&)> PSOModifier = nullptr)
+{
+	if (UNLIKELY(!FApp::CanEverRender()))
+	{
+		return;
+	}
+
+	// Set new states
+	FRHIBlendState* BlendStateRHI = bClearColor
+		? TStaticBlendState<>::GetRHI()
+		: TStaticBlendStateWriteMask<CW_NONE,CW_NONE,CW_NONE,CW_NONE,CW_NONE,CW_NONE,CW_NONE,CW_NONE>::GetRHI();
+	
+	FRHIDepthStencilState* DepthStencilStateRHI =
+		(bClearDepth && bClearStencil)
+			? TStaticDepthStencilState<
+				true, CF_Always,
+				true,CF_Always,SO_Replace,SO_Replace,SO_Replace,
+				false,CF_Always,SO_Replace,SO_Replace,SO_Replace,
+				0xff,0xff
+				>::GetRHI()
+			: bClearDepth
+				? TStaticDepthStencilState<true, CF_Always>::GetRHI()
+				: bClearStencil
+					? TStaticDepthStencilState<
+						false, CF_Always,
+						true,CF_Always,SO_Replace,SO_Replace,SO_Replace,
+						false,CF_Always,SO_Replace,SO_Replace,SO_Replace,
+						0xff,0xff
+						>::GetRHI()
+					: TStaticDepthStencilState<false, CF_Always>::GetRHI();
+
+	FGraphicsPipelineStateInitializer GraphicsPSOInit;
+	RHICmdList.ApplyCachedRenderTargets(GraphicsPSOInit);
+
+	GraphicsPSOInit.RasterizerState = TStaticRasterizerState<FM_Solid, CM_None>::GetRHI();
+	GraphicsPSOInit.BlendState = BlendStateRHI;
+	GraphicsPSOInit.DepthStencilState = DepthStencilStateRHI;
+
+	auto* ShaderMap = GetGlobalShaderMap(GMaxRHIFeatureLevel);
+
+	// Set the new shaders
+	TShaderMapRef<TOneColorVS<true> > VertexShader(ShaderMap);
+
+	FOneColorPS* PixelShader = NULL;
+
+	// Set the shader to write to the appropriate number of render targets
+	// On AMD PC hardware, outputting to a color index in the shader without a matching render target set has a significant performance hit
+	if (NumClearColors <= 1)
+	{
+		TShaderMapRef<TOneColorPixelShaderMRT<1> > MRTPixelShader(ShaderMap);
+		PixelShader = *MRTPixelShader;
+	}
+	else if (NumClearColors == 2)
+	{
+		TShaderMapRef<TOneColorPixelShaderMRT<2> > MRTPixelShader(ShaderMap);
+		PixelShader = *MRTPixelShader;
+	}
+	else if (NumClearColors == 3)
+	{
+		TShaderMapRef<TOneColorPixelShaderMRT<3> > MRTPixelShader(ShaderMap);
+		PixelShader = *MRTPixelShader;
+	}
+	else if (NumClearColors == 4)
+	{
+		TShaderMapRef<TOneColorPixelShaderMRT<4> > MRTPixelShader(ShaderMap);
+		PixelShader = *MRTPixelShader;
+	}
+	else if (NumClearColors == 5)
+	{
+		TShaderMapRef<TOneColorPixelShaderMRT<5> > MRTPixelShader(ShaderMap);
+		PixelShader = *MRTPixelShader;
+	}
+	else if (NumClearColors == 6)
+	{
+		TShaderMapRef<TOneColorPixelShaderMRT<6> > MRTPixelShader(ShaderMap);
+		PixelShader = *MRTPixelShader;
+	}
+	else if (NumClearColors == 7)
+	{
+		TShaderMapRef<TOneColorPixelShaderMRT<7> > MRTPixelShader(ShaderMap);
+		PixelShader = *MRTPixelShader;
+	}
+	else if (NumClearColors == 8)
+	{
+		TShaderMapRef<TOneColorPixelShaderMRT<8> > MRTPixelShader(ShaderMap);
+		PixelShader = *MRTPixelShader;
+	}
+
+	GraphicsPSOInit.BoundShaderState.VertexDeclarationRHI = GetVertexDeclarationFVector4();
+	GraphicsPSOInit.BoundShaderState.VertexShaderRHI = GETSAFERHISHADER_VERTEX(*VertexShader);
+	GraphicsPSOInit.BoundShaderState.PixelShaderRHI = GETSAFERHISHADER_PIXEL(PixelShader);
+	GraphicsPSOInit.PrimitiveType = PT_TriangleStrip;
+
+	if (PSOModifier)
+	{
+		PSOModifier(GraphicsPSOInit);
+	}
+
+	SetGraphicsPipelineState(RHICmdList, GraphicsPSOInit);
+	RHICmdList.SetStencilRef(Stencil);
+
+	VertexShader->SetDepthParameter(RHICmdList, Depth);
+	PixelShader->SetColors(RHICmdList, ClearColorArray, NumClearColors);
+}
+
+static void ClearUAVShader(FRHICommandList& RHICmdList, FRHIUnorderedAccessView* UnorderedAccessViewRHI, uint32 SizeInBytes, uint32 ClearValue, bool bBarriers = true)
+{
+	UE_CLOG((SizeInBytes & 0x3) != 0, LogClearQuad, Warning,
+		TEXT("Buffer size is not a multiple of DWORDs. Up to 3 bytes after buffer end will also be cleared"));
+
+	TShaderMapRef<FClearBufferReplacementCS> ComputeShader(GetGlobalShaderMap(GMaxRHIFeatureLevel));
+	FRHIComputeShader* ShaderRHI = ComputeShader->GetComputeShader();
+	
+	uint32 NumDWordsToClear = (SizeInBytes + 3) / 4;
+	uint32 NumThreadGroupsX = (NumDWordsToClear + 63) / 64;
+	RHICmdList.SetComputeShader(ShaderRHI);
+	ComputeShader->SetParameters(RHICmdList, UnorderedAccessViewRHI, NumDWordsToClear, ClearValue, bBarriers);
+	RHICmdList.DispatchComputeShader(NumThreadGroupsX, 1, 1);
+	ComputeShader->FinalizeParameters(RHICmdList, UnorderedAccessViewRHI, bBarriers);
+}
+
+void ClearUAV(FRHICommandList& RHICmdList, const FRWBufferStructured& StructuredBuffer, uint32 Value)
+{
+	
+	if (StructuredBuffer.NumBytes <= uint32(CVarFastClearUAVMaxSize.GetValueOnRenderThread()))
+	{
+		uint32 Values[4] = { Value, Value, Value, Value };
+		RHICmdList.ClearTinyUAV(StructuredBuffer.UAV, Values);
+	}
+	else
+	{
+		ClearUAVShader(RHICmdList, StructuredBuffer.UAV, StructuredBuffer.NumBytes, Value);
+	}
+}
+
+
+void ClearUAV(FRHICommandList& RHICmdList, const FTextureRWBuffer& Buffer, FLinearColor Value)
+{
+	ClearTexture2DUAV(RHICmdList, Buffer.UAV, Buffer.Buffer->GetSizeX(), Buffer.Buffer->GetSizeY(), Value);	
+}
+
+void ClearUAV(FRHICommandList& RHICmdList, const FRWBuffer& Buffer, uint32 Value, bool bBarriers)
+{
+	if (Buffer.NumBytes <= uint32(CVarFastClearUAVMaxSize.GetValueOnRenderThread()))
+	{
+		uint32 Values[4] = { Value, Value, Value, Value };
+		RHICmdList.ClearTinyUAV(Buffer.UAV, Values);
+		check(bBarriers); //  TODO ClearTinyUAV is doing transitions as of today
+	}
+	else
+	{
+		ClearUAVShader(RHICmdList, Buffer.UAV, Buffer.NumBytes, Value, bBarriers);
+	}
+}
+
+void ClearUAV(FRHICommandList& RHICmdList, FRHIUnorderedAccessView* Buffer, uint32 NumBytes, uint32 Value)
+{
+	if (NumBytes <= uint32(CVarFastClearUAVMaxSize.GetValueOnRenderThread()))
+	{
+		uint32 Values[4] = { Value, Value, Value, Value };
+		RHICmdList.ClearTinyUAV(Buffer, Values);
+	}
+	else
+	{
+		ClearUAVShader(RHICmdList, Buffer, NumBytes, Value);
+	}
+}
+
+template< typename T >
+inline void ClearUAV_T(FRHICommandList& RHICmdList, FRHITexture* Texture, FRHIUnorderedAccessView* TextureUAV, const T(&ClearValues)[4])
+{
+	check( Texture );
+	check( TextureUAV );
+
+	if (auto Texture2d = Texture->GetTexture2D())
+	{
+		TShaderMapRef< FClearTexture2DReplacementCS<T> > ComputeShader(GetGlobalShaderMap(GMaxRHIFeatureLevel));
+		FRHIComputeShader* ShaderRHI = ComputeShader->GetComputeShader();
+		RHICmdList.SetComputeShader(ShaderRHI);
+		ComputeShader->SetParameters(RHICmdList, TextureUAV, ClearValues);
+		uint32 x = (Texture2d->GetSizeX() + 7) / 8;
+		uint32 y = (Texture2d->GetSizeY() + 7) / 8;
+		RHICmdList.DispatchComputeShader(x, y, 1);
+		ComputeShader->FinalizeParameters(RHICmdList, TextureUAV);
+	}
+	else if (auto Texture2dArray = Texture->GetTexture2DArray())
+	{
+		TShaderMapRef< FClearTexture2DArrayReplacementCS<T> > ComputeShader(GetGlobalShaderMap(GMaxRHIFeatureLevel));
+		FRHIComputeShader* ShaderRHI = ComputeShader->GetComputeShader();
+		RHICmdList.SetComputeShader(ShaderRHI);
+		ComputeShader->SetParameters(RHICmdList, TextureUAV, ClearValues);
+		uint32 x = (Texture2dArray->GetSizeX() + 7) / 8;
+		uint32 y = (Texture2dArray->GetSizeY() + 7) / 8;
+		uint32 z = Texture2dArray->GetSizeZ();
+		RHICmdList.DispatchComputeShader(x, y, z);
+		ComputeShader->FinalizeParameters(RHICmdList, TextureUAV);
+	}
+	else if (auto TextureCube = Texture->GetTextureCube())
+	{
+		TShaderMapRef< FClearTexture2DArrayReplacementCS<T> > ComputeShader(GetGlobalShaderMap(GMaxRHIFeatureLevel));
+		FRHIComputeShader* ShaderRHI = ComputeShader->GetComputeShader();
+		RHICmdList.SetComputeShader(ShaderRHI);
+		ComputeShader->SetParameters(RHICmdList, TextureUAV, ClearValues);
+		uint32 x = (TextureCube->GetSize() + 7) / 8;
+		uint32 y = (TextureCube->GetSize() + 7) / 8;
+		RHICmdList.DispatchComputeShader(x, y, 6);
+		ComputeShader->FinalizeParameters(RHICmdList, TextureUAV);
+	}
+	else if (auto Texture3d = Texture->GetTexture3D())
+	{
+		TShaderMapRef< FClearVolumeReplacementCS<T> > ComputeShader(GetGlobalShaderMap(GMaxRHIFeatureLevel));
+		FRHIComputeShader* ShaderRHI = ComputeShader->GetComputeShader();
+		RHICmdList.SetComputeShader(ShaderRHI);
+		ComputeShader->SetParameters(RHICmdList, TextureUAV, ClearValues);
+		uint32 x = (Texture3d->GetSizeX() + 3) / 4;
+		uint32 y = (Texture3d->GetSizeY() + 3) / 4;
+		uint32 z = (Texture3d->GetSizeZ() + 3) / 4;
+		RHICmdList.DispatchComputeShader(x, y, z);
+		ComputeShader->FinalizeParameters(RHICmdList, TextureUAV);
+	}
+	else
+	{
+		check(0);
+	}
+}
+
+void ClearTexture2DUAV(FRHICommandList& RHICmdList, FRHIUnorderedAccessView* UAV, int32 Width, int32 Height, const FLinearColor& ClearColor)
+{
+	TShaderMapRef< FClearTexture2DReplacementCS<float> > ComputeShader(GetGlobalShaderMap(GMaxRHIFeatureLevel));
+	FRHIComputeShader* ShaderRHI = ComputeShader->GetComputeShader();
+	RHICmdList.SetComputeShader(ShaderRHI);
+	ComputeShader->SetParameters(RHICmdList, UAV, reinterpret_cast<const float(&)[4]>(ClearColor));
+	uint32 x = (Width + 7) / 8;
+	uint32 y = (Height + 7) / 8;
+	RHICmdList.DispatchComputeShader(x, y, 1);
+	ComputeShader->FinalizeParameters(RHICmdList, UAV);
+}
+
+void ClearUAV(FRHICommandList& RHICmdList, const FSceneRenderTargetItem& RenderTargetItem, const float(&ClearValues)[4])
+{
+	ClearUAV_T(RHICmdList, RenderTargetItem.TargetableTexture, RenderTargetItem.UAV, ClearValues);
+}
+
+void ClearUAV(FRHICommandList& RHICmdList, const FSceneRenderTargetItem& RenderTargetItem, const uint32(&ClearValues)[4])
+{
+	ClearUAV_T(RHICmdList, RenderTargetItem.TargetableTexture, RenderTargetItem.UAV, ClearValues);
+}
+
+void ClearUAV(FRHICommandList& RHICmdList, const FSceneRenderTargetItem& RenderTargetItem, const FLinearColor& ClearColor)
+{
+	ClearUAV_T(RHICmdList, RenderTargetItem.TargetableTexture, RenderTargetItem.UAV, reinterpret_cast<const float(&)[4]>(ClearColor));
+}
+
+void ClearUAV(FRHICommandList& RHICmdList, FRHITexture* Texture, FRHIUnorderedAccessView* TextureUAV, const float(&ClearValues)[4])
+{
+	ClearUAV_T(RHICmdList, Texture, TextureUAV, ClearValues);
+}
+
+void ClearUAV(FRHICommandList& RHICmdList, FRHITexture* Texture, FRHIUnorderedAccessView* TextureUAV, const uint32(&ClearValues)[4])
+{
+	ClearUAV_T(RHICmdList, Texture, TextureUAV, ClearValues);
+}
+
+void ClearUAV(FRHICommandList& RHICmdList, FRHITexture* Texture, FRHIUnorderedAccessView* TextureUAV, const FLinearColor& ClearColor)
+{
+	ClearUAV_T(RHICmdList, Texture, TextureUAV, reinterpret_cast<const float(&)[4]>(ClearColor));
+}
+
+void DrawClearQuadMRT(FRHICommandList& RHICmdList, bool bClearColor, int32 NumClearColors, const FLinearColor* ClearColorArray, bool bClearDepth, float Depth, bool bClearStencil, uint32 Stencil)
+{
+	ClearQuadSetup(RHICmdList, bClearColor, NumClearColors, ClearColorArray, bClearDepth, Depth, bClearStencil, Stencil);
+
+	RHICmdList.SetStreamSource(0, GClearVertexBuffer.VertexBufferRHI, 0);
+	RHICmdList.DrawPrimitive(0, 2, 1);
+}
+
+void DrawClearQuadMRT(FRHICommandList& RHICmdList, bool bClearColor, int32 NumClearColors, const FLinearColor* ClearColorArray, bool bClearDepth, float Depth, bool bClearStencil, uint32 Stencil, FClearQuadCallbacks ClearQuadCallbacks)
+{
+	ClearQuadSetup(RHICmdList, bClearColor, NumClearColors, ClearColorArray, bClearDepth, Depth, bClearStencil, Stencil, ClearQuadCallbacks.PSOModifier);
+
+	if (ClearQuadCallbacks.PreClear)
+	{
+		ClearQuadCallbacks.PreClear(RHICmdList);
+	}
+
+	// Draw a fullscreen quad without a hole
+	RHICmdList.SetStreamSource(0, GClearVertexBuffer.VertexBufferRHI, 0);
+	RHICmdList.DrawPrimitive(0, 2, 1);
+
+	if (ClearQuadCallbacks.PostClear)
+	{
+		ClearQuadCallbacks.PostClear(RHICmdList);
+	}
+}
+
+void DrawClearQuadMRT(FRHICommandList& RHICmdList, bool bClearColor, int32 NumClearColors, const FLinearColor* ClearColorArray, bool bClearDepth, float Depth, bool bClearStencil, uint32 Stencil, FIntPoint ViewSize, FIntRect ExcludeRect)
+{
+	if (ExcludeRect.Min == FIntPoint::ZeroValue && ExcludeRect.Max == ViewSize)
+	{
+		// Early out if the entire surface is excluded
+		return;
+	}
+
+	ClearQuadSetup(RHICmdList, bClearColor, NumClearColors, ClearColorArray, bClearDepth, Depth, bClearStencil, Stencil);
+
+	// Draw a fullscreen quad
+	if (ExcludeRect.Width() > 0 && ExcludeRect.Height() > 0)
+	{
+		// with a hole in it
+		FVector4 OuterVertices[4];
+		OuterVertices[0].Set(-1.0f, 1.0f, Depth, 1.0f);
+		OuterVertices[1].Set(1.0f, 1.0f, Depth, 1.0f);
+		OuterVertices[2].Set(1.0f, -1.0f, Depth, 1.0f);
+		OuterVertices[3].Set(-1.0f, -1.0f, Depth, 1.0f);
+
+		float InvViewWidth = 1.0f / ViewSize.X;
+		float InvViewHeight = 1.0f / ViewSize.Y;
+		FVector4 FractionRect = FVector4(ExcludeRect.Min.X * InvViewWidth, ExcludeRect.Min.Y * InvViewHeight, (ExcludeRect.Max.X - 1) * InvViewWidth, (ExcludeRect.Max.Y - 1) * InvViewHeight);
+
+		FVector4 InnerVertices[4];
+		InnerVertices[0].Set(FMath::Lerp(-1.0f, 1.0f, FractionRect.X), FMath::Lerp(1.0f, -1.0f, FractionRect.Y), Depth, 1.0f);
+		InnerVertices[1].Set(FMath::Lerp(-1.0f, 1.0f, FractionRect.Z), FMath::Lerp(1.0f, -1.0f, FractionRect.Y), Depth, 1.0f);
+		InnerVertices[2].Set(FMath::Lerp(-1.0f, 1.0f, FractionRect.Z), FMath::Lerp(1.0f, -1.0f, FractionRect.W), Depth, 1.0f);
+		InnerVertices[3].Set(FMath::Lerp(-1.0f, 1.0f, FractionRect.X), FMath::Lerp(1.0f, -1.0f, FractionRect.W), Depth, 1.0f);
+
+		FRHIResourceCreateInfo CreateInfo;
+		FVertexBufferRHIRef VertexBufferRHI = RHICreateVertexBuffer(sizeof(FVector4) * 10, BUF_Volatile, CreateInfo);
+		void* VoidPtr = RHILockVertexBuffer(VertexBufferRHI, 0, sizeof(FVector4) * 10, RLM_WriteOnly);
+		
+		FVector4* Vertices = reinterpret_cast<FVector4*>(VoidPtr);
+		Vertices[0] = OuterVertices[0];
+		Vertices[1] = InnerVertices[0];
+		Vertices[2] = OuterVertices[1];
+		Vertices[3] = InnerVertices[1];
+		Vertices[4] = OuterVertices[2];
+		Vertices[5] = InnerVertices[2];
+		Vertices[6] = OuterVertices[3];
+		Vertices[7] = InnerVertices[3];
+		Vertices[8] = OuterVertices[0];
+		Vertices[9] = InnerVertices[0];
+
+		RHIUnlockVertexBuffer(VertexBufferRHI);
+		RHICmdList.SetStreamSource(0, VertexBufferRHI, 0);
+
+		RHICmdList.DrawPrimitive(0, 8, 1);
+
+		VertexBufferRHI.SafeRelease();
+	}
+	else
+	{
+		// without a hole
+		RHICmdList.SetStreamSource(0, GClearVertexBuffer.VertexBufferRHI, 0);
+		RHICmdList.DrawPrimitive(0, 2, 1);
+	}
+}
\ No newline at end of file
diff --git a/Engine/Source/Runtime/RenderCore/Private/ClearReplacementShaders.cpp b/Engine/Source/Runtime/RenderCore/Private/ClearReplacementShaders.cpp
new file mode 100644
--- /dev/null
+++ b/Engine/Source/Runtime/RenderCore/Private/ClearReplacementShaders.cpp
@@ -0,0 +1,117 @@
+// Copyright 1998-2019 Epic Games, Inc. All Rights Reserved.
+
+
+#include "ClearReplacementShaders.h"
+#include "ShaderParameterUtils.h"
+#include "RendererInterface.h"
+
+template< typename T >
+void FClearTexture2DReplacementCS<T>::SetParameters( FRHICommandList& RHICmdList, FRHIUnorderedAccessView* TextureRW, const T(&Values)[4] )
+{
+	FRHIComputeShader* ComputeShaderRHI = GetComputeShader();
+	SetShaderValue(RHICmdList, ComputeShaderRHI, ClearColor, Values);
+
+	RHICmdList.TransitionResource(EResourceTransitionAccess::ERWBarrier, EResourceTransitionPipeline::EGfxToCompute, TextureRW);
+	RHICmdList.SetUAVParameter(ComputeShaderRHI, ClearTextureRW.GetBaseIndex(), TextureRW);
+}
+
+template< typename T >
+void FClearTexture2DReplacementCS<T>::FinalizeParameters(FRHICommandList& RHICmdList, FRHIUnorderedAccessView* TextureRW)
+{
+	RHICmdList.TransitionResource(EResourceTransitionAccess::ERWBarrier, EResourceTransitionPipeline::EComputeToCompute, TextureRW);
+	FRHIComputeShader* ComputeShaderRHI = GetComputeShader();
+	SetUAVParameter(RHICmdList, ComputeShaderRHI, ClearTextureRW, FUnorderedAccessViewRHIRef());
+}
+
+template< typename T >
+void FClearTexture2DArrayReplacementCS<T>::SetParameters(FRHICommandList& RHICmdList, FRHIUnorderedAccessView* TextureRW, const T(&Values)[4])
+{
+	FRHIComputeShader* ComputeShaderRHI = GetComputeShader();
+	SetShaderValue(RHICmdList, ComputeShaderRHI, ClearColor, Values);
+
+	RHICmdList.TransitionResource(EResourceTransitionAccess::ERWBarrier, EResourceTransitionPipeline::EGfxToCompute, TextureRW);
+	RHICmdList.SetUAVParameter(ComputeShaderRHI, ClearTextureArrayRW.GetBaseIndex(), TextureRW);
+}
+
+template< typename T >
+void FClearTexture2DArrayReplacementCS<T>::FinalizeParameters(FRHICommandList& RHICmdList, FRHIUnorderedAccessView* TextureRW)
+{
+	RHICmdList.TransitionResource(EResourceTransitionAccess::ERWBarrier, EResourceTransitionPipeline::EComputeToCompute, TextureRW);
+	FRHIComputeShader* ComputeShaderRHI = GetComputeShader();
+	SetUAVParameter(RHICmdList, ComputeShaderRHI, ClearTextureArrayRW, FUnorderedAccessViewRHIRef());
+}
+
+template< typename T >
+void FClearVolumeReplacementCS<T>::SetParameters(FRHICommandList& RHICmdList, FRHIUnorderedAccessView* TextureRW, const T(&Values)[4])
+{
+	FRHIComputeShader* ComputeShaderRHI = GetComputeShader();
+	SetShaderValue(RHICmdList, ComputeShaderRHI, ClearColor, Values);
+
+	RHICmdList.TransitionResource(EResourceTransitionAccess::ERWBarrier, EResourceTransitionPipeline::EGfxToCompute, TextureRW);
+	RHICmdList.SetUAVParameter(ComputeShaderRHI, ClearVolumeRW.GetBaseIndex(), TextureRW);
+}
+
+template< typename T >
+void FClearVolumeReplacementCS<T>::FinalizeParameters(FRHICommandList& RHICmdList, FRHIUnorderedAccessView* TextureRW)
+{
+	RHICmdList.TransitionResource(EResourceTransitionAccess::ERWBarrier, EResourceTransitionPipeline::EComputeToCompute, TextureRW);
+	FRHIComputeShader* ComputeShaderRHI = GetComputeShader();
+	SetUAVParameter(RHICmdList, ComputeShaderRHI, ClearVolumeRW, FUnorderedAccessViewRHIRef());
+}
+
+void FClearTexture2DReplacementScissorCS::SetParameters(FRHICommandList& RHICmdList, FRHIUnorderedAccessView* TextureRW, FLinearColor InClearColor, const FVector4& InTargetBounds)
+{
+	FRHIComputeShader* ComputeShaderRHI = GetComputeShader();
+	SetShaderValue(RHICmdList, ComputeShaderRHI, ClearColor, InClearColor);
+	SetShaderValue(RHICmdList, ComputeShaderRHI, TargetBounds, InTargetBounds);
+	RHICmdList.TransitionResource(EResourceTransitionAccess::ERWBarrier, EResourceTransitionPipeline::EGfxToCompute, TextureRW);
+	RHICmdList.SetUAVParameter(ComputeShaderRHI, ClearTextureRW.GetBaseIndex(), TextureRW);
+}
+
+void FClearTexture2DReplacementScissorCS::FinalizeParameters(FRHICommandList& RHICmdList, FRHIUnorderedAccessView* TextureRW)
+{
+	RHICmdList.TransitionResource(EResourceTransitionAccess::ERWBarrier, EResourceTransitionPipeline::EComputeToCompute, TextureRW);
+	FRHIComputeShader* ComputeShaderRHI = GetComputeShader();
+	SetUAVParameter(RHICmdList, ComputeShaderRHI, ClearTextureRW, FUnorderedAccessViewRHIRef());
+}
+
+void FClearBufferReplacementCS::SetParameters(FRHICommandList& RHICmdList, FRHIUnorderedAccessView* BufferRW, uint32 NumDWordsToClear,uint32 ClearValue, bool bBarriers)
+{
+	FRHIComputeShader* ComputeShaderRHI = GetComputeShader();
+	SetShaderValue(RHICmdList, ComputeShaderRHI, ClearBufferCSParams, FUintVector4(ClearValue, NumDWordsToClear, 0, 0));
+	if (bBarriers)
+	{
+		RHICmdList.TransitionResource(EResourceTransitionAccess::ERWBarrier, EResourceTransitionPipeline::EGfxToCompute, BufferRW);
+	}
+	RHICmdList.SetUAVParameter(ComputeShaderRHI, ClearBufferRW.GetBaseIndex(), BufferRW);
+}
+
+void FClearBufferReplacementCS::FinalizeParameters(FRHICommandList& RHICmdList, FRHIUnorderedAccessView* BufferRW, bool bBarriers)
+{
+	if (bBarriers)
+	{
+		RHICmdList.TransitionResource(EResourceTransitionAccess::ERWBarrier, EResourceTransitionPipeline::EComputeToCompute, BufferRW);
+	}
+	FRHIComputeShader* ComputeShaderRHI = GetComputeShader();
+	SetUAVParameter(RHICmdList, ComputeShaderRHI, ClearBufferRW, FUnorderedAccessViewRHIRef());
+}
+
+template class FClearTexture2DReplacementCS<float>;
+template class FClearTexture2DReplacementCS<uint32>;
+template class FClearTexture2DArrayReplacementCS<float>;
+template class FClearTexture2DArrayReplacementCS<uint32>;
+template class FClearVolumeReplacementCS<float>;
+template class FClearVolumeReplacementCS<uint32>;
+
+IMPLEMENT_SHADER_TYPE(, FClearReplacementVS, TEXT("/Engine/Private/ClearReplacementShaders.usf"), TEXT("ClearVS"), SF_Vertex);
+IMPLEMENT_SHADER_TYPE(, FClearReplacementPS, TEXT("/Engine/Private/ClearReplacementShaders.usf"), TEXT("ClearPS"), SF_Pixel);
+
+IMPLEMENT_SHADER_TYPE(template<>, FClearTexture2DReplacementCS<float>, TEXT("/Engine/Private/ClearReplacementShaders.usf"), TEXT("ClearTexture2DCS"), SF_Compute);
+IMPLEMENT_SHADER_TYPE(template<>, FClearTexture2DReplacementCS<uint32>, TEXT("/Engine/Private/ClearReplacementShaders.usf"), TEXT("ClearTexture2DCS"), SF_Compute);
+IMPLEMENT_SHADER_TYPE(template<>, FClearTexture2DArrayReplacementCS<float>, TEXT("/Engine/Private/ClearReplacementShaders.usf"), TEXT("ClearTexture2DArrayCS"), SF_Compute);
+IMPLEMENT_SHADER_TYPE(template<>, FClearTexture2DArrayReplacementCS<uint32>, TEXT("/Engine/Private/ClearReplacementShaders.usf"), TEXT("ClearTexture2DArrayCS"), SF_Compute);
+IMPLEMENT_SHADER_TYPE(template<>, FClearVolumeReplacementCS<float>, TEXT("/Engine/Private/ClearReplacementShaders.usf"), TEXT("ClearVolumeCS"), SF_Compute);
+IMPLEMENT_SHADER_TYPE(template<>, FClearVolumeReplacementCS<uint32>, TEXT("/Engine/Private/ClearReplacementShaders.usf"), TEXT("ClearVolumeCS"), SF_Compute);
+IMPLEMENT_SHADER_TYPE(, FClearTexture2DReplacementScissorCS, TEXT("/Engine/Private/ClearReplacementShaders.usf"), TEXT("ClearTexture2DScissorCS"), SF_Compute);
+
+IMPLEMENT_SHADER_TYPE(, FClearBufferReplacementCS, TEXT("/Engine/Private/ClearReplacementShaders.usf"), TEXT("ClearBufferCS"), SF_Compute);
\ No newline at end of file
diff --git a/Engine/Source/Runtime/RenderCore/Private/LongGPUTask.cpp b/Engine/Source/Runtime/RenderCore/Private/LongGPUTask.cpp
new file mode 100644
--- /dev/null
+++ b/Engine/Source/Runtime/RenderCore/Private/LongGPUTask.cpp
@@ -0,0 +1,100 @@
+// Copyright 1998-2019 Epic Games, Inc. All Rights Reserved.
+
+#include "LongGPUTask.h"
+#include "OneColorShader.h"
+#include "PipelineStateCache.h"
+#include "RHIStaticStates.h"
+#include "RenderUtils.h"
+#include "ClearQuad.h"
+
+IMPLEMENT_SHADER_TYPE(, FLongGPUTaskPS, TEXT("/Engine/Private/OneColorShader.usf"), TEXT("MainLongGPUTask"), SF_Pixel);
+
+int32 NumMeasuredIterationsToAchieve500ms = 0;
+
+void IssueScalableLongGPUTask(FRHICommandListImmediate& RHICmdList, int32 NumIteration /* = -1 by default */)
+{
+	FRHIResourceCreateInfo Info;
+	FTexture2DRHIRef LongTaskRenderTarget = RHICreateTexture2D(1920, 1080, PF_B8G8R8A8, 1, 1, TexCreate_RenderTargetable, Info);
+
+	FRHIRenderPassInfo RPInfo(LongTaskRenderTarget, ERenderTargetActions::DontLoad_Store);
+	TransitionRenderPassTargets(RHICmdList, RPInfo);
+	RHICmdList.BeginRenderPass(RPInfo, TEXT("LongGPUTask"));
+
+	{
+		FGraphicsPipelineStateInitializer GraphicsPSOInit;
+		RHICmdList.ApplyCachedRenderTargets(GraphicsPSOInit);
+
+		GraphicsPSOInit.BlendState = TStaticBlendState<>::GetRHI();
+		GraphicsPSOInit.RasterizerState = TStaticRasterizerState<>::GetRHI();
+		GraphicsPSOInit.DepthStencilState = TStaticDepthStencilState<false, CF_Always>::GetRHI();
+
+		auto ShaderMap = GetGlobalShaderMap(GMaxRHIFeatureLevel);
+		TShaderMapRef<TOneColorVS<true>> VertexShader(ShaderMap);
+		TShaderMapRef<FLongGPUTaskPS> PixelShader(ShaderMap);
+
+		GraphicsPSOInit.BoundShaderState.VertexDeclarationRHI = GetVertexDeclarationFVector4();
+		GraphicsPSOInit.BoundShaderState.VertexShaderRHI = VertexShader->GetVertexShader();
+		GraphicsPSOInit.BoundShaderState.PixelShaderRHI = PixelShader->GetPixelShader();
+		GraphicsPSOInit.PrimitiveType = PT_TriangleStrip;
+
+		SetGraphicsPipelineState(RHICmdList, GraphicsPSOInit);
+		VertexShader->SetDepthParameter(RHICmdList, 0.0f);
+
+		RHICmdList.SetStreamSource(0, GClearVertexBuffer.VertexBufferRHI, 0);
+
+		if (NumIteration == -1)
+		{
+			NumIteration = NumMeasuredIterationsToAchieve500ms;
+		}
+
+		for (int32 Iteration = 0; Iteration < NumIteration; Iteration++)
+		{
+			RHICmdList.DrawPrimitive(0, 2, 1);
+		}
+	}
+
+	RHICmdList.EndRenderPass();
+}
+
+void MeasureLongGPUTaskExecutionTime(FRHICommandListImmediate& RHICmdList)
+{
+	const int32 NumIterationsForMeasurement = 5;
+
+	FRenderQueryPoolRHIRef TimerQueryPool = RHICreateRenderQueryPool(RQT_AbsoluteTime, 2);
+
+	FRHIPooledRenderQuery TimeQueryStart = TimerQueryPool->AllocateQuery();
+	FRHIPooledRenderQuery TimeQueryEnd = TimerQueryPool->AllocateQuery();
+
+	if (!GSupportsTimestampRenderQueries)
+	{
+		// Not all platforms/drivers support RQT_AbsoluteTime queries
+		// Use fixed number of iterations on those platforms
+		NumMeasuredIterationsToAchieve500ms = 5;
+		return;
+	}
+
+	uint64 StartTime = 0;
+	uint64 EndTime = 0;
+
+	RHICmdList.EndRenderQuery(TimeQueryStart.GetQuery());
+
+	IssueScalableLongGPUTask(RHICmdList, NumIterationsForMeasurement);
+
+	RHICmdList.EndRenderQuery(TimeQueryEnd.GetQuery());
+
+	// Required by DX12 to resolve the query
+	RHICmdList.SubmitCommandsHint();
+	RHICmdList.ImmediateFlush(EImmediateFlushType::FlushRHIThread);
+
+	if (RHICmdList.GetRenderQueryResult(TimeQueryStart.GetQuery(), StartTime, true) && RHICmdList.GetRenderQueryResult(TimeQueryEnd.GetQuery(), EndTime, true))
+	{
+		NumMeasuredIterationsToAchieve500ms = FMath::Clamp(FMath::FloorToInt(500.0f / ((EndTime - StartTime) / 1000.0f / NumIterationsForMeasurement)), 1, 2000);
+	}
+	else
+	{
+		// Sometimes it fails even the platform supports RQT_AbsoluteTime
+		// Fallback and show a warning
+		NumMeasuredIterationsToAchieve500ms = 5;
+		UE_LOG(LogTemp, Display, TEXT("Unable to get render query result on a platform supporting RQT_AbsoluteTime queries, defaulting to %d iterations for LongGPUTask"), NumMeasuredIterationsToAchieve500ms);
+	}
+}
diff --git a/Engine/Source/Runtime/RenderCore/Private/MediaShaders.cpp b/Engine/Source/Runtime/RenderCore/Private/MediaShaders.cpp
new file mode 100644
--- /dev/null
+++ b/Engine/Source/Runtime/RenderCore/Private/MediaShaders.cpp
@@ -0,0 +1,564 @@
+// Copyright 1998-2019 Epic Games, Inc. All Rights Reserved.
+
+#include "MediaShaders.h"
+#include "ShaderParameterUtils.h"
+#include "RHIStaticStates.h"
+
+
+TGlobalResource<FMediaVertexDeclaration> GMediaVertexDeclaration;
+
+
+/* MediaShaders namespace
+ *****************************************************************************/
+
+namespace MediaShaders
+{
+	const FMatrix YuvToSrgbJpeg = FMatrix(
+		FPlane(1.000000f,  0.000000f,  1.402000f, 0.000000f),
+		FPlane(1.000000f, -0.344140f, -0.714140f, 0.000000f),
+		FPlane(1.000000f,  1.772000f,  0.000000f, 0.000000f),
+		FPlane(0.000000f,  0.000000f,  0.000000f, 0.000000f)
+	);
+
+	const FMatrix YuvToSrgbDefault = FMatrix(
+		FPlane(1.164383f,  0.000000f,  1.596027f, 0.000000f),
+		FPlane(1.164383f, -0.391762f, -0.812968f, 0.000000f),
+		FPlane(1.164383f,  2.017232f,  0.000000f, 0.000000f),
+		FPlane(0.000000f,  0.000000f,  0.000000f, 0.000000f)
+	);
+
+	const FMatrix YuvToSrgbPs4 = FMatrix(
+		FPlane(1.164400f,  0.000000f,  1.792700f, 0.000000f),
+		FPlane(1.164400f, -0.213300f, -0.532900f, 0.000000f),
+		FPlane(1.164400f,  2.112400f,  0.000000f, 0.000000f),
+		FPlane(0.000000f,  0.000000f,  0.000000f, 0.000000f)
+	);
+
+	const FMatrix YuvToSrgbRec601 = FMatrix(
+		FPlane(1.000000f, 0.000000f, 1.139830f, 0.000000f),
+		FPlane(1.000000f, -0.394650f, -0.580600f, 0.000000f),
+		FPlane(1.000000f, 2.032110, 0.000000, 0.000000f),
+		FPlane(0.000000f, 0.000000f, 0.000000f, 0.000000f)
+	);
+
+	const FMatrix YuvToRgbRec709 = FMatrix(
+		FPlane(1.000000f, 0.000000f, 1.280330f, 0.000000f),
+		FPlane(1.000000f, -0.214820f, -0.380590f, 0.000000f),
+		FPlane(1.000000f, 2.127980f, 0.000000f, 0.000000f),
+		FPlane(0.000000f, 0.000000f, 0.000000f, 0.000000f)
+	);
+
+	const FMatrix YuvToRgbRec709Full = FMatrix(
+		FPlane(1.164400f, 0.000000f, 1.792700f, 0.000000f),
+		FPlane(1.164400f, -0.213300f, -0.532900f, 0.000000f),
+		FPlane(1.164400f, 2.112400f, 0.000000f, 0.000000f),
+		FPlane(0.000000f, 0.000000f, 0.000000f, 0.000000f)
+	);
+
+	const FMatrix RgbToYuvRec709Full = FMatrix(
+		FPlane(0.182581f, 0.614210f, 0.062020f, 0.000000f),
+		FPlane(-0.100642f, -0.338566f, 0.439208f, 0.000000f),
+		FPlane(0.439227f, -0.398944f, -0.040283f, 0.000000f),
+		FPlane(0.000000f, 0.000000f, 0.000000f, 0.000000f)
+	);
+
+	const FVector YUVOffset8bits = FVector(0.06274509803921568627f, 0.5019607843137254902f, 0.5019607843137254902f);
+
+	const FVector YUVOffset10bits = FVector(0.06256109481915933529f, 0.50048875855327468231f, 0.50048875855327468231f);
+
+	/** Setup YUV Offset in matrix */
+	FMatrix CombineColorTransformAndOffset(const FMatrix& InMatrix, const FVector& InYUVOffset)
+	{
+		FMatrix Result = InMatrix;
+		// Offset in last column:
+		// 1) to allow for 4x4 matrix multiplication optimization when going from RGB to YUV (hence the 1.0 in the [3][3] matrix location)
+		// 2) stored in empty space when going from YUV to RGB
+		Result.M[0][3] = InYUVOffset.X;
+		Result.M[1][3] = InYUVOffset.Y;
+		Result.M[2][3] = InYUVOffset.Z;
+		Result.M[3][3] = 1.0f;
+		return Result;
+	}
+}
+
+/* FMediaShadersVS shader
+ *****************************************************************************/
+
+IMPLEMENT_SHADER_TYPE(, FMediaShadersVS, TEXT("/Engine/Private/MediaShaders.usf"), TEXT("MainVertexShader"), SF_Vertex);
+
+
+/* FAYUVConvertPS shader
+ *****************************************************************************/
+
+BEGIN_GLOBAL_SHADER_PARAMETER_STRUCT(FAYUVConvertUB, )
+SHADER_PARAMETER(FMatrix, ColorTransform)
+SHADER_PARAMETER(uint32, SrgbToLinear)
+SHADER_PARAMETER_TEXTURE(Texture2D, Texture)
+SHADER_PARAMETER_SAMPLER(SamplerState, Sampler)
+END_GLOBAL_SHADER_PARAMETER_STRUCT()
+
+IMPLEMENT_GLOBAL_SHADER_PARAMETER_STRUCT(FAYUVConvertUB, "AYUVConvertUB");
+IMPLEMENT_SHADER_TYPE(, FAYUVConvertPS, TEXT("/Engine/Private/MediaShaders.usf"), TEXT("AYUVConvertPS"), SF_Pixel);
+
+
+void FAYUVConvertPS::SetParameters(FRHICommandList& CommandList, TRefCountPtr<FRHITexture2D> AYUVTexture, const FMatrix& ColorTransform, const FVector& YUVOffset, bool SrgbToLinear)
+{
+	FAYUVConvertUB UB;
+	{
+		UB.ColorTransform = MediaShaders::CombineColorTransformAndOffset(ColorTransform, YUVOffset);
+		UB.Sampler = TStaticSamplerState<SF_Bilinear>::GetRHI();
+		UB.SrgbToLinear = SrgbToLinear;
+		UB.Texture = AYUVTexture;
+	}
+
+	TUniformBufferRef<FAYUVConvertUB> Data = TUniformBufferRef<FAYUVConvertUB>::CreateUniformBufferImmediate(UB, UniformBuffer_SingleFrame);
+	SetUniformBufferParameter(CommandList, GetPixelShader(), GetUniformBufferParameter<FAYUVConvertUB>(), Data);
+}
+
+
+/* FRGBConvertPS shader
+ *****************************************************************************/
+
+BEGIN_GLOBAL_SHADER_PARAMETER_STRUCT(FBMPConvertUB, )
+SHADER_PARAMETER(uint32, SrgbToLinear)
+SHADER_PARAMETER(FVector2D, UVScale)
+SHADER_PARAMETER_TEXTURE(Texture2D, Texture)
+SHADER_PARAMETER_SAMPLER(SamplerState, Sampler)
+END_GLOBAL_SHADER_PARAMETER_STRUCT()
+
+IMPLEMENT_GLOBAL_SHADER_PARAMETER_STRUCT(FBMPConvertUB, "BMPConvertUB");
+IMPLEMENT_SHADER_TYPE(, FBMPConvertPS, TEXT("/Engine/Private/MediaShaders.usf"), TEXT("BMPConvertPS"), SF_Pixel);
+
+
+void FBMPConvertPS::SetParameters(FRHICommandList& CommandList, TRefCountPtr<FRHITexture2D> BMPTexture, const FIntPoint& OutputDimensions, bool SrgbToLinear)
+{
+	FBMPConvertUB UB;
+	{
+		UB.SrgbToLinear = SrgbToLinear;
+		UB.UVScale = FVector2D((float)OutputDimensions.X / (float)BMPTexture->GetSizeX(), (float)OutputDimensions.Y / (float)BMPTexture->GetSizeY());
+		UB.Texture = BMPTexture;
+		UB.Sampler = TStaticSamplerState<SF_Bilinear>::GetRHI();
+	}
+
+	TUniformBufferRef<FBMPConvertUB> Data = TUniformBufferRef<FBMPConvertUB>::CreateUniformBufferImmediate(UB, UniformBuffer_SingleFrame);
+	SetUniformBufferParameter(CommandList, GetPixelShader(), GetUniformBufferParameter<FBMPConvertUB>(), Data);
+}
+
+
+/* FNV12ConvertPS shader
+ *****************************************************************************/
+
+BEGIN_GLOBAL_SHADER_PARAMETER_STRUCT(FNV12ConvertUB, )
+SHADER_PARAMETER(FMatrix, ColorTransform)
+SHADER_PARAMETER(uint32, OutputWidth)
+SHADER_PARAMETER(uint32, SrgbToLinear)
+SHADER_PARAMETER(FVector2D, UVScale)
+SHADER_PARAMETER_TEXTURE(Texture2D, Texture)
+SHADER_PARAMETER_SAMPLER(SamplerState, SamplerB)
+SHADER_PARAMETER_SAMPLER(SamplerState, SamplerP)
+END_GLOBAL_SHADER_PARAMETER_STRUCT()
+
+IMPLEMENT_GLOBAL_SHADER_PARAMETER_STRUCT(FNV12ConvertUB, "NV12ConvertUB");
+IMPLEMENT_SHADER_TYPE(, FNV12ConvertPS, TEXT("/Engine/Private/MediaShaders.usf"), TEXT("NV12ConvertPS"), SF_Pixel);
+
+
+void FNV12ConvertPS::SetParameters(FRHICommandList& CommandList, TRefCountPtr<FRHITexture2D> NV12Texture, const FIntPoint& OutputDimensions, const FMatrix& ColorTransform, const FVector& YUVOffset, bool SrgbToLinear)
+{
+	FNV12ConvertUB UB;
+	{
+		UB.ColorTransform = MediaShaders::CombineColorTransformAndOffset(ColorTransform, YUVOffset);
+		UB.OutputWidth = OutputDimensions.X;
+		UB.SamplerB = TStaticSamplerState<SF_Bilinear>::GetRHI();
+		UB.SamplerP = TStaticSamplerState<SF_Point>::GetRHI();
+		UB.SrgbToLinear = SrgbToLinear;
+		UB.Texture = NV12Texture;
+		UB.UVScale = FVector2D((float)OutputDimensions.X / (float)NV12Texture->GetSizeX(), (float)OutputDimensions.Y / (float)NV12Texture->GetSizeY());
+	}
+
+	TUniformBufferRef<FNV12ConvertUB> Data = TUniformBufferRef<FNV12ConvertUB>::CreateUniformBufferImmediate(UB, UniformBuffer_SingleFrame);
+	SetUniformBufferParameter(CommandList, GetPixelShader(), GetUniformBufferParameter<FNV12ConvertUB>(), Data);
+}
+
+
+/* FNV21ConvertPS shader
+ *****************************************************************************/
+
+BEGIN_GLOBAL_SHADER_PARAMETER_STRUCT(FNV21ConvertUB, )
+SHADER_PARAMETER(FMatrix, ColorTransform)
+SHADER_PARAMETER(uint32, OutputWidth)
+SHADER_PARAMETER(uint32, SrgbToLinear)
+SHADER_PARAMETER(FVector2D, UVScale)
+SHADER_PARAMETER_TEXTURE(Texture2D, Texture)
+SHADER_PARAMETER_SAMPLER(SamplerState, SamplerB)
+SHADER_PARAMETER_SAMPLER(SamplerState, SamplerP)
+END_GLOBAL_SHADER_PARAMETER_STRUCT()
+
+IMPLEMENT_GLOBAL_SHADER_PARAMETER_STRUCT(FNV21ConvertUB, "NV21ConvertUB");
+IMPLEMENT_SHADER_TYPE(, FNV21ConvertPS, TEXT("/Engine/Private/MediaShaders.usf"), TEXT("NV21ConvertPS"), SF_Pixel);
+
+
+void FNV21ConvertPS::SetParameters(FRHICommandList& CommandList, TRefCountPtr<FRHITexture2D> NV21Texture, const FIntPoint& OutputDimensions, const FMatrix& ColorTransform, const FVector& YUVOffset, bool SrgbToLinear)
+{
+	FNV21ConvertUB UB;
+	{
+		UB.ColorTransform = MediaShaders::CombineColorTransformAndOffset(ColorTransform, YUVOffset);
+		UB.OutputWidth = OutputDimensions.Y;
+		UB.SamplerB = TStaticSamplerState<SF_Bilinear>::GetRHI();
+		UB.SamplerP = TStaticSamplerState<SF_Point>::GetRHI();
+		UB.SrgbToLinear = SrgbToLinear;
+		UB.Texture = NV21Texture;
+		UB.UVScale = FVector2D((float)OutputDimensions.X / (float)NV21Texture->GetSizeX(), (float)OutputDimensions.Y / (float)NV21Texture->GetSizeY());
+	}
+
+	TUniformBufferRef<FNV21ConvertUB> Data = TUniformBufferRef<FNV21ConvertUB>::CreateUniformBufferImmediate(UB, UniformBuffer_SingleFrame);
+	SetUniformBufferParameter(CommandList, GetPixelShader(), GetUniformBufferParameter<FNV21ConvertUB>(), Data);
+}
+
+
+/* FRGBConvertPS shader
+ *****************************************************************************/
+
+BEGIN_GLOBAL_SHADER_PARAMETER_STRUCT(FRGBConvertUB, )
+SHADER_PARAMETER(FVector2D, UVScale)
+SHADER_PARAMETER(uint32, SrgbToLinear)
+SHADER_PARAMETER_TEXTURE(Texture2D, Texture)
+SHADER_PARAMETER_SAMPLER(SamplerState, Sampler)
+END_GLOBAL_SHADER_PARAMETER_STRUCT()
+
+IMPLEMENT_GLOBAL_SHADER_PARAMETER_STRUCT(FRGBConvertUB, "RGBConvertUB");
+IMPLEMENT_SHADER_TYPE(, FRGBConvertPS, TEXT("/Engine/Private/MediaShaders.usf"), TEXT("RGBConvertPS"), SF_Pixel);
+
+
+void FRGBConvertPS::SetParameters(FRHICommandList& CommandList, TRefCountPtr<FRHITexture2D> RGBTexture, const FIntPoint& OutputDimensions, bool bSrgbToLinear)
+{
+	FRGBConvertUB UB;
+	{
+		UB.Sampler = TStaticSamplerState<SF_Bilinear>::GetRHI();
+		UB.SrgbToLinear = bSrgbToLinear;
+		UB.Texture = RGBTexture;
+		UB.UVScale = FVector2D((float)OutputDimensions.X / (float)RGBTexture->GetSizeX(), (float)OutputDimensions.Y / (float)RGBTexture->GetSizeY());
+	}
+
+	TUniformBufferRef<FRGBConvertUB> Data = TUniformBufferRef<FRGBConvertUB>::CreateUniformBufferImmediate(UB, UniformBuffer_SingleFrame);
+	SetUniformBufferParameter(CommandList, GetPixelShader(), GetUniformBufferParameter<FRGBConvertUB>(), Data);
+}
+
+
+/* FYCbCrConvertUB shader
+ *****************************************************************************/
+
+BEGIN_GLOBAL_SHADER_PARAMETER_STRUCT(FYCbCrConvertUB, )
+SHADER_PARAMETER(FMatrix, ColorTransform)
+SHADER_PARAMETER(uint32, SrgbToLinear)
+SHADER_PARAMETER_TEXTURE(Texture2D, LumaTexture)
+SHADER_PARAMETER_TEXTURE(Texture2D, CbCrTexture)
+SHADER_PARAMETER_SAMPLER(SamplerState, LumaSampler)
+SHADER_PARAMETER_SAMPLER(SamplerState, CbCrSampler)
+END_GLOBAL_SHADER_PARAMETER_STRUCT()
+
+IMPLEMENT_GLOBAL_SHADER_PARAMETER_STRUCT(FYCbCrConvertUB, "YCbCrConvertUB");
+IMPLEMENT_SHADER_TYPE(, FYCbCrConvertPS, TEXT("/Engine/Private/MediaShaders.usf"), TEXT("YCbCrConvertPS"), SF_Pixel);
+IMPLEMENT_SHADER_TYPE(, FYCbCrConvertPS_4x4Matrix, TEXT("/Engine/Private/MediaShaders.usf"), TEXT("YCbCrConvertPS_4x4Matrix"), SF_Pixel);
+
+
+void FYCbCrConvertPS::SetParameters(FRHICommandList& CommandList, TRefCountPtr<FRHITexture2D> LumaTexture, TRefCountPtr<FRHITexture2D> CbCrTexture, const FMatrix& ColorTransform, const FVector& YUVOffset, bool SrgbToLinear)
+{
+	FYCbCrConvertUB UB;
+	{
+		// Chroma is not usually 1:1 with the output textxure
+		UB.CbCrSampler = TStaticSamplerState<SF_Bilinear>::GetRHI();
+		UB.CbCrTexture = CbCrTexture;
+		UB.ColorTransform = MediaShaders::CombineColorTransformAndOffset(ColorTransform, YUVOffset);
+		
+		// Luma should be 1:1 with the output texture and needs to be point sampled
+		UB.LumaSampler = TStaticSamplerState<SF_Point>::GetRHI();
+		UB.LumaTexture = LumaTexture;
+		UB.SrgbToLinear = SrgbToLinear;
+	}
+
+	TUniformBufferRef<FYCbCrConvertUB> Data = TUniformBufferRef<FYCbCrConvertUB>::CreateUniformBufferImmediate(UB, UniformBuffer_SingleFrame);
+	SetUniformBufferParameter(CommandList, GetPixelShader(), GetUniformBufferParameter<FYCbCrConvertUB>(), Data);	
+}
+
+
+/* FUYVYConvertPS shader
+ *****************************************************************************/
+
+BEGIN_GLOBAL_SHADER_PARAMETER_STRUCT(FUYVYConvertUB, )
+SHADER_PARAMETER(FMatrix, ColorTransform)
+SHADER_PARAMETER(uint32, SrgbToLinear)
+SHADER_PARAMETER(uint32, Width)
+SHADER_PARAMETER_TEXTURE(Texture2D, Texture)
+SHADER_PARAMETER_SAMPLER(SamplerState, SamplerB)
+SHADER_PARAMETER_SAMPLER(SamplerState, SamplerP)
+END_GLOBAL_SHADER_PARAMETER_STRUCT()
+
+IMPLEMENT_GLOBAL_SHADER_PARAMETER_STRUCT(FUYVYConvertUB, "UYVYConvertUB");
+IMPLEMENT_SHADER_TYPE(, FUYVYConvertPS, TEXT("/Engine/Private/MediaShaders.usf"), TEXT("UYVYConvertPS"), SF_Pixel);
+
+
+void FUYVYConvertPS::SetParameters(FRHICommandList& CommandList, TRefCountPtr<FRHITexture2D> UYVYTexture, const FMatrix& ColorTransform, const FVector& YUVOffset, bool SrgbToLinear)
+{
+	FUYVYConvertUB UB;
+	{
+		UB.ColorTransform = MediaShaders::CombineColorTransformAndOffset(ColorTransform, YUVOffset);
+		UB.SamplerB = TStaticSamplerState<SF_Bilinear>::GetRHI();
+		UB.SamplerP = TStaticSamplerState<SF_Point>::GetRHI();
+		UB.SrgbToLinear = SrgbToLinear;
+		UB.Texture = UYVYTexture;
+		UB.Width = UYVYTexture->GetSizeX();
+	}
+
+	TUniformBufferRef<FUYVYConvertUB> Data = TUniformBufferRef<FUYVYConvertUB>::CreateUniformBufferImmediate(UB, UniformBuffer_SingleFrame);
+	SetUniformBufferParameter(CommandList, GetPixelShader(), GetUniformBufferParameter<FUYVYConvertUB>(), Data);
+}
+
+
+/* FYUVConvertPS shader
+ *****************************************************************************/
+
+BEGIN_GLOBAL_SHADER_PARAMETER_STRUCT(FYUVConvertUB, )
+SHADER_PARAMETER(FMatrix, ColorTransform)
+SHADER_PARAMETER(uint32, SrgbToLinear)
+SHADER_PARAMETER(FVector2D, UVScale)
+SHADER_PARAMETER_TEXTURE(Texture2D, YTexture)
+SHADER_PARAMETER_TEXTURE(Texture2D, UTexture)
+SHADER_PARAMETER_TEXTURE(Texture2D, VTexture)
+SHADER_PARAMETER_SAMPLER(SamplerState, YSampler)
+SHADER_PARAMETER_SAMPLER(SamplerState, USampler)
+SHADER_PARAMETER_SAMPLER(SamplerState, VSampler)
+END_GLOBAL_SHADER_PARAMETER_STRUCT()
+
+IMPLEMENT_GLOBAL_SHADER_PARAMETER_STRUCT(FYUVConvertUB, "YUVConvertUB");
+IMPLEMENT_SHADER_TYPE(, FYUVConvertPS, TEXT("/Engine/Private/MediaShaders.usf"), TEXT("YUVConvertPS"), SF_Pixel);
+
+
+void FYUVConvertPS::SetParameters(FRHICommandList& CommandList, TRefCountPtr<FRHITexture2D> YTexture, TRefCountPtr<FRHITexture2D> UTexture, TRefCountPtr<FRHITexture2D> VTexture, const FIntPoint& OutputDimensions, const FMatrix& ColorTransform, const FVector& YUVOffset, bool SrgbToLinear)
+{
+	FYUVConvertUB UB;
+	{
+		UB.ColorTransform = MediaShaders::CombineColorTransformAndOffset(ColorTransform, YUVOffset);
+		UB.SrgbToLinear = SrgbToLinear;
+		UB.YTexture = YTexture;
+		UB.UTexture = UTexture;
+		UB.VTexture = VTexture;
+		UB.YSampler = TStaticSamplerState<SF_Bilinear>::GetRHI();
+		UB.USampler = TStaticSamplerState<SF_Bilinear>::GetRHI();
+		UB.VSampler = TStaticSamplerState<SF_Bilinear>::GetRHI();
+		UB.UVScale = FVector2D((float) OutputDimensions.X / (float) YTexture->GetSizeX(), (float) OutputDimensions.Y / (float) YTexture->GetSizeY());
+	}
+
+	TUniformBufferRef<FYUVConvertUB> Data = TUniformBufferRef<FYUVConvertUB>::CreateUniformBufferImmediate(UB, UniformBuffer_SingleFrame);
+	SetUniformBufferParameter(CommandList, GetPixelShader(), GetUniformBufferParameter<FYUVConvertUB>(), Data);
+}
+
+
+/* FYUVv210ConvertPS shader
+ *****************************************************************************/
+
+BEGIN_GLOBAL_SHADER_PARAMETER_STRUCT(FYUVv210ConvertUB, )
+SHADER_PARAMETER(FMatrix, ColorTransform)
+SHADER_PARAMETER(uint32, SrgbToLinear)
+SHADER_PARAMETER(uint32, OutputDimX)
+SHADER_PARAMETER(uint32, OutputDimY)
+SHADER_PARAMETER_TEXTURE(Texture2D<uint4>, YUVTexture)
+END_GLOBAL_SHADER_PARAMETER_STRUCT()
+
+IMPLEMENT_GLOBAL_SHADER_PARAMETER_STRUCT(FYUVv210ConvertUB, "YUVv210ConvertUB");
+IMPLEMENT_SHADER_TYPE(, FYUVv210ConvertPS, TEXT("/Engine/Private/MediaShaders.usf"), TEXT("YUVv210ConvertPS"), SF_Pixel);
+
+
+void FYUVv210ConvertPS::SetParameters(FRHICommandList& CommandList, TRefCountPtr<FRHITexture2D> YUVTexture, const FIntPoint& OutputDimensions, const FMatrix& ColorTransform, const FVector& YUVOffset, bool SrgbToLinear)
+{
+	FYUVv210ConvertUB UB;
+	{
+		UB.ColorTransform = MediaShaders::CombineColorTransformAndOffset(ColorTransform, YUVOffset);
+		UB.SrgbToLinear = SrgbToLinear;
+		UB.OutputDimX = OutputDimensions.X;
+		UB.OutputDimY = OutputDimensions.Y;
+		UB.YUVTexture = YUVTexture;
+	}
+
+	TUniformBufferRef<FYUVv210ConvertUB> Data = TUniformBufferRef<FYUVv210ConvertUB>::CreateUniformBufferImmediate(UB, UniformBuffer_SingleFrame);
+	SetUniformBufferParameter(CommandList, GetPixelShader(), GetUniformBufferParameter<FYUVv210ConvertUB>(), Data);
+}
+
+
+/* FYUY2ConvertPS shader
+ *****************************************************************************/
+
+BEGIN_GLOBAL_SHADER_PARAMETER_STRUCT(FYUY2ConvertUB, )
+SHADER_PARAMETER(FMatrix, ColorTransform)
+SHADER_PARAMETER(uint32, OutputWidth)
+SHADER_PARAMETER(uint32, SrgbToLinear)
+SHADER_PARAMETER(FVector2D, UVScale)
+SHADER_PARAMETER_TEXTURE(Texture2D, Texture)
+SHADER_PARAMETER_SAMPLER(SamplerState, SamplerB)
+SHADER_PARAMETER_SAMPLER(SamplerState, SamplerP)
+END_GLOBAL_SHADER_PARAMETER_STRUCT()
+
+IMPLEMENT_GLOBAL_SHADER_PARAMETER_STRUCT(FYUY2ConvertUB, "YUY2ConvertUB");
+IMPLEMENT_SHADER_TYPE(, FYUY2ConvertPS, TEXT("/Engine/Private/MediaShaders.usf"), TEXT("YUY2ConvertPS"), SF_Pixel);
+
+
+void FYUY2ConvertPS::SetParameters(FRHICommandList& CommandList, TRefCountPtr<FRHITexture2D> YUY2Texture, const FIntPoint& OutputDimensions, const FMatrix& ColorTransform, const FVector& YUVOffset, bool SrgbToLinear)
+{
+	FYUY2ConvertUB UB;
+	{
+		UB.ColorTransform = MediaShaders::CombineColorTransformAndOffset(ColorTransform, YUVOffset);
+		UB.OutputWidth = OutputDimensions.X;
+		UB.SamplerB = TStaticSamplerState<SF_Bilinear>::GetRHI();
+		UB.SamplerP = TStaticSamplerState<SF_Point>::GetRHI();
+		UB.SrgbToLinear = SrgbToLinear;
+		UB.Texture = YUY2Texture;
+		UB.UVScale = FVector2D((float)OutputDimensions.X / (2.0f * YUY2Texture->GetSizeX()), (float)OutputDimensions.Y / (float)YUY2Texture->GetSizeY());
+	}
+
+	TUniformBufferRef<FYUY2ConvertUB> Data = TUniformBufferRef<FYUY2ConvertUB>::CreateUniformBufferImmediate(UB, UniformBuffer_SingleFrame);
+	SetUniformBufferParameter(CommandList, GetPixelShader(), GetUniformBufferParameter<FYUY2ConvertUB>(), Data);
+}
+
+
+/* FYVYUConvertPS shader
+ *****************************************************************************/
+
+BEGIN_GLOBAL_SHADER_PARAMETER_STRUCT(FYVYUConvertUB, )
+SHADER_PARAMETER(FMatrix, ColorTransform)
+SHADER_PARAMETER(uint32, SrgbToLinear)
+SHADER_PARAMETER(uint32, Width)
+SHADER_PARAMETER_TEXTURE(Texture2D, Texture)
+SHADER_PARAMETER_SAMPLER(SamplerState, SamplerB)
+SHADER_PARAMETER_SAMPLER(SamplerState, SamplerP)
+END_GLOBAL_SHADER_PARAMETER_STRUCT()
+
+IMPLEMENT_GLOBAL_SHADER_PARAMETER_STRUCT(FYVYUConvertUB, "YVYUConvertUB");
+IMPLEMENT_SHADER_TYPE(, FYVYUConvertPS, TEXT("/Engine/Private/MediaShaders.usf"), TEXT("YVYUConvertPS"), SF_Pixel);
+
+
+void FYVYUConvertPS::SetParameters(FRHICommandList& CommandList, TRefCountPtr<FRHITexture2D> YVYUTexture, const FMatrix& ColorTransform, const FVector& YUVOffset, bool SrgbToLinear)
+{
+	FYVYUConvertUB UB;
+	{
+		UB.ColorTransform = MediaShaders::CombineColorTransformAndOffset(ColorTransform, YUVOffset);
+		UB.SamplerB = TStaticSamplerState<SF_Bilinear>::GetRHI();
+		UB.SamplerP = TStaticSamplerState<SF_Point>::GetRHI();
+		UB.SrgbToLinear = SrgbToLinear;
+		UB.Texture = YVYUTexture;
+		UB.Width = YVYUTexture->GetSizeX();
+	}
+
+	TUniformBufferRef<FYVYUConvertUB> Data = TUniformBufferRef<FYVYUConvertUB>::CreateUniformBufferImmediate(UB, UniformBuffer_SingleFrame);
+	SetUniformBufferParameter(CommandList, GetPixelShader(), GetUniformBufferParameter<FYVYUConvertUB>(), Data);
+}
+
+
+/* FRGB8toUYVY8ConvertPS shader
+ *****************************************************************************/
+
+BEGIN_GLOBAL_SHADER_PARAMETER_STRUCT(FRGB8toUYVY8ConvertUB, )
+SHADER_PARAMETER(FMatrix, ColorTransform)
+SHADER_PARAMETER(uint32, LinearToSrgb)
+SHADER_PARAMETER(float, OnePixelDeltaX)
+SHADER_PARAMETER_TEXTURE(Texture2D, Texture)
+SHADER_PARAMETER_SAMPLER(SamplerState, SamplerP)
+END_GLOBAL_SHADER_PARAMETER_STRUCT()
+
+IMPLEMENT_GLOBAL_SHADER_PARAMETER_STRUCT(FRGB8toUYVY8ConvertUB, "RGB8toUYVY8ConvertUB");
+IMPLEMENT_SHADER_TYPE(, FRGB8toUYVY8ConvertPS, TEXT("/Engine/Private/MediaShaders.usf"), TEXT("RGB8toUYVY8ConvertPS"), SF_Pixel);
+
+
+void FRGB8toUYVY8ConvertPS::SetParameters(FRHICommandList& CommandList, TRefCountPtr<FRHITexture2D> RGBATexture, const FMatrix& ColorTransform, const FVector& YUVOffset, bool LinearToSrgb)
+{
+	FRGB8toUYVY8ConvertUB UB;
+	{
+		UB.ColorTransform = MediaShaders::CombineColorTransformAndOffset(ColorTransform, YUVOffset);
+		UB.SamplerP = TStaticSamplerState<SF_Point>::GetRHI();
+		UB.LinearToSrgb = LinearToSrgb;
+		UB.Texture = RGBATexture;
+		UB.OnePixelDeltaX = 1.0f / float(RGBATexture->GetSizeX());
+	}
+
+	TUniformBufferRef<FRGB8toUYVY8ConvertUB> Data = TUniformBufferRef<FRGB8toUYVY8ConvertUB>::CreateUniformBufferImmediate(UB, UniformBuffer_SingleFrame);
+	SetUniformBufferParameter(CommandList, GetPixelShader(), GetUniformBufferParameter<FRGB8toUYVY8ConvertUB>(), Data);
+}
+
+
+/* FRGB10toYUVv210ConvertPS shader
+ *****************************************************************************/
+
+BEGIN_GLOBAL_SHADER_PARAMETER_STRUCT(FRGB10toYUVv210ConvertUB, )
+SHADER_PARAMETER(FMatrix, ColorTransform)
+SHADER_PARAMETER(uint32, LinearToSrgb)
+SHADER_PARAMETER(float, OnePixelDeltaX)
+SHADER_PARAMETER_TEXTURE(Texture2D, Texture)
+SHADER_PARAMETER_SAMPLER(SamplerState, SamplerP)
+END_GLOBAL_SHADER_PARAMETER_STRUCT()
+
+IMPLEMENT_GLOBAL_SHADER_PARAMETER_STRUCT(FRGB10toYUVv210ConvertUB, "RGB10toYUVv210ConvertUB");
+IMPLEMENT_SHADER_TYPE(, FRGB10toYUVv210ConvertPS, TEXT("/Engine/Private/MediaShaders.usf"), TEXT("RGB10toYUVv210ConvertPS"), SF_Pixel);
+
+
+void FRGB10toYUVv210ConvertPS::SetParameters(FRHICommandList& CommandList, TRefCountPtr<FRHITexture2D> RGBATexture, const FMatrix& ColorTransform, const FVector& YUVOffset, bool LinearToSrgb)
+{
+	FRGB10toYUVv210ConvertUB UB;
+	{
+		UB.ColorTransform = MediaShaders::CombineColorTransformAndOffset(ColorTransform, YUVOffset);
+		UB.SamplerP = TStaticSamplerState<SF_Point>::GetRHI();
+		UB.LinearToSrgb = LinearToSrgb;
+		UB.Texture = RGBATexture;
+		UB.OnePixelDeltaX = 1.0f / float(RGBATexture->GetSizeX());
+	}
+
+	TUniformBufferRef<FRGB10toYUVv210ConvertUB> Data = TUniformBufferRef<FRGB10toYUVv210ConvertUB>::CreateUniformBufferImmediate(UB, UniformBuffer_SingleFrame);
+	SetUniformBufferParameter(CommandList, GetPixelShader(), GetUniformBufferParameter<FRGB10toYUVv210ConvertUB>(), Data);
+}
+
+
+/* FInvertAlphaPS shader
+ *****************************************************************************/
+
+BEGIN_GLOBAL_SHADER_PARAMETER_STRUCT(FInvertAlphaUB, )
+SHADER_PARAMETER_TEXTURE(Texture2D, Texture)
+SHADER_PARAMETER_SAMPLER(SamplerState, SamplerP)
+END_GLOBAL_SHADER_PARAMETER_STRUCT()
+
+IMPLEMENT_GLOBAL_SHADER_PARAMETER_STRUCT(FInvertAlphaUB, "InvertAlphaUB");
+IMPLEMENT_SHADER_TYPE(, FInvertAlphaPS, TEXT("/Engine/Private/MediaShaders.usf"), TEXT("InvertAlphaPS"), SF_Pixel);
+
+
+void FInvertAlphaPS::SetParameters(FRHICommandList& CommandList, TRefCountPtr<FRHITexture2D> RGBATexture)
+{
+	FInvertAlphaUB UB;
+	{
+		UB.SamplerP = TStaticSamplerState<SF_Point>::GetRHI();
+		UB.Texture = RGBATexture;
+	}
+
+	TUniformBufferRef<FInvertAlphaUB> Data = TUniformBufferRef<FInvertAlphaUB>::CreateUniformBufferImmediate(UB, UniformBuffer_SingleFrame);
+	SetUniformBufferParameter(CommandList, GetPixelShader(), GetUniformBufferParameter<FInvertAlphaUB>(), Data);
+}
+
+
+/* FSetAlphaOnePS shader
+ *****************************************************************************/
+
+BEGIN_GLOBAL_SHADER_PARAMETER_STRUCT(FSetAlphaOneUB, )
+SHADER_PARAMETER_TEXTURE(Texture2D, Texture)
+SHADER_PARAMETER_SAMPLER(SamplerState, SamplerP)
+END_GLOBAL_SHADER_PARAMETER_STRUCT()
+
+IMPLEMENT_GLOBAL_SHADER_PARAMETER_STRUCT(FSetAlphaOneUB, "SetAlphaOneUB");
+IMPLEMENT_SHADER_TYPE(, FSetAlphaOnePS, TEXT("/Engine/Private/MediaShaders.usf"), TEXT("SetAlphaOnePS"), SF_Pixel);
+
+
+void FSetAlphaOnePS::SetParameters(FRHICommandList& CommandList, TRefCountPtr<FRHITexture2D> RGBATexture)
+{
+	FSetAlphaOneUB UB;
+	{
+		UB.SamplerP = TStaticSamplerState<SF_Point>::GetRHI();
+		UB.Texture = RGBATexture;
+	}
+
+	TUniformBufferRef<FSetAlphaOneUB> Data = TUniformBufferRef<FSetAlphaOneUB>::CreateUniformBufferImmediate(UB, UniformBuffer_SingleFrame);
+	SetUniformBufferParameter(CommandList, GetPixelShader(), GetUniformBufferParameter<FSetAlphaOneUB>(), Data);
+}
diff --git a/Engine/Source/Runtime/RenderCore/Private/OculusShaders.cpp b/Engine/Source/Runtime/RenderCore/Private/OculusShaders.cpp
new file mode 100644
--- /dev/null
+++ b/Engine/Source/Runtime/RenderCore/Private/OculusShaders.cpp
@@ -0,0 +1,10 @@
+// Copyright 1998-2019 Epic Games, Inc. All Rights Reserved.
+
+#include "OculusShaders.h"
+
+
+IMPLEMENT_SHADER_TYPE(, FOculusVertexShader, TEXT("/Engine/Private/OculusShaders.usf"), TEXT("MainVertexShader"), SF_Vertex);
+IMPLEMENT_SHADER_TYPE(, FOculusWhiteShader, TEXT("/Engine/Private/OculusShaders.usf"), TEXT("MainWhiteShader"), SF_Pixel);
+IMPLEMENT_SHADER_TYPE(, FOculusBlackShader, TEXT("/Engine/Private/OculusShaders.usf"), TEXT("MainBlackShader"), SF_Pixel);
+IMPLEMENT_SHADER_TYPE(, FOculusAlphaInverseShader, TEXT("/Engine/Private/OculusShaders.usf"), TEXT("MainAlphaInverseShader"), SF_Pixel);
+IMPLEMENT_SHADER_TYPE(, FOculusCubemapPS, TEXT("/Engine/Private/OculusShaders.usf"), TEXT("MainForCubemap"), SF_Pixel);
diff --git a/Engine/Source/Runtime/RenderCore/Private/OneColorShader.cpp b/Engine/Source/Runtime/RenderCore/Private/OneColorShader.cpp
new file mode 100644
--- /dev/null
+++ b/Engine/Source/Runtime/RenderCore/Private/OneColorShader.cpp
@@ -0,0 +1,61 @@
+// Copyright 1998-2019 Epic Games, Inc. All Rights Reserved.
+
+
+#include "OneColorShader.h"
+
+BEGIN_GLOBAL_SHADER_PARAMETER_STRUCT(FClearShaderUB, )
+	SHADER_PARAMETER_ARRAY(FVector4, DrawColorMRT, [8] )
+END_GLOBAL_SHADER_PARAMETER_STRUCT()
+
+IMPLEMENT_GLOBAL_SHADER_PARAMETER_STRUCT(FClearShaderUB, "ClearShaderUB");
+
+void FOneColorPS::SetColors(FRHICommandList& RHICmdList, const FLinearColor* Colors, int32 NumColors)
+{
+	check(NumColors <= MaxSimultaneousRenderTargets);
+
+	auto& ClearUBParam = GetUniformBufferParameter<FClearShaderUB>();
+	if (ClearUBParam.IsInitialized())
+	{
+		if (ClearUBParam.IsBound())
+		{
+			FClearShaderUB ClearData;
+			FMemory::Memzero(ClearData.DrawColorMRT);			
+			for (int32 i = 0; i < NumColors; ++i)
+			{
+				ClearData.DrawColorMRT[i].X = Colors[i].R;
+				ClearData.DrawColorMRT[i].Y = Colors[i].G;
+				ClearData.DrawColorMRT[i].Z = Colors[i].B;
+				ClearData.DrawColorMRT[i].W = Colors[i].A;
+			}
+
+			FLocalUniformBuffer LocalUB = TUniformBufferRef<FClearShaderUB>::CreateLocalUniformBuffer(RHICmdList, ClearData, UniformBuffer_SingleFrame);	
+			RHICmdList.SetLocalShaderUniformBuffer(GetPixelShader(), ClearUBParam.GetBaseIndex(), LocalUB);
+		}
+	}
+
+	
+}
+
+// #define avoids a lot of code duplication
+#define IMPLEMENT_ONECOLORVS(A,B) typedef TOneColorVS<A,B> TOneColorVS##A##B; \
+	IMPLEMENT_SHADER_TYPE2_WITH_TEMPLATE_PREFIX(template<> RENDERCORE_API, TOneColorVS##A##B, SF_Vertex);
+
+IMPLEMENT_ONECOLORVS(false,false)
+IMPLEMENT_ONECOLORVS(false,true)
+IMPLEMENT_ONECOLORVS(true,true)
+IMPLEMENT_ONECOLORVS(true,false)
+#undef IMPLEMENT_ONECOLORVS
+
+IMPLEMENT_SHADER_TYPE(,FOneColorPS,TEXT("/Engine/Private/OneColorShader.usf"),TEXT("MainPixelShader"),SF_Pixel);
+// Compiling a version for every number of MRT's
+// On AMD PC hardware, outputting to a color index in the shader without a matching render target set has a significant performance hit
+IMPLEMENT_SHADER_TYPE(template<> RENDERCORE_API,TOneColorPixelShaderMRT<1>,TEXT("/Engine/Private/OneColorShader.usf"),TEXT("MainPixelShaderMRT"),SF_Pixel);
+IMPLEMENT_SHADER_TYPE(template<> RENDERCORE_API,TOneColorPixelShaderMRT<2>,TEXT("/Engine/Private/OneColorShader.usf"),TEXT("MainPixelShaderMRT"),SF_Pixel);
+IMPLEMENT_SHADER_TYPE(template<> RENDERCORE_API,TOneColorPixelShaderMRT<3>,TEXT("/Engine/Private/OneColorShader.usf"),TEXT("MainPixelShaderMRT"),SF_Pixel);
+IMPLEMENT_SHADER_TYPE(template<> RENDERCORE_API,TOneColorPixelShaderMRT<4>,TEXT("/Engine/Private/OneColorShader.usf"),TEXT("MainPixelShaderMRT"),SF_Pixel);
+IMPLEMENT_SHADER_TYPE(template<> RENDERCORE_API,TOneColorPixelShaderMRT<5>,TEXT("/Engine/Private/OneColorShader.usf"),TEXT("MainPixelShaderMRT"),SF_Pixel);
+IMPLEMENT_SHADER_TYPE(template<> RENDERCORE_API,TOneColorPixelShaderMRT<6>,TEXT("/Engine/Private/OneColorShader.usf"),TEXT("MainPixelShaderMRT"),SF_Pixel);
+IMPLEMENT_SHADER_TYPE(template<> RENDERCORE_API,TOneColorPixelShaderMRT<7>,TEXT("/Engine/Private/OneColorShader.usf"),TEXT("MainPixelShaderMRT"),SF_Pixel);
+IMPLEMENT_SHADER_TYPE(template<> RENDERCORE_API,TOneColorPixelShaderMRT<8>,TEXT("/Engine/Private/OneColorShader.usf"),TEXT("MainPixelShaderMRT"),SF_Pixel);
+
+IMPLEMENT_SHADER_TYPE(,FFillTextureCS,TEXT("/Engine/Private/OneColorShader.usf"),TEXT("MainFillTextureCS"),SF_Compute);
diff --git a/Engine/Source/Runtime/RenderCore/Private/RGBAToYUV420Shader.cpp b/Engine/Source/Runtime/RenderCore/Private/RGBAToYUV420Shader.cpp
new file mode 100644
--- /dev/null
+++ b/Engine/Source/Runtime/RenderCore/Private/RGBAToYUV420Shader.cpp
@@ -0,0 +1,34 @@
+// Copyright 1998-2019 Epic Games, Inc. All Rights Reserved.
+
+#include "RGBAToYUV420Shader.h"
+#include "ShaderParameterUtils.h"
+
+#if defined(HAS_MORPHEUS) && HAS_MORPHEUS
+
+/* FRGBAToYUV420CS shader
+ *****************************************************************************/
+
+IMPLEMENT_SHADER_TYPE(, FRGBAToYUV420CS, TEXT("/Engine/Private/RGBAToYUV420.usf"), TEXT("RGBAToYUV420Main"), SF_Compute);
+
+
+void FRGBAToYUV420CS::SetParameters(FRHICommandList& RHICmdList, TRefCountPtr<FRHITexture2D> InSrcTexture, FRHIUnorderedAccessView* InOutUAV, float InTargetHeight, float InScaleFactorX, float InScaleFactorY, float InTextureYOffset)
+{
+	FRHIComputeShader* ComputeShaderRHI = GetComputeShader();
+	SetShaderValue(RHICmdList, ComputeShaderRHI, TargetHeight, InTargetHeight);
+	SetShaderValue(RHICmdList, ComputeShaderRHI, ScaleFactorX, InScaleFactorX);
+	SetShaderValue(RHICmdList, ComputeShaderRHI, ScaleFactorY, InScaleFactorY);
+	SetShaderValue(RHICmdList, ComputeShaderRHI, TextureYOffset, InTextureYOffset);
+	SetTextureParameter(RHICmdList, ComputeShaderRHI, SrcTexture, InSrcTexture);
+	RHICmdList.SetUAVParameter(ComputeShaderRHI, OutTextureRW.GetBaseIndex(), InOutUAV);
+}
+
+/**
+* Unbinds any buffers that have been bound.
+*/
+void FRGBAToYUV420CS::UnbindBuffers(FRHICommandList& RHICmdList)
+{
+	FRHIComputeShader* ComputeShaderRHI = GetComputeShader();
+	RHICmdList.SetUAVParameter(ComputeShaderRHI, OutTextureRW.GetBaseIndex(), nullptr);
+}
+
+#endif //defined(HAS_MORPHEUS) && HAS_MORPHEUS
diff --git a/Engine/Source/Runtime/RenderCore/Private/ResolveShader.cpp b/Engine/Source/Runtime/RenderCore/Private/ResolveShader.cpp
new file mode 100644
--- /dev/null
+++ b/Engine/Source/Runtime/RenderCore/Private/ResolveShader.cpp
@@ -0,0 +1,50 @@
+// Copyright 1998-2019 Epic Games, Inc. All Rights Reserved.
+
+
+#include "ResolveShader.h"
+#include "ShaderParameterUtils.h"
+
+IMPLEMENT_SHADER_TYPE(, FResolveDepthPS, TEXT("/Engine/Private/ResolvePixelShader.usf"), TEXT("MainDepth"), SF_Pixel);
+IMPLEMENT_SHADER_TYPE(, FResolveDepth2XPS, TEXT("/Engine/Private/ResolvePixelShader.usf"), TEXT("MainDepth"), SF_Pixel);
+IMPLEMENT_SHADER_TYPE(, FResolveDepth4XPS, TEXT("/Engine/Private/ResolvePixelShader.usf"), TEXT("MainDepth"), SF_Pixel);
+IMPLEMENT_SHADER_TYPE(, FResolveDepth8XPS, TEXT("/Engine/Private/ResolvePixelShader.usf"), TEXT("MainDepth"), SF_Pixel);
+IMPLEMENT_SHADER_TYPE(, FResolveDepthNonMSPS, TEXT("/Engine/Private/ResolvePixelShader.usf"), TEXT("MainDepthNonMS"), SF_Pixel);
+IMPLEMENT_SHADER_TYPE(, FResolveSingleSamplePS, TEXT("/Engine/Private/ResolvePixelShader.usf"), TEXT("MainSingleSample"), SF_Pixel);
+IMPLEMENT_SHADER_TYPE(, FResolveVS, TEXT("/Engine/Private/ResolveVertexShader.usf"), TEXT("Main"), SF_Vertex);
+
+RENDERCORE_API TGlobalResource<FResolveVertexBuffer> GResolveVertexBuffer;
+
+void FResolveSingleSamplePS::SetParameters(FRHICommandList& RHICmdList, uint32 SingleSampleIndexValue)
+{
+	SetShaderValue(RHICmdList, GetPixelShader(),SingleSampleIndex,SingleSampleIndexValue);
+}
+
+void FResolveVS::SetParameters(FRHICommandList& RHICmdList, FVector4 InPositionMinMax, FVector4 InUVMinMax)
+{
+	SetShaderValue(RHICmdList, GetVertexShader(), PositionMinMax, InPositionMinMax);
+	SetShaderValue(RHICmdList, GetVertexShader(), UVMinMax, InUVMinMax);
+}
+
+void FResolveVertexBuffer::InitDynamicRHI()
+{
+	FVertexDeclarationElementList Elements;
+	Elements.Add(FVertexElement(0, 0, VET_Float2, 0, sizeof(FVector2D)));
+	VertexDeclarationRHI = RHICreateVertexDeclaration(Elements);
+
+	void* VoidPtr;
+	FRHIResourceCreateInfo CreateInfo;
+	VertexBufferRHI = RHICreateAndLockVertexBuffer(sizeof(FVector2D) * 4, BUF_Static, CreateInfo, VoidPtr);
+	// Generate the vertices used
+	FVector2D* Vertices = (FVector2D*)VoidPtr;
+	Vertices[0] = { 0.0f, 0.0f };
+	Vertices[1] = { 1.0f, 0.0f };
+	Vertices[2] = { 0.0f, 1.0f };
+	Vertices[3] = { 1.0f, 1.0f };
+	RHIUnlockVertexBuffer(VertexBufferRHI);
+}
+
+void FResolveVertexBuffer::ReleaseDynamicRHI()
+{
+	VertexBufferRHI.SafeRelease();
+	VertexDeclarationRHI.SafeRelease();
+}
diff --git a/Engine/Source/Runtime/RenderCore/Public/BuiltInRayTracingShaders.h b/Engine/Source/Runtime/RenderCore/Public/BuiltInRayTracingShaders.h
new file mode 100644
--- /dev/null
+++ b/Engine/Source/Runtime/RenderCore/Public/BuiltInRayTracingShaders.h
@@ -0,0 +1,89 @@
+// Copyright 1998-2019 Epic Games, Inc. All Rights Reserved.
+
+#pragma once
+
+#include "GlobalShader.h"
+#include "ShaderParameterStruct.h"
+#include "ShaderParameterMacros.h"
+
+#if RHI_RAYTRACING
+
+class FBuiltInRayTracingShader : public FGlobalShader
+{
+protected:
+	static void ModifyCompilationEnvironment(const FGlobalShaderPermutationParameters& Parameters, FShaderCompilerEnvironment& OutEnvironment)
+	{}
+
+public:
+	static bool ShouldCompilePermutation(const FGlobalShaderPermutationParameters& Parameters)
+	{
+		// Built-in ray tracing shaders are always compiled for RHIs that support them, regardless of whether RT is enabled for the project.
+		return RHISupportsRayTracingShaders(Parameters.Platform);
+	}
+
+	FBuiltInRayTracingShader() = default;
+	FBuiltInRayTracingShader(const ShaderMetaType::CompiledShaderInitializerType& Initializer)
+		: FGlobalShader(Initializer)
+	{}
+};
+
+class RENDERCORE_API FOcclusionMainRG : public FBuiltInRayTracingShader
+{
+	DECLARE_GLOBAL_SHADER(FOcclusionMainRG);
+	SHADER_USE_ROOT_PARAMETER_STRUCT(FOcclusionMainRG, FBuiltInRayTracingShader);
+
+	BEGIN_SHADER_PARAMETER_STRUCT(FParameters, )
+		SHADER_PARAMETER_RDG_BUFFER_SRV(RaytracingAccelerationStructure, TLAS)
+		SHADER_PARAMETER_RDG_BUFFER_SRV(StructuredBuffer<FBasicRayData>, Rays)
+		SHADER_PARAMETER_RDG_BUFFER_UAV(RWStructuredBuffer<uint>, OcclusionOutput)
+	END_SHADER_PARAMETER_STRUCT()
+};
+
+class RENDERCORE_API FIntersectionMainRG : public FBuiltInRayTracingShader
+{
+	DECLARE_GLOBAL_SHADER(FIntersectionMainRG);
+	SHADER_USE_ROOT_PARAMETER_STRUCT(FIntersectionMainRG, FBuiltInRayTracingShader);
+
+	BEGIN_SHADER_PARAMETER_STRUCT(FParameters, )
+		SHADER_PARAMETER_RDG_BUFFER_SRV(RaytracingAccelerationStructure, TLAS)
+		SHADER_PARAMETER_RDG_BUFFER_SRV(StructuredBuffer<FBasicRayData>, Rays)
+		SHADER_PARAMETER_RDG_BUFFER_UAV(RWStructuredBuffer<FIntersectionPayload>, IntersectionOutput)
+	END_SHADER_PARAMETER_STRUCT()
+};
+
+class FIntersectionMainCHS : public FBuiltInRayTracingShader
+{
+	DECLARE_EXPORTED_SHADER_TYPE(FIntersectionMainCHS, Global, RENDERCORE_API);
+public:
+
+	FIntersectionMainCHS() = default;
+	FIntersectionMainCHS(const ShaderMetaType::CompiledShaderInitializerType& Initializer)
+		: FBuiltInRayTracingShader(Initializer)
+	{}
+};
+
+class FDefaultMainCHS : public FBuiltInRayTracingShader
+{
+	DECLARE_EXPORTED_SHADER_TYPE(FDefaultMainCHS, Global, RENDERCORE_API);
+public:
+
+	FDefaultMainCHS() = default;
+	FDefaultMainCHS(const ShaderMetaType::CompiledShaderInitializerType& Initializer)
+		: FBuiltInRayTracingShader(Initializer)
+	{}
+};
+
+
+class FDefaultMainMS : public FBuiltInRayTracingShader
+{
+	DECLARE_EXPORTED_SHADER_TYPE(FDefaultMainMS, Global, RENDERCORE_API);
+public:
+
+	FDefaultMainMS() = default;
+	FDefaultMainMS(const ShaderMetaType::CompiledShaderInitializerType& Initializer)
+		: FBuiltInRayTracingShader(Initializer)
+	{}
+};
+
+#endif // RHI_RAYTRACING
+
diff --git a/Engine/Source/Runtime/RenderCore/Public/ClearQuad.h b/Engine/Source/Runtime/RenderCore/Public/ClearQuad.h
new file mode 100644
--- /dev/null
+++ b/Engine/Source/Runtime/RenderCore/Public/ClearQuad.h
@@ -0,0 +1,83 @@
+// Copyright 1998-2019 Epic Games, Inc. All Rights Reserved.
+
+#pragma once
+
+#include "CoreMinimal.h"
+#include "RHIDefinitions.h"
+#include "RenderResource.h"
+
+class FRHICommandList;
+struct FRWBufferStructured;
+struct FRWBuffer;
+struct FSceneRenderTargetItem;
+class FRHIUnorderedAccessView;
+class FGraphicsPipelineStateInitializer;
+
+class FClearVertexBuffer : public FVertexBuffer
+{
+public:
+	/**
+	* Initialize the RHI for this rendering resource
+	*/
+	virtual void InitRHI() override
+	{
+		// create a static vertex buffer
+		FRHIResourceCreateInfo CreateInfo;
+		VertexBufferRHI = RHICreateVertexBuffer(sizeof(FVector4) * 4, BUF_Static, CreateInfo);
+		void* VoidPtr = RHILockVertexBuffer(VertexBufferRHI, 0, sizeof(FVector4) * 4, RLM_WriteOnly);
+		// Generate the vertices used
+		FVector4* Vertices = reinterpret_cast<FVector4*>(VoidPtr);
+		Vertices[0] = FVector4(-1.0f, 1.0f, 0.0f, 1.0f);
+		Vertices[1] = FVector4(1.0f, 1.0f, 0.0f, 1.0f);
+		Vertices[2] = FVector4(-1.0f, -1.0f, 0.0f, 1.0f);
+		Vertices[3] = FVector4(1.0f, -1.0f, 0.0f, 1.0f);
+		RHIUnlockVertexBuffer(VertexBufferRHI);
+	}
+};
+extern RENDERCORE_API TGlobalResource<FClearVertexBuffer> GClearVertexBuffer;
+
+struct FClearQuadCallbacks
+{
+	TFunction<void(FGraphicsPipelineStateInitializer&)> PSOModifier = nullptr;
+	TFunction<void(FRHICommandList&)> PreClear = nullptr;
+	TFunction<void(FRHICommandList&)> PostClear = nullptr;
+};
+
+extern RENDERCORE_API void ClearUAV(FRHICommandList& RHICmdList, const FRWBufferStructured& StructuredBuffer, uint32 Value);
+extern RENDERCORE_API void ClearUAV(FRHICommandList& RHICmdList, const FTextureRWBuffer& Buffer, FLinearColor Value);
+extern RENDERCORE_API void ClearUAV(FRHICommandList& RHICmdList, const FRWBuffer& Buffer, uint32 Value, bool bBarriers = true);
+extern RENDERCORE_API void ClearUAV(FRHICommandList& RHICmdList, FRHIUnorderedAccessView* Buffer, uint32 NumBytes, uint32 Value);
+
+extern RENDERCORE_API void ClearUAV(FRHICommandList& RHICmdList, FRHITexture* Texture, FRHIUnorderedAccessView* TextureUAV, const float(&ClearValues)[4]);
+extern RENDERCORE_API void ClearUAV(FRHICommandList& RHICmdList, FRHITexture* Texture, FRHIUnorderedAccessView* TextureUAV, const uint32(&ClearValues)[4]);
+extern RENDERCORE_API void ClearUAV(FRHICommandList& RHICmdList, FRHITexture* Texture, FRHIUnorderedAccessView* TextureUAV, const FLinearColor& ClearColor);
+
+extern RENDERCORE_API void ClearUAV(FRHICommandList& RHICmdList, const FSceneRenderTargetItem& RenderTargetItem, const float(&ClearValues)[4]);
+extern RENDERCORE_API void ClearUAV(FRHICommandList& RHICmdList, const FSceneRenderTargetItem& RenderTargetItem, const uint32(&ClearValues)[4]);
+extern RENDERCORE_API void ClearUAV(FRHICommandList& RHICmdList, const FSceneRenderTargetItem& RenderTargetItem, const FLinearColor& ClearColor);
+
+extern RENDERCORE_API void ClearTexture2DUAV(FRHICommandList& RHICmdList, FRHIUnorderedAccessView* UAV, int32 Width, int32 Height, const FLinearColor& ClearColor);
+
+extern RENDERCORE_API void DrawClearQuadMRT( FRHICommandList& RHICmdList, bool bClearColor, int32 NumClearColors, const FLinearColor* ClearColorArray, bool bClearDepth, float Depth, bool bClearStencil, uint32 Stencil );
+extern RENDERCORE_API void DrawClearQuadMRT(FRHICommandList& RHICmdList, bool bClearColor, int32 NumClearColors, const FLinearColor* ClearColorArray, bool bClearDepth, float Depth, bool bClearStencil, uint32 Stencil, FClearQuadCallbacks ClearQuadCallbacks);
+extern RENDERCORE_API void DrawClearQuadMRT( FRHICommandList& RHICmdList, bool bClearColor, int32 NumClearColors, const FLinearColor* ClearColorArray, bool bClearDepth, float Depth, bool bClearStencil, uint32 Stencil, FIntPoint ViewSize, FIntRect ExcludeRect );
+
+inline void DrawClearQuad(FRHICommandList& RHICmdList, bool bClearColor, const FLinearColor& Color, bool bClearDepth, float Depth, bool bClearStencil, uint32 Stencil)
+{
+	DrawClearQuadMRT(RHICmdList, bClearColor, 1, &Color, bClearDepth, Depth, bClearStencil, Stencil);
+}
+
+inline void DrawClearQuad(FRHICommandList& RHICmdList, bool bClearColor, const FLinearColor& Color, bool bClearDepth, float Depth, bool bClearStencil, uint32 Stencil, FIntPoint ViewSize, FIntRect ExcludeRect)
+{
+	DrawClearQuadMRT(RHICmdList, bClearColor, 1, &Color, bClearDepth, Depth, bClearStencil, Stencil, ViewSize, ExcludeRect);
+}
+
+inline void DrawClearQuad(FRHICommandList& RHICmdList, const FLinearColor& Color)
+{
+	DrawClearQuadMRT(RHICmdList, true, 1, &Color, false, 0, false, 0);
+}
+
+inline void DrawClearQuad(FRHICommandList& RHICmdList, const FLinearColor& Color, FClearQuadCallbacks ClearQuadCallbacks)
+{
+	DrawClearQuadMRT(RHICmdList, true, 1, &Color, false, 0, false, 0, ClearQuadCallbacks);
+}
\ No newline at end of file
diff --git a/Engine/Source/Runtime/RenderCore/Public/ClearReplacementShaders.h b/Engine/Source/Runtime/RenderCore/Public/ClearReplacementShaders.h
new file mode 100644
--- /dev/null
+++ b/Engine/Source/Runtime/RenderCore/Public/ClearReplacementShaders.h
@@ -0,0 +1,299 @@
+// Copyright 1998-2019 Epic Games, Inc. All Rights Reserved.
+
+
+#pragma once
+
+#include "CoreMinimal.h"
+#include "ShaderParameters.h"
+#include "Shader.h"
+#include "GlobalShader.h"
+
+class FClearReplacementVS : public FGlobalShader
+{
+	DECLARE_EXPORTED_SHADER_TYPE(FClearReplacementVS, Global, RENDERCORE_API);
+public:
+	FClearReplacementVS() {}
+	FClearReplacementVS(const ShaderMetaType::CompiledShaderInitializerType& Initializer)
+	: FGlobalShader( Initializer )
+	{
+	}
+	
+	// FShader interface.
+	virtual bool Serialize(FArchive& Ar) override
+	{
+		bool bShaderHasOutdatedParameters = FGlobalShader::Serialize(Ar);
+		return bShaderHasOutdatedParameters;
+	}
+	
+	static bool ShouldCompilePermutation(const FGlobalShaderPermutationParameters& Parameters)
+	{
+		return IsFeatureLevelSupported(Parameters.Platform, ERHIFeatureLevel::SM5);
+	}
+};
+
+class FClearReplacementPS : public FGlobalShader
+{
+	DECLARE_EXPORTED_SHADER_TYPE(FClearReplacementPS, Global, RENDERCORE_API);
+public:
+	FClearReplacementPS() {}
+	FClearReplacementPS(const ShaderMetaType::CompiledShaderInitializerType& Initializer)
+	: FGlobalShader( Initializer )
+	{
+		ClearColor.Bind(Initializer.ParameterMap, TEXT("ClearColor"), SPF_Mandatory);
+	}
+	
+	// FShader interface.
+	virtual bool Serialize(FArchive& Ar) override
+	{
+		bool bShaderHasOutdatedParameters = FGlobalShader::Serialize(Ar);
+		Ar << ClearColor;
+		return bShaderHasOutdatedParameters;
+	}
+	
+	static bool ShouldCompilePermutation(const FGlobalShaderPermutationParameters& Parameters)
+	{
+		return IsFeatureLevelSupported(Parameters.Platform, ERHIFeatureLevel::SM5);
+	}
+	
+	static void ModifyCompilationEnvironment(const FGlobalShaderPermutationParameters& Parameters, FShaderCompilerEnvironment& OutEnvironment)
+	{
+		FGlobalShader::ModifyCompilationEnvironment(Parameters, OutEnvironment);
+		OutEnvironment.SetRenderTargetOutputFormat(0, PF_A32B32G32R32F);
+	}
+	
+protected:
+	FShaderParameter ClearColor;
+};
+
+template< typename T >
+class FClearTexture2DReplacementCS : public FGlobalShader
+{
+	DECLARE_EXPORTED_SHADER_TYPE(FClearTexture2DReplacementCS, Global, RENDERCORE_API);
+public:
+	FClearTexture2DReplacementCS() {}
+	FClearTexture2DReplacementCS(const ShaderMetaType::CompiledShaderInitializerType& Initializer)
+	: FGlobalShader( Initializer )
+	{
+		ClearColor.Bind(Initializer.ParameterMap, TEXT("ClearColor"), SPF_Mandatory);
+		ClearTextureRW.Bind(Initializer.ParameterMap, TEXT("ClearTextureRW"), SPF_Mandatory);
+	}
+	
+	// FShader interface.
+	virtual bool Serialize(FArchive& Ar) override
+	{
+		bool bShaderHasOutdatedParameters = FGlobalShader::Serialize(Ar);
+		Ar << ClearColor << ClearTextureRW;
+		return bShaderHasOutdatedParameters;
+	}
+	
+	RENDERCORE_API void SetParameters(FRHICommandList& RHICmdList, FRHIUnorderedAccessView* TextureRW, const T(&Values)[4]);
+	RENDERCORE_API void FinalizeParameters(FRHICommandList& RHICmdList, FRHIUnorderedAccessView* TextureRW);
+	
+	static bool ShouldCompilePermutation(const FGlobalShaderPermutationParameters& Parameters)
+	{
+		return IsFeatureLevelSupported(Parameters.Platform, ERHIFeatureLevel::SM5);
+	}
+
+	static void ModifyCompilationEnvironment(const FGlobalShaderPermutationParameters& Parameters, FShaderCompilerEnvironment& OutEnvironment)
+	{
+		FGlobalShader::ModifyCompilationEnvironment(Parameters, OutEnvironment);
+		OutEnvironment.SetDefine( TEXT("Type"), TIsSame< T, float >::Value ? TEXT("float4") : TEXT("uint4") );
+	}
+	
+	const FShaderParameter& GetClearColorParameter()
+	{
+		return ClearColor;
+	}
+	
+	const FShaderResourceParameter& GetClearTextureRWParameter()
+	{
+		return ClearTextureRW;
+	}
+	
+protected:
+	FShaderParameter ClearColor;
+	FShaderResourceParameter ClearTextureRW;
+};
+
+template< typename T >
+class FClearTexture2DArrayReplacementCS : public FGlobalShader
+{
+	DECLARE_EXPORTED_SHADER_TYPE(FClearTexture2DArrayReplacementCS, Global, RENDERCORE_API);
+public:
+	FClearTexture2DArrayReplacementCS() {}
+	FClearTexture2DArrayReplacementCS(const ShaderMetaType::CompiledShaderInitializerType& Initializer)
+		: FGlobalShader(Initializer)
+	{
+		ClearColor.Bind(Initializer.ParameterMap, TEXT("ClearColor"), SPF_Mandatory);
+		ClearTextureArrayRW.Bind(Initializer.ParameterMap, TEXT("ClearTextureArrayRW"), SPF_Mandatory);
+	}
+
+	// FShader interface.
+	virtual bool Serialize(FArchive& Ar) override
+	{
+		bool bShaderHasOutdatedParameters = FGlobalShader::Serialize(Ar);
+		Ar << ClearColor << ClearTextureArrayRW;
+		return bShaderHasOutdatedParameters;
+	}
+
+	RENDERCORE_API void SetParameters(FRHICommandList& RHICmdList, FRHIUnorderedAccessView* TextureRW, const T(&Values)[4]);
+	RENDERCORE_API void FinalizeParameters(FRHICommandList& RHICmdList, FRHIUnorderedAccessView* TextureRW);
+
+	static bool ShouldCompilePermutation(const FGlobalShaderPermutationParameters& Parameters)
+	{
+		return IsFeatureLevelSupported(Parameters.Platform, ERHIFeatureLevel::SM5);
+	}
+
+	static void ModifyCompilationEnvironment(const FGlobalShaderPermutationParameters& Parameters, FShaderCompilerEnvironment& OutEnvironment)
+	{
+		FGlobalShader::ModifyCompilationEnvironment(Parameters, OutEnvironment);
+		OutEnvironment.SetDefine( TEXT("Type"), TIsSame< T, float >::Value ? TEXT("float4") : TEXT("uint4") );
+	}
+
+	const FShaderParameter& GetClearColorParameter()
+	{
+		return ClearColor;
+	}
+
+	const FShaderResourceParameter& GetClearTextureRWParameter()
+	{
+		return ClearTextureArrayRW;
+	}
+
+protected:
+	FShaderParameter ClearColor;
+	FShaderResourceParameter ClearTextureArrayRW;
+};
+
+template< typename T >
+class FClearVolumeReplacementCS : public FGlobalShader
+{
+	DECLARE_EXPORTED_SHADER_TYPE(FClearVolumeReplacementCS, Global, RENDERCORE_API);
+public:
+	FClearVolumeReplacementCS() {}
+	FClearVolumeReplacementCS(const ShaderMetaType::CompiledShaderInitializerType& Initializer)
+		: FGlobalShader(Initializer)
+	{
+		ClearColor.Bind(Initializer.ParameterMap, TEXT("ClearColor"), SPF_Mandatory);
+		ClearVolumeRW.Bind(Initializer.ParameterMap, TEXT("ClearVolumeRW"), SPF_Mandatory);
+	}
+
+	// FShader interface.
+	virtual bool Serialize(FArchive& Ar) override
+	{
+		bool bShaderHasOutdatedParameters = FGlobalShader::Serialize(Ar);
+		Ar << ClearColor << ClearVolumeRW;
+		return bShaderHasOutdatedParameters;
+	}
+
+	RENDERCORE_API void SetParameters(FRHICommandList& RHICmdList, FRHIUnorderedAccessView* TextureRW, const T(&Values)[4]);
+	RENDERCORE_API void FinalizeParameters(FRHICommandList& RHICmdList, FRHIUnorderedAccessView* TextureRW);
+
+	static bool ShouldCompilePermutation(const FGlobalShaderPermutationParameters& Parameters)
+	{
+		return IsFeatureLevelSupported(Parameters.Platform, ERHIFeatureLevel::SM5);
+	}
+
+	static void ModifyCompilationEnvironment(const FGlobalShaderPermutationParameters& Parameters, FShaderCompilerEnvironment& OutEnvironment)
+	{
+		FGlobalShader::ModifyCompilationEnvironment(Parameters, OutEnvironment);
+		OutEnvironment.SetDefine( TEXT("Type"), TIsSame< T, float >::Value ? TEXT("float4") : TEXT("uint4") );
+	}
+
+	const FShaderParameter& GetClearColorParameter()
+	{
+		return ClearColor;
+	}
+
+	const FShaderResourceParameter& GetClearTextureRWParameter()
+	{
+		return ClearVolumeRW;
+	}
+
+protected:
+	FShaderParameter ClearColor;
+	FShaderResourceParameter ClearVolumeRW;
+};
+
+class FClearTexture2DReplacementScissorCS : public FGlobalShader
+{
+	DECLARE_EXPORTED_SHADER_TYPE(FClearTexture2DReplacementScissorCS, Global, RENDERCORE_API);
+public:
+	FClearTexture2DReplacementScissorCS() {}
+	FClearTexture2DReplacementScissorCS(const ShaderMetaType::CompiledShaderInitializerType& Initializer)
+		: FGlobalShader(Initializer)
+	{
+		ClearColor.Bind(Initializer.ParameterMap, TEXT("ClearColor"), SPF_Mandatory);
+		TargetBounds.Bind(Initializer.ParameterMap, TEXT("TargetBounds"), SPF_Mandatory);
+		ClearTextureRW.Bind(Initializer.ParameterMap, TEXT("ClearTextureRW"), SPF_Mandatory);
+	}
+
+	// FShader interface.
+	virtual bool Serialize(FArchive& Ar) override
+	{
+		bool bShaderHasOutdatedParameters = FGlobalShader::Serialize(Ar);
+		Ar << ClearColor << TargetBounds << ClearTextureRW;
+		return bShaderHasOutdatedParameters;
+	}
+
+	RENDERCORE_API void SetParameters(FRHICommandList& RHICmdList, FRHIUnorderedAccessView* TextureRW, FLinearColor ClearColor, const FVector4& InTargetBounds);
+	RENDERCORE_API void FinalizeParameters(FRHICommandList& RHICmdList, FRHIUnorderedAccessView* TextureRW);
+
+	static bool ShouldCompilePermutation(const FGlobalShaderPermutationParameters& Parameters)
+	{
+		return IsFeatureLevelSupported(Parameters.Platform, ERHIFeatureLevel::SM5);
+	}
+
+	const FShaderParameter& GetClearColorParameter()
+	{
+		return ClearColor;
+	}
+
+	const FShaderParameter& GetTargetBoundsParameter()
+	{
+		return TargetBounds;
+	}
+
+	const FShaderResourceParameter& GetClearTextureRWParameter()
+	{
+		return ClearTextureRW;
+	}
+
+protected:
+	FShaderParameter ClearColor;
+	FShaderParameter TargetBounds;
+	FShaderResourceParameter ClearTextureRW;
+};
+
+class FClearBufferReplacementCS : public FGlobalShader
+{
+	DECLARE_EXPORTED_SHADER_TYPE(FClearBufferReplacementCS, Global, RENDERCORE_API);
+public:
+	FClearBufferReplacementCS() {}
+	FClearBufferReplacementCS(const ShaderMetaType::CompiledShaderInitializerType& Initializer)
+	: FGlobalShader( Initializer )
+	{
+		ClearBufferCSParams.Bind(Initializer.ParameterMap, TEXT("ClearBufferCSParams"), SPF_Mandatory);
+		ClearBufferRW.Bind(Initializer.ParameterMap, TEXT("ClearBufferRW"), SPF_Mandatory);
+	}
+	
+	RENDERCORE_API void SetParameters(FRHICommandList& RHICmdList, FRHIUnorderedAccessView* BufferRW, uint32 NumDWordsToClear, uint32 ClearValue, bool bBarriers = true);
+	RENDERCORE_API void FinalizeParameters(FRHICommandList& RHICmdList, FRHIUnorderedAccessView* BufferRW, bool bBarriers = true);
+	
+	// FShader interface.
+	virtual bool Serialize(FArchive& Ar) override
+	{
+		bool bShaderHasOutdatedParameters = FGlobalShader::Serialize(Ar);
+		Ar << ClearBufferCSParams << ClearBufferRW;
+		return bShaderHasOutdatedParameters;
+	}
+	
+	static bool ShouldCompilePermutation(const FGlobalShaderPermutationParameters& Parameters)
+	{
+		return RHISupportsComputeShaders(Parameters.Platform);
+	}
+	
+	//protected:
+	FShaderParameter ClearBufferCSParams;
+	FShaderResourceParameter ClearBufferRW;
+};
diff --git a/Engine/Source/Runtime/RenderCore/Public/LongGPUTask.h b/Engine/Source/Runtime/RenderCore/Public/LongGPUTask.h
new file mode 100644
--- /dev/null
+++ b/Engine/Source/Runtime/RenderCore/Public/LongGPUTask.h
@@ -0,0 +1,10 @@
+// Copyright 1998-2019 Epic Games, Inc. All Rights Reserved.
+
+#pragma once
+
+#include "CoreMinimal.h"
+
+class FRHICommandListImmediate;
+
+extern RENDERCORE_API void IssueScalableLongGPUTask(FRHICommandListImmediate& RHICmdList, int32 NumIteration = -1);
+extern RENDERCORE_API void MeasureLongGPUTaskExecutionTime(FRHICommandListImmediate& RHICmdList);
diff --git a/Engine/Source/Runtime/RenderCore/Public/MediaShaders.h b/Engine/Source/Runtime/RenderCore/Public/MediaShaders.h
new file mode 100644
--- /dev/null
+++ b/Engine/Source/Runtime/RenderCore/Public/MediaShaders.h
@@ -0,0 +1,666 @@
+// Copyright 1998-2019 Epic Games, Inc. All Rights Reserved.
+
+#pragma once
+
+#include "CoreMinimal.h"
+#include "RHI.h"
+#include "RenderResource.h"
+#include "Shader.h"
+#include "GlobalShader.h"
+#include "PipelineStateCache.h"
+
+
+namespace MediaShaders
+{
+	/** Color transform from YUV to sRGB (using values from MSDN). */
+	RENDERCORE_API extern const FMatrix YuvToSrgbDefault;
+
+	/** Color transform from YUV to sRGB (in JPEG color space). */
+	RENDERCORE_API extern const FMatrix YuvToSrgbJpeg;
+
+	/** Color transform from YUV to sRGB (using values from PS4 AvPlayer codec). */
+	RENDERCORE_API extern const FMatrix YuvToSrgbPs4;
+
+	/** Color transform from YUV to sRGB (in Rec. 601 color space). */
+	RENDERCORE_API extern const FMatrix YuvToSrgbRec601;
+
+	/** Color transform from YUV to sRGB (in Rec. 709 color space). */
+	RENDERCORE_API extern const FMatrix YuvToRgbRec709;
+
+	/** Color transform from YUV to RGB (in Rec. 709 color space, RGB full range) */
+	RENDERCORE_API extern const FMatrix YuvToRgbRec709Full;
+
+	/** Color transform from RGB to YUV (in Rec. 709 color space, RGB full range) */
+	RENDERCORE_API extern const FMatrix RgbToYuvRec709Full;
+
+	/** YUV Offset for 8 bit conversion (Computed as 16/255, 128/255, 128/255) */
+	RENDERCORE_API extern const FVector YUVOffset8bits;
+
+	/** YUV Offset for 10 bit conversion (Computed as 64/1023, 512/1023, 512/1023) */
+	RENDERCORE_API extern const FVector YUVOffset10bits;
+
+	/** Combine color transform matrix with yuv offset in a single matrix */
+	RENDERCORE_API FMatrix CombineColorTransformAndOffset(const FMatrix& InMatrix, const FVector& InYUVOffset);
+}
+
+
+/**
+ * Stores media drawing vertices.
+ */
+struct FMediaElementVertex
+{
+	FVector4 Position;
+	FVector2D TextureCoordinate;
+
+	FMediaElementVertex() { }
+
+	FMediaElementVertex(const FVector4& InPosition, const FVector2D& InTextureCoordinate)
+		: Position(InPosition)
+		, TextureCoordinate(InTextureCoordinate)
+	{ }
+};
+
+inline FVertexBufferRHIRef CreateTempMediaVertexBuffer(float ULeft = 0.0f, float URight = 1.0f, float VTop = 0.0f, float VBottom = 1.0f)
+{
+	FRHIResourceCreateInfo CreateInfo;
+	FVertexBufferRHIRef VertexBufferRHI = RHICreateVertexBuffer(sizeof(FMediaElementVertex) * 4, BUF_Volatile, CreateInfo);
+	void* VoidPtr = RHILockVertexBuffer(VertexBufferRHI, 0, sizeof(FMediaElementVertex) * 4, RLM_WriteOnly);
+
+	FMediaElementVertex* Vertices = (FMediaElementVertex*)VoidPtr;
+	Vertices[0].Position.Set(-1.0f, 1.0f, 1.0f, 1.0f); // Top Left
+	Vertices[1].Position.Set(1.0f, 1.0f, 1.0f, 1.0f); // Top Right
+	Vertices[2].Position.Set(-1.0f, -1.0f, 1.0f, 1.0f); // Bottom Left
+	Vertices[3].Position.Set(1.0f, -1.0f, 1.0f, 1.0f); // Bottom Right
+
+	Vertices[0].TextureCoordinate.Set(ULeft, VTop);
+	Vertices[1].TextureCoordinate.Set(URight, VTop);
+	Vertices[2].TextureCoordinate.Set(ULeft, VBottom);
+	Vertices[3].TextureCoordinate.Set(URight, VBottom);
+	RHIUnlockVertexBuffer(VertexBufferRHI);
+
+	return VertexBufferRHI;
+}
+
+/**
+ * The simple element vertex declaration resource type.
+ */
+class FMediaVertexDeclaration
+	: public FRenderResource
+{
+public:
+
+	FVertexDeclarationRHIRef VertexDeclarationRHI;
+
+	virtual ~FMediaVertexDeclaration() { }
+
+	virtual void InitRHI() override
+	{
+		FVertexDeclarationElementList Elements;
+		uint16 Stride = sizeof(FMediaElementVertex);
+		Elements.Add(FVertexElement(0, STRUCT_OFFSET(FMediaElementVertex, Position), VET_Float4, 0, Stride));
+		Elements.Add(FVertexElement(0, STRUCT_OFFSET(FMediaElementVertex, TextureCoordinate), VET_Float2, 1, Stride));
+		VertexDeclarationRHI = PipelineStateCache::GetOrCreateVertexDeclaration(Elements);
+	}
+
+	virtual void ReleaseRHI() override
+	{
+		VertexDeclarationRHI.SafeRelease();
+	}
+};
+
+
+RENDERCORE_API extern TGlobalResource<FMediaVertexDeclaration> GMediaVertexDeclaration;
+
+
+/**
+ * Media vertex shader (shared by all media shaders).
+ */
+class FMediaShadersVS
+	: public FGlobalShader
+{
+	DECLARE_EXPORTED_SHADER_TYPE(FMediaShadersVS, Global, RENDERCORE_API);
+
+public:
+
+	static bool ShouldCompilePermutation(const FGlobalShaderPermutationParameters& Parameters)
+	{
+		return IsFeatureLevelSupported(Parameters.Platform, ERHIFeatureLevel::ES2);
+	}
+
+	/** Default constructor. */
+	FMediaShadersVS() { }
+
+	/** Initialization constructor. */
+	FMediaShadersVS(const ShaderMetaType::CompiledShaderInitializerType& Initializer)
+		: FGlobalShader(Initializer)
+	{ }
+
+	virtual bool Serialize(FArchive& Ar) override
+	{
+		bool bShaderHasOutdatedParameters = FGlobalShader::Serialize(Ar);
+		return bShaderHasOutdatedParameters;
+	}
+};
+
+
+/**
+ * Pixel shader to convert an AYUV texture to RGBA.
+ *
+ * This shader expects a single texture consisting of a N x M array of pixels
+ * in AYUV format. Each pixel is encoded as four consecutive unsigned chars
+ * with the following layout: [V0 U0 Y0 A0][V1 U1 Y1 A1]..
+ */
+class FAYUVConvertPS
+	: public FGlobalShader
+{
+	DECLARE_EXPORTED_SHADER_TYPE(FAYUVConvertPS, Global, RENDERCORE_API);
+
+public:
+
+	static bool ShouldCompilePermutation(const FGlobalShaderPermutationParameters& Parameters)
+	{
+		return IsFeatureLevelSupported(Parameters.Platform, ERHIFeatureLevel::ES2);
+	}
+
+	FAYUVConvertPS() { }
+
+	FAYUVConvertPS(const ShaderMetaType::CompiledShaderInitializerType& Initializer)
+		: FGlobalShader(Initializer)
+	{ }
+
+	virtual bool Serialize(FArchive& Ar) override
+	{
+		bool bShaderHasOutdatedParameters = FGlobalShader::Serialize(Ar);
+		return bShaderHasOutdatedParameters;
+	}
+
+	RENDERCORE_API void SetParameters(FRHICommandList& RHICmdList, TRefCountPtr<FRHITexture2D> AYUVTexture, const FMatrix& ColorTransform, const FVector& YUVOffset, bool SrgbToLinear);
+};
+
+
+/**
+ * Pixel shader to convert a Windows Bitmap texture.
+ *
+ * This shader expects a BMP frame packed into a single texture in PF_B8G8R8A8 format.
+ */
+class FBMPConvertPS
+	: public FGlobalShader
+{
+	DECLARE_EXPORTED_SHADER_TYPE(FBMPConvertPS, Global, RENDERCORE_API);
+
+public:
+
+	static bool ShouldCompilePermutation(const FGlobalShaderPermutationParameters& Parameters)
+	{
+		return IsFeatureLevelSupported(Parameters.Platform, ERHIFeatureLevel::ES2);
+	}
+
+	FBMPConvertPS() { }
+
+	FBMPConvertPS(const ShaderMetaType::CompiledShaderInitializerType& Initializer)
+		: FGlobalShader(Initializer)
+	{ }
+
+	virtual bool Serialize(FArchive& Ar) override
+	{
+		bool bShaderHasOutdatedParameters = FGlobalShader::Serialize(Ar);
+		return bShaderHasOutdatedParameters;
+	}
+
+	RENDERCORE_API void SetParameters(FRHICommandList& RHICmdList, TRefCountPtr<FRHITexture2D> BMPTexture, const FIntPoint& OutputDimensions, bool SrgbToLinear);
+};
+
+
+/**
+ * Pixel shader to convert a NV12 frame to RGBA.
+ *
+ * This shader expects an NV12 frame packed into a single texture in PF_G8 format.
+ *
+ * @see http://www.fourcc.org/yuv.php#NV12
+ */
+class FNV12ConvertPS
+	: public FGlobalShader
+{
+	DECLARE_EXPORTED_SHADER_TYPE(FNV12ConvertPS, Global, RENDERCORE_API);
+
+public:
+
+	static bool ShouldCompilePermutation(const FGlobalShaderPermutationParameters& Parameters)
+	{
+		return IsFeatureLevelSupported(Parameters.Platform, ERHIFeatureLevel::ES2);
+	}
+
+	FNV12ConvertPS() { }
+
+	FNV12ConvertPS(const ShaderMetaType::CompiledShaderInitializerType& Initializer)
+		: FGlobalShader(Initializer)
+	{ }
+
+	virtual bool Serialize(FArchive& Ar) override
+	{
+		bool bShaderHasOutdatedParameters = FGlobalShader::Serialize(Ar);
+		return bShaderHasOutdatedParameters;
+	}
+
+	RENDERCORE_API void SetParameters(FRHICommandList& RHICmdList, TRefCountPtr<FRHITexture2D> NV12Texture, const FIntPoint& OutputDimensions, const FMatrix& ColorTransform, const FVector& YUVOffset, bool SrgbToLinear);
+};
+
+
+/**
+ * Pixel shader to convert a NV21 frame to RGBA.
+ *
+ * This shader expects an NV21 frame packed into a single texture in PF_G8 format.
+ *
+ * @see http://www.fourcc.org/yuv.php#NV21
+ */
+class FNV21ConvertPS
+	: public FGlobalShader
+{
+	DECLARE_EXPORTED_SHADER_TYPE(FNV21ConvertPS, Global, RENDERCORE_API);
+
+public:
+
+	static bool ShouldCompilePermutation(const FGlobalShaderPermutationParameters& Parameters)
+	{
+		return IsFeatureLevelSupported(Parameters.Platform, ERHIFeatureLevel::ES2);
+	}
+
+	FNV21ConvertPS() { }
+
+	FNV21ConvertPS(const ShaderMetaType::CompiledShaderInitializerType& Initializer)
+		: FGlobalShader(Initializer)
+	{ }
+
+	virtual bool Serialize(FArchive& Ar) override
+	{
+		bool bShaderHasOutdatedParameters = FGlobalShader::Serialize(Ar);
+		return bShaderHasOutdatedParameters;
+	}
+
+	RENDERCORE_API void SetParameters(FRHICommandList& RHICmdList, TRefCountPtr<FRHITexture2D> NV21Texture, const FIntPoint& OutputDimensions, const FMatrix& ColorTransform, const FVector& YUVOffset, bool SrgbToLinear);
+};
+
+
+/**
+ * Pixel shader to resize an RGB texture.
+ *
+ * This shader expects an RGB or RGBA frame packed into a single texture
+ * in PF_B8G8R8A8 or PF_FloatRGB format.
+ */
+class FRGBConvertPS
+	: public FGlobalShader
+{
+	DECLARE_EXPORTED_SHADER_TYPE(FRGBConvertPS, Global, RENDERCORE_API);
+
+public:
+
+	static bool ShouldCompilePermutation(const FGlobalShaderPermutationParameters& Parameters)
+	{
+		return IsFeatureLevelSupported(Parameters.Platform, ERHIFeatureLevel::ES2);
+	}
+
+	FRGBConvertPS() { }
+
+	FRGBConvertPS(const ShaderMetaType::CompiledShaderInitializerType& Initializer)
+		: FGlobalShader(Initializer)
+	{ }
+
+	virtual bool Serialize(FArchive& Ar) override
+	{
+		bool bShaderHasOutdatedParameters = FGlobalShader::Serialize(Ar);
+		return bShaderHasOutdatedParameters;
+	}
+
+	RENDERCORE_API void SetParameters(FRHICommandList& RHICmdList, TRefCountPtr<FRHITexture2D> RGBTexture, const FIntPoint& OutputDimensions, bool SrgbToLinear);
+};
+
+
+/**
+ * Pixel shader to convert a PS4 YCbCr texture to RGBA.
+ *
+ * This shader expects a separate chroma and luma plane stored in two textures
+ * in PF_B8G8R8A8 format. The full-size luma plane contains the Y-components.
+ * The half-size chroma plane contains the UV components in the following
+ * memory layout: [U0, V0][U1, V1]
+ * 
+ */
+class FYCbCrConvertPS
+	: public FGlobalShader
+{
+	DECLARE_EXPORTED_SHADER_TYPE(FYCbCrConvertPS, Global, RENDERCORE_API);
+
+public:
+
+	static bool ShouldCompilePermutation(const FGlobalShaderPermutationParameters& Parameters)
+	{
+		return IsFeatureLevelSupported(Parameters.Platform, ERHIFeatureLevel::ES2);
+	}
+
+	FYCbCrConvertPS() { }
+
+	FYCbCrConvertPS(const ShaderMetaType::CompiledShaderInitializerType& Initializer)
+		: FGlobalShader(Initializer)
+	{ }
+
+	virtual bool Serialize(FArchive& Ar) override
+	{
+		bool bShaderHasOutdatedParameters = FGlobalShader::Serialize(Ar);		
+		return bShaderHasOutdatedParameters;
+	}
+
+	RENDERCORE_API void SetParameters(FRHICommandList& RHICmdList, TRefCountPtr<FRHITexture2D> LumaTexture, TRefCountPtr<FRHITexture2D> CbCrTexture, const FMatrix& ColorTransform, const FVector& YUVOffset, bool SrgbToLinear);
+};
+
+class FYCbCrConvertPS_4x4Matrix : public FYCbCrConvertPS
+{
+    DECLARE_EXPORTED_SHADER_TYPE(FYCbCrConvertPS_4x4Matrix, Global, RENDERCORE_API);
+    
+public:
+    
+    static bool ShouldCompilePermutation(const FGlobalShaderPermutationParameters& Parameters)
+    {
+        return IsFeatureLevelSupported(Parameters.Platform, ERHIFeatureLevel::ES2);
+    }
+    
+    FYCbCrConvertPS_4x4Matrix() { }
+    
+    FYCbCrConvertPS_4x4Matrix(const ShaderMetaType::CompiledShaderInitializerType& Initializer)
+    : FYCbCrConvertPS(Initializer)
+    { }
+};
+
+
+/**
+ * Pixel shader to convert a UYVY (Y422, UYNV) frame to RGBA.
+ *
+ * This shader expects a UYVY frame packed into a single texture in PF_B8G8R8A8
+ * format with the following memory layout: [U0, Y0, V1, Y1][U1, Y2, V1, Y3]..
+ *
+ * @see http://www.fourcc.org/yuv.php#UYVY
+ */
+class FUYVYConvertPS
+	: public FGlobalShader
+{
+	DECLARE_EXPORTED_SHADER_TYPE(FUYVYConvertPS, Global, RENDERCORE_API);
+
+public:
+
+	static bool ShouldCompilePermutation(const FGlobalShaderPermutationParameters& Parameters)
+	{
+		return IsFeatureLevelSupported(Parameters.Platform, ERHIFeatureLevel::ES2);
+	}
+
+	FUYVYConvertPS() { }
+
+	FUYVYConvertPS(const ShaderMetaType::CompiledShaderInitializerType& Initializer)
+		: FGlobalShader(Initializer)
+	{ }
+
+	virtual bool Serialize(FArchive& Ar) override
+	{
+		bool bShaderHasOutdatedParameters = FGlobalShader::Serialize(Ar);
+		return bShaderHasOutdatedParameters;
+	}
+
+	RENDERCORE_API void SetParameters(FRHICommandList& RHICmdList, TRefCountPtr<FRHITexture2D> UYVYTexture, const FMatrix& ColorTransform, const FVector& YUVOffset, bool SrgbToLinear);
+};
+
+
+/**
+ * Pixel shader to convert Y, U, and V planes to RGBA.
+ *
+ * This shader expects three textures in PF_G8 format,
+ * one for each plane of Y, U, and V components.
+ */
+class FYUVConvertPS
+	: public FGlobalShader
+{
+	DECLARE_EXPORTED_SHADER_TYPE(FYUVConvertPS, Global, RENDERCORE_API);
+
+public:
+
+	static bool ShouldCompilePermutation(const FGlobalShaderPermutationParameters& Parameters)
+	{
+		return IsFeatureLevelSupported(Parameters.Platform, ERHIFeatureLevel::ES2);
+	}
+
+	FYUVConvertPS() { }
+
+	FYUVConvertPS(const ShaderMetaType::CompiledShaderInitializerType& Initializer)
+		: FGlobalShader(Initializer)
+	{ }
+
+	virtual bool Serialize(FArchive& Ar) override
+	{
+		bool bShaderHasOutdatedParameters = FGlobalShader::Serialize(Ar);
+		return bShaderHasOutdatedParameters;
+	}
+
+	RENDERCORE_API void SetParameters(FRHICommandList& RHICmdList, TRefCountPtr<FRHITexture2D> YTexture, TRefCountPtr<FRHITexture2D> UTexture, TRefCountPtr<FRHITexture2D> VTexture, const FIntPoint& OutputDimensions, const FMatrix& ColorTransform, const FVector& YUVOffset, bool SrgbToLinear);
+};
+
+
+/**
+ * Pixel shader to convert YUV v210 to RGB
+ *
+ * This shader expects a single texture in PF_R32G32B32A32_UINT format.
+ */
+class FYUVv210ConvertPS
+	: public FGlobalShader
+{
+	DECLARE_EXPORTED_SHADER_TYPE(FYUVv210ConvertPS, Global, RENDERCORE_API);
+
+public:
+
+	static bool ShouldCompilePermutation(const FGlobalShaderPermutationParameters& Parameters)
+	{
+		return IsFeatureLevelSupported(Parameters.Platform, ERHIFeatureLevel::ES2);
+	}
+
+	FYUVv210ConvertPS() { }
+
+	FYUVv210ConvertPS(const ShaderMetaType::CompiledShaderInitializerType& Initializer)
+		: FGlobalShader(Initializer)
+	{ }
+
+	virtual bool Serialize(FArchive& Ar) override
+	{
+		bool bShaderHasOutdatedParameters = FGlobalShader::Serialize(Ar);
+		return bShaderHasOutdatedParameters;
+	}
+
+	RENDERCORE_API void SetParameters(FRHICommandList& RHICmdList, TRefCountPtr<FRHITexture2D> YUVTexture, const FIntPoint& OutputDimensions, const FMatrix& ColorTransform, const FVector& YUVOffset, bool SrgbToLinear);
+};
+
+
+
+/**
+ * Pixel shader to convert a YUY2 frame to RGBA.
+ *
+ * This shader expects an YUY2 frame packed into a single texture in PF_B8G8R8A8
+ * format with the following memory layout: [Y0, U0, Y1, V0][Y2, U1, Y3, V1]...
+ *
+ * @see http://www.fourcc.org/yuv.php#YUY2
+ */
+class FYUY2ConvertPS
+	: public FGlobalShader
+{
+	DECLARE_EXPORTED_SHADER_TYPE(FYUY2ConvertPS, Global, RENDERCORE_API);
+
+public:
+
+	static bool ShouldCompilePermutation(const FGlobalShaderPermutationParameters& Parameters)
+	{
+		return IsFeatureLevelSupported(Parameters.Platform, ERHIFeatureLevel::ES2);
+	}
+
+	FYUY2ConvertPS() { }
+
+	FYUY2ConvertPS(const ShaderMetaType::CompiledShaderInitializerType& Initializer)
+		: FGlobalShader(Initializer)
+	{ }
+
+	virtual bool Serialize(FArchive& Ar) override
+	{
+		bool bShaderHasOutdatedParameters = FGlobalShader::Serialize(Ar);
+		return bShaderHasOutdatedParameters;
+	}
+
+	RENDERCORE_API void SetParameters(FRHICommandList& RHICmdList, TRefCountPtr<FRHITexture2D> YUY2Texture, const FIntPoint& OutputDimensions, const FMatrix& ColorTransform, const FVector& YUVOffset, bool SrgbToLinear);
+};
+
+
+/**
+ * Pixel shader to convert a YVYU frame to RGBA.
+ *
+ * This shader expects a YVYU frame packed into a single texture in PF_B8G8R8A8
+ * format with the following memory layout: [Y0, V0, Y1, U0][Y2, V1, Y3, U1]..
+ *
+ * @see http://www.fourcc.org/yuv.php#YVYU
+ */
+class FYVYUConvertPS
+	: public FGlobalShader
+{
+	DECLARE_EXPORTED_SHADER_TYPE(FYVYUConvertPS, Global, RENDERCORE_API);
+
+public:
+
+	static bool ShouldCompilePermutation(const FGlobalShaderPermutationParameters& Parameters)
+	{
+		return IsFeatureLevelSupported(Parameters.Platform, ERHIFeatureLevel::ES2);
+	}
+
+	FYVYUConvertPS() { }
+
+	FYVYUConvertPS(const ShaderMetaType::CompiledShaderInitializerType& Initializer)
+		: FGlobalShader(Initializer)
+	{ }
+
+	virtual bool Serialize(FArchive& Ar) override
+	{
+		bool bShaderHasOutdatedParameters = FGlobalShader::Serialize(Ar);
+		return bShaderHasOutdatedParameters;
+	}
+
+	RENDERCORE_API void SetParameters(FRHICommandList& RHICmdList, TRefCountPtr<FRHITexture2D> YVYUTexture, const FMatrix& ColorTransform, const FVector& YUVOffset, bool SrgbToLinear);
+};
+
+
+/**
+ * Pixel shader to convert RGB 8 bits to UYVY 8 bits
+ *
+ * This shader expects a single texture in PF_B8G8R8A8 format.
+ */
+class FRGB8toUYVY8ConvertPS
+	: public FGlobalShader
+{
+	DECLARE_EXPORTED_SHADER_TYPE(FRGB8toUYVY8ConvertPS, Global, RENDERCORE_API);
+
+public:
+
+	static bool ShouldCompilePermutation(const FGlobalShaderPermutationParameters& Parameters)
+	{
+		return IsFeatureLevelSupported(Parameters.Platform, ERHIFeatureLevel::ES2);
+	}
+
+	FRGB8toUYVY8ConvertPS() { }
+
+	FRGB8toUYVY8ConvertPS(const ShaderMetaType::CompiledShaderInitializerType& Initializer)
+		: FGlobalShader(Initializer)
+	{ }
+
+	virtual bool Serialize(FArchive& Ar) override
+	{
+		bool bShaderHasOutdatedParameters = FGlobalShader::Serialize(Ar);
+		return bShaderHasOutdatedParameters;
+	}
+
+	RENDERCORE_API void SetParameters(FRHICommandList& RHICmdList, TRefCountPtr<FRHITexture2D> RGBATexture, const FMatrix& ColorTransform, const FVector& YUVOffset, bool LinearToSrgb);
+};
+
+
+/**
+ * Pixel shader to convert RGB 10 bits to YUV v210
+ *
+ * This shader expects a single texture in PF_A2B10G10R10 format.
+ */
+class FRGB10toYUVv210ConvertPS
+	: public FGlobalShader
+{
+	DECLARE_EXPORTED_SHADER_TYPE(FRGB10toYUVv210ConvertPS, Global, RENDERCORE_API);
+
+public:
+
+	static bool ShouldCompilePermutation(const FGlobalShaderPermutationParameters& Parameters)
+	{
+		return IsFeatureLevelSupported(Parameters.Platform, ERHIFeatureLevel::ES2);
+	}
+
+	FRGB10toYUVv210ConvertPS() { }
+
+	FRGB10toYUVv210ConvertPS(const ShaderMetaType::CompiledShaderInitializerType& Initializer)
+		: FGlobalShader(Initializer)
+	{ }
+
+	virtual bool Serialize(FArchive& Ar) override
+	{
+		bool bShaderHasOutdatedParameters = FGlobalShader::Serialize(Ar);
+		return bShaderHasOutdatedParameters;
+	}
+
+	RENDERCORE_API void SetParameters(FRHICommandList& RHICmdList, TRefCountPtr<FRHITexture2D> RGBATexture, const FMatrix& ColorTransform, const FVector& YUVOffset, bool LinearToSrgb);
+};
+
+
+/**
+ * Pixel shader to inverse alpha component 
+ *
+ * This shader expects a single texture in RGBA 8 or 10 bits format.
+ */
+class FInvertAlphaPS
+	: public FGlobalShader
+{
+	DECLARE_EXPORTED_SHADER_TYPE(FInvertAlphaPS, Global, RENDERCORE_API);
+
+public:
+
+	static bool ShouldCompilePermutation(const FGlobalShaderPermutationParameters& Parameters)
+	{
+		return IsFeatureLevelSupported(Parameters.Platform, ERHIFeatureLevel::ES2);
+	}
+
+	FInvertAlphaPS() { }
+
+	FInvertAlphaPS(const ShaderMetaType::CompiledShaderInitializerType& Initializer)
+		: FGlobalShader(Initializer)
+	{ }
+
+	RENDERCORE_API void SetParameters(FRHICommandList& RHICmdList, TRefCountPtr<FRHITexture2D> RGBATexture);
+};
+
+
+/**
+ * Pixel shader to set alpha component to 1.0f
+ *
+ * This shader expects a single texture in RGBA 8 or 10 bits format.
+ */
+class FSetAlphaOnePS
+	: public FGlobalShader
+{
+	DECLARE_EXPORTED_SHADER_TYPE(FSetAlphaOnePS, Global, RENDERCORE_API);
+
+public:
+
+	static bool ShouldCompilePermutation(const FGlobalShaderPermutationParameters& Parameters)
+	{
+		return IsFeatureLevelSupported(Parameters.Platform, ERHIFeatureLevel::ES2);
+	}
+
+	FSetAlphaOnePS() { }
+
+	FSetAlphaOnePS(const ShaderMetaType::CompiledShaderInitializerType& Initializer)
+		: FGlobalShader(Initializer)
+	{ }
+
+	RENDERCORE_API void SetParameters(FRHICommandList& RHICmdList, TRefCountPtr<FRHITexture2D> RGBATexture);
+};
diff --git a/Engine/Source/Runtime/RenderCore/Public/OculusShaders.h b/Engine/Source/Runtime/RenderCore/Public/OculusShaders.h
new file mode 100644
--- /dev/null
+++ b/Engine/Source/Runtime/RenderCore/Public/OculusShaders.h
@@ -0,0 +1,138 @@
+// Copyright 1998-2019 Epic Games, Inc. All Rights Reserved.
+
+#pragma once
+
+#include "CoreMinimal.h"
+#include "ShaderParameters.h"
+#include "Shader.h"
+#include "GlobalShader.h"
+#include "ShaderParameterUtils.h"
+#include "RenderResource.h"
+
+class FOculusVertexShader : public FGlobalShader
+{
+	DECLARE_EXPORTED_SHADER_TYPE(FOculusVertexShader, Global, RENDERCORE_API);
+public:
+
+	static bool ShouldCompilePermutation(const FGlobalShaderPermutationParameters& Parameters) { return true; }
+
+	FOculusVertexShader(const ShaderMetaType::CompiledShaderInitializerType& Initializer) :
+		FGlobalShader(Initializer)
+	{}
+	FOculusVertexShader() {}
+};
+
+class FOculusWhiteShader : public FGlobalShader
+{
+	DECLARE_EXPORTED_SHADER_TYPE(FOculusWhiteShader, Global, RENDERCORE_API);
+public:
+
+	static bool ShouldCompilePermutation(const FGlobalShaderPermutationParameters& Parameters) { return true; }
+
+	FOculusWhiteShader(const ShaderMetaType::CompiledShaderInitializerType& Initializer) :
+		FGlobalShader(Initializer)
+	{
+	}
+	
+	FOculusWhiteShader() {}
+};
+
+class FOculusBlackShader : public FGlobalShader
+{
+	DECLARE_EXPORTED_SHADER_TYPE(FOculusBlackShader, Global, RENDERCORE_API);
+public:
+
+	static bool ShouldCompilePermutation(const FGlobalShaderPermutationParameters& Parameters) { return true; }
+
+	FOculusBlackShader(const ShaderMetaType::CompiledShaderInitializerType& Initializer) :
+		FGlobalShader(Initializer)
+	{
+	}
+
+	FOculusBlackShader() {}
+};
+
+class FOculusAlphaInverseShader : public FGlobalShader
+{
+	DECLARE_EXPORTED_SHADER_TYPE(FOculusAlphaInverseShader, Global, RENDERCORE_API);
+public:
+
+	static bool ShouldCompilePermutation(const FGlobalShaderPermutationParameters& Parameters) { return true; }
+
+	FOculusAlphaInverseShader(const ShaderMetaType::CompiledShaderInitializerType& Initializer) :
+		FGlobalShader(Initializer)
+	{
+		InTexture.Bind(Initializer.ParameterMap, TEXT("InTexture"), SPF_Mandatory);
+		InTextureSampler.Bind(Initializer.ParameterMap, TEXT("InTextureSampler"));
+	}
+	FOculusAlphaInverseShader() {}
+
+	RENDERCORE_API void SetParameters(FRHICommandList& RHICmdList, const FTexture* Texture)
+	{
+		SetTextureParameter(RHICmdList, GetPixelShader(), InTexture, InTextureSampler, Texture);
+	}
+
+	RENDERCORE_API void SetParameters(FRHICommandList& RHICmdList, FRHISamplerState* SamplerStateRHI, FRHITexture* TextureRHI)
+	{
+		SetTextureParameter(RHICmdList, GetPixelShader(), InTexture, InTextureSampler, SamplerStateRHI, TextureRHI);
+	}
+
+	virtual bool Serialize(FArchive& Ar) override
+	{
+		bool bShaderHasOutdatedParameters = FGlobalShader::Serialize(Ar);
+		Ar << InTexture;
+		Ar << InTextureSampler;
+		return bShaderHasOutdatedParameters;
+	}
+
+private:
+	FShaderResourceParameter InTexture;
+	FShaderResourceParameter InTextureSampler;
+};
+
+
+/**
+* A pixel shader for rendering a textured screen element.
+*/
+class FOculusCubemapPS : public FGlobalShader
+{
+	DECLARE_EXPORTED_SHADER_TYPE(FOculusCubemapPS, Global, RENDERCORE_API);
+public:
+
+	static bool ShouldCompilePermutation(const FGlobalShaderPermutationParameters& Parameters) { return true; }
+
+	FOculusCubemapPS(const ShaderMetaType::CompiledShaderInitializerType& Initializer) :
+		FGlobalShader(Initializer)
+	{
+		InTexture.Bind(Initializer.ParameterMap, TEXT("InTextureCube"), SPF_Mandatory);
+		InTextureSampler.Bind(Initializer.ParameterMap, TEXT("InTextureSampler"));
+		InFaceIndexParameter.Bind(Initializer.ParameterMap, TEXT("CubeFaceIndex"));
+	}
+	FOculusCubemapPS() {}
+
+	void SetParameters(FRHICommandList& RHICmdList, const FTexture* Texture, int FaceIndex)
+	{
+		SetTextureParameter(RHICmdList, GetPixelShader(), InTexture, InTextureSampler, Texture);
+		SetShaderValue(RHICmdList, GetPixelShader(), InFaceIndexParameter, FaceIndex);
+	}
+
+	void SetParameters(FRHICommandList& RHICmdList, FRHISamplerState* SamplerStateRHI, FRHITexture* TextureRHI, int FaceIndex)
+	{
+		SetTextureParameter(RHICmdList, GetPixelShader(), InTexture, InTextureSampler, SamplerStateRHI, TextureRHI);
+		SetShaderValue(RHICmdList, GetPixelShader(), InFaceIndexParameter, FaceIndex);
+	}
+
+	virtual bool Serialize(FArchive& Ar) override
+	{
+		bool bShaderHasOutdatedParameters = FGlobalShader::Serialize(Ar);
+		Ar << InTexture;
+		Ar << InTextureSampler;
+		Ar << InFaceIndexParameter;
+		return bShaderHasOutdatedParameters;
+	}
+
+private:
+	FShaderResourceParameter InTexture;
+	FShaderResourceParameter InTextureSampler;
+	FShaderParameter		 InFaceIndexParameter;
+};
\ No newline at end of file
diff --git a/Engine/Source/Runtime/RenderCore/Public/OneColorShader.h b/Engine/Source/Runtime/RenderCore/Public/OneColorShader.h
new file mode 100644
--- /dev/null
+++ b/Engine/Source/Runtime/RenderCore/Public/OneColorShader.h
@@ -0,0 +1,195 @@
+// Copyright 1998-2019 Epic Games, Inc. All Rights Reserved.
+
+
+#pragma once
+
+#include "CoreMinimal.h"
+#include "ShaderParameters.h"
+#include "Shader.h"
+#include "GlobalShader.h"
+#include "ShaderParameterUtils.h"
+
+/**
+ * Vertex shader for rendering a single, constant color.
+ */
+template<bool bUsingNDCPositions=true, bool bUsingVertexLayers=false>
+class TOneColorVS : public FGlobalShader
+{
+	DECLARE_EXPORTED_SHADER_TYPE(TOneColorVS, Global, RENDERCORE_API);
+
+	/** Default constructor. */
+	TOneColorVS() {}
+
+public:
+
+	TOneColorVS(const ShaderMetaType::CompiledShaderInitializerType& Initializer)
+		: FGlobalShader(Initializer)
+	{
+		DepthParameter.Bind(Initializer.ParameterMap, TEXT("InputDepth"), SPF_Mandatory);
+	}
+
+	static void ModifyCompilationEnvironment(const FGlobalShaderPermutationParameters& Parameters, FShaderCompilerEnvironment& OutEnvironment)
+	{
+		FGlobalShader::ModifyCompilationEnvironment(Parameters, OutEnvironment);
+		OutEnvironment.SetDefine(TEXT("USING_NDC_POSITIONS"), (uint32)(bUsingNDCPositions ? 1 : 0));
+		OutEnvironment.SetDefine(TEXT("USING_LAYERS"), (uint32)(bUsingVertexLayers ? 1 : 0));
+	}
+
+	void SetDepthParameter(FRHICommandList& RHICmdList, float Depth)
+	{
+		SetShaderValue(RHICmdList, GetVertexShader(), DepthParameter, Depth);
+	}
+
+	// FShader interface.
+	virtual bool Serialize(FArchive& Ar) override
+	{
+		bool bShaderHasOutdatedParameters = FGlobalShader::Serialize(Ar);
+		Ar << DepthParameter;
+		return bShaderHasOutdatedParameters;
+	}
+
+	static bool ShouldCompilePermutation(const FGlobalShaderPermutationParameters& Parameters)
+	{
+		return true;
+	}
+	
+	static const TCHAR* GetSourceFilename()
+	{
+		return TEXT("/Engine/Private/OneColorShader.usf");
+	}
+	
+	static const TCHAR* GetFunctionName()
+	{
+		return TEXT("MainVertexShader");
+	}
+
+private:
+	FShaderParameter DepthParameter;
+};
+
+/**
+ * Pixel shader for rendering a single, constant color.
+ */
+class FOneColorPS : public FGlobalShader
+{
+	DECLARE_EXPORTED_SHADER_TYPE(FOneColorPS,Global,RENDERCORE_API);
+public:
+	
+	FOneColorPS( )	{ }
+	FOneColorPS(const ShaderMetaType::CompiledShaderInitializerType& Initializer)
+	:	FGlobalShader( Initializer )
+	{
+		//ColorParameter.Bind( Initializer.ParameterMap, TEXT("DrawColorMRT"), SPF_Mandatory);
+	}
+
+	virtual void SetColors(FRHICommandList& RHICmdList, const FLinearColor* Colors, int32 NumColors);	
+
+	// FShader interface.
+	virtual bool Serialize(FArchive& Ar) override
+	{
+		bool bShaderHasOutdatedParameters = FGlobalShader::Serialize(Ar);
+		//Ar << ColorParameter;
+		return bShaderHasOutdatedParameters;
+	}
+	static bool ShouldCompilePermutation(const FGlobalShaderPermutationParameters& Parameters)
+	{
+		return true;
+	}
+
+	/** The parameter to use for setting the draw Color. */
+	//FShaderParameter ColorParameter;
+};
+
+/**
+ * Pixel shader for rendering a single, constant color to MRTs.
+ */
+template<int32 NumOutputs>
+class TOneColorPixelShaderMRT : public FOneColorPS
+{
+	DECLARE_EXPORTED_SHADER_TYPE(TOneColorPixelShaderMRT,Global,RENDERCORE_API);
+public:
+	TOneColorPixelShaderMRT( )	{ }
+	TOneColorPixelShaderMRT(const ShaderMetaType::CompiledShaderInitializerType& Initializer)
+	:	FOneColorPS( Initializer )
+	{
+	}
+
+	static bool ShouldCompilePermutation(const FGlobalShaderPermutationParameters& Parameters)
+	{
+		if( NumOutputs > 1 )
+		{
+			return IsFeatureLevelSupported( Parameters.Platform, ERHIFeatureLevel::ES3_1);
+		}
+		return true;
+	}
+
+	static void ModifyCompilationEnvironment(const FGlobalShaderPermutationParameters& Parameters, FShaderCompilerEnvironment& OutEnvironment)
+	{
+		FOneColorPS::ModifyCompilationEnvironment(Parameters, OutEnvironment);
+		OutEnvironment.SetDefine(TEXT("NUM_OUTPUTS"), NumOutputs);
+	}
+};
+
+/**
+ * Compute shader for writing values
+ */
+class FFillTextureCS : public FGlobalShader
+{
+	DECLARE_EXPORTED_SHADER_TYPE(FFillTextureCS,Global,RENDERCORE_API);
+public:
+	FFillTextureCS( )	{ }
+	FFillTextureCS(const ShaderMetaType::CompiledShaderInitializerType& Initializer)
+	:	FGlobalShader( Initializer )
+	{
+		FillValue.Bind( Initializer.ParameterMap, TEXT("FillValue"), SPF_Mandatory);
+		Params0.Bind( Initializer.ParameterMap, TEXT("Params0"), SPF_Mandatory);
+		Params1.Bind( Initializer.ParameterMap, TEXT("Params1"), SPF_Mandatory);
+		Params2.Bind( Initializer.ParameterMap, TEXT("Params2"), SPF_Optional);
+		FillTexture.Bind( Initializer.ParameterMap, TEXT("FillTexture"), SPF_Mandatory);
+	}
+
+	// FShader interface.
+	virtual bool Serialize(FArchive& Ar) override
+	{
+		bool bShaderHasOutdatedParameters = FGlobalShader::Serialize(Ar);
+		Ar << FillValue << Params0 << Params1 << Params2 << FillTexture;
+		return bShaderHasOutdatedParameters;
+	}
+
+	static bool ShouldCompilePermutation(const FGlobalShaderPermutationParameters& Parameters)
+	{
+		return IsFeatureLevelSupported(Parameters.Platform, ERHIFeatureLevel::SM5);
+	}
+
+	FShaderParameter FillValue;
+	FShaderParameter Params0;	// Texture Width,Height (.xy); Use Exclude Rect 1 : 0 (.z)
+	FShaderParameter Params1;	// Include X0,Y0 (.xy) - X1,Y1 (.zw)
+	FShaderParameter Params2;	// ExcludeRect X0,Y0 (.xy) - X1,Y1 (.zw)
+	FShaderResourceParameter FillTexture;
+};
+
+class FLongGPUTaskPS : public FGlobalShader
+{
+	DECLARE_EXPORTED_SHADER_TYPE(FLongGPUTaskPS,Global,RENDERCORE_API);
+public:
+	FLongGPUTaskPS( )	{ }
+	FLongGPUTaskPS(const ShaderMetaType::CompiledShaderInitializerType& Initializer)
+	:	FGlobalShader( Initializer )
+	{
+	}
+	
+	// FShader interface.
+	virtual bool Serialize(FArchive& Ar) override
+	{
+		bool bShaderHasOutdatedParameters = FGlobalShader::Serialize(Ar);
+		return bShaderHasOutdatedParameters;
+	}
+	
+	static bool ShouldCompilePermutation(const FGlobalShaderPermutationParameters& Parameters)
+	{
+		// MLCHANGES BEGIN
+		return true;
+		// MLCHANGES END
+	}
+};
+
diff --git a/Engine/Source/Runtime/RenderCore/Public/RGBAToYUV420Shader.h b/Engine/Source/Runtime/RenderCore/Public/RGBAToYUV420Shader.h
new file mode 100644
--- /dev/null
+++ b/Engine/Source/Runtime/RenderCore/Public/RGBAToYUV420Shader.h
@@ -0,0 +1,74 @@
+// Copyright 1998-2019 Epic Games, Inc. All Rights Reserved.
+
+#pragma once
+
+#include "CoreMinimal.h"
+#include "RHI.h"
+#include "RenderResource.h"
+#include "Shader.h"
+#include "GlobalShader.h"
+#include "Modules/ModuleManager.h"
+
+#if defined(HAS_MORPHEUS) && HAS_MORPHEUS
+
+/**
+ * Pixel shader to convert RGBA to YUV420.
+ *
+ * YUV420 has 8 bit intensity values in the top 2/3 of the texture for every pixel and 8 bit each UV coordinates 
+ * into the YUV colorspace in the bottom 1/3 for every pixel quad.
+ *
+ * This is only used by PS4.
+ */
+class FRGBAToYUV420CS
+	: public FGlobalShader
+{
+	DECLARE_SHADER_TYPE(FRGBAToYUV420CS, Global);
+
+public:
+
+	static bool ShouldCompilePermutation(const FGlobalShaderPermutationParameters& Parameters)
+	{
+		if (Parameters.Platform == EShaderPlatform::SP_PS4)
+		{
+			// we must use a run time check for this because the builds the build machines create will have Morpheus defined,
+			// but a user will not necessarily have the Morpheus files
+			return FModuleManager::Get().ModuleExists(TEXT("Morpheus"));
+		}
+
+		return false;
+	}
+
+	FRGBAToYUV420CS() { }
+
+	FRGBAToYUV420CS(const ShaderMetaType::CompiledShaderInitializerType& Initializer)
+		: FGlobalShader(Initializer)
+	{
+		TargetHeight.Bind(Initializer.ParameterMap, TEXT("TargetHeight"), SPF_Mandatory);
+		ScaleFactorX.Bind(Initializer.ParameterMap, TEXT("ScaleFactorX"), SPF_Mandatory);
+		ScaleFactorY.Bind(Initializer.ParameterMap, TEXT("ScaleFactorY"), SPF_Mandatory);
+		TextureYOffset.Bind(Initializer.ParameterMap, TEXT("TextureYOffset"), SPF_Mandatory);
+		SrcTexture.Bind(Initializer.ParameterMap, TEXT("SrcTexture"), SPF_Mandatory);
+		OutTextureRW.Bind(Initializer.ParameterMap, TEXT("OutTexture"), SPF_Mandatory);
+	}
+
+	virtual bool Serialize(FArchive& Ar) override
+	{
+		bool bShaderHasOutdatedParameters = FGlobalShader::Serialize(Ar);
+		Ar << TargetHeight << ScaleFactorX << ScaleFactorY << TextureYOffset << SrcTexture << OutTextureRW;
+		return bShaderHasOutdatedParameters;
+	}
+
+	void SetParameters(FRHICommandList& RHICmdList, TRefCountPtr<FRHITexture2D> SrcTex, FRHIUnorderedAccessView* OutUAV, float TargetHeight, float ScaleFactorX, float ScaleFactorY, float TextureYOffset);
+	void UnbindBuffers(FRHICommandList& RHICmdList);
+
+protected:
+	FShaderParameter TargetHeight;
+	FShaderParameter ScaleFactorX;
+	FShaderParameter ScaleFactorY;
+	FShaderParameter TextureYOffset;
+	FShaderResourceParameter SrcTexture;
+
+	FShaderResourceParameter OutTextureRW;
+};
+
+#endif //defined(HAS_MORPHEUS) && HAS_MORPHEUS
diff --git a/Engine/Source/Runtime/RenderCore/Public/ResolveShader.h b/Engine/Source/Runtime/RenderCore/Public/ResolveShader.h
new file mode 100644
--- /dev/null
+++ b/Engine/Source/Runtime/RenderCore/Public/ResolveShader.h
@@ -0,0 +1,283 @@
+// Copyright 1998-2019 Epic Games, Inc. All Rights Reserved.
+
+
+#pragma once
+
+#include "CoreMinimal.h"
+#include "ShaderParameters.h"
+#include "Shader.h"
+#include "GlobalShader.h"
+
+struct FDummyResolveParameter {};
+
+class FResolveDepthPS : public FGlobalShader
+{
+	DECLARE_EXPORTED_SHADER_TYPE(FResolveDepthPS, Global, RENDERCORE_API);
+public:
+	
+	typedef FDummyResolveParameter FParameter;
+	
+	static bool ShouldCompilePermutation(const FGlobalShaderPermutationParameters& Parameters) 
+	{ 
+		return GetMaxSupportedFeatureLevel(Parameters.Platform) >= ERHIFeatureLevel::SM5
+			|| IsSimulatedPlatform(Parameters.Platform); // support resolving MSAA depth in mobile emulation
+	}
+
+	static void ModifyCompilationEnvironment(const FGlobalShaderPermutationParameters& Parameters, FShaderCompilerEnvironment& OutEnvironment)
+	{
+		FGlobalShader::ModifyCompilationEnvironment(Parameters, OutEnvironment);
+		OutEnvironment.SetDefine(TEXT("SIMULATED_PLATFORM"), IsSimulatedPlatform(Parameters.Platform) ? 1 : 0);
+	}
+		
+	FResolveDepthPS(const ShaderMetaType::CompiledShaderInitializerType& Initializer):
+	FGlobalShader(Initializer)
+	{
+		UnresolvedSurface.Bind(Initializer.ParameterMap,TEXT("UnresolvedSurface"), SPF_Mandatory);
+	}
+	FResolveDepthPS() {}
+	
+	void SetParameters(FRHICommandList& RHICmdList, FParameter)
+	{
+	}
+
+	virtual bool Serialize(FArchive& Ar) override
+	{
+		bool bShaderHasOutdatedParameters = FGlobalShader::Serialize(Ar);
+		Ar << UnresolvedSurface;
+		return bShaderHasOutdatedParameters;
+	}
+	
+	FShaderResourceParameter UnresolvedSurface;
+};
+
+class FResolveDepth2XPS : public FGlobalShader
+{
+	DECLARE_EXPORTED_SHADER_TYPE(FResolveDepth2XPS, Global, RENDERCORE_API);
+public:
+
+	typedef FDummyResolveParameter FParameter;
+
+	static bool ShouldCompilePermutation(const FGlobalShaderPermutationParameters& Parameters) 
+	{ 
+		return GetMaxSupportedFeatureLevel(Parameters.Platform) >= ERHIFeatureLevel::SM5 
+			|| IsSimulatedPlatform(Parameters.Platform); // support resolving MSAA depth in mobile emulation
+	}
+
+	static void ModifyCompilationEnvironment(const FGlobalShaderPermutationParameters& Parameters, FShaderCompilerEnvironment& OutEnvironment)
+	{
+		FGlobalShader::ModifyCompilationEnvironment(Parameters, OutEnvironment);
+		OutEnvironment.SetDefine(TEXT("DEPTH_RESOLVE_NUM_SAMPLES"), 2);
+		OutEnvironment.SetDefine(TEXT("SIMULATED_PLATFORM"), IsSimulatedPlatform(Parameters.Platform) ? 1 : 0);
+	}
+
+	FResolveDepth2XPS(const ShaderMetaType::CompiledShaderInitializerType& Initializer) :
+		FGlobalShader(Initializer)
+	{
+		UnresolvedSurface.Bind(Initializer.ParameterMap, TEXT("UnresolvedSurface"), SPF_Mandatory);
+	}
+	FResolveDepth2XPS() {}
+
+	void SetParameters(FRHICommandList& RHICmdList, FParameter)
+	{
+	}
+
+	virtual bool Serialize(FArchive& Ar) override
+	{
+		bool bShaderHasOutdatedParameters = FGlobalShader::Serialize(Ar);
+		Ar << UnresolvedSurface;
+		return bShaderHasOutdatedParameters;
+	}
+
+	FShaderResourceParameter UnresolvedSurface;
+};
+
+
+class FResolveDepth4XPS : public FGlobalShader
+{
+	DECLARE_EXPORTED_SHADER_TYPE(FResolveDepth4XPS, Global, RENDERCORE_API);
+public:
+
+	typedef FDummyResolveParameter FParameter;
+
+	static bool ShouldCompilePermutation(const FGlobalShaderPermutationParameters& Parameters) 
+	{ 
+		return GetMaxSupportedFeatureLevel(Parameters.Platform) >= ERHIFeatureLevel::SM5
+			|| IsSimulatedPlatform(Parameters.Platform); // support resolving MSAA depth in mobile emulation
+	}
+
+	static void ModifyCompilationEnvironment(const FGlobalShaderPermutationParameters& Parameters, FShaderCompilerEnvironment& OutEnvironment)
+	{
+		FGlobalShader::ModifyCompilationEnvironment(Parameters, OutEnvironment);
+		OutEnvironment.SetDefine(TEXT("DEPTH_RESOLVE_NUM_SAMPLES"), 4);
+		OutEnvironment.SetDefine(TEXT("SIMULATED_PLATFORM"), IsSimulatedPlatform(Parameters.Platform) ? 1 : 0);
+	}
+
+	FResolveDepth4XPS(const ShaderMetaType::CompiledShaderInitializerType& Initializer) :
+		FGlobalShader(Initializer)
+	{
+		UnresolvedSurface.Bind(Initializer.ParameterMap, TEXT("UnresolvedSurface"), SPF_Mandatory);
+	}
+	FResolveDepth4XPS() {}
+
+	void SetParameters(FRHICommandList& RHICmdList, FParameter)
+	{
+	}
+
+	virtual bool Serialize(FArchive& Ar) override
+	{
+		bool bShaderHasOutdatedParameters = FGlobalShader::Serialize(Ar);
+		Ar << UnresolvedSurface;
+		return bShaderHasOutdatedParameters;
+	}
+
+	FShaderResourceParameter UnresolvedSurface;
+};
+
+
+class FResolveDepth8XPS : public FGlobalShader
+{
+	DECLARE_EXPORTED_SHADER_TYPE(FResolveDepth8XPS, Global, RENDERCORE_API);
+public:
+
+	typedef FDummyResolveParameter FParameter;
+
+	static bool ShouldCache(EShaderPlatform Platform) { return GetMaxSupportedFeatureLevel(Platform) >= ERHIFeatureLevel::SM5; }
+
+	static bool ShouldCompilePermutation(const FGlobalShaderPermutationParameters& Parameters)
+	{
+		return GetMaxSupportedFeatureLevel(Parameters.Platform) >= ERHIFeatureLevel::SM5
+			|| IsSimulatedPlatform(Parameters.Platform); // support resolving MSAA depth in mobile emulation
+	}
+
+	static void ModifyCompilationEnvironment(const FGlobalShaderPermutationParameters& Parameters, FShaderCompilerEnvironment& OutEnvironment)
+	{
+		FGlobalShader::ModifyCompilationEnvironment(Parameters, OutEnvironment);
+		OutEnvironment.SetDefine(TEXT("DEPTH_RESOLVE_NUM_SAMPLES"), 8);
+		OutEnvironment.SetDefine(TEXT("SIMULATED_PLATFORM"), IsSimulatedPlatform(Parameters.Platform) ? 1 : 0);
+	}
+
+	FResolveDepth8XPS(const ShaderMetaType::CompiledShaderInitializerType& Initializer) :
+		FGlobalShader(Initializer)
+	{
+		UnresolvedSurface.Bind(Initializer.ParameterMap, TEXT("UnresolvedSurface"), SPF_Mandatory);
+	}
+	FResolveDepth8XPS() {}
+
+	void SetParameters(FRHICommandList& RHICmdList, FParameter)
+	{
+	}
+
+	virtual bool Serialize(FArchive& Ar) override
+	{
+		bool bShaderHasOutdatedParameters = FGlobalShader::Serialize(Ar);
+		Ar << UnresolvedSurface;
+		return bShaderHasOutdatedParameters;
+	}
+
+	FShaderResourceParameter UnresolvedSurface;
+};
+
+
+class FResolveDepthNonMSPS : public FGlobalShader
+{
+	DECLARE_EXPORTED_SHADER_TYPE(FResolveDepthNonMSPS, Global, RENDERCORE_API);
+public:
+	
+	typedef FDummyResolveParameter FParameter;
+	
+	static bool ShouldCompilePermutation(const FGlobalShaderPermutationParameters& Parameters) { return GetMaxSupportedFeatureLevel(Parameters.Platform) <= ERHIFeatureLevel::ES3_1; }
+	
+	FResolveDepthNonMSPS(const ShaderMetaType::CompiledShaderInitializerType& Initializer):
+	FGlobalShader(Initializer)
+	{
+		UnresolvedSurface.Bind(Initializer.ParameterMap,TEXT("UnresolvedSurfaceNonMS"), SPF_Mandatory);
+	}
+	FResolveDepthNonMSPS() {}
+	
+	void SetParameters(FRHICommandList& RHICmdList, FParameter)
+	{
+	}
+	
+	virtual bool Serialize(FArchive& Ar) override
+	{
+		bool bShaderHasOutdatedParameters = FGlobalShader::Serialize(Ar);
+		Ar << UnresolvedSurface;
+		return bShaderHasOutdatedParameters;
+	}
+	
+	FShaderResourceParameter UnresolvedSurface;
+};
+
+class FResolveSingleSamplePS : public FGlobalShader
+{
+	DECLARE_EXPORTED_SHADER_TYPE(FResolveSingleSamplePS, Global, RENDERCORE_API);
+public:
+	
+	typedef uint32 FParameter;
+	
+	static bool ShouldCompilePermutation(const FGlobalShaderPermutationParameters& Parameters) { return Parameters.Platform == SP_PCD3D_SM5; }
+	
+	FResolveSingleSamplePS(const ShaderMetaType::CompiledShaderInitializerType& Initializer):
+	FGlobalShader(Initializer)
+	{
+		UnresolvedSurface.Bind(Initializer.ParameterMap,TEXT("UnresolvedSurface"), SPF_Mandatory);
+		SingleSampleIndex.Bind(Initializer.ParameterMap,TEXT("SingleSampleIndex"), SPF_Mandatory);
+	}
+	FResolveSingleSamplePS() {}
+	
+	RENDERCORE_API void SetParameters(FRHICommandList& RHICmdList, uint32 SingleSampleIndexValue);
+	
+	virtual bool Serialize(FArchive& Ar) override
+	{
+		bool bShaderHasOutdatedParameters = FGlobalShader::Serialize(Ar);
+		Ar << UnresolvedSurface;
+		Ar << SingleSampleIndex;
+		return bShaderHasOutdatedParameters;
+	}
+	
+	FShaderResourceParameter UnresolvedSurface;
+	FShaderParameter SingleSampleIndex;
+};
+
+/**
+ * A vertex shader for rendering a textured screen element.
+ */
+class FResolveVS : public FGlobalShader
+{
+	DECLARE_EXPORTED_SHADER_TYPE(FResolveVS, Global, RENDERCORE_API);
+public:
+	
+	static bool ShouldCompilePermutation(const FGlobalShaderPermutationParameters& Parameters) { return true; }
+	
+	FResolveVS(const ShaderMetaType::CompiledShaderInitializerType& Initializer)
+		: FGlobalShader(Initializer)
+	{
+		PositionMinMax.Bind(Initializer.ParameterMap, TEXT("PositionMinMax"), SPF_Mandatory);
+		UVMinMax.Bind(Initializer.ParameterMap, TEXT("UVMinMax"), SPF_Mandatory);
+	}
+	FResolveVS() {}
+
+	RENDERCORE_API void SetParameters(FRHICommandList& RHICmdList, FVector4 PositionMinMax, FVector4 UVMinMax);
+
+	virtual bool Serialize(FArchive& Ar) override
+	{
+		bool bShaderHasOutdatedParameters = FGlobalShader::Serialize(Ar);
+		Ar << PositionMinMax;
+		Ar << UVMinMax;
+		return bShaderHasOutdatedParameters;
+	}
+
+	FShaderParameter PositionMinMax;
+	FShaderParameter UVMinMax;
+};
+
+struct FResolveVertexBuffer : public FRenderResource
+{
+	FVertexBufferRHIRef VertexBufferRHI;
+	FVertexDeclarationRHIRef VertexDeclarationRHI;
+
+	virtual void InitDynamicRHI() override;
+	virtual void ReleaseDynamicRHI() override;
+};
+
+extern RENDERCORE_API TGlobalResource<FResolveVertexBuffer> GResolveVertexBuffer;
diff --git a/Engine/Source/Runtime/RenderCore/Public/ShaderParameterMacros.h b/Engine/Source/Runtime/RenderCore/Public/ShaderParameterMacros.h
--- a/Engine/Source/Runtime/RenderCore/Public/ShaderParameterMacros.h
+++ b/Engine/Source/Runtime/RenderCore/Public/ShaderParameterMacros.h
@@ -704,7 +704,7 @@ struct TShaderParameterStructTypeInfo<StructType[InNumElements]>
 		static_assert(BaseType != UBMT_INVALID, "Invalid type " #MemberType " of member " #MemberName "."); \
 	private: \
 		struct zzNextMemberId##MemberName { enum { HasDeclaredResource = zzMemberId##MemberName::HasDeclaredResource || !TypeInfo::bIsStoredInConstantBuffer }; }; \
-		static FORCEINLINE TArray<FShaderParametersMetadata::FMember> zzGetMembersBefore(zzNextMemberId##MemberName) \
+		static TArray<FShaderParametersMetadata::FMember> zzGetMembersBefore(zzNextMemberId##MemberName) \
 		{ \
 			static_assert(TypeInfo::bIsStoredInConstantBuffer || TIsArrayOrRefOfType<decltype(OptionalShaderType), TCHAR>::Value, "No shader type for " #MemberName "."); \
 			/* Route the member enumeration on to the function for the member following this. */ \
diff --git a/Engine/Source/Runtime/UtilityShaders/Private/BuiltInRayTracingShaders.cpp b/Engine/Source/Runtime/UtilityShaders/Private/BuiltInRayTracingShaders.cpp
deleted file mode 100644
diff --git a/Engine/Source/Runtime/UtilityShaders/Private/ClearQuad.cpp b/Engine/Source/Runtime/UtilityShaders/Private/ClearQuad.cpp
deleted file mode 100644
diff --git a/Engine/Source/Runtime/UtilityShaders/Private/ClearReplacementShaders.cpp b/Engine/Source/Runtime/UtilityShaders/Private/ClearReplacementShaders.cpp
deleted file mode 100644
diff --git a/Engine/Source/Runtime/UtilityShaders/Private/LongGPUTask.cpp b/Engine/Source/Runtime/UtilityShaders/Private/LongGPUTask.cpp
deleted file mode 100644
diff --git a/Engine/Source/Runtime/UtilityShaders/Private/MediaShaders.cpp b/Engine/Source/Runtime/UtilityShaders/Private/MediaShaders.cpp
deleted file mode 100644
diff --git a/Engine/Source/Runtime/UtilityShaders/Private/OculusShaders.cpp b/Engine/Source/Runtime/UtilityShaders/Private/OculusShaders.cpp
deleted file mode 100644
diff --git a/Engine/Source/Runtime/UtilityShaders/Private/OneColorShader.cpp b/Engine/Source/Runtime/UtilityShaders/Private/OneColorShader.cpp
deleted file mode 100644
diff --git a/Engine/Source/Runtime/UtilityShaders/Private/RGBAToYUV420Shader.cpp b/Engine/Source/Runtime/UtilityShaders/Private/RGBAToYUV420Shader.cpp
deleted file mode 100644
diff --git a/Engine/Source/Runtime/UtilityShaders/Private/ResolveShader.cpp b/Engine/Source/Runtime/UtilityShaders/Private/ResolveShader.cpp
deleted file mode 100644
diff --git a/Engine/Source/Runtime/UtilityShaders/Public/BuiltInRayTracingShaders.h b/Engine/Source/Runtime/UtilityShaders/Public/BuiltInRayTracingShaders.h
deleted file mode 100644
diff --git a/Engine/Source/Runtime/UtilityShaders/Public/ClearQuad.h b/Engine/Source/Runtime/UtilityShaders/Public/ClearQuad.h
deleted file mode 100644
diff --git a/Engine/Source/Runtime/UtilityShaders/Public/ClearReplacementShaders.h b/Engine/Source/Runtime/UtilityShaders/Public/ClearReplacementShaders.h
deleted file mode 100644
diff --git a/Engine/Source/Runtime/UtilityShaders/Public/LongGPUTask.h b/Engine/Source/Runtime/UtilityShaders/Public/LongGPUTask.h
deleted file mode 100644
diff --git a/Engine/Source/Runtime/UtilityShaders/Public/MediaShaders.h b/Engine/Source/Runtime/UtilityShaders/Public/MediaShaders.h
deleted file mode 100644
diff --git a/Engine/Source/Runtime/UtilityShaders/Public/OculusShaders.h b/Engine/Source/Runtime/UtilityShaders/Public/OculusShaders.h
deleted file mode 100644
diff --git a/Engine/Source/Runtime/UtilityShaders/Public/OneColorShader.h b/Engine/Source/Runtime/UtilityShaders/Public/OneColorShader.h
deleted file mode 100644
diff --git a/Engine/Source/Runtime/UtilityShaders/Public/RGBAToYUV420Shader.h b/Engine/Source/Runtime/UtilityShaders/Public/RGBAToYUV420Shader.h
deleted file mode 100644
diff --git a/Engine/Source/Runtime/UtilityShaders/Public/ResolveShader.h b/Engine/Source/Runtime/UtilityShaders/Public/ResolveShader.h
deleted file mode 100644
