Back out CL 8638795 - Inlining resulted in too much compile time cost and instruction count increase.
#rnx

diff --git a/Engine/Source/Runtime/RenderCore/Public/ShaderParameterMacros.h b/Engine/Source/Runtime/RenderCore/Public/ShaderParameterMacros.h
--- a/Engine/Source/Runtime/RenderCore/Public/ShaderParameterMacros.h
+++ b/Engine/Source/Runtime/RenderCore/Public/ShaderParameterMacros.h
@@ -704,7 +704,7 @@ struct TShaderParameterStructTypeInfo<StructType[InNumElements]>
 		static_assert(BaseType != UBMT_INVALID, "Invalid type " #MemberType " of member " #MemberName "."); \
 	private: \
 		struct zzNextMemberId##MemberName { enum { HasDeclaredResource = zzMemberId##MemberName::HasDeclaredResource || !TypeInfo::bIsStoredInConstantBuffer }; }; \
-		static FORCEINLINE TArray<FShaderParametersMetadata::FMember> zzGetMembersBefore(zzNextMemberId##MemberName) \
+		static TArray<FShaderParametersMetadata::FMember> zzGetMembersBefore(zzNextMemberId##MemberName) \
 		{ \
 			static_assert(TypeInfo::bIsStoredInConstantBuffer || TIsArrayOrRefOfType<decltype(OptionalShaderType), TCHAR>::Value, "No shader type for " #MemberName "."); \
 			/* Route the member enumeration on to the function for the member following this. */ \
