Bug fix for a race condition where render tickables could be updating objects that have been deleted.

diff --git a/Engine/Source/Runtime/RenderCore/Private/RenderingThread.cpp b/Engine/Source/Runtime/RenderCore/Private/RenderingThread.cpp
--- a/Engine/Source/Runtime/RenderCore/Private/RenderingThread.cpp
+++ b/Engine/Source/Runtime/RenderCore/Private/RenderingThread.cpp
@@ -498,6 +498,21 @@ public:
 TAtomic<bool> GRunRenderingThreadHeartbeat;
 
 FThreadSafeCounter OutstandingHeartbeats;
+
+/** rendering tickables shouldn't be updated during a flush */
+TAtomic<int32> GSuspendRenderingTickables = 0;
+struct FSuspendRenderingTickables
+{
+	FSuspendRenderingTickables()
+	{
+		++GSuspendRenderingTickables;
+	}
+	~FSuspendRenderingTickables()
+	{
+		--GSuspendRenderingTickables;
+	}
+};
+
 /** The rendering thread heartbeat runnable object. */
 class FRenderingThreadTickHeartbeat : public FRunnable
 {
@@ -531,7 +546,8 @@ public:
 					{
 						OutstandingHeartbeats.Decrement();
 						// make sure that rendering thread tickables get a chance to tick, even if the render thread is starving
-						if (!GIsRenderingThreadSuspended.Load(EMemoryOrder::Relaxed))
+						// but if GSuspendRenderingTickables is != 0 a flush is happening so don't tick during it
+						if (!GIsRenderingThreadSuspended.Load(EMemoryOrder::Relaxed) && !GSuspendRenderingTickables.Load(EMemoryOrder::Relaxed))
 						{
 							TickRenderingTickables();
 						}
@@ -1211,6 +1227,7 @@ void FlushRenderingCommands(bool bFlushDeferredDeletes)
 	{
 		return;
 	}
+	FSuspendRenderingTickables SuspendRenderingTickables;
 
 	// Need to flush GT because render commands from threads other than GT are sent to
 	// the main queue of GT when RT is disabled
