Buildbreak fix.

diff --git a/Engine/Source/Runtime/RenderCore/Private/RenderingThread.cpp b/Engine/Source/Runtime/RenderCore/Private/RenderingThread.cpp
--- a/Engine/Source/Runtime/RenderCore/Private/RenderingThread.cpp
+++ b/Engine/Source/Runtime/RenderCore/Private/RenderingThread.cpp
@@ -500,7 +500,7 @@ TAtomic<bool> GRunRenderingThreadHeartbeat;
 FThreadSafeCounter OutstandingHeartbeats;
 
 /** rendering tickables shouldn't be updated during a flush */
-TAtomic<int32> GSuspendRenderingTickables = 0;
+TAtomic<int32> GSuspendRenderingTickables;
 struct FSuspendRenderingTickables
 {
 	FSuspendRenderingTickables()
@@ -520,7 +520,8 @@ public:
 
 	// FRunnable interface.
 	virtual bool Init(void) 
-	{ 
+	{
+		GSuspendRenderingTickables = 0;
 		OutstandingHeartbeats.Reset();
 		return true; 
 	}
