Clarified shader compiler output.

diff --git a/Engine/Source/Runtime/RenderCore/Private/ShaderParameterStruct.cpp b/Engine/Source/Runtime/RenderCore/Private/ShaderParameterStruct.cpp
--- a/Engine/Source/Runtime/RenderCore/Private/ShaderParameterStruct.cpp
+++ b/Engine/Source/Runtime/RenderCore/Private/ShaderParameterStruct.cpp
@@ -198,10 +198,15 @@ struct FShaderParameterStructBindingContext
 					Parameter.BaseIndex = BaseIndex;
 					Parameter.ByteOffset = ByteOffset + ArrayElementId * SHADER_PARAMETER_POINTER_ALIGNMENT;
 
-					checkf(
-						BoundSize == 1,
-						TEXT("The shader compiler should give precisely which elements of an array did not get compiled out, ")
-						TEXT("for optimal automatic render graph pass dependency with ClearUnusedGraphResources()."));
+					if (BoundSize != 1)
+					{
+						UE_LOG(LogShaders, Fatal, 
+							TEXT("Error with shader %s's (Permutation Id %d) parameter %s is %i bytes, cpp name = %s.")
+							TEXT("The shader compiler should give precisely which elements of an array did not get compiled out, ")
+							TEXT("for optimal automatic render graph pass dependency with ClearUnusedGraphResources()."),
+							Shader->GetType()->GetName(), Shader->GetPermutationId(),
+							*ElementShaderBindingName, BoundSize, *CppName);
+					}
 
 					if (BaseType == UBMT_TEXTURE)
 						Bindings->Textures.Add(Parameter);
