FShaderCodeLibrary - Editor fixes.  It's up to the caller if we want to initialize the code library - if we have an implementation then allow shader creation.  Impl creation was checked against FPlatformProperties::RequiresCookedData() but its a bit more complex than that.  Logic for creation is in LaunchEngineLoop.cpp.

#rnx

diff --git a/Engine/Source/Runtime/Launch/Private/LaunchEngineLoop.cpp b/Engine/Source/Runtime/Launch/Private/LaunchEngineLoop.cpp
--- a/Engine/Source/Runtime/Launch/Private/LaunchEngineLoop.cpp
+++ b/Engine/Source/Runtime/Launch/Private/LaunchEngineLoop.cpp
@@ -2274,20 +2274,27 @@ int32 FEngineLoop::PreInit(const TCHAR* CmdLine)
 		RenderUtilsInit();
 	}
 
-	if (FPlatformProperties::RequiresCookedData())
 	{
+		bool bUseCodeLibrary = FPlatformProperties::RequiresCookedData() || GAllowCookedDataInEditorBuilds;
+		if (bUseCodeLibrary)
 		{
-			SCOPED_BOOT_TIMING("FShaderCodeLibrary::InitForRuntime");
-			// Will open material shader code storage if project was packaged with it
-			// This only opens the Global shader library, which is always in the content dir.
-			FShaderCodeLibrary::InitForRuntime(GMaxRHIShaderPlatform);
-		}
+			{
+				SCOPED_BOOT_TIMING("FShaderCodeLibrary::InitForRuntime");
+				// Will open material shader code storage if project was packaged with it
+				// This only opens the Global shader library, which is always in the content dir.
+				FShaderCodeLibrary::InitForRuntime(GMaxRHIShaderPlatform);
+			}
 
-		{
-			SCOPED_BOOT_TIMING("FShaderPipelineCache::Initialize");
-			// Initialize the pipeline cache system. Opening is deferred until the manual call to
-			// OpenPipelineFileCache below, after content pak's ShaderCodeLibraries are loaded.
-			FShaderPipelineCache::Initialize(GMaxRHIShaderPlatform);
+	#if !UE_EDITOR
+			// Cooked data only - but also requires the code library - game only
+			if (FPlatformProperties::RequiresCookedData())
+			{
+				SCOPED_BOOT_TIMING("FShaderPipelineCache::Initialize");
+				// Initialize the pipeline cache system. Opening is deferred until the manual call to
+				// OpenPipelineFileCache below, after content pak's ShaderCodeLibraries are loaded.
+				FShaderPipelineCache::Initialize(GMaxRHIShaderPlatform);
+			}
+	#endif //!UE_EDITOR
 		}
 	}
 
@@ -2586,7 +2593,7 @@ int32 FEngineLoop::PreInit(const TCHAR* CmdLine)
 					FShaderCodeLibrary::OpenLibrary(FApp::GetProjectName(), FPaths::Combine(RootDir, FApp::GetProjectName(), TEXT("Content")));
 				}
 
-				// Now our shader code main library is opened, kick off the precompile.
+				// Now our shader code main library is opened, kick off the precompile, if already initialized
 				FShaderPipelineCache::OpenPipelineFileCache(GMaxRHIShaderPlatform);
 			}
 		}
diff --git a/Engine/Source/Runtime/RenderCore/Private/ShaderCodeLibrary.cpp b/Engine/Source/Runtime/RenderCore/Private/ShaderCodeLibrary.cpp
--- a/Engine/Source/Runtime/RenderCore/Private/ShaderCodeLibrary.cpp
+++ b/Engine/Source/Runtime/RenderCore/Private/ShaderCodeLibrary.cpp
@@ -2149,8 +2149,6 @@ static void FShaderCodeLibraryPluginMountedCallback(IPlugin& Plugin)
 
 void FShaderCodeLibrary::InitForRuntime(EShaderPlatform ShaderPlatform)
 {
-	check(FPlatformProperties::RequiresCookedData());
-
 	if (FShaderCodeLibraryImpl::Impl != nullptr)
 	{
 		//cooked, can't change shader platform on the fly
@@ -2277,7 +2275,7 @@ void FShaderCodeLibrary::LazyReleaseShaderCode(const FSHAHash& Hash)
 FVertexShaderRHIRef FShaderCodeLibrary::CreateVertexShader(EShaderPlatform Platform, FSHAHash Hash, TArray<uint8> const& Code)
 {
 	FVertexShaderRHIRef Shader;
-	if (FShaderCodeLibraryImpl::Impl && FPlatformProperties::RequiresCookedData())
+	if (FShaderCodeLibraryImpl::Impl)
 	{
 		Shader = FShaderCodeLibraryImpl::Impl->CreateVertexShader(Platform, Hash);
 	}
@@ -2292,7 +2290,7 @@ FVertexShaderRHIRef FShaderCodeLibrary::CreateVertexShader(EShaderPlatform Platf
 FPixelShaderRHIRef FShaderCodeLibrary::CreatePixelShader(EShaderPlatform Platform, FSHAHash Hash, TArray<uint8> const& Code)
 {
 	FPixelShaderRHIRef Shader;
-	if (FShaderCodeLibraryImpl::Impl && FPlatformProperties::RequiresCookedData())
+	if (FShaderCodeLibraryImpl::Impl)
 	{
 		Shader = FShaderCodeLibraryImpl::Impl->CreatePixelShader(Platform, Hash);
 	}
@@ -2307,7 +2305,7 @@ FPixelShaderRHIRef FShaderCodeLibrary::CreatePixelShader(EShaderPlatform Platfor
 FGeometryShaderRHIRef FShaderCodeLibrary::CreateGeometryShader(EShaderPlatform Platform, FSHAHash Hash, TArray<uint8> const& Code)
 {
 	FGeometryShaderRHIRef Shader;
-	if (FShaderCodeLibraryImpl::Impl && FPlatformProperties::RequiresCookedData())
+	if (FShaderCodeLibraryImpl::Impl)
 	{
 		Shader = FShaderCodeLibraryImpl::Impl->CreateGeometryShader(Platform, Hash);
 	}
@@ -2322,7 +2320,7 @@ FGeometryShaderRHIRef FShaderCodeLibrary::CreateGeometryShader(EShaderPlatform P
 FHullShaderRHIRef FShaderCodeLibrary::CreateHullShader(EShaderPlatform Platform, FSHAHash Hash, TArray<uint8> const& Code)
 {
 	FHullShaderRHIRef Shader;
-	if (FShaderCodeLibraryImpl::Impl && FPlatformProperties::RequiresCookedData())
+	if (FShaderCodeLibraryImpl::Impl)
 	{
 		Shader = FShaderCodeLibraryImpl::Impl->CreateHullShader(Platform, Hash);
 	}
@@ -2337,7 +2335,7 @@ FHullShaderRHIRef FShaderCodeLibrary::CreateHullShader(EShaderPlatform Platform,
 FDomainShaderRHIRef FShaderCodeLibrary::CreateDomainShader(EShaderPlatform Platform, FSHAHash Hash, TArray<uint8> const& Code)
 {
 	FDomainShaderRHIRef Shader;
-	if (FShaderCodeLibraryImpl::Impl && FPlatformProperties::RequiresCookedData())
+	if (FShaderCodeLibraryImpl::Impl)
 	{
 		Shader = FShaderCodeLibraryImpl::Impl->CreateDomainShader(Platform, Hash);
 	}
@@ -2352,7 +2350,7 @@ FDomainShaderRHIRef FShaderCodeLibrary::CreateDomainShader(EShaderPlatform Platf
 FComputeShaderRHIRef FShaderCodeLibrary::CreateComputeShader(EShaderPlatform Platform, FSHAHash Hash, TArray<uint8> const& Code)
 {
 	FComputeShaderRHIRef Shader;
-	if (FShaderCodeLibraryImpl::Impl && FPlatformProperties::RequiresCookedData())
+	if (FShaderCodeLibraryImpl::Impl)
 	{
 		Shader = FShaderCodeLibraryImpl::Impl->CreateComputeShader(Platform, Hash);
 	}
@@ -2369,7 +2367,7 @@ FComputeShaderRHIRef FShaderCodeLibrary::CreateComputeShader(EShaderPlatform Pla
 TRefCountPtr<FRHIShaderLibrary::FShaderLibraryIterator> FShaderCodeLibrary::CreateIterator(void)
 {
 	TRefCountPtr<FRHIShaderLibrary::FShaderLibraryIterator> It;
-	if (FShaderCodeLibraryImpl::Impl && FPlatformProperties::RequiresCookedData())
+	if (FShaderCodeLibraryImpl::Impl)
 	{
 		It = FShaderCodeLibraryImpl::Impl->CreateIterator();
 	}
@@ -2379,7 +2377,7 @@ TRefCountPtr<FRHIShaderLibrary::FShaderLibraryIterator> FShaderCodeLibrary::Crea
 uint32 FShaderCodeLibrary::GetShaderCount(void)
 {
 	uint32 Num = 0;
-	if (FShaderCodeLibraryImpl::Impl && FPlatformProperties::RequiresCookedData())
+	if (FShaderCodeLibraryImpl::Impl)
 	{
 		Num = FShaderCodeLibraryImpl::Impl->GetShaderCount();
 	}
@@ -2389,7 +2387,7 @@ uint32 FShaderCodeLibrary::GetShaderCount(void)
 TSet<FShaderCodeLibraryPipeline> const* FShaderCodeLibrary::GetShaderPipelines(EShaderPlatform Platform)
 {
 	TSet<FShaderCodeLibraryPipeline> const* Pipelines = nullptr;
-	if (FShaderCodeLibraryImpl::Impl && FPlatformProperties::RequiresCookedData())
+	if (FShaderCodeLibraryImpl::Impl)
 	{
 		Pipelines = FShaderCodeLibraryImpl::Impl->GetShaderPipelines(Platform);
 	}
@@ -2399,7 +2397,7 @@ TSet<FShaderCodeLibraryPipeline> const* FShaderCodeLibrary::GetShaderPipelines(E
 EShaderPlatform FShaderCodeLibrary::GetRuntimeShaderPlatform(void)
 {
 	EShaderPlatform Platform = SP_NumPlatforms;
-	if (FShaderCodeLibraryImpl::Impl && FPlatformProperties::RequiresCookedData())
+	if (FShaderCodeLibraryImpl::Impl)
 	{
 		Platform = FShaderCodeLibraryImpl::Impl->GetRuntimeShaderPlatform();
 	}
