Fixed render target pool to actually release unused resources instead of all of them.

diff --git a/Engine/Source/Runtime/RenderCore/Private/RenderTargetPool.cpp b/Engine/Source/Runtime/RenderCore/Private/RenderTargetPool.cpp
--- a/Engine/Source/Runtime/RenderCore/Private/RenderTargetPool.cpp
+++ b/Engine/Source/Runtime/RenderCore/Private/RenderTargetPool.cpp
@@ -975,7 +975,7 @@ void FRenderTargetPool::FreeUnusedResource(TRefCountPtr<IPooledRenderTarget>& In
 	{
 		FPooledRenderTarget* Element = PooledRenderTargets[Index];
 
-		if (Element)
+		if (Element && Element->IsFree())
 		{
 			check(!Element->IsSnapshot());
 			AllocationLevelInKB -= ComputeSizeInKB(*Element);
