fix for not properly deleting the OS cache on the first run

diff --git a/Engine/Source/Runtime/Apple/MetalRHI/Private/MetalPipeline.cpp b/Engine/Source/Runtime/Apple/MetalRHI/Private/MetalPipeline.cpp
--- a/Engine/Source/Runtime/Apple/MetalRHI/Private/MetalPipeline.cpp
+++ b/Engine/Source/Runtime/Apple/MetalRHI/Private/MetalPipeline.cpp
@@ -1878,10 +1878,12 @@ void FMetalPipelineStateCacheManager::OnShaderPipelineCachePreOpen(FString const
 	if (stat(TCHAR_TO_UTF8(*Result), &FileInfo) != -1 && ((FileInfo.st_size / 1024 / 1024) > GMetalCacheMinSize) && stat(TCHAR_TO_UTF8(*Result2), &FileInfo) != -1)
 	{
 		bReady = false;
+		FShaderPipelineCache::SetBatchMode(FShaderPipelineCache::BatchMode::Background);
 	}
 	else
 	{
 		bReady = true;
+		FShaderPipelineCache::SetBatchMode(FShaderPipelineCache::BatchMode::Precompile);
 	}
 }
 
diff --git a/Engine/Source/Runtime/RHI/Private/PipelineFileCache.cpp b/Engine/Source/Runtime/RHI/Private/PipelineFileCache.cpp
--- a/Engine/Source/Runtime/RHI/Private/PipelineFileCache.cpp
+++ b/Engine/Source/Runtime/RHI/Private/PipelineFileCache.cpp
@@ -2458,7 +2458,7 @@ bool FPipelineFileCache::ShouldEnableFileCache()
 		static FString PrivateWritePathBase = FString([NSSearchPathForDirectoriesInDomains(NSLibraryDirectory, NSUserDomainMask, YES) objectAtIndex:0]) + TEXT("/");
 		FString Result = PrivateWritePathBase + FString([NSString stringWithFormat:@"/Caches/%@/com.apple.metal/functions.data", [NSBundle mainBundle].bundleIdentifier]);
 		FString Result2 = PrivateWritePathBase + FString([NSString stringWithFormat:@"/Caches/%@/com.apple.metal/usecache.txt", [NSBundle mainBundle].bundleIdentifier]);
-		if (stat(TCHAR_TO_UTF8(*Result), &FileInfo) != -1 && ((FileInfo.st_size / 1024 / 1024) > 4) && stat(TCHAR_TO_UTF8(*Result2), &FileInfo) != -1)
+		if (stat(TCHAR_TO_UTF8(*Result), &FileInfo) != -1 && stat(TCHAR_TO_UTF8(*Result2), &FileInfo) != -1)
 		{
 			return false;
 		}
@@ -2474,12 +2474,16 @@ void FPipelineFileCache::PreCompileComplete()
 	static FString PrivateWritePathBase = FString([NSSearchPathForDirectoriesInDomains(NSLibraryDirectory, NSUserDomainMask, YES) objectAtIndex:0]) + TEXT("/");
 	FString Result = PrivateWritePathBase + FString([NSString stringWithFormat:@"/Caches/%@/com.apple.metal/usecache.txt", [NSBundle mainBundle].bundleIdentifier]);
 	int32 Handle = open(TCHAR_TO_UTF8(*Result), O_CREAT, S_IRUSR | S_IWUSR | S_IRGRP | S_IWGRP | S_IROTH | S_IWOTH);
+	char* Version = TCHAR_TO_ANSI(*FEngineVersion::Current().ToString());
+	write(Handle, Version, strlen(Version));
 	close(Handle);
 #endif
 }
 
 void FPipelineFileCache::ClearOSPipelineCache()
 {
+	UE_LOG(LogTemp, Warning, TEXT("Clearing the OS Cache"));
+	
 	bool bCmdLineSkip = FParse::Param(FCommandLine::Get(), TEXT("skippsoclear"));
 	if (CVarClearOSPSOFileCache.GetValueOnAnyThread() > 0 && !bCmdLineSkip)
 	{
diff --git a/Engine/Source/Runtime/RenderCore/Private/ShaderPipelineCache.cpp b/Engine/Source/Runtime/RenderCore/Private/ShaderPipelineCache.cpp
--- a/Engine/Source/Runtime/RenderCore/Private/ShaderPipelineCache.cpp
+++ b/Engine/Source/Runtime/RenderCore/Private/ShaderPipelineCache.cpp
@@ -69,6 +69,12 @@ static TAutoConsoleVariable<int32> CVarPSOFileCacheBatchSize(
 														   TEXT("Set the number of PipelineStateObjects to compile in a single batch operation when compiling takes priority. Defaults to a maximum of 50 per frame, due to async. file IO it is less in practice."),
 														   ECVF_Default | ECVF_RenderThreadSafe
 														   );
+static TAutoConsoleVariable<int32> CVarPSOFileCachePrecompileBatchSize(
+															 TEXT("r.ShaderPipelineCache.PrecompileBatchSize"),
+															 50,
+															 TEXT("Set the number of PipelineStateObjects to compile in a single batch operation when pre-optimizing the cache. Defaults to a maximum of 50 per frame, due to async. file IO it is less in practice."),
+															 ECVF_Default | ECVF_RenderThreadSafe
+															 );
 static TAutoConsoleVariable<float> CVarPSOFileCacheBackgroundBatchTime(
 														  TEXT("r.ShaderPipelineCache.BackgroundBatchTime"),
 														  0.0f,
@@ -81,6 +87,12 @@ static TAutoConsoleVariable<float> CVarPSOFileCacheBatchTime(
 														   TEXT("The target time (in ms) to spend precompiling each frame when compiling takes priority or 0.0 to disable. When precompiling is faster the batch size will grow and when slower will shrink to attempt to occupy the full amount. Defaults to 16.0 (max. ms per-frame of precompilation)."),
 														   ECVF_Default | ECVF_RenderThreadSafe
 														   );
+static TAutoConsoleVariable<float> CVarPSOFileCachePrecompileBatchTime(
+															 TEXT("r.ShaderPipelineCache.PrecompileBatchTime"),
+															 0.0f,
+															 TEXT("The target time (in ms) to spend precompiling each frame when cpre-optimizing or 0.0 to disable. When precompiling is faster the batch size will grow and when slower will shrink to attempt to occupy the full amount. Defaults to 10.0 (off)."),
+															 ECVF_Default | ECVF_RenderThreadSafe
+															 );
 static TAutoConsoleVariable<int32> CVarPSOFileCacheSaveAfterPSOsLogged(
 														   TEXT("r.ShaderPipelineCache.SaveAfterPSOsLogged"),
 #if !UE_BUILD_SHIPPING
@@ -199,6 +211,11 @@ void ConsoleCommandSwitchModePipelineCacheCmd(const TArray< FString >& Args)
             FShaderPipelineCache::SetBatchMode(FShaderPipelineCache::BatchMode::Fast);
             FShaderPipelineCache::ResumeBatching();
         }
+		else if (Mode == TEXT("Precompile"))
+		{
+			FShaderPipelineCache::SetBatchMode(FShaderPipelineCache::BatchMode::Precompile);
+			FShaderPipelineCache::ResumeBatching();
+		}
     }
 }
 
@@ -420,6 +437,12 @@ void FShaderPipelineCache::SetBatchMode(BatchMode Mode)
 	{
 		switch (Mode)
 		{
+			case BatchMode::Precompile:
+			{
+				ShaderPipelineCache->BatchSize = CVarPSOFileCachePrecompileBatchSize.GetValueOnAnyThread();
+				ShaderPipelineCache->BatchTime = CVarPSOFileCachePrecompileBatchTime.GetValueOnAnyThread();
+				break;
+			}
 			case BatchMode::Fast:
 			{
 				ShaderPipelineCache->BatchSize = CVarPSOFileCacheBatchSize.GetValueOnAnyThread();
diff --git a/Engine/Source/Runtime/RenderCore/Public/ShaderPipelineCache.h b/Engine/Source/Runtime/RenderCore/Public/ShaderPipelineCache.h
--- a/Engine/Source/Runtime/RenderCore/Public/ShaderPipelineCache.h
+++ b/Engine/Source/Runtime/RenderCore/Public/ShaderPipelineCache.h
@@ -90,7 +90,8 @@ public:
 	enum class BatchMode
 	{
 		Background, // The maximum batch size is defined by r.ShaderPipelineCache.BackgroundBatchSize
-		Fast // The maximum batch size is defined by r.ShaderPipelineCache.BatchSize
+		Fast, // The maximum batch size is defined by r.ShaderPipelineCache.BatchSize
+		Precompile // The maximum batch size is defined by r.ShaderPipelineCache.PrecompileBatchSize
 	};
 	
 	/** Sets the precompilation batching mode. */
