Editgrate CL#8757427:

Refactor INTERNAL_SHADER_PARAMETER_STRUCT initialization to avoid recursion, avoiding stack overflow on large structures.  Processes the members in reverse, requesting the next function to call and looping until the first member is hit.  Then the final list is reversed to put in proper order.

diff --git a/Engine/Source/Runtime/RenderCore/Public/ShaderParameterMacros.h b/Engine/Source/Runtime/RenderCore/Public/ShaderParameterMacros.h
--- a/Engine/Source/Runtime/RenderCore/Public/ShaderParameterMacros.h
+++ b/Engine/Source/Runtime/RenderCore/Public/ShaderParameterMacros.h
@@ -8,6 +8,7 @@
 #pragma once
 
 #include "ShaderParameterMetadata.h"
+#include "Algo/Reverse.h"
 
 
 class FRDGTexture;
@@ -690,9 +691,11 @@ struct TShaderParameterStructTypeInfo<StructType[InNumElements]>
 	private: \
 		typedef StructTypeName zzTThisStruct; \
 		struct zzFirstMemberId { enum { HasDeclaredResource = 0 }; }; \
-		static TArray<FShaderParametersMetadata::FMember> zzGetMembersBefore(zzFirstMemberId) \
+		typedef void* zzFuncPtr; \
+		typedef zzFuncPtr(*zzMemberFunc)(zzFirstMemberId, TArray<FShaderParametersMetadata::FMember>*); \
+		static zzFuncPtr zzAppendMemberGetPrev(zzFirstMemberId, TArray<FShaderParametersMetadata::FMember>*) \
 		{ \
-			return TArray<FShaderParametersMetadata::FMember>(); \
+			return nullptr; \
 		} \
 		typedef zzFirstMemberId
 
@@ -704,13 +707,13 @@ struct TShaderParameterStructTypeInfo<StructType[InNumElements]>
 		static_assert(BaseType != UBMT_INVALID, "Invalid type " #MemberType " of member " #MemberName "."); \
 	private: \
 		struct zzNextMemberId##MemberName { enum { HasDeclaredResource = zzMemberId##MemberName::HasDeclaredResource || !TypeInfo::bIsStoredInConstantBuffer }; }; \
-		static TArray<FShaderParametersMetadata::FMember> zzGetMembersBefore(zzNextMemberId##MemberName) \
+		static zzFuncPtr zzAppendMemberGetPrev(zzNextMemberId##MemberName, TArray<FShaderParametersMetadata::FMember>* Members) \
 		{ \
 			static_assert(TypeInfo::bIsStoredInConstantBuffer || TIsArrayOrRefOfType<decltype(OptionalShaderType), TCHAR>::Value, "No shader type for " #MemberName "."); \
-			/* Route the member enumeration on to the function for the member following this. */ \
-			TArray<FShaderParametersMetadata::FMember> OutMembers = zzGetMembersBefore(zzMemberId##MemberName()); \
-			/* Add this member. */ \
-			OutMembers.Add(FShaderParametersMetadata::FMember( \
+			static_assert(\
+				(STRUCT_OFFSET(zzTThisStruct, MemberName) & (TypeInfo::Alignment - 1)) == 0, \
+				"Misaligned uniform buffer struct member " #MemberName "."); \
+			Members->Add(FShaderParametersMetadata::FMember( \
 				TEXT(#MemberName), \
 				OptionalShaderType, \
 				STRUCT_OFFSET(zzTThisStruct,MemberName), \
@@ -721,10 +724,9 @@ struct TShaderParameterStructTypeInfo<StructType[InNumElements]>
 				TypeInfo::NumElements, \
 				TypeInfo::GetStructMetadata() \
 				)); \
-			static_assert( \
-				(STRUCT_OFFSET(zzTThisStruct,MemberName) & (TypeInfo::Alignment - 1)) == 0, \
-				"Misaligned uniform buffer struct member " #MemberName "."); \
-			return OutMembers; \
+			zzFuncPtr(*PrevFunc)(zzMemberId##MemberName, TArray<FShaderParametersMetadata::FMember>*); \
+			PrevFunc = zzAppendMemberGetPrev; \
+			return (zzFuncPtr)PrevFunc; \
 		} \
 		typedef zzNextMemberId##MemberName
 
@@ -744,7 +746,18 @@ extern RENDERCORE_API FShaderParametersMetadata* FindUniformBufferStructByFName(
 
 #define END_SHADER_PARAMETER_STRUCT() \
 		zzLastMemberId; \
-		static TArray<FShaderParametersMetadata::FMember> zzGetMembers() { return zzGetMembersBefore(zzLastMemberId()); } \
+		static TArray<FShaderParametersMetadata::FMember> zzGetMembers() { \
+			TArray<FShaderParametersMetadata::FMember> Members; \
+			zzFuncPtr(*LastFunc)(zzLastMemberId, TArray<FShaderParametersMetadata::FMember>*); \
+			LastFunc = zzAppendMemberGetPrev; \
+			zzFuncPtr Ptr = (zzFuncPtr)LastFunc; \
+			do \
+			{ \
+				Ptr = reinterpret_cast<zzMemberFunc>(Ptr)(zzFirstMemberId(), &Members); \
+			} while (Ptr); \
+			Algo::Reverse(Members); \
+			return Members; \
+		} \
 	} GCC_ALIGN(SHADER_PARAMETER_STRUCT_ALIGNMENT);
 
 /** Begins & ends a shader global parameter structure.
