Build fix: Add missing forward declare

diff --git a/Engine/Source/Runtime/RenderCore/Public/VirtualTexturing.h b/Engine/Source/Runtime/RenderCore/Public/VirtualTexturing.h
--- a/Engine/Source/Runtime/RenderCore/Public/VirtualTexturing.h
+++ b/Engine/Source/Runtime/RenderCore/Public/VirtualTexturing.h
@@ -10,6 +10,7 @@
 class FRHICommandListImmediate;
 class FRHIShaderResourceView;
 class FRHITexture;
+class FRHIUnorderedAccessView;
 
 union FVirtualTextureProducerHandle
 {
