Add project virtual texture support to shader compilation environment

diff --git a/Engine/Source/Runtime/Engine/Private/ShaderCompiler/ShaderCompiler.cpp b/Engine/Source/Runtime/Engine/Private/ShaderCompiler/ShaderCompiler.cpp
--- a/Engine/Source/Runtime/Engine/Private/ShaderCompiler/ShaderCompiler.cpp
+++ b/Engine/Source/Runtime/Engine/Private/ShaderCompiler/ShaderCompiler.cpp
@@ -3379,6 +3379,13 @@ void GlobalBeginCompileShader(
 		Input.Environment.SetDefine(TEXT("PROJECT_MOBILE_USE_LEGACY_SHADING"), CVar ? (CVar->GetInt() != 0) : 0);
 	}
 
+	{
+		const ERHIFeatureLevel::Type FeatureLevel = GetMaxSupportedFeatureLevel((EShaderPlatform)Target.Platform);
+		const ITargetPlatform* TargetPlatform = GetTargetPlatformManager()->FindTargetPlatform(ShaderPlatformToPlatformName(EShaderPlatform(Target.Platform)).ToString());
+		const bool bVirtualTextureEnabled = UseVirtualTexturing(FeatureLevel, TargetPlatform);
+		Input.Environment.SetDefine(TEXT("PROJECT_SUPPORT_VIRTUAL_TEXTURE"), bVirtualTextureEnabled ? 1 : 0);
+	}
+
 	{
 		static IConsoleVariable* CVar = IConsoleManager::Get().FindConsoleVariable(TEXT("r.Mobile.ForceFullPrecisionInPS"));
 		if (CVar && CVar->GetInt() != 0)
diff --git a/Engine/Source/Runtime/RenderCore/Private/Shader.cpp b/Engine/Source/Runtime/RenderCore/Private/Shader.cpp
--- a/Engine/Source/Runtime/RenderCore/Private/Shader.cpp
+++ b/Engine/Source/Runtime/RenderCore/Private/Shader.cpp
@@ -2414,9 +2414,10 @@ void ShaderMapAppendKeyString(EShaderPlatform Platform, FString& KeyString)
 		}
 	}
 
+	ITargetPlatform* TargetPlatform = GetTargetPlatformManager()->FindTargetPlatformWithSupport(TEXT("ShaderFormat"), LegacyShaderPlatformToShaderFormat(Platform));
+
 	{
 		bool bForwardShading = false;
-		ITargetPlatform* TargetPlatform = GetTargetPlatformManager()->FindTargetPlatformWithSupport(TEXT("ShaderFormat"), LegacyShaderPlatformToShaderFormat(Platform));
 		if (TargetPlatform)
 		{
 			// if there is a specific target platform that matches our shader platform, use that to drive forward shading
@@ -2435,6 +2436,13 @@ void ShaderMapAppendKeyString(EShaderPlatform Platform, FString& KeyString)
 		}
 	}
 
+	{
+		if (UseVirtualTexturing(GetMaxSupportedFeatureLevel(Platform), TargetPlatform))
+		{
+			KeyString += TEXT("_VT");
+		}
+	}
+
 	{
 		static const auto CVar = IConsoleManager::Get().FindTConsoleVariableDataInt(TEXT("r.PostProcessing.PropagateAlpha"));
 		if (CVar && CVar->GetValueOnAnyThread() > 0)
