Fix shadowed variable

diff --git a/Engine/Source/Runtime/RenderCore/Private/Shader.cpp b/Engine/Source/Runtime/RenderCore/Private/Shader.cpp
--- a/Engine/Source/Runtime/RenderCore/Private/Shader.cpp
+++ b/Engine/Source/Runtime/RenderCore/Private/Shader.cpp
@@ -2527,9 +2527,9 @@ void ShaderMapAppendKeyString(EShaderPlatform Platform, FString& KeyString)
 
 		ITargetPlatformManagerModule* TPM = GetTargetPlatformManager();
 		check(TPM);
-		auto TargetPlatform = TPM->GetRunningTargetPlatform();
-		check(TargetPlatform);
-		const bool VTSupported = TargetPlatform->SupportsFeature(ETargetPlatformFeatures::VirtualTextureStreaming);
+		auto RunningTargetPlatform = TPM->GetRunningTargetPlatform();
+		check(RunningTargetPlatform);
+		const bool VTSupported = RunningTargetPlatform->SupportsFeature(ETargetPlatformFeatures::VirtualTextureStreaming);
 
 		auto tt = FString::Printf(TEXT("_VT-%d-%d-%d-%d"), VTLightmaps, VTTextures, VTFeedbackFactor, VTSupported);
  		KeyString += tt;
