Fix unused extra depth/stencil created and destroyed when depth/stencil HMD-created.

FSceneRenderTargets::AllocateCommonDepthTargets gives the HMD the option to create the depth/stencil textures. When following that path (e.g. with the OculusHMD plugin), FRenderTargetPool::FindFreeElement will (the first time through) create the resource, which will be immediately discarded if the HMD created one. The fix was to add a boolean flag to FRenderTargetPool::FindFreeElement to signal deferred resource creation to avoid this extra create/destroy cycle.

diff --git a/Engine/Source/Runtime/RenderCore/Private/RenderTargetPool.cpp b/Engine/Source/Runtime/RenderCore/Private/RenderTargetPool.cpp
--- a/Engine/Source/Runtime/RenderCore/Private/RenderTargetPool.cpp
+++ b/Engine/Source/Runtime/RenderCore/Private/RenderTargetPool.cpp
@@ -285,7 +285,7 @@ static void ClobberAllocatedRenderTarget(FRHICommandList& RHICmdList, const FPoo
 
 #endif
 
-bool FRenderTargetPool::FindFreeElement(FRHICommandList& RHICmdList, const FPooledRenderTargetDesc& InputDesc, TRefCountPtr<IPooledRenderTarget> &Out, const TCHAR* InDebugName, bool bDoWritableBarrier, ERenderTargetTransience TransienceHint)
+bool FRenderTargetPool::FindFreeElement(FRHICommandList& RHICmdList, const FPooledRenderTargetDesc& InputDesc, TRefCountPtr<IPooledRenderTarget> &Out, const TCHAR* InDebugName, bool bDoWritableBarrier, ERenderTargetTransience TransienceHint, bool bDeferTextureAllocation)
 {
 	check(IsInRenderingThread());
 
@@ -421,8 +421,9 @@ Done:
 		FRHIResourceCreateInfo CreateInfo(Desc.ClearValue);
 		CreateInfo.DebugName = InDebugName;
 
-		if(Desc.TargetableFlags & (TexCreate_RenderTargetable | TexCreate_DepthStencilTargetable | TexCreate_UAV))
+		if(Desc.TargetableFlags & (TexCreate_RenderTargetable | TexCreate_DepthStencilTargetable | TexCreate_UAV) && !bDeferTextureAllocation)
 		{
+			// Only create resources if we're not asked to defer creation.
 			if(Desc.Is2DTexture())
 			{
 				if (!Desc.IsArray())
@@ -533,8 +534,9 @@ Done:
 
 			RHIBindDebugLabelName(Found->RenderTargetItem.TargetableTexture, InDebugName);
 		}
-		else
+		else if (!bDeferTextureAllocation)
 		{
+			// Only create resources if we're not asked to defer creation.
 			if(Desc.Is2DTexture())
 			{
 				// this is useful to get a CPU lockable texture through the same interface
@@ -576,7 +578,7 @@ Done:
 			RHIBindDebugLabelName(Found->RenderTargetItem.ShaderResourceTexture, InDebugName);
 		}
 
-		if(Desc.TargetableFlags & TexCreate_UAV)
+		if((Desc.TargetableFlags & TexCreate_UAV) && !bDeferTextureAllocation)
 		{
 			// The render target desc is invalid if a UAV is requested with an RHI that doesn't support the high-end feature level.
 			check(GMaxRHIFeatureLevel == ERHIFeatureLevel::SM5);
@@ -589,8 +591,13 @@ Done:
 			Found->RenderTargetItem.UAV = Found->RenderTargetItem.MipUAVs[0];
 		}
 
-		AllocationLevelInKB += ComputeSizeInKB(*Found);
-		VerifyAllocationLevel();
+		if (!bDeferTextureAllocation)
+		{
+			// Only calculate allocation level if we actually allocated something. If bDeferTextureAllocation is true, the caller should call 
+			// UpdateElementSize once it's set the resources on the created object.
+			AllocationLevelInKB += ComputeSizeInKB(*Found);
+			VerifyAllocationLevel();
+		}
 
 		FoundIndex = PooledRenderTargets.Num() - 1;
 
diff --git a/Engine/Source/Runtime/RenderCore/Public/RenderTargetPool.h b/Engine/Source/Runtime/RenderCore/Public/RenderTargetPool.h
--- a/Engine/Source/Runtime/RenderCore/Public/RenderTargetPool.h
+++ b/Engine/Source/Runtime/RenderCore/Public/RenderTargetPool.h
@@ -223,7 +223,7 @@ public:
 	 * call from RenderThread only
 	 * @return true if the old element was still valid, false if a new one was assigned
 	 */
-	bool FindFreeElement(FRHICommandList& RHICmdList, const FPooledRenderTargetDesc& Desc, TRefCountPtr<IPooledRenderTarget>& Out, const TCHAR* InDebugName, bool bDoWritableBarrier = true, ERenderTargetTransience TransienceHint = ERenderTargetTransience::Transient );
+	bool FindFreeElement(FRHICommandList& RHICmdList, const FPooledRenderTargetDesc& Desc, TRefCountPtr<IPooledRenderTarget>& Out, const TCHAR* InDebugName, bool bDoWritableBarrier = true, ERenderTargetTransience TransienceHint = ERenderTargetTransience::Transient, bool bDeferTextureAllocation = false);
 
 	void CreateUntrackedElement(const FPooledRenderTargetDesc& Desc, TRefCountPtr<IPooledRenderTarget>& Out, const FSceneRenderTargetItem& Item);
 
diff --git a/Engine/Source/Runtime/Renderer/Private/PostProcess/SceneRenderTargets.cpp b/Engine/Source/Runtime/Renderer/Private/PostProcess/SceneRenderTargets.cpp
--- a/Engine/Source/Runtime/Renderer/Private/PostProcess/SceneRenderTargets.cpp
+++ b/Engine/Source/Runtime/Renderer/Private/PostProcess/SceneRenderTargets.cpp
@@ -2155,7 +2155,7 @@ void FSceneRenderTargets::AllocateCommonDepthTargets(FRHICommandList& RHICmdList
 		FPooledRenderTargetDesc Desc(FPooledRenderTargetDesc::Create2DDesc(BufferSize, PF_DepthStencil, DefaultDepthClear, TexCreate_None, TexCreate_DepthStencilTargetable | TexCreate_ShaderResource | TexCreate_InputAttachmentRead, false));
 		Desc.NumSamples = GetNumSceneColorMSAASamples(CurrentFeatureLevel);
 		Desc.Flags |= GFastVRamConfig.SceneDepth;
-		GRenderTargetPool.FindFreeElement(RHICmdList, Desc, SceneDepthZ, TEXT("SceneDepthZ"));
+		GRenderTargetPool.FindFreeElement(RHICmdList, Desc, SceneDepthZ, TEXT("SceneDepthZ"), true, ERenderTargetTransience::Transient, bHMDAllocated);
 
 		if (bHMDAllocated)
 		{
