Change vertex factor interface to use an enum instead of a bool, for selecting between default vertex stream and position only stream.

diff --git a/Engine/Plugins/Experimental/AR/PointCloud/Source/PointCloud/Private/PointCloudVertexFactory.cpp b/Engine/Plugins/Experimental/AR/PointCloud/Source/PointCloud/Private/PointCloudVertexFactory.cpp
--- a/Engine/Plugins/Experimental/AR/PointCloud/Source/PointCloud/Private/PointCloudVertexFactory.cpp
+++ b/Engine/Plugins/Experimental/AR/PointCloud/Source/PointCloud/Private/PointCloudVertexFactory.cpp
@@ -34,7 +34,7 @@ public:
 		const class FSceneInterface* Scene,
 		const class FSceneView* View,
 		const class FMeshMaterialShader* Shader,
-		bool bShaderRequiresPositionOnlyStream,
+		const EVertexInputStreamType InputStreamType,
 		ERHIFeatureLevel::Type FeatureLevel,
 		const class FVertexFactory* InVertexFactory,
 		const struct FMeshBatchElement& BatchElement,
diff --git a/Engine/Plugins/FX/Niagara/Source/NiagaraVertexFactories/Private/NiagaraMeshVertexFactory.cpp b/Engine/Plugins/FX/Niagara/Source/NiagaraVertexFactories/Private/NiagaraMeshVertexFactory.cpp
--- a/Engine/Plugins/FX/Niagara/Source/NiagaraVertexFactories/Private/NiagaraMeshVertexFactory.cpp
+++ b/Engine/Plugins/FX/Niagara/Source/NiagaraVertexFactories/Private/NiagaraMeshVertexFactory.cpp
@@ -54,7 +54,7 @@ public:
 		const FSceneInterface* Scene,
 		const FSceneView* View,
 		const FMeshMaterialShader* Shader,
-		bool bShaderRequiresPositionOnlyStream,
+		const EVertexInputStreamType InputStreamType,
 		ERHIFeatureLevel::Type FeatureLevel,
 		const FVertexFactory* VertexFactory,
 		const FMeshBatchElement& BatchElement,
diff --git a/Engine/Plugins/FX/Niagara/Source/NiagaraVertexFactories/Private/NiagaraRibbonVertexFactory.cpp b/Engine/Plugins/FX/Niagara/Source/NiagaraVertexFactories/Private/NiagaraRibbonVertexFactory.cpp
--- a/Engine/Plugins/FX/Niagara/Source/NiagaraVertexFactories/Private/NiagaraRibbonVertexFactory.cpp
+++ b/Engine/Plugins/FX/Niagara/Source/NiagaraVertexFactories/Private/NiagaraRibbonVertexFactory.cpp
@@ -65,7 +65,7 @@ public:
 		const FSceneInterface* Scene,
 		const FSceneView* View,
 		const FMeshMaterialShader* Shader,
-		bool bShaderRequiresPositionOnlyStream,
+		const EVertexInputStreamType InputStreamType,
 		ERHIFeatureLevel::Type FeatureLevel,
 		const FVertexFactory* VertexFactory,
 		const FMeshBatchElement& BatchElement,
@@ -117,7 +117,7 @@ public:
 		const FSceneInterface* Scene,
 		const FSceneView* View,
 		const FMeshMaterialShader* Shader,
-		bool bShaderRequiresPositionOnlyStream,
+		const EVertexInputStreamType InputStreamType,
 		ERHIFeatureLevel::Type FeatureLevel,
 		const FVertexFactory* VertexFactory,
 		const FMeshBatchElement& BatchElement,
diff --git a/Engine/Plugins/FX/Niagara/Source/NiagaraVertexFactories/Private/NiagaraSpriteVertexFactory.cpp b/Engine/Plugins/FX/Niagara/Source/NiagaraVertexFactories/Private/NiagaraSpriteVertexFactory.cpp
--- a/Engine/Plugins/FX/Niagara/Source/NiagaraVertexFactories/Private/NiagaraSpriteVertexFactory.cpp
+++ b/Engine/Plugins/FX/Niagara/Source/NiagaraVertexFactories/Private/NiagaraSpriteVertexFactory.cpp
@@ -79,7 +79,7 @@ public:
 		const FSceneInterface* Scene,
 		const FSceneView* View,
 		const FMeshMaterialShader* Shader,
-		bool bShaderRequiresPositionOnlyStream,
+		const EVertexInputStreamType VertexStreamType,
 		ERHIFeatureLevel::Type FeatureLevel,
 		const FVertexFactory* VertexFactory,
 		const FMeshBatchElement& BatchElement,
@@ -133,7 +133,7 @@ public:
 		const FSceneInterface* Scene,
 		const FSceneView* View,
 		const FMeshMaterialShader* Shader,
-		bool bShaderRequiresPositionOnlyStream,
+		const EVertexInputStreamType InputStreamType,
 		ERHIFeatureLevel::Type FeatureLevel,
 		const FVertexFactory* VertexFactory,
 		const FMeshBatchElement& BatchElement,
diff --git a/Engine/Source/Runtime/Engine/Private/Components/SplineMeshComponent.cpp b/Engine/Source/Runtime/Engine/Private/Components/SplineMeshComponent.cpp
--- a/Engine/Source/Runtime/Engine/Private/Components/SplineMeshComponent.cpp
+++ b/Engine/Source/Runtime/Engine/Private/Components/SplineMeshComponent.cpp
@@ -61,7 +61,7 @@ void FSplineMeshVertexFactoryShaderParameters::GetElementShaderBindings(
 	const class FSceneInterface* Scene,
 	const FSceneView* View,
 	const class FMeshMaterialShader* Shader,
-	bool bShaderRequiresPositionOnlyStream,
+	const EVertexInputStreamType InputStreamType,
 	ERHIFeatureLevel::Type FeatureLevel,
 	const FVertexFactory* VertexFactory,
 	const FMeshBatchElement& BatchElement,
diff --git a/Engine/Source/Runtime/Engine/Private/GPUSkinVertexFactory.cpp b/Engine/Source/Runtime/Engine/Private/GPUSkinVertexFactory.cpp
--- a/Engine/Source/Runtime/Engine/Private/GPUSkinVertexFactory.cpp
+++ b/Engine/Source/Runtime/Engine/Private/GPUSkinVertexFactory.cpp
@@ -529,7 +529,7 @@ public:
 		const FSceneInterface* Scene,
 		const FSceneView* View,
 		const FMeshMaterialShader* Shader,
-		bool bShaderRequiresPositionOnlyStream,
+		const EVertexInputStreamType InputStreamType,
 		ERHIFeatureLevel::Type FeatureLevel,
 		const FVertexFactory* VertexFactory,
 		const FMeshBatchElement& BatchElement,
@@ -613,7 +613,7 @@ public:
 		const FSceneInterface* Scene,
 		const FSceneView* View,
 		const FMeshMaterialShader* Shader,
-		bool bShaderRequiresPositionOnlyStream,
+		const EVertexInputStreamType InputStreamType,
 		ERHIFeatureLevel::Type FeatureLevel,
 		const FVertexFactory* VertexFactory,
 		const FMeshBatchElement& BatchElement,
@@ -629,7 +629,7 @@ public:
 		VertexFactoryUniformBuffer = LocalVertexFactory->GetUniformBuffer();
 
 		// #dxr_todo do we need this call to the base?
-		FLocalVertexFactoryShaderParametersBase::GetElementShaderBindingsBase(Scene, View, Shader, bShaderRequiresPositionOnlyStream, FeatureLevel, VertexFactory, BatchElement, VertexFactoryUniformBuffer, ShaderBindings, VertexStreams);
+		FLocalVertexFactoryShaderParametersBase::GetElementShaderBindingsBase(Scene, View, Shader, InputStreamType, FeatureLevel, VertexFactory, BatchElement, VertexFactoryUniformBuffer, ShaderBindings, VertexStreams);
 		FGPUSkinCache::GetShaderBindings(BatchUserData->Entry, BatchUserData->Section, Shader, (const FGPUSkinPassthroughVertexFactory*)VertexFactory, BatchElement.MinVertexIndex, GPUSkinCachePreviousPositionBuffer, ShaderBindings, VertexStreams);
 	}
 
@@ -845,7 +845,7 @@ public:
 		const FSceneInterface* Scene,
 		const FSceneView* View,
 		const FMeshMaterialShader* Shader,
-		bool bShaderRequiresPositionOnlyStream,
+		const EVertexInputStreamType InputStreamType,
 		ERHIFeatureLevel::Type FeatureLevel,
 		const FVertexFactory* VertexFactory,
 		const FMeshBatchElement& BatchElement,
@@ -853,7 +853,7 @@ public:
 		FVertexInputStreamArray& VertexStreams) const override
 	{
 		// Call regular GPU skinning shader parameters
-		FGPUSkinVertexFactoryShaderParameters::GetElementShaderBindings(Scene, View, Shader, bShaderRequiresPositionOnlyStream, FeatureLevel, VertexFactory, BatchElement, ShaderBindings, VertexStreams);
+		FGPUSkinVertexFactoryShaderParameters::GetElementShaderBindings(Scene, View, Shader, InputStreamType, FeatureLevel, VertexFactory, BatchElement, ShaderBindings, VertexStreams);
 		const auto* GPUSkinVertexFactory = (const FGPUBaseSkinVertexFactory*)VertexFactory;
 		// A little hacky; problem is we can't upcast from FGPUBaseSkinVertexFactory to FGPUBaseSkinAPEXClothVertexFactory as they are unrelated; a nice solution would be
 		// to use virtual inheritance, but that requires RTTI and complicates things further...
diff --git a/Engine/Source/Runtime/Engine/Private/GeometryCacheVertexFactory.cpp b/Engine/Source/Runtime/Engine/Private/GeometryCacheVertexFactory.cpp
--- a/Engine/Source/Runtime/Engine/Private/GeometryCacheVertexFactory.cpp
+++ b/Engine/Source/Runtime/Engine/Private/GeometryCacheVertexFactory.cpp
@@ -53,7 +53,7 @@ public:
 		const class FSceneInterface* Scene,
 		const FSceneView* View,
 		const class FMeshMaterialShader* Shader,
-		bool bShaderRequiresPositionOnlyStream,
+		const EVertexInputStreamType InputStreamType,
 		ERHIFeatureLevel::Type FeatureLevel,
 		const FVertexFactory* GenericVertexFactory,
 		const FMeshBatchElement& BatchElement,
@@ -186,9 +186,11 @@ void FGeometryCacheVertexVertexFactory::InitRHI()
 	// then initialize PositionStream and PositionDeclaration.
 	if (Data.PositionComponent.VertexBuffer != Data.TangentBasisComponents[0].VertexBuffer)
 	{
-		FVertexDeclarationElementList PositionOnlyStreamElements;
-		PositionOnlyStreamElements.Add(AccessPositionStreamComponent(Data.PositionComponent, 0));
-		InitPositionDeclaration(PositionOnlyStreamElements);
+		{
+			FVertexDeclarationElementList PositionOnlyStreamElements;
+			PositionOnlyStreamElements.Add(AccessStreamComponent(Data.PositionComponent, 0, EVertexInputStreamType::PositionOnly));
+			InitDeclaration(PositionOnlyStreamElements, EVertexInputStreamType::PositionOnly);
+		}
 	}
 
 	FVertexDeclarationElementList Elements;
diff --git a/Engine/Source/Runtime/Engine/Private/InstancedStaticMesh.cpp b/Engine/Source/Runtime/Engine/Private/InstancedStaticMesh.cpp
--- a/Engine/Source/Runtime/Engine/Private/InstancedStaticMesh.cpp
+++ b/Engine/Source/Runtime/Engine/Private/InstancedStaticMesh.cpp
@@ -487,18 +487,26 @@ void FInstancedStaticMeshVertexFactory::InitRHI()
 	// then initialize PositionStream and PositionDeclaration.
 	if(Data.PositionComponent.VertexBuffer != Data.TangentBasisComponents[0].VertexBuffer)
 	{
-		FVertexDeclarationElementList PositionOnlyStreamElements;
-		PositionOnlyStreamElements.Add(AccessPositionStreamComponent(Data.PositionComponent,0));
-
-		if (bInstanced)
+		auto AddDeclaration = [&Data](EVertexInputStreamType InputStreamType, bool bInstanced, bool bAddNormal)
 		{
-			// toss in the instanced location stream
-			PositionOnlyStreamElements.Add(AccessPositionStreamComponent(Data.InstanceOriginComponent,8));
-			PositionOnlyStreamElements.Add(AccessPositionStreamComponent(Data.InstanceTransformComponent[0],9));
-			PositionOnlyStreamElements.Add(AccessPositionStreamComponent(Data.InstanceTransformComponent[1],10));
-			PositionOnlyStreamElements.Add(AccessPositionStreamComponent(Data.InstanceTransformComponent[2],11));
-		}
-		InitPositionDeclaration(PositionOnlyStreamElements);
+			FVertexDeclarationElementList StreamElements;
+			StreamElements.Add(AccessPositionStreamComponent(Data.PositionComponent, 0));
+
+			if (bAddNormal)
+				StreamElements.Add(AccessPositionStreamComponent(Data.TangentBasisComponents[2], 2));
+
+			if (bInstanced)
+			{
+				// toss in the instanced location stream
+				StreamElements.Add(AccessPositionStreamComponent(Data.InstanceOriginComponent, 8));
+				StreamElements.Add(AccessPositionStreamComponent(Data.InstanceTransformComponent[0], 9));
+				StreamElements.Add(AccessPositionStreamComponent(Data.InstanceTransformComponent[1], 10));
+				StreamElements.Add(AccessPositionStreamComponent(Data.InstanceTransformComponent[2], 11));
+			}
+
+			InitDeclaration(StreamElements, InputStreamType);
+		};
+		AddDeclaration(EVertexInputStreamType::PositionOnly, bInstanced, false);
 	}
 #endif
 
@@ -2437,7 +2445,7 @@ void FInstancedStaticMeshVertexFactoryShaderParameters::GetElementShaderBindings
 	const class FSceneInterface* Scene,
 	const FSceneView* View,
 	const FMeshMaterialShader* Shader,
-	bool bShaderRequiresPositionOnlyStream,
+	const EVertexInputStreamType InputStreamType,
 	ERHIFeatureLevel::Type FeatureLevel,
 	const FVertexFactory* VertexFactory,
 	const FMeshBatchElement& BatchElement,
@@ -2447,7 +2455,7 @@ void FInstancedStaticMeshVertexFactoryShaderParameters::GetElementShaderBindings
 {
 	// Decode VertexFactoryUserData as VertexFactoryUniformBuffer
 	FUniformBufferRHIParamRef VertexFactoryUniformBuffer = static_cast<FUniformBufferRHIParamRef>(BatchElement.VertexFactoryUserData);
-	FLocalVertexFactoryShaderParametersBase::GetElementShaderBindingsBase(Scene, View, Shader, bShaderRequiresPositionOnlyStream, FeatureLevel, VertexFactory, BatchElement, VertexFactoryUniformBuffer, ShaderBindings, VertexStreams);
+	FLocalVertexFactoryShaderParametersBase::GetElementShaderBindingsBase(Scene, View, Shader, InputStreamType, FeatureLevel, VertexFactory, BatchElement, VertexFactoryUniformBuffer, ShaderBindings, VertexStreams);
 
 	const FInstancingUserData* InstancingUserData = (const FInstancingUserData*)BatchElement.UserData;
 	const auto* InstancedVertexFactory = static_cast<const FInstancedStaticMeshVertexFactory*>(VertexFactory);
@@ -2465,7 +2473,7 @@ void FInstancedStaticMeshVertexFactoryShaderParameters::GetElementShaderBindings
 
 	if (GRHISupportsInstancing && InstanceOffsetValue > 0 && VertexStreams.Num() > 0)
 	{
-		VertexFactory->OffsetInstanceStreams(InstanceOffsetValue, bShaderRequiresPositionOnlyStream, VertexStreams);
+		VertexFactory->OffsetInstanceStreams(InstanceOffsetValue, InputStreamType, VertexStreams);
 	}
 
 	if( InstancingWorldViewOriginOneParameter.IsBound() )
diff --git a/Engine/Source/Runtime/Engine/Private/InstancedStaticMesh.h b/Engine/Source/Runtime/Engine/Private/InstancedStaticMesh.h
--- a/Engine/Source/Runtime/Engine/Private/InstancedStaticMesh.h
+++ b/Engine/Source/Runtime/Engine/Private/InstancedStaticMesh.h
@@ -333,7 +333,7 @@ class FInstancedStaticMeshVertexFactoryShaderParameters : public FLocalVertexFac
 		const class FSceneInterface* Scene,
 		const FSceneView* View,
 		const FMeshMaterialShader* Shader,
-		bool bShaderRequiresPositionOnlyStream,
+		const EVertexInputStreamType InputStreamType,
 		ERHIFeatureLevel::Type FeatureLevel,
 		const FVertexFactory* VertexFactory,
 		const FMeshBatchElement& BatchElement,
diff --git a/Engine/Source/Runtime/Engine/Private/LocalVertexFactory.cpp b/Engine/Source/Runtime/Engine/Private/LocalVertexFactory.cpp
--- a/Engine/Source/Runtime/Engine/Private/LocalVertexFactory.cpp
+++ b/Engine/Source/Runtime/Engine/Private/LocalVertexFactory.cpp
@@ -94,7 +94,7 @@ void FLocalVertexFactoryShaderParametersBase::GetElementShaderBindingsBase(
 	const FSceneInterface* Scene,
 	const FSceneView* View,
 	const FMeshMaterialShader* Shader, 
-	bool bShaderRequiresPositionOnlyStream,
+	const EVertexInputStreamType InputStreamType,
 	ERHIFeatureLevel::Type FeatureLevel,
 	const FVertexFactory* VertexFactory, 
 	const FMeshBatchElement& BatchElement,
@@ -152,7 +152,7 @@ void FLocalVertexFactoryShaderParameters::GetElementShaderBindings(
 	const FSceneInterface* Scene,
 	const FSceneView* View,
 	const FMeshMaterialShader* Shader,
-	bool bShaderRequiresPositionOnlyStream,
+	const EVertexInputStreamType InputStreamType,
 	ERHIFeatureLevel::Type FeatureLevel,
 	const FVertexFactory* VertexFactory,
 	const FMeshBatchElement& BatchElement,
@@ -167,7 +167,7 @@ void FLocalVertexFactoryShaderParameters::GetElementShaderBindings(
 		Scene,
 		View,
 		Shader,
-		bShaderRequiresPositionOnlyStream,
+		InputStreamType,
 		FeatureLevel,
 		VertexFactory,
 		BatchElement,
@@ -257,20 +257,31 @@ void FLocalVertexFactory::InitRHI()
 
 	// If the vertex buffer containing position is not the same vertex buffer containing the rest of the data,
 	// then initialize PositionStream and PositionDeclaration.
-	if(Data.PositionComponent.VertexBuffer != Data.TangentBasisComponents[0].VertexBuffer)
+	if (Data.PositionComponent.VertexBuffer != Data.TangentBasisComponents[0].VertexBuffer)
 	{
-		FVertexDeclarationElementList PositionOnlyStreamElements;
-		PositionOnlyStreamElements.Add(AccessPositionStreamComponent(Data.PositionComponent,0));
-
-		PositionOnlyPrimitiveIdStreamIndex = -1;
-		if (GetType()->SupportsPrimitiveIdStream() && bCanUseGPUScene)
+		auto AddDeclaration = [this, bCanUseGPUScene](EVertexInputStreamType InputStreamType, bool bAddNormal)
 		{
-			// When the VF is used for rendering in normal mesh passes, this vertex buffer and offset will be overridden
-			PositionOnlyStreamElements.Add(AccessPositionStreamComponent(FVertexStreamComponent(&GPrimitiveIdDummy, 0, 0, sizeof(uint32), VET_UInt, EVertexStreamUsage::Instancing), 1));
-			PositionOnlyPrimitiveIdStreamIndex = PositionOnlyStreamElements.Last().StreamIndex;
-		}
-
-		InitPositionDeclaration(PositionOnlyStreamElements);
+			FVertexDeclarationElementList StreamElements;
+			StreamElements.Add(AccessStreamComponent(Data.PositionComponent, 0, InputStreamType));
+
+			bAddNormal = bAddNormal && Data.TangentBasisComponents[1].VertexBuffer != NULL;
+			if (bAddNormal)
+			{
+				StreamElements.Add(AccessStreamComponent(Data.TangentBasisComponents[1], 2, InputStreamType));
+			}
+
+			const uint8 TypeIndex = static_cast<uint8>(InputStreamType);
+			PrimitiveIdStreamIndex[TypeIndex] = -1;
+			if (GetType()->SupportsPrimitiveIdStream() && bCanUseGPUScene)
+			{
+				// When the VF is used for rendering in normal mesh passes, this vertex buffer and offset will be overridden
+				StreamElements.Add(AccessStreamComponent(FVertexStreamComponent(&GPrimitiveIdDummy, 0, 0, sizeof(uint32), VET_UInt, EVertexStreamUsage::Instancing), 1, InputStreamType));
+				PrimitiveIdStreamIndex[TypeIndex] = StreamElements.Last().StreamIndex;
+			}
+
+			InitDeclaration(StreamElements, InputStreamType);
+		};
+		AddDeclaration(EVertexInputStreamType::PositionOnly, false);
 	}
 
 	FVertexDeclarationElementList Elements;
@@ -279,12 +290,15 @@ void FLocalVertexFactory::InitRHI()
 		Elements.Add(AccessStreamComponent(Data.PositionComponent,0));
 	}
 
-	PrimitiveIdStreamIndex = -1;
-	if (GetType()->SupportsPrimitiveIdStream() && bCanUseGPUScene)
 	{
-		// When the VF is used for rendering in normal mesh passes, this vertex buffer and offset will be overridden
-		Elements.Add(AccessStreamComponent(FVertexStreamComponent(&GPrimitiveIdDummy, 0, 0, sizeof(uint32), VET_UInt, EVertexStreamUsage::Instancing), 13));
-		PrimitiveIdStreamIndex = Elements.Last().StreamIndex;
+		const uint8 Index = static_cast<uint8>(EVertexInputStreamType::Default);
+		PrimitiveIdStreamIndex[Index] = -1;
+		if (GetType()->SupportsPrimitiveIdStream() && bCanUseGPUScene)
+		{
+			// When the VF is used for rendering in normal mesh passes, this vertex buffer and offset will be overridden
+			Elements.Add(AccessStreamComponent(FVertexStreamComponent(&GPrimitiveIdDummy, 0, 0, sizeof(uint32), VET_UInt, EVertexStreamUsage::Instancing), 13));
+			PrimitiveIdStreamIndex[Index] = Elements.Last().StreamIndex;
+		}
 	}
 
 	// only tangent,normal are used by the stream. the binormal is derived in the shader
diff --git a/Engine/Source/Runtime/Engine/Private/MeshParticleVertexFactory.cpp b/Engine/Source/Runtime/Engine/Private/MeshParticleVertexFactory.cpp
--- a/Engine/Source/Runtime/Engine/Private/MeshParticleVertexFactory.cpp
+++ b/Engine/Source/Runtime/Engine/Private/MeshParticleVertexFactory.cpp
@@ -51,7 +51,7 @@ public:
 		const FSceneInterface* Scene,
 		const FSceneView* View,
 		const FMeshMaterialShader* Shader,
-		bool bShaderRequiresPositionOnlyStream,
+		const EVertexInputStreamType InputStreamType,
 		ERHIFeatureLevel::Type FeatureLevel,
 		const FVertexFactory* VertexFactory,
 		const FMeshBatchElement& BatchElement,
diff --git a/Engine/Source/Runtime/Engine/Private/Particles/ParticleBeamTrailVertexFactory.cpp b/Engine/Source/Runtime/Engine/Private/Particles/ParticleBeamTrailVertexFactory.cpp
--- a/Engine/Source/Runtime/Engine/Private/Particles/ParticleBeamTrailVertexFactory.cpp
+++ b/Engine/Source/Runtime/Engine/Private/Particles/ParticleBeamTrailVertexFactory.cpp
@@ -30,7 +30,7 @@ public:
 		const FSceneInterface* Scene,
 		const FSceneView* View,
 		const FMeshMaterialShader* Shader,
-		bool bShaderRequiresPositionOnlyStream,
+		const EVertexInputStreamType InputStreamType,
 		ERHIFeatureLevel::Type FeatureLevel,
 		const FVertexFactory* VertexFactory,
 		const FMeshBatchElement& BatchElement,
diff --git a/Engine/Source/Runtime/Engine/Private/Particles/ParticleGpuSimulation.cpp b/Engine/Source/Runtime/Engine/Private/Particles/ParticleGpuSimulation.cpp
--- a/Engine/Source/Runtime/Engine/Private/Particles/ParticleGpuSimulation.cpp
+++ b/Engine/Source/Runtime/Engine/Private/Particles/ParticleGpuSimulation.cpp
@@ -583,7 +583,7 @@ public:
 		const FSceneInterface* Scene,
 		const FSceneView* View,
 		const FMeshMaterialShader* Shader,
-		bool bShaderRequiresPositionOnlyStream,
+		const EVertexInputStreamType InputStreamType,
 		ERHIFeatureLevel::Type FeatureLevel,
 		const FVertexFactory* VertexFactory,
 		const FMeshBatchElement& BatchElement,
@@ -626,7 +626,7 @@ public:
 		const FSceneInterface* Scene,
 		const FSceneView* View,
 		const FMeshMaterialShader* Shader,
-		bool bShaderRequiresPositionOnlyStream,
+		const EVertexInputStreamType InputStreamType,
 		ERHIFeatureLevel::Type FeatureLevel,
 		const FVertexFactory* VertexFactory,
 		const FMeshBatchElement& BatchElement,
@@ -788,7 +788,7 @@ void FGPUSpriteVertexFactoryShaderParametersVS::GetElementShaderBindings(
 	const FSceneInterface* Scene,
 	const FSceneView* View,
 	const FMeshMaterialShader* Shader,
-	bool bShaderRequiresPositionOnlyStream,
+	const EVertexInputStreamType InputStreamType,
 	ERHIFeatureLevel::Type FeatureLevel,
 	const FVertexFactory* VertexFactory,
 	const FMeshBatchElement& BatchElement,
@@ -814,7 +814,7 @@ void FGPUSpriteVertexFactoryShaderParametersPS::GetElementShaderBindings(
 	const FSceneInterface* Scene,
 	const FSceneView* View,
 	const FMeshMaterialShader* Shader,
-	bool bShaderRequiresPositionOnlyStream,
+	const EVertexInputStreamType InputStreamType,
 	ERHIFeatureLevel::Type FeatureLevel,
 	const FVertexFactory* VertexFactory,
 	const FMeshBatchElement& BatchElement,
diff --git a/Engine/Source/Runtime/Engine/Private/Particles/ParticleVertexFactory.cpp b/Engine/Source/Runtime/Engine/Private/Particles/ParticleVertexFactory.cpp
--- a/Engine/Source/Runtime/Engine/Private/Particles/ParticleVertexFactory.cpp
+++ b/Engine/Source/Runtime/Engine/Private/Particles/ParticleVertexFactory.cpp
@@ -81,7 +81,7 @@ public:
 		const FSceneInterface* Scene,
 		const FSceneView* View,
 		const FMeshMaterialShader* Shader,
-		bool bShaderRequiresPositionOnlyStream,
+		const EVertexInputStreamType InputStreamType,
 		ERHIFeatureLevel::Type FeatureLevel,
 		const FVertexFactory* VertexFactory,
 		const FMeshBatchElement& BatchElement,
@@ -109,7 +109,7 @@ public:
 		const FSceneInterface* Scene,
 		const FSceneView* View,
 		const FMeshMaterialShader* Shader,
-		bool bShaderRequiresPositionOnlyStream,
+		const EVertexInputStreamType InputStreamType,
 		ERHIFeatureLevel::Type FeatureLevel,
 		const FVertexFactory* VertexFactory,
 		const FMeshBatchElement& BatchElement,
diff --git a/Engine/Source/Runtime/Engine/Private/SceneManagement.cpp b/Engine/Source/Runtime/Engine/Private/SceneManagement.cpp
--- a/Engine/Source/Runtime/Engine/Private/SceneManagement.cpp
+++ b/Engine/Source/Runtime/Engine/Private/SceneManagement.cpp
@@ -271,7 +271,7 @@ void FMeshElementCollector::AddMesh(int32 ViewIndex, FMeshBatch& MeshBatch)
 	}
 
 	// If we are maintaining primitive scene data on the GPU, copy the primitive uniform buffer data to a unified array so it can be uploaded later
-	if (UseGPUScene(GMaxRHIShaderPlatform, FeatureLevel) && MeshBatch.VertexFactory->GetPrimitiveIdStreamIndex(false) >= 0)
+	if (UseGPUScene(GMaxRHIShaderPlatform, FeatureLevel) && MeshBatch.VertexFactory->GetPrimitiveIdStreamIndex(EVertexInputStreamType::Default) >= 0)
 	{
 		for (int32 Index = 0; Index < MeshBatch.Elements.Num(); ++Index)
 		{
@@ -918,7 +918,7 @@ void FMeshBatch::PreparePrimitiveUniformBuffer(const FPrimitiveSceneProxy* Primi
 	const bool bVFSupportsPrimitiveIdStream = VertexFactory->GetType()->SupportsPrimitiveIdStream();
 	checkf((PrimitiveSceneProxy->DoesVFRequirePrimitiveUniformBuffer() || bVFSupportsPrimitiveIdStream), TEXT("PrimitiveSceneProxy has bVFRequiresPrimitiveUniformBuffer disabled yet tried to draw with a vertex factory (%s) that did not support PrimitiveIdStream."), VertexFactory->GetType()->GetName());
 
-	const bool bPrimitiveShaderDataComesFromSceneBuffer = VertexFactory->GetPrimitiveIdStreamIndex(false) >= 0;
+	const bool bPrimitiveShaderDataComesFromSceneBuffer = VertexFactory->GetPrimitiveIdStreamIndex(EVertexInputStreamType::Default) >= 0;
 
 	for (int32 ElementIndex = 0; ElementIndex < Elements.Num(); ElementIndex++)
 	{
diff --git a/Engine/Source/Runtime/Engine/Private/SplineMeshSceneProxy.h b/Engine/Source/Runtime/Engine/Private/SplineMeshSceneProxy.h
--- a/Engine/Source/Runtime/Engine/Private/SplineMeshSceneProxy.h
+++ b/Engine/Source/Runtime/Engine/Private/SplineMeshSceneProxy.h
@@ -75,7 +75,7 @@ class FSplineMeshVertexFactoryShaderParameters : public FVertexFactoryShaderPara
 		const class FSceneInterface* Scene,
 		const FSceneView* View,
 		const class FMeshMaterialShader* Shader,
-		bool bShaderRequiresPositionOnlyStream,
+		const EVertexInputStreamType InputStreamType,
 		ERHIFeatureLevel::Type FeatureLevel,
 		const FVertexFactory* VertexFactory,
 		const FMeshBatchElement& BatchElement,
diff --git a/Engine/Source/Runtime/Engine/Private/VectorFieldVisualization.cpp b/Engine/Source/Runtime/Engine/Private/VectorFieldVisualization.cpp
--- a/Engine/Source/Runtime/Engine/Private/VectorFieldVisualization.cpp
+++ b/Engine/Source/Runtime/Engine/Private/VectorFieldVisualization.cpp
@@ -43,7 +43,7 @@ public:
 		const FSceneInterface* Scene,
 		const FSceneView* View,
 		const FMeshMaterialShader* Shader,
-		bool bShaderRequiresPositionOnlyStream,
+		const EVertexInputStreamType InputStreamType,
 		ERHIFeatureLevel::Type FeatureLevel,
 		const FVertexFactory* VertexFactory,
 		const FMeshBatchElement& BatchElement,
@@ -175,7 +175,7 @@ void FVectorFieldVisualizationVertexFactoryShaderParameters::GetElementShaderBin
 	const FSceneInterface* Scene,
 	const FSceneView* View,
 	const FMeshMaterialShader* Shader,
-	bool bShaderRequiresPositionOnlyStream,
+	const EVertexInputStreamType InputStreamType,
 	ERHIFeatureLevel::Type FeatureLevel,
 	const FVertexFactory* InVertexFactory,
 	const FMeshBatchElement& BatchElement,
diff --git a/Engine/Source/Runtime/Engine/Public/LocalVertexFactory.h b/Engine/Source/Runtime/Engine/Public/LocalVertexFactory.h
--- a/Engine/Source/Runtime/Engine/Public/LocalVertexFactory.h
+++ b/Engine/Source/Runtime/Engine/Public/LocalVertexFactory.h
@@ -171,7 +171,7 @@ public:
 		const class FSceneInterface* Scene,
 		const FSceneView* View,
 		const FMeshMaterialShader* Shader,
-		bool bShaderRequiresPositionOnlyStream,
+		const EVertexInputStreamType InputStreamType,
 		ERHIFeatureLevel::Type FeatureLevel,
 		const FVertexFactory* VertexFactory,
 		const FMeshBatchElement& BatchElement,
@@ -201,7 +201,7 @@ public:
 		const FSceneInterface* Scene,
 		const FSceneView* View,
 		const FMeshMaterialShader* Shader,
-		bool bShaderRequiresPositionOnlyStream,
+		const EVertexInputStreamType InputStreamType,
 		ERHIFeatureLevel::Type FeatureLevel,
 		const FVertexFactory* VertexFactory,
 		const FMeshBatchElement& BatchElement,
diff --git a/Engine/Source/Runtime/Landscape/Private/LandscapeRender.cpp b/Engine/Source/Runtime/Landscape/Private/LandscapeRender.cpp
--- a/Engine/Source/Runtime/Landscape/Private/LandscapeRender.cpp
+++ b/Engine/Source/Runtime/Landscape/Private/LandscapeRender.cpp
@@ -3406,7 +3406,7 @@ public:
 		const class FSceneInterface* Scene,
 		const FSceneView* InView,
 		const class FMeshMaterialShader* Shader,
-		bool bShaderRequiresPositionOnlyStream,
+		const EVertexInputStreamType InputStreamType,
 		ERHIFeatureLevel::Type FeatureLevel,
 		const FVertexFactory* VertexFactory,
 		const FMeshBatchElement& BatchElement,
@@ -3593,7 +3593,7 @@ void FLandscapeVertexFactoryPixelShaderParameters::GetElementShaderBindings(
 	const class FSceneInterface* Scene,
 	const FSceneView* InView,
 	const class FMeshMaterialShader* Shader,
-	bool bShaderRequiresPositionOnlyStream,
+	const EVertexInputStreamType InputStreamType,
 	ERHIFeatureLevel::Type FeatureLevel,
 	const FVertexFactory* VertexFactory,
 	const FMeshBatchElement& BatchElement,
diff --git a/Engine/Source/Runtime/Landscape/Private/LandscapeRenderMobile.cpp b/Engine/Source/Runtime/Landscape/Private/LandscapeRenderMobile.cpp
--- a/Engine/Source/Runtime/Landscape/Private/LandscapeRenderMobile.cpp
+++ b/Engine/Source/Runtime/Landscape/Private/LandscapeRenderMobile.cpp
@@ -67,7 +67,7 @@ public:
 		const class FSceneInterface* Scene,
 		const FSceneView* InView,
 		const class FMeshMaterialShader* Shader,
-		bool bShaderRequiresPositionOnlyStream,
+		const EVertexInputStreamType InputStreamType,
 		ERHIFeatureLevel::Type FeatureLevel,
 		const FVertexFactory* VertexFactory,
 		const FMeshBatchElement& BatchElement,
@@ -199,7 +199,7 @@ public:
 		const class FSceneInterface* Scene,
 		const FSceneView* InView,
 		const class FMeshMaterialShader* Shader,
-		bool bShaderRequiresPositionOnlyStream,
+		const EVertexInputStreamType InputStreamType,
 		ERHIFeatureLevel::Type FeatureLevel,
 		const FVertexFactory* VertexFactory,
 		const FMeshBatchElement& BatchElement,
@@ -209,7 +209,7 @@ public:
 	{
 		SCOPE_CYCLE_COUNTER(STAT_LandscapeVFDrawTimePS);
 		
-		FLandscapeVertexFactoryPixelShaderParameters::GetElementShaderBindings(Scene, InView, Shader, bShaderRequiresPositionOnlyStream, FeatureLevel, VertexFactory, BatchElement, ShaderBindings, VertexStreams);
+		FLandscapeVertexFactoryPixelShaderParameters::GetElementShaderBindings(Scene, InView, Shader, InputStreamType, FeatureLevel, VertexFactory, BatchElement, ShaderBindings, VertexStreams);
 
 		if (BlendableLayerMaskParameter.IsBound())
 		{
diff --git a/Engine/Source/Runtime/Landscape/Public/LandscapeRender.h b/Engine/Source/Runtime/Landscape/Public/LandscapeRender.h
--- a/Engine/Source/Runtime/Landscape/Public/LandscapeRender.h
+++ b/Engine/Source/Runtime/Landscape/Public/LandscapeRender.h
@@ -133,7 +133,7 @@ public:
 		const class FSceneInterface* Scene,
 		const FSceneView* InView,
 		const class FMeshMaterialShader* Shader,
-		bool bShaderRequiresPositionOnlyStream,
+		const EVertexInputStreamType InputStreamType,
 		ERHIFeatureLevel::Type FeatureLevel,
 		const FVertexFactory* VertexFactory,
 		const FMeshBatchElement& BatchElement,
diff --git a/Engine/Source/Runtime/RenderCore/Private/VertexFactory.cpp b/Engine/Source/Runtime/RenderCore/Private/VertexFactory.cpp
--- a/Engine/Source/Runtime/RenderCore/Private/VertexFactory.cpp
+++ b/Engine/Source/Runtime/RenderCore/Private/VertexFactory.cpp
@@ -178,40 +178,61 @@ FVertexFactoryType* FindVertexFactoryType(FName TypeName)
 	return NULL;
 }
 
-void FVertexFactory::GetStreams(ERHIFeatureLevel::Type InFeatureLevel, FVertexInputStreamArray& OutVertexStreams) const
+void FVertexFactory::GetStreams(ERHIFeatureLevel::Type InFeatureLevel, EVertexInputStreamType VertexStreamType, FVertexInputStreamArray& OutVertexStreams) const
 {
-	bool bSupportsVertexFetch = SupportsManualVertexFetch(InFeatureLevel);
 	check(IsInitialized());
-
-	for (int32 StreamIndex = 0;StreamIndex < Streams.Num();StreamIndex++)
+	if (VertexStreamType == EVertexInputStreamType::Default)
 	{
-		const FVertexStream& Stream = Streams[StreamIndex];
 
-		if (!(EnumHasAnyFlags(EVertexStreamUsage::ManualFetch, Stream.VertexStreamUsage) && bSupportsVertexFetch))
+		bool bSupportsVertexFetch = SupportsManualVertexFetch(InFeatureLevel);
+
+		for (int32 StreamIndex = 0;StreamIndex < Streams.Num();StreamIndex++)
 		{
-			if (!Stream.VertexBuffer)
-			{
-				OutVertexStreams.Add(FVertexInputStream(StreamIndex, 0, nullptr));
-			}
-			else
+			const FVertexStream& Stream = Streams[StreamIndex];
+
+			if (!(EnumHasAnyFlags(EVertexStreamUsage::ManualFetch, Stream.VertexStreamUsage) && bSupportsVertexFetch))
 			{
-				if (EnumHasAnyFlags(EVertexStreamUsage::Overridden, Stream.VertexStreamUsage) && !Stream.VertexBuffer->IsInitialized())
+				if (!Stream.VertexBuffer)
 				{
 					OutVertexStreams.Add(FVertexInputStream(StreamIndex, 0, nullptr));
 				}
 				else
 				{
-					checkf(Stream.VertexBuffer->IsInitialized(), TEXT("Vertex buffer was not initialized! Stream %u, Stride %u, Name %s"), StreamIndex, Stream.Stride, *Stream.VertexBuffer->GetFriendlyName());
-					OutVertexStreams.Add(FVertexInputStream(StreamIndex, Stream.Offset, Stream.VertexBuffer->VertexBufferRHI));
+					if (EnumHasAnyFlags(EVertexStreamUsage::Overridden, Stream.VertexStreamUsage) && !Stream.VertexBuffer->IsInitialized())
+					{
+						OutVertexStreams.Add(FVertexInputStream(StreamIndex, 0, nullptr));
+					}
+					else
+					{
+						checkf(Stream.VertexBuffer->IsInitialized(), TEXT("Vertex buffer was not initialized! Stream %u, Stride %u, Name %s"), StreamIndex, Stream.Stride, *Stream.VertexBuffer->GetFriendlyName());
+						OutVertexStreams.Add(FVertexInputStream(StreamIndex, Stream.Offset, Stream.VertexBuffer->VertexBufferRHI));
+					}
 				}
 			}
 		}
 	}
+	else if (VertexStreamType == EVertexInputStreamType::PositionOnly)
+	{
+		// Set the predefined vertex streams.
+		for (int32 StreamIndex = 0; StreamIndex < PositionStream.Num(); StreamIndex++)
+		{
+			const FVertexStream& Stream = PositionStream[StreamIndex];
+			check(Stream.VertexBuffer->IsInitialized());
+			OutVertexStreams.Add(FVertexInputStream(StreamIndex, Stream.Offset, Stream.VertexBuffer->VertexBufferRHI));
+		}
+	}
+	else
+	{
+		// NOT_IMPLEMENTED
+	}
 }
 
-void FVertexFactory::OffsetInstanceStreams(uint32 InstanceOffset, bool bOperateOnPositionOnly, FVertexInputStreamArray& VertexStreams) const
+void FVertexFactory::OffsetInstanceStreams(uint32 InstanceOffset, EVertexInputStreamType VertexStreamType, FVertexInputStreamArray& VertexStreams) const
 {
-	const TArrayView<const FVertexStream>& StreamArray = bOperateOnPositionOnly ? MakeArrayView(PositionStream) : MakeArrayView(Streams);
+	const TArrayView<const FVertexStream>& StreamArray = 
+		  VertexStreamType == EVertexInputStreamType::PositionOnly ?			MakeArrayView(PositionStream) : 
+		/*VertexStreamType == EVertexInputStreamType::Default*/					MakeArrayView(Streams);
+
 	for (int32 StreamIndex = 0; StreamIndex < StreamArray.Num(); StreamIndex++)
 	{
 		const FVertexStream& Stream = StreamArray[StreamIndex];
@@ -229,18 +250,6 @@ void FVertexFactory::OffsetInstanceStreams(uint32 InstanceOffset, bool bOperateO
 	}
 }
 
-void FVertexFactory::GetPositionOnlyStream(FVertexInputStreamArray& OutVertexStreams) const
-{
-	check(IsInitialized());
-	// Set the predefined vertex streams.
-	for(int32 StreamIndex = 0;StreamIndex < PositionStream.Num();StreamIndex++)
-	{
-		const FVertexStream& Stream = PositionStream[StreamIndex];
-		check(Stream.VertexBuffer->IsInitialized());
-		OutVertexStreams.Add(FVertexInputStream(StreamIndex, Stream.Offset, Stream.VertexBuffer->VertexBufferRHI));
-	}
-}
-
 void FVertexFactory::ReleaseRHI()
 {
 	Declaration.SafeRelease();
@@ -249,7 +258,7 @@ void FVertexFactory::ReleaseRHI()
 	PositionStream.Empty();
 }
 
-FVertexElement FVertexFactory::AccessStreamComponent(const FVertexStreamComponent& Component,uint8 AttributeIndex)
+FVertexElement FVertexFactory::AccessStreamComponent(const FVertexStreamComponent& Component, uint8 AttributeIndex)
 {
 	FVertexStream VertexStream;
 	VertexStream.VertexBuffer = Component.VertexBuffer;
@@ -260,7 +269,7 @@ FVertexElement FVertexFactory::AccessStreamComponent(const FVertexStreamComponen
 	return FVertexElement(Streams.AddUnique(VertexStream),Component.Offset,Component.Type,AttributeIndex,VertexStream.Stride, EnumHasAnyFlags(EVertexStreamUsage::Instancing, VertexStream.VertexStreamUsage));
 }
 
-FVertexElement FVertexFactory::AccessPositionStreamComponent(const FVertexStreamComponent& Component,uint8 AttributeIndex)
+FVertexElement FVertexFactory::AccessStreamComponent(const FVertexStreamComponent& Component,uint8 AttributeIndex, EVertexInputStreamType InputStreamType)
 {
 	FVertexStream VertexStream;
 	VertexStream.VertexBuffer = Component.VertexBuffer;
@@ -268,18 +277,24 @@ FVertexElement FVertexFactory::AccessPositionStreamComponent(const FVertexStream
 	VertexStream.Offset = Component.StreamOffset;
 	VertexStream.VertexStreamUsage = Component.VertexStreamUsage;
 
-	return FVertexElement(PositionStream.AddUnique(VertexStream),Component.Offset,Component.Type,AttributeIndex,VertexStream.Stride, EnumHasAnyFlags(EVertexStreamUsage::Instancing, VertexStream.VertexStreamUsage));
-}
-
-void FVertexFactory::InitDeclaration(FVertexDeclarationElementList& Elements)
-{
-	// Create the vertex declaration for rendering the factory normally.
-	Declaration = PipelineStateCache::GetOrCreateVertexDeclaration(Elements);
+	if (InputStreamType == EVertexInputStreamType::PositionOnly)
+		return FVertexElement(PositionStream.AddUnique(VertexStream), Component.Offset, Component.Type, AttributeIndex, VertexStream.Stride, EnumHasAnyFlags(EVertexStreamUsage::Instancing, VertexStream.VertexStreamUsage));
+	else /* (InputStreamType == EVertexInputStreamType::Default) */
+		return FVertexElement(Streams.AddUnique(VertexStream), Component.Offset, Component.Type, AttributeIndex, VertexStream.Stride, EnumHasAnyFlags(EVertexStreamUsage::Instancing, VertexStream.VertexStreamUsage));
 }
 
-void FVertexFactory::InitPositionDeclaration(const FVertexDeclarationElementList& Elements)
+void FVertexFactory::InitDeclaration(const FVertexDeclarationElementList& Elements, EVertexInputStreamType StreamType)
 {
-	PositionDeclaration = PipelineStateCache::GetOrCreateVertexDeclaration(Elements);
+	
+	if (StreamType == EVertexInputStreamType::PositionOnly)
+	{
+		PositionDeclaration = PipelineStateCache::GetOrCreateVertexDeclaration(Elements);
+	}
+	else // (StreamType == EVertexInputStreamType::Default)
+	{
+		// Create the vertex declaration for rendering the factory normally.
+		Declaration = PipelineStateCache::GetOrCreateVertexDeclaration(Elements);
+	}
 }
 
 FVertexFactoryParameterRef::FVertexFactoryParameterRef(FVertexFactoryType* InVertexFactoryType,const FShaderParameterMap& ParameterMap, EShaderFrequency InShaderFrequency, EShaderPlatform InShaderPlatform)
diff --git a/Engine/Source/Runtime/RenderCore/Public/VertexFactory.h b/Engine/Source/Runtime/RenderCore/Public/VertexFactory.h
--- a/Engine/Source/Runtime/RenderCore/Public/VertexFactory.h
+++ b/Engine/Source/Runtime/RenderCore/Public/VertexFactory.h
@@ -77,6 +77,13 @@ enum class EVertexStreamUsage : uint8
 	ManualFetch		= 1 << 2
 };
 
+
+enum class EVertexInputStreamType : uint8
+{
+	Default = 0,
+	PositionOnly
+};
+
 ENUM_CLASS_FLAGS(EVertexStreamUsage);
 /**
  * A typed data source for a vertex factory which streams data from a vertex buffer.
@@ -158,7 +165,7 @@ public:
 		const class FSceneInterface* Scene,
 		const class FSceneView* View,
 		const class FMeshMaterialShader* Shader,
-		bool bShaderRequiresPositionOnlyStream,
+		const EVertexInputStreamType InputStreamType,
 		ERHIFeatureLevel::Type FeatureLevel,
 		const class FVertexFactory* VertexFactory,
 		const struct FMeshBatchElement& BatchElement,
@@ -479,11 +486,9 @@ public:
 
 	virtual FVertexFactoryType* GetType() const { return NULL; }
 
-	void GetStreams(ERHIFeatureLevel::Type InFeatureLevel, FVertexInputStreamArray& OutVertexStreams) const;
+	void GetStreams(ERHIFeatureLevel::Type InFeatureLevel, EVertexInputStreamType VertexStreamType, FVertexInputStreamArray& OutVertexStreams) const;
 
-	void OffsetInstanceStreams(uint32 InstanceOffset, bool bOperateOnPositionOnly, FVertexInputStreamArray& VertexStreams) const;
-
-	void GetPositionOnlyStream(FVertexInputStreamArray& OutVertexStreams) const;
+	void OffsetInstanceStreams(uint32 InstanceOffset, EVertexInputStreamType VertexStreamType, FVertexInputStreamArray& VertexStreams) const;
 
 	/**
 	* Can be overridden by FVertexFactory subclasses to modify their compile environment just before compilation occurs.
@@ -507,8 +512,15 @@ public:
 	FVertexDeclarationRHIRef& GetDeclaration() { return Declaration; }
 	void SetDeclaration(FVertexDeclarationRHIRef& NewDeclaration) { Declaration = NewDeclaration; }
 
-	const FVertexDeclarationRHIRef& GetDeclaration() const { return Declaration; }
-	const FVertexDeclarationRHIRef& GetPositionDeclaration() const { return PositionDeclaration; }
+	const FVertexDeclarationRHIRef& GetDeclaration(EVertexInputStreamType InputStreamType/* = FVertexInputStreamType::Default*/) const 
+	{
+		switch (InputStreamType)
+		{
+		case EVertexInputStreamType::Default:				return Declaration;
+		case EVertexInputStreamType::PositionOnly:			return PositionDeclaration;
+		}
+		return Declaration;
+	}
 
 	virtual bool IsGPUSkinned() const { return false; }
 
@@ -533,9 +545,9 @@ public:
 		return bSupportsManualVertexFetch && (InFeatureLevel > ERHIFeatureLevel::ES3_1) && RHISupportsManualVertexFetch(GMaxRHIShaderPlatform);
 	}
 
-	inline int32 GetPrimitiveIdStreamIndex(bool bPositionOnly) const
+	inline int32 GetPrimitiveIdStreamIndex(EVertexInputStreamType InputStreamType) const
 	{
-		return bPositionOnly ? PositionOnlyPrimitiveIdStreamIndex : PrimitiveIdStreamIndex;
+		return PrimitiveIdStreamIndex[static_cast<uint8>(InputStreamType)];
 	}
 
 protected:
@@ -549,25 +561,19 @@ protected:
 	FVertexElement AccessStreamComponent(const FVertexStreamComponent& Component,uint8 AttributeIndex);
 
 	/**
-	 * Creates a vertex element for a position vertex stream component.  Adds a unique position stream index for the vertex buffer used by the component.
-	 * @param Component - The position vertex stream component.
+	 * Creates a vertex element for a vertex stream component.  Adds a unique position stream index for the vertex buffer used by the component.
+	 * @param Component - The vertex stream component.
 	 * @param Usage - The vertex element usage semantics.
 	 * @param AttributeIndex - The attribute index to which the stream component is bound.
 	 * @return The vertex element which corresponds to Component.
 	 */
-	FVertexElement AccessPositionStreamComponent(const FVertexStreamComponent& Component,uint8 AttributeIndex);
+	FVertexElement AccessStreamComponent(const FVertexStreamComponent& Component, uint8 AttributeIndex, EVertexInputStreamType InputStreamType);
 
 	/**
 	 * Initializes the vertex declaration.
 	 * @param Elements - The elements of the vertex declaration.
 	 */
-	void InitDeclaration(FVertexDeclarationElementList& Elements);
-
-	/**
-	 * Initializes the position-only vertex declaration.
-	 * @param Elements - The elements of the vertex declaration.
-	 */
-	void InitPositionDeclaration(const FVertexDeclarationElementList& Elements);
+	void InitDeclaration(const FVertexDeclarationElementList& Elements, EVertexInputStreamType StreamType = EVertexInputStreamType::Default);
 
 	/**
 	 * Information needed to set a vertex stream.
@@ -598,8 +604,7 @@ protected:
 	
 	bool bSupportsManualVertexFetch = false;
 
-	int8 PositionOnlyPrimitiveIdStreamIndex = -1;
-	int8 PrimitiveIdStreamIndex = -1;
+	int8 PrimitiveIdStreamIndex[3] = { -1, -1 }; // Need to match entry count of EVertexInputStreamType
 
 private:
 
@@ -636,7 +641,7 @@ public:
 		const class FSceneInterface* Scene,
 		const FSceneView* View,
 		const class FMeshMaterialShader* Shader,
-		bool bShaderRequiresPositionOnlyStream,
+		const EVertexInputStreamType InputStreamType,
 		ERHIFeatureLevel::Type FeatureLevel,
 		const FVertexFactory* VertexFactory,
 		const struct FMeshBatchElement& BatchElement,
@@ -648,7 +653,7 @@ public:
 		{
 			checkSlow(VertexFactory->GetType() == VertexFactoryType);
 			checkSlow(View || VertexFactoryType->SupportsCachingMeshDrawCommands());
-			Parameters->GetElementShaderBindings(Scene, View, Shader, bShaderRequiresPositionOnlyStream, FeatureLevel, VertexFactory, BatchElement, ShaderBindings, VertexStreams);
+			Parameters->GetElementShaderBindings(Scene, View, Shader, InputStreamType, FeatureLevel, VertexFactory, BatchElement, ShaderBindings, VertexStreams);
 		}
 	}
 
diff --git a/Engine/Source/Runtime/Renderer/Private/BasePassRendering.h b/Engine/Source/Runtime/Renderer/Private/BasePassRendering.h
--- a/Engine/Source/Runtime/Renderer/Private/BasePassRendering.h
+++ b/Engine/Source/Runtime/Renderer/Private/BasePassRendering.h
@@ -219,7 +219,7 @@ public:
 		const FScene* Scene,
 		const FSceneView* ViewIfDynamicMeshCommand,
 		const FVertexFactory* VertexFactory,
-		bool bShaderRequiresPositionOnlyStream,
+		const EVertexInputStreamType InputStreamType,
 		ERHIFeatureLevel::Type FeatureLevel,
 		const FPrimitiveSceneProxy* PrimitiveSceneProxy,
 		const FMeshBatch& MeshBatch,
diff --git a/Engine/Source/Runtime/Renderer/Private/BasePassRendering.inl b/Engine/Source/Runtime/Renderer/Private/BasePassRendering.inl
--- a/Engine/Source/Runtime/Renderer/Private/BasePassRendering.inl
+++ b/Engine/Source/Runtime/Renderer/Private/BasePassRendering.inl
@@ -54,7 +54,7 @@ void TBasePassVertexShaderPolicyParamType<LightMapPolicyType>::GetElementShaderB
 	const FScene* Scene, 
 	const FSceneView* ViewIfDynamicMeshCommand, 
 	const FVertexFactory* VertexFactory,
-	bool bShaderRequiresPositionOnlyStream,
+	const EVertexInputStreamType InputStreamType,
 	ERHIFeatureLevel::Type FeatureLevel,
 	const FPrimitiveSceneProxy* PrimitiveSceneProxy,
 	const FMeshBatch& MeshBatch,
@@ -63,7 +63,7 @@ void TBasePassVertexShaderPolicyParamType<LightMapPolicyType>::GetElementShaderB
 	FMeshDrawSingleShaderBindings& ShaderBindings,
 	FVertexInputStreamArray& VertexStreams) const
 {
-	FMeshMaterialShader::GetElementShaderBindings(Scene, ViewIfDynamicMeshCommand, VertexFactory, bShaderRequiresPositionOnlyStream, FeatureLevel, PrimitiveSceneProxy, MeshBatch, BatchElement, ShaderElementData, ShaderBindings, VertexStreams);
+	FMeshMaterialShader::GetElementShaderBindings(Scene, ViewIfDynamicMeshCommand, VertexFactory, InputStreamType, FeatureLevel, PrimitiveSceneProxy, MeshBatch, BatchElement, ShaderElementData, ShaderBindings, VertexStreams);
 }
 
 template<typename LightMapPolicyType>
diff --git a/Engine/Source/Runtime/Renderer/Private/DebugViewModeRendering.h b/Engine/Source/Runtime/Renderer/Private/DebugViewModeRendering.h
--- a/Engine/Source/Runtime/Renderer/Private/DebugViewModeRendering.h
+++ b/Engine/Source/Runtime/Renderer/Private/DebugViewModeRendering.h
@@ -185,7 +185,7 @@ public:
 		const FScene* Scene, 
 		const FSceneView* ViewIfDynamicMeshCommand, 
 		const FVertexFactory* VertexFactory,
-		bool bShaderRequiresPositionOnlyStream,
+		const EVertexInputStreamType InputStreamType,
 		ERHIFeatureLevel::Type FeatureLevel,
 		const FPrimitiveSceneProxy* PrimitiveSceneProxy,
 		const FMeshBatch& MeshBatch,
@@ -194,7 +194,7 @@ public:
 		FMeshDrawSingleShaderBindings& ShaderBindings,
 		FVertexInputStreamArray& VertexStreams) const
 	{
-		FMeshMaterialShader::GetElementShaderBindings(Scene, ViewIfDynamicMeshCommand, VertexFactory, bShaderRequiresPositionOnlyStream, FeatureLevel, PrimitiveSceneProxy, MeshBatch, BatchElement, ShaderElementData, ShaderBindings, VertexStreams);
+		FMeshMaterialShader::GetElementShaderBindings(Scene, ViewIfDynamicMeshCommand, VertexFactory, InputStreamType, FeatureLevel, PrimitiveSceneProxy, MeshBatch, BatchElement, ShaderElementData, ShaderBindings, VertexStreams);
 	
 		int8 VisualizeElementIndex = 0;
 #if !(UE_BUILD_SHIPPING || UE_BUILD_TEST)
diff --git a/Engine/Source/Runtime/Renderer/Private/DynamicPrimitiveDrawing.inl b/Engine/Source/Runtime/Renderer/Private/DynamicPrimitiveDrawing.inl
--- a/Engine/Source/Runtime/Renderer/Private/DynamicPrimitiveDrawing.inl
+++ b/Engine/Source/Runtime/Renderer/Private/DynamicPrimitiveDrawing.inl
@@ -199,7 +199,7 @@ inline int32 FViewElementPDI::DrawMesh(const FMeshBatch& Mesh)
 			ENQUEUE_RENDER_COMMAND(FCopyDynamicPrimitiveShaderData)(
 				[NewMesh, DynamicPrimitiveShaderDataForRT, FeatureLevel](FRHICommandListImmediate& RHICmdList)
 				{
-					const bool bPrimitiveShaderDataComesFromSceneBuffer = NewMesh->VertexFactory->GetPrimitiveIdStreamIndex(false) >= 0;
+					const bool bPrimitiveShaderDataComesFromSceneBuffer = NewMesh->VertexFactory->GetPrimitiveIdStreamIndex(EVertexInputStreamType::Default) >= 0;
 
 					for (int32 ElementIndex = 0; ElementIndex < NewMesh->Elements.Num(); ElementIndex++)
 					{
diff --git a/Engine/Source/Runtime/Renderer/Private/RayTracing/RayTracingDynamicGeometry.cpp b/Engine/Source/Runtime/Renderer/Private/RayTracing/RayTracingDynamicGeometry.cpp
--- a/Engine/Source/Runtime/Renderer/Private/RayTracing/RayTracingDynamicGeometry.cpp
+++ b/Engine/Source/Runtime/Renderer/Private/RayTracing/RayTracingDynamicGeometry.cpp
@@ -60,7 +60,7 @@ public:
 		const FScene* Scene,
 		const FSceneView* ViewIfDynamicMeshCommand,
 		const FVertexFactory* VertexFactory,
-		bool bShaderRequiresPositionOnlyStream,
+		const EVertexInputStreamType InputStreamType,
 		ERHIFeatureLevel::Type FeatureLevel,
 		const FPrimitiveSceneProxy* PrimitiveSceneProxy,
 		const FMeshBatch& MeshBatch, 
@@ -69,7 +69,7 @@ public:
 		FMeshDrawSingleShaderBindings& ShaderBindings,
 		FVertexInputStreamArray& VertexStreams) const
 	{
-		FMeshMaterialShader::GetElementShaderBindings(Scene, ViewIfDynamicMeshCommand, VertexFactory, bShaderRequiresPositionOnlyStream, FeatureLevel, PrimitiveSceneProxy, MeshBatch, BatchElement, ShaderElementData, ShaderBindings, VertexStreams);
+		FMeshMaterialShader::GetElementShaderBindings(Scene, ViewIfDynamicMeshCommand, VertexFactory, InputStreamType, FeatureLevel, PrimitiveSceneProxy, MeshBatch, BatchElement, ShaderElementData, ShaderBindings, VertexStreams);
 	}
 
 	FRWShaderParameter RWVertexPositions;
@@ -119,7 +119,7 @@ void FRayTracingDynamicGeometryCollection::AddDynamicMeshBatchForGeometryUpdate(
 	Shader->GetShaderBindings(Scene, Scene->GetFeatureLevel(), PrimitiveSceneProxy, MaterialRenderProxy, Material, DrawRenderState, ShaderElementData, SingleShaderBindings);
 
 	FVertexInputStreamArray DummyArray;
-	Shader->GetElementShaderBindings(Scene, View, MeshBatch.VertexFactory, false, Scene->GetFeatureLevel(), PrimitiveSceneProxy, MeshBatch, MeshBatch.Elements[0], ShaderElementData, SingleShaderBindings, DummyArray);
+	Shader->GetElementShaderBindings(Scene, View, MeshBatch.VertexFactory, EVertexInputStreamType::Default, Scene->GetFeatureLevel(), PrimitiveSceneProxy, MeshBatch, MeshBatch.Elements[0], ShaderElementData, SingleShaderBindings, DummyArray);
 
 	DispatchCmd.TargetBuffer = &Buffer;
 	DispatchCmd.TargetGeometry = &Geometry;
diff --git a/Engine/Source/Runtime/Renderer/Private/RayTracing/RayTracingMaterialHitShaders.cpp b/Engine/Source/Runtime/Renderer/Private/RayTracing/RayTracingMaterialHitShaders.cpp
--- a/Engine/Source/Runtime/Renderer/Private/RayTracing/RayTracingMaterialHitShaders.cpp
+++ b/Engine/Source/Runtime/Renderer/Private/RayTracing/RayTracingMaterialHitShaders.cpp
@@ -60,7 +60,7 @@ public:
 		const FScene* Scene,
 		const FSceneView* ViewIfDynamicMeshCommand,
 		const FVertexFactory* VertexFactory,
-		bool bShaderRequiresPositionOnlyStream,
+		const EVertexInputStreamType InputStreamType,
 		ERHIFeatureLevel::Type FeatureLevel,
 		const FPrimitiveSceneProxy* PrimitiveSceneProxy,
 		const FMeshBatch& MeshBatch, 
@@ -69,7 +69,7 @@ public:
 		FMeshDrawSingleShaderBindings& ShaderBindings,
 		FVertexInputStreamArray& VertexStreams) const
 	{
-		FMeshMaterialShader::GetElementShaderBindings(Scene, ViewIfDynamicMeshCommand, VertexFactory, bShaderRequiresPositionOnlyStream, FeatureLevel, PrimitiveSceneProxy, MeshBatch, BatchElement, ShaderElementData, ShaderBindings, VertexStreams);
+		FMeshMaterialShader::GetElementShaderBindings(Scene, ViewIfDynamicMeshCommand, VertexFactory, InputStreamType, FeatureLevel, PrimitiveSceneProxy, MeshBatch, BatchElement, ShaderElementData, ShaderBindings, VertexStreams);
 	}
 };
 
diff --git a/Engine/Source/Runtime/Renderer/Private/Renderer.cpp b/Engine/Source/Runtime/Renderer/Private/Renderer.cpp
--- a/Engine/Source/Runtime/Renderer/Private/Renderer.cpp
+++ b/Engine/Source/Runtime/Renderer/Private/Renderer.cpp
@@ -109,7 +109,7 @@ void FRendererModule::DrawTileMesh(FRHICommandListImmediate& RHICmdList, FMeshPa
 
 		FSinglePrimitiveStructuredBuffer SinglePrimitiveStructuredBuffer;
 
-		if (Mesh.VertexFactory->GetPrimitiveIdStreamIndex(true) >= 0)
+		if (Mesh.VertexFactory->GetPrimitiveIdStreamIndex(EVertexInputStreamType::PositionOnly) >= 0)
 		{
 			FMeshBatchElement& MeshElement = Mesh.Elements[0];
 
diff --git a/Engine/Source/Runtime/Renderer/Private/ShaderBaseClasses.cpp b/Engine/Source/Runtime/Renderer/Private/ShaderBaseClasses.cpp
--- a/Engine/Source/Runtime/Renderer/Private/ShaderBaseClasses.cpp
+++ b/Engine/Source/Runtime/Renderer/Private/ShaderBaseClasses.cpp
@@ -534,7 +534,7 @@ void FMeshMaterialShader::GetElementShaderBindings(
 	const FScene* Scene, 
 	const FSceneView* ViewIfDynamicMeshCommand, 
 	const FVertexFactory* VertexFactory,
-	bool bShaderRequiresPositionOnlyStream,
+	const EVertexInputStreamType InputStreamType,
 	ERHIFeatureLevel::Type FeatureLevel,
 	const FPrimitiveSceneProxy* PrimitiveSceneProxy,
 	const FMeshBatch& MeshBatch,
@@ -545,9 +545,9 @@ void FMeshMaterialShader::GetElementShaderBindings(
 {
 	checkSlow(ShaderBindings.Frequency == GetType()->GetFrequency());
 
-	VertexFactoryParameters.GetElementShaderBindings(Scene, ViewIfDynamicMeshCommand, this, bShaderRequiresPositionOnlyStream, FeatureLevel, VertexFactory, BatchElement, ShaderBindings, VertexStreams);
+	VertexFactoryParameters.GetElementShaderBindings(Scene, ViewIfDynamicMeshCommand, this, InputStreamType, FeatureLevel, VertexFactory, BatchElement, ShaderBindings, VertexStreams);
 		
-	if (UseGPUScene(GMaxRHIShaderPlatform, FeatureLevel) && VertexFactory->GetPrimitiveIdStreamIndex(bShaderRequiresPositionOnlyStream) >= 0)
+	if (UseGPUScene(GMaxRHIShaderPlatform, FeatureLevel) && VertexFactory->GetPrimitiveIdStreamIndex(InputStreamType) >= 0)
 	{
 		ensureMsgf(!GetUniformBufferParameter<FPrimitiveUniformShaderParameters>().IsBound(), TEXT("Shader %s attempted to bind the Primitive uniform buffer even though Vertex Factory computes a PrimitiveId per-instance.  This will break auto-instancing.  Shaders should use GetPrimitiveData(PrimitiveId).Member instead of Primitive.Member."), GetType()->GetName());
 		ensureMsgf(!BatchElement.PrimitiveUniformBuffer, TEXT("FMeshBatchElement was assigned a PrimitiveUniformBuffer even though Vertex Factory %s fetches primitive shader data through a Scene buffer.  The assigned PrimitiveUniformBuffer cannot be respected.  Use PrimitiveUniformBufferResource instead for dynamic primitive data."), GetType()->GetName());
diff --git a/Engine/Source/Runtime/Renderer/Public/MeshMaterialShader.h b/Engine/Source/Runtime/Renderer/Public/MeshMaterialShader.h
--- a/Engine/Source/Runtime/Renderer/Public/MeshMaterialShader.h
+++ b/Engine/Source/Runtime/Renderer/Public/MeshMaterialShader.h
@@ -70,7 +70,7 @@ public:
 		const FScene* Scene, 
 		const FSceneView* ViewIfDynamicMeshCommand, 
 		const FVertexFactory* VertexFactory,
-		bool bShaderRequiresPositionOnlyStream,
+		const EVertexInputStreamType InputStreamType,
 		ERHIFeatureLevel::Type FeatureLevel,
 		const FPrimitiveSceneProxy* PrimitiveSceneProxy,
 		const FMeshBatch& MeshBatch,
diff --git a/Engine/Source/Runtime/Renderer/Public/MeshPassProcessor.inl b/Engine/Source/Runtime/Renderer/Public/MeshPassProcessor.inl
--- a/Engine/Source/Runtime/Renderer/Public/MeshPassProcessor.inl
+++ b/Engine/Source/Runtime/Renderer/Public/MeshPassProcessor.inl
@@ -33,10 +33,12 @@ void FMeshPassProcessor::BuildMeshDrawCommands(
 	FGraphicsMinimalPipelineStateInitializer PipelineState;
 	PipelineState.PrimitiveType = (EPrimitiveType)MeshBatch.Type;
 
-	const bool bPositionOnly = (MeshPassFeatures & EMeshPassFeatures::PositionOnly) != EMeshPassFeatures::Default;
+	EVertexInputStreamType InputStreamType = EVertexInputStreamType::Default;
+	if ((MeshPassFeatures & EMeshPassFeatures::PositionOnly) != EMeshPassFeatures::Default)				InputStreamType = EVertexInputStreamType::PositionOnly;
 
 	check(VertexFactory && VertexFactory->IsInitialized());
-	FVertexDeclarationRHIParamRef VertexDeclaration = bPositionOnly ? VertexFactory->GetPositionDeclaration() : VertexFactory->GetDeclaration();
+	FVertexDeclarationRHIParamRef VertexDeclaration = VertexFactory->GetDeclaration(InputStreamType);
+
 	check(!VertexFactory->NeedsDeclaration() || VertexDeclaration);
 
 	SharedMeshDrawCommand.SetShaders(VertexDeclaration, PassShaders.GetUntypedShaders(), PipelineState);
@@ -50,17 +52,9 @@ void FMeshPassProcessor::BuildMeshDrawCommands(
 	PipelineState.DepthStencilState = DrawRenderState.GetDepthStencilState();
 
 	check(VertexFactory && VertexFactory->IsInitialized());
+	VertexFactory->GetStreams(FeatureLevel, InputStreamType, SharedMeshDrawCommand.VertexStreams);
 
-	if (bPositionOnly)
-	{
-		VertexFactory->GetPositionOnlyStream(SharedMeshDrawCommand.VertexStreams);
-	}
-	else
-	{
-		VertexFactory->GetStreams(FeatureLevel, SharedMeshDrawCommand.VertexStreams);
-	}
-
-	SharedMeshDrawCommand.PrimitiveIdStreamIndex = VertexFactory->GetPrimitiveIdStreamIndex(bPositionOnly);
+	SharedMeshDrawCommand.PrimitiveIdStreamIndex = VertexFactory->GetPrimitiveIdStreamIndex(InputStreamType);
 
 	if (PassShaders.VertexShader)
 	{
@@ -106,28 +100,28 @@ void FMeshPassProcessor::BuildMeshDrawCommands(
 			if (PassShaders.VertexShader)
 			{
 				FMeshDrawSingleShaderBindings VertexShaderBindings = MeshDrawCommand.ShaderBindings.GetSingleShaderBindings(SF_Vertex);
-				PassShaders.VertexShader->GetElementShaderBindings(Scene, ViewIfDynamicMeshCommand, VertexFactory, bPositionOnly, FeatureLevel, PrimitiveSceneProxy, MeshBatch, BatchElement, ShaderElementData, VertexShaderBindings, MeshDrawCommand.VertexStreams);
+				PassShaders.VertexShader->GetElementShaderBindings(Scene, ViewIfDynamicMeshCommand, VertexFactory, InputStreamType, FeatureLevel, PrimitiveSceneProxy, MeshBatch, BatchElement, ShaderElementData, VertexShaderBindings, MeshDrawCommand.VertexStreams);
 			}
 
 			if (PassShaders.HullShader && PassShaders.DomainShader)
 			{
 				FMeshDrawSingleShaderBindings HullShaderBindings = MeshDrawCommand.ShaderBindings.GetSingleShaderBindings(SF_Hull);
 				FMeshDrawSingleShaderBindings DomainShaderBindings = MeshDrawCommand.ShaderBindings.GetSingleShaderBindings(SF_Domain);
-				PassShaders.HullShader->GetElementShaderBindings(Scene, ViewIfDynamicMeshCommand, VertexFactory, false, FeatureLevel, PrimitiveSceneProxy, MeshBatch, BatchElement, ShaderElementData, HullShaderBindings, MeshDrawCommand.VertexStreams);
-				PassShaders.DomainShader->GetElementShaderBindings(Scene, ViewIfDynamicMeshCommand, VertexFactory, false, FeatureLevel, PrimitiveSceneProxy, MeshBatch, BatchElement, ShaderElementData, DomainShaderBindings, MeshDrawCommand.VertexStreams);
+				PassShaders.HullShader->GetElementShaderBindings(Scene, ViewIfDynamicMeshCommand, VertexFactory, EVertexInputStreamType::Default, FeatureLevel, PrimitiveSceneProxy, MeshBatch, BatchElement, ShaderElementData, HullShaderBindings, MeshDrawCommand.VertexStreams);
+				PassShaders.DomainShader->GetElementShaderBindings(Scene, ViewIfDynamicMeshCommand, VertexFactory, EVertexInputStreamType::Default, FeatureLevel, PrimitiveSceneProxy, MeshBatch, BatchElement, ShaderElementData, DomainShaderBindings, MeshDrawCommand.VertexStreams);
 			}
 
 			if (PassShaders.PixelShader)
 			{
 				FMeshDrawSingleShaderBindings PixelShaderBindings = MeshDrawCommand.ShaderBindings.GetSingleShaderBindings(SF_Pixel);
-				PassShaders.PixelShader->GetElementShaderBindings(Scene, ViewIfDynamicMeshCommand, VertexFactory, false, FeatureLevel, PrimitiveSceneProxy, MeshBatch, BatchElement, ShaderElementData, PixelShaderBindings, MeshDrawCommand.VertexStreams);
+				PassShaders.PixelShader->GetElementShaderBindings(Scene, ViewIfDynamicMeshCommand, VertexFactory, EVertexInputStreamType::Default, FeatureLevel, PrimitiveSceneProxy, MeshBatch, BatchElement, ShaderElementData, PixelShaderBindings, MeshDrawCommand.VertexStreams);
 			}
 
 
 			if (PassShaders.GeometryShader)
 			{
 				FMeshDrawSingleShaderBindings GeometryShaderBindings = MeshDrawCommand.ShaderBindings.GetSingleShaderBindings(SF_Geometry);
-				PassShaders.GeometryShader->GetElementShaderBindings(Scene, ViewIfDynamicMeshCommand, VertexFactory, false, FeatureLevel, PrimitiveSceneProxy, MeshBatch, BatchElement, ShaderElementData, GeometryShaderBindings, MeshDrawCommand.VertexStreams);
+				PassShaders.GeometryShader->GetElementShaderBindings(Scene, ViewIfDynamicMeshCommand, VertexFactory, EVertexInputStreamType::Default, FeatureLevel, PrimitiveSceneProxy, MeshBatch, BatchElement, ShaderElementData, GeometryShaderBindings, MeshDrawCommand.VertexStreams);
 			}
 
 			const int32 DrawPrimitiveId = GetDrawCommandPrimitiveId(PrimitiveSceneInfo, BatchElement);
@@ -179,7 +173,7 @@ void FMeshPassProcessor::BuildRayTracingDrawCommands(
 	FMeshDrawCommand SharedMeshDrawCommand;
 
 	check(VertexFactory && VertexFactory->IsInitialized());
-	VertexFactory->GetStreams(FeatureLevel, SharedMeshDrawCommand.VertexStreams);
+	VertexFactory->GetStreams(FeatureLevel, EVertexInputStreamType::Default, SharedMeshDrawCommand.VertexStreams);
 
 #if RHI_RAYTRACING
 	SharedMeshDrawCommand.SetRayTracingShaders(PassShaders.GetUntypedShaders());
@@ -208,7 +202,7 @@ void FMeshPassProcessor::BuildRayTracingDrawCommands(
 			if (PassShaders.RayHitGroupShader)
 			{
 				FMeshDrawSingleShaderBindings RayHitGroupShaderBindings = MeshDrawCommand.ShaderBindings.GetSingleShaderBindings(SF_RayHitGroup);
-				PassShaders.RayHitGroupShader->GetElementShaderBindings(Scene, ViewIfDynamicMeshCommand, VertexFactory, false, FeatureLevel, PrimitiveSceneProxy, MeshBatch, BatchElement, ShaderElementData, RayHitGroupShaderBindings, MeshDrawCommand.VertexStreams);
+				PassShaders.RayHitGroupShader->GetElementShaderBindings(Scene, ViewIfDynamicMeshCommand, VertexFactory, EVertexInputStreamType::Default, FeatureLevel, PrimitiveSceneProxy, MeshBatch, BatchElement, ShaderElementData, RayHitGroupShaderBindings, MeshDrawCommand.VertexStreams);
 			}
 #endif
 
