Fixing assert when using a UAV in a ShaderParameter Struct.

diff --git a/Engine/Source/Runtime/RenderCore/Private/ShaderParameterStruct.cpp b/Engine/Source/Runtime/RenderCore/Private/ShaderParameterStruct.cpp
--- a/Engine/Source/Runtime/RenderCore/Private/ShaderParameterStruct.cpp
+++ b/Engine/Source/Runtime/RenderCore/Private/ShaderParameterStruct.cpp
@@ -65,6 +65,7 @@ struct FShaderParameterStructBindingContext
 			const bool bIsRHIResource = (
 				BaseType == UBMT_TEXTURE ||
 				BaseType == UBMT_SRV ||
+				BaseType == UBMT_UAV ||
 				BaseType == UBMT_SAMPLER);
 			const bool bIsRDGResource = IsRDGResourceReferenceShaderParameterType(BaseType) && BaseType != UBMT_RDG_BUFFER;
 			const bool bIsVariableNativeType = (
