Added RegisterExternalBuffer and QueueBufferExtraction to Render Graph

diff --git a/Engine/Source/Runtime/RenderCore/Private/RenderGraphBuilder.cpp b/Engine/Source/Runtime/RenderCore/Private/RenderGraphBuilder.cpp
--- a/Engine/Source/Runtime/RenderCore/Private/RenderGraphBuilder.cpp
+++ b/Engine/Source/Runtime/RenderCore/Private/RenderGraphBuilder.cpp
@@ -537,6 +537,10 @@ void FRDGBuilder::WalkGraphDependencies()
 	{
 		Query.Texture->ReferenceCount++;
 	}
+	for (const auto& Query : DeferredInternalBufferQueries)
+	{
+		Query.Buffer->ReferenceCount++;
+	}
 
 	// Release external texture that have ReferenceCount == 0 and yet are already allocated.
 	for (auto Pair : AllocatedTextures)
@@ -548,6 +552,17 @@ void FRDGBuilder::WalkGraphDependencies()
 			Pair.Key->CachedRHI.Resource = nullptr;
 		}
 	}
+
+	// Release external buffers that have ReferenceCount == 0 and yet are already allocated.
+	for (auto Pair : AllocatedBuffers)
+	{
+		if (Pair.Key->ReferenceCount == 0)
+		{
+			Pair.Value = nullptr;
+			Pair.Key->PooledBuffer = nullptr;
+			Pair.Key->CachedRHI.Resource = nullptr;
+		}
+	}
 }
 
 void FRDGBuilder::AllocateRHITextureIfNeeded(const FRDGTexture* Texture, bool bComputePass)
@@ -1293,7 +1308,7 @@ void FRDGBuilder::ProcessDeferredInternalResourceQueries()
 		
 		#if RENDER_GRAPH_DEBUGGING
 		{
-			// Increment the number of time the texture has been accessed to avoid warning on produced but never used resources that were produced
+			// Increment the number of times the texture has been accessed to avoid warning on produced but never used resources that were produced
 			// only to be extracted for the graph.
 			Query.Texture->DebugPassAccessCount += 1;
 		}
@@ -1305,6 +1320,25 @@ void FRDGBuilder::ProcessDeferredInternalResourceQueries()
 			ReleaseRHITextureIfPossible(Query.Texture);
 		}
 	}
+
+	for (const auto& Query : DeferredInternalBufferQueries)
+	{
+		*Query.OutBufferPtr = AllocatedBuffers.FindChecked(Query.Buffer);
+
+#if RENDER_GRAPH_DEBUGGING
+		{
+			// Increment the number of times the buffer has been accessed to avoid warning on produced but never used resources that were produced
+			// only to be extracted for the graph.
+			Query.Buffer->DebugPassAccessCount += 1;
+		}
+#endif
+
+		// No need to manually release in immediate mode, since it is done directly when emptying AllocatedBuffer in DestructPasses().
+		if (!GRenderGraphImmediateMode)
+		{
+			ReleaseRHIBufferIfPossible(Query.Buffer);
+		}
+	}
 }
 
 void FRDGBuilder::DestructPasses()
@@ -1349,7 +1383,9 @@ void FRDGBuilder::DestructPasses()
 
 	Passes.Empty();
 	DeferredInternalTextureQueries.Empty();
+	DeferredInternalBufferQueries.Empty();
 	AllocatedTextures.Empty();
+	AllocatedBuffers.Empty();
 }
 
 FRDGBuilder::~FRDGBuilder()
diff --git a/Engine/Source/Runtime/RenderCore/Private/RenderGraphResourcePool.cpp b/Engine/Source/Runtime/RenderCore/Private/RenderGraphResourcePool.cpp
--- a/Engine/Source/Runtime/RenderCore/Private/RenderGraphResourcePool.cpp
+++ b/Engine/Source/Runtime/RenderCore/Private/RenderGraphResourcePool.cpp
@@ -7,24 +7,6 @@
 #include "RenderGraphResourcePool.h"
 #include "RenderGraphResources.h"
 
-
-uint32 FPooledRDGBuffer::Release()
-{
-	return RefCount--;
-
-	if (RefCount == 0)
-	{
-		VertexBuffer.SafeRelease();
-		IndexBuffer.SafeRelease();
-		StructuredBuffer.SafeRelease();
-		UAVs.Empty();
-		SRVs.Empty();
-	}
-
-	return RefCount;
-}
-
-
 FRenderGraphResourcePool::FRenderGraphResourcePool()
 { }
 
diff --git a/Engine/Source/Runtime/RenderCore/Public/RenderGraphBuilder.h b/Engine/Source/Runtime/RenderCore/Public/RenderGraphBuilder.h
--- a/Engine/Source/Runtime/RenderCore/Public/RenderGraphBuilder.h
+++ b/Engine/Source/Runtime/RenderCore/Public/RenderGraphBuilder.h
@@ -224,6 +224,27 @@ public:
 		return OutTexture;
 	}
 
+	/** Register a external buffer to be tracked by the render graph. */
+	inline FRDGBufferRef RegisterExternalBuffer(const TRefCountPtr<FPooledRDGBuffer>& ExternalPooledBuffer, const TCHAR* Name = TEXT("External"))
+	{
+#if RENDER_GRAPH_DEBUGGING
+		{
+			ensureMsgf(ExternalPooledBuffer.IsValid(), TEXT("Attempted to register NULL external buffer: %s"), Name);
+		}
+#endif
+		FRDGBuffer* OutBuffer = AllocateForRHILifeTime<FRDGBuffer>(Name, ExternalPooledBuffer->Desc);
+		OutBuffer->PooledBuffer = ExternalPooledBuffer;
+		AllocatedBuffers.Add(OutBuffer, ExternalPooledBuffer);
+#if RENDER_GRAPH_DEBUGGING
+		{
+			OutBuffer->bHasEverBeenProduced = true;
+			Resources.Add(OutBuffer);
+		}
+#endif
+		return OutBuffer;
+	}
+
+
 	/** Create graph tracked resource from a descriptor with a debug name.
 	 *
 	 * The debug name is the name used for GPU debugging tools, but also for the VisualizeTexture/Vis command.
@@ -429,7 +450,7 @@ public:
 
 	/** Queue a texture extraction. This will set *OutTexturePtr with the internal pooled render target at the Execute().
 	 *
-	 * Note: even when the render graph uses the immediate debugging mode (executing passes as they get added), the texture extrations
+	 * Note: even when the render graph uses the immediate debugging mode (executing passes as they get added), the texture extractions
 	 * will still happen in the Execute(), to ensure there is no bug caused in code outside the render graph on whether this mode is used or not.
 	 */
 	inline void QueueTextureExtraction(FRDGTextureRef Texture, TRefCountPtr<IPooledRenderTarget>* OutTexturePtr, bool bTransitionToRead = true)
@@ -454,6 +475,32 @@ public:
 		DeferredInternalTextureQueries.Emplace(Query);
 	}
 
+	/** Queue a buffer extraction. This will set *OutBufferPtr with the internal pooled buffer at the Execute().
+	 *
+	 * Note: even when the render graph uses the immediate debugging mode (executing passes as they get added), the buffer extractions
+	 * will still happen in the Execute(), to ensure there is no bug caused in code outside the render graph on whether this mode is used or not.
+	 */
+	inline void QueueBufferExtraction(FRDGBufferRef Buffer, TRefCountPtr<FPooledRDGBuffer>* OutBufferPtr)
+	{
+		check(Buffer);
+		check(OutBufferPtr);
+#if RENDER_GRAPH_DEBUGGING
+		{
+			checkf(!bHasExecuted,
+				TEXT("Accessing render graph internal buffer %s with QueueBufferExtraction() needs to happen before the builder's execution."),
+				Buffer->Name);
+
+			checkf(Buffer->bHasEverBeenProduced,
+				TEXT("Unable to queue the extraction of the buffer %s because it has not been produced by any pass."),
+				Buffer->Name);
+		}
+#endif
+		FDeferredInternalBufferQuery Query;
+		Query.Buffer = Buffer;
+		Query.OutBufferPtr = OutBufferPtr;
+		DeferredInternalBufferQueries.Emplace(Query);
+	}
+
 	/** Flag a texture that is only produced by only 1 pass, but never used or extracted, to avoid generating a warning at runtime. */
 	FORCEINLINE_DEBUGGABLE void RemoveUnusedTextureWarning(FRDGTextureRef Texture)
 	{
@@ -508,6 +555,14 @@ private:
 	};
 	TArray<FDeferredInternalTextureQuery, SceneRenderingAllocator> DeferredInternalTextureQueries;
 
+	/** Array of all deferred access to internal buffers. */
+	struct FDeferredInternalBufferQuery
+	{
+		const FRDGBuffer* Buffer;
+		TRefCountPtr<FPooledRDGBuffer>* OutBufferPtr;
+	};
+	TArray<FDeferredInternalBufferQuery, SceneRenderingAllocator> DeferredInternalBufferQueries;
+
 	#if RENDER_GRAPH_DRAW_EVENTS == 2
 		/** All scopes allocated that needs to be arround to call destructors. */
 		TArray<FRDGEventScope*, SceneRenderingAllocator> EventScopes;
diff --git a/Engine/Source/Runtime/RenderCore/Public/RenderGraphResources.h b/Engine/Source/Runtime/RenderCore/Public/RenderGraphResources.h
--- a/Engine/Source/Runtime/RenderCore/Public/RenderGraphResources.h
+++ b/Engine/Source/Runtime/RenderCore/Public/RenderGraphResources.h
@@ -408,7 +408,21 @@ struct FPooledRDGBuffer
 		return ++RefCount;
 	}
 
-	uint32 Release();
+	uint32 Release()
+	{
+		RefCount--;
+
+		if (RefCount == 0)
+		{
+			VertexBuffer.SafeRelease();
+			IndexBuffer.SafeRelease();
+			StructuredBuffer.SafeRelease();
+			UAVs.Empty();
+			SRVs.Empty();
+		}
+
+		return RefCount;
+	}
 
 	inline uint32 GetRefCount()
 	{
