Fixes render graph resource mutable member.

diff --git a/Engine/Source/Runtime/RenderCore/Private/RenderGraphBuilder.cpp b/Engine/Source/Runtime/RenderCore/Private/RenderGraphBuilder.cpp
--- a/Engine/Source/Runtime/RenderCore/Private/RenderGraphBuilder.cpp
+++ b/Engine/Source/Runtime/RenderCore/Private/RenderGraphBuilder.cpp
@@ -336,7 +336,7 @@ void FRDGBuilder::ValidatePass(const FRenderGraphPass* Pass) const
 		for (int32 i = 0; i < RenderTargets->Output.Num(); i++)
 		{
 			const FRenderTargetBinding& RenderTarget = RenderTargets->Output[i];
-			const FRDGTexture* Texture = RenderTarget.GetTexture();
+			FRDGTexture* Texture = RenderTarget.GetTexture();
 
 			if (!Texture)
 			{
@@ -364,7 +364,7 @@ void FRDGBuilder::ValidatePass(const FRenderGraphPass* Pass) const
 		}
 		ensureMsgf(!bGeneratingMips || bFoundRTBound, TEXT("GenerateMips enabled but no RT found as source!"));
 
-		const FRDGTexture* Texture = RenderTargets->DepthStencil.Texture;
+		FRDGTexture* Texture = RenderTargets->DepthStencil.Texture;
 		if (Texture && !Texture->bHasEverBeenProduced)
 		{
 			checkf(RenderTargets->DepthStencil.DepthLoadAction != ERenderTargetLoadAction::ELoad,
@@ -562,7 +562,7 @@ void FRDGBuilder::WalkGraphDependencies()
 	}
 }
 
-void FRDGBuilder::AllocateRHITextureIfNeeded(const FRDGTexture* Texture, bool bComputePass)
+void FRDGBuilder::AllocateRHITextureIfNeeded(FRDGTexture* Texture, bool bComputePass)
 {
 	check(Texture);
 
@@ -582,7 +582,7 @@ void FRDGBuilder::AllocateRHITextureIfNeeded(const FRDGTexture* Texture, bool bC
 	check(Texture->CachedRHI.Resource);
 }
 
-void FRDGBuilder::AllocateRHITextureUAVIfNeeded(const FRDGTextureUAV* UAV, bool bComputePass)
+void FRDGBuilder::AllocateRHITextureUAVIfNeeded(FRDGTextureUAV* UAV, bool bComputePass)
 {
 	check(UAV);
 
@@ -596,7 +596,7 @@ void FRDGBuilder::AllocateRHITextureUAVIfNeeded(const FRDGTextureUAV* UAV, bool
 	UAV->CachedRHI.UAV = UAV->Desc.Texture->PooledRenderTarget->GetRenderTargetItem().MipUAVs[UAV->Desc.MipLevel];
 }
 
-void FRDGBuilder::AllocateRHIBufferSRVIfNeeded(const FRDGBufferSRV* SRV, bool bComputePass)
+void FRDGBuilder::AllocateRHIBufferSRVIfNeeded(FRDGBufferSRV* SRV, bool bComputePass)
 {
 	check(SRV);
 
@@ -638,7 +638,7 @@ void FRDGBuilder::AllocateRHIBufferSRVIfNeeded(const FRDGBufferSRV* SRV, bool bC
 	SRV->Desc.Buffer->PooledBuffer->SRVs.Add(SRV->Desc, RHIShaderResourceView);
 }
 
-void FRDGBuilder::AllocateRHIBufferUAVIfNeeded(const FRDGBufferUAV* UAV, bool bComputePass)
+void FRDGBuilder::AllocateRHIBufferUAVIfNeeded(FRDGBufferUAV* UAV, bool bComputePass)
 {
 	check(UAV);
 
@@ -705,7 +705,7 @@ static EResourceTransitionPipeline CalcTransitionPipeline(bool bCurrentCompute,
 	return static_cast< EResourceTransitionPipeline >( Table[ Bits ] );
 }
 
-void FRDGBuilder::TransitionTexture( const FRDGTexture* Texture, EResourceTransitionAccess TransitionAccess, bool bRequiredCompute ) const
+void FRDGBuilder::TransitionTexture(FRDGTexture* Texture, EResourceTransitionAccess TransitionAccess, bool bRequiredCompute) const
 {
 	const bool bRequiredWritable = TransitionAccess != EResourceTransitionAccess::EReadable;
 
@@ -717,7 +717,7 @@ void FRDGBuilder::TransitionTexture( const FRDGTexture* Texture, EResourceTransi
 	}
 }
 
-void FRDGBuilder::TransitionUAV(FUnorderedAccessViewRHIParamRef UAV, const FRDGResource* UnderlyingResource, EResourceTransitionAccess TransitionAccess, bool bRequiredCompute ) const
+void FRDGBuilder::TransitionUAV(FUnorderedAccessViewRHIParamRef UAV, FRDGResource* UnderlyingResource, EResourceTransitionAccess TransitionAccess, bool bRequiredCompute ) const
 {
 	const bool bRequiredWritable = true;
 
@@ -1160,7 +1160,7 @@ void FRDGBuilder::WarnForUselessPassDependencies(const FRenderGraphPass* Pass)
 		if (!IsRDGResourceReferenceShaderParameterType(Type))
 			continue;
 
-		const FRDGResource* Resource = *ParameterStruct.GetMemberPtrAtOffset<const FRDGResource*>(Offset);
+		FRDGResource* Resource = *ParameterStruct.GetMemberPtrAtOffset<FRDGResource*>(Offset);
 
 		if (!Resource)
 			continue;
@@ -1169,7 +1169,7 @@ void FRDGBuilder::WarnForUselessPassDependencies(const FRenderGraphPass* Pass)
 	}
 }
 
-void FRDGBuilder::ReleaseRHITextureIfPossible(const FRDGTexture* Texture)
+void FRDGBuilder::ReleaseRHITextureIfPossible(FRDGTexture* Texture)
 {
 	check(Texture->ReferenceCount > 0);
 	Texture->ReferenceCount--;
@@ -1182,7 +1182,7 @@ void FRDGBuilder::ReleaseRHITextureIfPossible(const FRDGTexture* Texture)
 	}
 }
 
-void FRDGBuilder::ReleaseRHIBufferIfPossible(const FRDGBuffer* Buffer)
+void FRDGBuilder::ReleaseRHIBufferIfPossible(FRDGBuffer* Buffer)
 {
 	check(Buffer->ReferenceCount > 0);
 	Buffer->ReferenceCount--;
diff --git a/Engine/Source/Runtime/RenderCore/Public/RenderGraphBuilder.h b/Engine/Source/Runtime/RenderCore/Public/RenderGraphBuilder.h
--- a/Engine/Source/Runtime/RenderCore/Public/RenderGraphBuilder.h
+++ b/Engine/Source/Runtime/RenderCore/Public/RenderGraphBuilder.h
@@ -203,7 +203,7 @@ public:
 	~FRDGBuilder();
 
 	/** Register a external texture to be tracked by the render graph. */
-	inline FRDGTextureRef RegisterExternalTexture(const TRefCountPtr<IPooledRenderTarget>& ExternalPooledTexture, const TCHAR* DebugName = TEXT("External"))
+	FORCEINLINE_DEBUGGABLE FRDGTextureRef RegisterExternalTexture(const TRefCountPtr<IPooledRenderTarget>& ExternalPooledTexture, const TCHAR* DebugName = TEXT("External"))
 	{
 		#if RENDER_GRAPH_DEBUGGING
 		{
@@ -225,7 +225,7 @@ public:
 	}
 
 	/** Register a external buffer to be tracked by the render graph. */
-	inline FRDGBufferRef RegisterExternalBuffer(const TRefCountPtr<FPooledRDGBuffer>& ExternalPooledBuffer, const TCHAR* Name = TEXT("External"))
+	FORCEINLINE_DEBUGGABLE FRDGBufferRef RegisterExternalBuffer(const TRefCountPtr<FPooledRDGBuffer>& ExternalPooledBuffer, const TCHAR* Name = TEXT("External"))
 	{
 #if RENDER_GRAPH_DEBUGGING
 		{
@@ -249,7 +249,7 @@ public:
 	 *
 	 * The debug name is the name used for GPU debugging tools, but also for the VisualizeTexture/Vis command.
 	 */
-	inline FRDGTextureRef CreateTexture(const FPooledRenderTargetDesc& Desc, const TCHAR* DebugName)
+	FORCEINLINE_DEBUGGABLE FRDGTextureRef CreateTexture(const FPooledRenderTargetDesc& Desc, const TCHAR* DebugName)
 	{
 		#if RENDER_GRAPH_DEBUGGING
 		{
@@ -269,7 +269,7 @@ public:
 	 *
 	 * The debug name is the name used for GPU debugging tools, but also for the VisualizeTexture/Vis command.
 	 */
-	inline FRDGBufferRef CreateBuffer(const FRDGBufferDesc& Desc, const TCHAR* DebugName)
+	FORCEINLINE_DEBUGGABLE FRDGBufferRef CreateBuffer(const FRDGBufferDesc& Desc, const TCHAR* DebugName)
 	{
 		#if RENDER_GRAPH_DEBUGGING
 		{
@@ -285,7 +285,7 @@ public:
 	}
 
 	/** Create graph tracked SRV for a texture from a descriptor. */
-	inline FRDGTextureSRVRef CreateSRV(const FRDGTextureSRVDesc& Desc)
+	FORCEINLINE_DEBUGGABLE FRDGTextureSRVRef CreateSRV(const FRDGTextureSRVDesc& Desc)
 	{
 		check(Desc.Texture);
 		#if RENDER_GRAPH_DEBUGGING
@@ -303,7 +303,7 @@ public:
 	}
 
 	/** Create graph tracked SRV for a buffer from a descriptor. */
-	inline FRDGBufferSRVRef CreateSRV(const FRDGBufferSRVDesc& Desc)
+	FORCEINLINE_DEBUGGABLE FRDGBufferSRVRef CreateSRV(const FRDGBufferSRVDesc& Desc)
 	{
 		check(Desc.Buffer);
 		#if RENDER_GRAPH_DEBUGGING
@@ -320,7 +320,7 @@ public:
 	}
 
 	/** Create graph tracked UAV for a texture from a descriptor. */
-	inline FRDGTextureUAVRef CreateUAV(const FRDGTextureUAVDesc& Desc)
+	FORCEINLINE_DEBUGGABLE FRDGTextureUAVRef CreateUAV(const FRDGTextureUAVDesc& Desc)
 	{
 		check(Desc.Texture);
 		#if RENDER_GRAPH_DEBUGGING
@@ -338,7 +338,7 @@ public:
 	}
 
 	/** Create graph tracked UAV for a buffer from a descriptor. */
-	inline FRDGBufferUAVRef CreateUAV(const FRDGBufferUAVDesc& Desc)
+	FORCEINLINE_DEBUGGABLE FRDGBufferUAVRef CreateUAV(const FRDGBufferUAVDesc& Desc)
 	{
 		check(Desc.Buffer);
 		#if RENDER_GRAPH_DEBUGGING
@@ -354,14 +354,14 @@ public:
 		return UAV;
 	}
 
-	inline FRDGBufferUAVRef CreateUAV(FRDGBufferRef Buffer, EPixelFormat Format)
+	FORCEINLINE_DEBUGGABLE FRDGBufferUAVRef CreateUAV(FRDGBufferRef Buffer, EPixelFormat Format)
 	{
 		return CreateUAV( FRDGBufferUAVDesc(Buffer, Format) );
 	}
 
 	/** Allocates parameter struct specifically to survive through the life time of the render graph. */
 	template< typename ParameterStructType >
-	inline ParameterStructType* AllocParameters()
+	FORCEINLINE_DEBUGGABLE ParameterStructType* AllocParameters()
 	{
 		// TODO(RDG): could allocate using AllocateForRHILifeTime() to avoid the copy done when using FRHICommandList::BuildLocalUniformBuffer()
 		ParameterStructType* OutParameterPtr = new(FMemStack::Get()) ParameterStructType;
@@ -430,7 +430,7 @@ public:
 	 *
 	 * Caution: You are on your own to have correct memory lifetime of the FRenderGraphPass.
 	 */
-	void AddProcedurallyCreatedPass(FRenderGraphPass* NewPass)
+	FORCEINLINE_DEBUGGABLE void AddProcedurallyCreatedPass(FRenderGraphPass* NewPass)
 	{
 		#if RENDER_GRAPH_DEBUGGING
 		{
@@ -453,7 +453,7 @@ public:
 	 * Note: even when the render graph uses the immediate debugging mode (executing passes as they get added), the texture extractions
 	 * will still happen in the Execute(), to ensure there is no bug caused in code outside the render graph on whether this mode is used or not.
 	 */
-	inline void QueueTextureExtraction(FRDGTextureRef Texture, TRefCountPtr<IPooledRenderTarget>* OutTexturePtr, bool bTransitionToRead = true)
+	FORCEINLINE_DEBUGGABLE void QueueTextureExtraction(FRDGTextureRef Texture, TRefCountPtr<IPooledRenderTarget>* OutTexturePtr, bool bTransitionToRead = true)
 	{
 		check(Texture);
 		check(OutTexturePtr);
@@ -480,7 +480,7 @@ public:
 	 * Note: even when the render graph uses the immediate debugging mode (executing passes as they get added), the buffer extractions
 	 * will still happen in the Execute(), to ensure there is no bug caused in code outside the render graph on whether this mode is used or not.
 	 */
-	inline void QueueBufferExtraction(FRDGBufferRef Buffer, TRefCountPtr<FPooledRDGBuffer>* OutBufferPtr)
+	FORCEINLINE_DEBUGGABLE void QueueBufferExtraction(FRDGBufferRef Buffer, TRefCountPtr<FPooledRDGBuffer>* OutBufferPtr)
 	{
 		check(Buffer);
 		check(OutBufferPtr);
@@ -528,9 +528,12 @@ public:
 	{
 		return CurrentScope;
 	}
+
 public:
 	/** The RHI command list used for the render graph. */
 	FRHICommandListImmediate& RHICmdList;
+
+
 private:
 	static constexpr int32 kMaxScopeCount = 8;
 
@@ -538,15 +541,15 @@ private:
 	TArray<FRenderGraphPass*, SceneRenderingAllocator> Passes;
 
 	/** Keep the references over the pooled render target, since FRDGTexture is allocated on FMemStack. */
-	TMap<const FRDGTexture*, TRefCountPtr<IPooledRenderTarget>, SceneRenderingSetAllocator> AllocatedTextures;
+	TMap<FRDGTexture*, TRefCountPtr<IPooledRenderTarget>, SceneRenderingSetAllocator> AllocatedTextures;
 
 	/** Keep the references over the pooled render target, since FRDGTexture is allocated on FMemStack. */
-	TMap<const FRDGBuffer*, TRefCountPtr<FPooledRDGBuffer>, SceneRenderingSetAllocator> AllocatedBuffers;
+	TMap<FRDGBuffer*, TRefCountPtr<FPooledRDGBuffer>, SceneRenderingSetAllocator> AllocatedBuffers;
 
 	/** Array of all deferred access to internal textures. */
 	struct FDeferredInternalTextureQuery
 	{
-		const FRDGTexture* Texture;
+		FRDGTexture* Texture;
 		TRefCountPtr<IPooledRenderTarget>* OutTexturePtr;
 		bool bTransitionToRead;
 	};
@@ -555,7 +558,7 @@ private:
 	/** Array of all deferred access to internal buffers. */
 	struct FDeferredInternalBufferQuery
 	{
-		const FRDGBuffer* Buffer;
+		FRDGBuffer* Buffer;
 		TRefCountPtr<FPooledRDGBuffer>* OutBufferPtr;
 	};
 	TArray<FDeferredInternalBufferQuery, SceneRenderingAllocator> DeferredInternalBufferQueries;
@@ -605,22 +608,22 @@ private:
 		}
 	}
 
-	void AllocateRHITextureIfNeeded(const FRDGTexture* Texture, bool bComputePass);
-	void AllocateRHITextureUAVIfNeeded(const FRDGTextureUAV* UAV, bool bComputePass);
-	void AllocateRHIBufferSRVIfNeeded(const FRDGBufferSRV* SRV, bool bComputePass);
-	void AllocateRHIBufferUAVIfNeeded(const FRDGBufferUAV* UAV, bool bComputePass);
+	void AllocateRHITextureIfNeeded(FRDGTexture* Texture, bool bComputePass);
+	void AllocateRHITextureUAVIfNeeded(FRDGTextureUAV* UAV, bool bComputePass);
+	void AllocateRHIBufferSRVIfNeeded(FRDGBufferSRV* SRV, bool bComputePass);
+	void AllocateRHIBufferUAVIfNeeded(FRDGBufferUAV* UAV, bool bComputePass);
 
 
-	void TransitionTexture( const FRDGTexture* Texture, EResourceTransitionAccess TransitionAccess, bool bRequiredCompute ) const;
-	void TransitionUAV(FUnorderedAccessViewRHIParamRef UAV, const FRDGResource* UnderlyingResource, EResourceTransitionAccess TransitionAccess, bool bRequiredCompute ) const;
+	void TransitionTexture(FRDGTexture* Texture, EResourceTransitionAccess TransitionAccess, bool bRequiredCompute) const;
+	void TransitionUAV(FUnorderedAccessViewRHIParamRef UAV, FRDGResource* UnderlyingResource, EResourceTransitionAccess TransitionAccess, bool bRequiredCompute ) const;
 
 	void PushDrawEventStack(const FRenderGraphPass* Pass);
 	void ExecutePass( const FRenderGraphPass* Pass );
 	void AllocateAndTransitionPassResources(const FRenderGraphPass* Pass, struct FRHIRenderPassInfo* OutRPInfo, bool* bOutHasRenderTargets);
 	static void WarnForUselessPassDependencies(const FRenderGraphPass* Pass);
 
-	void ReleaseRHITextureIfPossible(const FRDGTexture* Texture);
-	void ReleaseRHIBufferIfPossible(const FRDGBuffer* Buffer);
+	void ReleaseRHITextureIfPossible(FRDGTexture* Texture);
+	void ReleaseRHIBufferIfPossible(FRDGBuffer* Buffer);
 	void ReleaseUnecessaryResources(const FRenderGraphPass* Pass);
 
 	void ProcessDeferredInternalResourceQueries();
@@ -645,7 +648,7 @@ public:
 	FStackRDGEventScopeRef(FStackRDGEventScopeRef&&) = delete;
 	void operator = (const FStackRDGEventScopeRef&) = delete;
 
-	inline FStackRDGEventScopeRef(FRDGBuilder& InGraphBuilder, FRDGEventName&& ScopeName)
+	FORCEINLINE_DEBUGGABLE FStackRDGEventScopeRef(FRDGBuilder& InGraphBuilder, FRDGEventName&& ScopeName)
 		: GraphBuilder(InGraphBuilder)
 	{
 		checkf(!GraphBuilder.bHasExecuted, TEXT("Render graph bulider has already been executed."));
@@ -661,7 +664,7 @@ public:
 		GraphBuilder.CurrentScope = NewScope;
 	}
 
-	inline ~FStackRDGEventScopeRef()
+	FORCEINLINE_DEBUGGABLE ~FStackRDGEventScopeRef()
 	{
 		check(GraphBuilder.CurrentScope != nullptr);
 		GraphBuilder.CurrentScope = GraphBuilder.CurrentScope->ParentScope;
diff --git a/Engine/Source/Runtime/RenderCore/Public/RenderGraphResources.h b/Engine/Source/Runtime/RenderCore/Public/RenderGraphResources.h
--- a/Engine/Source/Runtime/RenderCore/Public/RenderGraphResources.h
+++ b/Engine/Source/Runtime/RenderCore/Public/RenderGraphResources.h
@@ -30,13 +30,13 @@ struct FRayTracingShaderBindingsWriter;
 
 
 /** Defines the RDG resource references for user code not forgetting the const every time. */
-using FRDGResourceRef = const FRDGResource*;
-using FRDGTextureRef = const FRDGTexture*;
-using FRDGTextureSRVRef = const FRDGTextureSRV*;
-using FRDGTextureUAVRef = const FRDGTextureUAV*;
-using FRDGBufferRef = const  FRDGBuffer*;
-using FRDGBufferSRVRef = const FRDGBufferSRV*;
-using FRDGBufferUAVRef = const FRDGBufferUAV*;
+using FRDGResourceRef = FRDGResource*;
+using FRDGTextureRef = FRDGTexture*;
+using FRDGTextureSRVRef = FRDGTextureSRV*;
+using FRDGTextureUAVRef = FRDGTextureUAV*;
+using FRDGBufferRef = FRDGBuffer*;
+using FRDGBufferSRVRef = FRDGBufferSRV*;
+using FRDGBufferUAVRef = FRDGBufferUAV*;
 
 
 /** Generic graph resource. */
@@ -46,10 +46,10 @@ protected:
 	/** Pointer on the RHI resource once allocated, that the RHI can dereferenced according to IsUniformBufferResourceIndirectionType(). */
 	union
 	{
-		mutable FRHIResource* Resource;
-		mutable FTextureRHIParamRef Texture;
-		mutable FShaderResourceViewRHIParamRef SRV;
-		mutable FUnorderedAccessViewRHIParamRef UAV;
+		FRHIResource* Resource;
+		FTextureRHIParamRef Texture;
+		FShaderResourceViewRHIParamRef SRV;
+		FUnorderedAccessViewRHIParamRef UAV;
 	} CachedRHI;
 
 public:
@@ -67,25 +67,25 @@ public:
 	void operator = (const FRDGResource&) = delete;
 
 	/** Boolean to track at runtime whether a ressource is actually used by the lambda of a pass or not, to detect unnecessary resource dependencies on passes. */
-	mutable bool bIsActuallyUsedByPass = false;
+	bool bIsActuallyUsedByPass = false;
 
 private:
 	/** Number of references in passes and deferred queries. */
-	mutable int32 ReferenceCount = 0;
+	int32 ReferenceCount = 0;
 
 	// Used for tracking resource state during execution
-	mutable bool bWritable = false;
-	mutable bool bCompute = false;
+	bool bWritable = false;
+	bool bCompute = false;
 
 #if RENDER_GRAPH_DEBUGGING
 	/** Boolean to track at wiring time if a resource has ever been produced by a pass, to error out early if accessing a resource that has not been produced. */
-	mutable bool bHasEverBeenProduced = false;
+	bool bHasEverBeenProduced = false;
 
 	/** Pointer towards the pass that is the first to produce it, for even more convenient error message. */
-	mutable const FRenderGraphPass* DebugFirstProducer = nullptr;
+	const FRenderGraphPass* DebugFirstProducer = nullptr;
 
 	/** Count the number of times it has been used by a pass. */
-	mutable int32 DebugPassAccessCount = 0;
+	int32 DebugPassAccessCount = 0;
 #endif
 
 	friend class FRDGBuilder;
diff --git a/Engine/Source/Runtime/RenderCore/Public/ShaderParameterMacros.h b/Engine/Source/Runtime/RenderCore/Public/ShaderParameterMacros.h
--- a/Engine/Source/Runtime/RenderCore/Public/ShaderParameterMacros.h
+++ b/Engine/Source/Runtime/RenderCore/Public/ShaderParameterMacros.h
@@ -149,26 +149,26 @@ struct alignas(SHADER_PARAMETER_STRUCT_ALIGNMENT) FRenderTargetBinding
 	 *
 	 * Notes: Load and store action are on purpose without default values, to force the user to not forget one of these.
 	 */
-	FRenderTargetBinding(const FRDGTexture* InTexture, ERenderTargetLoadAction InLoadAction, ERenderTargetStoreAction InStoreAction, uint8 InMipIndex = 0)
+	FRenderTargetBinding(FRDGTexture* InTexture, ERenderTargetLoadAction InLoadAction, ERenderTargetStoreAction InStoreAction, uint8 InMipIndex = 0)
 		: Texture(InTexture)
 		, LoadAction(InLoadAction)
 		, StoreAction(InStoreAction)
 		, MipIndex(InMipIndex)
 	{}
 
-	inline const FRDGTexture* GetTexture() const
+	FORCEINLINE FRDGTexture* GetTexture() const
 	{
 		return Texture;
 	}
-	inline ERenderTargetLoadAction GetLoadAction() const
+	FORCEINLINE ERenderTargetLoadAction GetLoadAction() const
 	{
 		return LoadAction;
 	}
-	inline ERenderTargetStoreAction GetStoreAction() const
+	FORCEINLINE ERenderTargetStoreAction GetStoreAction() const
 	{
 		return StoreAction;
 	}
-	inline uint8 GetMipIndex() const
+	FORCEINLINE uint8 GetMipIndex() const
 	{
 		return MipIndex;
 	}
@@ -177,7 +177,7 @@ private:
 	/** All parameters required to bind a render target deferred. This are purposefully private to
 	 * force the user to call FRenderTargetBinding() constructor, forcing him to specify the load and store action.
 	 */
-	TAlignedShaderParameterPtr<const FRDGTexture*> Texture;
+	TAlignedShaderParameterPtr<FRDGTexture*> Texture;
 	ERenderTargetLoadAction		LoadAction		= ERenderTargetLoadAction::ENoAction;
 	ERenderTargetStoreAction	StoreAction		= ERenderTargetStoreAction::ENoAction;
 	uint8						MipIndex		= 0;
@@ -192,7 +192,7 @@ struct alignas(SHADER_PARAMETER_STRUCT_ALIGNMENT) FDepthStencilBinding
 	{ }
 	FDepthStencilBinding(const FDepthStencilBinding&) = default;
 
-	TAlignedShaderParameterPtr<const FRDGTexture*> Texture;
+	TAlignedShaderParameterPtr<FRDGTexture*> Texture;
 	ERenderTargetLoadAction		DepthLoadAction		= ERenderTargetLoadAction::ENoAction;
 	ERenderTargetStoreAction	DepthStoreAction	= ERenderTargetStoreAction::ENoAction;
 	ERenderTargetLoadAction		StencilLoadAction	= ERenderTargetLoadAction::ENoAction;
diff --git a/Engine/Source/Runtime/RenderCore/Public/ShaderParameterStruct.h b/Engine/Source/Runtime/RenderCore/Public/ShaderParameterStruct.h
--- a/Engine/Source/Runtime/RenderCore/Public/ShaderParameterStruct.h
+++ b/Engine/Source/Runtime/RenderCore/Public/ShaderParameterStruct.h
@@ -73,7 +73,11 @@ FORCEINLINE void ValidateShaderParameters(const FShader* Shader, const FShaderPa
 /** Set compute shader UAVs. */
 template<typename TRHICmdList, typename TShaderClass, typename TShaderRHI>
 inline void SetShaderUAVs(TRHICmdList& RHICmdList, const TShaderClass* Shader, TShaderRHI* ShadeRHI, const typename TShaderClass::FParameters& Parameters)
-{}
+{
+	checkf(
+		Shader->Bindings.UAVs.Num() == 0 && Shader->Bindings.GraphUAVs.Num() == 0,
+		TEXT("TShaderRHI Can't have compute shader to be set. UAVs are not supported on vertex, tessellation, geometry and pixel shaders."));
+}
 
 template<typename TRHICmdList, typename TShaderClass>
 inline void SetShaderUAVs(TRHICmdList& RHICmdList, const TShaderClass* Shader, FComputeShaderRHIParamRef ShadeRHI, const typename TShaderClass::FParameters& Parameters)
@@ -99,7 +103,7 @@ inline void SetShaderUAVs(TRHICmdList& RHICmdList, const TShaderClass* Shader, F
 	// Graph UAVs
 	for (const FShaderParameterBindings::FResourceParameter& ParameterBinding : Bindings.GraphUAVs)
 	{
-		auto GraphUAV = *reinterpret_cast<const FRDGTextureUAV* const*>(Base + ParameterBinding.ByteOffset);
+		auto GraphUAV = *reinterpret_cast<FRDGTextureUAV* const*>(Base + ParameterBinding.ByteOffset);
 
 		if (DO_CHECK)
 		{
@@ -202,7 +206,7 @@ inline void SetShaderParameters(TRHICmdList& RHICmdList, const TShaderClass* Sha
 	// Graph Textures
 	for (const FShaderParameterBindings::FResourceParameter& ParameterBinding : Bindings.GraphTextures)
 	{
-		auto GraphTexture = *reinterpret_cast<const FRDGTexture* const*>(Base + ParameterBinding.ByteOffset);
+		auto GraphTexture = *reinterpret_cast<FRDGTexture* const*>(Base + ParameterBinding.ByteOffset);
 
 		if (DO_CHECK)
 		{
@@ -224,7 +228,7 @@ inline void SetShaderParameters(TRHICmdList& RHICmdList, const TShaderClass* Sha
 	// Graph SRVs
 	for (const FShaderParameterBindings::FResourceParameter& ParameterBinding : Bindings.GraphSRVs)
 	{
-		auto GraphSRV = *reinterpret_cast<const FRDGTextureSRV* const*>(Base + ParameterBinding.ByteOffset);
+		auto GraphSRV = *reinterpret_cast<FRDGTextureSRV* const*>(Base + ParameterBinding.ByteOffset);
 
 		if (DO_CHECK)
 		{
@@ -329,7 +333,7 @@ void SetShaderParameters(FRayTracingShaderBindingsWriter& RTBindingsWriter, cons
 	// Graph Textures
 	for (const FShaderParameterBindings::FResourceParameter& ParameterBinding : Bindings.GraphTextures)
 	{
-		auto GraphTexture = *reinterpret_cast<const FRDGTexture* const*>(Base + ParameterBinding.ByteOffset);
+		auto GraphTexture = *reinterpret_cast<FRDGTexture* const*>(Base + ParameterBinding.ByteOffset);
 
 		if (DO_CHECK)
 		{
@@ -351,7 +355,7 @@ void SetShaderParameters(FRayTracingShaderBindingsWriter& RTBindingsWriter, cons
 	// Graph SRVs
 	for (const FShaderParameterBindings::FResourceParameter& ParameterBinding : Bindings.GraphSRVs)
 	{
-		auto GraphSRV = *reinterpret_cast<const FRDGResource* const*>(Base + ParameterBinding.ByteOffset);
+		auto GraphSRV = *reinterpret_cast<FRDGResource* const*>(Base + ParameterBinding.ByteOffset);
 
 		if (DO_CHECK)
 		{
@@ -373,7 +377,7 @@ void SetShaderParameters(FRayTracingShaderBindingsWriter& RTBindingsWriter, cons
 	// Render graph UAVs
 	for (const FShaderParameterBindings::FResourceParameter& ParameterBinding : Bindings.GraphUAVs)
 	{
-		auto UAV = *reinterpret_cast<const FRDGResource* const*>(Base + ParameterBinding.ByteOffset);
+		auto UAV = *reinterpret_cast<FRDGResource* const*>(Base + ParameterBinding.ByteOffset);
 
 		if (DO_CHECK)
 		{
diff --git a/Engine/Source/Runtime/RenderCore/Public/VisualizeTexture.h b/Engine/Source/Runtime/RenderCore/Public/VisualizeTexture.h
--- a/Engine/Source/Runtime/RenderCore/Public/VisualizeTexture.h
+++ b/Engine/Source/Runtime/RenderCore/Public/VisualizeTexture.h
@@ -104,7 +104,7 @@ private:
 	bool ShouldCapture(const TCHAR* DebugName);
 
 	/** Create a pass capturing a texture. */
-	void CreateContentCapturePass(FRDGBuilder& GraphBuilder, const FRDGTextureRef Texture);
+	void CreateContentCapturePass(FRDGBuilder& GraphBuilder, FRDGTextureRef Texture);
 #endif
 
 	friend class FRDGBuilder;
diff --git a/Engine/Source/Runtime/Renderer/Private/DeferredShadingRenderer.h b/Engine/Source/Runtime/Renderer/Private/DeferredShadingRenderer.h
--- a/Engine/Source/Runtime/Renderer/Private/DeferredShadingRenderer.h
+++ b/Engine/Source/Runtime/Renderer/Private/DeferredShadingRenderer.h
@@ -390,7 +390,7 @@ private:
 		FIntVector VolumetricFogGridSize,
 		FVector GridZParams,
 		const FPooledRenderTargetDesc& VolumeDesc,
-		const FRDGTexture*& OutLocalShadowedLightScattering);
+		FRDGTexture*& OutLocalShadowedLightScattering);
 
 	void RenderLightFunctionForVolumetricFog(
 		FRDGBuilder& GraphBuilder,
@@ -398,7 +398,7 @@ private:
 		FIntVector VolumetricFogGridSize,
 		float VolumetricFogMaxDistance,
 		FMatrix& OutLightFunctionWorldToShadow,
-		const FRDGTexture*& OutLightFunctionTexture,
+		FRDGTexture*& OutLightFunctionTexture,
 		bool& bOutUseDirectionalLightShadowing);
 
 	void VoxelizeFogVolumePrimitives(
diff --git a/Engine/Source/Runtime/Renderer/Private/ScreenSpaceDiffuseIndirect.cpp b/Engine/Source/Runtime/Renderer/Private/ScreenSpaceDiffuseIndirect.cpp
--- a/Engine/Source/Runtime/Renderer/Private/ScreenSpaceDiffuseIndirect.cpp
+++ b/Engine/Source/Runtime/Renderer/Private/ScreenSpaceDiffuseIndirect.cpp
@@ -89,12 +89,12 @@ void RenderScreenSpaceDiffuseIndirect( FRHICommandListImmediate& RHICmdList, FVi
 
 	FSceneRenderTargets& SceneContext = FSceneRenderTargets::Get( RHICmdList );
 
-	const FRDGTexture* SceneColor	= GraphBuilder.RegisterExternalTexture( SceneContext.GetSceneColor() );
-	const FRDGTexture* ScreenSpaceAO= GraphBuilder.RegisterExternalTexture( SceneContext.ScreenSpaceAO );
-	const FRDGTexture* HZBTexture	= GraphBuilder.RegisterExternalTexture( View.HZB );
-	const FRDGTexture* ColorTexture	= GraphBuilder.RegisterExternalTexture( TemporalAAHistory.RT[0] );
+	FRDGTexture* SceneColor	= GraphBuilder.RegisterExternalTexture( SceneContext.GetSceneColor() );
+	FRDGTexture* ScreenSpaceAO= GraphBuilder.RegisterExternalTexture( SceneContext.ScreenSpaceAO );
+	FRDGTexture* HZBTexture	= GraphBuilder.RegisterExternalTexture( View.HZB );
+	FRDGTexture* ColorTexture	= GraphBuilder.RegisterExternalTexture( TemporalAAHistory.RT[0] );
 
-	const FRDGTexture* VelocityTexture;
+	FRDGTexture* VelocityTexture;
 	if( VelocityRT && !View.bCameraCut )
 	{
 		VelocityTexture = GraphBuilder.RegisterExternalTexture( VelocityRT );
diff --git a/Engine/Source/Runtime/Renderer/Private/VolumetricFog.cpp b/Engine/Source/Runtime/Renderer/Private/VolumetricFog.cpp
--- a/Engine/Source/Runtime/Renderer/Private/VolumetricFog.cpp
+++ b/Engine/Source/Runtime/Renderer/Private/VolumetricFog.cpp
@@ -473,7 +473,7 @@ void FDeferredShadingSceneRenderer::RenderLocalLightsForVolumetricFog(
 	FIntVector VolumetricFogGridSize,
 	FVector GridZParams,
 	const FPooledRenderTargetDesc& VolumeDesc,
-	const FRDGTexture*& OutLocalShadowedLightScattering)
+	FRDGTexture*& OutLocalShadowedLightScattering)
 {
 	TArray<const FLightSceneInfo*, SceneRenderingAllocator> LightsToInject;
 
@@ -1016,7 +1016,7 @@ void FDeferredShadingSceneRenderer::ComputeVolumetricFog(FRHICommandListImmediat
 			FRDGBuilder GraphBuilder(RHICmdListImmediate);
 
 			//@DW - register GWhiteTexture as a graph external for when there's no light function - later a shader is going to bind it whether we rendered to it or not
-			const FRDGTexture* LightFunctionTexture = GraphBuilder.RegisterExternalTexture(GSystemTextures.WhiteDummy);
+			FRDGTexture* LightFunctionTexture = GraphBuilder.RegisterExternalTexture(GSystemTextures.WhiteDummy);
 			bool bUseDirectionalLightShadowing;
 
 			RenderLightFunctionForVolumetricFog(
@@ -1084,7 +1084,7 @@ void FDeferredShadingSceneRenderer::ComputeVolumetricFog(FRHICommandListImmediat
 					FogInfo.VolumetricFogDistance);
 			}
 
-			const FRDGTexture* LocalShadowedLightScattering = GraphBuilder.RegisterExternalTexture(GSystemTextures.VolumetricBlackDummy);
+			FRDGTexture* LocalShadowedLightScattering = GraphBuilder.RegisterExternalTexture(GSystemTextures.VolumetricBlackDummy);
 			RenderLocalLightsForVolumetricFog(GraphBuilder, View, bUseTemporalReprojection, IntegrationData, FogInfo, VolumetricFogGridSize, GridZParams, VolumeDescFastVRAM, LocalShadowedLightScattering);
 
 			IntegrationData.LightScattering = GraphBuilder.CreateTexture(VolumeDesc, TEXT("LightScattering"));
@@ -1147,8 +1147,8 @@ void FDeferredShadingSceneRenderer::ComputeVolumetricFog(FRHICommandListImmediat
 				});
 			}
 
-			const FRDGTexture* IntegratedLightScattering = GraphBuilder.CreateTexture(VolumeDesc, TEXT("IntegratedLightScattering"));
-			const FRDGTextureUAV* IntegratedLightScatteringUAV = GraphBuilder.CreateUAV(FRDGTextureUAVDesc(IntegratedLightScattering));
+			FRDGTexture* IntegratedLightScattering = GraphBuilder.CreateTexture(VolumeDesc, TEXT("IntegratedLightScattering"));
+			FRDGTextureUAV* IntegratedLightScatteringUAV = GraphBuilder.CreateUAV(FRDGTextureUAVDesc(IntegratedLightScattering));
 
 			{
 				FVolumetricFogFinalIntegrationCS::FParameters* PassParameters = GraphBuilder.AllocParameters<FVolumetricFogFinalIntegrationCS::FParameters>();
diff --git a/Engine/Source/Runtime/Renderer/Private/VolumetricFogLightFunction.cpp b/Engine/Source/Runtime/Renderer/Private/VolumetricFogLightFunction.cpp
--- a/Engine/Source/Runtime/Renderer/Private/VolumetricFogLightFunction.cpp
+++ b/Engine/Source/Runtime/Renderer/Private/VolumetricFogLightFunction.cpp
@@ -103,7 +103,7 @@ void FDeferredShadingSceneRenderer::RenderLightFunctionForVolumetricFog(
 	FIntVector VolumetricFogGridSize,
 	float VolumetricFogMaxDistance,
 	FMatrix& OutLightFunctionWorldToShadow,
-	const FRDGTexture*& OutLightFunctionTexture,
+	FRDGTexture*& OutLightFunctionTexture,
 	bool& bOutUseDirectionalLightShadowing)
 {
 	OutLightFunctionWorldToShadow = FMatrix::Identity;
diff --git a/Engine/Source/Runtime/Renderer/Private/VolumetricFogShared.h b/Engine/Source/Runtime/Renderer/Private/VolumetricFogShared.h
--- a/Engine/Source/Runtime/Renderer/Private/VolumetricFogShared.h
+++ b/Engine/Source/Runtime/Renderer/Private/VolumetricFogShared.h
@@ -27,13 +27,13 @@ struct FVolumetricFogIntegrationParameterData
 
 	bool bTemporalHistoryIsValid;
 	TArray<FVector4, TInlineAllocator<16>> FrameJitterOffsetValues;
-	const FRDGTexture* VBufferA;
-	const FRDGTexture* VBufferB;
-	const FRDGTextureUAV* VBufferA_UAV;
-	const FRDGTextureUAV* VBufferB_UAV;
+	FRDGTexture* VBufferA;
+	FRDGTexture* VBufferB;
+	FRDGTextureUAV* VBufferA_UAV;
+	FRDGTextureUAV* VBufferB_UAV;
 
-	const FRDGTexture* LightScattering;
-	const FRDGTextureUAV* LightScatteringUAV;
+	FRDGTexture* LightScattering;
+	FRDGTextureUAV* LightScatteringUAV;
 };
 
 /**  */
