Fixes a TODO(RDG) to catch user mistake when using a binding a render target for the first time without store action.

diff --git a/Engine/Source/Runtime/RenderCore/Private/RenderGraphBuilder.cpp b/Engine/Source/Runtime/RenderCore/Private/RenderGraphBuilder.cpp
--- a/Engine/Source/Runtime/RenderCore/Private/RenderGraphBuilder.cpp
+++ b/Engine/Source/Runtime/RenderCore/Private/RenderGraphBuilder.cpp
@@ -371,13 +371,11 @@ void FRDGBuilder::ValidatePass(const FRenderGraphPass* Pass) const
 				TEXT("Can't load depth from a render target that has never been produced."));
 			checkf(RenderTargets->DepthStencil.StencilLoadAction != ERenderTargetLoadAction::ELoad,
 				TEXT("Can't load stencil from a render target that has never been produced."));
+			checkf(RenderTargets->DepthStencil.StencilStoreAction != ERenderTargetStoreAction::ENoAction,
+				TEXT("Using for the first time a render target without storing it, witch doesn't make much sens."));
 
-			// TODO(RDG): should only be done when there is a store action.
-			if (!Texture->bHasEverBeenProduced)
-			{
-				Texture->bHasEverBeenProduced = true;
-				Texture->DebugFirstProducer = Pass;
-			}
+			Texture->bHasEverBeenProduced = true;
+			Texture->DebugFirstProducer = Pass;
 		}
 	}
 	else
