Adds a TODO(RDG) on an important CPU perf investigation to be done

diff --git a/Engine/Source/Runtime/RenderCore/Private/ShaderParameterMetadata.cpp b/Engine/Source/Runtime/RenderCore/Private/ShaderParameterMetadata.cpp
--- a/Engine/Source/Runtime/RenderCore/Private/ShaderParameterMetadata.cpp
+++ b/Engine/Source/Runtime/RenderCore/Private/ShaderParameterMetadata.cpp
@@ -213,7 +213,7 @@ void FShaderParametersMetadata::InitializeLayout()
 		}
 	} // for (int32 i = 0; i < MemberStack.Num(); ++i)
 
-#if 0
+#if 0 // TODO(RDG)
 	/** Sort the resource on MemberType first to avoid CPU miss predictions when iterating over the resources. Then based on ascending offset
 	 * to still allow O(N) complexity on offset cross referencing such as done in ClearUnusedGraphResourcesImpl().
 	 */
