Fixes a bug on FDepthStencilBinding::DepthStencilAccess getting ignored by render graph.

diff --git a/Engine/Source/Runtime/RenderCore/Private/RenderGraphBuilder.cpp b/Engine/Source/Runtime/RenderCore/Private/RenderGraphBuilder.cpp
--- a/Engine/Source/Runtime/RenderCore/Private/RenderGraphBuilder.cpp
+++ b/Engine/Source/Runtime/RenderCore/Private/RenderGraphBuilder.cpp
@@ -1063,7 +1063,7 @@ void FRDGBuilder::AllocateAndTransitionPassResources(const FRenderGraphPass* Pas
 				OutRPInfo->DepthStencilRenderTarget.Action = MakeDepthStencilTargetActions(
 					MakeRenderTargetActions(DepthStencil.DepthLoadAction, DepthStencil.DepthStoreAction),
 					MakeRenderTargetActions(DepthStencil.StencilLoadAction, DepthStencil.StencilStoreAction));
-				OutRPInfo->DepthStencilRenderTarget.ExclusiveDepthStencil = FExclusiveDepthStencil::DepthWrite_StencilWrite;
+				OutRPInfo->DepthStencilRenderTarget.ExclusiveDepthStencil = DepthStencil.DepthStencilAccess;
 
 				TransitionTexture(DepthStencil.Texture, EResourceTransitionAccess::EWritable, false);
 
