Cherry-pick 5211844: Improves error message when their is unbound shader parameters

diff --git a/Engine/Source/Runtime/RenderCore/Private/ShaderParameterStruct.cpp b/Engine/Source/Runtime/RenderCore/Private/ShaderParameterStruct.cpp
--- a/Engine/Source/Runtime/RenderCore/Private/ShaderParameterStruct.cpp
+++ b/Engine/Source/Runtime/RenderCore/Private/ShaderParameterStruct.cpp
@@ -220,20 +220,18 @@ void FShaderParameterBindings::BindForLegacyShaderParameters(const FShader* Shad
 	ParametersMap.GetAllParameterNames(AllParameterNames);
 	if (bShouldBindEverything && BindingContext.ShaderGlobalScopeBindings.Num() != AllParameterNames.Num())
 	{
-		UE_LOG(LogShaders, Error, TEXT("%i shader parameters have not been bound for %s:"),
-			AllParameterNames.Num() - BindingContext.ShaderGlobalScopeBindings.Num(),
-			Shader->GetType()->GetName());
+		FString ErrorString = FString::Printf(
+			TEXT("Shader %s has unbound parameters not represented in the parameter struct:"), Shader->GetType()->GetName());
 
 		for (const FString& GlobalParameterName : AllParameterNames)
 		{
 			if (!BindingContext.ShaderGlobalScopeBindings.Contains(GlobalParameterName))
 			{
-				UE_LOG(LogShaders, Error, TEXT("  %s"), *GlobalParameterName);
+				ErrorString += FString::Printf(TEXT("\n  %s"), *GlobalParameterName);
 			}
 		}
 
-		UE_LOG(LogShaders, Fatal, TEXT("Unable to bind all shader parameters of %s."),
-			Shader->GetType()->GetName());
+		UE_LOG(LogShaders, Fatal, TEXT("%s"), *ErrorString);
 	}
 }
 
@@ -274,20 +272,18 @@ void FShaderParameterBindings::BindForRootShaderParameters(const FShader* Shader
 	ParametersMap.GetAllParameterNames(AllParameterNames);
 	if (BindingContext.ShaderGlobalScopeBindings.Num() != AllParameterNames.Num())
 	{
-		UE_LOG(LogShaders, Error, TEXT("%i shader parameters have not been bound for %s:"),
-			AllParameterNames.Num() - BindingContext.ShaderGlobalScopeBindings.Num(),
-			Shader->GetType()->GetName());
+		FString ErrorString = FString::Printf(
+			TEXT("Shader %s has unbound parameters not represented in the parameter struct:"), Shader->GetType()->GetName());
 
 		for (const FString& GlobalParameterName : AllParameterNames)
 		{
 			if (!BindingContext.ShaderGlobalScopeBindings.Contains(GlobalParameterName))
 			{
-				UE_LOG(LogShaders, Error, TEXT("  %s"), *GlobalParameterName);
+				ErrorString += FString::Printf(TEXT("\n  %s"), *GlobalParameterName);
 			}
 		}
 
-		UE_LOG(LogShaders, Fatal, TEXT("Unable to bind all shader parameters of %s."),
-			Shader->GetType()->GetName());
+		UE_LOG(LogShaders, Fatal, TEXT("%s"), *ErrorString);
 	}
 }
 
