Cherry-pick 5211847: Remove unnecessary restriction on where shader parameter resources are placed in a table

diff --git a/Engine/Source/Runtime/RenderCore/Public/ShaderParameterMacros.h b/Engine/Source/Runtime/RenderCore/Public/ShaderParameterMacros.h
--- a/Engine/Source/Runtime/RenderCore/Public/ShaderParameterMacros.h
+++ b/Engine/Source/Runtime/RenderCore/Public/ShaderParameterMacros.h
@@ -616,7 +616,6 @@ struct TShaderParameterTypeInfo<TUniformBufferRef<UniformBufferStructType>>
 		struct zzNextMemberId##MemberName { enum { HasDeclaredResource = zzMemberId##MemberName::HasDeclaredResource || !TypeInfo::bIsStoredInConstantBuffer }; }; \
 		static TArray<FShaderParametersMetadata::FMember> zzGetMembersBefore(zzNextMemberId##MemberName) \
 		{ \
-			static_assert(!TypeInfo::bIsStoredInConstantBuffer || zzMemberId##MemberName::HasDeclaredResource == 0, "Shader resources must be declared after " #MemberName "."); \
 			static_assert(TypeInfo::bIsStoredInConstantBuffer || TIsArrayOrRefOfType<decltype(OptionalShaderType), TCHAR>::Value, "No shader type for " #MemberName "."); \
 			/* Route the member enumeration on to the function for the member following this. */ \
 			TArray<FShaderParametersMetadata::FMember> OutMembers = zzGetMembersBefore(zzMemberId##MemberName()); \
