DR - Extend support for RDG buffers
#rnx

diff --git a/Engine/Source/Runtime/RenderCore/Public/RenderGraphBuilder.h b/Engine/Source/Runtime/RenderCore/Public/RenderGraphBuilder.h
--- a/Engine/Source/Runtime/RenderCore/Public/RenderGraphBuilder.h
+++ b/Engine/Source/Runtime/RenderCore/Public/RenderGraphBuilder.h
@@ -319,6 +319,11 @@ public:
 		return SRV;
 	}
 
+	FORCEINLINE_DEBUGGABLE FRDGBufferSRVRef CreateSRV(FRDGBufferRef Buffer, EPixelFormat Format)
+	{
+		return CreateSRV(FRDGBufferSRVDesc(Buffer, Format));
+	}
+
 	/** Create graph tracked UAV for a texture from a descriptor. */
 	FORCEINLINE_DEBUGGABLE FRDGTextureUAVRef CreateUAV(const FRDGTextureUAVDesc& Desc)
 	{
diff --git a/Engine/Source/Runtime/RenderCore/Public/RenderGraphResources.h b/Engine/Source/Runtime/RenderCore/Public/RenderGraphResources.h
--- a/Engine/Source/Runtime/RenderCore/Public/RenderGraphResources.h
+++ b/Engine/Source/Runtime/RenderCore/Public/RenderGraphResources.h
@@ -300,6 +300,16 @@ struct RENDERCORE_API FRDGBufferDesc
 		Desc.NumElements = NumElements;
 		return Desc;
 	}
+
+	static inline FRDGBufferDesc CreateBufferDesc(uint32 BytesPerElement, uint32 NumElements)
+	{
+		FRDGBufferDesc Desc;
+		Desc.UnderlyingType = EUnderlyingType::VertexBuffer;
+		Desc.Usage = EBufferUsageFlags(BUF_Static | BUF_UnorderedAccess | BUF_ShaderResource);
+		Desc.BytesPerElement = BytesPerElement;
+		Desc.NumElements = NumElements;
+		return Desc;
+	}
 };
 
 
