Implements FRDGResourceFlags::MultiFrame for TAA histories with Nvidia SLI

diff --git a/Engine/Source/Runtime/RenderCore/Private/RenderGraphBuilder.cpp b/Engine/Source/Runtime/RenderCore/Private/RenderGraphBuilder.cpp
--- a/Engine/Source/Runtime/RenderCore/Private/RenderGraphBuilder.cpp
+++ b/Engine/Source/Runtime/RenderCore/Private/RenderGraphBuilder.cpp
@@ -709,20 +709,46 @@ void FRDGBuilder::TransitionTexture(FRDGTexture* Texture, EResourceTransitionAcc
 
 	if( Texture->bWritable != bRequiredWritable || Texture->bCompute != bRequiredCompute )
 	{
-		RHICmdList.TransitionResource( TransitionAccess, Texture->PooledRenderTarget->GetRenderTargetItem().ShaderResourceTexture );
+		FTextureRHIParamRef RHITexture = Texture->PooledRenderTarget->GetRenderTargetItem().ShaderResourceTexture;
+
+		bool bIsMultiFrameResource = uint32(Texture->Flags) & uint32(ERDGResourceFlags::MultiFrame);
+		if (bIsMultiFrameResource && bRequiredWritable && !Texture->bWritable)
+		{
+			RHICmdList.BeginUpdateMultiFrameResource(RHITexture);
+		}
+
+		RHICmdList.TransitionResource( TransitionAccess, RHITexture );
+		
+		if (bIsMultiFrameResource && !bRequiredWritable && Texture->bWritable)
+		{
+			RHICmdList.EndUpdateMultiFrameResource(RHITexture);
+		}
+		
 		Texture->bWritable = bRequiredWritable;
 		Texture->bCompute = bRequiredCompute;
 	}
 }
 
-void FRDGBuilder::TransitionUAV(FUnorderedAccessViewRHIParamRef UAV, FRDGResource* UnderlyingResource, EResourceTransitionAccess TransitionAccess, bool bRequiredCompute ) const
+void FRDGBuilder::TransitionUAV(FUnorderedAccessViewRHIParamRef UAV, FRDGResource* UnderlyingResource, ERDGResourceFlags ResourceFlags, EResourceTransitionAccess TransitionAccess, bool bRequiredCompute ) const
 {
 	const bool bRequiredWritable = true;
 
 	if(UnderlyingResource->bWritable != bRequiredWritable || UnderlyingResource->bCompute != bRequiredCompute )
 	{
+		bool bIsMultiFrameResource = uint32(ResourceFlags) & uint32(ERDGResourceFlags::MultiFrame);
+		if (bIsMultiFrameResource && bRequiredWritable && !UnderlyingResource->bWritable)
+		{
+			RHICmdList.BeginUpdateMultiFrameResource(UAV);
+		}
+
 		EResourceTransitionPipeline TransitionPipeline = CalcTransitionPipeline(UnderlyingResource->bCompute, bRequiredCompute );
 		RHICmdList.TransitionResource( TransitionAccess, TransitionPipeline, UAV);
+		
+		if (bIsMultiFrameResource && !bRequiredWritable && UnderlyingResource->bWritable)
+		{
+			RHICmdList.EndUpdateMultiFrameResource(UAV);
+		}
+		
 		UnderlyingResource->bWritable = bRequiredWritable;
 		UnderlyingResource->bCompute = bRequiredCompute;
 	}
@@ -927,7 +953,7 @@ void FRDGBuilder::AllocateAndTransitionPassResources(const FRenderGraphPass* Pas
 			if (UAV)
 			{
 				AllocateRHITextureUAVIfNeeded(UAV, bIsCompute);
-				TransitionUAV(UAV->CachedRHI.UAV, UAV->Desc.Texture, EResourceTransitionAccess::EWritable, bIsCompute);
+				TransitionUAV(UAV->CachedRHI.UAV, UAV->Desc.Texture, UAV->Desc.Texture->Flags, EResourceTransitionAccess::EWritable, bIsCompute);
 
 				#if RENDER_GRAPH_DEBUGGING
 				{
@@ -953,7 +979,7 @@ void FRDGBuilder::AllocateAndTransitionPassResources(const FRenderGraphPass* Pas
 				// TODO(RDG): supper hacky, find the UAV and transition it. Hopefully there is one...
 				check(Buffer->PooledBuffer->UAVs.Num() == 1);
 				FUnorderedAccessViewRHIParamRef UAV = Buffer->PooledBuffer->UAVs.CreateIterator().Value();
-				TransitionUAV(UAV, Buffer, EResourceTransitionAccess::EReadable, bIsCompute);
+				TransitionUAV(UAV, Buffer, Buffer->Flags, EResourceTransitionAccess::EReadable, bIsCompute);
 
 				#if RENDER_GRAPH_DEBUGGING
 				{
@@ -982,7 +1008,7 @@ void FRDGBuilder::AllocateAndTransitionPassResources(const FRenderGraphPass* Pas
 				// TODO(RDG): supper hacky, find the UAV and transition it. Hopefully there is one...
 				check(SRV->Desc.Buffer->PooledBuffer->UAVs.Num() == 1);
 				FUnorderedAccessViewRHIParamRef UAV = SRV->Desc.Buffer->PooledBuffer->UAVs.CreateIterator().Value();
-				TransitionUAV(UAV, SRV->Desc.Buffer, EResourceTransitionAccess::EReadable, bIsCompute);
+				TransitionUAV(UAV, SRV->Desc.Buffer, SRV->Desc.Buffer->Flags, EResourceTransitionAccess::EReadable, bIsCompute);
 
 				#if RENDER_GRAPH_DEBUGGING
 				{
@@ -998,7 +1024,7 @@ void FRDGBuilder::AllocateAndTransitionPassResources(const FRenderGraphPass* Pas
 			if (UAV)
 			{
 				AllocateRHIBufferUAVIfNeeded(UAV, bIsCompute);
-				TransitionUAV(UAV->CachedRHI.UAV, UAV->Desc.Buffer, EResourceTransitionAccess::EWritable, bIsCompute);
+				TransitionUAV(UAV->CachedRHI.UAV, UAV->Desc.Buffer, UAV->Desc.Buffer->Flags, EResourceTransitionAccess::EWritable, bIsCompute);
 
 				#if RENDER_GRAPH_DEBUGGING
 				{
@@ -1296,7 +1322,7 @@ void FRDGBuilder::ProcessDeferredInternalResourceQueries()
 
 		if (Query.bTransitionToRead)
 		{
-			RHICmdList.TransitionResource(EResourceTransitionAccess::EReadable, Query.Texture->PooledRenderTarget->GetRenderTargetItem().ShaderResourceTexture);
+			TransitionTexture(Query.Texture, EResourceTransitionAccess::EReadable, /* bRequiredCompute = */ false);
 		}
 
 		*Query.OutTexturePtr = AllocatedTextures.FindChecked(Query.Texture);
diff --git a/Engine/Source/Runtime/RenderCore/Public/RenderGraphBuilder.h b/Engine/Source/Runtime/RenderCore/Public/RenderGraphBuilder.h
--- a/Engine/Source/Runtime/RenderCore/Public/RenderGraphBuilder.h
+++ b/Engine/Source/Runtime/RenderCore/Public/RenderGraphBuilder.h
@@ -211,7 +211,7 @@ public:
 			checkf(DebugName, TEXT("Externally allocated texture requires a debug name when registering them to render graph."));
 		}
 		#endif
-		FRDGTexture* OutTexture = AllocateForRHILifeTime<FRDGTexture>(DebugName, ExternalPooledTexture->GetDesc());
+		FRDGTexture* OutTexture = AllocateForRHILifeTime<FRDGTexture>(DebugName, ExternalPooledTexture->GetDesc(), ERDGResourceFlags::None);
 		OutTexture->PooledRenderTarget = ExternalPooledTexture;
 		OutTexture->CachedRHI.Texture = ExternalPooledTexture->GetRenderTargetItem().ShaderResourceTexture;
 		AllocatedTextures.Add(OutTexture, ExternalPooledTexture);
@@ -232,7 +232,7 @@ public:
 			ensureMsgf(ExternalPooledBuffer.IsValid(), TEXT("Attempted to register NULL external buffer: %s"), Name);
 		}
 #endif
-		FRDGBuffer* OutBuffer = AllocateForRHILifeTime<FRDGBuffer>(Name, ExternalPooledBuffer->Desc);
+		FRDGBuffer* OutBuffer = AllocateForRHILifeTime<FRDGBuffer>(Name, ExternalPooledBuffer->Desc, ERDGResourceFlags::None);
 		OutBuffer->PooledBuffer = ExternalPooledBuffer;
 		AllocatedBuffers.Add(OutBuffer, ExternalPooledBuffer);
 #if RENDER_GRAPH_DEBUGGING
@@ -249,7 +249,7 @@ public:
 	 *
 	 * The debug name is the name used for GPU debugging tools, but also for the VisualizeTexture/Vis command.
 	 */
-	FORCEINLINE_DEBUGGABLE FRDGTextureRef CreateTexture(const FPooledRenderTargetDesc& Desc, const TCHAR* DebugName)
+	FORCEINLINE_DEBUGGABLE FRDGTextureRef CreateTexture(const FPooledRenderTargetDesc& Desc, const TCHAR* DebugName, ERDGResourceFlags Flags = ERDGResourceFlags::None)
 	{
 		#if RENDER_GRAPH_DEBUGGING
 		{
@@ -258,7 +258,7 @@ public:
 			checkf(Desc.Format != PF_Unknown, TEXT("Illegal to create texture %s with an invalid pixel format."), DebugName);
 		}
 		#endif
-		FRDGTexture* Texture = AllocateForRHILifeTime<FRDGTexture>(DebugName, Desc);
+		FRDGTexture* Texture = AllocateForRHILifeTime<FRDGTexture>(DebugName, Desc, Flags);
 		#if RENDER_GRAPH_DEBUGGING
 			Resources.Add(Texture);
 		#endif
@@ -269,7 +269,7 @@ public:
 	 *
 	 * The debug name is the name used for GPU debugging tools, but also for the VisualizeTexture/Vis command.
 	 */
-	FORCEINLINE_DEBUGGABLE FRDGBufferRef CreateBuffer(const FRDGBufferDesc& Desc, const TCHAR* DebugName)
+	FORCEINLINE_DEBUGGABLE FRDGBufferRef CreateBuffer(const FRDGBufferDesc& Desc, const TCHAR* DebugName, ERDGResourceFlags Flags = ERDGResourceFlags::None)
 	{
 		#if RENDER_GRAPH_DEBUGGING
 		{
@@ -277,7 +277,7 @@ public:
 			checkf(DebugName, TEXT("Creating a render graph buffer requires a valid debug name."));
 		}
 		#endif
-		FRDGBuffer* Buffer = AllocateForRHILifeTime<FRDGBuffer>(DebugName, Desc);
+		FRDGBuffer* Buffer = AllocateForRHILifeTime<FRDGBuffer>(DebugName, Desc, Flags);
 		#if RENDER_GRAPH_DEBUGGING
 			Resources.Add(Buffer);
 		#endif
@@ -620,7 +620,7 @@ private:
 
 
 	void TransitionTexture(FRDGTexture* Texture, EResourceTransitionAccess TransitionAccess, bool bRequiredCompute) const;
-	void TransitionUAV(FUnorderedAccessViewRHIParamRef UAV, FRDGResource* UnderlyingResource, EResourceTransitionAccess TransitionAccess, bool bRequiredCompute ) const;
+	void TransitionUAV(FUnorderedAccessViewRHIParamRef UAV, FRDGResource* UnderlyingResource, ERDGResourceFlags ResourceFlags, EResourceTransitionAccess TransitionAccess, bool bRequiredCompute ) const;
 
 	void PushDrawEventStack(const FRenderGraphPass* Pass);
 	void ExecutePass( const FRenderGraphPass* Pass );
diff --git a/Engine/Source/Runtime/RenderCore/Public/RenderGraphResources.h b/Engine/Source/Runtime/RenderCore/Public/RenderGraphResources.h
--- a/Engine/Source/Runtime/RenderCore/Public/RenderGraphResources.h
+++ b/Engine/Source/Runtime/RenderCore/Public/RenderGraphResources.h
@@ -39,6 +39,16 @@ using FRDGBufferSRVRef = FRDGBufferSRV*;
 using FRDGBufferUAVRef = FRDGBufferUAV*;
 
 
+/** Render graph specific flags for resources. */
+enum class ERDGResourceFlags : uint8
+{
+	None = 0x00,
+
+	// Tag the resource to survive through frame, that is important for multi GPU alternate frame rendering.
+	MultiFrame = 0x01,
+};
+
+
 /** Generic graph resource. */
 class RENDERCORE_API FRDGResource
 {
@@ -139,9 +149,14 @@ private:
 	 */
 	mutable IPooledRenderTarget* PooledRenderTarget = nullptr;
 
-	FRDGTexture(const TCHAR* DebugName, const FPooledRenderTargetDesc& InDesc)
+	// Flags specific to this resource for render graph.
+	const ERDGResourceFlags Flags;
+
+
+	FRDGTexture(const TCHAR* DebugName, const FPooledRenderTargetDesc& InDesc, ERDGResourceFlags InFlags)
 		: FRDGResource(DebugName)
 		, Desc(InDesc)
+		, Flags(InFlags)
 	{}
 
 	friend class FRDGBuilder;
@@ -462,10 +477,15 @@ private:
 	 * FGraphBuilder::AllocatedBuffers is actually keeping the reference.
 	 */
 	mutable FPooledRDGBuffer* PooledBuffer = nullptr;
+	
+	// Flags specific to this resource for render graph.
+	const ERDGResourceFlags Flags;
+
 
-	FRDGBuffer(const TCHAR* DebugName, const FRDGBufferDesc& InDesc)
+	FRDGBuffer(const TCHAR* DebugName, const FRDGBufferDesc& InDesc, ERDGResourceFlags InFlags)
 		: FRDGResource(DebugName)
 		, Desc(InDesc)
+		, Flags(InFlags)
 	{ }
 
 	friend class FRDGBuilder;
diff --git a/Engine/Source/Runtime/Renderer/Private/PostProcess/PostProcessTemporalAA.cpp b/Engine/Source/Runtime/Renderer/Private/PostProcess/PostProcessTemporalAA.cpp
--- a/Engine/Source/Runtime/Renderer/Private/PostProcess/PostProcessTemporalAA.cpp
+++ b/Engine/Source/Runtime/Renderer/Private/PostProcess/PostProcessTemporalAA.cpp
@@ -585,12 +585,18 @@ FTAAOutputs FTAAPassParameters::AddTemporalAAPass(
 				SceneColorDesc.TargetableFlags |= TexCreate_UAV;
 			}
 
-			Outputs.SceneColor = GraphBuilder.CreateTexture(SceneColorDesc, kTAAOutputNames[static_cast<int32>(Pass)]);
+			Outputs.SceneColor = GraphBuilder.CreateTexture(
+				SceneColorDesc,
+				kTAAOutputNames[static_cast<int32>(Pass)],
+				ERDGResourceFlags::MultiFrame);
 		}
 
 		if (RenderTargetCount == 2)
 		{
-			Outputs.SceneMetadata = GraphBuilder.CreateTexture(SceneColorDesc, kTAAOutputNames[static_cast<int32>(Pass)]);
+			Outputs.SceneMetadata = GraphBuilder.CreateTexture(
+				SceneColorDesc,
+				kTAAOutputNames[static_cast<int32>(Pass)],
+				ERDGResourceFlags::MultiFrame);
 		}
 
 		if (bDownsample)
@@ -679,12 +685,6 @@ FTAAOutputs FTAAPassParameters::AddTemporalAAPass(
 			if (Outputs.DownsampledSceneColor)
 				PassParameters->OutComputeTexDownsampled = GraphBuilder.CreateUAV(Outputs.DownsampledSceneColor);
 		}
-		
-#if 0 // TODO.
-		Context.RHICmdList.BeginUpdateMultiFrameResource(DestRenderTarget[0]->ShaderResourceTexture);
-		if (RenderTargetCount == 2)
-			Context.RHICmdList.BeginUpdateMultiFrameResource(DestRenderTarget[1]->ShaderResourceTexture);
-#endif
 
 		TShaderMapRef<FTemporalAACS> ComputeShader(View.ShaderMap, PermutationVector);
 		ClearUnusedGraphResources(*ComputeShader, PassParameters);
