Implements -rdgdebug to enable RDG warnings from command line, and format RDG event names.

diff --git a/Engine/Config/ConsoleVariables.ini b/Engine/Config/ConsoleVariables.ini
--- a/Engine/Config/ConsoleVariables.ini
+++ b/Engine/Config/ConsoleVariables.ini
@@ -51,7 +51,7 @@
 
 ; Uncomment to get render graph executing passes as they get created to easily debug crashes caused by pass wiring logic.
 ;r.RDG.ImmediateMode=1
-; Uncomment to have render graph emitting warnings about inefficiencies.
+; Uncomment to get render graph to emit warnings for inneficiencies, and have event names formatting even if frame are not emiting draw events.
 ;r.RDG.EmitWarnings=1
 
 ; r.XGEShaderCompile is now enabled by default in source. Uncomment to disable XGE shader compilation.
diff --git a/Engine/Source/Runtime/RenderCore/Private/RenderGraphBuilder.cpp b/Engine/Source/Runtime/RenderCore/Private/RenderGraphBuilder.cpp
--- a/Engine/Source/Runtime/RenderCore/Private/RenderGraphBuilder.cpp
+++ b/Engine/Source/Runtime/RenderCore/Private/RenderGraphBuilder.cpp
@@ -41,6 +41,11 @@ void InitRenderGraph()
 	{
 		GRenderGraphImmediateMode = 1;
 	}
+
+	if (FParse::Param(FCommandLine::Get(), TEXT("rdgdebug")))
+	{
+		GRenderGraphEmitWarnings = 1;
+	}
 #endif
 }
 
@@ -104,7 +109,7 @@ static bool IsBoundAsReadable(const FRDGTexture* Texture, FShaderParameterStruct
 
 FRDGEventName::FRDGEventName(const TCHAR* EventFormat, ...)
 {
-	if (GetEmitDrawEvents())
+	if (GetEmitDrawEvents() || GRenderGraphEmitWarnings)
 	{
 		va_list ptr;
 		va_start(ptr, EventFormat);
@@ -113,7 +118,12 @@ FRDGEventName::FRDGEventName(const TCHAR* EventFormat, ...)
 		FCString::GetVarArgs(TempStr, ARRAY_COUNT(TempStr), EventFormat, ptr);
 		va_end(ptr);
 
-		EventName = TempStr;
+		EventNameStorage = TempStr;
+		EventName = *EventNameStorage;
+	}
+	else
+	{
+		EventName = TEXT("!!!Unavailable RDG event name: try r.RDG.EmitWarnings=1 or -rdgdebug!!!");
 	}
 }
 
diff --git a/Engine/Source/Runtime/RenderCore/Public/RenderGraphBuilder.h b/Engine/Source/Runtime/RenderCore/Public/RenderGraphBuilder.h
--- a/Engine/Source/Runtime/RenderCore/Public/RenderGraphBuilder.h
+++ b/Engine/Source/Runtime/RenderCore/Public/RenderGraphBuilder.h
@@ -40,23 +40,23 @@ public:
 	{ }
 
 	#endif
-
+	
 	const TCHAR* GetTCHAR() const
 	{
-		#if RENDER_GRAPH_DRAW_EVENTS == 2
-			return *EventName;
-		#elif RENDER_GRAPH_DRAW_EVENTS == 1
+		#if RENDER_GRAPH_DRAW_EVENTS == 1 || RENDER_GRAPH_DRAW_EVENTS == 2
 			return EventName;
 		#else
-			return TEXT("UnknownRDVEvent");
+			// Render graph draw events have been completely compiled for CPU performance reasons.
+			return TEXT("!!!Unavailable RDG event name: need RENDER_GRAPH_DRAW_EVENTS>=0 and r.RDG.EmitWarnings=1 or -rdgdebug!!!");
 		#endif
 	}
 
 private:
-	#if RENDER_GRAPH_DRAW_EVENTS == 2
-		FString EventName;
-	#elif RENDER_GRAPH_DRAW_EVENTS == 1
+	#if RENDER_GRAPH_DRAW_EVENTS == 1 || RENDER_GRAPH_DRAW_EVENTS == 2
 		const TCHAR* EventName;
+		#if RENDER_GRAPH_DRAW_EVENTS == 2
+			FString EventNameStorage;
+		#endif
 	#endif
 };
 
