Implements SHADER_PARAMETER_STRUCT_ARRAY()

diff --git a/Engine/Source/Runtime/Engine/Private/ShaderCompiler/ShaderCompiler.cpp b/Engine/Source/Runtime/Engine/Private/ShaderCompiler/ShaderCompiler.cpp
--- a/Engine/Source/Runtime/Engine/Private/ShaderCompiler/ShaderCompiler.cpp
+++ b/Engine/Source/Runtime/Engine/Private/ShaderCompiler/ShaderCompiler.cpp
@@ -2769,16 +2769,26 @@ static void PullRootShaderParametersLayout(FShaderCompilerInput& CompileInput, c
 	{
 		EUniformBufferBaseType BaseType = Member.GetBaseType();
 		uint16 MemberOffset = ByteOffset + uint16(Member.GetOffset());
+		uint32 NumElements = Member.GetNumElements();
 
 		if (BaseType == UBMT_INCLUDED_STRUCT)
 		{
+			check(NumElements == 0);
 			PullRootShaderParametersLayout(CompileInput, *Member.GetStructMetadata(), MemberOffset, Prefix);
 		}
-		else if (BaseType == UBMT_NESTED_STRUCT)
+		else if (BaseType == UBMT_NESTED_STRUCT && NumElements == 0)
 		{
 			FString NewPrefix = FString::Printf(TEXT("%s%s_"), *Prefix, Member.GetName());
 			PullRootShaderParametersLayout(CompileInput, *Member.GetStructMetadata(), MemberOffset, NewPrefix);
 		}
+		else if (BaseType == UBMT_NESTED_STRUCT && NumElements > 0)
+		{
+			for (uint32 ArrayElementId = 0; ArrayElementId < NumElements; ArrayElementId++)
+			{
+				FString NewPrefix = FString::Printf(TEXT("%s%s_%u_"), *Prefix, Member.GetName(), ArrayElementId);
+				PullRootShaderParametersLayout(CompileInput, *Member.GetStructMetadata(), MemberOffset, NewPrefix);
+			}
+		}
 		else if (
 			BaseType == UBMT_BOOL ||
 			BaseType == UBMT_INT32 ||
diff --git a/Engine/Source/Runtime/RenderCore/Private/ShaderParameterMetadata.cpp b/Engine/Source/Runtime/RenderCore/Private/ShaderParameterMetadata.cpp
--- a/Engine/Source/Runtime/RenderCore/Private/ShaderParameterMetadata.cpp
+++ b/Engine/Source/Runtime/RenderCore/Private/ShaderParameterMetadata.cpp
@@ -184,7 +184,7 @@ void FShaderParametersMetadata::InitializeLayout()
 				}
 			}
 
-			const bool bTypeCanBeArray = (bAllowResourceArrays && (bIsRHIResource || bIsRDGResource)) || bIsVariableNativeType;
+			const bool bTypeCanBeArray = (bAllowResourceArrays && (bIsRHIResource || bIsRDGResource)) || bIsVariableNativeType || BaseType == UBMT_NESTED_STRUCT;
 			if (bIsArray && !bTypeCanBeArray)
 			{
 				UE_LOG(LogRendererCore, Fatal, TEXT("Shader parameter %s error: Not allowed to be an array."), *CppName);
@@ -203,12 +203,15 @@ void FShaderParametersMetadata::InitializeLayout()
 
 		if (ChildStruct && BaseType != UBMT_REFERENCED_STRUCT)
 		{
-			int32 AbsoluteStructOffset = CurrentMember.GetOffset() + MemberStack[i].StructOffset;
-
-			for (int32 StructMemberIndex = 0; StructMemberIndex < ChildStruct->Members.Num(); StructMemberIndex++)
+			for (uint32 ArrayElementId = 0; ArrayElementId < (bIsArray ? ArraySize : 1u); ArrayElementId++)
 			{
-				const FMember& StructMember = ChildStruct->Members[StructMemberIndex];
-				MemberStack.Insert(FUniformBufferMemberAndOffset(*ChildStruct, StructMember, AbsoluteStructOffset), i + 1 + StructMemberIndex);
+				int32 AbsoluteStructOffset = CurrentMember.GetOffset() + MemberStack[i].StructOffset + ArrayElementId * ChildStruct->GetSize();
+
+				for (int32 StructMemberIndex = 0; StructMemberIndex < ChildStruct->Members.Num(); StructMemberIndex++)
+				{
+					const FMember& StructMember = ChildStruct->Members[StructMemberIndex];
+					MemberStack.Insert(FUniformBufferMemberAndOffset(*ChildStruct, StructMember, AbsoluteStructOffset), i + 1 + StructMemberIndex);
+				}
 			}
 		}
 	} // for (int32 i = 0; i < MemberStack.Num(); ++i)
@@ -263,6 +266,8 @@ void FShaderParametersMetadata::AddResourceTableEntriesRecursive(const TCHAR* Un
 	for (int32 MemberIndex = 0; MemberIndex < Members.Num(); ++MemberIndex)
 	{
 		const FMember& Member = Members[MemberIndex];
+		uint32 NumElements = Member.GetNumElements();
+
 		if (IsShaderParameterTypeForUniformBufferLayout(Member.GetBaseType()))
 		{
 			FResourceTableEntry& Entry = ResourceTableMap.FindOrAdd(FString::Printf(TEXT("%s%s"), Prefix, Member.GetName()));
@@ -273,15 +278,25 @@ void FShaderParametersMetadata::AddResourceTableEntriesRecursive(const TCHAR* Un
 				Entry.ResourceIndex = ResourceIndex++;
 			}
 		}
-		else if (Member.GetBaseType() == UBMT_NESTED_STRUCT)
+		else if (Member.GetBaseType() == UBMT_NESTED_STRUCT && NumElements == 0)
 		{
 			check(Member.GetStructMetadata());
 			FString MemberPrefix = FString::Printf(TEXT("%s%s_"), Prefix, Member.GetName());
 			Member.GetStructMetadata()->AddResourceTableEntriesRecursive(UniformBufferName, *MemberPrefix, ResourceIndex, ResourceTableMap);
 		}
+		else if (Member.GetBaseType() == UBMT_NESTED_STRUCT && NumElements > 0)
+		{
+			for (uint32 ArrayElementId = 0; ArrayElementId < NumElements; ArrayElementId++)
+			{
+				check(Member.GetStructMetadata());
+				FString MemberPrefix = FString::Printf(TEXT("%s%s_%u_"), Prefix, Member.GetName(), ArrayElementId);
+				Member.GetStructMetadata()->AddResourceTableEntriesRecursive(UniformBufferName, *MemberPrefix, ResourceIndex, ResourceTableMap);
+			}
+		}
 		else if (Member.GetBaseType() == UBMT_INCLUDED_STRUCT)
 		{
 			check(Member.GetStructMetadata());
+			check(NumElements == 0);
 			Member.GetStructMetadata()->AddResourceTableEntriesRecursive(UniformBufferName, Prefix, ResourceIndex, ResourceTableMap);
 		}
 	}
@@ -294,10 +309,12 @@ void FShaderParametersMetadata::FindMemberFromOffset(uint16 MemberOffset, const
 	for (const FMember& Member : Members)
 	{
 		EUniformBufferBaseType BaseType = Member.GetBaseType();
-		if (BaseType == UBMT_NESTED_STRUCT || BaseType == UBMT_INCLUDED_STRUCT)
+		uint32 NumElements = Member.GetNumElements();
+
+		if ((BaseType == UBMT_NESTED_STRUCT && NumElements == 0) || BaseType == UBMT_INCLUDED_STRUCT)
 		{
 			const FShaderParametersMetadata* SubStruct = Member.GetStructMetadata();
-			if (MemberOffset < Member.GetOffset() + SubStruct->GetSize())
+			if (MemberOffset < (Member.GetOffset() + SubStruct->GetSize()))
 			{
 				if (NamePrefix)
 				{
@@ -307,14 +324,38 @@ void FShaderParametersMetadata::FindMemberFromOffset(uint16 MemberOffset, const
 				return SubStruct->FindMemberFromOffset(MemberOffset - Member.GetOffset(), OutContainingStruct, OutMember, ArrayElementId, NamePrefix);
 			}
 		}
-		else if (Member.GetNumElements() > 0 && (
+		else if (BaseType == UBMT_NESTED_STRUCT && NumElements > 0)
+		{
+			const FShaderParametersMetadata* SubStruct = Member.GetStructMetadata();
+			uint32 StructSize = SubStruct->GetSize();
+			
+			uint16 ArrayStartOffset = Member.GetOffset();
+			uint16 ArrayEndOffset = ArrayStartOffset + SubStruct->GetSize() * NumElements;
+			
+			if (MemberOffset >= ArrayStartOffset && MemberOffset < ArrayEndOffset)
+			{
+				uint32 MemberOffsetInArray = MemberOffset - ArrayStartOffset;
+				check((MemberOffsetInArray % StructSize) == 0);
+
+				uint32 MemberPosInStructArray = MemberOffsetInArray / StructSize;
+				uint32 MemberOffsetInStructElement = MemberOffsetInArray - MemberPosInStructArray * StructSize;
+
+				if (NamePrefix)
+				{
+					*NamePrefix = FString::Printf(TEXT("%s%s[%u]::"), **NamePrefix, Member.GetName(), MemberPosInStructArray);
+				}
+
+				return SubStruct->FindMemberFromOffset(MemberOffsetInStructElement, OutContainingStruct, OutMember, ArrayElementId, NamePrefix);
+			}
+		}
+		else if (NumElements > 0 && (
 			BaseType == UBMT_TEXTURE ||
 			BaseType == UBMT_SRV ||
 			BaseType == UBMT_SAMPLER ||
 			IsRDGResourceReferenceShaderParameterType(BaseType)))
 		{
 			uint16 ArrayStartOffset = Member.GetOffset();
-			uint16 ArrayEndOffset = ArrayStartOffset + SHADER_PARAMETER_POINTER_ALIGNMENT * Member.GetNumElements();
+			uint16 ArrayEndOffset = ArrayStartOffset + SHADER_PARAMETER_POINTER_ALIGNMENT * NumElements;
 
 			if (MemberOffset >= ArrayStartOffset && MemberOffset < ArrayEndOffset)
 			{
diff --git a/Engine/Source/Runtime/RenderCore/Private/ShaderParameterStruct.cpp b/Engine/Source/Runtime/RenderCore/Private/ShaderParameterStruct.cpp
--- a/Engine/Source/Runtime/RenderCore/Private/ShaderParameterStruct.cpp
+++ b/Engine/Source/Runtime/RenderCore/Private/ShaderParameterStruct.cpp
@@ -74,13 +74,33 @@ struct FShaderParameterStructBindingContext
 				BaseType == UBMT_UINT32 ||
 				BaseType == UBMT_FLOAT32);
 
-			if (BaseType == UBMT_NESTED_STRUCT || BaseType == UBMT_INCLUDED_STRUCT)
+			if (BaseType == UBMT_INCLUDED_STRUCT)
+			{
+				checkf(!bIsArray, TEXT("Array of included structure is impossible."));
+				Bind(
+					*Member.GetStructMetadata(),
+					/* MemberPrefix = */ MemberPrefix,
+					/* GeneralByteOffset = */ ByteOffset);
+				continue;
+			}
+			else if (BaseType == UBMT_NESTED_STRUCT && bIsArray)
+			{
+				for (uint32 ArrayElementId = 0; ArrayElementId < (bIsArray ? ArraySize : 1u); ArrayElementId++)
+				{
+					FString NewPrefix = FString::Printf(TEXT("%s%s_%d_"), MemberPrefix, Member.GetName(), ArrayElementId);
+					Bind(
+						*Member.GetStructMetadata(),
+						/* MemberPrefix = */ *NewPrefix,
+						/* GeneralByteOffset = */ ByteOffset);
+				}
+				continue;
+			}
+			else if (BaseType == UBMT_NESTED_STRUCT && !bIsArray)
 			{
-				checkf(!bIsArray, TEXT("Array of structure bindings is not supported."));
 				FString NewPrefix = FString::Printf(TEXT("%s%s_"), MemberPrefix, Member.GetName());
 				Bind(
 					*Member.GetStructMetadata(),
-					/* MemberPrefix = */ BaseType == UBMT_INCLUDED_STRUCT ? MemberPrefix : *NewPrefix,
+					/* MemberPrefix = */ *NewPrefix,
 					/* GeneralByteOffset = */ ByteOffset);
 				continue;
 			}
diff --git a/Engine/Source/Runtime/RenderCore/Private/ShaderParameters.cpp b/Engine/Source/Runtime/RenderCore/Private/ShaderParameters.cpp
--- a/Engine/Source/Runtime/RenderCore/Private/ShaderParameters.cpp
+++ b/Engine/Source/Runtime/RenderCore/Private/ShaderParameters.cpp
@@ -169,6 +169,7 @@ static void CreateHLSLUniformBufferStructMembersDeclaration(
 
 		if(Member.GetBaseType() == UBMT_NESTED_STRUCT)
 		{
+			checkf(Member.GetNumElements() == 0, TEXT("SHADER_PARAMETER_STRUCT_ARRAY() is not supported in uniform buffer yet."));
 			Decl.StructMembers += TEXT("struct {\r\n");
 			Decl.Initializer += TEXT(",{");
 			CreateHLSLUniformBufferStructMembersDeclaration(*Member.GetStructMetadata(), FString::Printf(TEXT("%s%s_"), *NamePrefix, Member.GetName()), StructOffset + Member.GetOffset(), Decl, HLSLBaseOffset);
diff --git a/Engine/Source/Runtime/RenderCore/Public/ShaderParameterMacros.h b/Engine/Source/Runtime/RenderCore/Public/ShaderParameterMacros.h
--- a/Engine/Source/Runtime/RenderCore/Public/ShaderParameterMacros.h
+++ b/Engine/Source/Runtime/RenderCore/Public/ShaderParameterMacros.h
@@ -557,6 +557,34 @@ struct TShaderParameterTypeInfo<TUniformBufferRef<UniformBufferStructType>>
 	static const FShaderParametersMetadata* GetStructMetadata() { return &UniformBufferStructType::StaticStructMetadata; }
 };
 
+template<typename StructType>
+struct TShaderParameterStructTypeInfo
+{
+	static constexpr int32 NumRows = 1;
+	static constexpr int32 NumColumns = 1;
+	static constexpr int32 NumElements = 0;
+	static constexpr int32 Alignment = SHADER_PARAMETER_STRUCT_ALIGNMENT;
+	static constexpr bool bIsStoredInConstantBuffer = true;
+
+	using TAlignedType = StructType;
+
+	static const FShaderParametersMetadata* GetStructMetadata() { return StructType::FTypeInfo::GetStructMetadata(); }
+};
+
+template<typename StructType, size_t InNumElements>
+struct TShaderParameterStructTypeInfo<StructType[InNumElements]>
+{
+	static constexpr int32 NumRows = 1;
+	static constexpr int32 NumColumns = 1;
+	static constexpr int32 NumElements = InNumElements;
+	static constexpr int32 Alignment = SHADER_PARAMETER_STRUCT_ALIGNMENT;
+	static constexpr bool bIsStoredInConstantBuffer = true;
+
+	using TAlignedType = TStaticArray<StructType, InNumElements>;
+
+	static const FShaderParametersMetadata* GetStructMetadata() { return StructType::FTypeInfo::GetStructMetadata(); }
+};
+
 
 #define INTERNAL_BEGIN_GLOBAL_SHADER_PARAMETER_STRUCT \
 	static FShaderParametersMetadata StaticStructMetadata; \
@@ -851,6 +879,30 @@ extern RENDERCORE_API FShaderParametersMetadata* FindUniformBufferStructByFName(
 #define SHADER_PARAMETER_STRUCT(StructType,MemberName) \
 	INTERNAL_SHADER_PARAMETER_EXPLICIT(UBMT_NESTED_STRUCT, StructType::FTypeInfo, StructType, MemberName,,,EShaderPrecisionModifier::Float,TEXT(#StructType),true)
 
+/** Nests an array of shader parameter structure into another one, in C++ and shader code.
+ *
+ * Example:
+ *	BEGIN_SHADER_PARAMETER_STRUCT(FMyNestedStruct,)
+ *		SHADER_PARAMETER(float, MyScalar)
+ *		// ...
+ *	END_SHADER_PARAMETER_STRUCT()
+ *
+ *	BEGIN_SHADER_PARAMETER_STRUCT(FOtherStruct)
+ *		SHADER_PARAMETER_STRUCT_ARRAY(FMyNestedStruct, MyStructArray, [4])
+ *
+ * C++ use case:
+ *	FOtherStruct Parameters;
+ *	Parameters.MyStructArray[0].MyScalar = 1.0f;
+ *
+ * Shader code for a globally named shader parameter struct (UNSUPPORTED):
+ *	float MyScalar = MyGlobalShaderBindingName.MyStructArray_0.MyScalar;
+ *
+ * Shader code for a unnamed shader parameter struct:
+ *	float MyScalar = MyStructArray_0_MyScalar;
+ */
+#define SHADER_PARAMETER_STRUCT_ARRAY(StructType,MemberName, ArrayDecl) \
+	INTERNAL_SHADER_PARAMETER_EXPLICIT(UBMT_NESTED_STRUCT, TShaderParameterStructTypeInfo<StructType ArrayDecl>, StructType, MemberName,ArrayDecl,,EShaderPrecisionModifier::Float,TEXT(#StructType),true)
+
 /** Include a shader parameter structure into another one in shader code.
  *
  * Example:
