Fixes a bug on member offsets with SHADER_PARAMETER_STRUCT_ARRAY()

diff --git a/Engine/Source/Runtime/RenderCore/Private/ShaderParameterStruct.cpp b/Engine/Source/Runtime/RenderCore/Private/ShaderParameterStruct.cpp
--- a/Engine/Source/Runtime/RenderCore/Private/ShaderParameterStruct.cpp
+++ b/Engine/Source/Runtime/RenderCore/Private/ShaderParameterStruct.cpp
@@ -85,13 +85,15 @@ struct FShaderParameterStructBindingContext
 			}
 			else if (BaseType == UBMT_NESTED_STRUCT && bIsArray)
 			{
+				const FShaderParametersMetadata* ChildStruct = Member.GetStructMetadata();
+				uint32 StructSize = ChildStruct->GetSize();
 				for (uint32 ArrayElementId = 0; ArrayElementId < (bIsArray ? ArraySize : 1u); ArrayElementId++)
 				{
 					FString NewPrefix = FString::Printf(TEXT("%s%s_%d_"), MemberPrefix, Member.GetName(), ArrayElementId);
 					Bind(
-						*Member.GetStructMetadata(),
+						*ChildStruct,
 						/* MemberPrefix = */ *NewPrefix,
-						/* GeneralByteOffset = */ ByteOffset);
+						/* GeneralByteOffset = */ ByteOffset + ArrayElementId * StructSize);
 				}
 				continue;
 			}
