Implements a mechanism to verify a pass only access its RDG resources

diff --git a/Engine/Source/Runtime/RenderCore/Private/RenderGraphBuilder.cpp b/Engine/Source/Runtime/RenderCore/Private/RenderGraphBuilder.cpp
--- a/Engine/Source/Runtime/RenderCore/Private/RenderGraphBuilder.cpp
+++ b/Engine/Source/Runtime/RenderCore/Private/RenderGraphBuilder.cpp
@@ -849,6 +849,11 @@ void FRDGBuilder::ExecutePass( const FRenderGraphPass* Pass )
 	bool bHasRenderTargets = false;
 
 	AllocateAndTransitionPassResources(Pass, &RPInfo, &bHasRenderTargets);
+	
+	if (RENDER_GRAPH_DEBUGGING)
+	{
+		UpdateAccessGuardForPassResources(Pass, /** bAllowAccess */ true);
+	}
 
 	if (RENDER_GRAPH_DRAW_EVENTS)
 	{
@@ -888,6 +893,7 @@ void FRDGBuilder::ExecutePass( const FRenderGraphPass* Pass )
 
 	if (RENDER_GRAPH_DEBUGGING)
 	{
+		UpdateAccessGuardForPassResources(Pass, /** bAllowAccess */ false);
 		WarnForUselessPassDependencies(Pass);
 	}
 
@@ -1133,6 +1139,46 @@ void FRDGBuilder::AllocateAndTransitionPassResources(const FRenderGraphPass* Pas
 	OutRPInfo->bGeneratingMips = bGeneratingMips;
 }
 
+// static
+void FRDGBuilder::UpdateAccessGuardForPassResources(const FRenderGraphPass* Pass, bool bAllowAccess)
+{
+#if RENDER_GRAPH_DEBUGGING
+	FShaderParameterStructRef ParameterStruct = Pass->GetParameters();
+
+	for (int ResourceIndex = 0, Num = ParameterStruct.Layout->Resources.Num(); ResourceIndex < Num; ResourceIndex++)
+	{
+		EUniformBufferBaseType Type = ParameterStruct.Layout->Resources[ResourceIndex].MemberType;
+		uint16 Offset = ParameterStruct.Layout->Resources[ResourceIndex].MemberOffset;
+
+		if (IsRDGResourceReferenceShaderParameterType(Type))
+		{
+			FRDGResource* RESTRICT Resource = *ParameterStruct.GetMemberPtrAtOffset<FRDGResource*>(Offset);
+			if (Resource)
+				Resource->bAllowAccessToRHIResource = bAllowAccess;
+		}
+		else if (Type == UBMT_RENDER_TARGET_BINDING_SLOTS)
+		{
+			FRenderTargetBindingSlots* RESTRICT RenderTargets = ParameterStruct.GetMemberPtrAtOffset<FRenderTargetBindingSlots>(Offset);
+
+			for (int32 i = 0; i < RenderTargets->Output.Num(); i++)
+			{
+				FRDGTexture* Texture = RenderTargets->Output[i].GetTexture();
+				if (Texture)
+					Texture->bAllowAccessToRHIResource = bAllowAccess;
+				else
+					break;
+			}
+			
+			const FDepthStencilBinding& DepthStencil = RenderTargets->DepthStencil;
+			if (DepthStencil.Texture)
+			{
+				DepthStencil.Texture->bAllowAccessToRHIResource = bAllowAccess;
+			}
+		}
+	}
+#endif
+}
+
 // static 
 void FRDGBuilder::WarnForUselessPassDependencies(const FRenderGraphPass* Pass)
 {
diff --git a/Engine/Source/Runtime/RenderCore/Public/RenderGraphBuilder.h b/Engine/Source/Runtime/RenderCore/Public/RenderGraphBuilder.h
--- a/Engine/Source/Runtime/RenderCore/Public/RenderGraphBuilder.h
+++ b/Engine/Source/Runtime/RenderCore/Public/RenderGraphBuilder.h
@@ -625,6 +625,7 @@ private:
 	void PushDrawEventStack(const FRenderGraphPass* Pass);
 	void ExecutePass( const FRenderGraphPass* Pass );
 	void AllocateAndTransitionPassResources(const FRenderGraphPass* Pass, struct FRHIRenderPassInfo* OutRPInfo, bool* bOutHasRenderTargets);
+	static void UpdateAccessGuardForPassResources(const FRenderGraphPass* Pass, bool bAllowAccess);
 	static void WarnForUselessPassDependencies(const FRenderGraphPass* Pass);
 
 	void ReleaseRHITextureIfPossible(FRDGTexture* Texture);
diff --git a/Engine/Source/Runtime/RenderCore/Public/RenderGraphResources.h b/Engine/Source/Runtime/RenderCore/Public/RenderGraphResources.h
--- a/Engine/Source/Runtime/RenderCore/Public/RenderGraphResources.h
+++ b/Engine/Source/Runtime/RenderCore/Public/RenderGraphResources.h
@@ -71,14 +71,31 @@ public:
 	{
 		CachedRHI.Resource = nullptr;
 	}
+	
+	/** Verify that the RHI resource can be accessed at a pass execution. */
+	inline void VerifyRHIAccess() const
+	{
+		#if RENDER_GRAPH_DEBUGGING
+		{
+			checkf(bAllowAccessToRHIResource,
+				TEXT("Accessing the RHI resource of %s at this time is not allowed. If you hit this check in pass, ")
+				TEXT("that is due to this resource not being referenced in the parameters of your pass."),
+				Name);
+		}
+		#endif
+	}
+
+	/** Marks this resource as actually used by a resource. This is to track what dependencies on pass was actually unecessary. */
+	inline void MarkResourceAsUsed()
+	{
+		VerifyRHIAccess();
+		bIsActuallyUsedByPass = true;
+	}
 
 	FRDGResource() = delete;
 	FRDGResource(const FRDGResource&) = delete;
 	void operator = (const FRDGResource&) = delete;
 
-	/** Boolean to track at runtime whether a ressource is actually used by the lambda of a pass or not, to detect unnecessary resource dependencies on passes. */
-	bool bIsActuallyUsedByPass = false;
-
 private:
 	/** Number of references in passes and deferred queries. */
 	int32 ReferenceCount = 0;
@@ -87,10 +104,16 @@ private:
 	bool bWritable = false;
 	bool bCompute = false;
 
+	/** Boolean to track at runtime whether a ressource is actually used by the lambda of a pass or not, to detect unnecessary resource dependencies on passes. */
+	bool bIsActuallyUsedByPass = false;
+
 #if RENDER_GRAPH_DEBUGGING
 	/** Boolean to track at wiring time if a resource has ever been produced by a pass, to error out early if accessing a resource that has not been produced. */
 	bool bHasEverBeenProduced = false;
 
+	/** Boolean to track at pass execution whether the underlying RHI resource is allowed to be accessed. */
+	bool bAllowAccessToRHIResource = false;
+
 	/** Pointer towards the pass that is the first to produce it, for even more convenient error message. */
 	const FRenderGraphPass* DebugFirstProducer = nullptr;
 
@@ -131,6 +154,7 @@ public:
 	/** Returns the allocated pooled render target. Must only be called within a pass's lambda. */
 	inline IPooledRenderTarget* GetPooledRenderTarget() const
 	{
+		VerifyRHIAccess();
 		check(PooledRenderTarget);
 		return PooledRenderTarget;
 	}
@@ -139,6 +163,7 @@ public:
 	/** Returns the allocated RHI texture. Must only be called within a pass's lambda. */
 	inline FTextureRHIParamRef GetRHITexture() const
 	{
+		VerifyRHIAccess();
 		check(CachedRHI.Texture);
 		return CachedRHI.Texture;
 	}
@@ -466,6 +491,7 @@ public:
 	/** Returns the buffer to use for indirect RHI calls. */
 	FVertexBufferRHIParamRef GetIndirectRHICallBuffer() const
 	{
+		VerifyRHIAccess();
 		check(PooledBuffer);
 		checkf(Desc.UnderlyingType == FRDGBufferDesc::EUnderlyingType::VertexBuffer, TEXT("Indirect buffers needs to be underlying vertex buffer."));
 		check(PooledBuffer->VertexBuffer.IsValid());
diff --git a/Engine/Source/Runtime/RenderCore/Public/ShaderParameterStruct.h b/Engine/Source/Runtime/RenderCore/Public/ShaderParameterStruct.h
--- a/Engine/Source/Runtime/RenderCore/Public/ShaderParameterStruct.h
+++ b/Engine/Source/Runtime/RenderCore/Public/ShaderParameterStruct.h
@@ -113,8 +113,7 @@ inline void SetShaderUAVs(TRHICmdList& RHICmdList, const TShaderClass* Shader, F
 			}
 			else
 			{
-				// Mark this resource as used by the pass for inefficient pass resource dependency debugging purpose.
-				GraphUAV->bIsActuallyUsedByPass = true;
+				GraphUAV->MarkResourceAsUsed();
 			}
 		}
 
@@ -216,8 +215,7 @@ inline void SetShaderParameters(TRHICmdList& RHICmdList, const TShaderClass* Sha
 			}
 			else
 			{
-				// Mark this resource as used by the pass for inefficient pass resource dependency debugging purpose.
-				GraphTexture->bIsActuallyUsedByPass = true;
+				GraphTexture->MarkResourceAsUsed();
 			}
 		}
 
@@ -238,8 +236,7 @@ inline void SetShaderParameters(TRHICmdList& RHICmdList, const TShaderClass* Sha
 			}
 			else
 			{
-				// Mark this resource as used by the pass for inefficient pass resource dependency debugging purpose.
-				GraphSRV->bIsActuallyUsedByPass = true;
+				GraphSRV->MarkResourceAsUsed();
 			}
 		}
 
@@ -343,8 +340,7 @@ void SetShaderParameters(FRayTracingShaderBindingsWriter& RTBindingsWriter, cons
 			}
 			else
 			{
-				// Mark this resource as used by the pass for inefficient pass resource dependency debugging purpose.
-				GraphTexture->bIsActuallyUsedByPass = true;
+				GraphTexture->MarkResourceAsUsed();
 			}
 		}
 
@@ -365,8 +361,7 @@ void SetShaderParameters(FRayTracingShaderBindingsWriter& RTBindingsWriter, cons
 			}
 			else
 			{
-				// Mark this resource as used by the pass for inefficient pass resource dependency debugging purpose.
-				GraphSRV->bIsActuallyUsedByPass = true;
+				GraphSRV->MarkResourceAsUsed();
 			}
 		}
 
@@ -387,8 +382,7 @@ void SetShaderParameters(FRayTracingShaderBindingsWriter& RTBindingsWriter, cons
 			}
 			else
 			{
-				// Mark this resource as used by the pass for inefficient pass resource dependency debugging purpose.
-				UAV->bIsActuallyUsedByPass = true;
+				UAV->MarkResourceAsUsed();
 			}
 		}
 
