Fixes a bug in ClearUnusedGraphResources() by ensuring member metadata is sorted by byte offset in the shader parameter structure.

diff --git a/Engine/Shaders/Public/ShaderVersion.ush b/Engine/Shaders/Public/ShaderVersion.ush
--- a/Engine/Shaders/Public/ShaderVersion.ush
+++ b/Engine/Shaders/Public/ShaderVersion.ush
@@ -2,5 +2,4 @@
 // Each time the console command is executed it generates a new GUID. As this file is included
 // in Platform.ush (which should be included in any shader) it allows to invalidate the shader DDC.
 // 
-// If you are merging streams and there is a conflict with this GUID you should make a new GUID rather than taking one or the other.
-// GUID = A9685212-F729-4542-8B27-35936DC94536
+// GUID = 9686A7C745F60675F92F2DBFF3598712
\ No newline at end of file
diff --git a/Engine/Source/Runtime/RenderCore/Private/ShaderParameterMetadata.cpp b/Engine/Source/Runtime/RenderCore/Private/ShaderParameterMetadata.cpp
--- a/Engine/Source/Runtime/RenderCore/Private/ShaderParameterMetadata.cpp
+++ b/Engine/Source/Runtime/RenderCore/Private/ShaderParameterMetadata.cpp
@@ -216,21 +216,24 @@ void FShaderParametersMetadata::InitializeLayout()
 		}
 	} // for (int32 i = 0; i < MemberStack.Num(); ++i)
 
-#if 0 // TODO(RDG)
-	/** Sort the resource on MemberType first to avoid CPU miss predictions when iterating over the resources. Then based on ascending offset
-	 * to still allow O(N) complexity on offset cross referencing such as done in ClearUnusedGraphResourcesImpl().
-	 */
 	Layout.Resources.Sort([](
 		const FRHIUniformBufferLayout::FResourceParameter& A,
 		const FRHIUniformBufferLayout::FResourceParameter& B)
 	{
+#if 0 // TODO(RDG)
+		/** Sort the resource on MemberType first to avoid CPU miss predictions when iterating over the resources. Then based on ascending offset
+		 * to still allow O(N) complexity on offset cross referencing such as done in ClearUnusedGraphResourcesImpl().
+		 */
 		if (A.MemberType == B.MemberType)
 		{
 			return A.MemberOffset < B.MemberOffset;
 		}
 		return A.MemberType < B.MemberType;
-	});
+#else
+		// Sorts the resource based on MemberOffset to allow O(N) complexity on offset cross referencing such as done in ClearUnusedGraphResourcesImpl().
+		return A.MemberOffset < B.MemberOffset;
 #endif
+	});
 
 	Layout.ComputeHash();
 
