Implements a ValidateShaderParameters() for user code to manually validate its shader parameters.

diff --git a/Engine/Source/Runtime/RenderCore/Public/ShaderParameterStruct.h b/Engine/Source/Runtime/RenderCore/Public/ShaderParameterStruct.h
--- a/Engine/Source/Runtime/RenderCore/Public/ShaderParameterStruct.h
+++ b/Engine/Source/Runtime/RenderCore/Public/ShaderParameterStruct.h
@@ -69,6 +69,13 @@ FORCEINLINE void ValidateShaderParameters(const FShader* Shader, const FShaderPa
 
 #endif // !DO_CHECK
 
+template<typename TShaderClass>
+FORCEINLINE void ValidateShaderParameters(const TShaderClass* Shader, const typename TShaderClass::FParameters& Parameters)
+{
+	const typename TShaderClass::FParameters* ParameterPtr = &Parameters;
+	return ValidateShaderParameters(Shader, TShaderClass::FParameters::FTypeInfo::GetStructMetadata(), ParameterPtr);
+}
+
 
 /** Set compute shader UAVs. */
 template<typename TRHICmdList, typename TShaderClass, typename TShaderRHI>
