Adds a FRDGBuilder::CreateUAV(FRDGTextureRef) to make some C++ compiler happy with the implicit casts otherwise used.

diff --git a/Engine/Source/Runtime/RenderCore/Public/RenderGraphBuilder.h b/Engine/Source/Runtime/RenderCore/Public/RenderGraphBuilder.h
--- a/Engine/Source/Runtime/RenderCore/Public/RenderGraphBuilder.h
+++ b/Engine/Source/Runtime/RenderCore/Public/RenderGraphBuilder.h
@@ -376,6 +376,11 @@ public:
 		return UAV;
 	}
 
+	FORCEINLINE_DEBUGGABLE FRDGTextureUAVRef CreateUAV(FRDGTextureRef Texture)
+	{
+		return CreateUAV( FRDGTextureUAVDesc(Texture) );
+	}
+
 	/** Create graph tracked UAV for a buffer from a descriptor. */
 	FORCEINLINE_DEBUGGABLE FRDGBufferUAVRef CreateUAV(const FRDGBufferUAVDesc& Desc)
 	{
