Fixed build break in graph builder when draw events are disabled.

diff --git a/Engine/Source/Runtime/RenderCore/Public/RenderGraphBuilder.h b/Engine/Source/Runtime/RenderCore/Public/RenderGraphBuilder.h
--- a/Engine/Source/Runtime/RenderCore/Public/RenderGraphBuilder.h
+++ b/Engine/Source/Runtime/RenderCore/Public/RenderGraphBuilder.h
@@ -21,7 +21,7 @@
 #if WITH_PROFILEGPU
 	#define RENDER_GRAPH_DRAW_EVENTS RENDER_GRAPH_DRAW_EVENTS_STRING_COPY
 #else
-	#define RENDER_GRAPH_DRAW_EVENTS 0
+	#define RENDER_GRAPH_DRAW_EVENTS RENDER_GRAPH_DRAW_EVENTS_NONE
 #endif
 
 
@@ -31,17 +31,15 @@ class RENDERCORE_API FRDGEventName
 public:
 	FRDGEventName() = default;
 
-	#if RENDER_GRAPH_DRAW_EVENTS == RENDER_GRAPH_DRAW_EVENTS_STRING_COPY
-
+#if RENDER_GRAPH_DRAW_EVENTS == RENDER_GRAPH_DRAW_EVENTS_STRING_COPY
 	explicit FRDGEventName(const TCHAR* EventFormat, ...);
-
-	#elif RENDER_GRAPH_DRAW_EVENTS == RENDER_GRAPH_DRAW_EVENTS_STRING_REF
-
-	explicit inline FRDGEventName(const TCHAR* EventFormat)
+#else
+	explicit inline FRDGEventName(const TCHAR* EventFormat, ...)
+	#if RENDER_GRAPH_DRAW_EVENTS == RENDER_GRAPH_DRAW_EVENTS_STRING_REF
 		: EventName(EventFormat)
-	{}
-
 	#endif
+	{}
+#endif
 
 	FRDGEventName(const FRDGEventName& Other)
 	{
