Fixed allocation of temporary FScreenPassContext objects that lead to memory leaks.

FScreenPassContext::Create used to allocate FScreenPassContext objects on FMemStack, however it has a non-trivial destructor which was never called.

Instead of allocating temporary objects on the FMemStack automatically to ensure that they remain alive for the duration of GraphBuilder lifetime, FRDGBuilder::AllocObject() has been added. It explicitly associates the lifetime of allocated temporary object with the specific graph builder.

Correct destructors are automatically called when FRDGBuilder itself is destroyed.

diff --git a/Engine/Source/Runtime/RenderCore/Private/RenderGraphBuilder.cpp b/Engine/Source/Runtime/RenderCore/Private/RenderGraphBuilder.cpp
--- a/Engine/Source/Runtime/RenderCore/Private/RenderGraphBuilder.cpp
+++ b/Engine/Source/Runtime/RenderCore/Private/RenderGraphBuilder.cpp
@@ -1474,6 +1474,14 @@ void FRDGBuilder::DestructPasses()
 
 FRDGBuilder::~FRDGBuilder()
 {
+	// Owned objects must be destroyed in reverse allocation order to ensure that children are destroyed before parents.
+	while(OwnedObjects.Num())
+	{
+		FOwnedObjectPtr& Obj = OwnedObjects.Last();
+		Obj.Deleter(Obj.Ptr);
+		OwnedObjects.Pop(false);
+	}
+
 	#if RENDER_GRAPH_DEBUGGING
 	{
 		checkf(bHasExecuted, TEXT("Render graph execution si required to ensure consistency with immediate mode."));
diff --git a/Engine/Source/Runtime/RenderCore/Public/RenderGraphBuilder.h b/Engine/Source/Runtime/RenderCore/Public/RenderGraphBuilder.h
--- a/Engine/Source/Runtime/RenderCore/Public/RenderGraphBuilder.h
+++ b/Engine/Source/Runtime/RenderCore/Public/RenderGraphBuilder.h
@@ -401,6 +401,27 @@ public:
 		return CreateUAV( FRDGBufferUAVDesc(Buffer, Format) );
 	}
 
+	/**
+	 * Allocates an object that will have the same lifetime as this graph builder.
+	 * Destructor will be automatically called for this object when the graph builder itself is destroyed.
+	 */
+	template<typename T, typename... TArgs>
+	FORCEINLINE_DEBUGGABLE T* AllocObject(TArgs&&... Args)
+	{
+		// Owned objects must be deleted in stack order, therefore destruction stack is pushed before object is constructed.
+		// This ensures that any child objects are placed after the parents in the OwnedObjects array.
+		int32 Index = OwnedObjects.AddDefaulted();
+
+		T* Result = new (FMemStack::Get()) T(Forward<TArgs>(Args)...);
+
+		FOwnedObjectPtr OwnedObject;
+		OwnedObject.Ptr = Result;
+		OwnedObject.Deleter = [](void* InPtr) { ((T*)InPtr)->~T(); };
+		OwnedObjects[Index] = OwnedObject;
+
+		return Result;
+	}
+
 	/** Allocates parameter struct specifically to survive through the life time of the render graph. */
 	template< typename ParameterStructType >
 	FORCEINLINE_DEBUGGABLE ParameterStructType* AllocParameters()
@@ -627,6 +648,15 @@ private:
 		TSet<void*> AllocatedUnusedPassParameters;
 	#endif
 
+	/** Holds objects that are allocated and owned by this graph builder. Allow calling correct destructors at shutdown. */
+	struct FOwnedObjectPtr
+	{
+		typedef void(*DeleterCallback)(void*);
+		void* Ptr = nullptr;
+		DeleterCallback Deleter = nullptr;
+	};
+	TArray<FOwnedObjectPtr, SceneRenderingAllocator> OwnedObjects;
+
 	void DebugPass(const FRenderGraphPass* Pass);
 	void ValidatePass(const FRenderGraphPass* Pass) const;
 	void CaptureAnyInterestingPassOutput(const FRenderGraphPass* Pass);
diff --git a/Engine/Source/Runtime/Renderer/Private/DeferredShadingRenderer.cpp b/Engine/Source/Runtime/Renderer/Private/DeferredShadingRenderer.cpp
--- a/Engine/Source/Runtime/Renderer/Private/DeferredShadingRenderer.cpp
+++ b/Engine/Source/Runtime/Renderer/Private/DeferredShadingRenderer.cpp
@@ -1781,7 +1781,7 @@ void FDeferredShadingSceneRenderer::Render(FRHICommandListImmediate& RHICmdList)
 
 				if (IsSubsurfaceRequiredForView(View))
 				{
-					FScreenPassContext* Context = FScreenPassContext::Create(RHICmdList, View);
+					FScreenPassContext* Context = GraphBuilder.AllocObject<FScreenPassContext>(RHICmdList, View);
 
 					const FScreenPassTexture Input = FScreenPassTexture::Create(SceneColorTexture, View.ViewRect);
 
diff --git a/Engine/Source/Runtime/Renderer/Private/PostProcess/PostProcessSubsurface.cpp b/Engine/Source/Runtime/Renderer/Private/PostProcess/PostProcessSubsurface.cpp
--- a/Engine/Source/Runtime/Renderer/Private/PostProcess/PostProcessSubsurface.cpp
+++ b/Engine/Source/Runtime/Renderer/Private/PostProcess/PostProcessSubsurface.cpp
@@ -682,7 +682,7 @@ void FSubsurfaceVisualizeCompositePass::Process(FRenderingCompositePassContext&
 
 	FRDGBuilder GraphBuilder(RHICmdList);
 
-	FScreenPassContext* Context = FScreenPassContext::Create(RHICmdList, CompositePassContext.View);
+	FScreenPassContext* Context = GraphBuilder.AllocObject<FScreenPassContext>(RHICmdList, CompositePassContext.View);
 
 	const FScreenPassTexture TextureInput = FScreenPassTexture::Create(
 		CreateRDGTextureForInput(GraphBuilder, ePId_Input0, TEXT("SceneColor"), eFC_0000),
diff --git a/Engine/Source/Runtime/Renderer/Private/PostProcess/PostProcessing.cpp b/Engine/Source/Runtime/Renderer/Private/PostProcess/PostProcessing.cpp
--- a/Engine/Source/Runtime/Renderer/Private/PostProcess/PostProcessing.cpp
+++ b/Engine/Source/Runtime/Renderer/Private/PostProcess/PostProcessing.cpp
@@ -1771,7 +1771,7 @@ void FPostProcessing::Process(FRHICommandListImmediate& RHICmdList, const FViewI
 					{
 						OutColorTexture = VisualizeMotionBlur(
 							GraphBuilder,
-							FScreenPassContext::Create(InContext.RHICmdList, InContext.View),
+							GraphBuilder.AllocObject<FScreenPassContext>(InContext.RHICmdList, InContext.View),
 							ColorTexture,
 							DepthTexture,
 							VelocityTexture);
@@ -1780,7 +1780,7 @@ void FPostProcessing::Process(FRHICommandListImmediate& RHICmdList, const FViewI
 					{
 						OutColorTexture = ComputeMotionBlur(
 							GraphBuilder,
-							FScreenPassContext::Create(InContext.RHICmdList, InContext.View),
+							GraphBuilder.AllocObject<FScreenPassContext>(InContext.RHICmdList, InContext.View),
 							ColorTexture,
 							DepthTexture,
 							VelocityTexture);
diff --git a/Engine/Source/Runtime/Renderer/Private/ScreenPass.cpp b/Engine/Source/Runtime/Renderer/Private/ScreenPass.cpp
--- a/Engine/Source/Runtime/Renderer/Private/ScreenPass.cpp
+++ b/Engine/Source/Runtime/Renderer/Private/ScreenPass.cpp
@@ -163,11 +163,6 @@ FScreenPassCommonParameters GetScreenPassCommonParameters(FRHICommandListImmedia
 	return Parameters;
 }
 
-FScreenPassContext* FScreenPassContext::Create(FRHICommandListImmediate& RHICmdList, const FViewInfo& InView)
-{
-	return new(FMemStack::Get()) FScreenPassContext(RHICmdList, InView);
-}
-
 FScreenPassContext::FScreenPassContext(FRHICommandListImmediate& RHICmdList, const FViewInfo& InView)
 	: View(InView)
 	, ViewFamily(*View.Family)
diff --git a/Engine/Source/Runtime/Renderer/Private/ScreenPass.h b/Engine/Source/Runtime/Renderer/Private/ScreenPass.h
--- a/Engine/Source/Runtime/Renderer/Private/ScreenPass.h
+++ b/Engine/Source/Runtime/Renderer/Private/ScreenPass.h
@@ -217,8 +217,10 @@ FScreenPassCommonParameters GetScreenPassCommonParameters(FRHICommandListImmedia
 class FScreenPassContext
 {
 public:
-	// Creates an instance of the context on the stack specifically to survive through the lifetime of the render graph.
-	static FScreenPassContext* Create(FRHICommandListImmediate& RHICmdList, const FViewInfo& InView);
+
+	UE_NONCOPYABLE(FScreenPassContext)
+
+	FScreenPassContext(FRHICommandListImmediate& RHICmdList, const FViewInfo& InView);
 
 	// The current view instance being processed.
 	const FViewInfo& View;
@@ -246,10 +248,6 @@ public:
 
 	// Common screen space parameters, filled at context creation time.
 	const FScreenPassCommonParameters ScreenPassCommonParameters;
-
-private:
-	FScreenPassContext(FRHICommandListImmediate& RHICmdList, const FViewInfo& InView);
-	FScreenPassContext(const FScreenPassContext&) = delete;
 };
 
 using FScreenPassContextRef = const FScreenPassContext*;
