Queries the FMemStackBase used to allocate pretty much any RDG resources in FRDGBuilder's constructor to reduce TLS queries.

diff --git a/Engine/Source/Runtime/RenderCore/Public/RenderGraphBuilder.h b/Engine/Source/Runtime/RenderCore/Public/RenderGraphBuilder.h
--- a/Engine/Source/Runtime/RenderCore/Public/RenderGraphBuilder.h
+++ b/Engine/Source/Runtime/RenderCore/Public/RenderGraphBuilder.h
@@ -227,6 +227,7 @@ public:
 	/** A RHI cmd list is required, if using the immediate mode. */
 	FRDGBuilder(FRHICommandListImmediate& InRHICmdList)
 		: RHICmdList(InRHICmdList)
+		, MemStack(FMemStack::Get())
 	{
 		for (int32 i = 0; i < kMaxScopeCount; i++)
 			ScopesStack[i] = nullptr;
@@ -412,7 +413,7 @@ public:
 		// This ensures that any child objects are placed after the parents in the OwnedObjects array.
 		int32 Index = OwnedObjects.AddDefaulted();
 
-		T* Result = new (FMemStack::Get()) T(Forward<TArgs>(Args)...);
+		T* Result = new (MemStack) T(Forward<TArgs>(Args)...);
 
 		FOwnedObjectPtr OwnedObject;
 		OwnedObject.Ptr = Result;
@@ -427,7 +428,7 @@ public:
 	FORCEINLINE_DEBUGGABLE ParameterStructType* AllocParameters()
 	{
 		// TODO(RDG): could allocate using AllocateForRHILifeTime() to avoid the copy done when using FRHICommandList::BuildLocalUniformBuffer()
-		ParameterStructType* OutParameterPtr = new(FMemStack::Get()) ParameterStructType;
+		ParameterStructType* OutParameterPtr = new(MemStack) ParameterStructType;
 		FMemory::Memzero(OutParameterPtr, sizeof(ParameterStructType));
 		#if RENDER_GRAPH_DEBUGGING
 		{
@@ -472,7 +473,7 @@ public:
 		}
 		#endif
 
-		auto NewPass = new(FMemStack::Get()) TLambdaRenderPass<ParameterStructType, ExecuteLambdaType>(
+		auto NewPass = new(MemStack) TLambdaRenderPass<ParameterStructType, ExecuteLambdaType>(
 			static_cast<FRDGEventName&&>(Name), CurrentScope,
 			{ ParameterStruct, &ParameterStructType::FTypeInfo::GetStructMetadata()->GetLayout() },
 			Flags,
@@ -592,10 +593,12 @@ public:
 		return CurrentScope;
 	}
 
-public:
 	/** The RHI command list used for the render graph. */
 	FRHICommandListImmediate& RHICmdList;
 
+	/** Memory stack use for allocating RDG resource and passes. */
+	FMemStackBase& MemStack;
+
 
 private:
 	static constexpr int32 kMaxScopeCount = 8;
@@ -603,10 +606,10 @@ private:
 	/** Array of all pass created */
 	TArray<FRenderGraphPass*, SceneRenderingAllocator> Passes;
 
-	/** Keep the references over the pooled render target, since FRDGTexture is allocated on FMemStack. */
+	/** Keep the references over the pooled render target, since FRDGTexture is allocated on MemStack. */
 	TMap<FRDGTexture*, TRefCountPtr<IPooledRenderTarget>, SceneRenderingSetAllocator> AllocatedTextures;
 
-	/** Keep the references over the pooled render target, since FRDGTexture is allocated on FMemStack. */
+	/** Keep the references over the pooled render target, since FRDGTexture is allocated on MemStack. */
 	TMap<FRDGBuffer*, TRefCountPtr<FPooledRDGBuffer>, SceneRenderingSetAllocator> AllocatedBuffers;
 
 	/** Array of all deferred access to internal textures. */
@@ -671,7 +674,7 @@ private:
 		// on the RHI's stack allocator so RHICreateUniformBuffer() can dereference render graph resources.
 		if (RHICmdList.Bypass() || 1) // TODO: UE-68018
 		{
-			return new (FMemStack::Get()) Type(Forward<ConstructorParameterTypes>(ConstructorParameters)...);
+			return new (MemStack) Type(Forward<ConstructorParameterTypes>(ConstructorParameters)...);
 		}
 		else
 		{
@@ -726,7 +729,7 @@ public:
 	{
 		checkf(!GraphBuilder.bHasExecuted, TEXT("Render graph bulider has already been executed."));
 
-		auto NewScope = new(FMemStack::Get()) FRDGEventScope(GraphBuilder.CurrentScope, Forward<FRDGEventName>(ScopeName));
+		auto NewScope = new(GraphBuilder.MemStack) FRDGEventScope(GraphBuilder.CurrentScope, Forward<FRDGEventName>(ScopeName));
 
 		#if RENDER_GRAPH_DRAW_EVENTS == RENDER_GRAPH_DRAW_EVENTS_STRING_COPY
 		{
