Validates the FDepthStencilBinding as early as possible to catch common mistake quicker.

diff --git a/Engine/Source/Runtime/RenderCore/Private/RenderGraphBuilder.cpp b/Engine/Source/Runtime/RenderCore/Private/RenderGraphBuilder.cpp
--- a/Engine/Source/Runtime/RenderCore/Private/RenderGraphBuilder.cpp
+++ b/Engine/Source/Runtime/RenderCore/Private/RenderGraphBuilder.cpp
@@ -359,14 +359,14 @@ void FRDGBuilder::ValidatePass(const FRenderGraphPass* Pass) const
 			{
 				checkf(RenderTarget.GetLoadAction() != ERenderTargetLoadAction::ELoad,
 					TEXT("Can't load a render target %s that has never been produced."),
-					RenderTarget.GetTexture()->Name);
+					Texture->Name);
 
 				// TODO(RDG): should only be done when there is a store action.
 				Texture->bHasEverBeenProduced = true;
 				Texture->DebugFirstProducer = Pass;
 			}
 
-			bFoundRTBound = bFoundRTBound || IsBoundAsReadable(RenderTarget.GetTexture(), ParameterStruct);
+			bFoundRTBound = bFoundRTBound || IsBoundAsReadable(Texture, ParameterStruct);
 		}
 		for (int32 i = NumRenderTargets; i < RenderTargets->Output.Num(); i++)
 		{
@@ -375,15 +375,16 @@ void FRDGBuilder::ValidatePass(const FRenderGraphPass* Pass) const
 		}
 		ensureMsgf(!bGeneratingMips || bFoundRTBound, TEXT("GenerateMips enabled but no RT found as source!"));
 
-		FRDGTexture* Texture = RenderTargets->DepthStencil.Texture;
+		FRDGTexture* Texture = RenderTargets->DepthStencil.GetTexture();
 		if (Texture && !Texture->bHasEverBeenProduced)
 		{
-			checkf(RenderTargets->DepthStencil.DepthLoadAction != ERenderTargetLoadAction::ELoad,
-				TEXT("Can't load depth from a render target that has never been produced."));
-			checkf(RenderTargets->DepthStencil.StencilLoadAction != ERenderTargetLoadAction::ELoad,
-				TEXT("Can't load stencil from a render target that has never been produced."));
-			checkf(RenderTargets->DepthStencil.StencilStoreAction != ERenderTargetStoreAction::ENoAction,
-				TEXT("Using for the first time a render target without storing it, which doesn't make much sense."));
+			checkf(RenderTargets->DepthStencil.GetDepthLoadAction() != ERenderTargetLoadAction::ELoad,
+				TEXT("Can't load depth from render target %s that has never been produced."),
+				Texture->Name);
+
+			checkf(RenderTargets->DepthStencil.GetStencilLoadAction() != ERenderTargetLoadAction::ELoad,
+				TEXT("Can't load stencil from render target %s that has never been produced."),
+				Texture->Name);
 
 			Texture->bHasEverBeenProduced = true;
 			Texture->DebugFirstProducer = Pass;
@@ -419,11 +420,11 @@ void FRDGBuilder::CaptureAnyInterestingPassOutput(const FRenderGraphPass* Pass)
 		case UBMT_RENDER_TARGET_BINDING_SLOTS:
 		{
 			FRenderTargetBindingSlots* RESTRICT RenderTargets = ParameterStruct.GetMemberPtrAtOffset<FRenderTargetBindingSlots>(Offset);
-			if (RenderTargets->DepthStencil.Texture && 
-				(RenderTargets->DepthStencil.DepthStoreAction != ERenderTargetStoreAction::ENoAction || RenderTargets->DepthStencil.StencilStoreAction != ERenderTargetStoreAction::ENoAction) &&
-				GVisualizeTexture.ShouldCapture(RenderTargets->DepthStencil.Texture->Name))
+			if (RenderTargets->DepthStencil.GetTexture() && 
+				(RenderTargets->DepthStencil.GetDepthStoreAction() != ERenderTargetStoreAction::ENoAction || RenderTargets->DepthStencil.GetStencilStoreAction() != ERenderTargetStoreAction::ENoAction) &&
+				GVisualizeTexture.ShouldCapture(RenderTargets->DepthStencil.GetTexture()->Name))
 			{
-				GVisualizeTexture.CreateContentCapturePass(*this, RenderTargets->DepthStencil.Texture);
+				GVisualizeTexture.CreateContentCapturePass(*this, RenderTargets->DepthStencil.GetTexture());
 			}
 			for (int32 i = 0; i < RenderTargets->Output.Num(); i++)
 			{
@@ -526,9 +527,9 @@ void FRDGBuilder::WalkGraphDependencies()
 				}
 
 				const FDepthStencilBinding& DepthStencil = RenderTargets->DepthStencil;
-				if (DepthStencil.Texture)
+				if (DepthStencil.GetTexture())
 				{
-					DepthStencil.Texture->ReferenceCount++;
+					DepthStencil.GetTexture()->ReferenceCount++;
 				}
 			}
 			break;
@@ -1104,25 +1105,25 @@ void FRDGBuilder::AllocateAndTransitionPassResources(const FRenderGraphPass* Pas
 			}
 
 			const FDepthStencilBinding& DepthStencil = RenderTargets->DepthStencil;
-			if (DepthStencil.Texture)
+			if (DepthStencil.GetTexture())
 			{
-				AllocateRHITextureIfNeeded(DepthStencil.Texture, false);
+				AllocateRHITextureIfNeeded(DepthStencil.GetTexture(), false);
 
-				OutRPInfo->DepthStencilRenderTarget.DepthStencilTarget = DepthStencil.Texture->PooledRenderTarget->GetRenderTargetItem().TargetableTexture;
+				OutRPInfo->DepthStencilRenderTarget.DepthStencilTarget = DepthStencil.GetTexture()->PooledRenderTarget->GetRenderTargetItem().TargetableTexture;
 				OutRPInfo->DepthStencilRenderTarget.ResolveTarget = nullptr;
 				OutRPInfo->DepthStencilRenderTarget.Action = MakeDepthStencilTargetActions(
-					MakeRenderTargetActions(DepthStencil.DepthLoadAction, DepthStencil.DepthStoreAction),
-					MakeRenderTargetActions(DepthStencil.StencilLoadAction, DepthStencil.StencilStoreAction));
-				OutRPInfo->DepthStencilRenderTarget.ExclusiveDepthStencil = DepthStencil.DepthStencilAccess;
+					MakeRenderTargetActions(DepthStencil.GetDepthLoadAction(), DepthStencil.GetDepthStoreAction()),
+					MakeRenderTargetActions(DepthStencil.GetStencilLoadAction(), DepthStencil.GetStencilStoreAction()));
+				OutRPInfo->DepthStencilRenderTarget.ExclusiveDepthStencil = DepthStencil.GetDepthStencilAccess();
 
-				TransitionTexture(DepthStencil.Texture, EResourceTransitionAccess::EWritable, false);
+				TransitionTexture(DepthStencil.GetTexture(), EResourceTransitionAccess::EWritable, false);
 
 				NumSamples |= OutRPInfo->DepthStencilRenderTarget.DepthStencilTarget->GetNumSamples();
 				NumDepthStencilTargets++;
 					
 				#if RENDER_GRAPH_DEBUGGING
 				{
-					DepthStencil.Texture->DebugPassAccessCount++;
+					DepthStencil.GetTexture()->DebugPassAccessCount++;
 				}
 				#endif
 			}
@@ -1171,9 +1172,9 @@ void FRDGBuilder::UpdateAccessGuardForPassResources(const FRenderGraphPass* Pass
 			}
 			
 			const FDepthStencilBinding& DepthStencil = RenderTargets->DepthStencil;
-			if (DepthStencil.Texture)
+			if (DepthStencil.GetTexture())
 			{
-				DepthStencil.Texture->bAllowAccessToRHIResource = bAllowAccess;
+				DepthStencil.GetTexture()->bAllowAccessToRHIResource = bAllowAccess;
 			}
 		}
 	}
@@ -1366,9 +1367,9 @@ void FRDGBuilder::ReleaseUnecessaryResources(const FRenderGraphPass* Pass)
 					break;
 				}
 			}
-			if (RenderTargets->DepthStencil.Texture)
+			if (RenderTargets->DepthStencil.GetTexture())
 			{
-				ReleaseRHITextureIfPossible(RenderTargets->DepthStencil.Texture);
+				ReleaseRHITextureIfPossible(RenderTargets->DepthStencil.GetTexture());
 			}
 		}
 		break;
diff --git a/Engine/Source/Runtime/RenderCore/Private/ShaderParameterStruct.cpp b/Engine/Source/Runtime/RenderCore/Private/ShaderParameterStruct.cpp
--- a/Engine/Source/Runtime/RenderCore/Private/ShaderParameterStruct.cpp
+++ b/Engine/Source/Runtime/RenderCore/Private/ShaderParameterStruct.cpp
@@ -309,6 +309,47 @@ void FShaderParameterBindings::BindForRootShaderParameters(const FShader* Shader
 	}
 }
 
+bool FDepthStencilBinding::Validate() const
+{
+	if (Texture)
+	{
+		EPixelFormat PixelFormat = Texture->Desc.Format;
+		const TCHAR* FormatString = GetPixelFormatString(PixelFormat);
+
+		bool bIsDepthFormat = PixelFormat == PF_DepthStencil || PixelFormat == PF_ShadowDepth || PixelFormat == PF_D24;
+		checkf(bIsDepthFormat,
+			TEXT("Can't bind texture %s as a depth stencil because its pixel format is %s."),
+			Texture->Name, FormatString);
+		
+		checkf(DepthStencilAccess != FExclusiveDepthStencil::DepthNop_StencilNop,
+			TEXT("Why binding texture %s if there is no access?"),
+			Texture->Name);
+
+		bool bHasStencil = PixelFormat == PF_DepthStencil;
+		if (!bHasStencil)
+		{
+			checkf(StencilLoadAction == ERenderTargetLoadAction::ENoAction && StencilStoreAction == ERenderTargetStoreAction::ENoAction,
+				TEXT("Unable to load stencil of texture %s that have a pixel format %s that does not support stencil."),
+				Texture->Name, FormatString);
+		
+			checkf(!DepthStencilAccess.IsUsingStencil(),
+				TEXT("Unable to have stencil access on texture %s that have a pixel format %s that does not support stencil."),
+				Texture->Name, FormatString);
+		}
+	}
+	else
+	{
+		checkf(DepthLoadAction == ERenderTargetLoadAction::ENoAction && DepthStoreAction == ERenderTargetStoreAction::ENoAction,
+			TEXT("Can't have a depth load or store action when no texture are bound."));
+		checkf(StencilLoadAction == ERenderTargetLoadAction::ENoAction && StencilStoreAction == ERenderTargetStoreAction::ENoAction,
+			TEXT("Can't have a stencil load or store action when no texture are bound."));
+		checkf(DepthStencilAccess == FExclusiveDepthStencil::DepthNop_StencilNop,
+			TEXT("Can't have a depth stencil access when no texture are bound."));
+	}
+
+	return true;
+}
+
 void EmitNullShaderParameterFatalError(const FShader* Shader, const FShaderParametersMetadata* ParametersMetadata, uint16 MemberOffset)
 {
 	const FShaderParametersMetadata* MemberContainingStruct = nullptr;
diff --git a/Engine/Source/Runtime/RenderCore/Public/ShaderParameterMacros.h b/Engine/Source/Runtime/RenderCore/Public/ShaderParameterMacros.h
--- a/Engine/Source/Runtime/RenderCore/Public/ShaderParameterMacros.h
+++ b/Engine/Source/Runtime/RenderCore/Public/ShaderParameterMacros.h
@@ -189,8 +189,13 @@ struct alignas(SHADER_PARAMETER_STRUCT_ALIGNMENT) FDepthStencilBinding
 {
 	FDepthStencilBinding() = default;
 	FDepthStencilBinding(const FDepthStencilBinding&) = default;
-
-	FDepthStencilBinding(
+	
+	/**
+	 * Creates a render target binding informations for a depth/stencil texture.
+	 *
+	 * Notes: Load and store action are on explicit without default values, to force the user to not forget one of these.
+	 */
+	FORCEINLINE FDepthStencilBinding(
 		FRDGTexture* InTexture,
 		ERenderTargetLoadAction InDepthLoadAction,
 		ERenderTargetStoreAction InDepthStoreAction,
@@ -203,9 +208,11 @@ struct alignas(SHADER_PARAMETER_STRUCT_ALIGNMENT) FDepthStencilBinding
 		, StencilLoadAction(InStencilLoadAction)
 		, StencilStoreAction(InStencilStoreAction)
 		, DepthStencilAccess(InDepthStencilAccess)
-	{}
+	{
+		check(Validate());
+	}
 
-	FDepthStencilBinding(
+	FORCEINLINE FDepthStencilBinding(
 		FRDGTexture* InTexture,
 		ERenderTargetLoadAction InDepthLoadAction,
 		ERenderTargetStoreAction InDepthStoreAction,
@@ -214,14 +221,48 @@ struct alignas(SHADER_PARAMETER_STRUCT_ALIGNMENT) FDepthStencilBinding
 		, DepthLoadAction(InDepthLoadAction)
 		, DepthStoreAction(InDepthStoreAction)
 		, DepthStencilAccess(InDepthStencilAccess)
-	{}
+	{
+		check(Validate());
+	}
+
+	FORCEINLINE FRDGTexture* GetTexture() const
+	{
+		return Texture;
+	}
+	FORCEINLINE ERenderTargetLoadAction GetDepthLoadAction() const
+	{
+		return DepthLoadAction;
+	}
+	FORCEINLINE ERenderTargetStoreAction GetDepthStoreAction() const
+	{
+		return DepthStoreAction;
+	}
+	FORCEINLINE ERenderTargetLoadAction GetStencilLoadAction() const
+	{
+		return StencilLoadAction;
+	}
+	FORCEINLINE ERenderTargetStoreAction GetStencilStoreAction() const
+	{
+		return StencilStoreAction;
+	}
+	FORCEINLINE FExclusiveDepthStencil GetDepthStencilAccess() const
+	{
+		return DepthStencilAccess;
+	}
 
+private:
+	/** 
+	 * All parameters required to bind a depth render target deferred. This are purposefully private to
+	 * force the user to call FDepthStencilBinding() constructors, forcing him to specify the mendatory parameter without forgetting one.
+	 */
 	TAlignedShaderParameterPtr<FRDGTexture*> Texture = nullptr;
 	ERenderTargetLoadAction		DepthLoadAction		= ERenderTargetLoadAction::ENoAction;
 	ERenderTargetStoreAction	DepthStoreAction	= ERenderTargetStoreAction::ENoAction;
 	ERenderTargetLoadAction		StencilLoadAction	= ERenderTargetLoadAction::ENoAction;
 	ERenderTargetStoreAction	StencilStoreAction	= ERenderTargetStoreAction::ENoAction;
 	FExclusiveDepthStencil		DepthStencilAccess	= FExclusiveDepthStencil::DepthNop_StencilNop;
+
+	RENDERCORE_API bool Validate() const;
 };
 
 
diff --git a/Engine/Source/Runtime/Renderer/Private/PostProcess/PostProcessTemporalAA.cpp b/Engine/Source/Runtime/Renderer/Private/PostProcess/PostProcessTemporalAA.cpp
--- a/Engine/Source/Runtime/Renderer/Private/PostProcess/PostProcessTemporalAA.cpp
+++ b/Engine/Source/Runtime/Renderer/Private/PostProcess/PostProcessTemporalAA.cpp
@@ -707,12 +707,11 @@ FTAAOutputs FTAAPassParameters::AddTemporalAAPass(
 
 		if (bUseResponsiveStencilTest)
 		{
-			PassParameters->RenderTargets.DepthStencil.Texture = SceneBlackboard.SceneDepthBuffer;
-			PassParameters->RenderTargets.DepthStencil.DepthLoadAction = ERenderTargetLoadAction::ENoAction;
-			PassParameters->RenderTargets.DepthStencil.DepthStoreAction = ERenderTargetStoreAction::ENoAction;
-			PassParameters->RenderTargets.DepthStencil.StencilLoadAction = ERenderTargetLoadAction::ELoad;
-			PassParameters->RenderTargets.DepthStencil.StencilStoreAction = ERenderTargetStoreAction::EStore;
-			PassParameters->RenderTargets.DepthStencil.DepthStencilAccess = FExclusiveDepthStencil::DepthRead_StencilWrite;
+			PassParameters->RenderTargets.DepthStencil = FDepthStencilBinding(
+				SceneBlackboard.SceneDepthBuffer,
+				ERenderTargetLoadAction::ENoAction, ERenderTargetStoreAction::ENoAction,
+				ERenderTargetLoadAction::ELoad, ERenderTargetStoreAction::EStore,
+				FExclusiveDepthStencil::DepthRead_StencilWrite);
 		}
 		
 		FTemporalAAPS::FPermutationDomain BasePermutationVector;
