Fixes a bug where RDG was unable to capture resources on render target binding slot >= 1.

diff --git a/Engine/Source/Runtime/RenderCore/Private/RenderGraphBuilder.cpp b/Engine/Source/Runtime/RenderCore/Private/RenderGraphBuilder.cpp
--- a/Engine/Source/Runtime/RenderCore/Private/RenderGraphBuilder.cpp
+++ b/Engine/Source/Runtime/RenderCore/Private/RenderGraphBuilder.cpp
@@ -429,15 +429,15 @@ void FRDGBuilder::CaptureAnyInterestingPassOutput(const FRenderGraphPass* Pass)
 			for (int32 i = 0; i < RenderTargets->Output.Num(); i++)
 			{
 				const FRenderTargetBinding& RenderTarget = RenderTargets->Output[i];
-				if (RenderTarget.GetTexture() &&
-					RenderTarget.GetStoreAction() != ERenderTargetStoreAction::ENoAction &&
-					GVisualizeTexture.ShouldCapture(RenderTarget.GetTexture()->Name))
+				if (!RenderTarget.GetTexture())
 				{
-					GVisualizeTexture.CreateContentCapturePass(*this, RenderTarget.GetTexture());
+					break;
 				}
-				else
+
+				if (RenderTarget.GetStoreAction() != ERenderTargetStoreAction::ENoAction &&
+					GVisualizeTexture.ShouldCapture(RenderTarget.GetTexture()->Name))
 				{
-					break;
+					GVisualizeTexture.CreateContentCapturePass(*this, RenderTarget.GetTexture());
 				}
 			}
 		}
