Implements r.VisualizeTexture.AllowBlinking to be able to turn off the blinking on NaNs and INFs that is not helping for the eyes.

diff --git a/Engine/Source/Runtime/RenderCore/Private/VisualizeTexture.cpp b/Engine/Source/Runtime/RenderCore/Private/VisualizeTexture.cpp
--- a/Engine/Source/Runtime/RenderCore/Private/VisualizeTexture.cpp
+++ b/Engine/Source/Runtime/RenderCore/Private/VisualizeTexture.cpp
@@ -61,6 +61,12 @@ TGlobalResource<FVisualizeTexture> GVisualizeTexture;
 
 
 #if WITH_ENGINE
+
+static TAutoConsoleVariable<int32> CVarAllowBlinking(
+	TEXT("r.VisualizeTexture.AllowBlinking"), 1,
+	TEXT("Whether to allow blinking when visualizing NaN or inf that can become irritating over time.\n"),
+	ECVF_RenderThreadSafe);
+
 /** A pixel shader which filters a texture. */
 // @param TextureType 0:Cube, 1:1D(not yet supported), 2:2D no MSAA, 3:3D, 4:Cube[], 5:2D MSAA, 6:2D DepthStencil no MSAA (needed to avoid D3DDebug error)
 class FVisualizeTexturePS : public FGlobalShader
@@ -218,7 +224,7 @@ void FVisualizeTexture::CreateContentCapturePass(FRDGBuilder& GraphBuilder, cons
 
 			// w * almost_1 to avoid frac(1) => 0
 			PassParameters->VisualizeParam[0] = FVector4(RGBMul, SingleChannelMul, Add, FracScale * 0.9999f);
-			PassParameters->VisualizeParam[1] = FVector4(BlinkState, bSaturateInsteadOfFrac ? 1.0f : 0.0f, ArrayIndex, CustomMip);
+			PassParameters->VisualizeParam[1] = FVector4(CVarAllowBlinking.GetValueOnRenderThread() ? BlinkState : 1.0f, bSaturateInsteadOfFrac ? 1.0f : 0.0f, ArrayIndex, CustomMip);
 			PassParameters->VisualizeParam[2] = FVector4(InputValueMapping, 0.0f, SingleChannel);
 		}
 
