Refactored render graph defines and event management.

 - Moved defines to a dedicated header file. Renamed to RDG_ convention.
 - Moved render graph event classes to its own file. Moved event management out of the builder
   and into a separate class.
 - RDGPass::ParentScope is now set by the builder instead of passed by the user.

diff --git a/Engine/Source/Runtime/RenderCore/Private/RenderGraphBuilder.cpp b/Engine/Source/Runtime/RenderCore/Private/RenderGraphBuilder.cpp
--- a/Engine/Source/Runtime/RenderCore/Private/RenderGraphBuilder.cpp
+++ b/Engine/Source/Runtime/RenderCore/Private/RenderGraphBuilder.cpp
@@ -7,36 +7,44 @@
 #include "VisualizeTexture.h"
 #include "ProfilingDebugging/CsvProfiler.h"
 
-#if RENDER_GRAPH_DEBUGGING
-
-static int32 GRenderGraphImmediateMode = 0;
-static FAutoConsoleVariableRef CVarImmediateMode(
-	TEXT("r.RDG.ImmediateMode"),
-	GRenderGraphImmediateMode,
-	TEXT("Executes passes as they get created. Useful to have a callstack of the wiring code when crashing in the pass' lambda."),
-	ECVF_RenderThreadSafe);
-
-static int32 GRenderGraphEmitWarnings = 0;
-static FAutoConsoleVariableRef CVarEmitWarnings(
-	TEXT("r.RDG.EmitWarnings"),
-	GRenderGraphEmitWarnings,
-	TEXT("Allow to output warnings for inefficiencies found during wiring and execution of the passes.\n")
-	TEXT(" 0: disabled;\n")
-	TEXT(" 1: emit warning once (default);\n")
-	TEXT(" 2: emit warning everytime issue is detected."),
-	ECVF_RenderThreadSafe);
+namespace
+{
+#if RDG_ENABLE_DEBUG
+	int32 GRenderGraphImmediateMode = 0;
+	FAutoConsoleVariableRef CVarImmediateMode(
+		TEXT("r.RDG.ImmediateMode"),
+		GRenderGraphImmediateMode,
+		TEXT("Executes passes as they get created. Useful to have a callstack of the wiring code when crashing in the pass' lambda."),
+		ECVF_RenderThreadSafe);
+
+	int32 GRenderGraphEmitWarnings = 0;
+	FAutoConsoleVariableRef CVarEmitWarnings(
+		TEXT("r.RDG.EmitWarnings"),
+		GRenderGraphEmitWarnings,
+		TEXT("Allow to output warnings for inefficiencies found during wiring and execution of the passes.\n")
+		TEXT(" 0: disabled;\n")
+		TEXT(" 1: emit warning once (default);\n")
+		TEXT(" 2: emit warning everytime issue is detected."),
+		ECVF_RenderThreadSafe);
 
 #else
-
-static const int32 GRenderGraphImmediateMode = 0;
-static const int32 GRenderGraphEmitWarnings = 0;
-
+	const int32 GRenderGraphImmediateMode = 0;
+	const int32 GRenderGraphEmitWarnings = 0;
 #endif
+}
 
+bool GetEmitRDGEvents()
+{
+#if RDG_EVENTS != RDG_EVENTS_NONE
+	return GetEmitDrawEvents() || GRenderGraphEmitWarnings;
+#else
+	return false;
+#endif
+}
 
 void InitRenderGraph()
 {
-#if RENDER_GRAPH_DEBUGGING && WITH_ENGINE
+#if RDG_ENABLE_DEBUG_WITH_ENGINE
 	if (FParse::Param(FCommandLine::Get(), TEXT("rdgimmediate")))
 	{
 		GRenderGraphImmediateMode = 1;
@@ -104,50 +112,12 @@ static bool IsBoundAsReadable(const FRDGTexture* Texture, FShaderParameterStruct
 	return false;
 }
 
-
-#if RENDER_GRAPH_DRAW_EVENTS == 2
-
-FRDGEventName::FRDGEventName(const TCHAR* EventFormat, ...)
-{
-	if (GetEmitDrawEvents() || GRenderGraphEmitWarnings)
-	{
-		va_list ptr;
-		va_start(ptr, EventFormat);
-		TCHAR TempStr[256];
-		// Build the string in the temp buffer
-		FCString::GetVarArgs(TempStr, ARRAY_COUNT(TempStr), EventFormat, ptr);
-		va_end(ptr);
-
-		EventNameStorage = TempStr;
-		EventName = *EventNameStorage;
-	}
-	else
-	{
-		EventName = TEXT("!!!Unavailable RDG event name: try r.RDG.EmitWarnings=1 or -rdgdebug!!!");
-	}
-}
-
-#endif
-
-
 void FRDGBuilder::Execute()
 {
 	CSV_SCOPED_TIMING_STAT_EXCLUSIVE(FRDGBuilder_Execute);
-	#if RENDER_GRAPH_DEBUGGING
-	{
-		/** The usage need RDG_EVENT_SCOPE() needs to happen in inner scope of the one containing FRDGBuilder because of
-		 *  FStackRDGEventScopeRef's destructor modifying this FRDGBuilder instance.
-		 * 
-		 *
-		 *  FRDGBuilder GraphBuilder(RHICmdList);
-		 *  {
-		 *  	RDG_EVENT_SCOPE(GraphBuilder, "MyEventScope");
-		 *  	// ...
-		 *  }
-		 *  GraphBuilder.Execute();
-		 */
-		checkf(CurrentScope == nullptr, TEXT("Render graph needs to have all scopes ended to execute."));
 
+	#if RDG_ENABLE_DEBUG
+	{
 		checkf(!bHasExecuted, TEXT("Render graph execution should only happen once to ensure consistency with immediate mode."));
 
 		/** FRDGBuilder::AllocParameters() allocates shader parameter structure for the life time until pass execution.
@@ -162,45 +132,68 @@ void FRDGBuilder::Execute()
 	}
 	#endif
 
+	EventScopeStack.BeginExecute();
+
 	if (!GRenderGraphImmediateMode)
 	{
 		WalkGraphDependencies();
 
 		QUICK_SCOPE_CYCLE_COUNTER(STAT_FRDGBuilder_Execute);
-		for (const FRenderGraphPass* Pass : Passes)
+		for (const FRDGPass* Pass : Passes)
 		{
 			ExecutePass(Pass);
 		}
 	}
 
-	// Pops remaining scopes
-	if (RENDER_GRAPH_DRAW_EVENTS)
-	{
-		if (GetEmitDrawEvents())
-		{
-			for (int32 i = 0; i < kMaxScopeCount; i++)
-			{
-				if (!ScopesStack[i])
-					break;
-				RHICmdList.PopEvent();
-			}
-		}
-	}
+	EventScopeStack.EndExecute();
 
 	ProcessDeferredInternalResourceQueries();
 
 	DestructPasses();
 
-	#if RENDER_GRAPH_DEBUGGING
+	#if RDG_ENABLE_DEBUG
 	{
 		bHasExecuted = true;
 	}
 	#endif
 }
 
-void FRDGBuilder::DebugPass(const FRenderGraphPass* Pass)
+void FRDGBuilder::AddPassInternal(FRDGPass* Pass)
 {
-#if RENDER_GRAPH_DEBUGGING
+#if RDG_ENABLE_DEBUG
+	{
+		const void* ParameterStructData = Pass->GetParameters().Contents;
+
+		checkf(!bHasExecuted, TEXT("Render graph pass %s needs to be added before the builder execution."), Pass->GetName());
+
+		/** A pass parameter structure requires a correct life time until the pass execution, and therefor needs to be
+			* allocated with FRDGBuilder::AllocParameters().
+			*
+			* Moreover, because the destructor of this parameter structure will be done after the pass execution, a it can
+			* only be used by a single AddPass().
+			*/
+		checkf(
+			AllocatedUnusedPassParameters.Contains(ParameterStructData),
+			TEXT("The pass parameter structure has not been allocated for correct life time FRDGBuilder::AllocParameters() or has already ")
+			TEXT("been used by another previous FRDGBuilder::AddPass()."));
+
+		AllocatedUnusedPassParameters.Remove(ParameterStructData);
+	}
+#endif
+
+	Pass->ParentScope = EventScopeStack.GetCurrentScope();
+	Passes.Emplace(Pass);
+
+#if RDG_ENABLE_DEBUG || SUPPORTS_VISUALIZE_TEXTURE
+	{
+		DebugPass(Pass);
+	}
+#endif
+}
+
+void FRDGBuilder::DebugPass(const FRDGPass* Pass)
+{
+#if RDG_ENABLE_DEBUG
 	// Verify all the settings of the pass make sense.
 	ValidatePass(Pass);
 
@@ -211,7 +204,7 @@ void FRDGBuilder::DebugPass(const FRenderGraphPass* Pass)
 	}
 #endif
 
-#if WITH_ENGINE && SUPPORTS_VISUALIZE_TEXTURE
+#if SUPPORTS_VISUALIZE_TEXTURE
 	// If visualizing a texture, look for any output of the pass. This must be done after the
 	// GRenderGraphImmediateMode's ExecutePass() because this will actually create a capturing
 	// pass if needed that would have to be executed right away as well.
@@ -222,9 +215,9 @@ void FRDGBuilder::DebugPass(const FRenderGraphPass* Pass)
 #endif
 }
 
-void FRDGBuilder::ValidatePass(const FRenderGraphPass* Pass) const
+void FRDGBuilder::ValidatePass(const FRDGPass* Pass) const
 {
-#if RENDER_GRAPH_DEBUGGING
+#if RDG_ENABLE_DEBUG
 	FRenderTargetBindingSlots* RESTRICT RenderTargets = nullptr;
 	FShaderParameterStructRef ParameterStruct = Pass->GetParameters();
 
@@ -340,7 +333,7 @@ void FRDGBuilder::ValidatePass(const FRenderGraphPass* Pass) const
 	{
 		checkf(bRequiresRenderTargetSlots, TEXT("Render pass %s does not need render target binging slots"), Pass->GetName());
 
-		const bool bGeneratingMips = (Pass->GetFlags() & ERenderGraphPassFlags::GenerateMips) == ERenderGraphPassFlags::GenerateMips;
+		const bool bGeneratingMips = (Pass->GetFlags() & ERDGPassFlags::GenerateMips) == ERDGPassFlags::GenerateMips;
 		bool bFoundRTBound = false;
 
 		int32 NumRenderTargets = 0;
@@ -394,10 +387,10 @@ void FRDGBuilder::ValidatePass(const FRenderGraphPass* Pass) const
 	{
 		checkf(!bRequiresRenderTargetSlots, TEXT("Render pass %s requires render target binging slots"), Pass->GetName());
 	}
-#endif // RENDER_GRAPH_DEBUGGING
+#endif // RDG_ENABLE_DEBUG
 }
 
-void FRDGBuilder::CaptureAnyInterestingPassOutput(const FRenderGraphPass* Pass)
+void FRDGBuilder::CaptureAnyInterestingPassOutput(const FRDGPass* Pass)
 {
 #if SUPPORTS_VISUALIZE_TEXTURE
 	FShaderParameterStructRef ParameterStruct = Pass->GetParameters();
@@ -451,7 +444,7 @@ void FRDGBuilder::CaptureAnyInterestingPassOutput(const FRenderGraphPass* Pass)
 
 void FRDGBuilder::WalkGraphDependencies()
 {
-	for (const FRenderGraphPass* Pass : Passes)
+	for (const FRDGPass* Pass : Passes)
 	{
 		FShaderParameterStructRef ParameterStruct = Pass->GetParameters();
 
@@ -537,7 +530,7 @@ void FRDGBuilder::WalkGraphDependencies()
 				break;
 			}
 		}
-	} // for (const FRenderGraphPass* Pass : Passes)
+	} 
 
 	// Add additional dependencies from deferred queries.
 	for (const auto& Query : DeferredInternalTextureQueries)
@@ -616,7 +609,7 @@ void FRDGBuilder::AllocateRHIBufferSRVIfNeeded(FRDGBufferSRV* SRV, bool bCompute
 	}
 	
 	// The underlying buffer have already been allocated by a prior pass through AllocateRHIBufferUAVIfNeeded().
-	#if RENDER_GRAPH_DEBUGGING
+	#if RDG_ENABLE_DEBUG
 	{
 		check(SRV->Desc.Buffer->bHasEverBeenProduced);
 	}	
@@ -773,77 +766,7 @@ void FRDGBuilder::TransitionUAV(FUnorderedAccessViewRHIParamRef UAV, FRDGResourc
 	}
 }
 
-void FRDGBuilder::PushDrawEventStack(const FRenderGraphPass* Pass)
-{
-	// Push the scope event.
-	{
-		// Find out how many scope events needs to be poped.
-		TStaticArray<const FRDGEventScope*, kMaxScopeCount> TraversedScopes;
-		int32 CommonScopeId = -1;
-		int32 TraversedScopeCount = 0;
-		const FRDGEventScope* PassParentScope = Pass->ParentScope;
-		while (PassParentScope)
-		{
-			TraversedScopes[TraversedScopeCount] = PassParentScope;
-
-			for (int32 i = 0; i < ScopesStack.Num(); i++)
-			{
-				if (ScopesStack[i] == PassParentScope)
-				{
-					CommonScopeId = i;
-					break;
-				}
-			}
-
-			if (CommonScopeId != -1)
-			{
-				break;
-			}
-
-			TraversedScopeCount++;
-			PassParentScope = PassParentScope->ParentScope;
-		}
-
-		// Pop no longer used scopes
-		for (int32 i = CommonScopeId + 1; i < kMaxScopeCount; i++)
-		{
-			if (!ScopesStack[i])
-				break;
-
-			RHICmdList.PopEvent();
-			ScopesStack[i] = nullptr;
-		}
-
-		// Push new scopes
-		const FColor ScopeColor(0);
-		for (int32 i = TraversedScopeCount - 1; i >= 0; i--)
-		{
-			RHICmdList.PushEvent(TraversedScopes[i]->Name.GetTCHAR(), ScopeColor);
-			CommonScopeId++;
-			ScopesStack[CommonScopeId] = TraversedScopes[i];
-		}
-	}
-
-	// Push the pass's event with some color.
-	{
-		FColor Color(0, 0, 0);
-
-		if (Pass->IsCompute())
-		{
-			// Green for compute.
-			Color = FColor(128, 255, 128);
-		}
-		else
-		{
-			// Ref for rasterizer.
-			Color = FColor(255, 128, 128);
-		}
-
-		RHICmdList.PushEvent(Pass->GetName(), Color);
-	}
-}
-
-void FRDGBuilder::ExecutePass( const FRenderGraphPass* Pass )
+void FRDGBuilder::ExecutePass(const FRDGPass* Pass)
 {
 	QUICK_SCOPE_CYCLE_COUNTER(STAT_FRDGBuilder_ExecutePass);
 
@@ -852,20 +775,14 @@ void FRDGBuilder::ExecutePass( const FRenderGraphPass* Pass )
 
 	AllocateAndTransitionPassResources(Pass, &RPInfo, &bHasRenderTargets);
 	
-	if (RENDER_GRAPH_DEBUGGING)
+	if (RDG_ENABLE_DEBUG)
 	{
 		UpdateAccessGuardForPassResources(Pass, /** bAllowAccess */ true);
 	}
 
-	if (RENDER_GRAPH_DRAW_EVENTS)
-	{
-		if (GetEmitDrawEvents())
-		{
-			PushDrawEventStack(Pass);
-		}
-	}
+	EventScopeStack.BeginExecutePass(Pass);
 
-	if( !Pass->IsCompute())
+	if (!Pass->IsCompute())
 	{
 		check(bHasRenderTargets);
 		RHICmdList.BeginRenderPass( RPInfo, Pass->GetName() );
@@ -885,15 +802,9 @@ void FRDGBuilder::ExecutePass( const FRenderGraphPass* Pass )
 		RHICmdList.EndRenderPass();
 	}
 
-	if (RENDER_GRAPH_DRAW_EVENTS)
-	{
-		if (GetEmitDrawEvents())
-		{
-			RHICmdList.PopEvent();
-		}
-	}
+	EventScopeStack.EndExecutePass();
 
-	if (RENDER_GRAPH_DEBUGGING)
+	if (RDG_ENABLE_DEBUG)
 	{
 		UpdateAccessGuardForPassResources(Pass, /** bAllowAccess */ false);
 		WarnForUselessPassDependencies(Pass);
@@ -906,12 +817,12 @@ void FRDGBuilder::ExecutePass( const FRenderGraphPass* Pass )
 	}
 }
 
-void FRDGBuilder::AllocateAndTransitionPassResources(const FRenderGraphPass* Pass, struct FRHIRenderPassInfo* OutRPInfo, bool* bOutHasRenderTargets)
+void FRDGBuilder::AllocateAndTransitionPassResources(const FRDGPass* Pass, struct FRHIRenderPassInfo* OutRPInfo, bool* bOutHasRenderTargets)
 {
 	bool bIsCompute = Pass->IsCompute();
 	FShaderParameterStructRef ParameterStruct = Pass->GetParameters();
 
-	const bool bGeneratingMips = (Pass->GetFlags() & ERenderGraphPassFlags::GenerateMips) == ERenderGraphPassFlags::GenerateMips;
+	const bool bGeneratingMips = (Pass->GetFlags() & ERDGPassFlags::GenerateMips) == ERDGPassFlags::GenerateMips;
 	for (int ResourceIndex = 0, Num = ParameterStruct.Layout->Resources.Num(); ResourceIndex < Num; ResourceIndex++)
 	{
 		EUniformBufferBaseType Type = ParameterStruct.Layout->Resources[ResourceIndex].MemberType;
@@ -925,7 +836,7 @@ void FRDGBuilder::AllocateAndTransitionPassResources(const FRenderGraphPass* Pas
 			if (Texture)
 			{
 				// The underlying texture have already been allocated by a prior pass.
-				#if RENDER_GRAPH_DEBUGGING
+				#if RDG_ENABLE_DEBUG
 				{
 					check(Texture->bHasEverBeenProduced);
 				}	
@@ -934,7 +845,7 @@ void FRDGBuilder::AllocateAndTransitionPassResources(const FRenderGraphPass* Pas
 				check(Texture->CachedRHI.Resource);
 				TransitionTexture(Texture, EResourceTransitionAccess::EReadable, bIsCompute);
 
-				#if RENDER_GRAPH_DEBUGGING
+				#if RDG_ENABLE_DEBUG
 				{
 					Texture->DebugPassAccessCount++;
 				}
@@ -949,7 +860,7 @@ void FRDGBuilder::AllocateAndTransitionPassResources(const FRenderGraphPass* Pas
 			{
 				// The underlying texture have already been allocated by a prior pass.
 				check(SRV->Desc.Texture);
-				#if RENDER_GRAPH_DEBUGGING
+				#if RDG_ENABLE_DEBUG
 				{
 					check(SRV->Desc.Texture->bHasEverBeenProduced);
 				}	
@@ -964,7 +875,7 @@ void FRDGBuilder::AllocateAndTransitionPassResources(const FRenderGraphPass* Pas
 
 				TransitionTexture(SRV->Desc.Texture, EResourceTransitionAccess::EReadable, bIsCompute);
 
-				#if RENDER_GRAPH_DEBUGGING
+				#if RDG_ENABLE_DEBUG
 				{
 					SRV->Desc.Texture->DebugPassAccessCount++;
 				}
@@ -980,7 +891,7 @@ void FRDGBuilder::AllocateAndTransitionPassResources(const FRenderGraphPass* Pas
 				AllocateRHITextureUAVIfNeeded(UAV, bIsCompute);
 				TransitionUAV(UAV->CachedRHI.UAV, UAV->Desc.Texture, UAV->Desc.Texture->Flags, EResourceTransitionAccess::EWritable, bIsCompute);
 
-				#if RENDER_GRAPH_DEBUGGING
+				#if RDG_ENABLE_DEBUG
 				{
 					UAV->Desc.Texture->DebugPassAccessCount++;
 				}
@@ -994,7 +905,7 @@ void FRDGBuilder::AllocateAndTransitionPassResources(const FRenderGraphPass* Pas
 			if (Buffer)
 			{
 				// The underlying buffer have already been allocated by a prior pass through AllocateRHIBufferUAVIfNeeded().
-				#if RENDER_GRAPH_DEBUGGING
+				#if RDG_ENABLE_DEBUG
 				{
 					check(Buffer->bHasEverBeenProduced);
 				}
@@ -1006,7 +917,7 @@ void FRDGBuilder::AllocateAndTransitionPassResources(const FRenderGraphPass* Pas
 				FUnorderedAccessViewRHIParamRef UAV = Buffer->PooledBuffer->UAVs.CreateIterator().Value();
 				TransitionUAV(UAV, Buffer, Buffer->Flags, EResourceTransitionAccess::EReadable, bIsCompute);
 
-				#if RENDER_GRAPH_DEBUGGING
+				#if RDG_ENABLE_DEBUG
 				{
 					Buffer->DebugPassAccessCount++;
 				}
@@ -1021,7 +932,7 @@ void FRDGBuilder::AllocateAndTransitionPassResources(const FRenderGraphPass* Pas
 			{
 				// The underlying buffer have already been allocated by a prior pass through AllocateRHIBufferUAVIfNeeded().
 				check(SRV->Desc.Buffer);
-				#if RENDER_GRAPH_DEBUGGING
+				#if RDG_ENABLE_DEBUG
 				{
 					check(SRV->Desc.Buffer->bHasEverBeenProduced);
 				}
@@ -1035,7 +946,7 @@ void FRDGBuilder::AllocateAndTransitionPassResources(const FRenderGraphPass* Pas
 				FUnorderedAccessViewRHIParamRef UAV = SRV->Desc.Buffer->PooledBuffer->UAVs.CreateIterator().Value();
 				TransitionUAV(UAV, SRV->Desc.Buffer, SRV->Desc.Buffer->Flags, EResourceTransitionAccess::EReadable, bIsCompute);
 
-				#if RENDER_GRAPH_DEBUGGING
+				#if RDG_ENABLE_DEBUG
 				{
 					SRV->Desc.Buffer->DebugPassAccessCount++;
 				}
@@ -1051,7 +962,7 @@ void FRDGBuilder::AllocateAndTransitionPassResources(const FRenderGraphPass* Pas
 				AllocateRHIBufferUAVIfNeeded(UAV, bIsCompute);
 				TransitionUAV(UAV->CachedRHI.UAV, UAV->Desc.Buffer, UAV->Desc.Buffer->Flags, EResourceTransitionAccess::EWritable, bIsCompute);
 
-				#if RENDER_GRAPH_DEBUGGING
+				#if RDG_ENABLE_DEBUG
 				{
 					UAV->Desc.Buffer->DebugPassAccessCount++;
 				}
@@ -1092,7 +1003,7 @@ void FRDGBuilder::AllocateAndTransitionPassResources(const FRenderGraphPass* Pas
 					NumSamples |= OutRPInfo->ColorRenderTargets[i].RenderTarget->GetNumSamples();
 					NumRenderTargets++;
 					
-					#if RENDER_GRAPH_DEBUGGING
+					#if RDG_ENABLE_DEBUG
 					{
 						RenderTarget.GetTexture()->DebugPassAccessCount++;
 					}
@@ -1121,7 +1032,7 @@ void FRDGBuilder::AllocateAndTransitionPassResources(const FRenderGraphPass* Pas
 				NumSamples |= OutRPInfo->DepthStencilRenderTarget.DepthStencilTarget->GetNumSamples();
 				NumDepthStencilTargets++;
 					
-				#if RENDER_GRAPH_DEBUGGING
+				#if RDG_ENABLE_DEBUG
 				{
 					DepthStencil.GetTexture()->DebugPassAccessCount++;
 				}
@@ -1142,9 +1053,9 @@ void FRDGBuilder::AllocateAndTransitionPassResources(const FRenderGraphPass* Pas
 }
 
 // static
-void FRDGBuilder::UpdateAccessGuardForPassResources(const FRenderGraphPass* Pass, bool bAllowAccess)
+void FRDGBuilder::UpdateAccessGuardForPassResources(const FRDGPass* Pass, bool bAllowAccess)
 {
-#if RENDER_GRAPH_DEBUGGING
+#if RDG_ENABLE_DEBUG
 	FShaderParameterStructRef ParameterStruct = Pass->GetParameters();
 
 	for (int ResourceIndex = 0, Num = ParameterStruct.Layout->Resources.Num(); ResourceIndex < Num; ResourceIndex++)
@@ -1182,7 +1093,7 @@ void FRDGBuilder::UpdateAccessGuardForPassResources(const FRenderGraphPass* Pass
 }
 
 // static 
-void FRDGBuilder::WarnForUselessPassDependencies(const FRenderGraphPass* Pass)
+void FRDGBuilder::WarnForUselessPassDependencies(const FRDGPass* Pass)
 {
 	if (!GRenderGraphEmitWarnings)
 	{
@@ -1284,7 +1195,7 @@ void FRDGBuilder::ReleaseRHIBufferIfPossible(FRDGBuffer* Buffer)
 	}
 }
 
-void FRDGBuilder::ReleaseUnecessaryResources(const FRenderGraphPass* Pass)
+void FRDGBuilder::ReleaseUnecessaryResources(const FRDGPass* Pass)
 {
 	FShaderParameterStructRef ParameterStruct = Pass->GetParameters();
 
@@ -1392,7 +1303,7 @@ void FRDGBuilder::ProcessDeferredInternalResourceQueries()
 
 		*Query.OutTexturePtr = AllocatedTextures.FindChecked(Query.Texture);
 		
-		#if RENDER_GRAPH_DEBUGGING
+		#if RDG_ENABLE_DEBUG
 		{
 			// Increment the number of times the texture has been accessed to avoid warning on produced but never used resources that were produced
 			// only to be extracted for the graph.
@@ -1411,7 +1322,7 @@ void FRDGBuilder::ProcessDeferredInternalResourceQueries()
 	{
 		*Query.OutBufferPtr = AllocatedBuffers.FindChecked(Query.Buffer);
 
-#if RENDER_GRAPH_DEBUGGING
+#if RDG_ENABLE_DEBUG
 		{
 			// Increment the number of times the buffer has been accessed to avoid warning on produced but never used resources that were produced
 			// only to be extracted for the graph.
@@ -1429,18 +1340,7 @@ void FRDGBuilder::ProcessDeferredInternalResourceQueries()
 
 void FRDGBuilder::DestructPasses()
 {
-	#if RENDER_GRAPH_DRAW_EVENTS == 2
-	{
-		// Event scopes are allocated on FMemStack, so need to call their destructor because have a FString within them.
-		for (FRDGEventScope* EventScope : EventScopes)
-		{
-			EventScope->~FRDGEventScope();
-		}
-		EventScopes.Empty();
-	}
-	#endif
-
-	#if RENDER_GRAPH_DEBUGGING
+	#if RDG_ENABLE_DEBUG
 	{
 		// Make sure all resource references have been released to ensure no leaks happen,
 		// and emit warning if produced resource has not been used.
@@ -1462,9 +1362,9 @@ void FRDGBuilder::DestructPasses()
 	#endif
 
 	// Passes are allocated on FMemStack, so need to call destructor manually.
-	for (FRenderGraphPass* Pass : Passes)
+	for (int32 Index = Passes.Num() - 1; Index >= 0; --Index)
 	{
-		Pass->~FRenderGraphPass();
+		Passes[Index]->~FRDGPass();
 	}
 
 	Passes.Empty();
@@ -1474,17 +1374,19 @@ void FRDGBuilder::DestructPasses()
 	AllocatedBuffers.Empty();
 }
 
-FRDGBuilder::~FRDGBuilder()
+void FRDGBuilder::BeginEventScope(FRDGEventName&& ScopeName)
 {
-	// Owned objects must be destroyed in reverse allocation order to ensure that children are destroyed before parents.
-	while(OwnedObjects.Num())
-	{
-		FOwnedObjectPtr& Obj = OwnedObjects.Last();
-		Obj.Deleter(Obj.Ptr);
-		OwnedObjects.Pop(false);
-	}
+	EventScopeStack.BeginEventScope(MoveTemp(ScopeName));
+}
+
+void FRDGBuilder::EndEventScope()
+{
+	EventScopeStack.EndEventScope();
+}
 
-	#if RENDER_GRAPH_DEBUGGING
+FRDGBuilder::~FRDGBuilder()
+{
+	#if RDG_ENABLE_DEBUG
 	{
 		checkf(bHasExecuted, TEXT("Render graph execution is required to ensure consistency with immediate mode."));
 	}
diff --git a/Engine/Source/Runtime/RenderCore/Private/RenderGraphEvent.cpp b/Engine/Source/Runtime/RenderCore/Private/RenderGraphEvent.cpp
new file mode 100644
--- /dev/null
+++ b/Engine/Source/Runtime/RenderCore/Private/RenderGraphEvent.cpp
@@ -0,0 +1,228 @@
+// Copyright 1998-2019 Epic Games, Inc. All Rights Reserved.
+
+#pragma once
+
+#include "RenderGraphEvent.h"
+#include "RenderGraphBuilder.h"
+
+extern bool GetEmitRDGEvents();
+
+class FRDGEventScope
+{
+public:
+	FRDGEventScope(const FRDGEventScope* InParentScope, FRDGEventName&& InName)
+		: ParentScope(InParentScope)
+		, Name(InName)
+	{}
+
+	// Pointer towards this one is contained in.
+	const FRDGEventScope* const ParentScope;
+
+	// Name of the event.
+	const FRDGEventName Name;
+};
+
+#if RDG_EVENTS == RDG_EVENTS_STRING_COPY
+
+FRDGEventName::FRDGEventName(const TCHAR* EventFormat, ...)
+{
+	if (GetEmitRDGEvents())
+	{
+		va_list VAList;
+		va_start(VAList, EventFormat);
+		TCHAR TempStr[256];
+		// Build the string in the temp buffer
+		FCString::GetVarArgs(TempStr, ARRAY_COUNT(TempStr), EventFormat, VAList);
+		va_end(VAList);
+
+		EventNameStorage = TempStr;
+		EventName = *EventNameStorage;
+	}
+	else
+	{
+		EventName = TEXT("!!!Unavailable RDG event name: try r.RDG.EmitWarnings=1 or -rdgdebug!!!");
+	}
+}
+
+#endif
+
+FRDGEventScopeGuard::FRDGEventScopeGuard(FRDGBuilder& InGraphBuilder, FRDGEventName&& ScopeName)
+	: GraphBuilder(InGraphBuilder)
+{
+	GraphBuilder.BeginEventScope(MoveTemp(ScopeName));
+}
+
+FRDGEventScopeGuard::~FRDGEventScopeGuard()
+{
+	GraphBuilder.EndEventScope();
+}
+
+FRDGEventScopeStack::FRDGEventScopeStack(FRHICommandListImmediate& InRHICmdList)
+	: RHICmdList(InRHICmdList)
+	, MemStack(FMemStack::Get())
+	, ScopeStack(MakeUniformStaticArray<const FRDGEventScope*, kScopeStackDepthMax>(nullptr))
+{}
+
+FRDGEventScopeStack::~FRDGEventScopeStack()
+{
+	check(!bIsExecuting);
+}
+
+void FRDGEventScopeStack::BeginEventScope(FRDGEventName&& EventName)
+{
+	checkf(!bIsExecuting, TEXT("Can't begin an event scope while executing."));
+
+	if (RDG_EVENTS)
+	{
+		auto Scope = new(MemStack) FRDGEventScope(CurrentScope, Forward<FRDGEventName>(EventName));
+		EventScopes.Add(Scope);
+		CurrentScope = Scope;
+	}
+}
+
+void FRDGEventScopeStack::EndEventScope()
+{
+	checkf(!bIsExecuting, TEXT("Can't end an event scope while executing."));
+
+	if (RDG_EVENTS)
+	{
+		checkf(CurrentScope != nullptr, TEXT("Current scope is null."));
+		CurrentScope = CurrentScope->ParentScope;
+	}
+}
+
+void FRDGEventScopeStack::BeginExecute()
+{
+	/** The usage need RDG_EVENT_SCOPE() needs to happen in inner scope of the one containing FRDGBuilder because of
+	 *  FStackRDGEventScopeRef's destructor modifying this FRDGBuilder instance.
+	 *
+	 *
+	 *  FRDGBuilder GraphBuilder(RHICmdList);
+	 *  {
+	 *  	RDG_EVENT_SCOPE(GraphBuilder, "MyEventScope");
+	 *  	// ...
+	 *  }
+	 *  GraphBuilder.Execute();
+	 */
+	checkf(CurrentScope == nullptr, TEXT("Render graph needs to have all scopes ended to execute."));
+	check(!bIsExecuting);
+	bIsExecuting = true;
+}
+
+void FRDGEventScopeStack::BeginExecutePass(const FRDGPass* Pass)
+{
+	check(bIsExecuting);
+
+	if (!GetEmitRDGEvents())
+	{
+		return;
+	}
+
+	// Push the scope event.
+	{
+		// Find out how many scope events needs to be popped.
+		TStaticArray<const FRDGEventScope*, kScopeStackDepthMax> TraversedScopes;
+		int32 CommonScopeId = -1;
+		int32 TraversedScopeCount = 0;
+
+		const FRDGEventScope* PassParentScope = Pass->GetParentScope();
+
+		while (PassParentScope)
+		{
+			TraversedScopes[TraversedScopeCount] = PassParentScope;
+
+			for (int32 i = 0; i < ScopeStack.Num(); i++)
+			{
+				if (ScopeStack[i] == PassParentScope)
+				{
+					CommonScopeId = i;
+					break;
+				}
+			}
+
+			if (CommonScopeId != -1)
+			{
+				break;
+			}
+
+			TraversedScopeCount++;
+			PassParentScope = PassParentScope->ParentScope;
+		}
+
+		// Pop no longer used scopes
+		for (int32 i = CommonScopeId + 1; i < kScopeStackDepthMax; i++)
+		{
+			if (!ScopeStack[i])
+				break;
+
+			RHICmdList.PopEvent();
+			ScopeStack[i] = nullptr;
+		}
+
+		// Push new scopes
+		const FColor ScopeColor(0);
+		for (int32 i = TraversedScopeCount - 1; i >= 0; i--)
+		{
+			RHICmdList.PushEvent(TraversedScopes[i]->Name.GetTCHAR(), ScopeColor);
+			CommonScopeId++;
+			ScopeStack[CommonScopeId] = TraversedScopes[i];
+		}
+	}
+
+	// Push the pass's event with some color.
+	{
+		FColor Color(0, 0, 0);
+
+		if (Pass->IsCompute())
+		{
+			// Green for compute.
+			Color = FColor(128, 255, 128);
+		}
+		else
+		{
+			// Ref for rasterizer.
+			Color = FColor(255, 128, 128);
+		}
+
+		RHICmdList.PushEvent(Pass->GetName(), Color);
+	}
+}
+
+void FRDGEventScopeStack::EndExecutePass()
+{
+	check(bIsExecuting);
+
+	if (GetEmitRDGEvents())
+	{
+		RHICmdList.PopEvent();
+	}
+}
+
+void FRDGEventScopeStack::EndExecute()
+{
+	check(bIsExecuting);
+	bIsExecuting = false;
+
+	// Pops remaining scopes
+	if (GetEmitRDGEvents())
+	{
+		for (uint32 ScopeIndex = 0; ScopeIndex < kScopeStackDepthMax; ++ScopeIndex)
+		{
+			if (!ScopeStack[ScopeIndex])
+			{
+				break;
+			}
+
+			RHICmdList.PopEvent();
+		}
+	}
+
+#if RDG_EVENTS
+	// Event scopes are allocated on FMemStack, so need to call their destructor because have a FString within them.
+	for (int32 Index = EventScopes.Num() - 1; Index >= 0; --Index)
+	{
+		EventScopes[Index]->~FRDGEventScope();
+	}
+	EventScopes.Empty();
+#endif
+}
\ No newline at end of file
diff --git a/Engine/Source/Runtime/RenderCore/Public/RenderGraph.h b/Engine/Source/Runtime/RenderCore/Public/RenderGraph.h
--- a/Engine/Source/Runtime/RenderCore/Public/RenderGraph.h
+++ b/Engine/Source/Runtime/RenderCore/Public/RenderGraph.h
@@ -56,6 +56,7 @@
  * or FPixelShaderUtils.
  */
 
+#include "RenderGraphDefinitions.h"
 #include "RenderGraphResources.h"
 #include "RenderGraphBuilder.h"
 #include "RenderGraphUtils.h"
diff --git a/Engine/Source/Runtime/RenderCore/Public/RenderGraphBuilder.h b/Engine/Source/Runtime/RenderCore/Public/RenderGraphBuilder.h
--- a/Engine/Source/Runtime/RenderCore/Public/RenderGraphBuilder.h
+++ b/Engine/Source/Runtime/RenderCore/Public/RenderGraphBuilder.h
@@ -2,120 +2,13 @@
 
 #pragma once
 
+#include "RenderGraphDefinitions.h"
+#include "RenderGraphEvent.h"
 #include "RenderGraphResources.h"
 #include "ShaderParameterMacros.h"
 
-
-/** Whether visualize texture tool is supported. */
-#define SUPPORTS_VISUALIZE_TEXTURE (WITH_ENGINE && !(UE_BUILD_SHIPPING || UE_BUILD_TEST))
-
-/** Whether render graph should support draw events or not.
- * RENDER_GRAPH_DRAW_EVENTS == 0 means there is no string processing at all.
- * RENDER_GRAPH_DRAW_EVENTS == 1 means const TCHAR* is only passdown.
- * RENDER_GRAPH_DRAW_EVENTS == 2 means complex formated FString is passdown.
- */
-#define RENDER_GRAPH_DRAW_EVENTS_NONE 0
-#define RENDER_GRAPH_DRAW_EVENTS_STRING_REF 1
-#define RENDER_GRAPH_DRAW_EVENTS_STRING_COPY 2
-
-#if WITH_PROFILEGPU
-	#define RENDER_GRAPH_DRAW_EVENTS RENDER_GRAPH_DRAW_EVENTS_STRING_COPY
-#else
-	#define RENDER_GRAPH_DRAW_EVENTS RENDER_GRAPH_DRAW_EVENTS_NONE
-#endif
-
-
-/** Opaque object to store a draw event. */
-class RENDERCORE_API FRDGEventName
-{
-public:
-	FRDGEventName() = default;
-
-#if RENDER_GRAPH_DRAW_EVENTS == RENDER_GRAPH_DRAW_EVENTS_STRING_COPY
-	explicit FRDGEventName(const TCHAR* EventFormat, ...);
-#else
-	explicit inline FRDGEventName(const TCHAR* EventFormat, ...)
-	#if RENDER_GRAPH_DRAW_EVENTS == RENDER_GRAPH_DRAW_EVENTS_STRING_REF
-		: EventName(EventFormat)
-	#endif
-	{}
-#endif
-
-	FRDGEventName(const FRDGEventName& Other)
-	{
-		*this = Other;
-	}
-
-	FRDGEventName(FRDGEventName&& Other)
-	{
-		*this = Forward<FRDGEventName>(Other);
-	}
-
-	FRDGEventName& operator=(const FRDGEventName& Other)
-	{
-#if RENDER_GRAPH_DRAW_EVENTS == RENDER_GRAPH_DRAW_EVENTS_STRING_REF
-		EventName = Other.EventName;
-#elif RENDER_GRAPH_DRAW_EVENTS == RENDER_GRAPH_DRAW_EVENTS_STRING_COPY
-		EventNameStorage = Other.EventNameStorage;
-		EventName = *EventNameStorage;
-#endif
-		return *this;
-	}
-
-	FRDGEventName& operator=(FRDGEventName&& Other)
-	{
-#if RENDER_GRAPH_DRAW_EVENTS == RENDER_GRAPH_DRAW_EVENTS_STRING_REF
-		EventName = Other.EventName;
-		Other.EventName = nullptr;
-#elif RENDER_GRAPH_DRAW_EVENTS == RENDER_GRAPH_DRAW_EVENTS_STRING_COPY
-		EventNameStorage = MoveTemp(Other.EventNameStorage);
-		EventName = *EventNameStorage;
-		Other.EventName = nullptr;
-#endif
-		return *this;
-	}
-
-	const TCHAR* GetTCHAR() const
-	{
-		#if RENDER_GRAPH_DRAW_EVENTS == RENDER_GRAPH_DRAW_EVENTS_STRING_REF || RENDER_GRAPH_DRAW_EVENTS == RENDER_GRAPH_DRAW_EVENTS_STRING_COPY
-			return EventName;
-		#else
-			// Render graph draw events have been completely compiled for CPU performance reasons.
-			return TEXT("!!!Unavailable RDG event name: need RENDER_GRAPH_DRAW_EVENTS>=0 and r.RDG.EmitWarnings=1 or -rdgdebug!!!");
-		#endif
-	}
-
-private:
-	#if RENDER_GRAPH_DRAW_EVENTS == RENDER_GRAPH_DRAW_EVENTS_STRING_REF || RENDER_GRAPH_DRAW_EVENTS == RENDER_GRAPH_DRAW_EVENTS_STRING_COPY
-		const TCHAR* EventName;
-		#if RENDER_GRAPH_DRAW_EVENTS == RENDER_GRAPH_DRAW_EVENTS_STRING_COPY
-			FString EventNameStorage;
-		#endif
-	#endif
-};
-
-
-/** Hierarchical scope for draw events of passes. */
-class RENDERCORE_API FRDGEventScope
-{
-private:
-	// Pointer towards this one is contained in.
-	const FRDGEventScope* const ParentScope;
-
-	// Name of the event.
-	const FRDGEventName Name;
-
-
-	FRDGEventScope(const FRDGEventScope* InParentScope, FRDGEventName&& InName)
-		: ParentScope(InParentScope), Name(InName)
-	{ }
-	friend class FRDGBuilder;
-	friend class FStackRDGEventScopeRef;
-};
-
-
-/** Flags to anotate passes. */
-enum class ERenderGraphPassFlags
+/** Flags to annotate passes. */
+enum class ERDGPassFlags
 {
 	None = 0,
 
@@ -127,8 +20,10 @@ enum class ERenderGraphPassFlags
 	GenerateMips = 1 << 1,
 };
 
-ENUM_CLASS_FLAGS(ERenderGraphPassFlags)
+ENUM_CLASS_FLAGS(ERDGPassFlags)
 
+// TODO(RDG): Bulk rename across codebase.
+using ERenderGraphPassFlags = ERDGPassFlags;
 
 // TODO(RDG): remove from global scope?
 struct RENDERCORE_API FShaderParameterStructRef
@@ -143,25 +38,27 @@ struct RENDERCORE_API FShaderParameterStructRef
 	}
 };
 
-
 /** 
  * Base class of a render graph pass
  */
-struct RENDERCORE_API FRenderGraphPass
+class RENDERCORE_API FRDGPass
 {
-	FRenderGraphPass(FRDGEventName&& InName, const FRDGEventScope* InParentScope, FShaderParameterStructRef InParameterStruct, ERenderGraphPassFlags InPassFlags)
+public:
+	FRDGPass(
+		FRDGEventName&& InName,
+		FShaderParameterStructRef InParameterStruct,
+		ERDGPassFlags InPassFlags)
 		: Name(static_cast<FRDGEventName&&>(InName))
-		, ParentScope(InParentScope)
 		, ParameterStruct(InParameterStruct)
 		, PassFlags(InPassFlags)
 	{
 		if (IsCompute())
 		{
-			ensureMsgf(ParameterStruct.Layout->NumRenderTargets() == 0, TEXT("Pass %s was declared as ERenderGraphPassFlags::Compute yet has RenderTargets in its ResourceTable"), GetName());
+			ensureMsgf(ParameterStruct.Layout->NumRenderTargets() == 0, TEXT("Pass %s was declared as ERDGPassFlags::Compute yet has RenderTargets in its ResourceTable"), GetName());
 		}
 	}
 
-	virtual ~FRenderGraphPass() {}
+	virtual ~FRDGPass() {}
 
 	virtual void Execute(FRHICommandListImmediate& RHICmdList) const = 0;
 
@@ -170,23 +67,31 @@ struct RENDERCORE_API FRenderGraphPass
 		return Name.GetTCHAR();
 	}
 
-	const ERenderGraphPassFlags& GetFlags() const
+	ERDGPassFlags GetFlags() const
 	{
 		return PassFlags;
 	}
 
 	bool IsCompute() const
 	{
-		return (PassFlags & ERenderGraphPassFlags::Compute) == ERenderGraphPassFlags::Compute;
+		return (PassFlags & ERDGPassFlags::Compute) == ERDGPassFlags::Compute;
 	}
 	
-	FShaderParameterStructRef GetParameters() const { return ParameterStruct; }
+	FShaderParameterStructRef GetParameters() const
+	{
+		return ParameterStruct;
+	}
+
+	const FRDGEventScope* GetParentScope() const
+	{
+		return ParentScope;
+	}
 
 protected:
 	const FRDGEventName Name;
-	const FRDGEventScope* const ParentScope;
+	const FRDGEventScope* ParentScope = nullptr;
 	const FShaderParameterStructRef ParameterStruct;
-	const ERenderGraphPassFlags PassFlags;
+	const ERDGPassFlags PassFlags;
 
 	friend class FRDGBuilder;
 };
@@ -195,18 +100,22 @@ protected:
  * Render graph pass with lambda execute function
  */
 template <typename ParameterStructType, typename ExecuteLambdaType>
-struct TLambdaRenderPass final : public FRenderGraphPass
+struct TRDGLambdaPass final : public FRDGPass
 {
-	TLambdaRenderPass(FRDGEventName&& InName, const FRDGEventScope* InParentScope, FShaderParameterStructRef InParameterStruct, ERenderGraphPassFlags InPassFlags, ExecuteLambdaType&& InExecuteLambda)
-		: FRenderGraphPass(static_cast<FRDGEventName&&>(InName), InParentScope, InParameterStruct, InPassFlags)
+	TRDGLambdaPass(
+		FRDGEventName&& InName,
+		FShaderParameterStructRef InParameterStruct,
+		ERDGPassFlags InPassFlags,
+		ExecuteLambdaType&& InExecuteLambda)
+		: FRDGPass(static_cast<FRDGEventName&&>(InName), InParameterStruct, InPassFlags)
 		, ExecuteLambda(static_cast<ExecuteLambdaType&&>(InExecuteLambda))
 	{ }
 
-	~TLambdaRenderPass()
+	~TRDGLambdaPass()
 	{
 		// Manually call the destructor of the pass parameter, to make sure RHI references are released since the pass parameters are allocated on FMemStack.
 		// TODO(RDG): this may lead to RHI resource leaks if a struct allocated in FMemStack does not actually get used through FRDGBuilder::AddPass().
-		ParameterStructType* Struct = reinterpret_cast<ParameterStructType*>(const_cast<void*>(FRenderGraphPass::ParameterStruct.Contents));
+		ParameterStructType* Struct = reinterpret_cast<ParameterStructType*>(const_cast<void*>(FRDGPass::ParameterStruct.Contents));
 		Struct->~ParameterStructType();
 	}
 
@@ -230,17 +139,15 @@ public:
 	FRDGBuilder(FRHICommandListImmediate& InRHICmdList)
 		: RHICmdList(InRHICmdList)
 		, MemStack(FMemStack::Get())
-	{
-		for (int32 i = 0; i < kMaxScopeCount; i++)
-			ScopesStack[i] = nullptr;
-	}
+		, EventScopeStack(RHICmdList)
+	{}
 
 	~FRDGBuilder();
 
 	/** Register a external texture to be tracked by the render graph. */
 	FORCEINLINE_DEBUGGABLE FRDGTextureRef RegisterExternalTexture(const TRefCountPtr<IPooledRenderTarget>& ExternalPooledTexture, const TCHAR* DebugName = TEXT("External"))
 	{
-		#if RENDER_GRAPH_DEBUGGING
+		#if RDG_ENABLE_DEBUG
 		{
 			ensureMsgf(ExternalPooledTexture.IsValid(), TEXT("Attempted to register NULL external texture: %s"), DebugName);
 			checkf(DebugName, TEXT("Externally allocated texture requires a debug name when registering them to render graph."));
@@ -250,7 +157,7 @@ public:
 		OutTexture->PooledRenderTarget = ExternalPooledTexture;
 		OutTexture->CachedRHI.Texture = ExternalPooledTexture->GetRenderTargetItem().ShaderResourceTexture;
 		AllocatedTextures.Add(OutTexture, ExternalPooledTexture);
-		#if RENDER_GRAPH_DEBUGGING
+		#if RDG_ENABLE_DEBUG
 		{
 			OutTexture->bHasEverBeenProduced = true;
 			Resources.Add(OutTexture);
@@ -262,7 +169,7 @@ public:
 	/** Register a external buffer to be tracked by the render graph. */
 	FORCEINLINE_DEBUGGABLE FRDGBufferRef RegisterExternalBuffer(const TRefCountPtr<FPooledRDGBuffer>& ExternalPooledBuffer, const TCHAR* Name = TEXT("External"))
 	{
-#if RENDER_GRAPH_DEBUGGING
+#if RDG_ENABLE_DEBUG
 		{
 			ensureMsgf(ExternalPooledBuffer.IsValid(), TEXT("Attempted to register NULL external buffer: %s"), Name);
 		}
@@ -270,7 +177,7 @@ public:
 		FRDGBuffer* OutBuffer = AllocateForRHILifeTime<FRDGBuffer>(Name, ExternalPooledBuffer->Desc, ERDGResourceFlags::None);
 		OutBuffer->PooledBuffer = ExternalPooledBuffer;
 		AllocatedBuffers.Add(OutBuffer, ExternalPooledBuffer);
-#if RENDER_GRAPH_DEBUGGING
+#if RDG_ENABLE_DEBUG
 		{
 			OutBuffer->bHasEverBeenProduced = true;
 			Resources.Add(OutBuffer);
@@ -286,7 +193,7 @@ public:
 	 */
 	FORCEINLINE_DEBUGGABLE FRDGTextureRef CreateTexture(const FPooledRenderTargetDesc& Desc, const TCHAR* DebugName, ERDGResourceFlags Flags = ERDGResourceFlags::None)
 	{
-		#if RENDER_GRAPH_DEBUGGING
+		#if RDG_ENABLE_DEBUG
 		{
 			ensureMsgf(!bHasExecuted, TEXT("Render graph texture %s needs to be created before the builder execution."), DebugName);
 			checkf(DebugName, TEXT("Creating a render graph texture requires a valid debug name."));
@@ -294,7 +201,7 @@ public:
 		}
 		#endif
 		FRDGTexture* Texture = AllocateForRHILifeTime<FRDGTexture>(DebugName, Desc, Flags);
-		#if RENDER_GRAPH_DEBUGGING
+		#if RDG_ENABLE_DEBUG
 			Resources.Add(Texture);
 		#endif
 		return Texture;
@@ -306,14 +213,14 @@ public:
 	 */
 	FORCEINLINE_DEBUGGABLE FRDGBufferRef CreateBuffer(const FRDGBufferDesc& Desc, const TCHAR* DebugName, ERDGResourceFlags Flags = ERDGResourceFlags::None)
 	{
-		#if RENDER_GRAPH_DEBUGGING
+		#if RDG_ENABLE_DEBUG
 		{
 			ensureMsgf(!bHasExecuted, TEXT("Render graph buffer %s needs to be created before the builder execution."), DebugName);
 			checkf(DebugName, TEXT("Creating a render graph buffer requires a valid debug name."));
 		}
 		#endif
 		FRDGBuffer* Buffer = AllocateForRHILifeTime<FRDGBuffer>(DebugName, Desc, Flags);
-		#if RENDER_GRAPH_DEBUGGING
+		#if RDG_ENABLE_DEBUG
 			Resources.Add(Buffer);
 		#endif
 		return Buffer;
@@ -323,7 +230,7 @@ public:
 	FORCEINLINE_DEBUGGABLE FRDGTextureSRVRef CreateSRV(const FRDGTextureSRVDesc& Desc)
 	{
 		check(Desc.Texture);
-		#if RENDER_GRAPH_DEBUGGING
+		#if RDG_ENABLE_DEBUG
 		{
 			ensureMsgf(!bHasExecuted, TEXT("Render graph SRV %s needs to be created before the builder execution."), Desc.Texture->Name);
 			ensureMsgf(Desc.Texture->Desc.TargetableFlags & TexCreate_ShaderResource, TEXT("Attempted to create SRV from texture %s which was not created with TexCreate_ShaderResource"), Desc.Texture->Name);
@@ -331,7 +238,7 @@ public:
 		#endif
 		
 		FRDGTextureSRV* SRV = AllocateForRHILifeTime<FRDGTextureSRV>(Desc.Texture->Name, Desc);
-		#if RENDER_GRAPH_DEBUGGING
+		#if RDG_ENABLE_DEBUG
 			Resources.Add(SRV);
 		#endif
 		return SRV;
@@ -341,14 +248,14 @@ public:
 	FORCEINLINE_DEBUGGABLE FRDGBufferSRVRef CreateSRV(const FRDGBufferSRVDesc& Desc)
 	{
 		check(Desc.Buffer);
-		#if RENDER_GRAPH_DEBUGGING
+		#if RDG_ENABLE_DEBUG
 		{
 			ensureMsgf(!bHasExecuted, TEXT("Render graph SRV %s needs to be created before the builder execution."), Desc.Buffer->Name);
 		}
 		#endif
 		
 		FRDGBufferSRV* SRV = AllocateForRHILifeTime<FRDGBufferSRV>(Desc.Buffer->Name, Desc);
-		#if RENDER_GRAPH_DEBUGGING
+		#if RDG_ENABLE_DEBUG
 			Resources.Add(SRV);
 		#endif
 		return SRV;
@@ -363,7 +270,7 @@ public:
 	FORCEINLINE_DEBUGGABLE FRDGTextureUAVRef CreateUAV(const FRDGTextureUAVDesc& Desc)
 	{
 		check(Desc.Texture);
-		#if RENDER_GRAPH_DEBUGGING
+		#if RDG_ENABLE_DEBUG
 		{
 			ensureMsgf(!bHasExecuted, TEXT("Render graph UAV %s needs to be created before the builder execution."), Desc.Texture->Name);
 			ensureMsgf(Desc.Texture->Desc.TargetableFlags & TexCreate_UAV, TEXT("Attempted to create UAV from texture %s which was not created with TexCreate_UAV"), Desc.Texture->Name);
@@ -371,7 +278,7 @@ public:
 		#endif
 		
 		FRDGTextureUAV* UAV = AllocateForRHILifeTime<FRDGTextureUAV>(Desc.Texture->Name, Desc);
-		#if RENDER_GRAPH_DEBUGGING
+		#if RDG_ENABLE_DEBUG
 			Resources.Add(UAV);
 		#endif
 		return UAV;
@@ -386,14 +293,14 @@ public:
 	FORCEINLINE_DEBUGGABLE FRDGBufferUAVRef CreateUAV(const FRDGBufferUAVDesc& Desc)
 	{
 		check(Desc.Buffer);
-		#if RENDER_GRAPH_DEBUGGING
+		#if RDG_ENABLE_DEBUG
 		{
 			ensureMsgf(!bHasExecuted, TEXT("Render graph UAV %s needs to be created before the builder execution."), Desc.Buffer->Name);
 		}
 		#endif
 		
 		FRDGBufferUAV* UAV = AllocateForRHILifeTime<FRDGBufferUAV>(Desc.Buffer->Name, Desc);
-		#if RENDER_GRAPH_DEBUGGING
+		#if RDG_ENABLE_DEBUG
 			Resources.Add(UAV);
 		#endif
 		return UAV;
@@ -404,27 +311,6 @@ public:
 		return CreateUAV( FRDGBufferUAVDesc(Buffer, Format) );
 	}
 
-	/**
-	 * Allocates an object that will have the same lifetime as this graph builder.
-	 * Destructor will be automatically called for this object when the graph builder itself is destroyed.
-	 */
-	template<typename T, typename... TArgs>
-	FORCEINLINE_DEBUGGABLE T* AllocObject(TArgs&&... Args)
-	{
-		// Owned objects must be deleted in stack order, therefore destruction stack is pushed before object is constructed.
-		// This ensures that any child objects are placed after the parents in the OwnedObjects array.
-		int32 Index = OwnedObjects.AddDefaulted();
-
-		T* Result = new (MemStack) T(Forward<TArgs>(Args)...);
-
-		FOwnedObjectPtr OwnedObject;
-		OwnedObject.Ptr = Result;
-		OwnedObject.Deleter = [](void* InPtr) { ((T*)InPtr)->~T(); };
-		OwnedObjects[Index] = OwnedObject;
-
-		return Result;
-	}
-
 	/** Allocates parameter struct specifically to survive through the life time of the render graph. */
 	template< typename ParameterStructType >
 	FORCEINLINE_DEBUGGABLE ParameterStructType* AllocParameters()
@@ -432,7 +318,7 @@ public:
 		// TODO(RDG): could allocate using AllocateForRHILifeTime() to avoid the copy done when using FRHICommandList::BuildLocalUniformBuffer()
 		ParameterStructType* OutParameterPtr = new(MemStack) ParameterStructType;
 		FMemory::Memzero(OutParameterPtr, sizeof(ParameterStructType));
-		#if RENDER_GRAPH_DEBUGGING
+		#if RDG_ENABLE_DEBUG
 		{
 			AllocatedUnusedPassParameters.Add(static_cast<void *>(OutParameterPtr));
 		}
@@ -453,40 +339,16 @@ public:
 	void AddPass(
 		FRDGEventName&& Name, 
 		ParameterStructType* ParameterStruct,
-		ERenderGraphPassFlags Flags,
+		ERDGPassFlags Flags,
 		ExecuteLambdaType&& ExecuteLambda)
 	{
-		#if RENDER_GRAPH_DEBUGGING
-		{
-			checkf(!bHasExecuted, TEXT("Render graph pass %s needs to be added before the builder execution."), Name.GetTCHAR());
-
-			/** A pass parameter structure requires a correct life time until the pass execution, and therefor needs to be
-			 * allocated with FRDGBuilder::AllocParameters().
-			 *
-			 * Moreover, because the destructor of this parameter structure will be done after the pass execution, a it can
-			 * only be used by a single AddPass().
-			 */
-			checkf(
-				AllocatedUnusedPassParameters.Contains(static_cast<void *>(ParameterStruct)),
-				TEXT("The pass parameter structure has not been allocated for correct life time FRDGBuilder::AllocParameters() or has already ")
-				TEXT("been used by another previous FRDGBuilder::AddPass()."));
-
-			AllocatedUnusedPassParameters.Remove(static_cast<void *>(ParameterStruct));
-		}
-		#endif
-
-		auto NewPass = new(MemStack) TLambdaRenderPass<ParameterStructType, ExecuteLambdaType>(
-			static_cast<FRDGEventName&&>(Name), CurrentScope,
+		auto NewPass = new(MemStack) TRDGLambdaPass<ParameterStructType, ExecuteLambdaType>(
+			static_cast<FRDGEventName&&>(Name),
 			{ ParameterStruct, &ParameterStructType::FTypeInfo::GetStructMetadata()->GetLayout() },
 			Flags,
 			static_cast<ExecuteLambdaType&&>(ExecuteLambda) );
-		Passes.Emplace(NewPass);
 
-		#if RENDER_GRAPH_DEBUGGING || SUPPORTS_VISUALIZE_TEXTURE
-		{
-			DebugPass(NewPass);
-		}
-		#endif
+		AddPassInternal(NewPass);
 	}
 
 	/** Adds a procedurally created pass to the render graph.
@@ -494,24 +356,11 @@ public:
 	 * Note: You want to use this only when the layout of the pass might be procedurally generated from data driven, as opose to AddPass() that have,
 	 * constant hard coded pass layout.
 	 *
-	 * Caution: You are on your own to have correct memory lifetime of the FRenderGraphPass.
+	 * Caution: You are on your own to have correct memory lifetime of the FRDGPass.
 	 */
-	FORCEINLINE_DEBUGGABLE void AddProcedurallyCreatedPass(FRenderGraphPass* NewPass)
+	FORCEINLINE_DEBUGGABLE void AddProcedurallyCreatedPass(FRDGPass* NewPass)
 	{
-		#if RENDER_GRAPH_DEBUGGING
-		{
-			checkf(!bHasExecuted, TEXT("Render graph pass %s needs to be added before the builder execution."), NewPass->GetName());
-		}
-		#endif
-
-		// TODO(RDG): perhaps the CurrentScope could be set here instead of GetCurrentScope(), to not allow user code to start adding pass to random scopes. 
-		Passes.Emplace(NewPass);
-
-		#if RENDER_GRAPH_DEBUGGING || SUPPORTS_VISUALIZE_TEXTURE
-		{
-			DebugPass(NewPass);
-		}
-		#endif
+		AddPassInternal(NewPass);
 	}
 
 	/** Queue a texture extraction. This will set *OutTexturePtr with the internal pooled render target at the Execute().
@@ -523,7 +372,7 @@ public:
 	{
 		check(Texture);
 		check(OutTexturePtr);
-		#if RENDER_GRAPH_DEBUGGING
+		#if RDG_ENABLE_DEBUG
 		{
 			checkf(!bHasExecuted,
 				TEXT("Accessing render graph internal texture %s with QueueTextureExtraction() needs to happen before the builder's execution."),
@@ -550,7 +399,7 @@ public:
 	{
 		check(Buffer);
 		check(OutBufferPtr);
-#if RENDER_GRAPH_DEBUGGING
+#if RDG_ENABLE_DEBUG
 		{
 			checkf(!bHasExecuted,
 				TEXT("Accessing render graph internal buffer %s with QueueBufferExtraction() needs to happen before the builder's execution."),
@@ -571,7 +420,7 @@ public:
 	FORCEINLINE_DEBUGGABLE void RemoveUnusedTextureWarning(FRDGTextureRef Texture)
 	{
 		check(Texture);
-		#if RENDER_GRAPH_DEBUGGING
+		#if RDG_ENABLE_DEBUG
 		{
 			checkf(!bHasExecuted,
 				TEXT("Flaging texture %s with FlagUnusedTexture() needs to happen before the builder's execution."),
@@ -589,24 +438,15 @@ public:
 	 */
 	void Execute();
 
-	/** Returns the draw event scope, where passes are currently being added in. */
-	const FRDGEventScope* GetCurrentScope() const
-	{
-		return CurrentScope;
-	}
-
 	/** The RHI command list used for the render graph. */
 	FRHICommandListImmediate& RHICmdList;
 
 	/** Memory stack use for allocating RDG resource and passes. */
 	FMemStackBase& MemStack;
 
-
 private:
-	static constexpr int32 kMaxScopeCount = 8;
-
 	/** Array of all pass created */
-	TArray<FRenderGraphPass*, SceneRenderingAllocator> Passes;
+	TArray<FRDGPass*, SceneRenderingAllocator> Passes;
 
 	/** Keep the references over the pooled render target, since FRDGTexture is allocated on MemStack. */
 	TMap<FRDGTexture*, TRefCountPtr<IPooledRenderTarget>, SceneRenderingSetAllocator> AllocatedTextures;
@@ -631,18 +471,9 @@ private:
 	};
 	TArray<FDeferredInternalBufferQuery, SceneRenderingAllocator> DeferredInternalBufferQueries;
 
-	#if RENDER_GRAPH_DRAW_EVENTS == 2
-		/** All scopes allocated that needs to be arround to call destructors. */
-		TArray<FRDGEventScope*, SceneRenderingAllocator> EventScopes;
-	#endif
-
-	/** The current event scope as creating passes. */
-	const FRDGEventScope* CurrentScope = nullptr;
+	FRDGEventScopeStack EventScopeStack;
 
-	/** Stacks of scopes pushed to the RHI command list. */
-	TStaticArray<const FRDGEventScope*, kMaxScopeCount> ScopesStack;
-
-	#if RENDER_GRAPH_DEBUGGING
+	#if RDG_ENABLE_DEBUG
 		/** Whether the Execute() has already been called. */
 		bool bHasExecuted = false;
 
@@ -650,21 +481,12 @@ private:
 		TArray<const FRDGResource*, SceneRenderingAllocator> Resources;
 
 		// All recently allocated pass parameter structure, but not used by a AddPass() yet.
-		TSet<void*> AllocatedUnusedPassParameters;
+		TSet<const void*> AllocatedUnusedPassParameters;
 	#endif
 
-	/** Holds objects that are allocated and owned by this graph builder. Allow calling correct destructors at shutdown. */
-	struct FOwnedObjectPtr
-	{
-		typedef void(*DeleterCallback)(void*);
-		void* Ptr = nullptr;
-		DeleterCallback Deleter = nullptr;
-	};
-	TArray<FOwnedObjectPtr, SceneRenderingAllocator> OwnedObjects;
-
-	void DebugPass(const FRenderGraphPass* Pass);
-	void ValidatePass(const FRenderGraphPass* Pass) const;
-	void CaptureAnyInterestingPassOutput(const FRenderGraphPass* Pass);
+	void DebugPass(const FRDGPass* Pass);
+	void ValidatePass(const FRDGPass* Pass) const;
+	void CaptureAnyInterestingPassOutput(const FRDGPass* Pass);
 
 	void WalkGraphDependencies();
 	
@@ -685,97 +507,35 @@ private:
 		}
 	}
 
+	void AddPassInternal(FRDGPass* Pass);
+
 	void AllocateRHITextureIfNeeded(FRDGTexture* Texture, bool bComputePass);
 	void AllocateRHITextureUAVIfNeeded(FRDGTextureUAV* UAV, bool bComputePass);
 	void AllocateRHIBufferSRVIfNeeded(FRDGBufferSRV* SRV, bool bComputePass);
 	void AllocateRHIBufferUAVIfNeeded(FRDGBufferUAV* UAV, bool bComputePass);
 
+	void BeginEventScope(FRDGEventName&& ScopeName);
+	void EndEventScope();
 
 	void TransitionTexture(FRDGTexture* Texture, EResourceTransitionAccess TransitionAccess, bool bRequiredCompute) const;
 	void TransitionUAV(FUnorderedAccessViewRHIParamRef UAV, FRDGResource* UnderlyingResource, ERDGResourceFlags ResourceFlags, EResourceTransitionAccess TransitionAccess, bool bRequiredCompute ) const;
 
-	void PushDrawEventStack(const FRenderGraphPass* Pass);
-	void ExecutePass( const FRenderGraphPass* Pass );
-	void AllocateAndTransitionPassResources(const FRenderGraphPass* Pass, struct FRHIRenderPassInfo* OutRPInfo, bool* bOutHasRenderTargets);
-	static void UpdateAccessGuardForPassResources(const FRenderGraphPass* Pass, bool bAllowAccess);
-	static void WarnForUselessPassDependencies(const FRenderGraphPass* Pass);
+	void PushDrawEventStack(const FRDGPass* Pass);
+	void ExecutePass(const FRDGPass* Pass);
+	void AllocateAndTransitionPassResources(const FRDGPass* Pass, struct FRHIRenderPassInfo* OutRPInfo, bool* bOutHasRenderTargets);
+	static void UpdateAccessGuardForPassResources(const FRDGPass* Pass, bool bAllowAccess);
+	static void WarnForUselessPassDependencies(const FRDGPass* Pass);
 
 	void ReleaseRHITextureIfPossible(FRDGTexture* Texture);
 	void ReleaseRHIBufferIfPossible(FRDGBuffer* Buffer);
-	void ReleaseUnecessaryResources(const FRenderGraphPass* Pass);
+	void ReleaseUnecessaryResources(const FRDGPass* Pass);
 
 	void ProcessDeferredInternalResourceQueries();
 	void DestructPasses();
 
-	friend class FStackRDGEventScopeRef;
+	friend class FRDGEventScopeGuard;
 
 	/** To allow greater flexibility in the user code, the RHI can dereferenced RDG resource when creating uniform buffer. */
 	// TODO(RDG): Make this a little more explicit in RHI code.
 	static_assert(STRUCT_OFFSET(FRDGResource, CachedRHI) == 0, "FRDGResource::CachedRHI requires to be at offset 0 so the RHI can dereferenced them.");
-}; // class FRDGBuilder
-
-
-#if RENDER_GRAPH_DRAW_EVENTS
-
-/** Stack reference of render graph scope. */
-class RENDERCORE_API FStackRDGEventScopeRef
-{
-public:
-	FStackRDGEventScopeRef() = delete;
-	FStackRDGEventScopeRef(const FStackRDGEventScopeRef&) = delete;
-	FStackRDGEventScopeRef(FStackRDGEventScopeRef&&) = delete;
-	void operator = (const FStackRDGEventScopeRef&) = delete;
-
-	FORCEINLINE_DEBUGGABLE FStackRDGEventScopeRef(FRDGBuilder& InGraphBuilder, FRDGEventName&& ScopeName)
-		: GraphBuilder(InGraphBuilder)
-	{
-		checkf(!GraphBuilder.bHasExecuted, TEXT("Render graph bulider has already been executed."));
-
-		auto NewScope = new(GraphBuilder.MemStack) FRDGEventScope(GraphBuilder.CurrentScope, Forward<FRDGEventName>(ScopeName));
-
-		#if RENDER_GRAPH_DRAW_EVENTS == RENDER_GRAPH_DRAW_EVENTS_STRING_COPY
-		{
-			GraphBuilder.EventScopes.Add(NewScope);
-		}
-		#endif
-
-		GraphBuilder.CurrentScope = NewScope;
-	}
-
-	FORCEINLINE_DEBUGGABLE ~FStackRDGEventScopeRef()
-	{
-		check(GraphBuilder.CurrentScope != nullptr);
-		GraphBuilder.CurrentScope = GraphBuilder.CurrentScope->ParentScope;
-	}
-
-private:
-	FRDGBuilder& GraphBuilder;
-};
-
-#endif // RENDER_GRAPH_DRAW_EVENTS
-
-
-/** Macros for create render graph event names and scopes.
- *
- *		FRDGEventName Name = RDG_EVENT_NAME("MyPass %sx%s", ViewRect.Width(), ViewRect.Height());
- *
- *		RDG_EVENT_SCOPE(GraphBuilder, "MyProcessing %sx%s", ViewRect.Width(), ViewRect.Height());
- */
-#if RENDER_GRAPH_DRAW_EVENTS == RENDER_GRAPH_DRAW_EVENTS_STRING_COPY
-
-#define RDG_EVENT_NAME(Format, ...) FRDGEventName(TEXT(Format), ##__VA_ARGS__)
-#define RDG_EVENT_SCOPE(GraphBuilder, Format, ...) \
-	FStackRDGEventScopeRef PREPROCESSOR_JOIN(__RDG_ScopeRef_,__LINE__) ((GraphBuilder), RDG_EVENT_NAME(Format, ##__VA_ARGS__))
-
-#elif RENDER_GRAPH_DRAW_EVENTS == RENDER_GRAPH_DRAW_EVENTS_STRING_COPY
-
-#define RDG_EVENT_NAME(Format, ...) FRDGEventName(TEXT(Format))
-#define RDG_EVENT_SCOPE(GraphBuilder, Format, ...) \
-	FStackRDGEventScopeRef PREPROCESSOR_JOIN(__RDG_ScopeRef_,__LINE__) ((GraphBuilder), RDG_EVENT_NAME(Format, ##__VA_ARGS__))
-
-#else // !RENDER_GRAPH_DRAW_EVENTS
-
-#define RDG_EVENT_NAME(Format, ...) FRDGEventName()
-#define RDG_EVENT_SCOPE(GraphBuilder, Format, ...) 
-
-#endif // !RENDER_GRAPH_DRAW_EVENTS
+};
\ No newline at end of file
diff --git a/Engine/Source/Runtime/RenderCore/Public/RenderGraphDefinitions.h b/Engine/Source/Runtime/RenderCore/Public/RenderGraphDefinitions.h
new file mode 100644
--- /dev/null
+++ b/Engine/Source/Runtime/RenderCore/Public/RenderGraphDefinitions.h
@@ -0,0 +1,29 @@
+// Copyright 1998-2019 Epic Games, Inc. All Rights Reserved.
+
+#pragma once
+
+#include "CoreMinimal.h"
+#include "RHI.h"
+#include "RendererInterface.h"
+
+/** Whether render graph debugging is enabled. */
+#define RDG_ENABLE_DEBUG (!UE_BUILD_SHIPPING && !UE_BUILD_TEST)
+
+/** Whether render graph debugging is enabled and we are compiling with the engine. */
+#define RDG_ENABLE_DEBUG_WITH_ENGINE (RDG_ENABLE_DEBUG && WITH_ENGINE)
+
+/** The type of GPU events the render graph system supports.
+ *  RDG_EVENTS == 0 means there is no string processing at all.
+ *  RDG_EVENTS == 1 means the format component of the event name is stored as a const TCHAR*.
+ *  RDG_EVENTS == 2 means string formatting is evaluated and stored in an FString.
+ */
+#define RDG_EVENTS_NONE 0
+#define RDG_EVENTS_STRING_REF 1
+#define RDG_EVENTS_STRING_COPY 2
+
+/** Whether render graph GPU events are enabled. */
+#if WITH_PROFILEGPU
+	#define RDG_EVENTS RDG_EVENTS_STRING_COPY
+#else
+	#define RDG_EVENTS RDG_EVENTS_NONE
+#endif
\ No newline at end of file
diff --git a/Engine/Source/Runtime/RenderCore/Public/RenderGraphEvent.h b/Engine/Source/Runtime/RenderCore/Public/RenderGraphEvent.h
new file mode 100644
--- /dev/null
+++ b/Engine/Source/Runtime/RenderCore/Public/RenderGraphEvent.h
@@ -0,0 +1,168 @@
+// Copyright 1998-2019 Epic Games, Inc. All Rights Reserved.
+
+#pragma once
+
+#include "RenderGraphDefinitions.h"
+
+class FRDGPass;
+class FRDGBuilder;
+class FRDGEventScope;
+
+/** Stores a GPU event name for the render graph. Draw events can be compiled out entirely from
+ *  a release build for performance.
+ */
+class RENDERCORE_API FRDGEventName final
+{
+public:
+	FRDGEventName() = default;
+
+	~FRDGEventName()
+	{
+#if RDG_EVENTS != RDG_EVENTS_NONE
+		EventName = nullptr;
+#endif
+	}
+
+#if RDG_EVENTS == RDG_EVENTS_STRING_COPY
+	explicit FRDGEventName(const TCHAR* EventFormat, ...);
+#else
+	explicit FRDGEventName(const TCHAR* EventFormat, ...)
+#if RDG_EVENTS == RDG_EVENTS_STRING_REF
+		: EventName(EventFormat)
+#endif
+	{}
+#endif
+
+	FRDGEventName(const FRDGEventName& Other)
+	{
+		*this = Other;
+	}
+
+	FRDGEventName(FRDGEventName&& Other)
+	{
+		*this = Forward<FRDGEventName>(Other);
+	}
+
+	FRDGEventName& operator=(const FRDGEventName& Other)
+	{
+#if RDG_EVENTS == RDG_EVENTS_STRING_REF
+		EventName = Other.EventName;
+#elif RDG_EVENTS == RDG_EVENTS_STRING_COPY
+		EventNameStorage = Other.EventNameStorage;
+		EventName = *EventNameStorage;
+#endif
+		return *this;
+	}
+
+	FRDGEventName& operator=(FRDGEventName&& Other)
+	{
+#if RDG_EVENTS == RDG_EVENTS_STRING_REF
+		EventName = Other.EventName;
+		Other.EventName = nullptr;
+#elif RDG_EVENTS == RDG_EVENTS_STRING_COPY
+		EventNameStorage = MoveTemp(Other.EventNameStorage);
+		EventName = *EventNameStorage;
+		Other.EventName = nullptr;
+#endif
+		return *this;
+	}
+
+	const TCHAR* GetTCHAR() const
+	{
+#if RDG_EVENTS == RDG_EVENTS_STRING_REF || RDG_EVENTS == RDG_EVENTS_STRING_COPY
+		return EventName;
+#else
+		// Render graph draw events have been completely compiled for CPU performance reasons.
+		return TEXT("!!!Unavailable RDG event name: need RDG_EVENTS>=0 and r.RDG.EmitWarnings=1 or -rdgdebug!!!");
+#endif
+	}
+
+private:
+#if RDG_EVENTS == RDG_EVENTS_STRING_REF || RDG_EVENTS == RDG_EVENTS_STRING_COPY
+	const TCHAR* EventName;
+#if RDG_EVENTS == RDG_EVENTS_STRING_COPY
+	FString EventNameStorage;
+#endif
+#endif
+};
+
+/** Manages a stack of event scopes. Scopes are recorded ahead of time in a hierarchical fashion
+ *  and later executed topologically during pass execution.
+ */
+class RENDERCORE_API FRDGEventScopeStack final
+{
+public:
+	FRDGEventScopeStack(FRHICommandListImmediate& RHICmdList);
+	~FRDGEventScopeStack();
+
+	//////////////////////////////////////////////////////////////////////////
+	//! Called during graph setup phase.
+
+	void BeginEventScope(FRDGEventName&& EventName);
+
+	void EndEventScope();
+
+	//////////////////////////////////////////////////////////////////////////
+	//! Called during graph execute phase.
+
+	void BeginExecute();
+
+	void BeginExecutePass(const FRDGPass* Pass);
+
+	void EndExecutePass();
+
+	void EndExecute();
+
+	//////////////////////////////////////////////////////////////////////////
+
+	const FRDGEventScope* GetCurrentScope() const
+	{
+		return CurrentScope;
+	}
+
+private:
+	static constexpr uint32 kScopeStackDepthMax = 8;
+
+	FRHICommandListImmediate& RHICmdList;
+	FMemStackBase& MemStack;
+
+	/** The top of the scope stack during setup. */
+	const FRDGEventScope* CurrentScope = nullptr;
+
+	/** Tracks scopes allocated through MemStack for destruction. */
+	TArray<FRDGEventScope*, SceneRenderingAllocator> EventScopes;
+
+	/** Stacks of scopes pushed to the RHI command list during execution. */
+	TStaticArray<const FRDGEventScope*, kScopeStackDepthMax> ScopeStack;
+
+	/** Used to validate Begin / End symmetry and that setup and execute are mutually exclusive. */
+	bool bIsExecuting = false;
+};
+
+/** RAII class for begin / end of an event scope through the graph builder. */
+class RENDERCORE_API FRDGEventScopeGuard final
+{
+public:
+	FRDGEventScopeGuard(FRDGBuilder& InGraphBuilder, FRDGEventName&& ScopeName);
+	FRDGEventScopeGuard(const FRDGEventScopeGuard&) = delete;
+	~FRDGEventScopeGuard();
+
+private:
+	FRDGBuilder& GraphBuilder;
+};
+
+/** Macros for create render graph event names and scopes.
+ *
+ *  FRDGEventName Name = RDG_EVENT_NAME("MyPass %sx%s", ViewRect.Width(), ViewRect.Height());
+ *
+ *  RDG_EVENT_SCOPE(GraphBuilder, "MyProcessing %sx%s", ViewRect.Width(), ViewRect.Height());
+ */
+#if RDG_EVENTS == RDG_EVENTS_STRING_REF || RDG_EVENTS == RDG_EVENTS_STRING_COPY
+	#define RDG_EVENT_NAME(Format, ...) FRDGEventName(TEXT(Format), ##__VA_ARGS__)
+	#define RDG_EVENT_SCOPE(GraphBuilder, Format, ...) FRDGEventScopeGuard PREPROCESSOR_JOIN(__RDG_ScopeRef_,__LINE__) ((GraphBuilder), RDG_EVENT_NAME(Format, ##__VA_ARGS__))
+#elif RDG_EVENTS == RDG_EVENTS_NONE
+	#define RDG_EVENT_NAME(Format, ...) FRDGEventName()
+	#define RDG_EVENT_SCOPE(GraphBuilder, Format, ...) 
+#else
+	#error "RDG_EVENTS is not a valid value."
+#endif
\ No newline at end of file
diff --git a/Engine/Source/Runtime/RenderCore/Public/RenderGraphResources.h b/Engine/Source/Runtime/RenderCore/Public/RenderGraphResources.h
--- a/Engine/Source/Runtime/RenderCore/Public/RenderGraphResources.h
+++ b/Engine/Source/Runtime/RenderCore/Public/RenderGraphResources.h
@@ -2,16 +2,9 @@
 
 #pragma once
 
-#include "CoreMinimal.h"
-#include "RHI.h"
-#include "RendererInterface.h"
-
-
-/** Whether render graph debugging is compiled. */
-#define RENDER_GRAPH_DEBUGGING (!UE_BUILD_SHIPPING && !UE_BUILD_TEST)
-
-struct FRenderGraphPass;
+#include "RenderGraphDefinitions.h"
 
+class FRDGPass;
 class FRDGBuilder;
 
 class FRDGResource;
@@ -75,7 +68,7 @@ public:
 	/** Verify that the RHI resource can be accessed at a pass execution. */
 	inline void VerifyRHIAccess() const
 	{
-		#if RENDER_GRAPH_DEBUGGING
+		#if RDG_ENABLE_DEBUG
 		{
 			checkf(bAllowAccessToRHIResource,
 				TEXT("Accessing the RHI resource of %s at this time is not allowed. If you hit this check in pass, ")
@@ -107,7 +100,7 @@ private:
 	/** Boolean to track at runtime whether a ressource is actually used by the lambda of a pass or not, to detect unnecessary resource dependencies on passes. */
 	bool bIsActuallyUsedByPass = false;
 
-#if RENDER_GRAPH_DEBUGGING
+#if RDG_ENABLE_DEBUG
 	/** Boolean to track at wiring time if a resource has ever been produced by a pass, to error out early if accessing a resource that has not been produced. */
 	bool bHasEverBeenProduced = false;
 
@@ -115,7 +108,7 @@ private:
 	bool bAllowAccessToRHIResource = false;
 
 	/** Pointer towards the pass that is the first to produce it, for even more convenient error message. */
-	const FRenderGraphPass* DebugFirstProducer = nullptr;
+	const FRDGPass* DebugFirstProducer = nullptr;
 
 	/** Count the number of times it has been used by a pass. */
 	int32 DebugPassAccessCount = 0;
diff --git a/Engine/Source/Runtime/RenderCore/Public/VisualizeTexture.h b/Engine/Source/Runtime/RenderCore/Public/VisualizeTexture.h
--- a/Engine/Source/Runtime/RenderCore/Public/VisualizeTexture.h
+++ b/Engine/Source/Runtime/RenderCore/Public/VisualizeTexture.h
@@ -11,6 +11,7 @@
 #include "RenderResource.h"
 #include "RenderGraph.h"
 
+#define SUPPORTS_VISUALIZE_TEXTURE (WITH_ENGINE && !(UE_BUILD_SHIPPING || UE_BUILD_TEST))
 
 class RENDERCORE_API FVisualizeTexture : public FRenderResource
 {
