Pixelshader UAV support in RDG

diff --git a/Engine/Source/Runtime/RenderCore/Private/RenderGraphBuilder.cpp b/Engine/Source/Runtime/RenderCore/Private/RenderGraphBuilder.cpp
--- a/Engine/Source/Runtime/RenderCore/Private/RenderGraphBuilder.cpp
+++ b/Engine/Source/Runtime/RenderCore/Private/RenderGraphBuilder.cpp
@@ -222,7 +222,7 @@ void FRDGBuilder::ValidatePass(const FRDGPass* Pass) const
 	FShaderParameterStructRef ParameterStruct = Pass->GetParameters();
 
 	bool bIsCompute = Pass->IsCompute();
-	bool bCanUseUAVs = bIsCompute;
+	bool bCanUseUAVs = true;	// Graphics passes can now also use UAVs in pixelshaders.
 	bool bRequiresRenderTargetSlots = !bIsCompute;
 
 	for (int ResourceIndex = 0, Num = ParameterStruct.Layout->Resources.Num(); ResourceIndex < Num; ResourceIndex++)
@@ -822,6 +822,9 @@ void FRDGBuilder::AllocateAndTransitionPassResources(const FRDGPass* Pass, struc
 	bool bIsCompute = Pass->IsCompute();
 	FShaderParameterStructRef ParameterStruct = Pass->GetParameters();
 
+	OutRPInfo->NumUAVs = 0;
+	OutRPInfo->UAVIndex = 0;
+
 	const bool bGeneratingMips = (Pass->GetFlags() & ERDGPassFlags::GenerateMips) == ERDGPassFlags::GenerateMips;
 	for (int ResourceIndex = 0, Num = ParameterStruct.Layout->Resources.Num(); ResourceIndex < Num; ResourceIndex++)
 	{
@@ -891,6 +894,11 @@ void FRDGBuilder::AllocateAndTransitionPassResources(const FRDGPass* Pass, struc
 				AllocateRHITextureUAVIfNeeded(UAV, bIsCompute);
 				TransitionUAV(UAV->CachedRHI.UAV, UAV->Desc.Texture, UAV->Desc.Texture->Flags, EResourceTransitionAccess::EWritable, bIsCompute);
 
+				if (!bIsCompute)
+				{
+					OutRPInfo->UAVs[OutRPInfo->NumUAVs++] = UAV->CachedRHI.UAV;	// Bind UAVs in declaration order
+				}
+
 				#if RDG_ENABLE_DEBUG
 				{
 					UAV->Desc.Texture->DebugPassAccessCount++;
@@ -962,6 +970,11 @@ void FRDGBuilder::AllocateAndTransitionPassResources(const FRDGPass* Pass, struc
 				AllocateRHIBufferUAVIfNeeded(UAV, bIsCompute);
 				TransitionUAV(UAV->CachedRHI.UAV, UAV->Desc.Buffer, UAV->Desc.Buffer->Flags, EResourceTransitionAccess::EWritable, bIsCompute);
 
+				if (!bIsCompute)
+				{
+					OutRPInfo->UAVs[OutRPInfo->NumUAVs++] = UAV->CachedRHI.UAV;	// Bind UAVs in declaration order
+				}
+				
 				#if RDG_ENABLE_DEBUG
 				{
 					UAV->Desc.Buffer->DebugPassAccessCount++;
@@ -1015,6 +1028,8 @@ void FRDGBuilder::AllocateAndTransitionPassResources(const FRDGPass* Pass, struc
 				}
 			}
 
+			OutRPInfo->UAVIndex = NumRenderTargets;
+
 			const FDepthStencilBinding& DepthStencil = RenderTargets->DepthStencil;
 			if (DepthStencil.GetTexture())
 			{
diff --git a/Engine/Source/Runtime/RenderCore/Public/ShaderParameterStruct.h b/Engine/Source/Runtime/RenderCore/Public/ShaderParameterStruct.h
--- a/Engine/Source/Runtime/RenderCore/Public/ShaderParameterStruct.h
+++ b/Engine/Source/Runtime/RenderCore/Public/ShaderParameterStruct.h
@@ -83,7 +83,13 @@ inline void SetShaderUAVs(TRHICmdList& RHICmdList, const TShaderClass* Shader, T
 {
 	checkf(
 		Shader->Bindings.UAVs.Num() == 0 && Shader->Bindings.GraphUAVs.Num() == 0,
-		TEXT("TShaderRHI Can't have compute shader to be set. UAVs are not supported on vertex, tessellation, geometry and pixel shaders."));
+		TEXT("TShaderRHI Can't have compute shader to be set. UAVs are not supported on vertex, tessellation and geometry shaders."));
+}
+
+template<typename TRHICmdList, typename TShaderClass>
+inline void SetShaderUAVs(TRHICmdList& RHICmdList, const TShaderClass* Shader, FPixelShaderRHIParamRef ShadeRHI, const typename TShaderClass::FParameters& Parameters)
+{
+	// Pixelshader UAVs are bound together with rendertargets using BeginRenderPass
 }
 
 template<typename TRHICmdList, typename TShaderClass>
