Removed 'bIsExecuting' validation from FRDGEventScopeStack because it doesn't work in immediate mode.

diff --git a/Engine/Source/Runtime/RenderCore/Private/RenderGraphEvent.cpp b/Engine/Source/Runtime/RenderCore/Private/RenderGraphEvent.cpp
--- a/Engine/Source/Runtime/RenderCore/Private/RenderGraphEvent.cpp
+++ b/Engine/Source/Runtime/RenderCore/Private/RenderGraphEvent.cpp
@@ -63,15 +63,8 @@ FRDGEventScopeStack::FRDGEventScopeStack(FRHICommandListImmediate& InRHICmdList)
 	, ScopeStack(MakeUniformStaticArray<const FRDGEventScope*, kScopeStackDepthMax>(nullptr))
 {}
 
-FRDGEventScopeStack::~FRDGEventScopeStack()
-{
-	check(!bIsExecuting);
-}
-
 void FRDGEventScopeStack::BeginEventScope(FRDGEventName&& EventName)
 {
-	checkf(!bIsExecuting, TEXT("Can't begin an event scope while executing."));
-
 	if (RDG_EVENTS)
 	{
 		auto Scope = new(MemStack) FRDGEventScope(CurrentScope, Forward<FRDGEventName>(EventName));
@@ -82,8 +75,6 @@ void FRDGEventScopeStack::BeginEventScope(FRDGEventName&& EventName)
 
 void FRDGEventScopeStack::EndEventScope()
 {
-	checkf(!bIsExecuting, TEXT("Can't end an event scope while executing."));
-
 	if (RDG_EVENTS)
 	{
 		checkf(CurrentScope != nullptr, TEXT("Current scope is null."));
@@ -105,14 +96,10 @@ void FRDGEventScopeStack::BeginExecute()
 	 *  GraphBuilder.Execute();
 	 */
 	checkf(CurrentScope == nullptr, TEXT("Render graph needs to have all scopes ended to execute."));
-	check(!bIsExecuting);
-	bIsExecuting = true;
 }
 
 void FRDGEventScopeStack::BeginExecutePass(const FRDGPass* Pass)
 {
-	check(bIsExecuting);
-
 	if (!GetEmitRDGEvents())
 	{
 		return;
@@ -190,8 +177,6 @@ void FRDGEventScopeStack::BeginExecutePass(const FRDGPass* Pass)
 
 void FRDGEventScopeStack::EndExecutePass()
 {
-	check(bIsExecuting);
-
 	if (GetEmitRDGEvents())
 	{
 		RHICmdList.PopEvent();
@@ -200,9 +185,6 @@ void FRDGEventScopeStack::EndExecutePass()
 
 void FRDGEventScopeStack::EndExecute()
 {
-	check(bIsExecuting);
-	bIsExecuting = false;
-
 	// Pops remaining scopes
 	if (GetEmitRDGEvents())
 	{
diff --git a/Engine/Source/Runtime/RenderCore/Public/RenderGraphEvent.h b/Engine/Source/Runtime/RenderCore/Public/RenderGraphEvent.h
--- a/Engine/Source/Runtime/RenderCore/Public/RenderGraphEvent.h
+++ b/Engine/Source/Runtime/RenderCore/Public/RenderGraphEvent.h
@@ -93,7 +93,6 @@ class RENDERCORE_API FRDGEventScopeStack final
 {
 public:
 	FRDGEventScopeStack(FRHICommandListImmediate& RHICmdList);
-	~FRDGEventScopeStack();
 
 	//////////////////////////////////////////////////////////////////////////
 	//! Called during graph setup phase.
@@ -134,9 +133,6 @@ private:
 
 	/** Stacks of scopes pushed to the RHI command list during execution. */
 	TStaticArray<const FRDGEventScope*, kScopeStackDepthMax> ScopeStack;
-
-	/** Used to validate Begin / End symmetry and that setup and execute are mutually exclusive. */
-	bool bIsExecuting = false;
 };
 
 /** RAII class for begin / end of an event scope through the graph builder. */
