Moved TracingProfiler to RenderCore so that RenderGraph can access it.

diff --git a/Engine/Source/Runtime/Engine/Private/ProfilingDebugging/TracingProfiler.cpp b/Engine/Source/Runtime/Engine/Private/ProfilingDebugging/TracingProfiler.cpp
deleted file mode 100644
diff --git a/Engine/Source/Runtime/Engine/Public/ProfilingDebugging/TracingProfiler.h b/Engine/Source/Runtime/Engine/Public/ProfilingDebugging/TracingProfiler.h
deleted file mode 100644
diff --git a/Engine/Source/Runtime/RenderCore/Private/ProfilingDebugging/TracingProfiler.cpp b/Engine/Source/Runtime/RenderCore/Private/ProfilingDebugging/TracingProfiler.cpp
new file mode 100644
--- /dev/null
+++ b/Engine/Source/Runtime/RenderCore/Private/ProfilingDebugging/TracingProfiler.cpp
@@ -0,0 +1,366 @@
+// Copyright 1998-2019 Epic Games, Inc. All Rights Reserved.
+
+#include "ProfilingDebugging/TracingProfiler.h"
+
+#if TRACING_PROFILER
+
+#include "GPUProfiler.h"
+#include "HAL/ConsoleManager.h"
+#include "HAL/FileManager.h"
+#include "HAL/LowLevelMemTracker.h"
+#include "Misc/CoreDelegates.h"
+#include "Misc/Paths.h"
+#include "RHI.h"
+#include "RenderingThread.h"
+
+static TAutoConsoleVariable<int32> GTracingProfileBufferSize(
+	TEXT("TracingProfiler.BufferSize"),
+	65536,
+	TEXT("Defines the maximum umber of events stored in the internal ring buffer of the tracing profiler. ")
+	TEXT("Only read at process startup and can't be changed at runtime."),
+	ECVF_Default);
+
+FTracingProfiler* FTracingProfiler::Instance = nullptr;
+
+FTracingProfiler* FTracingProfiler::Get()
+{
+	if (Instance == nullptr)
+	{
+		Instance = new FTracingProfiler;
+	}
+	return Instance;
+}
+
+static void HandleTracingProfileCommand(const TArray<FString>& Args)
+{
+	if (Args.Num() < 1)
+	{
+		return;
+	}
+	FString Param = Args[0];
+	if (Param == TEXT("START"))
+	{
+		FTracingProfiler::Get()->BeginCapture();
+	}
+	else if (Param == TEXT("STOP"))
+	{
+		FTracingProfiler::Get()->EndCapture();
+	}
+	else
+	{
+		int32 CaptureFrames = 0;
+		if (FParse::Value(*Param, TEXT("FRAMES="), CaptureFrames))
+		{
+			FTracingProfiler::Get()->BeginCapture(CaptureFrames);
+		}
+	}
+}
+
+static FAutoConsoleCommand HandleTracingProfileCmd(
+	TEXT("TracingProfile"),
+	TEXT("Starts or stops tracing profiler"),
+	FConsoleCommandWithArgsDelegate::CreateStatic(&HandleTracingProfileCommand)
+);
+
+FTracingProfiler::FTracingProfiler()
+{
+	LLM_SCOPE(ELLMTag::Stats);
+
+	MaxNumCapturedEvents = GTracingProfileBufferSize.GetValueOnAnyThread();
+
+	FEvent DefaultEvent = {};
+	CapturedEvents.Init(DefaultEvent, MaxNumCapturedEvents);
+
+	FCoreDelegates::OnBeginFrame.AddLambda([]() { FTracingProfiler::Get()->BeginFrame(); });
+	FCoreDelegates::OnEndFrame.AddLambda([]() { FTracingProfiler::Get()->EndFrame(); });
+
+	FCoreDelegates::OnBeginFrameRT.AddLambda([]() { FTracingProfiler::Get()->BeginFrameRT(); });
+	FCoreDelegates::OnEndFrameRT.AddLambda([]() { FTracingProfiler::Get()->EndFrameRT(); });
+}
+
+void FTracingProfiler::Init()
+{
+	
+}
+
+void FTracingProfiler::BeginFrame()
+{
+	check(IsInGameThread());
+
+	if (bRequestStartCapture)
+	{
+		bCapturing = true;
+		bRequestStartCapture = false;
+	}
+
+	GameThreadFrameBeginCycle = FPlatformTime::Cycles64();
+}
+
+static inline uint64 CyclesToMicroseconds64(uint64 Cycles)
+{
+	return uint64(FPlatformTime::ToSeconds64(Cycles) * 1e6);
+}
+
+void FTracingProfiler::EndFrame()
+{
+	check(IsInGameThread());
+
+	if (!bCapturing)
+	{
+		return;
+	}
+
+	GameThreadFrameEndCycle = FPlatformTime::Cycles64();
+
+	AddCPUEvent("GameThread",
+		GameThreadFrameBeginCycle,
+		GameThreadFrameEndCycle,
+		GGameThreadId,
+		GFrameNumber);
+
+	if (NumFramesToCapture >= 0)
+	{
+		NumFramesToCapture--;
+		if (NumFramesToCapture == 0)
+		{
+			bRequestStopCapture = true;
+		}
+	}
+
+	if (bRequestStopCapture)
+	{
+		FlushRenderingCommands();
+
+		bCapturing = false;
+		bRequestStopCapture = false;
+
+		WriteCaptureToFile();
+
+	}
+
+	CaptureFrameNumber++;
+}
+
+void FTracingProfiler::BeginFrameRT()
+{
+	RenderThreadFrameBeginCycle = FPlatformTime::Cycles64();
+
+	if (!bCapturingRT && bCapturing)
+	{
+		GDynamicRHI->RHICalibrateTimers();
+	}
+
+	bCapturingRT = bCapturing;
+}
+
+void FTracingProfiler::EndFrameRT()
+{
+	RenderThreadFrameEndCycle = FPlatformTime::Cycles64();
+
+	if (!bCapturingRT)
+	{
+		return;
+	}
+
+	AddCPUEvent("RenderThread",
+		RenderThreadFrameBeginCycle,
+		RenderThreadFrameEndCycle,
+		GRenderThreadId,
+		GFrameNumberRenderThread);
+}
+
+
+int32 FTracingProfiler::GetCaptureFrameNumber()
+{
+	return CaptureFrameNumber;
+}
+
+void FTracingProfiler::BeginCapture(int InNumFramesToCapture)
+{
+	check(IsInGameThread());
+
+	NumFramesToCapture = InNumFramesToCapture;
+	CaptureFrameNumber = 0;
+
+	bRequestStartCapture = true;
+
+	SessionId++;
+}
+
+void FTracingProfiler::EndCapture()
+{
+	bRequestStopCapture = true;
+}
+
+bool FTracingProfiler::IsCapturing() const
+{
+	check(IsInGameThread());
+	return bRequestStartCapture || bCapturing;
+}
+
+void FTracingProfiler::WriteCaptureToFile()
+{
+	check(!bCapturing);
+
+	FString Filename = FString::Printf(TEXT("Profile(%s)"), *FDateTime::Now().ToString(TEXT("%Y%m%d_%H%M%S")));
+	FString TracingRootPath = FPaths::ProfilingDir() + TEXT("Traces/");
+	FString OutputFilename = TracingRootPath + Filename + TEXT(".json");
+
+	FArchive* OutputFile = IFileManager::Get().CreateDebugFileWriter(*OutputFilename);
+
+	auto WriteString = [OutputFile](const char* String)
+	{
+		OutputFile->Serialize((void*)String, sizeof(ANSICHAR)*FCStringAnsi::Strlen(String));
+	};
+
+	WriteString(R"({"traceEvents":[)" "\n");
+
+	enum { StringBufferSize = 1024 };
+	char StringBuffer[StringBufferSize];
+
+	const uint32 Pid = FPlatformProcess::GetCurrentProcessId();
+
+	const uint32 NumEvents = CapturedEvents.Num();
+
+	const uint32 LocalSessionId = SessionId;
+	auto IsEventValid = [LocalSessionId](const FEvent& Event)
+	{
+		if (Event.SessionId != LocalSessionId)
+		{
+			return false;
+		}
+
+		if (Event.Type == EEventType::GPU)
+		{
+			return Event.GPU.EndMicroseconds > Event.GPU.BeginMicroseconds
+				&& (Event.GPU.EndMicroseconds - Event.GPU.BeginMicroseconds) < 10000000; // Arbitrary threshold of 10 seconds for any one GPU event
+		}
+		else
+		{
+			return Event.CPU.EndCycles > Event.CPU.BeginCycles;
+		}
+	};
+
+	int32 MaxGPUIndex = -1;
+	uint64 FirstCPUTimestampCycles = UINT64_MAX;
+	uint64 FirstGPUTimestampMicroseconds = UINT64_MAX;
+	for (const FEvent& Event : CapturedEvents)
+	{
+		if (IsEventValid(Event))
+		{
+			if (Event.Type == EEventType::GPU)
+			{
+				MaxGPUIndex = FMath::Max<int32>(MaxGPUIndex, Event.GPU.GPUIndex);
+				FirstGPUTimestampMicroseconds = FMath::Min(FirstGPUTimestampMicroseconds, Event.GPU.BeginMicroseconds);
+			}
+			else
+			{
+				FirstCPUTimestampCycles = FMath::Min(FirstCPUTimestampCycles, Event.CPU.BeginCycles);
+			}
+		}
+	}
+
+	// Write metadata (thread names, sorting order, etc.)
+
+	int32 SortIndex = 0; // Lower numbers result in higher position in the visualizer.
+
+	for (int32 GPUIndex = 0; GPUIndex <= MaxGPUIndex; ++GPUIndex)
+	{
+		FCStringAnsi::Snprintf(StringBuffer, StringBufferSize,
+			R"({"pid":%d, "tid":%d, "ph": "M", "name": "thread_name", "args":{"name":"GPU %d"}},)"
+			R"({"pid":%d, "tid":%d, "ph": "M", "name": "thread_sort_index", "args":{"sort_index": %d}},)"
+			"\n",
+			Pid, GPUIndex, GPUIndex, Pid, GPUIndex, SortIndex);
+		WriteString(StringBuffer);
+		SortIndex++;
+	}
+
+	FCStringAnsi::Snprintf(StringBuffer, StringBufferSize,
+		R"({"pid":%d, "tid":%d, "ph": "M", "name": "thread_name", "args":{"name":"Render thread"}},)"
+		R"({"pid":%d, "tid":%d, "ph": "M", "name": "thread_sort_index", "args":{"sort_index": %d}},)"
+		"\n",
+		Pid, GRenderThreadId, Pid, GRenderThreadId, SortIndex);
+	WriteString(StringBuffer);
+	SortIndex++;
+
+	FCStringAnsi::Snprintf(StringBuffer, StringBufferSize,
+		R"({"pid":%d, "tid":%d, "ph": "M", "name": "thread_name", "args":{"name":"Game thread"}},)"
+		R"({"pid":%d, "tid":%d, "ph": "M", "name": "thread_sort_index", "args":{"sort_index": %d}},)"
+		"\n",
+		Pid, GGameThreadId, Pid, GGameThreadId, SortIndex);
+	WriteString(StringBuffer);
+	SortIndex++;
+
+	// Align GPU and CPU timestamps
+
+	FGPUTimingCalibrationTimestamp CalibrationTimestamp = FGPUTiming::GetCalibrationTimestamp();
+
+	// If platform does not support GPU/CPU timer alignment, then simply align GPU and CPU on first event
+	if (CalibrationTimestamp.CPUMicroseconds == 0 || CalibrationTimestamp.GPUMicroseconds == 0)
+	{
+		CalibrationTimestamp.CPUMicroseconds = CyclesToMicroseconds64(FirstCPUTimestampCycles);
+		CalibrationTimestamp.GPUMicroseconds = FirstGPUTimestampMicroseconds;
+	}
+
+	const uint64 GPUTimeOffset = CalibrationTimestamp.CPUMicroseconds - CalibrationTimestamp.GPUMicroseconds;
+
+	// Write out all events
+
+	uint32 NumProcessedEvents = 0;
+	for (const FEvent& Event : CapturedEvents)
+	{
+		if (!IsEventValid(Event))
+		{
+			continue;
+		}
+
+		uint64 TimeBeginMicroseconds = 0;
+		uint64 TimeEndMicroseconds = 0;
+
+		switch (Event.Type)
+		{
+		case EEventType::CPU:
+			TimeBeginMicroseconds = CyclesToMicroseconds64(Event.CPU.BeginCycles);
+			TimeEndMicroseconds = CyclesToMicroseconds64(Event.CPU.EndCycles);
+			break;
+		case EEventType::GPU:
+			// Note: We could also use `clock_sync` metadata to synchronize events in Chrome viewer.
+			// Advantage of doing it manually is that trace log can be consumed by a less sophisticated
+			// parser that does not implement all the features of Chrome viewer.
+			TimeBeginMicroseconds = Event.GPU.BeginMicroseconds + GPUTimeOffset;
+			TimeEndMicroseconds = Event.GPU.EndMicroseconds + GPUTimeOffset;
+			break;
+		default:
+			verifyf(0, TEXT("Unexpected profiling event type"));
+		}
+
+		uint32 ThreadOrGpuId;
+		if (Event.Type == EEventType::CPU)
+		{
+			ThreadOrGpuId = Event.CPU.ThreadId;
+		}
+		else
+		{
+			ThreadOrGpuId = Event.GPU.GPUIndex;
+		}
+
+		FCStringAnsi::Snprintf(StringBuffer, StringBufferSize,
+			R"({"pid":%d, "tid":%d, "ph": "X", "name": "%s", "ts": %llu, "dur": %llu, "args":{"frame":%d}},)"
+			"\n",
+			Pid, ThreadOrGpuId, Event.Name, TimeBeginMicroseconds,
+			TimeEndMicroseconds - TimeBeginMicroseconds,
+			Event.FrameNumber);
+
+		WriteString(StringBuffer);
+
+		++NumProcessedEvents;
+	}
+
+	// All done
+
+	WriteString("{}]}");
+
+	OutputFile->Close();
+}
+
+#endif //TRACING_PROFILER
diff --git a/Engine/Source/Runtime/RenderCore/Public/ProfilingDebugging/TracingProfiler.h b/Engine/Source/Runtime/RenderCore/Public/ProfilingDebugging/TracingProfiler.h
new file mode 100644
--- /dev/null
+++ b/Engine/Source/Runtime/RenderCore/Public/ProfilingDebugging/TracingProfiler.h
@@ -0,0 +1,152 @@
+// Copyright 1998-2019 Epic Games, Inc. All Rights Reserved.
+
+/**
+*
+* A lightweight profiler that can output logs compatible with Google Chrome tracing visualizer.
+* Captured events are written as a flat array (fixed size ring buffer), without any kind of aggregation.
+* Tracing events may be added from multiple threads simultaneously. 
+* Old trace events are overwritten when ring buffer wraps.
+*/
+
+#pragma once
+
+#include "HAL/PlatformAtomics.h"
+#include "Modules/ModuleInterface.h"
+#include "Modules/ModuleManager.h"
+#include "CoreTypes.h"
+#include "CoreMinimal.h"
+
+#define TRACING_PROFILER (WITH_ENGINE && !UE_BUILD_SHIPPING)
+
+#if TRACING_PROFILER
+
+class FTracingProfiler
+{
+private:
+
+	static FTracingProfiler* Instance;
+	FTracingProfiler();
+
+public:
+
+	enum class EEventType
+	{
+		CPU,
+		GPU,
+	};
+
+	struct FEvent
+	{
+		const char* Name;
+		uint32 FrameNumber;
+		uint32 SessionId;
+		EEventType Type;
+		union
+		{
+			struct
+			{
+				uint64 BeginMicroseconds;
+				uint64 EndMicroseconds;
+				uint64 GPUIndex;
+			} GPU;
+
+			struct
+			{
+				uint64 BeginCycles;
+				uint64 EndCycles;
+				uint64 ThreadId;
+			} CPU;
+		};
+	};
+
+	static RENDERCORE_API FTracingProfiler* Get();
+
+	RENDERCORE_API void Init();
+
+	uint32 AddCPUEvent(const char* Name,
+		uint64 TimestampBeginCycles,
+		uint64 TimestempEndCycles,
+		uint32 ThreadId,
+		uint32 FrameNumber)
+	{
+		FEvent Event;
+		Event.Name = Name;
+		Event.Type = EEventType::CPU;
+		Event.FrameNumber = FrameNumber;
+		Event.SessionId = SessionId;
+		Event.CPU.BeginCycles = TimestampBeginCycles;
+		Event.CPU.EndCycles = TimestempEndCycles;
+		Event.CPU.ThreadId = ThreadId;
+		return AddEvent(Event);
+	}
+
+	uint32 AddGPUEvent(const char* Name,
+		uint64 TimestempBeginMicroseconds,
+		uint64 TimestampEndMicroseconds,
+		uint64 GPUIndex,
+		uint32 FrameNumber)
+	{
+		FEvent Event;
+		Event.Name = Name;
+		Event.Type = EEventType::GPU;
+		Event.FrameNumber = FrameNumber;
+		Event.SessionId = SessionId;
+		Event.GPU.BeginMicroseconds = TimestempBeginMicroseconds;
+		Event.GPU.EndMicroseconds = TimestampEndMicroseconds;
+		Event.GPU.GPUIndex = GPUIndex;
+		return AddEvent(Event);
+	}
+
+	uint32 AddEvent(const FEvent& Event)
+	{
+		if (!bCapturing)
+		{
+			return ~0u;
+		}
+
+		uint32 EventId = (FPlatformAtomics::InterlockedIncrement(&EventAtomicConter)-1) % MaxNumCapturedEvents;
+		CapturedEvents[EventId] = Event;
+		return EventId;
+	}
+
+	RENDERCORE_API int32 GetCaptureFrameNumber();
+
+	RENDERCORE_API void BeginCapture(int InNumFramesToCapture = -1);
+	RENDERCORE_API void EndCapture();
+
+	RENDERCORE_API bool IsCapturing() const;
+
+private:
+
+	void WriteCaptureToFile();
+
+	void BeginFrame();
+	void EndFrame();
+
+	void BeginFrameRT();
+	void EndFrameRT();
+
+	TArray<FEvent> CapturedEvents;
+	uint32 MaxNumCapturedEvents = 0;
+
+	int32 NumFramesToCapture = -1;
+	int32 CaptureFrameNumber = 0;
+
+	bool bRequestStartCapture = false;
+	bool bRequestStopCapture = false;
+	bool bCapturing = false;
+	bool bCapturingRT = false;
+
+	uint64 GameThreadFrameBeginCycle = 0;
+	uint64 GameThreadFrameEndCycle = 0;
+
+	uint64 RenderThreadFrameBeginCycle = 0;
+	uint64 RenderThreadFrameEndCycle = 0;
+
+	uint32 SessionId = 0;
+
+	uint32 Pad[(PLATFORM_CACHE_LINE_SIZE / 4 - 2)];
+	volatile int32 EventAtomicConter = 0;
+};
+
+#endif //TRACING_PROFILER
