Refactored RDG pass parameters.

 - Added FRDGPassParameterStruct and FRDGPassParameter helpers to make iterating parameters simpler.
 - FRDGPass now releases uniform buffer references instead of the lambda variant.
 - Added FRDGTrackedResource as a base class for buffers / textures. This reduces confusion by eliminating zombie data from UAVs / SRVs.
 - Some minor formatting cleanup.
 - Removed erroneous warning that UAVs can't be used with pixel shaders (they can).
 - Attempting to create a UAV / SRV with a null parent resource now returns nullptr instead of polluting the graph.

diff --git a/Engine/Source/Runtime/RenderCore/Private/RenderGraphBuilder.cpp b/Engine/Source/Runtime/RenderCore/Private/RenderGraphBuilder.cpp
--- a/Engine/Source/Runtime/RenderCore/Private/RenderGraphBuilder.cpp
+++ b/Engine/Source/Runtime/RenderCore/Private/RenderGraphBuilder.cpp
@@ -9,34 +9,39 @@
 
 namespace
 {
+const int32 kRDGEmitWarningsOnce = 1;
+
 #if RDG_ENABLE_DEBUG
-	int32 GRenderGraphImmediateMode = 0;
-	FAutoConsoleVariableRef CVarImmediateMode(
-		TEXT("r.RDG.ImmediateMode"),
-		GRenderGraphImmediateMode,
-		TEXT("Executes passes as they get created. Useful to have a callstack of the wiring code when crashing in the pass' lambda."),
-		ECVF_RenderThreadSafe);
-
-	int32 GRenderGraphEmitWarnings = 0;
-	FAutoConsoleVariableRef CVarEmitWarnings(
-		TEXT("r.RDG.EmitWarnings"),
-		GRenderGraphEmitWarnings,
-		TEXT("Allow to output warnings for inefficiencies found during wiring and execution of the passes.\n")
-		TEXT(" 0: disabled;\n")
-		TEXT(" 1: emit warning once (default);\n")
-		TEXT(" 2: emit warning everytime issue is detected."),
-		ECVF_RenderThreadSafe);
+
+int32 GRDGImmediateMode = 0;
+FAutoConsoleVariableRef CVarImmediateMode(
+	TEXT("r.RDG.ImmediateMode"),
+	GRDGImmediateMode,
+	TEXT("Executes passes as they get created. Useful to have a callstack of the wiring code when crashing in the pass' lambda."),
+	ECVF_RenderThreadSafe);
+
+int32 GRDGEmitWarnings = 0;
+FAutoConsoleVariableRef CVarEmitWarnings(
+	TEXT("r.RDG.EmitWarnings"),
+	GRDGEmitWarnings,
+	TEXT("Allow to output warnings for inefficiencies found during wiring and execution of the passes.\n")
+	TEXT(" 0: disabled;\n")
+	TEXT(" 1: emit warning once (default);\n")
+	TEXT(" 2: emit warning everytime issue is detected."),
+	ECVF_RenderThreadSafe);
 
 #else
-	const int32 GRenderGraphImmediateMode = 0;
-	const int32 GRenderGraphEmitWarnings = 0;
+
+const int32 GRDGImmediateMode = 0;
+const int32 GRDGEmitWarnings = 0;
+
 #endif
-}
+} //! namespace
 
 bool GetEmitRDGEvents()
 {
 #if RDG_EVENTS != RDG_EVENTS_NONE
-	return GetEmitDrawEvents() || GRenderGraphEmitWarnings;
+	return GetEmitDrawEvents() || GRDGEmitWarnings;
 #else
 	return false;
 #endif
@@ -47,69 +52,55 @@ void InitRenderGraph()
 #if RDG_ENABLE_DEBUG_WITH_ENGINE
 	if (FParse::Param(FCommandLine::Get(), TEXT("rdgimmediate")))
 	{
-		GRenderGraphImmediateMode = 1;
+		GRDGImmediateMode = 1;
 	}
 
 	if (FParse::Param(FCommandLine::Get(), TEXT("rdgdebug")))
 	{
-		GRenderGraphEmitWarnings = 1;
+		GRDGEmitWarnings = 1;
 	}
 #endif
 }
 
-static void EmitRenderGraphWarning(const FString& WarningMessage)
+void EmitRDGWarning(const FString& WarningMessage)
 {
-	check(GRenderGraphEmitWarnings);
+	if (!GRDGEmitWarnings)
+	{
+		return;
+	}
 
 	static TSet<FString> GAlreadyEmittedWarnings;
 
-	if (GRenderGraphEmitWarnings == 2)
+	if (GRDGEmitWarnings == kRDGEmitWarningsOnce)
 	{
-		UE_LOG(LogRendererCore, Warning, TEXT("%s"), *WarningMessage);
+		if (!GAlreadyEmittedWarnings.Contains(WarningMessage))
+		{
+			GAlreadyEmittedWarnings.Add(WarningMessage);
+			UE_LOG(LogRendererCore, Warning, TEXT("%s"), *WarningMessage);
+		}
 	}
-	else if (!GAlreadyEmittedWarnings.Contains(WarningMessage))
+	else
 	{
-		GAlreadyEmittedWarnings.Add(WarningMessage);
 		UE_LOG(LogRendererCore, Warning, TEXT("%s"), *WarningMessage);
 	}
 }
 
-#define EmitRenderGraphWarningf(WarningMessageFormat, ...) \
-	EmitRenderGraphWarning(FString::Printf(WarningMessageFormat, ##__VA_ARGS__));
+#define EmitRDGWarningf(WarningMessageFormat, ...) \
+	EmitRDGWarning(FString::Printf(WarningMessageFormat, ##__VA_ARGS__));
 
-static bool IsBoundAsReadable(const FRDGTexture* Texture, FShaderParameterStructRef ParameterStruct)
+FRDGBuilder::FRDGBuilder(FRHICommandListImmediate& InRHICmdList)
+	: RHICmdList(InRHICmdList)
+	, MemStack(FMemStack::Get())
+	, EventScopeStack(RHICmdList)
+{}
+
+FRDGBuilder::~FRDGBuilder()
 {
-	for (int ResourceIndex = 0, Num = ParameterStruct.Layout->Resources.Num(); ResourceIndex < Num; ResourceIndex++)
+	#if RDG_ENABLE_DEBUG
 	{
-		EUniformBufferBaseType Type = ParameterStruct.Layout->Resources[ResourceIndex].MemberType;
-		uint16 Offset = ParameterStruct.Layout->Resources[ResourceIndex].MemberOffset;
-
-		switch (Type)
-		{
-		case UBMT_RDG_TEXTURE:
-		{
-			const FRDGTexture* InputTexture = *ParameterStruct.GetMemberPtrAtOffset<const FRDGTexture*>(Offset);
-			if (Texture == InputTexture)
-			{
-				return true;
-			}
-		}
-		break;
-		case UBMT_RDG_TEXTURE_SRV:
-		{
-			const FRDGTextureSRV* InputSRV = *ParameterStruct.GetMemberPtrAtOffset<const FRDGTextureSRV*>(Offset);
-			if (InputSRV && Texture == InputSRV->Desc.Texture)
-			{
-				return true;
-			}
-		}
-		break;
-		default:
-			break;
-		}
+		checkf(bHasExecuted, TEXT("Render graph execution is required to ensure consistency with immediate mode."));
 	}
-
-	return false;
+	#endif
 }
 
 void FRDGBuilder::Execute()
@@ -134,7 +125,7 @@ void FRDGBuilder::Execute()
 
 	EventScopeStack.BeginExecute();
 
-	if (!GRenderGraphImmediateMode)
+	if (!GRDGImmediateMode)
 	{
 		WalkGraphDependencies();
 
@@ -160,18 +151,17 @@ void FRDGBuilder::Execute()
 
 void FRDGBuilder::AddPassInternal(FRDGPass* Pass)
 {
-#if RDG_ENABLE_DEBUG
+	#if RDG_ENABLE_DEBUG
 	{
-		const void* ParameterStructData = Pass->GetParameters().Contents;
-
 		checkf(!bHasExecuted, TEXT("Render graph pass %s needs to be added before the builder execution."), Pass->GetName());
 
-		/** A pass parameter structure requires a correct life time until the pass execution, and therefor needs to be
-			* allocated with FRDGBuilder::AllocParameters().
-			*
-			* Moreover, because the destructor of this parameter structure will be done after the pass execution, a it can
-			* only be used by a single AddPass().
-			*/
+		const void* ParameterStructData = Pass->GetParameters().GetContents();
+
+		/** The lifetime of each pass parameter structure must extend until deferred pass execution; therefore, it needs to be
+		 *  allocated with FRDGBuilder::AllocParameters(). Also, all references held by the parameter structure are released
+		 *  immediately after pass execution, so a pass parameter struct instance must be 1-to-1 with a pass instance (i.e. one
+		 *  per AddPass() call).
+		 */
 		checkf(
 			AllocatedUnusedPassParameters.Contains(ParameterStructData),
 			TEXT("The pass parameter structure has not been allocated for correct life time FRDGBuilder::AllocParameters() or has already ")
@@ -179,148 +169,115 @@ void FRDGBuilder::AddPassInternal(FRDGPass* Pass)
 
 		AllocatedUnusedPassParameters.Remove(ParameterStructData);
 	}
-#endif
+	#endif
 
 	Pass->ParentScope = EventScopeStack.GetCurrentScope();
 	Passes.Emplace(Pass);
 
-#if RDG_ENABLE_DEBUG || SUPPORTS_VISUALIZE_TEXTURE
-	{
-		DebugPass(Pass);
-	}
-#endif
-}
-
-void FRDGBuilder::DebugPass(const FRDGPass* Pass)
-{
-#if RDG_ENABLE_DEBUG
-	// Verify all the settings of the pass make sense.
 	ValidatePass(Pass);
 
-	// Executes the pass immediatly as they get added mode, to have callstack of wiring code when crashing within the pass.
-	if (GRenderGraphImmediateMode)
+	if (GRDGImmediateMode)
 	{
 		ExecutePass(Pass);
 	}
-#endif
 
-#if SUPPORTS_VISUALIZE_TEXTURE
-	// If visualizing a texture, look for any output of the pass. This must be done after the
-	// GRenderGraphImmediateMode's ExecutePass() because this will actually create a capturing
-	// pass if needed that would have to be executed right away as well.
-	if (GVisualizeTexture.bEnabled)
-	{
-		CaptureAnyInterestingPassOutput(Pass);
-	}
-#endif
+	VisualizePassOutputs(Pass);
 }
 
 void FRDGBuilder::ValidatePass(const FRDGPass* Pass) const
 {
 #if RDG_ENABLE_DEBUG
-	FRenderTargetBindingSlots* RESTRICT RenderTargets = nullptr;
-	FShaderParameterStructRef ParameterStruct = Pass->GetParameters();
+	const FRenderTargetBindingSlots* RenderTargetBindingSlots = nullptr;
+
+	const TCHAR* PassName = Pass->GetName();
+	const bool bIsCompute = Pass->IsCompute();
+	const bool bRenderTargetSlotsRequired = !bIsCompute;
 
-	bool bIsCompute = Pass->IsCompute();
-	bool bCanUseUAVs = true;	// Graphics passes can now also use UAVs in pixelshaders.
-	bool bRequiresRenderTargetSlots = !bIsCompute;
+	FRDGPassParameterStruct ParameterStruct = Pass->GetParameters();
 
-	for (int ResourceIndex = 0, Num = ParameterStruct.Layout->Resources.Num(); ResourceIndex < Num; ResourceIndex++)
+	const uint32 ParameterCount = ParameterStruct.GetParameterCount();
+
+	for (uint32 ParameterIndex = 0; ParameterIndex < ParameterCount; ++ParameterIndex)
 	{
-		EUniformBufferBaseType Type = ParameterStruct.Layout->Resources[ResourceIndex].MemberType;
-		uint16 Offset = ParameterStruct.Layout->Resources[ResourceIndex].MemberOffset;
+		FRDGPassParameter Parameter = ParameterStruct.GetParameter(ParameterIndex);
 
-		switch (Type)
+		switch (Parameter.GetType())
 		{
 		case UBMT_RDG_TEXTURE:
 		{
-			FRDGTexture* RESTRICT Texture = *ParameterStruct.GetMemberPtrAtOffset<FRDGTexture*>(Offset);
-			checkf(!Texture || Texture->bHasEverBeenProduced,
-				TEXT("Pass %s has a dependency over the texture %s that has never been produced."),
-				Pass->GetName(), Texture->Name);
+			if (FRDGTextureRef Texture = Parameter.GetAsTexture())
+			{
+				checkf(Texture->HasBeenProduced(),
+					TEXT("Pass %s has a dependency over the texture %s that has never been produced."),
+					PassName, Texture->Name);
+			}
 		}
 		break;
 		case UBMT_RDG_TEXTURE_SRV:
 		{
-			FRDGTextureSRV* RESTRICT SRV = *ParameterStruct.GetMemberPtrAtOffset<FRDGTextureSRV*>(Offset);
-			if (SRV)
+			if (FRDGTextureSRVRef SRV = Parameter.GetAsTextureSRV())
 			{
-				checkf(SRV->Desc.Texture->bHasEverBeenProduced,
+				FRDGTextureRef Texture = SRV->Desc.Texture;
+
+				checkf(Texture->HasBeenProduced(),
 					TEXT("Pass %s has a dependency over the texture %s that has never been produced."),
-					Pass->GetName(), SRV->Desc.Texture->Name);
+					PassName, Texture->Name);
 			}
 		}
 		break;
 		case UBMT_RDG_TEXTURE_UAV:
 		{
-			FRDGTextureUAV* RESTRICT UAV = *ParameterStruct.GetMemberPtrAtOffset<FRDGTextureUAV*>(Offset);
-			if (UAV)
+			if (FRDGTextureUAVRef UAV = Parameter.GetAsTextureUAV())
 			{
-				if (!UAV->Desc.Texture->bHasEverBeenProduced)
-				{
-					UAV->Desc.Texture->bHasEverBeenProduced = true;
-					UAV->Desc.Texture->DebugFirstProducer = Pass;
-				}
+				FRDGTextureRef Texture = UAV->Desc.Texture;
 
-				if (!bCanUseUAVs && GRenderGraphEmitWarnings)
-				{
-					EmitRenderGraphWarningf(
-						TEXT("UAV can only been bound to compute shaders, therefore UAV %s is certainly useless for pass %s."),
-						UAV->Name, Pass->GetName());
-				}
+				Texture->MarkAsProducedBy(Pass);
 			}
 		}
 		break;
 		case UBMT_RDG_BUFFER:
 		{
-			FRDGBuffer* RESTRICT Buffer = *ParameterStruct.GetMemberPtrAtOffset<FRDGBuffer*>(Offset);
-			checkf(!Buffer || Buffer->bHasEverBeenProduced,
-				TEXT("Pass %s has a dependency over the buffer %s that has never been produced."),
-				Pass->GetName(), Buffer->Name);
+			if (FRDGBufferRef Buffer = Parameter.GetAsBuffer())
+			{
+				checkf(Buffer->HasBeenProduced(),
+					TEXT("Pass %s has a dependency over the buffer %s that has never been produced."),
+					PassName, Buffer->Name);
+			}
 		}
 		break;
 		case UBMT_RDG_BUFFER_SRV:
 		{
-			FRDGBufferSRV* RESTRICT SRV = *ParameterStruct.GetMemberPtrAtOffset<FRDGBufferSRV*>(Offset);
-			if (SRV)
+			if (FRDGBufferSRVRef SRV = Parameter.GetAsBufferSRV())
 			{
-				checkf(SRV->Desc.Buffer->bHasEverBeenProduced,
+				FRDGBufferRef Buffer = SRV->Desc.Buffer;
+
+				checkf(Buffer->HasBeenProduced(),
 					TEXT("Pass %s has a dependency over the buffer %s that has never been produced."),
-					Pass->GetName(), SRV->Desc.Buffer->Name);
+					PassName, SRV->Desc.Buffer->Name);
 			}
 		}
 		break;
 		case UBMT_RDG_BUFFER_UAV:
 		{
-			FRDGBufferUAV* RESTRICT UAV = *ParameterStruct.GetMemberPtrAtOffset<FRDGBufferUAV*>(Offset);
-			if (UAV)
+			if (FRDGBufferUAVRef UAV = Parameter.GetAsBufferUAV())
 			{
-				if (!UAV->Desc.Buffer->bHasEverBeenProduced)
-				{
-					UAV->Desc.Buffer->bHasEverBeenProduced = true;
-					UAV->Desc.Buffer->DebugFirstProducer = Pass;
-				}
+				FRDGBufferRef Buffer = UAV->Desc.Buffer;
 
-				if (!bCanUseUAVs && GRenderGraphEmitWarnings)
-				{
-					EmitRenderGraphWarningf(
-						TEXT("UAV can only been bound to compute shaders, therefore UAV %s is certainly useless for pass %s."),
-						UAV->Name, Pass->GetName());
-				}
+				Buffer->MarkAsProducedBy(Pass);
 			}
 		}
 		break;
 		case UBMT_RENDER_TARGET_BINDING_SLOTS:
 		{
-			if (!RenderTargets)
+			if (!RenderTargetBindingSlots)
 			{
-				RenderTargets = ParameterStruct.GetMemberPtrAtOffset<FRenderTargetBindingSlots>(Offset);
+				RenderTargetBindingSlots = &Parameter.GetAsRenderTargetBindingSlots();
 			}
-			else if (GRenderGraphEmitWarnings)
+			else if (GRDGEmitWarnings)
 			{
-				EmitRenderGraphWarningf(
+				EmitRDGWarningf(
 					TEXT("Pass %s have duplicated render target binding slots."),
-					Pass->GetName());
+					PassName);
 			}
 		}
 		break;
@@ -329,108 +286,189 @@ void FRDGBuilder::ValidatePass(const FRDGPass* Pass) const
 		}
 	}
 
-	if (RenderTargets)
+	/** Validate that raster passes have render target binding slots and compute passes don't. */
+	if (RenderTargetBindingSlots)
 	{
-		checkf(bRequiresRenderTargetSlots, TEXT("Render pass %s does not need render target binging slots"), Pass->GetName());
+		checkf(!bIsCompute, TEXT("Pass '%s' has render target binding slots but is flagged as 'Compute'."), PassName);
+	}
+	else
+	{
+		checkf(bIsCompute, TEXT("Pass '%s' is missing render target binding slots. Add the 'Compute' flag if render targets are not required."), PassName);
+	}
 
-		const bool bGeneratingMips = (Pass->GetFlags() & ERDGPassFlags::GenerateMips) == ERDGPassFlags::GenerateMips;
-		bool bFoundRTBound = false;
+	/** Validate render target / depth stencil binding usage. */
+	if (RenderTargetBindingSlots)
+	{
+		const auto& RenderTargets = RenderTargetBindingSlots->Output;
+
+		const uint32 RenderTargetCount = RenderTargets.Num();
 
-		int32 NumRenderTargets = 0;
-		for (int32 i = 0; i < RenderTargets->Output.Num(); i++)
 		{
-			const FRenderTargetBinding& RenderTarget = RenderTargets->Output[i];
-			FRDGTexture* Texture = RenderTarget.GetTexture();
+			/** Tracks the number of contiguous, non-null textures in the render target output array. */
+			uint32 ValidRenderTargetCount = 0;
 
-			if (!Texture)
+			for (uint32 RenderTargetIndex = 0; RenderTargetIndex < RenderTargetCount; ++RenderTargetIndex)
 			{
-				NumRenderTargets = i;
-				break;
+				const FRenderTargetBinding& RenderTarget = RenderTargets[RenderTargetIndex];
+
+				if (FRDGTextureRef Texture = RenderTarget.GetTexture())
+				{
+					/** Validate that load action is correct. */
+					const bool bIsLoadAction = RenderTarget.GetLoadAction() == ERenderTargetLoadAction::ELoad;
+					const bool bIsLoadActionInvalid = bIsLoadAction && !Texture->HasBeenProduced();
+					checkf(
+						!bIsLoadActionInvalid,
+						TEXT("Can't load a render target '%s' that has never been produced."),
+						Texture->Name);
+
+					/** Mark the pass as a producer for render targets with a store action. */
+					const bool bIsStoreAction = RenderTarget.GetStoreAction() != ERenderTargetStoreAction::ENoAction;
+					if (bIsStoreAction)
+					{
+						Texture->MarkAsProducedBy(Pass);
+					}
+				}
+				else
+				{
+					/** Found end of contiguous interval of valid render targets. */
+					ValidRenderTargetCount = RenderTargetIndex;
+					break;
+				}
 			}
 
-			if (!Texture->bHasEverBeenProduced)
+			/** Validate that no holes exist in the render target output array. Render targets must be bound contiguously. */
+			for (uint32 RenderTargetIndex = ValidRenderTargetCount; RenderTargetIndex < RenderTargetCount; ++RenderTargetIndex)
 			{
-				checkf(RenderTarget.GetLoadAction() != ERenderTargetLoadAction::ELoad,
-					TEXT("Can't load a render target %s that has never been produced."),
-					Texture->Name);
-
-				// TODO(RDG): should only be done when there is a store action.
-				Texture->bHasEverBeenProduced = true;
-				Texture->DebugFirstProducer = Pass;
+				const FRenderTargetBinding& RenderTarget = RenderTargets[RenderTargetIndex];
+				checkf(RenderTarget.GetTexture() == nullptr, TEXT("Render targets must be packed. No empty spaces in the array."));
 			}
-
-			bFoundRTBound = bFoundRTBound || IsBoundAsReadable(Texture, ParameterStruct);
-		}
-		for (int32 i = NumRenderTargets; i < RenderTargets->Output.Num(); i++)
-		{
-			const FRenderTargetBinding& RenderTarget = RenderTargets->Output[i];
-			checkf(RenderTarget.GetTexture() == nullptr, TEXT("Render targets must be packed. No empty spaces in the array."));
 		}
-		ensureMsgf(!bGeneratingMips || bFoundRTBound, TEXT("GenerateMips enabled but no RT found as source!"));
 
-		FRDGTexture* Texture = RenderTargets->DepthStencil.GetTexture();
-		if (Texture && !Texture->bHasEverBeenProduced)
+		const auto IsMipLevelBoundForRead = [ParameterStruct, ParameterCount](FRDGTextureRef Texture, uint32 MipLevel)
 		{
-			checkf(RenderTargets->DepthStencil.GetDepthLoadAction() != ERenderTargetLoadAction::ELoad,
-				TEXT("Can't load depth from render target %s that has never been produced."),
-				Texture->Name);
+			for (uint32 ParameterIndex = 0; ParameterIndex < ParameterCount; ++ParameterIndex)
+			{
+				FRDGPassParameter Parameter = ParameterStruct.GetParameter(ParameterIndex);
+
+				switch (Parameter.GetType())
+				{
+				case UBMT_RDG_TEXTURE:
+				{
+					if (Parameter.GetAsTexture() == Texture)
+					{
+						/** Full texture view bound. */
+						return true;
+					}
+					break;
+				}
+				case UBMT_RDG_TEXTURE_SRV:
+				{
+					if (const FRDGTextureSRVRef InputSRV = Parameter.GetAsTextureSRV())
+					{
+						const FRDGTextureSRVDesc& InputSRVDesc = InputSRV->Desc;
 
-			checkf(RenderTargets->DepthStencil.GetStencilLoadAction() != ERenderTargetLoadAction::ELoad,
-				TEXT("Can't load stencil from render target %s that has never been produced."),
-				Texture->Name);
+						if (InputSRVDesc.Texture == Texture)
+						{
+							return InputSRVDesc.MipLevel == MipLevel;
+						}
+					}
+					break;
+				}
+				default:
+					break;
+				}
+			}
 
-			Texture->bHasEverBeenProduced = true;
-			Texture->DebugFirstProducer = Pass;
+			return false;
+		};
+
+		/** Validate that texture mips are not bound as a render target and SRV at the same time. */
+		for (uint32 RenderTargetIndex = 0; RenderTargetIndex < RenderTargetCount; ++RenderTargetIndex)
+		{
+			const FRenderTargetBinding& RenderTarget = RenderTargets[RenderTargetIndex];
+
+			if (FRDGTextureRef Texture = RenderTarget.GetTexture())
+			{
+				if (IsMipLevelBoundForRead(Texture, RenderTarget.GetMipIndex()))
+				{
+					checkf(false, TEXT("Texture '%s', Mip '%d' is bound as a render target and SRV at the same time."), Texture->Name, RenderTarget.GetMipIndex());
+				}
+			}
+			else
+			{
+				break;
+			}
 		}
 	}
-	else
-	{
-		checkf(!bRequiresRenderTargetSlots, TEXT("Render pass %s requires render target binging slots"), Pass->GetName());
-	}
 #endif // RDG_ENABLE_DEBUG
 }
 
-void FRDGBuilder::CaptureAnyInterestingPassOutput(const FRDGPass* Pass)
+void FRDGBuilder::VisualizePassOutputs(const FRDGPass* Pass)
 {
 #if SUPPORTS_VISUALIZE_TEXTURE
-	FShaderParameterStructRef ParameterStruct = Pass->GetParameters();
-	for (int ResourceIndex = 0, Num = ParameterStruct.Layout->Resources.Num(); ResourceIndex < Num; ResourceIndex++)
+	if (!!GVisualizeTexture.bEnabled)
+	{
+		return;
+	}
+
+	FRDGPassParameterStruct ParameterStruct = Pass->GetParameters();
+
+	const uint32 ParameterCount = ParameterStruct.GetParameterCount();
+
+	for (uint32 ParameterIndex = 0; ParameterIndex < ParameterCount; ++ParameterIndex)
 	{
-		EUniformBufferBaseType Type = ParameterStruct.Layout->Resources[ResourceIndex].MemberType;
-		uint16 Offset = ParameterStruct.Layout->Resources[ResourceIndex].MemberOffset;
+		FRDGPassParameter Parameter = ParameterStruct.GetParameter(ParameterIndex);
 
-		switch (Type)
+		switch (Parameter.GetType())
 		{
 		case UBMT_RDG_TEXTURE_UAV:
 		{
-			FRDGTextureUAV* RESTRICT UAV = *ParameterStruct.GetMemberPtrAtOffset<FRDGTextureUAV*>(Offset);
-			if (UAV && GVisualizeTexture.ShouldCapture(UAV->Desc.Texture->Name))
+			if (FRDGTextureUAVRef UAV = Parameter.GetAsTextureUAV())
 			{
-				GVisualizeTexture.CreateContentCapturePass(*this, UAV->Desc.Texture);
+				if (FRDGTextureRef Texture = UAV->Desc.Texture)
+				{
+					if (GVisualizeTexture.ShouldCapture(Texture->Name))
+					{
+						GVisualizeTexture.CreateContentCapturePass(*this, Texture);
+					}
+				}
 			}
 		}
 		break;
 		case UBMT_RENDER_TARGET_BINDING_SLOTS:
 		{
-			FRenderTargetBindingSlots* RESTRICT RenderTargets = ParameterStruct.GetMemberPtrAtOffset<FRenderTargetBindingSlots>(Offset);
-			if (RenderTargets->DepthStencil.GetTexture() && 
-				(RenderTargets->DepthStencil.GetDepthStoreAction() != ERenderTargetStoreAction::ENoAction || RenderTargets->DepthStencil.GetStencilStoreAction() != ERenderTargetStoreAction::ENoAction) &&
-				GVisualizeTexture.ShouldCapture(RenderTargets->DepthStencil.GetTexture()->Name))
+			const FRenderTargetBindingSlots& RenderTargetBindingSlots = Parameter.GetAsRenderTargetBindingSlots();
+			const auto& DepthStencil = RenderTargetBindingSlots.DepthStencil;
+			const auto& RenderTargets = RenderTargetBindingSlots.Output;
+
+			if (FRDGTextureRef Texture = DepthStencil.GetTexture())
 			{
-				GVisualizeTexture.CreateContentCapturePass(*this, RenderTargets->DepthStencil.GetTexture());
+				const bool bHasStoreAction = DepthStencil.GetDepthStoreAction() != ERenderTargetStoreAction::ENoAction || DepthStencil.GetStencilStoreAction() != ERenderTargetStoreAction::ENoAction;
+
+				if (bHasStoreAction && GVisualizeTexture.ShouldCapture(Texture->Name))
+				{
+					GVisualizeTexture.CreateContentCapturePass(*this, Texture);
+				}
 			}
-			for (int32 i = 0; i < RenderTargets->Output.Num(); i++)
+
+			const uint32 RenderTargetCount = RenderTargets.Num();
+
+			for (uint32 RenderTargetIndex = 0; RenderTargetIndex < RenderTargetCount; ++RenderTargetIndex)
 			{
-				const FRenderTargetBinding& RenderTarget = RenderTargets->Output[i];
-				if (!RenderTarget.GetTexture())
+				const FRenderTargetBinding& RenderTarget = RenderTargets[RenderTargetIndex];
+
+				if (FRDGTextureRef Texture = RenderTarget.GetTexture())
 				{
-					break;
-				}
+					const bool bHasStoreAction = RenderTarget.GetStoreAction() != ERenderTargetStoreAction::ENoAction;
 
-				if (RenderTarget.GetStoreAction() != ERenderTargetStoreAction::ENoAction &&
-					GVisualizeTexture.ShouldCapture(RenderTarget.GetTexture()->Name))
+					if (bHasStoreAction && GVisualizeTexture.ShouldCapture(Texture->Name))
+					{
+						GVisualizeTexture.CreateContentCapturePass(*this, Texture);
+					}
+				}
+				else
 				{
-					GVisualizeTexture.CreateContentCapturePass(*this, RenderTarget.GetTexture());
+					break;
 				}
 			}
 		}
@@ -439,28 +477,27 @@ void FRDGBuilder::CaptureAnyInterestingPassOutput(const FRDGPass* Pass)
 			break;
 		}
 	}
-#endif // SUPPORTS_VISUALIZE_TEXTURE
+#endif
 }
 
 void FRDGBuilder::WalkGraphDependencies()
 {
 	for (const FRDGPass* Pass : Passes)
 	{
-		FShaderParameterStructRef ParameterStruct = Pass->GetParameters();
+		FRDGPassParameterStruct ParameterStruct = Pass->GetParameters();
+
+		const uint32 ParameterCount = ParameterStruct.GetParameterCount();
 
-		/** Increments all the FRDGResource::ReferenceCount. */
-		for (int ResourceIndex = 0, Num = ParameterStruct.Layout->Resources.Num(); ResourceIndex < Num; ResourceIndex++)
+		for (uint32 ParameterIndex = 0; ParameterIndex < ParameterCount; ++ParameterIndex)
 		{
-			EUniformBufferBaseType Type = ParameterStruct.Layout->Resources[ResourceIndex].MemberType;
-			uint16 Offset = ParameterStruct.Layout->Resources[ResourceIndex].MemberOffset;
+			FRDGPassParameter Parameter = ParameterStruct.GetParameter(ParameterIndex);
 
-			switch (Type)
+			switch (Parameter.GetType())
 			{
 			case UBMT_RDG_TEXTURE:
 			case UBMT_RDG_BUFFER:
 			{
-				FRDGResource* RESTRICT Resource = *ParameterStruct.GetMemberPtrAtOffset<FRDGResource*>(Offset);
-				if (Resource)
+				if (FRDGTrackedResourceRef Resource = Parameter.GetAsTrackedResource())
 				{
 					Resource->ReferenceCount++;
 				}
@@ -468,8 +505,7 @@ void FRDGBuilder::WalkGraphDependencies()
 			break;
 			case UBMT_RDG_TEXTURE_SRV:
 			{
-				FRDGTextureSRV* RESTRICT SRV = *ParameterStruct.GetMemberPtrAtOffset<FRDGTextureSRV*>(Offset);
-				if (SRV)
+				if (FRDGTextureSRVRef SRV = Parameter.GetAsTextureSRV())
 				{
 					SRV->Desc.Texture->ReferenceCount++;
 				}
@@ -477,8 +513,7 @@ void FRDGBuilder::WalkGraphDependencies()
 			break;
 			case UBMT_RDG_TEXTURE_UAV:
 			{
-				FRDGTextureUAV* RESTRICT UAV = *ParameterStruct.GetMemberPtrAtOffset<FRDGTextureUAV*>(Offset);
-				if (UAV)
+				if (FRDGTextureUAVRef UAV = Parameter.GetAsTextureUAV())
 				{
 					UAV->Desc.Texture->ReferenceCount++;
 				}
@@ -486,8 +521,7 @@ void FRDGBuilder::WalkGraphDependencies()
 			break;
 			case UBMT_RDG_BUFFER_SRV:
 			{
-				FRDGBufferSRV* RESTRICT SRV = *ParameterStruct.GetMemberPtrAtOffset<FRDGBufferSRV*>(Offset);
-				if (SRV)
+				if (FRDGBufferSRVRef SRV = Parameter.GetAsBufferSRV())
 				{
 					SRV->Desc.Buffer->ReferenceCount++;
 				}
@@ -495,8 +529,7 @@ void FRDGBuilder::WalkGraphDependencies()
 			break;
 			case UBMT_RDG_BUFFER_UAV:
 			{
-				FRDGBufferUAV* RESTRICT UAV = *ParameterStruct.GetMemberPtrAtOffset<FRDGBufferUAV*>(Offset);
-				if (UAV)
+				if (FRDGBufferUAVRef UAV = Parameter.GetAsBufferUAV())
 				{
 					UAV->Desc.Buffer->ReferenceCount++;
 				}
@@ -504,14 +537,18 @@ void FRDGBuilder::WalkGraphDependencies()
 			break;
 			case UBMT_RENDER_TARGET_BINDING_SLOTS:
 			{
-				FRenderTargetBindingSlots* RESTRICT RenderTargets = ParameterStruct.GetMemberPtrAtOffset<FRenderTargetBindingSlots>(Offset);
+				const FRenderTargetBindingSlots& RenderTargetBindingSlots = Parameter.GetAsRenderTargetBindingSlots();
+				const auto& DepthStencil = RenderTargetBindingSlots.DepthStencil;
+				const auto& RenderTargets = RenderTargetBindingSlots.Output;
+				const uint32 RenderTargetCount = RenderTargets.Num();
 
-				for (int32 i = 0; i < RenderTargets->Output.Num(); i++)
+				for (uint32 RenderTargetIndex = 0; RenderTargetIndex < RenderTargetCount; ++RenderTargetIndex)
 				{
-					const FRenderTargetBinding& RenderTarget = RenderTargets->Output[i];
-					if (RenderTarget.GetTexture())
+					const FRenderTargetBinding& RenderTarget = RenderTargets[RenderTargetIndex];
+
+					if (FRDGTextureRef Texture = RenderTarget.GetTexture())
 					{
-						RenderTarget.GetTexture()->ReferenceCount++;
+						Texture->ReferenceCount++;
 					}
 					else
 					{
@@ -519,10 +556,9 @@ void FRDGBuilder::WalkGraphDependencies()
 					}
 				}
 
-				const FDepthStencilBinding& DepthStencil = RenderTargets->DepthStencil;
-				if (DepthStencil.GetTexture())
+				if (FRDGTextureRef Texture = DepthStencil.GetTexture())
 				{
-					DepthStencil.GetTexture()->ReferenceCount++;
+					Texture->ReferenceCount++;
 				}
 			}
 			break;
@@ -574,7 +610,7 @@ void FRDGBuilder::AllocateRHITextureIfNeeded(FRDGTexture* Texture, bool bCompute
 		return;
 	}
 
-	check(Texture->ReferenceCount > 0 || GRenderGraphImmediateMode);
+	check(Texture->ReferenceCount > 0 || GRDGImmediateMode);
 
 	// TODO(RDG): should avoid bDoWritableBarrier = true
 	TRefCountPtr<IPooledRenderTarget>& PooledRenderTarget = AllocatedTextures.FindOrAdd(Texture);
@@ -607,30 +643,33 @@ void FRDGBuilder::AllocateRHIBufferSRVIfNeeded(FRDGBufferSRV* SRV, bool bCompute
 	{
 		return;
 	}
-	
+
+	FRDGBufferRef Buffer = SRV->Desc.Buffer;
+
 	// The underlying buffer have already been allocated by a prior pass through AllocateRHIBufferUAVIfNeeded().
 	#if RDG_ENABLE_DEBUG
 	{
-		check(SRV->Desc.Buffer->bHasEverBeenProduced);
-	}	
+		check(Buffer->HasBeenProduced());
+	}
 	#endif
-	check(SRV->Desc.Buffer->PooledBuffer);
 
-	if (SRV->Desc.Buffer->PooledBuffer->SRVs.Contains(SRV->Desc))
+	check(Buffer->PooledBuffer);
+
+	if (Buffer->PooledBuffer->SRVs.Contains(SRV->Desc))
 	{
-		SRV->CachedRHI.SRV = SRV->Desc.Buffer->PooledBuffer->SRVs[SRV->Desc];
+		SRV->CachedRHI.SRV = Buffer->PooledBuffer->SRVs[SRV->Desc];
 		return;
 	}
 
 	FShaderResourceViewRHIRef RHIShaderResourceView;
 
-	if (SRV->Desc.Buffer->Desc.UnderlyingType == FRDGBufferDesc::EUnderlyingType::VertexBuffer)
+	if (Buffer->Desc.UnderlyingType == FRDGBufferDesc::EUnderlyingType::VertexBuffer)
 	{
-		RHIShaderResourceView = RHICreateShaderResourceView(SRV->Desc.Buffer->PooledBuffer->VertexBuffer, SRV->Desc.BytesPerElement, SRV->Desc.Format);
+		RHIShaderResourceView = RHICreateShaderResourceView(Buffer->PooledBuffer->VertexBuffer, SRV->Desc.BytesPerElement, SRV->Desc.Format);
 	}
-	else if (SRV->Desc.Buffer->Desc.UnderlyingType == FRDGBufferDesc::EUnderlyingType::StructuredBuffer)
+	else if (Buffer->Desc.UnderlyingType == FRDGBufferDesc::EUnderlyingType::StructuredBuffer)
 	{
-		RHIShaderResourceView = RHICreateShaderResourceView(SRV->Desc.Buffer->PooledBuffer->StructuredBuffer);
+		RHIShaderResourceView = RHICreateShaderResourceView(Buffer->PooledBuffer->StructuredBuffer);
 	}
 	else
 	{
@@ -638,7 +677,7 @@ void FRDGBuilder::AllocateRHIBufferSRVIfNeeded(FRDGBufferSRV* SRV, bool bCompute
 	}
 
 	SRV->CachedRHI.SRV = RHIShaderResourceView;
-	SRV->Desc.Buffer->PooledBuffer->SRVs.Add(SRV->Desc, RHIShaderResourceView);
+	Buffer->PooledBuffer->SRVs.Add(SRV->Desc, RHIShaderResourceView);
 }
 
 void FRDGBuilder::AllocateRHIBufferUAVIfNeeded(FRDGBufferUAV* UAV, bool bComputePass)
@@ -655,7 +694,7 @@ void FRDGBuilder::AllocateRHIBufferUAVIfNeeded(FRDGBufferUAV* UAV, bool bCompute
 	// Allocate a buffer resource.
 	if (!Buffer->PooledBuffer)
 	{
-		check(Buffer->ReferenceCount > 0 || GRenderGraphImmediateMode);
+		check(Buffer->ReferenceCount > 0 || GRDGImmediateMode);
 
 		TRefCountPtr<FPooledRDGBuffer>& AllocatedBuffer = AllocatedBuffers.FindOrAdd(Buffer);
 		GRenderGraphResourcePool.FindFreeBuffer(RHICmdList, Buffer->Desc, AllocatedBuffer, Buffer->Name);
@@ -669,7 +708,7 @@ void FRDGBuilder::AllocateRHIBufferUAVIfNeeded(FRDGBufferUAV* UAV, bool bCompute
 		return;
 	}
 
-	// Hack to make sure only one UAVs is arround.
+	// Hack to make sure only one UAVs is around.
 	Buffer->PooledBuffer->UAVs.Empty();
 
 	FUnorderedAccessViewRHIRef RHIUnorderedAccessView;
@@ -712,11 +751,11 @@ void FRDGBuilder::TransitionTexture(FRDGTexture* Texture, EResourceTransitionAcc
 {
 	const bool bRequiredWritable = TransitionAccess != EResourceTransitionAccess::EReadable;
 
-	if( Texture->bWritable != bRequiredWritable || Texture->bCompute != bRequiredCompute )
+	if (Texture->bWritable != bRequiredWritable || Texture->bCompute != bRequiredCompute)
 	{
 		FTextureRHIParamRef RHITexture = Texture->PooledRenderTarget->GetRenderTargetItem().ShaderResourceTexture;
 
-		bool bIsMultiFrameResource = uint32(Texture->Flags) & uint32(ERDGResourceFlags::MultiFrame);
+		const bool bIsMultiFrameResource = (Texture->Flags & ERDGResourceFlags::MultiFrame) == ERDGResourceFlags::MultiFrame;
 		if (bIsMultiFrameResource && bRequiredWritable && !Texture->bWritable)
 		{
 			RHICmdList.BeginUpdateMultiFrameResource(RHITexture);
@@ -734,20 +773,20 @@ void FRDGBuilder::TransitionTexture(FRDGTexture* Texture, EResourceTransitionAcc
 	}
 }
 
-void FRDGBuilder::TransitionUAV(FUnorderedAccessViewRHIParamRef UAV, FRDGResource* UnderlyingResource, ERDGResourceFlags ResourceFlags, EResourceTransitionAccess TransitionAccess, bool bRequiredCompute ) const
+void FRDGBuilder::TransitionUAV(FUnorderedAccessViewRHIParamRef UAV, FRDGTrackedResource* UnderlyingResource, ERDGResourceFlags ResourceFlags, EResourceTransitionAccess TransitionAccess, bool bRequiredCompute ) const
 {
 	const bool bRequiredWritable = TransitionAccess != EResourceTransitionAccess::EReadable;
 
 	if (bRequiredWritable && UnderlyingResource->bWritable)
 	{
 		// Force a RW barrier between UAV write.
-		// TODO(RDG): allow to have no barriere in the API when multiple pass write concurrently to same resource.
+		// TODO(RDG): allow to have no barrier in the API when multiple pass write concurrently to same resource.
 		EResourceTransitionPipeline TransitionPipeline = CalcTransitionPipeline(UnderlyingResource->bCompute, bRequiredCompute);
 		RHICmdList.TransitionResource(EResourceTransitionAccess::ERWBarrier, TransitionPipeline, UAV);
 	}
 	else if(UnderlyingResource->bWritable != bRequiredWritable || UnderlyingResource->bCompute != bRequiredCompute )
 	{
-		bool bIsMultiFrameResource = uint32(ResourceFlags) & uint32(ERDGResourceFlags::MultiFrame);
+		const bool bIsMultiFrameResource = (ResourceFlags & ERDGResourceFlags::MultiFrame) == ERDGResourceFlags::MultiFrame;
 		if (bIsMultiFrameResource && bRequiredWritable && !UnderlyingResource->bWritable)
 		{
 			RHICmdList.BeginUpdateMultiFrameResource(UAV);
@@ -773,11 +812,11 @@ void FRDGBuilder::ExecutePass(const FRDGPass* Pass)
 	FRHIRenderPassInfo RPInfo;
 	bool bHasRenderTargets = false;
 
-	AllocateAndTransitionPassResources(Pass, &RPInfo, &bHasRenderTargets);
+	PrepareResourcesForExecute(Pass, &RPInfo, &bHasRenderTargets);
 	
-	if (RDG_ENABLE_DEBUG)
 	{
-		UpdateAccessGuardForPassResources(Pass, /** bAllowAccess */ true);
+		const bool bAllowAccess = true;
+		UpdateAccessGuardForPassResources(Pass, bAllowAccess);
 	}
 
 	EventScopeStack.BeginExecutePass(Pass);
@@ -792,107 +831,94 @@ void FRDGBuilder::ExecutePass(const FRDGPass* Pass)
 		UnbindRenderTargets(RHICmdList);
 	}
 	
-	// The name of the pass just for debuging convenience when crashing in the Execute().
-	const TCHAR* PassName = Pass->GetName();
-
 	Pass->Execute(RHICmdList);
 
-	if( bHasRenderTargets )
+	if (bHasRenderTargets)
 	{
 		RHICmdList.EndRenderPass();
 	}
 
 	EventScopeStack.EndExecutePass();
 
-	if (RDG_ENABLE_DEBUG)
 	{
-		UpdateAccessGuardForPassResources(Pass, /** bAllowAccess */ false);
-		WarnForUselessPassDependencies(Pass);
+		const bool bAllowAccess = false;
+		UpdateAccessGuardForPassResources(Pass, bAllowAccess);
 	}
 
+	UnmarkUsedResources(Pass);
+
 	// Can't release resources with immediate mode, because don't know if whether they are gonna be used.
-	if (!GRenderGraphImmediateMode)
+	if (!GRDGImmediateMode)
 	{
-		ReleaseUnecessaryResources(Pass);
+		ReleaseUnreferencedResources(Pass);
 	}
 }
 
-void FRDGBuilder::AllocateAndTransitionPassResources(const FRDGPass* Pass, struct FRHIRenderPassInfo* OutRPInfo, bool* bOutHasRenderTargets)
+void FRDGBuilder::PrepareResourcesForExecute(const FRDGPass* Pass, struct FRHIRenderPassInfo* OutRPInfo, bool* bOutHasRenderTargets)
 {
-	bool bIsCompute = Pass->IsCompute();
-	FShaderParameterStructRef ParameterStruct = Pass->GetParameters();
-
 	OutRPInfo->NumUAVs = 0;
 	OutRPInfo->UAVIndex = 0;
 
+	const bool bIsCompute = Pass->IsCompute();
 	const bool bGeneratingMips = (Pass->GetFlags() & ERDGPassFlags::GenerateMips) == ERDGPassFlags::GenerateMips;
-	for (int ResourceIndex = 0, Num = ParameterStruct.Layout->Resources.Num(); ResourceIndex < Num; ResourceIndex++)
+
+	FRDGPassParameterStruct ParameterStruct = Pass->GetParameters();
+
+	const uint32 ParameterCount = ParameterStruct.GetParameterCount();
+
+	for (uint32 ParameterIndex = 0; ParameterIndex < ParameterCount; ++ParameterIndex)
 	{
-		EUniformBufferBaseType Type = ParameterStruct.Layout->Resources[ResourceIndex].MemberType;
-		uint16 Offset = ParameterStruct.Layout->Resources[ResourceIndex].MemberOffset;
+		FRDGPassParameter Parameter = ParameterStruct.GetParameter(ParameterIndex);
 
-		switch (Type)
+		switch (Parameter.GetType())
 		{
 		case UBMT_RDG_TEXTURE:
 		{
-			FRDGTexture* RESTRICT Texture = *ParameterStruct.GetMemberPtrAtOffset<FRDGTexture*>(Offset);
-			if (Texture)
+			if (FRDGTextureRef Texture = Parameter.GetAsTexture())
 			{
 				// The underlying texture have already been allocated by a prior pass.
 				#if RDG_ENABLE_DEBUG
 				{
-					check(Texture->bHasEverBeenProduced);
-				}	
+					check(Texture->HasBeenProduced());
+					Texture->PassAccessCount++;
+				}
 				#endif
+
 				check(Texture->PooledRenderTarget);
 				check(Texture->CachedRHI.Resource);
 				TransitionTexture(Texture, EResourceTransitionAccess::EReadable, bIsCompute);
-
-				#if RDG_ENABLE_DEBUG
-				{
-					Texture->DebugPassAccessCount++;
-				}
-				#endif
 			}
 		}
 		break;
 		case UBMT_RDG_TEXTURE_SRV:
 		{
-			FRDGTextureSRV* RESTRICT SRV = *ParameterStruct.GetMemberPtrAtOffset<FRDGTextureSRV*>(Offset);
-			if (SRV)
+			if (FRDGTextureSRVRef SRV = Parameter.GetAsTextureSRV())
 			{
+				FRDGTextureRef Texture = SRV->Desc.Texture;
+
 				// The underlying texture have already been allocated by a prior pass.
-				check(SRV->Desc.Texture);
 				#if RDG_ENABLE_DEBUG
 				{
-					check(SRV->Desc.Texture->bHasEverBeenProduced);
-				}	
+					check(Texture->HasBeenProduced());
+					Texture->PassAccessCount++;
+				}
 				#endif
-				check(SRV->Desc.Texture->PooledRenderTarget);
 
 				// Might be the first time using this render graph SRV, so need to setup the cached rhi resource.
 				if (!SRV->CachedRHI.SRV)
 				{
-					SRV->CachedRHI.SRV = SRV->Desc.Texture->PooledRenderTarget->GetRenderTargetItem().MipSRVs[SRV->Desc.MipLevel];
+					SRV->CachedRHI.SRV = Texture->PooledRenderTarget->GetRenderTargetItem().MipSRVs[SRV->Desc.MipLevel];
 				}
 
-				TransitionTexture(SRV->Desc.Texture, EResourceTransitionAccess::EReadable, bIsCompute);
-
-				#if RDG_ENABLE_DEBUG
-				{
-					SRV->Desc.Texture->DebugPassAccessCount++;
-				}
-				#endif
+				TransitionTexture(Texture, EResourceTransitionAccess::EReadable, bIsCompute);
 			}
 		}
 		break;
 		case UBMT_RDG_TEXTURE_UAV:
 		{
-			FRDGTextureUAV* RESTRICT UAV = *ParameterStruct.GetMemberPtrAtOffset<FRDGTextureUAV*>(Offset);
-			if (UAV)
+			if (FRDGTextureUAVRef UAV = Parameter.GetAsTextureUAV())
 			{
-				AllocateRHITextureUAVIfNeeded(UAV, bIsCompute);
-				TransitionUAV(UAV->CachedRHI.UAV, UAV->Desc.Texture, UAV->Desc.Texture->Flags, EResourceTransitionAccess::EWritable, bIsCompute);
+				FRDGTextureRef Texture = UAV->Desc.Texture;
 
 				if (!bIsCompute)
 				{
@@ -901,74 +927,64 @@ void FRDGBuilder::AllocateAndTransitionPassResources(const FRDGPass* Pass, struc
 
 				#if RDG_ENABLE_DEBUG
 				{
-					UAV->Desc.Texture->DebugPassAccessCount++;
+					Texture->PassAccessCount++;
 				}
 				#endif
+
+				AllocateRHITextureUAVIfNeeded(UAV, bIsCompute);
+				TransitionUAV(UAV->CachedRHI.UAV, Texture, Texture->Flags, EResourceTransitionAccess::EWritable, bIsCompute);
 			}
 		}
 		break;
 		case UBMT_RDG_BUFFER:
 		{
-			FRDGBuffer* RESTRICT Buffer = *ParameterStruct.GetMemberPtrAtOffset<FRDGBuffer*>(Offset);
-			if (Buffer)
+			if (FRDGBufferRef Buffer = Parameter.GetAsBuffer())
 			{
 				// The underlying buffer have already been allocated by a prior pass through AllocateRHIBufferUAVIfNeeded().
 				#if RDG_ENABLE_DEBUG
 				{
-					check(Buffer->bHasEverBeenProduced);
+					check(Buffer->HasBeenProduced());
+					Buffer->PassAccessCount++;
 				}
 				#endif
-				check(Buffer->PooledBuffer);
 
-				// TODO(RDG): supper hacky, find the UAV and transition it. Hopefully there is one...
+				// TODO(RDG): super hacky, find the UAV and transition it. Hopefully there is one...
+				check(Buffer->PooledBuffer);
 				check(Buffer->PooledBuffer->UAVs.Num() == 1);
 				FUnorderedAccessViewRHIParamRef UAV = Buffer->PooledBuffer->UAVs.CreateIterator().Value();
 				TransitionUAV(UAV, Buffer, Buffer->Flags, EResourceTransitionAccess::EReadable, bIsCompute);
-
-				#if RDG_ENABLE_DEBUG
-				{
-					Buffer->DebugPassAccessCount++;
-				}
-				#endif
 			}
 		}
 		break;
 		case UBMT_RDG_BUFFER_SRV:
 		{
-			FRDGBufferSRV* RESTRICT SRV = *ParameterStruct.GetMemberPtrAtOffset<FRDGBufferSRV*>(Offset);
-			if (SRV)
+			if (FRDGBufferSRVRef SRV = Parameter.GetAsBufferSRV())
 			{
+				FRDGBufferRef Buffer = SRV->Desc.Buffer;
+
 				// The underlying buffer have already been allocated by a prior pass through AllocateRHIBufferUAVIfNeeded().
-				check(SRV->Desc.Buffer);
 				#if RDG_ENABLE_DEBUG
 				{
-					check(SRV->Desc.Buffer->bHasEverBeenProduced);
+					check(Buffer->HasBeenProduced());
+					Buffer->PassAccessCount++;
 				}
 				#endif
-				check(SRV->Desc.Buffer->PooledBuffer);
-				
-				AllocateRHIBufferSRVIfNeeded(SRV, bIsCompute);
 
-				// TODO(RDG): supper hacky, find the UAV and transition it. Hopefully there is one...
-				check(SRV->Desc.Buffer->PooledBuffer->UAVs.Num() == 1);
-				FUnorderedAccessViewRHIParamRef UAV = SRV->Desc.Buffer->PooledBuffer->UAVs.CreateIterator().Value();
-				TransitionUAV(UAV, SRV->Desc.Buffer, SRV->Desc.Buffer->Flags, EResourceTransitionAccess::EReadable, bIsCompute);
+				// TODO(RDG): super hacky, find the UAV and transition it. Hopefully there is one...
+				check(Buffer->PooledBuffer);
+				check(Buffer->PooledBuffer->UAVs.Num() == 1);
+				FUnorderedAccessViewRHIParamRef UAV = Buffer->PooledBuffer->UAVs.CreateIterator().Value();
 
-				#if RDG_ENABLE_DEBUG
-				{
-					SRV->Desc.Buffer->DebugPassAccessCount++;
-				}
-				#endif
+				AllocateRHIBufferSRVIfNeeded(SRV, bIsCompute);
+				TransitionUAV(UAV, Buffer, Buffer->Flags, EResourceTransitionAccess::EReadable, bIsCompute);
 			}
 		}
 		break;
 		case UBMT_RDG_BUFFER_UAV:
 		{
-			FRDGBufferUAV* RESTRICT UAV = *ParameterStruct.GetMemberPtrAtOffset<FRDGBufferUAV*>(Offset);
-			if (UAV)
+			if (FRDGBufferUAVRef UAV = Parameter.GetAsBufferUAV())
 			{
-				AllocateRHIBufferUAVIfNeeded(UAV, bIsCompute);
-				TransitionUAV(UAV->CachedRHI.UAV, UAV->Desc.Buffer, UAV->Desc.Buffer->Flags, EResourceTransitionAccess::EWritable, bIsCompute);
+				FRDGBufferRef Buffer = UAV->Desc.Buffer;
 
 				if (!bIsCompute)
 				{
@@ -977,9 +993,12 @@ void FRDGBuilder::AllocateAndTransitionPassResources(const FRDGPass* Pass, struc
 				
 				#if RDG_ENABLE_DEBUG
 				{
-					UAV->Desc.Buffer->DebugPassAccessCount++;
+					Buffer->PassAccessCount++;
 				}
 				#endif
+
+				AllocateRHIBufferUAVIfNeeded(UAV, bIsCompute);
+				TransitionUAV(UAV->CachedRHI.UAV, Buffer, Buffer->Flags, EResourceTransitionAccess::EWritable, bIsCompute);
 			}
 		}
 		break;
@@ -987,38 +1006,44 @@ void FRDGBuilder::AllocateAndTransitionPassResources(const FRDGPass* Pass, struc
 		{
 			check(!bIsCompute);
 
-			FRenderTargetBindingSlots* RESTRICT RenderTargets = ParameterStruct.GetMemberPtrAtOffset<FRenderTargetBindingSlots>(Offset);
+			const FRenderTargetBindingSlots& RenderTargetBindingSlots = Parameter.GetAsRenderTargetBindingSlots();
+			const auto& RenderTargets = RenderTargetBindingSlots.Output;
+			const auto& DepthStencil = RenderTargetBindingSlots.DepthStencil;
+			const uint32 RenderTargetCount = RenderTargets.Num();
 
-			uint32 NumRenderTargets = 0;
-			uint32 NumDepthStencilTargets = 0;
-			uint32 NumSamples = 0;
+			uint32 ValidRenderTargetCount = 0;
+			uint32 ValidDepthStencilCount = 0;
+			uint32 SampleCount = 0;
 
-			for (int32 i = 0; i < RenderTargets->Output.Num(); i++)
+			for (uint32 RenderTargetIndex = 0; RenderTargetIndex < RenderTargetCount; RenderTargetIndex++)
 			{
-				const FRenderTargetBinding& RenderTarget = RenderTargets->Output[i];
-				if (RenderTarget.GetTexture())
+				const FRenderTargetBinding& RenderTarget = RenderTargets[RenderTargetIndex];
+
+				if (FRDGTextureRef Texture = RenderTarget.GetTexture())
 				{
-					AllocateRHITextureIfNeeded(RenderTarget.GetTexture(), false);
+					AllocateRHITextureIfNeeded(Texture, bIsCompute);
+
+					auto& OutRenderTarget = OutRPInfo->ColorRenderTargets[RenderTargetIndex];
 
 					// TODO(RDG): should force TargetableTexture == ShaderResourceTexture with MSAA, and instead have an explicit MSAA resolve pass.
-					OutRPInfo->ColorRenderTargets[i].RenderTarget = RenderTarget.GetTexture()->PooledRenderTarget->GetRenderTargetItem().TargetableTexture;
-					OutRPInfo->ColorRenderTargets[i].ResolveTarget = nullptr;
-					OutRPInfo->ColorRenderTargets[i].ArraySlice = -1;
-					OutRPInfo->ColorRenderTargets[i].MipIndex = RenderTarget.GetMipIndex();
-					OutRPInfo->ColorRenderTargets[i].Action = MakeRenderTargetActions(RenderTarget.GetLoadAction(), RenderTarget.GetStoreAction());
+					OutRenderTarget.RenderTarget = Texture->PooledRenderTarget->GetRenderTargetItem().TargetableTexture;
+					OutRenderTarget.ResolveTarget = nullptr;
+					OutRenderTarget.ArraySlice = -1;
+					OutRenderTarget.MipIndex = RenderTarget.GetMipIndex();
+					OutRenderTarget.Action = MakeRenderTargetActions(RenderTarget.GetLoadAction(), RenderTarget.GetStoreAction());
 
 					if (!bGeneratingMips)
 					{
 						// Implicit assurance the RHI will do the correct transitions
-						TransitionTexture(RenderTarget.GetTexture(), EResourceTransitionAccess::EWritable, false);
+						TransitionTexture(Texture, EResourceTransitionAccess::EWritable, bIsCompute);
 					}
 
-					NumSamples |= OutRPInfo->ColorRenderTargets[i].RenderTarget->GetNumSamples();
-					NumRenderTargets++;
-					
+					SampleCount |= OutRenderTarget.RenderTarget->GetNumSamples();
+					ValidRenderTargetCount++;
+
 					#if RDG_ENABLE_DEBUG
 					{
-						RenderTarget.GetTexture()->DebugPassAccessCount++;
+						Texture->PassAccessCount++;
 					}
 					#endif
 				}
@@ -1028,35 +1053,37 @@ void FRDGBuilder::AllocateAndTransitionPassResources(const FRDGPass* Pass, struc
 				}
 			}
 
-			OutRPInfo->UAVIndex = NumRenderTargets;
+			OutRPInfo->UAVIndex = ValidRenderTargetCount;
+
+			if (FRDGTextureRef Texture = DepthStencil.GetTexture())
 
-			const FDepthStencilBinding& DepthStencil = RenderTargets->DepthStencil;
-			if (DepthStencil.GetTexture())
 			{
-				AllocateRHITextureIfNeeded(DepthStencil.GetTexture(), false);
+				AllocateRHITextureIfNeeded(Texture, false);
+
+				auto& OutDepthStencil = OutRPInfo->DepthStencilRenderTarget;
 
-				OutRPInfo->DepthStencilRenderTarget.DepthStencilTarget = DepthStencil.GetTexture()->PooledRenderTarget->GetRenderTargetItem().TargetableTexture;
-				OutRPInfo->DepthStencilRenderTarget.ResolveTarget = nullptr;
-				OutRPInfo->DepthStencilRenderTarget.Action = MakeDepthStencilTargetActions(
+				OutDepthStencil.DepthStencilTarget = Texture->PooledRenderTarget->GetRenderTargetItem().TargetableTexture;
+				OutDepthStencil.ResolveTarget = nullptr;
+				OutDepthStencil.Action = MakeDepthStencilTargetActions(
 					MakeRenderTargetActions(DepthStencil.GetDepthLoadAction(), DepthStencil.GetDepthStoreAction()),
 					MakeRenderTargetActions(DepthStencil.GetStencilLoadAction(), DepthStencil.GetStencilStoreAction()));
-				OutRPInfo->DepthStencilRenderTarget.ExclusiveDepthStencil = DepthStencil.GetDepthStencilAccess();
+				OutDepthStencil.ExclusiveDepthStencil = DepthStencil.GetDepthStencilAccess();
 
-				TransitionTexture(DepthStencil.GetTexture(), EResourceTransitionAccess::EWritable, false);
+				TransitionTexture(Texture, EResourceTransitionAccess::EWritable, bIsCompute);
+
+				SampleCount |= OutDepthStencil.DepthStencilTarget->GetNumSamples();
+				ValidDepthStencilCount++;
 
-				NumSamples |= OutRPInfo->DepthStencilRenderTarget.DepthStencilTarget->GetNumSamples();
-				NumDepthStencilTargets++;
-					
 				#if RDG_ENABLE_DEBUG
 				{
-					DepthStencil.GetTexture()->DebugPassAccessCount++;
+					Texture->PassAccessCount++;
 				}
 				#endif
 			}
 
-			OutRPInfo->bIsMSAA = NumSamples > 1;
+			OutRPInfo->bIsMSAA = SampleCount > 1;
 
-			*bOutHasRenderTargets = NumRenderTargets + NumDepthStencilTargets > 0;
+			*bOutHasRenderTargets = ValidRenderTargetCount + ValidDepthStencilCount > 0;
 		}
 		break;
 		default:
@@ -1067,124 +1094,122 @@ void FRDGBuilder::AllocateAndTransitionPassResources(const FRDGPass* Pass, struc
 	OutRPInfo->bGeneratingMips = bGeneratingMips;
 }
 
-// static
 void FRDGBuilder::UpdateAccessGuardForPassResources(const FRDGPass* Pass, bool bAllowAccess)
 {
 #if RDG_ENABLE_DEBUG
-	FShaderParameterStructRef ParameterStruct = Pass->GetParameters();
+	FRDGPassParameterStruct ParameterStruct = Pass->GetParameters();
+
+	const uint32 ParameterCount = ParameterStruct.GetParameterCount();
 
-	for (int ResourceIndex = 0, Num = ParameterStruct.Layout->Resources.Num(); ResourceIndex < Num; ResourceIndex++)
+	for (uint32 ParameterIndex = 0; ParameterIndex < ParameterCount; ++ParameterIndex)
 	{
-		EUniformBufferBaseType Type = ParameterStruct.Layout->Resources[ResourceIndex].MemberType;
-		uint16 Offset = ParameterStruct.Layout->Resources[ResourceIndex].MemberOffset;
+		FRDGPassParameter Parameter = ParameterStruct.GetParameter(ParameterIndex);
 
-		if (IsRDGResourceReferenceShaderParameterType(Type))
+		if (Parameter.IsResource())
 		{
-			FRDGResource* RESTRICT Resource = *ParameterStruct.GetMemberPtrAtOffset<FRDGResource*>(Offset);
-			if (Resource)
-				Resource->bAllowAccessToRHIResource = bAllowAccess;
+			if (FRDGResourceRef Resource = Parameter.GetAsResource())
+			{
+				Resource->bAllowRHIAccess = bAllowAccess;
+			}
 		}
-		else if (Type == UBMT_RENDER_TARGET_BINDING_SLOTS)
+		else if (Parameter.GetType() == UBMT_RENDER_TARGET_BINDING_SLOTS)
 		{
-			FRenderTargetBindingSlots* RESTRICT RenderTargets = ParameterStruct.GetMemberPtrAtOffset<FRenderTargetBindingSlots>(Offset);
+			const FRenderTargetBindingSlots& RenderTargetBindingSlots = Parameter.GetAsRenderTargetBindingSlots();
+			const auto& RenderTargets = RenderTargetBindingSlots.Output;
+			const auto& DepthStencil = RenderTargetBindingSlots.DepthStencil;
+			const uint32 RenderTargetCount = RenderTargets.Num();
 
-			for (int32 i = 0; i < RenderTargets->Output.Num(); i++)
+			for (uint32 RenderTargetIndex = 0; RenderTargetIndex < RenderTargetCount; RenderTargetIndex++)
 			{
-				FRDGTexture* Texture = RenderTargets->Output[i].GetTexture();
-				if (Texture)
-					Texture->bAllowAccessToRHIResource = bAllowAccess;
+				const FRenderTargetBinding& RenderTarget = RenderTargets[RenderTargetIndex];
+
+				if (FRDGTextureRef Texture = RenderTarget.GetTexture())
+				{
+					Texture->bAllowRHIAccess = bAllowAccess;
+				}
 				else
+				{
 					break;
+				}
 			}
 			
-			const FDepthStencilBinding& DepthStencil = RenderTargets->DepthStencil;
-			if (DepthStencil.GetTexture())
+			if (FRDGTextureRef Texture = DepthStencil.GetTexture())
 			{
-				DepthStencil.GetTexture()->bAllowAccessToRHIResource = bAllowAccess;
+				Texture->bAllowRHIAccess = bAllowAccess;
 			}
 		}
 	}
 #endif
 }
 
-// static 
-void FRDGBuilder::WarnForUselessPassDependencies(const FRDGPass* Pass)
+void FRDGBuilder::UnmarkUsedResources(const FRDGPass* Pass)
 {
-	if (!GRenderGraphEmitWarnings)
-	{
-		return;
-	}
-
-	FShaderParameterStructRef ParameterStruct = Pass->GetParameters();
-
-	int32 TrackedResourceCount = 0;
-	int32 UsedResourceCount = 0;
-
-	// First pass to count resources.
-	for (int ResourceIndex = 0, Num = ParameterStruct.Layout->Resources.Num(); ResourceIndex < Num; ResourceIndex++)
-	{
-		EUniformBufferBaseType Type = ParameterStruct.Layout->Resources[ResourceIndex].MemberType;
-		uint16 Offset = ParameterStruct.Layout->Resources[ResourceIndex].MemberOffset;
-
-		if (!IsRDGResourceReferenceShaderParameterType(Type))
-			continue;
-
-		const FRDGResource* Resource = *ParameterStruct.GetMemberPtrAtOffset<const FRDGResource*>(Offset);
-
-		if (!Resource)
-			continue;
+#if RDG_ENABLE_DEBUG
+	FRDGPassParameterStruct ParameterStruct = Pass->GetParameters();
 
-		TrackedResourceCount++;
-		UsedResourceCount += Resource->bIsActuallyUsedByPass ? 1 : 0;
-	}
+	const uint32 ParameterCount = ParameterStruct.GetParameterCount();
 
-	if (TrackedResourceCount != UsedResourceCount)
+	if (GRDGEmitWarnings)
 	{
-		FString WarningMessage = FString::Printf(
-			TEXT("%i of the %i resources of the pass %s where not actually used."),
-			TrackedResourceCount - UsedResourceCount, TrackedResourceCount, Pass->GetName());
+		uint32 TrackedResourceCount = 0;
+		uint32 UsedResourceCount = 0;
 
-		for (int ResourceIndex = 0, Num = ParameterStruct.Layout->Resources.Num(); ResourceIndex < Num; ResourceIndex++)
+		for (uint32 ParameterIndex = 0; ParameterIndex < ParameterCount; ++ParameterIndex)
 		{
-			EUniformBufferBaseType Type = ParameterStruct.Layout->Resources[ResourceIndex].MemberType;
-			uint16 Offset = ParameterStruct.Layout->Resources[ResourceIndex].MemberOffset;
-
-			if (!IsRDGResourceReferenceShaderParameterType(Type))
-				continue;
-
-			const FRDGResource* Resource = *ParameterStruct.GetMemberPtrAtOffset<const FRDGResource*>(Offset);
+			FRDGPassParameter Parameter = ParameterStruct.GetParameter(ParameterIndex);
 
-			if (!Resource)
-				continue;
-
-			if (!Resource->bIsActuallyUsedByPass)
+			if (Parameter.IsResource())
 			{
-				WarningMessage += FString::Printf(TEXT("\n    %s"), Resource->Name);
+				if (FRDGResourceRef Resource = Parameter.GetAsResource())
+				{
+					TrackedResourceCount++;
+					UsedResourceCount += Resource->bIsActuallyUsedByPass ? 1 : 0;
+				}
 			}
 		}
 
-		EmitRenderGraphWarning(WarningMessage);
-	}
+		if (TrackedResourceCount != UsedResourceCount)
+		{
+			FString WarningMessage = FString::Printf(
+				TEXT("'%d' of the '%d' resources of the pass '%s' where not actually used."),
+				TrackedResourceCount - UsedResourceCount, TrackedResourceCount, Pass->GetName());
 
-	// Last pass to clean the bIsActuallyUsedByPass flags.
-	for (int ResourceIndex = 0, Num = ParameterStruct.Layout->Resources.Num(); ResourceIndex < Num; ResourceIndex++)
-	{
-		EUniformBufferBaseType Type = ParameterStruct.Layout->Resources[ResourceIndex].MemberType;
-		uint16 Offset = ParameterStruct.Layout->Resources[ResourceIndex].MemberOffset;
+			for (uint32 ParameterIndex = 0; ParameterIndex < ParameterCount; ++ParameterIndex)
+			{
+				FRDGPassParameter Parameter = ParameterStruct.GetParameter(ParameterIndex);
 
-		if (!IsRDGResourceReferenceShaderParameterType(Type))
-			continue;
+				if (Parameter.IsResource())
+				{
+					if (const FRDGResourceRef Resource = Parameter.GetAsResource())
+					{
+						if (!Resource->bIsActuallyUsedByPass)
+						{
+							WarningMessage += FString::Printf(TEXT("\n    %s"), Resource->Name);
+						}
+					}
+				}
+			}
 
-		FRDGResource* Resource = *ParameterStruct.GetMemberPtrAtOffset<FRDGResource*>(Offset);
+			EmitRDGWarning(WarningMessage);
+		}
+	}
 
-		if (!Resource)
-			continue;
+	for (uint32 ParameterIndex = 0; ParameterIndex < ParameterCount; ++ParameterIndex)
+	{
+		FRDGPassParameter Parameter = ParameterStruct.GetParameter(ParameterIndex);
 
-		Resource->bIsActuallyUsedByPass = false;
+		if (Parameter.IsResource())
+		{
+			if (const FRDGResourceRef Resource = Parameter.GetAsResource())
+			{
+				Resource->bIsActuallyUsedByPass = false;
+			}
+		}
 	}
+#endif
 }
 
-void FRDGBuilder::ReleaseRHITextureIfPossible(FRDGTexture* Texture)
+void FRDGBuilder::ReleaseRHITextureIfUnreferenced(FRDGTexture* Texture)
 {
 	check(Texture->ReferenceCount > 0);
 	Texture->ReferenceCount--;
@@ -1197,7 +1222,7 @@ void FRDGBuilder::ReleaseRHITextureIfPossible(FRDGTexture* Texture)
 	}
 }
 
-void FRDGBuilder::ReleaseRHIBufferIfPossible(FRDGBuffer* Buffer)
+void FRDGBuilder::ReleaseRHIBufferIfUnreferenced(FRDGBuffer* Buffer)
 {
 	check(Buffer->ReferenceCount > 0);
 	Buffer->ReferenceCount--;
@@ -1210,92 +1235,90 @@ void FRDGBuilder::ReleaseRHIBufferIfPossible(FRDGBuffer* Buffer)
 	}
 }
 
-void FRDGBuilder::ReleaseUnecessaryResources(const FRDGPass* Pass)
+void FRDGBuilder::ReleaseUnreferencedResources(const FRDGPass* Pass)
 {
-	FShaderParameterStructRef ParameterStruct = Pass->GetParameters();
+	FRDGPassParameterStruct ParameterStruct = Pass->GetParameters();
+
+	const uint32 ParameterCount = ParameterStruct.GetParameterCount();
 
-	/** Increments all the FRDGResource::ReferenceCount. */
-	// TODO(RDG): Investigate the cost of branch miss-prediction.
-	for (int ResourceIndex = 0, Num = ParameterStruct.Layout->Resources.Num(); ResourceIndex < Num; ResourceIndex++)
+	for (uint32 ParameterIndex = 0; ParameterIndex < ParameterCount; ++ParameterIndex)
 	{
-		EUniformBufferBaseType Type = ParameterStruct.Layout->Resources[ResourceIndex].MemberType;
-		uint16 Offset = ParameterStruct.Layout->Resources[ResourceIndex].MemberOffset;
+		FRDGPassParameter Parameter = ParameterStruct.GetParameter(ParameterIndex);
 
-		switch (Type)
+		switch (Parameter.GetType())
 		{
 		case UBMT_RDG_TEXTURE:
 		{
-			FRDGTexture* RESTRICT Texture = *ParameterStruct.GetMemberPtrAtOffset<FRDGTexture*>(Offset);
-			if (Texture)
+			if (FRDGTextureRef Texture = Parameter.GetAsTexture())
 			{
-				ReleaseRHITextureIfPossible(Texture);
+				ReleaseRHITextureIfUnreferenced(Texture);
 			}
 		}
 		break;
 		case UBMT_RDG_TEXTURE_SRV:
 		{
-			FRDGTextureSRV* RESTRICT SRV = *ParameterStruct.GetMemberPtrAtOffset<FRDGTextureSRV*>(Offset);
-			if (SRV)
+			if (FRDGTextureSRVRef SRV = Parameter.GetAsTextureSRV())
 			{
-				ReleaseRHITextureIfPossible(SRV->Desc.Texture);
+				ReleaseRHITextureIfUnreferenced(SRV->Desc.Texture);
 			}
 		}
 		break;
 		case UBMT_RDG_TEXTURE_UAV:
 		{
-			FRDGTextureUAV* RESTRICT UAV = *ParameterStruct.GetMemberPtrAtOffset<FRDGTextureUAV*>(Offset);
-			if (UAV)
+			if (FRDGTextureUAVRef UAV = Parameter.GetAsTextureUAV())
 			{
-				ReleaseRHITextureIfPossible(UAV->Desc.Texture);
+				ReleaseRHITextureIfUnreferenced(UAV->Desc.Texture);
 			}
 		}
 		break;
 		case UBMT_RDG_BUFFER:
 		{
-			FRDGBuffer* RESTRICT Buffer = *ParameterStruct.GetMemberPtrAtOffset<FRDGBuffer*>(Offset);
-			if (Buffer)
+			if (FRDGBufferRef Buffer = Parameter.GetAsBuffer())
 			{
-				ReleaseRHIBufferIfPossible(Buffer);
+				ReleaseRHIBufferIfUnreferenced(Buffer);
 			}
 		}
 		break;
 		case UBMT_RDG_BUFFER_SRV:
 		{
-			FRDGBufferSRV* RESTRICT SRV = *ParameterStruct.GetMemberPtrAtOffset<FRDGBufferSRV*>(Offset);
-			if (SRV)
+			if (FRDGBufferSRVRef SRV = Parameter.GetAsBufferSRV())
 			{
-				ReleaseRHIBufferIfPossible(SRV->Desc.Buffer);
+				ReleaseRHIBufferIfUnreferenced(SRV->Desc.Buffer);
 			}
 		}
 		break;
 		case UBMT_RDG_BUFFER_UAV:
 		{
-			FRDGBufferUAV* RESTRICT UAV = *ParameterStruct.GetMemberPtrAtOffset<FRDGBufferUAV*>(Offset);
-			if (UAV)
+			if (FRDGBufferUAVRef UAV = Parameter.GetAsBufferUAV())
 			{
-				ReleaseRHIBufferIfPossible(UAV->Desc.Buffer);
+				ReleaseRHIBufferIfUnreferenced(UAV->Desc.Buffer);
 			}
 		}
 		break;
 		case UBMT_RENDER_TARGET_BINDING_SLOTS:
 		{
-			FRenderTargetBindingSlots* RESTRICT RenderTargets = ParameterStruct.GetMemberPtrAtOffset<FRenderTargetBindingSlots>(Offset);
+			const FRenderTargetBindingSlots& RenderTargetBindingSlots = Parameter.GetAsRenderTargetBindingSlots();
+			const auto& RenderTargets = RenderTargetBindingSlots.Output;
+			const auto& DepthStencil = RenderTargetBindingSlots.DepthStencil;
+			const uint32 RenderTargetCount = RenderTargets.Num();
 
-			for (int32 i = 0; i < RenderTargets->Output.Num(); i++)
+			for (uint32 RenderTargetIndex = 0; RenderTargetIndex < RenderTargetCount; RenderTargetIndex++)
 			{
-				const FRenderTargetBinding& RenderTarget = RenderTargets->Output[i];
-				if (RenderTarget.GetTexture())
+				const FRenderTargetBinding& RenderTarget = RenderTargets[RenderTargetIndex];
+
+				if (FRDGTextureRef Texture = RenderTarget.GetTexture())
 				{
-					ReleaseRHITextureIfPossible(RenderTarget.GetTexture());
+					ReleaseRHITextureIfUnreferenced(Texture);
 				}
 				else
 				{
 					break;
 				}
 			}
-			if (RenderTargets->DepthStencil.GetTexture())
+
+			if (FRDGTextureRef Texture = DepthStencil.GetTexture())
 			{
-				ReleaseRHITextureIfPossible(RenderTargets->DepthStencil.GetTexture());
+				ReleaseRHITextureIfUnreferenced(Texture);
 			}
 		}
 		break;
@@ -1322,14 +1345,14 @@ void FRDGBuilder::ProcessDeferredInternalResourceQueries()
 		{
 			// Increment the number of times the texture has been accessed to avoid warning on produced but never used resources that were produced
 			// only to be extracted for the graph.
-			Query.Texture->DebugPassAccessCount += 1;
+			Query.Texture->PassAccessCount += 1;
 		}
 		#endif
 
 		// No need to manually release in immediate mode, since it is done directly when emptying AllocatedTextures in DestructPasses().
-		if (!GRenderGraphImmediateMode)
+		if (!GRDGImmediateMode)
 		{
-			ReleaseRHITextureIfPossible(Query.Texture);
+			ReleaseRHITextureIfUnreferenced(Query.Texture);
 		}
 	}
 
@@ -1337,18 +1360,18 @@ void FRDGBuilder::ProcessDeferredInternalResourceQueries()
 	{
 		*Query.OutBufferPtr = AllocatedBuffers.FindChecked(Query.Buffer);
 
-#if RDG_ENABLE_DEBUG
+		#if RDG_ENABLE_DEBUG
 		{
 			// Increment the number of times the buffer has been accessed to avoid warning on produced but never used resources that were produced
 			// only to be extracted for the graph.
-			Query.Buffer->DebugPassAccessCount += 1;
+			Query.Buffer->PassAccessCount += 1;
 		}
-#endif
+		#endif
 
 		// No need to manually release in immediate mode, since it is done directly when emptying AllocatedBuffer in DestructPasses().
-		if (!GRenderGraphImmediateMode)
+		if (!GRDGImmediateMode)
 		{
-			ReleaseRHIBufferIfPossible(Query.Buffer);
+			ReleaseRHIBufferIfUnreferenced(Query.Buffer);
 		}
 	}
 }
@@ -1357,32 +1380,34 @@ void FRDGBuilder::DestructPasses()
 {
 	#if RDG_ENABLE_DEBUG
 	{
-		// Make sure all resource references have been released to ensure no leaks happen,
-		// and emit warning if produced resource has not been used.
-		for (const FRDGResource* Resource : Resources)
+		const bool bEmitWarnings = GRDGEmitWarnings;
+
+		for (const FRDGTrackedResourceRef Resource : TrackedResources)
 		{
 			check(Resource->ReferenceCount == 0);
 
-			if (GRenderGraphEmitWarnings && Resource->DebugPassAccessCount == 1 && Resource->DebugFirstProducer)
+			const bool bProducedButNeverUsed = Resource->PassAccessCount == 1 && Resource->FirstProducer;
+
+			if (bEmitWarnings && bProducedButNeverUsed)
 			{
-				check(Resource->bHasEverBeenProduced);
+				check(Resource->HasBeenProduced());
 
-				EmitRenderGraphWarningf(
+				EmitRDGWarningf(
 					TEXT("Resources %s has been produced by the pass %s, but never used by another pass."),
-					Resource->Name, Resource->DebugFirstProducer->GetName());
+					Resource->Name, Resource->FirstProducer->GetName());
 			}
 		}
-		Resources.Empty();
+
+		TrackedResources.Empty();
 	}
 	#endif
 
-	// Passes are allocated on FMemStack, so need to call destructor manually.
-	for (int32 Index = Passes.Num() - 1; Index >= 0; --Index)
+	for (int32 PassIndex = Passes.Num() - 1; PassIndex >= 0; --PassIndex)
 	{
-		Passes[Index]->~FRDGPass();
+		Passes[PassIndex]->~FRDGPass();
 	}
-
 	Passes.Empty();
+
 	DeferredInternalTextureQueries.Empty();
 	DeferredInternalBufferQueries.Empty();
 	AllocatedTextures.Empty();
@@ -1397,13 +1422,4 @@ void FRDGBuilder::BeginEventScope(FRDGEventName&& ScopeName)
 void FRDGBuilder::EndEventScope()
 {
 	EventScopeStack.EndEventScope();
-}
-
-FRDGBuilder::~FRDGBuilder()
-{
-	#if RDG_ENABLE_DEBUG
-	{
-		checkf(bHasExecuted, TEXT("Render graph execution is required to ensure consistency with immediate mode."));
-	}
-	#endif
-}
+}
\ No newline at end of file
diff --git a/Engine/Source/Runtime/RenderCore/Private/RenderGraphPass.cpp b/Engine/Source/Runtime/RenderCore/Private/RenderGraphPass.cpp
new file mode 100644
--- /dev/null
+++ b/Engine/Source/Runtime/RenderCore/Private/RenderGraphPass.cpp
@@ -0,0 +1,46 @@
+// Copyright 1998-2019 Epic Games, Inc. All Rights Reserved.
+
+#include "RenderGraphPass.h"
+
+void FRDGPassParameterStruct::ClearUniformBuffers() const
+{
+	/** The pass parameter struct is mostly POD, with the exception of uniform buffers.
+	 *  Since the explicit type of the struct is unknown, the method traverses and destructs
+	 *  all uniform buffer references manually.
+	 */
+	const auto& Resources = Layout->Resources;
+
+	for (int32 ResourceIndex = 0; ResourceIndex < Resources.Num(); ++ResourceIndex)
+	{
+		const EUniformBufferBaseType MemberType = Resources[ResourceIndex].MemberType;
+
+		if (MemberType == UBMT_REFERENCED_STRUCT)
+		{
+			const uint16 MemberOffset = Resources[ResourceIndex].MemberOffset;
+
+			FUniformBufferRHIRef* UniformBufferPtr = reinterpret_cast<FUniformBufferRHIRef*>(Contents + MemberOffset);
+			*UniformBufferPtr = FUniformBufferRHIRef();
+		}
+	}
+}
+
+FRDGPass::FRDGPass(
+	FRDGEventName&& InName,
+	FRDGPassParameterStruct InParameterStruct,
+	ERDGPassFlags InPassFlags)
+	: Name(static_cast<FRDGEventName&&>(InName))
+	, ParameterStruct(InParameterStruct)
+	, PassFlags(InPassFlags)
+{
+	if (IsCompute())
+	{
+		ensureMsgf(ParameterStruct.GetLayout().NumRenderTargets() == 0, TEXT("Pass %s was declared as ERDGPassFlags::Compute yet has RenderTargets in its ResourceTable"), GetName());
+	}
+}
+
+void FRDGPass::Execute(FRHICommandListImmediate& RHICmdList) const
+{
+	ExecuteImpl(RHICmdList);
+
+	ParameterStruct.ClearUniformBuffers();
+}
\ No newline at end of file
diff --git a/Engine/Source/Runtime/RenderCore/Public/RenderGraph.h b/Engine/Source/Runtime/RenderCore/Public/RenderGraph.h
--- a/Engine/Source/Runtime/RenderCore/Public/RenderGraph.h
+++ b/Engine/Source/Runtime/RenderCore/Public/RenderGraph.h
@@ -58,6 +58,7 @@
 
 #include "RenderGraphDefinitions.h"
 #include "RenderGraphResources.h"
+#include "RenderGraphPass.h"
 #include "RenderGraphBuilder.h"
 #include "RenderGraphUtils.h"
 #include "ShaderParameterStruct.h"
diff --git a/Engine/Source/Runtime/RenderCore/Public/RenderGraphBuilder.h b/Engine/Source/Runtime/RenderCore/Public/RenderGraphBuilder.h
--- a/Engine/Source/Runtime/RenderCore/Public/RenderGraphBuilder.h
+++ b/Engine/Source/Runtime/RenderCore/Public/RenderGraphBuilder.h
@@ -5,143 +5,18 @@
 #include "RenderGraphDefinitions.h"
 #include "RenderGraphEvent.h"
 #include "RenderGraphResources.h"
+#include "RenderGraphPass.h"
 #include "ShaderParameterMacros.h"
 
-/** Flags to annotate passes. */
-enum class ERDGPassFlags
-{
-	None = 0,
-
-	/** Pass uses compute only */
-	Compute = 1 << 0,
-
-	//#todo-rco: Remove this when we can do split/per mip layout transitions.
-	/** Hint to some RHIs this pass will be generating mips to optimize transitions. */
-	GenerateMips = 1 << 1,
-};
-
-ENUM_CLASS_FLAGS(ERDGPassFlags)
-
-// TODO(RDG): Bulk rename across codebase.
-using ERenderGraphPassFlags = ERDGPassFlags;
-
-// TODO(RDG): remove from global scope?
-struct RENDERCORE_API FShaderParameterStructRef
-{
-	const void*						Contents;
-	const FRHIUniformBufferLayout*	Layout;
-
-	template<typename MemberType>
-	MemberType* GetMemberPtrAtOffset(uint16 Offset) const
-	{
-		return reinterpret_cast<MemberType*>(((uint8*)Contents) + Offset);
-	}
-};
-
-/** 
- * Base class of a render graph pass
- */
-class RENDERCORE_API FRDGPass
-{
-public:
-	FRDGPass(
-		FRDGEventName&& InName,
-		FShaderParameterStructRef InParameterStruct,
-		ERDGPassFlags InPassFlags)
-		: Name(static_cast<FRDGEventName&&>(InName))
-		, ParameterStruct(InParameterStruct)
-		, PassFlags(InPassFlags)
-	{
-		if (IsCompute())
-		{
-			ensureMsgf(ParameterStruct.Layout->NumRenderTargets() == 0, TEXT("Pass %s was declared as ERDGPassFlags::Compute yet has RenderTargets in its ResourceTable"), GetName());
-		}
-	}
-
-	virtual ~FRDGPass() {}
-
-	virtual void Execute(FRHICommandListImmediate& RHICmdList) const = 0;
-
-	const TCHAR* GetName() const
-	{
-		return Name.GetTCHAR();
-	}
-
-	ERDGPassFlags GetFlags() const
-	{
-		return PassFlags;
-	}
-
-	bool IsCompute() const
-	{
-		return (PassFlags & ERDGPassFlags::Compute) == ERDGPassFlags::Compute;
-	}
-	
-	FShaderParameterStructRef GetParameters() const
-	{
-		return ParameterStruct;
-	}
-
-	const FRDGEventScope* GetParentScope() const
-	{
-		return ParentScope;
-	}
-
-protected:
-	const FRDGEventName Name;
-	const FRDGEventScope* ParentScope = nullptr;
-	const FShaderParameterStructRef ParameterStruct;
-	const ERDGPassFlags PassFlags;
-
-	friend class FRDGBuilder;
-};
-
-/** 
- * Render graph pass with lambda execute function
- */
-template <typename ParameterStructType, typename ExecuteLambdaType>
-struct TRDGLambdaPass final : public FRDGPass
-{
-	TRDGLambdaPass(
-		FRDGEventName&& InName,
-		FShaderParameterStructRef InParameterStruct,
-		ERDGPassFlags InPassFlags,
-		ExecuteLambdaType&& InExecuteLambda)
-		: FRDGPass(static_cast<FRDGEventName&&>(InName), InParameterStruct, InPassFlags)
-		, ExecuteLambda(static_cast<ExecuteLambdaType&&>(InExecuteLambda))
-	{ }
-
-	~TRDGLambdaPass()
-	{
-		// Manually call the destructor of the pass parameter, to make sure RHI references are released since the pass parameters are allocated on FMemStack.
-		// TODO(RDG): this may lead to RHI resource leaks if a struct allocated in FMemStack does not actually get used through FRDGBuilder::AddPass().
-		ParameterStructType* Struct = reinterpret_cast<ParameterStructType*>(const_cast<void*>(FRDGPass::ParameterStruct.Contents));
-		Struct->~ParameterStructType();
-	}
-
-	virtual void Execute(FRHICommandListImmediate& RHICmdList) const override
-	{
-		ExecuteLambda(RHICmdList);
-	}
-
-	ExecuteLambdaType ExecuteLambda;
-};
-
-/** 
- * Builds the per-frame render graph.
- * Resources must be created from the builder before they can be bound to Pass ResourceTables.
- * These resources are descriptors only until the graph is executed, where RHI resources are allocated as needed.
+/** Builds the per-frame render graph.
+ *  Resources must be created from the builder before they can be bound to Pass ResourceTables.
+ *  These resources are descriptors only until the graph is executed, where RHI resources are allocated as needed.
  */
 class RENDERCORE_API FRDGBuilder
 {
 public:
 	/** A RHI cmd list is required, if using the immediate mode. */
-	FRDGBuilder(FRHICommandListImmediate& InRHICmdList)
-		: RHICmdList(InRHICmdList)
-		, MemStack(FMemStack::Get())
-		, EventScopeStack(RHICmdList)
-	{}
-
+	FRDGBuilder(FRHICommandListImmediate& InRHICmdList);
 	~FRDGBuilder();
 
 	/** Register a external texture to be tracked by the render graph. */
@@ -153,40 +28,45 @@ public:
 			checkf(DebugName, TEXT("Externally allocated texture requires a debug name when registering them to render graph."));
 		}
 		#endif
+
 		FRDGTexture* OutTexture = AllocateForRHILifeTime<FRDGTexture>(DebugName, ExternalPooledTexture->GetDesc(), ERDGResourceFlags::None);
 		OutTexture->PooledRenderTarget = ExternalPooledTexture;
 		OutTexture->CachedRHI.Texture = ExternalPooledTexture->GetRenderTargetItem().ShaderResourceTexture;
 		AllocatedTextures.Add(OutTexture, ExternalPooledTexture);
+
 		#if RDG_ENABLE_DEBUG
 		{
-			OutTexture->bHasEverBeenProduced = true;
-			Resources.Add(OutTexture);
+			OutTexture->MarkAsExternal();
+			TrackedResources.Add(OutTexture);
 		}
 		#endif
+
 		return OutTexture;
 	}
 
 	/** Register a external buffer to be tracked by the render graph. */
 	FORCEINLINE_DEBUGGABLE FRDGBufferRef RegisterExternalBuffer(const TRefCountPtr<FPooledRDGBuffer>& ExternalPooledBuffer, const TCHAR* Name = TEXT("External"))
 	{
-#if RDG_ENABLE_DEBUG
+		#if RDG_ENABLE_DEBUG
 		{
 			ensureMsgf(ExternalPooledBuffer.IsValid(), TEXT("Attempted to register NULL external buffer: %s"), Name);
 		}
-#endif
+		#endif
+
 		FRDGBuffer* OutBuffer = AllocateForRHILifeTime<FRDGBuffer>(Name, ExternalPooledBuffer->Desc, ERDGResourceFlags::None);
 		OutBuffer->PooledBuffer = ExternalPooledBuffer;
 		AllocatedBuffers.Add(OutBuffer, ExternalPooledBuffer);
-#if RDG_ENABLE_DEBUG
+
+		#if RDG_ENABLE_DEBUG
 		{
-			OutBuffer->bHasEverBeenProduced = true;
-			Resources.Add(OutBuffer);
+			OutBuffer->MarkAsExternal();
+			TrackedResources.Add(OutBuffer);
 		}
-#endif
+		#endif
+
 		return OutBuffer;
 	}
 
-
 	/** Create graph tracked resource from a descriptor with a debug name.
 	 *
 	 * The debug name is the name used for GPU debugging tools, but also for the VisualizeTexture/Vis command.
@@ -200,10 +80,15 @@ public:
 			checkf(Desc.Format != PF_Unknown, TEXT("Illegal to create texture %s with an invalid pixel format."), DebugName);
 		}
 		#endif
+
 		FRDGTexture* Texture = AllocateForRHILifeTime<FRDGTexture>(DebugName, Desc, Flags);
+
 		#if RDG_ENABLE_DEBUG
-			Resources.Add(Texture);
+		{
+			TrackedResources.Add(Texture);
+		}
 		#endif
+
 		return Texture;
 	}
 
@@ -219,46 +104,57 @@ public:
 			checkf(DebugName, TEXT("Creating a render graph buffer requires a valid debug name."));
 		}
 		#endif
-		FRDGBuffer* Buffer = AllocateForRHILifeTime<FRDGBuffer>(DebugName, Desc, Flags);
+
+		FRDGBufferRef Buffer = AllocateForRHILifeTime<FRDGBuffer>(DebugName, Desc, Flags);
+
 		#if RDG_ENABLE_DEBUG
-			Resources.Add(Buffer);
+		{
+			TrackedResources.Add(Buffer);
+		}
 		#endif
+
 		return Buffer;
 	}
 
 	/** Create graph tracked SRV for a texture from a descriptor. */
 	FORCEINLINE_DEBUGGABLE FRDGTextureSRVRef CreateSRV(const FRDGTextureSRVDesc& Desc)
 	{
-		check(Desc.Texture);
+		FRDGTextureRef Texture = Desc.Texture;
+
 		#if RDG_ENABLE_DEBUG
 		{
+			if (!Texture)
+			{
+				ensureMsgf(false, TEXT("RenderGraph texture SRV created with a null texture."));
+				return nullptr;
+			}
+
 			ensureMsgf(!bHasExecuted, TEXT("Render graph SRV %s needs to be created before the builder execution."), Desc.Texture->Name);
 			ensureMsgf(Desc.Texture->Desc.TargetableFlags & TexCreate_ShaderResource, TEXT("Attempted to create SRV from texture %s which was not created with TexCreate_ShaderResource"), Desc.Texture->Name);
 		}
 		#endif
-		
-		FRDGTextureSRV* SRV = AllocateForRHILifeTime<FRDGTextureSRV>(Desc.Texture->Name, Desc);
-		#if RDG_ENABLE_DEBUG
-			Resources.Add(SRV);
-		#endif
-		return SRV;
+
+		return AllocateForRHILifeTime<FRDGTextureSRV>(Desc.Texture->Name, Desc);
 	}
 
 	/** Create graph tracked SRV for a buffer from a descriptor. */
 	FORCEINLINE_DEBUGGABLE FRDGBufferSRVRef CreateSRV(const FRDGBufferSRVDesc& Desc)
 	{
-		check(Desc.Buffer);
+		FRDGBufferRef Buffer = Desc.Buffer;
+
 		#if RDG_ENABLE_DEBUG
 		{
+			if (!Buffer)
+			{
+				ensureMsgf(false, TEXT("RenderGraph buffer SRV created with a null buffer."));
+				return nullptr;
+			}
+
 			ensureMsgf(!bHasExecuted, TEXT("Render graph SRV %s needs to be created before the builder execution."), Desc.Buffer->Name);
 		}
 		#endif
-		
-		FRDGBufferSRV* SRV = AllocateForRHILifeTime<FRDGBufferSRV>(Desc.Buffer->Name, Desc);
-		#if RDG_ENABLE_DEBUG
-			Resources.Add(SRV);
-		#endif
-		return SRV;
+
+		return AllocateForRHILifeTime<FRDGBufferSRV>(Buffer->Name, Desc);
 	}
 
 	FORCEINLINE_DEBUGGABLE FRDGBufferSRVRef CreateSRV(FRDGBufferRef Buffer, EPixelFormat Format)
@@ -269,46 +165,52 @@ public:
 	/** Create graph tracked UAV for a texture from a descriptor. */
 	FORCEINLINE_DEBUGGABLE FRDGTextureUAVRef CreateUAV(const FRDGTextureUAVDesc& Desc)
 	{
-		check(Desc.Texture);
+		FRDGTextureRef Texture = Desc.Texture;
+		
 		#if RDG_ENABLE_DEBUG
 		{
-			ensureMsgf(!bHasExecuted, TEXT("Render graph UAV %s needs to be created before the builder execution."), Desc.Texture->Name);
-			ensureMsgf(Desc.Texture->Desc.TargetableFlags & TexCreate_UAV, TEXT("Attempted to create UAV from texture %s which was not created with TexCreate_UAV"), Desc.Texture->Name);
+			if (!Texture)
+			{
+				checkf(false, TEXT("RenderGraph texture UAV created with a null texture."));
+				return nullptr;
+			}
+
+			ensureMsgf(!bHasExecuted, TEXT("Render graph UAV %s needs to be created before the builder execution."), Texture->Name);
+			ensureMsgf(Texture->Desc.TargetableFlags & TexCreate_UAV, TEXT("Attempted to create UAV from texture %s which was not created with TexCreate_UAV"), Texture->Name);
 		}
 		#endif
 		
-		FRDGTextureUAV* UAV = AllocateForRHILifeTime<FRDGTextureUAV>(Desc.Texture->Name, Desc);
-		#if RDG_ENABLE_DEBUG
-			Resources.Add(UAV);
-		#endif
-		return UAV;
+		return AllocateForRHILifeTime<FRDGTextureUAV>(Texture->Name, Desc);
 	}
 
 	FORCEINLINE_DEBUGGABLE FRDGTextureUAVRef CreateUAV(FRDGTextureRef Texture)
 	{
-		return CreateUAV( FRDGTextureUAVDesc(Texture) );
+		return CreateUAV(FRDGTextureUAVDesc(Texture));
 	}
 
 	/** Create graph tracked UAV for a buffer from a descriptor. */
 	FORCEINLINE_DEBUGGABLE FRDGBufferUAVRef CreateUAV(const FRDGBufferUAVDesc& Desc)
 	{
-		check(Desc.Buffer);
+		FRDGBufferRef Buffer = Desc.Buffer;
+
 		#if RDG_ENABLE_DEBUG
 		{
-			ensureMsgf(!bHasExecuted, TEXT("Render graph UAV %s needs to be created before the builder execution."), Desc.Buffer->Name);
+			if (!Buffer)
+			{
+				checkf(false, TEXT("RenderGraph buffer UAV created with a null buffer."));
+				return nullptr;
+			}
+
+			ensureMsgf(!bHasExecuted, TEXT("Render graph UAV %s needs to be created before the builder execution."), Buffer->Name);
 		}
 		#endif
 		
-		FRDGBufferUAV* UAV = AllocateForRHILifeTime<FRDGBufferUAV>(Desc.Buffer->Name, Desc);
-		#if RDG_ENABLE_DEBUG
-			Resources.Add(UAV);
-		#endif
-		return UAV;
+		return AllocateForRHILifeTime<FRDGBufferUAV>(Buffer->Name, Desc);
 	}
 
 	FORCEINLINE_DEBUGGABLE FRDGBufferUAVRef CreateUAV(FRDGBufferRef Buffer, EPixelFormat Format)
 	{
-		return CreateUAV( FRDGBufferUAVDesc(Buffer, Format) );
+		return CreateUAV(FRDGBufferUAVDesc(Buffer, Format));
 	}
 
 	/** Allocates parameter struct specifically to survive through the life time of the render graph. */
@@ -343,10 +245,10 @@ public:
 		ExecuteLambdaType&& ExecuteLambda)
 	{
 		auto NewPass = new(MemStack) TRDGLambdaPass<ParameterStructType, ExecuteLambdaType>(
-			static_cast<FRDGEventName&&>(Name),
-			{ ParameterStruct, &ParameterStructType::FTypeInfo::GetStructMetadata()->GetLayout() },
+			MoveTemp(Name),
+			FRDGPassParameterStruct(ParameterStruct),
 			Flags,
-			static_cast<ExecuteLambdaType&&>(ExecuteLambda) );
+			MoveTemp(ExecuteLambda));
 
 		AddPassInternal(NewPass);
 	}
@@ -378,7 +280,7 @@ public:
 				TEXT("Accessing render graph internal texture %s with QueueTextureExtraction() needs to happen before the builder's execution."),
 				Texture->Name);
 
-			checkf(Texture->bHasEverBeenProduced,
+			checkf(Texture->HasBeenProduced(),
 				TEXT("Unable to queue the extraction of the texture %s because it has not been produced by any pass."),
 				Texture->Name);
 		}
@@ -399,17 +301,17 @@ public:
 	{
 		check(Buffer);
 		check(OutBufferPtr);
-#if RDG_ENABLE_DEBUG
+		#if RDG_ENABLE_DEBUG
 		{
 			checkf(!bHasExecuted,
 				TEXT("Accessing render graph internal buffer %s with QueueBufferExtraction() needs to happen before the builder's execution."),
 				Buffer->Name);
 
-			checkf(Buffer->bHasEverBeenProduced,
+			checkf(Buffer->HasBeenProduced(),
 				TEXT("Unable to queue the extraction of the buffer %s because it has not been produced by any pass."),
 				Buffer->Name);
 		}
-#endif
+		#endif
 		FDeferredInternalBufferQuery Query;
 		Query.Buffer = Buffer;
 		Query.OutBufferPtr = OutBufferPtr;
@@ -428,7 +330,7 @@ public:
 			
 			// Increment the number of time the texture has been accessed to avoid warning on produced but never used resources that were produced
 			// only to be extracted for the graph.
-			Texture->DebugPassAccessCount += 1;
+			Texture->PassAccessCount += 1;
 		}
 		#endif
 	}
@@ -473,20 +375,19 @@ private:
 
 	FRDGEventScopeStack EventScopeStack;
 
-	#if RDG_ENABLE_DEBUG
-		/** Whether the Execute() has already been called. */
-		bool bHasExecuted = false;
+#if RDG_ENABLE_DEBUG
+	/** Whether the Execute() has already been called. */
+	bool bHasExecuted = false;
 
-		/** Lists of all created resources */
-		TArray<const FRDGResource*, SceneRenderingAllocator> Resources;
+	/** All recently allocated pass parameter structure, but not used by a AddPass() yet. */
+	TSet<const void*> AllocatedUnusedPassParameters;
 
-		// All recently allocated pass parameter structure, but not used by a AddPass() yet.
-		TSet<const void*> AllocatedUnusedPassParameters;
-	#endif
+	/** List of tracked resources for validation prior to shutdown. */
+	TArray<FRDGTrackedResourceRef, SceneRenderingAllocator> TrackedResources;
+#endif
 
-	void DebugPass(const FRDGPass* Pass);
 	void ValidatePass(const FRDGPass* Pass) const;
-	void CaptureAnyInterestingPassOutput(const FRDGPass* Pass);
+	void VisualizePassOutputs(const FRDGPass* Pass);
 
 	void WalkGraphDependencies();
 	
@@ -518,17 +419,17 @@ private:
 	void EndEventScope();
 
 	void TransitionTexture(FRDGTexture* Texture, EResourceTransitionAccess TransitionAccess, bool bRequiredCompute) const;
-	void TransitionUAV(FUnorderedAccessViewRHIParamRef UAV, FRDGResource* UnderlyingResource, ERDGResourceFlags ResourceFlags, EResourceTransitionAccess TransitionAccess, bool bRequiredCompute ) const;
+	void TransitionUAV(FUnorderedAccessViewRHIParamRef UAV, FRDGTrackedResource* UnderlyingResource, ERDGResourceFlags ResourceFlags, EResourceTransitionAccess TransitionAccess, bool bRequiredCompute ) const;
 
-	void PushDrawEventStack(const FRDGPass* Pass);
 	void ExecutePass(const FRDGPass* Pass);
-	void AllocateAndTransitionPassResources(const FRDGPass* Pass, struct FRHIRenderPassInfo* OutRPInfo, bool* bOutHasRenderTargets);
-	static void UpdateAccessGuardForPassResources(const FRDGPass* Pass, bool bAllowAccess);
-	static void WarnForUselessPassDependencies(const FRDGPass* Pass);
+	void PrepareResourcesForExecute(const FRDGPass* Pass, struct FRHIRenderPassInfo* OutRPInfo, bool* bOutHasRenderTargets);
+
+	void UpdateAccessGuardForPassResources(const FRDGPass* Pass, bool bAllowAccess);
+	void UnmarkUsedResources(const FRDGPass* Pass);
 
-	void ReleaseRHITextureIfPossible(FRDGTexture* Texture);
-	void ReleaseRHIBufferIfPossible(FRDGBuffer* Buffer);
-	void ReleaseUnecessaryResources(const FRDGPass* Pass);
+	void ReleaseRHITextureIfUnreferenced(FRDGTexture* Texture);
+	void ReleaseRHIBufferIfUnreferenced(FRDGBuffer* Buffer);
+	void ReleaseUnreferencedResources(const FRDGPass* Pass);
 
 	void ProcessDeferredInternalResourceQueries();
 	void DestructPasses();
diff --git a/Engine/Source/Runtime/RenderCore/Public/RenderGraphPass.h b/Engine/Source/Runtime/RenderCore/Public/RenderGraphPass.h
new file mode 100644
--- /dev/null
+++ b/Engine/Source/Runtime/RenderCore/Public/RenderGraphPass.h
@@ -0,0 +1,256 @@
+// Copyright 1998-2019 Epic Games, Inc. All Rights Reserved.
+
+#pragma once
+
+#include "RenderGraphDefinitions.h"
+#include "RenderGraphResources.h"
+#include "RenderGraphEvent.h"
+#include "ShaderParameterMacros.h"
+
+class FRDGPass;
+class FRDGBuilder;
+
+/** A helper class for identifying and accessing a render graph pass parameter. */
+class FRDGPassParameter final
+{
+	friend class FRDGPassParameterStruct;
+public:
+	FRDGPassParameter() = default;
+
+	bool IsResource() const
+	{
+		return IsRDGResourceReferenceShaderParameterType(MemberType);
+	}
+
+	bool IsTrackedResource() const
+	{
+		return MemberType == UBMT_RDG_TEXTURE || MemberType == UBMT_RDG_BUFFER;
+	}
+
+	EUniformBufferBaseType GetType() const
+	{
+		return MemberType;
+	}
+
+	FRDGResourceRef GetAsResource() const
+	{
+		check(IsResource());
+		return *GetAs<FRDGResourceRef>();
+	}
+
+	FRDGTrackedResourceRef GetAsTrackedResource() const
+	{
+		check(IsTrackedResource());
+		return *GetAs<FRDGTrackedResourceRef>();
+	}
+
+	FRDGTextureRef GetAsTexture() const
+	{
+		check(MemberType == UBMT_RDG_TEXTURE);
+		return *GetAs<FRDGTextureRef>();
+	}
+
+	FRDGBufferRef GetAsBuffer() const
+	{
+		check(MemberType == UBMT_RDG_BUFFER);
+		return *GetAs<FRDGBufferRef>();
+	}
+
+	FRDGTextureSRVRef GetAsTextureSRV() const
+	{
+		check(MemberType == UBMT_RDG_TEXTURE_SRV);
+		return *GetAs<FRDGTextureSRVRef>();
+	}
+
+	FRDGBufferSRVRef GetAsBufferSRV() const
+	{
+		check(MemberType == UBMT_RDG_BUFFER_SRV);
+		return *GetAs<FRDGBufferSRVRef>();
+	}
+
+	FRDGTextureUAVRef GetAsTextureUAV() const
+	{
+		check(MemberType == UBMT_RDG_TEXTURE_UAV);
+		return *GetAs<FRDGTextureUAVRef>();
+	}
+
+	FRDGBufferUAVRef GetAsBufferUAV() const
+	{
+		check(MemberType == UBMT_RDG_BUFFER_UAV);
+		return *GetAs<FRDGBufferUAVRef>();
+	}
+
+	const FRenderTargetBindingSlots& GetAsRenderTargetBindingSlots() const
+	{
+		check(MemberType == UBMT_RENDER_TARGET_BINDING_SLOTS);
+		return *GetAs<FRenderTargetBindingSlots>();
+	}
+
+private:
+	FRDGPassParameter(EUniformBufferBaseType InMemberType, void* InMemberPtr)
+		: MemberType(InMemberType)
+		, MemberPtr(InMemberPtr)
+	{}
+
+	template <typename T>
+	T* GetAs() const
+	{
+		return reinterpret_cast<T*>(MemberPtr);
+	}
+
+	const EUniformBufferBaseType MemberType = UBMT_INVALID;
+	void* MemberPtr = nullptr;
+};
+
+/** Wraps a pass parameter struct payload and provides helpers for traversing members. */
+class FRDGPassParameterStruct final
+{
+public:
+	explicit FRDGPassParameterStruct(void* InContents, const FRHIUniformBufferLayout* InLayout)
+		: Contents(reinterpret_cast<uint8*>(InContents))
+		, Layout(InLayout)
+	{
+		checkf(Contents && Layout, TEXT("Pass parameter struct created with null inputs."));
+	}
+
+	template <typename FParameterStruct>
+	explicit FRDGPassParameterStruct(FParameterStruct* Parameters)
+		: FRDGPassParameterStruct(Parameters, &FParameterStruct::FTypeInfo::GetStructMetadata()->GetLayout())
+	{}
+
+	bool IsCompute() const
+	{
+		return Layout->NumRenderTargets() == 0;
+	}
+
+	uint32 GetParameterCount() const
+	{
+		return Layout->Resources.Num();
+	}
+
+	FRDGPassParameter GetParameter(uint32 ParameterIndex) const
+	{
+		const auto& Resources = Layout->Resources;
+		const EUniformBufferBaseType MemberType = Resources[ParameterIndex].MemberType;
+		const uint16 MemberOffset = Resources[ParameterIndex].MemberOffset;
+
+		return FRDGPassParameter(MemberType, Contents + MemberOffset);
+	}
+
+	const void* GetContents() const
+	{
+		return Contents;
+	}
+
+	const FRHIUniformBufferLayout& GetLayout() const
+	{
+		return *Layout;
+	}
+
+private:
+	/** Releases all active uniform buffer references held inside the struct. */
+	RENDERCORE_API void ClearUniformBuffers() const;
+
+	uint8* Contents;
+	const FRHIUniformBufferLayout* Layout;
+
+	friend FRDGPass;
+};
+
+/** Flags to annotate passes. */
+enum class ERDGPassFlags
+{
+	None = 0,
+
+	/** Pass uses compute only */
+	Compute = 1 << 0,
+
+	//#todo-rco: Remove this when we can do split/per mip layout transitions.
+	/** Hint to some RHIs this pass will be generating mips to optimize transitions. */
+	GenerateMips = 1 << 1,
+};
+
+ENUM_CLASS_FLAGS(ERDGPassFlags)
+
+// TODO(RDG): Bulk rename across codebase.
+using ERenderGraphPassFlags = ERDGPassFlags;
+
+/** Base class of a render graph pass. */
+class RENDERCORE_API FRDGPass
+{
+public:
+	FRDGPass(
+		FRDGEventName&& InName,
+		FRDGPassParameterStruct InParameterStruct,
+		ERDGPassFlags InPassFlags);
+
+	FRDGPass(const FRDGPass&) = delete;
+
+	virtual ~FRDGPass() = default;
+
+	const TCHAR* GetName() const
+	{
+		return Name.GetTCHAR();
+	}
+
+	ERDGPassFlags GetFlags() const
+	{
+		return PassFlags;
+	}
+
+	bool IsCompute() const
+	{
+		return (PassFlags & ERDGPassFlags::Compute) == ERDGPassFlags::Compute;
+	}
+
+	FRDGPassParameterStruct GetParameters() const
+	{
+		return ParameterStruct;
+	}
+
+	const FRDGEventScope* GetParentScope() const
+	{
+		return ParentScope;
+	}
+
+private:
+	//////////////////////////////////////////////////////////////////////////
+	//! User Methods to Override
+
+	virtual void ExecuteImpl(FRHICommandListImmediate& RHICmdList) const = 0;
+
+	//////////////////////////////////////////////////////////////////////////
+
+	void Execute(FRHICommandListImmediate& RHICmdList) const;
+
+	const FRDGEventName Name;
+	const FRDGEventScope* ParentScope = nullptr;
+	FRDGPassParameterStruct ParameterStruct;
+	const ERDGPassFlags PassFlags;
+
+	friend FRDGBuilder;
+};
+
+/** Render graph pass with lambda execute function. */
+template <typename ParameterStructType, typename ExecuteLambdaType>
+class TRDGLambdaPass final
+	: public FRDGPass
+{
+public:
+	TRDGLambdaPass(
+		FRDGEventName&& InName,
+		FRDGPassParameterStruct InParameterStruct,
+		ERDGPassFlags InPassFlags,
+		ExecuteLambdaType&& InExecuteLambda)
+		: FRDGPass(MoveTemp(InName), InParameterStruct, InPassFlags)
+		, ExecuteLambda(MoveTemp(InExecuteLambda))
+	{}
+
+private:
+	void ExecuteImpl(FRHICommandListImmediate& RHICmdList) const override
+	{
+		ExecuteLambda(RHICmdList);
+	}
+
+	ExecuteLambdaType ExecuteLambda;
+};
diff --git a/Engine/Source/Runtime/RenderCore/Public/RenderGraphResources.h b/Engine/Source/Runtime/RenderCore/Public/RenderGraphResources.h
--- a/Engine/Source/Runtime/RenderCore/Public/RenderGraphResources.h
+++ b/Engine/Source/Runtime/RenderCore/Public/RenderGraphResources.h
@@ -6,24 +6,22 @@
 
 class FRDGPass;
 class FRDGBuilder;
-
+class FRDGEventName;
 class FRDGResource;
+class FRDGTrackedResource;
 class FRDGTexture;
 class FRDGTextureSRV;
 class FRDGTextureUAV;
-
 class FRDGBuffer;
 class FRDGBufferSRV;
 class FRDGBufferUAV;
 
-class FRDGEventName;
-
 struct FPooledRDGBuffer;
 struct FRayTracingShaderBindingsWriter;
 
-
 /** Defines the RDG resource references for user code not forgetting the const every time. */
 using FRDGResourceRef = FRDGResource*;
+using FRDGTrackedResourceRef = FRDGTrackedResource*;
 using FRDGTextureRef = FRDGTexture*;
 using FRDGTextureSRVRef = FRDGTextureSRV*;
 using FRDGTextureUAVRef = FRDGTextureUAV*;
@@ -31,16 +29,16 @@ using FRDGBufferRef = FRDGBuffer*;
 using FRDGBufferSRVRef = FRDGBufferSRV*;
 using FRDGBufferUAVRef = FRDGBufferUAV*;
 
-
 /** Render graph specific flags for resources. */
 enum class ERDGResourceFlags : uint8
 {
-	None = 0x00,
+	None = 0,
 
 	// Tag the resource to survive through frame, that is important for multi GPU alternate frame rendering.
-	MultiFrame = 0x01,
+	MultiFrame = 1 << 1,
 };
 
+ENUM_CLASS_FLAGS(ERDGResourceFlags);
 
 /** Generic graph resource. */
 class RENDERCORE_API FRDGResource
@@ -55,22 +53,12 @@ protected:
 		FUnorderedAccessViewRHIParamRef UAV;
 	} CachedRHI;
 
-public:
-	// Name of the resource for debugging purpose.
-	const TCHAR* const Name = nullptr;
-
-	FRDGResource(const TCHAR* InDebugName)
-		: Name(InDebugName)
-	{
-		CachedRHI.Resource = nullptr;
-	}
-	
 	/** Verify that the RHI resource can be accessed at a pass execution. */
-	inline void VerifyRHIAccess() const
+	void ValidateRHIAccess() const
 	{
 		#if RDG_ENABLE_DEBUG
 		{
-			checkf(bAllowAccessToRHIResource,
+			checkf(bAllowRHIAccess,
 				TEXT("Accessing the RHI resource of %s at this time is not allowed. If you hit this check in pass, ")
 				TEXT("that is due to this resource not being referenced in the parameters of your pass."),
 				Name);
@@ -78,40 +66,37 @@ public:
 		#endif
 	}
 
-	/** Marks this resource as actually used by a resource. This is to track what dependencies on pass was actually unecessary. */
-	inline void MarkResourceAsUsed()
+	FRDGResource(const TCHAR* InName)
+		: Name(InName)
 	{
-		VerifyRHIAccess();
-		bIsActuallyUsedByPass = true;
+		CachedRHI.Resource = nullptr;
 	}
 
-	FRDGResource() = delete;
+public:
 	FRDGResource(const FRDGResource&) = delete;
-	void operator = (const FRDGResource&) = delete;
 
-private:
-	/** Number of references in passes and deferred queries. */
-	int32 ReferenceCount = 0;
+	/** Marks this resource as actually used by a resource. This is to track what dependencies on pass was actually unnecessary. */
+	void MarkResourceAsUsed()
+	{
+		ValidateRHIAccess();
 
-	// Used for tracking resource state during execution
-	bool bWritable = false;
-	bool bCompute = false;
+		#if RDG_ENABLE_DEBUG
+		{
+			bIsActuallyUsedByPass = true;
+		}
+		#endif
+	}
 
-	/** Boolean to track at runtime whether a ressource is actually used by the lambda of a pass or not, to detect unnecessary resource dependencies on passes. */
-	bool bIsActuallyUsedByPass = false;
+	// Name of the resource for debugging purpose.
+	const TCHAR* const Name = nullptr;
 
+private:
 #if RDG_ENABLE_DEBUG
-	/** Boolean to track at wiring time if a resource has ever been produced by a pass, to error out early if accessing a resource that has not been produced. */
-	bool bHasEverBeenProduced = false;
+	/** Boolean to track at runtime whether a resource is actually used by the lambda of a pass or not, to detect unnecessary resource dependencies on passes. */
+	bool bIsActuallyUsedByPass = false;
 
 	/** Boolean to track at pass execution whether the underlying RHI resource is allowed to be accessed. */
-	bool bAllowAccessToRHIResource = false;
-
-	/** Pointer towards the pass that is the first to produce it, for even more convenient error message. */
-	const FRDGPass* DebugFirstProducer = nullptr;
-
-	/** Count the number of times it has been used by a pass. */
-	int32 DebugPassAccessCount = 0;
+	bool bAllowRHIAccess = false;
 #endif
 
 	friend class FRDGBuilder;
@@ -134,49 +119,106 @@ private:
 	friend void AddPass_ClearUAV(FRDGBuilder&, FRDGEventName&&, FRDGBufferUAVRef, uint32);
 };
 
+/** A render graph resource with an allocation lifetime tracked by the graph. */
+class RENDERCORE_API FRDGTrackedResource
+	: public FRDGResource
+{
+protected:
+	FRDGTrackedResource(const TCHAR* InName, const ERDGResourceFlags InFlags)
+		: FRDGResource(InName)
+		, Flags(InFlags)
+	{}
+
+	/** Flags specific to this resource for render graph. */
+	const ERDGResourceFlags Flags;
+
+private:
+	/** Number of references in passes and deferred queries. */
+	int32 ReferenceCount = 0;
+
+#if RDG_ENABLE_DEBUG
+	void MarkAsProducedBy(const FRDGPass* Pass)
+	{
+		if (!bHasBeenProduced)
+		{
+			bHasBeenProduced = true;
+			FirstProducer = Pass;
+		}
+	}
+
+	void MarkAsExternal()
+	{
+		bHasBeenProduced = true;
+	}
+
+	bool HasBeenProduced() const
+	{
+		return bHasBeenProduced;
+	}
+
+	/** Pointer towards the pass that is the first to produce it, for even more convenient error message. */
+	const FRDGPass* FirstProducer = nullptr;
+
+	/** Count the number of times it has been used by a pass. */
+	int32 PassAccessCount = 0;
+
+	/** Boolean to track at wiring time if a resource has ever been produced by a pass, to error out early if accessing a resource that has not been produced. */
+	bool bHasBeenProduced = false;
+#endif
+
+	/** Used for tracking resource state during execution. */
+	bool bWritable = false;
+	bool bCompute = false;
+
+	friend class FRDGBuilder;
+};
+
 /** Descriptor of a graph tracked texture. */
 using FRDGTextureDesc = FPooledRenderTargetDesc;
 
 /** Render graph tracked Texture. */
-class RENDERCORE_API FRDGTexture : public FRDGResource
+class RENDERCORE_API FRDGTexture final
+	: public FRDGTrackedResource
 {
 public:
 	/** Descriptor of the graph tracked texture. */
 	const FRDGTextureDesc Desc;
 
-	/** Returns the allocated pooled render target. Must only be called within a pass's lambda. */
-	inline IPooledRenderTarget* GetPooledRenderTarget() const
+	//////////////////////////////////////////////////////////////////////////
+	//! The following methods may only be called during pass execution.
+
+	/** Returns the allocated pooled render target. */
+	IPooledRenderTarget* GetPooledRenderTarget() const
 	{
-		VerifyRHIAccess();
+		ValidateRHIAccess();
 		check(PooledRenderTarget);
 		return PooledRenderTarget;
 	}
 
-
-	/** Returns the allocated RHI texture. Must only be called within a pass's lambda. */
-	inline FTextureRHIParamRef GetRHITexture() const
+	/** Returns the allocated RHI texture. */
+	FTextureRHIParamRef GetRHITexture() const
 	{
-		VerifyRHIAccess();
+		ValidateRHIAccess();
 		check(CachedRHI.Texture);
 		return CachedRHI.Texture;
 	}
 
-private:
-	/** This is not a TRefCountPtr<> because FRDGTexture is allocated on the FMemStack
-	 * FGraphBuilder::AllocatedTextures is actually keeping the reference.
-	 */
-	mutable IPooledRenderTarget* PooledRenderTarget = nullptr;
-
-	// Flags specific to this resource for render graph.
-	const ERDGResourceFlags Flags;
-
+	//////////////////////////////////////////////////////////////////////////
 
-	FRDGTexture(const TCHAR* DebugName, const FPooledRenderTargetDesc& InDesc, ERDGResourceFlags InFlags)
-		: FRDGResource(DebugName)
+private:
+	FRDGTexture(
+		const TCHAR* InName,
+		const FPooledRenderTargetDesc& InDesc,
+		ERDGResourceFlags InFlags)
+		: FRDGTrackedResource(InName, InFlags)
 		, Desc(InDesc)
-		, Flags(InFlags)
 	{}
 
+	/** This is not a TRefCountPtr<> because FRDGTexture is allocated on the FMemStack
+	 *  FGraphBuilder::AllocatedTextures is actually keeping the reference.
+	 */
+	IPooledRenderTarget* PooledRenderTarget = nullptr;
+
 	friend class FRDGBuilder;
 
 	template<typename TRHICmdList, typename TShaderClass, typename TShaderRHI>
@@ -189,32 +231,35 @@ private:
 };
 
 /** Descriptor for render graph tracked SRV. */
-class RENDERCORE_API FRDGTextureSRVDesc
+class RENDERCORE_API FRDGTextureSRVDesc final
 {
 public:
-	FRDGTextureRef Texture;
-	uint8 MipLevel = 0;
-
-	FRDGTextureSRVDesc() {}
+	FRDGTextureSRVDesc() = default;
 
-	FRDGTextureSRVDesc(FRDGTextureRef InTexture, uint8 InMipLevel) :
-		Texture(InTexture),
-		MipLevel(InMipLevel)
+	FRDGTextureSRVDesc(FRDGTextureRef InTexture, uint8 InMipLevel)
+		: Texture(InTexture)
+		, MipLevel(InMipLevel)
 	{}
+
+	FRDGTextureRef Texture = nullptr;
+	uint8 MipLevel = 0;
 };
 
 /** Render graph tracked SRV. */
-class RENDERCORE_API FRDGTextureSRV : public FRDGResource
+class RENDERCORE_API FRDGTextureSRV final
+	: public FRDGResource
 {
 public:
 	/** Descriptor of the graph tracked SRV. */
 	const FRDGTextureSRVDesc Desc;
 
 private:
-	FRDGTextureSRV(const TCHAR* DebugName, const FRDGTextureSRVDesc& InDesc)
-		: FRDGResource(DebugName)
+	FRDGTextureSRV(
+		const TCHAR* InName,
+		const FRDGTextureSRVDesc& InDesc)
+		: FRDGResource(InName)
 		, Desc(InDesc)
-	{ }
+	{}
 
 	friend class FRDGBuilder;
 };
@@ -223,34 +268,38 @@ private:
 class RENDERCORE_API FRDGTextureUAVDesc
 {
 public:
-	FRDGTextureRef Texture;
-	uint8 MipLevel = 0;
-
-	FRDGTextureUAVDesc() {}
+	FRDGTextureUAVDesc() = default;
 
-	FRDGTextureUAVDesc(FRDGTextureRef InTexture, uint8 InMipLevel = 0) :
-		Texture(InTexture),
-		MipLevel(InMipLevel)
+	FRDGTextureUAVDesc(
+		FRDGTextureRef InTexture,
+		uint8 InMipLevel = 0)
+		: Texture(InTexture)
+		, MipLevel(InMipLevel)
 	{}
+
+	FRDGTextureRef Texture = nullptr;
+	uint8 MipLevel = 0;
 };
 
 /** Render graph tracked UAV. */
-class RENDERCORE_API FRDGTextureUAV : public FRDGResource
+class RENDERCORE_API FRDGTextureUAV final
+	: public FRDGResource
 {
 public:
 	/** Descriptor of the graph tracked UAV. */
 	const FRDGTextureUAVDesc Desc;
 
 private:
-	FRDGTextureUAV(const TCHAR* DebugName, const FRDGTextureUAVDesc& InDesc)
-		: FRDGResource(DebugName)
+	FRDGTextureUAV(
+		const TCHAR* InName,
+		const FRDGTextureUAVDesc& InDesc)
+		: FRDGResource(InName)
 		, Desc(InDesc)
-	{ }
+	{}
 
 	friend class FRDGBuilder;
 };
 
-
 /** Descriptor for render graph tracked Buffer. */
 struct RENDERCORE_API FRDGBufferDesc
 {
@@ -263,39 +312,6 @@ struct RENDERCORE_API FRDGBufferDesc
 		StructuredBuffer,
 	};
 
-	/** Stride in bytes for index and structured buffers. */
-	uint32 BytesPerElement = 1;
-
-	/** Number of elements. */
-	uint32 NumElements = 1;
-
-	/** Bitfields describing the uses of that buffer. */
-	EBufferUsageFlags Usage = BUF_None;
-
-	/** The underlying RHI type to use. A bit of a work around because RHI still have 3 different objects. */
-	EUnderlyingType UnderlyingType = EUnderlyingType::VertexBuffer;
-
-	/** Returns the total number of bytes allocated for a such buffer. */
-	uint32 GetTotalNumBytes() const
-	{
-		return BytesPerElement * NumElements;
-	}
-
-
-	bool operator == (const FRDGBufferDesc& Other) const
-	{
-		return (
-			BytesPerElement == Other.BytesPerElement &&
-			NumElements == Other.NumElements &&
-			Usage == Other.Usage &&
-			UnderlyingType == Other.UnderlyingType);
-	}
-
-	bool operator != (const FRDGBufferDesc& Other) const
-	{
-		return !(*this == Other);
-	}
-
 	/** Create the descriptor for an indirect RHI call.
 	 *
 	 * Note, IndirectParameterStruct should be one of the:
@@ -304,7 +320,7 @@ struct RENDERCORE_API FRDGBufferDesc
 	 *		struct FRHIDrawIndexedIndirectParameters
 	 */
 	template<typename IndirectParameterStruct>
-	static inline FRDGBufferDesc CreateIndirectDesc(uint32 NumElements = 1)
+	static FRDGBufferDesc CreateIndirectDesc(uint32 NumElements = 1)
 	{
 		FRDGBufferDesc Desc;
 		Desc.UnderlyingType = EUnderlyingType::VertexBuffer;
@@ -314,7 +330,7 @@ struct RENDERCORE_API FRDGBufferDesc
 		return Desc;
 	}
 
-	static inline FRDGBufferDesc CreateIndirectDesc(uint32 NumElements = 1)
+	static FRDGBufferDesc CreateIndirectDesc(uint32 NumElements = 1)
 	{
 		FRDGBufferDesc Desc;
 		Desc.UnderlyingType = EUnderlyingType::VertexBuffer;
@@ -324,7 +340,7 @@ struct RENDERCORE_API FRDGBufferDesc
 		return Desc;
 	}
 
-	static inline FRDGBufferDesc CreateStructuredDesc(uint32 BytesPerElement, uint32 NumElements)
+	static FRDGBufferDesc CreateStructuredDesc(uint32 BytesPerElement, uint32 NumElements)
 	{
 		FRDGBufferDesc Desc;
 		Desc.UnderlyingType = EUnderlyingType::StructuredBuffer;
@@ -334,7 +350,7 @@ struct RENDERCORE_API FRDGBufferDesc
 		return Desc;
 	}
 
-	static inline FRDGBufferDesc CreateBufferDesc(uint32 BytesPerElement, uint32 NumElements)
+	static FRDGBufferDesc CreateBufferDesc(uint32 BytesPerElement, uint32 NumElements)
 	{
 		FRDGBufferDesc Desc;
 		Desc.UnderlyingType = EUnderlyingType::VertexBuffer;
@@ -343,48 +359,82 @@ struct RENDERCORE_API FRDGBufferDesc
 		Desc.NumElements = NumElements;
 		return Desc;
 	}
-};
 
+	/** Returns the total number of bytes allocated for a such buffer. */
+	uint32 GetTotalNumBytes() const
+	{
+		return BytesPerElement * NumElements;
+	}
 
-struct RENDERCORE_API FRDGBufferSRVDesc
-{
-	FRDGBufferRef Buffer = nullptr;
+	bool operator == (const FRDGBufferDesc& Other) const
+	{
+		return (
+			BytesPerElement == Other.BytesPerElement &&
+			NumElements == Other.NumElements &&
+			Usage == Other.Usage &&
+			UnderlyingType == Other.UnderlyingType);
+	}
 
-	/** Number of bytes per element (used for vertex buffer). */
+	bool operator != (const FRDGBufferDesc& Other) const
+	{
+		return !(*this == Other);
+	}
+
+	/** Stride in bytes for index and structured buffers. */
 	uint32 BytesPerElement = 1;
 
-	/** Encoding format for the element (used for vertex buffer). */
-	EPixelFormat Format = PF_Unknown;
+	/** Number of elements. */
+	uint32 NumElements = 1;
+
+	/** Bitfields describing the uses of that buffer. */
+	EBufferUsageFlags Usage = BUF_None;
+
+	/** The underlying RHI type to use. A bit of a work around because RHI still have 3 different objects. */
+	EUnderlyingType UnderlyingType = EUnderlyingType::VertexBuffer;
+};
+
+struct RENDERCORE_API FRDGBufferSRVDesc
+{
+	FRDGBufferSRVDesc() = default;
 
-	FRDGBufferSRVDesc() {}
 	FRDGBufferSRVDesc(FRDGBufferRef InBuffer);
+
 	FRDGBufferSRVDesc(FRDGBufferRef InBuffer, EPixelFormat InFormat)
 		: Buffer(InBuffer)
 		, Format(InFormat)
 	{
-		BytesPerElement = GPixelFormats[ Format ].BlockBytes;
+		BytesPerElement = GPixelFormats[Format].BlockBytes;
 	}
-};
 
-struct RENDERCORE_API FRDGBufferUAVDesc
-{
 	FRDGBufferRef Buffer = nullptr;
 
 	/** Number of bytes per element (used for vertex buffer). */
+	uint32 BytesPerElement = 1;
+
+	/** Encoding format for the element (used for vertex buffer). */
 	EPixelFormat Format = PF_Unknown;
+};
 
-	/** Whether the uav supports atomic counter or append buffer operations (used for structured buffers) */
-	bool bSupportsAtomicCounter = false;
-	bool bSupportsAppendBuffer = false;
+struct RENDERCORE_API FRDGBufferUAVDesc
+{
+	FRDGBufferUAVDesc() = default;
 
-	FRDGBufferUAVDesc() {}
 	FRDGBufferUAVDesc(FRDGBufferRef InBuffer);
+
 	FRDGBufferUAVDesc(FRDGBufferRef InBuffer, EPixelFormat InFormat)
 		: Buffer(InBuffer)
 		, Format(InFormat)
 	{}
-};
 
+	FRDGBufferRef Buffer = nullptr;
+
+	/** Number of bytes per element (used for vertex buffer). */
+	EPixelFormat Format = PF_Unknown;
+
+	/** Whether the uav supports atomic counter or append buffer operations (used for structured buffers) */
+	bool bSupportsAtomicCounter = false;
+	bool bSupportsAppendBuffer = false;
+};
 
 /** Defines how the map's pairs are hashed. */
 template<typename KeyType, typename ValueType>
@@ -473,80 +523,84 @@ private:
 	uint32 RefCount = 0;
 };
 
-
 /** Render graph tracked buffers. */
-class RENDERCORE_API FRDGBuffer : public FRDGResource
+class RENDERCORE_API FRDGBuffer final
+	: public FRDGTrackedResource
 {
 public:
-	/** Descriptor of the graph */
+	/** Descriptor of the graph. */
 	const FRDGBufferDesc Desc;
 
+	//////////////////////////////////////////////////////////////////////////
+	//! The following methods may only be called during pass execution.
+
 	/** Returns the buffer to use for indirect RHI calls. */
 	FVertexBufferRHIParamRef GetIndirectRHICallBuffer() const
 	{
-		VerifyRHIAccess();
+		ValidateRHIAccess();
 		check(PooledBuffer);
 		checkf(Desc.UnderlyingType == FRDGBufferDesc::EUnderlyingType::VertexBuffer, TEXT("Indirect buffers needs to be underlying vertex buffer."));
 		check(PooledBuffer->VertexBuffer.IsValid());
 		return PooledBuffer->VertexBuffer;
 	}
 
+	//////////////////////////////////////////////////////////////////////////
+
 private:
+	FRDGBuffer(
+		const TCHAR* InName,
+		const FRDGBufferDesc& InDesc,
+		ERDGResourceFlags InFlags)
+		: FRDGTrackedResource(InName, InFlags)
+		, Desc(InDesc)
+	{}
+
 	/** This is not a TRefCountPtr<> because FRDGBuffer is allocated on the FMemStack
-	 * FGraphBuilder::AllocatedBuffers is actually keeping the reference.
+	 *  FGraphBuilder::AllocatedBuffers is actually keeping the reference.
 	 */
-	mutable FPooledRDGBuffer* PooledBuffer = nullptr;
+	FPooledRDGBuffer* PooledBuffer = nullptr;
 	
-	// Flags specific to this resource for render graph.
-	const ERDGResourceFlags Flags;
-
-
-	FRDGBuffer(const TCHAR* DebugName, const FRDGBufferDesc& InDesc, ERDGResourceFlags InFlags)
-		: FRDGResource(DebugName)
-		, Desc(InDesc)
-		, Flags(InFlags)
-	{ }
-
 	friend class FRDGBuilder;
 };
 
 /** Render graph tracked SRV. */
-class RENDERCORE_API FRDGBufferSRV : public FRDGResource
+class RENDERCORE_API FRDGBufferSRV final
+	: public FRDGResource
 {
 public:
 	/** Descriptor of the graph tracked SRV. */
 	const FRDGBufferSRVDesc Desc;
 
 private:
-	FRDGBufferSRV(const TCHAR* DebugName, const FRDGBufferSRVDesc& InDesc)
-		: FRDGResource(DebugName)
+	FRDGBufferSRV(const TCHAR* InName, const FRDGBufferSRVDesc& InDesc)
+		: FRDGResource(InName)
 		, Desc(InDesc)
-	{ }
+	{}
 
 	friend class FRDGBuilder;
 };
 
 /** Render graph tracked UAV. */
-class RENDERCORE_API FRDGBufferUAV : public FRDGResource
+class RENDERCORE_API FRDGBufferUAV final
+	: public FRDGResource
 {
 public:
 	/** Descriptor of the graph tracked UAV. */
 	const FRDGBufferUAVDesc Desc;
 
 private:
-	FRDGBufferUAV(const TCHAR* DebugName, const FRDGBufferUAVDesc& InDesc)
-		: FRDGResource(DebugName)
+	FRDGBufferUAV(const TCHAR* InName, const FRDGBufferUAVDesc& InDesc)
+		: FRDGResource(InName)
 		, Desc(InDesc)
-	{ }
+	{}
 
 	friend class FRDGBuilder;
 };
 
-
 inline FRDGBufferSRVDesc::FRDGBufferSRVDesc(FRDGBufferRef InBuffer)
 	: Buffer(InBuffer)
 {
-	if( Buffer->Desc.Usage & BUF_DrawIndirect )
+	if (Buffer->Desc.Usage & BUF_DrawIndirect)
 	{
 		BytesPerElement = 4;
 		Format = PF_R32_UINT;
@@ -560,7 +614,7 @@ inline FRDGBufferSRVDesc::FRDGBufferSRVDesc(FRDGBufferRef InBuffer)
 inline FRDGBufferUAVDesc::FRDGBufferUAVDesc(FRDGBufferRef InBuffer)
 	: Buffer(InBuffer)
 {
-	if( Buffer->Desc.Usage & BUF_DrawIndirect )
+	if (Buffer->Desc.Usage & BUF_DrawIndirect)
 	{
 		Format = PF_R32_UINT;
 	}
@@ -568,4 +622,4 @@ inline FRDGBufferUAVDesc::FRDGBufferUAVDesc(FRDGBufferRef InBuffer)
 	{
 		checkf(Buffer->Desc.UnderlyingType != FRDGBufferDesc::EUnderlyingType::VertexBuffer, TEXT("VertexBuffer %s requires a type when creating a UAV."), Buffer->Name);
 	}
-}
+}
\ No newline at end of file
