Fixes RDG no longer listing all resources available to debug.

diff --git a/Engine/Source/Runtime/RenderCore/Private/RenderGraphBuilder.cpp b/Engine/Source/Runtime/RenderCore/Private/RenderGraphBuilder.cpp
--- a/Engine/Source/Runtime/RenderCore/Private/RenderGraphBuilder.cpp
+++ b/Engine/Source/Runtime/RenderCore/Private/RenderGraphBuilder.cpp
@@ -406,11 +406,6 @@ void FRDGBuilder::ValidatePass(const FRDGPass* Pass) const
 void FRDGBuilder::VisualizePassOutputs(const FRDGPass* Pass)
 {
 #if SUPPORTS_VISUALIZE_TEXTURE
-	if (!!GVisualizeTexture.bEnabled)
-	{
-		return;
-	}
-
 	FRDGPassParameterStruct ParameterStruct = Pass->GetParameters();
 
 	const uint32 ParameterCount = ParameterStruct.GetParameterCount();
