Fixed build break from uber builds.

diff --git a/Engine/Source/Runtime/RenderCore/Private/ProfilingDebugging/RealtimeGPUProfiler.cpp b/Engine/Source/Runtime/RenderCore/Private/ProfilingDebugging/RealtimeGPUProfiler.cpp
--- a/Engine/Source/Runtime/RenderCore/Private/ProfilingDebugging/RealtimeGPUProfiler.cpp
+++ b/Engine/Source/Runtime/RenderCore/Private/ProfilingDebugging/RealtimeGPUProfiler.cpp
@@ -3,8 +3,7 @@
 #include "ProfilingDebugging/RealtimeGPUProfiler.h"
 #include "ProfilingDebugging/CsvProfiler.h"
 #include "ProfilingDebugging/TracingProfiler.h"
-
-DEFINE_LOG_CATEGORY_STATIC(LogRendererCore,All,All);
+#include "RenderCore.h"
 
 // Only exposed for debugging. Disabling this carries a severe performance penalty
 #define RENDER_QUERY_POOLING_ENABLED 1
