Oups, didn't notice the !! that was actually the cause of the issue that motivated 6332286

diff --git a/Engine/Source/Runtime/RenderCore/Private/RenderGraphBuilder.cpp b/Engine/Source/Runtime/RenderCore/Private/RenderGraphBuilder.cpp
--- a/Engine/Source/Runtime/RenderCore/Private/RenderGraphBuilder.cpp
+++ b/Engine/Source/Runtime/RenderCore/Private/RenderGraphBuilder.cpp
@@ -406,6 +406,11 @@ void FRDGBuilder::ValidatePass(const FRDGPass* Pass) const
 void FRDGBuilder::VisualizePassOutputs(const FRDGPass* Pass)
 {
 #if SUPPORTS_VISUALIZE_TEXTURE
+	if (!GVisualizeTexture.bEnabled)
+	{
+		return;
+	}
+
 	FRDGPassParameterStruct ParameterStruct = Pass->GetParameters();
 
 	const uint32 ParameterCount = ParameterStruct.GetParameterCount();
