Added support for RDG_GPU_STAT_SCOPE.

diff --git a/Engine/Source/Runtime/RenderCore/Private/RenderGraphBuilder.cpp b/Engine/Source/Runtime/RenderCore/Private/RenderGraphBuilder.cpp
--- a/Engine/Source/Runtime/RenderCore/Private/RenderGraphBuilder.cpp
+++ b/Engine/Source/Runtime/RenderCore/Private/RenderGraphBuilder.cpp
@@ -92,6 +92,7 @@ FRDGBuilder::FRDGBuilder(FRHICommandListImmediate& InRHICmdList)
 	: RHICmdList(InRHICmdList)
 	, MemStack(FMemStack::Get())
 	, EventScopeStack(RHICmdList)
+	, StatScopeStack(RHICmdList)
 {}
 
 FRDGBuilder::~FRDGBuilder()
@@ -124,6 +125,7 @@ void FRDGBuilder::Execute()
 	#endif
 
 	EventScopeStack.BeginExecute();
+	StatScopeStack.BeginExecute();
 
 	if (!GRDGImmediateMode)
 	{
@@ -137,6 +139,7 @@ void FRDGBuilder::Execute()
 	}
 
 	EventScopeStack.EndExecute();
+	StatScopeStack.EndExecute();
 
 	ProcessDeferredInternalResourceQueries();
 
@@ -171,7 +174,8 @@ void FRDGBuilder::AddPassInternal(FRDGPass* Pass)
 	}
 	#endif
 
-	Pass->ParentScope = EventScopeStack.GetCurrentScope();
+	Pass->EventScope = EventScopeStack.GetCurrentScope();
+	Pass->StatScope = StatScopeStack.GetCurrentScope();
 	Passes.Emplace(Pass);
 
 	ValidatePass(Pass);
@@ -820,6 +824,7 @@ void FRDGBuilder::ExecutePass(const FRDGPass* Pass)
 	}
 
 	EventScopeStack.BeginExecutePass(Pass);
+	StatScopeStack.BeginExecutePass(Pass);
 
 	if (!Pass->IsCompute())
 	{
@@ -1416,10 +1421,20 @@ void FRDGBuilder::DestructPasses()
 
 void FRDGBuilder::BeginEventScope(FRDGEventName&& ScopeName)
 {
-	EventScopeStack.BeginEventScope(MoveTemp(ScopeName));
+	EventScopeStack.BeginScope(Forward<FRDGEventName&&>(ScopeName));
 }
 
 void FRDGBuilder::EndEventScope()
 {
-	EventScopeStack.EndEventScope();
+	EventScopeStack.EndScope();
+}
+
+void FRDGBuilder::BeginStatScope(const FName& Name, const FName& StatName)
+{
+	StatScopeStack.BeginScope(Name, StatName);
+}
+
+void FRDGBuilder::EndStatScope()
+{
+	StatScopeStack.EndScope();
 }
\ No newline at end of file
diff --git a/Engine/Source/Runtime/RenderCore/Private/RenderGraphEvent.cpp b/Engine/Source/Runtime/RenderCore/Private/RenderGraphEvent.cpp
--- a/Engine/Source/Runtime/RenderCore/Private/RenderGraphEvent.cpp
+++ b/Engine/Source/Runtime/RenderCore/Private/RenderGraphEvent.cpp
@@ -5,21 +5,6 @@
 
 extern bool GetEmitRDGEvents();
 
-class FRDGEventScope
-{
-public:
-	FRDGEventScope(const FRDGEventScope* InParentScope, FRDGEventName&& InName)
-		: ParentScope(InParentScope)
-		, Name(InName)
-	{}
-
-	// Pointer towards this one is contained in.
-	const FRDGEventScope* const ParentScope;
-
-	// Name of the event.
-	const FRDGEventName Name;
-};
-
 #if RDG_EVENTS == RDG_EVENTS_STRING_COPY
 
 FRDGEventName::FRDGEventName(const TCHAR* EventFormat, ...)
@@ -55,107 +40,56 @@ FRDGEventScopeGuard::~FRDGEventScopeGuard()
 	GraphBuilder.EndEventScope();
 }
 
+static void OnPushEvent(FRHICommandListImmediate& RHICmdList, const FRDGEventScope* Scope)
+{
+	RHICmdList.PushEvent(Scope->Name.GetTCHAR(), FColor(0, 0, 0));
+}
+
+static void OnPopEvent(FRHICommandListImmediate& RHICmdList)
+{
+	RHICmdList.PopEvent();
+}
+
+bool FRDGEventScopeStack::IsEnabled()
+{
+	return RDG_EVENTS && GetEmitRDGEvents();
+}
+
 FRDGEventScopeStack::FRDGEventScopeStack(FRHICommandListImmediate& InRHICmdList)
-	: RHICmdList(InRHICmdList)
-	, MemStack(FMemStack::Get())
-	, ScopeStack(MakeUniformStaticArray<const FRDGEventScope*, kScopeStackDepthMax>(nullptr))
+	: ScopeStack(InRHICmdList, &OnPushEvent, &OnPopEvent)
 {}
 
-void FRDGEventScopeStack::BeginEventScope(FRDGEventName&& EventName)
+void FRDGEventScopeStack::BeginScope(FRDGEventName&& EventName)
 {
-	if (RDG_EVENTS)
+	if (IsEnabled())
 	{
-		auto Scope = new(MemStack) FRDGEventScope(CurrentScope, Forward<FRDGEventName>(EventName));
-		EventScopes.Add(Scope);
-		CurrentScope = Scope;
+		ScopeStack.BeginScope(Forward<FRDGEventName&&>(EventName));
 	}
 }
 
-void FRDGEventScopeStack::EndEventScope()
+void FRDGEventScopeStack::EndScope()
 {
-	if (RDG_EVENTS)
+	if (IsEnabled())
 	{
-		checkf(CurrentScope != nullptr, TEXT("Current scope is null."));
-		CurrentScope = CurrentScope->ParentScope;
+		ScopeStack.EndScope();
 	}
 }
 
 void FRDGEventScopeStack::BeginExecute()
 {
-	/** The usage need RDG_EVENT_SCOPE() needs to happen in inner scope of the one containing FRDGBuilder because of
-	 *  FStackRDGEventScopeRef's destructor modifying this FRDGBuilder instance.
-	 *
-	 *
-	 *  FRDGBuilder GraphBuilder(RHICmdList);
-	 *  {
-	 *  	RDG_EVENT_SCOPE(GraphBuilder, "MyEventScope");
-	 *  	// ...
-	 *  }
-	 *  GraphBuilder.Execute();
-	 */
-	checkf(CurrentScope == nullptr, TEXT("Render graph needs to have all scopes ended to execute."));
+	if (IsEnabled())
+	{
+		ScopeStack.BeginExecute();
+	}
 }
 
 void FRDGEventScopeStack::BeginExecutePass(const FRDGPass* Pass)
 {
-	if (!GetEmitRDGEvents())
+	if (IsEnabled())
 	{
-		return;
-	}
-
-	// Push the scope event.
-	{
-		// Find out how many scope events needs to be popped.
-		TStaticArray<const FRDGEventScope*, kScopeStackDepthMax> TraversedScopes;
-		int32 CommonScopeId = -1;
-		int32 TraversedScopeCount = 0;
+		ScopeStack.BeginExecutePass(Pass->GetEventScope());
 
-		const FRDGEventScope* PassParentScope = Pass->GetParentScope();
-
-		while (PassParentScope)
-		{
-			TraversedScopes[TraversedScopeCount] = PassParentScope;
-
-			for (int32 i = 0; i < ScopeStack.Num(); i++)
-			{
-				if (ScopeStack[i] == PassParentScope)
-				{
-					CommonScopeId = i;
-					break;
-				}
-			}
-
-			if (CommonScopeId != -1)
-			{
-				break;
-			}
-
-			TraversedScopeCount++;
-			PassParentScope = PassParentScope->ParentScope;
-		}
-
-		// Pop no longer used scopes
-		for (int32 i = CommonScopeId + 1; i < kScopeStackDepthMax; i++)
-		{
-			if (!ScopeStack[i])
-				break;
-
-			RHICmdList.PopEvent();
-			ScopeStack[i] = nullptr;
-		}
-
-		// Push new scopes
-		const FColor ScopeColor(0);
-		for (int32 i = TraversedScopeCount - 1; i >= 0; i--)
-		{
-			RHICmdList.PushEvent(TraversedScopes[i]->Name.GetTCHAR(), ScopeColor);
-			CommonScopeId++;
-			ScopeStack[CommonScopeId] = TraversedScopes[i];
-		}
-	}
-
-	// Push the pass's event with some color.
-	{
+		// Push the pass event with some color.
 		FColor Color(0, 0, 0);
 
 		if (Pass->IsCompute())
@@ -165,44 +99,102 @@ void FRDGEventScopeStack::BeginExecutePass(const FRDGPass* Pass)
 		}
 		else
 		{
-			// Ref for rasterizer.
+			// Red for rasterizer.
 			Color = FColor(255, 128, 128);
 		}
 
-		RHICmdList.PushEvent(Pass->GetName(), Color);
+		ScopeStack.RHICmdList.PushEvent(Pass->GetName(), Color);
 	}
 }
 
 void FRDGEventScopeStack::EndExecutePass()
 {
-	if (GetEmitRDGEvents())
+	if (IsEnabled())
 	{
-		RHICmdList.PopEvent();
+		ScopeStack.RHICmdList.PopEvent();
 	}
 }
 
 void FRDGEventScopeStack::EndExecute()
 {
-	// Pops remaining scopes
-	if (GetEmitRDGEvents())
+	if (IsEnabled())
 	{
-		for (uint32 ScopeIndex = 0; ScopeIndex < kScopeStackDepthMax; ++ScopeIndex)
-		{
-			if (!ScopeStack[ScopeIndex])
-			{
-				break;
-			}
+		ScopeStack.EndExecute();
+	}
+}
 
-			RHICmdList.PopEvent();
-		}
+FRDGStatScopeGuard::FRDGStatScopeGuard(FRDGBuilder& InGraphBuilder, const FName& Name, const FName& StatName)
+	: GraphBuilder(InGraphBuilder)
+{
+	GraphBuilder.BeginStatScope(Name, StatName);
+}
+
+FRDGStatScopeGuard::~FRDGStatScopeGuard()
+{
+	GraphBuilder.EndStatScope();
+}
+
+extern bool AreGPUStatsEnabled();
+
+static void OnPushStat(FRHICommandListImmediate& RHICmdList, const FRDGStatScope* Scope)
+{
+#if HAS_GPU_STATS
+	FRealtimeGPUProfiler::Get()->PushEvent(RHICmdList, Scope->Name, Scope->StatName);
+#endif
+}
+
+static void OnPopStat(FRHICommandListImmediate& RHICmdList)
+{
+#if HAS_GPU_STATS
+	FRealtimeGPUProfiler::Get()->PopEvent(RHICmdList);
+#endif
+}
+
+bool FRDGStatScopeStack::IsEnabled()
+{
+	return HAS_GPU_STATS && AreGPUStatsEnabled();
+}
+
+FRDGStatScopeStack::FRDGStatScopeStack(FRHICommandListImmediate& InRHICmdList)
+	: ScopeStack(InRHICmdList, &OnPushStat, &OnPopStat)
+{}
+
+void FRDGStatScopeStack::BeginScope(const FName& Name, const FName& StatName)
+{
+	if (IsEnabled())
+	{
+		ScopeStack.BeginScope(Name, StatName);
 	}
+}
 
-#if RDG_EVENTS
-	// Event scopes are allocated on FMemStack, so need to call their destructor because have a FString within them.
-	for (int32 Index = EventScopes.Num() - 1; Index >= 0; --Index)
+void FRDGStatScopeStack::EndScope()
+{
+	if (IsEnabled())
 	{
-		EventScopes[Index]->~FRDGEventScope();
+		ScopeStack.EndScope();
+	}
+}
+
+void FRDGStatScopeStack::BeginExecute()
+{
+	if (IsEnabled())
+	{
+		ScopeStack.BeginExecute();
+	}
+}
+
+void FRDGStatScopeStack::BeginExecutePass(const FRDGPass* Pass)
+{
+	if (IsEnabled())
+	{
+		ScopeStack.BeginExecutePass(Pass->GetStatScope());
+	}
+}
+
+void FRDGStatScopeStack::EndExecute()
+{
+	if (IsEnabled())
+	{
+		ScopeStack.EndExecute();
 	}
-	EventScopes.Empty();
-#endif
 }
\ No newline at end of file
diff --git a/Engine/Source/Runtime/RenderCore/Public/RenderGraphBuilder.h b/Engine/Source/Runtime/RenderCore/Public/RenderGraphBuilder.h
--- a/Engine/Source/Runtime/RenderCore/Public/RenderGraphBuilder.h
+++ b/Engine/Source/Runtime/RenderCore/Public/RenderGraphBuilder.h
@@ -335,9 +335,17 @@ public:
 		#endif
 	}
 
-	/** 
-	 * Executes the queued passes, managing setting of render targets (RHI RenderPasses), resource transitions and queued texture extraction.
-	 */
+	/** Begins / ends a named event scope. These scopes are visible in most external GPU profilers.
+	 *  Prefer to use RDG_EVENT_SCOPE instead of calling this manually. */
+	void BeginEventScope(FRDGEventName&& ScopeName);
+	void EndEventScope();
+
+	/** Begins / ends a stat scope. Stat scopes are visible via 'stat GPU'. Must be accompanied by a
+	 *  respective EndStatScope call. Prefer to use RDG_GPU_STAT_SCOPE instead of calling this manually. */
+	void BeginStatScope(const FName& Name, const FName& StatName);
+	void EndStatScope();
+
+	/** Executes the queued passes, managing setting of render targets (RHI RenderPasses), resource transitions and queued texture extraction. */
 	void Execute();
 
 	/** The RHI command list used for the render graph. */
@@ -374,6 +382,7 @@ private:
 	TArray<FDeferredInternalBufferQuery, SceneRenderingAllocator> DeferredInternalBufferQueries;
 
 	FRDGEventScopeStack EventScopeStack;
+	FRDGStatScopeStack StatScopeStack;
 
 #if RDG_ENABLE_DEBUG
 	/** Whether the Execute() has already been called. */
@@ -415,9 +424,6 @@ private:
 	void AllocateRHIBufferSRVIfNeeded(FRDGBufferSRV* SRV, bool bComputePass);
 	void AllocateRHIBufferUAVIfNeeded(FRDGBufferUAV* UAV, bool bComputePass);
 
-	void BeginEventScope(FRDGEventName&& ScopeName);
-	void EndEventScope();
-
 	void TransitionTexture(FRDGTexture* Texture, EResourceTransitionAccess TransitionAccess, bool bRequiredCompute) const;
 	void TransitionUAV(FUnorderedAccessViewRHIParamRef UAV, FRDGTrackedResource* UnderlyingResource, ERDGResourceFlags ResourceFlags, EResourceTransitionAccess TransitionAccess, bool bRequiredCompute ) const;
 
@@ -434,8 +440,6 @@ private:
 	void ProcessDeferredInternalResourceQueries();
 	void DestructPasses();
 
-	friend class FRDGEventScopeGuard;
-
 	/** To allow greater flexibility in the user code, the RHI can dereferenced RDG resource when creating uniform buffer. */
 	// TODO(RDG): Make this a little more explicit in RHI code.
 	static_assert(STRUCT_OFFSET(FRDGResource, CachedRHI) == 0, "FRDGResource::CachedRHI requires to be at offset 0 so the RHI can dereferenced them.");
diff --git a/Engine/Source/Runtime/RenderCore/Public/RenderGraphEvent.h b/Engine/Source/Runtime/RenderCore/Public/RenderGraphEvent.h
--- a/Engine/Source/Runtime/RenderCore/Public/RenderGraphEvent.h
+++ b/Engine/Source/Runtime/RenderCore/Public/RenderGraphEvent.h
@@ -3,10 +3,76 @@
 #pragma once
 
 #include "RenderGraphDefinitions.h"
+#include "ProfilingDebugging/RealtimeGPUProfiler.h"
+
+/** A helper profiler class for tracking and evaluating hierarchical scopes in the context of render graph. */
+template <typename TScopeType>
+class FRDGScopeStack final
+{
+	static constexpr uint32 kScopeStackDepthMax = 8;
+public:
+	using FPushFunction = void(*)(FRHICommandListImmediate&, const TScopeType*);
+	using FPopFunction = void(*)(FRHICommandListImmediate&);
+
+	FRDGScopeStack(FRHICommandListImmediate& InRHICmdList, FPushFunction InPushFunction, FPopFunction InPopFunction);
+	~FRDGScopeStack();
+
+	//////////////////////////////////////////////////////////////////////////
+	//! Called during graph setup phase.
+
+	/** Call to begin recording a scope. */
+	template <typename... FScopeConstructArgs>
+	void BeginScope(FScopeConstructArgs... ScopeConstructArgs);
+
+	/** Call to end recording a scope. */
+	void EndScope();
+
+	//////////////////////////////////////////////////////////////////////////
+	//! Called during graph execute phase.
+
+	/** Call prior to executing the graph. */
+	void BeginExecute();
+
+	/** Call prior to executing a pass in the graph. */
+	void BeginExecutePass(const TScopeType* ParentScope);
+
+	/** Call after executing the graph. */
+	void EndExecute();
+
+	//////////////////////////////////////////////////////////////////////////
+
+	const TScopeType* GetCurrentScope() const
+	{
+		return CurrentScope;
+	}
+
+	FRHICommandListImmediate& RHICmdList;
+	FMemStackBase& MemStack;
+
+private:
+	void ClearScopes();
+
+	FPushFunction PushFunction;
+	FPopFunction PopFunction;
+
+	/** The top of the scope stack during setup. */
+	const TScopeType* CurrentScope = nullptr;
+
+	/** Tracks scopes allocated through MemStack for destruction. */
+	TArray<TScopeType*, SceneRenderingAllocator> Scopes;
+
+	/** Stacks of scopes pushed to the RHI command list during execution. */
+	TStaticArray<const TScopeType*, kScopeStackDepthMax> ScopeStack;
+};
 
 class FRDGPass;
 class FRDGBuilder;
-class FRDGEventScope;
+
+//////////////////////////////////////////////////////////////////////////
+//
+// GPU Events - Named hierarchical events emitted to external profiling tools.
+//
+//////////////////////////////////////////////////////////////////////////
 
 /** Stores a GPU event name for the render graph. Draw events can be compiled out entirely from
  *  a release build for performance.
@@ -16,66 +82,16 @@ class RENDERCORE_API FRDGEventName final
 public:
 	FRDGEventName() = default;
 
-	~FRDGEventName()
-	{
-#if RDG_EVENTS != RDG_EVENTS_NONE
-		EventName = nullptr;
-#endif
-	}
-
-#if RDG_EVENTS == RDG_EVENTS_STRING_COPY
 	explicit FRDGEventName(const TCHAR* EventFormat, ...);
-#else
-	explicit FRDGEventName(const TCHAR* EventFormat, ...)
-#if RDG_EVENTS == RDG_EVENTS_STRING_REF
-		: EventName(EventFormat)
-#endif
-	{}
-#endif
 
-	FRDGEventName(const FRDGEventName& Other)
-	{
-		*this = Other;
-	}
+	~FRDGEventName();
 
-	FRDGEventName(FRDGEventName&& Other)
-	{
-		*this = Forward<FRDGEventName>(Other);
-	}
+	FRDGEventName(const FRDGEventName& Other);
+	FRDGEventName(FRDGEventName&& Other);
+	FRDGEventName& operator=(const FRDGEventName& Other);
+	FRDGEventName& operator=(FRDGEventName&& Other);
 
-	FRDGEventName& operator=(const FRDGEventName& Other)
-	{
-#if RDG_EVENTS == RDG_EVENTS_STRING_REF
-		EventName = Other.EventName;
-#elif RDG_EVENTS == RDG_EVENTS_STRING_COPY
-		EventNameStorage = Other.EventNameStorage;
-		EventName = *EventNameStorage;
-#endif
-		return *this;
-	}
-
-	FRDGEventName& operator=(FRDGEventName&& Other)
-	{
-#if RDG_EVENTS == RDG_EVENTS_STRING_REF
-		EventName = Other.EventName;
-		Other.EventName = nullptr;
-#elif RDG_EVENTS == RDG_EVENTS_STRING_COPY
-		EventNameStorage = MoveTemp(Other.EventNameStorage);
-		EventName = *EventNameStorage;
-		Other.EventName = nullptr;
-#endif
-		return *this;
-	}
-
-	const TCHAR* GetTCHAR() const
-	{
-#if RDG_EVENTS == RDG_EVENTS_STRING_REF || RDG_EVENTS == RDG_EVENTS_STRING_COPY
-		return EventName;
-#else
-		// Render graph draw events have been completely compiled for CPU performance reasons.
-		return TEXT("!!!Unavailable RDG event name: need RDG_EVENTS>=0 and r.RDG.EmitWarnings=1 or -rdgdebug!!!");
-#endif
-	}
+	const TCHAR* GetTCHAR() const;
 
 private:
 #if RDG_EVENTS == RDG_EVENTS_STRING_REF || RDG_EVENTS == RDG_EVENTS_STRING_COPY
@@ -86,6 +102,18 @@ private:
 #endif
 };
 
+class FRDGEventScope final
+{
+public:
+	FRDGEventScope(const FRDGEventScope* InParentScope, FRDGEventName&& InName)
+		: ParentScope(InParentScope)
+		, Name(Forward<FRDGEventName&&>(InName))
+	{}
+
+	const FRDGEventScope* const ParentScope;
+	const FRDGEventName Name;
+};
+
 /** Manages a stack of event scopes. Scopes are recorded ahead of time in a hierarchical fashion
  *  and later executed topologically during pass execution.
  */
@@ -97,42 +125,37 @@ public:
 	//////////////////////////////////////////////////////////////////////////
 	//! Called during graph setup phase.
 
-	void BeginEventScope(FRDGEventName&& EventName);
+	/** Call to begin recording an event scope. */
+	void BeginScope(FRDGEventName&& EventName);
 
-	void EndEventScope();
+	/** Call to end recording an event scope. */
+	void EndScope();
 
 	//////////////////////////////////////////////////////////////////////////
 	//! Called during graph execute phase.
 
+	/** Call prior to executing the graph. */
 	void BeginExecute();
 
+	/** Call prior to executing a pass in the graph. */
 	void BeginExecutePass(const FRDGPass* Pass);
 
+	/** Call after executing a pass in the graph. */
 	void EndExecutePass();
 
+	/** Call after executing the graph. */
 	void EndExecute();
 
 	//////////////////////////////////////////////////////////////////////////
 
 	const FRDGEventScope* GetCurrentScope() const
 	{
-		return CurrentScope;
+		return ScopeStack.GetCurrentScope();
 	}
 
 private:
-	static constexpr uint32 kScopeStackDepthMax = 8;
-
-	FRHICommandListImmediate& RHICmdList;
-	FMemStackBase& MemStack;
-
-	/** The top of the scope stack during setup. */
-	const FRDGEventScope* CurrentScope = nullptr;
-
-	/** Tracks scopes allocated through MemStack for destruction. */
-	TArray<FRDGEventScope*, SceneRenderingAllocator> EventScopes;
-
-	/** Stacks of scopes pushed to the RHI command list during execution. */
-	TStaticArray<const FRDGEventScope*, kScopeStackDepthMax> ScopeStack;
+	static bool IsEnabled();
+	FRDGScopeStack<FRDGEventScope> ScopeStack;
 };
 
 /** RAII class for begin / end of an event scope through the graph builder. */
@@ -161,4 +184,86 @@ private:
 	#define RDG_EVENT_SCOPE(GraphBuilder, Format, ...) 
 #else
 	#error "RDG_EVENTS is not a valid value."
-#endif
\ No newline at end of file
+#endif
+
+//////////////////////////////////////////////////////////////////////////
+//
+// GPU Stats - Aggregated counters emitted to the runtime 'stat GPU' profiler.
+//
+//////////////////////////////////////////////////////////////////////////
+
+class FRDGStatScope final
+{
+public:
+	FRDGStatScope(const FRDGStatScope* InParentScope, const FName& InName, const FName& InStatName)
+		: ParentScope(InParentScope)
+		, Name(InName)
+		, StatName(InStatName)
+	{}
+
+	const FRDGStatScope* const ParentScope;
+	const FName Name;
+	const FName StatName;
+};
+
+class RENDERCORE_API FRDGStatScopeStack final
+{
+public:
+	FRDGStatScopeStack(FRHICommandListImmediate& RHICmdList);
+
+	//////////////////////////////////////////////////////////////////////////
+	//! Called during graph setup phase.
+
+	/** Call to begin recording a stat scope. */
+	void BeginScope(const FName& Name, const FName& StatName);
+
+	/** Call to end recording a stat scope. */
+	void EndScope();
+
+	//////////////////////////////////////////////////////////////////////////
+	//! Called during graph execute phase.
+
+	/** Call prior to executing the graph. */
+	void BeginExecute();
+
+	/** Call prior to executing a pass in the graph. */
+	void BeginExecutePass(const FRDGPass* Pass);
+
+	/** Call after executing the graph. */
+	void EndExecute();
+
+	//////////////////////////////////////////////////////////////////////////
+
+	const FRDGStatScope* GetCurrentScope() const
+	{
+		return ScopeStack.GetCurrentScope();
+	}
+
+private:
+	static bool IsEnabled();
+	FRDGScopeStack<FRDGStatScope> ScopeStack;
+};
+
+/** RAII class for begin / end of a stat scope through the graph builder. */
+class RENDERCORE_API FRDGStatScopeGuard final
+{
+public:
+	FRDGStatScopeGuard(FRDGBuilder& InGraphBuilder, const FName& Name, const FName& StatName);
+	FRDGStatScopeGuard(const FRDGStatScopeGuard&) = delete;
+	~FRDGStatScopeGuard();
+
+private:
+	FRDGBuilder& GraphBuilder;
+};
+
+#if HAS_GPU_STATS
+	#if STATS
+		#define RDG_GPU_STAT_SCOPE(GraphBuilder, StatName) FRDGStatScopeGuard PREPROCESSOR_JOIN(GPUStatEvent_##StatName,__LINE__) ((GraphBuilder), CSV_STAT_FNAME(StatName), GET_STATID(Stat_GPU_##StatName).GetName());
+	#else
+		#define RDG_GPU_STAT_SCOPE(GraphBuilder, StatName) FRDGStatScopeGuard PREPROCESSOR_JOIN(GPUStatEvent_##StatName,__LINE__) ((GraphBuilder), CSV_STAT_FNAME(StatName), FName());
+	#endif
+#else
+	#define RDG_GPU_STAT_SCOPE(GraphBuilder, StatName)
+#endif
+
+#include "RenderGraphEvent.inl"
\ No newline at end of file
diff --git a/Engine/Source/Runtime/RenderCore/Public/RenderGraphEvent.inl b/Engine/Source/Runtime/RenderCore/Public/RenderGraphEvent.inl
new file mode 100644
--- /dev/null
+++ b/Engine/Source/Runtime/RenderCore/Public/RenderGraphEvent.inl
@@ -0,0 +1,179 @@
+// Copyright 1998-2019 Epic Games, Inc. All Rights Reserved.
+
+#pragma once
+
+template <typename TScopeType>
+FRDGScopeStack<TScopeType>::FRDGScopeStack(
+	FRHICommandListImmediate& InRHICmdList,
+	FPushFunction InPushFunction,
+	FPopFunction InPopFunction)
+	: RHICmdList(InRHICmdList)
+	, MemStack(FMemStack::Get())
+	, PushFunction(InPushFunction)
+	, PopFunction(InPopFunction)
+	, ScopeStack(MakeUniformStaticArray<const TScopeType*, kScopeStackDepthMax>(nullptr))
+{}
+
+template <typename TScopeType>
+FRDGScopeStack<TScopeType>::~FRDGScopeStack()
+{
+	ClearScopes();
+}
+
+template <typename TScopeType>
+template <typename... FScopeConstructArgs>
+void FRDGScopeStack<TScopeType>::BeginScope(FScopeConstructArgs... ScopeConstructArgs)
+{
+	auto Scope = new(MemStack) TScopeType(CurrentScope, Forward<FScopeConstructArgs>(ScopeConstructArgs)...);
+	Scopes.Add(Scope);
+	CurrentScope = Scope;
+}
+
+template <typename TScopeType>
+void FRDGScopeStack<TScopeType>::EndScope()
+{
+	checkf(CurrentScope != nullptr, TEXT("Current scope is null."));
+	CurrentScope = CurrentScope->ParentScope;
+}
+
+template <typename TScopeType>
+void FRDGScopeStack<TScopeType>::BeginExecute()
+{
+	checkf(CurrentScope == nullptr, TEXT("Render graph needs to have all scopes ended to execute."));
+}
+
+template <typename TScopeType>
+void FRDGScopeStack<TScopeType>::BeginExecutePass(const TScopeType* ParentScope)
+{
+	// Find out how many scopes needs to be popped.
+	TStaticArray<const TScopeType*, kScopeStackDepthMax> TraversedScopes;
+	int32 CommonScopeId = -1;
+	int32 TraversedScopeCount = 0;
+
+	// Find common ancestor between current stack and requested scope.
+	while (ParentScope)
+	{
+		TraversedScopes[TraversedScopeCount] = ParentScope;
+
+		for (int32 i = 0; i < ScopeStack.Num(); i++)
+		{
+			if (ScopeStack[i] == ParentScope)
+			{
+				CommonScopeId = i;
+				break;
+			}
+		}
+
+		if (CommonScopeId != -1)
+		{
+			break;
+		}
+
+		TraversedScopeCount++;
+		ParentScope = ParentScope->ParentScope;
+	}
+
+	// Pop no longer used scopes.
+	for (int32 i = CommonScopeId + 1; i < kScopeStackDepthMax; i++)
+	{
+		if (!ScopeStack[i])
+		{
+			break;
+		}
+
+		PopFunction(RHICmdList);
+		ScopeStack[i] = nullptr;
+	}
+
+	// Push new scopes.
+	for (int32 i = TraversedScopeCount - 1; i >= 0; i--)
+	{
+		PushFunction(RHICmdList, TraversedScopes[i]);
+		CommonScopeId++;
+		ScopeStack[CommonScopeId] = TraversedScopes[i];
+	}
+}
+
+template <typename TScopeType>
+void FRDGScopeStack<TScopeType>::EndExecute()
+{
+	for (uint32 ScopeIndex = 0; ScopeIndex < kScopeStackDepthMax; ++ScopeIndex)
+	{
+		if (!ScopeStack[ScopeIndex])
+		{
+			break;
+		}
+
+		PopFunction(RHICmdList);
+	}
+	ClearScopes();
+}
+
+template <typename TScopeType>
+void FRDGScopeStack<TScopeType>::ClearScopes()
+{
+	for (int32 Index = Scopes.Num() - 1; Index >= 0; --Index)
+	{
+		Scopes[Index]->~TScopeType();
+	}
+	Scopes.Empty();
+}
+
+inline FRDGEventName::~FRDGEventName()
+{
+#if RDG_EVENTS != RDG_EVENTS_NONE
+	EventName = nullptr;
+#endif
+}
+
+#if RDG_EVENTS != RDG_EVENTS_STRING_COPY
+inline FRDGEventName::FRDGEventName(const TCHAR* EventFormat, ...)
+#if RDG_EVENTS == RDG_EVENTS_STRING_REF
+	: EventName(EventFormat)
+#endif
+{}
+#endif
+
+inline FRDGEventName::FRDGEventName(const FRDGEventName& Other)
+{
+	*this = Other;
+}
+
+inline FRDGEventName::FRDGEventName(FRDGEventName&& Other)
+{
+	*this = Forward<FRDGEventName>(Other);
+}
+
+inline FRDGEventName& FRDGEventName::operator=(const FRDGEventName& Other)
+{
+#if RDG_EVENTS == RDG_EVENTS_STRING_REF
+	EventName = Other.EventName;
+#elif RDG_EVENTS == RDG_EVENTS_STRING_COPY
+	EventNameStorage = Other.EventNameStorage;
+	EventName = *EventNameStorage;
+#endif
+	return *this;
+}
+
+inline FRDGEventName& FRDGEventName::operator=(FRDGEventName&& Other)
+{
+#if RDG_EVENTS == RDG_EVENTS_STRING_REF
+	EventName = Other.EventName;
+	Other.EventName = nullptr;
+#elif RDG_EVENTS == RDG_EVENTS_STRING_COPY
+	EventNameStorage = MoveTemp(Other.EventNameStorage);
+	EventName = *EventNameStorage;
+	Other.EventName = nullptr;
+#endif
+	return *this;
+}
+
+inline const TCHAR* FRDGEventName::GetTCHAR() const
+{
+#if RDG_EVENTS == RDG_EVENTS_STRING_REF || RDG_EVENTS == RDG_EVENTS_STRING_COPY
+	return EventName;
+#else
+	// Render graph draw events have been completely compiled for CPU performance reasons.
+	return TEXT("!!!Unavailable RDG event name: need RDG_EVENTS>=0 and r.RDG.EmitWarnings=1 or -rdgdebug!!!");
+#endif
+}
\ No newline at end of file
diff --git a/Engine/Source/Runtime/RenderCore/Public/RenderGraphPass.h b/Engine/Source/Runtime/RenderCore/Public/RenderGraphPass.h
--- a/Engine/Source/Runtime/RenderCore/Public/RenderGraphPass.h
+++ b/Engine/Source/Runtime/RenderCore/Public/RenderGraphPass.h
@@ -208,9 +208,14 @@ public:
 		return ParameterStruct;
 	}
 
-	const FRDGEventScope* GetParentScope() const
+	const FRDGEventScope* GetEventScope() const
 	{
-		return ParentScope;
+		return EventScope;
+	}
+
+	const FRDGStatScope* GetStatScope() const
+	{
+		return StatScope;
 	}
 
 private:
@@ -224,7 +229,8 @@ private:
 	void Execute(FRHICommandListImmediate& RHICmdList) const;
 
 	const FRDGEventName Name;
-	const FRDGEventScope* ParentScope = nullptr;
+	const FRDGEventScope* EventScope = nullptr;
+	const FRDGStatScope* StatScope = nullptr;
 	FRDGPassParameterStruct ParameterStruct;
 	const ERDGPassFlags PassFlags;
 
