Moved inline methods in FRDGBuilder to RenderGraphBuilder.inl.

diff --git a/Engine/Source/Runtime/RenderCore/Public/RenderGraphBuilder.h b/Engine/Source/Runtime/RenderCore/Public/RenderGraphBuilder.h
--- a/Engine/Source/Runtime/RenderCore/Public/RenderGraphBuilder.h
+++ b/Engine/Source/Runtime/RenderCore/Public/RenderGraphBuilder.h
@@ -20,213 +20,63 @@ public:
 	~FRDGBuilder();
 
 	/** Register a external texture to be tracked by the render graph. */
-	FORCEINLINE_DEBUGGABLE FRDGTextureRef RegisterExternalTexture(const TRefCountPtr<IPooledRenderTarget>& ExternalPooledTexture, const TCHAR* DebugName = TEXT("External"))
-	{
-		#if RDG_ENABLE_DEBUG
-		{
-			ensureMsgf(ExternalPooledTexture.IsValid(), TEXT("Attempted to register NULL external texture: %s"), DebugName);
-			checkf(DebugName, TEXT("Externally allocated texture requires a debug name when registering them to render graph."));
-		}
-		#endif
-
-		FRDGTexture* OutTexture = AllocateForRHILifeTime<FRDGTexture>(DebugName, ExternalPooledTexture->GetDesc(), ERDGResourceFlags::None);
-		OutTexture->PooledRenderTarget = ExternalPooledTexture;
-		OutTexture->CachedRHI.Texture = ExternalPooledTexture->GetRenderTargetItem().ShaderResourceTexture;
-		AllocatedTextures.Add(OutTexture, ExternalPooledTexture);
-
-		#if RDG_ENABLE_DEBUG
-		{
-			OutTexture->MarkAsExternal();
-			TrackedResources.Add(OutTexture);
-		}
-		#endif
-
-		return OutTexture;
-	}
+	FRDGTextureRef RegisterExternalTexture(
+		const TRefCountPtr<IPooledRenderTarget>& ExternalPooledTexture,
+		const TCHAR* Name = TEXT("External"));
 
 	/** Register a external buffer to be tracked by the render graph. */
-	FORCEINLINE_DEBUGGABLE FRDGBufferRef RegisterExternalBuffer(const TRefCountPtr<FPooledRDGBuffer>& ExternalPooledBuffer, const TCHAR* Name = TEXT("External"))
-	{
-		#if RDG_ENABLE_DEBUG
-		{
-			ensureMsgf(ExternalPooledBuffer.IsValid(), TEXT("Attempted to register NULL external buffer: %s"), Name);
-		}
-		#endif
-
-		FRDGBuffer* OutBuffer = AllocateForRHILifeTime<FRDGBuffer>(Name, ExternalPooledBuffer->Desc, ERDGResourceFlags::None);
-		OutBuffer->PooledBuffer = ExternalPooledBuffer;
-		AllocatedBuffers.Add(OutBuffer, ExternalPooledBuffer);
-
-		#if RDG_ENABLE_DEBUG
-		{
-			OutBuffer->MarkAsExternal();
-			TrackedResources.Add(OutBuffer);
-		}
-		#endif
-
-		return OutBuffer;
-	}
+	FRDGBufferRef RegisterExternalBuffer(
+		const TRefCountPtr<FPooledRDGBuffer>& ExternalPooledBuffer,
+		const TCHAR* Name = TEXT("External"));
 
 	/** Create graph tracked resource from a descriptor with a debug name.
 	 *
-	 * The debug name is the name used for GPU debugging tools, but also for the VisualizeTexture/Vis command.
+	 *  The debug name is the name used for GPU debugging tools, but also for the VisualizeTexture/Vis command.
 	 */
-	FORCEINLINE_DEBUGGABLE FRDGTextureRef CreateTexture(const FPooledRenderTargetDesc& Desc, const TCHAR* DebugName, ERDGResourceFlags Flags = ERDGResourceFlags::None)
-	{
-		#if RDG_ENABLE_DEBUG
-		{
-			ensureMsgf(!bHasExecuted, TEXT("Render graph texture %s needs to be created before the builder execution."), DebugName);
-			checkf(DebugName, TEXT("Creating a render graph texture requires a valid debug name."));
-			checkf(Desc.Format != PF_Unknown, TEXT("Illegal to create texture %s with an invalid pixel format."), DebugName);
-		}
-		#endif
-
-		FRDGTexture* Texture = AllocateForRHILifeTime<FRDGTexture>(DebugName, Desc, Flags);
-
-		#if RDG_ENABLE_DEBUG
-		{
-			TrackedResources.Add(Texture);
-		}
-		#endif
-
-		return Texture;
-	}
+	FRDGTextureRef CreateTexture(
+		const FPooledRenderTargetDesc& Desc,
+		const TCHAR* Name,
+		ERDGResourceFlags Flags = ERDGResourceFlags::None);
 
 	/** Create graph tracked resource from a descriptor with a debug name.
 	 *
-	 * The debug name is the name used for GPU debugging tools, but also for the VisualizeTexture/Vis command.
+	 *  The debug name is the name used for GPU debugging tools, but also for the VisualizeTexture/Vis command.
 	 */
-	FORCEINLINE_DEBUGGABLE FRDGBufferRef CreateBuffer(const FRDGBufferDesc& Desc, const TCHAR* DebugName, ERDGResourceFlags Flags = ERDGResourceFlags::None)
-	{
-		#if RDG_ENABLE_DEBUG
-		{
-			ensureMsgf(!bHasExecuted, TEXT("Render graph buffer %s needs to be created before the builder execution."), DebugName);
-			checkf(DebugName, TEXT("Creating a render graph buffer requires a valid debug name."));
-		}
-		#endif
-
-		FRDGBufferRef Buffer = AllocateForRHILifeTime<FRDGBuffer>(DebugName, Desc, Flags);
-
-		#if RDG_ENABLE_DEBUG
-		{
-			TrackedResources.Add(Buffer);
-		}
-		#endif
-
-		return Buffer;
-	}
+	FRDGBufferRef CreateBuffer(
+		const FRDGBufferDesc& Desc,
+		const TCHAR* Name,
+		ERDGResourceFlags Flags = ERDGResourceFlags::None);
 
 	/** Create graph tracked SRV for a texture from a descriptor. */
-	FORCEINLINE_DEBUGGABLE FRDGTextureSRVRef CreateSRV(const FRDGTextureSRVDesc& Desc)
-	{
-		FRDGTextureRef Texture = Desc.Texture;
-
-		#if RDG_ENABLE_DEBUG
-		{
-			if (!Texture)
-			{
-				ensureMsgf(false, TEXT("RenderGraph texture SRV created with a null texture."));
-				return nullptr;
-			}
-
-			ensureMsgf(!bHasExecuted, TEXT("Render graph SRV %s needs to be created before the builder execution."), Desc.Texture->Name);
-			ensureMsgf(Desc.Texture->Desc.TargetableFlags & TexCreate_ShaderResource, TEXT("Attempted to create SRV from texture %s which was not created with TexCreate_ShaderResource"), Desc.Texture->Name);
-		}
-		#endif
-
-		return AllocateForRHILifeTime<FRDGTextureSRV>(Desc.Texture->Name, Desc);
-	}
+	FRDGTextureSRVRef CreateSRV(const FRDGTextureSRVDesc& Desc);
 
 	/** Create graph tracked SRV for a buffer from a descriptor. */
-	FORCEINLINE_DEBUGGABLE FRDGBufferSRVRef CreateSRV(const FRDGBufferSRVDesc& Desc)
-	{
-		FRDGBufferRef Buffer = Desc.Buffer;
-
-		#if RDG_ENABLE_DEBUG
-		{
-			if (!Buffer)
-			{
-				ensureMsgf(false, TEXT("RenderGraph buffer SRV created with a null buffer."));
-				return nullptr;
-			}
-
-			ensureMsgf(!bHasExecuted, TEXT("Render graph SRV %s needs to be created before the builder execution."), Desc.Buffer->Name);
-		}
-		#endif
-
-		return AllocateForRHILifeTime<FRDGBufferSRV>(Buffer->Name, Desc);
-	}
+	FRDGBufferSRVRef CreateSRV(const FRDGBufferSRVDesc& Desc);
 
-	FORCEINLINE_DEBUGGABLE FRDGBufferSRVRef CreateSRV(FRDGBufferRef Buffer, EPixelFormat Format)
+	FRDGBufferSRVRef CreateSRV(FRDGBufferRef Buffer, EPixelFormat Format)
 	{
 		return CreateSRV(FRDGBufferSRVDesc(Buffer, Format));
 	}
 
 	/** Create graph tracked UAV for a texture from a descriptor. */
-	FORCEINLINE_DEBUGGABLE FRDGTextureUAVRef CreateUAV(const FRDGTextureUAVDesc& Desc)
-	{
-		FRDGTextureRef Texture = Desc.Texture;
-		
-		#if RDG_ENABLE_DEBUG
-		{
-			if (!Texture)
-			{
-				checkf(false, TEXT("RenderGraph texture UAV created with a null texture."));
-				return nullptr;
-			}
-
-			ensureMsgf(!bHasExecuted, TEXT("Render graph UAV %s needs to be created before the builder execution."), Texture->Name);
-			ensureMsgf(Texture->Desc.TargetableFlags & TexCreate_UAV, TEXT("Attempted to create UAV from texture %s which was not created with TexCreate_UAV"), Texture->Name);
-		}
-		#endif
-		
-		return AllocateForRHILifeTime<FRDGTextureUAV>(Texture->Name, Desc);
-	}
+	FRDGTextureUAVRef CreateUAV(const FRDGTextureUAVDesc& Desc);
 
-	FORCEINLINE_DEBUGGABLE FRDGTextureUAVRef CreateUAV(FRDGTextureRef Texture)
+	FRDGTextureUAVRef CreateUAV(FRDGTextureRef Texture)
 	{
 		return CreateUAV(FRDGTextureUAVDesc(Texture));
 	}
 
 	/** Create graph tracked UAV for a buffer from a descriptor. */
-	FORCEINLINE_DEBUGGABLE FRDGBufferUAVRef CreateUAV(const FRDGBufferUAVDesc& Desc)
-	{
-		FRDGBufferRef Buffer = Desc.Buffer;
-
-		#if RDG_ENABLE_DEBUG
-		{
-			if (!Buffer)
-			{
-				checkf(false, TEXT("RenderGraph buffer UAV created with a null buffer."));
-				return nullptr;
-			}
-
-			ensureMsgf(!bHasExecuted, TEXT("Render graph UAV %s needs to be created before the builder execution."), Buffer->Name);
-		}
-		#endif
-		
-		return AllocateForRHILifeTime<FRDGBufferUAV>(Buffer->Name, Desc);
-	}
+	FRDGBufferUAVRef CreateUAV(const FRDGBufferUAVDesc& Desc);
 
-	FORCEINLINE_DEBUGGABLE FRDGBufferUAVRef CreateUAV(FRDGBufferRef Buffer, EPixelFormat Format)
+	FRDGBufferUAVRef CreateUAV(FRDGBufferRef Buffer, EPixelFormat Format)
 	{
 		return CreateUAV(FRDGBufferUAVDesc(Buffer, Format));
 	}
 
 	/** Allocates parameter struct specifically to survive through the life time of the render graph. */
 	template< typename ParameterStructType >
-	FORCEINLINE_DEBUGGABLE ParameterStructType* AllocParameters()
-	{
-		// TODO(RDG): could allocate using AllocateForRHILifeTime() to avoid the copy done when using FRHICommandList::BuildLocalUniformBuffer()
-		ParameterStructType* OutParameterPtr = new(MemStack) ParameterStructType;
-		FMemory::Memzero(OutParameterPtr, sizeof(ParameterStructType));
-		#if RDG_ENABLE_DEBUG
-		{
-			AllocatedUnusedPassParameters.Add(static_cast<void *>(OutParameterPtr));
-		}
-		#endif
-		return OutParameterPtr;
-	}
+	ParameterStructType* AllocParameters();
 
 	/** Adds a hard coded lambda pass to the graph.
 	 *
@@ -239,109 +89,46 @@ public:
 	 */
 	template<typename ParameterStructType, typename ExecuteLambdaType>
 	void AddPass(
-		FRDGEventName&& Name, 
+		FRDGEventName&& Name,
 		ParameterStructType* ParameterStruct,
 		ERDGPassFlags Flags,
-		ExecuteLambdaType&& ExecuteLambda)
-	{
-		auto NewPass = new(MemStack) TRDGLambdaPass<ParameterStructType, ExecuteLambdaType>(
-			MoveTemp(Name),
-			FRDGPassParameterStruct(ParameterStruct),
-			Flags,
-			MoveTemp(ExecuteLambda));
+		ExecuteLambdaType&& ExecuteLambda);
 
-		AddPassInternal(NewPass);
-	}
-
-	/** Adds a procedurally created pass to the render graph.
-	 *
-	 * Note: You want to use this only when the layout of the pass might be procedurally generated from data driven, as opose to AddPass() that have,
-	 * constant hard coded pass layout.
+	/** Adds an externally created pass to the render graph. This is useful if the pass parameter structure is generated at runtime.
 	 *
-	 * Caution: You are on your own to have correct memory lifetime of the FRDGPass.
+	 *  Caution: You are on your own to have correct memory lifetime of the FRDGPass.
 	 */
-	FORCEINLINE_DEBUGGABLE void AddProcedurallyCreatedPass(FRDGPass* NewPass)
+	void AddExternalPass(FRDGPass* NewPass)
 	{
 		AddPassInternal(NewPass);
 	}
 
 	/** Queue a texture extraction. This will set *OutTexturePtr with the internal pooled render target at the Execute().
 	 *
-	 * Note: even when the render graph uses the immediate debugging mode (executing passes as they get added), the texture extractions
-	 * will still happen in the Execute(), to ensure there is no bug caused in code outside the render graph on whether this mode is used or not.
+	 *  Note: even when the render graph uses the immediate debugging mode (executing passes as they get added), the texture extractions
+	 *  will still happen in the Execute(), to ensure there is no bug caused in code outside the render graph on whether this mode is used or not.
 	 */
-	FORCEINLINE_DEBUGGABLE void QueueTextureExtraction(FRDGTextureRef Texture, TRefCountPtr<IPooledRenderTarget>* OutTexturePtr, bool bTransitionToRead = true)
-	{
-		check(Texture);
-		check(OutTexturePtr);
-		#if RDG_ENABLE_DEBUG
-		{
-			checkf(!bHasExecuted,
-				TEXT("Accessing render graph internal texture %s with QueueTextureExtraction() needs to happen before the builder's execution."),
-				Texture->Name);
-
-			checkf(Texture->HasBeenProduced(),
-				TEXT("Unable to queue the extraction of the texture %s because it has not been produced by any pass."),
-				Texture->Name);
-		}
-		#endif
-		FDeferredInternalTextureQuery Query;
-		Query.Texture = Texture;
-		Query.OutTexturePtr = OutTexturePtr;
-		Query.bTransitionToRead = bTransitionToRead;
-		DeferredInternalTextureQueries.Emplace(Query);
-	}
+	void QueueTextureExtraction(FRDGTextureRef Texture, TRefCountPtr<IPooledRenderTarget>* OutTexturePtr, bool bTransitionToRead = true);
 
 	/** Queue a buffer extraction. This will set *OutBufferPtr with the internal pooled buffer at the Execute().
 	 *
-	 * Note: even when the render graph uses the immediate debugging mode (executing passes as they get added), the buffer extractions
-	 * will still happen in the Execute(), to ensure there is no bug caused in code outside the render graph on whether this mode is used or not.
+	 *  Note: even when the render graph uses the immediate debugging mode (executing passes as they get added), the buffer extractions
+	 *  will still happen in the Execute(), to ensure there is no bug caused in code outside the render graph on whether this mode is used or not.
 	 */
-	FORCEINLINE_DEBUGGABLE void QueueBufferExtraction(FRDGBufferRef Buffer, TRefCountPtr<FPooledRDGBuffer>* OutBufferPtr)
-	{
-		check(Buffer);
-		check(OutBufferPtr);
-		#if RDG_ENABLE_DEBUG
-		{
-			checkf(!bHasExecuted,
-				TEXT("Accessing render graph internal buffer %s with QueueBufferExtraction() needs to happen before the builder's execution."),
-				Buffer->Name);
-
-			checkf(Buffer->HasBeenProduced(),
-				TEXT("Unable to queue the extraction of the buffer %s because it has not been produced by any pass."),
-				Buffer->Name);
-		}
-		#endif
-		FDeferredInternalBufferQuery Query;
-		Query.Buffer = Buffer;
-		Query.OutBufferPtr = OutBufferPtr;
-		DeferredInternalBufferQueries.Emplace(Query);
-	}
+	void QueueBufferExtraction(FRDGBufferRef Buffer, TRefCountPtr<FPooledRDGBuffer>* OutBufferPtr);
 
 	/** Flag a texture that is only produced by only 1 pass, but never used or extracted, to avoid generating a warning at runtime. */
-	FORCEINLINE_DEBUGGABLE void RemoveUnusedTextureWarning(FRDGTextureRef Texture)
-	{
-		check(Texture);
-		#if RDG_ENABLE_DEBUG
-		{
-			checkf(!bHasExecuted,
-				TEXT("Flaging texture %s with FlagUnusedTexture() needs to happen before the builder's execution."),
-				Texture->Name);
-			
-			// Increment the number of time the texture has been accessed to avoid warning on produced but never used resources that were produced
-			// only to be extracted for the graph.
-			Texture->PassAccessCount += 1;
-		}
-		#endif
-	}
+	void RemoveUnusedTextureWarning(FRDGTextureRef Texture);
 
 	/** Begins / ends a named event scope. These scopes are visible in most external GPU profilers.
-	 *  Prefer to use RDG_EVENT_SCOPE instead of calling this manually. */
+	 *  Prefer to use RDG_EVENT_SCOPE instead of calling this manually.
+	 */
 	void BeginEventScope(FRDGEventName&& ScopeName);
 	void EndEventScope();
 
 	/** Begins / ends a stat scope. Stat scopes are visible via 'stat GPU'. Must be accompanied by a
-	 *  respective EndStatScope call. Prefer to use RDG_GPU_STAT_SCOPE instead of calling this manually. */
+	 *  respective EndStatScope call. Prefer to use RDG_GPU_STAT_SCOPE instead of calling this manually.
+	 */
 	void BeginStatScope(const FName& Name, const FName& StatName);
 	void EndStatScope();
 
@@ -401,21 +188,7 @@ private:
 	void WalkGraphDependencies();
 	
 	template<class Type, class ...ConstructorParameterTypes>
-	Type* AllocateForRHILifeTime(ConstructorParameterTypes&&... ConstructorParameters)
-	{
-		check(IsInRenderingThread());
-		// When bypassing the RHI command queuing, can allocate directly on render thread memory stack allocator, otherwise allocate
-		// on the RHI's stack allocator so RHICreateUniformBuffer() can dereference render graph resources.
-		if (RHICmdList.Bypass() || 1) // TODO: UE-68018
-		{
-			return new (MemStack) Type(Forward<ConstructorParameterTypes>(ConstructorParameters)...);
-		}
-		else
-		{
-			void* UnitializedType = RHICmdList.Alloc<Type>();
-			return new (UnitializedType) Type(Forward<ConstructorParameterTypes>(ConstructorParameters)...);
-		}
-	}
+	Type* AllocateForRHILifeTime(ConstructorParameterTypes&&... ConstructorParameters);
 
 	void AddPassInternal(FRDGPass* Pass);
 
@@ -443,4 +216,6 @@ private:
 	/** To allow greater flexibility in the user code, the RHI can dereferenced RDG resource when creating uniform buffer. */
 	// TODO(RDG): Make this a little more explicit in RHI code.
 	static_assert(STRUCT_OFFSET(FRDGResource, CachedRHI) == 0, "FRDGResource::CachedRHI requires to be at offset 0 so the RHI can dereferenced them.");
-};
\ No newline at end of file
+};
+
+#include "RenderGraphBuilder.inl"
\ No newline at end of file
diff --git a/Engine/Source/Runtime/RenderCore/Public/RenderGraphBuilder.inl b/Engine/Source/Runtime/RenderCore/Public/RenderGraphBuilder.inl
new file mode 100644
--- /dev/null
+++ b/Engine/Source/Runtime/RenderCore/Public/RenderGraphBuilder.inl
@@ -0,0 +1,283 @@
+// Copyright 1998-2019 Epic Games, Inc. All Rights Reserved.
+
+#pragma once
+
+inline FRDGTextureRef FRDGBuilder::RegisterExternalTexture(const TRefCountPtr<IPooledRenderTarget>& ExternalPooledTexture, const TCHAR* Name)
+{
+#if RDG_ENABLE_DEBUG
+	{
+		ensureMsgf(ExternalPooledTexture.IsValid(), TEXT("Attempted to register NULL external texture: %s"), Name);
+		checkf(Name, TEXT("Externally allocated texture requires a debug name when registering them to render graph."));
+	}
+#endif
+
+	FRDGTexture* OutTexture = AllocateForRHILifeTime<FRDGTexture>(Name, ExternalPooledTexture->GetDesc(), ERDGResourceFlags::None);
+	OutTexture->PooledRenderTarget = ExternalPooledTexture;
+	OutTexture->CachedRHI.Texture = ExternalPooledTexture->GetRenderTargetItem().ShaderResourceTexture;
+	AllocatedTextures.Add(OutTexture, ExternalPooledTexture);
+
+#if RDG_ENABLE_DEBUG
+	{
+		OutTexture->MarkAsExternal();
+		TrackedResources.Add(OutTexture);
+	}
+#endif
+
+	return OutTexture;
+}
+
+inline FRDGBufferRef FRDGBuilder::RegisterExternalBuffer(const TRefCountPtr<FPooledRDGBuffer>& ExternalPooledBuffer, const TCHAR* Name)
+{
+#if RDG_ENABLE_DEBUG
+	{
+		ensureMsgf(ExternalPooledBuffer.IsValid(), TEXT("Attempted to register NULL external buffer: %s"), Name);
+	}
+#endif
+
+	FRDGBuffer* OutBuffer = AllocateForRHILifeTime<FRDGBuffer>(Name, ExternalPooledBuffer->Desc, ERDGResourceFlags::None);
+	OutBuffer->PooledBuffer = ExternalPooledBuffer;
+	AllocatedBuffers.Add(OutBuffer, ExternalPooledBuffer);
+
+#if RDG_ENABLE_DEBUG
+	{
+		OutBuffer->MarkAsExternal();
+		TrackedResources.Add(OutBuffer);
+	}
+#endif
+
+	return OutBuffer;
+}
+
+inline FRDGTextureRef FRDGBuilder::CreateTexture(
+	const FPooledRenderTargetDesc& Desc,
+	const TCHAR* DebugName,
+	ERDGResourceFlags Flags)
+{
+#if RDG_ENABLE_DEBUG
+	{
+		ensureMsgf(!bHasExecuted, TEXT("Render graph texture %s needs to be created before the builder execution."), DebugName);
+		checkf(DebugName, TEXT("Creating a render graph texture requires a valid debug name."));
+		checkf(Desc.Format != PF_Unknown, TEXT("Illegal to create texture %s with an invalid pixel format."), DebugName);
+	}
+#endif
+
+	FRDGTexture* Texture = AllocateForRHILifeTime<FRDGTexture>(DebugName, Desc, Flags);
+
+#if RDG_ENABLE_DEBUG
+	{
+		TrackedResources.Add(Texture);
+	}
+#endif
+
+	return Texture;
+}
+
+inline FRDGBufferRef FRDGBuilder::CreateBuffer(
+	const FRDGBufferDesc& Desc,
+	const TCHAR* Name,
+	ERDGResourceFlags Flags)
+{
+#if RDG_ENABLE_DEBUG
+	{
+		ensureMsgf(!bHasExecuted, TEXT("Render graph buffer %s needs to be created before the builder execution."), Name);
+		checkf(Name, TEXT("Creating a render graph buffer requires a valid debug name."));
+	}
+#endif
+
+	FRDGBufferRef Buffer = AllocateForRHILifeTime<FRDGBuffer>(Name, Desc, Flags);
+
+#if RDG_ENABLE_DEBUG
+	{
+		TrackedResources.Add(Buffer);
+	}
+#endif
+
+	return Buffer;
+}
+
+inline FRDGTextureSRVRef FRDGBuilder::CreateSRV(const FRDGTextureSRVDesc& Desc)
+{
+	FRDGTextureRef Texture = Desc.Texture;
+
+#if RDG_ENABLE_DEBUG
+	{
+		if (!Texture)
+		{
+			ensureMsgf(false, TEXT("RenderGraph texture SRV created with a null texture."));
+			return nullptr;
+		}
+
+		ensureMsgf(!bHasExecuted, TEXT("Render graph SRV %s needs to be created before the builder execution."), Desc.Texture->Name);
+		ensureMsgf(Desc.Texture->Desc.TargetableFlags & TexCreate_ShaderResource, TEXT("Attempted to create SRV from texture %s which was not created with TexCreate_ShaderResource"), Desc.Texture->Name);
+	}
+#endif
+
+	return AllocateForRHILifeTime<FRDGTextureSRV>(Desc.Texture->Name, Desc);
+}
+
+inline FRDGBufferSRVRef FRDGBuilder::CreateSRV(const FRDGBufferSRVDesc& Desc)
+{
+	FRDGBufferRef Buffer = Desc.Buffer;
+
+#if RDG_ENABLE_DEBUG
+	{
+		if (!Buffer)
+		{
+			ensureMsgf(false, TEXT("RenderGraph buffer SRV created with a null buffer."));
+			return nullptr;
+		}
+
+		ensureMsgf(!bHasExecuted, TEXT("Render graph SRV %s needs to be created before the builder execution."), Desc.Buffer->Name);
+	}
+#endif
+
+	return AllocateForRHILifeTime<FRDGBufferSRV>(Buffer->Name, Desc);
+}
+
+inline FRDGTextureUAVRef FRDGBuilder::CreateUAV(const FRDGTextureUAVDesc& Desc)
+{
+	FRDGTextureRef Texture = Desc.Texture;
+
+#if RDG_ENABLE_DEBUG
+	{
+		if (!Texture)
+		{
+			checkf(false, TEXT("RenderGraph texture UAV created with a null texture."));
+			return nullptr;
+		}
+
+		ensureMsgf(!bHasExecuted, TEXT("Render graph UAV %s needs to be created before the builder execution."), Texture->Name);
+		ensureMsgf(Texture->Desc.TargetableFlags & TexCreate_UAV, TEXT("Attempted to create UAV from texture %s which was not created with TexCreate_UAV"), Texture->Name);
+	}
+#endif
+
+	return AllocateForRHILifeTime<FRDGTextureUAV>(Texture->Name, Desc);
+}
+
+inline FRDGBufferUAVRef FRDGBuilder::CreateUAV(const FRDGBufferUAVDesc& Desc)
+{
+	FRDGBufferRef Buffer = Desc.Buffer;
+
+#if RDG_ENABLE_DEBUG
+	{
+		if (!Buffer)
+		{
+			checkf(false, TEXT("RenderGraph buffer UAV created with a null buffer."));
+			return nullptr;
+		}
+
+		ensureMsgf(!bHasExecuted, TEXT("Render graph UAV %s needs to be created before the builder execution."), Buffer->Name);
+	}
+#endif
+
+	return AllocateForRHILifeTime<FRDGBufferUAV>(Buffer->Name, Desc);
+}
+
+template <typename ParameterStructType>
+ParameterStructType* FRDGBuilder::AllocParameters()
+{
+	// TODO(RDG): could allocate using AllocateForRHILifeTime() to avoid the copy done when using FRHICommandList::BuildLocalUniformBuffer()
+	ParameterStructType* OutParameterPtr = new(MemStack) ParameterStructType;
+	FMemory::Memzero(OutParameterPtr, sizeof(ParameterStructType));
+#if RDG_ENABLE_DEBUG
+	{
+		AllocatedUnusedPassParameters.Add(static_cast<void *>(OutParameterPtr));
+	}
+#endif
+	return OutParameterPtr;
+}
+
+template <typename ParameterStructType, typename ExecuteLambdaType>
+void FRDGBuilder::AddPass(
+	FRDGEventName&& Name,
+	ParameterStructType* ParameterStruct,
+	ERDGPassFlags Flags,
+	ExecuteLambdaType&& ExecuteLambda)
+{
+	auto NewPass = new(MemStack) TRDGLambdaPass<ParameterStructType, ExecuteLambdaType>(
+		MoveTemp(Name),
+		FRDGPassParameterStruct(ParameterStruct),
+		Flags,
+		MoveTemp(ExecuteLambda));
+
+	AddPassInternal(NewPass);
+}
+
+inline void FRDGBuilder::QueueTextureExtraction(
+	FRDGTextureRef Texture,
+	TRefCountPtr<IPooledRenderTarget>* OutTexturePtr,
+	bool bTransitionToRead)
+{
+	check(Texture);
+	check(OutTexturePtr);
+#if RDG_ENABLE_DEBUG
+	{
+		checkf(!bHasExecuted,
+			TEXT("Accessing render graph internal texture %s with QueueTextureExtraction() needs to happen before the builder's execution."),
+			Texture->Name);
+
+		checkf(Texture->HasBeenProduced(),
+			TEXT("Unable to queue the extraction of the texture %s because it has not been produced by any pass."),
+			Texture->Name);
+	}
+#endif
+	FDeferredInternalTextureQuery Query;
+	Query.Texture = Texture;
+	Query.OutTexturePtr = OutTexturePtr;
+	Query.bTransitionToRead = bTransitionToRead;
+	DeferredInternalTextureQueries.Emplace(Query);
+}
+
+inline void FRDGBuilder::QueueBufferExtraction(FRDGBufferRef Buffer, TRefCountPtr<FPooledRDGBuffer>* OutBufferPtr)
+{
+	check(Buffer);
+	check(OutBufferPtr);
+#if RDG_ENABLE_DEBUG
+	{
+		checkf(!bHasExecuted,
+			TEXT("Accessing render graph internal buffer %s with QueueBufferExtraction() needs to happen before the builder's execution."),
+			Buffer->Name);
+
+		checkf(Buffer->HasBeenProduced(),
+			TEXT("Unable to queue the extraction of the buffer %s because it has not been produced by any pass."),
+			Buffer->Name);
+	}
+#endif
+	FDeferredInternalBufferQuery Query;
+	Query.Buffer = Buffer;
+	Query.OutBufferPtr = OutBufferPtr;
+	DeferredInternalBufferQueries.Emplace(Query);
+}
+
+inline void FRDGBuilder::RemoveUnusedTextureWarning(FRDGTextureRef Texture)
+{
+	check(Texture);
+#if RDG_ENABLE_DEBUG
+	{
+		checkf(!bHasExecuted,
+			TEXT("Flaging texture %s with FlagUnusedTexture() needs to happen before the builder's execution."),
+			Texture->Name);
+
+		// Increment the number of time the texture has been accessed to avoid warning on produced but never used resources that were produced
+		// only to be extracted for the graph.
+		Texture->PassAccessCount += 1;
+	}
+#endif
+}
+
+template<class Type, class ...ConstructorParameterTypes>
+Type* FRDGBuilder::AllocateForRHILifeTime(ConstructorParameterTypes&&... ConstructorParameters)
+{
+	check(IsInRenderingThread());
+	// When bypassing the RHI command queuing, can allocate directly on render thread memory stack allocator, otherwise allocate
+	// on the RHI's stack allocator so RHICreateUniformBuffer() can dereference render graph resources.
+	if (RHICmdList.Bypass() || 1) // TODO: UE-68018
+	{
+		return new (MemStack) Type(Forward<ConstructorParameterTypes>(ConstructorParameters)...);
+	}
+	else
+	{
+		void* UnitializedType = RHICmdList.Alloc<Type>();
+		return new (UnitializedType) Type(Forward<ConstructorParameterTypes>(ConstructorParameters)...);
+	}
+}
\ No newline at end of file
