Fix compiler error if the csv profiler is not enabled

diff --git a/Engine/Source/Runtime/RenderCore/Private/ProfilingDebugging/RealtimeGPUProfiler.cpp b/Engine/Source/Runtime/RenderCore/Private/ProfilingDebugging/RealtimeGPUProfiler.cpp
--- a/Engine/Source/Runtime/RenderCore/Private/ProfilingDebugging/RealtimeGPUProfiler.cpp
+++ b/Engine/Source/Runtime/RenderCore/Private/ProfilingDebugging/RealtimeGPUProfiler.cpp
@@ -455,11 +455,13 @@ public:
 			FThreadStats::AddMessage(Event->GetStatName(), StatOp, double(EventTime));
 #endif
 
+#if CSV_PROFILER
 			if (bCsvStatsEnabled)
 			{
 				ECsvCustomStatOp CsvStatOp = bIsNew ? ECsvCustomStatOp::Set : ECsvCustomStatOp::Accumulate;
 				FCsvProfiler::Get()->RecordCustomStat(Event->GetName(), CSV_CATEGORY_INDEX(GPU), EventTime, CsvStatOp);
 			}
+#endif
 
 #if TRACING_PROFILER
 			if (bTracingStatsEnabled)
