Fixing build failures.

diff --git a/Engine/Source/Runtime/RenderCore/Private/RenderGraphEvent.cpp b/Engine/Source/Runtime/RenderCore/Private/RenderGraphEvent.cpp
--- a/Engine/Source/Runtime/RenderCore/Private/RenderGraphEvent.cpp
+++ b/Engine/Source/Runtime/RenderCore/Private/RenderGraphEvent.cpp
@@ -134,7 +134,9 @@ FRDGStatScopeGuard::~FRDGStatScopeGuard()
 	GraphBuilder.EndStatScope();
 }
 
+#if HAS_GPU_STATS
 extern bool AreGPUStatsEnabled();
+#endif
 
 static void OnPushStat(FRHICommandListImmediate& RHICmdList, const FRDGStatScope* Scope)
 {
@@ -152,7 +154,11 @@ static void OnPopStat(FRHICommandListImmediate& RHICmdList)
 
 bool FRDGStatScopeStack::IsEnabled()
 {
-	return HAS_GPU_STATS && AreGPUStatsEnabled();
+#if HAS_GPU_STATS
+	return AreGPUStatsEnabled();
+#else
+	return false;
+#endif
 }
 
 FRDGStatScopeStack::FRDGStatScopeStack(FRHICommandListImmediate& InRHICmdList)
