Fixing build failures.

diff --git a/Engine/Source/Runtime/RenderCore/Private/ProfilingDebugging/RealtimeGPUProfiler.cpp b/Engine/Source/Runtime/RenderCore/Private/ProfilingDebugging/RealtimeGPUProfiler.cpp
--- a/Engine/Source/Runtime/RenderCore/Private/ProfilingDebugging/RealtimeGPUProfiler.cpp
+++ b/Engine/Source/Runtime/RenderCore/Private/ProfilingDebugging/RealtimeGPUProfiler.cpp
@@ -567,7 +567,7 @@ void FRealtimeGPUProfiler::BeginFrame(FRHICommandListImmediate& RHICmdList)
 	bInBeginEndBlock = true;
 }
 
-inline bool AreGPUStatsEnabled()
+bool AreGPUStatsEnabled()
 {
 	if (GSupportsTimestampRenderQueries == false || !CVarGPUStatsEnabled.GetValueOnRenderThread())
 	{
diff --git a/Engine/Source/Runtime/RenderCore/Private/RenderGraphEvent.cpp b/Engine/Source/Runtime/RenderCore/Private/RenderGraphEvent.cpp
--- a/Engine/Source/Runtime/RenderCore/Private/RenderGraphEvent.cpp
+++ b/Engine/Source/Runtime/RenderCore/Private/RenderGraphEvent.cpp
@@ -134,10 +134,6 @@ FRDGStatScopeGuard::~FRDGStatScopeGuard()
 	GraphBuilder.EndStatScope();
 }
 
-#if HAS_GPU_STATS
-extern bool AreGPUStatsEnabled();
-#endif
-
 static void OnPushStat(FRHICommandListImmediate& RHICmdList, const FRDGStatScope* Scope)
 {
 #if HAS_GPU_STATS
diff --git a/Engine/Source/Runtime/RenderCore/Public/ProfilingDebugging/RealtimeGPUProfiler.h b/Engine/Source/Runtime/RenderCore/Public/ProfilingDebugging/RealtimeGPUProfiler.h
--- a/Engine/Source/Runtime/RenderCore/Public/ProfilingDebugging/RealtimeGPUProfiler.h
+++ b/Engine/Source/Runtime/RenderCore/Public/ProfilingDebugging/RealtimeGPUProfiler.h
@@ -215,6 +215,8 @@ class FScopedGPUStatEvent;
  #define GPU_STATS_ENDFRAME(RHICmdList) 
 #endif
 
+bool AreGPUStatsEnabled();
+
 #if HAS_GPU_STATS
 
 class FRealtimeGPUProfilerEvent;
