Adds validations of RT's load/store actions.

diff --git a/Engine/Source/Runtime/RenderCore/Private/RenderGraphBuilder.cpp b/Engine/Source/Runtime/RenderCore/Private/RenderGraphBuilder.cpp
--- a/Engine/Source/Runtime/RenderCore/Private/RenderGraphBuilder.cpp
+++ b/Engine/Source/Runtime/RenderCore/Private/RenderGraphBuilder.cpp
@@ -324,11 +324,16 @@ void FRDGBuilder::ValidatePass(const FRDGPass* Pass) const
 						!bIsLoadActionInvalid,
 						TEXT("Can't load a render target '%s' that has never been produced."),
 						Texture->Name);
-
+					
+					const bool bOverwritesResourceEntirely = !bIsLoadAction && Texture->HasBeenProduced() && Texture->Desc.NumMips == 1;
+					ensureMsgf(!bOverwritesResourceEntirely,
+						TEXT("Clobbering render target %s result that has been produced previously by another pass. Should instead render to a new FRDGTexture to reduce memory pressure."),
+						Texture->Name);
+					
 					/** Mark the pass as a producer for render targets with a store action. */
-					const bool bIsStoreAction = RenderTarget.GetStoreAction() != ERenderTargetStoreAction::ENoAction;
-					if (bIsStoreAction)
 					{
+						const bool bIsStoreAction = RenderTarget.GetStoreAction() != ERenderTargetStoreAction::ENoAction;
+						check(bIsStoreAction); // already been validated in FRenderTargetBinding::Validate()
 						Texture->MarkAsProducedBy(Pass);
 					}
 				}
diff --git a/Engine/Source/Runtime/RenderCore/Private/ShaderParameterStruct.cpp b/Engine/Source/Runtime/RenderCore/Private/ShaderParameterStruct.cpp
--- a/Engine/Source/Runtime/RenderCore/Private/ShaderParameterStruct.cpp
+++ b/Engine/Source/Runtime/RenderCore/Private/ShaderParameterStruct.cpp
@@ -309,6 +309,23 @@ void FShaderParameterBindings::BindForRootShaderParameters(const FShader* Shader
 	}
 }
 
+bool FRenderTargetBinding::Validate() const
+{
+	if (Texture)
+	{
+		checkf(StoreAction != ERenderTargetStoreAction::ENoAction,
+			TEXT("You must specify a store action for non-null render target %s."),
+			Texture->Name);
+	}
+	else
+	{
+		checkf(LoadAction == ERenderTargetLoadAction::ENoAction && StoreAction == ERenderTargetStoreAction::ENoAction,
+			TEXT("Can't have a load or store action when no texture is bound."));
+	}
+	
+	return true;
+}
+
 bool FDepthStencilBinding::Validate() const
 {
 	if (Texture)
diff --git a/Engine/Source/Runtime/RenderCore/Public/ShaderParameterMacros.h b/Engine/Source/Runtime/RenderCore/Public/ShaderParameterMacros.h
--- a/Engine/Source/Runtime/RenderCore/Public/ShaderParameterMacros.h
+++ b/Engine/Source/Runtime/RenderCore/Public/ShaderParameterMacros.h
@@ -154,7 +154,9 @@ struct alignas(SHADER_PARAMETER_STRUCT_ALIGNMENT) FRenderTargetBinding
 		, LoadAction(InLoadAction)
 		, StoreAction(InStoreAction)
 		, MipIndex(InMipIndex)
-	{}
+	{
+		check(Validate());
+	}
 
 	FORCEINLINE FRDGTexture* GetTexture() const
 	{
@@ -181,6 +183,8 @@ private:
 	ERenderTargetLoadAction		LoadAction		= ERenderTargetLoadAction::ENoAction;
 	ERenderTargetStoreAction	StoreAction		= ERenderTargetStoreAction::ENoAction;
 	uint8						MipIndex		= 0;
+
+	RENDERCORE_API bool Validate() const;
 };
 
 
