Fixed merge error detected thanks to warning V517

diff --git a/Engine/Source/Runtime/RenderCore/Private/RenderGraphBuilder.cpp b/Engine/Source/Runtime/RenderCore/Private/RenderGraphBuilder.cpp
--- a/Engine/Source/Runtime/RenderCore/Private/RenderGraphBuilder.cpp
+++ b/Engine/Source/Runtime/RenderCore/Private/RenderGraphBuilder.cpp
@@ -794,9 +794,6 @@ void FRDGBuilder::TransitionUAV(FUnorderedAccessViewRHIParamRef UAV, FRDGTracked
 		RHICmdList.TransitionResource(EResourceTransitionAccess::ERWBarrier, TransitionPipeline, UAV);
 	}
 	else if(UnderlyingResource->bWritable != bRequiredWritable || UnderlyingResource->bCompute != bRequiredCompute )
-	{
-	}
-	else if(UnderlyingResource->bWritable != bRequiredWritable || UnderlyingResource->bCompute != bRequiredCompute )
 	{
 		const bool bIsMultiFrameResource = (ResourceFlags & ERDGResourceFlags::MultiFrame) == ERDGResourceFlags::MultiFrame;
 		if (bIsMultiFrameResource && bRequiredWritable && !UnderlyingResource->bWritable)
