Eliminate CachedRHI union and direct casting to FRHIResource**

diff --git a/Engine/Source/Runtime/D3D12RHI/Private/D3D12UniformBuffer.cpp b/Engine/Source/Runtime/D3D12RHI/Private/D3D12UniformBuffer.cpp
--- a/Engine/Source/Runtime/D3D12RHI/Private/D3D12UniformBuffer.cpp
+++ b/Engine/Source/Runtime/D3D12RHI/Private/D3D12UniformBuffer.cpp
@@ -6,6 +6,7 @@
 
 #include "D3D12RHIPrivate.h"
 #include "UniformBuffer.h"
+#include "RenderGraphResources.h"
 
 FUniformBufferRHIRef FD3D12DynamicRHI::RHICreateUniformBuffer(const void* Contents, const FRHIUniformBufferLayout& Layout, EUniformBufferUsage Usage, EUniformBufferValidation Validation)
 {
@@ -91,8 +92,8 @@ FUniformBufferRHIRef FD3D12DynamicRHI::RHICreateUniformBuffer(const void* Conten
 				else if (IsRDGResourceReferenceShaderParameterType(ResourceType))
 				{
 					check(IsInRenderingThread()); // TODO: UE-68018
-					FRHIResource** ResourcePtr = *(FRHIResource***)((uint8*)Contents + Layout.Resources[i].MemberOffset);
-					Resource = ResourcePtr ? *ResourcePtr : nullptr;
+					FRDGResource* ResourcePtr = *(FRDGResource**)((uint8*)Contents + Layout.Resources[i].MemberOffset);
+					Resource = ResourcePtr ? ResourcePtr->GetRHI() : nullptr;
 				}
 				else
 				{
@@ -185,8 +186,8 @@ void FD3D12DynamicRHI::RHIUpdateUniformBuffer(FUniformBufferRHIParamRef UniformB
 			else if (IsRDGResourceReferenceShaderParameterType(ResourceType))
 			{
 				check(IsInRenderingThread()); // TODO: UE-68018
-				FRHIResource** ResourcePtr = *(FRHIResource***)((uint8*)Contents + Layout.Resources[ResourceIndex].MemberOffset);
-				Resource = ResourcePtr ? *ResourcePtr : nullptr;
+				FRDGResource* ResourcePtr = *(FRDGResource**)((uint8*)Contents + Layout.Resources[ResourceIndex].MemberOffset);
+				Resource = ResourcePtr ? ResourcePtr->GetRHI() : nullptr;
 			}
 			else
 			{
diff --git a/Engine/Source/Runtime/RHI/Public/RHIDefinitions.h b/Engine/Source/Runtime/RHI/Public/RHIDefinitions.h
--- a/Engine/Source/Runtime/RHI/Public/RHIDefinitions.h
+++ b/Engine/Source/Runtime/RHI/Public/RHIDefinitions.h
@@ -1174,7 +1174,7 @@ inline bool IsShaderParameterTypeForUniformBufferLayout(EUniformBufferBaseType B
 		BaseType == UBMT_SRV ||
 		BaseType == UBMT_SAMPLER ||
 
-		// RHI is able to dereference the RHI resource allocated in FRDGResource::CachedRHI to avoid pain in the high level to passdown.
+		// RHI is able to access RHI resources from RDG.
 		IsRDGResourceReferenceShaderParameterType(BaseType) ||
 
 		// Render graph uses FRHIUniformBufferLayout to walk pass' parameters.
diff --git a/Engine/Source/Runtime/RenderCore/Private/RenderGraphBuilder.cpp b/Engine/Source/Runtime/RenderCore/Private/RenderGraphBuilder.cpp
--- a/Engine/Source/Runtime/RenderCore/Private/RenderGraphBuilder.cpp
+++ b/Engine/Source/Runtime/RenderCore/Private/RenderGraphBuilder.cpp
@@ -594,7 +594,7 @@ void FRDGBuilder::WalkGraphDependencies()
 		{
 			Pair.Value = nullptr;
 			Pair.Key->PooledRenderTarget = nullptr;
-			Pair.Key->CachedRHI.Resource = nullptr;
+			Pair.Key->ResourceRHI = nullptr;
 		}
 	}
 
@@ -605,7 +605,7 @@ void FRDGBuilder::WalkGraphDependencies()
 		{
 			Pair.Value = nullptr;
 			Pair.Key->PooledBuffer = nullptr;
-			Pair.Key->CachedRHI.Resource = nullptr;
+			Pair.Key->ResourceRHI = nullptr;
 		}
 	}
 }
@@ -626,29 +626,29 @@ void FRDGBuilder::AllocateRHITextureIfNeeded(FRDGTexture* Texture, bool bCompute
 	GRenderTargetPool.FindFreeElement(RHICmdList, Texture->Desc, PooledRenderTarget, Texture->Name, /* bDoWritableBarrier = */ true);
 
 	Texture->PooledRenderTarget = PooledRenderTarget;
-	Texture->CachedRHI.Texture = PooledRenderTarget->GetRenderTargetItem().ShaderResourceTexture;
-	check(Texture->CachedRHI.Resource);
+	Texture->ResourceRHI = PooledRenderTarget->GetRenderTargetItem().ShaderResourceTexture;
+	check(Texture->ResourceRHI);
 }
 
 void FRDGBuilder::AllocateRHITextureUAVIfNeeded(FRDGTextureUAV* UAV, bool bComputePass)
 {
 	check(UAV);
 
-	if (UAV->CachedRHI.UAV)
+	if (UAV->ResourceRHI)
 	{
 		return;
 	}
 
 	AllocateRHITextureIfNeeded(UAV->Desc.Texture, bComputePass);
 
-	UAV->CachedRHI.UAV = UAV->Desc.Texture->PooledRenderTarget->GetRenderTargetItem().MipUAVs[UAV->Desc.MipLevel];
+	UAV->ResourceRHI = UAV->Desc.Texture->PooledRenderTarget->GetRenderTargetItem().MipUAVs[UAV->Desc.MipLevel];
 }
 
 void FRDGBuilder::AllocateRHIBufferSRVIfNeeded(FRDGBufferSRV* SRV, bool bComputePass)
 {
 	check(SRV);
 
-	if (SRV->CachedRHI.SRV)
+	if (SRV->ResourceRHI)
 	{
 		return;
 	}
@@ -666,7 +666,7 @@ void FRDGBuilder::AllocateRHIBufferSRVIfNeeded(FRDGBufferSRV* SRV, bool bCompute
 
 	if (Buffer->PooledBuffer->SRVs.Contains(SRV->Desc))
 	{
-		SRV->CachedRHI.SRV = Buffer->PooledBuffer->SRVs[SRV->Desc];
+		SRV->ResourceRHI = Buffer->PooledBuffer->SRVs[SRV->Desc];
 		return;
 	}
 
@@ -685,7 +685,7 @@ void FRDGBuilder::AllocateRHIBufferSRVIfNeeded(FRDGBufferSRV* SRV, bool bCompute
 		check(0);
 	}
 
-	SRV->CachedRHI.SRV = RHIShaderResourceView;
+	SRV->ResourceRHI = RHIShaderResourceView;
 	Buffer->PooledBuffer->SRVs.Add(SRV->Desc, RHIShaderResourceView);
 }
 
@@ -693,7 +693,7 @@ void FRDGBuilder::AllocateRHIBufferUAVIfNeeded(FRDGBufferUAV* UAV, bool bCompute
 {
 	check(UAV);
 
-	if (UAV->CachedRHI.UAV)
+	if (UAV->ResourceRHI)
 	{
 		return;
 	}
@@ -713,7 +713,7 @@ void FRDGBuilder::AllocateRHIBufferUAVIfNeeded(FRDGBufferUAV* UAV, bool bCompute
 
 	if (Buffer->PooledBuffer->UAVs.Contains(UAV->Desc))
 	{
-		UAV->CachedRHI.UAV = Buffer->PooledBuffer->UAVs[UAV->Desc];
+		UAV->ResourceRHI = Buffer->PooledBuffer->UAVs[UAV->Desc];
 		return;
 	}
 
@@ -735,7 +735,7 @@ void FRDGBuilder::AllocateRHIBufferUAVIfNeeded(FRDGBufferUAV* UAV, bool bCompute
 		check(0);
 	}
 
-	UAV->CachedRHI.UAV = RHIUnorderedAccessView;
+	UAV->ResourceRHI = RHIUnorderedAccessView;
 	Buffer->PooledBuffer->UAVs.Add(UAV->Desc, RHIUnorderedAccessView);
 }
 
@@ -818,16 +818,16 @@ void FRDGBuilder::ExecutePass(const FRDGPass* Pass)
 {
 	QUICK_SCOPE_CYCLE_COUNTER(STAT_FRDGBuilder_ExecutePass);
 
-	FRHIRenderPassInfo RPInfo;
-	bool bHasRenderTargets = false;
-
-	PrepareResourcesForExecute(Pass, &RPInfo, &bHasRenderTargets);
-	
 	{
 		const bool bAllowAccess = true;
 		UpdateAccessGuardForPassResources(Pass, bAllowAccess);
 	}
 
+	FRHIRenderPassInfo RPInfo;
+	bool bHasRenderTargets = false;
+
+	PrepareResourcesForExecute(Pass, &RPInfo, &bHasRenderTargets);
+	
 	EventScopeStack.BeginExecutePass(Pass);
 	StatScopeStack.BeginExecutePass(Pass);
 
@@ -895,7 +895,7 @@ void FRDGBuilder::PrepareResourcesForExecute(const FRDGPass* Pass, struct FRHIRe
 				#endif
 
 				check(Texture->PooledRenderTarget);
-				check(Texture->CachedRHI.Resource);
+				check(Texture->ResourceRHI);
 				TransitionTexture(Texture, EResourceTransitionAccess::EReadable, bIsCompute);
 			}
 		}
@@ -915,9 +915,9 @@ void FRDGBuilder::PrepareResourcesForExecute(const FRDGPass* Pass, struct FRHIRe
 				#endif
 
 				// Might be the first time using this render graph SRV, so need to setup the cached rhi resource.
-				if (!SRV->CachedRHI.SRV)
+				if (!SRV->ResourceRHI)
 				{
-					SRV->CachedRHI.SRV = Texture->PooledRenderTarget->GetRenderTargetItem().MipSRVs[SRV->Desc.MipLevel];
+					SRV->ResourceRHI = Texture->PooledRenderTarget->GetRenderTargetItem().MipSRVs[SRV->Desc.MipLevel];
 				}
 
 				TransitionTexture(Texture, EResourceTransitionAccess::EReadable, bIsCompute);
@@ -929,10 +929,11 @@ void FRDGBuilder::PrepareResourcesForExecute(const FRDGPass* Pass, struct FRHIRe
 			if (FRDGTextureUAVRef UAV = Parameter.GetAsTextureUAV())
 			{
 				FRDGTextureRef Texture = UAV->Desc.Texture;
+				FUnorderedAccessViewRHIParamRef UAVRHI = UAV->GetRHI();
 
 				if (!bIsCompute)
 				{
-					OutRPInfo->UAVs[OutRPInfo->NumUAVs++] = UAV->CachedRHI.UAV;	// Bind UAVs in declaration order
+					OutRPInfo->UAVs[OutRPInfo->NumUAVs++] = UAVRHI;	// Bind UAVs in declaration order
 				}
 
 				#if RDG_ENABLE_DEBUG
@@ -942,7 +943,7 @@ void FRDGBuilder::PrepareResourcesForExecute(const FRDGPass* Pass, struct FRHIRe
 				#endif
 
 				AllocateRHITextureUAVIfNeeded(UAV, bIsCompute);
-				TransitionUAV(UAV->CachedRHI.UAV, Texture, Texture->Flags, EResourceTransitionAccess::EWritable, bIsCompute);
+				TransitionUAV(UAVRHI, Texture, Texture->Flags, EResourceTransitionAccess::EWritable, bIsCompute);
 			}
 		}
 		break;
@@ -995,10 +996,11 @@ void FRDGBuilder::PrepareResourcesForExecute(const FRDGPass* Pass, struct FRHIRe
 			if (FRDGBufferUAVRef UAV = Parameter.GetAsBufferUAV())
 			{
 				FRDGBufferRef Buffer = UAV->Desc.Buffer;
+				FUnorderedAccessViewRHIParamRef UAVRHI = UAV->GetRHI();
 
 				if (!bIsCompute)
 				{
-					OutRPInfo->UAVs[OutRPInfo->NumUAVs++] = UAV->CachedRHI.UAV;	// Bind UAVs in declaration order
+					OutRPInfo->UAVs[OutRPInfo->NumUAVs++] = UAVRHI;	// Bind UAVs in declaration order
 				}
 				
 				#if RDG_ENABLE_DEBUG
@@ -1008,7 +1010,7 @@ void FRDGBuilder::PrepareResourcesForExecute(const FRDGPass* Pass, struct FRHIRe
 				#endif
 
 				AllocateRHIBufferUAVIfNeeded(UAV, bIsCompute);
-				TransitionUAV(UAV->CachedRHI.UAV, Buffer, Buffer->Flags, EResourceTransitionAccess::EWritable, bIsCompute);
+				TransitionUAV(UAVRHI, Buffer, Buffer->Flags, EResourceTransitionAccess::EWritable, bIsCompute);
 			}
 		}
 		break;
@@ -1227,7 +1229,7 @@ void FRDGBuilder::ReleaseRHITextureIfUnreferenced(FRDGTexture* Texture)
 	if (Texture->ReferenceCount == 0)
 	{
 		Texture->PooledRenderTarget = nullptr;
-		Texture->CachedRHI.Resource = nullptr;
+		Texture->ResourceRHI = nullptr;
 		AllocatedTextures.FindChecked(Texture) = nullptr;
 	}
 }
@@ -1240,7 +1242,7 @@ void FRDGBuilder::ReleaseRHIBufferIfUnreferenced(FRDGBuffer* Buffer)
 	if (Buffer->ReferenceCount == 0)
 	{
 		Buffer->PooledBuffer = nullptr;
-		Buffer->CachedRHI.Resource = nullptr;
+		Buffer->ResourceRHI = nullptr;
 		AllocatedBuffers.FindChecked(Buffer) = nullptr;
 	}
 }
diff --git a/Engine/Source/Runtime/RenderCore/Public/RenderGraphBuilder.h b/Engine/Source/Runtime/RenderCore/Public/RenderGraphBuilder.h
--- a/Engine/Source/Runtime/RenderCore/Public/RenderGraphBuilder.h
+++ b/Engine/Source/Runtime/RenderCore/Public/RenderGraphBuilder.h
@@ -212,10 +212,6 @@ private:
 
 	void ProcessDeferredInternalResourceQueries();
 	void DestructPasses();
-
-	/** To allow greater flexibility in the user code, the RHI can dereferenced RDG resource when creating uniform buffer. */
-	// TODO(RDG): Make this a little more explicit in RHI code.
-	static_assert(STRUCT_OFFSET(FRDGResource, CachedRHI) == 0, "FRDGResource::CachedRHI requires to be at offset 0 so the RHI can dereferenced them.");
 };
 
 #include "RenderGraphBuilder.inl"
\ No newline at end of file
diff --git a/Engine/Source/Runtime/RenderCore/Public/RenderGraphBuilder.inl b/Engine/Source/Runtime/RenderCore/Public/RenderGraphBuilder.inl
--- a/Engine/Source/Runtime/RenderCore/Public/RenderGraphBuilder.inl
+++ b/Engine/Source/Runtime/RenderCore/Public/RenderGraphBuilder.inl
@@ -13,7 +13,7 @@ inline FRDGTextureRef FRDGBuilder::RegisterExternalTexture(const TRefCountPtr<IP
 
 	FRDGTexture* OutTexture = AllocateForRHILifeTime<FRDGTexture>(Name, ExternalPooledTexture->GetDesc(), ERDGResourceFlags::None);
 	OutTexture->PooledRenderTarget = ExternalPooledTexture;
-	OutTexture->CachedRHI.Texture = ExternalPooledTexture->GetRenderTargetItem().ShaderResourceTexture;
+	OutTexture->ResourceRHI = ExternalPooledTexture->GetRenderTargetItem().ShaderResourceTexture;
 	AllocatedTextures.Add(OutTexture, ExternalPooledTexture);
 
 #if RDG_ENABLE_DEBUG
diff --git a/Engine/Source/Runtime/RenderCore/Public/RenderGraphResources.h b/Engine/Source/Runtime/RenderCore/Public/RenderGraphResources.h
--- a/Engine/Source/Runtime/RenderCore/Public/RenderGraphResources.h
+++ b/Engine/Source/Runtime/RenderCore/Public/RenderGraphResources.h
@@ -43,54 +43,56 @@
 /** Generic graph resource. */
 class RENDERCORE_API FRDGResource
 {
+public:
+	FRDGResource(const FRDGResource&) = delete;
+
+	// Name of the resource for debugging purpose.
+	const TCHAR* const Name = nullptr;
+
+	//////////////////////////////////////////////////////////////////////////
+	//! The following methods may only be called during pass execution.
+
+	/** Marks this resource as actually used by a resource. This is to track what dependencies on pass was actually unnecessary. */
+	void MarkResourceAsUsed()
+	{
+		ValidateRHIAccess();
+
+		#if RDG_ENABLE_DEBUG
+		{
+			bIsActuallyUsedByPass = true;
+		}
+		#endif
+	}
+
+	FRHIResource* GetRHI() const
+	{
+		ValidateRHIAccess();
+		return ResourceRHI;
+	}
+
+	//////////////////////////////////////////////////////////////////////////
+
 protected:
-	/** Pointer on the RHI resource once allocated, that the RHI can dereferenced according to IsUniformBufferResourceIndirectionType(). */
-	union
-	{
-		FRHIResource* Resource;
-		FTextureRHIParamRef Texture;
-		FShaderResourceViewRHIParamRef SRV;
-		FUnorderedAccessViewRHIParamRef UAV;
-	} CachedRHI;
+	FRDGResource(const TCHAR* InName)
+		: Name(InName)
+	{}
 
 	/** Verify that the RHI resource can be accessed at a pass execution. */
 	void ValidateRHIAccess() const
 	{
-		#if RDG_ENABLE_DEBUG
+#if RDG_ENABLE_DEBUG
 		{
 			checkf(bAllowRHIAccess,
 				TEXT("Accessing the RHI resource of %s at this time is not allowed. If you hit this check in pass, ")
 				TEXT("that is due to this resource not being referenced in the parameters of your pass."),
 				Name);
 		}
-		#endif
+#endif
 	}
 
-	FRDGResource(const TCHAR* InName)
-		: Name(InName)
-	{
-		CachedRHI.Resource = nullptr;
-	}
-
-public:
-	FRDGResource(const FRDGResource&) = delete;
-
-	/** Marks this resource as actually used by a resource. This is to track what dependencies on pass was actually unnecessary. */
-	void MarkResourceAsUsed()
-	{
-		ValidateRHIAccess();
-
-		#if RDG_ENABLE_DEBUG
-		{
-			bIsActuallyUsedByPass = true;
-		}
-		#endif
-	}
-
-	// Name of the resource for debugging purpose.
-	const TCHAR* const Name = nullptr;
-
 private:
+	FRHIResource* ResourceRHI = nullptr;
+
 #if RDG_ENABLE_DEBUG
 	/** Boolean to track at runtime whether a resource is actually used by the lambda of a pass or not, to detect unnecessary resource dependencies on passes. */
 	bool bIsActuallyUsedByPass = false;
@@ -100,38 +102,22 @@
 #endif
 
 	friend class FRDGBuilder;
-	/** Friendship over parameter settings for bIsActuallyUsedByPass. It means only this parameter settings can be
-	 * used for render graph resource, to force the useless dependency tracking to just work.
-	 */
-	template<typename TRHICmdList, typename TShaderClass, typename TShaderRHI>
-	friend void SetShaderParameters(TRHICmdList&, const TShaderClass*, TShaderRHI*, const typename TShaderClass::FParameters&);
-
-	template<typename TRHICmdList, typename TShaderClass>
-	friend void SetShaderUAVs(TRHICmdList&, const TShaderClass*, FComputeShaderRHIParamRef, const typename TShaderClass::FParameters&);
-
-#if RHI_RAYTRACING
-	template<typename TShaderClass>
-	friend void SetShaderParameters(FRayTracingShaderBindingsWriter& RTBindingsWriter, const TShaderClass* Shader, const typename TShaderClass::FParameters& Parameters);
-#endif
-
-	template< typename T >
-	friend void AddPass_ClearUAV(FRDGBuilder&, FRDGEventName&&, FRDGTextureRef, FRDGTextureUAVRef, const T(&ClearValues)[4]);
-	friend void AddPass_ClearUAV(FRDGBuilder&, FRDGEventName&&, FRDGBufferUAVRef, uint32);
 };
 
 /** A render graph resource with an allocation lifetime tracked by the graph. */
 class RENDERCORE_API FRDGTrackedResource
 	: public FRDGResource
 {
+public:
+	/** Flags specific to this resource for render graph. */
+	const ERDGResourceFlags Flags;
+
 protected:
 	FRDGTrackedResource(const TCHAR* InName, const ERDGResourceFlags InFlags)
 		: FRDGResource(InName)
 		, Flags(InFlags)
 	{}
 
-	/** Flags specific to this resource for render graph. */
-	const ERDGResourceFlags Flags;
-
 private:
 	/** Number of references in passes and deferred queries. */
 	int32 ReferenceCount = 0;
@@ -196,11 +182,9 @@
 	}
 
 	/** Returns the allocated RHI texture. */
-	FTextureRHIParamRef GetRHITexture() const
+	FTextureRHIParamRef GetRHI() const
 	{
-		ValidateRHIAccess();
-		check(CachedRHI.Texture);
-		return CachedRHI.Texture;
+		return static_cast<FTextureRHIParamRef>(FRDGResource::GetRHI());
 	}
 
 	//////////////////////////////////////////////////////////////////////////
@@ -220,14 +204,40 @@
 	IPooledRenderTarget* PooledRenderTarget = nullptr;
 
 	friend class FRDGBuilder;
-
-	template<typename TRHICmdList, typename TShaderClass, typename TShaderRHI>
-	friend void SetShaderParameters(TRHICmdList&, const TShaderClass*, TShaderRHI*, const typename TShaderClass::FParameters&);
-
-#if RHI_RAYTRACING
-	template<typename TShaderClass>
-	friend void SetShaderParameters(struct FRayTracingShaderBindingsWriter& RTBindingsWriter, const TShaderClass* Shader, const typename TShaderClass::FParameters& Parameters);
-#endif
 };
 
+/** Render graph tracked SRV. */
+class FRDGShaderResourceView
+	: public FRDGResource
+{
+public:
+	/** Returns the allocated RHI SRV. */
+	FShaderResourceViewRHIParamRef GetRHI() const
+	{
+		return static_cast<FShaderResourceViewRHIParamRef>(FRDGResource::GetRHI());
+	}
+
+protected:
+	FRDGShaderResourceView(const TCHAR* Name)
+		: FRDGResource(Name)
+	{}
+};
+
+/** Render graph tracked UAV. */
+class FRDGUnorderedAccessView
+	: public FRDGResource
+{
+public:
+	/** Returns the allocated RHI SRV. */
+	FUnorderedAccessViewRHIParamRef GetRHI() const
+	{
+		return static_cast<FUnorderedAccessViewRHIParamRef>(FRDGResource::GetRHI());
+	}
+
+protected:
+	FRDGUnorderedAccessView(const TCHAR* Name)
+		: FRDGResource(Name)
+	{}
+};
+
 /** Descriptor for render graph tracked SRV. */
@@ -247,7 +257,7 @@
 
 /** Render graph tracked SRV. */
 class RENDERCORE_API FRDGTextureSRV final
-	: public FRDGResource
+	: public FRDGShaderResourceView
 {
 public:
 	/** Descriptor of the graph tracked SRV. */
@@ -257,7 +267,7 @@
 	FRDGTextureSRV(
 		const TCHAR* InName,
 		const FRDGTextureSRVDesc& InDesc)
-		: FRDGResource(InName)
+		: FRDGShaderResourceView(InName)
 		, Desc(InDesc)
 	{}
 
@@ -281,9 +291,9 @@
 	uint8 MipLevel = 0;
 };
 
-/** Render graph tracked UAV. */
+/** Render graph tracked texture UAV. */
 class RENDERCORE_API FRDGTextureUAV final
-	: public FRDGResource
+	: public FRDGUnorderedAccessView
 {
 public:
 	/** Descriptor of the graph tracked UAV. */
@@ -293,7 +303,7 @@
 	FRDGTextureUAV(
 		const TCHAR* InName,
 		const FRDGTextureUAVDesc& InDesc)
-		: FRDGResource(InName)
+		: FRDGUnorderedAccessView(InName)
 		, Desc(InDesc)
 	{}
 
@@ -563,9 +573,9 @@
 	friend class FRDGBuilder;
 };
 
-/** Render graph tracked SRV. */
+/** Render graph tracked buffer SRV. */
 class RENDERCORE_API FRDGBufferSRV final
-	: public FRDGResource
+	: public FRDGShaderResourceView
 {
 public:
 	/** Descriptor of the graph tracked SRV. */
@@ -573,16 +583,16 @@
 
 private:
 	FRDGBufferSRV(const TCHAR* InName, const FRDGBufferSRVDesc& InDesc)
-		: FRDGResource(InName)
+		: FRDGShaderResourceView(InName)
 		, Desc(InDesc)
 	{}
 
 	friend class FRDGBuilder;
 };
 
-/** Render graph tracked UAV. */
+/** Render graph tracked buffer UAV. */
 class RENDERCORE_API FRDGBufferUAV final
-	: public FRDGResource
+	: public FRDGUnorderedAccessView
 {
 public:
 	/** Descriptor of the graph tracked UAV. */
@@ -590,7 +600,7 @@
 
 private:
 	FRDGBufferUAV(const TCHAR* InName, const FRDGBufferUAVDesc& InDesc)
-		: FRDGResource(InName)
+		: FRDGUnorderedAccessView(InName)
 		, Desc(InDesc)
 	{}
 
diff --git a/Engine/Source/Runtime/RenderCore/Public/ShaderParameterStruct.h b/Engine/Source/Runtime/RenderCore/Public/ShaderParameterStruct.h
--- a/Engine/Source/Runtime/RenderCore/Public/ShaderParameterStruct.h
+++ b/Engine/Source/Runtime/RenderCore/Public/ShaderParameterStruct.h
@@ -116,7 +116,7 @@ inline void SetShaderUAVs(TRHICmdList& RHICmdList, const TShaderClass* Shader, F
 	// Graph UAVs
 	for (const FShaderParameterBindings::FResourceParameter& ParameterBinding : Bindings.GraphUAVs)
 	{
-		auto GraphUAV = *reinterpret_cast<FRDGTextureUAV* const*>(Base + ParameterBinding.ByteOffset);
+		auto GraphUAV = *reinterpret_cast<FRDGUnorderedAccessView* const*>(Base + ParameterBinding.ByteOffset);
 
 		if (DO_CHECK)
 		{
@@ -131,7 +131,7 @@ inline void SetShaderUAVs(TRHICmdList& RHICmdList, const TShaderClass* Shader, F
 		}
 
 		checkSlow(GraphUAV);
-		RHICmdList.SetUAVParameter(ShadeRHI, ParameterBinding.BaseIndex, GraphUAV->CachedRHI.UAV);
+		RHICmdList.SetUAVParameter(ShadeRHI, ParameterBinding.BaseIndex, GraphUAV->GetRHI());
 	}
 }
 
@@ -233,13 +233,13 @@ inline void SetShaderParameters(TRHICmdList& RHICmdList, const TShaderClass* Sha
 		}
 
 		checkSlow(GraphTexture);
-		RHICmdList.SetShaderTexture(ShadeRHI, ParameterBinding.BaseIndex, GraphTexture->GetRHITexture());
+		RHICmdList.SetShaderTexture(ShadeRHI, ParameterBinding.BaseIndex, GraphTexture->GetRHI());
 	}
 
 	// Graph SRVs
 	for (const FShaderParameterBindings::FResourceParameter& ParameterBinding : Bindings.GraphSRVs)
 	{
-		auto GraphSRV = *reinterpret_cast<FRDGTextureSRV* const*>(Base + ParameterBinding.ByteOffset);
+		auto GraphSRV = *reinterpret_cast<FRDGShaderResourceView* const*>(Base + ParameterBinding.ByteOffset);
 
 		if (DO_CHECK)
 		{
@@ -254,7 +254,7 @@ inline void SetShaderParameters(TRHICmdList& RHICmdList, const TShaderClass* Sha
 		}
 
 		checkSlow(GraphSRV);
-		RHICmdList.SetShaderResourceViewParameter(ShadeRHI, ParameterBinding.BaseIndex, GraphSRV->CachedRHI.SRV);
+		RHICmdList.SetShaderResourceViewParameter(ShadeRHI, ParameterBinding.BaseIndex, GraphSRV->GetRHI());
 	}
 
 	// UAVs for compute shaders
@@ -358,13 +358,13 @@ void SetShaderParameters(FRayTracingShaderBindingsWriter& RTBindingsWriter, cons
 		}
 
 		checkSlow(GraphTexture);
-		RTBindingsWriter.SetTexture(ParameterBinding.BaseIndex, GraphTexture->GetRHITexture());
+		RTBindingsWriter.SetTexture(ParameterBinding.BaseIndex, GraphTexture->GetRHI());
 	}
 
 	// Graph SRVs
 	for (const FShaderParameterBindings::FResourceParameter& ParameterBinding : Bindings.GraphSRVs)
 	{
-		auto GraphSRV = *reinterpret_cast<FRDGResource* const*>(Base + ParameterBinding.ByteOffset);
+		auto GraphSRV = *reinterpret_cast<FRDGShaderResourceView* const*>(Base + ParameterBinding.ByteOffset);
 
 		if (DO_CHECK)
 		{
@@ -379,13 +379,13 @@ void SetShaderParameters(FRayTracingShaderBindingsWriter& RTBindingsWriter, cons
 		}
 
 		checkSlow(GraphSRV);
-		RTBindingsWriter.SetSRV(ParameterBinding.BaseIndex, GraphSRV->CachedRHI.SRV);
+		RTBindingsWriter.SetSRV(ParameterBinding.BaseIndex, GraphSRV->GetRHI());
 	}
 
 	// Render graph UAVs
 	for (const FShaderParameterBindings::FResourceParameter& ParameterBinding : Bindings.GraphUAVs)
 	{
-		auto UAV = *reinterpret_cast<FRDGResource* const*>(Base + ParameterBinding.ByteOffset);
+		auto UAV = *reinterpret_cast<FRDGUnorderedAccessView* const*>(Base + ParameterBinding.ByteOffset);
 
 		if (DO_CHECK)
 		{
@@ -400,7 +400,7 @@ void SetShaderParameters(FRayTracingShaderBindingsWriter& RTBindingsWriter, cons
 		}
 
 		checkSlow(UAV);
-		RTBindingsWriter.SetUAV(ParameterBinding.BaseIndex, UAV->CachedRHI.UAV);
+		RTBindingsWriter.SetUAV(ParameterBinding.BaseIndex, UAV->GetRHI());
 	}
 
 	// Referenced uniform buffers
diff --git a/Engine/Source/Runtime/Renderer/Private/LightGridInjection.cpp b/Engine/Source/Runtime/Renderer/Private/LightGridInjection.cpp
--- a/Engine/Source/Runtime/Renderer/Private/LightGridInjection.cpp
+++ b/Engine/Source/Runtime/Renderer/Private/LightGridInjection.cpp
@@ -95,7 +95,7 @@ void AddPass_ClearUAV(
 		ERenderGraphPassFlags::Compute,
 		[&Parameters, BufferUAV, Value](FRHICommandList& RHICmdList)
 		{
-			ClearUAV( RHICmdList, BufferUAV->CachedRHI.UAV, BufferUAV->Desc.Buffer->Desc.GetTotalNumBytes(), Value );
+			ClearUAV( RHICmdList, BufferUAV->GetRHI(), BufferUAV->Desc.Buffer->Desc.GetTotalNumBytes(), Value );
 		} );
 }
 
diff --git a/Engine/Source/Runtime/Renderer/Private/PostProcess/PostProcessMotionBlur.cpp b/Engine/Source/Runtime/Renderer/Private/PostProcess/PostProcessMotionBlur.cpp
--- a/Engine/Source/Runtime/Renderer/Private/PostProcess/PostProcessMotionBlur.cpp
+++ b/Engine/Source/Runtime/Renderer/Private/PostProcess/PostProcessMotionBlur.cpp
@@ -849,7 +849,7 @@ FRDGTextureRef VisualizeMotionBlur(
 		{
 			const FViewInfo& LocalView = ScreenPassView.View;
 			const FSceneViewFamily& ViewFamily = *LocalView.Family;
-			FRenderTargetTemp TempRenderTarget(static_cast<FTexture2DRHIParamRef>(ColorTextureOutput->GetRHITexture()), ColorViewport.Extent);
+			FRenderTargetTemp TempRenderTarget(static_cast<FTexture2DRHIParamRef>(ColorTextureOutput->GetRHI()), ColorViewport.Extent);
 			FCanvas Canvas(&TempRenderTarget, nullptr, ViewFamily.CurrentRealTime, ViewFamily.CurrentWorldTime, ViewFamily.DeltaWorldTime, LocalView.GetFeatureLevel());
 
 			float X = 20;
diff --git a/Engine/Source/Runtime/Renderer/Private/PostProcess/PostProcessSubsurface.cpp b/Engine/Source/Runtime/Renderer/Private/PostProcess/PostProcessSubsurface.cpp
--- a/Engine/Source/Runtime/Renderer/Private/PostProcess/PostProcessSubsurface.cpp
+++ b/Engine/Source/Runtime/Renderer/Private/PostProcess/PostProcessSubsurface.cpp
@@ -715,7 +715,7 @@ void VisualizeSubsurface(
 		{
 			const FViewInfo& LocalView = ScreenPassView.View;
 			const FSceneViewFamily& ViewFamily = *LocalView.Family;
-			FRenderTargetTemp TempRenderTarget(static_cast<FTexture2DRHIParamRef>(SceneTextureOutput->GetRHITexture()), SceneTextureOutput->Desc.Extent);
+			FRenderTargetTemp TempRenderTarget(static_cast<FTexture2DRHIParamRef>(SceneTextureOutput->GetRHI()), SceneTextureOutput->Desc.Extent);
 			FCanvas Canvas(&TempRenderTarget, nullptr, ViewFamily.CurrentRealTime, ViewFamily.CurrentWorldTime, ViewFamily.DeltaWorldTime, LocalView.GetFeatureLevel());
 
 			float X = 30;
