DR - Fix RenderQuery RAII issue where a Slate Module allocated by the Module Manger was holding onto a SceneViewState. Also removed Logspam in Non debug build when the editor is stalled.

diff --git a/Engine/Source/Runtime/RenderCore/Private/ProfilingDebugging/RealtimeGPUProfiler.cpp b/Engine/Source/Runtime/RenderCore/Private/ProfilingDebugging/RealtimeGPUProfiler.cpp
--- a/Engine/Source/Runtime/RenderCore/Private/ProfilingDebugging/RealtimeGPUProfiler.cpp
+++ b/Engine/Source/Runtime/RenderCore/Private/ProfilingDebugging/RealtimeGPUProfiler.cpp
@@ -369,7 +369,9 @@ public:
 			}
 			if (!Event->HasValidResult())
 			{
+#if UE_BUILD_DEBUG
 				UE_LOG(LogRendererCore, Warning, TEXT("Query '%s' not ready."), *Event->GetName().ToString());
+#endif
 				// The frame isn't ready yet. Don't update stats - we'll try again next frame. 
 				bAnyEventFailed = true;
 				continue;
