Renames r.RDG.EmitWarnings to r.RDG.Debug to match -rdgdebug

diff --git a/Engine/Config/ConsoleVariables.ini b/Engine/Config/ConsoleVariables.ini
--- a/Engine/Config/ConsoleVariables.ini
+++ b/Engine/Config/ConsoleVariables.ini
@@ -51,8 +51,9 @@
 
 ; Uncomment to get render graph executing passes as they get created to easily debug crashes caused by pass wiring logic.
 ;r.RDG.ImmediateMode=1
-; Uncomment to get render graph to emit warnings for inneficiencies, and have event names formatting even if frame are not emiting draw events.
-;r.RDG.EmitWarnings=1
+; Uncomment to get render graph to emit warnings for inneficiencies that are normally too CPU costly, and have draw event names
+; formatting even if frame are not emiting draw events.
+;r.RDG.Debug=1
 
 ; r.XGEShaderCompile is now enabled by default in source. Uncomment to disable XGE shader compilation.
 ;r.XGEShaderCompile = 0
diff --git a/Engine/Source/Runtime/RenderCore/Private/RenderGraphBuilder.cpp b/Engine/Source/Runtime/RenderCore/Private/RenderGraphBuilder.cpp
--- a/Engine/Source/Runtime/RenderCore/Private/RenderGraphBuilder.cpp
+++ b/Engine/Source/Runtime/RenderCore/Private/RenderGraphBuilder.cpp
@@ -20,10 +20,10 @@ FAutoConsoleVariableRef CVarImmediateMode(
 	TEXT("Executes passes as they get created. Useful to have a callstack of the wiring code when crashing in the pass' lambda."),
 	ECVF_RenderThreadSafe);
 
-int32 GRDGEmitWarnings = 0;
-FAutoConsoleVariableRef CVarEmitWarnings(
-	TEXT("r.RDG.EmitWarnings"),
-	GRDGEmitWarnings,
+int32 GRDGDebug = 0;
+FAutoConsoleVariableRef CVarRDGDebug(
+	TEXT("r.RDG.Debug"),
+	GRDGDebug,
 	TEXT("Allow to output warnings for inefficiencies found during wiring and execution of the passes.\n")
 	TEXT(" 0: disabled;\n")
 	TEXT(" 1: emit warning once (default);\n")
@@ -33,7 +33,7 @@ FAutoConsoleVariableRef CVarEmitWarnings(
 #else
 
 const int32 GRDGImmediateMode = 0;
-const int32 GRDGEmitWarnings = 0;
+const int32 GRDGDebug = 0;
 
 #endif
 } //! namespace
@@ -41,7 +41,7 @@ const int32 GRDGEmitWarnings = 0;
 bool GetEmitRDGEvents()
 {
 #if RDG_EVENTS != RDG_EVENTS_NONE
-	return GetEmitDrawEvents() || GRDGEmitWarnings;
+	return GetEmitDrawEvents() || GRDGDebug;
 #else
 	return false;
 #endif
@@ -57,21 +57,21 @@ void InitRenderGraph()
 
 	if (FParse::Param(FCommandLine::Get(), TEXT("rdgdebug")))
 	{
-		GRDGEmitWarnings = 1;
+		GRDGDebug = 1;
 	}
 #endif
 }
 
 void EmitRDGWarning(const FString& WarningMessage)
 {
-	if (!GRDGEmitWarnings)
+	if (!GRDGDebug)
 	{
 		return;
 	}
 
 	static TSet<FString> GAlreadyEmittedWarnings;
 
-	if (GRDGEmitWarnings == kRDGEmitWarningsOnce)
+	if (GRDGDebug == kRDGEmitWarningsOnce)
 	{
 		if (!GAlreadyEmittedWarnings.Contains(WarningMessage))
 		{
@@ -277,7 +277,7 @@ void FRDGBuilder::ValidatePass(const FRDGPass* Pass) const
 			{
 				RenderTargetBindingSlots = &Parameter.GetAsRenderTargetBindingSlots();
 			}
-			else if (GRDGEmitWarnings)
+			else if (GRDGDebug)
 			{
 				EmitRDGWarningf(
 					TEXT("Pass %s have duplicated render target binding slots."),
@@ -1161,7 +1161,7 @@ void FRDGBuilder::UnmarkUsedResources(const FRDGPass* Pass)
 
 	const uint32 ParameterCount = ParameterStruct.GetParameterCount();
 
-	if (GRDGEmitWarnings)
+	if (GRDGDebug)
 	{
 		uint32 TrackedResourceCount = 0;
 		uint32 UsedResourceCount = 0;
@@ -1392,7 +1392,7 @@ void FRDGBuilder::DestructPasses()
 {
 	#if RDG_ENABLE_DEBUG
 	{
-		const bool bEmitWarnings = GRDGEmitWarnings;
+		const bool bEmitWarnings = GRDGDebug;
 
 		for (const FRDGTrackedResourceRef Resource : TrackedResources)
 		{
