Fixed scope event to be transparent black instead of opaque black.

diff --git a/Engine/Source/Runtime/RenderCore/Private/RenderGraphEvent.cpp b/Engine/Source/Runtime/RenderCore/Private/RenderGraphEvent.cpp
--- a/Engine/Source/Runtime/RenderCore/Private/RenderGraphEvent.cpp
+++ b/Engine/Source/Runtime/RenderCore/Private/RenderGraphEvent.cpp
@@ -42,7 +42,7 @@ FRDGEventScopeGuard::~FRDGEventScopeGuard()
 
 static void OnPushEvent(FRHICommandListImmediate& RHICmdList, const FRDGEventScope* Scope)
 {
-	RHICmdList.PushEvent(Scope->Name.GetTCHAR(), FColor(0, 0, 0));
+	RHICmdList.PushEvent(Scope->Name.GetTCHAR(), FColor(0));
 }
 
 static void OnPopEvent(FRHICommandListImmediate& RHICmdList)
