Fix clang compilation

diff --git a/Engine/Source/Runtime/RenderCore/Public/VirtualTexturing.h b/Engine/Source/Runtime/RenderCore/Public/VirtualTexturing.h
--- a/Engine/Source/Runtime/RenderCore/Public/VirtualTexturing.h
+++ b/Engine/Source/Runtime/RenderCore/Public/VirtualTexturing.h
@@ -168,7 +168,7 @@ struct FVTProduceTargetLayer
 	/** The texture to write to */
 	FRHITexture* TextureRHI = nullptr;
 
-	TRefCountPtr<struct IPooledRenderTarget> PooledRenderTarget = nullptr;
+	TRefCountPtr<struct IPooledRenderTarget> PooledRenderTarget;
 
 	/** Location within the texture to write */
 	FIntVector pPageLocation;
