Fixed static analysis warning in RDG.

diff --git a/Engine/Source/Runtime/RenderCore/Private/RenderGraphEvent.cpp b/Engine/Source/Runtime/RenderCore/Private/RenderGraphEvent.cpp
--- a/Engine/Source/Runtime/RenderCore/Private/RenderGraphEvent.cpp
+++ b/Engine/Source/Runtime/RenderCore/Private/RenderGraphEvent.cpp
@@ -52,7 +52,11 @@ static void OnPopEvent(FRHICommandListImmediate& RHICmdList)
 
 bool FRDGEventScopeStack::IsEnabled()
 {
-	return RDG_EVENTS && GetEmitRDGEvents();
+#if RDG_EVENTS
+	return GetEmitRDGEvents();
+#else
+	return false;
+#endif
 }
 
 FRDGEventScopeStack::FRDGEventScopeStack(FRHICommandListImmediate& InRHICmdList)
