A bit of cleanup in the RDG resource transition code.

diff --git a/Engine/Source/Runtime/RenderCore/Private/RenderGraphBuilder.cpp b/Engine/Source/Runtime/RenderCore/Private/RenderGraphBuilder.cpp
--- a/Engine/Source/Runtime/RenderCore/Private/RenderGraphBuilder.cpp
+++ b/Engine/Source/Runtime/RenderCore/Private/RenderGraphBuilder.cpp
@@ -610,7 +610,7 @@ void FRDGBuilder::WalkGraphDependencies()
 	}
 }
 
-void FRDGBuilder::AllocateRHITextureIfNeeded(FRDGTexture* Texture, bool bComputePass)
+void FRDGBuilder::AllocateRHITextureIfNeeded(FRDGTexture* Texture)
 {
 	check(Texture);
 
@@ -630,7 +630,7 @@ void FRDGBuilder::AllocateRHITextureIfNeeded(FRDGTexture* Texture, bool bCompute
 	check(Texture->ResourceRHI);
 }
 
-void FRDGBuilder::AllocateRHITextureUAVIfNeeded(FRDGTextureUAV* UAV, bool bComputePass)
+void FRDGBuilder::AllocateRHITextureUAVIfNeeded(FRDGTextureUAV* UAV)
 {
 	check(UAV);
 
@@ -639,12 +639,12 @@ void FRDGBuilder::AllocateRHITextureUAVIfNeeded(FRDGTextureUAV* UAV, bool bCompu
 		return;
 	}
 
-	AllocateRHITextureIfNeeded(UAV->Desc.Texture, bComputePass);
+	AllocateRHITextureIfNeeded(UAV->Desc.Texture);
 
 	UAV->ResourceRHI = UAV->Desc.Texture->PooledRenderTarget->GetRenderTargetItem().MipUAVs[UAV->Desc.MipLevel];
 }
 
-void FRDGBuilder::AllocateRHIBufferSRVIfNeeded(FRDGBufferSRV* SRV, bool bComputePass)
+void FRDGBuilder::AllocateRHIBufferSRVIfNeeded(FRDGBufferSRV* SRV)
 {
 	check(SRV);
 
@@ -689,7 +689,7 @@ void FRDGBuilder::AllocateRHIBufferSRVIfNeeded(FRDGBufferSRV* SRV, bool bCompute
 	Buffer->PooledBuffer->SRVs.Add(SRV->Desc, RHIShaderResourceView);
 }
 
-void FRDGBuilder::AllocateRHIBufferUAVIfNeeded(FRDGBufferUAV* UAV, bool bComputePass)
+void FRDGBuilder::AllocateRHIBufferUAVIfNeeded(FRDGBufferUAV* UAV)
 {
 	check(UAV);
 
@@ -739,78 +739,109 @@ void FRDGBuilder::AllocateRHIBufferUAVIfNeeded(FRDGBufferUAV* UAV, bool bCompute
 	Buffer->PooledBuffer->UAVs.Add(UAV->Desc, RHIUnorderedAccessView);
 }
 
-static EResourceTransitionPipeline CalcTransitionPipeline(bool bCurrentCompute, bool bTargetCompute)
+static EResourceTransitionPipeline GetResourceTransitionPipeline(ERDGPassPipeline PassPipelineBefore, ERDGPassPipeline PassPipelineAfter)
 {
-	// TODO(RDG) convert table to math
-	uint32 Bits;
-	Bits  = (uint32)bCurrentCompute;
-	Bits |= (uint32)bTargetCompute << 1;
-
-	EResourceTransitionPipeline Table[] = {
-		EResourceTransitionPipeline::EGfxToGfx,
-		EResourceTransitionPipeline::EComputeToGfx,
-		EResourceTransitionPipeline::EGfxToCompute,
-		EResourceTransitionPipeline::EComputeToCompute
+	static const EResourceTransitionPipeline Table[(uint32)ERDGPassPipeline::MAX][(uint32)ERDGPassPipeline::MAX] =
+	{
+		// Graphics
+		{
+			EResourceTransitionPipeline::EGfxToGfx,
+			EResourceTransitionPipeline::EGfxToCompute
+		},
+
+		// Compute
+		{
+			EResourceTransitionPipeline::EComputeToGfx,
+			EResourceTransitionPipeline::EComputeToCompute
+		}
 	};
-	
-	return static_cast< EResourceTransitionPipeline >( Table[ Bits ] );
+
+#if RDG_ENABLE_DEBUG
+	check(PassPipelineBefore < ERDGPassPipeline::MAX);
+	check(PassPipelineAfter < ERDGPassPipeline::MAX);
+#endif
+
+	return Table[(uint32)PassPipelineBefore][(uint32)PassPipelineAfter];
+};
+
+static EResourceTransitionAccess GetResourceTransitionAccess(ERDGPassAccess PassAccess)
+{
+	return PassAccess == ERDGPassAccess::Read ? EResourceTransitionAccess::EReadable : EResourceTransitionAccess::EWritable;
 }
 
-void FRDGBuilder::TransitionTexture(FRDGTexture* Texture, EResourceTransitionAccess TransitionAccess, bool bRequiredCompute) const
+void FRDGBuilder::TransitionTexture(FRDGTexture* Texture, ERDGPassAccess PassAccessAfter, ERDGPassPipeline PassPipelineAfter) const
 {
-	const bool bRequiredWritable = TransitionAccess != EResourceTransitionAccess::EReadable;
+	const ERDGPassAccess PassAccessBefore = Texture->PassAccess;
 
-	if (Texture->bWritable != bRequiredWritable || Texture->bCompute != bRequiredCompute)
+	if (PassAccessBefore != PassAccessAfter || Texture->PassPipeline != PassPipelineAfter)
 	{
-		FTextureRHIParamRef RHITexture = Texture->PooledRenderTarget->GetRenderTargetItem().ShaderResourceTexture;
+		FTextureRHIParamRef RHITexture = Texture->GetRHIUnchecked();
 
 		const bool bIsMultiFrameResource = (Texture->Flags & ERDGResourceFlags::MultiFrame) == ERDGResourceFlags::MultiFrame;
-		if (bIsMultiFrameResource && bRequiredWritable && !Texture->bWritable)
+
+		const bool bIsWriteAccessBegin = PassAccessBefore == ERDGPassAccess::Read && PassAccessAfter == ERDGPassAccess::Write;
+		const bool bIsWriteAccessEnd = PassAccessBefore == ERDGPassAccess::Write && PassAccessAfter == ERDGPassAccess::Read;
+
+		if (bIsMultiFrameResource && bIsWriteAccessBegin)
 		{
 			RHICmdList.BeginUpdateMultiFrameResource(RHITexture);
 		}
 
-		RHICmdList.TransitionResource( TransitionAccess, RHITexture );
+		const EResourceTransitionAccess TransitionAccess = GetResourceTransitionAccess(PassAccessAfter);
+
+		RHICmdList.TransitionResource(TransitionAccess, RHITexture);
 		
-		if (bIsMultiFrameResource && !bRequiredWritable && Texture->bWritable)
+		if (bIsMultiFrameResource && bIsWriteAccessEnd)
 		{
 			RHICmdList.EndUpdateMultiFrameResource(RHITexture);
 		}
 		
-		Texture->bWritable = bRequiredWritable;
-		Texture->bCompute = bRequiredCompute;
+		Texture->PassAccess = PassAccessAfter;
+		Texture->PassPipeline = PassPipelineAfter;
 	}
 }
 
-void FRDGBuilder::TransitionUAV(FUnorderedAccessViewRHIParamRef UAV, FRDGTrackedResource* UnderlyingResource, ERDGResourceFlags ResourceFlags, EResourceTransitionAccess TransitionAccess, bool bRequiredCompute ) const
+void FRDGBuilder::TransitionUAV(
+	FUnorderedAccessViewRHIParamRef UAV,
+	FRDGTrackedResource* UnderlyingResource,
+	ERDGPassAccess PassAccessAfter,
+	ERDGPassPipeline PassPipelineAfter) const
 {
-	const bool bRequiredWritable = TransitionAccess != EResourceTransitionAccess::EReadable;
+	const ERDGPassAccess PassAccessBefore = UnderlyingResource->PassAccess;
+	const ERDGPassPipeline PassPipelineBefore = UnderlyingResource->PassPipeline;
 
-	if (bRequiredWritable && UnderlyingResource->bWritable)
+	if (PassAccessBefore == ERDGPassAccess::Write && PassAccessAfter == ERDGPassAccess::Write)
 	{
 		// Force a RW barrier between UAV write.
 		// TODO(RDG): allow to have no barrier in the API when multiple pass write concurrently to same resource.
-		EResourceTransitionPipeline TransitionPipeline = CalcTransitionPipeline(UnderlyingResource->bCompute, bRequiredCompute);
+		const EResourceTransitionPipeline TransitionPipeline = GetResourceTransitionPipeline(PassPipelineBefore, PassPipelineAfter);
+
 		RHICmdList.TransitionResource(EResourceTransitionAccess::ERWBarrier, TransitionPipeline, UAV);
 	}
-	else if(UnderlyingResource->bWritable != bRequiredWritable || UnderlyingResource->bCompute != bRequiredCompute )
+	else if(PassAccessBefore != PassAccessAfter|| PassPipelineBefore != PassPipelineAfter)
 	{
-		const bool bIsMultiFrameResource = (ResourceFlags & ERDGResourceFlags::MultiFrame) == ERDGResourceFlags::MultiFrame;
-		if (bIsMultiFrameResource && bRequiredWritable && !UnderlyingResource->bWritable)
+		const bool bIsMultiFrameResource = (UnderlyingResource->Flags & ERDGResourceFlags::MultiFrame) == ERDGResourceFlags::MultiFrame;
+
+		const bool bIsWriteAccessBegin = PassAccessBefore == ERDGPassAccess::Read && PassAccessAfter == ERDGPassAccess::Write;
+		const bool bIsWriteAccessEnd = PassAccessBefore == ERDGPassAccess::Write && PassAccessAfter == ERDGPassAccess::Read;
+
+		if (bIsMultiFrameResource && bIsWriteAccessBegin)
 		{
 			RHICmdList.BeginUpdateMultiFrameResource(UAV);
 		}
 
-		EResourceTransitionPipeline TransitionPipeline = CalcTransitionPipeline(UnderlyingResource->bCompute, bRequiredCompute );
-		RHICmdList.TransitionResource( TransitionAccess, TransitionPipeline, UAV);
+		const EResourceTransitionPipeline TransitionPipeline = GetResourceTransitionPipeline(PassPipelineBefore, PassPipelineAfter);
+		const EResourceTransitionAccess TransitionAccess = GetResourceTransitionAccess(PassAccessAfter);
+
+		RHICmdList.TransitionResource(TransitionAccess, TransitionPipeline, UAV);
 		
-		if (bIsMultiFrameResource && !bRequiredWritable && UnderlyingResource->bWritable)
+		if (bIsMultiFrameResource && bIsWriteAccessEnd)
 		{
 			RHICmdList.EndUpdateMultiFrameResource(UAV);
 		}
 		
-		UnderlyingResource->bWritable = bRequiredWritable;
-		UnderlyingResource->bCompute = bRequiredCompute;
+		UnderlyingResource->PassAccess = PassAccessAfter;
+		UnderlyingResource->PassPipeline = PassPipelineAfter;
 	}
 }
 
@@ -869,7 +900,8 @@ void FRDGBuilder::PrepareResourcesForExecute(const FRDGPass* Pass, struct FRHIRe
 	OutRPInfo->NumUAVs = 0;
 	OutRPInfo->UAVIndex = 0;
 
-	const bool bIsCompute = Pass->IsCompute();
+	const ERDGPassPipeline PassPipeline = Pass->IsCompute() ? ERDGPassPipeline::Compute : ERDGPassPipeline::Graphics;
+
 	const bool bGeneratingMips = (Pass->GetFlags() & ERDGPassFlags::GenerateMips) == ERDGPassFlags::GenerateMips;
 
 	FRDGPassParameterStruct ParameterStruct = Pass->GetParameters();
@@ -896,7 +928,7 @@ void FRDGBuilder::PrepareResourcesForExecute(const FRDGPass* Pass, struct FRHIRe
 
 				check(Texture->PooledRenderTarget);
 				check(Texture->ResourceRHI);
-				TransitionTexture(Texture, EResourceTransitionAccess::EReadable, bIsCompute);
+				TransitionTexture(Texture, ERDGPassAccess::Read, PassPipeline);
 			}
 		}
 		break;
@@ -920,7 +952,7 @@ void FRDGBuilder::PrepareResourcesForExecute(const FRDGPass* Pass, struct FRHIRe
 					SRV->ResourceRHI = Texture->PooledRenderTarget->GetRenderTargetItem().MipSRVs[SRV->Desc.MipLevel];
 				}
 
-				TransitionTexture(Texture, EResourceTransitionAccess::EReadable, bIsCompute);
+				TransitionTexture(Texture, ERDGPassAccess::Read, PassPipeline);
 			}
 		}
 		break;
@@ -931,7 +963,7 @@ void FRDGBuilder::PrepareResourcesForExecute(const FRDGPass* Pass, struct FRHIRe
 				FRDGTextureRef Texture = UAV->Desc.Texture;
 				FUnorderedAccessViewRHIParamRef UAVRHI = UAV->GetRHI();
 
-				if (!bIsCompute)
+				if (PassPipeline == ERDGPassPipeline::Graphics)
 				{
 					OutRPInfo->UAVs[OutRPInfo->NumUAVs++] = UAVRHI;	// Bind UAVs in declaration order
 				}
@@ -942,8 +974,8 @@ void FRDGBuilder::PrepareResourcesForExecute(const FRDGPass* Pass, struct FRHIRe
 				}
 				#endif
 
-				AllocateRHITextureUAVIfNeeded(UAV, bIsCompute);
-				TransitionUAV(UAVRHI, Texture, Texture->Flags, EResourceTransitionAccess::EWritable, bIsCompute);
+				AllocateRHITextureUAVIfNeeded(UAV);
+				TransitionUAV(UAVRHI, Texture, ERDGPassAccess::Write, PassPipeline);
 			}
 		}
 		break;
@@ -963,7 +995,8 @@ void FRDGBuilder::PrepareResourcesForExecute(const FRDGPass* Pass, struct FRHIRe
 				check(Buffer->PooledBuffer);
 				check(Buffer->PooledBuffer->UAVs.Num() == 1);
 				FUnorderedAccessViewRHIParamRef UAV = Buffer->PooledBuffer->UAVs.CreateIterator().Value();
-				TransitionUAV(UAV, Buffer, Buffer->Flags, EResourceTransitionAccess::EReadable, bIsCompute);
+
+				TransitionUAV(UAV, Buffer, ERDGPassAccess::Write, PassPipeline);
 			}
 		}
 		break;
@@ -986,8 +1019,8 @@ void FRDGBuilder::PrepareResourcesForExecute(const FRDGPass* Pass, struct FRHIRe
 				check(Buffer->PooledBuffer->UAVs.Num() == 1);
 				FUnorderedAccessViewRHIParamRef UAV = Buffer->PooledBuffer->UAVs.CreateIterator().Value();
 
-				AllocateRHIBufferSRVIfNeeded(SRV, bIsCompute);
-				TransitionUAV(UAV, Buffer, Buffer->Flags, EResourceTransitionAccess::EReadable, bIsCompute);
+				AllocateRHIBufferSRVIfNeeded(SRV);
+				TransitionUAV(UAV, Buffer, ERDGPassAccess::Read, PassPipeline);
 			}
 		}
 		break;
@@ -998,7 +1031,7 @@ void FRDGBuilder::PrepareResourcesForExecute(const FRDGPass* Pass, struct FRHIRe
 				FRDGBufferRef Buffer = UAV->Desc.Buffer;
 				FUnorderedAccessViewRHIParamRef UAVRHI = UAV->GetRHI();
 
-				if (!bIsCompute)
+				if (PassPipeline == ERDGPassPipeline::Graphics)
 				{
 					OutRPInfo->UAVs[OutRPInfo->NumUAVs++] = UAVRHI;	// Bind UAVs in declaration order
 				}
@@ -1009,14 +1042,14 @@ void FRDGBuilder::PrepareResourcesForExecute(const FRDGPass* Pass, struct FRHIRe
 				}
 				#endif
 
-				AllocateRHIBufferUAVIfNeeded(UAV, bIsCompute);
-				TransitionUAV(UAVRHI, Buffer, Buffer->Flags, EResourceTransitionAccess::EWritable, bIsCompute);
+				AllocateRHIBufferUAVIfNeeded(UAV);
+				TransitionUAV(UAVRHI, Buffer, ERDGPassAccess::Write, PassPipeline);
 			}
 		}
 		break;
 		case UBMT_RENDER_TARGET_BINDING_SLOTS:
 		{
-			check(!bIsCompute);
+			check(PassPipeline == ERDGPassPipeline::Graphics);
 
 			const FRenderTargetBindingSlots& RenderTargetBindingSlots = Parameter.GetAsRenderTargetBindingSlots();
 			const auto& RenderTargets = RenderTargetBindingSlots.Output;
@@ -1033,7 +1066,7 @@ void FRDGBuilder::PrepareResourcesForExecute(const FRDGPass* Pass, struct FRHIRe
 
 				if (FRDGTextureRef Texture = RenderTarget.GetTexture())
 				{
-					AllocateRHITextureIfNeeded(Texture, bIsCompute);
+					AllocateRHITextureIfNeeded(Texture);
 
 					auto& OutRenderTarget = OutRPInfo->ColorRenderTargets[RenderTargetIndex];
 
@@ -1047,7 +1080,7 @@ void FRDGBuilder::PrepareResourcesForExecute(const FRDGPass* Pass, struct FRHIRe
 					if (!bGeneratingMips)
 					{
 						// Implicit assurance the RHI will do the correct transitions
-						TransitionTexture(Texture, EResourceTransitionAccess::EWritable, bIsCompute);
+						TransitionTexture(Texture, ERDGPassAccess::Write, PassPipeline);
 					}
 
 					SampleCount |= OutRenderTarget.RenderTarget->GetNumSamples();
@@ -1068,9 +1101,8 @@ void FRDGBuilder::PrepareResourcesForExecute(const FRDGPass* Pass, struct FRHIRe
 			OutRPInfo->UAVIndex = ValidRenderTargetCount;
 
 			if (FRDGTextureRef Texture = DepthStencil.GetTexture())
-
 			{
-				AllocateRHITextureIfNeeded(Texture, false);
+				AllocateRHITextureIfNeeded(Texture);
 
 				auto& OutDepthStencil = OutRPInfo->DepthStencilRenderTarget;
 
@@ -1081,7 +1113,7 @@ void FRDGBuilder::PrepareResourcesForExecute(const FRDGPass* Pass, struct FRHIRe
 					MakeRenderTargetActions(DepthStencil.GetStencilLoadAction(), DepthStencil.GetStencilStoreAction()));
 				OutDepthStencil.ExclusiveDepthStencil = DepthStencil.GetDepthStencilAccess();
 
-				TransitionTexture(Texture, EResourceTransitionAccess::EWritable, bIsCompute);
+				TransitionTexture(Texture, ERDGPassAccess::Write, PassPipeline);
 
 				SampleCount |= OutDepthStencil.DepthStencilTarget->GetNumSamples();
 				ValidDepthStencilCount++;
@@ -1348,7 +1380,7 @@ void FRDGBuilder::ProcessDeferredInternalResourceQueries()
 
 		if (Query.bTransitionToRead)
 		{
-			TransitionTexture(Query.Texture, EResourceTransitionAccess::EReadable, /* bRequiredCompute = */ false);
+			TransitionTexture(Query.Texture, ERDGPassAccess::Read, ERDGPassPipeline::Graphics);
 		}
 
 		*Query.OutTexturePtr = AllocatedTextures.FindChecked(Query.Texture);
diff --git a/Engine/Source/Runtime/RenderCore/Public/RenderGraph.h b/Engine/Source/Runtime/RenderCore/Public/RenderGraph.h
--- a/Engine/Source/Runtime/RenderCore/Public/RenderGraph.h
+++ b/Engine/Source/Runtime/RenderCore/Public/RenderGraph.h
@@ -24,7 +24,7 @@
  * increases the graphs knowledge of that resource's lifetime. This can cause an increase in memory use or prevent passes
  * from overlapping their execution. An example is ClearUnusedGraphResources() that can automatically null out resource
  * references unused by a shader for passes that have only on shader. Warnings will be emitted if a resource is not used
- * in a pass (thanks to FRDGResource::bIsActuallyUsedByPass)
+ * in a pass (thanks to FRDGResource::MarkResourceAsUsed())
  *
  * The lambda scope of a pass execution may happen any time after FRDGBuilder::AddPass(). For debugging purpose, it may
  * happen directly in the AddPass() with the immediate mode. When a bug is happening during pass execution, immediate mode
diff --git a/Engine/Source/Runtime/RenderCore/Public/RenderGraphBuilder.h b/Engine/Source/Runtime/RenderCore/Public/RenderGraphBuilder.h
--- a/Engine/Source/Runtime/RenderCore/Public/RenderGraphBuilder.h
+++ b/Engine/Source/Runtime/RenderCore/Public/RenderGraphBuilder.h
@@ -172,14 +172,14 @@ private:
 	FRDGStatScopeStack StatScopeStack;
 
 #if RDG_ENABLE_DEBUG
-	/** Whether the Execute() has already been called. */
-	bool bHasExecuted = false;
-
 	/** All recently allocated pass parameter structure, but not used by a AddPass() yet. */
 	TSet<const void*> AllocatedUnusedPassParameters;
 
 	/** List of tracked resources for validation prior to shutdown. */
 	TArray<FRDGTrackedResourceRef, SceneRenderingAllocator> TrackedResources;
+
+	/** Whether the Execute() has already been called. */
+	bool bHasExecuted = false;
 #endif
 
 	void ValidatePass(const FRDGPass* Pass) const;
@@ -192,13 +192,21 @@ private:
 
 	void AddPassInternal(FRDGPass* Pass);
 
-	void AllocateRHITextureIfNeeded(FRDGTexture* Texture, bool bComputePass);
-	void AllocateRHITextureUAVIfNeeded(FRDGTextureUAV* UAV, bool bComputePass);
-	void AllocateRHIBufferSRVIfNeeded(FRDGBufferSRV* SRV, bool bComputePass);
-	void AllocateRHIBufferUAVIfNeeded(FRDGBufferUAV* UAV, bool bComputePass);
-
-	void TransitionTexture(FRDGTexture* Texture, EResourceTransitionAccess TransitionAccess, bool bRequiredCompute) const;
-	void TransitionUAV(FUnorderedAccessViewRHIParamRef UAV, FRDGTrackedResource* UnderlyingResource, ERDGResourceFlags ResourceFlags, EResourceTransitionAccess TransitionAccess, bool bRequiredCompute ) const;
+	void AllocateRHITextureIfNeeded(FRDGTexture* Texture);
+	void AllocateRHITextureUAVIfNeeded(FRDGTextureUAV* UAV);
+	void AllocateRHIBufferSRVIfNeeded(FRDGBufferSRV* SRV);
+	void AllocateRHIBufferUAVIfNeeded(FRDGBufferUAV* UAV);
+
+	void TransitionTexture(
+		FRDGTexture* Texture,
+		ERDGPassAccess PassAccess,
+		ERDGPassPipeline PassPipeline) const;
+
+	void TransitionUAV(
+		FUnorderedAccessViewRHIParamRef UAV,
+		FRDGTrackedResource* UnderlyingResource,
+		ERDGPassAccess PassAccess,
+		ERDGPassPipeline PassPipeline) const;
 
 	void ExecutePass(const FRDGPass* Pass);
 	void PrepareResourcesForExecute(const FRDGPass* Pass, struct FRHIRenderPassInfo* OutRPInfo, bool* bOutHasRenderTargets);
diff --git a/Engine/Source/Runtime/RenderCore/Public/RenderGraphResources.h b/Engine/Source/Runtime/RenderCore/Public/RenderGraphResources.h
--- a/Engine/Source/Runtime/RenderCore/Public/RenderGraphResources.h
+++ b/Engine/Source/Runtime/RenderCore/Public/RenderGraphResources.h
@@ -29,6 +29,22 @@ using FRDGBufferRef = FRDGBuffer*;
 using FRDGBufferSRVRef = FRDGBufferSRV*;
 using FRDGBufferUAVRef = FRDGBufferUAV*;
 
+/** The hardware pipeline for pass scheduling. */
+enum class ERDGPassPipeline : uint8
+{
+	Graphics,
+	Compute,
+	MAX
+};
+
+/** Describes how a pass accesses a resource. */
+enum class ERDGPassAccess : uint8
+{
+	Read,
+	Write,
+	MAX
+};
+
 /** Render graph specific flags for resources. */
 enum class ERDGResourceFlags : uint8
 {
@@ -90,6 +106,11 @@ protected:
 #endif
 	}
 
+	FRHIResource* GetRHIUnchecked() const
+	{
+		return ResourceRHI;
+	}
+
 private:
 	FRHIResource* ResourceRHI = nullptr;
 
@@ -134,6 +155,7 @@ private:
 
 	void MarkAsExternal()
 	{
+		checkf(!FirstProducer, TEXT("Resource %s with producer pass marked as external."), Name);
 		bHasBeenProduced = true;
 	}
 
@@ -153,8 +175,8 @@ private:
 #endif
 
 	/** Used for tracking resource state during execution. */
-	bool bWritable = false;
-	bool bCompute = false;
+	ERDGPassAccess PassAccess = ERDGPassAccess::Read;
+	ERDGPassPipeline PassPipeline = ERDGPassPipeline::Graphics;
 
 	friend class FRDGBuilder;
 };
@@ -198,6 +220,12 @@ private:
 		, Desc(InDesc)
 	{}
 
+	/** Returns the allocated RHI texture without access checks. */
+	FTextureRHIParamRef GetRHIUnchecked() const
+	{
+		return static_cast<FTextureRHIParamRef>(FRDGResource::GetRHIUnchecked());
+	}
+
 	/** This is not a TRefCountPtr<> because FRDGTexture is allocated on the FMemStack
 	 *  FGraphBuilder::AllocatedTextures is actually keeping the reference.
 	 */
@@ -221,6 +249,12 @@ protected:
 	FRDGShaderResourceView(const TCHAR* Name)
 		: FRDGResource(Name)
 	{}
+
+	/** Returns the allocated RHI SRV without access checks. */
+	FShaderResourceViewRHIParamRef GetRHIUnchecked() const
+	{
+		return static_cast<FShaderResourceViewRHIParamRef>(FRDGResource::GetRHIUnchecked());
+	}
 };
 
 /** Render graph tracked UAV. */
@@ -228,7 +262,7 @@ class FRDGUnorderedAccessView
 	: public FRDGResource
 {
 public:
-	/** Returns the allocated RHI SRV. */
+	/** Returns the allocated RHI UAV. */
 	FUnorderedAccessViewRHIParamRef GetRHI() const
 	{
 		return static_cast<FUnorderedAccessViewRHIParamRef>(FRDGResource::GetRHI());
@@ -238,6 +272,12 @@ protected:
 	FRDGUnorderedAccessView(const TCHAR* Name)
 		: FRDGResource(Name)
 	{}
+
+	/** Returns the allocated RHI UAV without access checks. */
+	FUnorderedAccessViewRHIParamRef GetRHIUnchecked() const
+	{
+		return static_cast<FUnorderedAccessViewRHIParamRef>(FRDGResource::GetRHIUnchecked());
+	}
 };
 
 /** Descriptor for render graph tracked SRV. */
