Replace previous ad-hoc operator overloads with ENUM_CLASS_FLAGS() macro

diff --git a/Engine/Source/Runtime/RenderCore/Public/VirtualTexturing.h b/Engine/Source/Runtime/RenderCore/Public/VirtualTexturing.h
--- a/Engine/Source/Runtime/RenderCore/Public/VirtualTexturing.h
+++ b/Engine/Source/Runtime/RenderCore/Public/VirtualTexturing.h
@@ -148,8 +148,7 @@ enum class EVTProducePageFlags : uint8
 	None = 0u,
 	SkipPageBorders = (1u << 0),
 };
-inline EVTProducePageFlags operator|(EVTProducePageFlags Lhs, EVTProducePageFlags Rhs) { return (EVTProducePageFlags)((uint8)Lhs | (uint8)Rhs); }
-inline EVTProducePageFlags operator&(EVTProducePageFlags Lhs, EVTProducePageFlags Rhs) { return (EVTProducePageFlags)((uint8)Lhs & (uint8)Rhs); }
+ENUM_CLASS_FLAGS(EVTProducePageFlags);
 
 struct FVTRequestPageResult
 {
