Validates the debug name of FRDGResource before any other validation in CreateTexture() and CreateBuffer(), since other debugging message are friendlier with resource names

diff --git a/Engine/Source/Runtime/RenderCore/Public/RenderGraphBuilder.inl b/Engine/Source/Runtime/RenderCore/Public/RenderGraphBuilder.inl
--- a/Engine/Source/Runtime/RenderCore/Public/RenderGraphBuilder.inl
+++ b/Engine/Source/Runtime/RenderCore/Public/RenderGraphBuilder.inl
@@ -55,8 +55,8 @@ inline FRDGTextureRef FRDGBuilder::CreateTexture(
 {
 #if RDG_ENABLE_DEBUG
 	{
-		ensureMsgf(!bHasExecuted, TEXT("Render graph texture %s needs to be created before the builder execution."), DebugName);
 		checkf(DebugName, TEXT("Creating a render graph texture requires a valid debug name."));
+		checkf(!bHasExecuted, TEXT("Render graph texture %s needs to be created before the builder execution."), DebugName);
 		checkf(Desc.Format != PF_Unknown, TEXT("Illegal to create texture %s with an invalid pixel format."), DebugName);
 	}
 #endif
@@ -79,8 +79,8 @@ inline FRDGBufferRef FRDGBuilder::CreateBuffer(
 {
 #if RDG_ENABLE_DEBUG
 	{
-		ensureMsgf(!bHasExecuted, TEXT("Render graph buffer %s needs to be created before the builder execution."), Name);
 		checkf(Name, TEXT("Creating a render graph buffer requires a valid debug name."));
+		checkf(!bHasExecuted, TEXT("Render graph buffer %s needs to be created before the builder execution."), Name);
 	}
 #endif
 
