Makes any validation failure when creating a RDG resource fatal.

diff --git a/Engine/Source/Runtime/RenderCore/Public/RenderGraphBuilder.inl b/Engine/Source/Runtime/RenderCore/Public/RenderGraphBuilder.inl
--- a/Engine/Source/Runtime/RenderCore/Public/RenderGraphBuilder.inl
+++ b/Engine/Source/Runtime/RenderCore/Public/RenderGraphBuilder.inl
@@ -6,7 +6,7 @@ inline FRDGTextureRef FRDGBuilder::RegisterExternalTexture(const TRefCountPtr<IP
 {
 #if RDG_ENABLE_DEBUG
 	{
-		ensureMsgf(ExternalPooledTexture.IsValid(), TEXT("Attempted to register NULL external texture: %s"), Name);
+		checkf(ExternalPooledTexture.IsValid(), TEXT("Attempted to register NULL external texture: %s"), Name);
 		checkf(Name, TEXT("Externally allocated texture requires a debug name when registering them to render graph."));
 	}
 #endif
@@ -30,7 +30,7 @@ inline FRDGBufferRef FRDGBuilder::RegisterExternalBuffer(const TRefCountPtr<FPoo
 {
 #if RDG_ENABLE_DEBUG
 	{
-		ensureMsgf(ExternalPooledBuffer.IsValid(), TEXT("Attempted to register NULL external buffer: %s"), Name);
+		checkf(ExternalPooledBuffer.IsValid(), TEXT("Attempted to register NULL external buffer: %s"), Name);
 	}
 #endif
 
@@ -101,14 +101,9 @@ inline FRDGTextureSRVRef FRDGBuilder::CreateSRV(const FRDGTextureSRVDesc& Desc)
 
 #if RDG_ENABLE_DEBUG
 	{
-		if (!Texture)
-		{
-			ensureMsgf(false, TEXT("RenderGraph texture SRV created with a null texture."));
-			return nullptr;
-		}
-
-		ensureMsgf(!bHasExecuted, TEXT("Render graph SRV %s needs to be created before the builder execution."), Desc.Texture->Name);
-		ensureMsgf(Desc.Texture->Desc.TargetableFlags & TexCreate_ShaderResource, TEXT("Attempted to create SRV from texture %s which was not created with TexCreate_ShaderResource"), Desc.Texture->Name);
+		checkf(Texture, TEXT("RenderGraph texture SRV created with a null texture."));
+		checkf(!bHasExecuted, TEXT("Render graph SRV %s needs to be created before the builder execution."), Desc.Texture->Name);
+		checkf(Desc.Texture->Desc.TargetableFlags & TexCreate_ShaderResource, TEXT("Attempted to create SRV from texture %s which was not created with TexCreate_ShaderResource"), Desc.Texture->Name);
 	}
 #endif
 
@@ -121,13 +116,8 @@ inline FRDGBufferSRVRef FRDGBuilder::CreateSRV(const FRDGBufferSRVDesc& Desc)
 
 #if RDG_ENABLE_DEBUG
 	{
-		if (!Buffer)
-		{
-			ensureMsgf(false, TEXT("RenderGraph buffer SRV created with a null buffer."));
-			return nullptr;
-		}
-
-		ensureMsgf(!bHasExecuted, TEXT("Render graph SRV %s needs to be created before the builder execution."), Desc.Buffer->Name);
+		checkf(Buffer, TEXT("RenderGraph buffer SRV created with a null buffer."));
+		checkf(!bHasExecuted, TEXT("Render graph SRV %s needs to be created before the builder execution."), Desc.Buffer->Name);
 	}
 #endif
 
@@ -140,14 +130,9 @@ inline FRDGTextureUAVRef FRDGBuilder::CreateUAV(const FRDGTextureUAVDesc& Desc)
 
 #if RDG_ENABLE_DEBUG
 	{
-		if (!Texture)
-		{
-			checkf(false, TEXT("RenderGraph texture UAV created with a null texture."));
-			return nullptr;
-		}
-
-		ensureMsgf(!bHasExecuted, TEXT("Render graph UAV %s needs to be created before the builder execution."), Texture->Name);
-		ensureMsgf(Texture->Desc.TargetableFlags & TexCreate_UAV, TEXT("Attempted to create UAV from texture %s which was not created with TexCreate_UAV"), Texture->Name);
+		checkf(Texture, TEXT("RenderGraph texture UAV created with a null texture."));
+		checkf(!bHasExecuted, TEXT("Render graph UAV %s needs to be created before the builder execution."), Texture->Name);
+		checkf(Texture->Desc.TargetableFlags & TexCreate_UAV, TEXT("Attempted to create UAV from texture %s which was not created with TexCreate_UAV"), Texture->Name);
 	}
 #endif
 
@@ -160,13 +145,8 @@ inline FRDGBufferUAVRef FRDGBuilder::CreateUAV(const FRDGBufferUAVDesc& Desc)
 
 #if RDG_ENABLE_DEBUG
 	{
-		if (!Buffer)
-		{
-			checkf(false, TEXT("RenderGraph buffer UAV created with a null buffer."));
-			return nullptr;
-		}
-
-		ensureMsgf(!bHasExecuted, TEXT("Render graph UAV %s needs to be created before the builder execution."), Buffer->Name);
+		checkf(Buffer, TEXT("RenderGraph buffer UAV created with a null buffer."));
+		checkf(!bHasExecuted, TEXT("Render graph UAV %s needs to be created before the builder execution."), Buffer->Name);
 	}
 #endif
 
