UseGPUScene is now used to calculate DDC key.

diff --git a/Engine/Source/Runtime/RenderCore/Private/Shader.cpp b/Engine/Source/Runtime/RenderCore/Private/Shader.cpp
--- a/Engine/Source/Runtime/RenderCore/Private/Shader.cpp
+++ b/Engine/Source/Runtime/RenderCore/Private/Shader.cpp
@@ -2393,6 +2393,16 @@ void ShaderMapAppendKeyString(EShaderPlatform Platform, FString& KeyString)
 			KeyString += TEXT("_2bi");
 		}
 	}
+	{
+		if(UseGPUScene(Platform, GetMaxSupportedFeatureLevel(Platform)))
+		{
+			KeyString += TEXT("_gs1");
+		}
+		else
+		{
+			KeyString += TEXT("_gs0");
+		}
+	}
 
 	{
 		static const auto CVarVirtualTextureLightmaps = IConsoleManager::Get().FindTConsoleVariableDataInt(TEXT("r.VirtualTexturedLightmaps"));
