Works arround hack where MSAA texture have different shader resource than render target.

diff --git a/Engine/Source/Runtime/RenderCore/Private/RenderGraphBuilder.cpp b/Engine/Source/Runtime/RenderCore/Private/RenderGraphBuilder.cpp
--- a/Engine/Source/Runtime/RenderCore/Private/RenderGraphBuilder.cpp
+++ b/Engine/Source/Runtime/RenderCore/Private/RenderGraphBuilder.cpp
@@ -1070,12 +1070,31 @@ void FRDGBuilder::PrepareResourcesForExecute(const FRDGPass* Pass, struct FRHIRe
 
 					auto& OutRenderTarget = OutRPInfo->ColorRenderTargets[RenderTargetIndex];
 
+					// TODO(RDG): Clean up this legacy hack of the FPooledRenderTarget that can have TargetableTexture != ShaderResourceTexture
+					// for MSAA texture. Instead the two texture should be independent FRDGTexture explicitly handled by the user code.
+					FRHITexture* TargetableTexture = Texture->PooledRenderTarget->GetRenderTargetItem().TargetableTexture;
+					FRHITexture* ShaderResourceTexture = Texture->PooledRenderTarget->GetRenderTargetItem().ShaderResourceTexture;
+
+					// TODO(RDG): Looks like the store action on FRenderTargetBinding is not necessary, because: if want to bind a RT,
+					// that is most certainly to modify it as oposed to depth-stencil that might be for read only purposes. And if modify
+					// this resource, that certainly for being used by another pass. Otherwise this pass should be culled.
+					//
+					// TODO(RDG): The load store action could actually be optimised by render graph for tile hardware when there is multiple
+					// consecutive rasterizer passes that have RDG resource as render target, a bit like resource transitions.
+					ERenderTargetStoreAction StoreAction = RenderTarget.GetStoreAction();
+
+					// Automatically switch the store action to MSAA resolve when there is MSAA texture.
+					if (TargetableTexture != ShaderResourceTexture && Texture->Desc.NumSamples > 1 && StoreAction == ERenderTargetStoreAction::EStore)
+					{
+						StoreAction = ERenderTargetStoreAction::EMultisampleResolve;
+					}
+
 					// TODO(RDG): should force TargetableTexture == ShaderResourceTexture with MSAA, and instead have an explicit MSAA resolve pass.
-					OutRenderTarget.RenderTarget = Texture->PooledRenderTarget->GetRenderTargetItem().TargetableTexture;
-					OutRenderTarget.ResolveTarget = nullptr;
+					OutRenderTarget.RenderTarget = TargetableTexture;
+					OutRenderTarget.ResolveTarget = ShaderResourceTexture != TargetableTexture ? ShaderResourceTexture : false;
 					OutRenderTarget.ArraySlice = -1;
 					OutRenderTarget.MipIndex = RenderTarget.GetMipIndex();
-					OutRenderTarget.Action = MakeRenderTargetActions(RenderTarget.GetLoadAction(), RenderTarget.GetStoreAction());
+					OutRenderTarget.Action = MakeRenderTargetActions(RenderTarget.GetLoadAction(), StoreAction);
 
 					if (!bGeneratingMips)
 					{
@@ -1106,6 +1125,8 @@ void FRDGBuilder::PrepareResourcesForExecute(const FRDGPass* Pass, struct FRHIRe
 
 				auto& OutDepthStencil = OutRPInfo->DepthStencilRenderTarget;
 
+				// TODO(RDG): Addresses the TODO of the color scene render target.
+				ensureMsgf(Texture->Desc.NumSamples == 1, TEXT("MSAA dept-stencil render target not yet supported."));
 				OutDepthStencil.DepthStencilTarget = Texture->PooledRenderTarget->GetRenderTargetItem().TargetableTexture;
 				OutDepthStencil.ResolveTarget = nullptr;
 				OutDepthStencil.Action = MakeDepthStencilTargetActions(
