Adds validation when attempting TexCreate_UAV with a MSAA texture.

diff --git a/Engine/Source/Runtime/RenderCore/Public/RenderGraphBuilder.inl b/Engine/Source/Runtime/RenderCore/Public/RenderGraphBuilder.inl
--- a/Engine/Source/Runtime/RenderCore/Public/RenderGraphBuilder.inl
+++ b/Engine/Source/Runtime/RenderCore/Public/RenderGraphBuilder.inl
@@ -58,6 +58,13 @@ inline FRDGTextureRef FRDGBuilder::CreateTexture(
 		checkf(DebugName, TEXT("Creating a render graph texture requires a valid debug name."));
 		checkf(!bHasExecuted, TEXT("Render graph texture %s needs to be created before the builder execution."), DebugName);
 		checkf(Desc.Format != PF_Unknown, TEXT("Illegal to create texture %s with an invalid pixel format."), DebugName);
+
+		const bool bCanHaveUAV = Desc.TargetableFlags & TexCreate_UAV;
+		const bool bIsMSAA = Desc.NumSamples > 1;
+
+		// D3D11 doesn't allow creating a UAV on MSAA texture.
+		const bool bIsUAVForMSAATexture = bIsMSAA && bCanHaveUAV;
+		checkf(!bIsUAVForMSAATexture, TEXT("TexCreate_UAV is not allowed on MSAA texture %s."), DebugName);
 	}
 #endif
 
