Fixes clang build.

diff --git a/Engine/Source/Runtime/RenderCore/Private/RenderGraphBuilder.cpp b/Engine/Source/Runtime/RenderCore/Private/RenderGraphBuilder.cpp
--- a/Engine/Source/Runtime/RenderCore/Private/RenderGraphBuilder.cpp
+++ b/Engine/Source/Runtime/RenderCore/Private/RenderGraphBuilder.cpp
@@ -1091,7 +1091,7 @@ void FRDGBuilder::PrepareResourcesForExecute(const FRDGPass* Pass, struct FRHIRe
 
 					// TODO(RDG): should force TargetableTexture == ShaderResourceTexture with MSAA, and instead have an explicit MSAA resolve pass.
 					OutRenderTarget.RenderTarget = TargetableTexture;
-					OutRenderTarget.ResolveTarget = ShaderResourceTexture != TargetableTexture ? ShaderResourceTexture : false;
+					OutRenderTarget.ResolveTarget = ShaderResourceTexture != TargetableTexture ? ShaderResourceTexture : nullptr;
 					OutRenderTarget.ArraySlice = -1;
 					OutRenderTarget.MipIndex = RenderTarget.GetMipIndex();
 					OutRenderTarget.Action = MakeRenderTargetActions(RenderTarget.GetLoadAction(), StoreAction);
