Fixed uniform buffer generation (see UE-69171); In the braced-initializer, all sub expressions are now separated by a comma at the end instead of the beginning (e.g. { x, y, } instead of { ,x ,y }).

#rnx

diff --git a/Engine/Shaders/Public/ShaderVersion.ush b/Engine/Shaders/Public/ShaderVersion.ush
--- a/Engine/Shaders/Public/ShaderVersion.ush
+++ b/Engine/Shaders/Public/ShaderVersion.ush
@@ -2,4 +2,4 @@
 // Each time the console command is executed it generates a new GUID. As this file is included
 // in Platform.ush (which should be included in any shader) it allows to invalidate the shader DDC.
 // 
-// GUID = 1DEC0C1545F80AEC49E44F915EE0DEDB
\ No newline at end of file
+// GUID = 3A7029DE8C044F0A97C52C4CD3F6182E
\ No newline at end of file
diff --git a/Engine/Source/Runtime/Engine/Public/GeometryCacheVertexFactory.h b/Engine/Source/Runtime/Engine/Public/GeometryCacheVertexFactory.h
--- a/Engine/Source/Runtime/Engine/Public/GeometryCacheVertexFactory.h
+++ b/Engine/Source/Runtime/Engine/Public/GeometryCacheVertexFactory.h
@@ -25,8 +25,6 @@ END_GLOBAL_SHADER_PARAMETER_STRUCT()
 typedef TUniformBufferRef<FGeometryCacheVertexFactoryUniformBufferParameters> FGeometryCacheVertexFactoryUniformBufferParametersRef;
 
 BEGIN_GLOBAL_SHADER_PARAMETER_STRUCT(FGeometryCacheManualVertexFetchUniformBufferParameters, ENGINE_API)
-	// FIXME: UE-69107 error X3000: syntax error: unexpected token ','
-	SHADER_PARAMETER(float, Dummy)
 	SHADER_PARAMETER_SRV(Buffer<float>, Position)
 	SHADER_PARAMETER_SRV(Buffer<float>, MotionBlurData)
 	SHADER_PARAMETER_SRV(Buffer<half4>, TangentX)
diff --git a/Engine/Source/Runtime/RenderCore/Private/ShaderParameters.cpp b/Engine/Source/Runtime/RenderCore/Private/ShaderParameters.cpp
--- a/Engine/Source/Runtime/RenderCore/Private/ShaderParameters.cpp
+++ b/Engine/Source/Runtime/RenderCore/Private/ShaderParameters.cpp
@@ -154,8 +154,6 @@ static void CreateHLSLUniformBufferStructMembersDeclaration(
 {
 	const TArray<FShaderParametersMetadata::FMember>& StructMembers = UniformBufferStruct.GetMembers();
 	
-	int32 OpeningBraceLocPlusOne = Decl.Initializer.Len();
-
 	FString PreviousBaseTypeName = TEXT("float");
 	for (int32 MemberIndex = 0; MemberIndex < StructMembers.Num(); ++MemberIndex)
 	{
@@ -171,14 +169,13 @@ static void CreateHLSLUniformBufferStructMembersDeclaration(
 		{
 			checkf(Member.GetNumElements() == 0, TEXT("SHADER_PARAMETER_STRUCT_ARRAY() is not supported in uniform buffer yet."));
 			Decl.StructMembers += TEXT("struct {\r\n");
-			Decl.Initializer += TEXT(",{");
+			Decl.Initializer += TEXT("{");
 			CreateHLSLUniformBufferStructMembersDeclaration(*Member.GetStructMetadata(), FString::Printf(TEXT("%s%s_"), *NamePrefix, Member.GetName()), StructOffset + Member.GetOffset(), Decl, HLSLBaseOffset);
-			Decl.Initializer += TEXT("}");
+			Decl.Initializer += TEXT("},");
 			Decl.StructMembers += FString::Printf(TEXT("} %s%s;\r\n"),Member.GetName(),*ArrayDim);
 		}
 		else if (Member.GetBaseType() == UBMT_INCLUDED_STRUCT)
 		{
-			Decl.Initializer += TEXT(",");
 			CreateHLSLUniformBufferStructMembersDeclaration(*Member.GetStructMetadata(), NamePrefix, StructOffset + Member.GetOffset(), Decl, HLSLBaseOffset);
 		}
 		else if (IsShaderParameterTypeForUniformBufferLayout(Member.GetBaseType()))
@@ -255,7 +252,7 @@ static void CreateHLSLUniformBufferStructMembersDeclaration(
 			FString ParameterName = FString::Printf(TEXT("%s%s"),*NamePrefix,Member.GetName());
 			Decl.ConstantBufferMembers += FString::Printf(TEXT("\t%s%s %s%s;\r\n"),*BaseTypeName,*TypeDim,*ParameterName,*ArrayDim);
 			Decl.StructMembers += FString::Printf(TEXT("\t%s%s %s%s;\r\n"),*BaseTypeName,*TypeDim,Member.GetName(),*ArrayDim);
-			Decl.Initializer += FString::Printf(TEXT(",%s"),*ParameterName);
+			Decl.Initializer += FString::Printf(TEXT("%s,"),*ParameterName);
 		}
 	}
 
@@ -273,7 +270,7 @@ static void CreateHLSLUniformBufferStructMembersDeclaration(
 				FString ParameterName = FString::Printf(TEXT("%s%s"),*NamePrefix,Member.GetName());
 				Decl.ResourceMembers += FString::Printf(TEXT("PLATFORM_SUPPORTS_SRV_UB_MACRO( %s %s; ) \r\n"), Member.GetShaderType(), *ParameterName);
 				Decl.StructMembers += FString::Printf(TEXT("\tPLATFORM_SUPPORTS_SRV_UB_MACRO( %s %s; ) \r\n"), Member.GetShaderType(), Member.GetName());
-				Decl.Initializer += FString::Printf(TEXT(" PLATFORM_SUPPORTS_SRV_UB_MACRO( ,%s ) "), *ParameterName);
+				Decl.Initializer += FString::Printf(TEXT(" PLATFORM_SUPPORTS_SRV_UB_MACRO( %s, ) "), *ParameterName);
 			}
 			else
 			{
@@ -281,15 +278,10 @@ static void CreateHLSLUniformBufferStructMembersDeclaration(
 				FString ParameterName = FString::Printf(TEXT("%s%s"),*NamePrefix,Member.GetName());
 				Decl.ResourceMembers += FString::Printf(TEXT("%s %s;\r\n"), Member.GetShaderType(), *ParameterName);
 				Decl.StructMembers += FString::Printf(TEXT("\t%s %s;\r\n"), Member.GetShaderType(), Member.GetName());
-				Decl.Initializer += FString::Printf(TEXT(",%s"), *ParameterName);
+				Decl.Initializer += FString::Printf(TEXT("%s,"), *ParameterName);
 			}
 		}
 	}
-
-	if (Decl.Initializer[OpeningBraceLocPlusOne] == TEXT(','))
-	{
-		Decl.Initializer[OpeningBraceLocPlusOne] = TEXT(' ');
-	}
 }
 
 /** Creates a HLSL declaration of a uniform buffer with the given structure. */
