Adds static_assert() on the size of the lambda in FRDGBuilder::AddPass()

diff --git a/Engine/Source/Runtime/RenderCore/Public/RenderGraphBuilder.inl b/Engine/Source/Runtime/RenderCore/Public/RenderGraphBuilder.inl
--- a/Engine/Source/Runtime/RenderCore/Public/RenderGraphBuilder.inl
+++ b/Engine/Source/Runtime/RenderCore/Public/RenderGraphBuilder.inl
@@ -181,6 +181,11 @@ void FRDGBuilder::AddPass(
 	ERDGPassFlags Flags,
 	ExecuteLambdaType&& ExecuteLambda)
 {
+	// Verify that the amount of stuff captured by a pass's lambda is reasonable. Changing this maximum as desired without review is not OK,
+	// the all point being to catch things that should be captured by referenced instead of captured by copy.
+	constexpr int32 kMaximumLambdaCaptureSize = 1024;
+	static_assert(sizeof(ExecuteLambdaType) <= kMaximumLambdaCaptureSize, "The amount of data of captured for the pass looks abnormally high.");
+
 	auto NewPass = new(MemStack) TRDGLambdaPass<ParameterStructType, ExecuteLambdaType>(
 		MoveTemp(Name),
 		FRDGPassParameterStruct(ParameterStruct),
