Catches when using wrong shader binding mechanism for a shader frequency.

diff --git a/Engine/Source/Runtime/RenderCore/Private/ShaderParameterStruct.cpp b/Engine/Source/Runtime/RenderCore/Private/ShaderParameterStruct.cpp
--- a/Engine/Source/Runtime/RenderCore/Private/ShaderParameterStruct.cpp
+++ b/Engine/Source/Runtime/RenderCore/Private/ShaderParameterStruct.cpp
@@ -225,6 +225,20 @@ void FShaderParameterBindings::BindForLegacyShaderParameters(const FShader* Shad
 {
 	checkf(StructMetaData.GetSize() < (1 << (sizeof(uint16) * 8)), TEXT("Shader parameter structure can only have a size < 65536 bytes."));
 	check(this == &Shader->Bindings);
+	
+	switch (Shader->GetType()->GetFrequency())
+	{
+	case SF_Vertex:
+	case SF_Hull:
+	case SF_Domain:
+	case SF_Pixel:
+	case SF_Geometry:
+	case SF_Compute:
+		break;
+	default:
+		checkf(0, TEXT("Invalid shader frequency for this shader binding technic."));
+		break;
+	}
 
 	FShaderParameterStructBindingContext BindingContext;
 	BindingContext.Shader = Shader;
@@ -265,6 +279,18 @@ void FShaderParameterBindings::BindForRootShaderParameters(const FShader* Shader
 	const FShaderParametersMetadata& StructMetaData = *Shader->GetType()->GetRootParametersMetadata();
 	checkf(StructMetaData.GetSize() < (1 << (sizeof(uint16) * 8)), TEXT("Shader parameter structure can only have a size < 65536 bytes."));
 
+	switch (Shader->GetType()->GetFrequency())
+	{
+	case SF_RayGen:
+	case SF_RayMiss:
+	case SF_RayHitGroup:
+	case SF_RayCallable:
+		break;
+	default:
+		checkf(0, TEXT("Invalid shader frequency for this shader binding technic."));
+		break;
+	}
+
 	FShaderParameterStructBindingContext BindingContext;
 	BindingContext.Shader = Shader;
 	BindingContext.Bindings = this;
