Disable CL 6673299 temporary until ray tracing pixel composition shaders that call BindForRootShaderParameters are refactored.

diff --git a/Engine/Source/Runtime/RenderCore/Private/ShaderParameterStruct.cpp b/Engine/Source/Runtime/RenderCore/Private/ShaderParameterStruct.cpp
--- a/Engine/Source/Runtime/RenderCore/Private/ShaderParameterStruct.cpp
+++ b/Engine/Source/Runtime/RenderCore/Private/ShaderParameterStruct.cpp
@@ -287,7 +287,7 @@ void FShaderParameterBindings::BindForRootShaderParameters(const FShader* Shader
 	case SF_RayCallable:
 		break;
 	default:
-		checkf(0, TEXT("Invalid shader frequency for this shader binding technic."));
+		//checkf(0, TEXT("Invalid shader frequency for this shader binding technic."));
 		break;
 	}
 
