DR - Deprecate CreateGeometryShaderWithStreamOutput

diff --git a/Engine/Source/Runtime/RenderCore/Private/Shader.cpp b/Engine/Source/Runtime/RenderCore/Private/Shader.cpp
--- a/Engine/Source/Runtime/RenderCore/Private/Shader.cpp
+++ b/Engine/Source/Runtime/RenderCore/Private/Shader.cpp
@@ -912,7 +912,9 @@ void FShaderResource::InitRHI()
 			checkf(ElementList.Num(), TEXT("Shader type %s was given GetStreamOutElements implementation that had no elements!"), SpecificType->GetName());
 
 			//@todo - not using the cache
+			PRAGMA_DISABLE_DEPRECATION_WARNINGS
 			Shader = FShaderCodeLibrary::CreateGeometryShaderWithStreamOutput((EShaderPlatform)Target.Platform, OutputHash, UncompressedCode, ElementList, StreamStrides.Num(), StreamStrides.GetData(), RasterizedStream);
+			PRAGMA_ENABLE_DEPRECATION_WARNINGS
 		}
 		else
 		{
