DR - Deprecate CreateGeometryShaderWithStreamOutput usages
#rnx

diff --git a/Engine/Source/Runtime/RHI/Public/RHICommandList.h b/Engine/Source/Runtime/RHI/Public/RHICommandList.h
--- a/Engine/Source/Runtime/RHI/Public/RHICommandList.h
+++ b/Engine/Source/Runtime/RHI/Public/RHICommandList.h
@@ -3609,12 +3609,14 @@ public:
 		return GDynamicRHI->CreateGeometryShader_RenderThread(*this, Library, Hash);
 	}
 	
+	UE_DEPRECATED(4.23, "Geometry Stream out is deprecated.")
 	FORCEINLINE FGeometryShaderRHIRef CreateGeometryShaderWithStreamOutput(const TArray<uint8>& Code, const FStreamOutElementList& ElementList, uint32 NumStrides, const uint32* Strides, int32 RasterizedStream)
 	{
 		LLM_SCOPE(ELLMTag::Shaders);
 		return GDynamicRHI->CreateGeometryShaderWithStreamOutput_RenderThread(*this, Code, ElementList, NumStrides, Strides, RasterizedStream);
 	}
 	
+	UE_DEPRECATED(4.23, "Geometry Stream out is deprecated.")
 	FORCEINLINE FGeometryShaderRHIRef CreateGeometryShaderWithStreamOutput(const FStreamOutElementList& ElementList, uint32 NumStrides, const uint32* Strides, int32 RasterizedStream, FRHIShaderLibrary* Library, FSHAHash Hash)
 	{
 		LLM_SCOPE(ELLMTag::Shaders);
@@ -4548,14 +4550,20 @@ FORCEINLINE FGeometryShaderRHIRef RHICreateGeometryShader(FRHIShaderLibrary* Lib
 	return FRHICommandListExecutor::GetImmediateCommandList().CreateGeometryShader(Library, Hash);
 }
 
+UE_DEPRECATED(4.23, "Geometry Stream out is deprecated.")
 FORCEINLINE FGeometryShaderRHIRef RHICreateGeometryShaderWithStreamOutput(const TArray<uint8>& Code, const FStreamOutElementList& ElementList, uint32 NumStrides, const uint32* Strides, int32 RasterizedStream)
 {
+	PRAGMA_DISABLE_DEPRECATION_WARNINGS
 	return FRHICommandListExecutor::GetImmediateCommandList().CreateGeometryShaderWithStreamOutput(Code, ElementList, NumStrides, Strides, RasterizedStream);
+	PRAGMA_ENABLE_DEPRECATION_WARNINGS
 }
 
+UE_DEPRECATED(4.23, "Geometry Stream out is deprecated.")
 FORCEINLINE FGeometryShaderRHIRef RHICreateGeometryShaderWithStreamOutput(const FStreamOutElementList& ElementList, uint32 NumStrides, const uint32* Strides, int32 RasterizedStream, FRHIShaderLibrary* Library, FSHAHash Hash)
 {
+	PRAGMA_DISABLE_DEPRECATION_WARNINGS
 	return FRHICommandListExecutor::GetImmediateCommandList().CreateGeometryShaderWithStreamOutput(ElementList, NumStrides, Strides, RasterizedStream, Library, Hash);
+	PRAGMA_ENABLE_DEPRECATION_WARNINGS
 }
 
 FORCEINLINE FComputeShaderRHIRef RHICreateComputeShader(const TArray<uint8>& Code)
diff --git a/Engine/Source/Runtime/RenderCore/Private/ShaderCodeLibrary.cpp b/Engine/Source/Runtime/RenderCore/Private/ShaderCodeLibrary.cpp
--- a/Engine/Source/Runtime/RenderCore/Private/ShaderCodeLibrary.cpp
+++ b/Engine/Source/Runtime/RenderCore/Private/ShaderCodeLibrary.cpp
@@ -718,7 +718,9 @@ public:
 		{
 			TArray<uint8> UCode;
 			TArray<uint8>& UncompressedCode = FShaderLibraryHelperUncompressCode(Platform, Size, *Code, UCode);
+			PRAGMA_DISABLE_DEPRECATION_WARNINGS
 			Shader = RHICreateGeometryShaderWithStreamOutput(UncompressedCode, ElementList, NumStrides, Strides, RasterizedStream);
+			PRAGMA_ENABLE_DEPRECATION_WARNINGS
 			CheckShaderCreation(Shader.GetReference(), Hash);
 			if (bWasSync)
 			{
@@ -1758,9 +1760,13 @@ public:
 		if (ShaderCodeArchive)
 		{
 			if (bNativeFormat || GRHILazyShaderCodeLoading)
+			{
 				Result = RHICreateVertexShader(ShaderCodeArchive, Hash);
+			}
 			else
+			{
 				Result = ((FShaderCodeArchive*)ShaderCodeArchive)->CreateVertexShader(Hash);
+			}
 		}
 		return Result;
 	}
@@ -1774,9 +1780,13 @@ public:
 		if (ShaderCodeArchive)
 		{
 			if (bNativeFormat || GRHILazyShaderCodeLoading)
+			{
 				Result = RHICreatePixelShader(ShaderCodeArchive, Hash);
+			}
 			else
+			{
 				Result = ((FShaderCodeArchive*)ShaderCodeArchive)->CreatePixelShader(Hash);
+			}
 		}
 		return Result;
 	}
@@ -1790,9 +1800,13 @@ public:
 		if (ShaderCodeArchive)
 		{
 			if (bNativeFormat || GRHILazyShaderCodeLoading)
+			{
 				Result = RHICreateGeometryShader(ShaderCodeArchive, Hash);
+			}
 			else
+			{
 				Result = ((FShaderCodeArchive*)ShaderCodeArchive)->CreateGeometryShader(Hash);
+			}
 		}
 		return Result;
 	}
@@ -1805,10 +1819,16 @@ public:
 		FRHIShaderLibrary* ShaderCodeArchive = FindShaderLibrary(Hash);
 		if (ShaderCodeArchive)
 		{
+			PRAGMA_DISABLE_DEPRECATION_WARNINGS
 			if (bNativeFormat || GRHILazyShaderCodeLoading)
+			{
 				Result = RHICreateGeometryShaderWithStreamOutput(ElementList, NumStrides, Strides, RasterizedStream, ShaderCodeArchive, Hash);
+			}
 			else
+			{
 				Result = ((FShaderCodeArchive*)ShaderCodeArchive)->CreateGeometryShaderWithStreamOutput(Hash, ElementList, NumStrides, Strides, RasterizedStream);
+			}
+			PRAGMA_ENABLE_DEPRECATION_WARNINGS
 		}
 		return Result;
 	}
@@ -1822,9 +1842,13 @@ public:
 		if (ShaderCodeArchive)
 		{
 			if (bNativeFormat || GRHILazyShaderCodeLoading)
+			{
 				Result = RHICreateHullShader(ShaderCodeArchive, Hash);
+			}
 			else
+			{
 				Result = ((FShaderCodeArchive*)ShaderCodeArchive)->CreateHullShader(Hash);
+			}
 		}
 		return Result;
 	}
@@ -1838,9 +1862,13 @@ public:
 		if (ShaderCodeArchive)
 		{
 			if (bNativeFormat || GRHILazyShaderCodeLoading)
+			{
 				Result = RHICreateDomainShader(ShaderCodeArchive, Hash);
+			}
 			else
+			{
 				Result = ((FShaderCodeArchive*)ShaderCodeArchive)->CreateDomainShader(Hash);
+			}
 		}
 		return Result;
 	}
@@ -1854,9 +1882,13 @@ public:
 		if (ShaderCodeArchive)
 		{
 			if (bNativeFormat || GRHILazyShaderCodeLoading)
+			{
 				Result = RHICreateComputeShader(ShaderCodeArchive, Hash);
+			}
 			else
+			{
 				Result = ((FShaderCodeArchive*)ShaderCodeArchive)->CreateComputeShader(Hash);
+			}
 		}
 		return Result;
 	}
@@ -2324,7 +2356,9 @@ FGeometryShaderRHIRef FShaderCodeLibrary::CreateGeometryShaderWithStreamOutput(E
 	}
 	if (!IsValidRef(Shader))
 	{
+		PRAGMA_DISABLE_DEPRECATION_WARNINGS
 		Shader = RHICreateGeometryShaderWithStreamOutput(Code, ElementList, NumStrides, Strides, RasterizedStream);
+		PRAGMA_ENABLE_DEPRECATION_WARNINGS
 	}
 	SafeAssignHash(Shader, Hash);
 	return Shader;
diff --git a/Engine/Source/Runtime/RenderCore/Public/ShaderCodeLibrary.h b/Engine/Source/Runtime/RenderCore/Public/ShaderCodeLibrary.h
--- a/Engine/Source/Runtime/RenderCore/Public/ShaderCodeLibrary.h
+++ b/Engine/Source/Runtime/RenderCore/Public/ShaderCodeLibrary.h
@@ -159,6 +159,7 @@ struct RENDERCORE_API FShaderCodeLibrary
 	/** Instantiate or retrieve a geometry shader from the cache for the provided code & hash. */
 	static FGeometryShaderRHIRef CreateGeometryShader(EShaderPlatform Platform, FSHAHash Hash, TArray<uint8> const& Code);
 	/** Instantiate or retrieve a geometry shader from the cache for the provided code & hash. */
+	UE_DEPRECATED(4.23, "Geometry Stream out is deprecated.")
 	static FGeometryShaderRHIRef CreateGeometryShaderWithStreamOutput(EShaderPlatform Platform, FSHAHash Hash, const TArray<uint8>& Code, const FStreamOutElementList& ElementList, uint32 NumStrides, const uint32* Strides, int32 RasterizedStream);
 	/** Instantiate or retrieve a hull shader from the cache for the provided code & hash. */
 	static FHullShaderRHIRef CreateHullShader(EShaderPlatform Platform, FSHAHash Hash, TArray<uint8> const& Code);
