Fixes a RDG warning in the DOF's IndirectScatter pass that forgot to mark the indirect buffer as used by the pass execution.

diff --git a/Engine/Source/Runtime/RenderCore/Public/RenderGraphResources.h b/Engine/Source/Runtime/RenderCore/Public/RenderGraphResources.h
--- a/Engine/Source/Runtime/RenderCore/Public/RenderGraphResources.h
+++ b/Engine/Source/Runtime/RenderCore/Public/RenderGraphResources.h
@@ -590,6 +590,7 @@ public:
 		ValidateRHIAccess();
 		check(PooledBuffer);
 		checkf(Desc.UnderlyingType == FRDGBufferDesc::EUnderlyingType::VertexBuffer, TEXT("Indirect buffers needs to be underlying vertex buffer."));
+		checkf(Desc.Usage & BUF_DrawIndirect, TEXT("The buffer %s was not flagged for indirect draw call"), Name);
 		check(PooledBuffer->VertexBuffer.IsValid());
 		return PooledBuffer->VertexBuffer;
 	}
diff --git a/Engine/Source/Runtime/Renderer/Private/PostProcess/DiaphragmDOF.cpp b/Engine/Source/Runtime/Renderer/Private/PostProcess/DiaphragmDOF.cpp
--- a/Engine/Source/Runtime/Renderer/Private/PostProcess/DiaphragmDOF.cpp
+++ b/Engine/Source/Runtime/Renderer/Private/PostProcess/DiaphragmDOF.cpp
@@ -2294,6 +2294,9 @@ FRDGTextureRef DiaphragmDOF::AddPasses(
 
 				RHICmdList.SetStreamSource(0, NULL, 0);
 
+				// Marks the indirect draw parameter as used by the pass, given it's not used directly by any of the shaders.
+				PassParameters->IndirectDrawParameter->MarkResourceAsUsed();
+
 				if (GRHISupportsRectTopology)
 				{
 					RHICmdList.DrawPrimitiveIndirect(
