Fixes a bug in FComputeShaderUtils::AddPass() for indirect dispatch that was emitting an used warning for the indirect parameter buffer

diff --git a/Engine/Source/Runtime/RenderCore/Public/RenderGraphUtils.h b/Engine/Source/Runtime/RenderCore/Public/RenderGraphUtils.h
--- a/Engine/Source/Runtime/RenderCore/Public/RenderGraphUtils.h
+++ b/Engine/Source/Runtime/RenderCore/Public/RenderGraphUtils.h
@@ -116,6 +116,22 @@ struct RENDERCORE_API FComputeShaderUtils
 		RHICmdList.DispatchComputeShader(GroupCount.X, GroupCount.Y, GroupCount.Z);
 		UnsetShaderUAVs(RHICmdList, ComputeShader, ShaderRHI);
 	}
+	
+	/** Indirect dispatch a compute shader to rhi command list with its parameters. */
+	template<typename TShaderClass>
+	static FORCEINLINE_DEBUGGABLE void DispatchIndirect(
+		FRHICommandList& RHICmdList,
+		const TShaderClass* ComputeShader,
+		const typename TShaderClass::FParameters& Parameters,
+		FVertexBufferRHIParamRef IndirectArgsBuffer,
+		uint32 IndirectArgOffset)
+	{
+		FRHIComputeShader* ShaderRHI = ComputeShader->GetComputeShader();
+		RHICmdList.SetComputeShader(ShaderRHI);
+		SetShaderParameters(RHICmdList, ComputeShader, ShaderRHI, Parameters);
+		RHICmdList.DispatchIndirectComputeShader(IndirectArgsBuffer, IndirectArgOffset);
+		UnsetShaderUAVs(RHICmdList, ComputeShader, ShaderRHI);
+	}
 
 	/** Dispatch a compute shader to render graph builder with its parameters. */
 	template<typename TShaderClass>
@@ -148,6 +164,8 @@ struct RENDERCORE_API FComputeShaderUtils
 		FRDGBufferRef IndirectArgsBuffer,
 		uint32 IndirectArgOffset)
 	{
+		checkf(IndirectArgsBuffer->Desc.Usage & BUF_DrawIndirect, TEXT("The buffer %s was not flagged for indirect draw parameters"), IndirectArgsBuffer->Name);
+
 		ClearUnusedGraphResources(ComputeShader, Parameters, { IndirectArgsBuffer });
 
 		GraphBuilder.AddPass(
@@ -155,12 +173,12 @@ struct RENDERCORE_API FComputeShaderUtils
 			Parameters,
 			ERenderGraphPassFlags::Compute,
 			[Parameters, ComputeShader, IndirectArgsBuffer, IndirectArgOffset](FRHICommandList& RHICmdList)
-		{
-			FRHIComputeShader* ShaderRHI = ComputeShader->GetComputeShader();
-			RHICmdList.SetComputeShader(ShaderRHI);
-			SetShaderParameters(RHICmdList, ComputeShader, ShaderRHI, *Parameters);
-			RHICmdList.DispatchIndirectComputeShader( IndirectArgsBuffer->GetIndirectRHICallBuffer(), IndirectArgOffset );
-			UnsetShaderUAVs(RHICmdList, ComputeShader, ShaderRHI);
+		{			
+			// Marks the indirect draw parameter as used by the pass manually, given it can't be bound directly by any of the shader,
+			// meaning SetShaderParameters() won't be able to do it.
+			IndirectArgsBuffer->MarkResourceAsUsed();
+
+			FComputeShaderUtils::DispatchIndirect(RHICmdList, ComputeShader, *Parameters, IndirectArgsBuffer->GetIndirectRHICallBuffer(), IndirectArgOffset);
 		});
 	}
 };
\ No newline at end of file
