Add R32_UINT texture visualizer

diff --git a/Engine/Source/Runtime/RenderCore/Private/VisualizeTexture.cpp b/Engine/Source/Runtime/RenderCore/Private/VisualizeTexture.cpp
--- a/Engine/Source/Runtime/RenderCore/Private/VisualizeTexture.cpp
+++ b/Engine/Source/Runtime/RenderCore/Private/VisualizeTexture.cpp
@@ -52,6 +52,7 @@ enum class EVisualisePSType
 	Texture2DMSAA = 5,
 	Texture2DDepthStencilNoMSAA = 6,
 	Texture2DUINT8 = 7,
+	Texture2DUINT32 = 8,
 	MAX
 };
 
@@ -128,6 +129,10 @@ static EVisualisePSType GetVisualizePSType(const FRDGTextureDesc& Desc)
 			{
 				return EVisualisePSType::Texture2DUINT8;
 			}
+			else if (Desc.Format == PF_R32_UINT)
+			{
+				return EVisualisePSType::Texture2DUINT32;
+			}
 			else
 			{
 				// non MSAA
