Fixed build warning.

diff --git a/Engine/Source/Runtime/RenderCore/Private/RenderGraphBuilder.cpp b/Engine/Source/Runtime/RenderCore/Private/RenderGraphBuilder.cpp
--- a/Engine/Source/Runtime/RenderCore/Private/RenderGraphBuilder.cpp
+++ b/Engine/Source/Runtime/RenderCore/Private/RenderGraphBuilder.cpp
@@ -1445,7 +1445,7 @@ void FRDGBuilder::DestructPasses()
 {
 	#if RDG_ENABLE_DEBUG
 	{
-		const bool bEmitWarnings = GRDGDebug;
+		const bool bEmitWarnings = (GRDGDebug > 0);
 
 		for (const FRDGTrackedResourceRef Resource : TrackedResources)
 		{
