Fixed bug where first buffer UAV usage would silently fail to transition.

diff --git a/Engine/Source/Runtime/RenderCore/Private/RenderGraphBuilder.cpp b/Engine/Source/Runtime/RenderCore/Private/RenderGraphBuilder.cpp
--- a/Engine/Source/Runtime/RenderCore/Private/RenderGraphBuilder.cpp
+++ b/Engine/Source/Runtime/RenderCore/Private/RenderGraphBuilder.cpp
@@ -807,6 +807,9 @@ void FRDGBuilder::TransitionUAV(
 	ERDGPassAccess PassAccessAfter,
 	ERDGPassPipeline PassPipelineAfter) const
 {
+	check(UAV);
+	check(UnderlyingResource);
+
 	const ERDGPassAccess PassAccessBefore = UnderlyingResource->PassAccess;
 	const ERDGPassPipeline PassPipelineBefore = UnderlyingResource->PassPipeline;
 
@@ -961,13 +964,7 @@ void FRDGBuilder::PrepareResourcesForExecute(const FRDGPass* Pass, struct FRHIRe
 			if (FRDGTextureUAVRef UAV = Parameter.GetAsTextureUAV())
 			{
 				FRDGTextureRef Texture = UAV->Desc.Texture;
-				FRHIUnorderedAccessView* UAVRHI = UAV->GetRHI();
-
-				if (PassPipeline == ERDGPassPipeline::Graphics)
-				{
-					OutRPInfo->UAVs[OutRPInfo->NumUAVs++] = UAVRHI;	// Bind UAVs in declaration order
-				}
-
+	
 				#if RDG_ENABLE_DEBUG
 				{
 					Texture->PassAccessCount++;
@@ -975,6 +972,14 @@ void FRDGBuilder::PrepareResourcesForExecute(const FRDGPass* Pass, struct FRHIRe
 				#endif
 
 				AllocateRHITextureUAVIfNeeded(UAV);
+
+				FRHIUnorderedAccessView* UAVRHI = UAV->GetRHI();
+
+				if (PassPipeline == ERDGPassPipeline::Graphics)
+				{
+					OutRPInfo->UAVs[OutRPInfo->NumUAVs++] = UAVRHI;	// Bind UAVs in declaration order
+				}
+
 				TransitionUAV(UAVRHI, Texture, ERDGPassAccess::Write, PassPipeline);
 			}
 		}
@@ -1014,12 +1019,13 @@ void FRDGBuilder::PrepareResourcesForExecute(const FRDGPass* Pass, struct FRHIRe
 				}
 				#endif
 
+				AllocateRHIBufferSRVIfNeeded(SRV);
+
 				// TODO(RDG): super hacky, find the UAV and transition it. Hopefully there is one...
 				check(Buffer->PooledBuffer);
 				check(Buffer->PooledBuffer->UAVs.Num() == 1);
 				FRHIUnorderedAccessView* UAV = Buffer->PooledBuffer->UAVs.CreateIterator().Value();
 
-				AllocateRHIBufferSRVIfNeeded(SRV);
 				TransitionUAV(UAV, Buffer, ERDGPassAccess::Read, PassPipeline);
 			}
 		}
@@ -1029,13 +1035,7 @@ void FRDGBuilder::PrepareResourcesForExecute(const FRDGPass* Pass, struct FRHIRe
 			if (FRDGBufferUAVRef UAV = Parameter.GetAsBufferUAV())
 			{
 				FRDGBufferRef Buffer = UAV->Desc.Buffer;
-				FRHIUnorderedAccessView* UAVRHI = UAV->GetRHI();
 
-				if (PassPipeline == ERDGPassPipeline::Graphics)
-				{
-					OutRPInfo->UAVs[OutRPInfo->NumUAVs++] = UAVRHI;	// Bind UAVs in declaration order
-				}
-				
 				#if RDG_ENABLE_DEBUG
 				{
 					Buffer->PassAccessCount++;
@@ -1043,6 +1043,14 @@ void FRDGBuilder::PrepareResourcesForExecute(const FRDGPass* Pass, struct FRHIRe
 				#endif
 
 				AllocateRHIBufferUAVIfNeeded(UAV);
+
+				FRHIUnorderedAccessView* UAVRHI = UAV->GetRHI();
+
+				if (PassPipeline == ERDGPassPipeline::Graphics)
+				{
+					OutRPInfo->UAVs[OutRPInfo->NumUAVs++] = UAVRHI;	// Bind UAVs in declaration order
+				}
+	
 				TransitionUAV(UAVRHI, Buffer, ERDGPassAccess::Write, PassPipeline);
 			}
 		}
