Fixed bad merge from main causing UAVs to be bound twice.

diff --git a/Engine/Source/Runtime/RenderCore/Public/ShaderParameterStruct.h b/Engine/Source/Runtime/RenderCore/Public/ShaderParameterStruct.h
--- a/Engine/Source/Runtime/RenderCore/Public/ShaderParameterStruct.h
+++ b/Engine/Source/Runtime/RenderCore/Public/ShaderParameterStruct.h
@@ -260,9 +260,6 @@ inline void SetShaderParameters(TRHICmdList& RHICmdList, const TShaderClass* Sha
 		RHICmdList.SetShaderResourceViewParameter(ShadeRHI, ParameterBinding.BaseIndex, GraphSRV->GetRHI());
 	}
 
-	// UAVs for compute shaders
-	SetShaderUAVs(RHICmdList, Shader, ShadeRHI, Parameters);
-
 	// Reference structures
 	for (const FShaderParameterBindings::FParameterStructReference& ParameterBinding : Bindings.ParameterReferences)
 	{
