Adds breakpoint when emitting a RDG warning.

diff --git a/Engine/Source/Runtime/RenderCore/Private/RenderGraphBuilder.cpp b/Engine/Source/Runtime/RenderCore/Private/RenderGraphBuilder.cpp
--- a/Engine/Source/Runtime/RenderCore/Private/RenderGraphBuilder.cpp
+++ b/Engine/Source/Runtime/RenderCore/Private/RenderGraphBuilder.cpp
@@ -30,6 +30,12 @@ FAutoConsoleVariableRef CVarRDGDebug(
 	TEXT(" 2: emit warning everytime issue is detected."),
 	ECVF_RenderThreadSafe);
 
+static TAutoConsoleVariable<int32> CVarRDGEnableBreakPoint(
+	TEXT("r.RDG.EnableBreakPoint"), 0,
+	TEXT("Breakpoint in debugger when a warning is raised.\n"),
+	ECVF_RenderThreadSafe);
+
+
 #else
 
 const int32 GRDGImmediateMode = 0;
@@ -64,6 +70,7 @@ void InitRenderGraph()
 
 void EmitRDGWarning(const FString& WarningMessage)
 {
+#if RDG_ENABLE_DEBUG_WITH_ENGINE
 	if (!GRDGDebug)
 	{
 		return;
@@ -71,18 +78,31 @@ void EmitRDGWarning(const FString& WarningMessage)
 
 	static TSet<FString> GAlreadyEmittedWarnings;
 
+	const bool bEnableBreakPoint = CVarRDGEnableBreakPoint.GetValueOnRenderThread();
+
 	if (GRDGDebug == kRDGEmitWarningsOnce)
 	{
 		if (!GAlreadyEmittedWarnings.Contains(WarningMessage))
 		{
 			GAlreadyEmittedWarnings.Add(WarningMessage);
 			UE_LOG(LogRendererCore, Warning, TEXT("%s"), *WarningMessage);
+			
+			if (bEnableBreakPoint)
+			{
+				UE_DEBUG_BREAK();
+			}
 		}
 	}
 	else
 	{
 		UE_LOG(LogRendererCore, Warning, TEXT("%s"), *WarningMessage);
+		
+		if (bEnableBreakPoint)
+		{
+			UE_DEBUG_BREAK();
+		}
 	}
+#endif
 }
 
 #define EmitRDGWarningf(WarningMessageFormat, ...) \
