Replace FShader UniformBufferParameters array of separately allocated pointers with array of values

diff --git a/Engine/Source/Runtime/RenderCore/Private/Shader.cpp b/Engine/Source/Runtime/RenderCore/Private/Shader.cpp
--- a/Engine/Source/Runtime/RenderCore/Private/Shader.cpp
+++ b/Engine/Source/Runtime/RenderCore/Private/Shader.cpp
@@ -1172,9 +1172,8 @@ FShader::FShader(const CompiledShaderInitializerType& Initializer)
 		if (Initializer.ParameterMap.ContainsParameterAllocation(StructIt->GetShaderVariableName()))
 		{
 			UniformBufferParameterStructs.Add(*StructIt);
-			UniformBufferParameters.Add(new FShaderUniformBufferParameter());
-			FShaderUniformBufferParameter* Parameter = UniformBufferParameters.Last();
-			Parameter->Bind(Initializer.ParameterMap, StructIt->GetShaderVariableName(), SPF_Mandatory);
+			FShaderUniformBufferParameter& Parameter = UniformBufferParameters.AddDefaulted_GetRef();
+			Parameter.Bind(Initializer.ParameterMap, StructIt->GetShaderVariableName(), SPF_Mandatory);
 		}
 	}
 
@@ -1188,11 +1187,6 @@ FShader::FShader(const CompiledShaderInitializerType& Initializer)
 FShader::~FShader()
 {
 	check(NumRefs == 0);
-
-	for (int32 StructIndex = 0; StructIndex < UniformBufferParameters.Num(); StructIndex++)
-	{
-		delete UniformBufferParameters[StructIndex];
-	}
 }
 
 
@@ -1266,13 +1260,10 @@ bool FShader::SerializeBase(FArchive& Ar, bool bShadersInline, bool bLoadedByCoo
 				Struct = FindUniformBufferStructByFName(StructFName);
 				checkf(Struct, TEXT("Uniform Buffer Struct %s no longer exists, which shader of type %s was compiled with.  Modify ShaderVersion.ush to invalidate old shaders."), *StructFName.ToString(), Type->GetName());
 			}
-			
-			FShaderUniformBufferParameter* Parameter = new FShaderUniformBufferParameter();
-
-			Ar << *Parameter;
 
 			UniformBufferParameterStructs.Add(Struct);
-			UniformBufferParameters.Add(Parameter);
+			FShaderUniformBufferParameter& Parameter = UniformBufferParameters.AddDefaulted_GetRef();
+			Ar << Parameter;
 		}
 	}
 	else
@@ -1294,7 +1285,7 @@ bool FShader::SerializeBase(FArchive& Ar, bool bShadersInline, bool bLoadedByCoo
 				Ar << StructFName;
 			}
 
-			Ar << *UniformBufferParameters[StructIndex];
+			Ar << UniformBufferParameters[StructIndex];
 		}
 	}
 
diff --git a/Engine/Source/Runtime/RenderCore/Public/Shader.h b/Engine/Source/Runtime/RenderCore/Public/Shader.h
--- a/Engine/Source/Runtime/RenderCore/Public/Shader.h
+++ b/Engine/Source/Runtime/RenderCore/Public/Shader.h
@@ -1000,39 +1000,14 @@ public:
 	template<typename UniformBufferStructType>
 	FORCEINLINE_DEBUGGABLE const TShaderUniformBufferParameter<UniformBufferStructType>& GetUniformBufferParameter() const
 	{
-		const FShaderParametersMetadata* SearchStruct = &UniformBufferStructType::StaticStructMetadata;
-		int32 FoundIndex = INDEX_NONE;
-
-		for (int32 StructIndex = 0, Count = UniformBufferParameterStructs.Num(); StructIndex < Count; StructIndex++)
-		{
-			if (UniformBufferParameterStructs[StructIndex] == SearchStruct)
-			{
-				FoundIndex = StructIndex;
-				break;
-			}
-		}
-
-		if (FoundIndex != INDEX_NONE)
-		{
-			const TShaderUniformBufferParameter<UniformBufferStructType>& FoundParameter = (const TShaderUniformBufferParameter<UniformBufferStructType>&)*UniformBufferParameters[FoundIndex];
-			return FoundParameter;
-		}
-		else
-		{
-			// This can happen if the uniform buffer was not bound
-			// There's no good way to distinguish not being bound due to temporary debugging / compiler optimizations or an actual code bug,
-			// Hence failing silently instead of an error message
-			static TShaderUniformBufferParameter<UniformBufferStructType> UnboundParameter;
-			UnboundParameter.SetInitialized();
-			return UnboundParameter;
-		}
+		const FShaderUniformBufferParameter& FoundParameter = GetUniformBufferParameter(&UniformBufferStructType::StaticStructMetadata);
+		return static_cast<const TShaderUniformBufferParameter<UniformBufferStructType>&>(FoundParameter);
 	}
 
 	/** Finds an automatically bound uniform buffer matching the given uniform buffer struct if one exists, or returns an unbound parameter. */
-	const FShaderUniformBufferParameter& GetUniformBufferParameter(const FShaderParametersMetadata* SearchStruct) const
+	FORCEINLINE_DEBUGGABLE const FShaderUniformBufferParameter& GetUniformBufferParameter(const FShaderParametersMetadata* SearchStruct) const
 	{
 		int32 FoundIndex = INDEX_NONE;
-
 		for (int32 StructIndex = 0, Count = UniformBufferParameterStructs.Num(); StructIndex < Count; StructIndex++)
 		{
 			if (UniformBufferParameterStructs[StructIndex] == SearchStruct)
@@ -1044,11 +1019,14 @@ public:
 
 		if (FoundIndex != INDEX_NONE)
 		{
-			const FShaderUniformBufferParameter& FoundParameter = *UniformBufferParameters[FoundIndex];
+			const FShaderUniformBufferParameter& FoundParameter = UniformBufferParameters[FoundIndex];
 			return FoundParameter;
 		}
 		else
 		{
+			// This can happen if the uniform buffer was not bound
+			// There's no good way to distinguish not being bound due to temporary debugging / compiler optimizations or an actual code bug,
+			// Hence failing silently instead of an error message
 			static FShaderUniformBufferParameter UnboundParameter;
 			UnboundParameter.SetInitialized();
 			return UnboundParameter;
@@ -1059,7 +1037,7 @@ public:
 	{
 		for (int32 i = 0; i < UniformBufferParameters.Num(); i++)
 		{
-			if (UniformBufferParameters[i]->GetBaseIndex() == BaseIndex)
+			if (UniformBufferParameters[i].GetBaseIndex() == BaseIndex)
 			{
 				return UniformBufferParameterStructs[i];
 			}
@@ -1094,7 +1072,7 @@ protected:
 
 	/** Indexed the same as UniformBufferParameters.  Packed densely for coherent traversal. */
 	TArray<const FShaderParametersMetadata*> UniformBufferParameterStructs;
-	TArray<FShaderUniformBufferParameter*> UniformBufferParameters;
+	TArray<FShaderUniformBufferParameter> UniformBufferParameters;
 
 private:
 	/** Hash of the material shader map this shader belongs to, stored so that an FShaderId can be constructed from this shader. */
