Renamed stat GenerateVisibleRayTracingMeshCommands to GatherRayTracingWorldInstances, as that's the function it actually measures.

diff --git a/Engine/Source/Runtime/RenderCore/Private/RenderCore.cpp b/Engine/Source/Runtime/RenderCore/Private/RenderCore.cpp
--- a/Engine/Source/Runtime/RenderCore/Private/RenderCore.cpp
+++ b/Engine/Source/Runtime/RenderCore/Private/RenderCore.cpp
@@ -60,7 +60,7 @@ DEFINE_STAT(STAT_DynamicShadowSetupTime);
 DEFINE_STAT(STAT_RenderQueryResultTime);
 // Use 'stat initviews' to get more detail
 DEFINE_STAT(STAT_InitViewsTime);
-DEFINE_STAT(STAT_GenerateVisibleRayTracingMeshCommands);
+DEFINE_STAT(STAT_GatherRayTracingWorldInstances);
 DEFINE_STAT(STAT_InitViewsPossiblyAfterPrepass);
 // Measures the time spent in RenderViewFamily_RenderThread
 // Note that this is not the total rendering thread time, any other rendering commands will not be counted here
diff --git a/Engine/Source/Runtime/RenderCore/Public/RenderCore.h b/Engine/Source/Runtime/RenderCore/Public/RenderCore.h
--- a/Engine/Source/Runtime/RenderCore/Public/RenderCore.h
+++ b/Engine/Source/Runtime/RenderCore/Public/RenderCore.h
@@ -18,7 +18,7 @@ DECLARE_LOG_CATEGORY_EXTERN(LogRendererCore, Log, All);
 DECLARE_CYCLE_STAT_EXTERN(TEXT("BeginOcclusionTests"),STAT_BeginOcclusionTestsTime,STATGROUP_SceneRendering, RENDERCORE_API);
 DECLARE_CYCLE_STAT_EXTERN(TEXT("RenderQuery Result"),STAT_RenderQueryResultTime,STATGROUP_SceneRendering, RENDERCORE_API);
 DECLARE_CYCLE_STAT_EXTERN(TEXT("InitViews"),STAT_InitViewsTime,STATGROUP_SceneRendering, RENDERCORE_API);
-DECLARE_CYCLE_STAT_EXTERN(TEXT("GenerateVisibleRayTracingMeshCommands"), STAT_GenerateVisibleRayTracingMeshCommands, STATGROUP_SceneRendering, RENDERCORE_API);
+DECLARE_CYCLE_STAT_EXTERN(TEXT("GatherRayTracingWorldInstances"), STAT_GatherRayTracingWorldInstances, STATGROUP_SceneRendering, RENDERCORE_API);
 DECLARE_CYCLE_STAT_EXTERN(TEXT("InitViewsPossiblyAfterPrepass"), STAT_InitViewsPossiblyAfterPrepass, STATGROUP_SceneRendering, RENDERCORE_API);
 DECLARE_CYCLE_STAT_EXTERN(TEXT("Dynamic shadow setup"), STAT_DynamicShadowSetupTime, STATGROUP_SceneRendering, RENDERCORE_API);
 DECLARE_CYCLE_STAT_EXTERN(TEXT("Total GPU rendering"),STAT_TotalGPUFrameTime,STATGROUP_SceneRendering, RENDERCORE_API);
diff --git a/Engine/Source/Runtime/Renderer/Private/DeferredShadingRenderer.cpp b/Engine/Source/Runtime/Renderer/Private/DeferredShadingRenderer.cpp
--- a/Engine/Source/Runtime/Renderer/Private/DeferredShadingRenderer.cpp
+++ b/Engine/Source/Runtime/Renderer/Private/DeferredShadingRenderer.cpp
@@ -530,7 +530,7 @@ bool FDeferredShadingSceneRenderer::GatherRayTracingWorldInstances(FRHICommandLi
 	}
 
 	{
-		SCOPE_CYCLE_COUNTER(STAT_GenerateVisibleRayTracingMeshCommands);
+		SCOPE_CYCLE_COUNTER(STAT_GatherRayTracingWorldInstances);
 		RayTracingCollector.ClearViewMeshArrays();
 		TArray<int> DynamicMeshBatchStartOffset;
 		TArray<int> VisibleDrawCommandStartOffset;
