Shader resource loading optimizations

~80% of loaded shader resources were already loaded and therefor discarded in an internal project. They were serialized on the async loading thread and later destroyed on the main thread causing lots of redundant heap allocations with odd life times.

Shader resources are now immediately registered after being loaded on the async loading thread. They get deserialized onto the stack and then moved to a heap-allocated shader resource if the resource didn't already exist. This speeds up shader post-load in particular, since it no longer deletes lots of shader resources. Asset loading in general becomes a bit faster too due to reduced allocator pressure.

Shader resource life times are now handled via thread-safe refcounting. This required some minor create / find shader resource API changes.

Integrated from CL 6926407 in Dev-Core

diff --git a/Engine/Plugins/Compositing/OpenColorIO/Source/OpenColorIO/Private/OpenColorIOShaderMap.cpp b/Engine/Plugins/Compositing/OpenColorIO/Source/OpenColorIO/Private/OpenColorIOShaderMap.cpp
--- a/Engine/Plugins/Compositing/OpenColorIO/Source/OpenColorIO/Private/OpenColorIOShaderMap.cpp
+++ b/Engine/Plugins/Compositing/OpenColorIO/Source/OpenColorIO/Private/OpenColorIOShaderMap.cpp
@@ -234,7 +234,7 @@ FShader* FOpenColorIOShaderType::FinishCompileShader(
 
 	// Reuse an existing resource with the same key or create a new one based on the compile output
 	// This allows FShaders to share compiled bytecode and RHI shader references
-	FShaderResource* Resource = FShaderResource::FindOrCreateShaderResource(InCurrentJob.Output, SpecificType, /* SpecificPermutationId = */ 0);
+	TRefCountPtr<FShaderResource> Resource = FShaderResource::FindOrCreate(InCurrentJob.Output, SpecificType, /* SpecificPermutationId = */ 0);
 
 	// Find a shader with the same key in memory
 	FShader* Shader = InCurrentJob.ShaderType->FindShaderById(FShaderId(InShaderMapHash, nullptr, nullptr, InCurrentJob.ShaderType, /* SpecificPermutationId = */ 0, InCurrentJob.Input.Target));
diff --git a/Engine/Plugins/FX/Niagara/Source/NiagaraShader/Private/NiagaraShader.cpp b/Engine/Plugins/FX/Niagara/Source/NiagaraShader/Private/NiagaraShader.cpp
--- a/Engine/Plugins/FX/Niagara/Source/NiagaraShader/Private/NiagaraShader.cpp
+++ b/Engine/Plugins/FX/Niagara/Source/NiagaraShader/Private/NiagaraShader.cpp
@@ -410,7 +410,7 @@ FShader* FNiagaraShaderType::FinishCompileShader(
 
 	// Reuse an existing resource with the same key or create a new one based on the compile output
 	// This allows FShaders to share compiled bytecode and RHI shader references
-	FShaderResource* Resource = FShaderResource::FindOrCreateShaderResource(CurrentJob.Output, SpecificType, /* SpecificPermutationId = */ 0);
+	TRefCountPtr<FShaderResource> Resource = FShaderResource::FindOrCreate(CurrentJob.Output, SpecificType, /* SpecificPermutationId = */ 0);
 
 	// Find a shader with the same key in memory
 	FShader* Shader = CurrentJob.ShaderType->FindShaderById(FShaderId(ShaderMapHash, nullptr, nullptr, CurrentJob.ShaderType, /* SpecificPermutationId = */ 0, CurrentJob.Input.Target));
@@ -427,7 +427,7 @@ FShader* FNiagaraShaderType::FinishCompileShader(
 			check(false);
 		}
 
-		Shader = (*ConstructCompiledRef)(FNiagaraShaderType::CompiledShaderInitializerType(this, PermutationId, CurrentJob.Output, Resource, ShaderMapHash, InDebugDescription, DIParamInfo));
+		Shader = (*ConstructCompiledRef)(FNiagaraShaderType::CompiledShaderInitializerType(this, PermutationId, CurrentJob.Output, MoveTemp(Resource), ShaderMapHash, InDebugDescription, DIParamInfo));
 		//CurrentJob.Output.ParameterMap.VerifyBindingsAreComplete(GetName(), CurrentJob.Output.Target, nullptr); // b/c we don't bind data interfaces yet...
 
 	}
@@ -607,9 +607,8 @@ void FNiagaraShaderMap::SaveForRemoteRecompile(FArchive& Ar, const TMap<FString,
 					if (ClientResourceIds.Contains(ShaderId) == false)
 					{
 						// lookup the resource by ID
-						FShaderResource* Resource = FShaderResource::FindShaderResourceById(ShaderId);
 						// add it if it's unique
-						UniqueResources.AddUnique(Resource);
+						UniqueResources.AddUnique(FShaderResource::FindById(ShaderId));
 					}
 					else
 					{
diff --git a/Engine/Plugins/FX/Niagara/Source/NiagaraShader/Public/NiagaraShaderType.h b/Engine/Plugins/FX/Niagara/Source/NiagaraShader/Public/NiagaraShaderType.h
--- a/Engine/Plugins/FX/Niagara/Source/NiagaraShader/Public/NiagaraShaderType.h
+++ b/Engine/Plugins/FX/Niagara/Source/NiagaraShader/Public/NiagaraShaderType.h
@@ -56,12 +56,12 @@ public:
 			FShaderType* InType,
 			int32 InPermutationId,
 			const FShaderCompilerOutput& CompilerOutput,
-			FShaderResource* InResource,
+			TRefCountPtr<FShaderResource>&& InResource,
 			const FSHAHash& InNiagaraShaderMapHash,
 			const FString& InDebugDescription,
 			const TArray< FNiagaraDataInterfaceGPUParamInfo > &InDIParamInfo
 			)
-		: FGlobalShaderType::CompiledShaderInitializerType(InType,InPermutationId,CompilerOutput,InResource, InNiagaraShaderMapHash,nullptr,nullptr)
+		: FGlobalShaderType::CompiledShaderInitializerType(InType,InPermutationId,CompilerOutput, MoveTemp(InResource), InNiagaraShaderMapHash,nullptr,nullptr)
 		, DebugDescription(InDebugDescription)
 		, DIParamInfo(InDIParamInfo)
 		{}
diff --git a/Engine/Source/Runtime/Engine/Private/Materials/MaterialShader.cpp b/Engine/Source/Runtime/Engine/Private/Materials/MaterialShader.cpp
--- a/Engine/Source/Runtime/Engine/Private/Materials/MaterialShader.cpp
+++ b/Engine/Source/Runtime/Engine/Private/Materials/MaterialShader.cpp
@@ -1012,7 +1012,7 @@ FShader* FMaterialShaderType::FinishCompileShader(
 
 	// Reuse an existing resource with the same key or create a new one based on the compile output
 	// This allows FShaders to share compiled bytecode and RHI shader references
-	FShaderResource* Resource = FShaderResource::FindOrCreateShaderResource(CurrentJob.Output, SpecificType, SpecificPermutationId);
+	TRefCountPtr<FShaderResource> Resource = FShaderResource::FindOrCreate(CurrentJob.Output, SpecificType, SpecificPermutationId);
 
 	if (ShaderPipelineType && !ShaderPipelineType->ShouldOptimizeUnusedOutputs(CurrentJob.Input.Target.GetPlatform()))
 	{
@@ -1026,7 +1026,7 @@ FShader* FMaterialShaderType::FinishCompileShader(
 	// There was no shader with the same key so create a new one with the compile output, which will bind shader parameters
 	if (!Shader)
 	{
-		Shader = (*ConstructCompiledRef)(CompiledShaderInitializerType(this, CurrentJob.PermutationId, CurrentJob.Output, Resource, UniformExpressionSet, MaterialShaderMapHash, ShaderPipelineType, nullptr, InDebugDescription));
+		Shader = (*ConstructCompiledRef)(CompiledShaderInitializerType(this, CurrentJob.PermutationId, CurrentJob.Output, MoveTemp(Resource), UniformExpressionSet, MaterialShaderMapHash, ShaderPipelineType, nullptr, InDebugDescription));
 		CurrentJob.Output.ParameterMap.VerifyBindingsAreComplete(GetName(), CurrentJob.Output.Target, CurrentJob.VFType);
 	}
 
@@ -1215,7 +1215,7 @@ void FMaterialShaderMap::SaveForRemoteRecompile(FArchive& Ar, const TMap<FString
 	UE_LOG(LogMaterial, Display, TEXT("Looking for unique resources, %d were on client"), ClientResourceIds.Num());	
 
 	// first, we look for the unique shader resources
-	TArray<FShaderResource*> UniqueResources;
+	TArray<TRefCountPtr<FShaderResource>> UniqueResources;
 	int32 NumSkippedResources = 0;
 
 	for (TMap<FString, TArray<TRefCountPtr<FMaterialShaderMap> > >::TConstIterator It(CompiledShaderMaps); It; ++It)
@@ -1254,10 +1254,8 @@ void FMaterialShaderMap::SaveForRemoteRecompile(FArchive& Ar, const TMap<FString
 					// skip this shader if the Id was already on the client (ie, it didn't change)
 					if (ClientResourceIds.Contains(ShaderId) == false)
 					{
-						// lookup the resource by ID
-						FShaderResource* Resource = FShaderResource::FindShaderResourceById(ShaderId);
-						// add it if it's unique
-						UniqueResources.AddUnique(Resource);
+						// lookup the resource by ID and add it if it's unique
+						UniqueResources.AddUnique(FShaderResource::FindById(ShaderId));
 					}
 					else
 					{
@@ -1319,29 +1317,17 @@ void FMaterialShaderMap::LoadForRemoteRecompile(FArchive& Ar, EShaderPlatform Sh
 	Ar << NumResources;
 
 	// KeepAliveReferences keeps resources alive until we are finished serializing in this function
-	TArray<TRefCountPtr<FShaderResource> > KeepAliveReferences;
+	TArray<TRefCountPtr<FShaderResource>> KeepAliveReferences;
+	KeepAliveReferences.SetNum(NumResources);
 
 	// load and register the resources
 	for (int32 Index = 0; Index < NumResources; Index++)
 	{
 		// Load the inlined shader resource
-		FShaderResource* Resource = new FShaderResource();
-		Resource->Serialize(Ar, false);
+		FShaderResource ResourceTemp;
+		ResourceTemp.Serialize(Ar, false);
 
-		// if this Id is already in memory, that means that this is a repeated resource and so we skip it
-		if (FShaderResource::FindShaderResourceById(Resource->GetId()) != NULL)
-		{
-			delete Resource;
-		}
-		// otherwise, it's a new resource, so we register it for the maps to find below
-		else
-		{
-			Resource->Register();
-
-			// Keep this guy alive until we finish serializing all the FShaders in
-			// The FShaders which are discarded may cause these resources to be discarded 
-			KeepAliveReferences.Add( Resource );
-		}
+		KeepAliveReferences[Index] = FShaderResource::FindOrClone(MoveTemp(ResourceTemp));
 	}
 
 	int32 MapSize;
diff --git a/Engine/Source/Runtime/Engine/Private/Materials/MeshMaterialShader.cpp b/Engine/Source/Runtime/Engine/Private/Materials/MeshMaterialShader.cpp
--- a/Engine/Source/Runtime/Engine/Private/Materials/MeshMaterialShader.cpp
+++ b/Engine/Source/Runtime/Engine/Private/Materials/MeshMaterialShader.cpp
@@ -151,7 +151,7 @@ FShader* FMeshMaterialShaderType::FinishCompileShader(
 
 	// Reuse an existing resource with the same key or create a new one based on the compile output
 	// This allows FShaders to share compiled bytecode and RHI shader references
-	FShaderResource* Resource = FShaderResource::FindOrCreateShaderResource(CurrentJob.Output, SpecificType, CurrentJob.PermutationId);
+	TRefCountPtr<FShaderResource> Resource = FShaderResource::FindOrCreate(CurrentJob.Output, SpecificType, CurrentJob.PermutationId);
 
 	if (ShaderPipelineType && !ShaderPipelineType->ShouldOptimizeUnusedOutputs(CurrentJob.Input.Target.GetPlatform()))
 	{
@@ -165,7 +165,7 @@ FShader* FMeshMaterialShaderType::FinishCompileShader(
 	// There was no shader with the same key so create a new one with the compile output, which will bind shader parameters
 	if (!Shader)
 	{
-		Shader = (*ConstructCompiledRef)(CompiledShaderInitializerType(this, CurrentJob.PermutationId, CurrentJob.Output, Resource, UniformExpressionSet, MaterialShaderMapHash, InDebugDescription, ShaderPipelineType, CurrentJob.VFType));
+		Shader = (*ConstructCompiledRef)(CompiledShaderInitializerType(this, CurrentJob.PermutationId, CurrentJob.Output, MoveTemp(Resource), UniformExpressionSet, MaterialShaderMapHash, InDebugDescription, ShaderPipelineType, CurrentJob.VFType));
 		ValidateAfterBind((FMeshMaterialShader*)Shader);
 		CurrentJob.Output.ParameterMap.VerifyBindingsAreComplete(GetName(), CurrentJob.Output.Target, CurrentJob.VFType);
 	}
diff --git a/Engine/Source/Runtime/Engine/Private/ShaderCompiler/ShaderCompiler.cpp b/Engine/Source/Runtime/Engine/Private/ShaderCompiler/ShaderCompiler.cpp
--- a/Engine/Source/Runtime/Engine/Private/ShaderCompiler/ShaderCompiler.cpp
+++ b/Engine/Source/Runtime/Engine/Private/ShaderCompiler/ShaderCompiler.cpp
@@ -3764,7 +3764,7 @@ FShader* FGlobalShaderTypeCompiler::FinishCompileShader(FGlobalShaderType* Shade
 
 		// Reuse an existing resource with the same key or create a new one based on the compile output
 		// This allows FShaders to share compiled bytecode and RHI shader references
-		FShaderResource* Resource = FShaderResource::FindOrCreateShaderResource(CurrentJob.Output, SpecificType, SpecificPermutationId);
+		TRefCountPtr<FShaderResource> Resource = FShaderResource::FindOrCreate(CurrentJob.Output, SpecificType, SpecificPermutationId);
 		check(Resource);
 
 		if (ShaderPipelineType && !ShaderPipelineType->ShouldOptimizeUnusedOutputs(CurrentJob.Input.Target.GetPlatform()))
@@ -3790,7 +3790,7 @@ FShader* FGlobalShaderTypeCompiler::FinishCompileShader(FGlobalShaderType* Shade
 		// There was no shader with the same key so create a new one with the compile output, which will bind shader parameters
 		if (!Shader)
 		{
-			Shader = (*(ShaderType->ConstructCompiledRef))(FGlobalShaderType::CompiledShaderInitializerType(ShaderType, CurrentJob.PermutationId, CurrentJob.Output, Resource, GlobalShaderMapHash, ShaderPipelineType, nullptr));
+			Shader = (*(ShaderType->ConstructCompiledRef))(FGlobalShaderType::CompiledShaderInitializerType(ShaderType, CurrentJob.PermutationId, CurrentJob.Output, MoveTemp(Resource), GlobalShaderMapHash, ShaderPipelineType, nullptr));
 			CurrentJob.Output.ParameterMap.VerifyBindingsAreComplete(ShaderType->GetName(), CurrentJob.Output.Target, CurrentJob.VFType);
 		}
 	}
diff --git a/Engine/Source/Runtime/Engine/Public/MaterialShaderType.h b/Engine/Source/Runtime/Engine/Public/MaterialShaderType.h
--- a/Engine/Source/Runtime/Engine/Public/MaterialShaderType.h
+++ b/Engine/Source/Runtime/Engine/Public/MaterialShaderType.h
@@ -85,14 +85,14 @@ public:
 			FShaderType* InType,
 			int32 InPermutationId,
 			const FShaderCompilerOutput& CompilerOutput,
-			FShaderResource* InResource,
+			TRefCountPtr<FShaderResource>&& InResource,
 			const FUniformExpressionSet& InUniformExpressionSet,
 			const FSHAHash& InMaterialShaderMapHash,
 			const FShaderPipelineType* InShaderPipeline,
 			FVertexFactoryType* InVertexFactoryType,
 			const FString& InDebugDescription
 			)
-		: FGlobalShaderType::CompiledShaderInitializerType(InType,InPermutationId,CompilerOutput,InResource,InMaterialShaderMapHash,InShaderPipeline,InVertexFactoryType)
+		: FGlobalShaderType::CompiledShaderInitializerType(InType,InPermutationId,CompilerOutput,MoveTemp(InResource),InMaterialShaderMapHash,InShaderPipeline,InVertexFactoryType)
 		, UniformExpressionSet(InUniformExpressionSet)
 		, DebugDescription(InDebugDescription)
 		{}
diff --git a/Engine/Source/Runtime/RenderCore/Private/Shader.cpp b/Engine/Source/Runtime/RenderCore/Private/Shader.cpp
--- a/Engine/Source/Runtime/RenderCore/Private/Shader.cpp
+++ b/Engine/Source/Runtime/RenderCore/Private/Shader.cpp
@@ -19,7 +19,7 @@
 #include "Misc/ScopeLock.h"
 #include "UObject/RenderingObjectVersion.h"
 #include "UObject/FortniteMainBranchObjectVersion.h"
-#include "Misc/ScopeLock.h"
+#include "Misc/ScopeRwLock.h"
 
 #if WITH_EDITORONLY_DATA
 #include "Interfaces/IShaderFormat.h"
@@ -391,8 +391,113 @@ void FShaderType::Uninitialize()
 	bInitializedSerializationHistory = false;
 }
 
-TMap<FShaderResourceId, FShaderResource*> FShaderResource::ShaderResourceIdMap;
+struct FShaderResourceKeyFuncs
+{
+	typedef FShaderResource* KeyType;
+	typedef FShaderResource* KeyInitType;
+	typedef FShaderResource* ElementInitType;
+
+	enum { bAllowDuplicateKeys = false };
 
+	static KeyInitType GetSetKey(FShaderResource* Element)
+	{
+		return Element;
+	}
+
+	static bool Matches(FShaderResource* A, FShaderResource* B)
+	{
+		return A == B || A->GetId() == B->GetId();
+	}
+
+	static bool Matches(FShaderResource* Resource, const FShaderResourceId& Id)
+	{
+		return Resource->GetId() == Id;
+	}
+
+	static uint32 GetKeyHash(FShaderResource* Key)
+	{
+		return GetTypeHash(Key->GetOutputHash());
+	}
+
+	static uint32 GetKeyHash(const FShaderResourceId& Key)
+	{
+		return GetTypeHash(Key.OutputHash);
+	}
+};
+
+
+class FShaderResourceSet 
+{
+public:
+	TRefCountPtr<FShaderResource> Find(const FShaderResourceId& Id) const
+	{
+		uint32 IdHash = FShaderResourceKeyFuncs::GetKeyHash(Id);
+
+		FRWScopeLock Locker(Lock, SLT_ReadOnly);
+		
+		if (FShaderResource*const* ResourceIt = Set.FindByHash(IdHash, Id))
+		{
+			return TRefCountPtr<FShaderResource>(*ResourceIt);
+		}
+
+		return TRefCountPtr<FShaderResource>();
+	}
+
+	/** 
+	 * Add resource and return it as smart pointer unless an identical resource exists.
+	 * In that case the resource is deleted and the existing resource is returned.
+	 */
+	TRefCountPtr<FShaderResource> AddOrDelete(FShaderResource* Resource)
+	{
+		TRefCountPtr<FShaderResource> OutResource(Resource);
+
+		{
+			FRWScopeLock Locker(Lock, SLT_Write);
+		
+			bool bAlreadyInSet;
+			FSetElementId ElementId = Set.Add(Resource, &bAlreadyInSet);
+			if (bAlreadyInSet)
+			{
+				return TRefCountPtr<FShaderResource>(Set[ElementId]);
+			}
+		}
+
+		return OutResource;
+	}
+
+	/** Returns true if resource was removed before someone created a new reference to it */
+	bool TryRemove(FShaderResource& Resource)
+	{
+		FRWScopeLock Locker(Lock, SLT_Write);
+
+		if (Resource.NumRefs == 0)
+		{
+			Set.Remove(&Resource);
+			return true;
+		}
+
+		// Another thread found the resource after ref-count was decremented to zero
+		return false;
+	}
+
+	void GetAllIds(TArray<FShaderResourceId>& Out) const
+	{
+		FRWScopeLock Locker(Lock, SLT_ReadOnly);
+
+		Out.Reserve(Set.Num());
+		for (FShaderResource* Resource : Set)
+		{
+			Out.Add(Resource->GetId());
+			Out.Last().MakeSelfContained();			
+		}
+	}
+
+private:
+	mutable FRWLock Lock;
+	TSet<FShaderResource*, FShaderResourceKeyFuncs> Set;
+};
+
+static FShaderResourceSet gLoadedShaderResources;
 #if RHI_RAYTRACING
 TArray<uint32> FShaderResource::GlobalUnusedIndicies;
 TArray<FRHIRayTracingShader*> FShaderResource::GlobalRayTracingMaterialLibrary;
@@ -476,12 +581,7 @@ FShaderResource::FShaderResource(const FShaderCompilerOutput& Output, FShaderTyp
 #if WITH_EDITORONLY_DATA
 	PlatformDebugData = Output.PlatformDebugData;
 #endif
-
-	{
-		check(IsInGameThread());
-		ShaderResourceIdMap.Add(GetId(), this);
-	}
-	
+		
 	INC_DWORD_STAT_BY_FName(GetMemoryStatType((EShaderFrequency)Target.Frequency).GetName(), Code.Num());
 	INC_DWORD_STAT_BY(STAT_Shaders_ShaderResourceMemory, GetSizeBytes());
 	INC_DWORD_STAT_BY(STAT_Shaders_NumShaderResourcesLoaded, 1);
@@ -610,12 +710,6 @@ void FShaderResource::CompressCode(const TArray<uint8>& UncompressedCode)
 	Code.Shrink();
 }
 
-void FShaderResource::Register()
-{
-	check(IsInGameThread());
-	ShaderResourceIdMap.Add(GetId(), this);
-}
-
 // Note: this is derived data.  Bump guid in ShaderVersion.ush if changing the format, no backwards compat is necessary
 void FShaderResource::Serialize(FArchive& Ar, bool bLoadedByCookedMaterial)
 {
@@ -746,19 +840,15 @@ void FShaderResource::SerializePlatformDebugData(FArchive& Ar)
 
 void FShaderResource::AddRef()
 {
-	checkSlow(IsInGameThread());
-	++NumRefs;
+	FPlatformAtomics::InterlockedIncrement((volatile int32*)&NumRefs);
 }
 
 
 void FShaderResource::Release()
 {
-	checkSlow(IsInGameThread());
-	check(NumRefs != 0);
-	if(--NumRefs == 0)
+	check(NumRefs > 0);
+	if (FPlatformAtomics::InterlockedDecrement((volatile int32*)&NumRefs) == 0 && gLoadedShaderResources.TryRemove(*this))
 	{
-		ShaderResourceIdMap.Remove(GetId());
-
 		// Send a release message to the rendering thread when the shader loses its last reference.
 		BeginReleaseResource(this);
 		BeginCleanup(this);
@@ -777,31 +867,38 @@ void FShaderResource::Release()
 	}
 }
 
-
-FShaderResource* FShaderResource::FindShaderResourceById(const FShaderResourceId& Id)
+TRefCountPtr<FShaderResource> FShaderResource::FindById(const FShaderResourceId& Id)
 {
-	check(IsInGameThread());
-	FShaderResource* Result = ShaderResourceIdMap.FindRef(Id);
-	return Result;
+	return gLoadedShaderResources.Find(Id);
 }
 
+TRefCountPtr<FShaderResource> FShaderResource::FindOrClone(FShaderResource&& Temp)
+{
+	if (TRefCountPtr<FShaderResource> ExistingResource = FindById(Temp.GetId()))
+	{
+		return ExistingResource;
+	}
+
+	FShaderResource* Clone = new FShaderResource(MoveTemp(Temp));
+	return gLoadedShaderResources.AddOrDelete(Clone);
+}
 
-FShaderResource* FShaderResource::FindOrCreateShaderResource(const FShaderCompilerOutput& Output, FShaderType* SpecificType, int32 SpecificPermutationId)
+TRefCountPtr<FShaderResource> FShaderResource::FindOrCreate(const FShaderCompilerOutput& Output, FShaderType* SpecificType, int32 SpecificPermutationId)
 {
 	const FShaderResourceId ResourceId(Output.Target, Output.OutputHash, SpecificType ? SpecificType->GetName() : nullptr, SpecificPermutationId);
-	FShaderResource* Resource = FindShaderResourceById(ResourceId);
-	if (!Resource)
+
+	if (TRefCountPtr<FShaderResource> ExistingResource = FindById(ResourceId))
 	{
-		Resource = new FShaderResource(Output, SpecificType, SpecificPermutationId);
+		return ExistingResource;
 	}
-	
-	return Resource;
+
+	FShaderResource* NewResource = new FShaderResource(Output, SpecificType, SpecificPermutationId);
+	return gLoadedShaderResources.AddOrDelete(NewResource);
 }
 
 void FShaderResource::GetAllShaderResourceId(TArray<FShaderResourceId>& Ids)
 {
-	check(IsInGameThread());
-	ShaderResourceIdMap.GetKeys(Ids);
+	gLoadedShaderResources.GetAllIds(Ids);
 }
 
 bool FShaderResource::ArePlatformsCompatible(EShaderPlatform CurrentPlatform, EShaderPlatform TargetPlatform)
@@ -1124,8 +1221,7 @@ FArchive& operator<<(FArchive& Ar,class FSelfContainedShaderId& Ref)
  * This still needs to initialize members to safe values since FShaderType::GenerateSerializationHistory uses this constructor.
  */
 FShader::FShader()
-	: SerializedResource(nullptr)
-	, ShaderPipeline(nullptr)
+	: ShaderPipeline(nullptr)
 	, VFType(nullptr)
 	, Type(nullptr)
 	, PermutationId(0)
@@ -1141,7 +1237,7 @@ FShader::FShader()
  */
 FShader::FShader(const CompiledShaderInitializerType& Initializer)
 	: MaterialShaderMapHash(Initializer.MaterialShaderMapHash)
-	, SerializedResource(nullptr)
+	, Resource(Initializer.Resource)
 	, ShaderPipeline(Initializer.ShaderPipeline)
 	, VFType(Initializer.VertexFactoryType)
 	, Type(Initializer.Type)
@@ -1177,7 +1273,9 @@ FShader::FShader(const CompiledShaderInitializerType& Initializer)
 		}
 	}
 
-	SetResource(Initializer.Resource);
+	
+	check(Resource && Resource->Target == Target);
+	checkSlow(FShaderResource::FindById(Resource->GetId()) == Resource);
 
 	// Register the shader now that it is valid, so that it can be reused
 	Register(false);
@@ -1233,6 +1331,8 @@ bool FShader::SerializeBase(FArchive& Ar, bool bShadersInline, bool bLoadedByCoo
 	Ar << FShaderResource::FilterShaderSourceHashForSerialization(Ar, Hash);
 	Ar << Target;
 
+	const bool SerializeStructTypeNameAsStrings = Ar.CustomVer(FFortniteMainBranchObjectVersion::GUID) < FFortniteMainBranchObjectVersion::MaterialInstanceSerializeOptimization_ShaderFName;
+
 	// TODO(RDG): Kill that once all shaders are refactored.
 	if (Ar.IsLoading())
 	{
@@ -1246,7 +1346,7 @@ bool FShader::SerializeBase(FArchive& Ar, bool bShadersInline, bool bLoadedByCoo
 		{
 			FShaderParametersMetadata* Struct = nullptr;
 
-			if (Ar.CustomVer(FFortniteMainBranchObjectVersion::GUID) < FFortniteMainBranchObjectVersion::MaterialInstanceSerializeOptimization_ShaderFName)
+			if (SerializeStructTypeNameAsStrings)
 			{
 				FString StructName;
 				Ar << StructName;
@@ -1271,17 +1371,18 @@ bool FShader::SerializeBase(FArchive& Ar, bool bShadersInline, bool bLoadedByCoo
 		int32 NumUniformParameters = UniformBufferParameters.Num();
 		Ar << NumUniformParameters;
 
-		for (int32 StructIndex = 0; StructIndex < UniformBufferParameters.Num(); StructIndex++)
+		for (int32 StructIndex = 0; StructIndex < NumUniformParameters; StructIndex++)
 		{
-			FString StructName(UniformBufferParameterStructs[StructIndex]->GetStructTypeName());
+			const TCHAR* StructName = UniformBufferParameterStructs[StructIndex]->GetStructTypeName();
 
-			if (Ar.CustomVer(FFortniteMainBranchObjectVersion::GUID) < FFortniteMainBranchObjectVersion::MaterialInstanceSerializeOptimization_ShaderFName)
+			if (SerializeStructTypeNameAsStrings)
 			{
-				Ar << StructName;
+				FString StructNameCopy(StructName);
+				Ar << StructNameCopy;
 			}
 			else
 			{
-				FName StructFName(*StructName);
+				FName StructFName(StructName);
 				Ar << StructFName;
 			}
 
@@ -1301,9 +1402,12 @@ bool FShader::SerializeBase(FArchive& Ar, bool bShadersInline, bool bLoadedByCoo
 		if (Ar.IsLoading())
 		{
 			// Load the inlined shader resource
-			SerializedResource = new FShaderResource();
-			SerializedResource->Serialize(Ar, bLoadedByCookedMaterial);
-			checkSlow(OutputHash == SerializedResource->OutputHash);
+			FShaderResource ResourceTemp;
+			ResourceTemp.Serialize(Ar, bLoadedByCookedMaterial);
+
+			Resource = FShaderResource::FindOrClone(MoveTemp(ResourceTemp));
+
+			checkSlow(OutputHash == Resource->OutputHash);
 		}
 	}
 	else
@@ -1315,8 +1419,8 @@ bool FShader::SerializeBase(FArchive& Ar, bool bShadersInline, bool bLoadedByCoo
 			FShaderResourceId ResourceId(Target, OutputHash, Type->LimitShaderResourceToThisType() ? Type->GetName() : nullptr, Type->LimitShaderResourceToThisType() ? PermutationId : 0);
 
 			// use it to look up in the registered resource map
-			FShaderResource* ExistingResource = FShaderResource::FindShaderResourceById(ResourceId);
-			SetResource(ExistingResource);
+			Resource = FShaderResource::FindById(ResourceId);
+			check(Resource && Resource->Target == Target);
 		}
 	}
 
@@ -1382,34 +1486,6 @@ FShaderId FShader::GetId() const
 	return ShaderId;
 }
 
-void FShader::RegisterSerializedResource()
-{
-	if (SerializedResource)
-	{
-		FShaderResource* ExistingResource = FShaderResource::FindShaderResourceById(SerializedResource->GetId());
-
-		// Reuse an existing shader resource if a matching one already exists in memory
-		if (ExistingResource)
-		{
-			delete SerializedResource;
-			SerializedResource = ExistingResource;
-		}
-		else
-		{
-			// Register the newly loaded shader resource so it can be reused by other shaders
-			SerializedResource->Register();
-		}
-
-		SetResource(SerializedResource);
-	}
-}
-
-void FShader::SetResource(FShaderResource* InResource)
-{
-	check(InResource && InResource->Target == Target);
-	Resource = InResource;
-}
-
 void FShader::DumpDebugInfo()
 {
 	UE_LOG(LogConsoleResponse, Display, TEXT("      FShader  :MaterialShaderMapHash %s"), *MaterialShaderMapHash.ToString());
diff --git a/Engine/Source/Runtime/RenderCore/Private/ShaderCodeLibrary.cpp b/Engine/Source/Runtime/RenderCore/Private/ShaderCodeLibrary.cpp
--- a/Engine/Source/Runtime/RenderCore/Private/ShaderCodeLibrary.cpp
+++ b/Engine/Source/Runtime/RenderCore/Private/ShaderCodeLibrary.cpp
@@ -1922,19 +1922,16 @@ public:
 	FRHIShaderLibrary* FindShaderLibrary(const FSHAHash& Hash)
 	{
 		FRWScopeLock(LibraryMutex, SLT_ReadOnly);
-		FRHIShaderLibrary* Result = nullptr;
 
 		// Search in library opened order
-		for (int32 i = 0; i < ShaderCodeArchiveStack.Num(); ++i)
+		for (FRHIShaderLibrary* ShaderCodeArchive : ShaderCodeArchiveStack)
 		{
-			FRHIShaderLibrary* ShaderCodeArchive = ShaderCodeArchiveStack[i];
 			if (ShaderCodeArchive->ContainsEntry(Hash))
 			{
-				Result = ShaderCodeArchive;
-				break;
+				return ShaderCodeArchive;
 			}
 		}
-		return Result;
+		return nullptr;
 	}
 
 	bool ContainsShaderCode(const FSHAHash& Hash)
diff --git a/Engine/Source/Runtime/RenderCore/Public/Shader.h b/Engine/Source/Runtime/RenderCore/Public/Shader.h
--- a/Engine/Source/Runtime/RenderCore/Public/Shader.h
+++ b/Engine/Source/Runtime/RenderCore/Public/Shader.h
@@ -166,6 +166,16 @@ public:
 		return Ar;
 	}
 
+	/** Copy string data that belongs to FShaderResource for cases where this id can outlive the resource. */
+	void MakeSelfContained()
+	{
+		if (SpecificShaderTypeName && SpecificShaderTypeName != *SpecificShaderTypeStorage)
+		{
+			SpecificShaderTypeStorage = SpecificShaderTypeName;
+			SpecificShaderTypeName = *SpecificShaderTypeStorage;
+		}
+	}
+
 	/** Hash of the compiled shader output, which is used to create the FShaderResource. */
 	FSHAHash OutputHash;
 
@@ -276,15 +286,17 @@ public:
 class FShaderResource : public FRenderResource, public FDeferredCleanupInterface
 {
 	friend class FShader;
+	friend class FShaderResourceSet;
 public:
 
 	/** Constructor used for deserialization. */
 	RENDERCORE_API FShaderResource();
+	FShaderResource(FShaderResource&&) = default;
 
 	/** Constructor used when creating a new shader resource from compiled output. */
 	FShaderResource(const FShaderCompilerOutput& Output, FShaderType* InSpecificType, int32 InSpecificPermutationId);
 
-	~FShaderResource();
+	RENDERCORE_API ~FShaderResource();
 
 	RENDERCORE_API void Serialize(FArchive& Ar, bool bLoadedByCookedMaterial);
 
@@ -292,8 +304,6 @@ public:
 	RENDERCORE_API void AddRef();
 	RENDERCORE_API void Release();
 
-	RENDERCORE_API void Register();
-
 	/** @return the shader's vertex shader */
 	FORCEINLINE FRHIVertexShader* GetVertexShader()
 	{
@@ -395,6 +405,11 @@ public:
 #endif // RHI_RAYTRACING
 
 	RENDERCORE_API FShaderResourceId GetId() const;
+	
+	const FSHAHash& GetOutputHash() const
+	{
+		return OutputHash;
+	}
 
 	uint32 GetSizeBytes() const
 	{
@@ -405,14 +420,32 @@ public:
 	virtual void InitRHI();
 	virtual void ReleaseRHI();
 
+
+	template<class T>
+	static FShaderResource* FindShaderResourceById(T Id)
+	{
+		static_assert(sizeof(T) == 0, "Deprecated. Please save return value as TRefCountPtr<FShaderResource> and switch to FindById().");
+		return nullptr;
+	}
+
 	/** Finds a matching shader resource in memory if possible. */
-	RENDERCORE_API static FShaderResource* FindShaderResourceById(const FShaderResourceId& Id);
+	RENDERCORE_API static TRefCountPtr<FShaderResource> FindById(const FShaderResourceId& Id);
+
+	/** Finds a matching shader resource in memory or clone the temporary and register the new resource */
+	RENDERCORE_API static TRefCountPtr<FShaderResource> FindOrClone(FShaderResource&& Temp);
+
+	template<class T>
+	static FShaderResource* FindOrCreateShaderResource(const FShaderCompilerOutput& Output, class FShaderType* SpecificType, T SpecificPermutationId)
+	{
+		static_assert(sizeof(T) == 0, "Deprecated. Please save return value as TRefCountPtr<FShaderResource> and switch to FindById().");
+		return nullptr;
+	}
 
 	/** 
 	 * Finds a matching shader resource in memory or creates a new one with the given compiler output.  
 	 * SpecificType can be NULL
 	 */
-	RENDERCORE_API static FShaderResource* FindOrCreateShaderResource(const FShaderCompilerOutput& Output, class FShaderType* SpecificType, int32 SpecificPermutationId);
+	RENDERCORE_API static TRefCountPtr<FShaderResource> FindOrCreate(const FShaderCompilerOutput& Output, class FShaderType* SpecificType, int32 SpecificPermutationId);
 
 	/** Return a list of all shader Ids currently known */
 	RENDERCORE_API static void GetAllShaderResourceId(TArray<FShaderResourceId>& Ids);
@@ -480,7 +513,7 @@ private:
 	int32 SpecificPermutationId;
 
 	/** The number of references to this shader. */
-	mutable uint32 NumRefs;
+	mutable int32 NumRefs;
 
 	/** The number of instructions the shader takes to execute. */
 	uint32 NumInstructions;
@@ -501,8 +534,6 @@ private:
 
 	void BuildParameterMapInfo(const TMap<FString, FParameterAllocation>& ParameterMap);
 
-	/** Tracks loaded shader resources by id. */
-	static TMap<FShaderResourceId, FShaderResource*> ShaderResourceIdMap;
 	/** Critical section for ShaderResourceIdMap. */
 	static FCriticalSection ShaderResourceIdMapCritical;
 };
@@ -819,7 +850,7 @@ public:
 		const TArray<uint8>& Code;
 		const FShaderParameterMap& ParameterMap;
 		const FSHAHash& OutputHash;
-		FShaderResource* Resource;
+		TRefCountPtr<FShaderResource> Resource;
 		FSHAHash MaterialShaderMapHash;
 		const FShaderPipelineType* ShaderPipeline;
 		FVertexFactoryType* VertexFactoryType;
@@ -829,7 +860,7 @@ public:
 			FShaderType* InType,
 			int32 InPermutationId,
 			const FShaderCompilerOutput& CompilerOutput,
-			FShaderResource* InResource,
+			TRefCountPtr<FShaderResource>&& InResource,
 			const FSHAHash& InMaterialShaderMapHash,
 			const FShaderPipelineType* InShaderPipeline,
 			FVertexFactoryType* InVertexFactoryType
@@ -985,9 +1016,6 @@ public:
 
 	void SetResource(FShaderResource* InResource);
 
-	/** Called from the main thread to register and set the serialized resource */
-	void RegisterSerializedResource();
-
 	/** Implement for geometry shaders that want to use stream out. */
 	static void GetStreamOutElements(FStreamOutElementList& ElementList, TArray<uint32>& StreamStrides, int32& RasterizedStream) {}
 
@@ -1052,13 +1080,6 @@ public:
 		return this;
 	}
 
-	/** Discards the serialized resource, used when the engine is using NullRHI */
-	void DiscardSerializedResource()
-	{
-		delete SerializedResource;
-		SerializedResource = nullptr;
-	}
-
 	void DumpDebugInfo();
 	void SaveShaderStableKeys(EShaderPlatform TargetShaderPlatform, const struct FStableShaderKeyAndValue& SaveKeyVal);
 
@@ -1095,9 +1116,6 @@ private:
 	/** Reference to the shader resource, which stores the compiled bytecode and the RHI shader resource. */
 	TRefCountPtr<FShaderResource> Resource;
 
-	/** Pointer to the shader resource that has been serialized from disk, to be registered on the main thread later. */
-	FShaderResource* SerializedResource;
-
 	/** Shader pipeline this shader belongs to, stored so that an FShaderId can be constructed from this shader. */
 	const FShaderPipelineType* ShaderPipeline;
 
@@ -2272,8 +2290,6 @@ public:
 		check(IsInGameThread());
 		for (FShader* Shader : SerializedShaders)
 		{
-			Shader->RegisterSerializedResource();
-
 			FShaderType* Type = Shader->GetType();
 			FShader* ExistingShader = Type->FindShaderById(Shader->GetId());
 
@@ -2293,10 +2309,6 @@ public:
 
 		for (FSerializedShaderPipeline* SerializedPipeline : SerializedShaderPipelines)
 		{
-			for (TRefCountPtr<FShader> Shader : SerializedPipeline->ShaderStages)
-			{
-				Shader->RegisterSerializedResource();
-			}
 			FShaderPipeline* ShaderPipeline = new FShaderPipeline(SerializedPipeline->ShaderPipelineType, SerializedPipeline->ShaderStages);
 			AddShaderPipeline(SerializedPipeline->ShaderPipelineType, ShaderPipeline);
 
@@ -2310,20 +2322,12 @@ public:
 	{
 		for (FShader* Shader : SerializedShaders)
 		{
-			if (Shader)
-			{
-				Shader->DiscardSerializedResource();
-			}
 			delete Shader;
 		}
 		SerializedShaders.Empty();
 
 		for (FSerializedShaderPipeline* SerializedPipeline : SerializedShaderPipelines)
 		{
-			for (TRefCountPtr<FShader> Shader : SerializedPipeline->ShaderStages)
-			{
-				Shader->DiscardSerializedResource();
-			}
 			delete SerializedPipeline;
 		}
 		SerializedShaderPipelines.Empty();
