Fixed compile error on Linux due to non-portable header include file path

diff --git a/Engine/Source/Runtime/RenderCore/Private/Shader.cpp b/Engine/Source/Runtime/RenderCore/Private/Shader.cpp
--- a/Engine/Source/Runtime/RenderCore/Private/Shader.cpp
+++ b/Engine/Source/Runtime/RenderCore/Private/Shader.cpp
@@ -19,7 +19,7 @@
 #include "Misc/ScopeLock.h"
 #include "UObject/RenderingObjectVersion.h"
 #include "UObject/FortniteMainBranchObjectVersion.h"
-#include "Misc/ScopeRwLock.h"
+#include "Misc/ScopeRWLock.h"
 
 #if WITH_EDITORONLY_DATA
 #include "Interfaces/IShaderFormat.h"
