Added RDG_{TEXTURE, BUFFER}_COPY_DEST usage to pass parameters.

diff --git a/Engine/Source/Runtime/Engine/Private/ShaderCompiler/ShaderCompiler.cpp b/Engine/Source/Runtime/Engine/Private/ShaderCompiler/ShaderCompiler.cpp
--- a/Engine/Source/Runtime/Engine/Private/ShaderCompiler/ShaderCompiler.cpp
+++ b/Engine/Source/Runtime/Engine/Private/ShaderCompiler/ShaderCompiler.cpp
@@ -2831,6 +2831,12 @@ static void PullRootShaderParametersLayout(FShaderCompilerInput& CompileInput, c
 		{
 			// RHI don't need to care about render target bindings slot anyway.
 		}
+		else if (
+			BaseType == UBMT_RDG_BUFFER_COPY_DEST ||
+			BaseType == UBMT_RDG_TEXTURE_COPY_DEST)
+		{
+			// Shaders don't care about copy destination parameters.
+		}
 		else if (
 			BaseType == UBMT_RDG_BUFFER_UAV ||
 			BaseType == UBMT_RDG_TEXTURE_UAV)
diff --git a/Engine/Source/Runtime/RHI/Public/RHIDefinitions.h b/Engine/Source/Runtime/RHI/Public/RHIDefinitions.h
--- a/Engine/Source/Runtime/RHI/Public/RHIDefinitions.h
+++ b/Engine/Source/Runtime/RHI/Public/RHIDefinitions.h
@@ -487,9 +487,11 @@ enum EUniformBufferBaseType : uint8
 	UBMT_RDG_TEXTURE,
 	UBMT_RDG_TEXTURE_SRV,
 	UBMT_RDG_TEXTURE_UAV,
+	UBMT_RDG_TEXTURE_COPY_DEST,
 	UBMT_RDG_BUFFER,
 	UBMT_RDG_BUFFER_SRV,
 	UBMT_RDG_BUFFER_UAV,
+	UBMT_RDG_BUFFER_COPY_DEST,
 
 	// Nested structure.
 	UBMT_NESTED_STRUCT,
@@ -1249,9 +1251,11 @@ inline bool IsRDGResourceReferenceShaderParameterType(EUniformBufferBaseType Bas
 		BaseType == UBMT_RDG_TEXTURE ||
 		BaseType == UBMT_RDG_TEXTURE_SRV ||
 		BaseType == UBMT_RDG_TEXTURE_UAV ||
+		BaseType == UBMT_RDG_TEXTURE_COPY_DEST ||
 		BaseType == UBMT_RDG_BUFFER ||
 		BaseType == UBMT_RDG_BUFFER_SRV ||
-		BaseType == UBMT_RDG_BUFFER_UAV;
+		BaseType == UBMT_RDG_BUFFER_UAV ||
+		BaseType == UBMT_RDG_BUFFER_COPY_DEST;
 }
 
 /** Returns whether the shader parameter type needs to be passdown to RHI through FRHIUniformBufferLayout when creating an uniform buffer. */
@@ -1276,7 +1280,11 @@ inline bool IsShaderParameterTypeIgnoredByRHI(EUniformBufferBaseType BaseType)
 {
 	return
 		// Render targets bindings slots needs to be in FRHIUniformBufferLayout for render graph, but the RHI does not actually need to know about it.
-		BaseType == UBMT_RENDER_TARGET_BINDING_SLOTS || 
+		BaseType == UBMT_RENDER_TARGET_BINDING_SLOTS ||
+
+		// Copy destination states are used by the render graph.
+		BaseType == UBMT_RDG_TEXTURE_COPY_DEST ||
+		BaseType == UBMT_RDG_BUFFER_COPY_DEST ||
 
 		// #yuriy_todo: RHI is able to dereference uniform buffer in root shader parameter structures
 		BaseType == UBMT_REFERENCED_STRUCT;
diff --git a/Engine/Source/Runtime/RenderCore/Private/RenderGraphBuilder.cpp b/Engine/Source/Runtime/RenderCore/Private/RenderGraphBuilder.cpp
--- a/Engine/Source/Runtime/RenderCore/Private/RenderGraphBuilder.cpp
+++ b/Engine/Source/Runtime/RenderCore/Private/RenderGraphBuilder.cpp
@@ -292,6 +292,15 @@ void FRDGBuilder::ValidatePass(const FRDGPass* Pass) const
 			}
 		}
 		break;
+		case UBMT_RDG_TEXTURE_COPY_DEST:
+		case UBMT_RDG_BUFFER_COPY_DEST:
+		{
+			if (FRDGTrackedResourceRef Resource = Parameter.GetAsTrackedResource())
+			{
+				Resource->MarkAsProducedBy(Pass);
+			}
+		}
+		break;
 		case UBMT_RENDER_TARGET_BINDING_SLOTS:
 		{
 			if (!RenderTargetBindingSlots)
@@ -338,18 +347,36 @@ void FRDGBuilder::ValidatePass(const FRDGPass* Pass) const
 
 				if (FRDGTextureRef Texture = RenderTarget.GetTexture())
 				{
-					/** Validate that load action is correct. */
 					const bool bIsLoadAction = RenderTarget.GetLoadAction() == ERenderTargetLoadAction::ELoad;
-					const bool bIsLoadActionInvalid = bIsLoadAction && !Texture->HasBeenProduced();
-					checkf(
-						!bIsLoadActionInvalid,
-						TEXT("Can't load a render target '%s' that has never been produced."),
-						Texture->Name);
-					
-					const bool bOverwritesResourceEntirely = !bIsLoadAction && Texture->HasBeenProduced() && Texture->Desc.NumMips == 1;
-					ensureMsgf(!bOverwritesResourceEntirely,
-						TEXT("Clobbering render target %s result that has been produced previously by another pass. Should instead render to a new FRDGTexture to reduce memory pressure."),
-						Texture->Name);
+
+					/** Validate that load action is correct. We can only load contents if a pass previously produced something. */
+					{
+						const bool bIsLoadActionInvalid = bIsLoadAction && !Texture->HasBeenProduced();
+						checkf(
+							!bIsLoadActionInvalid,
+							TEXT("Pass '%s' attempted to bind texture '%s' as a render target with the 'Load' action specified, but the texture has not been produced yet. The render target must use either 'Clear' or 'NoAction' action instead."),
+							Pass->GetName(),
+							Texture->Name);
+					}
+
+					/** Validate that any previously produced texture contents are loaded. This occurs if the user failed to specify a load action
+					 *  on a texture that was produced by a previous pass, effectively losing that data. This can also happen if the user 're-uses'
+					 *  a texture for some other purpose. The latter is considered bad practice, since it increases memory pressure on the render
+					 *  target pool. Instead, the user should create a new texture instance. An exception to this rule are untracked render targets,
+					 *  which are not actually managed by the render target pool and likely represent the frame buffer.
+					 */
+					{
+						// We only validate single-mip textures since we don't track production at the subresource level.
+						const bool bFailedToLoadProducedContent = !bIsLoadAction && Texture->HasBeenProduced() && Texture->Desc.NumMips == 1;
+
+						// Untracked render targets aren't actually managed by the render target pool.
+						const bool bIsUntrackedRenderTarget = Texture->PooledRenderTarget && Texture->PooledRenderTarget->IsTracked();
+
+						ensureMsgf(!bFailedToLoadProducedContent || bIsUntrackedRenderTarget,
+							TEXT("Pass '%s' attempted to bind texture '%s' as a render target without the 'Load' action specified, despite a prior pass having produced it. It's invalid to completely clobber the contents of a resource. Create a new texture instance instead."),
+							Pass->GetName(),
+							Texture->Name);
+					}
 					
 					/** Mark the pass as a producer for render targets with a store action. */
 					{
@@ -525,7 +552,9 @@ void FRDGBuilder::WalkGraphDependencies()
 			switch (Parameter.GetType())
 			{
 			case UBMT_RDG_TEXTURE:
+			case UBMT_RDG_TEXTURE_COPY_DEST:
 			case UBMT_RDG_BUFFER:
+			case UBMT_RDG_BUFFER_COPY_DEST:
 			{
 				if (FRDGTrackedResourceRef Resource = Parameter.GetAsTrackedResource())
 				{
@@ -917,6 +946,20 @@ void FRDGBuilder::PrepareResourcesForExecute(const FRDGPass* Pass, struct FRHIRe
 			}
 		}
 		break;
+		case UBMT_RDG_TEXTURE_COPY_DEST:
+		{
+			if (FRDGTextureRef Texture = Parameter.GetAsTexture())
+			{
+				#if RDG_ENABLE_DEBUG
+				{
+					Texture->PassAccessCount++;
+				}
+				#endif
+
+				AllocateRHITextureIfNeeded(Texture);
+			}
+		}
+		break;
 		case UBMT_RDG_BUFFER:
 		{
 			if (FRDGBufferRef Buffer = Parameter.GetAsBuffer())
@@ -988,6 +1031,20 @@ void FRDGBuilder::PrepareResourcesForExecute(const FRDGPass* Pass, struct FRHIRe
 			}
 		}
 		break;
+		case UBMT_RDG_BUFFER_COPY_DEST:
+		{
+			if (FRDGBufferRef Buffer = Parameter.GetAsBuffer())
+			{
+				#if RDG_ENABLE_DEBUG
+				{
+					Buffer->PassAccessCount++;
+				}
+				#endif
+
+				AllocateRHIBufferIfNeeded(Buffer);
+			}
+		}
+		break;
 		case UBMT_RENDER_TARGET_BINDING_SLOTS:
 		{
 			check(!Pass->IsCompute());
@@ -1258,6 +1315,7 @@ void FRDGBuilder::ReleaseUnreferencedResources(const FRDGPass* Pass)
 		switch (Parameter.GetType())
 		{
 		case UBMT_RDG_TEXTURE:
+		case UBMT_RDG_TEXTURE_COPY_DEST:
 		{
 			if (FRDGTextureRef Texture = Parameter.GetAsTexture())
 			{
@@ -1282,6 +1340,7 @@ void FRDGBuilder::ReleaseUnreferencedResources(const FRDGPass* Pass)
 		}
 		break;
 		case UBMT_RDG_BUFFER:
+		case UBMT_RDG_BUFFER_COPY_DEST:
 		{
 			if (FRDGBufferRef Buffer = Parameter.GetAsBuffer())
 			{
diff --git a/Engine/Source/Runtime/RenderCore/Private/RenderGraphUtils.cpp b/Engine/Source/Runtime/RenderCore/Private/RenderGraphUtils.cpp
--- a/Engine/Source/Runtime/RenderCore/Private/RenderGraphUtils.cpp
+++ b/Engine/Source/Runtime/RenderCore/Private/RenderGraphUtils.cpp
@@ -3,6 +3,14 @@
 #include "RenderGraphUtils.h"
 #include <initializer_list>
 
+ /** Adds a render graph tracked buffer suitable for use as a copy destination. */
+#define RDG_BUFFER_COPY_DEST(MemberName) \
+	INTERNAL_SHADER_PARAMETER_EXPLICIT(UBMT_RDG_BUFFER_COPY_DEST, TShaderResourceParameterTypeInfo<FRDGBufferRef>, FRDGBufferRef,MemberName,, = nullptr,EShaderPrecisionModifier::Float,TEXT(""),false)
+
+  /** Adds a render graph tracked texture suitable for use as a copy destination. */
+#define RDG_TEXTURE_COPY_DEST(MemberName) \
+	INTERNAL_SHADER_PARAMETER_EXPLICIT(UBMT_RDG_TEXTURE_COPY_DEST, TShaderResourceParameterTypeInfo<FRDGTextureRef>, FRDGTextureRef,MemberName,, = nullptr,EShaderPrecisionModifier::Float,TEXT(""),false)
+
 void ClearUnusedGraphResourcesImpl(
 	const FShaderParameterBindings& ShaderBindings,
 	const FShaderParametersMetadata* ParametersMetadata,
@@ -81,3 +89,28 @@ FRDGTextureRef RegisterExternalTextureWithFallback(
 	}
 }
 
+BEGIN_SHADER_PARAMETER_STRUCT(FCopyTextureParameters, )
+	SHADER_PARAMETER_RDG_TEXTURE(, Input)
+	RDG_TEXTURE_COPY_DEST(Output)
+END_SHADER_PARAMETER_STRUCT()
+
+void AddCopyTexturePass(
+	FRDGBuilder& GraphBuilder,
+	FRDGTextureRef InputTexture,
+	FRDGTextureRef OutputTexture,
+	const FRHICopyTextureInfo& CopyInfo)
+{
+	const FRDGTextureDesc& InputDesc = InputTexture->Desc;
+	const FRDGTextureDesc& OutputDesc = OutputTexture->Desc;
+	checkf(InputDesc.Format == OutputDesc.Format, TEXT("This method does not support format conversion."));
+
+	FCopyTextureParameters* Parameters = GraphBuilder.AllocParameters<FCopyTextureParameters>();
+	Parameters->Input = InputTexture;
+	Parameters->Output = OutputTexture;
+
+	GraphBuilder.AddPass(RDG_EVENT_NAME("CopyTexture"), Parameters, ERDGPassFlags::Copy,
+		[InputTexture, OutputTexture, CopyInfo](FRHICommandList& RHICmdList)
+	{
+		RHICmdList.CopyTexture(InputTexture->GetRHI(), OutputTexture->GetRHI(), CopyInfo);
+	});
+}
diff --git a/Engine/Source/Runtime/RenderCore/Private/ShaderParameterStruct.cpp b/Engine/Source/Runtime/RenderCore/Private/ShaderParameterStruct.cpp
--- a/Engine/Source/Runtime/RenderCore/Private/ShaderParameterStruct.cpp
+++ b/Engine/Source/Runtime/RenderCore/Private/ShaderParameterStruct.cpp
@@ -67,7 +67,13 @@ struct FShaderParameterStructBindingContext
 				BaseType == UBMT_SRV ||
 				BaseType == UBMT_UAV ||
 				BaseType == UBMT_SAMPLER);
-			const bool bIsRDGResource = IsRDGResourceReferenceShaderParameterType(BaseType) && BaseType != UBMT_RDG_BUFFER;
+
+			const bool bIsRDGResource =
+				IsRDGResourceReferenceShaderParameterType(BaseType) &&
+				BaseType != UBMT_RDG_BUFFER &&
+				BaseType != UBMT_RDG_BUFFER_COPY_DEST &&
+				BaseType != UBMT_RDG_TEXTURE_COPY_DEST;
+
 			const bool bIsVariableNativeType = (
 				BaseType == UBMT_BOOL ||
 				BaseType == UBMT_INT32 ||
diff --git a/Engine/Source/Runtime/RenderCore/Private/ShaderParameters.cpp b/Engine/Source/Runtime/RenderCore/Private/ShaderParameters.cpp
--- a/Engine/Source/Runtime/RenderCore/Private/ShaderParameters.cpp
+++ b/Engine/Source/Runtime/RenderCore/Private/ShaderParameters.cpp
@@ -263,6 +263,7 @@ static void CreateHLSLUniformBufferStructMembersDeclaration(
 		if (IsShaderParameterTypeForUniformBufferLayout(Member.GetBaseType()))
 		{
 			check(Member.GetBaseType() != UBMT_RDG_TEXTURE_SRV && Member.GetBaseType() != UBMT_RDG_TEXTURE_UAV);
+			checkf(Member.GetBaseType() != UBMT_RDG_TEXTURE_COPY_DEST && Member.GetBaseType() != UBMT_RDG_TEXTURE_COPY_DEST, TEXT("Copy destination usage is not supported in uniform buffers."));
 			checkf(Member.GetNumElements() == 0, TEXT("Resources array are not supported in uniform buffers yet."));
 			if (Member.GetBaseType() == UBMT_SRV)
 			{
diff --git a/Engine/Source/Runtime/RenderCore/Public/RenderGraphPass.h b/Engine/Source/Runtime/RenderCore/Public/RenderGraphPass.h
--- a/Engine/Source/Runtime/RenderCore/Public/RenderGraphPass.h
+++ b/Engine/Source/Runtime/RenderCore/Public/RenderGraphPass.h
@@ -24,7 +24,11 @@ public:
 
 	bool IsTrackedResource() const
 	{
-		return MemberType == UBMT_RDG_TEXTURE || MemberType == UBMT_RDG_BUFFER;
+		return
+			MemberType == UBMT_RDG_TEXTURE ||
+			MemberType == UBMT_RDG_TEXTURE_COPY_DEST ||
+			MemberType == UBMT_RDG_BUFFER ||
+			MemberType == UBMT_RDG_BUFFER_COPY_DEST;
 	}
 
 	EUniformBufferBaseType GetType() const
@@ -46,13 +50,13 @@ public:
 
 	FRDGTextureRef GetAsTexture() const
 	{
-		check(MemberType == UBMT_RDG_TEXTURE);
+		check(MemberType == UBMT_RDG_TEXTURE || MemberType == UBMT_RDG_TEXTURE_COPY_DEST);
 		return *GetAs<FRDGTextureRef>();
 	}
 
 	FRDGBufferRef GetAsBuffer() const
 	{
-		check(MemberType == UBMT_RDG_BUFFER);
+		check(MemberType == UBMT_RDG_BUFFER || MemberType == UBMT_RDG_BUFFER_COPY_DEST);
 		return *GetAs<FRDGBufferRef>();
 	}
 
@@ -163,15 +167,16 @@ enum class ERDGPassFlags
 	None = 0,
 
 	/** Pass uses compute only */
-	Compute = 1 << 0,
+	Compute,
+
+	/** Pass uses RHI copy commands only. */
+	Copy,
 
 	//#todo-rco: Remove this when we can do split/per mip layout transitions.
 	/** Hint to some RHIs this pass will be generating mips to optimize transitions. */
-	GenerateMips = 1 << 1,
+	GenerateMips
 };
 
-ENUM_CLASS_FLAGS(ERDGPassFlags)
-
 // TODO(RDG): Bulk rename across codebase.
 using ERenderGraphPassFlags = ERDGPassFlags;
 
@@ -200,12 +205,12 @@ public:
 
 	bool IsGenerateMips() const
 	{
-		return (PassFlags & ERDGPassFlags::GenerateMips) == ERDGPassFlags::GenerateMips;
+		return (PassFlags == ERDGPassFlags::GenerateMips);
 	}
 
 	bool IsCompute() const
 	{
-		return (PassFlags & ERDGPassFlags::Compute) == ERDGPassFlags::Compute;
+		return (PassFlags == ERDGPassFlags::Compute);
 	}
 
 	FRDGPassParameterStruct GetParameters() const
diff --git a/Engine/Source/Runtime/RenderCore/Public/RenderGraphUtils.h b/Engine/Source/Runtime/RenderCore/Public/RenderGraphUtils.h
--- a/Engine/Source/Runtime/RenderCore/Public/RenderGraphUtils.h
+++ b/Engine/Source/Runtime/RenderCore/Public/RenderGraphUtils.h
@@ -181,4 +181,40 @@ struct RENDERCORE_API FComputeShaderUtils
 			FComputeShaderUtils::DispatchIndirect(RHICmdList, ComputeShader, *Parameters, IndirectArgsBuffer->GetIndirectRHICallBuffer(), IndirectArgOffset);
 		});
 	}
-};
\ No newline at end of file
+};
+
+/** Adds a render graph pass that copies a region from one texture to another;
+ *  Uses RHICopyTexture under the hood. Formats of the two textures must match.
+ *  The output and output texture regions be within the respective extents.
+ */
+RENDERCORE_API void AddCopyTexturePass(
+	FRDGBuilder& GraphBuilder,
+	FRDGTextureRef InputTexture,
+	FRDGTextureRef OutputTexture,
+	const FRHICopyTextureInfo& CopyInfo);
+
+/** Simpler variant of the above function for 2D textures.
+ *  @param InputPosition The pixel position within the input texture of the top-left corner of the box.
+ *  @param OutputPosition The pixel position within the output texture of the top-left corner of the box.
+ *  @param Size The size in pixels of the region to copy from input to output. If zero, the full extent of
+ *         the input texture is copied.
+ */
+inline void AddCopyTexturePass(
+	FRDGBuilder& GraphBuilder,
+	FRDGTextureRef InputTexture,
+	FRDGTextureRef OutputTexture,
+	FIntPoint InputPosition = FIntPoint::ZeroValue,
+	FIntPoint OutputPosition = FIntPoint::ZeroValue,
+	FIntPoint Size = FIntPoint::ZeroValue)
+{
+	FRHICopyTextureInfo CopyInfo;
+	CopyInfo.SourcePosition.X = InputPosition.X;
+	CopyInfo.SourcePosition.Y = InputPosition.Y;
+	CopyInfo.DestPosition.X = OutputPosition.X;
+	CopyInfo.DestPosition.Y = OutputPosition.Y;
+	if (Size != FIntPoint::ZeroValue)
+	{
+		CopyInfo.Size = FIntVector(Size.X, Size.Y, 1);
+	}
+	AddCopyTexturePass(GraphBuilder, InputTexture, OutputTexture, CopyInfo);
+}
diff --git a/Engine/Source/Runtime/RenderCore/Public/RenderTargetPool.h b/Engine/Source/Runtime/RenderCore/Public/RenderTargetPool.h
--- a/Engine/Source/Runtime/RenderCore/Public/RenderTargetPool.h
+++ b/Engine/Source/Runtime/RenderCore/Public/RenderTargetPool.h
@@ -68,6 +68,11 @@ struct RENDERCORE_API FPooledRenderTarget : public IPooledRenderTarget
 	{
 		return !!(Desc.Flags & TexCreate_Transient);
 	}
+	bool IsTracked() const override
+	{
+		return RenderTargetPool != nullptr;
+	}
+
 	virtual void SetDebugName(const TCHAR *InName);
 	virtual const FPooledRenderTargetDesc& GetDesc() const;
 	virtual uint32 ComputeMemorySize() const;
diff --git a/Engine/Source/Runtime/RenderCore/Public/RendererInterface.h b/Engine/Source/Runtime/RenderCore/Public/RendererInterface.h
--- a/Engine/Source/Runtime/RenderCore/Public/RendererInterface.h
+++ b/Engine/Source/Runtime/RenderCore/Public/RendererInterface.h
@@ -448,6 +448,8 @@ struct IPooledRenderTarget
 	inline FSceneRenderTargetItem& GetRenderTargetItem() { return RenderTargetItem; }
 	/** Get the low level internals (texture/surface) */
 	inline const FSceneRenderTargetItem& GetRenderTargetItem() const { return RenderTargetItem; }
+	/** Returns if the render target is tracked by a pool. */
+	virtual bool IsTracked() const = 0;
 
 	// Refcounting
 	virtual uint32 AddRef() const = 0;
diff --git a/Engine/Source/Runtime/RenderCore/Public/ShaderParameterStruct.h b/Engine/Source/Runtime/RenderCore/Public/ShaderParameterStruct.h
--- a/Engine/Source/Runtime/RenderCore/Public/ShaderParameterStruct.h
+++ b/Engine/Source/Runtime/RenderCore/Public/ShaderParameterStruct.h
@@ -12,7 +12,7 @@
 #include "RenderGraphResources.h"
 
 
-/** Instrument a shader class to use strutured shader parameters API.
+/** Tag a shader class to use the structured shader parameters API.
  *
  * class FMyShaderClassCS : public FGlobalShader
  * {
@@ -25,10 +25,10 @@
  *		END_SHADER_PARAMETER_STRUCT()
  * };
  *
- * Notes: Long therm, this macro will be no longer be needed.
+ * Notes: Long term, this macro will no longer be needed. Instead, parameter binding will become the default behavior for shader declarations.
  */
- // TODO(RDG): would not even need ShaderParentClass anymore. And in fact should not so Bindings.Bind() is not being called twice.
-#define SHADER_USE_PARAMETER_STRUCT(ShaderClass, ShaderParentClass) \
+
+#define SHADER_USE_PARAMETER_STRUCT_INTERNAL(ShaderClass, ShaderParentClass, bShouldBindEverything) \
 	ShaderClass(const ShaderMetaType::CompiledShaderInitializerType& Initializer) \
 		: ShaderParentClass(Initializer) \
 	{ \
@@ -36,12 +36,22 @@
 			this, \
 			Initializer.ParameterMap, \
 			*FParameters::FTypeInfo::GetStructMetadata(), \
-			true); \
+			bShouldBindEverything); \
 	} \
 	\
 	ShaderClass() \
 	{ } \
 
+// TODO(RDG): would not even need ShaderParentClass anymore. And in fact should not so Bindings.Bind() is not being called twice.
+#define SHADER_USE_PARAMETER_STRUCT(ShaderClass, ShaderParentClass) \
+	SHADER_USE_PARAMETER_STRUCT_INTERNAL(ShaderClass, ShaderParentClass, true)
+
+/** Use when sharing shader parameter binding with legacy parameters in the base class; i.e. FMaterialShader or FMeshMaterialShader.
+ *  Note that this disables validation that the parameter struct contains all shader bindings.
+ */
+#define SHADER_USE_PARAMETER_STRUCT_WITH_LEGACY_BASE(ShaderClass, ShaderParentClass) \
+	SHADER_USE_PARAMETER_STRUCT_INTERNAL(ShaderClass, ShaderParentClass, false)
+
 #define SHADER_USE_ROOT_PARAMETER_STRUCT(ShaderClass, ShaderParentClass) \
 	static inline const FShaderParametersMetadata* GetRootParametersMetadata() { return FParameters::FTypeInfo::GetStructMetadata(); } \
 	\
