Multiple external resource registrations now de-duplicate to the same RDG resource.

diff --git a/Engine/Source/Runtime/RenderCore/Private/RenderGraphBuilder.cpp b/Engine/Source/Runtime/RenderCore/Private/RenderGraphBuilder.cpp
--- a/Engine/Source/Runtime/RenderCore/Private/RenderGraphBuilder.cpp
+++ b/Engine/Source/Runtime/RenderCore/Private/RenderGraphBuilder.cpp
@@ -1489,6 +1489,8 @@ void FRDGBuilder::DestructPasses()
 
 	DeferredInternalTextureQueries.Empty();
 	DeferredInternalBufferQueries.Empty();
+	ExternalTextures.Empty();
+	ExternalBuffers.Empty();
 	AllocatedTextures.Empty();
 	AllocatedBuffers.Empty();
 }
diff --git a/Engine/Source/Runtime/RenderCore/Public/RenderGraphBuilder.h b/Engine/Source/Runtime/RenderCore/Public/RenderGraphBuilder.h
--- a/Engine/Source/Runtime/RenderCore/Public/RenderGraphBuilder.h
+++ b/Engine/Source/Runtime/RenderCore/Public/RenderGraphBuilder.h
@@ -151,6 +151,10 @@ private:
 	/** Keep the references over the pooled render target, since FRDGTexture is allocated on MemStack. */
 	TMap<FRDGBuffer*, TRefCountPtr<FPooledRDGBuffer>, SceneRenderingSetAllocator> AllocatedBuffers;
 
+	/** Tracks external resources to their registered render graph counterparts for de-duplication. */
+	TMap<const IPooledRenderTarget*, FRDGTexture*, SceneRenderingSetAllocator> ExternalTextures;
+	TMap<const FPooledRDGBuffer*, FRDGBuffer*, SceneRenderingSetAllocator> ExternalBuffers;
+
 	/** Array of all deferred access to internal textures. */
 	struct FDeferredInternalTextureQuery
 	{
diff --git a/Engine/Source/Runtime/RenderCore/Public/RenderGraphBuilder.inl b/Engine/Source/Runtime/RenderCore/Public/RenderGraphBuilder.inl
--- a/Engine/Source/Runtime/RenderCore/Public/RenderGraphBuilder.inl
+++ b/Engine/Source/Runtime/RenderCore/Public/RenderGraphBuilder.inl
@@ -11,6 +11,11 @@ inline FRDGTextureRef FRDGBuilder::RegisterExternalTexture(const TRefCountPtr<IP
 	}
 #endif
 
+	if (FRDGTextureRef* Texture = ExternalTextures.Find(ExternalPooledTexture.GetReference()))
+	{
+		return *Texture;
+	}
+
 	FRDGTexture* OutTexture = AllocateForRHILifeTime<FRDGTexture>(Name, ExternalPooledTexture->GetDesc(), ERDGResourceFlags::None);
 	OutTexture->PooledRenderTarget = ExternalPooledTexture;
 	OutTexture->ResourceRHI = ExternalPooledTexture->GetRenderTargetItem().ShaderResourceTexture;
@@ -23,6 +28,8 @@ inline FRDGTextureRef FRDGBuilder::RegisterExternalTexture(const TRefCountPtr<IP
 	}
 #endif
 
+	ExternalTextures.Add(ExternalPooledTexture.GetReference(), OutTexture);
+
 	return OutTexture;
 }
 
@@ -34,6 +41,11 @@ inline FRDGBufferRef FRDGBuilder::RegisterExternalBuffer(const TRefCountPtr<FPoo
 	}
 #endif
 
+	if (FRDGBufferRef* Buffer = ExternalBuffers.Find(ExternalPooledBuffer.GetReference()))
+	{
+		return *Buffer;
+	}
+
 	FRDGBuffer* OutBuffer = AllocateForRHILifeTime<FRDGBuffer>(Name, ExternalPooledBuffer->Desc, ERDGResourceFlags::None);
 	OutBuffer->PooledBuffer = ExternalPooledBuffer;
 	AllocatedBuffers.Add(OutBuffer, ExternalPooledBuffer);
@@ -45,6 +57,8 @@ inline FRDGBufferRef FRDGBuilder::RegisterExternalBuffer(const TRefCountPtr<FPoo
 	}
 #endif
 
+	ExternalBuffers.Add(ExternalPooledBuffer.GetReference(), OutBuffer);
+
 	return OutBuffer;
 }
 
