Fixed render graph validation when doing copy passes.

diff --git a/Engine/Source/Runtime/RenderCore/Private/RenderGraphBuilder.cpp b/Engine/Source/Runtime/RenderCore/Private/RenderGraphBuilder.cpp
--- a/Engine/Source/Runtime/RenderCore/Private/RenderGraphBuilder.cpp
+++ b/Engine/Source/Runtime/RenderCore/Private/RenderGraphBuilder.cpp
@@ -215,8 +215,7 @@ void FRDGBuilder::ValidatePass(const FRDGPass* Pass) const
 	const FRenderTargetBindingSlots* RenderTargetBindingSlots = nullptr;
 
 	const TCHAR* PassName = Pass->GetName();
-	const bool bIsCompute = Pass->IsCompute();
-	const bool bRenderTargetSlotsRequired = !bIsCompute;
+	const bool bIsGraphics = Pass->IsGraphics();
 
 	FRDGPassParameterStruct ParameterStruct = Pass->GetParameters();
 
@@ -323,11 +322,11 @@ void FRDGBuilder::ValidatePass(const FRDGPass* Pass) const
 	/** Validate that raster passes have render target binding slots and compute passes don't. */
 	if (RenderTargetBindingSlots)
 	{
-		checkf(!bIsCompute, TEXT("Pass '%s' has render target binding slots but is flagged as 'Compute'."), PassName);
+		checkf(bIsGraphics, TEXT("Pass '%s' has render target binding slots but is flagged as 'Compute'."), PassName);
 	}
 	else
 	{
-		checkf(bIsCompute, TEXT("Pass '%s' is missing render target binding slots. Add the 'Compute' flag if render targets are not required."), PassName);
+		checkf(!bIsGraphics, TEXT("Pass '%s' is missing render target binding slots. Set the 'Compute' or 'Copy' flag if render targets are not required."), PassName);
 	}
 
 	/** Validate render target / depth stencil binding usage. */
@@ -814,7 +813,7 @@ void FRDGBuilder::ExecutePass(const FRDGPass* Pass)
 	EventScopeStack.BeginExecutePass(Pass);
 	StatScopeStack.BeginExecutePass(Pass);
 
-	if (!Pass->IsCompute())
+	if (Pass->IsGraphics())
 	{
 		check(bHasRenderTargets);
 		RHICmdList.BeginRenderPass( RPInfo, Pass->GetName() );
diff --git a/Engine/Source/Runtime/RenderCore/Public/RenderGraphPass.h b/Engine/Source/Runtime/RenderCore/Public/RenderGraphPass.h
--- a/Engine/Source/Runtime/RenderCore/Public/RenderGraphPass.h
+++ b/Engine/Source/Runtime/RenderCore/Public/RenderGraphPass.h
@@ -213,6 +213,16 @@ public:
 		return (PassFlags == ERDGPassFlags::Compute);
 	}
 
+	bool IsCopy() const
+	{
+		return (PassFlags == ERDGPassFlags::Copy);
+	}
+
+	bool IsGraphics() const
+	{
+		return (PassFlags == ERDGPassFlags::None) || (PassFlags == ERDGPassFlags::GenerateMips);
+	}
+
 	FRDGPassParameterStruct GetParameters() const
 	{
 		return ParameterStruct;
