Fixed static analysis warning.

diff --git a/Engine/Source/Runtime/RenderCore/Private/ShaderParameters.cpp b/Engine/Source/Runtime/RenderCore/Private/ShaderParameters.cpp
--- a/Engine/Source/Runtime/RenderCore/Private/ShaderParameters.cpp
+++ b/Engine/Source/Runtime/RenderCore/Private/ShaderParameters.cpp
@@ -263,7 +263,7 @@ static void CreateHLSLUniformBufferStructMembersDeclaration(
 		if (IsShaderParameterTypeForUniformBufferLayout(Member.GetBaseType()))
 		{
 			check(Member.GetBaseType() != UBMT_RDG_TEXTURE_SRV && Member.GetBaseType() != UBMT_RDG_TEXTURE_UAV);
-			checkf(Member.GetBaseType() != UBMT_RDG_TEXTURE_COPY_DEST && Member.GetBaseType() != UBMT_RDG_TEXTURE_COPY_DEST, TEXT("Copy destination usage is not supported in uniform buffers."));
+			checkf(Member.GetBaseType() != UBMT_RDG_TEXTURE_COPY_DEST && Member.GetBaseType() != UBMT_RDG_BUFFER_COPY_DEST, TEXT("Copy destination usage is not supported in uniform buffers."));
 			checkf(Member.GetNumElements() == 0, TEXT("Resources array are not supported in uniform buffers yet."));
 			if (Member.GetBaseType() == UBMT_SRV)
 			{
