Don't store pointers to serialization history in FShaderId
- These always point back to values stored in shader/vertex-factory type
- If types are equal, no need to also compare history for equality (since in that case they're guaranteed to be the same object)
- Reduces size of sets/maps that are keyed on FShaderId, reduces cost of comparison
- Also re-order FShaderId fields to reduce padding

diff --git a/Engine/Source/Runtime/RenderCore/Private/Shader.cpp b/Engine/Source/Runtime/RenderCore/Private/Shader.cpp
--- a/Engine/Source/Runtime/RenderCore/Private/Shader.cpp
+++ b/Engine/Source/Runtime/RenderCore/Private/Shader.cpp
@@ -1117,29 +1117,22 @@ FShaderResourceId FShaderResource::GetId() const
 }
 
 FShaderId::FShaderId(const FSHAHash& InMaterialShaderMapHash, const FShaderPipelineType* InShaderPipeline, FVertexFactoryType* InVertexFactoryType, FShaderType* InShaderType, int32 InPermutationId, FShaderTarget InTarget)
-	: MaterialShaderMapHash(InMaterialShaderMapHash)
+	: ShaderType(InShaderType)
+	, VertexFactoryType(InVertexFactoryType)
+	, ShaderPipeline(InShaderPipeline)
+	, MaterialShaderMapHash(InMaterialShaderMapHash)
 #if KEEP_SHADER_SOURCE_HASHES
 	, SourceHash(InShaderType->GetSourceHash(InTarget.GetPlatform()))
 #endif
 	, Target(InTarget)
-	, ShaderPipeline(InShaderPipeline)
-	, ShaderType(InShaderType)
 	, PermutationId(InPermutationId)
-	, SerializationHistory(InShaderType->GetSerializationHistory())
 {
+#if KEEP_SHADER_SOURCE_HASHES
 	if (InVertexFactoryType)
 	{
-		VFSerializationHistory = InVertexFactoryType->GetSerializationHistory(InTarget.GetFrequency());
-		VertexFactoryType = InVertexFactoryType;
-#if KEEP_SHADER_SOURCE_HASHES
 		VFSourceHash = InVertexFactoryType->GetSourceHash(InTarget.GetPlatform());
-#endif
-	}
-	else
-	{
-		VFSerializationHistory = nullptr;
-		VertexFactoryType = nullptr;
 	}
+#endif
 }
 
 FSelfContainedShaderId::FSelfContainedShaderId() :
@@ -1149,17 +1142,25 @@ FSelfContainedShaderId::FSelfContainedShaderId() :
 FSelfContainedShaderId::FSelfContainedShaderId(const FShaderId& InShaderId)
 {
 	MaterialShaderMapHash = InShaderId.MaterialShaderMapHash;
-	VertexFactoryTypeName = InShaderId.VertexFactoryType ? InShaderId.VertexFactoryType->GetName() : TEXT("");
 	ShaderPipelineName = InShaderId.ShaderPipeline ? InShaderId.ShaderPipeline->GetName() : TEXT("");
-	VFSerializationHistory = InShaderId.VFSerializationHistory ? *InShaderId.VFSerializationHistory : FSerializationHistory();
 	ShaderTypeName = InShaderId.ShaderType->GetName();
 	PermutationId = InShaderId.PermutationId;
+	Target = InShaderId.Target;
 #if KEEP_SHADER_SOURCE_HASHES
 	SourceHash = InShaderId.SourceHash;
 	VFSourceHash = InShaderId.VFSourceHash;
 #endif
-	SerializationHistory = InShaderId.SerializationHistory;
-	Target = InShaderId.Target;
+	SerializationHistory = InShaderId.ShaderType->GetSerializationHistory();
+
+	if (InShaderId.VertexFactoryType)
+	{
+		const FSerializationHistory* InVFSerializationHistory =InShaderId.VertexFactoryType->GetSerializationHistory(InShaderId.Target.GetFrequency());
+		VertexFactoryTypeName = InShaderId.VertexFactoryType->GetName();
+		if (InVFSerializationHistory)
+		{
+			VFSerializationHistory = *InVFSerializationHistory;
+		}
+	}
 }
 
 bool FSelfContainedShaderId::IsValid()
@@ -1471,11 +1472,10 @@ void FShader::Deregister()
 
 FShaderId FShader::GetId() const
 {
-	FShaderId ShaderId(Type->GetSerializationHistory());
+	FShaderId ShaderId;
 	ShaderId.MaterialShaderMapHash = MaterialShaderMapHash;
 	ShaderId.ShaderPipeline = ShaderPipeline;
 	ShaderId.VertexFactoryType = VFType;
-	ShaderId.VFSerializationHistory = VFType ? VFType->GetSerializationHistory((EShaderFrequency)GetTarget().Frequency) : NULL;
 	ShaderId.ShaderType = Type;
 	ShaderId.PermutationId = PermutationId;
 #if KEEP_SHADER_SOURCE_HASHES
diff --git a/Engine/Source/Runtime/RenderCore/Private/ShaderCore.cpp b/Engine/Source/Runtime/RenderCore/Private/ShaderCore.cpp
--- a/Engine/Source/Runtime/RenderCore/Private/ShaderCore.cpp
+++ b/Engine/Source/Runtime/RenderCore/Private/ShaderCore.cpp
@@ -92,6 +92,9 @@ DEFINE_STAT(STAT_Shaders_ShaderMemory);
 DEFINE_STAT(STAT_Shaders_ShaderResourceMemory);
 DEFINE_STAT(STAT_Shaders_ShaderMapMemory);
 
+DEFINE_STAT(STAT_Shaders_NumShadersRegistered);
+DEFINE_STAT(STAT_Shaders_NumShadersDuplicated);
+
 /** Protects GShaderFileCache from simultaneous access by multiple threads. */
 FCriticalSection FileCacheCriticalSection;
 
diff --git a/Engine/Source/Runtime/RenderCore/Public/Shader.h b/Engine/Source/Runtime/RenderCore/Public/Shader.h
--- a/Engine/Source/Runtime/RenderCore/Public/Shader.h
+++ b/Engine/Source/Runtime/RenderCore/Public/Shader.h
@@ -618,6 +618,18 @@ public:
 class FShaderId
 {
 public:
+	/** Shader type */
+	FShaderType* ShaderType;
+
+	/**
+	 * Vertex factory type that the shader was created for,
+	 * This is needed in the Id since a single shader type will be compiled for multiple vertex factories within a material shader map.
+	 * Will be NULL for global shaders.
+	 */
+	FVertexFactoryType* VertexFactoryType;
+
+	/** Shader Pipeline linked to this shader, needed since a single shader might be used on different Pipelines. */
+	const FShaderPipelineType* ShaderPipeline;
 
 	/** 
 	 * Hash of the material shader map Id, since this shader depends on the generated material code from that shader map.
@@ -636,36 +648,10 @@ public:
 	/** Shader platform and frequency. */
 	FShaderTarget Target;
 
-	/** Shader Pipeline linked to this shader, needed since a single shader might be used on different Pipelines. */
-	const FShaderPipelineType* ShaderPipeline;
-
-	/** 
-	 * Vertex factory type that the shader was created for, 
-	 * This is needed in the Id since a single shader type will be compiled for multiple vertex factories within a material shader map.
-	 * Will be NULL for global shaders.
-	 */
-	FVertexFactoryType* VertexFactoryType;
-
-	/** 
-	 * Used to detect changes to the vertex factory parameter class serialization, or NULL for global shaders. 
-	 * Note: This is referencing memory in the VF Type, since it is the same for all shaders using that VF Type.
-	 */
-	const FSerializationHistory* VFSerializationHistory;
-
-	/** Shader type */
-	FShaderType* ShaderType;
-
 	/** Unique permutation identifier within the ShaderType. */
 	int32 PermutationId;
 
-	/** Used to detect changes to the shader serialization.  Note: this is referencing memory in the FShaderType. */
-	const FSerializationHistory& SerializationHistory;
-
-	/** Create a minimally initialized Id.  Members will have to be assigned individually. */
-	FShaderId(const FSerializationHistory& InSerializationHistory)
-		: PermutationId(0)
-		, SerializationHistory(InSerializationHistory)
-	{}
+	FShaderId() = default;
 
 	/** Creates an Id for the given material, vertex factory, shader type and target. */
 	RENDERCORE_API FShaderId(const FSHAHash& InMaterialShaderMapHash, const FShaderPipelineType* InShaderPipeline, FVertexFactoryType* InVertexFactoryType, FShaderType* InShaderType, int32 PermutationId, FShaderTarget InTarget);
@@ -683,16 +669,12 @@ public:
 		return X.MaterialShaderMapHash == Y.MaterialShaderMapHash
 			&& X.ShaderPipeline == Y.ShaderPipeline
 			&& X.VertexFactoryType == Y.VertexFactoryType
-			&& ((X.VFSerializationHistory == NULL && Y.VFSerializationHistory == NULL)
-				|| (X.VFSerializationHistory != NULL && Y.VFSerializationHistory != NULL &&
-					*X.VFSerializationHistory == *Y.VFSerializationHistory))
 			&& X.ShaderType == Y.ShaderType
 			&& X.PermutationId == Y.PermutationId 
 #if KEEP_SHADER_SOURCE_HASHES
 			&& X.SourceHash == Y.SourceHash 
 			&& X.VFSourceHash == Y.VFSourceHash
 #endif
-			&& X.SerializationHistory == Y.SerializationHistory
 			&& X.Target == Y.Target;
 	}
 
@@ -2295,12 +2277,14 @@ public:
 
 			if (ExistingShader != nullptr)
 			{
+				INC_DWORD_STAT(STAT_Shaders_NumShadersDuplicated);
 				delete Shader;
 				Shader = ExistingShader;
 			}
 			else
 			{
 				// Register the shader now that it is valid, so that it can be reused
+				INC_DWORD_STAT(STAT_Shaders_NumShadersRegistered);
 				Shader->Register(bCookedMaterial);
 			}
 			AddShader(Shader->GetType(), Shader->GetPermutationId(), Shader);
diff --git a/Engine/Source/Runtime/RenderCore/Public/ShaderCore.h b/Engine/Source/Runtime/RenderCore/Public/ShaderCore.h
--- a/Engine/Source/Runtime/RenderCore/Public/ShaderCore.h
+++ b/Engine/Source/Runtime/RenderCore/Public/ShaderCore.h
@@ -62,6 +62,9 @@ DECLARE_MEMORY_STAT_EXTERN(TEXT("Shader Memory"),STAT_Shaders_ShaderMemory,STATG
 DECLARE_MEMORY_STAT_EXTERN(TEXT("Shader Resource Mem"),STAT_Shaders_ShaderResourceMemory,STATGROUP_Shaders, RENDERCORE_API);
 DECLARE_MEMORY_STAT_EXTERN(TEXT("Shader MapMemory"),STAT_Shaders_ShaderMapMemory,STATGROUP_Shaders, RENDERCORE_API);
 
+DECLARE_DWORD_ACCUMULATOR_STAT_EXTERN(TEXT("Num Shaders Registered"), STAT_Shaders_NumShadersRegistered, STATGROUP_Shaders, RENDERCORE_API);
+DECLARE_DWORD_ACCUMULATOR_STAT_EXTERN(TEXT("Num Shaders Duplicated"), STAT_Shaders_NumShadersDuplicated, STATGROUP_Shaders, RENDERCORE_API);
+
 inline TStatId GetMemoryStatType(EShaderFrequency ShaderFrequency)
 {
 	static_assert(10 == SF_NumFrequencies, "EShaderFrequency has a bad size.");
