Validates mip level on UAV and SRV.

diff --git a/Engine/Source/Runtime/RenderCore/Public/RenderGraphBuilder.inl b/Engine/Source/Runtime/RenderCore/Public/RenderGraphBuilder.inl
--- a/Engine/Source/Runtime/RenderCore/Public/RenderGraphBuilder.inl
+++ b/Engine/Source/Runtime/RenderCore/Public/RenderGraphBuilder.inl
@@ -125,6 +125,7 @@ inline FRDGTextureSRVRef FRDGBuilder::CreateSRV(const FRDGTextureSRVDesc& Desc)
 		checkf(Texture, TEXT("RenderGraph texture SRV created with a null texture."));
 		checkf(!bHasExecuted, TEXT("Render graph SRV %s needs to be created before the builder execution."), Desc.Texture->Name);
 		checkf(Desc.Texture->Desc.TargetableFlags & TexCreate_ShaderResource, TEXT("Attempted to create SRV from texture %s which was not created with TexCreate_ShaderResource"), Desc.Texture->Name);
+		checkf(Desc.MipLevel < Texture->Desc.NumMips, TEXT("Failed to create SRV at mip %d: the texture %s has only %d mip levels."), Desc.MipLevel, Texture->Name, Texture->Desc.NumMips);
 	}
 #endif
 
@@ -154,6 +155,7 @@ inline FRDGTextureUAVRef FRDGBuilder::CreateUAV(const FRDGTextureUAVDesc& Desc)
 		checkf(Texture, TEXT("RenderGraph texture UAV created with a null texture."));
 		checkf(!bHasExecuted, TEXT("Render graph UAV %s needs to be created before the builder execution."), Texture->Name);
 		checkf(Texture->Desc.TargetableFlags & TexCreate_UAV, TEXT("Attempted to create UAV from texture %s which was not created with TexCreate_UAV"), Texture->Name);
+		checkf(Desc.MipLevel < Texture->Desc.NumMips, TEXT("Failed to create UAV at mip %d: the texture %s has only %d mip levels."), Desc.MipLevel, Texture->Name, Texture->Desc.NumMips);
 	}
 #endif
 
