Add FShaderKey to replace FShaderId for looking up shaders on FShaderTypes
- FShaderId includes FShaderType*, which is extra/wasted data for looking up shaders when the specific type is already known

diff --git a/Engine/Plugins/Compositing/OpenColorIO/Source/OpenColorIO/Private/OpenColorIOShaderMap.cpp b/Engine/Plugins/Compositing/OpenColorIO/Source/OpenColorIO/Private/OpenColorIOShaderMap.cpp
--- a/Engine/Plugins/Compositing/OpenColorIO/Source/OpenColorIO/Private/OpenColorIOShaderMap.cpp
+++ b/Engine/Plugins/Compositing/OpenColorIO/Source/OpenColorIO/Private/OpenColorIOShaderMap.cpp
@@ -237,7 +237,7 @@ FShader* FOpenColorIOShaderType::FinishCompileShader(
 	TRefCountPtr<FShaderResource> Resource = FShaderResource::FindOrCreate(InCurrentJob.Output, SpecificType, /* SpecificPermutationId = */ 0);
 
 	// Find a shader with the same key in memory
-	FShader* Shader = InCurrentJob.ShaderType->FindShaderById(FShaderId(InShaderMapHash, nullptr, nullptr, InCurrentJob.ShaderType, /* SpecificPermutationId = */ 0, InCurrentJob.Input.Target));
+	FShader* Shader = InCurrentJob.ShaderType->FindShaderByKey(FShaderKey(InShaderMapHash, nullptr, nullptr, /* SpecificPermutationId = */ 0, InCurrentJob.Input.Target.GetPlatform()));
 
 	// There was no shader with the same key so create a new one with the compile output, which will bind shader parameters
 	if (!Shader)
@@ -614,8 +614,8 @@ void FOpenColorIOShaderMap::LoadMissingShadersFromMemory(const FOpenColorIOTrans
 		FOpenColorIOShaderType* ShaderType = ShaderTypeIt->GetOpenColorIOShaderType();
 		if (ShaderType && ShouldCacheOpenColorIOShader(ShaderType, Platform, InColorTransform) && !HasShader(ShaderType, /* PermutationId = */ 0))
 		{
-			FShaderId ShaderId(ShaderMapHash, nullptr, nullptr, ShaderType, /** PermutationId = */ 0, FShaderTarget(ShaderType->GetFrequency(), Platform));
-			FShader* FoundShader = ShaderType->FindShaderById(ShaderId);
+			FShaderKey ShaderKey(ShaderMapHash, nullptr, nullptr, /** PermutationId = */ 0, Platform);
+			FShader* FoundShader = ShaderType->FindShaderByKey(ShaderKey);
 			if (FoundShader)
 			{
 				AddShader(ShaderType, /* PermutationId = */ 0, FoundShader);
diff --git a/Engine/Plugins/FX/Niagara/Source/NiagaraShader/Private/NiagaraShader.cpp b/Engine/Plugins/FX/Niagara/Source/NiagaraShader/Private/NiagaraShader.cpp
--- a/Engine/Plugins/FX/Niagara/Source/NiagaraShader/Private/NiagaraShader.cpp
+++ b/Engine/Plugins/FX/Niagara/Source/NiagaraShader/Private/NiagaraShader.cpp
@@ -413,7 +413,7 @@ FShader* FNiagaraShaderType::FinishCompileShader(
 	TRefCountPtr<FShaderResource> Resource = FShaderResource::FindOrCreate(CurrentJob.Output, SpecificType, /* SpecificPermutationId = */ 0);
 
 	// Find a shader with the same key in memory
-	FShader* Shader = CurrentJob.ShaderType->FindShaderById(FShaderId(ShaderMapHash, nullptr, nullptr, CurrentJob.ShaderType, /* SpecificPermutationId = */ 0, CurrentJob.Input.Target));
+	FShader* Shader = CurrentJob.ShaderType->FindShaderByKey(FShaderKey(ShaderMapHash, nullptr, nullptr, /* SpecificPermutationId = */ 0, CurrentJob.Input.Target.GetPlatform()));
 
 	// There was no shader with the same key so create a new one with the compile output, which will bind shader parameters
 	if (!Shader)
@@ -1034,8 +1034,8 @@ void FNiagaraShaderMap::LoadMissingShadersFromMemory(const FNiagaraShaderScript*
 		FNiagaraShaderType* ShaderType = ShaderTypeIt->GetNiagaraShaderType();
 		if (ShaderType && ShouldCacheNiagaraShader(ShaderType, Platform, Script) && !HasShader(ShaderType, /* PermutationId = */ 0))
 		{
-			FShaderId ShaderId(ShaderMapHash, nullptr, nullptr, ShaderType, /** PermutationId = */ 0, FShaderTarget(ShaderType->GetFrequency(), Platform));
-			FShader* FoundShader = ShaderType->FindShaderById(ShaderId);
+			FShaderKey ShaderKey(ShaderMapHash, nullptr, nullptr, /** PermutationId = */ 0, Platform);
+			FShader* FoundShader = ShaderType->FindShaderByKey(ShaderKey);
 			if (FoundShader)
 			{
 				AddShader(ShaderType, /* PermutationId = */ 0, FoundShader);
diff --git a/Engine/Source/Runtime/Engine/Private/Materials/MaterialInstance.cpp b/Engine/Source/Runtime/Engine/Private/Materials/MaterialInstance.cpp
--- a/Engine/Source/Runtime/Engine/Private/Materials/MaterialInstance.cpp
+++ b/Engine/Source/Runtime/Engine/Private/Materials/MaterialInstance.cpp
@@ -1192,9 +1192,8 @@ void UMaterialInstance::ValidateTextureOverrides(ERHIFeatureLevel::Type InFeatur
 			{
 				UTexture* Texture = nullptr;
 				const FMaterialParameterInfo& ParameterInfo = ParameterExpression->GetParameterInfo();
-				if (GetTextureParameterValue(ParameterInfo, Texture, true))
+				if (GetTextureParameterValue(ParameterInfo, Texture, true) && Texture != nullptr)
 				{
-					check(Texture);
 					const EMaterialValueType TextureType = Texture->GetMaterialType();
 					switch (TypeIndex)
 					{
diff --git a/Engine/Source/Runtime/Engine/Private/Materials/MaterialShader.cpp b/Engine/Source/Runtime/Engine/Private/Materials/MaterialShader.cpp
--- a/Engine/Source/Runtime/Engine/Private/Materials/MaterialShader.cpp
+++ b/Engine/Source/Runtime/Engine/Private/Materials/MaterialShader.cpp
@@ -1021,7 +1021,7 @@ FShader* FMaterialShaderType::FinishCompileShader(
 	}
 
 	// Find a shader with the same key in memory
-	FShader* Shader = CurrentJob.ShaderType->FindShaderById(FShaderId(MaterialShaderMapHash, ShaderPipelineType, CurrentJob.VFType, CurrentJob.ShaderType, CurrentJob.PermutationId, CurrentJob.Input.Target));
+	FShader* Shader = CurrentJob.ShaderType->FindShaderByKey(FShaderKey(MaterialShaderMapHash, ShaderPipelineType, CurrentJob.VFType, CurrentJob.PermutationId, CurrentJob.Input.Target.GetPlatform()));
 
 	// There was no shader with the same key so create a new one with the compile output, which will bind shader parameters
 	if (!Shader)
@@ -2102,8 +2102,8 @@ void FMaterialShaderMap::LoadMissingShadersFromMemory(const FMaterial* Material)
 		{
 			if (ShouldCacheMaterialShader(ShaderType, GetShaderPlatform(), Material, PermutationId) && !HasShader(ShaderType, PermutationId))
 			{
-				FShaderId ShaderId(MaterialShaderMapHash, nullptr, nullptr, ShaderType, PermutationId, FShaderTarget(ShaderType->GetFrequency(), GetShaderPlatform()));
-				FShader* FoundShader = ShaderType->FindShaderById(ShaderId);
+				FShaderKey ShaderKey(MaterialShaderMapHash, nullptr, nullptr, PermutationId, GetShaderPlatform());
+				FShader* FoundShader = ShaderType->FindShaderByKey(ShaderKey);
 				if (FoundShader)
 				{
 					AddShader(ShaderType, PermutationId, FoundShader);
@@ -2138,8 +2138,8 @@ void FMaterialShaderMap::LoadMissingShadersFromMemory(const FMaterial* Material)
 					FMaterialShaderType* ShaderType = (FMaterialShaderType*)Shader->GetMaterialShaderType();
 					if (!HasShader(ShaderType, kUniqueShaderPermutationId))
 					{
-						FShaderId ShaderId(MaterialShaderMapHash, PipelineType->ShouldOptimizeUnusedOutputs(GetShaderPlatform()) ? PipelineType : nullptr, nullptr, ShaderType, kUniqueShaderPermutationId, FShaderTarget(ShaderType->GetFrequency(), GetShaderPlatform()));
-						FShader* FoundShader = ShaderType->FindShaderById(ShaderId);
+						FShaderKey ShaderKey(MaterialShaderMapHash, PipelineType->ShouldOptimizeUnusedOutputs(GetShaderPlatform()) ? PipelineType : nullptr, nullptr, kUniqueShaderPermutationId, GetShaderPlatform());
+						FShader* FoundShader = ShaderType->FindShaderByKey(ShaderKey);
 						if (FoundShader)
 						{
 							AddShader(ShaderType, kUniqueShaderPermutationId, FoundShader);
diff --git a/Engine/Source/Runtime/Engine/Private/Materials/MeshMaterialShader.cpp b/Engine/Source/Runtime/Engine/Private/Materials/MeshMaterialShader.cpp
--- a/Engine/Source/Runtime/Engine/Private/Materials/MeshMaterialShader.cpp
+++ b/Engine/Source/Runtime/Engine/Private/Materials/MeshMaterialShader.cpp
@@ -160,7 +160,7 @@ FShader* FMeshMaterialShaderType::FinishCompileShader(
 	}
 
 	// Find a shader with the same key in memory
-	FShader* Shader = CurrentJob.ShaderType->FindShaderById(FShaderId(MaterialShaderMapHash, ShaderPipelineType, CurrentJob.VFType, CurrentJob.ShaderType, CurrentJob.PermutationId, CurrentJob.Input.Target));
+	FShader* Shader = CurrentJob.ShaderType->FindShaderByKey(FShaderKey(MaterialShaderMapHash, ShaderPipelineType, CurrentJob.VFType, CurrentJob.PermutationId, CurrentJob.Input.Target.GetPlatform()));
 
 	// There was no shader with the same key so create a new one with the compile output, which will bind shader parameters
 	if (!Shader)
@@ -428,8 +428,8 @@ void FMeshMaterialShaderMap::LoadMissingShadersFromMemory(
 		{
 			if (ShouldCacheMeshShader(ShaderType, InPlatform, Material, VertexFactoryType, PermutationId) && !HasShader((FShaderType*)ShaderType, PermutationId))
 			{
-				const FShaderId ShaderId(MaterialShaderMapHash, nullptr, VertexFactoryType, (FShaderType*)ShaderType, PermutationId, FShaderTarget(ShaderType->GetFrequency(), InPlatform));
-				FShader* FoundShader = ((FShaderType*)ShaderType)->FindShaderById(ShaderId);
+				const FShaderKey ShaderKey(MaterialShaderMapHash, nullptr, VertexFactoryType, PermutationId, InPlatform);
+				FShader* FoundShader = ((FShaderType*)ShaderType)->FindShaderByKey(ShaderKey);
 
 				if (FoundShader)
 				{
@@ -469,8 +469,8 @@ void FMeshMaterialShaderMap::LoadMissingShadersFromMemory(
 					FMeshMaterialShaderType* ShaderType = (FMeshMaterialShaderType*)Shader->GetMeshMaterialShaderType();
 					if (!HasShader(ShaderType, kUniqueShaderPermutationId))
 					{
-						const FShaderId ShaderId(MaterialShaderMapHash, PipelineType->ShouldOptimizeUnusedOutputs(InPlatform) ? PipelineType : nullptr, VertexFactoryType, ShaderType, kUniqueShaderPermutationId, FShaderTarget(ShaderType->GetFrequency(), InPlatform));
-						FShader* FoundShader = ShaderType->FindShaderById(ShaderId);
+						const FShaderKey ShaderKey(MaterialShaderMapHash, PipelineType->ShouldOptimizeUnusedOutputs(InPlatform) ? PipelineType : nullptr, VertexFactoryType, kUniqueShaderPermutationId, InPlatform);
+						FShader* FoundShader = ShaderType->FindShaderByKey(ShaderKey);
 						if (FoundShader)
 						{
 							AddShader(ShaderType, kUniqueShaderPermutationId, FoundShader);
diff --git a/Engine/Source/Runtime/Engine/Private/ShaderCompiler/ShaderCompiler.cpp b/Engine/Source/Runtime/Engine/Private/ShaderCompiler/ShaderCompiler.cpp
--- a/Engine/Source/Runtime/Engine/Private/ShaderCompiler/ShaderCompiler.cpp
+++ b/Engine/Source/Runtime/Engine/Private/ShaderCompiler/ShaderCompiler.cpp
@@ -3791,7 +3791,7 @@ FShader* FGlobalShaderTypeCompiler::FinishCompileShader(FGlobalShaderType* Shade
 		}
 
 		// Find a shader with the same key in memory
-		Shader = CurrentJob.ShaderType->FindShaderById(FShaderId(GlobalShaderMapHash, ShaderPipelineType, nullptr, CurrentJob.ShaderType, CurrentJob.PermutationId, CurrentJob.Input.Target));
+		Shader = CurrentJob.ShaderType->FindShaderByKey(FShaderKey(GlobalShaderMapHash, ShaderPipelineType, nullptr, CurrentJob.PermutationId, CurrentJob.Input.Target.GetPlatform()));
 
 		// There was no shader with the same key so create a new one with the compile output, which will bind shader parameters
 		if (!Shader)
diff --git a/Engine/Source/Runtime/RenderCore/Private/Shader.cpp b/Engine/Source/Runtime/RenderCore/Private/Shader.cpp
--- a/Engine/Source/Runtime/RenderCore/Private/Shader.cpp
+++ b/Engine/Source/Runtime/RenderCore/Private/Shader.cpp
@@ -236,7 +236,7 @@ inline bool FShaderType::GetOutdatedCurrentType(TArray<FShaderType*>& OutdatedSh
 {
 	bool bOutdated = false;
 #if WITH_EDITOR
-	for (TMap<FShaderId, FShader*>::TConstIterator ShaderIt(ShaderIdMap);ShaderIt;++ShaderIt)
+	for (TMap<FShaderKey, FShader*>::TConstIterator ShaderIt(ShaderIdMap);ShaderIt;++ShaderIt)
 	{
 		FShader* Shader = ShaderIt.Value();
 		const FVertexFactoryParameterRef* VFParameterRef = Shader->GetVertexFactoryParameterRef();
@@ -314,11 +314,11 @@ FArchive& operator<<(FArchive& Ar,FShaderType*& Ref)
 }
 
 
-FShader* FShaderType::FindShaderById(const FShaderId& Id)
+FShader* FShaderType::FindShaderByKey(const FShaderKey& Key)
 {
 	check(IsInGameThread());
-	FShader* Result = ShaderIdMap.FindRef(Id);
-	check(!Result || Result->GetId() == Id);
+	FShader* Result = ShaderIdMap.FindRef(Key);
+	checkSlow(!Result || Result->GetKey() == Key);
 	return Result;
 }
 
@@ -1135,6 +1135,19 @@ FShaderId::FShaderId(const FSHAHash& InMaterialShaderMapHash, const FShaderPipel
 #endif
 }
 
+FArchive& operator<<(FArchive& Ar, FShaderKey& Ref)
+{
+	uint32 Platform = Ref.Platform;
+	Ar
+		<< Ref.VertexFactoryType
+		<< Ref.ShaderPipeline
+		<< Ref.MaterialShaderMapHash
+		<< Ref.PermutationId
+		<< Platform;
+	Ref.Platform = Platform;
+	return Ar;
+}
+
 FSelfContainedShaderId::FSelfContainedShaderId() :
 	Target(FShaderTarget(SF_NumFrequencies, SP_NumPlatforms))
 {}
@@ -1456,18 +1469,18 @@ void FShader::Release()
 
 void FShader::Register(bool bLoadedByCookedMaterial)
 {
-	FShaderId ShaderId = GetId();
-	check(ShaderId.MaterialShaderMapHash != FSHAHash());
+	const FShaderKey ShaderKey = GetKey();
+	check(ShaderKey.MaterialShaderMapHash != FSHAHash());
 #if KEEP_SHADER_SOURCE_HASHES
-	check(ShaderId.SourceHash != FSHAHash() || FPlatformProperties::RequiresCookedData() || bLoadedByCookedMaterial);
+	check(SourceHash != FSHAHash() || FPlatformProperties::RequiresCookedData() || bLoadedByCookedMaterial);
 #endif
 	check(Resource);
-	Type->AddToShaderIdMap(ShaderId, this);
+	Type->AddToShaderMap(ShaderKey, this);
 }
 
 void FShader::Deregister()
 {
-	Type->RemoveFromShaderIdMap(GetId());
+	Type->RemoveFromShaderMap(GetKey());
 }
 
 FShaderId FShader::GetId() const
@@ -1935,7 +1948,7 @@ void DumpShaderStats(EShaderPlatform Platform, EShaderFrequency Frequency)
 			int32 NumShaders					= 0;
 			int32 NumPipelines = 0;
 			int32 NumSharedPipelines = 0;
-			for (TMap<FShaderId,FShader*>::TConstIterator ShaderIt(Type->ShaderIdMap);ShaderIt;++ShaderIt)
+			for (TMap<FShaderKey,FShader*>::TConstIterator ShaderIt(Type->ShaderIdMap);ShaderIt;++ShaderIt)
 			{
 				const FShader* Shader = ShaderIt.Value();
 				// Skip shaders that don't match frequency.
diff --git a/Engine/Source/Runtime/RenderCore/Public/Shader.h b/Engine/Source/Runtime/RenderCore/Public/Shader.h
--- a/Engine/Source/Runtime/RenderCore/Public/Shader.h
+++ b/Engine/Source/Runtime/RenderCore/Public/Shader.h
@@ -611,6 +611,43 @@ public:
 	}
 };
 
+class FShaderKey
+{
+public:
+	inline FShaderKey(const FSHAHash& InMaterialShaderMapHash, const FShaderPipelineType* InShaderPipeline, FVertexFactoryType* InVertexFactoryType, int32 PermutationId, EShaderPlatform InPlatform)
+		: VertexFactoryType(InVertexFactoryType)
+		, ShaderPipeline(InShaderPipeline)
+		, MaterialShaderMapHash(InMaterialShaderMapHash)
+		, PermutationId(PermutationId)
+		, Platform(InPlatform)
+	{}
+
+	friend inline uint32 GetTypeHash(const FShaderKey& Id)
+	{
+		return
+			HashCombine(
+				HashCombine(*(uint32*)&Id.MaterialShaderMapHash, GetTypeHash(Id.Platform)),
+				HashCombine(GetTypeHash(Id.VertexFactoryType), uint32(Id.PermutationId)));
+	}
+
+	friend bool operator==(const FShaderKey& X, const FShaderKey& Y)
+	{
+		return X.MaterialShaderMapHash == Y.MaterialShaderMapHash
+			&& X.ShaderPipeline == Y.ShaderPipeline
+			&& X.VertexFactoryType == Y.VertexFactoryType
+			&& X.PermutationId == Y.PermutationId
+			&& X.Platform == Y.Platform;
+	}
+
+	RENDERCORE_API friend FArchive& operator<<(FArchive& Ar, FShaderKey& Ref);
+
+	FVertexFactoryType* VertexFactoryType;
+	const FShaderPipelineType* ShaderPipeline;
+	FSHAHash MaterialShaderMapHash;
+	int32 PermutationId;
+	uint32 Platform : SP_NumBits;
+};
+
 /** 
  * Uniquely identifies an FShader instance.  
  * Used to link FMaterialShaderMaps and FShaders on load. 
@@ -651,7 +688,8 @@ public:
 	/** Unique permutation identifier within the ShaderType. */
 	int32 PermutationId;
 
-	FShaderId() = default;
+	/** Creates empty Id */
+	inline FShaderId() : ShaderType(nullptr), VertexFactoryType(nullptr), ShaderPipeline(nullptr), PermutationId(0) {}
 
 	/** Creates an Id for the given material, vertex factory, shader type and target. */
 	RENDERCORE_API FShaderId(const FSHAHash& InMaterialShaderMapHash, const FShaderPipelineType* InShaderPipeline, FVertexFactoryType* InVertexFactoryType, FShaderType* InShaderType, int32 PermutationId, FShaderTarget InTarget);
@@ -965,6 +1003,7 @@ public:
 #endif
 	}
 	FShaderId GetId() const;
+	inline FShaderKey GetKey() const { return FShaderKey(MaterialShaderMapHash, ShaderPipeline, VFType, PermutationId, Target.GetPlatform()); }
 	inline FVertexFactoryType* GetVertexFactoryType() const { return VFType; }
 	inline int32 GetNumRefs() const { return NumRefs; }
 	const FShaderParameterMapInfo& GetParameterMapInfo() const { return Resource->ParameterMapInfo; }
@@ -1178,10 +1217,10 @@ public:
 	virtual ~FShaderType();
 
 	/**
-	 * Finds a shader of this type by ID.
-	 * @return NULL if no shader with the specified ID was found.
+	 * Finds a shader of this type by Key.
+	 * @return NULL if no shader with the specified Key was found.
 	 */
-	FShader* FindShaderById(const FShaderId& Id);
+	FShader* FindShaderByKey(const FShaderKey& Key);
 
 	/** Constructs a new instance of the shader type for deserialization. */
 	FShader* ConstructForDeserialization() const;
@@ -1297,16 +1336,16 @@ public:
 		bCachedUniformBufferStructDeclarations = false;
 	}
 
-	void AddToShaderIdMap(FShaderId Id, FShader* Shader)
+	void AddToShaderMap(const FShaderKey& Key, FShader* Shader)
 	{
 		check(IsInGameThread());
-		ShaderIdMap.Add(Id, Shader);
+		ShaderIdMap.Add(Key, Shader);
 	}
 
-	inline void RemoveFromShaderIdMap(FShaderId Id)
+	inline void RemoveFromShaderMap(const FShaderKey& Key)
 	{
 		check(IsInGameThread());
-		ShaderIdMap.Remove(Id);
+		ShaderIdMap.Remove(Key);
 	}
 
 	bool LimitShaderResourceToThisType() const
@@ -1337,7 +1376,7 @@ private:
 	const FShaderParametersMetadata* const RootParametersMetadata;
 
 	/** A map from shader ID to shader.  A shader will be removed from it when deleted, so this doesn't need to use a TRefCountPtr. */
-	TMap<FShaderId,FShader*> ShaderIdMap;
+	TMap<FShaderKey,FShader*> ShaderIdMap;
 
 	TLinkedList<FShaderType*> GlobalListLink;
 
@@ -2273,7 +2312,7 @@ public:
 		for (FShader* Shader : SerializedShaders)
 		{
 			FShaderType* Type = Shader->GetType();
-			FShader* ExistingShader = Type->FindShaderById(Shader->GetId());
+			FShader* ExistingShader = Type->FindShaderByKey(Shader->GetKey());
 
 			if (ExistingShader != nullptr)
 			{
