Fixes bug where the event format was not being returned by FRDGEventName::GetTCHAR() when GetEmitRDGEvents() == false

diff --git a/Engine/Source/Runtime/RenderCore/Private/RenderGraphEvent.cpp b/Engine/Source/Runtime/RenderCore/Private/RenderGraphEvent.cpp
--- a/Engine/Source/Runtime/RenderCore/Private/RenderGraphEvent.cpp
+++ b/Engine/Source/Runtime/RenderCore/Private/RenderGraphEvent.cpp
@@ -7,7 +7,8 @@ extern bool GetEmitRDGEvents();
 
 #if RDG_EVENTS == RDG_EVENTS_STRING_COPY
 
-FRDGEventName::FRDGEventName(const TCHAR* EventFormat, ...)
+FRDGEventName::FRDGEventName(const TCHAR* InEventFormat, ...)
+	: EventFormat(InEventFormat)
 {
 	if (GetEmitRDGEvents())
 	{
@@ -18,12 +19,7 @@ FRDGEventName::FRDGEventName(const TCHAR* EventFormat, ...)
 		FCString::GetVarArgs(TempStr, ARRAY_COUNT(TempStr), EventFormat, VAList);
 		va_end(VAList);
 
-		EventNameStorage = TempStr;
-		EventName = *EventNameStorage;
-	}
-	else
-	{
-		EventName = TEXT("!!!Unavailable RDG event name: try r.RDG.EmitWarnings=1 or -rdgdebug!!!");
+		FormatedEventName = TempStr;
 	}
 }
 
diff --git a/Engine/Source/Runtime/RenderCore/Public/RenderGraphEvent.h b/Engine/Source/Runtime/RenderCore/Public/RenderGraphEvent.h
--- a/Engine/Source/Runtime/RenderCore/Public/RenderGraphEvent.h
+++ b/Engine/Source/Runtime/RenderCore/Public/RenderGraphEvent.h
@@ -95,9 +95,12 @@ public:
 
 private:
 #if RDG_EVENTS == RDG_EVENTS_STRING_REF || RDG_EVENTS == RDG_EVENTS_STRING_COPY
-	const TCHAR* EventName;
+	// Event format kept arround to still have a clue what error might be causing the problem in error messages.
+	const TCHAR* EventFormat;
+
 #if RDG_EVENTS == RDG_EVENTS_STRING_COPY
-	FString EventNameStorage;
+	// Formated event name if GetEmitRDGEvents() == true.
+	FString FormatedEventName;
 #endif
 #endif
 };
diff --git a/Engine/Source/Runtime/RenderCore/Public/RenderGraphEvent.inl b/Engine/Source/Runtime/RenderCore/Public/RenderGraphEvent.inl
--- a/Engine/Source/Runtime/RenderCore/Public/RenderGraphEvent.inl
+++ b/Engine/Source/Runtime/RenderCore/Public/RenderGraphEvent.inl
@@ -121,15 +121,15 @@ void FRDGScopeStack<TScopeType>::ClearScopes()
 
 inline FRDGEventName::~FRDGEventName()
 {
-#if RDG_EVENTS != RDG_EVENTS_NONE
-	EventName = nullptr;
+#if RDG_EVENTS == RDG_EVENTS_STRING_REF || RDG_EVENTS == RDG_EVENTS_STRING_COPY
+	EventFormat = nullptr;
 #endif
 }
 
 #if RDG_EVENTS != RDG_EVENTS_STRING_COPY
-inline FRDGEventName::FRDGEventName(const TCHAR* EventFormat, ...)
+inline FRDGEventName::FRDGEventName(const TCHAR* InEventFormat, ...)
 #if RDG_EVENTS == RDG_EVENTS_STRING_REF
-	: EventName(EventFormat)
+	: EventFormat(InEventFormat)
 #endif
 {}
 #endif
@@ -147,10 +147,10 @@ inline FRDGEventName::FRDGEventName(FRDGEventName&& Other)
 inline FRDGEventName& FRDGEventName::operator=(const FRDGEventName& Other)
 {
 #if RDG_EVENTS == RDG_EVENTS_STRING_REF
-	EventName = Other.EventName;
+	EventFormat = Other.EventFormat;
 #elif RDG_EVENTS == RDG_EVENTS_STRING_COPY
-	EventNameStorage = Other.EventNameStorage;
-	EventName = *EventNameStorage;
+	EventFormat = Other.EventFormat;
+	FormatedEventName = FormatedEventName;
 #endif
 	return *this;
 }
@@ -158,12 +158,12 @@ inline FRDGEventName& FRDGEventName::operator=(const FRDGEventName& Other)
 inline FRDGEventName& FRDGEventName::operator=(FRDGEventName&& Other)
 {
 #if RDG_EVENTS == RDG_EVENTS_STRING_REF
-	EventName = Other.EventName;
-	Other.EventName = nullptr;
+	EventFormat = Other.EventFormat;
+	Other.EventFormat = nullptr;
 #elif RDG_EVENTS == RDG_EVENTS_STRING_COPY
-	EventNameStorage = MoveTemp(Other.EventNameStorage);
-	EventName = *EventNameStorage;
-	Other.EventName = nullptr;
+	EventFormat = Other.EventFormat;
+	Other.EventFormat = nullptr;
+	FormatedEventName = MoveTemp(Other.FormatedEventName);
 #endif
 	return *this;
 }
@@ -171,7 +171,16 @@ inline FRDGEventName& FRDGEventName::operator=(FRDGEventName&& Other)
 inline const TCHAR* FRDGEventName::GetTCHAR() const
 {
 #if RDG_EVENTS == RDG_EVENTS_STRING_REF || RDG_EVENTS == RDG_EVENTS_STRING_COPY
-	return EventName;
+	#if RDG_EVENTS == RDG_EVENTS_STRING_COPY
+		if (!FormatedEventName.IsEmpty())
+		{
+			return *FormatedEventName;
+		}
+	#endif
+
+	// The event has not been formated, at least return the event format to have
+	// error messages that give some clue when GetEmitRDGEvents() == false.
+	return EventFormat;
 #else
 	// Render graph draw events have been completely compiled for CPU performance reasons.
 	return TEXT("!!!Unavailable RDG event name: need RDG_EVENTS>=0 and r.RDG.EmitWarnings=1 or -rdgdebug!!!");
