Fixes crashes in FRDGEventName due to passding member to va_start()

diff --git a/Engine/Source/Runtime/RenderCore/Private/RenderGraphEvent.cpp b/Engine/Source/Runtime/RenderCore/Private/RenderGraphEvent.cpp
--- a/Engine/Source/Runtime/RenderCore/Private/RenderGraphEvent.cpp
+++ b/Engine/Source/Runtime/RenderCore/Private/RenderGraphEvent.cpp
@@ -13,10 +13,10 @@ FRDGEventName::FRDGEventName(const TCHAR* InEventFormat, ...)
 	if (GetEmitRDGEvents())
 	{
 		va_list VAList;
-		va_start(VAList, EventFormat);
+		va_start(VAList, InEventFormat);
 		TCHAR TempStr[256];
 		// Build the string in the temp buffer
-		FCString::GetVarArgs(TempStr, ARRAY_COUNT(TempStr), EventFormat, VAList);
+		FCString::GetVarArgs(TempStr, ARRAY_COUNT(TempStr), InEventFormat, VAList);
 		va_end(VAList);
 
 		FormatedEventName = TempStr;
