Fixes compilation failure on clang.

diff --git a/Engine/Source/Runtime/RenderCore/Public/RenderGraphEvent.inl b/Engine/Source/Runtime/RenderCore/Public/RenderGraphEvent.inl
--- a/Engine/Source/Runtime/RenderCore/Public/RenderGraphEvent.inl
+++ b/Engine/Source/Runtime/RenderCore/Public/RenderGraphEvent.inl
@@ -150,7 +150,7 @@ inline FRDGEventName& FRDGEventName::operator=(const FRDGEventName& Other)
 	EventFormat = Other.EventFormat;
 #elif RDG_EVENTS == RDG_EVENTS_STRING_COPY
 	EventFormat = Other.EventFormat;
-	FormatedEventName = FormatedEventName;
+	FormatedEventName = Other.FormatedEventName;
 #endif
 	return *this;
 }
