Validates the pixel format when creating a RDG texture.

diff --git a/Engine/Source/Runtime/RenderCore/Public/RenderGraphBuilder.inl b/Engine/Source/Runtime/RenderCore/Public/RenderGraphBuilder.inl
--- a/Engine/Source/Runtime/RenderCore/Public/RenderGraphBuilder.inl
+++ b/Engine/Source/Runtime/RenderCore/Public/RenderGraphBuilder.inl
@@ -71,7 +71,12 @@ inline FRDGTextureRef FRDGBuilder::CreateTexture(
 	{
 		checkf(DebugName, TEXT("Creating a render graph texture requires a valid debug name."));
 		checkf(!bHasExecuted, TEXT("Render graph texture %s needs to be created before the builder execution."), DebugName);
+
+		// Validate the pixel format.
 		checkf(Desc.Format != PF_Unknown, TEXT("Illegal to create texture %s with an invalid pixel format."), DebugName);
+		checkf(Desc.Format < PF_MAX, TEXT("Illegal to create texture %s with invalid FPooledRenderTargetDesc::Format."), DebugName);
+		checkf(GPixelFormats[Desc.Format].Supported, TEXT("Failed to create texture %s with pixel format %s because it is not supported."),
+			DebugName, GPixelFormats[Desc.Format].Name);
 
 		const bool bCanHaveUAV = Desc.TargetableFlags & TexCreate_UAV;
 		const bool bIsMSAA = Desc.NumSamples > 1;
