Improves visualise texture tool with more information on the draw event of the pass and avoid capturing resource with undesired mip level.

diff --git a/Engine/Source/Runtime/RenderCore/Private/RenderGraphBuilder.cpp b/Engine/Source/Runtime/RenderCore/Private/RenderGraphBuilder.cpp
--- a/Engine/Source/Runtime/RenderCore/Private/RenderGraphBuilder.cpp
+++ b/Engine/Source/Runtime/RenderCore/Private/RenderGraphBuilder.cpp
@@ -47,6 +47,7 @@ const int32 GRDGDebug = 0;
 
 bool GetEmitRDGEvents()
 {
+	check(IsInRenderingThread());
 #if RDG_EVENTS != RDG_EVENTS_NONE
 	return GetEmitDrawEvents() || GRDGDebug;
 #else
@@ -481,12 +482,13 @@ void FRDGBuilder::VisualizePassOutputs(const FRDGPass* Pass)
 		{
 			if (FRDGTextureUAVRef UAV = Parameter.GetAsTextureUAV())
 			{
-				if (FRDGTextureRef Texture = UAV->Desc.Texture)
+				FRDGTextureRef Texture = UAV->Desc.Texture;
+				check(Texture);
+
+				int32 CaptureId = GVisualizeTexture.ShouldCapture(Texture->Name);
+				if (CaptureId != FVisualizeTexture::kInvalidCaptureId && UAV->Desc.MipLevel == GVisualizeTexture.CustomMip)
 				{
-					if (GVisualizeTexture.ShouldCapture(Texture->Name))
-					{
-						GVisualizeTexture.CreateContentCapturePass(*this, Texture);
-					}
+					GVisualizeTexture.CreateContentCapturePass(*this, Texture, CaptureId);
 				}
 			}
 		}
@@ -501,9 +503,16 @@ void FRDGBuilder::VisualizePassOutputs(const FRDGPass* Pass)
 			{
 				const bool bHasStoreAction = DepthStencil.GetDepthStoreAction() != ERenderTargetStoreAction::ENoAction || DepthStencil.GetStencilStoreAction() != ERenderTargetStoreAction::ENoAction;
 
-				if (bHasStoreAction && GVisualizeTexture.ShouldCapture(Texture->Name))
+				if (bHasStoreAction)
 				{
-					GVisualizeTexture.CreateContentCapturePass(*this, Texture);
+					// Depth render target binding can only be done on mip level 0.
+					const int32 MipLevel = 0;
+
+					int32 CaptureId = GVisualizeTexture.ShouldCapture(Texture->Name);
+					if (CaptureId != FVisualizeTexture::kInvalidCaptureId && MipLevel == GVisualizeTexture.CustomMip)
+					{
+						GVisualizeTexture.CreateContentCapturePass(*this, Texture, CaptureId);
+					}
 				}
 			}
 
@@ -517,9 +526,14 @@ void FRDGBuilder::VisualizePassOutputs(const FRDGPass* Pass)
 				{
 					const bool bHasStoreAction = RenderTarget.GetStoreAction() != ERenderTargetStoreAction::ENoAction;
 
-					if (bHasStoreAction && GVisualizeTexture.ShouldCapture(Texture->Name))
+					if (bHasStoreAction)
 					{
-						GVisualizeTexture.CreateContentCapturePass(*this, Texture);
+						int32 CaptureId = GVisualizeTexture.ShouldCapture(Texture->Name);
+						
+						if (CaptureId != FVisualizeTexture::kInvalidCaptureId && RenderTarget.GetMipIndex() == GVisualizeTexture.CustomMip)
+						{
+							GVisualizeTexture.CreateContentCapturePass(*this, Texture, CaptureId);
+						}
 					}
 				}
 				else
diff --git a/Engine/Source/Runtime/RenderCore/Private/RenderGraphEvent.cpp b/Engine/Source/Runtime/RenderCore/Private/RenderGraphEvent.cpp
--- a/Engine/Source/Runtime/RenderCore/Private/RenderGraphEvent.cpp
+++ b/Engine/Source/Runtime/RenderCore/Private/RenderGraphEvent.cpp
@@ -3,7 +3,6 @@
 #include "RenderGraphEvent.h"
 #include "RenderGraphBuilder.h"
 
-extern bool GetEmitRDGEvents();
 
 #if RDG_EVENTS == RDG_EVENTS_STRING_COPY
 
diff --git a/Engine/Source/Runtime/RenderCore/Private/VisualizeTexture.cpp b/Engine/Source/Runtime/RenderCore/Private/VisualizeTexture.cpp
--- a/Engine/Source/Runtime/RenderCore/Private/VisualizeTexture.cpp
+++ b/Engine/Source/Runtime/RenderCore/Private/VisualizeTexture.cpp
@@ -164,8 +164,10 @@ void FVisualizeTexture::ReleaseDynamicRHI()
 	StencilSRV.SafeRelease();
 }
 
-void FVisualizeTexture::CreateContentCapturePass(FRDGBuilder& GraphBuilder, const FRDGTextureRef SrcTexture)
+void FVisualizeTexture::CreateContentCapturePass(FRDGBuilder& GraphBuilder, const FRDGTextureRef SrcTexture, int32 CaptureId)
 {
+	check(CaptureId >= 0);
+
 #if !(UE_BUILD_SHIPPING || UE_BUILD_TEST)
 	if (!SrcTexture || !SrcTexture->Desc.IsValid())
 	{
@@ -183,7 +185,7 @@ void FVisualizeTexture::CreateContentCapturePass(FRDGBuilder& GraphBuilder, cons
 
 	FRDGTextureRef CopyTexture;
 	{
-		FIntPoint Size = SrcTexture->Desc.Extent;
+		FIntPoint Size = SrcDesc.Extent;
 
 		// clamp to reasonable value to prevent crash
 		Size.X = FMath::Max(Size.X, 1);
@@ -196,7 +198,7 @@ void FVisualizeTexture::CreateContentCapturePass(FRDGBuilder& GraphBuilder, cons
 		CopyTexture = GraphBuilder.CreateTexture(CopyDesc, TEXT("VisualizeTexture"));
 	}
 
-	FIntPoint RTExtent = SrcTexture->Desc.Extent;
+	FIntPoint RTExtent = SrcDesc.Extent;
 
 	uint32 LocalVisualizeTextureInputMapping = UVInputMapping;
 
@@ -257,8 +259,35 @@ void FVisualizeTexture::CreateContentCapturePass(FRDGBuilder& GraphBuilder, cons
 
 	TShaderMapRef<FVisualizeTexturePS> PixelShader(ShaderMap, PermutationVector);
 
+	FString ExtendedDrawEvent;
+	if (GetEmitRDGEvents())
+	{
+		// If this is a 3D texture or texture array, precise.
+		if (SrcDesc.Depth > 0)
+		{
+			if (SrcDesc.bIsArray)
+			{
+				ExtendedDrawEvent += FString::Printf(TEXT(" ArraySize=%d CapturedSlice=%d"), SrcDesc.Depth, ArrayIndex);
+			}
+			else
+			{
+				ExtendedDrawEvent += FString::Printf(TEXT("x%d CapturedSlice=%d"), SrcDesc.Depth, ArrayIndex);
+			}
+		}
+
+		// Precise the mip level being captured in the mip level when there is a mip chain.
+		if (SrcDesc.NumMips > 1)
+		{
+			ExtendedDrawEvent += FString::Printf(TEXT(" Mips=%d CapturedMip=%d"), SrcDesc.NumMips, CustomMip);
+		}
+	}
+
 	GraphBuilder.AddPass(
-		RDG_EVENT_NAME("VisualizeTextureCapture(%s)", SrcTexture->Name),
+		RDG_EVENT_NAME("VisualizeTextureCapture(%s@%d %s %dx%d%s)",
+			SrcTexture->Name, CaptureId,
+			GPixelFormats[SrcDesc.Format].Name,
+			SrcDesc.Extent.X, SrcDesc.Extent.Y,
+			*ExtendedDrawEvent),
 		PassParameters,
 		ERenderGraphPassFlags::None,
 		[this, PassParameters, ShaderMap, PixelShader, RTExtent](FRHICommandList& RHICmdList)
@@ -376,10 +405,10 @@ void FVisualizeTexture::CreateContentCapturePass(FRDGBuilder& GraphBuilder, cons
 #endif // !(UE_BUILD_SHIPPING || UE_BUILD_TEST)
 }
 
-bool FVisualizeTexture::ShouldCapture(const TCHAR* DebugName)
+int32 FVisualizeTexture::ShouldCapture(const TCHAR* DebugName)
 {
 #if (UE_BUILD_SHIPPING || UE_BUILD_TEST)
-	return false;
+	return FVisualizeTexture::kInvalidCaptureId;
 #else
 	if (!bEnabled)
 	{
@@ -399,16 +428,18 @@ bool FVisualizeTexture::ShouldCapture(const TCHAR* DebugName)
 	// First check if we need to find anything to avoid string the comparison
 	if (!ObservedDebugName.IsEmpty() && ObservedDebugName == DebugName)
 	{
+		uint32 CaptureId = *UsageCountPtr;
+
 		// if multiple times reused during the frame, is that the one we want to look at?
-		if (*UsageCountPtr == ObservedDebugNameReusedGoal || ObservedDebugNameReusedGoal == 0xffffffff)
+		if (CaptureId == ObservedDebugNameReusedGoal || ObservedDebugNameReusedGoal == 0xffffffff)
 		{
-			*UsageCountPtr = *UsageCountPtr + 1;
-			return true;
+			*UsageCountPtr = CaptureId + 1;
+			return int32(CaptureId);
 		}
 	}
 	// only needed for VisualizeTexture (todo: optimize out when possible)
 	*UsageCountPtr = *UsageCountPtr + 1;
-	return false;
+	return FVisualizeTexture::kInvalidCaptureId;
 #endif
 }
 
@@ -424,7 +455,8 @@ void FVisualizeTexture::SetCheckPoint(FRHICommandList& RHICmdList, const IPooled
 
 	const TCHAR* DebugName = PooledRenderTarget->GetDesc().DebugName;
 
-	if (!ShouldCapture(DebugName))
+	int32 CaptureId = ShouldCapture(DebugName);
+	if (CaptureId == FVisualizeTexture::kInvalidCaptureId)
 	{
 		return;
 	}
@@ -448,7 +480,7 @@ void FVisualizeTexture::SetCheckPoint(FRHICommandList& RHICmdList, const IPooled
 	TRefCountPtr<IPooledRenderTarget> PooledRenderTargetRef(const_cast<IPooledRenderTarget*>(PooledRenderTarget));
 	FRDGTextureRef TextureToCapture = GraphBuilder.RegisterExternalTexture(PooledRenderTargetRef, DebugName);
 
-	CreateContentCapturePass(GraphBuilder, TextureToCapture);
+	CreateContentCapturePass(GraphBuilder, TextureToCapture, CaptureId);
 	GraphBuilder.Execute();
 
 	if (&RHICmdList != &RHICmdListIm)
diff --git a/Engine/Source/Runtime/RenderCore/Public/RenderGraphEvent.h b/Engine/Source/Runtime/RenderCore/Public/RenderGraphEvent.h
--- a/Engine/Source/Runtime/RenderCore/Public/RenderGraphEvent.h
+++ b/Engine/Source/Runtime/RenderCore/Public/RenderGraphEvent.h
@@ -5,6 +5,11 @@
 #include "RenderGraphDefinitions.h"
 #include "ProfilingDebugging/RealtimeGPUProfiler.h"
 
+
+/** Returns whether the current frame is emitting render graph events. */
+extern RENDERCORE_API bool GetEmitRDGEvents();
+
+
 /** A helper profiler class for tracking and evaluating hierarchical scopes in the context of render graph. */
 template <typename TScopeType>
 class FRDGScopeStack final
diff --git a/Engine/Source/Runtime/RenderCore/Public/VisualizeTexture.h b/Engine/Source/Runtime/RenderCore/Public/VisualizeTexture.h
--- a/Engine/Source/Runtime/RenderCore/Public/VisualizeTexture.h
+++ b/Engine/Source/Runtime/RenderCore/Public/VisualizeTexture.h
@@ -18,6 +18,8 @@ class RENDERCORE_API FVisualizeTexture : public FRenderResource
 public:
 	FVisualizeTexture();
 
+	static constexpr int32 kInvalidCaptureId = -1;
+
 #if WITH_ENGINE
 	virtual void ReleaseDynamicRHI() override;
 
@@ -102,10 +104,11 @@ private:
 	ERHIFeatureLevel::Type FeatureLevel;
 
 #if WITH_ENGINE
-	bool ShouldCapture(const TCHAR* DebugName);
+	/** Determine whether a texture should be captured for debugging purpose, and return the capture id if needed, or kInvalidCaptureId otherwise. */
+	int32 ShouldCapture(const TCHAR* DebugName);
 
 	/** Create a pass capturing a texture. */
-	void CreateContentCapturePass(FRDGBuilder& GraphBuilder, FRDGTextureRef Texture);
+	void CreateContentCapturePass(FRDGBuilder& GraphBuilder, FRDGTextureRef Texture, int32 CaptureId);
 #endif
 
 	friend class FRDGBuilder;
diff --git a/Engine/Source/Runtime/Renderer/Private/VisualizeTexturePresent.cpp b/Engine/Source/Runtime/Renderer/Private/VisualizeTexturePresent.cpp
--- a/Engine/Source/Runtime/Renderer/Private/VisualizeTexturePresent.cpp
+++ b/Engine/Source/Runtime/Renderer/Private/VisualizeTexturePresent.cpp
@@ -373,6 +373,7 @@ void FVisualizeTexturePresent::PresentContent(FRHICommandListImmediate& RHICmdLi
 		}
 	}
 
+	// TODO(RDG): delete this old technic of visualizing a buffer based on id that predates when render target had names.
 	if (GVisualizeTexture.Mode != 0)
 	{
 		// old mode is used, lets copy the specified texture to do it similar to the new system
@@ -383,7 +384,7 @@ void FVisualizeTexturePresent::PresentContent(FRHICommandListImmediate& RHICmdLi
 
 			TRefCountPtr<IPooledRenderTarget> ElementRefCount = Element;
 
-			GVisualizeTexture.CreateContentCapturePass(GraphBuilder, GraphBuilder.RegisterExternalTexture(ElementRefCount));
+			GVisualizeTexture.CreateContentCapturePass(GraphBuilder, GraphBuilder.RegisterExternalTexture(ElementRefCount), /* CaptureId = */ 0);
 			GraphBuilder.Execute();
 		}
 	}
