Integrated render graph fixes from 4.23.

diff --git a/Engine/Source/Runtime/Engine/Private/TextureRenderTarget2D.cpp b/Engine/Source/Runtime/Engine/Private/TextureRenderTarget2D.cpp
--- a/Engine/Source/Runtime/Engine/Private/TextureRenderTarget2D.cpp
+++ b/Engine/Source/Runtime/Engine/Private/TextureRenderTarget2D.cpp
@@ -53,6 +53,11 @@ FTextureResource* UTextureRenderTarget2D::CreateResource()
 	if (bAutoGenerateMips)
 	{
 		NumMips = FGenericPlatformMath::CeilToInt(FGenericPlatformMath::Log2(FGenericPlatformMath::Max(SizeX, SizeY)));
+
+		if (RHIRequiresComputeGenerateMips())
+		{
+			bCanCreateUAV = 1;
+		}
 	}
 	else
 	{
diff --git a/Engine/Source/Runtime/RenderCore/Private/GenerateMips.cpp b/Engine/Source/Runtime/RenderCore/Private/GenerateMips.cpp
--- a/Engine/Source/Runtime/RenderCore/Private/GenerateMips.cpp
+++ b/Engine/Source/Runtime/RenderCore/Private/GenerateMips.cpp
@@ -126,12 +126,16 @@ void FGenerateMips::Compute(FRHICommandListImmediate& RHIImmCmdList, FRHITexture
 			FMath::Max(DestTextureSizeY / MIPSSHADER_NUMTHREADS, 1),
 			1);
 		//Pass added per mip level to be written.
-		FComputeShaderUtils::AddPass(
-			GraphBuilder,
+		ClearUnusedGraphResources(*ComputeShader, PassParameters);
+
+		GraphBuilder.AddPass(
 			RDG_EVENT_NAME("Generate2DTextureMips DestMipLevel=%d", MipLevel),
-			*ComputeShader,
 			PassParameters,
-			GenMipsGroupCount);
+			ERenderGraphPassFlags::Compute | ERDGPassFlags::GenerateMips,
+			[PassParameters, ComputeShader, GenMipsGroupCount](FRHICommandList& RHICmdList)
+		{
+			FComputeShaderUtils::Dispatch(RHICmdList, *ComputeShader, *PassParameters, GenMipsGroupCount);
+		});
 	}
 	GraphBuilder.QueueTextureExtraction(GraphTexture, &GenMipsStruct->RenderTarget);
 	GraphBuilder.Execute();	
diff --git a/Engine/Source/Runtime/RenderCore/Private/RenderGraphBarrierBatcher.cpp b/Engine/Source/Runtime/RenderCore/Private/RenderGraphBarrierBatcher.cpp
--- a/Engine/Source/Runtime/RenderCore/Private/RenderGraphBarrierBatcher.cpp
+++ b/Engine/Source/Runtime/RenderCore/Private/RenderGraphBarrierBatcher.cpp
@@ -48,63 +48,71 @@ namespace
 		return EResourceTransitionPipeline::EGfxToGfx;
 	}
 
-	inline EResourceTransitionAccess GetResourceTransitionAccess(FRDGResourceState::EAccess AccessAfter)
-	{
-		return AccessAfter == FRDGResourceState::EAccess::Write ? EResourceTransitionAccess::EWritable : EResourceTransitionAccess::EReadable;
-	}
+} //! namespace
 
-	inline EResourceTransitionAccess GetResourceTransitionAccessForUAV(FRDGResourceState::EAccess AccessBefore, FRDGResourceState::EAccess AccessAfter)
+FRDGBarrierBatcher::FRDGBarrierBatcher(FRHICommandList& InRHICmdList, const FRDGPass* InPass)
+	: RHICmdList(InRHICmdList)
+	, Pass(InPass)
+{
+	if (Pass)
 	{
-		switch (AccessAfter)
-		{
-		case FRDGResourceState::EAccess::Read:
-			return EResourceTransitionAccess::EReadable;
-		case FRDGResourceState::EAccess::Write:
-			// If doing a Write -> Write transition, we use a UAV barrier.
-			if (AccessBefore == FRDGResourceState::EAccess::Write)
-			{
-				return EResourceTransitionAccess::ERWBarrier;
-			}
-			else
-			{
-				return EResourceTransitionAccess::EWritable;
-			}
-		}
-		check(false);
-		return EResourceTransitionAccess::EMaxAccess;
+		bIsGeneratingMips = Pass->IsGenerateMips();
+		Pipeline = Pass->IsCompute() ? FRDGResourceState::EPipeline::Compute : FRDGResourceState::EPipeline::Graphics;
 	}
-
-} //! namespace
+}
 
 FRDGBarrierBatcher::~FRDGBarrierBatcher()
 {
-	check(!bAllowQueueing);
-}
+	for (FRHITexture* RHITexture : TextureUpdateMultiFrameBegins)
+	{
+		RHICmdList.BeginUpdateMultiFrameResource(RHITexture);
+	}
 
-void FRDGBarrierBatcher::Begin()
-{
-#if RDG_ENABLE_DEBUG
-	check(!bAllowQueueing);
+	for (FRHIUnorderedAccessView* RHIUAV : UAVUpdateMultiFrameBegins)
+	{
+		RHICmdList.BeginUpdateMultiFrameResource(RHIUAV);
+	}
+
+	for (auto& Element : TextureBatchMap)
+	{
+		FTransitionParameters TransitionParameters = Element.Key;
+		FTextureBatch& Batch = Element.Value;
+		RHICmdList.TransitionResources(TransitionParameters.TransitionAccess, Batch.GetData(), Batch.Num());
+	}
 
-	check(!TextureUpdateMultiFrameBegins.Num());
-	check(!TextureUpdateMultiFrameEnds.Num());
-	check(!TextureBatchMap.Num());
+	for (auto& Element : UAVBatchMap)
+	{
+		FTransitionParameters TransitionParameters = Element.Key;
+		FUAVBatch& Batch = Element.Value;
+		RHICmdList.TransitionResources(TransitionParameters.TransitionAccess, TransitionParameters.TransitionPipeline, Batch.GetData(), Batch.Num());
+	}
 
-	check(!UAVUpdateMultiFrameBegins.Num());
-	check(!UAVUpdateMultiFrameEnds.Num());
-	check(!UAVBatchMap.Num());
-#endif
+	for (FRHITexture* RHITexture : TextureUpdateMultiFrameEnds)
+	{
+		RHICmdList.EndUpdateMultiFrameResource(RHITexture);
+	}
 
-	bAllowQueueing = true;
+	for (FRHIUnorderedAccessView* RHIUAV : UAVUpdateMultiFrameEnds)
+	{
+		RHICmdList.EndUpdateMultiFrameResource(RHIUAV);
+	}
 }
 
-void FRDGBarrierBatcher::QueueTransitionTexture(FRDGTexture* Texture, FRDGResourceState StateAfter)
+void FRDGBarrierBatcher::QueueTransitionTexture(FRDGTexture* Texture, FRDGResourceState::EAccess AccessAfter)
 {
 	check(Texture);
 
+	// Texture transitions are ignored when generating mips, since the render target binding call or UAV will
+	// perform the subresource transition.
+	if (bIsGeneratingMips)
+	{
+		return;
+	}
+
 	const FRDGResourceState StateBefore = Texture->State;
+	const FRDGResourceState StateAfter(Pass, Pipeline, AccessAfter);
 
-	ValidateTransition(StateBefore, StateAfter, Texture);
+	ValidateTransition(Texture, StateBefore, StateAfter);
 
 	if (StateBefore != StateAfter)
 	{
@@ -126,12 +134,12 @@ void FRDGBarrierBatcher::QueueTransitionTexture(FRDGTexture* Texture, FRDGResour
 			FTextureBatch& TextureBatch = TextureBatchMap.FindOrAdd(TransitionParameters);
 			TextureBatch.Reserve(kBatchReservationSize);
 
-#if RDG_ENABLE_DEBUG
+			#if RDG_ENABLE_DEBUG
 			{
 				// We should have filtered out duplicates in the first branch of this function.
-				check(TextureBatch.Find(RHITexture) == INDEX_NONE);
+				ensure(TextureBatch.Find(RHITexture) == INDEX_NONE);
 			}
-#endif
+			#endif
 
 			TextureBatch.Add(RHITexture);
 		}
@@ -148,14 +156,15 @@ void FRDGBarrierBatcher::QueueTransitionTexture(FRDGTexture* Texture, FRDGResour
 void FRDGBarrierBatcher::QueueTransitionUAV(
 	FRHIUnorderedAccessView* UAV,
 	FRDGTrackedResource* UnderlyingResource,
-	FRDGResourceState StateAfter)
+	FRDGResourceState::EAccess AccessAfter)
 {
 	check(UAV);
 	check(UnderlyingResource);
 
 	const FRDGResourceState StateBefore = UnderlyingResource->State;
+	const FRDGResourceState StateAfter(Pass, Pipeline, AccessAfter);
 
-	ValidateTransition(StateBefore, StateAfter, UnderlyingResource);
+	ValidateTransition(UnderlyingResource, StateBefore, StateAfter);
 
 	if (StateBefore != StateAfter)
 	{
@@ -175,12 +184,12 @@ void FRDGBarrierBatcher::QueueTransitionUAV(
 			FUAVBatch& UAVBatch = UAVBatchMap.FindOrAdd(TransitionParameters);
 			UAVBatch.Reserve(kBatchReservationSize);
 
-#if RDG_ENABLE_DEBUG
+			#if RDG_ENABLE_DEBUG
 			{
 				// We should have filtered out duplicates in the first branch of this function.
-				check(UAVBatch.Find(UAV) == INDEX_NONE);
+				ensure(UAVBatch.Find(UAV) == INDEX_NONE);
 			}
-#endif
+			#endif
 
 			UAVBatch.Add(UAV);
 		}
@@ -194,53 +203,7 @@ void FRDGBarrierBatcher::QueueTransitionUAV(
 	}
 }
 
-void FRDGBarrierBatcher::End(FRHICommandList& RHICmdList)
-{
-	check(bAllowQueueing);
-	bAllowQueueing = false;
-
-	for (FRHITexture* RHITexture : TextureUpdateMultiFrameBegins)
-	{
-		RHICmdList.BeginUpdateMultiFrameResource(RHITexture);
-	}
-	TextureUpdateMultiFrameBegins.Empty();
-
-	for (FRHIUnorderedAccessView* RHIUAV : UAVUpdateMultiFrameBegins)
-	{
-		RHICmdList.BeginUpdateMultiFrameResource(RHIUAV);
-	}
-	UAVUpdateMultiFrameBegins.Empty();
-
-	for (auto& Element : TextureBatchMap)
-	{
-		FTransitionParameters TransitionParameters = Element.Key;
-		FTextureBatch& Batch = Element.Value;
-		RHICmdList.TransitionResources(TransitionParameters.TransitionAccess, Batch.GetData(), Batch.Num());
-	}
-	TextureBatchMap.Empty();
-
-	for (auto& Element : UAVBatchMap)
-	{
-		FTransitionParameters TransitionParameters = Element.Key;
-		FUAVBatch& Batch = Element.Value;
-		RHICmdList.TransitionResources(TransitionParameters.TransitionAccess, TransitionParameters.TransitionPipeline, Batch.GetData(), Batch.Num());
-	}
-	UAVBatchMap.Empty();
-
-	for (FRHITexture* RHITexture : TextureUpdateMultiFrameEnds)
-	{
-		RHICmdList.EndUpdateMultiFrameResource(RHITexture);
-	}
-	TextureUpdateMultiFrameEnds.Empty();
-
-	for (FRHIUnorderedAccessView* RHIUAV : UAVUpdateMultiFrameEnds)
-	{
-		RHICmdList.EndUpdateMultiFrameResource(RHIUAV);
-	}
-	UAVUpdateMultiFrameEnds.Empty();
-}
-
-void FRDGBarrierBatcher::ValidateTransition(FRDGResourceState StateBefore, FRDGResourceState StateAfter, const FRDGTrackedResource* Resource)
+void FRDGBarrierBatcher::ValidateTransition(const FRDGTrackedResource* Resource, FRDGResourceState StateBefore, FRDGResourceState StateAfter)
 {
 #if RDG_ENABLE_DEBUG
 	check(StateAfter.Pipeline != FRDGResourceState::EPipeline::MAX);
@@ -249,11 +212,44 @@ void FRDGBarrierBatcher::ValidateTransition(FRDGResourceState StateBefore, FRDGR
 	if (StateBefore != StateAfter && StateAfter.Pass)
 	{
 		// We allow duplicate transitions of the same resource within the same pass, but not conflicting ones.
-		checkf(
+		ensureMsgf(
 			StateBefore.Pass != StateAfter.Pass,
 			TEXT("Pass %s attempted to transition resource %s to different states. Make sure the resource isn't being used\n")
 			TEXT("for both read and write at the same time. This can occur if the resource is used as both an SRV and UAV, or\n")
-			TEXT("SRV and Render Target, for example.\n"), StateAfter.Pass->GetName(), Resource->Name);
+			TEXT("SRV and Render Target, for example. If this pass is meant to generate mip maps, make sure the GenerateMips flag\n")
+			TEXT("is set.\n"), StateAfter.Pass->GetName(), Resource->Name);
 	}
 #endif
-}
\ No newline at end of file
+}
+
+EResourceTransitionAccess FRDGBarrierBatcher::GetResourceTransitionAccess(FRDGResourceState::EAccess AccessAfter) const
+{
+	return AccessAfter == FRDGResourceState::EAccess::Write ? EResourceTransitionAccess::EWritable : EResourceTransitionAccess::EReadable;
+}
+
+EResourceTransitionAccess FRDGBarrierBatcher::GetResourceTransitionAccessForUAV(FRDGResourceState::EAccess AccessBefore, FRDGResourceState::EAccess AccessAfter) const
+{
+	switch (AccessAfter)
+	{
+	case FRDGResourceState::EAccess::Read:
+		return EResourceTransitionAccess::EReadable;
+
+	case FRDGResourceState::EAccess::Write:
+		// Mip-map generation uses its own barrier.
+		if (bIsGeneratingMips)
+		{
+			return EResourceTransitionAccess::ERWSubResBarrier;
+		}
+		// If doing a Write -> Write transition, we use a UAV barrier.
+		else if (AccessBefore == FRDGResourceState::EAccess::Write)
+		{
+			return EResourceTransitionAccess::ERWBarrier;
+		}
+		else
+		{
+			return EResourceTransitionAccess::EWritable;
+		}
+	}
+	check(false);
+	return EResourceTransitionAccess::EMaxAccess;
+}
diff --git a/Engine/Source/Runtime/RenderCore/Private/RenderGraphBarrierBatcher.h b/Engine/Source/Runtime/RenderCore/Private/RenderGraphBarrierBatcher.h
--- a/Engine/Source/Runtime/RenderCore/Private/RenderGraphBarrierBatcher.h
+++ b/Engine/Source/Runtime/RenderCore/Private/RenderGraphBarrierBatcher.h
@@ -8,24 +8,25 @@
 class RENDERCORE_API FRDGBarrierBatcher final
 {
 public:
-	FRDGBarrierBatcher() = default;
+	/** RAII initialization of batcher. Batcher will flush all queued transitions in the destructor.
+	 *  @param Pass The current pass if performing inter-pass barriers. This can be null (e.g. for post-execution barriers).
+	 */
+	FRDGBarrierBatcher(FRHICommandList& RHICmdList, const FRDGPass* Pass);
 	~FRDGBarrierBatcher();
 
-	void Begin();
-
-	void QueueTransitionTexture(FRDGTexture* Texture, FRDGResourceState StateAfter);
+	/** Queues a transition of the texture to the requested access state. */
+	void QueueTransitionTexture(FRDGTexture* Texture, FRDGResourceState::EAccess AccessAfter);
 
+	/** Queues a transition of the UAV to the requested access state. The current state of the resource is checked and the transition
+	 *  is only performed if a change occurs. It is valid to queue the same state multiple times. However, it is invalid to queue a
+	 *  resource into multiple states at the same time. The underlying resource is the texture / buffer instance referenced by the view.
+	 */
 	void QueueTransitionUAV(
 		FRHIUnorderedAccessView* UAV,
 		FRDGTrackedResource* UnderlyingResource,
-		FRDGResourceState StateAfter);
-
-	/** Submits and clears all queued transitions to the provided RHI command list. */
-	void End(FRHICommandList& RHICmdList);
+		FRDGResourceState::EAccess AccessAfter);
 
 private:
-	static void ValidateTransition(FRDGResourceState StateBefore, FRDGResourceState StateAfter, const FRDGTrackedResource* Resource);
-
 	struct FTransitionParameters
 	{
 		FTransitionParameters() = default;
@@ -59,6 +60,8 @@ private:
 	using FTextureBatch = TArray<FRHITexture*, SceneRenderingAllocator>;
 	using FTextureBatchMap = TMap<FTransitionParameters, FTextureBatch, SceneRenderingSetAllocator, TBatchMapKeyFuncs<FTextureBatch>>;
 
+	FRHICommandList& RHICmdList;
+
 	FTextureBatch TextureUpdateMultiFrameBegins;
 	FTextureBatch TextureUpdateMultiFrameEnds;
 	FTextureBatchMap TextureBatchMap;
@@ -67,5 +70,12 @@ private:
 	FUAVBatch UAVUpdateMultiFrameEnds;
 	FUAVBatchMap UAVBatchMap;
 
-	bool bAllowQueueing = false;
-};
\ No newline at end of file
+	void ValidateTransition(const FRDGTrackedResource* Resource, FRDGResourceState StateBefore, FRDGResourceState StateAfter);
+
+	EResourceTransitionAccess GetResourceTransitionAccess(FRDGResourceState::EAccess AccessAfter) const;
+	EResourceTransitionAccess GetResourceTransitionAccessForUAV(FRDGResourceState::EAccess AccessBefore, FRDGResourceState::EAccess AccessAfter) const;
+
+	const FRDGPass* Pass = nullptr;
+	FRDGResourceState::EPipeline Pipeline = FRDGResourceState::EPipeline::Graphics;
+	bool bIsGeneratingMips = false;
+};
diff --git a/Engine/Source/Runtime/RenderCore/Private/RenderGraphBuilder.cpp b/Engine/Source/Runtime/RenderCore/Private/RenderGraphBuilder.cpp
--- a/Engine/Source/Runtime/RenderCore/Private/RenderGraphBuilder.cpp
+++ b/Engine/Source/Runtime/RenderCore/Private/RenderGraphBuilder.cpp
@@ -55,6 +55,17 @@ bool GetEmitRDGEvents()
 #endif
 }
 
+bool IsRDGDebugEnabled()
+{
+	return GRDGDebug != 0;
+}
+
+bool IsRDGImmediateModeEnabled()
+{
+	return GRDGImmediateMode != 0;
+}
+
+
 void InitRenderGraph()
 {
 #if RDG_ENABLE_DEBUG_WITH_ENGINE
@@ -107,9 +118,6 @@ void EmitRDGWarning(const FString& WarningMessage)
 #endif
 }
 
-#define EmitRDGWarningf(WarningMessageFormat, ...) \
-	EmitRDGWarning(FString::Printf(WarningMessageFormat, ##__VA_ARGS__));
-
 void FRDGBuilder::TickPoolElements()
 {
 	GRenderGraphResourcePool.TickPoolElements();
@@ -122,34 +130,11 @@ FRDGBuilder::FRDGBuilder(FRHICommandListImmediate& InRHICmdList)
 	, StatScopeStack(RHICmdList)
 {}
 
-FRDGBuilder::~FRDGBuilder()
-{
-	#if RDG_ENABLE_DEBUG
-	{
-		checkf(bHasExecuted, TEXT("Render graph execution is required to ensure consistency with immediate mode."));
-	}
-	#endif
-}
-
 void FRDGBuilder::Execute()
 {
 	CSV_SCOPED_TIMING_STAT_EXCLUSIVE(FRDGBuilder_Execute);
 
-	#if RDG_ENABLE_DEBUG
-	{
-		checkf(!bHasExecuted, TEXT("Render graph execution should only happen once to ensure consistency with immediate mode."));
-
-		/** FRDGBuilder::AllocParameters() allocates shader parameter structure for the life time until pass execution.
-		 * But they are allocated on a FMemStack for CPU performance reason, and have their destructor called right after
-		 * the pass execution. Therefore allocating pass parameter unused by a FRDGBuilder::AddPass() can lead on a memory
-		 * leak of RHI resource that have been reference in the parameter structure.
-		 */
-		checkf(
-			AllocatedUnusedPassParameters.Num() == 0,
-			TEXT("%i pass parameter structure has been allocated with FRDGBuilder::AllocParameters(), but has not be used by a ")
-			TEXT("FRDGBuilder::AddPass() that can cause RHI resource leak."), AllocatedUnusedPassParameters.Num());
-	}
-	#endif
+	IF_RDG_ENABLE_DEBUG(Validation.ValidateExecuteBegin());
 
 	EventScopeStack.BeginExecute();
 	StatScopeStack.BeginExecute();
@@ -170,43 +155,19 @@ void FRDGBuilder::Execute()
 
 	ProcessDeferredInternalResourceQueries();
 
-	DestructPasses();
+	IF_RDG_ENABLE_DEBUG(Validation.ValidateExecuteEnd());
 
-	#if RDG_ENABLE_DEBUG
-	{
-		bHasExecuted = true;
-	}
-	#endif
+	DestructPasses();
 }
 
 void FRDGBuilder::AddPassInternal(FRDGPass* Pass)
 {
-	#if RDG_ENABLE_DEBUG
-	{
-		checkf(!bHasExecuted, TEXT("Render graph pass %s needs to be added before the builder execution."), Pass->GetName());
-
-		const void* ParameterStructData = Pass->GetParameters().GetContents();
-
-		/** The lifetime of each pass parameter structure must extend until deferred pass execution; therefore, it needs to be
-		 *  allocated with FRDGBuilder::AllocParameters(). Also, all references held by the parameter structure are released
-		 *  immediately after pass execution, so a pass parameter struct instance must be 1-to-1 with a pass instance (i.e. one
-		 *  per AddPass() call).
-		 */
-		checkf(
-			AllocatedUnusedPassParameters.Contains(ParameterStructData),
-			TEXT("The pass parameter structure has not been allocated for correct life time FRDGBuilder::AllocParameters() or has already ")
-			TEXT("been used by another previous FRDGBuilder::AddPass()."));
-
-		AllocatedUnusedPassParameters.Remove(ParameterStructData);
-	}
-	#endif
+	IF_RDG_ENABLE_DEBUG(Validation.ValidateAddPass(Pass));
 
 	Pass->EventScope = EventScopeStack.GetCurrentScope();
 	Pass->StatScope = StatScopeStack.GetCurrentScope();
 	Passes.Emplace(Pass);
 
-	ValidatePass(Pass);
-
 	if (GRDGImmediateMode)
 	{
 		ExecutePass(Pass);
@@ -215,256 +176,6 @@ void FRDGBuilder::AddPassInternal(FRDGPass* Pass)
 	VisualizePassOutputs(Pass);
 }
 
-void FRDGBuilder::ValidatePass(const FRDGPass* Pass) const
-{
-#if RDG_ENABLE_DEBUG
-	const FRenderTargetBindingSlots* RenderTargetBindingSlots = nullptr;
-
-	const TCHAR* PassName = Pass->GetName();
-	const bool bIsGraphics = Pass->IsGraphics();
-
-	FRDGPassParameterStruct ParameterStruct = Pass->GetParameters();
-
-	const uint32 ParameterCount = ParameterStruct.GetParameterCount();
-
-	for (uint32 ParameterIndex = 0; ParameterIndex < ParameterCount; ++ParameterIndex)
-	{
-		FRDGPassParameter Parameter = ParameterStruct.GetParameter(ParameterIndex);
-
-		switch (Parameter.GetType())
-		{
-		case UBMT_RDG_TEXTURE:
-		{
-			if (FRDGTextureRef Texture = Parameter.GetAsTexture())
-			{
-				checkf(Texture->HasBeenProduced(),
-					TEXT("Pass %s has a dependency over the texture %s that has never been produced."),
-					PassName, Texture->Name);
-			}
-		}
-		break;
-		case UBMT_RDG_TEXTURE_SRV:
-		{
-			if (FRDGTextureSRVRef SRV = Parameter.GetAsTextureSRV())
-			{
-				FRDGTextureRef Texture = SRV->Desc.Texture;
-
-				checkf(Texture->HasBeenProduced(),
-					TEXT("Pass %s has a dependency over the texture %s that has never been produced."),
-					PassName, Texture->Name);
-			}
-		}
-		break;
-		case UBMT_RDG_TEXTURE_UAV:
-		{
-			if (FRDGTextureUAVRef UAV = Parameter.GetAsTextureUAV())
-			{
-				FRDGTextureRef Texture = UAV->Desc.Texture;
-
-				Texture->MarkAsProducedBy(Pass);
-			}
-		}
-		break;
-		case UBMT_RDG_BUFFER:
-		{
-			if (FRDGBufferRef Buffer = Parameter.GetAsBuffer())
-			{
-				checkf(Buffer->HasBeenProduced(),
-					TEXT("Pass %s has a dependency over the buffer %s that has never been produced."),
-					PassName, Buffer->Name);
-			}
-		}
-		break;
-		case UBMT_RDG_BUFFER_SRV:
-		{
-			if (FRDGBufferSRVRef SRV = Parameter.GetAsBufferSRV())
-			{
-				FRDGBufferRef Buffer = SRV->Desc.Buffer;
-
-				checkf(Buffer->HasBeenProduced(),
-					TEXT("Pass %s has a dependency over the buffer %s that has never been produced."),
-					PassName, SRV->Desc.Buffer->Name);
-			}
-		}
-		break;
-		case UBMT_RDG_BUFFER_UAV:
-		{
-			if (FRDGBufferUAVRef UAV = Parameter.GetAsBufferUAV())
-			{
-				FRDGBufferRef Buffer = UAV->Desc.Buffer;
-
-				Buffer->MarkAsProducedBy(Pass);
-			}
-		}
-		break;
-		case UBMT_RDG_TEXTURE_COPY_DEST:
-		case UBMT_RDG_BUFFER_COPY_DEST:
-		{
-			if (FRDGTrackedResourceRef Resource = Parameter.GetAsTrackedResource())
-			{
-				Resource->MarkAsProducedBy(Pass);
-			}
-		}
-		break;
-		case UBMT_RENDER_TARGET_BINDING_SLOTS:
-		{
-			if (!RenderTargetBindingSlots)
-			{
-				RenderTargetBindingSlots = &Parameter.GetAsRenderTargetBindingSlots();
-			}
-			else if (GRDGDebug)
-			{
-				EmitRDGWarningf(
-					TEXT("Pass %s have duplicated render target binding slots."),
-					PassName);
-			}
-		}
-		break;
-		default:
-			break;
-		}
-	}
-
-	/** Validate that raster passes have render target binding slots and compute passes don't. */
-	if (RenderTargetBindingSlots)
-	{
-		checkf(bIsGraphics, TEXT("Pass '%s' has render target binding slots but is flagged as 'Compute'."), PassName);
-	}
-	else
-	{
-		checkf(!bIsGraphics, TEXT("Pass '%s' is missing render target binding slots. Set the 'Compute' or 'Copy' flag if render targets are not required."), PassName);
-	}
-
-	/** Validate render target / depth stencil binding usage. */
-	if (RenderTargetBindingSlots)
-	{
-		const auto& RenderTargets = RenderTargetBindingSlots->Output;
-
-		const uint32 RenderTargetCount = RenderTargets.Num();
-
-		{
-			/** Tracks the number of contiguous, non-null textures in the render target output array. */
-			uint32 ValidRenderTargetCount = 0;
-
-			for (uint32 RenderTargetIndex = 0; RenderTargetIndex < RenderTargetCount; ++RenderTargetIndex)
-			{
-				const FRenderTargetBinding& RenderTarget = RenderTargets[RenderTargetIndex];
-
-				if (FRDGTextureRef Texture = RenderTarget.GetTexture())
-				{
-					const bool bIsLoadAction = RenderTarget.GetLoadAction() == ERenderTargetLoadAction::ELoad;
-
-					/** Validate that load action is correct. We can only load contents if a pass previously produced something. */
-					{
-						const bool bIsLoadActionInvalid = bIsLoadAction && !Texture->HasBeenProduced();
-						checkf(
-							!bIsLoadActionInvalid,
-							TEXT("Pass '%s' attempted to bind texture '%s' as a render target with the 'Load' action specified, but the texture has not been produced yet. The render target must use either 'Clear' or 'NoAction' action instead."),
-							Pass->GetName(),
-							Texture->Name);
-					}
-
-					/** Validate that any previously produced texture contents are loaded. This occurs if the user failed to specify a load action
-					 *  on a texture that was produced by a previous pass, effectively losing that data. This can also happen if the user 're-uses'
-					 *  a texture for some other purpose. The latter is considered bad practice, since it increases memory pressure on the render
-					 *  target pool. Instead, the user should create a new texture instance. An exception to this rule are untracked render targets,
-					 *  which are not actually managed by the render target pool and likely represent the frame buffer.
-					 */
-					{
-						// We only validate single-mip textures since we don't track production at the subresource level.
-						const bool bFailedToLoadProducedContent = !bIsLoadAction && Texture->HasBeenProduced() && Texture->Desc.NumMips == 1;
-
-						// Untracked render targets aren't actually managed by the render target pool.
-						const bool bIsUntrackedRenderTarget = Texture->PooledRenderTarget && Texture->PooledRenderTarget->IsTracked();
-
-						ensureMsgf(!bFailedToLoadProducedContent || bIsUntrackedRenderTarget,
-							TEXT("Pass '%s' attempted to bind texture '%s' as a render target without the 'Load' action specified, despite a prior pass having produced it. It's invalid to completely clobber the contents of a resource. Create a new texture instance instead."),
-							Pass->GetName(),
-							Texture->Name);
-					}
-					
-					/** Mark the pass as a producer for render targets with a store action. */
-					{
-						const bool bIsStoreAction = RenderTarget.GetStoreAction() != ERenderTargetStoreAction::ENoAction;
-						check(bIsStoreAction); // already been validated in FRenderTargetBinding::Validate()
-						Texture->MarkAsProducedBy(Pass);
-					}
-				}
-				else
-				{
-					/** Found end of contiguous interval of valid render targets. */
-					ValidRenderTargetCount = RenderTargetIndex;
-					break;
-				}
-			}
-
-			/** Validate that no holes exist in the render target output array. Render targets must be bound contiguously. */
-			for (uint32 RenderTargetIndex = ValidRenderTargetCount; RenderTargetIndex < RenderTargetCount; ++RenderTargetIndex)
-			{
-				const FRenderTargetBinding& RenderTarget = RenderTargets[RenderTargetIndex];
-				checkf(RenderTarget.GetTexture() == nullptr, TEXT("Render targets must be packed. No empty spaces in the array."));
-			}
-		}
-
-		const auto IsMipLevelBoundForRead = [ParameterStruct, ParameterCount](FRDGTextureRef Texture, uint32 MipLevel)
-		{
-			for (uint32 ParameterIndex = 0; ParameterIndex < ParameterCount; ++ParameterIndex)
-			{
-				FRDGPassParameter Parameter = ParameterStruct.GetParameter(ParameterIndex);
-
-				switch (Parameter.GetType())
-				{
-				case UBMT_RDG_TEXTURE:
-				{
-					if (Parameter.GetAsTexture() == Texture)
-					{
-						/** Full texture view bound. */
-						return true;
-					}
-					break;
-				}
-				case UBMT_RDG_TEXTURE_SRV:
-				{
-					if (const FRDGTextureSRVRef InputSRV = Parameter.GetAsTextureSRV())
-					{
-						const FRDGTextureSRVDesc& InputSRVDesc = InputSRV->Desc;
-
-						if (InputSRVDesc.Texture == Texture)
-						{
-							return InputSRVDesc.MipLevel == MipLevel;
-						}
-					}
-					break;
-				}
-				default:
-					break;
-				}
-			}
-
-			return false;
-		};
-
-		/** Validate that texture mips are not bound as a render target and SRV at the same time. */
-		for (uint32 RenderTargetIndex = 0; RenderTargetIndex < RenderTargetCount; ++RenderTargetIndex)
-		{
-			const FRenderTargetBinding& RenderTarget = RenderTargets[RenderTargetIndex];
-
-			if (FRDGTextureRef Texture = RenderTarget.GetTexture())
-			{
-				if (IsMipLevelBoundForRead(Texture, RenderTarget.GetMipIndex()))
-				{
-					checkf(false, TEXT("Texture '%s', Mip '%d' is bound as a render target and SRV at the same time."), Texture->Name, RenderTarget.GetMipIndex());
-				}
-			}
-			else
-			{
-				break;
-			}
-		}
-	}
-#endif // RDG_ENABLE_DEBUG
-}
-
 void FRDGBuilder::VisualizePassOutputs(const FRDGPass* Pass)
 {
 #if SUPPORTS_VISUALIZE_TEXTURE
@@ -741,13 +452,6 @@ void FRDGBuilder::AllocateRHIBufferSRVIfNeeded(FRDGBufferSRV* SRV)
 
 	FRDGBufferRef Buffer = SRV->Desc.Buffer;
 
-	// The underlying buffer have already been allocated by a prior pass through AllocateRHIBufferUAVIfNeeded().
-	#if RDG_ENABLE_DEBUG
-	{
-		check(Buffer->HasBeenProduced());
-	}
-	#endif
-
 	check(Buffer->PooledBuffer);
 
 	if (Buffer->PooledBuffer->SRVs.Contains(SRV->Desc))
@@ -819,10 +523,7 @@ void FRDGBuilder::ExecutePass(const FRDGPass* Pass)
 {
 	QUICK_SCOPE_CYCLE_COUNTER(STAT_FRDGBuilder_ExecutePass);
 
-	{
-		const bool bAllowAccess = true;
-		UpdateAccessGuardForPassResources(Pass, bAllowAccess);
-	}
+	IF_RDG_ENABLE_DEBUG(Validation.ValidateExecutePassBegin(Pass));
 
 	FRHIRenderPassInfo RPInfo;
 	bool bHasRenderTargets = false;
@@ -832,7 +533,7 @@ void FRDGBuilder::ExecutePass(const FRDGPass* Pass)
 	EventScopeStack.BeginExecutePass(Pass);
 	StatScopeStack.BeginExecutePass(Pass);
 
-	if (Pass->IsGraphics())
+	if (Pass->IsRaster())
 	{
 		check(bHasRenderTargets);
 		RHICmdList.BeginRenderPass( RPInfo, Pass->GetName() );
@@ -851,12 +552,7 @@ void FRDGBuilder::ExecutePass(const FRDGPass* Pass)
 
 	EventScopeStack.EndExecutePass();
 
-	{
-		const bool bAllowAccess = false;
-		UpdateAccessGuardForPassResources(Pass, bAllowAccess);
-	}
-
-	UnmarkUsedResources(Pass);
+	IF_RDG_ENABLE_DEBUG(Validation.ValidateExecutePassEnd(Pass));
 
 	// Can't release resources with immediate mode, because don't know if whether they are gonna be used.
 	if (!GRDGImmediateMode)
@@ -874,8 +570,7 @@ void FRDGBuilder::PrepareResourcesForExecute(const FRDGPass* Pass, struct FRHIRe
 
 	const bool bIsCompute = Pass->IsCompute();
 
-	FRDGBarrierBatcher BarrierBatcher;
-	BarrierBatcher.Begin();
+	FRDGBarrierBatcher BarrierBatcher(RHICmdList, Pass);
 
 	// NOTE: When generating mips, we don't perform any transitions on textures. They are done implicitly by the RHI.
 	const bool bGeneratingMips = Pass->IsGenerateMips();
@@ -894,21 +589,10 @@ void FRDGBuilder::PrepareResourcesForExecute(const FRDGPass* Pass, struct FRHIRe
 		{
 			if (FRDGTextureRef Texture = Parameter.GetAsTexture())
 			{
-				// The underlying texture have already been allocated by a prior pass.
-				#if RDG_ENABLE_DEBUG
-				{
-					check(Texture->HasBeenProduced());
-					Texture->PassAccessCount++;
-				}
-				#endif
-
 				check(Texture->PooledRenderTarget);
 				check(Texture->ResourceRHI);
 
-				if (!bGeneratingMips)
-				{
-					BarrierBatcher.QueueTransitionTexture(Texture, FRDGResourceState::CreateRead(Pass));
-				}
+				BarrierBatcher.QueueTransitionTexture(Texture, FRDGResourceState::EAccess::Read);
 			}
 		}
 		break;
@@ -918,24 +602,13 @@ void FRDGBuilder::PrepareResourcesForExecute(const FRDGPass* Pass, struct FRHIRe
 			{
 				FRDGTextureRef Texture = SRV->Desc.Texture;
 
-				// The underlying texture have already been allocated by a prior pass.
-				#if RDG_ENABLE_DEBUG
-				{
-					check(Texture->HasBeenProduced());
-					Texture->PassAccessCount++;
-				}
-				#endif
-
 				// Might be the first time using this render graph SRV, so need to setup the cached rhi resource.
 				if (!SRV->ResourceRHI)
 				{
 					SRV->ResourceRHI = Texture->PooledRenderTarget->GetRenderTargetItem().MipSRVs[SRV->Desc.MipLevel];
 				}
 
-				if (!bGeneratingMips)
-				{
-					BarrierBatcher.QueueTransitionTexture(Texture, FRDGResourceState::CreateRead(Pass));
-				}
+				BarrierBatcher.QueueTransitionTexture(Texture, FRDGResourceState::EAccess::Read);
 			}
 		}
 		break;
@@ -945,12 +618,6 @@ void FRDGBuilder::PrepareResourcesForExecute(const FRDGPass* Pass, struct FRHIRe
 			{
 				FRDGTextureRef Texture = UAV->Desc.Texture;
 	
-				#if RDG_ENABLE_DEBUG
-				{
-					Texture->PassAccessCount++;
-				}
-				#endif
-
 				AllocateRHITextureUAVIfNeeded(UAV);
 
 				FRHIUnorderedAccessView* UAVRHI = UAV->GetRHI();
@@ -960,7 +627,7 @@ void FRDGBuilder::PrepareResourcesForExecute(const FRDGPass* Pass, struct FRHIRe
 					OutRPInfo->UAVs[OutRPInfo->NumUAVs++] = UAVRHI;	// Bind UAVs in declaration order
 				}
 
-				BarrierBatcher.QueueTransitionUAV(UAVRHI, Texture, FRDGResourceState::CreateWrite(Pass));
+				BarrierBatcher.QueueTransitionUAV(UAVRHI, Texture, FRDGResourceState::EAccess::Write);
 			}
 		}
 		break;
@@ -968,12 +635,6 @@ void FRDGBuilder::PrepareResourcesForExecute(const FRDGPass* Pass, struct FRHIRe
 		{
 			if (FRDGTextureRef Texture = Parameter.GetAsTexture())
 			{
-				#if RDG_ENABLE_DEBUG
-				{
-					Texture->PassAccessCount++;
-				}
-				#endif
-
 				AllocateRHITextureIfNeeded(Texture);
 			}
 		}
@@ -982,20 +643,12 @@ void FRDGBuilder::PrepareResourcesForExecute(const FRDGPass* Pass, struct FRHIRe
 		{
 			if (FRDGBufferRef Buffer = Parameter.GetAsBuffer())
 			{
-				// The underlying buffer have already been allocated by a prior pass through AllocateRHIBufferUAVIfNeeded().
-				#if RDG_ENABLE_DEBUG
-				{
-					check(Buffer->HasBeenProduced());
-					Buffer->PassAccessCount++;
-				}
-				#endif
-
 				// TODO(RDG): super hacky, find the UAV and transition it. Hopefully there is one...
 				check(Buffer->PooledBuffer);
 				check(Buffer->PooledBuffer->UAVs.Num() == 1);
 				FRHIUnorderedAccessView* UAVRHI = Buffer->PooledBuffer->UAVs.CreateIterator().Value();
 
-				BarrierBatcher.QueueTransitionUAV(UAVRHI, Buffer, FRDGResourceState::CreateRead(Pass));
+				BarrierBatcher.QueueTransitionUAV(UAVRHI, Buffer, FRDGResourceState::EAccess::Read);
 			}
 		}
 		break;
@@ -1005,14 +658,6 @@ void FRDGBuilder::PrepareResourcesForExecute(const FRDGPass* Pass, struct FRHIRe
 			{
 				FRDGBufferRef Buffer = SRV->Desc.Buffer;
 
-				// The underlying buffer have already been allocated by a prior pass through AllocateRHIBufferUAVIfNeeded().
-				#if RDG_ENABLE_DEBUG
-				{
-					check(Buffer->HasBeenProduced());
-					Buffer->PassAccessCount++;
-				}
-				#endif
-
 				AllocateRHIBufferSRVIfNeeded(SRV);
 
 				// TODO(RDG): super hacky, find the UAV and transition it. Hopefully there is one...
@@ -1020,7 +665,7 @@ void FRDGBuilder::PrepareResourcesForExecute(const FRDGPass* Pass, struct FRHIRe
 				check(Buffer->PooledBuffer->UAVs.Num() == 1);
 				FRHIUnorderedAccessView* UAVRHI = Buffer->PooledBuffer->UAVs.CreateIterator().Value();
 
-				BarrierBatcher.QueueTransitionUAV(UAVRHI, Buffer, FRDGResourceState::CreateRead(Pass));
+				BarrierBatcher.QueueTransitionUAV(UAVRHI, Buffer, FRDGResourceState::EAccess::Read);
 			}
 		}
 		break;
@@ -1030,12 +675,6 @@ void FRDGBuilder::PrepareResourcesForExecute(const FRDGPass* Pass, struct FRHIRe
 			{
 				FRDGBufferRef Buffer = UAV->Desc.Buffer;
 
-				#if RDG_ENABLE_DEBUG
-				{
-					Buffer->PassAccessCount++;
-				}
-				#endif
-
 				AllocateRHIBufferUAVIfNeeded(UAV);
 
 				FRHIUnorderedAccessView* UAVRHI = UAV->GetRHI();
@@ -1045,7 +684,7 @@ void FRDGBuilder::PrepareResourcesForExecute(const FRDGPass* Pass, struct FRHIRe
 					OutRPInfo->UAVs[OutRPInfo->NumUAVs++] = UAVRHI;	// Bind UAVs in declaration order
 				}
 
-				BarrierBatcher.QueueTransitionUAV(UAVRHI, Buffer, FRDGResourceState::CreateWrite(Pass));
+				BarrierBatcher.QueueTransitionUAV(UAVRHI, Buffer, FRDGResourceState::EAccess::Write);
 			}
 		}
 		break;
@@ -1065,8 +704,6 @@ void FRDGBuilder::PrepareResourcesForExecute(const FRDGPass* Pass, struct FRHIRe
 		break;
 		case UBMT_RENDER_TARGET_BINDING_SLOTS:
 		{
-			check(!Pass->IsCompute());
-
 			const FRenderTargetBindingSlots& RenderTargetBindingSlots = Parameter.GetAsRenderTargetBindingSlots();
 			const auto& RenderTargets = RenderTargetBindingSlots.Output;
 			const auto& DepthStencil = RenderTargetBindingSlots.DepthStencil;
@@ -1112,19 +749,10 @@ void FRDGBuilder::PrepareResourcesForExecute(const FRDGPass* Pass, struct FRHIRe
 					OutRenderTarget.MipIndex = RenderTarget.GetMipIndex();
 					OutRenderTarget.Action = MakeRenderTargetActions(RenderTarget.GetLoadAction(), StoreAction);
 
-					if (!bGeneratingMips)
-					{
-						BarrierBatcher.QueueTransitionTexture(Texture, FRDGResourceState::CreateWrite(Pass));
-					}
+					BarrierBatcher.QueueTransitionTexture(Texture, FRDGResourceState::EAccess::Write);
 
 					SampleCount |= OutRenderTarget.RenderTarget->GetNumSamples();
 					ValidRenderTargetCount++;
-
-					#if RDG_ENABLE_DEBUG
-					{
-						Texture->PassAccessCount++;
-					}
-					#endif
 				}
 				else
 				{
@@ -1151,17 +779,11 @@ void FRDGBuilder::PrepareResourcesForExecute(const FRDGPass* Pass, struct FRHIRe
 
 				BarrierBatcher.QueueTransitionTexture(Texture,
 					DepthStencil.GetDepthStencilAccess().IsAnyWrite() ?
-					FRDGResourceState::CreateWrite(Pass) :
-					FRDGResourceState::CreateRead(Pass));
+					FRDGResourceState::EAccess::Write :
+					FRDGResourceState::EAccess::Read);
 
 				SampleCount |= OutDepthStencil.DepthStencilTarget->GetNumSamples();
 				ValidDepthStencilCount++;
-
-				#if RDG_ENABLE_DEBUG
-				{
-					Texture->PassAccessCount++;
-				}
-				#endif
 			}
 
 			OutRPInfo->bIsMSAA = SampleCount > 1;
@@ -1175,123 +797,6 @@ void FRDGBuilder::PrepareResourcesForExecute(const FRDGPass* Pass, struct FRHIRe
 	}
 
 	OutRPInfo->bGeneratingMips = bGeneratingMips;
-
-	BarrierBatcher.End(RHICmdList);
-}
-
-void FRDGBuilder::UpdateAccessGuardForPassResources(const FRDGPass* Pass, bool bAllowAccess)
-{
-#if RDG_ENABLE_DEBUG
-	FRDGPassParameterStruct ParameterStruct = Pass->GetParameters();
-
-	const uint32 ParameterCount = ParameterStruct.GetParameterCount();
-
-	for (uint32 ParameterIndex = 0; ParameterIndex < ParameterCount; ++ParameterIndex)
-	{
-		FRDGPassParameter Parameter = ParameterStruct.GetParameter(ParameterIndex);
-
-		if (Parameter.IsResource())
-		{
-			if (FRDGResourceRef Resource = Parameter.GetAsResource())
-			{
-				Resource->bAllowRHIAccess = bAllowAccess;
-			}
-		}
-		else if (Parameter.GetType() == UBMT_RENDER_TARGET_BINDING_SLOTS)
-		{
-			const FRenderTargetBindingSlots& RenderTargetBindingSlots = Parameter.GetAsRenderTargetBindingSlots();
-			const auto& RenderTargets = RenderTargetBindingSlots.Output;
-			const auto& DepthStencil = RenderTargetBindingSlots.DepthStencil;
-			const uint32 RenderTargetCount = RenderTargets.Num();
-
-			for (uint32 RenderTargetIndex = 0; RenderTargetIndex < RenderTargetCount; RenderTargetIndex++)
-			{
-				const FRenderTargetBinding& RenderTarget = RenderTargets[RenderTargetIndex];
-
-				if (FRDGTextureRef Texture = RenderTarget.GetTexture())
-				{
-					Texture->bAllowRHIAccess = bAllowAccess;
-				}
-				else
-				{
-					break;
-				}
-			}
-			
-			if (FRDGTextureRef Texture = DepthStencil.GetTexture())
-			{
-				Texture->bAllowRHIAccess = bAllowAccess;
-			}
-		}
-	}
-#endif
-}
-
-void FRDGBuilder::UnmarkUsedResources(const FRDGPass* Pass)
-{
-#if RDG_ENABLE_DEBUG
-	FRDGPassParameterStruct ParameterStruct = Pass->GetParameters();
-
-	const uint32 ParameterCount = ParameterStruct.GetParameterCount();
-
-	if (GRDGDebug)
-	{
-		uint32 TrackedResourceCount = 0;
-		uint32 UsedResourceCount = 0;
-
-		for (uint32 ParameterIndex = 0; ParameterIndex < ParameterCount; ++ParameterIndex)
-		{
-			FRDGPassParameter Parameter = ParameterStruct.GetParameter(ParameterIndex);
-
-			if (Parameter.IsResource())
-			{
-				if (FRDGResourceRef Resource = Parameter.GetAsResource())
-				{
-					TrackedResourceCount++;
-					UsedResourceCount += Resource->bIsActuallyUsedByPass ? 1 : 0;
-				}
-			}
-		}
-
-		if (TrackedResourceCount != UsedResourceCount)
-		{
-			FString WarningMessage = FString::Printf(
-				TEXT("'%d' of the '%d' resources of the pass '%s' where not actually used."),
-				TrackedResourceCount - UsedResourceCount, TrackedResourceCount, Pass->GetName());
-
-			for (uint32 ParameterIndex = 0; ParameterIndex < ParameterCount; ++ParameterIndex)
-			{
-				FRDGPassParameter Parameter = ParameterStruct.GetParameter(ParameterIndex);
-
-				if (Parameter.IsResource())
-				{
-					if (const FRDGResourceRef Resource = Parameter.GetAsResource())
-					{
-						if (!Resource->bIsActuallyUsedByPass)
-						{
-							WarningMessage += FString::Printf(TEXT("\n    %s"), Resource->Name);
-						}
-					}
-				}
-			}
-
-			EmitRDGWarning(WarningMessage);
-		}
-	}
-
-	for (uint32 ParameterIndex = 0; ParameterIndex < ParameterCount; ++ParameterIndex)
-	{
-		FRDGPassParameter Parameter = ParameterStruct.GetParameter(ParameterIndex);
-
-		if (Parameter.IsResource())
-		{
-			if (const FRDGResourceRef Resource = Parameter.GetAsResource())
-			{
-				Resource->bIsActuallyUsedByPass = false;
-			}
-		}
-	}
-#endif
 }
 
 void FRDGBuilder::ReleaseRHITextureIfUnreferenced(FRDGTexture* Texture)
@@ -1417,8 +922,7 @@ void FRDGBuilder::ReleaseUnreferencedResources(const FRDGPass* Pass)
 
 void FRDGBuilder::ProcessDeferredInternalResourceQueries()
 {
-	FRDGBarrierBatcher BarrierBatcher;
-	BarrierBatcher.Begin();
+	FRDGBarrierBatcher BarrierBatcher(RHICmdList, nullptr);
 
 	for (const auto& Query : DeferredInternalTextureQueries)
 	{
@@ -1426,24 +930,11 @@ void FRDGBuilder::ProcessDeferredInternalResourceQueries()
 
 		if (Query.bTransitionToRead)
 		{
-			BarrierBatcher.QueueTransitionTexture(
-				Query.Texture,
-				FRDGResourceState(
-					nullptr,
-					FRDGResourceState::EPipeline::Graphics,
-					FRDGResourceState::EAccess::Read));
+			BarrierBatcher.QueueTransitionTexture(Query.Texture, FRDGResourceState::EAccess::Read);
 		}
 
 		*Query.OutTexturePtr = AllocatedTextures.FindChecked(Query.Texture);
 		
-		#if RDG_ENABLE_DEBUG
-		{
-			// Increment the number of times the texture has been accessed to avoid warning on produced but never used resources that were produced
-			// only to be extracted for the graph.
-			Query.Texture->PassAccessCount += 1;
-		}
-		#endif
-
 		// No need to manually release in immediate mode, since it is done directly when emptying AllocatedTextures in DestructPasses().
 		if (!GRDGImmediateMode)
 		{
@@ -1455,50 +946,16 @@ void FRDGBuilder::ProcessDeferredInternalResourceQueries()
 	{
 		*Query.OutBufferPtr = AllocatedBuffers.FindChecked(Query.Buffer);
 
-		#if RDG_ENABLE_DEBUG
-		{
-			// Increment the number of times the buffer has been accessed to avoid warning on produced but never used resources that were produced
-			// only to be extracted for the graph.
-			Query.Buffer->PassAccessCount += 1;
-		}
-		#endif
-
 		// No need to manually release in immediate mode, since it is done directly when emptying AllocatedBuffer in DestructPasses().
 		if (!GRDGImmediateMode)
 		{
 			ReleaseRHIBufferIfUnreferenced(Query.Buffer);
 		}
 	}
-
-	BarrierBatcher.End(RHICmdList);
 }
 
 void FRDGBuilder::DestructPasses()
 {
-	#if RDG_ENABLE_DEBUG
-	{
-		const bool bEmitWarnings = (GRDGDebug > 0);
-
-		for (const FRDGTrackedResourceRef Resource : TrackedResources)
-		{
-			check(Resource->ReferenceCount == 0);
-
-			const bool bProducedButNeverUsed = Resource->PassAccessCount == 1 && Resource->FirstProducer;
-
-			if (bEmitWarnings && bProducedButNeverUsed)
-			{
-				check(Resource->HasBeenProduced());
-
-				EmitRDGWarningf(
-					TEXT("Resources %s has been produced by the pass %s, but never used by another pass."),
-					Resource->Name, Resource->FirstProducer->GetName());
-			}
-		}
-
-		TrackedResources.Empty();
-	}
-	#endif
-
 	for (int32 PassIndex = Passes.Num() - 1; PassIndex >= 0; --PassIndex)
 	{
 		Passes[PassIndex]->~FRDGPass();
diff --git a/Engine/Source/Runtime/RenderCore/Private/RenderGraphPass.cpp b/Engine/Source/Runtime/RenderCore/Private/RenderGraphPass.cpp
--- a/Engine/Source/Runtime/RenderCore/Private/RenderGraphPass.cpp
+++ b/Engine/Source/Runtime/RenderCore/Private/RenderGraphPass.cpp
@@ -31,12 +31,7 @@ FRDGPass::FRDGPass(
 	: Name(static_cast<FRDGEventName&&>(InName))
 	, ParameterStruct(InParameterStruct)
 	, PassFlags(InPassFlags)
-{
-	if (IsCompute())
-	{
-		ensureMsgf(ParameterStruct.GetLayout().NumRenderTargets() == 0, TEXT("Pass %s was declared as ERDGPassFlags::Compute yet has RenderTargets in its ResourceTable"), GetName());
-	}
-}
+{}
 
 void FRDGPass::Execute(FRHICommandListImmediate& RHICmdList) const
 {
diff --git a/Engine/Source/Runtime/RenderCore/Private/RenderGraphResources.cpp b/Engine/Source/Runtime/RenderCore/Private/RenderGraphResources.cpp
deleted file mode 100644
diff --git a/Engine/Source/Runtime/RenderCore/Public/RenderGraphBuilder.h b/Engine/Source/Runtime/RenderCore/Public/RenderGraphBuilder.h
--- a/Engine/Source/Runtime/RenderCore/Public/RenderGraphBuilder.h
+++ b/Engine/Source/Runtime/RenderCore/Public/RenderGraphBuilder.h
@@ -6,6 +6,7 @@
 #include "RenderGraphEvent.h"
 #include "RenderGraphResources.h"
 #include "RenderGraphPass.h"
+#include "RenderGraphValidation.h"
 #include "ShaderParameterMacros.h"
 
 /** Builds the per-frame render graph.
@@ -17,7 +18,7 @@ class RENDERCORE_API FRDGBuilder
 public:
 	/** A RHI cmd list is required, if using the immediate mode. */
 	FRDGBuilder(FRHICommandListImmediate& InRHICmdList);
-	~FRDGBuilder();
+	FRDGBuilder(const FRDGBuilder&) = delete;
 
 	/** Register a external texture to be tracked by the render graph. */
 	FRDGTextureRef RegisterExternalTexture(
@@ -117,9 +118,12 @@ public:
 	 */
 	void QueueBufferExtraction(FRDGBufferRef Buffer, TRefCountPtr<FPooledRDGBuffer>* OutBufferPtr);
 
-	/** Flag a texture that is only produced by only 1 pass, but never used or extracted, to avoid generating a warning at runtime. */
+	/** Flag a texture that is produced by a pass but never used or extracted to not emit an 'unused' warning. */
 	void RemoveUnusedTextureWarning(FRDGTextureRef Texture);
 
+	/** Flag a buffer that is produced by a pass but never used or extracted to not emit an 'unused' warning. */
+	void RemoveUnusedBufferWarning(FRDGBufferRef Buffer);
+
 	/** Begins / ends a named event scope. These scopes are visible in most external GPU profilers.
 	 *  Prefer to use RDG_EVENT_SCOPE instead of calling this manually.
 	 */
@@ -179,17 +183,9 @@ private:
 	FRDGStatScopeStack StatScopeStack;
 
 #if RDG_ENABLE_DEBUG
-	/** All recently allocated pass parameter structure, but not used by a AddPass() yet. */
-	TSet<const void*> AllocatedUnusedPassParameters;
-
-	/** List of tracked resources for validation prior to shutdown. */
-	TArray<FRDGTrackedResourceRef, SceneRenderingAllocator> TrackedResources;
-
-	/** Whether the Execute() has already been called. */
-	bool bHasExecuted = false;
+	FRDGUserValidation Validation;
 #endif
 
-	void ValidatePass(const FRDGPass* Pass) const;
 	void VisualizePassOutputs(const FRDGPass* Pass);
 
 	void WalkGraphDependencies();
@@ -208,9 +204,6 @@ private:
 	void ExecutePass(const FRDGPass* Pass);
 	void PrepareResourcesForExecute(const FRDGPass* Pass, struct FRHIRenderPassInfo* OutRPInfo, bool* bOutHasRenderTargets);
 
-	void UpdateAccessGuardForPassResources(const FRDGPass* Pass, bool bAllowAccess);
-	void UnmarkUsedResources(const FRDGPass* Pass);
-
 	void ReleaseRHITextureIfUnreferenced(FRDGTexture* Texture);
 	void ReleaseRHIBufferIfUnreferenced(FRDGBuffer* Buffer);
 	void ReleaseUnreferencedResources(const FRDGPass* Pass);
diff --git a/Engine/Source/Runtime/RenderCore/Public/RenderGraphBuilder.inl b/Engine/Source/Runtime/RenderCore/Public/RenderGraphBuilder.inl
--- a/Engine/Source/Runtime/RenderCore/Public/RenderGraphBuilder.inl
+++ b/Engine/Source/Runtime/RenderCore/Public/RenderGraphBuilder.inl
@@ -8,6 +8,7 @@ inline FRDGTextureRef FRDGBuilder::RegisterExternalTexture(const TRefCountPtr<IP
 	{
 		checkf(ExternalPooledTexture.IsValid(), TEXT("Attempted to register NULL external texture: %s"), Name);
 		checkf(Name, TEXT("Externally allocated texture requires a debug name when registering them to render graph."));
+		Validation.ExecuteGuard(TEXT("RegisterExternalTexture"), Name);
 	}
 #endif
 
@@ -21,12 +22,7 @@ inline FRDGTextureRef FRDGBuilder::RegisterExternalTexture(const TRefCountPtr<IP
 	OutTexture->ResourceRHI = ExternalPooledTexture->GetRenderTargetItem().ShaderResourceTexture;
 	AllocatedTextures.Add(OutTexture, ExternalPooledTexture);
 
-#if RDG_ENABLE_DEBUG
-	{
-		OutTexture->MarkAsExternal();
-		TrackedResources.Add(OutTexture);
-	}
-#endif
+	IF_RDG_ENABLE_DEBUG(Validation.ValidateCreateExternalResource(OutTexture));
 
 	ExternalTextures.Add(ExternalPooledTexture.GetReference(), OutTexture);
 
@@ -38,6 +34,7 @@ inline FRDGBufferRef FRDGBuilder::RegisterExternalBuffer(const TRefCountPtr<FPoo
 #if RDG_ENABLE_DEBUG
 	{
 		checkf(ExternalPooledBuffer.IsValid(), TEXT("Attempted to register NULL external buffer: %s"), Name);
+		Validation.ExecuteGuard(TEXT("RegisterExternalBuffer"), Name);
 	}
 #endif
 
@@ -50,12 +47,7 @@ inline FRDGBufferRef FRDGBuilder::RegisterExternalBuffer(const TRefCountPtr<FPoo
 	OutBuffer->PooledBuffer = ExternalPooledBuffer;
 	AllocatedBuffers.Add(OutBuffer, ExternalPooledBuffer);
 
-#if RDG_ENABLE_DEBUG
-	{
-		OutBuffer->MarkAsExternal();
-		TrackedResources.Add(OutBuffer);
-	}
-#endif
+	IF_RDG_ENABLE_DEBUG(Validation.ValidateCreateExternalResource(OutBuffer));
 
 	ExternalBuffers.Add(ExternalPooledBuffer.GetReference(), OutBuffer);
 
@@ -64,36 +56,32 @@ inline FRDGBufferRef FRDGBuilder::RegisterExternalBuffer(const TRefCountPtr<FPoo
 
 inline FRDGTextureRef FRDGBuilder::CreateTexture(
 	const FPooledRenderTargetDesc& Desc,
-	const TCHAR* DebugName,
+	const TCHAR* Name,
 	ERDGResourceFlags Flags)
 {
 #if RDG_ENABLE_DEBUG
 	{
-		checkf(DebugName, TEXT("Creating a render graph texture requires a valid debug name."));
-		checkf(!bHasExecuted, TEXT("Render graph texture %s needs to be created before the builder execution."), DebugName);
+		checkf(Name, TEXT("Creating a render graph texture requires a valid debug name."));
+		Validation.ExecuteGuard(TEXT("CreateTexture"), Name);
 
 		// Validate the pixel format.
-		checkf(Desc.Format != PF_Unknown, TEXT("Illegal to create texture %s with an invalid pixel format."), DebugName);
-		checkf(Desc.Format < PF_MAX, TEXT("Illegal to create texture %s with invalid FPooledRenderTargetDesc::Format."), DebugName);
+		checkf(Desc.Format != PF_Unknown, TEXT("Illegal to create texture %s with an invalid pixel format."), Name);
+		checkf(Desc.Format < PF_MAX, TEXT("Illegal to create texture %s with invalid FPooledRenderTargetDesc::Format."), Name);
 		checkf(GPixelFormats[Desc.Format].Supported, TEXT("Failed to create texture %s with pixel format %s because it is not supported."),
-			DebugName, GPixelFormats[Desc.Format].Name);
+			Name, GPixelFormats[Desc.Format].Name);
 
 		const bool bCanHaveUAV = Desc.TargetableFlags & TexCreate_UAV;
 		const bool bIsMSAA = Desc.NumSamples > 1;
 
 		// D3D11 doesn't allow creating a UAV on MSAA texture.
 		const bool bIsUAVForMSAATexture = bIsMSAA && bCanHaveUAV;
-		checkf(!bIsUAVForMSAATexture, TEXT("TexCreate_UAV is not allowed on MSAA texture %s."), DebugName);
+		checkf(!bIsUAVForMSAATexture, TEXT("TexCreate_UAV is not allowed on MSAA texture %s."), Name);
 	}
 #endif
 
-	FRDGTexture* Texture = AllocateForRHILifeTime<FRDGTexture>(DebugName, Desc, Flags);
+	FRDGTexture* Texture = AllocateForRHILifeTime<FRDGTexture>(Name, Desc, Flags);
 
-#if RDG_ENABLE_DEBUG
-	{
-		TrackedResources.Add(Texture);
-	}
-#endif
+	IF_RDG_ENABLE_DEBUG(Validation.ValidateCreateResource(Texture));
 
 	return Texture;
 }
@@ -106,17 +94,13 @@ inline FRDGBufferRef FRDGBuilder::CreateBuffer(
 #if RDG_ENABLE_DEBUG
 	{
 		checkf(Name, TEXT("Creating a render graph buffer requires a valid debug name."));
-		checkf(!bHasExecuted, TEXT("Render graph buffer %s needs to be created before the builder execution."), Name);
+		Validation.ExecuteGuard(TEXT("CreateBuffer"), Name);
 	}
 #endif
 
 	FRDGBufferRef Buffer = AllocateForRHILifeTime<FRDGBuffer>(Name, Desc, Flags);
 
-#if RDG_ENABLE_DEBUG
-	{
-		TrackedResources.Add(Buffer);
-	}
-#endif
+	IF_RDG_ENABLE_DEBUG(Validation.ValidateCreateResource(Buffer));
 
 	return Buffer;
 }
@@ -128,13 +112,13 @@ inline FRDGTextureSRVRef FRDGBuilder::CreateSRV(const FRDGTextureSRVDesc& Desc)
 #if RDG_ENABLE_DEBUG
 	{
 		checkf(Texture, TEXT("RenderGraph texture SRV created with a null texture."));
-		checkf(!bHasExecuted, TEXT("Render graph SRV %s needs to be created before the builder execution."), Desc.Texture->Name);
-		checkf(Desc.Texture->Desc.TargetableFlags & TexCreate_ShaderResource, TEXT("Attempted to create SRV from texture %s which was not created with TexCreate_ShaderResource"), Desc.Texture->Name);
+		Validation.ExecuteGuard(TEXT("CreateSRV"), Texture->Name);
+		checkf(Texture->Desc.TargetableFlags & TexCreate_ShaderResource, TEXT("Attempted to create SRV from texture %s which was not created with TexCreate_ShaderResource"), Desc.Texture->Name);
 		checkf(Desc.MipLevel < Texture->Desc.NumMips, TEXT("Failed to create SRV at mip %d: the texture %s has only %d mip levels."), Desc.MipLevel, Texture->Name, Texture->Desc.NumMips);
 	}
 #endif
 
-	return AllocateForRHILifeTime<FRDGTextureSRV>(Desc.Texture->Name, Desc);
+	return AllocateForRHILifeTime<FRDGTextureSRV>(Texture->Name, Desc);
 }
 
 inline FRDGBufferSRVRef FRDGBuilder::CreateSRV(const FRDGBufferSRVDesc& Desc)
@@ -144,7 +128,7 @@ inline FRDGBufferSRVRef FRDGBuilder::CreateSRV(const FRDGBufferSRVDesc& Desc)
 #if RDG_ENABLE_DEBUG
 	{
 		checkf(Buffer, TEXT("RenderGraph buffer SRV created with a null buffer."));
-		checkf(!bHasExecuted, TEXT("Render graph SRV %s needs to be created before the builder execution."), Desc.Buffer->Name);
+		Validation.ExecuteGuard(TEXT("CreateSRV"), Buffer->Name);
 	}
 #endif
 
@@ -158,7 +142,7 @@ inline FRDGTextureUAVRef FRDGBuilder::CreateUAV(const FRDGTextureUAVDesc& Desc)
 #if RDG_ENABLE_DEBUG
 	{
 		checkf(Texture, TEXT("RenderGraph texture UAV created with a null texture."));
-		checkf(!bHasExecuted, TEXT("Render graph UAV %s needs to be created before the builder execution."), Texture->Name);
+		Validation.ExecuteGuard(TEXT("CreateUAV"), Texture->Name);
 		checkf(Texture->Desc.TargetableFlags & TexCreate_UAV, TEXT("Attempted to create UAV from texture %s which was not created with TexCreate_UAV"), Texture->Name);
 		checkf(Desc.MipLevel < Texture->Desc.NumMips, TEXT("Failed to create UAV at mip %d: the texture %s has only %d mip levels."), Desc.MipLevel, Texture->Name, Texture->Desc.NumMips);
 	}
@@ -174,7 +158,7 @@ inline FRDGBufferUAVRef FRDGBuilder::CreateUAV(const FRDGBufferUAVDesc& Desc)
 #if RDG_ENABLE_DEBUG
 	{
 		checkf(Buffer, TEXT("RenderGraph buffer UAV created with a null buffer."));
-		checkf(!bHasExecuted, TEXT("Render graph UAV %s needs to be created before the builder execution."), Buffer->Name);
+		Validation.ExecuteGuard(TEXT("CreateUAV"), Buffer->Name);
 	}
 #endif
 
@@ -187,11 +171,7 @@ ParameterStructType* FRDGBuilder::AllocParameters()
 	// TODO(RDG): could allocate using AllocateForRHILifeTime() to avoid the copy done when using FRHICommandList::BuildLocalUniformBuffer()
 	ParameterStructType* OutParameterPtr = new(MemStack) ParameterStructType;
 	FMemory::Memzero(OutParameterPtr, sizeof(ParameterStructType));
-#if RDG_ENABLE_DEBUG
-	{
-		AllocatedUnusedPassParameters.Add(static_cast<void *>(OutParameterPtr));
-	}
-#endif
+	IF_RDG_ENABLE_DEBUG(Validation.ValidateAllocPassParameters(OutParameterPtr));
 	return OutParameterPtr;
 }
 
@@ -221,19 +201,10 @@ inline void FRDGBuilder::QueueTextureExtraction(
 	TRefCountPtr<IPooledRenderTarget>* OutTexturePtr,
 	bool bTransitionToRead)
 {
-	check(Texture);
+	IF_RDG_ENABLE_DEBUG(Validation.ValidateExtractResource(Texture));
+
 	check(OutTexturePtr);
-#if RDG_ENABLE_DEBUG
-	{
-		checkf(!bHasExecuted,
-			TEXT("Accessing render graph internal texture %s with QueueTextureExtraction() needs to happen before the builder's execution."),
-			Texture->Name);
 
-		checkf(Texture->HasBeenProduced(),
-			TEXT("Unable to queue the extraction of the texture %s because it has not been produced by any pass."),
-			Texture->Name);
-	}
-#endif
 	FDeferredInternalTextureQuery Query;
 	Query.Texture = Texture;
 	Query.OutTexturePtr = OutTexturePtr;
@@ -243,19 +214,10 @@ inline void FRDGBuilder::QueueTextureExtraction(
 
 inline void FRDGBuilder::QueueBufferExtraction(FRDGBufferRef Buffer, TRefCountPtr<FPooledRDGBuffer>* OutBufferPtr)
 {
-	check(Buffer);
+	IF_RDG_ENABLE_DEBUG(Validation.ValidateExtractResource(Buffer));
+
 	check(OutBufferPtr);
-#if RDG_ENABLE_DEBUG
-	{
-		checkf(!bHasExecuted,
-			TEXT("Accessing render graph internal buffer %s with QueueBufferExtraction() needs to happen before the builder's execution."),
-			Buffer->Name);
 
-		checkf(Buffer->HasBeenProduced(),
-			TEXT("Unable to queue the extraction of the buffer %s because it has not been produced by any pass."),
-			Buffer->Name);
-	}
-#endif
 	FDeferredInternalBufferQuery Query;
 	Query.Buffer = Buffer;
 	Query.OutBufferPtr = OutBufferPtr;
@@ -264,17 +226,19 @@ inline void FRDGBuilder::QueueBufferExtraction(FRDGBufferRef Buffer, TRefCountPt
 
 inline void FRDGBuilder::RemoveUnusedTextureWarning(FRDGTextureRef Texture)
 {
-	check(Texture);
 #if RDG_ENABLE_DEBUG
-	{
-		checkf(!bHasExecuted,
-			TEXT("Flaging texture %s with FlagUnusedTexture() needs to happen before the builder's execution."),
-			Texture->Name);
+	check(Texture);
+	Validation.ExecuteGuard(TEXT("RemoveUnusedTextureWarning"), Texture->Name);
+	Validation.RemoveUnusedWarning(Texture);
+#endif
+}
 
-		// Increment the number of time the texture has been accessed to avoid warning on produced but never used resources that were produced
-		// only to be extracted for the graph.
-		Texture->PassAccessCount += 1;
-	}
+inline void FRDGBuilder::RemoveUnusedBufferWarning(FRDGBufferRef Buffer)
+{
+#if RDG_ENABLE_DEBUG
+	check(Buffer);
+	Validation.ExecuteGuard(TEXT("RemoveUnusedBufferWarning"), Buffer->Name);
+	Validation.RemoveUnusedWarning(Buffer);
 #endif
 }
 
diff --git a/Engine/Source/Runtime/RenderCore/Public/RenderGraphDefinitions.h b/Engine/Source/Runtime/RenderCore/Public/RenderGraphDefinitions.h
--- a/Engine/Source/Runtime/RenderCore/Public/RenderGraphDefinitions.h
+++ b/Engine/Source/Runtime/RenderCore/Public/RenderGraphDefinitions.h
@@ -9,6 +9,13 @@
 /** Whether render graph debugging is enabled. */
 #define RDG_ENABLE_DEBUG (!UE_BUILD_SHIPPING && !UE_BUILD_TEST)
 
+/** Performs the operation if RDG_ENABLE_DEBUG is enabled. Useful for one-line checks without explicitly wrapping in #if. */
+#if RDG_ENABLE_DEBUG
+#define IF_RDG_ENABLE_DEBUG(Op) Op
+#else
+#define IF_RDG_ENABLE_DEBUG(Op)
+#endif
+
 /** Whether render graph debugging is enabled and we are compiling with the engine. */
 #define RDG_ENABLE_DEBUG_WITH_ENGINE (RDG_ENABLE_DEBUG && WITH_ENGINE)
 
diff --git a/Engine/Source/Runtime/RenderCore/Public/RenderGraphPass.h b/Engine/Source/Runtime/RenderCore/Public/RenderGraphPass.h
--- a/Engine/Source/Runtime/RenderCore/Public/RenderGraphPass.h
+++ b/Engine/Source/Runtime/RenderCore/Public/RenderGraphPass.h
@@ -162,20 +162,25 @@ private:
 };
 
 /** Flags to annotate passes. */
-enum class ERDGPassFlags
+enum class ERDGPassFlags : uint8
 {
-	None = 0,
+	/** Equivalent to 'Raster'. */
+	None = 0x1,
+
+	/** Pass uses raster pipeline. */
+	Raster = None,
 
 	/** Pass uses compute only */
-	Compute,
+	Compute = 0x2,
 
 	/** Pass uses RHI copy commands only. */
-	Copy,
+	Copy = 0x4,
 
 	//#todo-rco: Remove this when we can do split/per mip layout transitions.
 	/** Hint to some RHIs this pass will be generating mips to optimize transitions. */
-	GenerateMips
+	GenerateMips = 0x8
 };
+ENUM_CLASS_FLAGS(ERDGPassFlags);
 
 // TODO(RDG): Bulk rename across codebase.
 using ERenderGraphPassFlags = ERDGPassFlags;
@@ -203,24 +208,24 @@ public:
 		return PassFlags;
 	}
 
-	bool IsGenerateMips() const
+	bool IsRaster() const
 	{
-		return (PassFlags == ERDGPassFlags::GenerateMips);
+		return (PassFlags & ERDGPassFlags::Raster) == ERDGPassFlags::Raster;
 	}
 
 	bool IsCompute() const
 	{
-		return (PassFlags == ERDGPassFlags::Compute);
+		return (PassFlags & ERDGPassFlags::Compute) == ERDGPassFlags::Compute;
 	}
 
 	bool IsCopy() const
 	{
-		return (PassFlags == ERDGPassFlags::Copy);
+		return (PassFlags & ERDGPassFlags::Copy) == ERDGPassFlags::Copy;
 	}
 
-	bool IsGraphics() const
+	bool IsGenerateMips() const
 	{
-		return (PassFlags == ERDGPassFlags::None) || (PassFlags == ERDGPassFlags::GenerateMips);
+		return (PassFlags & ERDGPassFlags::GenerateMips) == ERDGPassFlags::GenerateMips;
 	}
 
 	FRDGPassParameterStruct GetParameters() const
diff --git a/Engine/Source/Runtime/RenderCore/Public/RenderGraphResources.h b/Engine/Source/Runtime/RenderCore/Public/RenderGraphResources.h
--- a/Engine/Source/Runtime/RenderCore/Public/RenderGraphResources.h
+++ b/Engine/Source/Runtime/RenderCore/Public/RenderGraphResources.h
@@ -46,10 +46,6 @@ struct FRDGResourceState
 		MAX = Unknown
 	};
 
-	/** Constructs a read / write state suitable for the provided pass. */
-	static FRDGResourceState CreateRead(const FRDGPass* Pass);
-	static FRDGResourceState CreateWrite(const FRDGPass* Pass);
-
 	FRDGResourceState() = default;
 
 	FRDGResourceState(const FRDGPass* InPass, EPipeline InPipeline, EAccess InAccess)
@@ -156,6 +152,7 @@ private:
 #endif
 
 	friend class FRDGBuilder;
+	friend class FRDGUserValidation;
 };
 
 /** A render graph resource with an allocation lifetime tracked by the graph. */
@@ -219,6 +216,7 @@ private:
 	FRDGResourceState State;
 
 	friend class FRDGBuilder;
+	friend class FRDGUserValidation;
 	friend class FRDGBarrierBatcher;
 };
 
@@ -273,6 +271,7 @@ private:
 	IPooledRenderTarget* PooledRenderTarget = nullptr;
 
 	friend class FRDGBuilder;
+	friend class FRDGUserValidation;
 	friend class FRDGBarrierBatcher;
 };
 
