Adding missing local files from integration.

diff --git a/Engine/Source/Runtime/RenderCore/Private/RenderGraphValidation.cpp b/Engine/Source/Runtime/RenderCore/Private/RenderGraphValidation.cpp
new file mode 100644
--- /dev/null
+++ b/Engine/Source/Runtime/RenderCore/Private/RenderGraphValidation.cpp
@@ -0,0 +1,688 @@
+// Copyright 1998-2019 Epic Games, Inc. All Rights Reserved.
+
+#include "RenderGraphValidation.h"
+
+#if RDG_ENABLE_DEBUG
+
+/** This utility class tracks read / write events on individual resources within a pass in order to ensure
+ *  that sub-resources are not being simultaneously bound for read / write. The RHI currently only supports
+ *  disjoint read-write access on a texture for mip-map generation purposes.
+ */
+class FRDGPassResourceValidator
+{
+public:
+	static const uint32 AllMipLevels = ~0;
+
+	FRDGPassResourceValidator(const FRDGPass& InPass)
+		: Pass(InPass)
+		, bIsGenerateMips(InPass.IsGenerateMips())
+	{}
+
+	void Read(FRDGBufferRef Buffer)
+	{
+		check(Buffer);
+
+		ensureMsgf(!bIsGenerateMips,
+			TEXT("Attempting to read from Buffer %s on Pass %s which is marked as 'GenerateMips'. Only a single texture is supported on this type of pass."),
+			Buffer->Name, Pass.GetName());
+
+		FBufferAccess& BufferAccess = BufferAccessMap.FindOrAdd(Buffer);
+
+		ensureMsgf(BufferAccess.Access != FRDGResourceState::EAccess::Write,
+			TEXT("Buffer %s is being used for both read and write access on Pass %s. This is not supported."),
+			Buffer->Name, Pass.GetName());
+
+		BufferAccess.Access = FRDGResourceState::EAccess::Read;
+	}
+
+	void Write(FRDGBufferRef Buffer)
+	{
+		check(Buffer);
+
+		ensureMsgf(!bIsGenerateMips,
+			TEXT("Attempting to write to Buffer %s on Pass %s which is marked as 'GenerateMips'. Only a single texture is supported on this type of pass."),
+			Buffer->Name, Pass.GetName());
+
+		FBufferAccess& BufferAccess = BufferAccessMap.FindOrAdd(Buffer);
+
+		ensureMsgf(BufferAccess.Access != FRDGResourceState::EAccess::Read,
+			TEXT("Attempting to simultaneously read and write from Buffer %s on Pass %s. This is not supported."),
+			Buffer->Name, Pass.GetName());
+
+		BufferAccess.Access = FRDGResourceState::EAccess::Read;
+	}
+
+	void Read(FRDGTextureRef Texture, uint32 MipLevel = AllMipLevels)
+	{
+		check(Texture);
+
+		const uint32 MipMask = (MipLevel == AllMipLevels) ? AllMipLevels : (1 << MipLevel);
+
+		FTextureAccess& TextureAccess = TextureAccessMap.FindOrAdd(Texture);
+		TextureAccess.SetAllRead(MipMask);
+
+		ensureMsgf(!TextureAccess.IsAnyWrite(MipMask),
+			TEXT("Attempting to access mips of Texture %s for both read and write access on Pass %s. This is not supported."),
+			MipLevel, Texture->Name, Pass.GetName());
+
+		if (bIsGenerateMips)
+		{
+			ensureMsgf(TextureAccessMap.Num() == 1,
+				TEXT("Attempting to access more than one texture in Pass %s when 'GenerateMips' is set. Only a single texture is supported on this type of pass."),
+				Pass.GetName());
+		}
+		else
+		{
+			ensureMsgf(!TextureAccess.Mips.Write,
+				TEXT("Texture %s is being used for both read and write access on the same Pass %s. This is only supported via the 'GenerateMips' flag."),
+				Pass.GetName());
+		}
+	}
+
+	void Write(FRDGTextureRef Texture, uint32 MipLevel = AllMipLevels)
+	{
+		check(Texture);
+
+		const uint32 MipMask = 1 << MipLevel;
+
+		FTextureAccess& TextureAccess = TextureAccessMap.FindOrAdd(Texture);
+		TextureAccess.SetAllWrite(MipMask);
+
+		ensureMsgf(!TextureAccess.IsAnyRead(MipMask),
+			TEXT("Attempting to use Mip %d of Texture %s for both read and write access on Pass %s. This is not supported."),
+			MipLevel, Texture->Name, Pass.GetName());
+
+		if (bIsGenerateMips)
+		{
+			ensureMsgf(TextureAccessMap.Num() == 1,
+				TEXT("Attempting to access more than one texture in Pass %s when 'GenerateMips' is set. Only a single texture is supported on this type of pass."),
+				Pass.GetName());
+		}
+		else
+		{
+			ensureMsgf(!TextureAccess.Mips.Read,
+				TEXT("Texture %s is being used for both read and write access on the same Pass %s. This is only supported via the 'GenerateMips' flag."),
+				Texture->Name, Pass.GetName());
+		}
+	}
+
+private:
+	struct FBufferAccess
+	{
+		FRDGResourceState::EAccess Access = FRDGResourceState::EAccess::Unknown;
+
+		FORCEINLINE bool operator!= (FBufferAccess Other) const
+		{
+			return Access != Other.Access;
+		}
+
+		FORCEINLINE bool operator< (FBufferAccess Other) const
+		{
+			return Access < Other.Access;
+		}
+	};
+
+	struct FTextureAccess
+	{
+		FTextureAccess() = default;
+
+		union
+		{
+			uint64 PackedData = 0;
+
+			struct
+			{
+				uint32 Read;
+				uint32 Write;
+			} Mips;
+		};
+
+		FORCEINLINE bool IsAnyRead(uint32 MipMask) const
+		{
+			return (Mips.Read & MipMask) != 0;
+		}
+
+		FORCEINLINE bool IsAnyWrite(uint32 MipMask) const
+		{
+			return (Mips.Write & MipMask) != 0;
+		}
+
+		FORCEINLINE void SetAllRead(uint32 MipMask)
+		{
+			Mips.Read |= MipMask;
+		}
+
+		FORCEINLINE void SetAllWrite(uint32 MipMask)
+		{
+			Mips.Write |= MipMask;
+		}
+
+		FORCEINLINE bool operator!= (FTextureAccess Other) const
+		{
+			return PackedData != Other.PackedData;
+		}
+
+		FORCEINLINE bool operator< (FTextureAccess Other) const
+		{
+			return PackedData < Other.PackedData;
+		}
+	};
+
+	TMap<FRDGBufferRef, FBufferAccess, SceneRenderingSetAllocator> BufferAccessMap;
+	TMap<FRDGTextureRef, FTextureAccess, SceneRenderingSetAllocator> TextureAccessMap;
+	const FRDGPass& Pass;
+	const bool bIsGenerateMips;
+};
+
+FRDGUserValidation::~FRDGUserValidation()
+{
+	checkf(bHasExecuted, TEXT("Render graph execution is required to ensure consistency with immediate mode."));
+}
+
+void FRDGUserValidation::ExecuteGuard(const TCHAR* Operation, const TCHAR* ResourceName)
+{
+	checkf(!bHasExecuted, TEXT("Render graph operation '%s' with resource '%s' must be performed prior to graph execution."), Operation, ResourceName);
+}
+
+void FRDGUserValidation::ValidateCreateResource(FRDGTrackedResourceRef Resource)
+{
+	check(Resource);
+	TrackedResources.Add(Resource);
+}
+
+void FRDGUserValidation::ValidateCreateExternalResource(FRDGTrackedResourceRef Resource)
+{
+	ValidateCreateResource(Resource);
+	Resource->MarkAsExternal();
+}
+
+void FRDGUserValidation::ValidateExtractResource(FRDGTrackedResourceRef Resource)
+{
+	check(Resource);
+
+	checkf(Resource->HasBeenProduced(),
+		TEXT("Unable to queue the extraction of the resource %s because it has not been produced by any pass."),
+		Resource->Name);
+
+	/** Increment pass access counts for externally registered buffers and textures to avoid
+	 *  emitting a 'produced but never used' warning. We don't have the history of registered
+	 *  resources to be able to emit a proper warning.
+	 */
+
+	Resource->PassAccessCount++;
+}
+
+void FRDGUserValidation::RemoveUnusedWarning(FRDGTrackedResourceRef Resource)
+{
+	check(Resource);
+	Resource->PassAccessCount++;
+}
+
+void FRDGUserValidation::ValidateAllocPassParameters(const void* Parameters)
+{
+	check(Parameters);
+	AllocatedUnusedPassParameters.Add(Parameters);
+}
+
+void FRDGUserValidation::ValidateAddPass(const FRDGPass* Pass)
+{
+	checkf(!bHasExecuted, TEXT("Render graph pass %s needs to be added before the builder execution."), Pass->GetName());
+
+	{
+		const void* ParameterStructData = Pass->GetParameters().GetContents();
+
+		/** The lifetime of each pass parameter structure must extend until deferred pass execution; therefore, it needs to be
+		 *  allocated with FRDGBuilder::AllocParameters(). Also, all references held by the parameter structure are released
+		 *  immediately after pass execution, so a pass parameter struct instance must be 1-to-1 with a pass instance (i.e. one
+		 *  per AddPass() call).
+		 */
+		checkf(
+			AllocatedUnusedPassParameters.Contains(ParameterStructData),
+			TEXT("The pass parameter structure has not been allocated for correct life time FRDGBuilder::AllocParameters() or has already ")
+			TEXT("been used by another previous FRDGBuilder::AddPass()."));
+
+		AllocatedUnusedPassParameters.Remove(ParameterStructData);
+	}
+
+	const FRenderTargetBindingSlots* RenderTargetBindingSlots = nullptr;
+
+	const TCHAR* PassName = Pass->GetName();
+	const bool bIsRaster = Pass->IsRaster();
+	const bool bIsCompute = Pass->IsCompute();
+	const bool bIsCopy = Pass->IsCopy();
+	const bool bIsGeneratingMips = Pass->IsGenerateMips();
+
+	/** Validate that the user set the correct pass flags. */
+	{
+		uint32 Count = 0;
+		Count += bIsRaster ? 1 : 0;
+		Count += bIsCompute ? 1 : 0;
+		Count += bIsCopy ? 1 : 0;
+		checkf(Count == 1,
+			TEXT("Pass %s must be declared as either Raster, Compute, or Copy. These flags cannot be combined."),
+			PassName);
+
+		if (bIsGeneratingMips)
+		{
+			checkf(bIsRaster || bIsCompute,
+				TEXT("Pass %s is declared as generating mips. This is only supported by Raster or Compute passes."),
+				PassName);
+		}
+	}
+
+	FRDGPassResourceValidator PassResourceValidator(*Pass);
+
+	FRDGPassParameterStruct ParameterStruct = Pass->GetParameters();
+
+	const uint32 ParameterCount = ParameterStruct.GetParameterCount();
+
+	for (uint32 ParameterIndex = 0; ParameterIndex < ParameterCount; ++ParameterIndex)
+	{
+		FRDGPassParameter Parameter = ParameterStruct.GetParameter(ParameterIndex);
+
+		switch (Parameter.GetType())
+		{
+		case UBMT_RDG_TEXTURE:
+		{
+			if (FRDGTextureRef Texture = Parameter.GetAsTexture())
+			{
+				PassResourceValidator.Read(Texture);
+
+				checkf(Texture->HasBeenProduced(),
+					TEXT("Pass %s has a dependency over the texture %s that has never been produced."),
+					PassName, Texture->Name);
+
+				Texture->PassAccessCount++;
+			}
+		}
+		break;
+		case UBMT_RDG_TEXTURE_SRV:
+		{
+			if (FRDGTextureSRVRef SRV = Parameter.GetAsTextureSRV())
+			{
+				FRDGTextureRef Texture = SRV->Desc.Texture;
+
+				PassResourceValidator.Read(Texture, SRV->Desc.MipLevel);
+
+				checkf(Texture->HasBeenProduced(),
+					TEXT("Pass %s has a dependency over the texture %s that has never been produced."),
+					PassName, Texture->Name);
+
+				Texture->PassAccessCount++;
+			}
+		}
+		break;
+		case UBMT_RDG_TEXTURE_UAV:
+		{
+			if (FRDGTextureUAVRef UAV = Parameter.GetAsTextureUAV())
+			{
+				FRDGTextureRef Texture = UAV->Desc.Texture;
+
+				PassResourceValidator.Write(Texture, UAV->Desc.MipLevel);
+
+				Texture->MarkAsProducedBy(Pass);
+				Texture->PassAccessCount++;
+			}
+		}
+		break;
+		case UBMT_RDG_BUFFER:
+		{
+			if (FRDGBufferRef Buffer = Parameter.GetAsBuffer())
+			{
+				PassResourceValidator.Read(Buffer);
+
+				checkf(Buffer->HasBeenProduced(),
+					TEXT("Pass %s has a dependency over the buffer %s that has never been produced."),
+					PassName, Buffer->Name);
+
+				Buffer->PassAccessCount++;
+			}
+		}
+		break;
+		case UBMT_RDG_BUFFER_SRV:
+		{
+			if (FRDGBufferSRVRef SRV = Parameter.GetAsBufferSRV())
+			{
+				FRDGBufferRef Buffer = SRV->Desc.Buffer;
+
+				PassResourceValidator.Read(Buffer);
+
+				checkf(Buffer->HasBeenProduced(),
+					TEXT("Pass %s has a dependency over the buffer %s that has never been produced."),
+					PassName, SRV->Desc.Buffer->Name);
+
+				Buffer->PassAccessCount++;
+			}
+		}
+		break;
+		case UBMT_RDG_BUFFER_UAV:
+		{
+			if (FRDGBufferUAVRef UAV = Parameter.GetAsBufferUAV())
+			{
+				FRDGBufferRef Buffer = UAV->Desc.Buffer;
+
+				PassResourceValidator.Write(Buffer);
+
+				Buffer->MarkAsProducedBy(Pass);
+				Buffer->PassAccessCount++;
+			}
+		}
+		break;
+		case UBMT_RDG_TEXTURE_COPY_DEST:
+		{
+			if (FRDGTextureRef Texture = Parameter.GetAsTexture())
+			{
+				PassResourceValidator.Write(Texture);
+
+				Texture->MarkAsProducedBy(Pass);
+				Texture->PassAccessCount++;
+			}
+		}
+		break;
+		case UBMT_RDG_BUFFER_COPY_DEST:
+		{
+			if (FRDGBufferRef Buffer = Parameter.GetAsBuffer())
+			{
+				PassResourceValidator.Write(Buffer);
+
+				Buffer->MarkAsProducedBy(Pass);
+				Buffer->PassAccessCount++;
+			}
+		}
+		break;
+		case UBMT_RENDER_TARGET_BINDING_SLOTS:
+		{
+			if (!RenderTargetBindingSlots)
+			{
+				RenderTargetBindingSlots = &Parameter.GetAsRenderTargetBindingSlots();
+			}
+			else if (IsRDGDebugEnabled())
+			{
+				EmitRDGWarningf(
+					TEXT("Pass %s have duplicated render target binding slots."),
+					PassName);
+			}
+		}
+		break;
+		default:
+			break;
+		}
+	}
+
+	/** Validate that raster passes have render target binding slots and compute passes don't. */
+	if (RenderTargetBindingSlots)
+	{
+		checkf(bIsRaster, TEXT("Pass '%s' has render target binding slots but is flagged as 'Compute'."), PassName);
+	}
+	else
+	{
+		checkf(!bIsRaster, TEXT("Pass '%s' is missing render target binding slots. Set the 'Compute' or 'Copy' flag if render targets are not required."), PassName);
+	}
+
+	/** Validate render target / depth stencil binding usage. */
+	if (RenderTargetBindingSlots)
+	{
+		const auto& RenderTargets = RenderTargetBindingSlots->Output;
+
+		{
+			const auto& DepthStencil = RenderTargetBindingSlots->DepthStencil;
+
+			if (FRDGTextureRef Texture = DepthStencil.GetTexture())
+			{
+				ensureMsgf(Texture->Desc.NumSamples == 1,
+					TEXT("Pass '%s' uses an MSAA depth-stencil render target. This is not yet supported."),
+					Pass->GetName());
+
+				ensureMsgf(!bIsGeneratingMips,
+					TEXT("Pass '%s' is marked to generate mips but has a depth stencil texture. This is not supported."),
+					Pass->GetName());
+
+				// Depth stencil only supports one mip, since there isn't actually a way to select the mip level.
+				const uint32 MipLevel = 0;
+				check(Texture->Desc.NumMips == 1);
+
+				if (DepthStencil.GetDepthStencilAccess().IsAnyWrite())
+				{
+					PassResourceValidator.Write(Texture, MipLevel);
+					Texture->MarkAsProducedBy(Pass);
+				}
+				else
+				{
+					PassResourceValidator.Read(Texture, MipLevel);
+				}
+
+				Texture->PassAccessCount++;
+			}
+		}
+
+		const uint32 RenderTargetCount = RenderTargets.Num();
+
+		{
+			/** Tracks the number of contiguous, non-null textures in the render target output array. */
+			uint32 ValidRenderTargetCount = 0;
+
+			for (uint32 RenderTargetIndex = 0; RenderTargetIndex < RenderTargetCount; ++RenderTargetIndex)
+			{
+				const FRenderTargetBinding& RenderTarget = RenderTargets[RenderTargetIndex];
+
+				if (FRDGTextureRef Texture = RenderTarget.GetTexture())
+				{
+					PassResourceValidator.Write(Texture, RenderTarget.GetMipIndex());
+
+					const bool bIsLoadAction = RenderTarget.GetLoadAction() == ERenderTargetLoadAction::ELoad;
+
+					/** Validate that load action is correct. We can only load contents if a pass previously produced something. */
+					{
+						const bool bIsLoadActionInvalid = bIsLoadAction && !Texture->HasBeenProduced();
+						checkf(
+							!bIsLoadActionInvalid,
+							TEXT("Pass '%s' attempted to bind texture '%s' as a render target with the 'Load' action specified, but the texture has not been produced yet. The render target must use either 'Clear' or 'NoAction' action instead."),
+							PassName,
+							Texture->Name);
+					}
+
+					/** Validate that any previously produced texture contents are loaded. This occurs if the user failed to specify a load action
+					 *  on a texture that was produced by a previous pass, effectively losing that data. This can also happen if the user 're-uses'
+					 *  a texture for some other purpose. The latter is considered bad practice, since it increases memory pressure on the render
+					 *  target pool. Instead, the user should create a new texture instance. An exception to this rule are untracked render targets,
+					 *  which are not actually managed by the render target pool and likely represent the frame buffer.
+					 */
+					{
+						// We only validate single-mip textures since we don't track production at the subresource level.
+						const bool bFailedToLoadProducedContent = !bIsLoadAction && Texture->HasBeenProduced() && Texture->Desc.NumMips == 1;
+
+						// Untracked render targets aren't actually managed by the render target pool.
+						const bool bIsUntrackedRenderTarget = Texture->PooledRenderTarget && Texture->PooledRenderTarget->IsTracked();
+
+						ensureMsgf(!bFailedToLoadProducedContent || bIsUntrackedRenderTarget,
+							TEXT("Pass '%s' attempted to bind texture '%s' as a render target without the 'Load' action specified, despite a prior pass having produced it. It's invalid to completely clobber the contents of a resource. Create a new texture instance instead."),
+							PassName,
+							Texture->Name);
+					}
+					
+					/** Mark the pass as a producer for render targets with a store action. */
+					{
+						const bool bIsStoreAction = RenderTarget.GetStoreAction() != ERenderTargetStoreAction::ENoAction;
+						check(bIsStoreAction); // already been validated in FRenderTargetBinding::Validate()
+						Texture->MarkAsProducedBy(Pass);
+					}
+
+					Texture->PassAccessCount++;
+				}
+				else
+				{
+					/** Found end of contiguous interval of valid render targets. */
+					ValidRenderTargetCount = RenderTargetIndex;
+					break;
+				}
+			}
+
+			/** Validate that no holes exist in the render target output array. Render targets must be bound contiguously. */
+			for (uint32 RenderTargetIndex = ValidRenderTargetCount; RenderTargetIndex < RenderTargetCount; ++RenderTargetIndex)
+			{
+				const FRenderTargetBinding& RenderTarget = RenderTargets[RenderTargetIndex];
+				checkf(RenderTarget.GetTexture() == nullptr, TEXT("Render targets must be packed. No empty spaces in the array."));
+			}
+		}
+	}
+}
+
+void FRDGUserValidation::ValidateExecuteBegin()
+{
+	checkf(!bHasExecuted, TEXT("Render graph execution should only happen once to ensure consistency with immediate mode."));
+
+	/** FRDGBuilder::AllocParameters() allocates shader parameter structure for the lifetime until pass execution.
+	 *  But they are allocated on a FMemStack for CPU performance, and are destructed immediately after pass execution.
+	 *  Therefore allocating pass parameter unused by a FRDGBuilder::AddPass() can lead on a memory leak of RHI resources
+	 *  that have been reference in the parameter structure.
+	 */
+	checkf(
+		AllocatedUnusedPassParameters.Num() == 0,
+		TEXT("%i pass parameter structure has been allocated with FRDGBuilder::AllocParameters(), but has not be used by a ")
+		TEXT("FRDGBuilder::AddPass() that can cause RHI resource leak."), AllocatedUnusedPassParameters.Num());
+}
+
+void FRDGUserValidation::ValidateExecuteEnd()
+{
+	bHasExecuted = true;
+
+	const bool bEmitWarnings = IsRDGDebugEnabled();
+
+	for (const FRDGTrackedResourceRef Resource : TrackedResources)
+	{
+		check(Resource->ReferenceCount == 0);
+
+		const bool bProducedButNeverUsed = Resource->PassAccessCount == 1 && Resource->FirstProducer;
+
+		if (bEmitWarnings && bProducedButNeverUsed)
+		{
+			check(Resource->HasBeenProduced());
+
+			EmitRDGWarningf(
+				TEXT("Resources %s has been produced by the pass %s, but never used by another pass."),
+				Resource->Name, Resource->FirstProducer->GetName());
+		}
+	}
+
+	TrackedResources.Empty();
+}
+
+void FRDGUserValidation::ValidateExecutePassBegin(const FRDGPass* Pass)
+{
+	SetAllowRHIAccess(Pass, true);
+}
+
+void FRDGUserValidation::ValidateExecutePassEnd(const FRDGPass* Pass)
+{
+	SetAllowRHIAccess(Pass, false);
+
+	FRDGPassParameterStruct ParameterStruct = Pass->GetParameters();
+
+	const uint32 ParameterCount = ParameterStruct.GetParameterCount();
+
+	if (IsRDGDebugEnabled())
+	{
+		uint32 TrackedResourceCount = 0;
+		uint32 UsedResourceCount = 0;
+
+		for (uint32 ParameterIndex = 0; ParameterIndex < ParameterCount; ++ParameterIndex)
+		{
+			FRDGPassParameter Parameter = ParameterStruct.GetParameter(ParameterIndex);
+
+			if (Parameter.IsResource())
+			{
+				if (FRDGResourceRef Resource = Parameter.GetAsResource())
+				{
+					TrackedResourceCount++;
+					UsedResourceCount += Resource->bIsActuallyUsedByPass ? 1 : 0;
+				}
+			}
+		}
+
+		if (TrackedResourceCount != UsedResourceCount)
+		{
+			FString WarningMessage = FString::Printf(
+				TEXT("'%d' of the '%d' resources of the pass '%s' where not actually used."),
+				TrackedResourceCount - UsedResourceCount, TrackedResourceCount, Pass->GetName());
+
+			for (uint32 ParameterIndex = 0; ParameterIndex < ParameterCount; ++ParameterIndex)
+			{
+				FRDGPassParameter Parameter = ParameterStruct.GetParameter(ParameterIndex);
+
+				if (Parameter.IsResource())
+				{
+					if (const FRDGResourceRef Resource = Parameter.GetAsResource())
+					{
+						if (!Resource->bIsActuallyUsedByPass)
+						{
+							WarningMessage += FString::Printf(TEXT("\n    %s"), Resource->Name);
+						}
+					}
+				}
+			}
+
+			EmitRDGWarning(WarningMessage);
+		}
+	}
+
+	for (uint32 ParameterIndex = 0; ParameterIndex < ParameterCount; ++ParameterIndex)
+	{
+		FRDGPassParameter Parameter = ParameterStruct.GetParameter(ParameterIndex);
+
+		if (Parameter.IsResource())
+		{
+			if (const FRDGResourceRef Resource = Parameter.GetAsResource())
+			{
+				Resource->bIsActuallyUsedByPass = false;
+			}
+		}
+	}
+}
+
+void FRDGUserValidation::SetAllowRHIAccess(const FRDGPass* Pass, bool bAllowAccess)
+{
+	FRDGPassParameterStruct ParameterStruct = Pass->GetParameters();
+
+	const uint32 ParameterCount = ParameterStruct.GetParameterCount();
+
+	for (uint32 ParameterIndex = 0; ParameterIndex < ParameterCount; ++ParameterIndex)
+	{
+		FRDGPassParameter Parameter = ParameterStruct.GetParameter(ParameterIndex);
+
+		if (Parameter.IsResource())
+		{
+			if (FRDGResourceRef Resource = Parameter.GetAsResource())
+			{
+				Resource->bAllowRHIAccess = bAllowAccess;
+			}
+		}
+		else if (Parameter.GetType() == UBMT_RENDER_TARGET_BINDING_SLOTS)
+		{
+			const FRenderTargetBindingSlots& RenderTargetBindingSlots = Parameter.GetAsRenderTargetBindingSlots();
+			const auto& RenderTargets = RenderTargetBindingSlots.Output;
+			const auto& DepthStencil = RenderTargetBindingSlots.DepthStencil;
+			const uint32 RenderTargetCount = RenderTargets.Num();
+
+			for (uint32 RenderTargetIndex = 0; RenderTargetIndex < RenderTargetCount; RenderTargetIndex++)
+			{
+				const FRenderTargetBinding& RenderTarget = RenderTargets[RenderTargetIndex];
+
+				if (FRDGTextureRef Texture = RenderTarget.GetTexture())
+				{
+					Texture->bAllowRHIAccess = bAllowAccess;
+				}
+				else
+				{
+					break;
+				}
+			}
+
+			if (FRDGTextureRef Texture = DepthStencil.GetTexture())
+			{
+				Texture->bAllowRHIAccess = bAllowAccess;
+			}
+		}
+	}
+}
+
+#endif
\ No newline at end of file
diff --git a/Engine/Source/Runtime/RenderCore/Public/RenderGraphValidation.h b/Engine/Source/Runtime/RenderCore/Public/RenderGraphValidation.h
new file mode 100644
--- /dev/null
+++ b/Engine/Source/Runtime/RenderCore/Public/RenderGraphValidation.h
@@ -0,0 +1,85 @@
+// Copyright 1998-2019 Epic Games, Inc. All Rights Reserved.
+
+#pragma once
+
+#include "RenderGraphPass.h"
+
+/** Returns whether render graph runtime debugging is enabled. */
+extern bool IsRDGDebugEnabled();
+
+/** Returns whether render graph is running in immediate mode. */
+extern bool IsRDGImmediateModeEnabled();
+
+/** Emits a render graph warning. */
+extern void EmitRDGWarning(const FString& WarningMessage);
+
+#define EmitRDGWarningf(WarningMessageFormat, ...) \
+	EmitRDGWarning(FString::Printf(WarningMessageFormat, ##__VA_ARGS__));
+
+#if RDG_ENABLE_DEBUG
+
+/** Used by the render graph builder to validate correct usage of the graph API from setup to execution.
+ *  Validation is compiled out in shipping builds. This class tracks resources and passes as they are
+ *  added to the graph. It will then validate execution of the graph, including whether resources are
+ *  used during execution, and that they are properly produced before being consumed. All found issues
+ *  must be clear enough to help the user identify the problem in client code. Validation should occur
+ *  as soon as possible in the graph lifecycle. It's much easier to catch an issue at the setup location
+ *  rather than during deferred execution.
+ *
+ *  Finally, this class is designed for user validation, not for internal graph validation. In other words,
+ *  if the user can break the graph externally via the client-facing API, this validation layer should catch it.
+ *  Any internal validation of the graph state should be kept out of this class in order to provide a clear
+ *  and modular location to extend the validation layer as well as clearly separate the graph implementation
+ *  details from events in the graph.
+ */
+class RENDERCORE_API FRDGUserValidation final
+{
+public:
+	FRDGUserValidation() = default;
+	FRDGUserValidation(const FRDGUserValidation&) = delete;
+	~FRDGUserValidation();
+
+	/** Validates that the graph has not executed yet. */
+	void ExecuteGuard(const TCHAR* Operation, const TCHAR* ResourceName);
+
+	/** Tracks and validates the creation of a new resource in the graph. */
+	void ValidateCreateResource(FRDGTrackedResourceRef Resource);
+
+	/** Tracks and validates the creation of a new externally registered resource. */
+	void ValidateCreateExternalResource(FRDGTrackedResourceRef Resource);
+
+	/** Tracks and validates usage of a pass parameters allocation. */
+	void ValidateAllocPassParameters(const void* Parameters);
+
+	/** Validates a resource extraction operation. */
+	void ValidateExtractResource(FRDGTrackedResourceRef Resource);
+
+	/** Tracks and validates the addition of a new pass to the graph. */
+	void ValidateAddPass(const FRDGPass* Pass);
+
+	/** Validate pass state before and after execution. */
+	void ValidateExecutePassBegin(const FRDGPass* Pass);
+	void ValidateExecutePassEnd(const FRDGPass* Pass);
+
+	/** Validate graph state before and after execution. */
+	void ValidateExecuteBegin();
+	void ValidateExecuteEnd();
+
+	/** Removes the 'produced but not used' warning from the requested resource. */
+	void RemoveUnusedWarning(FRDGTrackedResourceRef Resource);
+
+private:
+	/** Traverses all resources in the pass and marks whether they are externally accessible by user pass implementations. */
+	static void SetAllowRHIAccess(const FRDGPass* Pass, bool bAllowAccess);
+
+	/** All recently allocated pass parameter structure, but not used by a AddPass() yet. */
+	TSet<const void*, DefaultKeyFuncs<const void*>, SceneRenderingSetAllocator> AllocatedUnusedPassParameters;
+
+	/** List of tracked resources for validation prior to shutdown. */
+	TArray<FRDGTrackedResourceRef, SceneRenderingAllocator> TrackedResources;
+
+	/** Whether the Execute() has already been called. */
+	bool bHasExecuted = false;
+};
+
+#endif
\ No newline at end of file
