Fixing broken build.

diff --git a/Engine/Source/Runtime/RenderCore/Private/GenerateMips.cpp b/Engine/Source/Runtime/RenderCore/Private/GenerateMips.cpp
--- a/Engine/Source/Runtime/RenderCore/Private/GenerateMips.cpp
+++ b/Engine/Source/Runtime/RenderCore/Private/GenerateMips.cpp
@@ -131,7 +131,7 @@ void FGenerateMips::Compute(FRHICommandListImmediate& RHIImmCmdList, FRHITexture
 		GraphBuilder.AddPass(
 			RDG_EVENT_NAME("Generate2DTextureMips DestMipLevel=%d", MipLevel),
 			PassParameters,
-			ERenderGraphPassFlags::Compute | ERDGPassFlags::GenerateMips,
+			ERDGPassFlags::Compute | ERDGPassFlags::GenerateMips,
 			[PassParameters, ComputeShader, GenMipsGroupCount](FRHICommandList& RHICmdList)
 		{
 			FComputeShaderUtils::Dispatch(RHICmdList, *ComputeShader, *PassParameters, GenMipsGroupCount);
diff --git a/Engine/Source/Runtime/Renderer/Private/SceneOcclusion.cpp b/Engine/Source/Runtime/Renderer/Private/SceneOcclusion.cpp
--- a/Engine/Source/Runtime/Renderer/Private/SceneOcclusion.cpp
+++ b/Engine/Source/Runtime/Renderer/Private/SceneOcclusion.cpp
@@ -1203,7 +1203,7 @@ void BuildHZB(FRDGBuilder& GraphBuilder, FViewInfo& View)
 		GraphBuilder.AddPass(
 			RDG_EVENT_NAME("HZB(mip=%d) %dx%d", MipIndex, DstSize.X, DstSize.Y),
 			PassParameters,
-			ERenderGraphPassFlags::GenerateMips,
+			ERenderGraphPassFlags::Raster | ERenderGraphPassFlags::GenerateMips,
 			[PassParameters, SrcSize, DstSize, &View](FRHICommandListImmediate& RHICmdList)
 			{
 				FHZBBuildPS::FPermutationDomain PermutationVector;
