Fixed render graph validation bug.

diff --git a/Engine/Source/Runtime/RenderCore/Private/RenderGraphValidation.cpp b/Engine/Source/Runtime/RenderCore/Private/RenderGraphValidation.cpp
--- a/Engine/Source/Runtime/RenderCore/Private/RenderGraphValidation.cpp
+++ b/Engine/Source/Runtime/RenderCore/Private/RenderGraphValidation.cpp
@@ -492,14 +492,14 @@ void FRDGUserValidation::ValidateAddPass(const FRDGPass* Pass)
 						const bool bFailedToLoadProducedContent = !bIsLoadAction && Texture->HasBeenProduced() && Texture->Desc.NumMips == 1;
 
 						// Untracked render targets aren't actually managed by the render target pool.
-						const bool bIsUntrackedRenderTarget = Texture->PooledRenderTarget && Texture->PooledRenderTarget->IsTracked();
+						const bool bIsUntrackedRenderTarget = Texture->PooledRenderTarget && !Texture->PooledRenderTarget->IsTracked();
 
 						ensureMsgf(!bFailedToLoadProducedContent || bIsUntrackedRenderTarget,
 							TEXT("Pass '%s' attempted to bind texture '%s' as a render target without the 'Load' action specified, despite a prior pass having produced it. It's invalid to completely clobber the contents of a resource. Create a new texture instance instead."),
 							PassName,
 							Texture->Name);
 					}
-					
+
 					/** Mark the pass as a producer for render targets with a store action. */
 					{
 						const bool bIsStoreAction = RenderTarget.GetStoreAction() != ERenderTargetStoreAction::ENoAction;
