Extends FRDGTextureSRVDesc to all capabilities of the RHI

diff --git a/Engine/Shaders/Private/Tools/VisualizeTexture.usf b/Engine/Shaders/Private/Tools/VisualizeTexture.usf
--- a/Engine/Shaders/Private/Tools/VisualizeTexture.usf
+++ b/Engine/Shaders/Private/Tools/VisualizeTexture.usf
@@ -17,7 +17,7 @@ TextureCube VisualizeTextureCube;
 SamplerState VisualizeTextureCubeSampler;
 
 //default Texture2D type is <float4>, which won't read the X28_G8UINT format properly.
-Texture2D<uint4> VisualizeDepthStencilTexture; 
+Texture2D<uint4> VisualizeDepthStencil; 
 
 Texture2D<uint> VisualizeUINT8Texture2D;
 
@@ -39,7 +39,7 @@ void VisualizeTexturePS(
 	float4 SvPosition : SV_POSITION,
 	out float4 OutColor : SV_Target0)
 {
-	float2 UV = SvPosition / TextureExtent.xy;
+	float2 UV = SvPosition.xy / TextureExtent.xy;
 
 	// Longitude Latitude unwrapping
 	float3 SpherePos;
@@ -137,7 +137,7 @@ void VisualizeTexturePS(
 #if TEXTURE_TYPE == 6
 		// inside #ifdef to avoid D3DDebug error:
 		//    The resource return type for component 0 declared in the shader code (UINT) is not compatible with the Shader Resource View format bound to slot 1 of the Pixel Shader unit (UNORM).  This mismatch is invalid if the shader actually uses the view (e.g. it is not skipped due to shader code branching). [ EXECUTION ERROR #361: DEVICE_DRAW_RESOURCE_RETURN_TYPE_MISMATCH]
-		float fStencil = VisualizeDepthStencilTexture.Load(int3(UV * TextureExtent.xy, 0)).g;
+		float fStencil = VisualizeDepthStencil.Load(int3(UV * TextureExtent.xy, 0)).g;
 		fStencil /= 255.0f;
 		TexLookup.a = fStencil;
 #endif
diff --git a/Engine/Source/Runtime/RHI/Public/RHI.h b/Engine/Source/Runtime/RHI/Public/RHI.h
--- a/Engine/Source/Runtime/RHI/Public/RHI.h
+++ b/Engine/Source/Runtime/RHI/Public/RHI.h
@@ -1267,8 +1267,32 @@ struct FRHITextureSRVCreateInfo
 
 	/** Specify number of array slices. If FirstArraySlice and NumArraySlices are both zero, the SRV is created for all array slices. By default 0. */
 	uint32 NumArraySlices;
+
+
+	FORCEINLINE bool operator==(const FRHITextureSRVCreateInfo& Other)const
+	{
+		return (
+			Format == Other.Format &&
+			SRGBOverride == Other.SRGBOverride &&
+			MipLevel == Other.MipLevel &&
+			NumMipLevels == Other.NumMipLevels &&
+			FirstArraySlice == Other.FirstArraySlice &&
+			NumArraySlices == Other.NumArraySlices);
+	}
+
+	FORCEINLINE bool operator!=(const FRHITextureSRVCreateInfo& Other)const
+	{
+		return !(*this == Other);
+	}
 };
 
+FORCEINLINE uint32 GetTypeHash(const FRHITextureSRVCreateInfo& Var)
+{
+	uint32 Hash0 = uint32(Var.Format) | uint32(Var.MipLevel) << 8 | uint32(Var.NumMipLevels) << 16 | uint32(Var.SRGBOverride) << 24;
+	return HashCombine(HashCombine(GetTypeHash(Hash0), GetTypeHash(Var.FirstArraySlice)), GetTypeHash(Var.NumArraySlices));
+}
+
+
 // Forward-declaration.
 struct FResolveParams;
 
diff --git a/Engine/Source/Runtime/RenderCore/Private/GenerateMips.cpp b/Engine/Source/Runtime/RenderCore/Private/GenerateMips.cpp
--- a/Engine/Source/Runtime/RenderCore/Private/GenerateMips.cpp
+++ b/Engine/Source/Runtime/RenderCore/Private/GenerateMips.cpp
@@ -62,11 +62,14 @@ FGenerateMipsStruct* FGenerateMips::SetupTexture(FRHITexture* InTexture,
 		RenderTexture.TargetableTexture = InTexture;
 		RenderTexture.ShaderResourceTexture = InTexture;
 
-		RenderTexture.MipSRVs.SetNum(Desc.NumMips);
+		RenderTexture.SRVs.Reserve(Desc.NumMips);
 		RenderTexture.MipUAVs.Reserve(Desc.NumMips);
 		for (uint8 MipLevel = 0; MipLevel < Desc.NumMips; MipLevel++)
 		{
-			RenderTexture.MipSRVs[MipLevel] = RHICreateShaderResourceView((FTexture2DRHIRef&)InTexture, MipLevel);
+			FRHITextureSRVCreateInfo SRVDesc;
+			SRVDesc.MipLevel = MipLevel;
+			RenderTexture.SRVs.Add(SRVDesc, RHICreateShaderResourceView((FTexture2DRHIRef&)InTexture, SRVDesc));
+
 			RenderTexture.MipUAVs.Add(RHICreateUnorderedAccessView(InTexture, MipLevel));
 		}
 		RHIBindDebugLabelName(RenderTexture.TargetableTexture, Desc.DebugName);
@@ -109,7 +112,7 @@ void FGenerateMips::Compute(FRHICommandListImmediate& RHIImmCmdList, FRHITexture
 		int DestTextureSizeY = InTexture->GetSizeXYZ().Y >> MipLevel;
 
 		//Create the RDG viewable SRV, of a complete Mip, to read from
-		FRDGTextureSRVDesc SRVDesc(GraphTexture, MipLevel - 1);
+		FRDGTextureSRVDesc SRVDesc = FRDGTextureSRVDesc::CreateForMipLevel(GraphTexture, MipLevel - 1);
 		//Create the RDG writeable UAV for the next mip level to be written to.
 		FRDGTextureUAVDesc UAVDesc(GraphTexture, MipLevel);
 
diff --git a/Engine/Source/Runtime/RenderCore/Private/RenderGraphBuilder.cpp b/Engine/Source/Runtime/RenderCore/Private/RenderGraphBuilder.cpp
--- a/Engine/Source/Runtime/RenderCore/Private/RenderGraphBuilder.cpp
+++ b/Engine/Source/Runtime/RenderCore/Private/RenderGraphBuilder.cpp
@@ -410,6 +410,32 @@ void FRDGBuilder::AllocateRHITextureIfNeeded(FRDGTexture* Texture)
 	check(Texture->ResourceRHI);
 }
 
+void FRDGBuilder::AllocateRHITextureSRVIfNeeded(FRDGTextureSRV* SRV)
+{
+	check(SRV);
+
+	if (SRV->ResourceRHI)
+	{
+		return;
+	}
+
+	FRDGTextureRef Texture = SRV->Desc.Texture;
+	check(Texture->PooledRenderTarget);
+
+	FSceneRenderTargetItem& RenderTarget = Texture->PooledRenderTarget->GetRenderTargetItem();
+
+	if (RenderTarget.SRVs.Contains(SRV->Desc))
+	{
+		SRV->ResourceRHI = RenderTarget.SRVs[SRV->Desc];
+		return;
+	}
+
+	FShaderResourceViewRHIRef RHIShaderResourceView = RHICreateShaderResourceView(RenderTarget.ShaderResourceTexture, SRV->Desc);
+
+	SRV->ResourceRHI = RHIShaderResourceView;
+	RenderTarget.SRVs.Add(SRV->Desc, RHIShaderResourceView);
+}
+
 void FRDGBuilder::AllocateRHITextureUAVIfNeeded(FRDGTextureUAV* UAV)
 {
 	check(UAV);
@@ -602,11 +628,7 @@ void FRDGBuilder::PrepareResourcesForExecute(const FRDGPass* Pass, struct FRHIRe
 			{
 				FRDGTextureRef Texture = SRV->Desc.Texture;
 
-				// Might be the first time using this render graph SRV, so need to setup the cached rhi resource.
-				if (!SRV->ResourceRHI)
-				{
-					SRV->ResourceRHI = Texture->PooledRenderTarget->GetRenderTargetItem().MipSRVs[SRV->Desc.MipLevel];
-				}
+				AllocateRHITextureSRVIfNeeded(SRV);
 
 				BarrierBatcher.QueueTransitionTexture(Texture, FRDGResourceState::EAccess::Read);
 			}
diff --git a/Engine/Source/Runtime/RenderCore/Private/RenderTargetPool.cpp b/Engine/Source/Runtime/RenderCore/Private/RenderTargetPool.cpp
--- a/Engine/Source/Runtime/RenderCore/Private/RenderTargetPool.cpp
+++ b/Engine/Source/Runtime/RenderCore/Private/RenderTargetPool.cpp
@@ -470,15 +470,6 @@ Done:
 					Found->RenderTargetItem.RTWriteMaskDataBufferRHI = RHICreateRTWriteMaskBuffer((FTexture2DRHIRef&)Found->RenderTargetItem.TargetableTexture);
 					Found->RenderTargetItem.RTWriteMaskBufferRHI_SRV = RHICreateShaderResourceView(Found->RenderTargetItem.RTWriteMaskDataBufferRHI);
 				}
-
-				if( Desc.NumMips > 1 )
-				{
-					Found->RenderTargetItem.MipSRVs.SetNum( Desc.NumMips );
-					for( uint16 i = 0; i < Desc.NumMips; i++ )
-					{
-						Found->RenderTargetItem.MipSRVs[i] = RHICreateShaderResourceView( (FTexture2DRHIRef&)Found->RenderTargetItem.ShaderResourceTexture, i );
-					}
-				}
 			}
 			else if(Desc.Is3DTexture())
 			{
@@ -525,15 +516,6 @@ Done:
 						(FTextureCubeRHIRef&)Found->RenderTargetItem.TargetableTexture,
 						(FTextureCubeRHIRef&)Found->RenderTargetItem.ShaderResourceTexture
 						);
-
-					if( Desc.NumMips > 1 )
-					{
-						Found->RenderTargetItem.MipSRVs.SetNum( Desc.NumMips );
-						for( uint16 i = 0; i < Desc.NumMips; i++ )
-						{
-							Found->RenderTargetItem.MipSRVs[i] = RHICreateShaderResourceView( (FTextureCubeRHIRef&)Found->RenderTargetItem.ShaderResourceTexture, i );
-						}
-					}
 				}
 
 			}
diff --git a/Engine/Source/Runtime/RenderCore/Private/VisualizeTexture.cpp b/Engine/Source/Runtime/RenderCore/Private/VisualizeTexture.cpp
--- a/Engine/Source/Runtime/RenderCore/Private/VisualizeTexture.cpp
+++ b/Engine/Source/Runtime/RenderCore/Private/VisualizeTexture.cpp
@@ -97,7 +97,7 @@ class FVisualizeTexturePS : public FGlobalShader
 		SHADER_PARAMETER_SAMPLER(SamplerState, VisualizeTextureCubeSampler)
 		SHADER_PARAMETER_RDG_TEXTURE(TextureCubeArray, VisualizeTextureCubeArray)
 		SHADER_PARAMETER_SAMPLER(SamplerState, VisualizeTextureCubeArraySampler)
-		SHADER_PARAMETER_SRV(Texture2D<uint4>, VisualizeDepthStencilTexture)
+		SHADER_PARAMETER_RDG_TEXTURE_SRV(Texture2D<uint4>, VisualizeDepthStencil)
 		SHADER_PARAMETER_RDG_TEXTURE(Texture2DMS<float4>, VisualizeTexture2DMS)
 		SHADER_PARAMETER_RDG_TEXTURE(Texture2D<uint>, VisualizeUINT8Texture2D)
 
@@ -161,7 +161,6 @@ static EVisualisePSType GetVisualizePSType(const FRDGTextureDesc& Desc)
 void FVisualizeTexture::ReleaseDynamicRHI()
 {
 	VisualizeTextureContent.SafeRelease();
-	StencilSRV.SafeRelease();
 }
 
 void FVisualizeTexture::CreateContentCapturePass(FRDGBuilder& GraphBuilder, const FRDGTextureRef SrcTexture, int32 CaptureId)
@@ -214,6 +213,8 @@ void FVisualizeTexture::CreateContentCapturePass(FRDGBuilder& GraphBuilder, cons
 	bool bSaturateInsteadOfFrac = (Flags & 1) != 0;
 	int32 InputValueMapping = bShadowDepth ? 2 : (bDepthTexture ? 1 : 0);
 
+	EVisualisePSType VisualizeType = GetVisualizePSType(SrcDesc);
+
 	FVisualizeTexturePS::FParameters* PassParameters = GraphBuilder.AllocParameters<FVisualizeTexturePS::FParameters>();
 	{
 		PassParameters->TextureExtent = FVector(SrcDesc.Extent.X, SrcDesc.Extent.Y, SrcDesc.Depth);
@@ -246,7 +247,12 @@ void FVisualizeTexture::CreateContentCapturePass(FRDGBuilder& GraphBuilder, cons
 		PassParameters->VisualizeTextureCubeArray = SrcTexture;
 		PassParameters->VisualizeTextureCubeArraySampler = PointSampler;
 
-		PassParameters->VisualizeDepthStencilTexture = nullptr; // TODO: StencilSRV
+		if (VisualizeType == EVisualisePSType::Texture2DDepthStencilNoMSAA)
+		{
+			FRDGTextureSRVDesc SRVDesc = FRDGTextureSRVDesc::CreateWithPixelFormat(SrcTexture, PF_X24_G8);
+			PassParameters->VisualizeDepthStencil = GraphBuilder.CreateSRV(SRVDesc);
+		}
+
 		PassParameters->VisualizeTexture2DMS = SrcTexture;
 		PassParameters->VisualizeUINT8Texture2D = SrcTexture;
 
@@ -255,7 +261,7 @@ void FVisualizeTexture::CreateContentCapturePass(FRDGBuilder& GraphBuilder, cons
 
 	auto ShaderMap = GetGlobalShaderMap(FeatureLevel);
 	FVisualizeTexturePS::FPermutationDomain PermutationVector;
-	PermutationVector.Set<FVisualizeTexturePS::FVisualisePSTypeDim>(GetVisualizePSType(SrcDesc));
+	PermutationVector.Set<FVisualizeTexturePS::FVisualisePSTypeDim>(VisualizeType);
 
 	TShaderMapRef<FVisualizeTexturePS> PixelShader(ShaderMap, PermutationVector);
 
@@ -282,70 +288,17 @@ void FVisualizeTexture::CreateContentCapturePass(FRDGBuilder& GraphBuilder, cons
 		}
 	}
 
-	GraphBuilder.AddPass(
+	FPixelShaderUtils::AddFullscreenPass(
+		GraphBuilder,
+		ShaderMap,
 		RDG_EVENT_NAME("VisualizeTextureCapture(%s@%d %s %dx%d%s)",
 			SrcTexture->Name, CaptureId,
 			GPixelFormats[SrcDesc.Format].Name,
 			SrcDesc.Extent.X, SrcDesc.Extent.Y,
 			*ExtendedDrawEvent),
+		*PixelShader,
 		PassParameters,
-		ERDGPassFlags::Raster,
-		[this, PassParameters, ShaderMap, PixelShader, RTExtent](FRHICommandList& RHICmdList)
-	{
-		FVisualizeTexturePS::FParameters ShaderParameter = *PassParameters;
-
-		// TODO(RDG): technically could use FPixelShaderUtils::AddPass(), but there is lot of work to support creating arbitrary number of SRV for a FRDGTexture,
-		// so the VisualizeDepthStencilTexture has to be hacked in the lambda...
-		{
-			// Some RHI might be unhappy with RHICreateShaderResourceView() inside renderpass.
-			check(RHICmdList.IsInsideRenderPass());
-			RHICmdList.EndRenderPass();
-			check(RHICmdList.IsOutsideRenderPass());
-
-			const FRDGTextureDesc& SrcDesc2 = PassParameters->VisualizeTexture2D->Desc;
-			FSceneRenderTargetItem& RenderTargetItem = PassParameters->VisualizeTexture2D->GetPooledRenderTarget()->GetRenderTargetItem();
-
-			bool bIsDefault = this->StencilSRVSrc == GBlackTexture->TextureRHI;
-			bool bDepthStencil = SrcDesc2.Is2DTexture() && SrcDesc2.Format == PF_DepthStencil;
-
-			//clear if this is a new different Stencil buffer, or it's not a stencil buffer and we haven't switched to the default yet.
-			bool bNeedsClear = bDepthStencil && (this->StencilSRVSrc != RenderTargetItem.TargetableTexture);
-			bNeedsClear |= !bDepthStencil && !bIsDefault;
-			if (bNeedsClear)
-			{
-				this->StencilSRVSrc = nullptr;
-				this->StencilSRV.SafeRelease();
-			}
-
-			//always set something into the StencilSRV slot for platforms that require a full resource binding, even if
-			//dynamic branching will cause them not to be used.	
-			if (bDepthStencil && !GVisualizeTexture.StencilSRVSrc)
-			{
-				this->StencilSRVSrc = RenderTargetItem.TargetableTexture;
-				this->StencilSRV = RHICreateShaderResourceView((FTexture2DRHIRef&)RenderTargetItem.TargetableTexture, 0, 1, PF_X24_G8);
-			}
-			else if (!GVisualizeTexture.StencilSRVSrc)
-			{
-				this->StencilSRVSrc = GBlackTexture->TextureRHI;
-				this->StencilSRV = RHICreateShaderResourceView((FTexture2DRHIRef&)GBlackTexture->TextureRHI, 0, 1, PF_B8G8R8A8);
-			}
-
-			ShaderParameter.VisualizeDepthStencilTexture = this->StencilSRV;
-
-			// Rebind the render targets.
-			FRHIRenderPassInfo RPInfo;
-			RPInfo.ColorRenderTargets[0].RenderTarget = PassParameters->RenderTargets[0].GetTexture()->GetPooledRenderTarget()->GetRenderTargetItem().TargetableTexture;
-			RPInfo.ColorRenderTargets[0].ResolveTarget = nullptr;
-			RPInfo.ColorRenderTargets[0].ArraySlice = -1;
-			RPInfo.ColorRenderTargets[0].MipIndex = 0;
-			RPInfo.ColorRenderTargets[0].Action = MakeRenderTargetActions(ERenderTargetLoadAction::EClear, ERenderTargetStoreAction::EStore);
-
-			RHICmdList.BeginRenderPass(RPInfo, TEXT("VisualizeTextureCapture"));
-			check(RHICmdList.IsInsideRenderPass());
-		}
-
-		FPixelShaderUtils::DrawFullscreenPixelShader(RHICmdList, ShaderMap, *PixelShader, ShaderParameter, FIntRect(0, 0, RTExtent.X, RTExtent.Y));
-	});
+		FIntRect(0, 0, RTExtent.X, RTExtent.Y));
 
 	// Save the copied texture and descriptor about original informations.
 	{
diff --git a/Engine/Source/Runtime/RenderCore/Public/RenderGraphBuilder.h b/Engine/Source/Runtime/RenderCore/Public/RenderGraphBuilder.h
--- a/Engine/Source/Runtime/RenderCore/Public/RenderGraphBuilder.h
+++ b/Engine/Source/Runtime/RenderCore/Public/RenderGraphBuilder.h
@@ -196,6 +196,7 @@ private:
 	void AddPassInternal(FRDGPass* Pass);
 
 	void AllocateRHITextureIfNeeded(FRDGTexture* Texture);
+	void AllocateRHITextureSRVIfNeeded(FRDGTextureSRV* SRV);
 	void AllocateRHITextureUAVIfNeeded(FRDGTextureUAV* UAV);
 	void AllocateRHIBufferIfNeeded(FRDGBuffer* Buffer);
 	void AllocateRHIBufferSRVIfNeeded(FRDGBufferSRV* SRV);
diff --git a/Engine/Source/Runtime/RenderCore/Public/RenderGraphBuilder.inl b/Engine/Source/Runtime/RenderCore/Public/RenderGraphBuilder.inl
--- a/Engine/Source/Runtime/RenderCore/Public/RenderGraphBuilder.inl
+++ b/Engine/Source/Runtime/RenderCore/Public/RenderGraphBuilder.inl
@@ -115,7 +115,32 @@ inline FRDGTextureSRVRef FRDGBuilder::CreateSRV(const FRDGTextureSRVDesc& Desc)
 		checkf(Texture, TEXT("RenderGraph texture SRV created with a null texture."));
 		Validation.ExecuteGuard(TEXT("CreateSRV"), Texture->Name);
 		checkf(Texture->Desc.TargetableFlags & TexCreate_ShaderResource, TEXT("Attempted to create SRV from texture %s which was not created with TexCreate_ShaderResource"), Desc.Texture->Name);
-		checkf(Desc.MipLevel < Texture->Desc.NumMips, TEXT("Failed to create SRV at mip %d: the texture %s has only %d mip levels."), Desc.MipLevel, Texture->Name, Texture->Desc.NumMips);
+		
+		// Validate the pixel format if overridden by the SRV's descriptor.
+		if (Desc.Format == PF_X24_G8)
+		{
+			// PF_X24_G8 is a bit of mess in the RHI, used to read the stencil, but have varying BlockBytes.
+			checkf(Texture->Desc.Format == PF_DepthStencil, TEXT("PF_X24_G8 is only to read stencil from a PF_DepthStencil texture"));
+		}
+		else if (Desc.Format != PF_Unknown)
+		{
+			checkf(Desc.Format < PF_MAX, TEXT("Illegal to create SRV for texture %s with invalid FPooledRenderTargetDesc::Format."), Texture->Name);
+			checkf(GPixelFormats[Desc.Format].Supported, TEXT("Failed to create SRV for texture %s with pixel format %s because it is not supported."),
+				Texture->Name, GPixelFormats[Desc.Format].Name);
+
+			EPixelFormat ResourcePixelFormat = Texture->Desc.Format;
+
+			checkf(
+				GPixelFormats[Desc.Format].BlockBytes == GPixelFormats[ResourcePixelFormat].BlockBytes &&
+				GPixelFormats[Desc.Format].BlockSizeX == GPixelFormats[ResourcePixelFormat].BlockSizeX &&
+				GPixelFormats[Desc.Format].BlockSizeY == GPixelFormats[ResourcePixelFormat].BlockSizeY &&
+				GPixelFormats[Desc.Format].BlockSizeZ == GPixelFormats[ResourcePixelFormat].BlockSizeZ,
+				TEXT("Failed to create SRV for texture %s with pixel format %s because it does not match the byte size of the texture's pixel format %s."),
+				Texture->Name, GPixelFormats[Desc.Format].Name, GPixelFormats[ResourcePixelFormat].Name);
+		}
+
+		checkf((Desc.MipLevel + Desc.NumMipLevels) <= Texture->Desc.NumMips, TEXT("Failed to create SRV at mips %d-%d: the texture %s has only %d mip levels."),
+			Desc.MipLevel, (Desc.MipLevel + Desc.NumMipLevels), Texture->Name, Texture->Desc.NumMips);
 	}
 #endif
 
diff --git a/Engine/Source/Runtime/RenderCore/Public/RenderGraphResources.h b/Engine/Source/Runtime/RenderCore/Public/RenderGraphResources.h
--- a/Engine/Source/Runtime/RenderCore/Public/RenderGraphResources.h
+++ b/Engine/Source/Runtime/RenderCore/Public/RenderGraphResources.h
@@ -322,18 +322,47 @@ protected:
 };
 
 /** Descriptor for render graph tracked SRV. */
-class RENDERCORE_API FRDGTextureSRVDesc final
+class RENDERCORE_API FRDGTextureSRVDesc final : public FRHITextureSRVCreateInfo
 {
 public:
 	FRDGTextureSRVDesc() = default;
+	
+	FRDGTextureRef Texture = nullptr;
 
-	FRDGTextureSRVDesc(FRDGTextureRef InTexture, uint8 InMipLevel)
-		: Texture(InTexture)
-		, MipLevel(InMipLevel)
-	{}
+	/** Create SRV that access all sub resources of texture. */
+	static FRDGTextureSRVDesc Create(FRDGTextureRef Texture)
+	{
+		FRDGTextureSRVDesc Desc;
+		Desc.Texture = Texture;
+		Desc.NumMipLevels = Texture->Desc.NumMips;
+		if (Texture->Desc.bIsArray)
+		{
+			Desc.NumArraySlices = Texture->Desc.Depth;
+		}
+		return Desc;
+	}
 
-	FRDGTextureRef Texture = nullptr;
-	uint8 MipLevel = 0;
+	/** Create SRV that access one specific mip level. */
+	static FRDGTextureSRVDesc CreateForMipLevel(FRDGTextureRef Texture, int32 MipLevel)
+	{
+		FRDGTextureSRVDesc Desc;
+		Desc.Texture = Texture;
+		Desc.MipLevel = MipLevel;
+		Desc.NumMipLevels = 1;
+		if (Texture->Desc.bIsArray)
+		{
+			Desc.NumArraySlices = Texture->Desc.Depth;
+		}
+		return Desc;
+	}
+
+	/** Create SRV that access one specific mip level. */
+	static FRDGTextureSRVDesc CreateWithPixelFormat(FRDGTextureRef Texture, EPixelFormat PixelFormat)
+	{
+		FRDGTextureSRVDesc Desc = FRDGTextureSRVDesc::Create(Texture);
+		Desc.Format = PixelFormat;
+		return Desc;
+	}
 };
 
 /** Render graph tracked SRV. */
diff --git a/Engine/Source/Runtime/RenderCore/Public/RendererInterface.h b/Engine/Source/Runtime/RenderCore/Public/RendererInterface.h
--- a/Engine/Source/Runtime/RenderCore/Public/RendererInterface.h
+++ b/Engine/Source/Runtime/RenderCore/Public/RendererInterface.h
@@ -387,20 +387,13 @@ struct FSceneRenderTargetItem
 		,	UAV(InUAV)
 	{}
 
-	/** */
 	void SafeRelease()
 	{
 		TargetableTexture.SafeRelease();
 		ShaderResourceTexture.SafeRelease();
 		UAV.SafeRelease();
-		for (int32 i = 0; i < MipUAVs.Num(); i++)
-		{
-			MipUAVs[i].SafeRelease();
-		}
-		for( int32 i = 0; i < MipSRVs.Num(); i++ )
-		{
-			MipSRVs[i].SafeRelease();
-		}
+		MipUAVs.Empty();
+		SRVs.Empty();
 	}
 
 	bool IsValid() const
@@ -419,8 +412,8 @@ struct FSceneRenderTargetItem
 	FUnorderedAccessViewRHIRef UAV;
 	/** only created if requested through the flag  */
 	TArray< FUnorderedAccessViewRHIRef, TInlineAllocator<1> > MipUAVs;
-	/** only created if requested through the flag  */
-	TArray< FShaderResourceViewRHIRef > MipSRVs;
+	/** All SRVs that has been created on for that ShaderResourceTexture.  */
+	TMap<FRHITextureSRVCreateInfo, FShaderResourceViewRHIRef> SRVs;
 
 	FShaderResourceViewRHIRef RTWriteMaskBufferRHI_SRV;
 	FStructuredBufferRHIRef RTWriteMaskDataBufferRHI;
diff --git a/Engine/Source/Runtime/RenderCore/Public/VisualizeTexture.h b/Engine/Source/Runtime/RenderCore/Public/VisualizeTexture.h
--- a/Engine/Source/Runtime/RenderCore/Public/VisualizeTexture.h
+++ b/Engine/Source/Runtime/RenderCore/Public/VisualizeTexture.h
@@ -94,9 +94,6 @@ private:
 	/** Descriptor of the texture being visualized. */
 	FPooledRenderTargetDesc VisualizeTextureDesc;
 
-	TRefCountPtr<FRHIShaderResourceView> StencilSRV;
-	FTextureRHIRef StencilSRVSrc;
-
 	// Flag to determine whether texture visualization is enabled, currently based on the feature level we are rendering with
 	bool bEnabled;
 
diff --git a/Engine/Source/Runtime/Renderer/Private/ReflectionEnvironmentCapture.cpp b/Engine/Source/Runtime/Renderer/Private/ReflectionEnvironmentCapture.cpp
--- a/Engine/Source/Runtime/Renderer/Private/ReflectionEnvironmentCapture.cpp
+++ b/Engine/Source/Runtime/Renderer/Private/ReflectionEnvironmentCapture.cpp
@@ -185,6 +185,17 @@ void CreateCubeMips( FRHICommandListImmediate& RHICmdList, ERHIFeatureLevel::Typ
 
 	auto* ShaderMap = GetGlobalShaderMap(FeatureLevel);
 
+	for (int32 MipIndex = 0; MipIndex < NumMips; MipIndex++)
+	{
+		FRHITextureSRVCreateInfo SRVDesc;
+		SRVDesc.MipLevel = MipIndex;
+
+		if (!Cubemap.SRVs.Contains(SRVDesc))
+		{
+			Cubemap.SRVs.Add(SRVDesc, RHICreateShaderResourceView(Cubemap.ShaderResourceTexture, SRVDesc));
+		}
+	}
+
 	FGraphicsPipelineStateInitializer GraphicsPSOInit;
 	GraphicsPSOInit.RasterizerState = TStaticRasterizerState<FM_Solid, CM_None>::GetRHI();
 	GraphicsPSOInit.DepthStencilState = TStaticDepthStencilState<false, CF_Always>::GetRHI();
@@ -224,7 +235,10 @@ void CreateCubeMips( FRHICommandListImmediate& RHICmdList, ERHIFeatureLevel::Typ
 
 				SetShaderValue(RHICmdList, ShaderRHI, PixelShader->NumMips, NumMips);
 
-				SetSRVParameter(RHICmdList, ShaderRHI, PixelShader->SourceTexture, Cubemap.MipSRVs[MipIndex - 1]);
+				FRHITextureSRVCreateInfo SrcSRVDesc;
+				SrcSRVDesc.MipLevel = MipIndex - 1;
+
+				SetSRVParameter(RHICmdList, ShaderRHI, PixelShader->SourceTexture, Cubemap.SRVs[SrcSRVDesc]);
 				SetSamplerParameter(RHICmdList, ShaderRHI, PixelShader->SourceTextureSampler, TStaticSamplerState<SF_Bilinear, AM_Clamp, AM_Clamp, AM_Clamp>::GetRHI());
 			}
 
diff --git a/Engine/Source/Runtime/Renderer/Private/SceneTextureReductions.cpp b/Engine/Source/Runtime/Renderer/Private/SceneTextureReductions.cpp
--- a/Engine/Source/Runtime/Renderer/Private/SceneTextureReductions.cpp
+++ b/Engine/Source/Runtime/Renderer/Private/SceneTextureReductions.cpp
@@ -151,7 +151,7 @@ void BuildHZB(FRDGBuilder& GraphBuilder, FViewInfo& View)
 		DstSize.Y = FMath::Max(DstSize.Y, 1);
 
 		//@DW: Explicit creation of SRV, full configuration of SRV supported
-		FRDGTextureSRVDesc Desc(HZBTexture, MipIndex - 1);
+		FRDGTextureSRVDesc Desc = FRDGTextureSRVDesc::CreateForMipLevel(HZBTexture, MipIndex - 1);
 		FRDGTextureSRVRef ParentMipSRV = GraphBuilder.CreateSRV(Desc);
 
 		FHZBBuildPassParameters* PassParameters = GraphBuilder.AllocParameters<FHZBBuildPassParameters>();
