Validates a render targets has TexCreate_RenderTargetable flag when getting bound.

diff --git a/Engine/Source/Runtime/RenderCore/Private/RenderGraphValidation.cpp b/Engine/Source/Runtime/RenderCore/Private/RenderGraphValidation.cpp
--- a/Engine/Source/Runtime/RenderCore/Private/RenderGraphValidation.cpp
+++ b/Engine/Source/Runtime/RenderCore/Private/RenderGraphValidation.cpp
@@ -437,6 +437,11 @@ void FRDGUserValidation::ValidateAddPass(const FRDGPass* Pass)
 					TEXT("Pass '%s' is marked to generate mips but has a depth stencil texture. This is not supported."),
 					Pass->GetName());
 
+				checkf(
+					Texture->Desc.TargetableFlags & TexCreate_RenderTargetable,
+					TEXT("Pass '%s' attempted to bind texture '%s' as a depth stencil render target, but the texture has not been created with TexCreate_RenderTargetable."),
+					PassName, Texture->Name);
+
 				// Depth stencil only supports one mip, since there isn't actually a way to select the mip level.
 				const uint32 MipLevel = 0;
 				check(Texture->Desc.NumMips == 1);
@@ -471,6 +476,11 @@ void FRDGUserValidation::ValidateAddPass(const FRDGPass* Pass)
 
 					const bool bIsLoadAction = RenderTarget.GetLoadAction() == ERenderTargetLoadAction::ELoad;
 
+					checkf(
+						Texture->Desc.TargetableFlags & TexCreate_RenderTargetable,
+						TEXT("Pass '%s' attempted to bind texture '%s' as a render target, but the texture has not been created with TexCreate_RenderTargetable."),
+						PassName, Texture->Name);
+
 					/** Validate that load action is correct. We can only load contents if a pass previously produced something. */
 					{
 						const bool bIsLoadActionInvalid = bIsLoadAction && !Texture->HasBeenProduced();
