Catches unused TexCreate_UAV and TexCreate_RenderTargetable flags on FRDGTexture with -rdgdebug

diff --git a/Engine/Source/Runtime/RenderCore/Private/RenderGraphValidation.cpp b/Engine/Source/Runtime/RenderCore/Private/RenderGraphValidation.cpp
--- a/Engine/Source/Runtime/RenderCore/Private/RenderGraphValidation.cpp
+++ b/Engine/Source/Runtime/RenderCore/Private/RenderGraphValidation.cpp
@@ -184,16 +184,34 @@ void FRDGUserValidation::ExecuteGuard(const TCHAR* Operation, const TCHAR* Resou
 	checkf(!bHasExecuted, TEXT("Render graph operation '%s' with resource '%s' must be performed prior to graph execution."), Operation, ResourceName);
 }
 
-void FRDGUserValidation::ValidateCreateResource(FRDGTrackedResourceRef Resource)
+void FRDGUserValidation::ValidateCreateTexture(FRDGTextureRef Texture)
 {
-	check(Resource);
-	TrackedResources.Add(Resource);
+	check(Texture);
+	if (IsRDGDebugEnabled())
+	{
+		TrackedTextures.Add(Texture);
+	}
 }
 
-void FRDGUserValidation::ValidateCreateExternalResource(FRDGTrackedResourceRef Resource)
+void FRDGUserValidation::ValidateCreateBuffer(FRDGBufferRef Buffer)
 {
-	ValidateCreateResource(Resource);
-	Resource->MarkAsExternal();
+	check(Buffer);
+	if (IsRDGDebugEnabled())
+	{
+		TrackedBuffers.Add(Buffer);
+	}
+}
+
+void FRDGUserValidation::ValidateCreateExternalTexture(FRDGTextureRef Texture)
+{
+	ValidateCreateTexture(Texture);
+	Texture->MarkAsExternal();
+}
+
+void FRDGUserValidation::ValidateCreateExternalBuffer(FRDGBufferRef Buffer)
+{
+	ValidateCreateBuffer(Buffer);
+	Buffer->MarkAsExternal();
 }
 
 void FRDGUserValidation::ValidateExtractResource(FRDGTrackedResourceRef Resource)
@@ -556,30 +574,105 @@ void FRDGUserValidation::ValidateExecuteEnd()
 {
 	bHasExecuted = true;
 
-	const bool bEmitWarnings = IsRDGDebugEnabled();
-
-	for (const FRDGTrackedResourceRef Resource : TrackedResources)
+	if (IsRDGDebugEnabled())
 	{
-		check(Resource->ReferenceCount == 0);
+		auto ValidateResourceAtExecuteEnd = [](const FRDGTrackedResourceRef Resource)
+		{
+			check(Resource->ReferenceCount == 0);
+
+			const bool bProducedButNeverUsed = Resource->PassAccessCount == 1 && Resource->FirstProducer;
+
+			if (bProducedButNeverUsed)
+			{
+				check(Resource->HasBeenProduced());
 
-		const bool bProducedButNeverUsed = Resource->PassAccessCount == 1 && Resource->FirstProducer;
+				EmitRDGWarningf(
+					TEXT("Resource %s has been produced by the pass %s, but never used by another pass."),
+					Resource->Name, Resource->FirstProducer->GetName());
+			}
+		};
 
-		if (bEmitWarnings && bProducedButNeverUsed)
+		for (const FRDGTextureRef Texture : TrackedTextures)
 		{
-			check(Resource->HasBeenProduced());
+			ValidateResourceAtExecuteEnd(Texture);
+
+			bool bHasBeenProducedByGraph = !Texture->IsExternal() && Texture->PassAccessCount > 0;
+
+			if (bHasBeenProducedByGraph && !Texture->bHasNeededUAV && (Texture->Desc.TargetableFlags & TexCreate_UAV))
+			{
+				EmitRDGWarningf(
+					TEXT("Resource %s first produced by the pass %s had the TexCreate_UAV flag, but no UAV has been used."),
+					Texture->Name, Texture->FirstProducer->GetName());
+			}
+
+			if (bHasBeenProducedByGraph && !Texture->bHasBeenBoundAsRenderTarget && (Texture->Desc.TargetableFlags & TexCreate_RenderTargetable))
+			{
+				EmitRDGWarningf(
+					TEXT("Resource %s first produced by the pass %s had the TexCreate_RenderTargetable flag, but has never been bound as a render target of a pass."),
+					Texture->Name, Texture->FirstProducer->GetName());
+			}
+		}
 
-			EmitRDGWarningf(
-				TEXT("Resources %s has been produced by the pass %s, but never used by another pass."),
-				Resource->Name, Resource->FirstProducer->GetName());
+		for (const FRDGBufferRef Buffer : TrackedBuffers)
+		{
+			ValidateResourceAtExecuteEnd(Buffer);
 		}
-	}
+	} // if (IsRDGDebugEnabled())
 
-	TrackedResources.Empty();
+	TrackedTextures.Empty();
+	TrackedBuffers.Empty();
 }
 
 void FRDGUserValidation::ValidateExecutePassBegin(const FRDGPass* Pass)
 {
 	SetAllowRHIAccess(Pass, true);
+
+	if (IsRDGDebugEnabled())
+	{
+		FRDGPassParameterStruct ParameterStruct = Pass->GetParameters();
+
+		const uint32 ParameterCount = ParameterStruct.GetParameterCount();
+
+		for (uint32 ParameterIndex = 0; ParameterIndex < ParameterCount; ++ParameterIndex)
+		{
+			FRDGPassParameter Parameter = ParameterStruct.GetParameter(ParameterIndex);
+
+			if (Parameter.GetType() == UBMT_RDG_TEXTURE_UAV)
+			{
+				if (FRDGTextureUAVRef UAV = Parameter.GetAsTextureUAV())
+				{
+					FRDGTextureRef Texture = UAV->Desc.Texture;
+					Texture->bHasNeededUAV = true;
+				}
+			}
+			else if (Parameter.GetType() == UBMT_RENDER_TARGET_BINDING_SLOTS)
+			{
+				const FRenderTargetBindingSlots& RenderTargetBindingSlots = Parameter.GetAsRenderTargetBindingSlots();
+				const auto& RenderTargets = RenderTargetBindingSlots.Output;
+				const auto& DepthStencil = RenderTargetBindingSlots.DepthStencil;
+				const uint32 RenderTargetCount = RenderTargets.Num();
+
+				for (uint32 RenderTargetIndex = 0; RenderTargetIndex < RenderTargetCount; RenderTargetIndex++)
+				{
+					const FRenderTargetBinding& RenderTarget = RenderTargets[RenderTargetIndex];
+
+					if (FRDGTextureRef Texture = RenderTarget.GetTexture())
+					{
+						Texture->bHasBeenBoundAsRenderTarget = true;
+					}
+					else
+					{
+						break;
+					}
+				}
+
+				if (FRDGTextureRef Texture = DepthStencil.GetTexture())
+				{
+					Texture->bHasBeenBoundAsRenderTarget = true;
+				}
+			}
+		} // for (uint32 ParameterIndex = 0; ParameterIndex < ParameterCount; ++ParameterIndex)
+	} // if (IsRDGDebugEnabled())
 }
 
 void FRDGUserValidation::ValidateExecutePassEnd(const FRDGPass* Pass)
diff --git a/Engine/Source/Runtime/RenderCore/Public/RenderGraphBuilder.inl b/Engine/Source/Runtime/RenderCore/Public/RenderGraphBuilder.inl
--- a/Engine/Source/Runtime/RenderCore/Public/RenderGraphBuilder.inl
+++ b/Engine/Source/Runtime/RenderCore/Public/RenderGraphBuilder.inl
@@ -23,7 +23,7 @@ inline FRDGTextureRef FRDGBuilder::RegisterExternalTexture(const TRefCountPtr<IP
 	OutTexture->ResourceRHI = ExternalPooledTexture->GetRenderTargetItem().ShaderResourceTexture;
 	AllocatedTextures.Add(OutTexture, ExternalPooledTexture);
 
-	IF_RDG_ENABLE_DEBUG(Validation.ValidateCreateExternalResource(OutTexture));
+	IF_RDG_ENABLE_DEBUG(Validation.ValidateCreateExternalTexture(OutTexture));
 
 	ExternalTextures.Add(ExternalPooledTexture.GetReference(), OutTexture);
 
@@ -48,7 +48,7 @@ inline FRDGBufferRef FRDGBuilder::RegisterExternalBuffer(const TRefCountPtr<FPoo
 	OutBuffer->PooledBuffer = ExternalPooledBuffer;
 	AllocatedBuffers.Add(OutBuffer, ExternalPooledBuffer);
 
-	IF_RDG_ENABLE_DEBUG(Validation.ValidateCreateExternalResource(OutBuffer));
+	IF_RDG_ENABLE_DEBUG(Validation.ValidateCreateExternalBuffer(OutBuffer));
 
 	ExternalBuffers.Add(ExternalPooledBuffer.GetReference(), OutBuffer);
 
@@ -82,7 +82,7 @@ inline FRDGTextureRef FRDGBuilder::CreateTexture(
 
 	FRDGTexture* Texture = AllocateForRHILifeTime<FRDGTexture>(Name, Desc, Flags);
 
-	IF_RDG_ENABLE_DEBUG(Validation.ValidateCreateResource(Texture));
+	IF_RDG_ENABLE_DEBUG(Validation.ValidateCreateTexture(Texture));
 
 	return Texture;
 }
@@ -101,7 +101,7 @@ inline FRDGBufferRef FRDGBuilder::CreateBuffer(
 
 	FRDGBufferRef Buffer = AllocateForRHILifeTime<FRDGBuffer>(Name, Desc, Flags);
 
-	IF_RDG_ENABLE_DEBUG(Validation.ValidateCreateResource(Buffer));
+	IF_RDG_ENABLE_DEBUG(Validation.ValidateCreateBuffer(Buffer));
 
 	return Buffer;
 }
diff --git a/Engine/Source/Runtime/RenderCore/Public/RenderGraphResources.h b/Engine/Source/Runtime/RenderCore/Public/RenderGraphResources.h
--- a/Engine/Source/Runtime/RenderCore/Public/RenderGraphResources.h
+++ b/Engine/Source/Runtime/RenderCore/Public/RenderGraphResources.h
@@ -270,6 +270,14 @@ private:
 	 */
 	IPooledRenderTarget* PooledRenderTarget = nullptr;
 
+#if RDG_ENABLE_DEBUG
+	/** Tracks whether a UAV has ever been allocated to catch when TexCreate_UAV was unneeded. */
+	bool bHasNeededUAV = false;
+
+	/** Tracks whether has ever been bound as a render target to catch when TexCreate_RenderTargetable was unneeded. */
+	bool bHasBeenBoundAsRenderTarget = false;
+#endif 
+
 	friend class FRDGBuilder;
 	friend class FRDGUserValidation;
 	friend class FRDGBarrierBatcher;
diff --git a/Engine/Source/Runtime/RenderCore/Public/RenderGraphValidation.h b/Engine/Source/Runtime/RenderCore/Public/RenderGraphValidation.h
--- a/Engine/Source/Runtime/RenderCore/Public/RenderGraphValidation.h
+++ b/Engine/Source/Runtime/RenderCore/Public/RenderGraphValidation.h
@@ -43,10 +43,12 @@ public:
 	void ExecuteGuard(const TCHAR* Operation, const TCHAR* ResourceName);
 
 	/** Tracks and validates the creation of a new resource in the graph. */
-	void ValidateCreateResource(FRDGTrackedResourceRef Resource);
+	void ValidateCreateTexture(FRDGTextureRef Texture);
+	void ValidateCreateBuffer(FRDGBufferRef Buffer);
 
 	/** Tracks and validates the creation of a new externally registered resource. */
-	void ValidateCreateExternalResource(FRDGTrackedResourceRef Resource);
+	void ValidateCreateExternalTexture(FRDGTextureRef Texture);
+	void ValidateCreateExternalBuffer(FRDGBufferRef Buffer);
 
 	/** Tracks and validates usage of a pass parameters allocation. */
 	void ValidateAllocPassParameters(const void* Parameters);
@@ -76,7 +78,8 @@ private:
 	TSet<const void*, DefaultKeyFuncs<const void*>, SceneRenderingSetAllocator> AllocatedUnusedPassParameters;
 
 	/** List of tracked resources for validation prior to shutdown. */
-	TArray<FRDGTrackedResourceRef, SceneRenderingAllocator> TrackedResources;
+	TArray<FRDGTextureRef, SceneRenderingAllocator> TrackedTextures;
+	TArray<FRDGBufferRef, SceneRenderingAllocator> TrackedBuffers;
 
 	/** Whether the Execute() has already been called. */
 	bool bHasExecuted = false;
