Removes the store action on the render target binding that is just redondant.

diff --git a/Engine/Source/Runtime/RenderCore/Private/RenderGraphBuilder.cpp b/Engine/Source/Runtime/RenderCore/Private/RenderGraphBuilder.cpp
--- a/Engine/Source/Runtime/RenderCore/Private/RenderGraphBuilder.cpp
+++ b/Engine/Source/Runtime/RenderCore/Private/RenderGraphBuilder.cpp
@@ -217,7 +217,7 @@ void FRDGBuilder::VisualizePassOutputs(const FRDGPass* Pass)
 
 			if (FRDGTextureRef Texture = DepthStencil.GetTexture())
 			{
-				const bool bHasStoreAction = DepthStencil.GetDepthStoreAction() != ERenderTargetStoreAction::ENoAction || DepthStencil.GetStencilStoreAction() != ERenderTargetStoreAction::ENoAction;
+				const bool bHasStoreAction = DepthStencil.GetDepthStencilAccess().IsAnyWrite();
 
 				if (bHasStoreAction)
 				{
@@ -240,16 +240,11 @@ void FRDGBuilder::VisualizePassOutputs(const FRDGPass* Pass)
 
 				if (FRDGTextureRef Texture = RenderTarget.GetTexture())
 				{
-					const bool bHasStoreAction = RenderTarget.GetStoreAction() != ERenderTargetStoreAction::ENoAction;
-
-					if (bHasStoreAction)
-					{
-						int32 CaptureId = GVisualizeTexture.ShouldCapture(Texture->Name);
+					int32 CaptureId = GVisualizeTexture.ShouldCapture(Texture->Name);
 						
-						if (CaptureId != FVisualizeTexture::kInvalidCaptureId && RenderTarget.GetMipIndex() == GVisualizeTexture.CustomMip)
-						{
-							GVisualizeTexture.CreateContentCapturePass(*this, Texture, CaptureId);
-						}
+					if (CaptureId != FVisualizeTexture::kInvalidCaptureId && RenderTarget.GetMipIndex() == GVisualizeTexture.CustomMip)
+					{
+						GVisualizeTexture.CreateContentCapturePass(*this, Texture, CaptureId);
 					}
 				}
 				else
@@ -750,13 +745,9 @@ void FRDGBuilder::PrepareResourcesForExecute(const FRDGPass* Pass, struct FRHIRe
 					FRHITexture* TargetableTexture = Texture->PooledRenderTarget->GetRenderTargetItem().TargetableTexture;
 					FRHITexture* ShaderResourceTexture = Texture->PooledRenderTarget->GetRenderTargetItem().ShaderResourceTexture;
 
-					// TODO(RDG): Looks like the store action on FRenderTargetBinding is not necessary, because: if want to bind a RT,
-					// that is most certainly to modify it as oposed to depth-stencil that might be for read only purposes. And if modify
-					// this resource, that certainly for being used by another pass. Otherwise this pass should be culled.
-					//
 					// TODO(RDG): The load store action could actually be optimised by render graph for tile hardware when there is multiple
 					// consecutive rasterizer passes that have RDG resource as render target, a bit like resource transitions.
-					ERenderTargetStoreAction StoreAction = RenderTarget.GetStoreAction();
+					ERenderTargetStoreAction StoreAction = ERenderTargetStoreAction::EStore;
 
 					// Automatically switch the store action to MSAA resolve when there is MSAA texture.
 					if (TargetableTexture != ShaderResourceTexture && Texture->Desc.NumSamples > 1 && StoreAction == ERenderTargetStoreAction::EStore)
@@ -790,17 +781,22 @@ void FRDGBuilder::PrepareResourcesForExecute(const FRDGPass* Pass, struct FRHIRe
 
 				auto& OutDepthStencil = OutRPInfo->DepthStencilRenderTarget;
 
+				FExclusiveDepthStencil ExclusiveDepthStencil = DepthStencil.GetDepthStencilAccess();
+
+				ERenderTargetStoreAction DepthStoreAction = ExclusiveDepthStencil.IsDepthWrite() ? ERenderTargetStoreAction::EStore : ERenderTargetStoreAction::ENoAction;
+				ERenderTargetStoreAction StencilStoreAction = ExclusiveDepthStencil.IsStencilWrite() ? ERenderTargetStoreAction::EStore : ERenderTargetStoreAction::ENoAction;
+
 				// TODO(RDG): Addresses the TODO of the color scene render target.
 				ensureMsgf(Texture->Desc.NumSamples == 1, TEXT("MSAA dept-stencil render target not yet supported."));
 				OutDepthStencil.DepthStencilTarget = Texture->PooledRenderTarget->GetRenderTargetItem().TargetableTexture;
 				OutDepthStencil.ResolveTarget = nullptr;
 				OutDepthStencil.Action = MakeDepthStencilTargetActions(
-					MakeRenderTargetActions(DepthStencil.GetDepthLoadAction(), DepthStencil.GetDepthStoreAction()),
-					MakeRenderTargetActions(DepthStencil.GetStencilLoadAction(), DepthStencil.GetStencilStoreAction()));
-				OutDepthStencil.ExclusiveDepthStencil = DepthStencil.GetDepthStencilAccess();
+					MakeRenderTargetActions(DepthStencil.GetDepthLoadAction(), DepthStoreAction),
+					MakeRenderTargetActions(DepthStencil.GetStencilLoadAction(), StencilStoreAction));
+				OutDepthStencil.ExclusiveDepthStencil = ExclusiveDepthStencil;
 
 				BarrierBatcher.QueueTransitionTexture(Texture,
-					DepthStencil.GetDepthStencilAccess().IsAnyWrite() ?
+					ExclusiveDepthStencil.IsAnyWrite() ?
 					FRDGResourceState::EAccess::Write :
 					FRDGResourceState::EAccess::Read);
 
diff --git a/Engine/Source/Runtime/RenderCore/Private/RenderGraphValidation.cpp b/Engine/Source/Runtime/RenderCore/Private/RenderGraphValidation.cpp
--- a/Engine/Source/Runtime/RenderCore/Private/RenderGraphValidation.cpp
+++ b/Engine/Source/Runtime/RenderCore/Private/RenderGraphValidation.cpp
@@ -529,11 +529,7 @@ void FRDGUserValidation::ValidateAddPass(const FRDGPass* Pass)
 					}
 
 					/** Mark the pass as a producer for render targets with a store action. */
-					{
-						const bool bIsStoreAction = RenderTarget.GetStoreAction() != ERenderTargetStoreAction::ENoAction;
-						check(bIsStoreAction); // already been validated in FRenderTargetBinding::Validate()
-						Texture->MarkAsProducedBy(Pass);
-					}
+					Texture->MarkAsProducedBy(Pass);
 
 					Texture->PassAccessCount++;
 				}
diff --git a/Engine/Source/Runtime/RenderCore/Private/ShaderParameterStruct.cpp b/Engine/Source/Runtime/RenderCore/Private/ShaderParameterStruct.cpp
--- a/Engine/Source/Runtime/RenderCore/Private/ShaderParameterStruct.cpp
+++ b/Engine/Source/Runtime/RenderCore/Private/ShaderParameterStruct.cpp
@@ -343,16 +343,10 @@ void FShaderParameterBindings::BindForRootShaderParameters(const FShader* Shader
 
 bool FRenderTargetBinding::Validate() const
 {
-	if (Texture)
-	{
-		checkf(StoreAction != ERenderTargetStoreAction::ENoAction,
-			TEXT("You must specify a store action for non-null render target %s."),
-			Texture->Name);
-	}
-	else
+	if (!Texture)
 	{
-		checkf(LoadAction == ERenderTargetLoadAction::ENoAction && StoreAction == ERenderTargetStoreAction::ENoAction,
-			TEXT("Can't have a load or store action when no texture is bound."));
+		checkf(LoadAction == ERenderTargetLoadAction::ENoAction,
+			TEXT("Can't have a load action when no texture is bound."));
 	}
 	
 	return true;
@@ -377,7 +371,7 @@ bool FDepthStencilBinding::Validate() const
 		bool bHasStencil = PixelFormat == PF_DepthStencil;
 		if (!bHasStencil)
 		{
-			checkf(StencilLoadAction == ERenderTargetLoadAction::ENoAction && StencilStoreAction == ERenderTargetStoreAction::ENoAction,
+			checkf(StencilLoadAction == ERenderTargetLoadAction::ENoAction,
 				TEXT("Unable to load stencil of texture %s that have a pixel format %s that does not support stencil."),
 				Texture->Name, FormatString);
 		
@@ -388,10 +382,10 @@ bool FDepthStencilBinding::Validate() const
 	}
 	else
 	{
-		checkf(DepthLoadAction == ERenderTargetLoadAction::ENoAction && DepthStoreAction == ERenderTargetStoreAction::ENoAction,
-			TEXT("Can't have a depth load or store action when no texture are bound."));
-		checkf(StencilLoadAction == ERenderTargetLoadAction::ENoAction && StencilStoreAction == ERenderTargetStoreAction::ENoAction,
-			TEXT("Can't have a stencil load or store action when no texture are bound."));
+		checkf(DepthLoadAction == ERenderTargetLoadAction::ENoAction,
+			TEXT("Can't have a depth load action when no texture are bound."));
+		checkf(StencilLoadAction == ERenderTargetLoadAction::ENoAction,
+			TEXT("Can't have a stencil load action when no texture are bound."));
 		checkf(DepthStencilAccess == FExclusiveDepthStencil::DepthNop_StencilNop,
 			TEXT("Can't have a depth stencil access when no texture are bound."));
 	}
diff --git a/Engine/Source/Runtime/RenderCore/Private/VisualizeTexture.cpp b/Engine/Source/Runtime/RenderCore/Private/VisualizeTexture.cpp
--- a/Engine/Source/Runtime/RenderCore/Private/VisualizeTexture.cpp
+++ b/Engine/Source/Runtime/RenderCore/Private/VisualizeTexture.cpp
@@ -256,7 +256,7 @@ void FVisualizeTexture::CreateContentCapturePass(FRDGBuilder& GraphBuilder, cons
 		PassParameters->VisualizeTexture2DMS = SrcTexture;
 		PassParameters->VisualizeUINT8Texture2D = SrcTexture;
 
-		PassParameters->RenderTargets[0] = FRenderTargetBinding(CopyTexture, ERenderTargetLoadAction::EClear, ERenderTargetStoreAction::EStore);
+		PassParameters->RenderTargets[0] = FRenderTargetBinding(CopyTexture, ERenderTargetLoadAction::EClear);
 	}
 
 	auto ShaderMap = GetGlobalShaderMap(FeatureLevel);
diff --git a/Engine/Source/Runtime/RenderCore/Public/ShaderParameterMacros.h b/Engine/Source/Runtime/RenderCore/Public/ShaderParameterMacros.h
--- a/Engine/Source/Runtime/RenderCore/Public/ShaderParameterMacros.h
+++ b/Engine/Source/Runtime/RenderCore/Public/ShaderParameterMacros.h
@@ -149,10 +149,9 @@ struct alignas(SHADER_PARAMETER_STRUCT_ALIGNMENT) FRenderTargetBinding
 	 *
 	 * Notes: Load and store action are on purpose without default values, to force the user to not forget one of these.
 	 */
-	FRenderTargetBinding(FRDGTexture* InTexture, ERenderTargetLoadAction InLoadAction, ERenderTargetStoreAction InStoreAction, uint8 InMipIndex = 0)
+	FRenderTargetBinding(FRDGTexture* InTexture, ERenderTargetLoadAction InLoadAction, uint8 InMipIndex = 0)
 		: Texture(InTexture)
 		, LoadAction(InLoadAction)
-		, StoreAction(InStoreAction)
 		, MipIndex(InMipIndex)
 	{
 		check(Validate());
@@ -166,10 +165,6 @@ struct alignas(SHADER_PARAMETER_STRUCT_ALIGNMENT) FRenderTargetBinding
 	{
 		return LoadAction;
 	}
-	FORCEINLINE ERenderTargetStoreAction GetStoreAction() const
-	{
-		return StoreAction;
-	}
 	FORCEINLINE uint8 GetMipIndex() const
 	{
 		return MipIndex;
@@ -181,7 +176,6 @@ private:
 	 */
 	TAlignedShaderParameterPtr<FRDGTexture*> Texture;
 	ERenderTargetLoadAction		LoadAction		= ERenderTargetLoadAction::ENoAction;
-	ERenderTargetStoreAction	StoreAction		= ERenderTargetStoreAction::ENoAction;
 	uint8						MipIndex		= 0;
 
 	RENDERCORE_API bool Validate() const;
@@ -202,15 +196,11 @@ struct alignas(SHADER_PARAMETER_STRUCT_ALIGNMENT) FDepthStencilBinding
 	FORCEINLINE FDepthStencilBinding(
 		FRDGTexture* InTexture,
 		ERenderTargetLoadAction InDepthLoadAction,
-		ERenderTargetStoreAction InDepthStoreAction,
 		ERenderTargetLoadAction InStencilLoadAction,
-		ERenderTargetStoreAction InStencilStoreAction,
 		FExclusiveDepthStencil InDepthStencilAccess)
 		: Texture(InTexture)
 		, DepthLoadAction(InDepthLoadAction)
-		, DepthStoreAction(InDepthStoreAction)
 		, StencilLoadAction(InStencilLoadAction)
-		, StencilStoreAction(InStencilStoreAction)
 		, DepthStencilAccess(InDepthStencilAccess)
 	{
 		check(Validate());
@@ -219,11 +209,9 @@ struct alignas(SHADER_PARAMETER_STRUCT_ALIGNMENT) FDepthStencilBinding
 	FORCEINLINE FDepthStencilBinding(
 		FRDGTexture* InTexture,
 		ERenderTargetLoadAction InDepthLoadAction,
-		ERenderTargetStoreAction InDepthStoreAction,
 		FExclusiveDepthStencil InDepthStencilAccess)
 		: Texture(InTexture)
 		, DepthLoadAction(InDepthLoadAction)
-		, DepthStoreAction(InDepthStoreAction)
 		, DepthStencilAccess(InDepthStencilAccess)
 	{
 		check(Validate());
@@ -237,18 +225,10 @@ struct alignas(SHADER_PARAMETER_STRUCT_ALIGNMENT) FDepthStencilBinding
 	{
 		return DepthLoadAction;
 	}
-	FORCEINLINE ERenderTargetStoreAction GetDepthStoreAction() const
-	{
-		return DepthStoreAction;
-	}
 	FORCEINLINE ERenderTargetLoadAction GetStencilLoadAction() const
 	{
 		return StencilLoadAction;
 	}
-	FORCEINLINE ERenderTargetStoreAction GetStencilStoreAction() const
-	{
-		return StencilStoreAction;
-	}
 	FORCEINLINE FExclusiveDepthStencil GetDepthStencilAccess() const
 	{
 		return DepthStencilAccess;
@@ -261,9 +241,7 @@ private:
 	 */
 	TAlignedShaderParameterPtr<FRDGTexture*> Texture = nullptr;
 	ERenderTargetLoadAction		DepthLoadAction		= ERenderTargetLoadAction::ENoAction;
-	ERenderTargetStoreAction	DepthStoreAction	= ERenderTargetStoreAction::ENoAction;
 	ERenderTargetLoadAction		StencilLoadAction	= ERenderTargetLoadAction::ENoAction;
-	ERenderTargetStoreAction	StencilStoreAction	= ERenderTargetStoreAction::ENoAction;
 	FExclusiveDepthStencil		DepthStencilAccess	= FExclusiveDepthStencil::DepthNop_StencilNop;
 
 	RENDERCORE_API bool Validate() const;
diff --git a/Engine/Source/Runtime/Renderer/Private/ClusteredDeferredShadingPass.cpp b/Engine/Source/Runtime/Renderer/Private/ClusteredDeferredShadingPass.cpp
--- a/Engine/Source/Runtime/Renderer/Private/ClusteredDeferredShadingPass.cpp
+++ b/Engine/Source/Runtime/Renderer/Private/ClusteredDeferredShadingPass.cpp
@@ -133,7 +133,7 @@ void FDeferredShadingSceneRenderer::AddClusteredDeferredShadingPass(FRHICommandL
 			PassParameters->SSProfilesTexture = RegisterExternalTextureWithFallback(GraphBuilder, GetSubsufaceProfileTexture_RT(RHICmdList), GSystemTextures.BlackDummy);
 			PassParameters->TransmissionProfilesLinearSampler = TStaticSamplerState<SF_Bilinear, AM_Clamp, AM_Clamp, AM_Clamp>::GetRHI();
 
-			PassParameters->RenderTargets[0] = FRenderTargetBinding(GraphBuilder.RegisterExternalTexture(SceneContext.GetSceneColor()), ERenderTargetLoadAction::ELoad, ERenderTargetStoreAction::EStore);
+			PassParameters->RenderTargets[0] = FRenderTargetBinding(GraphBuilder.RegisterExternalTexture(SceneContext.GetSceneColor()), ERenderTargetLoadAction::ELoad);
 			// NOTE: if we wanted to get depth/stencil testing happening we'd need to add that here too, at the moment we don't.
 
 			GraphBuilder.AddPass(
diff --git a/Engine/Source/Runtime/Renderer/Private/IndirectLightRendering.cpp b/Engine/Source/Runtime/Renderer/Private/IndirectLightRendering.cpp
--- a/Engine/Source/Runtime/Renderer/Private/IndirectLightRendering.cpp
+++ b/Engine/Source/Runtime/Renderer/Private/IndirectLightRendering.cpp
@@ -229,7 +229,7 @@ void FDeferredShadingSceneRenderer::RenderDiffuseIndirectAndAmbientOcclusion(FRH
 			PassParameters->ViewUniformBuffer = View.ViewUniformBuffer;
 
 			PassParameters->RenderTargets[0] = FRenderTargetBinding(
-				SceneColor, ERenderTargetLoadAction::ELoad, ERenderTargetStoreAction::EStore);
+				SceneColor, ERenderTargetLoadAction::ELoad);
 		
 			FDiffuseIndirectCompositePS::FPermutationDomain PermutationVector;
 			PermutationVector.Set<FDiffuseIndirectCompositePS::FApplyDiffuseIndirectDim>(PassParameters->DiffuseIndirectTexture != nullptr);
@@ -290,7 +290,7 @@ void FDeferredShadingSceneRenderer::RenderDiffuseIndirectAndAmbientOcclusion(FRH
 			PassParameters->ViewUniformBuffer = View.ViewUniformBuffer;
 
 			PassParameters->RenderTargets[0] = FRenderTargetBinding(
-				SceneColor, ERenderTargetLoadAction::ELoad, ERenderTargetStoreAction::EStore);
+				SceneColor, ERenderTargetLoadAction::ELoad);
 		
 			TShaderMapRef<FAmbientCubemapCompositePS> PixelShader(View.ShaderMap);
 			GraphBuilder.AddPass(
diff --git a/Engine/Source/Runtime/Renderer/Private/PlanarReflectionRendering.cpp b/Engine/Source/Runtime/Renderer/Private/PlanarReflectionRendering.cpp
--- a/Engine/Source/Runtime/Renderer/Private/PlanarReflectionRendering.cpp
+++ b/Engine/Source/Runtime/Renderer/Private/PlanarReflectionRendering.cpp
@@ -734,7 +734,7 @@ void FDeferredShadingSceneRenderer::RenderDeferredPlanarReflections(FRDGBuilder&
 	SetupSceneTextureSamplers(&PassParameters->SceneTextureSamplers);
 	PassParameters->ViewUniformBuffer = View.ViewUniformBuffer;
 	PassParameters->RenderTargets[0] = FRenderTargetBinding(
-		ReflectionsOutputTexture, bClearReflectionsOutputTexture ? ERenderTargetLoadAction::EClear : ERenderTargetLoadAction::ELoad, ERenderTargetStoreAction::EStore);
+		ReflectionsOutputTexture, bClearReflectionsOutputTexture ? ERenderTargetLoadAction::EClear : ERenderTargetLoadAction::ELoad);
 
 	GraphBuilder.AddPass(
 		RDG_EVENT_NAME("CompositePlanarReflections"),
diff --git a/Engine/Source/Runtime/Renderer/Private/PostProcess/DiaphragmDOF.cpp b/Engine/Source/Runtime/Renderer/Private/PostProcess/DiaphragmDOF.cpp
--- a/Engine/Source/Runtime/Renderer/Private/PostProcess/DiaphragmDOF.cpp
+++ b/Engine/Source/Runtime/Renderer/Private/PostProcess/DiaphragmDOF.cpp
@@ -2287,7 +2287,7 @@ FRDGTextureRef DiaphragmDOF::AddPasses(
 			PassParameters->IndirectDrawParameter = DrawIndirectParametersBuffer;
 			PassParameters->ScatterDrawList = GraphBuilder.CreateSRV(ScatterDrawList);
 			PassParameters->RenderTargets[0] = FRenderTargetBinding(
-				ConvolutionTextures->SceneColor, ERenderTargetLoadAction::ELoad, ERenderTargetStoreAction::EStore);
+				ConvolutionTextures->SceneColor, ERenderTargetLoadAction::ELoad);
 
 			ValidateShaderParameters(*VertexShader, *PassParameters);
 			ValidateShaderParameters(*PixelShader, *PassParameters);
diff --git a/Engine/Source/Runtime/Renderer/Private/PostProcess/PostProcessMaterial.cpp b/Engine/Source/Runtime/Renderer/Private/PostProcess/PostProcessMaterial.cpp
--- a/Engine/Source/Runtime/Renderer/Private/PostProcess/PostProcessMaterial.cpp
+++ b/Engine/Source/Runtime/Renderer/Private/PostProcess/PostProcessMaterial.cpp
@@ -448,17 +448,14 @@ FRDGTextureRef ComputePostProcessMaterial(
 
 	PostProcessMaterialParameters->RenderTargets[0] = FRenderTargetBinding(
 		OutputTexture,
-		OutputLoadAction,
-		ERenderTargetStoreAction::EStore);
+		OutputLoadAction);
 
 	if (DepthStencilTexture)
 	{
 		PostProcessMaterialParameters->RenderTargets.DepthStencil = FDepthStencilBinding(
 			DepthStencilTexture,
 			ERenderTargetLoadAction::ENoAction,
-			ERenderTargetStoreAction::ENoAction,
 			ERenderTargetLoadAction::ELoad,
-			ERenderTargetStoreAction::ENoAction,
 			FExclusiveDepthStencil::DepthRead_StencilRead);
 	}
 
diff --git a/Engine/Source/Runtime/Renderer/Private/PostProcess/PostProcessMotionBlur.cpp b/Engine/Source/Runtime/Renderer/Private/PostProcess/PostProcessMotionBlur.cpp
--- a/Engine/Source/Runtime/Renderer/Private/PostProcess/PostProcessMotionBlur.cpp
+++ b/Engine/Source/Runtime/Renderer/Private/PostProcess/PostProcessMotionBlur.cpp
@@ -531,16 +531,13 @@ void ComputeMotionBlurVelocity(
 			FDepthStencilBinding(
 				VelocityTileDepthTexture,
 				ERenderTargetLoadAction::ENoAction,
-				ERenderTargetStoreAction::EStore,
 				ERenderTargetLoadAction::ENoAction,
-				ERenderTargetStoreAction::ENoAction,
 				FExclusiveDepthStencil::DepthWrite_StencilNop);
 
 		PassParameters->RenderTargets[0] =
 			FRenderTargetBinding(
 				VelocityTileTexture,
-				ERenderTargetLoadAction::ENoAction,
-				ERenderTargetStoreAction::EStore);
+				ERenderTargetLoadAction::ENoAction);
 
 		TShaderMapRef<FMotionBlurVelocityDilateScatterVS> VertexShader(View.ShaderMap);
 		TShaderMapRef<FMotionBlurVelocityDilateScatterPS> PixelShader(View.ShaderMap);
@@ -714,7 +711,7 @@ FRDGTextureRef ComputeMotionBlurFilter(
 	{
 		FMotionBlurFilterPS::FParameters* PassParameters = GraphBuilder.AllocParameters<FMotionBlurFilterPS::FParameters>();
 		PassParameters->Filter = MotionBlurFilterParameters;
-		PassParameters->RenderTargets[0] = FRenderTargetBinding(ColorTextureOutput, ERenderTargetLoadAction::ENoAction, ERenderTargetStoreAction::EStore);
+		PassParameters->RenderTargets[0] = FRenderTargetBinding(ColorTextureOutput, ERenderTargetLoadAction::ENoAction);
 
 		TShaderMapRef<FMotionBlurFilterPS> PixelShader(View.ShaderMap, PermutationVector);
 
@@ -833,7 +830,7 @@ FRDGTextureRef VisualizeMotionBlur(
 	PassParameters->ColorSampler = GetMotionBlurColorSampler();
 	PassParameters->VelocitySampler = GetMotionBlurVelocitySampler();
 	PassParameters->DepthSampler = GetMotionBlurVelocitySampler();
-	PassParameters->RenderTargets[0] = FRenderTargetBinding(ColorTextureOutput, ERenderTargetLoadAction::ENoAction, ERenderTargetStoreAction::EStore);
+	PassParameters->RenderTargets[0] = FRenderTargetBinding(ColorTextureOutput, ERenderTargetLoadAction::ENoAction);
 
 	TShaderMapRef<FMotionBlurVisualizePS> PixelShader(View.ShaderMap);
 
diff --git a/Engine/Source/Runtime/Renderer/Private/PostProcess/PostProcessSubsurface.cpp b/Engine/Source/Runtime/Renderer/Private/PostProcess/PostProcessSubsurface.cpp
--- a/Engine/Source/Runtime/Renderer/Private/PostProcess/PostProcessSubsurface.cpp
+++ b/Engine/Source/Runtime/Renderer/Private/PostProcess/PostProcessSubsurface.cpp
@@ -166,7 +166,7 @@ FRDGTextureRef ComputeSubsurface(
 	if (bHasNonSubsurfaceView)
 	{
 		FSubsurfaceViewportCopyPS::FParameters* PassParameters = GraphBuilder.AllocParameters<FSubsurfaceViewportCopyPS::FParameters>();
-		PassParameters->RenderTargets[0] = FRenderTargetBinding(SceneTextureOutput, ERenderTargetLoadAction::ENoAction, ERenderTargetStoreAction::EStore);
+		PassParameters->RenderTargets[0] = FRenderTargetBinding(SceneTextureOutput, ERenderTargetLoadAction::ENoAction);
 		PassParameters->SubsurfaceInput0_Texture = SceneTexture;
 		PassParameters->SubsurfaceSampler0 = TStaticSamplerState<SF_Point, AM_Clamp, AM_Clamp, AM_Clamp>::GetRHI();
 
@@ -240,7 +240,7 @@ void VisualizeSubsurface(
 
 	FSubsurfaceVisualizePS::FParameters* PassParameters = GraphBuilder.AllocParameters<FSubsurfaceVisualizePS::FParameters>();
 	PassParameters->Subsurface = GetSubsurfaceCommonParameters(GraphBuilder.RHICmdList, View);
-	PassParameters->RenderTargets[0] = FRenderTargetBinding(SceneTextureOutput, ERenderTargetLoadAction::EClear, ERenderTargetStoreAction::EStore);
+	PassParameters->RenderTargets[0] = FRenderTargetBinding(SceneTextureOutput, ERenderTargetLoadAction::EClear);
 	PassParameters->SubsurfaceInput0.Texture = SceneTexture;
 	PassParameters->SubsurfaceInput0.Viewport = GetScreenPassTextureViewportParameters(SceneViewport);
 	PassParameters->SubsurfaceSampler0 = TStaticSamplerState<SF_Point, AM_Clamp, AM_Clamp, AM_Clamp>::GetRHI();
diff --git a/Engine/Source/Runtime/Renderer/Private/PostProcess/PostProcessing.cpp b/Engine/Source/Runtime/Renderer/Private/PostProcess/PostProcessing.cpp
--- a/Engine/Source/Runtime/Renderer/Private/PostProcess/PostProcessing.cpp
+++ b/Engine/Source/Runtime/Renderer/Private/PostProcess/PostProcessing.cpp
@@ -1017,7 +1017,7 @@ FRDGTextureRef AddSeparateTranslucencyCompositionPass(FRDGBuilder& GraphBuilder,
 	PassParameters->SeparateTranslucency = SeparateTranslucency;
 	PassParameters->SeparateTranslucencySampler = TStaticSamplerState<SF_Point>::GetRHI();
 	PassParameters->ViewUniformBuffer = View.ViewUniformBuffer;
-	PassParameters->RenderTargets[0] = FRenderTargetBinding(NewSceneColor, ERenderTargetLoadAction::ENoAction, ERenderTargetStoreAction::EStore);
+	PassParameters->RenderTargets[0] = FRenderTargetBinding(NewSceneColor, ERenderTargetLoadAction::ENoAction);
 			
 	TShaderMapRef<FComposeSeparateTranslucencyPS> PixelShader(View.ShaderMap);
 	FPixelShaderUtils::AddFullscreenPass(
diff --git a/Engine/Source/Runtime/Renderer/Private/PostProcess/SubsurfaceBurleyNormalized.cpp b/Engine/Source/Runtime/Renderer/Private/PostProcess/SubsurfaceBurleyNormalized.cpp
--- a/Engine/Source/Runtime/Renderer/Private/PostProcess/SubsurfaceBurleyNormalized.cpp
+++ b/Engine/Source/Runtime/Renderer/Private/PostProcess/SubsurfaceBurleyNormalized.cpp
@@ -331,7 +331,7 @@ void ComputeBurleySubsurfaceForView(
 		//single pass
 		FSubsurfaceBurleySinglePassPS::FParameters* PassParameters = GraphBuilder.AllocParameters<FSubsurfaceBurleySinglePassPS::FParameters>();
 		PassParameters->Subsurface = SubsurfaceCommonParameters;
-		PassParameters->RenderTargets[0] = FRenderTargetBinding(SceneTextureOutput, SceneTextureLoadAction, ERenderTargetStoreAction::EStore);
+		PassParameters->RenderTargets[0] = FRenderTargetBinding(SceneTextureOutput, SceneTextureLoadAction);
 		PassParameters->SubsurfaceInput0 = GetSubsurfaceBurleyInput(SceneTexture, SceneViewportParameters);
 		PassParameters->SubsurfaceSampler0 = SubsurfaceSamplerState;
 
@@ -360,7 +360,7 @@ void ComputeBurleySubsurfaceForView(
 		{
 			FSubsurfaceBurleySetupPS::FParameters* PassParameters = GraphBuilder.AllocParameters<FSubsurfaceBurleySetupPS::FParameters>();
 			PassParameters->Subsurface = SubsurfaceCommonParameters;
-			PassParameters->RenderTargets[0] = FRenderTargetBinding(SetupTexture, ERenderTargetLoadAction::ENoAction, ERenderTargetStoreAction::EStore);
+			PassParameters->RenderTargets[0] = FRenderTargetBinding(SetupTexture, ERenderTargetLoadAction::ENoAction);
 			PassParameters->SubsurfaceInput0 = GetSubsurfaceBurleyInput(SceneTexture, SceneViewportParameters);
 			PassParameters->SubsurfaceSampler0 = PointClampSampler;
 
@@ -417,7 +417,7 @@ void ComputeBurleySubsurfaceForView(
 
 			FSubsurfaceBurleyPS::FParameters* PassParameters = GraphBuilder.AllocParameters<FSubsurfaceBurleyPS::FParameters>();
 			PassParameters->Subsurface = SubsurfaceCommonParameters;
-			PassParameters->RenderTargets[0] = FRenderTargetBinding(TextureOutput, ERenderTargetLoadAction::ENoAction, ERenderTargetStoreAction::EStore);
+			PassParameters->RenderTargets[0] = FRenderTargetBinding(TextureOutput, ERenderTargetLoadAction::ENoAction);
 			PassParameters->SubsurfaceInput0 = GetSubsurfaceBurleyInput(TextureInput, SubsurfaceViewportParameters);
 			PassParameters->SubsurfaceSampler0 = SubsurfaceSamplerState;
 
@@ -434,7 +434,7 @@ void ComputeBurleySubsurfaceForView(
 	{
 		FSubsurfaceBurleyRecombinePS::FParameters* PassParameters = GraphBuilder.AllocParameters<FSubsurfaceBurleyRecombinePS::FParameters>();
 		PassParameters->Subsurface = SubsurfaceCommonParameters;
-		PassParameters->RenderTargets[0] = FRenderTargetBinding(SceneTextureOutput, SceneTextureLoadAction, ERenderTargetStoreAction::EStore);
+		PassParameters->RenderTargets[0] = FRenderTargetBinding(SceneTextureOutput, SceneTextureLoadAction);
 		PassParameters->SubsurfaceInput0 = GetSubsurfaceBurleyInput(SceneTexture, SceneViewportParameters);
 		PassParameters->SubsurfaceSampler0 = BilinearBorderSampler;
 
diff --git a/Engine/Source/Runtime/Renderer/Private/PostProcess/SubsurfaceSeparable.cpp b/Engine/Source/Runtime/Renderer/Private/PostProcess/SubsurfaceSeparable.cpp
--- a/Engine/Source/Runtime/Renderer/Private/PostProcess/SubsurfaceSeparable.cpp
+++ b/Engine/Source/Runtime/Renderer/Private/PostProcess/SubsurfaceSeparable.cpp
@@ -243,7 +243,7 @@ void ComputeSeparableSubsurfaceForView(
 		{
 			FSubsurfaceSetupPS::FParameters* PassParameters = GraphBuilder.AllocParameters<FSubsurfaceSetupPS::FParameters>();
 			PassParameters->Subsurface = SubsurfaceCommonParameters;
-			PassParameters->RenderTargets[0] = FRenderTargetBinding(SetupTexture, ERenderTargetLoadAction::ENoAction, ERenderTargetStoreAction::EStore);
+			PassParameters->RenderTargets[0] = FRenderTargetBinding(SetupTexture, ERenderTargetLoadAction::ENoAction);
 			PassParameters->SubsurfaceInput0 = GetSubsurfaceInput(SceneTexture, SceneViewportParameters);
 			PassParameters->SubsurfaceSampler0 = PointClampSampler;
 
@@ -297,7 +297,7 @@ void ComputeSeparableSubsurfaceForView(
 
 			FSubsurfacePS::FParameters* PassParameters = GraphBuilder.AllocParameters<FSubsurfacePS::FParameters>();
 			PassParameters->Subsurface = SubsurfaceCommonParameters;
-			PassParameters->RenderTargets[0] = FRenderTargetBinding(TextureOutput, ERenderTargetLoadAction::ENoAction, ERenderTargetStoreAction::EStore);
+			PassParameters->RenderTargets[0] = FRenderTargetBinding(TextureOutput, ERenderTargetLoadAction::ENoAction);
 			PassParameters->SubsurfaceInput0 = GetSubsurfaceInput(TextureInput, SubsurfaceViewportParameters);
 			PassParameters->SubsurfaceSampler0 = SubsurfaceSamplerState;
 
@@ -314,7 +314,7 @@ void ComputeSeparableSubsurfaceForView(
 	{
 		FSubsurfaceRecombinePS::FParameters* PassParameters = GraphBuilder.AllocParameters<FSubsurfaceRecombinePS::FParameters>();
 		PassParameters->Subsurface = SubsurfaceCommonParameters;
-		PassParameters->RenderTargets[0] = FRenderTargetBinding(SceneTextureOutput, SceneTextureLoadAction, ERenderTargetStoreAction::EStore);
+		PassParameters->RenderTargets[0] = FRenderTargetBinding(SceneTextureOutput, SceneTextureLoadAction);
 		PassParameters->SubsurfaceInput0 = GetSubsurfaceInput(SceneTexture, SceneViewportParameters);
 		PassParameters->SubsurfaceSampler0 = BilinearBorderSampler;
 
diff --git a/Engine/Source/Runtime/Renderer/Private/RayTracing/RaytracingSkylight.cpp b/Engine/Source/Runtime/Renderer/Private/RayTracing/RaytracingSkylight.cpp
--- a/Engine/Source/Runtime/Renderer/Private/RayTracing/RaytracingSkylight.cpp
+++ b/Engine/Source/Runtime/Renderer/Private/RayTracing/RaytracingSkylight.cpp
@@ -1146,7 +1146,7 @@ void FDeferredShadingSceneRenderer::CompositeRayTracingSkyLight(
 		PassParameters->SkyLightTexture = GraphBuilder.RegisterExternalTexture(SkyLightRT);
 		PassParameters->SkyLightTextureSampler = TStaticSamplerState<SF_Point, AM_Clamp, AM_Clamp, AM_Clamp>::GetRHI();
 		PassParameters->ViewUniformBuffer = Views[ViewIndex].ViewUniformBuffer;
-		PassParameters->RenderTargets[0] = FRenderTargetBinding(GraphBuilder.RegisterExternalTexture(SceneContext.GetSceneColor()), ERenderTargetLoadAction::ELoad, ERenderTargetStoreAction::EStore);
+		PassParameters->RenderTargets[0] = FRenderTargetBinding(GraphBuilder.RegisterExternalTexture(SceneContext.GetSceneColor()), ERenderTargetLoadAction::ELoad);
 		PassParameters->SceneTextures = SceneTextures;
 
 		// dxr_todo: Unify with RTGI compositing workflow
diff --git a/Engine/Source/Runtime/Renderer/Private/ReflectionEnvironment.cpp b/Engine/Source/Runtime/Renderer/Private/ReflectionEnvironment.cpp
--- a/Engine/Source/Runtime/Renderer/Private/ReflectionEnvironment.cpp
+++ b/Engine/Source/Runtime/Renderer/Private/ReflectionEnvironment.cpp
@@ -847,7 +847,7 @@ void FDeferredShadingSceneRenderer::RenderDeferredReflectionsAndSkyLighting(FRHI
 				PassParameters->ForwardLightData = View.ForwardLightingResources->ForwardLightDataUniformBuffer;
 			}
 
-			PassParameters->RenderTargets[0] = FRenderTargetBinding(SceneColorTexture, ERenderTargetLoadAction::ELoad, ERenderTargetStoreAction::EStore);
+			PassParameters->RenderTargets[0] = FRenderTargetBinding(SceneColorTexture, ERenderTargetLoadAction::ELoad);
 
 			auto PermutationVector = FReflectionEnvironmentSkyLightingPS::BuildPermutationVector(
 				View, bHasBoxCaptures, bHasSphereCaptures, DynamicBentNormalAO != NULL, bSkyLight, bDynamicSkyLight, bApplySkyShadowing, bRayTracedReflections);
diff --git a/Engine/Source/Runtime/Renderer/Private/SceneTextureReductions.cpp b/Engine/Source/Runtime/Renderer/Private/SceneTextureReductions.cpp
--- a/Engine/Source/Runtime/Renderer/Private/SceneTextureReductions.cpp
+++ b/Engine/Source/Runtime/Renderer/Private/SceneTextureReductions.cpp
@@ -73,7 +73,7 @@ void BuildHZB(FRDGBuilder& GraphBuilder, FViewInfo& View)
 
 	{
 		FHZBBuildPassParameters* PassParameters = GraphBuilder.AllocParameters<FHZBBuildPassParameters>();
-		PassParameters->RenderTargets[0] = FRenderTargetBinding( HZBTexture, ERenderTargetLoadAction::ENoAction, ERenderTargetStoreAction::EStore);
+		PassParameters->RenderTargets[0] = FRenderTargetBinding( HZBTexture, ERenderTargetLoadAction::ENoAction);
 
 		//@DW - this pass only reads external textures, we don't have any graph inputs
 		GraphBuilder.AddPass(
@@ -155,7 +155,7 @@ void BuildHZB(FRDGBuilder& GraphBuilder, FViewInfo& View)
 		FRDGTextureSRVRef ParentMipSRV = GraphBuilder.CreateSRV(Desc);
 
 		FHZBBuildPassParameters* PassParameters = GraphBuilder.AllocParameters<FHZBBuildPassParameters>();
-		PassParameters->RenderTargets[0] = FRenderTargetBinding( HZBTexture, ERenderTargetLoadAction::ENoAction, ERenderTargetStoreAction::EStore, MipIndex );
+		PassParameters->RenderTargets[0] = FRenderTargetBinding( HZBTexture, ERenderTargetLoadAction::ENoAction, MipIndex );
 		PassParameters->Texture = ParentMipSRV;
 		PassParameters->TextureSampler = TStaticSamplerState<SF_Point, AM_Clamp, AM_Clamp, AM_Clamp>::GetRHI();
 
diff --git a/Engine/Source/Runtime/Renderer/Private/ScreenPass.cpp b/Engine/Source/Runtime/Renderer/Private/ScreenPass.cpp
--- a/Engine/Source/Runtime/Renderer/Private/ScreenPass.cpp
+++ b/Engine/Source/Runtime/Renderer/Private/ScreenPass.cpp
@@ -192,7 +192,7 @@ void AddDrawTexturePass(
 	FCopyRectPS::FParameters* Parameters = GraphBuilder.AllocParameters<FCopyRectPS::FParameters>();
 	Parameters->InputTexture = InputTexture;
 	Parameters->InputSampler = TStaticSamplerState<>::GetRHI();
-	Parameters->RenderTargets[0] = FRenderTargetBinding(OutputTexture, LoadAction, ERenderTargetStoreAction::EStore);
+	Parameters->RenderTargets[0] = FRenderTargetBinding(OutputTexture, LoadAction);
 
 	AddDrawScreenPass(
 		GraphBuilder,
diff --git a/Engine/Source/Runtime/Renderer/Private/ScreenSpaceRayTracing.cpp b/Engine/Source/Runtime/Renderer/Private/ScreenSpaceRayTracing.cpp
--- a/Engine/Source/Runtime/Renderer/Private/ScreenSpaceRayTracing.cpp
+++ b/Engine/Source/Runtime/Renderer/Private/ScreenSpaceRayTracing.cpp
@@ -463,11 +463,11 @@ void RenderScreenSpaceReflections(
 	SetupSceneTextureSamplers(&CommonParameters.SceneTextureSamplers);
 	
 	FRenderTargetBindingSlots RenderTargets;
-	RenderTargets[0] = FRenderTargetBinding(DenoiserInputs->Color, ERenderTargetLoadAction::ENoAction, ERenderTargetStoreAction::EStore);
+	RenderTargets[0] = FRenderTargetBinding(DenoiserInputs->Color, ERenderTargetLoadAction::ENoAction);
 
 	if (bDenoiser)
 	{
-		RenderTargets[1] = FRenderTargetBinding(DenoiserInputs->RayHitDistance, ERenderTargetLoadAction::ENoAction, ERenderTargetStoreAction::EStore);
+		RenderTargets[1] = FRenderTargetBinding(DenoiserInputs->RayHitDistance, ERenderTargetLoadAction::ENoAction);
 	}
 
 	// Do a pre pass that output 0, or set a stencil mask to run the more expensive pixel shader.
@@ -476,8 +476,8 @@ void RenderScreenSpaceReflections(
 		// Also bind the depth buffer
 		RenderTargets.DepthStencil = FDepthStencilBinding(
 			SceneTextures.SceneDepthBuffer,
-			ERenderTargetLoadAction::ENoAction, ERenderTargetStoreAction::ENoAction,
-			ERenderTargetLoadAction::ELoad, ERenderTargetStoreAction::EStore,
+			ERenderTargetLoadAction::ENoAction,
+			ERenderTargetLoadAction::ELoad,
 			FExclusiveDepthStencil::DepthRead_StencilWrite);
 
 		FScreenSpaceReflectionsStencilPS::FPermutationDomain PermutationVector;
diff --git a/Engine/Source/Runtime/Renderer/Private/ShaderPrint.cpp b/Engine/Source/Runtime/Renderer/Private/ShaderPrint.cpp
--- a/Engine/Source/Runtime/Renderer/Private/ShaderPrint.cpp
+++ b/Engine/Source/Runtime/Renderer/Private/ShaderPrint.cpp
@@ -415,7 +415,7 @@ namespace ShaderPrint
 			TShaderMapRef< FShaderDrawSymbolsPS > PixelShader(GlobalShaderMap);
 
 			SHADER::FParameters* PassParameters = GraphBuilder.AllocParameters<SHADER::FParameters>();
-			PassParameters->RenderTargets[0] = FRenderTargetBinding(OutputTexture, ERenderTargetLoadAction::ENoAction, ERenderTargetStoreAction::ENoAction);
+			PassParameters->RenderTargets[0] = FRenderTargetBinding(OutputTexture, ERenderTargetLoadAction::ENoAction);
 			PassParameters->UniformBufferParameters = UniformBuffer;
 			PassParameters->MiniFontTexture = FontTexture;
 			PassParameters->SymbolsBuffer = GraphBuilder.CreateSRV(SymbolBuffer);
diff --git a/Engine/Source/Runtime/Renderer/Private/VT/RuntimeVirtualTextureRender.cpp b/Engine/Source/Runtime/Renderer/Private/VT/RuntimeVirtualTextureRender.cpp
--- a/Engine/Source/Runtime/Renderer/Private/VT/RuntimeVirtualTextureRender.cpp
+++ b/Engine/Source/Runtime/Renderer/Private/VT/RuntimeVirtualTextureRender.cpp
@@ -651,9 +651,9 @@ namespace RuntimeVirtualTexture
 		if (GraphSetup.bRenderPass)
 		{
 			FShader_VirtualTextureMaterialDraw::FParameters* PassParameters = GraphBuilder.AllocParameters<FShader_VirtualTextureMaterialDraw::FParameters>();
-			PassParameters->RenderTargets[0] = GraphSetup.RenderTexture0 ? FRenderTargetBinding(GraphSetup.RenderTexture0, ERenderTargetLoadAction::EClear, ERenderTargetStoreAction::EStore) : FRenderTargetBinding();
-			PassParameters->RenderTargets[1] = GraphSetup.RenderTexture1 ? FRenderTargetBinding(GraphSetup.RenderTexture1, ERenderTargetLoadAction::EClear, ERenderTargetStoreAction::EStore) : FRenderTargetBinding();
-			PassParameters->RenderTargets[2] = GraphSetup.RenderTexture2 ? FRenderTargetBinding(GraphSetup.RenderTexture2, ERenderTargetLoadAction::EClear, ERenderTargetStoreAction::EStore) : FRenderTargetBinding();
+			PassParameters->RenderTargets[0] = GraphSetup.RenderTexture0 ? FRenderTargetBinding(GraphSetup.RenderTexture0, ERenderTargetLoadAction::EClear) : FRenderTargetBinding();
+			PassParameters->RenderTargets[1] = GraphSetup.RenderTexture1 ? FRenderTargetBinding(GraphSetup.RenderTexture1, ERenderTargetLoadAction::EClear) : FRenderTargetBinding();
+			PassParameters->RenderTargets[2] = GraphSetup.RenderTexture2 ? FRenderTargetBinding(GraphSetup.RenderTexture2, ERenderTargetLoadAction::EClear) : FRenderTargetBinding();
 
 			GraphBuilder.AddPass(
 				RDG_EVENT_NAME("VirtualTextureDraw"),
diff --git a/Engine/Source/Runtime/Renderer/Private/VolumetricFog.cpp b/Engine/Source/Runtime/Renderer/Private/VolumetricFog.cpp
--- a/Engine/Source/Runtime/Renderer/Private/VolumetricFog.cpp
+++ b/Engine/Source/Runtime/Renderer/Private/VolumetricFog.cpp
@@ -488,7 +488,7 @@ void FDeferredShadingSceneRenderer::RenderLocalLightsForVolumetricFog(
 		OutLocalShadowedLightScattering = GraphBuilder.CreateTexture(VolumeDesc, TEXT("LocalShadowedLightScattering"));
 
 		FRenderTargetParameters* PassParameters = GraphBuilder.AllocParameters<FRenderTargetParameters>();
-		PassParameters->RenderTargets[0] = FRenderTargetBinding(OutLocalShadowedLightScattering, ERenderTargetLoadAction::EClear, ERenderTargetStoreAction::EStore);
+		PassParameters->RenderTargets[0] = FRenderTargetBinding(OutLocalShadowedLightScattering, ERenderTargetLoadAction::EClear);
 
 		GraphBuilder.AddPass(
 			RDG_EVENT_NAME("ShadowedLights"),
diff --git a/Engine/Source/Runtime/Renderer/Private/VolumetricFogLightFunction.cpp b/Engine/Source/Runtime/Renderer/Private/VolumetricFogLightFunction.cpp
--- a/Engine/Source/Runtime/Renderer/Private/VolumetricFogLightFunction.cpp
+++ b/Engine/Source/Runtime/Renderer/Private/VolumetricFogLightFunction.cpp
@@ -202,7 +202,7 @@ void FDeferredShadingSceneRenderer::RenderLightFunctionForVolumetricFog(
 			OutLightFunctionWorldToShadow = WorldToShadowValue;
 
 			FRenderTargetParameters* PassParameters = GraphBuilder.AllocParameters<FRenderTargetParameters>();
-			PassParameters->RenderTargets[0] = FRenderTargetBinding(OutLightFunctionTexture, ERenderTargetLoadAction::ENoAction, ERenderTargetStoreAction::EStore);
+			PassParameters->RenderTargets[0] = FRenderTargetBinding(OutLightFunctionTexture, ERenderTargetLoadAction::ENoAction);
 
 			GraphBuilder.AddPass(
 				RDG_EVENT_NAME("LightFunction"),
diff --git a/Engine/Source/Runtime/Renderer/Private/VolumetricFogVoxelization.cpp b/Engine/Source/Runtime/Renderer/Private/VolumetricFogVoxelization.cpp
--- a/Engine/Source/Runtime/Renderer/Private/VolumetricFogVoxelization.cpp
+++ b/Engine/Source/Runtime/Renderer/Private/VolumetricFogVoxelization.cpp
@@ -590,8 +590,8 @@ void FDeferredShadingSceneRenderer::VoxelizeFogVolumePrimitives(
 	if (View.VolumetricMeshBatches.Num() > 0 && DoesPlatformSupportVolumetricFogVoxelization(View.GetShaderPlatform()))
 	{
 		FRenderTargetParameters* PassParameters = GraphBuilder.AllocParameters<FRenderTargetParameters>();
-		PassParameters->RenderTargets[0] = FRenderTargetBinding(IntegrationData.VBufferA, ERenderTargetLoadAction::ELoad, ERenderTargetStoreAction::EStore);
-		PassParameters->RenderTargets[1] = FRenderTargetBinding(IntegrationData.VBufferB, ERenderTargetLoadAction::ELoad, ERenderTargetStoreAction::EStore);
+		PassParameters->RenderTargets[0] = FRenderTargetBinding(IntegrationData.VBufferA, ERenderTargetLoadAction::ELoad);
+		PassParameters->RenderTargets[1] = FRenderTargetBinding(IntegrationData.VBufferB, ERenderTargetLoadAction::ELoad);
 
 		GraphBuilder.AddPass(
 			RDG_EVENT_NAME("VoxelizeVolumePrimitives"),
