Fixes wrong validations for depth-stencil that must have TexCreate_DepthStencilTargetable flag

diff --git a/Engine/Source/Runtime/RenderCore/Private/RenderGraphValidation.cpp b/Engine/Source/Runtime/RenderCore/Private/RenderGraphValidation.cpp
--- a/Engine/Source/Runtime/RenderCore/Private/RenderGraphValidation.cpp
+++ b/Engine/Source/Runtime/RenderCore/Private/RenderGraphValidation.cpp
@@ -456,8 +456,8 @@ void FRDGUserValidation::ValidateAddPass(const FRDGPass* Pass)
 					Pass->GetName());
 
 				checkf(
-					Texture->Desc.TargetableFlags & TexCreate_RenderTargetable,
-					TEXT("Pass '%s' attempted to bind texture '%s' as a depth stencil render target, but the texture has not been created with TexCreate_RenderTargetable."),
+					Texture->Desc.TargetableFlags & TexCreate_DepthStencilTargetable,
+					TEXT("Pass '%s' attempted to bind texture '%s' as a depth stencil render target, but the texture has not been created with TexCreate_DepthStencilTargetable."),
 					PassName, Texture->Name);
 
 				// Depth stencil only supports one mip, since there isn't actually a way to select the mip level.
