Adds validation between depth stencil load action and depth stencil access

diff --git a/Engine/Source/Runtime/RenderCore/Private/ShaderParameterStruct.cpp b/Engine/Source/Runtime/RenderCore/Private/ShaderParameterStruct.cpp
--- a/Engine/Source/Runtime/RenderCore/Private/ShaderParameterStruct.cpp
+++ b/Engine/Source/Runtime/RenderCore/Private/ShaderParameterStruct.cpp
@@ -379,6 +379,17 @@ bool FDepthStencilBinding::Validate() const
 				TEXT("Unable to have stencil access on texture %s that have a pixel format %s that does not support stencil."),
 				Texture->Name, FormatString);
 		}
+
+		bool bReadDepth = DepthStencilAccess.IsUsingDepth() && !DepthStencilAccess.IsDepthWrite();
+		bool bReadStencil = DepthStencilAccess.IsUsingStencil() && !DepthStencilAccess.IsStencilWrite();
+
+		checkf(!(bReadDepth && DepthLoadAction != ERenderTargetLoadAction::ELoad),
+			TEXT("Depth read access, but without depth load action on texture %s doesn't make any sens."),
+			Texture->Name);
+
+		checkf(!(bReadStencil && StencilLoadAction != ERenderTargetLoadAction::ELoad),
+			TEXT("Stencil read access, but without stencil load action on texture %s doesn't make any sens."),
+			Texture->Name);
 	}
 	else
 	{
diff --git a/Engine/Source/Runtime/Renderer/Private/ScreenSpaceRayTracing.cpp b/Engine/Source/Runtime/Renderer/Private/ScreenSpaceRayTracing.cpp
--- a/Engine/Source/Runtime/Renderer/Private/ScreenSpaceRayTracing.cpp
+++ b/Engine/Source/Runtime/Renderer/Private/ScreenSpaceRayTracing.cpp
@@ -478,7 +478,7 @@ void RenderScreenSpaceReflections(
 			SceneTextures.SceneDepthBuffer,
 			ERenderTargetLoadAction::ENoAction,
 			ERenderTargetLoadAction::ELoad,
-			FExclusiveDepthStencil::DepthRead_StencilWrite);
+			FExclusiveDepthStencil::DepthNop_StencilWrite);
 
 		FScreenSpaceReflectionsStencilPS::FPermutationDomain PermutationVector;
 		PermutationVector.Set<FSSROutputForDenoiser>(bDenoiser);
