Fixed validation error in motion blur.

diff --git a/Engine/Source/Runtime/RenderCore/Private/RenderGraphValidation.cpp b/Engine/Source/Runtime/RenderCore/Private/RenderGraphValidation.cpp
--- a/Engine/Source/Runtime/RenderCore/Private/RenderGraphValidation.cpp
+++ b/Engine/Source/Runtime/RenderCore/Private/RenderGraphValidation.cpp
@@ -494,7 +494,7 @@ void FRDGUserValidation::ValidateAddPass(const FRDGPass* Pass)
 
 					const bool bIsLoadAction = RenderTarget.GetLoadAction() == ERenderTargetLoadAction::ELoad;
 
-					checkf(
+					ensureMsgf(
 						Texture->Desc.TargetableFlags & TexCreate_RenderTargetable,
 						TEXT("Pass '%s' attempted to bind texture '%s' as a render target, but the texture has not been created with TexCreate_RenderTargetable."),
 						PassName, Texture->Name);
diff --git a/Engine/Source/Runtime/Renderer/Private/PostProcess/PostProcessMotionBlur.cpp b/Engine/Source/Runtime/Renderer/Private/PostProcess/PostProcessMotionBlur.cpp
--- a/Engine/Source/Runtime/Renderer/Private/PostProcess/PostProcessMotionBlur.cpp
+++ b/Engine/Source/Runtime/Renderer/Private/PostProcess/PostProcessMotionBlur.cpp
@@ -502,7 +502,7 @@ void ComputeMotionBlurVelocity(
 				PF_FloatRGBA,
 				FClearValueBinding::None,
 				GFastVRamConfig.MotionBlur,
-				TexCreate_ShaderResource | TexCreate_UAV,
+				TexCreate_ShaderResource | TexCreate_UAV | TexCreate_RenderTargetable,
 				false),
 			TEXT("DilatedVelocityTile"));
 
