Added validation to assert when multiple uniform buffer definitions are registered with the same shader variable name.

diff --git a/Engine/Source/Runtime/RenderCore/Private/ShaderParameterMetadata.cpp b/Engine/Source/Runtime/RenderCore/Private/ShaderParameterMetadata.cpp
--- a/Engine/Source/Runtime/RenderCore/Private/ShaderParameterMetadata.cpp
+++ b/Engine/Source/Runtime/RenderCore/Private/ShaderParameterMetadata.cpp
@@ -8,6 +8,11 @@
 #include "RenderCore.h"
 #include "ShaderCore.h"
 
+#define VALIDATE_UNIFORM_BUFFER_UNIQUE_NAME (!UE_BUILD_SHIPPING && !UE_BUILD_TEST)
+
+#if VALIDATE_UNIFORM_BUFFER_UNIQUE_NAME
+static TMap<FName, FName> GlobalShaderVariableToStructMap;
+#endif
 
 static TLinkedList<FShaderParametersMetadata*>* GUniformStructList = nullptr;
 
@@ -83,6 +88,25 @@ FShaderParametersMetadata::FShaderParametersMetadata(
 		// Verify that during FName creation there's no case conversion
 		checkSlow(FCString::Strcmp(StructTypeName, *StructTypeFName.GetPlainNameString()) == 0);
 		GetNameStructMap().Add(FName(StructTypeFName), this);
+
+#if VALIDATE_UNIFORM_BUFFER_UNIQUE_NAME
+		FName ShaderVariableFName(ShaderVariableName);
+
+		// Verify that the global variable name is unique so that we can disambiguate when reflecting from shader source.
+		if (FName* StructFName = GlobalShaderVariableToStructMap.Find(ShaderVariableFName))
+		{
+			checkf(
+				false,
+				TEXT("Found duplicate Uniform Buffer shader variable name %s defined by struct %s. Previous definition ")
+				TEXT("found on struct %s. Uniform buffer shader names must be unique to support name-based reflection of ")
+				TEXT("shader source files."),
+				ShaderVariableName,
+				StructTypeName,
+				*StructFName->GetPlainNameString());
+		}
+
+		GlobalShaderVariableToStructMap.Add(ShaderVariableFName, StructTypeFName);
+#endif
 	}
 	else
 	{
@@ -93,6 +117,20 @@ FShaderParametersMetadata::FShaderParametersMetadata(
 	}
 }
 
+FShaderParametersMetadata::~FShaderParametersMetadata()
+{
+	if (UseCase == EUseCase::GlobalShaderParameterStruct)
+	{
+		GlobalListLink.Unlink();
+		GetNameStructMap().Remove(FName(StructTypeName, FNAME_Find));
+
+#if VALIDATE_UNIFORM_BUFFER_UNIQUE_NAME
+		GlobalShaderVariableToStructMap.Remove(FName(ShaderVariableName, FNAME_Find));
+#endif
+	}
+}
+
+
 void FShaderParametersMetadata::InitializeAllGlobalStructs()
 {
 	for (TLinkedList<FShaderParametersMetadata*>::TIterator StructIt(FShaderParametersMetadata::GetStructList()); StructIt; StructIt.Next())
diff --git a/Engine/Source/Runtime/RenderCore/Public/ShaderParameterMetadata.h b/Engine/Source/Runtime/RenderCore/Public/ShaderParameterMetadata.h
--- a/Engine/Source/Runtime/RenderCore/Public/ShaderParameterMetadata.h
+++ b/Engine/Source/Runtime/RenderCore/Public/ShaderParameterMetadata.h
@@ -144,11 +144,7 @@ public:
 		uint32 InSize,
 		const TArray<FMember>& InMembers);
 
-	virtual ~FShaderParametersMetadata()
-	{
-		GlobalListLink.Unlink();
-		GetNameStructMap().Remove(FName(StructTypeName, FNAME_Find));
-	}
+	virtual ~FShaderParametersMetadata();
 
 	void GetNestedStructs(TArray<const FShaderParametersMetadata*>& OutNestedStructs) const;
 
