Adds helper for byte address buffer in FRDGBufferDesc

diff --git a/Engine/Source/Runtime/RenderCore/Public/RenderGraphBuilder.inl b/Engine/Source/Runtime/RenderCore/Public/RenderGraphBuilder.inl
--- a/Engine/Source/Runtime/RenderCore/Public/RenderGraphBuilder.inl
+++ b/Engine/Source/Runtime/RenderCore/Public/RenderGraphBuilder.inl
@@ -96,6 +96,11 @@ inline FRDGBufferRef FRDGBuilder::CreateBuffer(
 	{
 		checkf(Name, TEXT("Creating a render graph buffer requires a valid debug name."));
 		Validation.ExecuteGuard(TEXT("CreateBuffer"), Name);
+
+		if (Desc.UnderlyingType == FRDGBufferDesc::EUnderlyingType::StructuredBuffer && (Desc.Usage | BUF_ByteAddressBuffer))
+		{
+			check(Desc.BytesPerElement == 4);
+		}
 	}
 #endif
 
diff --git a/Engine/Source/Runtime/RenderCore/Public/RenderGraphResources.h b/Engine/Source/Runtime/RenderCore/Public/RenderGraphResources.h
--- a/Engine/Source/Runtime/RenderCore/Public/RenderGraphResources.h
+++ b/Engine/Source/Runtime/RenderCore/Public/RenderGraphResources.h
@@ -488,6 +488,17 @@ struct RENDERCORE_API FRDGBufferDesc
 		return Desc;
 	}
 
+	static FRDGBufferDesc CreateByteAddressDesc(uint32 NumBytes)
+	{
+		check(NumBytes % 4 == 0);
+		FRDGBufferDesc Desc;
+		Desc.UnderlyingType = EUnderlyingType::StructuredBuffer;
+		Desc.Usage = EBufferUsageFlags(BUF_UnorderedAccess | BUF_ShaderResource | BUF_ByteAddressBuffer);
+		Desc.BytesPerElement = 4;
+		Desc.NumElements = NumBytes / 4;
+		return Desc;
+	}
+
 	/** Returns the total number of bytes allocated for a such buffer. */
 	uint32 GetTotalNumBytes() const
 	{
