Fixed validation condition in FRDGBuilder::CreateBuffer

diff --git a/Engine/Source/Runtime/RenderCore/Public/RenderGraphBuilder.inl b/Engine/Source/Runtime/RenderCore/Public/RenderGraphBuilder.inl
--- a/Engine/Source/Runtime/RenderCore/Public/RenderGraphBuilder.inl
+++ b/Engine/Source/Runtime/RenderCore/Public/RenderGraphBuilder.inl
@@ -97,7 +97,7 @@ inline FRDGBufferRef FRDGBuilder::CreateBuffer(
 		checkf(Name, TEXT("Creating a render graph buffer requires a valid debug name."));
 		Validation.ExecuteGuard(TEXT("CreateBuffer"), Name);
 
-		if (Desc.UnderlyingType == FRDGBufferDesc::EUnderlyingType::StructuredBuffer && (Desc.Usage | BUF_ByteAddressBuffer))
+		if (Desc.UnderlyingType == FRDGBufferDesc::EUnderlyingType::StructuredBuffer && (Desc.Usage & BUF_ByteAddressBuffer))
 		{
 			check(Desc.BytesPerElement == 4);
 		}
