Extended RDG validation with error message.

diff --git a/Engine/Source/Runtime/RenderCore/Public/RenderGraphBuilder.inl b/Engine/Source/Runtime/RenderCore/Public/RenderGraphBuilder.inl
--- a/Engine/Source/Runtime/RenderCore/Public/RenderGraphBuilder.inl
+++ b/Engine/Source/Runtime/RenderCore/Public/RenderGraphBuilder.inl
@@ -99,7 +99,7 @@ inline FRDGBufferRef FRDGBuilder::CreateBuffer(
 
 		if (Desc.UnderlyingType == FRDGBufferDesc::EUnderlyingType::StructuredBuffer && (Desc.Usage & BUF_ByteAddressBuffer))
 		{
-			check(Desc.BytesPerElement == 4);
+			checkf(Desc.BytesPerElement == 4, TEXT("Creating buffer '%s' as a structured buffer that is also byte addressable, BytesPerElement must be 4! Instead it is %d"), Name, Desc.BytesPerElement);
 		}
 	}
 #endif
