Fixed AddCopyTexturePass to mark resources as used.

diff --git a/Engine/Source/Runtime/RenderCore/Private/RenderGraphUtils.cpp b/Engine/Source/Runtime/RenderCore/Private/RenderGraphUtils.cpp
--- a/Engine/Source/Runtime/RenderCore/Private/RenderGraphUtils.cpp
+++ b/Engine/Source/Runtime/RenderCore/Private/RenderGraphUtils.cpp
@@ -111,6 +111,10 @@ void AddCopyTexturePass(
 	GraphBuilder.AddPass(RDG_EVENT_NAME("CopyTexture"), Parameters, ERDGPassFlags::Copy,
 		[InputTexture, OutputTexture, CopyInfo](FRHICommandList& RHICmdList)
 	{
+        // Manually mark as used since we aren't invoking any shaders.
+		InputTexture->MarkResourceAsUsed();
+		OutputTexture->MarkResourceAsUsed();
+
 		RHICmdList.CopyTexture(InputTexture->GetRHI(), OutputTexture->GetRHI(), CopyInfo);
 	});
 }
