Extended RDG helpers with support for ClearUAV texture variant. Fixed validation so that parent texture can be accessed when child view is declared on a pass (needed for UAV clear methods).

diff --git a/Engine/Source/Runtime/RenderCore/Private/RenderGraphValidation.cpp b/Engine/Source/Runtime/RenderCore/Private/RenderGraphValidation.cpp
--- a/Engine/Source/Runtime/RenderCore/Private/RenderGraphValidation.cpp
+++ b/Engine/Source/Runtime/RenderCore/Private/RenderGraphValidation.cpp
@@ -233,7 +233,12 @@ void FRDGUserValidation::ValidateExtractResource(FRDGTrackedResourceRef Resource
 void FRDGUserValidation::RemoveUnusedWarning(FRDGTrackedResourceRef Resource)
 {
 	check(Resource);
+
+    // Removes 'produced but not used' warning.
 	Resource->PassAccessCount++;
+
+    // Removes 'not used' warning.
+	Resource->bIsActuallyUsedByPass = true;
 }
 
 void FRDGUserValidation::ValidateAllocPassParameters(const void* Parameters)
@@ -613,7 +618,7 @@ void FRDGUserValidation::ValidateExecuteEnd()
 		{
 			ValidateResourceAtExecuteEnd(Buffer);
 		}
-	} // if (IsRDGDebugEnabled())
+	}
 
 	TrackedTextures.Empty();
 	TrackedBuffers.Empty();
@@ -667,8 +672,8 @@ void FRDGUserValidation::ValidateExecutePassBegin(const FRDGPass* Pass)
 					Texture->bHasBeenBoundAsRenderTarget = true;
 				}
 			}
-		} // for (uint32 ParameterIndex = 0; ParameterIndex < ParameterCount; ++ParameterIndex)
-	} // if (IsRDGDebugEnabled())
+		}
+	} 
 }
 
 void FRDGUserValidation::ValidateExecutePassEnd(const FRDGPass* Pass)
@@ -701,7 +706,7 @@ void FRDGUserValidation::ValidateExecutePassEnd(const FRDGPass* Pass)
 		if (TrackedResourceCount != UsedResourceCount)
 		{
 			FString WarningMessage = FString::Printf(
-				TEXT("'%d' of the '%d' resources of the pass '%s' where not actually used."),
+				TEXT("'%d' of the '%d' resources of the pass '%s' were not actually used."),
 				TrackedResourceCount - UsedResourceCount, TrackedResourceCount, Pass->GetName());
 
 			for (uint32 ParameterIndex = 0; ParameterIndex < ParameterCount; ++ParameterIndex)
@@ -748,38 +753,72 @@ void FRDGUserValidation::SetAllowRHIAccess(const FRDGPass* Pass, bool bAllowAcce
 	{
 		FRDGPassParameter Parameter = ParameterStruct.GetParameter(ParameterIndex);
 
-		if (Parameter.IsResource())
+		switch (Parameter.GetType())
 		{
+		case UBMT_RDG_TEXTURE:
+		case UBMT_RDG_TEXTURE_COPY_DEST:
+		case UBMT_RDG_BUFFER:
+		case UBMT_RDG_BUFFER_COPY_DEST:
 			if (FRDGResourceRef Resource = Parameter.GetAsResource())
 			{
 				Resource->bAllowRHIAccess = bAllowAccess;
 			}
-		}
-		else if (Parameter.GetType() == UBMT_RENDER_TARGET_BINDING_SLOTS)
-		{
-			const FRenderTargetBindingSlots& RenderTargetBindingSlots = Parameter.GetAsRenderTargetBindingSlots();
-			const auto& RenderTargets = RenderTargetBindingSlots.Output;
-			const auto& DepthStencil = RenderTargetBindingSlots.DepthStencil;
-			const uint32 RenderTargetCount = RenderTargets.Num();
-
-			for (uint32 RenderTargetIndex = 0; RenderTargetIndex < RenderTargetCount; RenderTargetIndex++)
+			break;
+		case UBMT_RDG_BUFFER_SRV:
+			if (FRDGBufferSRVRef SRV = Parameter.GetAsBufferSRV())
 			{
-				const FRenderTargetBinding& RenderTarget = RenderTargets[RenderTargetIndex];
+				SRV->bAllowRHIAccess = bAllowAccess;
+				SRV->GetParent()->bAllowRHIAccess = bAllowAccess;
+			}
+			break;
+		case UBMT_RDG_BUFFER_UAV:
+			if (FRDGBufferUAVRef UAV = Parameter.GetAsBufferUAV())
+			{
+				UAV->bAllowRHIAccess = bAllowAccess;
+				UAV->GetParent()->bAllowRHIAccess = bAllowAccess;
+			}
+			break;
+		case UBMT_RDG_TEXTURE_SRV:
+			if (FRDGTextureSRVRef SRV = Parameter.GetAsTextureSRV())
+			{
+				SRV->bAllowRHIAccess = bAllowAccess;
+				SRV->GetParent()->bAllowRHIAccess = bAllowAccess;
+			}
+			break;
+		case UBMT_RDG_TEXTURE_UAV:
+			if (FRDGTextureUAVRef UAV = Parameter.GetAsTextureUAV())
+			{
+				UAV->bAllowRHIAccess = bAllowAccess;
+				UAV->GetParent()->bAllowRHIAccess = bAllowAccess;
+			}
+			break;
+		case UBMT_RENDER_TARGET_BINDING_SLOTS:
+			{
+				const FRenderTargetBindingSlots& RenderTargetBindingSlots = Parameter.GetAsRenderTargetBindingSlots();
+				const auto& RenderTargets = RenderTargetBindingSlots.Output;
+				const auto& DepthStencil = RenderTargetBindingSlots.DepthStencil;
+				const uint32 RenderTargetCount = RenderTargets.Num();
 
-				if (FRDGTextureRef Texture = RenderTarget.GetTexture())
+				for (uint32 RenderTargetIndex = 0; RenderTargetIndex < RenderTargetCount; RenderTargetIndex++)
 				{
-					Texture->bAllowRHIAccess = bAllowAccess;
+					const FRenderTargetBinding& RenderTarget = RenderTargets[RenderTargetIndex];
+
+					if (FRDGTextureRef Texture = RenderTarget.GetTexture())
+					{
+						Texture->bAllowRHIAccess = bAllowAccess;
+					}
+					else
+					{
+						break;
+					}
 				}
-				else
+
+				if (FRDGTextureRef Texture = DepthStencil.GetTexture())
 				{
-					break;
+					Texture->bAllowRHIAccess = bAllowAccess;
 				}
 			}
-
-			if (FRDGTextureRef Texture = DepthStencil.GetTexture())
-			{
-				Texture->bAllowRHIAccess = bAllowAccess;
-			}
+			break;
 		}
 	}
 }
diff --git a/Engine/Source/Runtime/RenderCore/Public/RenderGraphResources.h b/Engine/Source/Runtime/RenderCore/Public/RenderGraphResources.h
--- a/Engine/Source/Runtime/RenderCore/Public/RenderGraphResources.h
+++ b/Engine/Source/Runtime/RenderCore/Public/RenderGraphResources.h
@@ -381,6 +381,11 @@ public:
 	/** Descriptor of the graph tracked SRV. */
 	const FRDGTextureSRVDesc Desc;
 
+	FRDGTextureRef GetParent() const
+	{
+		return Desc.Texture;
+	}
+
 private:
 	FRDGTextureSRV(
 		const TCHAR* InName,
@@ -417,6 +422,11 @@ public:
 	/** Descriptor of the graph tracked UAV. */
 	const FRDGTextureUAVDesc Desc;
 
+	FRDGTextureRef GetParent() const
+	{
+		return Desc.Texture;
+	}
+
 private:
 	FRDGTextureUAV(
 		const TCHAR* InName,
@@ -715,6 +725,11 @@ public:
 	/** Descriptor of the graph tracked SRV. */
 	const FRDGBufferSRVDesc Desc;
 
+	FRDGBufferRef GetParent() const
+	{
+		return Desc.Buffer;
+	}
+
 private:
 	FRDGBufferSRV(const TCHAR* InName, const FRDGBufferSRVDesc& InDesc)
 		: FRDGShaderResourceView(InName)
@@ -732,6 +747,11 @@ public:
 	/** Descriptor of the graph tracked UAV. */
 	const FRDGBufferUAVDesc Desc;
 
+	FRDGBufferRef GetParent() const
+	{
+		return Desc.Buffer;
+	}
+
 private:
 	FRDGBufferUAV(const TCHAR* InName, const FRDGBufferUAVDesc& InDesc)
 		: FRDGUnorderedAccessView(InName)
diff --git a/Engine/Source/Runtime/Renderer/Private/LightGridInjection.cpp b/Engine/Source/Runtime/Renderer/Private/LightGridInjection.cpp
--- a/Engine/Source/Runtime/Renderer/Private/LightGridInjection.cpp
+++ b/Engine/Source/Runtime/Renderer/Private/LightGridInjection.cpp
@@ -73,36 +73,6 @@ FAutoConsoleVariableRef CVarLightCullingQuality(
 	ECVF_RenderThreadSafe
 );
 
-
-// TODO move to render graph utils
-BEGIN_SHADER_PARAMETER_STRUCT(FClearUAVParameters, )
-	SHADER_PARAMETER_RDG_TEXTURE_UAV(	RWTexture2D,	TextureUAV)
-	SHADER_PARAMETER_RDG_BUFFER_UAV(	RWBuffer<uint>,	BufferUAV)
-END_SHADER_PARAMETER_STRUCT()
-
-void AddPass_ClearUAV(
-	FRDGBuilder& GraphBuilder,
-	FRDGEventName&& PassName,
-	FRDGBufferUAVRef BufferUAV,
-	uint32 Value)
-{
-	FClearUAVParameters* Parameters = GraphBuilder.AllocParameters< FClearUAVParameters >();
-	Parameters->BufferUAV = BufferUAV;
-
-	GraphBuilder.AddPass(
-		Forward<FRDGEventName>(PassName),
-		Parameters,
-		ERDGPassFlags::Compute,
-		[&Parameters, BufferUAV, Value](FRHICommandList& RHICmdList)
-		{
-			BufferUAV->MarkResourceAsUsed();
-			ClearUAV( RHICmdList, BufferUAV->GetRHI(), BufferUAV->Desc.Buffer->Desc.GetTotalNumBytes(), Value );
-		} );
-}
-
-
-
-
 /** A minimal forwarding lighting setup. */
 class FMinimalDummyForwardLightingResources : public FRenderResource
 {
@@ -754,9 +724,9 @@ void FDeferredShadingSceneRenderer::ComputeLightGrid(FRHICommandListImmediate& R
 
 					if (GLightLinkedListCulling != 0)
 					{
-						AddPass_ClearUAV(GraphBuilder, RDG_EVENT_NAME("Clear:StartOffsetGrid"), PassParameters->RWStartOffsetGrid, 0xFFFFFFFF);
-						AddPass_ClearUAV(GraphBuilder, RDG_EVENT_NAME("Clear:NextCulledLightLink"), PassParameters->RWNextCulledLightLink, 0);
-						AddPass_ClearUAV(GraphBuilder, RDG_EVENT_NAME("Clear:NextCulledLightData"), GraphBuilder.CreateUAV(NextCulledLightDataBuffer, PF_R32_UINT), 0);
+						ClearUAV(GraphBuilder, RDG_EVENT_NAME("Clear:StartOffsetGrid"), PassParameters->RWStartOffsetGrid, 0xFFFFFFFF);
+						ClearUAV(GraphBuilder, RDG_EVENT_NAME("Clear:NextCulledLightLink"), PassParameters->RWNextCulledLightLink, 0);
+						ClearUAV(GraphBuilder, RDG_EVENT_NAME("Clear:NextCulledLightData"), GraphBuilder.CreateUAV(NextCulledLightDataBuffer, PF_R32_UINT), 0);
 						FComputeShaderUtils::AddPass(GraphBuilder, RDG_EVENT_NAME("LightGridInject:LinkedList"), *ComputeShader, PassParameters, NumGroups);
 
 
diff --git a/Engine/Source/Runtime/Renderer/Private/PostProcess/DiaphragmDOF.cpp b/Engine/Source/Runtime/Renderer/Private/PostProcess/DiaphragmDOF.cpp
--- a/Engine/Source/Runtime/Renderer/Private/PostProcess/DiaphragmDOF.cpp
+++ b/Engine/Source/Runtime/Renderer/Private/PostProcess/DiaphragmDOF.cpp
@@ -20,14 +20,6 @@
 #include "PostProcessTemporalAA.h"
 #include "SceneTextureParameters.h"
 
-
-void AddPass_ClearUAV(
-	FRDGBuilder& GraphBuilder,
-	FRDGEventName&& PassName,
-	FRDGBufferUAVRef BufferUAV,
-	uint32 Value);
-
-
 // ---------------------------------------------------- Cvars
 
 namespace
@@ -1901,7 +1893,7 @@ FRDGTextureRef DiaphragmDOF::AddPasses(
 				if (BackgroundScatterDrawListBuffer)
 					PassParameters->OutBackgroundScatterDrawList = GraphBuilder.CreateUAV(BackgroundScatterDrawListBuffer);
 
-				AddPass_ClearUAV(GraphBuilder, RDG_EVENT_NAME("ClearIndirectDraw"), PassParameters->OutScatterDrawIndirectParameters, 0);
+				ClearUAV(GraphBuilder, RDG_EVENT_NAME("ClearIndirectDraw"), PassParameters->OutScatterDrawIndirectParameters, 0);
 			}
 
 			TShaderMapRef<FDiaphragmDOFReduceCS> ComputeShader(View.ShaderMap, PermutationVector);
diff --git a/Engine/Source/Runtime/UtilityShaders/Private/ClearQuad.cpp b/Engine/Source/Runtime/UtilityShaders/Private/ClearQuad.cpp
--- a/Engine/Source/Runtime/UtilityShaders/Private/ClearQuad.cpp
+++ b/Engine/Source/Runtime/UtilityShaders/Private/ClearQuad.cpp
@@ -369,3 +369,60 @@ void DrawClearQuadMRT(FRHICommandList& RHICmdList, bool bClearColor, int32 NumCl
 		RHICmdList.DrawPrimitive(0, 2, 1);
 	}
 }
+
+BEGIN_SHADER_PARAMETER_STRUCT(FClearBufferUAVParameters, )
+	SHADER_PARAMETER_RDG_BUFFER_UAV(RWBuffer<uint>, BufferUAV)
+END_SHADER_PARAMETER_STRUCT()
+
+void ClearUAV(FRDGBuilder& GraphBuilder, FRDGEventName&& PassName, FRDGBufferUAVRef BufferUAV, uint32 Value)
+{
+	FClearBufferUAVParameters* Parameters = GraphBuilder.AllocParameters<FClearBufferUAVParameters>();
+	Parameters->BufferUAV = BufferUAV;
+
+	GraphBuilder.AddPass(
+		Forward<FRDGEventName>(PassName),
+		Parameters,
+		ERDGPassFlags::Compute,
+		[&Parameters, BufferUAV, Value](FRHICommandList& RHICmdList)
+	{
+		// This might be called if using ClearTinyUAV.
+		BufferUAV->MarkResourceAsUsed();
+
+		::ClearUAV(RHICmdList, BufferUAV->GetRHI(), BufferUAV->GetParent()->Desc.GetTotalNumBytes(), Value);
+	});
+}
+
+BEGIN_SHADER_PARAMETER_STRUCT(FClearTextureUAVParameters, )
+	SHADER_PARAMETER_RDG_TEXTURE_UAV(RWTexture2D, TextureUAV)
+END_SHADER_PARAMETER_STRUCT()
+
+template <typename T>
+void ClearUAV_T(FRDGBuilder& GraphBuilder, FRDGEventName&& PassName, FRDGTextureUAVRef TextureUAV, const T(&ClearValues)[4])
+{
+	FClearTextureUAVParameters* Parameters = GraphBuilder.AllocParameters<FClearTextureUAVParameters>();
+	Parameters->TextureUAV = TextureUAV;
+
+	GraphBuilder.AddPass(
+		Forward<FRDGEventName>(PassName),
+		Parameters,
+		ERDGPassFlags::Compute,
+		[&Parameters, TextureUAV, ClearValues](FRHICommandList& RHICmdList)
+	{
+		::ClearUAV(RHICmdList, TextureUAV->GetParent()->GetRHI(), TextureUAV->GetRHI(), ClearValues);
+	});
+}
+
+void ClearUAV(FRDGBuilder& GraphBuilder, FRDGEventName&& PassName, FRDGTextureUAVRef TextureUAV, const float(&ClearValues)[4])
+{
+	ClearUAV_T(GraphBuilder, Forward<FRDGEventName>(PassName), TextureUAV, ClearValues);
+}
+
+void ClearUAV(FRDGBuilder& GraphBuilder, FRDGEventName&& PassName, FRDGTextureUAVRef TextureUAV, const uint32(&ClearValues)[4])
+{
+	ClearUAV_T(GraphBuilder, Forward<FRDGEventName>(PassName), TextureUAV, ClearValues);
+}
+
+void ClearUAV(FRDGBuilder& GraphBuilder, FRDGEventName&& PassName, FRDGTextureUAVRef TextureUAV, const FLinearColor& ClearColor)
+{
+	ClearUAV_T(GraphBuilder, Forward<FRDGEventName>(PassName), TextureUAV, reinterpret_cast<const float(&)[4]>(ClearColor));
+}
diff --git a/Engine/Source/Runtime/UtilityShaders/Public/ClearQuad.h b/Engine/Source/Runtime/UtilityShaders/Public/ClearQuad.h
--- a/Engine/Source/Runtime/UtilityShaders/Public/ClearQuad.h
+++ b/Engine/Source/Runtime/UtilityShaders/Public/ClearQuad.h
@@ -5,6 +5,7 @@
 #include "CoreMinimal.h"
 #include "RHIDefinitions.h"
 #include "RenderResource.h"
+#include "RenderGraphBuilder.h"
 
 class FRHICommandList;
 struct FRWBufferStructured;
@@ -57,7 +58,6 @@ extern UTILITYSHADERS_API void ClearUAV(FRHICommandList& RHICmdList, const FScen
 
 extern UTILITYSHADERS_API void ClearTexture2DUAV(FRHICommandList& RHICmdList, FRHIUnorderedAccessView* UAV, int32 Width, int32 Height, const FLinearColor& ClearColor);
 
-
 extern UTILITYSHADERS_API void DrawClearQuadMRT( FRHICommandList& RHICmdList, bool bClearColor, int32 NumClearColors, const FLinearColor* ClearColorArray, bool bClearDepth, float Depth, bool bClearStencil, uint32 Stencil );
 extern UTILITYSHADERS_API void DrawClearQuadMRT(FRHICommandList& RHICmdList, bool bClearColor, int32 NumClearColors, const FLinearColor* ClearColorArray, bool bClearDepth, float Depth, bool bClearStencil, uint32 Stencil, FClearQuadCallbacks ClearQuadCallbacks);
 extern UTILITYSHADERS_API void DrawClearQuadMRT( FRHICommandList& RHICmdList, bool bClearColor, int32 NumClearColors, const FLinearColor* ClearColorArray, bool bClearDepth, float Depth, bool bClearStencil, uint32 Stencil, FIntPoint ViewSize, FIntRect ExcludeRect );
@@ -81,3 +81,8 @@ inline void DrawClearQuad(FRHICommandList& RHICmdList, const FLinearColor& Color
 {
 	DrawClearQuadMRT(RHICmdList, true, 1, &Color, false, 0, false, 0, ClearQuadCallbacks);
 }
+
+extern UTILITYSHADERS_API void ClearUAV(FRDGBuilder& GraphBuilder, FRDGEventName&& PassName, FRDGBufferUAVRef BufferUAV, uint32 Value);
+extern UTILITYSHADERS_API void ClearUAV(FRDGBuilder& GraphBuilder, FRDGEventName&& PassName, FRDGTextureUAVRef TextureUAV, const float(&ClearValues)[4]);
+extern UTILITYSHADERS_API void ClearUAV(FRDGBuilder& GraphBuilder, FRDGEventName&& PassName, FRDGTextureUAVRef TextureUAV, const uint32(&ClearValues)[4]);
+extern UTILITYSHADERS_API void ClearUAV(FRDGBuilder& GraphBuilder, FRDGEventName&& PassName, FRDGTextureUAVRef TextureUAV, const FLinearColor& ClearColor);
\ No newline at end of file
