Fixed static analysis warning.

diff --git a/Engine/Source/Runtime/RenderCore/Public/RenderGraphBuilder.inl b/Engine/Source/Runtime/RenderCore/Public/RenderGraphBuilder.inl
--- a/Engine/Source/Runtime/RenderCore/Public/RenderGraphBuilder.inl
+++ b/Engine/Source/Runtime/RenderCore/Public/RenderGraphBuilder.inl
@@ -97,7 +97,9 @@ inline FRDGBufferRef FRDGBuilder::CreateBuffer(
 		checkf(Name, TEXT("Creating a render graph buffer requires a valid debug name."));
 		Validation.ExecuteGuard(TEXT("CreateBuffer"), Name);
 
-		if (Desc.UnderlyingType == FRDGBufferDesc::EUnderlyingType::StructuredBuffer && (Desc.Usage & BUF_ByteAddressBuffer))
+		const bool bIsByteAddress = (Desc.Usage & BUF_ByteAddressBuffer) == BUF_ByteAddressBuffer;
+
+		if (bIsByteAddress && Desc.UnderlyingType == FRDGBufferDesc::EUnderlyingType::StructuredBuffer)
 		{
 			checkf(Desc.BytesPerElement == 4, TEXT("Creating buffer '%s' as a structured buffer that is also byte addressable, BytesPerElement must be 4! Instead it is %d"), Name, Desc.BytesPerElement);
 		}
