Minor cleanup of RDG internals.

diff --git a/Engine/Source/Runtime/RenderCore/Private/RenderGraphBarrierBatcher.cpp b/Engine/Source/Runtime/RenderCore/Private/RenderGraphBarrierBatcher.cpp
--- a/Engine/Source/Runtime/RenderCore/Private/RenderGraphBarrierBatcher.cpp
+++ b/Engine/Source/Runtime/RenderCore/Private/RenderGraphBarrierBatcher.cpp
@@ -155,20 +155,20 @@ void FRDGBarrierBatcher::QueueTransitionTexture(FRDGTexture* Texture, FRDGResour
 
 void FRDGBarrierBatcher::QueueTransitionUAV(
 	FRHIUnorderedAccessView* UAV,
-	FRDGTrackedResource* UnderlyingResource,
+	FRDGParentResource* ParentResource,
 	FRDGResourceState::EAccess AccessAfter)
 {
 	check(UAV);
-	check(UnderlyingResource);
+	check(ParentResource);
 
-	const FRDGResourceState StateBefore = UnderlyingResource->State;
+	const FRDGResourceState StateBefore = ParentResource->State;
 	const FRDGResourceState StateAfter(Pass, Pipeline, AccessAfter);
 
-	ValidateTransition(UnderlyingResource, StateBefore, StateAfter);
+	ValidateTransition(ParentResource, StateBefore, StateAfter);
 
 	if (StateBefore != StateAfter)
 	{
-		const bool bIsMultiFrameResource = (UnderlyingResource->Flags & ERDGResourceFlags::MultiFrame) == ERDGResourceFlags::MultiFrame;
+		const bool bIsMultiFrameResource = (ParentResource->Flags & ERDGResourceFlags::MultiFrame) == ERDGResourceFlags::MultiFrame;
 
 		if (bIsMultiFrameResource && IsWriteAccessBegin(StateBefore.Access, StateAfter.Access))
 		{
@@ -199,11 +199,11 @@ void FRDGBarrierBatcher::QueueTransitionUAV(
 			UAVUpdateMultiFrameEnds.AddUnique(UAV);
 		}
 
-		UnderlyingResource->State = StateAfter;
+		ParentResource->State = StateAfter;
 	}
 }
 
-void FRDGBarrierBatcher::ValidateTransition(const FRDGTrackedResource* Resource, FRDGResourceState StateBefore, FRDGResourceState StateAfter)
+void FRDGBarrierBatcher::ValidateTransition(const FRDGParentResource* Resource, FRDGResourceState StateBefore, FRDGResourceState StateAfter)
 {
 #if RDG_ENABLE_DEBUG
 	check(StateAfter.Pipeline != FRDGResourceState::EPipeline::MAX);
diff --git a/Engine/Source/Runtime/RenderCore/Private/RenderGraphBarrierBatcher.h b/Engine/Source/Runtime/RenderCore/Private/RenderGraphBarrierBatcher.h
--- a/Engine/Source/Runtime/RenderCore/Private/RenderGraphBarrierBatcher.h
+++ b/Engine/Source/Runtime/RenderCore/Private/RenderGraphBarrierBatcher.h
@@ -23,7 +23,7 @@ public:
 	 */
 	void QueueTransitionUAV(
 		FRHIUnorderedAccessView* UAV,
-		FRDGTrackedResource* UnderlyingResource,
+		FRDGParentResource* ParentResource,
 		FRDGResourceState::EAccess AccessAfter);
 
 private:
@@ -70,7 +70,7 @@ private:
 	FUAVBatch UAVUpdateMultiFrameEnds;
 	FUAVBatchMap UAVBatchMap;
 
-	void ValidateTransition(const FRDGTrackedResource* Resource, FRDGResourceState StateBefore, FRDGResourceState StateAfter);
+	void ValidateTransition(const FRDGParentResource* Resource, FRDGResourceState StateBefore, FRDGResourceState StateAfter);
 
 	EResourceTransitionAccess GetResourceTransitionAccess(FRDGResourceState::EAccess AccessAfter) const;
 	EResourceTransitionAccess GetResourceTransitionAccessForUAV(FRDGResourceState::EAccess AccessBefore, FRDGResourceState::EAccess AccessAfter) const;
diff --git a/Engine/Source/Runtime/RenderCore/Private/RenderGraphBuilder.cpp b/Engine/Source/Runtime/RenderCore/Private/RenderGraphBuilder.cpp
--- a/Engine/Source/Runtime/RenderCore/Private/RenderGraphBuilder.cpp
+++ b/Engine/Source/Runtime/RenderCore/Private/RenderGraphBuilder.cpp
@@ -273,52 +273,21 @@ void FRDGBuilder::WalkGraphDependencies()
 		{
 			FRDGPassParameter Parameter = ParameterStruct.GetParameter(ParameterIndex);
 
-			switch (Parameter.GetType())
+			if (Parameter.IsParentResource())
 			{
-			case UBMT_RDG_TEXTURE:
-			case UBMT_RDG_TEXTURE_COPY_DEST:
-			case UBMT_RDG_BUFFER:
-			case UBMT_RDG_BUFFER_COPY_DEST:
-			{
-				if (FRDGTrackedResourceRef Resource = Parameter.GetAsTrackedResource())
+				if (FRDGParentResourceRef Resource = Parameter.GetAsParentResource())
 				{
 					Resource->ReferenceCount++;
 				}
 			}
-			break;
-			case UBMT_RDG_TEXTURE_SRV:
-			{
-				if (FRDGTextureSRVRef SRV = Parameter.GetAsTextureSRV())
-				{
-					SRV->Desc.Texture->ReferenceCount++;
-				}
-			}
-			break;
-			case UBMT_RDG_TEXTURE_UAV:
-			{
-				if (FRDGTextureUAVRef UAV = Parameter.GetAsTextureUAV())
-				{
-					UAV->Desc.Texture->ReferenceCount++;
-				}
-			}
-			break;
-			case UBMT_RDG_BUFFER_SRV:
-			{
-				if (FRDGBufferSRVRef SRV = Parameter.GetAsBufferSRV())
-				{
-					SRV->Desc.Buffer->ReferenceCount++;
-				}
-			}
-			break;
-			case UBMT_RDG_BUFFER_UAV:
+			else if (Parameter.IsChildResource())
 			{
-				if (FRDGBufferUAVRef UAV = Parameter.GetAsBufferUAV())
+				if (FRDGChildResourceRef Resource = Parameter.GetAsChildResource())
 				{
-					UAV->Desc.Buffer->ReferenceCount++;
+					Resource->GetParent()->ReferenceCount++;
 				}
 			}
-			break;
-			case UBMT_RENDER_TARGET_BINDING_SLOTS:
+			else if (Parameter.IsRenderTargetBindingSlots())
 			{
 				const FRenderTargetBindingSlots& RenderTargetBindingSlots = Parameter.GetAsRenderTargetBindingSlots();
 				const auto& DepthStencil = RenderTargetBindingSlots.DepthStencil;
@@ -344,10 +313,6 @@ void FRDGBuilder::WalkGraphDependencies()
 					Texture->ReferenceCount++;
 				}
 			}
-			break;
-			default:
-				break;
-			}
 		}
 	} 
 
diff --git a/Engine/Source/Runtime/RenderCore/Private/RenderGraphValidation.cpp b/Engine/Source/Runtime/RenderCore/Private/RenderGraphValidation.cpp
--- a/Engine/Source/Runtime/RenderCore/Private/RenderGraphValidation.cpp
+++ b/Engine/Source/Runtime/RenderCore/Private/RenderGraphValidation.cpp
@@ -214,7 +214,7 @@ void FRDGUserValidation::ValidateCreateExternalBuffer(FRDGBufferRef Buffer)
 	Buffer->MarkAsExternal();
 }
 
-void FRDGUserValidation::ValidateExtractResource(FRDGTrackedResourceRef Resource)
+void FRDGUserValidation::ValidateExtractResource(FRDGParentResourceRef Resource)
 {
 	check(Resource);
 
@@ -230,7 +230,7 @@ void FRDGUserValidation::ValidateExtractResource(FRDGTrackedResourceRef Resource
 	Resource->PassAccessCount++;
 }
 
-void FRDGUserValidation::RemoveUnusedWarning(FRDGTrackedResourceRef Resource)
+void FRDGUserValidation::RemoveUnusedWarning(FRDGParentResourceRef Resource)
 {
 	check(Resource);
 
@@ -323,7 +323,7 @@ void FRDGUserValidation::ValidateAddPass(const FRDGPass* Pass)
 		{
 			if (FRDGTextureSRVRef SRV = Parameter.GetAsTextureSRV())
 			{
-				FRDGTextureRef Texture = SRV->Desc.Texture;
+				FRDGTextureRef Texture = SRV->GetParent();
 
 				PassResourceValidator.Read(Texture, SRV->Desc.MipLevel);
 
@@ -339,7 +339,7 @@ void FRDGUserValidation::ValidateAddPass(const FRDGPass* Pass)
 		{
 			if (FRDGTextureUAVRef UAV = Parameter.GetAsTextureUAV())
 			{
-				FRDGTextureRef Texture = UAV->Desc.Texture;
+				FRDGTextureRef Texture = UAV->GetParent();
 
 				PassResourceValidator.Write(Texture, UAV->Desc.MipLevel);
 
@@ -366,7 +366,7 @@ void FRDGUserValidation::ValidateAddPass(const FRDGPass* Pass)
 		{
 			if (FRDGBufferSRVRef SRV = Parameter.GetAsBufferSRV())
 			{
-				FRDGBufferRef Buffer = SRV->Desc.Buffer;
+				FRDGBufferRef Buffer = SRV->GetParent();
 
 				PassResourceValidator.Read(Buffer);
 
@@ -382,7 +382,7 @@ void FRDGUserValidation::ValidateAddPass(const FRDGPass* Pass)
 		{
 			if (FRDGBufferUAVRef UAV = Parameter.GetAsBufferUAV())
 			{
-				FRDGBufferRef Buffer = UAV->Desc.Buffer;
+				FRDGBufferRef Buffer = UAV->GetParent();
 
 				PassResourceValidator.Write(Buffer);
 
@@ -577,7 +577,7 @@ void FRDGUserValidation::ValidateExecuteEnd()
 
 	if (IsRDGDebugEnabled())
 	{
-		auto ValidateResourceAtExecuteEnd = [](const FRDGTrackedResourceRef Resource)
+		auto ValidateResourceAtExecuteEnd = [](const FRDGParentResourceRef Resource)
 		{
 			check(Resource->ReferenceCount == 0);
 
@@ -753,72 +753,46 @@ void FRDGUserValidation::SetAllowRHIAccess(const FRDGPass* Pass, bool bAllowAcce
 	{
 		FRDGPassParameter Parameter = ParameterStruct.GetParameter(ParameterIndex);
 
-		switch (Parameter.GetType())
+		if (Parameter.IsParentResource())
 		{
-		case UBMT_RDG_TEXTURE:
-		case UBMT_RDG_TEXTURE_COPY_DEST:
-		case UBMT_RDG_BUFFER:
-		case UBMT_RDG_BUFFER_COPY_DEST:
-			if (FRDGResourceRef Resource = Parameter.GetAsResource())
+			if (FRDGParentResourceRef Resource = Parameter.GetAsParentResource())
 			{
 				Resource->bAllowRHIAccess = bAllowAccess;
 			}
-			break;
-		case UBMT_RDG_BUFFER_SRV:
-			if (FRDGBufferSRVRef SRV = Parameter.GetAsBufferSRV())
-			{
-				SRV->bAllowRHIAccess = bAllowAccess;
-				SRV->GetParent()->bAllowRHIAccess = bAllowAccess;
-			}
-			break;
-		case UBMT_RDG_BUFFER_UAV:
-			if (FRDGBufferUAVRef UAV = Parameter.GetAsBufferUAV())
-			{
-				UAV->bAllowRHIAccess = bAllowAccess;
-				UAV->GetParent()->bAllowRHIAccess = bAllowAccess;
-			}
-			break;
-		case UBMT_RDG_TEXTURE_SRV:
-			if (FRDGTextureSRVRef SRV = Parameter.GetAsTextureSRV())
-			{
-				SRV->bAllowRHIAccess = bAllowAccess;
-				SRV->GetParent()->bAllowRHIAccess = bAllowAccess;
-			}
-			break;
-		case UBMT_RDG_TEXTURE_UAV:
-			if (FRDGTextureUAVRef UAV = Parameter.GetAsTextureUAV())
+		}
+		else if (Parameter.IsChildResource())
+		{
+			if (FRDGChildResourceRef Resource = Parameter.GetAsChildResource())
 			{
-				UAV->bAllowRHIAccess = bAllowAccess;
-				UAV->GetParent()->bAllowRHIAccess = bAllowAccess;
+				Resource->bAllowRHIAccess = bAllowAccess;
+				Resource->GetParent()->bAllowRHIAccess = bAllowAccess;
 			}
-			break;
-		case UBMT_RENDER_TARGET_BINDING_SLOTS:
+		}
+		else if (Parameter.IsRenderTargetBindingSlots())
+		{
+			const FRenderTargetBindingSlots& RenderTargetBindingSlots = Parameter.GetAsRenderTargetBindingSlots();
+			const auto& RenderTargets = RenderTargetBindingSlots.Output;
+			const auto& DepthStencil = RenderTargetBindingSlots.DepthStencil;
+			const uint32 RenderTargetCount = RenderTargets.Num();
+
+			for (uint32 RenderTargetIndex = 0; RenderTargetIndex < RenderTargetCount; RenderTargetIndex++)
 			{
-				const FRenderTargetBindingSlots& RenderTargetBindingSlots = Parameter.GetAsRenderTargetBindingSlots();
-				const auto& RenderTargets = RenderTargetBindingSlots.Output;
-				const auto& DepthStencil = RenderTargetBindingSlots.DepthStencil;
-				const uint32 RenderTargetCount = RenderTargets.Num();
+				const FRenderTargetBinding& RenderTarget = RenderTargets[RenderTargetIndex];
 
-				for (uint32 RenderTargetIndex = 0; RenderTargetIndex < RenderTargetCount; RenderTargetIndex++)
+				if (FRDGTextureRef Texture = RenderTarget.GetTexture())
 				{
-					const FRenderTargetBinding& RenderTarget = RenderTargets[RenderTargetIndex];
-
-					if (FRDGTextureRef Texture = RenderTarget.GetTexture())
-					{
-						Texture->bAllowRHIAccess = bAllowAccess;
-					}
-					else
-					{
-						break;
-					}
+					Texture->bAllowRHIAccess = bAllowAccess;
 				}
-
-				if (FRDGTextureRef Texture = DepthStencil.GetTexture())
+				else
 				{
-					Texture->bAllowRHIAccess = bAllowAccess;
+					break;
 				}
 			}
-			break;
+
+			if (FRDGTextureRef Texture = DepthStencil.GetTexture())
+			{
+				Texture->bAllowRHIAccess = bAllowAccess;
+			}
 		}
 	}
 }
diff --git a/Engine/Source/Runtime/RenderCore/Public/RenderGraphPass.h b/Engine/Source/Runtime/RenderCore/Public/RenderGraphPass.h
--- a/Engine/Source/Runtime/RenderCore/Public/RenderGraphPass.h
+++ b/Engine/Source/Runtime/RenderCore/Public/RenderGraphPass.h
@@ -22,7 +22,17 @@ public:
 		return IsRDGResourceReferenceShaderParameterType(MemberType);
 	}
 
-	bool IsTrackedResource() const
+	bool IsSRV() const
+	{
+		return MemberType == UBMT_RDG_TEXTURE_SRV || MemberType == UBMT_RDG_BUFFER_SRV;
+	}
+
+	bool IsUAV() const
+	{
+		return MemberType == UBMT_RDG_TEXTURE_UAV || MemberType == UBMT_RDG_BUFFER_UAV;
+	}
+
+	bool IsParentResource() const
 	{
 		return
 			MemberType == UBMT_RDG_TEXTURE ||
@@ -31,6 +41,16 @@ public:
 			MemberType == UBMT_RDG_BUFFER_COPY_DEST;
 	}
 
+	bool IsChildResource() const
+	{
+		return IsSRV() || IsUAV();
+	}
+
+	bool IsRenderTargetBindingSlots() const
+	{
+		return MemberType == UBMT_RENDER_TARGET_BINDING_SLOTS;
+	}
+
 	EUniformBufferBaseType GetType() const
 	{
 		return MemberType;
@@ -42,10 +62,28 @@ public:
 		return *GetAs<FRDGResourceRef>();
 	}
 
-	FRDGTrackedResourceRef GetAsTrackedResource() const
+	FRDGParentResourceRef GetAsParentResource() const
+	{
+		check(IsParentResource());
+		return *GetAs<FRDGParentResourceRef>();
+	}
+
+	FRDGChildResourceRef GetAsChildResource() const
+	{
+		check(IsChildResource());
+		return *GetAs<FRDGChildResourceRef>();
+	}
+
+	FRDGShaderResourceViewRef GetAsSRV() const
+	{
+		check(IsSRV());
+		return *GetAs<FRDGShaderResourceViewRef>();
+	}
+
+	FRDGUnorderedAccessViewRef GetAsUAV() const
 	{
-		check(IsTrackedResource());
-		return *GetAs<FRDGTrackedResourceRef>();
+		check(IsUAV());
+		return *GetAs<FRDGUnorderedAccessViewRef>();
 	}
 
 	FRDGTextureRef GetAsTexture() const
@@ -86,7 +124,7 @@ public:
 
 	const FRenderTargetBindingSlots& GetAsRenderTargetBindingSlots() const
 	{
-		check(MemberType == UBMT_RENDER_TARGET_BINDING_SLOTS);
+		check(IsRenderTargetBindingSlots());
 		return *GetAs<FRenderTargetBindingSlots>();
 	}
 
diff --git a/Engine/Source/Runtime/RenderCore/Public/RenderGraphResources.h b/Engine/Source/Runtime/RenderCore/Public/RenderGraphResources.h
--- a/Engine/Source/Runtime/RenderCore/Public/RenderGraphResources.h
+++ b/Engine/Source/Runtime/RenderCore/Public/RenderGraphResources.h
@@ -7,23 +7,44 @@
 class FRDGPass;
 class FRDGBuilder;
 class FRDGEventName;
-class FRDGResource;
-class FRDGTrackedResource;
-class FRDGTexture;
-class FRDGTextureSRV;
-class FRDGTextureUAV;
-class FRDGBuffer;
-class FRDGBufferSRV;
-class FRDGBufferUAV;
 
-/** Defines the RDG resource references for user code not forgetting the const every time. */
+class FRDGResource;
 using FRDGResourceRef = FRDGResource*;
-using FRDGTrackedResourceRef = FRDGTrackedResource*;
+
+class FRDGParentResource;
+using FRDGParentResourceRef = FRDGParentResource*;
+
+UE_DEPRECATED(4.24, "FRDGTrackedResource typedef is deprecated; use FRDGParentResource instead.")
+typedef FRDGParentResource FRDGTrackedResource;
+
+UE_DEPRECATED(4.24, "FRDGTrackedResourceRef typedef is deprecated; use FRDGParentResourceRef instead.")
+typedef FRDGParentResource* FRDGTrackedResourceRef;
+
+class FRDGChildResource;
+using FRDGChildResourceRef = FRDGChildResource*;
+
+class FRDGShaderResourceView;
+using FRDGShaderResourceViewRef = FRDGShaderResourceView*;
+
+class FRDGUnorderedAccessView;
+using FRDGUnorderedAccessViewRef = FRDGUnorderedAccessView*;
+
+class FRDGTexture;
 using FRDGTextureRef = FRDGTexture*;
+
+class FRDGTextureSRV;
 using FRDGTextureSRVRef = FRDGTextureSRV*;
+
+class FRDGTextureUAV;
 using FRDGTextureUAVRef = FRDGTextureUAV*;
+
+class FRDGBuffer;
 using FRDGBufferRef = FRDGBuffer*;
+
+class FRDGBufferSRV;
 using FRDGBufferSRVRef = FRDGBufferSRV*;
+
+class FRDGBufferUAV;
 using FRDGBufferUAVRef = FRDGBufferUAV*;
 
 /** Used for tracking resource state during execution. */
@@ -122,6 +143,9 @@ protected:
 		: Name(InName)
 	{}
 
+	// IMPORTANT: This is never actually called. RDG resources are allocated via MemStack. This is necessary to force a vtable.
+	virtual ~FRDGResource() = default;
+
 	/** Verify that the RHI resource can be accessed at a pass execution. */
 	void ValidateRHIAccess() const
 	{
@@ -155,8 +179,8 @@ private:
 	friend class FRDGUserValidation;
 };
 
-/** A render graph resource with an allocation lifetime tracked by the graph. */
-class RENDERCORE_API FRDGTrackedResource
+/** A render graph resource with an allocation lifetime tracked by the graph. May have child resources which reference it (e.g. views). */
+class RENDERCORE_API FRDGParentResource
 	: public FRDGResource
 {
 public:
@@ -164,7 +188,7 @@ public:
 	const ERDGResourceFlags Flags;
 
 protected:
-	FRDGTrackedResource(const TCHAR* InName, const ERDGResourceFlags InFlags)
+	FRDGParentResource(const TCHAR* InName, const ERDGResourceFlags InFlags)
 		: FRDGResource(InName)
 		, Flags(InFlags)
 	{}
@@ -220,12 +244,25 @@ private:
 	friend class FRDGBarrierBatcher;
 };
 
+/** A render graph resource (e.g. a view) which references a single parent resource (e.g. a texture / buffer). Provides an abstract way to access the parent resource. */
+class FRDGChildResource
+	: public FRDGResource
+{
+public:
+	FRDGChildResource(const TCHAR* Name)
+		: FRDGResource(Name)
+	{}
+
+	/** Returns the referenced parent render graph resource. */
+	virtual FRDGParentResourceRef GetParent() const = 0;
+};
+
 /** Descriptor of a graph tracked texture. */
 using FRDGTextureDesc = FPooledRenderTargetDesc;
 
 /** Render graph tracked Texture. */
 class RENDERCORE_API FRDGTexture final
-	: public FRDGTrackedResource
+	: public FRDGParentResource
 {
 public:
 	/** Descriptor of the graph tracked texture. */
@@ -255,7 +292,7 @@ private:
 		const TCHAR* InName,
 		const FPooledRenderTargetDesc& InDesc,
 		ERDGResourceFlags InFlags)
-		: FRDGTrackedResource(InName, InFlags)
+		: FRDGParentResource(InName, InFlags)
 		, Desc(InDesc)
 	{}
 
@@ -285,7 +322,7 @@ private:
 
 /** Render graph tracked SRV. */
 class FRDGShaderResourceView
-	: public FRDGResource
+	: public FRDGChildResource
 {
 public:
 	/** Returns the allocated RHI SRV. */
@@ -296,7 +333,7 @@ public:
 
 protected:
 	FRDGShaderResourceView(const TCHAR* Name)
-		: FRDGResource(Name)
+		: FRDGChildResource(Name)
 	{}
 
 	/** Returns the allocated RHI SRV without access checks. */
@@ -308,7 +345,7 @@ protected:
 
 /** Render graph tracked UAV. */
 class FRDGUnorderedAccessView
-	: public FRDGResource
+	: public FRDGChildResource
 {
 public:
 	/** Returns the allocated RHI UAV. */
@@ -319,7 +356,7 @@ public:
 
 protected:
 	FRDGUnorderedAccessView(const TCHAR* Name)
-		: FRDGResource(Name)
+		: FRDGChildResource(Name)
 	{}
 
 	/** Returns the allocated RHI UAV without access checks. */
@@ -330,7 +367,7 @@ protected:
 };
 
 /** Descriptor for render graph tracked SRV. */
-class RENDERCORE_API FRDGTextureSRVDesc final : public FRHITextureSRVCreateInfo
+class FRDGTextureSRVDesc final : public FRHITextureSRVCreateInfo
 {
 public:
 	FRDGTextureSRVDesc() = default;
@@ -374,14 +411,14 @@ public:
 };
 
 /** Render graph tracked SRV. */
-class RENDERCORE_API FRDGTextureSRV final
+class FRDGTextureSRV final
 	: public FRDGShaderResourceView
 {
 public:
 	/** Descriptor of the graph tracked SRV. */
 	const FRDGTextureSRVDesc Desc;
 
-	FRDGTextureRef GetParent() const
+	FRDGTextureRef GetParent() const override
 	{
 		return Desc.Texture;
 	}
@@ -398,7 +435,7 @@ private:
 };
 
 /** Descriptor for render graph tracked UAV. */
-class RENDERCORE_API FRDGTextureUAVDesc
+class FRDGTextureUAVDesc
 {
 public:
 	FRDGTextureUAVDesc() = default;
@@ -415,14 +452,14 @@ public:
 };
 
 /** Render graph tracked texture UAV. */
-class RENDERCORE_API FRDGTextureUAV final
+class FRDGTextureUAV final
 	: public FRDGUnorderedAccessView
 {
 public:
 	/** Descriptor of the graph tracked UAV. */
 	const FRDGTextureUAVDesc Desc;
 
-	FRDGTextureRef GetParent() const
+	FRDGTextureRef GetParent() const override
 	{
 		return Desc.Texture;
 	}
@@ -439,7 +476,7 @@ private:
 };
 
 /** Descriptor for render graph tracked Buffer. */
-struct RENDERCORE_API FRDGBufferDesc
+struct FRDGBufferDesc
 {
 	// Type of buffers to the RHI
 	// TODO(RDG): refactor RHI to only have one FRHIBuffer.
@@ -542,7 +579,7 @@ struct RENDERCORE_API FRDGBufferDesc
 	EUnderlyingType UnderlyingType = EUnderlyingType::VertexBuffer;
 };
 
-struct RENDERCORE_API FRDGBufferSRVDesc
+struct FRDGBufferSRVDesc
 {
 	FRDGBufferSRVDesc() = default;
 
@@ -564,7 +601,7 @@ struct RENDERCORE_API FRDGBufferSRVDesc
 	EPixelFormat Format = PF_Unknown;
 };
 
-struct RENDERCORE_API FRDGBufferUAVDesc
+struct FRDGBufferUAVDesc
 {
 	FRDGBufferUAVDesc() = default;
 
@@ -677,8 +714,8 @@ private:
 };
 
 /** Render graph tracked buffers. */
-class RENDERCORE_API FRDGBuffer final
-	: public FRDGTrackedResource
+class FRDGBuffer final
+	: public FRDGParentResource
 {
 public:
 	/** Descriptor of the graph. */
@@ -705,7 +742,7 @@ private:
 		const TCHAR* InName,
 		const FRDGBufferDesc& InDesc,
 		ERDGResourceFlags InFlags)
-		: FRDGTrackedResource(InName, InFlags)
+		: FRDGParentResource(InName, InFlags)
 		, Desc(InDesc)
 	{}
 
@@ -718,14 +755,14 @@ private:
 };
 
 /** Render graph tracked buffer SRV. */
-class RENDERCORE_API FRDGBufferSRV final
+class FRDGBufferSRV final
 	: public FRDGShaderResourceView
 {
 public:
 	/** Descriptor of the graph tracked SRV. */
 	const FRDGBufferSRVDesc Desc;
 
-	FRDGBufferRef GetParent() const
+	FRDGBufferRef GetParent() const override
 	{
 		return Desc.Buffer;
 	}
@@ -740,14 +777,14 @@ private:
 };
 
 /** Render graph tracked buffer UAV. */
-class RENDERCORE_API FRDGBufferUAV final
+class FRDGBufferUAV final
 	: public FRDGUnorderedAccessView
 {
 public:
 	/** Descriptor of the graph tracked UAV. */
 	const FRDGBufferUAVDesc Desc;
 
-	FRDGBufferRef GetParent() const
+	FRDGBufferRef GetParent() const override
 	{
 		return Desc.Buffer;
 	}
diff --git a/Engine/Source/Runtime/RenderCore/Public/RenderGraphValidation.h b/Engine/Source/Runtime/RenderCore/Public/RenderGraphValidation.h
--- a/Engine/Source/Runtime/RenderCore/Public/RenderGraphValidation.h
+++ b/Engine/Source/Runtime/RenderCore/Public/RenderGraphValidation.h
@@ -54,7 +54,7 @@ public:
 	void ValidateAllocPassParameters(const void* Parameters);
 
 	/** Validates a resource extraction operation. */
-	void ValidateExtractResource(FRDGTrackedResourceRef Resource);
+	void ValidateExtractResource(FRDGParentResourceRef Resource);
 
 	/** Tracks and validates the addition of a new pass to the graph. */
 	void ValidateAddPass(const FRDGPass* Pass);
@@ -68,7 +68,7 @@ public:
 	void ValidateExecuteEnd();
 
 	/** Removes the 'produced but not used' warning from the requested resource. */
-	void RemoveUnusedWarning(FRDGTrackedResourceRef Resource);
+	void RemoveUnusedWarning(FRDGParentResourceRef Resource);
 
 private:
 	/** Traverses all resources in the pass and marks whether they are externally accessible by user pass implementations. */
