DR - remove some unecessary refcounts when working with shaders

diff --git a/Engine/Source/Runtime/RHI/Public/RHICommandList.h b/Engine/Source/Runtime/RHI/Public/RHICommandList.h
--- a/Engine/Source/Runtime/RHI/Public/RHICommandList.h
+++ b/Engine/Source/Runtime/RHI/Public/RHICommandList.h
@@ -2137,7 +2137,7 @@ public:
 	}
 
 	template <typename TShaderRHI>
-	FORCEINLINE_DEBUGGABLE void SetLocalShaderUniformBuffer(TRefCountPtr<TShaderRHI>& Shader, uint32 BaseIndex, const FLocalUniformBuffer& UniformBuffer)
+	FORCEINLINE_DEBUGGABLE void SetLocalShaderUniformBuffer(const TRefCountPtr<TShaderRHI>& Shader, uint32 BaseIndex, const FLocalUniformBuffer& UniformBuffer)
 	{
 		SetLocalShaderUniformBuffer(Shader.GetReference(), BaseIndex, UniformBuffer);
 	}
@@ -2154,7 +2154,7 @@ public:
 		ALLOC_COMMAND(FRHICommandSetShaderUniformBuffer<TShaderRHI, ECmdList::EGfx>)(Shader, BaseIndex, UniformBuffer);
 	}
 	template <typename TShaderRHI>
-	FORCEINLINE void SetShaderUniformBuffer(TRefCountPtr<TShaderRHI>& Shader, uint32 BaseIndex, FRHIUniformBuffer* UniformBuffer)
+	FORCEINLINE void SetShaderUniformBuffer(const TRefCountPtr<TShaderRHI>& Shader, uint32 BaseIndex, FRHIUniformBuffer* UniformBuffer)
 	{
 		SetShaderUniformBuffer(Shader.GetReference(), BaseIndex, UniformBuffer);
 	}
@@ -2173,7 +2173,7 @@ public:
 		ALLOC_COMMAND(FRHICommandSetShaderParameter<TShaderRHI, ECmdList::EGfx>)(Shader, BufferIndex, BaseIndex, NumBytes, UseValue);
 	}
 	template <typename TShaderRHI>
-	FORCEINLINE void SetShaderParameter(TRefCountPtr<TShaderRHI>& Shader, uint32 BufferIndex, uint32 BaseIndex, uint32 NumBytes, const void* NewValue)
+	FORCEINLINE void SetShaderParameter(const TRefCountPtr<TShaderRHI>& Shader, uint32 BufferIndex, uint32 BaseIndex, uint32 NumBytes, const void* NewValue)
 	{
 		SetShaderParameter(Shader.GetReference(), BufferIndex, BaseIndex, NumBytes, NewValue);
 	}
@@ -2191,7 +2191,7 @@ public:
 	}
 
 	template <typename TShaderRHI>
-	FORCEINLINE_DEBUGGABLE void SetShaderTexture(TRefCountPtr<TShaderRHI>& Shader, uint32 TextureIndex, FRHITexture* Texture)
+	FORCEINLINE_DEBUGGABLE void SetShaderTexture(const TRefCountPtr<TShaderRHI>& Shader, uint32 TextureIndex, FRHITexture* Texture)
 	{
 		SetShaderTexture(Shader.GetReference(), TextureIndex, Texture);
 	}
@@ -2215,7 +2215,7 @@ public:
 	}
 
 	template <typename TShaderRHI>
-	FORCEINLINE_DEBUGGABLE void SetShaderResourceViewParameter(TRefCountPtr<TShaderRHI>& Shader, uint32 SamplerIndex, FRHIShaderResourceView* SRV)
+	FORCEINLINE_DEBUGGABLE void SetShaderResourceViewParameter(const TRefCountPtr<TShaderRHI>& Shader, uint32 SamplerIndex, FRHIShaderResourceView* SRV)
 	{
 		SetShaderResourceViewParameter(Shader.GetReference(), SamplerIndex, SRV);
 	}
@@ -2241,7 +2241,7 @@ public:
 	}
 
 	template <typename TShaderRHI>
-	FORCEINLINE_DEBUGGABLE void SetShaderSampler(TRefCountPtr<TShaderRHI>& Shader, uint32 SamplerIndex, FRHISamplerState* State)
+	FORCEINLINE_DEBUGGABLE void SetShaderSampler(const TRefCountPtr<TShaderRHI>& Shader, uint32 SamplerIndex, FRHISamplerState* State)
 	{
 		SetShaderSampler(Shader.GetReference(), SamplerIndex, State);
 	}
@@ -2256,7 +2256,7 @@ public:
 		ALLOC_COMMAND(FRHICommandSetUAVParameter<FRHIComputeShader, ECmdList::EGfx>)(Shader, UAVIndex, UAV);
 	}
 
-	FORCEINLINE_DEBUGGABLE void SetUAVParameter(TRefCountPtr<FRHIComputeShader>& Shader, uint32 UAVIndex, FRHIUnorderedAccessView* UAV)
+	FORCEINLINE_DEBUGGABLE void SetUAVParameter(const TRefCountPtr<FRHIComputeShader>& Shader, uint32 UAVIndex, FRHIUnorderedAccessView* UAV)
 	{
 		SetUAVParameter(Shader.GetReference(), UAVIndex, UAV);
 	}
@@ -2271,7 +2271,7 @@ public:
 		ALLOC_COMMAND(FRHICommandSetUAVParameter_IntialCount<FRHIComputeShader, ECmdList::EGfx>)(Shader, UAVIndex, UAV, InitialCount);
 	}
 
-	FORCEINLINE_DEBUGGABLE void SetUAVParameter(TRefCountPtr<FRHIComputeShader>& Shader, uint32 UAVIndex, FRHIUnorderedAccessView* UAV, uint32 InitialCount)
+	FORCEINLINE_DEBUGGABLE void SetUAVParameter(const TRefCountPtr<FRHIComputeShader>& Shader, uint32 UAVIndex, FRHIUnorderedAccessView* UAV, uint32 InitialCount)
 	{
 		SetUAVParameter(Shader.GetReference(), UAVIndex, UAV, InitialCount);
 	}
diff --git a/Engine/Source/Runtime/RenderCore/Public/ShaderParameterUtils.h b/Engine/Source/Runtime/RenderCore/Public/ShaderParameterUtils.h
--- a/Engine/Source/Runtime/RenderCore/Public/ShaderParameterUtils.h
+++ b/Engine/Source/Runtime/RenderCore/Public/ShaderParameterUtils.h
@@ -24,7 +24,7 @@ template<typename TBufferStruct> class TUniformBufferRef;
 template<typename ShaderRHIParamRef, class ParameterType, typename TRHICmdList>
 void SetShaderValue(
 	TRHICmdList& RHICmdList,
-	ShaderRHIParamRef Shader,
+	const ShaderRHIParamRef& Shader,
 	const FShaderParameter& Parameter,
 	const ParameterType& Value,
 	uint32 ElementIndex = 0
@@ -53,7 +53,7 @@ void SetShaderValue(
 template<typename ShaderRHIParamRef, class ParameterType>
 void SetShaderValueOnContext(
 	IRHICommandContext& RHICmdListContext,
-	ShaderRHIParamRef Shader,
+	const ShaderRHIParamRef& Shader,
 	const FShaderParameter& Parameter,
 	const ParameterType& Value,
 	uint32 ElementIndex = 0
@@ -85,7 +85,7 @@ void SetShaderValueOnContext(
 template<typename ShaderRHIParamRef>
 void SetShaderValue(
 	FRHICommandList& RHICmdList, 
-	ShaderRHIParamRef Shader,
+	const ShaderRHIParamRef& Shader,
 	const FShaderParameter& Parameter,
 	bool Value,
 	uint32 ElementIndex = 0
@@ -99,7 +99,7 @@ void SetShaderValue(
 template<typename ShaderRHIParamRef>
 void SetShaderValue(
 	FRHIAsyncComputeCommandList& RHICmdList,
-	ShaderRHIParamRef Shader,
+	const ShaderRHIParamRef& Shader,
 	const FShaderParameter& Parameter,
 	bool Value,
 	uint32 ElementIndex = 0
@@ -118,7 +118,7 @@ void SetShaderValue(
 template<typename ShaderRHIParamRef,class ParameterType, typename TRHICmdList>
 void SetShaderValueArray(
 	TRHICmdList& RHICmdList,
-	ShaderRHIParamRef Shader,
+	const ShaderRHIParamRef& Shader,
 	const FShaderParameter& Parameter,
 	const ParameterType* Values,
 	uint32 NumElements,
@@ -147,7 +147,7 @@ void SetShaderValueArray(
 template<typename ShaderRHIParamRef, typename TRHICmdList>
 void SetShaderValueArray(
 	TRHICmdList& RHICmdList,
-	ShaderRHIParamRef Shader,
+	const ShaderRHIParamRef& Shader,
 	const FShaderParameter& Parameter,
 	const bool* Values,
 	uint32 NumElements,
@@ -333,7 +333,7 @@ FORCEINLINE void SetSRVParameter(
 template<typename TRHIShader, typename TRHICmdList>
 FORCEINLINE void SetSRVParameter(
 	TRHICmdList& RHICmdList,
-	TRefCountPtr<TRHIShader> Shader,
+	const TRefCountPtr<TRHIShader>& Shader,
 	const FShaderResourceParameter& Parameter,
 	FRHIShaderResourceView* NewShaderResourceViewRHI
 )
@@ -407,7 +407,7 @@ inline bool SetUAVParameterIfCS(TRHICmdList& RHICmdList, FRHIComputeShader* Shad
 }
 
 template<typename TShaderRHIRef, typename TRHICmdList>
-inline void FRWShaderParameter::SetBuffer(TRHICmdList& RHICmdList, TShaderRHIRef Shader, const FRWBuffer& RWBuffer) const
+inline void FRWShaderParameter::SetBuffer(TRHICmdList& RHICmdList, const TShaderRHIRef& Shader, const FRWBuffer& RWBuffer) const
 {
 	if (!SetUAVParameterIfCS(RHICmdList, Shader, UAVParameter, RWBuffer.UAV))
 	{
@@ -416,7 +416,7 @@ inline void FRWShaderParameter::SetBuffer(TRHICmdList& RHICmdList, TShaderRHIRef
 }
 
 template<typename TShaderRHIRef, typename TRHICmdList>
-inline void FRWShaderParameter::SetBuffer(TRHICmdList& RHICmdList, TShaderRHIRef Shader, const FRWBufferStructured& RWBuffer) const
+inline void FRWShaderParameter::SetBuffer(TRHICmdList& RHICmdList, const TShaderRHIRef& Shader, const FRWBufferStructured& RWBuffer) const
 {
 	if (!SetUAVParameterIfCS(RHICmdList, Shader, UAVParameter, RWBuffer.UAV))
 	{
@@ -425,7 +425,7 @@ inline void FRWShaderParameter::SetBuffer(TRHICmdList& RHICmdList, TShaderRHIRef
 }
 
 template<typename TShaderRHIRef, typename TRHICmdList>
-inline void FRWShaderParameter::SetTexture(TRHICmdList& RHICmdList, TShaderRHIRef Shader, FRHITexture* Texture, FRHIUnorderedAccessView* UAV) const
+inline void FRWShaderParameter::SetTexture(TRHICmdList& RHICmdList, const TShaderRHIRef& Shader, FRHITexture* Texture, FRHIUnorderedAccessView* UAV) const
 {
 	if (!SetUAVParameterIfCS(RHICmdList, Shader, UAVParameter, UAV))
 	{
@@ -444,7 +444,7 @@ inline void FRWShaderParameter::UnsetUAV(TRHICmdList& RHICmdList, FRHIComputeSha
 template<typename TShaderRHIRef>
 inline void SetLocalUniformBufferParameter(
 	FRHICommandList& RHICmdList,
-	TShaderRHIRef Shader,
+	const TShaderRHIRef& Shader,
 	const FShaderUniformBufferParameter& Parameter,
 	const FLocalUniformBuffer& LocalUniformBuffer
 	)
@@ -461,7 +461,7 @@ inline void SetLocalUniformBufferParameter(
 template<typename TShaderRHIRef, typename TRHICmdList>
 inline void SetUniformBufferParameter(
 	TRHICmdList& RHICmdList,
-	TShaderRHIRef Shader,
+	const TShaderRHIRef& Shader,
 	const FShaderUniformBufferParameter& Parameter,
 	FRHIUniformBuffer* UniformBufferRHI
 	)
@@ -480,7 +480,7 @@ inline void SetUniformBufferParameter(
 template<typename TShaderRHIRef, typename TBufferStruct, typename TRHICmdList>
 inline void SetUniformBufferParameter(
 	TRHICmdList& RHICmdList,
-	TShaderRHIRef Shader,
+	const TShaderRHIRef& Shader,
 	const TShaderUniformBufferParameter<TBufferStruct>& Parameter,
 	const TUniformBufferRef<TBufferStruct>& UniformBufferRef
 	)
@@ -499,7 +499,7 @@ inline void SetUniformBufferParameter(
 template<typename TShaderRHIRef, typename TBufferStruct, typename TRHICmdList>
 inline void SetUniformBufferParameter(
 	TRHICmdList& RHICmdList,
-	TShaderRHIRef Shader,
+	const TShaderRHIRef& Shader,
 	const TShaderUniformBufferParameter<TBufferStruct>& Parameter,
 	const TUniformBuffer<TBufferStruct>& UniformBuffer
 	)
@@ -518,7 +518,7 @@ inline void SetUniformBufferParameter(
 template<typename TShaderRHIRef,typename TBufferStruct>
 inline void SetUniformBufferParameterImmediate(
 	FRHICommandList& RHICmdList,
-	TShaderRHIRef Shader,
+	const TShaderRHIRef& Shader,
 	const TShaderUniformBufferParameter<TBufferStruct>& Parameter,
 	const TBufferStruct& UniformBufferValue
 	)
@@ -539,7 +539,7 @@ inline void SetUniformBufferParameterImmediate(
 template<typename TShaderRHIRef,typename TBufferStruct, typename TRHICmdList>
 inline void SetUniformBufferParameterImmediate(
 	TRHICmdList& RHICmdList,
-	TShaderRHIRef Shader,
+	const TShaderRHIRef& Shader,
 	const TShaderUniformBufferParameter<TBufferStruct>& Parameter,
 	const TBufferStruct& UniformBufferValue
 	)
diff --git a/Engine/Source/Runtime/RenderCore/Public/ShaderParameters.h b/Engine/Source/Runtime/RenderCore/Public/ShaderParameters.h
--- a/Engine/Source/Runtime/RenderCore/Public/ShaderParameters.h
+++ b/Engine/Source/Runtime/RenderCore/Public/ShaderParameters.h
@@ -139,13 +139,13 @@ public:
 	}
 
 	template<typename TShaderRHIRef, typename TRHICmdList>
-	inline void SetBuffer(TRHICmdList& RHICmdList, TShaderRHIRef Shader, const FRWBuffer& RWBuffer) const;
+	inline void SetBuffer(TRHICmdList& RHICmdList, const TShaderRHIRef& Shader, const FRWBuffer& RWBuffer) const;
 
 	template<typename TShaderRHIRef, typename TRHICmdList>
-	inline void SetBuffer(TRHICmdList& RHICmdList, TShaderRHIRef Shader, const FRWBufferStructured& RWBuffer) const;
+	inline void SetBuffer(TRHICmdList& RHICmdList, const TShaderRHIRef& Shader, const FRWBufferStructured& RWBuffer) const;
 
 	template<typename TShaderRHIRef, typename TRHICmdList>
-	inline void SetTexture(TRHICmdList& RHICmdList, TShaderRHIRef Shader, FRHITexture* Texture, FRHIUnorderedAccessView* UAV) const;
+	inline void SetTexture(TRHICmdList& RHICmdList, const TShaderRHIRef& Shader, FRHITexture* Texture, FRHIUnorderedAccessView* UAV) const;
 
 	template<typename TRHICmdList>
 	inline void UnsetUAV(TRHICmdList& RHICmdList, FRHIComputeShader* ComputeShader) const;
