Fixed RDG texture transitions for MSAA.

diff --git a/Engine/Source/Runtime/RenderCore/Private/RenderGraphBarrierBatcher.cpp b/Engine/Source/Runtime/RenderCore/Private/RenderGraphBarrierBatcher.cpp
--- a/Engine/Source/Runtime/RenderCore/Private/RenderGraphBarrierBatcher.cpp
+++ b/Engine/Source/Runtime/RenderCore/Private/RenderGraphBarrierBatcher.cpp
@@ -116,7 +116,7 @@ void FRDGBarrierBatcher::QueueTransitionTexture(FRDGTexture* Texture, FRDGResour
 
 	if (StateBefore != StateAfter)
 	{
-		FRHITexture* RHITexture = Texture->GetRHIUnchecked();
+		FRHITexture* RHITexture = Texture->PooledRenderTarget->GetRenderTargetItem().TargetableTexture;
 
 		const bool bIsMultiFrameResource = (Texture->Flags & ERDGResourceFlags::MultiFrame) == ERDGResourceFlags::MultiFrame;
 
