Added helper function to bind shader parameters to legacy parameter map.

diff --git a/Engine/Source/Runtime/RenderCore/Public/ShaderParameterStruct.h b/Engine/Source/Runtime/RenderCore/Public/ShaderParameterStruct.h
--- a/Engine/Source/Runtime/RenderCore/Public/ShaderParameterStruct.h
+++ b/Engine/Source/Runtime/RenderCore/Public/ShaderParameterStruct.h
@@ -11,6 +11,11 @@
 #include "RHI.h"
 #include "RenderGraphResources.h"
 
+template <typename FParameterStruct>
+void BindForLegacyShaderParameters(FShader* Shader, const FShaderParameterMap& ParameterMap, bool bShouldBindEverything = false)
+{
+	Shader->Bindings.BindForLegacyShaderParameters(Shader, ParameterMap, *FParameterStruct::FTypeInfo::GetStructMetadata(), bShouldBindEverything);
+}
 
 /** Tag a shader class to use the structured shader parameters API.
  *
@@ -32,11 +37,7 @@
 	ShaderClass(const ShaderMetaType::CompiledShaderInitializerType& Initializer) \
 		: ShaderParentClass(Initializer) \
 	{ \
-		this->Bindings.BindForLegacyShaderParameters( \
-			this, \
-			Initializer.ParameterMap, \
-			*FParameters::FTypeInfo::GetStructMetadata(), \
-			bShouldBindEverything); \
+		BindForLegacyShaderParameters<FParameters>(this, Initializer.ParameterMap, bShouldBindEverything); \
 	} \
 	\
 	ShaderClass() \
