Added ability to use ERDGResourceFlags for external resources in RDG.

diff --git a/Engine/Source/Runtime/RenderCore/Public/RenderGraphBuilder.h b/Engine/Source/Runtime/RenderCore/Public/RenderGraphBuilder.h
--- a/Engine/Source/Runtime/RenderCore/Public/RenderGraphBuilder.h
+++ b/Engine/Source/Runtime/RenderCore/Public/RenderGraphBuilder.h
@@ -23,27 +23,31 @@ public:
 	/** Register a external texture to be tracked by the render graph. */
 	FRDGTextureRef RegisterExternalTexture(
 		const TRefCountPtr<IPooledRenderTarget>& ExternalPooledTexture,
-		const TCHAR* Name = TEXT("External"));
+		const TCHAR* Name = TEXT("External"),
+		ERDGResourceFlags Flags = ERDGResourceFlags::None);
 
 	/** Variant of RegisterExternalTexture which will returns null (rather than assert) if the external texture is null. */
 	inline FRDGTextureRef TryRegisterExternalTexture(
 		const TRefCountPtr<IPooledRenderTarget>& ExternalPooledTexture,
-		const TCHAR* Name = TEXT("External"))
+		const TCHAR* Name = TEXT("External"),
+		ERDGResourceFlags Flags = ERDGResourceFlags::None)
 	{
-		return ExternalPooledTexture ? RegisterExternalTexture(ExternalPooledTexture, Name) : nullptr;
+		return ExternalPooledTexture ? RegisterExternalTexture(ExternalPooledTexture, Name, Flags) : nullptr;
 	}
 
 	/** Register a external buffer to be tracked by the render graph. */
 	FRDGBufferRef RegisterExternalBuffer(
 		const TRefCountPtr<FPooledRDGBuffer>& ExternalPooledBuffer,
-		const TCHAR* Name = TEXT("External"));
+		const TCHAR* Name = TEXT("External"),
+		ERDGResourceFlags Flags = ERDGResourceFlags::None);
 
 	/** Variant of RegisterExternalBuffer which will return null (rather than assert) if the external buffer is null. */
 	inline FRDGBufferRef TryRegisterExternalBuffer(
 		const TRefCountPtr<FPooledRDGBuffer>& ExternalPooledBuffer,
-		const TCHAR* Name = TEXT("External"))
+		const TCHAR* Name = TEXT("External"),
+		ERDGResourceFlags Flags = ERDGResourceFlags::None)
 	{
-		return ExternalPooledBuffer ? RegisterExternalBuffer(ExternalPooledBuffer, Name) : nullptr;
+		return ExternalPooledBuffer ? RegisterExternalBuffer(ExternalPooledBuffer, Name, Flags) : nullptr;
 	}
 
 	/** Create graph tracked resource from a descriptor with a debug name.
diff --git a/Engine/Source/Runtime/RenderCore/Public/RenderGraphBuilder.inl b/Engine/Source/Runtime/RenderCore/Public/RenderGraphBuilder.inl
--- a/Engine/Source/Runtime/RenderCore/Public/RenderGraphBuilder.inl
+++ b/Engine/Source/Runtime/RenderCore/Public/RenderGraphBuilder.inl
@@ -2,7 +2,10 @@
 
 #pragma once
 
-inline FRDGTextureRef FRDGBuilder::RegisterExternalTexture(const TRefCountPtr<IPooledRenderTarget>& ExternalPooledTexture, const TCHAR* Name)
+inline FRDGTextureRef FRDGBuilder::RegisterExternalTexture(
+	const TRefCountPtr<IPooledRenderTarget>& ExternalPooledTexture,
+	const TCHAR* Name,
+	ERDGResourceFlags Flags)
 {
 #if RDG_ENABLE_DEBUG
 	{
@@ -18,7 +21,7 @@ inline FRDGTextureRef FRDGBuilder::RegisterExternalTexture(const TRefCountPtr<IP
 		return *Texture;
 	}
 
-	FRDGTexture* OutTexture = AllocateForRHILifeTime<FRDGTexture>(Name, ExternalPooledTexture->GetDesc(), ERDGResourceFlags::None);
+	FRDGTexture* OutTexture = AllocateForRHILifeTime<FRDGTexture>(Name, ExternalPooledTexture->GetDesc(), Flags);
 	OutTexture->PooledRenderTarget = ExternalPooledTexture;
 	OutTexture->ResourceRHI = ExternalPooledTexture->GetRenderTargetItem().ShaderResourceTexture;
 	AllocatedTextures.Add(OutTexture, ExternalPooledTexture);
@@ -30,7 +33,10 @@ inline FRDGTextureRef FRDGBuilder::RegisterExternalTexture(const TRefCountPtr<IP
 	return OutTexture;
 }
 
-inline FRDGBufferRef FRDGBuilder::RegisterExternalBuffer(const TRefCountPtr<FPooledRDGBuffer>& ExternalPooledBuffer, const TCHAR* Name)
+inline FRDGBufferRef FRDGBuilder::RegisterExternalBuffer(
+	const TRefCountPtr<FPooledRDGBuffer>& ExternalPooledBuffer,
+	const TCHAR* Name,
+	ERDGResourceFlags Flags)
 {
 #if RDG_ENABLE_DEBUG
 	{
@@ -44,7 +50,7 @@ inline FRDGBufferRef FRDGBuilder::RegisterExternalBuffer(const TRefCountPtr<FPoo
 		return *Buffer;
 	}
 
-	FRDGBuffer* OutBuffer = AllocateForRHILifeTime<FRDGBuffer>(Name, ExternalPooledBuffer->Desc, ERDGResourceFlags::None);
+	FRDGBuffer* OutBuffer = AllocateForRHILifeTime<FRDGBuffer>(Name, ExternalPooledBuffer->Desc, Flags);
 	OutBuffer->PooledBuffer = ExternalPooledBuffer;
 	AllocatedBuffers.Add(OutBuffer, ExternalPooledBuffer);
 
