Fixed ensure in RDG barrier tracking with non-targetable texture.

diff --git a/Engine/Source/Runtime/RenderCore/Private/RenderGraphBarrierBatcher.cpp b/Engine/Source/Runtime/RenderCore/Private/RenderGraphBarrierBatcher.cpp
--- a/Engine/Source/Runtime/RenderCore/Private/RenderGraphBarrierBatcher.cpp
+++ b/Engine/Source/Runtime/RenderCore/Private/RenderGraphBarrierBatcher.cpp
@@ -118,6 +118,12 @@ void FRDGBarrierBatcher::QueueTransitionTexture(FRDGTexture* Texture, FRDGResour
 	{
 		FRHITexture* RHITexture = Texture->PooledRenderTarget->GetRenderTargetItem().TargetableTexture;
 
+		// This particular texture does not have a targetable texture. It's effectively read-only.
+		if (!RHITexture)
+		{
+			return;
+		}
+
 		const bool bIsMultiFrameResource = (Texture->Flags & ERDGResourceFlags::MultiFrame) == ERDGResourceFlags::MultiFrame;
 
 		if (bIsMultiFrameResource && IsWriteAccessBegin(StateBefore.Access, StateAfter.Access))
