Increase LongGPUTask iteration cap to 2000 for new GPUs like 2080Ti to fix ProfileGPU

diff --git a/Engine/Source/Runtime/RenderCore/Private/LongGPUTask.cpp b/Engine/Source/Runtime/RenderCore/Private/LongGPUTask.cpp
--- a/Engine/Source/Runtime/RenderCore/Private/LongGPUTask.cpp
+++ b/Engine/Source/Runtime/RenderCore/Private/LongGPUTask.cpp
@@ -88,7 +88,7 @@ void MeasureLongGPUTaskExecutionTime(FRHICommandListImmediate& RHICmdList)
 
 	if (RHICmdList.GetRenderQueryResult(TimeQueryStart.GetQuery(), StartTime, true) && RHICmdList.GetRenderQueryResult(TimeQueryEnd.GetQuery(), EndTime, true))
 	{
-		NumMeasuredIterationsToAchieve500ms = FMath::Clamp(FMath::FloorToInt(500.0f / ((EndTime - StartTime) / 1000.0f / NumIterationsForMeasurement)), 1, 200);
+		NumMeasuredIterationsToAchieve500ms = FMath::Clamp(FMath::FloorToInt(500.0f / ((EndTime - StartTime) / 1000.0f / NumIterationsForMeasurement)), 1, 2000);
 	}
 	else
 	{
