DR - Remove deprecated {RHI}GeometryShaderWithStreamOutput (no data format changes)

diff --git a/Engine/Source/Runtime/Apple/MetalRHI/Private/MetalShaders.cpp b/Engine/Source/Runtime/Apple/MetalRHI/Private/MetalShaders.cpp
--- a/Engine/Source/Runtime/Apple/MetalRHI/Private/MetalShaders.cpp
+++ b/Engine/Source/Runtime/Apple/MetalRHI/Private/MetalShaders.cpp
@@ -1332,26 +1332,6 @@ FGeometryShaderRHIRef FMetalDynamicRHI::RHICreateGeometryShader(FRHIShaderLibrar
 	}
 }
 
-FGeometryShaderRHIRef FMetalDynamicRHI::RHICreateGeometryShaderWithStreamOutput(const TArray<uint8>& Code, const FStreamOutElementList& ElementList,
-	uint32 NumStrides, const uint32* Strides, int32 RasterizedStream)
-{
-	checkf(0, TEXT("Not supported yet"));
-	return NULL;
-}
-
-FGeometryShaderRHIRef FMetalDynamicRHI::RHICreateGeometryShaderWithStreamOutput(const FStreamOutElementList& ElementList,
-	uint32 NumStrides, const uint32* Strides, int32 RasterizedStream, FRHIShaderLibrary* Library, FSHAHash Hash)
-{
-	@autoreleasepool {
-	
-	checkSlow(Library && Library->IsNativeLibrary() && IsMetalPlatform(Library->GetPlatform()) && Library->GetPlatform() <= GMaxRHIShaderPlatform);
-	
-	FMetalShaderLibrary* MetalLibrary = ResourceCast(Library);
-	return MetalLibrary->CreateGeometryShaderWithStreamOutput(Hash, ElementList, NumStrides, Strides, RasterizedStream);
-	
-	}
-}
-
 FComputeShaderRHIRef FMetalDynamicRHI::RHICreateComputeShader(const TArray<uint8>& Code) 
 {
 	@autoreleasepool {
@@ -1395,14 +1375,7 @@ FGeometryShaderRHIRef FMetalDynamicRHI::CreateGeometryShader_RenderThread(class
 {
 	return RHICreateGeometryShader(Library, Hash);
 }
-FGeometryShaderRHIRef FMetalDynamicRHI::CreateGeometryShaderWithStreamOutput_RenderThread(class FRHICommandListImmediate& RHICmdList, const TArray<uint8>& Code, const FStreamOutElementList& ElementList, uint32 NumStrides, const uint32* Strides, int32 RasterizedStream)
-{
-	return RHICreateGeometryShaderWithStreamOutput(Code, ElementList, NumStrides, Strides, RasterizedStream);
-}
-FGeometryShaderRHIRef FMetalDynamicRHI::CreateGeometryShaderWithStreamOutput_RenderThread(class FRHICommandListImmediate& RHICmdList, const FStreamOutElementList& ElementList, uint32 NumStrides, const uint32* Strides, int32 RasterizedStream, FRHIShaderLibrary* Library, FSHAHash Hash)
-{
-	return RHICreateGeometryShaderWithStreamOutput(ElementList, NumStrides, Strides, RasterizedStream, Library, Hash);
-}
+
 FComputeShaderRHIRef FMetalDynamicRHI::CreateComputeShader_RenderThread(class FRHICommandListImmediate& RHICmdList, const TArray<uint8>& Code)
 {
 	return RHICreateComputeShader(Code);
@@ -1563,12 +1536,6 @@ FGeometryShaderRHIRef FMetalShaderLibrary::CreateGeometryShader(const FSHAHash&
 	return NULL;
 }
 
-FGeometryShaderRHIRef FMetalShaderLibrary::CreateGeometryShaderWithStreamOutput(const FSHAHash& Hash, const FStreamOutElementList& ElementList, uint32 NumStrides, const uint32* Strides, int32 RasterizedStream)
-{
-	checkf(0, TEXT("Not supported yet"));
-	return NULL;
-}
-
 FComputeShaderRHIRef FMetalShaderLibrary::CreateComputeShader(const FSHAHash& Hash)
 {
 	FMetalShadeEntry* Code = Map.HashMap.Find(Hash);
diff --git a/Engine/Source/Runtime/Apple/MetalRHI/Public/MetalRHI.h b/Engine/Source/Runtime/Apple/MetalRHI/Public/MetalRHI.h
--- a/Engine/Source/Runtime/Apple/MetalRHI/Public/MetalRHI.h
+++ b/Engine/Source/Runtime/Apple/MetalRHI/Public/MetalRHI.h
@@ -55,8 +55,6 @@ public:
 	virtual FDomainShaderRHIRef RHICreateDomainShader(FRHIShaderLibrary* Library, FSHAHash Hash) final override;
 	virtual FGeometryShaderRHIRef RHICreateGeometryShader(const TArray<uint8>& Code) final override;
 	virtual FGeometryShaderRHIRef RHICreateGeometryShader(FRHIShaderLibrary* Library, FSHAHash Hash) final override;
-	virtual FGeometryShaderRHIRef RHICreateGeometryShaderWithStreamOutput(const TArray<uint8>& Code, const FStreamOutElementList& ElementList, uint32 NumStrides, const uint32* Strides, int32 RasterizedStream) final override;
-	virtual FGeometryShaderRHIRef RHICreateGeometryShaderWithStreamOutput(const FStreamOutElementList& ElementList, uint32 NumStrides, const uint32* Strides, int32 RasterizedStream, FRHIShaderLibrary* Library, FSHAHash Hash) final override;
 	virtual FComputeShaderRHIRef RHICreateComputeShader(const TArray<uint8>& Code) final override;
 	virtual FComputeShaderRHIRef RHICreateComputeShader(FRHIShaderLibrary* Library, FSHAHash Hash) final override;
 	virtual FBoundShaderStateRHIRef RHICreateBoundShaderState(FRHIVertexDeclaration* VertexDeclaration, FRHIVertexShader* VertexShader, FRHIHullShader* HullShader, FRHIDomainShader* DomainShader, FRHIPixelShader* PixelShader, FRHIGeometryShader* GeometryShader) final override;
@@ -206,8 +204,6 @@ public:
 	virtual FPixelShaderRHIRef CreatePixelShader_RenderThread(class FRHICommandListImmediate& RHICmdList, FRHIShaderLibrary* Library, FSHAHash Hash) final override;
 	virtual FGeometryShaderRHIRef CreateGeometryShader_RenderThread(class FRHICommandListImmediate& RHICmdList, const TArray<uint8>& Code) final override;
 	virtual FGeometryShaderRHIRef CreateGeometryShader_RenderThread(class FRHICommandListImmediate& RHICmdList, FRHIShaderLibrary* Library, FSHAHash Hash) final override;
-	virtual FGeometryShaderRHIRef CreateGeometryShaderWithStreamOutput_RenderThread(class FRHICommandListImmediate& RHICmdList, const TArray<uint8>& Code, const FStreamOutElementList& ElementList, uint32 NumStrides, const uint32* Strides, int32 RasterizedStream) final override;
-	virtual FGeometryShaderRHIRef CreateGeometryShaderWithStreamOutput_RenderThread(class FRHICommandListImmediate& RHICmdList, const FStreamOutElementList& ElementList, uint32 NumStrides, const uint32* Strides, int32 RasterizedStream, FRHIShaderLibrary* Library, FSHAHash Hash) final override;
 	virtual FComputeShaderRHIRef CreateComputeShader_RenderThread(class FRHICommandListImmediate& RHICmdList, const TArray<uint8>& Code) final override;
 	virtual FComputeShaderRHIRef CreateComputeShader_RenderThread(class FRHICommandListImmediate& RHICmdList, FRHIShaderLibrary* Library, FSHAHash Hash) final override;
 	virtual FHullShaderRHIRef CreateHullShader_RenderThread(class FRHICommandListImmediate& RHICmdList, const TArray<uint8>& Code) final override;
diff --git a/Engine/Source/Runtime/Apple/MetalRHI/Public/MetalResources.h b/Engine/Source/Runtime/Apple/MetalRHI/Public/MetalResources.h
--- a/Engine/Source/Runtime/Apple/MetalRHI/Public/MetalResources.h
+++ b/Engine/Source/Runtime/Apple/MetalRHI/Public/MetalResources.h
@@ -1234,7 +1234,6 @@ private:
 	FHullShaderRHIRef CreateHullShader(const FSHAHash& Hash);
 	FDomainShaderRHIRef CreateDomainShader(const FSHAHash& Hash);
 	FGeometryShaderRHIRef CreateGeometryShader(const FSHAHash& Hash);
-	FGeometryShaderRHIRef CreateGeometryShaderWithStreamOutput(const FSHAHash& Hash, const FStreamOutElementList& ElementList, uint32 NumStrides, const uint32* Strides, int32 RasterizedStream);
 	FComputeShaderRHIRef CreateComputeShader(const FSHAHash& Hash);
 	
 private:
diff --git a/Engine/Source/Runtime/D3D12RHI/Private/D3D12PipelineState.h b/Engine/Source/Runtime/D3D12RHI/Private/D3D12PipelineState.h
--- a/Engine/Source/Runtime/D3D12RHI/Private/D3D12PipelineState.h
+++ b/Engine/Source/Runtime/D3D12RHI/Private/D3D12PipelineState.h
@@ -48,6 +48,7 @@ struct FD3D12_GRAPHICS_PIPELINE_STATE_DESC
 	D3D12_SHADER_BYTECODE DS;
 	D3D12_SHADER_BYTECODE HS;
 	D3D12_SHADER_BYTECODE GS;
+	//#todo-RemoveStreamOut
 	D3D12_STREAM_OUTPUT_DESC StreamOutput;
 #if !D3D12_USE_DERIVED_PSO
 	D3D12_BLEND_DESC BlendState;
diff --git a/Engine/Source/Runtime/D3D12RHI/Private/D3D12Pipelinestate.cpp b/Engine/Source/Runtime/D3D12RHI/Private/D3D12Pipelinestate.cpp
--- a/Engine/Source/Runtime/D3D12RHI/Private/D3D12Pipelinestate.cpp
+++ b/Engine/Source/Runtime/D3D12RHI/Private/D3D12Pipelinestate.cpp
@@ -62,6 +62,7 @@ FD3D12LowLevelGraphicsPipelineStateDesc GetLowLevelGraphicsPipelineStateDesc(con
 
 	if (Initializer.BoundShaderState.GeometryShaderRHI)
 	{
+		//#todo-RemoveStreamOut
 		Desc.Desc.StreamOutput = ((FD3D12GeometryShader*) Initializer.BoundShaderState.GeometryShaderRHI)->StreamOutput;
 	}
 
diff --git a/Engine/Source/Runtime/D3D12RHI/Private/D3D12RHIPrivate.h b/Engine/Source/Runtime/D3D12RHI/Private/D3D12RHIPrivate.h
--- a/Engine/Source/Runtime/D3D12RHI/Private/D3D12RHIPrivate.h
+++ b/Engine/Source/Runtime/D3D12RHI/Private/D3D12RHIPrivate.h
@@ -270,7 +270,6 @@ public:
 	virtual FHullShaderRHIRef RHICreateHullShader(const TArray<uint8>& Code) final override;
 	virtual FDomainShaderRHIRef RHICreateDomainShader(const TArray<uint8>& Code) final override;
 	virtual FGeometryShaderRHIRef RHICreateGeometryShader(const TArray<uint8>& Code) final override;
-	virtual FGeometryShaderRHIRef RHICreateGeometryShaderWithStreamOutput(const TArray<uint8>& Code, const FStreamOutElementList& ElementList, uint32 NumStrides, const uint32* Strides, int32 RasterizedStream) final override;
 	virtual FComputeShaderRHIRef RHICreateComputeShader(const TArray<uint8>& Code) override;
 	virtual FComputeFenceRHIRef RHICreateComputeFence(const FName& Name) final override;
 	virtual FGPUFenceRHIRef RHICreateGPUFence(const FName& Name) final override;
@@ -396,11 +395,6 @@ public:
 		return RHICreateGeometryShader(Code);
 	}
 
-	virtual FGeometryShaderRHIRef CreateGeometryShaderWithStreamOutput_RenderThread(class FRHICommandListImmediate& RHICmdList, const TArray<uint8>& Code, const FStreamOutElementList& ElementList, uint32 NumStrides, const uint32* Strides, int32 RasterizedStream) override final
-	{
-		return RHICreateGeometryShaderWithStreamOutput(Code, ElementList, NumStrides, Strides, RasterizedStream);
-	}
-
 	virtual FHullShaderRHIRef CreateHullShader_RenderThread(class FRHICommandListImmediate& RHICmdList, const TArray<uint8>& Code) override final
 	{
 		return RHICreateHullShader(Code);
diff --git a/Engine/Source/Runtime/D3D12RHI/Private/D3D12Shader.h b/Engine/Source/Runtime/D3D12RHI/Private/D3D12Shader.h
--- a/Engine/Source/Runtime/D3D12RHI/Private/D3D12Shader.h
+++ b/Engine/Source/Runtime/D3D12RHI/Private/D3D12Shader.h
@@ -68,10 +68,14 @@ public:
 	TArray<uint8> Code;
 
 	/** The shader's stream output description. */
+	//#todo-RemoveStreamOut
 	D3D12_STREAM_OUTPUT_DESC StreamOutput;
+	//#todo-RemoveStreamOut
 	D3D12_SO_DECLARATION_ENTRY* pStreamOutEntries;
+	//#todo-RemoveStreamOut
 	uint32* pStreamOutStrides;
 
+	//#todo-RemoveStreamOut
 	bool bShaderNeedsStreamOutput;
 
 	FShaderCodePackedResourceCounts ResourceCounts;
diff --git a/Engine/Source/Runtime/D3D12RHI/Private/D3D12Shaders.cpp b/Engine/Source/Runtime/D3D12RHI/Private/D3D12Shaders.cpp
--- a/Engine/Source/Runtime/D3D12RHI/Private/D3D12Shaders.cpp
+++ b/Engine/Source/Runtime/D3D12RHI/Private/D3D12Shaders.cpp
@@ -173,70 +173,6 @@ FGeometryShaderRHIRef FD3D12DynamicRHI::RHICreateGeometryShader(const TArray<uin
 	return Shader;
 }
 
-FGeometryShaderRHIRef FD3D12DynamicRHI::RHICreateGeometryShaderWithStreamOutput(const TArray<uint8>& Code, const FStreamOutElementList& ElementList, uint32 NumStrides, const uint32* Strides, int32 RasterizedStream)
-{
-	FShaderCodeReader ShaderCode(Code);
-
-	FD3D12GeometryShader* Shader = new FD3D12GeometryShader;
-
-	FMemoryReader Ar(Code, true);
-	Ar << Shader->ShaderResourceTable;
-	int32 Offset = Ar.Tell();
-	const uint8* CodePtr = Code.GetData() + Offset;
-	const SIZE_T CodeSize = ShaderCode.GetActualShaderCodeSize() - Offset;
-
-	Shader->StreamOutput.RasterizedStream = RasterizedStream;
-	if (RasterizedStream == -1)
-	{
-		Shader->StreamOutput.RasterizedStream = D3D12_SO_NO_RASTERIZED_STREAM;
-	}
-
-	int32 NumEntries = ElementList.Num();
-	Shader->StreamOutput.NumEntries = NumEntries;
-	Shader->pStreamOutEntries = new D3D12_SO_DECLARATION_ENTRY[NumEntries];
-	Shader->StreamOutput.pSODeclaration = Shader->pStreamOutEntries;
-	if (Shader->pStreamOutEntries == nullptr)
-	{
-		return nullptr;	// Out of memory
-	}
-	for (int32 EntryIndex = 0; EntryIndex < NumEntries; EntryIndex++)
-	{
-		Shader->pStreamOutEntries[EntryIndex].Stream = ElementList[EntryIndex].Stream;
-		Shader->pStreamOutEntries[EntryIndex].SemanticName = ElementList[EntryIndex].SemanticName;
-		Shader->pStreamOutEntries[EntryIndex].SemanticIndex = ElementList[EntryIndex].SemanticIndex;
-		Shader->pStreamOutEntries[EntryIndex].StartComponent = ElementList[EntryIndex].StartComponent;
-		Shader->pStreamOutEntries[EntryIndex].ComponentCount = ElementList[EntryIndex].ComponentCount;
-		Shader->pStreamOutEntries[EntryIndex].OutputSlot = ElementList[EntryIndex].OutputSlot;
-	}
-
-	if (!ReadShaderOptionalData(ShaderCode, *Shader))
-	{
-		return nullptr;
-	}
-
-	// Indicate this shader uses stream output
-	Shader->bShaderNeedsStreamOutput = true;
-
-	// Copy the strides
-	Shader->StreamOutput.NumStrides = NumStrides;
-	Shader->pStreamOutStrides = new uint32[NumStrides];
-	Shader->StreamOutput.pBufferStrides = Shader->pStreamOutStrides;
-	if (Shader->pStreamOutStrides == nullptr)
-	{
-		return nullptr;	// Out of memory
-	}
-	FMemory::Memcpy(Shader->pStreamOutStrides, Strides, sizeof(uint32) * NumStrides);
-
-	Shader->Code = Code;
-
-	D3D12_SHADER_BYTECODE ShaderBytecode;
-	ShaderBytecode.pShaderBytecode = Shader->Code.GetData() + Offset;
-	ShaderBytecode.BytecodeLength = CodeSize;
-	Shader->ShaderBytecode.SetShaderBytecode(ShaderBytecode);
-
-	return Shader;
-}
-
 FComputeShaderRHIRef FD3D12DynamicRHI::RHICreateComputeShader(const TArray<uint8>& Code)
 {
 	FShaderCodeReader ShaderCode(Code);
diff --git a/Engine/Source/Runtime/EmptyRHI/Private/EmptyShaders.cpp b/Engine/Source/Runtime/EmptyRHI/Private/EmptyShaders.cpp
--- a/Engine/Source/Runtime/EmptyRHI/Private/EmptyShaders.cpp
+++ b/Engine/Source/Runtime/EmptyRHI/Private/EmptyShaders.cpp
@@ -53,13 +53,6 @@ FGeometryShaderRHIRef FEmptyDynamicRHI::RHICreateGeometryShader(const TArray<uin
 	return Shader;
 }
 
-FGeometryShaderRHIRef FEmptyDynamicRHI::RHICreateGeometryShaderWithStreamOutput(const TArray<uint8>& Code, const FStreamOutElementList& ElementList,
-	uint32 NumStrides, const uint32* Strides, int32 RasterizedStream)
-{
-	checkf(0, TEXT("Not supported yet"));
-	return NULL;
-}
-
 FComputeShaderRHIRef FEmptyDynamicRHI::RHICreateComputeShader(const TArray<uint8>& Code) 
 { 
 	FEmptyComputeShader* Shader = new FEmptyComputeShader(Code);
diff --git a/Engine/Source/Runtime/EmptyRHI/Public/EmptyRHI.h b/Engine/Source/Runtime/EmptyRHI/Public/EmptyRHI.h
--- a/Engine/Source/Runtime/EmptyRHI/Public/EmptyRHI.h
+++ b/Engine/Source/Runtime/EmptyRHI/Public/EmptyRHI.h
@@ -46,7 +46,6 @@ public:
 	virtual FHullShaderRHIRef RHICreateHullShader(const TArray<uint8>& Code) final override;
 	virtual FDomainShaderRHIRef RHICreateDomainShader(const TArray<uint8>& Code) final override;
 	virtual FGeometryShaderRHIRef RHICreateGeometryShader(const TArray<uint8>& Code) final override;
-	virtual FGeometryShaderRHIRef RHICreateGeometryShaderWithStreamOutput(const TArray<uint8>& Code, const FStreamOutElementList& ElementList, uint32 NumStrides, const uint32* Strides, int32 RasterizedStream) final override;
 	virtual FComputeShaderRHIRef RHICreateComputeShader(const TArray<uint8>& Code) final override;
 	virtual FBoundShaderStateRHIRef RHICreateBoundShaderState(FRHIVertexDeclaration* VertexDeclaration, FRHIVertexShader* VertexShader, FRHIHullShader* HullShader, FRHIDomainShaderFRHIDomainShader*, FRHIPixelShader* PixelShader, FRHIGeometryShader* GeometryShader) final override;
 	virtual FUniformBufferRHIRef RHICreateUniformBuffer(const void* Contents, const FRHIUniformBufferLayout& Layout, EUniformBufferUsage Usage, EUniformBufferValidation Validation) final override;
diff --git a/Engine/Source/Runtime/NullDrv/Public/NullRHI.h b/Engine/Source/Runtime/NullDrv/Public/NullRHI.h
--- a/Engine/Source/Runtime/NullDrv/Public/NullRHI.h
+++ b/Engine/Source/Runtime/NullDrv/Public/NullRHI.h
@@ -66,12 +66,6 @@ public:
 		return new FRHIGeometryShader(); 
 	}
 
-
-	virtual FGeometryShaderRHIRef RHICreateGeometryShaderWithStreamOutput(const TArray<uint8>& Code, const FStreamOutElementList& ElementList, uint32 NumStrides, const uint32* Strides, int32 RasterizedStream) final override
-	{ 
-		return new FRHIGeometryShader(); 
-	}
-
 	virtual FComputeShaderRHIRef RHICreateComputeShader(const TArray<uint8>& Code) final override
 	{ 
 		return new FRHIComputeShader(); 
diff --git a/Engine/Source/Runtime/OpenGLDrv/Private/OpenGLShaders.cpp b/Engine/Source/Runtime/OpenGLDrv/Private/OpenGLShaders.cpp
--- a/Engine/Source/Runtime/OpenGLDrv/Private/OpenGLShaders.cpp
+++ b/Engine/Source/Runtime/OpenGLDrv/Private/OpenGLShaders.cpp
@@ -1491,12 +1491,6 @@ FDomainShaderRHIRef FOpenGLDynamicRHI::RHICreateDomainShader(const TArray<uint8>
 	return CreateProxyShader<FRHIDomainShader, FOpenGLDomainShaderProxy>(Code);
 }
 
-FGeometryShaderRHIRef FOpenGLDynamicRHI::RHICreateGeometryShaderWithStreamOutput(const TArray<uint8>& Code, const FStreamOutElementList& ElementList, uint32 NumStrides, const uint32* Strides, int32 RasterizedStream) 
-{
-	UE_LOG(LogRHI, Fatal,TEXT("OpenGL Render path does not support stream output!"));
-	return NULL;
-}
-
 template<typename RHIType, typename TOGLProxyType>
 RHIType* CreateProxyShader(FRHIShaderLibrary* Library, FSHAHash Hash)
 {
@@ -1545,13 +1539,6 @@ FDomainShaderRHIRef FOpenGLDynamicRHI::RHICreateDomainShader(FRHIShaderLibrary*
 	return CreateProxyShader<FRHIDomainShader, FOpenGLDomainShaderProxy>(Library, Hash);
 }
 
-FGeometryShaderRHIRef FOpenGLDynamicRHI::RHICreateGeometryShaderWithStreamOutput(const FStreamOutElementList& ElementList, uint32 NumStrides, const uint32* Strides, int32 RasterizedStream, FRHIShaderLibrary* Library, FSHAHash Hash)
-{
-	UE_LOG(LogRHI, Fatal, TEXT("OpenGL Render path does not support stream output!"));
-	return NULL;
-}
-
-
 static void MarkShaderParameterCachesDirty(FOpenGLShaderParameterCache* ShaderParameters, bool UpdateCompute)
 {
 	VERIFY_GL_SCOPE();
diff --git a/Engine/Source/Runtime/OpenGLDrv/Public/OpenGLDrv.h b/Engine/Source/Runtime/OpenGLDrv/Public/OpenGLDrv.h
--- a/Engine/Source/Runtime/OpenGLDrv/Public/OpenGLDrv.h
+++ b/Engine/Source/Runtime/OpenGLDrv/Public/OpenGLDrv.h
@@ -350,7 +350,6 @@ public:
 	virtual FHullShaderRHIRef RHICreateHullShader(const TArray<uint8>& Code) final override;
 	virtual FDomainShaderRHIRef RHICreateDomainShader(const TArray<uint8>& Code) final override;
 	virtual FGeometryShaderRHIRef RHICreateGeometryShader(const TArray<uint8>& Code) final override;
-	virtual FGeometryShaderRHIRef RHICreateGeometryShaderWithStreamOutput(const TArray<uint8>& Code, const FStreamOutElementList& ElementList, uint32 NumStrides, const uint32* Strides, int32 RasterizedStream) final override;
 	virtual FComputeShaderRHIRef RHICreateComputeShader(const TArray<uint8>& Code) final override;
 
 	virtual FPixelShaderRHIRef RHICreatePixelShader(FRHIShaderLibrary* Library, FSHAHash Hash) final override;
@@ -358,7 +357,6 @@ public:
 	virtual FHullShaderRHIRef RHICreateHullShader(FRHIShaderLibrary* Library, FSHAHash Hash) final override;
 	virtual FDomainShaderRHIRef RHICreateDomainShader(FRHIShaderLibrary* Library, FSHAHash Hash) final override;
 	virtual FGeometryShaderRHIRef RHICreateGeometryShader(FRHIShaderLibrary* Library, FSHAHash Hash) final override;
-	virtual FGeometryShaderRHIRef RHICreateGeometryShaderWithStreamOutput(const FStreamOutElementList& ElementList, uint32 NumStrides, const uint32* Strides, int32 RasterizedStream, FRHIShaderLibrary* Library, FSHAHash Hash) final override;
 	virtual FComputeShaderRHIRef RHICreateComputeShader(FRHIShaderLibrary* Library, FSHAHash Hash) final override;
 
 	virtual FUniformBufferRHIRef RHICreateUniformBuffer(const void* Contents, const FRHIUniformBufferLayout& Layout, EUniformBufferUsage Usage, EUniformBufferValidation Validation) final override;
@@ -939,11 +937,6 @@ public:
 		return this->RHICreateGeometryShader(Code);
 	}
 
-	virtual FGeometryShaderRHIRef CreateGeometryShaderWithStreamOutput_RenderThread(class FRHICommandListImmediate& RHICmdList, const TArray<uint8>& Code, const FStreamOutElementList& ElementList, uint32 NumStrides, const uint32* Strides, int32 RasterizedStream) final override
-	{
-		return this->RHICreateGeometryShaderWithStreamOutput(Code, ElementList, NumStrides, Strides, RasterizedStream);
-	}
-
 	virtual FComputeShaderRHIRef CreateComputeShader_RenderThread(class FRHICommandListImmediate& RHICmdList, const TArray<uint8>& Code) final override
 	{
 		return this->RHICreateComputeShader(Code);
diff --git a/Engine/Source/Runtime/RHI/Private/RHICommandList.cpp b/Engine/Source/Runtime/RHI/Private/RHICommandList.cpp
--- a/Engine/Source/Runtime/RHI/Private/RHICommandList.cpp
+++ b/Engine/Source/Runtime/RHI/Private/RHICommandList.cpp
@@ -2264,20 +2264,6 @@ FGeometryShaderRHIRef FDynamicRHI::CreateGeometryShader_RenderThread(class FRHIC
 	return GDynamicRHI->RHICreateGeometryShader(Library, Hash);
 }
 
-FGeometryShaderRHIRef FDynamicRHI::CreateGeometryShaderWithStreamOutput_RenderThread(class FRHICommandListImmediate& RHICmdList, const TArray<uint8>& Code, const FStreamOutElementList& ElementList, uint32 NumStrides, const uint32* Strides, int32 RasterizedStream)
-{
-	CSV_SCOPED_TIMING_STAT(RHITStalls, CreateGeometryShader_RenderThread);
-	FScopedRHIThreadStaller StallRHIThread(RHICmdList);
-	return GDynamicRHI->RHICreateGeometryShaderWithStreamOutput(Code, ElementList, NumStrides, Strides, RasterizedStream);
-}
-
-FGeometryShaderRHIRef FDynamicRHI::CreateGeometryShaderWithStreamOutput_RenderThread(class FRHICommandListImmediate& RHICmdList, const FStreamOutElementList& ElementList, uint32 NumStrides, const uint32* Strides, int32 RasterizedStream, FRHIShaderLibrary* Library, FSHAHash Hash)
-{
-	CSV_SCOPED_TIMING_STAT(RHITStalls, CreateGeometryShader_RenderThread);
-	FScopedRHIThreadStaller StallRHIThread(RHICmdList);
-	return GDynamicRHI->RHICreateGeometryShaderWithStreamOutput(ElementList, NumStrides, Strides, RasterizedStream, Library, Hash);
-}
-
 FComputeShaderRHIRef FDynamicRHI::CreateComputeShader_RenderThread(class FRHICommandListImmediate& RHICmdList, const TArray<uint8>& Code)
 {
 	CSV_SCOPED_TIMING_STAT(RHITStalls, CreateGeometryShader_RenderThread);
diff --git a/Engine/Source/Runtime/RHI/Public/DynamicRHI.h b/Engine/Source/Runtime/RHI/Public/DynamicRHI.h
--- a/Engine/Source/Runtime/RHI/Public/DynamicRHI.h
+++ b/Engine/Source/Runtime/RHI/Public/DynamicRHI.h
@@ -269,17 +269,6 @@ public:
 		return nullptr;
 	}
 
-	/** Creates a geometry shader with stream output ability, defined by ElementList. */
-	// FlushType: Wait RHI Thread
-	virtual FGeometryShaderRHIRef RHICreateGeometryShaderWithStreamOutput(const TArray<uint8>& Code, const FStreamOutElementList& ElementList, uint32 NumStrides, const uint32* Strides, int32 RasterizedStream) = 0;
-
-	/** Creates a geometry shader with stream output ability, defined by ElementList. */
-	// FlushType: Wait RHI Thread
-	virtual FGeometryShaderRHIRef RHICreateGeometryShaderWithStreamOutput(const FStreamOutElementList& ElementList, uint32 NumStrides, const uint32* Strides, int32 RasterizedStream, FRHIShaderLibrary* Library, FSHAHash Hash)
-	{
-		return nullptr;
-	}
-
 	// Some RHIs can have pending messages/logs for error tracking, or debug modes
 	virtual void FlushPendingLogs() {}
 
@@ -1124,8 +1113,6 @@ public:
 	virtual FPixelShaderRHIRef CreatePixelShader_RenderThread(class FRHICommandListImmediate& RHICmdList, FRHIShaderLibrary* Library, FSHAHash Hash);
 	virtual FGeometryShaderRHIRef CreateGeometryShader_RenderThread(class FRHICommandListImmediate& RHICmdList, const TArray<uint8>& Code);
 	virtual FGeometryShaderRHIRef CreateGeometryShader_RenderThread(class FRHICommandListImmediate& RHICmdList, FRHIShaderLibrary* Library, FSHAHash Hash);
-	virtual FGeometryShaderRHIRef CreateGeometryShaderWithStreamOutput_RenderThread(class FRHICommandListImmediate& RHICmdList, const TArray<uint8>& Code, const FStreamOutElementList& ElementList, uint32 NumStrides, const uint32* Strides, int32 RasterizedStream);
-	virtual FGeometryShaderRHIRef CreateGeometryShaderWithStreamOutput_RenderThread(class FRHICommandListImmediate& RHICmdList, const FStreamOutElementList& ElementList, uint32 NumStrides, const uint32* Strides, int32 RasterizedStream, FRHIShaderLibrary* Library, FSHAHash Hash);
 	virtual FComputeShaderRHIRef CreateComputeShader_RenderThread(class FRHICommandListImmediate& RHICmdList, const TArray<uint8>& Code);
 	virtual FComputeShaderRHIRef CreateComputeShader_RenderThread(class FRHICommandListImmediate& RHICmdList, FRHIShaderLibrary* Library, FSHAHash Hash);
 	virtual FHullShaderRHIRef CreateHullShader_RenderThread(class FRHICommandListImmediate& RHICmdList, const TArray<uint8>& Code);
diff --git a/Engine/Source/Runtime/RHI/Public/RHICommandList.h b/Engine/Source/Runtime/RHI/Public/RHICommandList.h
--- a/Engine/Source/Runtime/RHI/Public/RHICommandList.h
+++ b/Engine/Source/Runtime/RHI/Public/RHICommandList.h
@@ -3620,20 +3620,6 @@ public:
 		return GDynamicRHI->CreateGeometryShader_RenderThread(*this, Library, Hash);
 	}
 	
-	UE_DEPRECATED(4.23, "Geometry Stream out is deprecated.")
-	FORCEINLINE FGeometryShaderRHIRef CreateGeometryShaderWithStreamOutput(const TArray<uint8>& Code, const FStreamOutElementList& ElementList, uint32 NumStrides, const uint32* Strides, int32 RasterizedStream)
-	{
-		LLM_SCOPE(ELLMTag::Shaders);
-		return GDynamicRHI->CreateGeometryShaderWithStreamOutput_RenderThread(*this, Code, ElementList, NumStrides, Strides, RasterizedStream);
-	}
-	
-	UE_DEPRECATED(4.23, "Geometry Stream out is deprecated.")
-	FORCEINLINE FGeometryShaderRHIRef CreateGeometryShaderWithStreamOutput(const FStreamOutElementList& ElementList, uint32 NumStrides, const uint32* Strides, int32 RasterizedStream, FRHIShaderLibrary* Library, FSHAHash Hash)
-	{
-		LLM_SCOPE(ELLMTag::Shaders);
-		return GDynamicRHI->CreateGeometryShaderWithStreamOutput_RenderThread(*this, ElementList, NumStrides, Strides, RasterizedStream, Library, Hash);
-	}
-	
 	FORCEINLINE FComputeShaderRHIRef CreateComputeShader(const TArray<uint8>& Code)
 	{
 		LLM_SCOPE(ELLMTag::Shaders);
@@ -4580,22 +4566,6 @@ FORCEINLINE FGeometryShaderRHIRef RHICreateGeometryShader(FRHIShaderLibrary* Lib
 	return FRHICommandListExecutor::GetImmediateCommandList().CreateGeometryShader(Library, Hash);
 }
 
-UE_DEPRECATED(4.23, "Geometry Stream out is deprecated.")
-FORCEINLINE FGeometryShaderRHIRef RHICreateGeometryShaderWithStreamOutput(const TArray<uint8>& Code, const FStreamOutElementList& ElementList, uint32 NumStrides, const uint32* Strides, int32 RasterizedStream)
-{
-	PRAGMA_DISABLE_DEPRECATION_WARNINGS
-	return FRHICommandListExecutor::GetImmediateCommandList().CreateGeometryShaderWithStreamOutput(Code, ElementList, NumStrides, Strides, RasterizedStream);
-	PRAGMA_ENABLE_DEPRECATION_WARNINGS
-}
-
-UE_DEPRECATED(4.23, "Geometry Stream out is deprecated.")
-FORCEINLINE FGeometryShaderRHIRef RHICreateGeometryShaderWithStreamOutput(const FStreamOutElementList& ElementList, uint32 NumStrides, const uint32* Strides, int32 RasterizedStream, FRHIShaderLibrary* Library, FSHAHash Hash)
-{
-	PRAGMA_DISABLE_DEPRECATION_WARNINGS
-	return FRHICommandListExecutor::GetImmediateCommandList().CreateGeometryShaderWithStreamOutput(ElementList, NumStrides, Strides, RasterizedStream, Library, Hash);
-	PRAGMA_ENABLE_DEPRECATION_WARNINGS
-}
-
 FORCEINLINE FComputeShaderRHIRef RHICreateComputeShader(const TArray<uint8>& Code)
 {
 	return FRHICommandListExecutor::GetImmediateCommandList().CreateComputeShader(Code);
diff --git a/Engine/Source/Runtime/RHI/Public/RHIValidation.h b/Engine/Source/Runtime/RHI/Public/RHIValidation.h
--- a/Engine/Source/Runtime/RHI/Public/RHIValidation.h
+++ b/Engine/Source/Runtime/RHI/Public/RHIValidation.h
@@ -145,22 +145,6 @@ public:
 		return RHI->RHICreateGeometryShader(Library, Hash);
 	}
 
-	/** Creates a geometry shader with stream output ability, defined by ElementList. */
-	// FlushType: Wait RHI Thread
-	virtual FGeometryShaderRHIRef RHICreateGeometryShaderWithStreamOutput(const TArray<uint8>& Code, const FStreamOutElementList& ElementList, uint32 NumStrides, const uint32* Strides, int32 RasterizedStream) override final
-	{
-		check(RHISupportsGeometryShaders(GMaxRHIShaderPlatform));
-		return RHI->RHICreateGeometryShaderWithStreamOutput(Code, ElementList, NumStrides, Strides, RasterizedStream);
-	}
-
-	/** Creates a geometry shader with stream output ability, defined by ElementList. */
-	// FlushType: Wait RHI Thread
-	virtual FGeometryShaderRHIRef RHICreateGeometryShaderWithStreamOutput(const FStreamOutElementList& ElementList, uint32 NumStrides, const uint32* Strides, int32 RasterizedStream, FRHIShaderLibrary* Library, FSHAHash Hash) override final
-	{
-		check(RHISupportsGeometryShaders(GMaxRHIShaderPlatform));
-		return RHI->RHICreateGeometryShaderWithStreamOutput(ElementList, NumStrides, Strides, RasterizedStream, Library, Hash);
-	}
-
 	// Some RHIs can have pending messages/logs for error tracking, or debug modes
 	virtual void FlushPendingLogs() override final
 	{
@@ -1319,18 +1303,6 @@ public:
 		return RHI->CreateGeometryShader_RenderThread(RHICmdList, Library, Hash);
 	}
 
-	virtual FGeometryShaderRHIRef CreateGeometryShaderWithStreamOutput_RenderThread(class FRHICommandListImmediate& RHICmdList, const TArray<uint8>& Code, const FStreamOutElementList& ElementList, uint32 NumStrides, const uint32* Strides, int32 RasterizedStream) override final
-	{
-		check(RHISupportsGeometryShaders(GMaxRHIShaderPlatform));
-		return RHI->CreateGeometryShaderWithStreamOutput_RenderThread(RHICmdList, Code, ElementList, NumStrides, Strides, RasterizedStream);
-	}
-
-	virtual FGeometryShaderRHIRef CreateGeometryShaderWithStreamOutput_RenderThread(class FRHICommandListImmediate& RHICmdList, const FStreamOutElementList& ElementList, uint32 NumStrides, const uint32* Strides, int32 RasterizedStream, FRHIShaderLibrary* Library, FSHAHash Hash) override final
-	{
-		check(RHISupportsGeometryShaders(GMaxRHIShaderPlatform));
-		return RHI->CreateGeometryShaderWithStreamOutput_RenderThread(RHICmdList, ElementList, NumStrides, Strides, RasterizedStream, Library, Hash);
-	}
-
 	virtual FComputeShaderRHIRef CreateComputeShader_RenderThread(class FRHICommandListImmediate& RHICmdList, const TArray<uint8>& Code) override final
 	{
 		check(RHISupportsComputeShaders(GMaxRHIShaderPlatform));
diff --git a/Engine/Source/Runtime/RenderCore/Private/Shader.cpp b/Engine/Source/Runtime/RenderCore/Private/Shader.cpp
--- a/Engine/Source/Runtime/RenderCore/Private/Shader.cpp
+++ b/Engine/Source/Runtime/RenderCore/Private/Shader.cpp
@@ -1000,6 +1000,8 @@ void FShaderResource::InitRHI()
 	}
 	else if(Target.Frequency == SF_Geometry)
 	{
+		//#todo-RemoveStreamOut
+/*
 		if (SpecificType)
 		{
 			FStreamOutElementList ElementList;
@@ -1014,6 +1016,7 @@ void FShaderResource::InitRHI()
 			PRAGMA_ENABLE_DEPRECATION_WARNINGS
 		}
 		else
+*/
 		{
 			Shader = FShaderCodeLibrary::CreateGeometryShader((EShaderPlatform)Target.Platform, OutputHash, UncompressedCode);
 		}
@@ -1046,6 +1049,7 @@ void FShaderResource::InitRHI()
 		checkNoEntry(); // Unexpected shader target frequency
 	}
 
+		//#todo-RemoveStreamOut
 	if (Target.Frequency != SF_Geometry)
 	{
 		checkf(!SpecificType, TEXT("Only geometry shaders can use GetStreamOutElements, shader type %s"), SpecificType->GetName());
diff --git a/Engine/Source/Runtime/RenderCore/Private/ShaderCodeLibrary.cpp b/Engine/Source/Runtime/RenderCore/Private/ShaderCodeLibrary.cpp
--- a/Engine/Source/Runtime/RenderCore/Private/ShaderCodeLibrary.cpp
+++ b/Engine/Source/Runtime/RenderCore/Private/ShaderCodeLibrary.cpp
@@ -707,29 +707,6 @@ public:
 		return Shader;
 	}
 
-	FGeometryShaderRHIRef CreateGeometryShaderWithStreamOutput(const FSHAHash& Hash, const FStreamOutElementList& ElementList, uint32 NumStrides, const uint32* Strides, int32 RasterizedStream) override final
-	{
-		FGeometryShaderRHIRef Shader;
-
-		int32 Size = 0;
-		bool bWasSync = false;
-		TArray<uint8>* Code = LookupShaderCode(Hash, Size, bWasSync);
-		if (Code)
-		{
-			TArray<uint8> UCode;
-			TArray<uint8>& UncompressedCode = FShaderLibraryHelperUncompressCode(Platform, Size, *Code, UCode);
-			PRAGMA_DISABLE_DEPRECATION_WARNINGS
-			Shader = RHICreateGeometryShaderWithStreamOutput(UncompressedCode, ElementList, NumStrides, Strides, RasterizedStream);
-			PRAGMA_ENABLE_DEPRECATION_WARNINGS
-			CheckShaderCreation(Shader.GetReference(), Hash);
-			if (bWasSync)
-			{
-				ReleaseShaderCode(Hash);
-			}
-		}
-		return Shader;
-	}
-
 	FComputeShaderRHIRef CreateComputeShader(const FSHAHash& Hash) override final
 	{
 		FComputeShaderRHIRef Shader;
@@ -1811,28 +1788,6 @@ public:
 		return Result;
 	}
 
-	FGeometryShaderRHIRef CreateGeometryShaderWithStreamOutput(EShaderPlatform Platform, FSHAHash Hash, const FStreamOutElementList& ElementList, uint32 NumStrides, const uint32* Strides, int32 RasterizedStream)
-	{
-		checkSlow(Platform == GetRuntimeShaderPlatform());
-
-		FGeometryShaderRHIRef Result;
-		FRHIShaderLibrary* ShaderCodeArchive = FindShaderLibrary(Hash);
-		if (ShaderCodeArchive)
-		{
-			PRAGMA_DISABLE_DEPRECATION_WARNINGS
-			if (bNativeFormat || GRHILazyShaderCodeLoading)
-			{
-				Result = RHICreateGeometryShaderWithStreamOutput(ElementList, NumStrides, Strides, RasterizedStream, ShaderCodeArchive, Hash);
-			}
-			else
-			{
-				Result = ((FShaderCodeArchive*)ShaderCodeArchive)->CreateGeometryShaderWithStreamOutput(Hash, ElementList, NumStrides, Strides, RasterizedStream);
-			}
-			PRAGMA_ENABLE_DEPRECATION_WARNINGS
-		}
-		return Result;
-	}
-
 	FHullShaderRHIRef CreateHullShader(EShaderPlatform Platform, FSHAHash Hash)
 	{
 		checkSlow(Platform == GetRuntimeShaderPlatform());
@@ -2344,23 +2299,6 @@ FGeometryShaderRHIRef FShaderCodeLibrary::CreateGeometryShader(EShaderPlatform P
 	return Shader;
 }
 
-FGeometryShaderRHIRef FShaderCodeLibrary::CreateGeometryShaderWithStreamOutput(EShaderPlatform Platform, FSHAHash Hash, const TArray<uint8>& Code, const FStreamOutElementList& ElementList, uint32 NumStrides, const uint32* Strides, int32 RasterizedStream)
-{
-	FGeometryShaderRHIRef Shader;
-	if (FShaderCodeLibraryImpl::Impl && FPlatformProperties::RequiresCookedData())
-	{
-		Shader = FShaderCodeLibraryImpl::Impl->CreateGeometryShaderWithStreamOutput(Platform, Hash, ElementList, NumStrides, Strides, RasterizedStream);
-	}
-	if (!IsValidRef(Shader))
-	{
-		PRAGMA_DISABLE_DEPRECATION_WARNINGS
-		Shader = RHICreateGeometryShaderWithStreamOutput(Code, ElementList, NumStrides, Strides, RasterizedStream);
-		PRAGMA_ENABLE_DEPRECATION_WARNINGS
-	}
-	SafeAssignHash(Shader, Hash);
-	return Shader;
-}
-
 FHullShaderRHIRef FShaderCodeLibrary::CreateHullShader(EShaderPlatform Platform, FSHAHash Hash, TArray<uint8> const& Code)
 {
 	FHullShaderRHIRef Shader;
diff --git a/Engine/Source/Runtime/RenderCore/Public/Shader.h b/Engine/Source/Runtime/RenderCore/Public/Shader.h
--- a/Engine/Source/Runtime/RenderCore/Public/Shader.h
+++ b/Engine/Source/Runtime/RenderCore/Public/Shader.h
@@ -506,9 +506,11 @@ private:
 	/** Original bytecode size, before compression */
 	uint32 UncompressedCodeSize = 0;
 
+	//#todo-RemoveStreamOut
 	/** If not NULL, the shader type this resource must be used with. */
 	class FShaderType* SpecificType;
 
+	//#todo-RemoveStreamOut
 	/** Specific permutation identifier of the shader when SpecificType is non null, ignored otherwise. */
 	int32 SpecificPermutationId;
 
diff --git a/Engine/Source/Runtime/RenderCore/Public/ShaderCodeLibrary.h b/Engine/Source/Runtime/RenderCore/Public/ShaderCodeLibrary.h
--- a/Engine/Source/Runtime/RenderCore/Public/ShaderCodeLibrary.h
+++ b/Engine/Source/Runtime/RenderCore/Public/ShaderCodeLibrary.h
@@ -126,7 +126,6 @@ public:
 	virtual FHullShaderRHIRef CreateHullShader(const FSHAHash& Hash) = 0;
 	virtual FDomainShaderRHIRef CreateDomainShader(const FSHAHash& Hash) = 0;
 	virtual FGeometryShaderRHIRef CreateGeometryShader(const FSHAHash& Hash) = 0;
-	virtual FGeometryShaderRHIRef CreateGeometryShaderWithStreamOutput(const FSHAHash& Hash, const FStreamOutElementList& ElementList, uint32 NumStrides, const uint32* Strides, int32 RasterizedStream) = 0;
 	virtual FComputeShaderRHIRef CreateComputeShader(const FSHAHash& Hash) = 0;
 };
 
@@ -158,9 +157,6 @@ struct RENDERCORE_API FShaderCodeLibrary
 	static FPixelShaderRHIRef CreatePixelShader(EShaderPlatform Platform, FSHAHash Hash, TArray<uint8> const& Code);
 	/** Instantiate or retrieve a geometry shader from the cache for the provided code & hash. */
 	static FGeometryShaderRHIRef CreateGeometryShader(EShaderPlatform Platform, FSHAHash Hash, TArray<uint8> const& Code);
-	/** Instantiate or retrieve a geometry shader from the cache for the provided code & hash. */
-	UE_DEPRECATED(4.23, "Geometry Stream out is deprecated.")
-	static FGeometryShaderRHIRef CreateGeometryShaderWithStreamOutput(EShaderPlatform Platform, FSHAHash Hash, const TArray<uint8>& Code, const FStreamOutElementList& ElementList, uint32 NumStrides, const uint32* Strides, int32 RasterizedStream);
 	/** Instantiate or retrieve a hull shader from the cache for the provided code & hash. */
 	static FHullShaderRHIRef CreateHullShader(EShaderPlatform Platform, FSHAHash Hash, TArray<uint8> const& Code);
 	/** Instantiate or retrieve a domain shader from the cache for the provided code & hash. */
diff --git a/Engine/Source/Runtime/VulkanRHI/Private/VulkanShaders.cpp b/Engine/Source/Runtime/VulkanRHI/Private/VulkanShaders.cpp
--- a/Engine/Source/Runtime/VulkanRHI/Private/VulkanShaders.cpp
+++ b/Engine/Source/Runtime/VulkanRHI/Private/VulkanShaders.cpp
@@ -235,13 +235,6 @@ FGeometryShaderRHIRef FVulkanDynamicRHI::RHICreateGeometryShader(const TArray<ui
 	return Device->GetShaderFactory().CreateShader<FVulkanGeometryShader>(Code, Device);
 }
 
-FGeometryShaderRHIRef FVulkanDynamicRHI::RHICreateGeometryShaderWithStreamOutput(const TArray<uint8>& Code, const FStreamOutElementList& ElementList,
-	uint32 NumStrides, const uint32* Strides, int32 RasterizedStream)
-{
-	VULKAN_SIGNAL_UNIMPLEMENTED();
-	return nullptr;
-}
-
 FComputeShaderRHIRef FVulkanDynamicRHI::RHICreateComputeShader(const TArray<uint8>& Code) 
 { 
 	return Device->GetShaderFactory().CreateShader<FVulkanComputeShader>(Code, Device);
diff --git a/Engine/Source/Runtime/VulkanRHI/Public/VulkanDynamicRHI.h b/Engine/Source/Runtime/VulkanRHI/Public/VulkanDynamicRHI.h
--- a/Engine/Source/Runtime/VulkanRHI/Public/VulkanDynamicRHI.h
+++ b/Engine/Source/Runtime/VulkanRHI/Public/VulkanDynamicRHI.h
@@ -50,7 +50,6 @@ public:
 	virtual FHullShaderRHIRef RHICreateHullShader(const TArray<uint8>& Code) final override;
 	virtual FDomainShaderRHIRef RHICreateDomainShader(const TArray<uint8>& Code) final override;
 	virtual FGeometryShaderRHIRef RHICreateGeometryShader(const TArray<uint8>& Code) final override;
-	virtual FGeometryShaderRHIRef RHICreateGeometryShaderWithStreamOutput(const TArray<uint8>& Code, const FStreamOutElementList& ElementList, uint32 NumStrides, const uint32* Strides, int32 RasterizedStream) final override;
 	virtual FComputeShaderRHIRef RHICreateComputeShader(const TArray<uint8>& Code) final override;
 	virtual FComputeFenceRHIRef RHICreateComputeFence(const FName& Name) final override;
 	virtual FGPUFenceRHIRef RHICreateGPUFence(const FName &Name) final override;
@@ -205,11 +204,7 @@ public:
 		return RHICreateGeometryShader(Code);
 	}
 	//virtual FGeometryShaderRHIRef CreateGeometryShader_RenderThread(class FRHICommandListImmediate& RHICmdList, FRHIShaderLibrary* Library, FSHAHash Hash) override final;
-	virtual FGeometryShaderRHIRef CreateGeometryShaderWithStreamOutput_RenderThread(class FRHICommandListImmediate& RHICmdList, const TArray<uint8>& Code, const FStreamOutElementList& ElementList, uint32 NumStrides, const uint32* Strides, int32 RasterizedStream) override final
-	{
-		return RHICreateGeometryShaderWithStreamOutput(Code, ElementList, NumStrides, Strides, RasterizedStream);
-	}
-	//virtual FGeometryShaderRHIRef CreateGeometryShaderWithStreamOutput_RenderThread(class FRHICommandListImmediate& RHICmdList, const FStreamOutElementList& ElementList, uint32 NumStrides, const uint32* Strides, int32 RasterizedStream, FRHIShaderLibrary* Library, FSHAHash Hash) override final;
+
 	virtual FComputeShaderRHIRef CreateComputeShader_RenderThread(class FRHICommandListImmediate& RHICmdList, const TArray<uint8>& Code) override final
 	{
 		return RHICreateComputeShader(Code);
diff --git a/Engine/Source/Runtime/Windows/D3D11RHI/Private/D3D11RHIPrivate.h b/Engine/Source/Runtime/Windows/D3D11RHI/Private/D3D11RHIPrivate.h
--- a/Engine/Source/Runtime/Windows/D3D11RHI/Private/D3D11RHIPrivate.h
+++ b/Engine/Source/Runtime/Windows/D3D11RHI/Private/D3D11RHIPrivate.h
@@ -405,7 +405,6 @@ public:
 	virtual FHullShaderRHIRef RHICreateHullShader(const TArray<uint8>& Code) final override;
 	virtual FDomainShaderRHIRef RHICreateDomainShader(const TArray<uint8>& Code) final override;
 	virtual FGeometryShaderRHIRef RHICreateGeometryShader(const TArray<uint8>& Code) final override;
-	virtual FGeometryShaderRHIRef RHICreateGeometryShaderWithStreamOutput(const TArray<uint8>& Code, const FStreamOutElementList& ElementList, uint32 NumStrides, const uint32* Strides, int32 RasterizedStream) final override;
 	virtual FComputeShaderRHIRef RHICreateComputeShader(const TArray<uint8>& Code) final override;
 	virtual FStagingBufferRHIRef RHICreateStagingBuffer() final override;
 	virtual void RHICopyToStagingBuffer(FRHIVertexBuffer* SourceBuffer, FRHIStagingBuffer* DestinationStagingBuffer, uint32 Offset, uint32 NumBytes) final override;
@@ -616,7 +615,6 @@ public:
 	virtual FVertexShaderRHIRef CreateVertexShader_RenderThread(class FRHICommandListImmediate& RHICmdList, const TArray<uint8>& Code) final override;
 	virtual FPixelShaderRHIRef CreatePixelShader_RenderThread(class FRHICommandListImmediate& RHICmdList, const TArray<uint8>& Code) final override;
 	virtual FGeometryShaderRHIRef CreateGeometryShader_RenderThread(class FRHICommandListImmediate& RHICmdList, const TArray<uint8>& Code) final override;
-	virtual FGeometryShaderRHIRef CreateGeometryShaderWithStreamOutput_RenderThread(class FRHICommandListImmediate& RHICmdList, const TArray<uint8>& Code, const FStreamOutElementList& ElementList, uint32 NumStrides, const uint32* Strides, int32 RasterizedStream) final override;
 	virtual FComputeShaderRHIRef CreateComputeShader_RenderThread(class FRHICommandListImmediate& RHICmdList, const TArray<uint8>& Code) final override;
 	virtual FHullShaderRHIRef CreateHullShader_RenderThread(class FRHICommandListImmediate& RHICmdList, const TArray<uint8>& Code) final override;
 	virtual FDomainShaderRHIRef CreateDomainShader_RenderThread(class FRHICommandListImmediate& RHICmdList, const TArray<uint8>& Code);
diff --git a/Engine/Source/Runtime/Windows/D3D11RHI/Private/D3D11Shaders.cpp b/Engine/Source/Runtime/Windows/D3D11RHI/Private/D3D11Shaders.cpp
--- a/Engine/Source/Runtime/Windows/D3D11RHI/Private/D3D11Shaders.cpp
+++ b/Engine/Source/Runtime/Windows/D3D11RHI/Private/D3D11Shaders.cpp
@@ -87,71 +87,6 @@ FGeometryShaderRHIRef FD3D11DynamicRHI::CreateGeometryShader_RenderThread(
 	return RHICreateGeometryShader(Code);
 }
 
-FGeometryShaderRHIRef FD3D11DynamicRHI::RHICreateGeometryShaderWithStreamOutput(const TArray<uint8>& Code, const FStreamOutElementList& ElementList, uint32 NumStrides, const uint32* Strides, int32 RasterizedStream) 
-{ 
-	FShaderCodeReader ShaderCode(Code);
-
-	FD3D11GeometryShader* Shader = new FD3D11GeometryShader;
-
-	FMemoryReader Ar( Code, true );
-	Ar << Shader->ShaderResourceTable;
-	int32 Offset = Ar.Tell();
-	const uint8* CodePtr = Code.GetData() + Offset;
-	const size_t CodeSize = ShaderCode.GetActualShaderCodeSize() - Offset;
-
-	uint32 D3DRasterizedStream = RasterizedStream;
-	if (RasterizedStream == -1)
-	{
-		D3DRasterizedStream = D3D11_SO_NO_RASTERIZED_STREAM;
-	}
-
-	D3D11_SO_DECLARATION_ENTRY StreamOutEntries[D3D11_SO_STREAM_COUNT * D3D11_SO_OUTPUT_COMPONENT_COUNT];
-
-	for (int32 EntryIndex = 0; EntryIndex < ElementList.Num(); EntryIndex++)
-	{
-		StreamOutEntries[EntryIndex].Stream = ElementList[EntryIndex].Stream;
-		StreamOutEntries[EntryIndex].SemanticName = ElementList[EntryIndex].SemanticName;
-		StreamOutEntries[EntryIndex].SemanticIndex = ElementList[EntryIndex].SemanticIndex;
-		StreamOutEntries[EntryIndex].StartComponent = ElementList[EntryIndex].StartComponent;
-		StreamOutEntries[EntryIndex].ComponentCount = ElementList[EntryIndex].ComponentCount;
-		StreamOutEntries[EntryIndex].OutputSlot = ElementList[EntryIndex].OutputSlot;
-	}
-
-	VERIFYD3D11SHADERRESULT( Direct3DDevice->CreateGeometryShaderWithStreamOutput(
-			(void*)CodePtr,
-			CodeSize,
-			StreamOutEntries,
-			ElementList.Num(),
-			Strides,
-			NumStrides,
-			D3DRasterizedStream,
-			NULL,
-			Shader->Resource.GetInitReference()
-			),
-		Shader,
-		Direct3DDevice);
-	
-	auto PackedResourceCounts = ShaderCode.FindOptionalData<FShaderCodePackedResourceCounts>();
-	check(PackedResourceCounts);
-	Shader->bShaderNeedsGlobalConstantBuffer = PackedResourceCounts->bGlobalUniformBufferUsed;
-#if !(UE_BUILD_SHIPPING || UE_BUILD_TEST)
-	Shader->ShaderName = ShaderCode.FindOptionalData('n');
-#endif
-
-	return Shader;
-}
-
-FGeometryShaderRHIRef FD3D11DynamicRHI::CreateGeometryShaderWithStreamOutput_RenderThread(
-	class FRHICommandListImmediate& RHICmdList,
-	const TArray<uint8>& Code,
-	const FStreamOutElementList& ElementList,
-	uint32 NumStrides,
-	const uint32* Strides,
-	int32 RasterizedStream)
-{
-	return RHICreateGeometryShaderWithStreamOutput(Code, ElementList, NumStrides, Strides, RasterizedStream);
-}
-
 FHullShaderRHIRef FD3D11DynamicRHI::RHICreateHullShader(const TArray<uint8>& Code) 
 { 
 	FShaderCodeReader ShaderCode(Code);
