Refactored ClearUAV into RenderGraphUtils. Renamed to AddClearUAVPass to match convention. Implemented using texture descriptor instead of parent RHI texture. Removed RDG validation workaround that allowed parent texture to be accessed without being declared. Added r.RDG.ClobberResources implementation which can help catch bugs with clears or sampling outside a valid view rect.

diff --git a/Engine/Source/Runtime/RenderCore/Private/ClearQuad.cpp b/Engine/Source/Runtime/RenderCore/Private/ClearQuad.cpp
--- a/Engine/Source/Runtime/RenderCore/Private/ClearQuad.cpp
+++ b/Engine/Source/Runtime/RenderCore/Private/ClearQuad.cpp
@@ -374,61 +374,4 @@ void DrawClearQuadMRT(FRHICommandList& RHICmdList, bool bClearColor, int32 NumCl
 		RHICmdList.SetStreamSource(0, GClearVertexBuffer.VertexBufferRHI, 0);
 		RHICmdList.DrawPrimitive(0, 2, 1);
 	}
-}
-
-BEGIN_SHADER_PARAMETER_STRUCT(FClearBufferUAVParameters, )
-	SHADER_PARAMETER_RDG_BUFFER_UAV(RWBuffer<uint>, BufferUAV)
-END_SHADER_PARAMETER_STRUCT()
-
-void ClearUAV(FRDGBuilder& GraphBuilder, FRDGEventName&& PassName, FRDGBufferUAVRef BufferUAV, uint32 Value)
-{
-	FClearBufferUAVParameters* Parameters = GraphBuilder.AllocParameters<FClearBufferUAVParameters>();
-	Parameters->BufferUAV = BufferUAV;
-
-	GraphBuilder.AddPass(
-		Forward<FRDGEventName>(PassName),
-		Parameters,
-		ERDGPassFlags::Compute,
-		[&Parameters, BufferUAV, Value](FRHICommandList& RHICmdList)
-	{
-		// This might be called if using ClearTinyUAV.
-		BufferUAV->MarkResourceAsUsed();
-
-		::ClearUAV(RHICmdList, BufferUAV->GetRHI(), BufferUAV->GetParent()->Desc.GetTotalNumBytes(), Value);
-	});
-}
-
-BEGIN_SHADER_PARAMETER_STRUCT(FClearTextureUAVParameters, )
-	SHADER_PARAMETER_RDG_TEXTURE_UAV(RWTexture2D, TextureUAV)
-END_SHADER_PARAMETER_STRUCT()
-
-template <typename T>
-void ClearUAV_T(FRDGBuilder& GraphBuilder, FRDGEventName&& PassName, FRDGTextureUAVRef TextureUAV, const T(&ClearValues)[4])
-{
-	FClearTextureUAVParameters* Parameters = GraphBuilder.AllocParameters<FClearTextureUAVParameters>();
-	Parameters->TextureUAV = TextureUAV;
-
-	GraphBuilder.AddPass(
-		Forward<FRDGEventName>(PassName),
-		Parameters,
-		ERDGPassFlags::Compute,
-		[&Parameters, TextureUAV, ClearValues](FRHICommandList& RHICmdList)
-	{
-		::ClearUAV(RHICmdList, TextureUAV->GetParent()->GetRHI(), TextureUAV->GetRHI(), ClearValues);
-	});
-}
-
-void ClearUAV(FRDGBuilder& GraphBuilder, FRDGEventName&& PassName, FRDGTextureUAVRef TextureUAV, const float(&ClearValues)[4])
-{
-	ClearUAV_T(GraphBuilder, Forward<FRDGEventName>(PassName), TextureUAV, ClearValues);
-}
-
-void ClearUAV(FRDGBuilder& GraphBuilder, FRDGEventName&& PassName, FRDGTextureUAVRef TextureUAV, const uint32(&ClearValues)[4])
-{
-	ClearUAV_T(GraphBuilder, Forward<FRDGEventName>(PassName), TextureUAV, ClearValues);
-}
-
-void ClearUAV(FRDGBuilder& GraphBuilder, FRDGEventName&& PassName, FRDGTextureUAVRef TextureUAV, const FLinearColor& ClearColor)
-{
-	ClearUAV_T(GraphBuilder, Forward<FRDGEventName>(PassName), TextureUAV, reinterpret_cast<const float(&)[4]>(ClearColor));
-}
+}
\ No newline at end of file
diff --git a/Engine/Source/Runtime/RenderCore/Private/RenderGraphBuilder.cpp b/Engine/Source/Runtime/RenderCore/Private/RenderGraphBuilder.cpp
--- a/Engine/Source/Runtime/RenderCore/Private/RenderGraphBuilder.cpp
+++ b/Engine/Source/Runtime/RenderCore/Private/RenderGraphBuilder.cpp
@@ -31,20 +31,57 @@ FAutoConsoleVariableRef CVarRDGDebug(
 	TEXT(" 2: emit warning everytime issue is detected."),
 	ECVF_RenderThreadSafe);
 
-static TAutoConsoleVariable<int32> CVarRDGEnableBreakPoint(
-	TEXT("r.RDG.EnableBreakPoint"), 0,
+TAutoConsoleVariable<int32> CVarRDGEnableBreakpoint(
+	TEXT("r.RDG.EnableBreakpoint"), 0,
 	TEXT("Breakpoint in debugger when a warning is raised.\n"),
 	ECVF_RenderThreadSafe);
 
+int32 GRDGClobberResources = 0;
+FAutoConsoleVariableRef CVarRDGClobberResources(
+	TEXT("r.RDG.ClobberResources"),
+	GRDGClobberResources,
+	TEXT("Clears all render targets and texture / buffer UAVs with the requested clear color at allocation time. Useful for debugging.\n")
+	TEXT(" 0:off (default);\n")
+	TEXT(" 1: 1000 on RGBA channels;\n")
+	TEXT(" 2: NaN on RGBA channels;\n")
+	TEXT(" 3: +INFINITY on RGBA channels.\n"),
+	ECVF_Cheat | ECVF_RenderThreadSafe);
+
+FLinearColor GetClobberColor()
+{
+	switch (GRDGClobberResources)
+	{
+	case 2:
+		return FLinearColor(NAN, NAN, NAN, NAN);
+		break;
+	case 3:
+		return FLinearColor(INFINITY, INFINITY, INFINITY, INFINITY);
+		break;
+	default:
+		return FLinearColor(1000, 1000, 1000, 1000);
+	}
+}
+
+uint32 GetClobberBufferValue()
+{
+	return 1000;
+}
 
 #else
 
 const int32 GRDGImmediateMode = 0;
 const int32 GRDGDebug = 0;
+const int32 GRDGClobberResources = 0;
 
 #endif
+
 } //! namespace
 
+bool IsRDGClobberResourcesEnabled()
+{
+	return GRDGClobberResources > 0;
+}
+
 bool GetEmitRDGEvents()
 {
 	check(IsInRenderingThread());
@@ -65,7 +102,6 @@ bool IsRDGImmediateModeEnabled()
 	return GRDGImmediateMode != 0;
 }
 
-
 void InitRenderGraph()
 {
 #if RDG_ENABLE_DEBUG_WITH_ENGINE
@@ -91,7 +127,7 @@ void EmitRDGWarning(const FString& WarningMessage)
 
 	static TSet<FString> GAlreadyEmittedWarnings;
 
-	const bool bEnableBreakPoint = CVarRDGEnableBreakPoint.GetValueOnRenderThread();
+	const bool bEnableBreakpoint = CVarRDGEnableBreakpoint.GetValueOnRenderThread();
 
 	if (GRDGDebug == kRDGEmitWarningsOnce)
 	{
@@ -99,8 +135,8 @@ void EmitRDGWarning(const FString& WarningMessage)
 		{
 			GAlreadyEmittedWarnings.Add(WarningMessage);
 			UE_LOG(LogRendererCore, Warning, TEXT("%s"), *WarningMessage);
-			
-			if (bEnableBreakPoint)
+
+			if (bEnableBreakpoint)
 			{
 				UE_DEBUG_BREAK();
 			}
@@ -109,8 +145,8 @@ void EmitRDGWarning(const FString& WarningMessage)
 	else
 	{
 		UE_LOG(LogRendererCore, Warning, TEXT("%s"), *WarningMessage);
-		
-		if (bEnableBreakPoint)
+
+		if (bEnableBreakpoint)
 		{
 			UE_DEBUG_BREAK();
 		}
@@ -162,7 +198,9 @@ void FRDGBuilder::Execute()
 
 void FRDGBuilder::AddPassInternal(FRDGPass* Pass)
 {
-	IF_RDG_ENABLE_DEBUG(Validation.ValidateAddPass(Pass));
+	ClobberPassOutputs(Pass);
+
+	IF_RDG_ENABLE_DEBUG(Validation.ValidateAddPass(Pass, bInDebugPassScope));
 
 	Pass->EventScope = EventScopeStack.GetCurrentScope();
 	Pass->StatScope = StatScopeStack.GetCurrentScope();
@@ -179,7 +217,7 @@ void FRDGBuilder::AddPassInternal(FRDGPass* Pass)
 void FRDGBuilder::VisualizePassOutputs(const FRDGPass* Pass)
 {
 #if SUPPORTS_VISUALIZE_TEXTURE
-	if (!GVisualizeTexture.bEnabled)
+	if (!GVisualizeTexture.bEnabled || bInDebugPassScope)
 	{
 		return;
 	}
@@ -254,13 +292,117 @@ void FRDGBuilder::VisualizePassOutputs(const FRDGPass* Pass)
 			}
 		}
 		break;
-		default:
-			break;
 		}
 	}
 #endif
 }
 
+#if RDG_ENABLE_DEBUG
+
+void FRDGBuilder::ClobberPassOutputs(const FRDGPass* Pass)
+{
+	if (!IsRDGClobberResourcesEnabled())
+	{
+		return;
+	}
+
+	if (bInDebugPassScope)
+	{
+		return;
+	}
+	bInDebugPassScope = true;
+
+	const auto TryMarkForClobber = [](FRDGParentResourceRef Resource) -> bool
+	{
+		const bool bClobber = !Resource->bHasBeenClobbered && !Resource->IsExternal();
+
+		if (bClobber)
+		{
+			Resource->bHasBeenClobbered = true;
+		}
+
+		return bClobber;
+	};
+
+	const FLinearColor ClobberColor = GetClobberColor();
+
+	FRDGPassParameterStruct ParameterStruct = Pass->GetParameters();
+
+	const uint32 ParameterCount = ParameterStruct.GetParameterCount();
+
+	for (uint32 ParameterIndex = 0; ParameterIndex < ParameterCount; ++ParameterIndex)
+	{
+		FRDGPassParameter Parameter = ParameterStruct.GetParameter(ParameterIndex);
+
+		switch (Parameter.GetType())
+		{
+		case UBMT_RDG_BUFFER_UAV:
+		{
+			if (FRDGBufferUAVRef UAV = Parameter.GetAsBufferUAV())
+			{
+				FRDGBufferRef Buffer = UAV->GetParent();
+
+				if (TryMarkForClobber(Buffer))
+				{
+					AddClearUAVPass(*this, UAV, GetClobberBufferValue());
+				}
+			}
+		}
+		break;
+		case UBMT_RDG_TEXTURE_UAV:
+		{
+			if (FRDGTextureUAVRef UAV = Parameter.GetAsTextureUAV())
+			{
+				FRDGTextureRef Texture = UAV->GetParent();
+
+				if (TryMarkForClobber(Texture))
+				{
+					AddClearUAVPass(*this, UAV, ClobberColor);
+				}
+			}
+		}
+		break;
+		case UBMT_RENDER_TARGET_BINDING_SLOTS:
+		{
+			const FRenderTargetBindingSlots& RenderTargetBindingSlots = Parameter.GetAsRenderTargetBindingSlots();
+			const auto& DepthStencil = RenderTargetBindingSlots.DepthStencil;
+			const auto& RenderTargets = RenderTargetBindingSlots.Output;
+
+			if (FRDGTextureRef Texture = DepthStencil.GetTexture())
+			{
+				if (TryMarkForClobber(Texture))
+				{
+					AddClearDepthStencilPass(*this, Texture, true, 0.0f, true, 0);
+				}
+			}
+
+			const uint32 RenderTargetCount = RenderTargets.Num();
+
+			for (uint32 RenderTargetIndex = 0; RenderTargetIndex < RenderTargetCount; ++RenderTargetIndex)
+			{
+				const FRenderTargetBinding& RenderTarget = RenderTargets[RenderTargetIndex];
+
+				if (FRDGTextureRef Texture = RenderTarget.GetTexture())
+				{
+					if (TryMarkForClobber(Texture))
+					{
+						AddClearRenderTargetPass(*this, Texture, ClobberColor);
+					}
+				}
+				else
+				{
+					break;
+				}
+			}
+		}
+		break;
+		}
+	}
+	bInDebugPassScope = false;
+}
+
+#endif
+
 void FRDGBuilder::WalkGraphDependencies()
 {
 	for (const FRDGPass* Pass : Passes)
@@ -528,7 +670,7 @@ void FRDGBuilder::ExecutePass(const FRDGPass* Pass)
 	{
 		UnbindRenderTargets(RHICmdList);
 	}
-	
+
 	Pass->Execute(RHICmdList);
 
 	if (bHasRenderTargets)
diff --git a/Engine/Source/Runtime/RenderCore/Private/RenderGraphUtils.cpp b/Engine/Source/Runtime/RenderCore/Private/RenderGraphUtils.cpp
--- a/Engine/Source/Runtime/RenderCore/Private/RenderGraphUtils.cpp
+++ b/Engine/Source/Runtime/RenderCore/Private/RenderGraphUtils.cpp
@@ -1,13 +1,15 @@
 // Copyright 1998-2019 Epic Games, Inc. All Rights Reserved.
 
 #include "RenderGraphUtils.h"
+#include "ClearQuad.h"
+#include "ClearReplacementShaders.h"
 #include <initializer_list>
 
- /** Adds a render graph tracked buffer suitable for use as a copy destination. */
+/** Adds a render graph tracked buffer suitable for use as a copy destination. */
 #define RDG_BUFFER_COPY_DEST(MemberName) \
 	INTERNAL_SHADER_PARAMETER_EXPLICIT(UBMT_RDG_BUFFER_COPY_DEST, TShaderResourceParameterTypeInfo<FRDGBufferRef>, FRDGBufferRef,MemberName,, = nullptr,EShaderPrecisionModifier::Float,TEXT(""),false)
 
-  /** Adds a render graph tracked texture suitable for use as a copy destination. */
+/** Adds a render graph tracked texture suitable for use as a copy destination. */
 #define RDG_TEXTURE_COPY_DEST(MemberName) \
 	INTERNAL_SHADER_PARAMETER_EXPLICIT(UBMT_RDG_TEXTURE_COPY_DEST, TShaderResourceParameterTypeInfo<FRDGTextureRef>, FRDGTextureRef,MemberName,, = nullptr,EShaderPrecisionModifier::Float,TEXT(""),false)
 
@@ -71,7 +73,6 @@ void ClearUnusedGraphResourcesImpl(
 	}
 }
 
-
 FRDGTextureRef RegisterExternalTextureWithFallback(
 	FRDGBuilder& GraphBuilder,
 	const TRefCountPtr<IPooledRenderTarget>& ExternalPooledTexture,
@@ -118,3 +119,189 @@ void AddCopyTexturePass(
 		RHICmdList.CopyTexture(InputTexture->GetRHI(), OutputTexture->GetRHI(), CopyInfo);
 	});
 }
+
+BEGIN_SHADER_PARAMETER_STRUCT(FClearBufferUAVParameters, )
+	SHADER_PARAMETER_RDG_BUFFER_UAV(RWBuffer<uint>, BufferUAV)
+END_SHADER_PARAMETER_STRUCT()
+
+void AddClearUAVPass(FRDGBuilder& GraphBuilder, FRDGBufferUAVRef BufferUAV, uint32 Value)
+{
+	FClearBufferUAVParameters* Parameters = GraphBuilder.AllocParameters<FClearBufferUAVParameters>();
+	Parameters->BufferUAV = BufferUAV;
+
+	GraphBuilder.AddPass(
+		RDG_EVENT_NAME("ClearBuffer(%s)", BufferUAV->GetParent()->Name),
+		Parameters,
+		ERDGPassFlags::Compute,
+		[&Parameters, BufferUAV, Value](FRHICommandList& RHICmdList)
+	{
+		// This might be called if using ClearTinyUAV.
+		BufferUAV->MarkResourceAsUsed();
+
+		::ClearUAV(RHICmdList, BufferUAV->GetRHI(), BufferUAV->GetParent()->Desc.GetTotalNumBytes(), Value);
+	});
+}
+
+BEGIN_SHADER_PARAMETER_STRUCT(FClearTextureUAVParameters, )
+	SHADER_PARAMETER_RDG_TEXTURE_UAV(RWTexture2D, TextureUAV)
+END_SHADER_PARAMETER_STRUCT()
+
+template <typename T>
+void AddClearUAVPass_T(FRDGBuilder& GraphBuilder, FRDGTextureUAVRef TextureUAV, const T(&ClearValues)[4])
+{
+	check(TextureUAV);
+
+	FClearTextureUAVParameters* Parameters = GraphBuilder.AllocParameters<FClearTextureUAVParameters>();
+	Parameters->TextureUAV = TextureUAV;
+
+	const FRDGTextureDesc& TextureDesc = TextureUAV->GetParent()->Desc;
+
+	GraphBuilder.AddPass(
+		RDG_EVENT_NAME("ClearTexture(%s) %dx%d", TextureUAV->GetParent()->Name, TextureDesc.Extent.X, TextureDesc.Extent.Y),
+		Parameters,
+		ERDGPassFlags::Compute,
+		[&Parameters, TextureUAV, ClearValues](FRHICommandList& RHICmdList)
+	{
+		const FRDGTextureDesc& TextureDesc = TextureUAV->GetParent()->Desc;
+
+		FRHIUnorderedAccessView* RHITextureUAV = TextureUAV->GetRHI();
+
+		if (TextureDesc.Is2DTexture())
+		{
+			if (TextureDesc.IsArray())
+			{
+				TShaderMapRef<FClearTexture2DArrayReplacementCS<T>> ComputeShader(GetGlobalShaderMap(GMaxRHIFeatureLevel));
+				FRHIComputeShader* ShaderRHI = ComputeShader->GetComputeShader();
+				RHICmdList.SetComputeShader(ShaderRHI);
+				ComputeShader->SetParameters(RHICmdList, RHITextureUAV, ClearValues);
+				uint32 x = (TextureDesc.Extent.X + 7) / 8;
+				uint32 y = (TextureDesc.Extent.Y + 7) / 8;
+				uint32 z = TextureDesc.ArraySize;
+				RHICmdList.DispatchComputeShader(x, y, z);
+				ComputeShader->FinalizeParameters(RHICmdList, RHITextureUAV);
+			}
+			else
+			{
+				TShaderMapRef<FClearTexture2DReplacementCS<T>> ComputeShader(GetGlobalShaderMap(GMaxRHIFeatureLevel));
+				FRHIComputeShader* ShaderRHI = ComputeShader->GetComputeShader();
+				RHICmdList.SetComputeShader(ShaderRHI);
+				ComputeShader->SetParameters(RHICmdList, RHITextureUAV, ClearValues);
+				const uint32 X = (TextureDesc.Extent.X + 7) / 8;
+				const uint32 Y = (TextureDesc.Extent.Y + 7) / 8;
+				RHICmdList.DispatchComputeShader(X, Y, 1);
+				ComputeShader->FinalizeParameters(RHICmdList, RHITextureUAV);
+			}
+		}
+		else if (TextureDesc.IsCubemap())
+		{
+			TShaderMapRef<FClearTexture2DArrayReplacementCS<T>> ComputeShader(GetGlobalShaderMap(GMaxRHIFeatureLevel));
+			FRHIComputeShader* ShaderRHI = ComputeShader->GetComputeShader();
+			RHICmdList.SetComputeShader(ShaderRHI);
+			ComputeShader->SetParameters(RHICmdList, RHITextureUAV, ClearValues);
+			uint32 x = (TextureDesc.Extent.X + 7) / 8;
+			uint32 y = (TextureDesc.Extent.Y + 7) / 8;
+			RHICmdList.DispatchComputeShader(x, y, 6);
+			ComputeShader->FinalizeParameters(RHICmdList, RHITextureUAV);
+		}
+		else if (TextureDesc.Is3DTexture())
+		{
+			TShaderMapRef<FClearVolumeReplacementCS<T>> ComputeShader(GetGlobalShaderMap(GMaxRHIFeatureLevel));
+			FRHIComputeShader* ShaderRHI = ComputeShader->GetComputeShader();
+			RHICmdList.SetComputeShader(ShaderRHI);
+			ComputeShader->SetParameters(RHICmdList, RHITextureUAV, ClearValues);
+			uint32 x = (TextureDesc.Extent.X + 3) / 4;
+			uint32 y = (TextureDesc.Extent.Y + 3) / 4;
+			uint32 z = (TextureDesc.Depth + 3) / 4;
+			RHICmdList.DispatchComputeShader(x, y, z);
+			ComputeShader->FinalizeParameters(RHICmdList, RHITextureUAV);
+		}
+		else
+		{
+			check(0);
+		}
+	});
+}
+
+void AddClearUAVPass(FRDGBuilder& GraphBuilder, FRDGTextureUAVRef TextureUAV, const float(&ClearValues)[4])
+{
+	AddClearUAVPass_T(GraphBuilder, TextureUAV, ClearValues);
+}
+
+void AddClearUAVPass(FRDGBuilder& GraphBuilder, FRDGTextureUAVRef TextureUAV, const uint32(&ClearValues)[4])
+{
+	AddClearUAVPass_T(GraphBuilder, TextureUAV, ClearValues);
+}
+
+void AddClearUAVPass(FRDGBuilder& GraphBuilder, FRDGTextureUAVRef TextureUAV, const FLinearColor& ClearColor)
+{
+	AddClearUAVPass_T(GraphBuilder, TextureUAV, reinterpret_cast<const float(&)[4]>(ClearColor));
+}
+
+void AddClearRenderTargetPass(FRDGBuilder& GraphBuilder, FRDGTextureRef Texture, const FLinearColor& ClearColor)
+{
+	check(Texture);
+
+	FRenderTargetParameters* Parameters = GraphBuilder.AllocParameters<FRenderTargetParameters>();
+	Parameters->RenderTargets[0] = FRenderTargetBinding(Texture, ERenderTargetLoadAction::ENoAction);
+
+	GraphBuilder.AddPass(
+		RDG_EVENT_NAME("ClearRenderTarget(%s) %dx%d", Texture->Name, Texture->Desc.Extent.X, Texture->Desc.Extent.Y),
+		Parameters,
+		ERDGPassFlags::Raster,
+		[Parameters, ClearColor](FRHICommandList& RHICmdList)
+	{
+		DrawClearQuad(RHICmdList, ClearColor);
+	});
+}
+
+void AddClearDepthStencilPass(
+	FRDGBuilder& GraphBuilder,
+	FRDGTextureRef Texture,
+	bool bClearDepth,
+	float Depth,
+	bool bClearStencil,
+	uint8 Stencil)
+{
+	check(Texture);
+
+	FExclusiveDepthStencil ExclusiveDepthStencil;
+	ERenderTargetLoadAction DepthLoadAction = ERenderTargetLoadAction::ELoad;
+	ERenderTargetLoadAction StencilLoadAction = ERenderTargetLoadAction::ENoAction;
+
+	const bool bHasStencil = Texture->Desc.Format == PF_DepthStencil;
+
+	// We can't clear stencil if we don't have it.
+	bClearStencil &= bHasStencil;
+
+	if (bClearDepth)
+	{
+		ExclusiveDepthStencil.SetDepthWrite();
+		DepthLoadAction = ERenderTargetLoadAction::ENoAction;
+	}
+
+	if (bHasStencil)
+	{
+		if (bClearStencil)
+		{
+			ExclusiveDepthStencil.SetStencilWrite();
+			StencilLoadAction = ERenderTargetLoadAction::ENoAction;
+		}
+		else
+		{
+			// Preserve stencil contents.
+			StencilLoadAction = ERenderTargetLoadAction::ELoad;
+		}
+	}
+
+	FRenderTargetParameters* Parameters = GraphBuilder.AllocParameters<FRenderTargetParameters>();
+	Parameters->RenderTargets.DepthStencil = FDepthStencilBinding(Texture, DepthLoadAction, StencilLoadAction, ExclusiveDepthStencil);
+
+	GraphBuilder.AddPass(
+		RDG_EVENT_NAME("ClearDepthStencil(%s) %dx%d", Texture->Name, Texture->Desc.Extent.X, Texture->Desc.Extent.Y),
+		Parameters,
+		ERDGPassFlags::Raster,
+		[Parameters, bClearDepth, Depth, bClearStencil, Stencil](FRHICommandList& RHICmdList)
+	{
+		DrawClearQuad(RHICmdList, false, FLinearColor(), bClearDepth, Depth, bClearStencil, Stencil);
+	});
+}
\ No newline at end of file
diff --git a/Engine/Source/Runtime/RenderCore/Private/RenderGraphValidation.cpp b/Engine/Source/Runtime/RenderCore/Private/RenderGraphValidation.cpp
--- a/Engine/Source/Runtime/RenderCore/Private/RenderGraphValidation.cpp
+++ b/Engine/Source/Runtime/RenderCore/Private/RenderGraphValidation.cpp
@@ -4,6 +4,11 @@
 
 #if RDG_ENABLE_DEBUG
 
+/** Validates that we are only executing a single render graph instance in the callstack. Used to catch if a
+ *  user creates a second FRDGBuilder instance inside of a pass that is executing.
+ */
+static bool GRDGInExecutePassScope = false;
+
 /** This utility class tracks read / write events on individual resources within a pass in order to ensure
  *  that sub-resources are not being simultaneously bound for read / write. The RHI currently only supports
  *  disjoint read-write access on a texture for mip-map generation purposes.
@@ -226,7 +231,6 @@ void FRDGUserValidation::ValidateExtractResource(FRDGParentResourceRef Resource)
 	 *  emitting a 'produced but never used' warning. We don't have the history of registered
 	 *  resources to be able to emit a proper warning.
 	 */
-
 	Resource->PassAccessCount++;
 }
 
@@ -234,10 +238,10 @@ void FRDGUserValidation::RemoveUnusedWarning(FRDGParentResourceRef Resource)
 {
 	check(Resource);
 
-    // Removes 'produced but not used' warning.
+	// Removes 'produced but not used' warning.
 	Resource->PassAccessCount++;
 
-    // Removes 'not used' warning.
+	// Removes 'not used' warning.
 	Resource->bIsActuallyUsedByPass = true;
 }
 
@@ -247,7 +251,7 @@ void FRDGUserValidation::ValidateAllocPassParameters(const void* Parameters)
 	AllocatedUnusedPassParameters.Add(Parameters);
 }
 
-void FRDGUserValidation::ValidateAddPass(const FRDGPass* Pass)
+void FRDGUserValidation::ValidateAddPass(const FRDGPass* Pass, bool bSkipPassAccessMarking)
 {
 	checkf(!bHasExecuted, TEXT("Render graph pass %s needs to be added before the builder execution."), Pass->GetName());
 
@@ -267,6 +271,23 @@ void FRDGUserValidation::ValidateAddPass(const FRDGPass* Pass)
 		AllocatedUnusedPassParameters.Remove(ParameterStructData);
 	}
 
+	const auto MarkAsProduced = [Pass, bSkipPassAccessMarking] (FRDGParentResourceRef Resource)
+	{
+		if (!bSkipPassAccessMarking)
+		{
+			Resource->MarkAsProducedBy(Pass);
+			Resource->PassAccessCount++;
+		}
+	};
+
+	const auto MarkAsConsumed = [Pass, bSkipPassAccessMarking] (FRDGParentResourceRef Resource)
+	{
+		if (!bSkipPassAccessMarking)
+		{
+			Resource->PassAccessCount++;
+		}
+	};
+
 	const FRenderTargetBindingSlots* RenderTargetBindingSlots = nullptr;
 
 	const TCHAR* PassName = Pass->GetName();
@@ -315,7 +336,7 @@ void FRDGUserValidation::ValidateAddPass(const FRDGPass* Pass)
 					TEXT("Pass %s has a dependency over the texture %s that has never been produced."),
 					PassName, Texture->Name);
 
-				Texture->PassAccessCount++;
+				MarkAsConsumed(Texture);
 			}
 		}
 		break;
@@ -331,7 +352,7 @@ void FRDGUserValidation::ValidateAddPass(const FRDGPass* Pass)
 					TEXT("Pass %s has a dependency over the texture %s that has never been produced."),
 					PassName, Texture->Name);
 
-				Texture->PassAccessCount++;
+				MarkAsConsumed(Texture);
 			}
 		}
 		break;
@@ -343,8 +364,7 @@ void FRDGUserValidation::ValidateAddPass(const FRDGPass* Pass)
 
 				PassResourceValidator.Write(Texture, UAV->Desc.MipLevel);
 
-				Texture->MarkAsProducedBy(Pass);
-				Texture->PassAccessCount++;
+				MarkAsProduced(Texture);
 			}
 		}
 		break;
@@ -358,7 +378,7 @@ void FRDGUserValidation::ValidateAddPass(const FRDGPass* Pass)
 					TEXT("Pass %s has a dependency over the buffer %s that has never been produced."),
 					PassName, Buffer->Name);
 
-				Buffer->PassAccessCount++;
+				MarkAsConsumed(Buffer);
 			}
 		}
 		break;
@@ -374,7 +394,7 @@ void FRDGUserValidation::ValidateAddPass(const FRDGPass* Pass)
 					TEXT("Pass %s has a dependency over the buffer %s that has never been produced."),
 					PassName, SRV->Desc.Buffer->Name);
 
-				Buffer->PassAccessCount++;
+				MarkAsConsumed(Buffer);
 			}
 		}
 		break;
@@ -386,8 +406,7 @@ void FRDGUserValidation::ValidateAddPass(const FRDGPass* Pass)
 
 				PassResourceValidator.Write(Buffer);
 
-				Buffer->MarkAsProducedBy(Pass);
-				Buffer->PassAccessCount++;
+				MarkAsProduced(Buffer);
 			}
 		}
 		break;
@@ -397,8 +416,7 @@ void FRDGUserValidation::ValidateAddPass(const FRDGPass* Pass)
 			{
 				PassResourceValidator.Write(Texture);
 
-				Texture->MarkAsProducedBy(Pass);
-				Texture->PassAccessCount++;
+				MarkAsProduced(Texture);
 			}
 		}
 		break;
@@ -408,8 +426,7 @@ void FRDGUserValidation::ValidateAddPass(const FRDGPass* Pass)
 			{
 				PassResourceValidator.Write(Buffer);
 
-				Buffer->MarkAsProducedBy(Pass);
-				Buffer->PassAccessCount++;
+				MarkAsProduced(Buffer);
 			}
 		}
 		break;
@@ -421,14 +438,10 @@ void FRDGUserValidation::ValidateAddPass(const FRDGPass* Pass)
 			}
 			else if (IsRDGDebugEnabled())
 			{
-				EmitRDGWarningf(
-					TEXT("Pass %s have duplicated render target binding slots."),
-					PassName);
+				EmitRDGWarningf(TEXT("Pass %s have duplicated render target binding slots."), PassName);
 			}
 		}
 		break;
-		default:
-			break;
 		}
 	}
 
@@ -472,14 +485,13 @@ void FRDGUserValidation::ValidateAddPass(const FRDGPass* Pass)
 				if (DepthStencil.GetDepthStencilAccess().IsAnyWrite())
 				{
 					PassResourceValidator.Write(Texture, MipLevel);
-					Texture->MarkAsProducedBy(Pass);
+					MarkAsProduced(Texture);
 				}
 				else
 				{
 					PassResourceValidator.Read(Texture, MipLevel);
+					MarkAsConsumed(Texture);
 				}
-
-				Texture->PassAccessCount++;
 			}
 		}
 
@@ -521,21 +533,11 @@ void FRDGUserValidation::ValidateAddPass(const FRDGPass* Pass)
 					 *  which are not actually managed by the render target pool and likely represent the frame buffer.
 					 */
 					{
-						// We only validate single-mip textures since we don't track production at the subresource level.
-						const bool bFailedToLoadProducedContent = !bIsLoadAction && Texture->HasBeenProduced() && Texture->Desc.NumMips == 1;
-
-						// Untracked render targets aren't actually managed by the render target pool.
-						const bool bIsUntrackedRenderTarget = Texture->PooledRenderTarget && !Texture->PooledRenderTarget->IsTracked();
+						// Ignore external textures which are always marked as produced. We don't need to enforce this warning on them.
+						const bool bHasBeenProduced = Texture->HasBeenProduced() && !Texture->IsExternal();
 
-						ensureMsgf(!bFailedToLoadProducedContent || bIsUntrackedRenderTarget,
-							TEXT("Pass '%s' attempted to bind texture '%s' as a render target without the 'Load' action specified, despite a prior pass having produced it. It's invalid to completely clobber the contents of a resource. Create a new texture instance instead."),
-							PassName,
-							Texture->Name);
-					}
-
-					{
 						// We only validate single-mip textures since we don't track production at the subresource level.
-						const bool bFailedToLoadProducedContent = !bIsLoadAction && Texture->HasBeenProduced() && Texture->Desc.NumMips == 1;
+						const bool bFailedToLoadProducedContent = !bIsLoadAction && bHasBeenProduced && Texture->Desc.NumMips == 1;
 
 						// Untracked render targets aren't actually managed by the render target pool.
 						const bool bIsUntrackedRenderTarget = Texture->PooledRenderTarget && !Texture->PooledRenderTarget->IsTracked();
@@ -547,9 +549,7 @@ void FRDGUserValidation::ValidateAddPass(const FRDGPass* Pass)
 					}
 
 					/** Mark the pass as a producer for render targets with a store action. */
-					Texture->MarkAsProducedBy(Pass);
-
-					Texture->PassAccessCount++;
+					MarkAsProduced(Texture);
 				}
 				else
 				{
@@ -639,6 +639,10 @@ void FRDGUserValidation::ValidateExecuteEnd()
 
 void FRDGUserValidation::ValidateExecutePassBegin(const FRDGPass* Pass)
 {
+	checkf(!GRDGInExecutePassScope, TEXT("Render graph is being executed recursively. This usually means a separate FRDGBuilder instance was created inside of an executing pass."));
+
+	GRDGInExecutePassScope = true;
+
 	SetAllowRHIAccess(Pass, true);
 
 	if (IsRDGDebugEnabled())
@@ -754,6 +758,8 @@ void FRDGUserValidation::ValidateExecutePassEnd(const FRDGPass* Pass)
 			}
 		}
 	}
+
+	GRDGInExecutePassScope = false;
 }
 
 void FRDGUserValidation::SetAllowRHIAccess(const FRDGPass* Pass, bool bAllowAccess)
@@ -766,19 +772,11 @@ void FRDGUserValidation::SetAllowRHIAccess(const FRDGPass* Pass, bool bAllowAcce
 	{
 		FRDGPassParameter Parameter = ParameterStruct.GetParameter(ParameterIndex);
 
-		if (Parameter.IsParentResource())
-		{
-			if (FRDGParentResourceRef Resource = Parameter.GetAsParentResource())
-			{
-				Resource->bAllowRHIAccess = bAllowAccess;
-			}
-		}
-		else if (Parameter.IsChildResource())
+		if (Parameter.IsResource())
 		{
-			if (FRDGChildResourceRef Resource = Parameter.GetAsChildResource())
+			if (FRDGResourceRef Resource = Parameter.GetAsResource())
 			{
 				Resource->bAllowRHIAccess = bAllowAccess;
-				Resource->GetParent()->bAllowRHIAccess = bAllowAccess;
 			}
 		}
 		else if (Parameter.IsRenderTargetBindingSlots())
diff --git a/Engine/Source/Runtime/RenderCore/Private/RenderTargetPool.cpp b/Engine/Source/Runtime/RenderCore/Private/RenderTargetPool.cpp
--- a/Engine/Source/Runtime/RenderCore/Private/RenderTargetPool.cpp
+++ b/Engine/Source/Runtime/RenderCore/Private/RenderTargetPool.cpp
@@ -7,10 +7,6 @@
 #include "RenderTargetPool.h"
 #include "RHIStaticStates.h"
 
-#if !(UE_BUILD_SHIPPING || UE_BUILD_TEST) && 0
-	#include "ClearQuad.h"
-#endif
-
 /** The global render targets pool. */
 TGlobalResource<FRenderTargetPool> GRenderTargetPool;
 
@@ -57,19 +53,6 @@ static FAutoConsoleCommand GRenderTargetPoolEventsCmd(
 	FConsoleCommandWithArgsDelegate::CreateStatic(RenderTargetPoolEvents)
 	);
 
-#if !(UE_BUILD_SHIPPING || UE_BUILD_TEST)
-static TAutoConsoleVariable<int32> CVarClobberRenderTarget(
-	TEXT("r.Test.ClobberRenderRarget"),
-	0,
-	TEXT("Clears the texture returned by the rendertarget pool with a special color\n")
-	TEXT("so we can see better which passes would need to clear. Doesn't work on volume textures and non rendertargets yet.\n")
-	TEXT(" 0:off (default);\n")
-	TEXT(" 1: 1000 on RGBA channels;\n")
-	TEXT(" 2: NaN on RGBA channels;\n")
-	TEXT(" 3: +INFINITY on RGBA channels.\n"),
-	ECVF_Cheat | ECVF_RenderThreadSafe);
-#endif
-
 static TAutoConsoleVariable<int32> CVarAllowMultipleAliasingDiscardsPerFrame(
 	TEXT("r.RenderTargetPool.AllowMultipleAliasingDiscardsPerFrame"),
 	0,
@@ -229,66 +212,6 @@ void FRenderTargetPool::WaitForTransitionFence()
 	DeferredDeleteArray.Reset();
 }
 
-#if 1 // TODO(RDG): Needs dependency on UtilityShaders to reimplement ClobberAllocatedRenderTarget functionality, or just move the clearing shaders in RenderCore?
-
-static void ClobberAllocatedRenderTarget(FRHICommandList& RHICmdList, const FPooledRenderTargetDesc& Desc, IPooledRenderTarget* Out)
-{
-
-}
-
-
-#elif !(UE_BUILD_SHIPPING || UE_BUILD_TEST)
-
-static void ClobberAllocatedRenderTarget(FRHICommandList& RHICmdList, const FPooledRenderTargetDesc& Desc, IPooledRenderTarget* Out)
-{
-	int32 ClearId = CVarClobberRenderTarget.GetValueOnRenderThread();
-
-	if (!ClearId)
-	{
-		return;
-	}
-	if (Desc.bIsCubemap || Desc.bIsArray)
-	{
-		return;
-	}
-
-	FLinearColor Color = FLinearColor(1000, 1000, 1000, 1000);
-
-	if (ClearId == 2)
-	{
-		Color = FLinearColor(NAN, NAN, NAN, NAN);
-	}
-	else if (ClearId == 3)
-	{
-		Color = FLinearColor(INFINITY, INFINITY, INFINITY, INFINITY);
-	}
-
-	// TODO(RDG): draw events will need to be rewritten with render graph builder instead of RHICmdList.
-	//SCOPED_DRAW_EVENT(RHICmdList, ClobberAllocatedRenderTarget);
-
-	if (Out->GetDesc().TargetableFlags & TexCreate_RenderTargetable)
-	{
-		// Needs conversion to Render Passes
-		check(0);
-		//SetRenderTarget(RHICmdList, Out->GetRenderTargetItem().TargetableTexture, FTextureRHIRef());
-		DrawClearQuad(RHICmdList, Color);
-	}
-	else if (Out->GetDesc().TargetableFlags & TexCreate_UAV)
-	{
-		ClearUAV(RHICmdList, Out->GetRenderTargetItem(), Color);
-	}
-
-	if (Desc.TargetableFlags & TexCreate_DepthStencilTargetable)
-	{
-		// Needs conversion to Render Passes
-		check(0);
-		//SetRenderTarget(RHICmdList, FTextureRHIRef(), Out->GetRenderTargetItem().TargetableTexture);
-		DrawClearQuad(RHICmdList, false, FLinearColor::Black, true, 0.0f, true, 0);
-	}
-}
-
-#endif
-
 bool FRenderTargetPool::FindFreeElement(FRHICommandList& RHICmdList, const FPooledRenderTargetDesc& InputDesc, TRefCountPtr<IPooledRenderTarget> &Out, const TCHAR* InDebugName, bool bDoWritableBarrier, ERenderTargetTransience TransienceHint, bool bDeferTextureAllocation)
 {
 	check(IsInRenderingThread());
@@ -335,9 +258,6 @@ bool FRenderTargetPool::FindFreeElement(FRHICommandList& RHICmdList, const FPool
 				RHIBindDebugLabelName(Current->GetRenderTargetItem().TargetableTexture, InDebugName);
 			}
 			check(!Out->IsFree());
-			#if !(UE_BUILD_SHIPPING || UE_BUILD_TEST)
-				ClobberAllocatedRenderTarget(RHICmdList, Desc, Out);
-			#endif
 			return true;
 		}
 		else
@@ -592,10 +512,6 @@ Done:
 		Found->Desc.DebugName = InDebugName;
 	}
 
-	#if !(UE_BUILD_SHIPPING || UE_BUILD_TEST)
-		ClobberAllocatedRenderTarget(RHICmdList, Desc, Found);
-	#endif
-
 	check(Found->IsFree());
 	check(!Found->IsSnapshot());
 
diff --git a/Engine/Source/Runtime/RenderCore/Public/ClearQuad.h b/Engine/Source/Runtime/RenderCore/Public/ClearQuad.h
--- a/Engine/Source/Runtime/RenderCore/Public/ClearQuad.h
+++ b/Engine/Source/Runtime/RenderCore/Public/ClearQuad.h
@@ -5,7 +5,6 @@
 #include "CoreMinimal.h"
 #include "RHIDefinitions.h"
 #include "RenderResource.h"
-#include "RenderGraphBuilder.h"
 
 class FRHICommandList;
 struct FRWBufferStructured;
@@ -81,9 +80,4 @@ inline void DrawClearQuad(FRHICommandList& RHICmdList, const FLinearColor& Color
 inline void DrawClearQuad(FRHICommandList& RHICmdList, const FLinearColor& Color, FClearQuadCallbacks ClearQuadCallbacks)
 {
 	DrawClearQuadMRT(RHICmdList, true, 1, &Color, false, 0, false, 0, ClearQuadCallbacks);
-}
-
-extern RENDERCORE_API void ClearUAV(FRDGBuilder& GraphBuilder, FRDGEventName&& PassName, FRDGBufferUAVRef BufferUAV, uint32 Value);
-extern RENDERCORE_API void ClearUAV(FRDGBuilder& GraphBuilder, FRDGEventName&& PassName, FRDGTextureUAVRef TextureUAV, const float(&ClearValues)[4]);
-extern RENDERCORE_API void ClearUAV(FRDGBuilder& GraphBuilder, FRDGEventName&& PassName, FRDGTextureUAVRef TextureUAV, const uint32(&ClearValues)[4]);
-extern RENDERCORE_API void ClearUAV(FRDGBuilder& GraphBuilder, FRDGEventName&& PassName, FRDGTextureUAVRef TextureUAV, const FLinearColor& ClearColor);
\ No newline at end of file
+}
\ No newline at end of file
diff --git a/Engine/Source/Runtime/RenderCore/Public/RenderGraphBuilder.h b/Engine/Source/Runtime/RenderCore/Public/RenderGraphBuilder.h
--- a/Engine/Source/Runtime/RenderCore/Public/RenderGraphBuilder.h
+++ b/Engine/Source/Runtime/RenderCore/Public/RenderGraphBuilder.h
@@ -204,9 +204,13 @@ private:
 
 #if RDG_ENABLE_DEBUG
 	FRDGUserValidation Validation;
+
+	/** Tracks whether we are in a scope of adding passes to the builder. Used to avoid recursion. */
+	bool bInDebugPassScope = false;
 #endif
 
 	void VisualizePassOutputs(const FRDGPass* Pass);
+	void ClobberPassOutputs(const FRDGPass* Pass);
 
 	void WalkGraphDependencies();
 	
diff --git a/Engine/Source/Runtime/RenderCore/Public/RenderGraphResources.h b/Engine/Source/Runtime/RenderCore/Public/RenderGraphResources.h
--- a/Engine/Source/Runtime/RenderCore/Public/RenderGraphResources.h
+++ b/Engine/Source/Runtime/RenderCore/Public/RenderGraphResources.h
@@ -230,9 +230,12 @@ private:
 	/** Count the number of times it has been used by a pass. */
 	int32 PassAccessCount = 0;
 
-	/** Boolean to track at wiring time if a resource has ever been produced by a pass, to error out early if accessing a resource that has not been produced. */
+	/** Tracks at wiring time if a resource has ever been produced by a pass, to error out early if accessing a resource that has not been produced. */
 	bool bHasBeenProduced = false;
 
+	/** Tracks whether this resource was clobbered by the builder prior to use. */
+	bool bHasBeenClobbered = false;
+
 	/** Tracks whether this resource is an externally imported resource. */
 	bool bIsExternal = false;
 #endif
diff --git a/Engine/Source/Runtime/RenderCore/Public/RenderGraphUtils.h b/Engine/Source/Runtime/RenderCore/Public/RenderGraphUtils.h
--- a/Engine/Source/Runtime/RenderCore/Public/RenderGraphUtils.h
+++ b/Engine/Source/Runtime/RenderCore/Public/RenderGraphUtils.h
@@ -8,12 +8,10 @@
 #include "ShaderParameterStruct.h"
 #include "ShaderParameterMacros.h"
 
-
 /** An empty shader parameter structure ready to be used anywhere. */
 BEGIN_SHADER_PARAMETER_STRUCT(FEmptyShaderParameters, RENDERCORE_API)
 END_SHADER_PARAMETER_STRUCT()
 
-
 /** Useful parameter struct that only have render targets.
  *
  *	FRenderTargetParameters PassParameters;
@@ -24,7 +22,6 @@ BEGIN_SHADER_PARAMETER_STRUCT(FRenderTargetParameters, RENDERCORE_API)
 	RENDER_TARGET_BINDING_SLOTS()
 END_SHADER_PARAMETER_STRUCT()
 
-
 /**
  * Clears all render graph tracked resources that are not bound by a shader.
  * Excludes any resources on the ExcludeList from being cleared regardless of whether the 
@@ -38,7 +35,7 @@ extern RENDERCORE_API void ClearUnusedGraphResourcesImpl(
 	std::initializer_list< FRDGResourceRef > ExcludeList);
 
 template<typename TShaderClass>
-FORCEINLINE_DEBUGGABLE void ClearUnusedGraphResources(
+void ClearUnusedGraphResources(
 	const TShaderClass* Shader,
 	typename TShaderClass::FParameters* InoutParameters,
 	std::initializer_list< FRDGResourceRef > ExcludeList = {})
@@ -54,7 +51,6 @@ FORCEINLINE_DEBUGGABLE void ClearUnusedGraphResources(
 	return ClearUnusedGraphResourcesImpl(Shader->Bindings, ParametersMetadata, InoutParameters, ExcludeList);
 }
 
-
 /**
  * Register external texture with fallback if the resource is invalid.
  *
@@ -67,7 +63,6 @@ RENDERCORE_API FRDGTextureRef RegisterExternalTextureWithFallback(
 	const TRefCountPtr<IPooledRenderTarget>& FallbackPooledTexture,
 	const TCHAR* ExternalPooledTextureName = TEXT("External"));
 
-
 /** All utils for compute shaders.
  */
 struct RENDERCORE_API FComputeShaderUtils
@@ -108,7 +103,7 @@ struct RENDERCORE_API FComputeShaderUtils
 
 	/** Dispatch a compute shader to rhi command list with its parameters. */
 	template<typename TShaderClass>
-	static FORCEINLINE_DEBUGGABLE void Dispatch(FRHICommandList& RHICmdList, const TShaderClass* ComputeShader, const typename TShaderClass::FParameters& Parameters, FIntVector GroupCount)
+	static void Dispatch(FRHICommandList& RHICmdList, const TShaderClass* ComputeShader, const typename TShaderClass::FParameters& Parameters, FIntVector GroupCount)
 	{
 		FRHIComputeShader* ShaderRHI = ComputeShader->GetComputeShader();
 		RHICmdList.SetComputeShader(ShaderRHI);
@@ -119,7 +114,7 @@ struct RENDERCORE_API FComputeShaderUtils
 	
 	/** Indirect dispatch a compute shader to rhi command list with its parameters. */
 	template<typename TShaderClass>
-	static FORCEINLINE_DEBUGGABLE void DispatchIndirect(
+	static void DispatchIndirect(
 		FRHICommandList& RHICmdList,
 		const TShaderClass* ComputeShader,
 		const typename TShaderClass::FParameters& Parameters,
@@ -135,7 +130,7 @@ struct RENDERCORE_API FComputeShaderUtils
 
 	/** Dispatch a compute shader to render graph builder with its parameters. */
 	template<typename TShaderClass>
-	static FORCEINLINE_DEBUGGABLE void AddPass(
+	static void AddPass(
 		FRDGBuilder& GraphBuilder,
 		FRDGEventName&& PassName,
 		const TShaderClass* ComputeShader,
@@ -156,7 +151,7 @@ struct RENDERCORE_API FComputeShaderUtils
 
 	/** Dispatch a compute shader to render graph builder with its parameters. */
 	template<typename TShaderClass>
-	static FORCEINLINE_DEBUGGABLE void AddPass(
+	static void AddPass(
 		FRDGBuilder& GraphBuilder,
 		FRDGEventName&& PassName,
 		const TShaderClass* ComputeShader,
@@ -183,9 +178,8 @@ struct RENDERCORE_API FComputeShaderUtils
 	}
 };
 
-/** Adds a render graph pass that copies a region from one texture to another;
- *  Uses RHICopyTexture under the hood. Formats of the two textures must match.
- *  The output and output texture regions be within the respective extents.
+/** Adds a render graph pass to copy a region from one texture to another. Uses RHICopyTexture under the hood.
+ *  Formats of the two textures must match. The output and output texture regions be within the respective extents.
  */
 RENDERCORE_API void AddCopyTexturePass(
 	FRDGBuilder& GraphBuilder,
@@ -218,3 +212,24 @@ inline void AddCopyTexturePass(
 	}
 	AddCopyTexturePass(GraphBuilder, InputTexture, OutputTexture, CopyInfo);
 }
+
+/** Adds a render graph pass to clear a texture or buffer UAV with a single typed value. */
+RENDERCORE_API void AddClearUAVPass(FRDGBuilder& GraphBuilder, FRDGBufferUAVRef BufferUAV, uint32 Value);
+
+RENDERCORE_API void AddClearUAVPass(FRDGBuilder& GraphBuilder, FRDGTextureUAVRef TextureUAV, const float(&ClearValues)[4]);
+
+RENDERCORE_API void AddClearUAVPass(FRDGBuilder& GraphBuilder, FRDGTextureUAVRef TextureUAV, const uint32(&ClearValues)[4]);
+
+RENDERCORE_API void AddClearUAVPass(FRDGBuilder& GraphBuilder, FRDGTextureUAVRef TextureUAV, const FLinearColor& ClearColor);
+
+/** Adds a render graph pass to clear a render target. Prefer to use clear actions if possible. */
+RENDERCORE_API void AddClearRenderTargetPass(FRDGBuilder& GraphBuilder, FRDGTextureRef Texture, const FLinearColor& ClearColor);
+
+/** Adds a render graph pass to clear a depth stencil target. Prefer to use clear actions if possible. */
+RENDERCORE_API void AddClearDepthStencilPass(
+	FRDGBuilder& GraphBuilder,
+	FRDGTextureRef Texture,
+	bool bClearDepth,
+	float Depth,
+	bool bClearStencil,
+	uint8 Stencil);
\ No newline at end of file
diff --git a/Engine/Source/Runtime/RenderCore/Public/RenderGraphValidation.h b/Engine/Source/Runtime/RenderCore/Public/RenderGraphValidation.h
--- a/Engine/Source/Runtime/RenderCore/Public/RenderGraphValidation.h
+++ b/Engine/Source/Runtime/RenderCore/Public/RenderGraphValidation.h
@@ -56,8 +56,12 @@ public:
 	/** Validates a resource extraction operation. */
 	void ValidateExtractResource(FRDGParentResourceRef Resource);
 
-	/** Tracks and validates the addition of a new pass to the graph. */
-	void ValidateAddPass(const FRDGPass* Pass);
+	/** Tracks and validates the addition of a new pass to the graph.
+	 *  @param bSkipPassAccessMarking Skips marking the pass as a producer or incrementing the pass access. Useful when
+	 *      the builder needs to inject a pass for debugging while preserving error messages and warnings for the original
+	 *      graph structure.
+	 */
+	void ValidateAddPass(const FRDGPass* Pass, bool bSkipPassAccessMarking);
 
 	/** Validate pass state before and after execution. */
 	void ValidateExecutePassBegin(const FRDGPass* Pass);
diff --git a/Engine/Source/Runtime/Renderer/Private/LightGridInjection.cpp b/Engine/Source/Runtime/Renderer/Private/LightGridInjection.cpp
--- a/Engine/Source/Runtime/Renderer/Private/LightGridInjection.cpp
+++ b/Engine/Source/Runtime/Renderer/Private/LightGridInjection.cpp
@@ -724,9 +724,9 @@ void FDeferredShadingSceneRenderer::ComputeLightGrid(FRHICommandListImmediate& R
 
 					if (GLightLinkedListCulling != 0)
 					{
-						ClearUAV(GraphBuilder, RDG_EVENT_NAME("Clear:StartOffsetGrid"), PassParameters->RWStartOffsetGrid, 0xFFFFFFFF);
-						ClearUAV(GraphBuilder, RDG_EVENT_NAME("Clear:NextCulledLightLink"), PassParameters->RWNextCulledLightLink, 0);
-						ClearUAV(GraphBuilder, RDG_EVENT_NAME("Clear:NextCulledLightData"), GraphBuilder.CreateUAV(NextCulledLightDataBuffer, PF_R32_UINT), 0);
+						AddClearUAVPass(GraphBuilder, PassParameters->RWStartOffsetGrid, 0xFFFFFFFF);
+						AddClearUAVPass(GraphBuilder, PassParameters->RWNextCulledLightLink, 0);
+						AddClearUAVPass(GraphBuilder, GraphBuilder.CreateUAV(NextCulledLightDataBuffer, PF_R32_UINT), 0);
 						FComputeShaderUtils::AddPass(GraphBuilder, RDG_EVENT_NAME("LightGridInject:LinkedList"), *ComputeShader, PassParameters, NumGroups);
 
 
diff --git a/Engine/Source/Runtime/Renderer/Private/PostProcess/DiaphragmDOF.cpp b/Engine/Source/Runtime/Renderer/Private/PostProcess/DiaphragmDOF.cpp
--- a/Engine/Source/Runtime/Renderer/Private/PostProcess/DiaphragmDOF.cpp
+++ b/Engine/Source/Runtime/Renderer/Private/PostProcess/DiaphragmDOF.cpp
@@ -1908,7 +1908,7 @@ FRDGTextureRef DiaphragmDOF::AddPasses(
 				if (BackgroundScatterDrawListBuffer)
 					PassParameters->OutBackgroundScatterDrawList = GraphBuilder.CreateUAV(BackgroundScatterDrawListBuffer);
 
-				ClearUAV(GraphBuilder, RDG_EVENT_NAME("ClearIndirectDraw"), PassParameters->OutScatterDrawIndirectParameters, 0);
+				AddClearUAVPass(GraphBuilder, PassParameters->OutScatterDrawIndirectParameters, 0);
 			}
 
 			TShaderMapRef<FDiaphragmDOFReduceCS> ComputeShader(View.ShaderMap, PermutationVector);
diff --git a/Engine/Source/Runtime/Renderer/Private/PostProcess/PostProcessEyeAdaptation.cpp b/Engine/Source/Runtime/Renderer/Private/PostProcess/PostProcessEyeAdaptation.cpp
--- a/Engine/Source/Runtime/Renderer/Private/PostProcess/PostProcessEyeAdaptation.cpp
+++ b/Engine/Source/Runtime/Renderer/Private/PostProcess/PostProcessEyeAdaptation.cpp
@@ -349,7 +349,7 @@ FRDGTextureRef AddHistogramEyeAdaptationPass(
 	{
 		FEyeAdaptationPS::FParameters* PassParameters = GraphBuilder.AllocParameters<FEyeAdaptationPS::FParameters>();
 		PassParameters->Base = PassBaseParameters;
-		PassParameters->RenderTargets[0] = FRenderTargetBinding(OutputTexture, ERenderTargetLoadAction::ELoad);
+		PassParameters->RenderTargets[0] = FRenderTargetBinding(OutputTexture, ERenderTargetLoadAction::ENoAction);
 
 		TShaderMapRef<FEyeAdaptationPS> PixelShader(ScreenPassView.View.ShaderMap);
 
@@ -602,7 +602,7 @@ FRDGTextureRef AddBasicEyeAdaptationPass(
 	{
 		FBasicEyeAdaptationPS::FParameters* PassParameters = GraphBuilder.AllocParameters<FBasicEyeAdaptationPS::FParameters>();
 		PassParameters->Base = PassBaseParameters;
-		PassParameters->RenderTargets[0] = FRenderTargetBinding(OutputTexture, ERenderTargetLoadAction::ELoad);
+		PassParameters->RenderTargets[0] = FRenderTargetBinding(OutputTexture, ERenderTargetLoadAction::ENoAction);
 
 		TShaderMapRef<FBasicEyeAdaptationPS> PixelShader(ScreenPassView.View.ShaderMap);
 
