Clobber depth now uses something more exotic than 0.

diff --git a/Engine/Source/Runtime/RenderCore/Private/RenderGraphBuilder.cpp b/Engine/Source/Runtime/RenderCore/Private/RenderGraphBuilder.cpp
--- a/Engine/Source/Runtime/RenderCore/Private/RenderGraphBuilder.cpp
+++ b/Engine/Source/Runtime/RenderCore/Private/RenderGraphBuilder.cpp
@@ -372,7 +372,11 @@ void FRDGBuilder::ClobberPassOutputs(const FRDGPass* Pass)
 			{
 				if (TryMarkForClobber(Texture))
 				{
-					AddClearDepthStencilPass(*this, Texture, true, 0.0f, true, 0);
+					// These are arbitrarily chosen to be something unusual.
+					const float ClobberDepth = 0.56789f;
+					const uint8 ClobberStencil = 123;
+
+					AddClearDepthStencilPass(*this, Texture, true, ClobberDepth, true, ClobberStencil);
 				}
 			}
 
