Remove MSAA assert in render graph builder.

diff --git a/Engine/Source/Runtime/RenderCore/Private/RenderGraphBuilder.cpp b/Engine/Source/Runtime/RenderCore/Private/RenderGraphBuilder.cpp
--- a/Engine/Source/Runtime/RenderCore/Private/RenderGraphBuilder.cpp
+++ b/Engine/Source/Runtime/RenderCore/Private/RenderGraphBuilder.cpp
@@ -897,8 +897,6 @@ void FRDGBuilder::PrepareResourcesForExecute(const FRDGPass* Pass, struct FRHIRe
 				ERenderTargetStoreAction DepthStoreAction = ExclusiveDepthStencil.IsDepthWrite() ? ERenderTargetStoreAction::EStore : ERenderTargetStoreAction::ENoAction;
 				ERenderTargetStoreAction StencilStoreAction = ExclusiveDepthStencil.IsStencilWrite() ? ERenderTargetStoreAction::EStore : ERenderTargetStoreAction::ENoAction;
 
-				// TODO(RDG): Addresses the TODO of the color scene render target.
-				ensureMsgf(Texture->Desc.NumSamples == 1, TEXT("MSAA dept-stencil render target not yet supported."));
 				OutDepthStencil.DepthStencilTarget = Texture->PooledRenderTarget->GetRenderTargetItem().TargetableTexture;
 				OutDepthStencil.ResolveTarget = nullptr;
 				OutDepthStencil.Action = MakeDepthStencilTargetActions(
