Fixes for build warnings.

diff --git a/Engine/Source/Runtime/RenderCore/Private/RenderGraphUtils.cpp b/Engine/Source/Runtime/RenderCore/Private/RenderGraphUtils.cpp
--- a/Engine/Source/Runtime/RenderCore/Private/RenderGraphUtils.cpp
+++ b/Engine/Source/Runtime/RenderCore/Private/RenderGraphUtils.cpp
@@ -162,21 +162,21 @@ void AddClearUAVPass_T(FRDGBuilder& GraphBuilder, FRDGTextureUAVRef TextureUAV,
 		ERDGPassFlags::Compute,
 		[&Parameters, TextureUAV, ClearValues](FRHICommandList& RHICmdList)
 	{
-		const FRDGTextureDesc& TextureDesc = TextureUAV->GetParent()->Desc;
+		const FRDGTextureDesc& LocalTextureDesc = TextureUAV->GetParent()->Desc;
 
 		FRHIUnorderedAccessView* RHITextureUAV = TextureUAV->GetRHI();
 
-		if (TextureDesc.Is2DTexture())
+		if (LocalTextureDesc.Is2DTexture())
 		{
-			if (TextureDesc.IsArray())
+			if (LocalTextureDesc.IsArray())
 			{
 				TShaderMapRef<FClearTexture2DArrayReplacementCS<T>> ComputeShader(GetGlobalShaderMap(GMaxRHIFeatureLevel));
 				FRHIComputeShader* ShaderRHI = ComputeShader->GetComputeShader();
 				RHICmdList.SetComputeShader(ShaderRHI);
 				ComputeShader->SetParameters(RHICmdList, RHITextureUAV, ClearValues);
-				uint32 x = (TextureDesc.Extent.X + 7) / 8;
-				uint32 y = (TextureDesc.Extent.Y + 7) / 8;
-				uint32 z = TextureDesc.ArraySize;
+				uint32 x = (LocalTextureDesc.Extent.X + 7) / 8;
+				uint32 y = (LocalTextureDesc.Extent.Y + 7) / 8;
+				uint32 z = LocalTextureDesc.ArraySize;
 				RHICmdList.DispatchComputeShader(x, y, z);
 				ComputeShader->FinalizeParameters(RHICmdList, RHITextureUAV);
 			}
@@ -186,32 +186,32 @@ void AddClearUAVPass_T(FRDGBuilder& GraphBuilder, FRDGTextureUAVRef TextureUAV,
 				FRHIComputeShader* ShaderRHI = ComputeShader->GetComputeShader();
 				RHICmdList.SetComputeShader(ShaderRHI);
 				ComputeShader->SetParameters(RHICmdList, RHITextureUAV, ClearValues);
-				const uint32 X = (TextureDesc.Extent.X + 7) / 8;
-				const uint32 Y = (TextureDesc.Extent.Y + 7) / 8;
+				const uint32 X = (LocalTextureDesc.Extent.X + 7) / 8;
+				const uint32 Y = (LocalTextureDesc.Extent.Y + 7) / 8;
 				RHICmdList.DispatchComputeShader(X, Y, 1);
 				ComputeShader->FinalizeParameters(RHICmdList, RHITextureUAV);
 			}
 		}
-		else if (TextureDesc.IsCubemap())
+		else if (LocalTextureDesc.IsCubemap())
 		{
 			TShaderMapRef<FClearTexture2DArrayReplacementCS<T>> ComputeShader(GetGlobalShaderMap(GMaxRHIFeatureLevel));
 			FRHIComputeShader* ShaderRHI = ComputeShader->GetComputeShader();
 			RHICmdList.SetComputeShader(ShaderRHI);
 			ComputeShader->SetParameters(RHICmdList, RHITextureUAV, ClearValues);
-			uint32 x = (TextureDesc.Extent.X + 7) / 8;
-			uint32 y = (TextureDesc.Extent.Y + 7) / 8;
+			uint32 x = (LocalTextureDesc.Extent.X + 7) / 8;
+			uint32 y = (LocalTextureDesc.Extent.Y + 7) / 8;
 			RHICmdList.DispatchComputeShader(x, y, 6);
 			ComputeShader->FinalizeParameters(RHICmdList, RHITextureUAV);
 		}
-		else if (TextureDesc.Is3DTexture())
+		else if (LocalTextureDesc.Is3DTexture())
 		{
 			TShaderMapRef<FClearVolumeReplacementCS<T>> ComputeShader(GetGlobalShaderMap(GMaxRHIFeatureLevel));
 			FRHIComputeShader* ShaderRHI = ComputeShader->GetComputeShader();
 			RHICmdList.SetComputeShader(ShaderRHI);
 			ComputeShader->SetParameters(RHICmdList, RHITextureUAV, ClearValues);
-			uint32 x = (TextureDesc.Extent.X + 3) / 4;
-			uint32 y = (TextureDesc.Extent.Y + 3) / 4;
-			uint32 z = (TextureDesc.Depth + 3) / 4;
+			uint32 x = (LocalTextureDesc.Extent.X + 3) / 4;
+			uint32 y = (LocalTextureDesc.Extent.Y + 3) / 4;
+			uint32 z = (LocalTextureDesc.Depth + 3) / 4;
 			RHICmdList.DispatchComputeShader(x, y, z);
 			ComputeShader->FinalizeParameters(RHICmdList, RHITextureUAV);
 		}
diff --git a/Engine/Source/Runtime/Renderer/Private/CompositionLighting/PostProcessAmbientOcclusion.cpp b/Engine/Source/Runtime/Renderer/Private/CompositionLighting/PostProcessAmbientOcclusion.cpp
--- a/Engine/Source/Runtime/Renderer/Private/CompositionLighting/PostProcessAmbientOcclusion.cpp
+++ b/Engine/Source/Runtime/Renderer/Private/CompositionLighting/PostProcessAmbientOcclusion.cpp
@@ -668,14 +668,7 @@ class FPostProcessAmbientOcclusionPSandCS : public FGlobalShader
 
 	static bool ShouldCompilePermutation(const FGlobalShaderPermutationParameters& Parameters)
 	{
-		if(bComputeShader)
-		{
-			return IsFeatureLevelSupported(Parameters.Platform, ERHIFeatureLevel::SM5);
-		}
-		else
-		{
-			return IsFeatureLevelSupported(Parameters.Platform, ERHIFeatureLevel::SM5);
-		}
+		return IsFeatureLevelSupported(Parameters.Platform, ERHIFeatureLevel::SM5);
 	}
 
 	static void ModifyCompilationEnvironment(const FGlobalShaderPermutationParameters& Parameters, FShaderCompilerEnvironment& OutEnvironment)
diff --git a/Engine/Source/Runtime/Renderer/Private/LightGridInjection.cpp b/Engine/Source/Runtime/Renderer/Private/LightGridInjection.cpp
--- a/Engine/Source/Runtime/Renderer/Private/LightGridInjection.cpp
+++ b/Engine/Source/Runtime/Renderer/Private/LightGridInjection.cpp
@@ -87,33 +87,24 @@ public:
 	{
 		if (GMaxRHIFeatureLevel >= ERHIFeatureLevel::SM5)
 		{
-			if (GMaxRHIFeatureLevel >= ERHIFeatureLevel::SM5)
-			{
-				ForwardLightingResources.ForwardLocalLightBuffer.Initialize(sizeof(FVector4), sizeof(FForwardLocalLightData) / sizeof(FVector4), PF_A32B32G32R32F, BUF_Dynamic);
-				ForwardLightingResources.NumCulledLightsGrid.Initialize(sizeof(uint32), 1, PF_R32_UINT);
-
-				const bool bSupportFormatConversion = RHISupportsBufferLoadTypeConversion(GMaxRHIShaderPlatform);
+			ForwardLightingResources.ForwardLocalLightBuffer.Initialize(sizeof(FVector4), sizeof(FForwardLocalLightData) / sizeof(FVector4), PF_A32B32G32R32F, BUF_Dynamic);
+			ForwardLightingResources.NumCulledLightsGrid.Initialize(sizeof(uint32), 1, PF_R32_UINT);
 
-				if (bSupportFormatConversion)
-				{
-					ForwardLightingResources.CulledLightDataGrid.Initialize(sizeof(uint16), 1, PF_R16_UINT);
-				}
-				else
-				{
-					ForwardLightingResources.CulledLightDataGrid.Initialize(sizeof(uint32), 1, PF_R32_UINT);
-				}
+			const bool bSupportFormatConversion = RHISupportsBufferLoadTypeConversion(GMaxRHIShaderPlatform);
 
-				ForwardLightingResources.ForwardLightData.ForwardLocalLightBuffer = ForwardLightingResources.ForwardLocalLightBuffer.SRV;
-				ForwardLightingResources.ForwardLightData.NumCulledLightsGrid = ForwardLightingResources.NumCulledLightsGrid.SRV;
-				ForwardLightingResources.ForwardLightData.CulledLightDataGrid = ForwardLightingResources.CulledLightDataGrid.SRV;
+			if (bSupportFormatConversion)
+			{
+				ForwardLightingResources.CulledLightDataGrid.Initialize(sizeof(uint16), 1, PF_R16_UINT);
 			}
 			else
 			{
-				ForwardLightingResources.ForwardLightData.ForwardLocalLightBuffer = GNullColorVertexBuffer.VertexBufferSRV;
-				ForwardLightingResources.ForwardLightData.NumCulledLightsGrid = GNullColorVertexBuffer.VertexBufferSRV;
-				ForwardLightingResources.ForwardLightData.CulledLightDataGrid = GNullColorVertexBuffer.VertexBufferSRV;
+				ForwardLightingResources.CulledLightDataGrid.Initialize(sizeof(uint32), 1, PF_R32_UINT);
 			}
 
+			ForwardLightingResources.ForwardLightData.ForwardLocalLightBuffer = ForwardLightingResources.ForwardLocalLightBuffer.SRV;
+			ForwardLightingResources.ForwardLightData.NumCulledLightsGrid = ForwardLightingResources.NumCulledLightsGrid.SRV;
+			ForwardLightingResources.ForwardLightData.CulledLightDataGrid = ForwardLightingResources.CulledLightDataGrid.SRV;
+
 			ForwardLightingResources.ForwardLightDataUniformBuffer = TUniformBufferRef<FForwardLightData>::CreateUniformBufferImmediate(ForwardLightingResources.ForwardLightData, UniformBuffer_MultiFrame);
 		}
 	}
diff --git a/Engine/Source/Runtime/Renderer/Private/MobileShadingRenderer.cpp b/Engine/Source/Runtime/Renderer/Private/MobileShadingRenderer.cpp
--- a/Engine/Source/Runtime/Renderer/Private/MobileShadingRenderer.cpp
+++ b/Engine/Source/Runtime/Renderer/Private/MobileShadingRenderer.cpp
@@ -228,6 +228,8 @@ void FMobileSceneRenderer::InitViews(FRHICommandListImmediate& RHICmdList)
 	SCOPE_CYCLE_COUNTER(STAT_InitViewsTime);
 	CSV_SCOPED_TIMING_STAT_EXCLUSIVE(InitViews_Scene);
 
+	check(Scene);
+
 	FILCUpdatePrimTaskData ILCTaskData;
 	FViewVisibleCommandsPerView ViewCommandsPerView;
 	ViewCommandsPerView.SetNum(Views.Num());
@@ -238,7 +240,7 @@ void FMobileSceneRenderer::InitViews(FRHICommandListImmediate& RHICmdList)
 	ComputeViewVisibility(RHICmdList, BasePassDepthStencilAccess, ViewCommandsPerView, DynamicIndexBuffer, DynamicVertexBuffer, DynamicReadBuffer);
 
 	// Initialise Sky/View resources before the view global uniform buffer is built.
-	if (Scene && ShouldRenderSkyAtmosphere(Scene->GetSkyAtmosphereSceneInfo(), Scene->GetShaderPlatform()))
+	if (ShouldRenderSkyAtmosphere(Scene->GetSkyAtmosphereSceneInfo(), Scene->GetShaderPlatform()))
 	{
 		InitSkyAtmosphereForViews(RHICmdList);
 	}
@@ -276,7 +278,7 @@ void FMobileSceneRenderer::InitViews(FRHICommandListImmediate& RHICmdList)
 		// Create the directional light uniform buffers
 		CreateDirectionalLightUniformBuffers(Views[ViewIndex]);
 	}
-	
+
 	UpdateGPUScene(RHICmdList, *Scene);
 	for (int32 ViewIndex = 0; ViewIndex < Views.Num(); ViewIndex++)
 	{
