Adjusting RDG validation to allow multiple simultaneous access textures in a GenerateMips pass.

diff --git a/Engine/Source/Runtime/RenderCore/Private/RenderGraphValidation.cpp b/Engine/Source/Runtime/RenderCore/Private/RenderGraphValidation.cpp
--- a/Engine/Source/Runtime/RenderCore/Private/RenderGraphValidation.cpp
+++ b/Engine/Source/Runtime/RenderCore/Private/RenderGraphValidation.cpp
@@ -23,21 +23,57 @@ public:
 		, bIsGenerateMips(InPass.IsGenerateMips())
 	{}
 
+	~FRDGPassResourceValidator()
+	{
+		for (const auto& KeyValue : BufferAccessMap)
+		{
+			const FRDGBufferRef Buffer = KeyValue.Key;
+			const FBufferAccess& BufferAccess = KeyValue.Value;
+	
+			ensureMsgf(!(BufferAccess.Bits.Read && BufferAccess.Bits.Write),
+				TEXT("Attempting to use buffer '%s' for both SRV and UAV access in pass '%s'. Only one usage at a time is allowed."),
+				Buffer->Name, Pass.GetName());
+		}
+
+		for (const auto& KeyValue : TextureAccessMap)
+		{
+			const FRDGTextureRef Texture = KeyValue.Key;
+			const FTextureAccess& TextureAccess = KeyValue.Value;
+
+			// A generate mips pass requires both read / write access to the same texture for barriers to work correctly.
+			if (bIsGenerateMips)
+			{
+				const bool bReadWithoutWrite = TextureAccess.IsAnyRead() && !TextureAccess.IsAnyWrite();
+				ensureMsgf(!bReadWithoutWrite,
+					TEXT("Attempting to use texture '%s' for read but not write access in pass '%s' which is marked with 'GenerateMips'. This kind of pass only supports textures for simultaneous read and write."),
+					Texture->Name, Pass.GetName());
+
+				const bool bWriteWithoutRead = TextureAccess.IsAnyWrite() && !TextureAccess.IsAnyRead();
+				ensureMsgf(!bWriteWithoutRead,
+					TEXT("Attempting to use texture '%s' for write but not read access in pass '%s' which is marked with 'GenerateMips'. This kind of pass only supports textures for simultaneous read and write."),
+					Texture->Name, Pass.GetName());
+			}
+			// A normal pass requires only read OR write access to the same resource for barriers to work correctly.
+			else
+			{
+				const bool bReadAndWrite = TextureAccess.IsAnyRead() && TextureAccess.IsAnyWrite();
+				ensureMsgf(!bReadAndWrite,
+					TEXT("Attempting to use texture '%s' for both read and write access in pass '%s' which is NOT marked with 'GenerateMips'. You must specify GenerateMips in your pass flags for this to be valid."),
+					Texture->Name, Pass.GetName());
+			}
+		}
+	}
+
 	void Read(FRDGBufferRef Buffer)
 	{
 		check(Buffer);
 
 		ensureMsgf(!bIsGenerateMips,
-			TEXT("Attempting to read from Buffer %s on Pass %s which is marked as 'GenerateMips'. Only a single texture is supported on this type of pass."),
+			TEXT("Attempting to read from Buffer %s on Pass %s which is marked as 'GenerateMips'. Only textures are supported on this type of pass."),
 			Buffer->Name, Pass.GetName());
 
 		FBufferAccess& BufferAccess = BufferAccessMap.FindOrAdd(Buffer);
-
-		ensureMsgf(BufferAccess.Access != FRDGResourceState::EAccess::Write,
-			TEXT("Buffer %s is being used for both read and write access on Pass %s. This is not supported."),
-			Buffer->Name, Pass.GetName());
-
-		BufferAccess.Access = FRDGResourceState::EAccess::Read;
+		BufferAccess.Bits.Read = 1;
 	}
 
 	void Write(FRDGBufferRef Buffer)
@@ -45,16 +81,11 @@ public:
 		check(Buffer);
 
 		ensureMsgf(!bIsGenerateMips,
-			TEXT("Attempting to write to Buffer %s on Pass %s which is marked as 'GenerateMips'. Only a single texture is supported on this type of pass."),
+			TEXT("Attempting to write to Buffer %s on Pass %s which is marked as 'GenerateMips'. Only textures are supported on this type of pass."),
 			Buffer->Name, Pass.GetName());
 
 		FBufferAccess& BufferAccess = BufferAccessMap.FindOrAdd(Buffer);
-
-		ensureMsgf(BufferAccess.Access != FRDGResourceState::EAccess::Read,
-			TEXT("Attempting to simultaneously read and write from Buffer %s on Pass %s. This is not supported."),
-			Buffer->Name, Pass.GetName());
-
-		BufferAccess.Access = FRDGResourceState::EAccess::Read;
+		BufferAccess.Bits.Write = 1;
 	}
 
 	void Read(FRDGTextureRef Texture, uint32 MipLevel = AllMipLevels)
@@ -65,23 +96,6 @@ public:
 
 		FTextureAccess& TextureAccess = TextureAccessMap.FindOrAdd(Texture);
 		TextureAccess.SetAllRead(MipMask);
-
-		ensureMsgf(!TextureAccess.IsAnyWrite(MipMask),
-			TEXT("Attempting to access mips of Texture %s for both read and write access on Pass %s. This is not supported."),
-			MipLevel, Texture->Name, Pass.GetName());
-
-		if (bIsGenerateMips)
-		{
-			ensureMsgf(TextureAccessMap.Num() == 1,
-				TEXT("Attempting to access more than one texture in Pass %s when 'GenerateMips' is set. Only a single texture is supported on this type of pass."),
-				Pass.GetName());
-		}
-		else
-		{
-			ensureMsgf(!TextureAccess.Mips.Write,
-				TEXT("Texture %s is being used for both read and write access on the same Pass %s. This is only supported via the 'GenerateMips' flag."),
-				Pass.GetName());
-		}
 	}
 
 	void Write(FRDGTextureRef Texture, uint32 MipLevel = AllMipLevels)
@@ -92,38 +106,30 @@ public:
 
 		FTextureAccess& TextureAccess = TextureAccessMap.FindOrAdd(Texture);
 		TextureAccess.SetAllWrite(MipMask);
-
-		ensureMsgf(!TextureAccess.IsAnyRead(MipMask),
-			TEXT("Attempting to use Mip %d of Texture %s for both read and write access on Pass %s. This is not supported."),
-			MipLevel, Texture->Name, Pass.GetName());
-
-		if (bIsGenerateMips)
-		{
-			ensureMsgf(TextureAccessMap.Num() == 1,
-				TEXT("Attempting to access more than one texture in Pass %s when 'GenerateMips' is set. Only a single texture is supported on this type of pass."),
-				Pass.GetName());
-		}
-		else
-		{
-			ensureMsgf(!TextureAccess.Mips.Read,
-				TEXT("Texture %s is being used for both read and write access on the same Pass %s. This is only supported via the 'GenerateMips' flag."),
-				Texture->Name, Pass.GetName());
-		}
 	}
 
 private:
 	struct FBufferAccess
 	{
-		FRDGResourceState::EAccess Access = FRDGResourceState::EAccess::Unknown;
+		union
+		{
+			uint8 PackedData = 0;
+
+			struct
+			{
+				uint8 Read : 1;
+				uint8 Write : 1;
+			} Bits;
+		};
 
 		FORCEINLINE bool operator!= (FBufferAccess Other) const
 		{
-			return Access != Other.Access;
+			return PackedData != Other.PackedData;
 		}
 
 		FORCEINLINE bool operator< (FBufferAccess Other) const
 		{
-			return Access < Other.Access;
+			return PackedData < Other.PackedData;
 		}
 	};
 
@@ -142,22 +148,22 @@ private:
 			} Mips;
 		};
 
-		FORCEINLINE bool IsAnyRead(uint32 MipMask) const
+		FORCEINLINE bool IsAnyRead(uint32 MipMask = AllMipLevels) const
 		{
 			return (Mips.Read & MipMask) != 0;
 		}
 
-		FORCEINLINE bool IsAnyWrite(uint32 MipMask) const
+		FORCEINLINE bool IsAnyWrite(uint32 MipMask = AllMipLevels) const
 		{
 			return (Mips.Write & MipMask) != 0;
 		}
 
-		FORCEINLINE void SetAllRead(uint32 MipMask)
+		FORCEINLINE void SetAllRead(uint32 MipMask = AllMipLevels)
 		{
 			Mips.Read |= MipMask;
 		}
 
-		FORCEINLINE void SetAllWrite(uint32 MipMask)
+		FORCEINLINE void SetAllWrite(uint32 MipMask = AllMipLevels)
 		{
 			Mips.Write |= MipMask;
 		}
