Uncrustification of RenderTargetPool (no functional changes).

diff --git a/Engine/Source/Runtime/RenderCore/Private/RenderTargetPool.cpp b/Engine/Source/Runtime/RenderCore/Private/RenderTargetPool.cpp
--- a/Engine/Source/Runtime/RenderCore/Private/RenderTargetPool.cpp
+++ b/Engine/Source/Runtime/RenderCore/Private/RenderTargetPool.cpp
@@ -16,6 +16,7 @@ RENDERCORE_API void DumpRenderTargetPoolMemory(FOutputDevice& OutputDevice)
 {
 	GRenderTargetPool.DumpMemoryUsage(OutputDevice);
 }
+
 static FAutoConsoleCommandWithOutputDevice GDumpRenderTargetPoolMemoryCmd(
 	TEXT("r.DumpRenderTargetPoolMemory"),
 	TEXT("Dump allocation information for the render target pool."),
@@ -25,12 +26,12 @@ static FAutoConsoleCommandWithOutputDevice GDumpRenderTargetPoolMemoryCmd(
 void RenderTargetPoolEvents(const TArray<FString>& Args)
 {
 	uint32 SizeInKBThreshold = -1;
-	if(Args.Num() && Args[0].IsNumeric())
+	if (Args.Num() && Args[0].IsNumeric())
 	{
 		SizeInKBThreshold = FCString::Atof(*Args[0]);
 	}
 
-	if(SizeInKBThreshold != -1)
+	if (SizeInKBThreshold != -1)
 	{
 		UE_LOG(LogRenderTargetPool, Display, TEXT("r.DumpRenderTargetPoolEvents is now enabled, use r.DumpRenderTargetPoolEvents ? for help"));
 
@@ -85,7 +86,7 @@ IPooledRenderTarget* FRenderTargetPoolEvent::GetValidatedPointer() const
 {
 	int32 Index = GRenderTargetPool.FindIndex(Pointer);
 
-	if(Index >= 0)
+	if (Index >= 0)
 	{
 		return Pointer;
 	}
@@ -95,12 +96,12 @@ IPooledRenderTarget* FRenderTargetPoolEvent::GetValidatedPointer() const
 
 bool FRenderTargetPoolEvent::NeedsDeallocEvent()
 {
-	if(GetEventType() == ERTPE_Alloc)
+	if (GetEventType() == ERTPE_Alloc)
 	{
-		if(Pointer)
+		if (Pointer)
 		{
 			IPooledRenderTarget* ValidPointer = GetValidatedPointer();
-			if(!ValidPointer || ValidPointer->IsFree())
+			if (!ValidPointer || ValidPointer->IsFree())
 			{
 				Pointer = 0;
 				return true;
@@ -111,7 +112,6 @@ bool FRenderTargetPoolEvent::NeedsDeallocEvent()
 	return false;
 }
 
-
 static uint32 ComputeSizeInKB(FPooledRenderTarget& Element)
 {
 	return (Element.ComputeMemorySize() + 1023) / 1024;
@@ -216,7 +216,7 @@ bool FRenderTargetPool::FindFreeElement(FRHICommandList& RHICmdList, const FPool
 {
 	check(IsInRenderingThread());
 
-	if(!InputDesc.IsValid())
+	if (!InputDesc.IsValid())
 	{
 		// no need to do anything
 		return true;
@@ -241,7 +241,7 @@ bool FRenderTargetPool::FindFreeElement(FRHICommandList& RHICmdList, const FPool
 	const FPooledRenderTargetDesc& Desc = bMakeTransient ? ModifiedDesc : InputDesc;
 
 	// if we can keep the current one, do that
-	if(Out)
+	if (Out)
 	{
 		FPooledRenderTarget* Current = (FPooledRenderTarget*)Out.GetReference();
 
@@ -249,7 +249,7 @@ bool FRenderTargetPool::FindFreeElement(FRHICommandList& RHICmdList, const FPool
 
 		const bool bExactMatch = true;
 
-		if(Out->GetDesc().Compare(Desc, bExactMatch))
+		if (Out->GetDesc().Compare(Desc, bExactMatch))
 		{
 			// we can reuse the same, but the debug name might have changed
 			Current->Desc.DebugName = InDebugName;
@@ -265,7 +265,7 @@ bool FRenderTargetPool::FindFreeElement(FRHICommandList& RHICmdList, const FPool
 			// release old reference, it might free a RT we can use
 			Out = 0;
 
-			if(Current->IsFree())
+			if (Current->IsFree())
 			{
 				AllocationLevelInKB -= ComputeSizeInKB(*Current);
 
@@ -300,19 +300,19 @@ bool FRenderTargetPool::FindFreeElement(FRHICommandList& RHICmdList, const FPool
 		bool bAllowMultipleDiscards = ( CVarAllowMultipleAliasingDiscardsPerFrame.GetValueOnRenderThread() != 0 );
 		// first we try exact, if that fails we try without TexCreate_FastVRAM
 		// (easily we can run out of VRam, if this search becomes a performance problem we can optimize or we should use less TexCreate_FastVRAM)
-		for(uint32 Pass = 0; Pass < PassCount; ++Pass)
+		for (uint32 Pass = 0; Pass < PassCount; ++Pass)
 		{
 			bool bExactMatch = (Pass == 0); //-V547
     
-			for(uint32 i = 0, Num = (uint32)PooledRenderTargets.Num(); i < Num; ++i)
+			for (uint32 i = 0, Num = (uint32)PooledRenderTargets.Num(); i < Num; ++i)
 			{
 				FPooledRenderTarget* Element = PooledRenderTargets[i];
-				if(Element && Element->GetDesc().Compare(Desc, bExactMatch))
+				if (Element && Element->GetDesc().Compare(Desc, bExactMatch))
 				{
 					int a = 0;
 				}
     
-				if(Element && Element->IsFree() && Element->GetDesc().Compare(Desc, bExactMatch))
+				if (Element && Element->IsFree() && Element->GetDesc().Compare(Desc, bExactMatch))
 				{
 					if ( ( Desc.Flags & TexCreate_Transient ) && bAllowMultipleDiscards == false && Element->HasBeenDiscardedThisFrame() )
 					{
@@ -326,12 +326,12 @@ bool FRenderTargetPool::FindFreeElement(FRHICommandList& RHICmdList, const FPool
 					goto Done;
 				}
 			}
-		}		
+		}
 	}
 Done:
 
-	if(!Found)
-	{		
+	if (!Found)
+	{
 		UE_LOG(LogRenderTargetPool, Display, TEXT("%d MB, NewRT %s %s"), (AllocationLevelInKB + 1023) / 1024, *Desc.GenerateInfoString(), InDebugName);
 
 		// not found in the pool, create a new element
@@ -347,10 +347,10 @@ Done:
 		FRHIResourceCreateInfo CreateInfo(Desc.ClearValue);
 		CreateInfo.DebugName = InDebugName;
 
-		if(Desc.TargetableFlags & (TexCreate_RenderTargetable | TexCreate_DepthStencilTargetable | TexCreate_UAV) && !bDeferTextureAllocation)
+		if (Desc.TargetableFlags & (TexCreate_RenderTargetable | TexCreate_DepthStencilTargetable | TexCreate_UAV) && !bDeferTextureAllocation)
 		{
 			// Only create resources if we're not asked to defer creation.
-			if(Desc.Is2DTexture())
+			if (Desc.Is2DTexture())
 			{
 				if (!Desc.IsArray())
 				{
@@ -394,7 +394,7 @@ Done:
 					Found->RenderTargetItem.RTWriteMaskBufferRHI_SRV = RHICreateShaderResourceView(Found->RenderTargetItem.RTWriteMaskDataBufferRHI);
 				}
 			}
-			else if(Desc.Is3DTexture())
+			else if (Desc.Is3DTexture())
 			{
 				Found->RenderTargetItem.ShaderResourceTexture = RHICreateTexture3D(
 					Desc.Extent.X,
@@ -411,7 +411,7 @@ Done:
 			else
 			{
 				check(Desc.IsCubemap());
-				if(Desc.IsArray())
+				if (Desc.IsArray())
 				{
 					RHICreateTargetableShaderResourceCubeArray(
 						Desc.Extent.X,
@@ -440,7 +440,6 @@ Done:
 						(FTextureCubeRHIRef&)Found->RenderTargetItem.ShaderResourceTexture
 						);
 				}
-
 			}
 
 			RHIBindDebugLabelName(Found->RenderTargetItem.TargetableTexture, InDebugName);
@@ -448,7 +447,7 @@ Done:
 		else if (!bDeferTextureAllocation)
 		{
 			// Only create resources if we're not asked to defer creation.
-			if(Desc.Is2DTexture())
+			if (Desc.Is2DTexture())
 			{
 				// this is useful to get a CPU lockable texture through the same interface
 				Found->RenderTargetItem.ShaderResourceTexture = RHICreateTexture2D(
@@ -460,7 +459,7 @@ Done:
 					Desc.Flags,
 					CreateInfo);
 			}
-			else if(Desc.Is3DTexture())
+			else if (Desc.Is3DTexture())
 			{
 				Found->RenderTargetItem.ShaderResourceTexture = RHICreateTexture3D(
 					Desc.Extent.X,
@@ -474,7 +473,7 @@ Done:
 			else 
 			{
 				check(Desc.IsCubemap());
-				if(Desc.IsArray())
+				if (Desc.IsArray())
 				{
 					FTextureCubeRHIRef CubeTexture = RHICreateTextureCubeArray(Desc.Extent.X,Desc.ArraySize,Desc.Format,Desc.NumMips,Desc.Flags | Desc.TargetableFlags | TexCreate_ShaderResource,CreateInfo);
 					Found->RenderTargetItem.TargetableTexture = Found->RenderTargetItem.ShaderResourceTexture = CubeTexture;
@@ -489,7 +488,7 @@ Done:
 			RHIBindDebugLabelName(Found->RenderTargetItem.ShaderResourceTexture, InDebugName);
 		}
 
-		if((Desc.TargetableFlags & TexCreate_UAV) && !bDeferTextureAllocation)
+		if ((Desc.TargetableFlags & TexCreate_UAV) && !bDeferTextureAllocation)
 		{
 			// The render target desc is invalid if a UAV is requested with an RHI that doesn't support the high-end feature level.
 			check(GMaxRHIFeatureLevel == ERHIFeatureLevel::SM5 || GMaxRHIFeatureLevel == ERHIFeatureLevel::ES3_1);
@@ -578,25 +577,24 @@ IPooledRenderTarget* FRenderTargetPool::MakeSnapshot(const TRefCountPtr<IPooledR
 	return NewSnapshot;
 }
 
-
 void FRenderTargetPool::GetStats(uint32& OutWholeCount, uint32& OutWholePoolInKB, uint32& OutUsedInKB) const
 {
 	OutWholeCount = (uint32)PooledRenderTargets.Num();
 	OutUsedInKB = 0;
 	OutWholePoolInKB = 0;
 		
-	for(uint32 i = 0; i < (uint32)PooledRenderTargets.Num(); ++i)
+	for (uint32 i = 0; i < (uint32)PooledRenderTargets.Num(); ++i)
 	{
 		FPooledRenderTarget* Element = PooledRenderTargets[i];
 
-		if(Element)
+		if (Element)
 		{
 			check(!Element->IsSnapshot());
 			uint32 SizeInKB = ComputeSizeInKB(*Element);
 
 			OutWholePoolInKB += SizeInKB;
 
-			if(!Element->IsFree())
+			if (!Element->IsFree())
 			{
 				OutUsedInKB += SizeInKB;
 			}
@@ -609,15 +607,15 @@ void FRenderTargetPool::GetStats(uint32& OutWholeCount, uint32& OutWholePoolInKB
 
 void FRenderTargetPool::AddPhaseEvent(const TCHAR *InPhaseName)
 {
-	if(IsEventRecordingEnabled())
+	if (IsEventRecordingEnabled())
 	{
 		AddDeallocEvents();
 
 		const FString* LastName = GetLastEventPhaseName();
 
-		if(!LastName || *LastName != InPhaseName)
+		if (!LastName || *LastName != InPhaseName)
 		{
-			if(CurrentEventRecordingTime)
+			if (CurrentEventRecordingTime)
 			{
 				// put a break to former data
 				++CurrentEventRecordingTime;
@@ -635,11 +633,11 @@ const FString* FRenderTargetPool::GetLastEventPhaseName()
 	// could be optimized but this is a debug view
 
 	// start from the end for better performance
-	for(int32 i = RenderTargetPoolEvents.Num() - 1; i >= 0; --i)
+	for (int32 i = RenderTargetPoolEvents.Num() - 1; i >= 0; --i)
 	{
 		const FRenderTargetPoolEvent* Event = &RenderTargetPoolEvents[i];
 
-		if(Event->GetEventType() == ERTPE_Phase)
+		if (Event->GetEventType() == ERTPE_Phase)
 		{
 			return &Event->GetPhaseName();
 		}
@@ -681,9 +679,6 @@ FRenderTargetPool::SMemoryStats FRenderTargetPool::ComputeView()
 			// sort criteria
 			bool operator <(const FRTPColumn& rhs) const
 			{
-				// sort VRam first (only matters on XboxOne but nice to always see it)
-// sorting only useful for XboxOne		if(bVRam != rhs.bVRam) return bVRam > rhs.bVRam;
-
 				// we want the large ones first
 				return SizeInBytes > rhs.SizeInBytes;
 			}
@@ -692,15 +687,15 @@ FRenderTargetPool::SMemoryStats FRenderTargetPool::ComputeView()
 		TArray<FRTPColumn> Colums;
 
 		// generate Colums
-		for(int32 i = 0, Num = RenderTargetPoolEvents.Num(); i < Num; i++)
+		for (int32 i = 0, Num = RenderTargetPoolEvents.Num(); i < Num; i++)
 		{
 			FRenderTargetPoolEvent* Event = &RenderTargetPoolEvents[i];
 
-			if(Event->GetEventType() == ERTPE_Alloc)
+			if (Event->GetEventType() == ERTPE_Alloc)
 			{
 				uint32 PoolEntryId = Event->GetPoolEntryId();
 
-				if(PoolEntryId >= (uint32)Colums.Num())
+				if (PoolEntryId >= (uint32)Colums.Num())
 				{
 					Colums.SetNum(PoolEntryId + 1);
 				}
@@ -714,14 +709,14 @@ FRenderTargetPool::SMemoryStats FRenderTargetPool::ComputeView()
 		{
 			uint32 ColumnX = 0;
 
-			for(int32 ColumnIndex = 0, ColumnsNum = Colums.Num(); ColumnIndex < ColumnsNum; ++ColumnIndex)
+			for (int32 ColumnIndex = 0, ColumnsNum = Colums.Num(); ColumnIndex < ColumnsNum; ++ColumnIndex)
 			{
 				const FRTPColumn& RTPColumn = Colums[ColumnIndex];
 
 				uint32 ColumnSize = RTPColumn.SizeInBytes;
 
 				// hide columns that are too small to make a difference (e.g. <1 MB)
-				if(RTPColumn.SizeInBytes <= EventRecordingSizeThreshold * 1024)
+				if (RTPColumn.SizeInBytes <= EventRecordingSizeThreshold * 1024)
 				{
 					ColumnSize = 0;
 				}
@@ -736,15 +731,15 @@ FRenderTargetPool::SMemoryStats FRenderTargetPool::ComputeView()
 				MemoryStats.TotalColumnSize += ColumnSize;
 				MemoryStats.TotalUsageInBytes += RTPColumn.SizeInBytes;
 				
-				for(int32 EventIndex = 0, PoolEventsNum = RenderTargetPoolEvents.Num(); EventIndex < PoolEventsNum; EventIndex++)
+				for (int32 EventIndex = 0, PoolEventsNum = RenderTargetPoolEvents.Num(); EventIndex < PoolEventsNum; EventIndex++)
 				{
 					FRenderTargetPoolEvent* Event = &RenderTargetPoolEvents[EventIndex];
 
-					if(Event->GetEventType() != ERTPE_Phase)
+					if (Event->GetEventType() != ERTPE_Phase)
 					{
 						uint32 PoolEntryId = Event->GetPoolEntryId();
 
-						if(RTPColumn.PoolEntryId == PoolEntryId)
+						if (RTPColumn.PoolEntryId == PoolEntryId)
 						{
 							Event->SetColumn(ColumnIndex, ColumnX, ColumnSize);
 						}
@@ -772,11 +767,11 @@ void FRenderTargetPool::AddDeallocEvents()
 
 	bool bWorkWasDone = false;
 
-	for(uint32 i = 0, Num = (uint32)RenderTargetPoolEvents.Num(); i < Num; ++i)
+	for (uint32 i = 0, Num = (uint32)RenderTargetPoolEvents.Num(); i < Num; ++i)
 	{
 		FRenderTargetPoolEvent& Event = RenderTargetPoolEvents[i];
 
-		if(Event.NeedsDeallocEvent())
+		if (Event.NeedsDeallocEvent())
 		{
 			FRenderTargetPoolEvent NewEvent(Event.GetPoolEntryId(), CurrentEventRecordingTime);
 
@@ -788,7 +783,7 @@ void FRenderTargetPool::AddDeallocEvents()
 		}
 	}
 
-	if(bWorkWasDone)
+	if (bWorkWasDone)
 	{
 		++CurrentEventRecordingTime;
 	}
@@ -798,7 +793,7 @@ void FRenderTargetPool::AddAllocEvent(uint32 InPoolEntryId, FPooledRenderTarget*
 {
 	check(In);
 
-	if(IsEventRecordingEnabled())
+	if (IsEventRecordingEnabled())
 	{
 		AddDeallocEvents();
 
@@ -810,10 +805,9 @@ void FRenderTargetPool::AddAllocEvent(uint32 InPoolEntryId, FPooledRenderTarget*
 	}
 }
 
-
 void FRenderTargetPool::AddAllocEventsFromCurrentState()
 {
-	if(!IsEventRecordingEnabled())
+	if (!IsEventRecordingEnabled())
 	{
 		return;
 	}
@@ -822,11 +816,11 @@ void FRenderTargetPool::AddAllocEventsFromCurrentState()
 
 	bool bWorkWasDone = false;
 
-	for(uint32 i = 0; i < (uint32)PooledRenderTargets.Num(); ++i)
+	for (uint32 i = 0; i < (uint32)PooledRenderTargets.Num(); ++i)
 	{
 		FPooledRenderTarget* Element = PooledRenderTargets[i];
 
-		if(Element && !Element->IsFree())
+		if (Element && !Element->IsFree())
 		{
 			FRenderTargetPoolEvent NewEvent(i, CurrentEventRecordingTime, Element);
 
@@ -835,7 +829,7 @@ void FRenderTargetPool::AddAllocEventsFromCurrentState()
 		}
 	}
 
-	if(bWorkWasDone)
+	if (bWorkWasDone)
 	{
 		++CurrentEventRecordingTime;
 	}
@@ -846,7 +840,7 @@ void FRenderTargetPool::TickPoolElements()
 	check(IsInRenderingThread());
 	WaitForTransitionFence();
 
-	if(bStartEventRecordingNextTick)
+	if (bStartEventRecordingNextTick)
 	{
 		bStartEventRecordingNextTick = false;
 		bEventRecordingStarted = true;
@@ -861,11 +855,11 @@ void FRenderTargetPool::TickPoolElements()
 
 	CompactPool();
 
-	for(uint32 i = 0; i < (uint32)PooledRenderTargets.Num(); ++i)
+	for (uint32 i = 0; i < (uint32)PooledRenderTargets.Num(); ++i)
 	{
 		FPooledRenderTarget* Element = PooledRenderTargets[i];
 
-		if(Element)
+		if (Element)
 		{
 			check(!Element->IsSnapshot());
 			Element->OnFrameStart();
@@ -873,21 +867,21 @@ void FRenderTargetPool::TickPoolElements()
 	}
 
 	// we need to release something, take the oldest ones first
-	while(AllocationLevelInKB > MinimumPoolSizeInKB)
+	while (AllocationLevelInKB > MinimumPoolSizeInKB)
 	{
 		// -1: not set
 		int32 OldestElementIndex = -1;
 
 		// find oldest element we can remove
-		for(uint32 i = 0, Num = (uint32)PooledRenderTargets.Num(); i < Num; ++i)
+		for (uint32 i = 0, Num = (uint32)PooledRenderTargets.Num(); i < Num; ++i)
 		{
 			FPooledRenderTarget* Element = PooledRenderTargets[i];
 
-			if(Element && Element->UnusedForNFrames > 2)
+			if (Element && Element->UnusedForNFrames > 2)
 			{
-				if(OldestElementIndex != -1)
+				if (OldestElementIndex != -1)
 				{
-					if(PooledRenderTargets[OldestElementIndex]->UnusedForNFrames < Element->UnusedForNFrames)
+					if (PooledRenderTargets[OldestElementIndex]->UnusedForNFrames < Element->UnusedForNFrames)
 					{
 						OldestElementIndex = i;
 					}
@@ -899,7 +893,7 @@ void FRenderTargetPool::TickPoolElements()
 			}
 		}
 
-		if(OldestElementIndex != -1)
+		if (OldestElementIndex != -1)
 		{
 			AllocationLevelInKB -= ComputeSizeInKB(*PooledRenderTargets[OldestElementIndex]);
 
@@ -917,7 +911,7 @@ void FRenderTargetPool::TickPoolElements()
 			//   * Reduce rendering features or resolution
 			//   * Investigate allocations, order or reusing other render targets can help
 			//   * Ignore (editor case, might start using slow memory which can be ok)
-			if(!bCurrentlyOverBudget)
+			if (!bCurrentlyOverBudget)
 			{
 				UE_CLOG(IsRunningClientOnly(), LogRenderTargetPool, Warning, TEXT("r.RenderTargetPoolMin exceeded %d/%d MB (ok in editor, bad on fixed memory platform)"), (AllocationLevelInKB + 1023) / 1024, MinimumPoolSizeInKB / 1024);
 				bCurrentlyOverBudget = true;
@@ -925,26 +919,17 @@ void FRenderTargetPool::TickPoolElements()
 			// at this point we need to give up
 			break;
 		}
+	}
 
-/*	
-	// confused more than it helps (often a name is used on two elements in the pool and some pool elements are not rendered to this frame)
-	else
-	{
-		// initial state of a render target (e.g. Velocity@0)
-		GVisualizeTexture.SetCheckPoint(Element);
-*/	}
-
-	if(AllocationLevelInKB <= MinimumPoolSizeInKB)
+	if (AllocationLevelInKB <= MinimumPoolSizeInKB)
 	{
-		if(bCurrentlyOverBudget)
+		if (bCurrentlyOverBudget)
 		{
 			UE_LOG(LogRenderTargetPool, Display, TEXT("r.RenderTargetPoolMin resolved %d/%d MB"), (AllocationLevelInKB + 1023) / 1024, MinimumPoolSizeInKB / 1024);
 			bCurrentlyOverBudget = false;
 		}
 	}
 
-//	CompactEventArray();
-
 	AddPhaseEvent(TEXT("FromLastFrame"));
 	AddAllocEventsFromCurrentState();
 	AddPhaseEvent(TEXT("Rendering"));
@@ -958,18 +943,17 @@ void FRenderTargetPool::TickPoolElements()
 #endif // STATS
 }
 
-
 int32 FRenderTargetPool::FindIndex(IPooledRenderTarget* In) const
 {
 	check(IsInRenderingThread());
 
-	if(In)
+	if (In)
 	{
-		for(uint32 i = 0, Num = (uint32)PooledRenderTargets.Num(); i < Num; ++i)
+		for (uint32 i = 0, Num = (uint32)PooledRenderTargets.Num(); i < Num; ++i)
 		{
 			const FPooledRenderTarget* Element = PooledRenderTargets[i];
 
-			if(Element == In)
+			if (Element == In)
 			{
 				check(!Element->IsSnapshot());
 				return i;
@@ -987,11 +971,11 @@ void FRenderTargetPool::FreeUnusedResource(TRefCountPtr<IPooledRenderTarget>& In
 	
 	int32 Index = FindIndex(In);
 
-	if(Index != -1)
+	if (Index != -1)
 	{
 		FPooledRenderTarget* Element = PooledRenderTargets[Index];
 
-		if(Element)
+		if (Element)
 		{
 			check(!Element->IsSnapshot());
 			AllocationLevelInKB -= ComputeSizeInKB(*Element);
@@ -1011,11 +995,11 @@ void FRenderTargetPool::FreeUnusedResources()
 {
 	check(IsInRenderingThread());
 
-	for(uint32 i = 0, Num = (uint32)PooledRenderTargets.Num(); i < Num; ++i)
+	for (uint32 i = 0, Num = (uint32)PooledRenderTargets.Num(); i < Num; ++i)
 	{
 		FPooledRenderTarget* Element = PooledRenderTargets[i];
 
-		if(Element && Element->IsFree())
+		if (Element && Element->IsFree())
 		{
 			check(!Element->IsSnapshot());
 			AllocationLevelInKB -= ComputeSizeInKB(*Element);
@@ -1029,15 +1013,14 @@ void FRenderTargetPool::FreeUnusedResources()
 	VerifyAllocationLevel();
 }
 
-
 void FRenderTargetPool::DumpMemoryUsage(FOutputDevice& OutputDevice)
 {
 	OutputDevice.Logf(TEXT("Pooled Render Targets:"));
-	for(int32 i = 0; i < PooledRenderTargets.Num(); ++i)
+	for (int32 i = 0; i < PooledRenderTargets.Num(); ++i)
 	{
 		FPooledRenderTarget* Element = PooledRenderTargets[i];
 
-		if(Element)
+		if (Element)
 		{
 			check(!Element->IsSnapshot());
 			OutputDevice.Logf(
@@ -1061,7 +1044,6 @@ void FRenderTargetPool::DumpMemoryUsage(FOutputDevice& OutputDevice)
 	GetStats(NumTargets,PoolKB,UsedKB);
 	OutputDevice.Logf(TEXT("%.3fMB total, %.3fMB used, %d render targets"), PoolKB / 1024.f, UsedKB / 1024.f, NumTargets);
 
-
 	uint32 DeferredTotal = 0;
 	OutputDevice.Logf(TEXT("Deferred Render Targets:"));
 	for (int32 i = 0; i < DeferredDeleteArray.Num(); ++i)
@@ -1089,7 +1071,6 @@ void FRenderTargetPool::DumpMemoryUsage(FOutputDevice& OutputDevice)
 		}
 	}
 	OutputDevice.Logf(TEXT("%.3fMB Deferred total"), DeferredTotal / 1024.f);
-
 }
 
 uint32 FPooledRenderTarget::AddRef() const
@@ -1109,7 +1090,7 @@ uint32 FPooledRenderTarget::Release()
 	{
 		checkf(IsInRenderingThread(), TEXT("Tried to delete on non-render thread, PooledRT %s %s"), Desc.DebugName ? Desc.DebugName : TEXT("<Unnamed>"), *Desc.GenerateInfoString());
 		uint32 Refs = uint32(--NumRefs);
-		if(Refs == 0)
+		if (Refs == 0)
 		{
 			RenderTargetItem.SafeRelease();
 			delete this;
@@ -1173,7 +1154,7 @@ FPooledRenderTarget* FRenderTargetPool::GetElementById(uint32 Id) const
 {
 	// is used in game and render thread
 
-	if(Id >= (uint32)PooledRenderTargets.Num())
+	if (Id >= (uint32)PooledRenderTargets.Num())
 	{
 		return 0;
 	}
@@ -1183,23 +1164,15 @@ FPooledRenderTarget* FRenderTargetPool::GetElementById(uint32 Id) const
 
 void FRenderTargetPool::VerifyAllocationLevel() const
 {
-/*
-	// uncomment to verify internal consistency
-	uint32 OutWholeCount;
-	uint32 OutWholePoolInKB;
-	uint32 OutUsedInKB;
-
-	GetStats(OutWholeCount, OutWholePoolInKB, OutUsedInKB);
-*/
 }
 
 void FRenderTargetPool::CompactPool()
 {
-	for(uint32 i = 0, Num = (uint32)PooledRenderTargets.Num(); i < Num; ++i)
+	for (uint32 i = 0, Num = (uint32)PooledRenderTargets.Num(); i < Num; ++i)
 	{
 		FPooledRenderTarget* Element = PooledRenderTargets[i];
 
-		if(!Element)
+		if (!Element)
 		{
 			PooledRenderTargets.RemoveAtSwap(i);
 			--Num;
@@ -1212,7 +1185,7 @@ bool FPooledRenderTarget::OnFrameStart()
 	check(IsInRenderingThread() && !bSnapshot);
 
 	// If there are any references to the pooled render target other than the pool itself, then it may not be freed.
-	if(!IsFree())
+	if (!IsFree())
 	{
 		check(!UnusedForNFrames);
 		return false;
@@ -1221,7 +1194,7 @@ bool FPooledRenderTarget::OnFrameStart()
 	++UnusedForNFrames;
 
 	// this logic can be improved
-	if(UnusedForNFrames > 10)
+	if (UnusedForNFrames > 10)
 	{
 		// release
 		return true;
@@ -1235,18 +1208,18 @@ uint32 FPooledRenderTarget::ComputeMemorySize() const
 	uint32 Size = 0;
 	if (!bSnapshot)
 	{
-		if(Desc.Is2DTexture())
+		if (Desc.Is2DTexture())
 		{
 			Size += RHIComputeMemorySize((const FTexture2DRHIRef&)RenderTargetItem.TargetableTexture);
-			if(RenderTargetItem.ShaderResourceTexture != RenderTargetItem.TargetableTexture)
+			if (RenderTargetItem.ShaderResourceTexture != RenderTargetItem.TargetableTexture)
 			{
 				Size += RHIComputeMemorySize((const FTexture2DRHIRef&)RenderTargetItem.ShaderResourceTexture);
 			}
 		}
-		else if(Desc.Is3DTexture())
+		else if (Desc.Is3DTexture())
 		{
 			Size += RHIComputeMemorySize((const FTexture3DRHIRef&)RenderTargetItem.TargetableTexture);
-			if(RenderTargetItem.ShaderResourceTexture != RenderTargetItem.TargetableTexture)
+			if (RenderTargetItem.ShaderResourceTexture != RenderTargetItem.TargetableTexture)
 			{
 				Size += RHIComputeMemorySize((const FTexture3DRHIRef&)RenderTargetItem.ShaderResourceTexture);
 			}
@@ -1254,7 +1227,7 @@ uint32 FPooledRenderTarget::ComputeMemorySize() const
 		else
 		{
 			Size += RHIComputeMemorySize((const FTextureCubeRHIRef&)RenderTargetItem.TargetableTexture);
-			if(RenderTargetItem.ShaderResourceTexture != RenderTargetItem.TargetableTexture)
+			if (RenderTargetItem.ShaderResourceTexture != RenderTargetItem.TargetableTexture)
 			{
 				Size += RHIComputeMemorySize((const FTextureCubeRHIRef&)RenderTargetItem.ShaderResourceTexture);
 			}
