More repairs //UE4/Dev-RenderPlat-Staging to Dev-VR (//UE4/Dev-VR)

diff --git a/Engine/Source/Runtime/RenderCore/Public/BuiltInRayTracingShaders.h b/Engine/Source/Runtime/RenderCore/Public/BuiltInRayTracingShaders.h
new file mode 100644
--- /dev/null
+++ b/Engine/Source/Runtime/RenderCore/Public/BuiltInRayTracingShaders.h
@@ -0,0 +1,89 @@
+// Copyright 1998-2019 Epic Games, Inc. All Rights Reserved.
+
+#pragma once
+
+#include "GlobalShader.h"
+#include "ShaderParameterStruct.h"
+#include "ShaderParameterMacros.h"
+
+#if RHI_RAYTRACING
+
+class FBuiltInRayTracingShader : public FGlobalShader
+{
+protected:
+	static void ModifyCompilationEnvironment(const FGlobalShaderPermutationParameters& Parameters, FShaderCompilerEnvironment& OutEnvironment)
+	{}
+
+public:
+	static bool ShouldCompilePermutation(const FGlobalShaderPermutationParameters& Parameters)
+	{
+		// Built-in ray tracing shaders are always compiled for RHIs that support them, regardless of whether RT is enabled for the project.
+		return RHISupportsRayTracingShaders(Parameters.Platform);
+	}
+
+	FBuiltInRayTracingShader() = default;
+	FBuiltInRayTracingShader(const ShaderMetaType::CompiledShaderInitializerType& Initializer)
+		: FGlobalShader(Initializer)
+	{}
+};
+
+class RENDERCORE_API FOcclusionMainRG : public FBuiltInRayTracingShader
+{
+	DECLARE_GLOBAL_SHADER(FOcclusionMainRG);
+	SHADER_USE_ROOT_PARAMETER_STRUCT(FOcclusionMainRG, FBuiltInRayTracingShader);
+
+	BEGIN_SHADER_PARAMETER_STRUCT(FParameters, )
+		SHADER_PARAMETER_RDG_BUFFER_SRV(RaytracingAccelerationStructure, TLAS)
+		SHADER_PARAMETER_RDG_BUFFER_SRV(StructuredBuffer<FBasicRayData>, Rays)
+		SHADER_PARAMETER_RDG_BUFFER_UAV(RWStructuredBuffer<uint>, OcclusionOutput)
+	END_SHADER_PARAMETER_STRUCT()
+};
+
+class RENDERCORE_API FIntersectionMainRG : public FBuiltInRayTracingShader
+{
+	DECLARE_GLOBAL_SHADER(FIntersectionMainRG);
+	SHADER_USE_ROOT_PARAMETER_STRUCT(FIntersectionMainRG, FBuiltInRayTracingShader);
+
+	BEGIN_SHADER_PARAMETER_STRUCT(FParameters, )
+		SHADER_PARAMETER_RDG_BUFFER_SRV(RaytracingAccelerationStructure, TLAS)
+		SHADER_PARAMETER_RDG_BUFFER_SRV(StructuredBuffer<FBasicRayData>, Rays)
+		SHADER_PARAMETER_RDG_BUFFER_UAV(RWStructuredBuffer<FIntersectionPayload>, IntersectionOutput)
+	END_SHADER_PARAMETER_STRUCT()
+};
+
+class FIntersectionMainCHS : public FBuiltInRayTracingShader
+{
+	DECLARE_EXPORTED_SHADER_TYPE(FIntersectionMainCHS, Global, RENDERCORE_API);
+public:
+
+	FIntersectionMainCHS() = default;
+	FIntersectionMainCHS(const ShaderMetaType::CompiledShaderInitializerType& Initializer)
+		: FBuiltInRayTracingShader(Initializer)
+	{}
+};
+
+class FDefaultMainCHS : public FBuiltInRayTracingShader
+{
+	DECLARE_EXPORTED_SHADER_TYPE(FDefaultMainCHS, Global, RENDERCORE_API);
+public:
+
+	FDefaultMainCHS() = default;
+	FDefaultMainCHS(const ShaderMetaType::CompiledShaderInitializerType& Initializer)
+		: FBuiltInRayTracingShader(Initializer)
+	{}
+};
+
+
+class FDefaultMainMS : public FBuiltInRayTracingShader
+{
+	DECLARE_EXPORTED_SHADER_TYPE(FDefaultMainMS, Global, RENDERCORE_API);
+public:
+
+	FDefaultMainMS() = default;
+	FDefaultMainMS(const ShaderMetaType::CompiledShaderInitializerType& Initializer)
+		: FBuiltInRayTracingShader(Initializer)
+	{}
+};
+
+#endif // RHI_RAYTRACING
+
diff --git a/Engine/Source/Runtime/RenderCore/Public/ClearQuad.h b/Engine/Source/Runtime/RenderCore/Public/ClearQuad.h
new file mode 100644
--- /dev/null
+++ b/Engine/Source/Runtime/RenderCore/Public/ClearQuad.h
@@ -0,0 +1,83 @@
+// Copyright 1998-2019 Epic Games, Inc. All Rights Reserved.
+
+#pragma once
+
+#include "CoreMinimal.h"
+#include "RHIDefinitions.h"
+#include "RenderResource.h"
+
+class FRHICommandList;
+struct FRWBufferStructured;
+struct FRWBuffer;
+struct FSceneRenderTargetItem;
+class FRHIUnorderedAccessView;
+class FGraphicsPipelineStateInitializer;
+
+class FClearVertexBuffer : public FVertexBuffer
+{
+public:
+	/**
+	* Initialize the RHI for this rendering resource
+	*/
+	virtual void InitRHI() override
+	{
+		// create a static vertex buffer
+		FRHIResourceCreateInfo CreateInfo;
+		VertexBufferRHI = RHICreateVertexBuffer(sizeof(FVector4) * 4, BUF_Static, CreateInfo);
+		void* VoidPtr = RHILockVertexBuffer(VertexBufferRHI, 0, sizeof(FVector4) * 4, RLM_WriteOnly);
+		// Generate the vertices used
+		FVector4* Vertices = reinterpret_cast<FVector4*>(VoidPtr);
+		Vertices[0] = FVector4(-1.0f, 1.0f, 0.0f, 1.0f);
+		Vertices[1] = FVector4(1.0f, 1.0f, 0.0f, 1.0f);
+		Vertices[2] = FVector4(-1.0f, -1.0f, 0.0f, 1.0f);
+		Vertices[3] = FVector4(1.0f, -1.0f, 0.0f, 1.0f);
+		RHIUnlockVertexBuffer(VertexBufferRHI);
+	}
+};
+extern RENDERCORE_API TGlobalResource<FClearVertexBuffer> GClearVertexBuffer;
+
+struct FClearQuadCallbacks
+{
+	TFunction<void(FGraphicsPipelineStateInitializer&)> PSOModifier = nullptr;
+	TFunction<void(FRHICommandList&)> PreClear = nullptr;
+	TFunction<void(FRHICommandList&)> PostClear = nullptr;
+};
+
+extern RENDERCORE_API void ClearUAV(FRHICommandList& RHICmdList, const FRWBufferStructured& StructuredBuffer, uint32 Value);
+extern RENDERCORE_API void ClearUAV(FRHICommandList& RHICmdList, const FTextureRWBuffer& Buffer, FLinearColor Value);
+extern RENDERCORE_API void ClearUAV(FRHICommandList& RHICmdList, const FRWBuffer& Buffer, uint32 Value, bool bBarriers = true);
+extern RENDERCORE_API void ClearUAV(FRHICommandList& RHICmdList, FRHIUnorderedAccessView* Buffer, uint32 NumBytes, uint32 Value);
+
+extern RENDERCORE_API void ClearUAV(FRHICommandList& RHICmdList, FRHITexture* Texture, FRHIUnorderedAccessView* TextureUAV, const float(&ClearValues)[4]);
+extern RENDERCORE_API void ClearUAV(FRHICommandList& RHICmdList, FRHITexture* Texture, FRHIUnorderedAccessView* TextureUAV, const uint32(&ClearValues)[4]);
+extern RENDERCORE_API void ClearUAV(FRHICommandList& RHICmdList, FRHITexture* Texture, FRHIUnorderedAccessView* TextureUAV, const FLinearColor& ClearColor);
+
+extern RENDERCORE_API void ClearUAV(FRHICommandList& RHICmdList, const FSceneRenderTargetItem& RenderTargetItem, const float(&ClearValues)[4]);
+extern RENDERCORE_API void ClearUAV(FRHICommandList& RHICmdList, const FSceneRenderTargetItem& RenderTargetItem, const uint32(&ClearValues)[4]);
+extern RENDERCORE_API void ClearUAV(FRHICommandList& RHICmdList, const FSceneRenderTargetItem& RenderTargetItem, const FLinearColor& ClearColor);
+
+extern RENDERCORE_API void ClearTexture2DUAV(FRHICommandList& RHICmdList, FRHIUnorderedAccessView* UAV, int32 Width, int32 Height, const FLinearColor& ClearColor);
+
+extern RENDERCORE_API void DrawClearQuadMRT( FRHICommandList& RHICmdList, bool bClearColor, int32 NumClearColors, const FLinearColor* ClearColorArray, bool bClearDepth, float Depth, bool bClearStencil, uint32 Stencil );
+extern RENDERCORE_API void DrawClearQuadMRT(FRHICommandList& RHICmdList, bool bClearColor, int32 NumClearColors, const FLinearColor* ClearColorArray, bool bClearDepth, float Depth, bool bClearStencil, uint32 Stencil, FClearQuadCallbacks ClearQuadCallbacks);
+extern RENDERCORE_API void DrawClearQuadMRT( FRHICommandList& RHICmdList, bool bClearColor, int32 NumClearColors, const FLinearColor* ClearColorArray, bool bClearDepth, float Depth, bool bClearStencil, uint32 Stencil, FIntPoint ViewSize, FIntRect ExcludeRect );
+
+inline void DrawClearQuad(FRHICommandList& RHICmdList, bool bClearColor, const FLinearColor& Color, bool bClearDepth, float Depth, bool bClearStencil, uint32 Stencil)
+{
+	DrawClearQuadMRT(RHICmdList, bClearColor, 1, &Color, bClearDepth, Depth, bClearStencil, Stencil);
+}
+
+inline void DrawClearQuad(FRHICommandList& RHICmdList, bool bClearColor, const FLinearColor& Color, bool bClearDepth, float Depth, bool bClearStencil, uint32 Stencil, FIntPoint ViewSize, FIntRect ExcludeRect)
+{
+	DrawClearQuadMRT(RHICmdList, bClearColor, 1, &Color, bClearDepth, Depth, bClearStencil, Stencil, ViewSize, ExcludeRect);
+}
+
+inline void DrawClearQuad(FRHICommandList& RHICmdList, const FLinearColor& Color)
+{
+	DrawClearQuadMRT(RHICmdList, true, 1, &Color, false, 0, false, 0);
+}
+
+inline void DrawClearQuad(FRHICommandList& RHICmdList, const FLinearColor& Color, FClearQuadCallbacks ClearQuadCallbacks)
+{
+	DrawClearQuadMRT(RHICmdList, true, 1, &Color, false, 0, false, 0, ClearQuadCallbacks);
+}
\ No newline at end of file
diff --git a/Engine/Source/Runtime/RenderCore/Public/ClearReplacementShaders.h b/Engine/Source/Runtime/RenderCore/Public/ClearReplacementShaders.h
new file mode 100644
--- /dev/null
+++ b/Engine/Source/Runtime/RenderCore/Public/ClearReplacementShaders.h
@@ -0,0 +1,299 @@
+// Copyright 1998-2019 Epic Games, Inc. All Rights Reserved.
+
+
+#pragma once
+
+#include "CoreMinimal.h"
+#include "ShaderParameters.h"
+#include "Shader.h"
+#include "GlobalShader.h"
+
+class FClearReplacementVS : public FGlobalShader
+{
+	DECLARE_EXPORTED_SHADER_TYPE(FClearReplacementVS, Global, RENDERCORE_API);
+public:
+	FClearReplacementVS() {}
+	FClearReplacementVS(const ShaderMetaType::CompiledShaderInitializerType& Initializer)
+	: FGlobalShader( Initializer )
+	{
+	}
+	
+	// FShader interface.
+	virtual bool Serialize(FArchive& Ar) override
+	{
+		bool bShaderHasOutdatedParameters = FGlobalShader::Serialize(Ar);
+		return bShaderHasOutdatedParameters;
+	}
+	
+	static bool ShouldCompilePermutation(const FGlobalShaderPermutationParameters& Parameters)
+	{
+		return IsFeatureLevelSupported(Parameters.Platform, ERHIFeatureLevel::SM5);
+	}
+};
+
+class FClearReplacementPS : public FGlobalShader
+{
+	DECLARE_EXPORTED_SHADER_TYPE(FClearReplacementPS, Global, RENDERCORE_API);
+public:
+	FClearReplacementPS() {}
+	FClearReplacementPS(const ShaderMetaType::CompiledShaderInitializerType& Initializer)
+	: FGlobalShader( Initializer )
+	{
+		ClearColor.Bind(Initializer.ParameterMap, TEXT("ClearColor"), SPF_Mandatory);
+	}
+	
+	// FShader interface.
+	virtual bool Serialize(FArchive& Ar) override
+	{
+		bool bShaderHasOutdatedParameters = FGlobalShader::Serialize(Ar);
+		Ar << ClearColor;
+		return bShaderHasOutdatedParameters;
+	}
+	
+	static bool ShouldCompilePermutation(const FGlobalShaderPermutationParameters& Parameters)
+	{
+		return IsFeatureLevelSupported(Parameters.Platform, ERHIFeatureLevel::SM5);
+	}
+	
+	static void ModifyCompilationEnvironment(const FGlobalShaderPermutationParameters& Parameters, FShaderCompilerEnvironment& OutEnvironment)
+	{
+		FGlobalShader::ModifyCompilationEnvironment(Parameters, OutEnvironment);
+		OutEnvironment.SetRenderTargetOutputFormat(0, PF_A32B32G32R32F);
+	}
+	
+protected:
+	FShaderParameter ClearColor;
+};
+
+template< typename T >
+class FClearTexture2DReplacementCS : public FGlobalShader
+{
+	DECLARE_EXPORTED_SHADER_TYPE(FClearTexture2DReplacementCS, Global, RENDERCORE_API);
+public:
+	FClearTexture2DReplacementCS() {}
+	FClearTexture2DReplacementCS(const ShaderMetaType::CompiledShaderInitializerType& Initializer)
+	: FGlobalShader( Initializer )
+	{
+		ClearColor.Bind(Initializer.ParameterMap, TEXT("ClearColor"), SPF_Mandatory);
+		ClearTextureRW.Bind(Initializer.ParameterMap, TEXT("ClearTextureRW"), SPF_Mandatory);
+	}
+	
+	// FShader interface.
+	virtual bool Serialize(FArchive& Ar) override
+	{
+		bool bShaderHasOutdatedParameters = FGlobalShader::Serialize(Ar);
+		Ar << ClearColor << ClearTextureRW;
+		return bShaderHasOutdatedParameters;
+	}
+	
+	RENDERCORE_API void SetParameters(FRHICommandList& RHICmdList, FRHIUnorderedAccessView* TextureRW, const T(&Values)[4]);
+	RENDERCORE_API void FinalizeParameters(FRHICommandList& RHICmdList, FRHIUnorderedAccessView* TextureRW);
+	
+	static bool ShouldCompilePermutation(const FGlobalShaderPermutationParameters& Parameters)
+	{
+		return IsFeatureLevelSupported(Parameters.Platform, ERHIFeatureLevel::SM5);
+	}
+
+	static void ModifyCompilationEnvironment(const FGlobalShaderPermutationParameters& Parameters, FShaderCompilerEnvironment& OutEnvironment)
+	{
+		FGlobalShader::ModifyCompilationEnvironment(Parameters, OutEnvironment);
+		OutEnvironment.SetDefine( TEXT("Type"), TIsSame< T, float >::Value ? TEXT("float4") : TEXT("uint4") );
+	}
+	
+	const FShaderParameter& GetClearColorParameter()
+	{
+		return ClearColor;
+	}
+	
+	const FShaderResourceParameter& GetClearTextureRWParameter()
+	{
+		return ClearTextureRW;
+	}
+	
+protected:
+	FShaderParameter ClearColor;
+	FShaderResourceParameter ClearTextureRW;
+};
+
+template< typename T >
+class FClearTexture2DArrayReplacementCS : public FGlobalShader
+{
+	DECLARE_EXPORTED_SHADER_TYPE(FClearTexture2DArrayReplacementCS, Global, RENDERCORE_API);
+public:
+	FClearTexture2DArrayReplacementCS() {}
+	FClearTexture2DArrayReplacementCS(const ShaderMetaType::CompiledShaderInitializerType& Initializer)
+		: FGlobalShader(Initializer)
+	{
+		ClearColor.Bind(Initializer.ParameterMap, TEXT("ClearColor"), SPF_Mandatory);
+		ClearTextureArrayRW.Bind(Initializer.ParameterMap, TEXT("ClearTextureArrayRW"), SPF_Mandatory);
+	}
+
+	// FShader interface.
+	virtual bool Serialize(FArchive& Ar) override
+	{
+		bool bShaderHasOutdatedParameters = FGlobalShader::Serialize(Ar);
+		Ar << ClearColor << ClearTextureArrayRW;
+		return bShaderHasOutdatedParameters;
+	}
+
+	RENDERCORE_API void SetParameters(FRHICommandList& RHICmdList, FRHIUnorderedAccessView* TextureRW, const T(&Values)[4]);
+	RENDERCORE_API void FinalizeParameters(FRHICommandList& RHICmdList, FRHIUnorderedAccessView* TextureRW);
+
+	static bool ShouldCompilePermutation(const FGlobalShaderPermutationParameters& Parameters)
+	{
+		return IsFeatureLevelSupported(Parameters.Platform, ERHIFeatureLevel::SM5);
+	}
+
+	static void ModifyCompilationEnvironment(const FGlobalShaderPermutationParameters& Parameters, FShaderCompilerEnvironment& OutEnvironment)
+	{
+		FGlobalShader::ModifyCompilationEnvironment(Parameters, OutEnvironment);
+		OutEnvironment.SetDefine( TEXT("Type"), TIsSame< T, float >::Value ? TEXT("float4") : TEXT("uint4") );
+	}
+
+	const FShaderParameter& GetClearColorParameter()
+	{
+		return ClearColor;
+	}
+
+	const FShaderResourceParameter& GetClearTextureRWParameter()
+	{
+		return ClearTextureArrayRW;
+	}
+
+protected:
+	FShaderParameter ClearColor;
+	FShaderResourceParameter ClearTextureArrayRW;
+};
+
+template< typename T >
+class FClearVolumeReplacementCS : public FGlobalShader
+{
+	DECLARE_EXPORTED_SHADER_TYPE(FClearVolumeReplacementCS, Global, RENDERCORE_API);
+public:
+	FClearVolumeReplacementCS() {}
+	FClearVolumeReplacementCS(const ShaderMetaType::CompiledShaderInitializerType& Initializer)
+		: FGlobalShader(Initializer)
+	{
+		ClearColor.Bind(Initializer.ParameterMap, TEXT("ClearColor"), SPF_Mandatory);
+		ClearVolumeRW.Bind(Initializer.ParameterMap, TEXT("ClearVolumeRW"), SPF_Mandatory);
+	}
+
+	// FShader interface.
+	virtual bool Serialize(FArchive& Ar) override
+	{
+		bool bShaderHasOutdatedParameters = FGlobalShader::Serialize(Ar);
+		Ar << ClearColor << ClearVolumeRW;
+		return bShaderHasOutdatedParameters;
+	}
+
+	RENDERCORE_API void SetParameters(FRHICommandList& RHICmdList, FRHIUnorderedAccessView* TextureRW, const T(&Values)[4]);
+	RENDERCORE_API void FinalizeParameters(FRHICommandList& RHICmdList, FRHIUnorderedAccessView* TextureRW);
+
+	static bool ShouldCompilePermutation(const FGlobalShaderPermutationParameters& Parameters)
+	{
+		return IsFeatureLevelSupported(Parameters.Platform, ERHIFeatureLevel::SM5);
+	}
+
+	static void ModifyCompilationEnvironment(const FGlobalShaderPermutationParameters& Parameters, FShaderCompilerEnvironment& OutEnvironment)
+	{
+		FGlobalShader::ModifyCompilationEnvironment(Parameters, OutEnvironment);
+		OutEnvironment.SetDefine( TEXT("Type"), TIsSame< T, float >::Value ? TEXT("float4") : TEXT("uint4") );
+	}
+
+	const FShaderParameter& GetClearColorParameter()
+	{
+		return ClearColor;
+	}
+
+	const FShaderResourceParameter& GetClearTextureRWParameter()
+	{
+		return ClearVolumeRW;
+	}
+
+protected:
+	FShaderParameter ClearColor;
+	FShaderResourceParameter ClearVolumeRW;
+};
+
+class FClearTexture2DReplacementScissorCS : public FGlobalShader
+{
+	DECLARE_EXPORTED_SHADER_TYPE(FClearTexture2DReplacementScissorCS, Global, RENDERCORE_API);
+public:
+	FClearTexture2DReplacementScissorCS() {}
+	FClearTexture2DReplacementScissorCS(const ShaderMetaType::CompiledShaderInitializerType& Initializer)
+		: FGlobalShader(Initializer)
+	{
+		ClearColor.Bind(Initializer.ParameterMap, TEXT("ClearColor"), SPF_Mandatory);
+		TargetBounds.Bind(Initializer.ParameterMap, TEXT("TargetBounds"), SPF_Mandatory);
+		ClearTextureRW.Bind(Initializer.ParameterMap, TEXT("ClearTextureRW"), SPF_Mandatory);
+	}
+
+	// FShader interface.
+	virtual bool Serialize(FArchive& Ar) override
+	{
+		bool bShaderHasOutdatedParameters = FGlobalShader::Serialize(Ar);
+		Ar << ClearColor << TargetBounds << ClearTextureRW;
+		return bShaderHasOutdatedParameters;
+	}
+
+	RENDERCORE_API void SetParameters(FRHICommandList& RHICmdList, FRHIUnorderedAccessView* TextureRW, FLinearColor ClearColor, const FVector4& InTargetBounds);
+	RENDERCORE_API void FinalizeParameters(FRHICommandList& RHICmdList, FRHIUnorderedAccessView* TextureRW);
+
+	static bool ShouldCompilePermutation(const FGlobalShaderPermutationParameters& Parameters)
+	{
+		return IsFeatureLevelSupported(Parameters.Platform, ERHIFeatureLevel::SM5);
+	}
+
+	const FShaderParameter& GetClearColorParameter()
+	{
+		return ClearColor;
+	}
+
+	const FShaderParameter& GetTargetBoundsParameter()
+	{
+		return TargetBounds;
+	}
+
+	const FShaderResourceParameter& GetClearTextureRWParameter()
+	{
+		return ClearTextureRW;
+	}
+
+protected:
+	FShaderParameter ClearColor;
+	FShaderParameter TargetBounds;
+	FShaderResourceParameter ClearTextureRW;
+};
+
+class FClearBufferReplacementCS : public FGlobalShader
+{
+	DECLARE_EXPORTED_SHADER_TYPE(FClearBufferReplacementCS, Global, RENDERCORE_API);
+public:
+	FClearBufferReplacementCS() {}
+	FClearBufferReplacementCS(const ShaderMetaType::CompiledShaderInitializerType& Initializer)
+	: FGlobalShader( Initializer )
+	{
+		ClearBufferCSParams.Bind(Initializer.ParameterMap, TEXT("ClearBufferCSParams"), SPF_Mandatory);
+		ClearBufferRW.Bind(Initializer.ParameterMap, TEXT("ClearBufferRW"), SPF_Mandatory);
+	}
+	
+	RENDERCORE_API void SetParameters(FRHICommandList& RHICmdList, FRHIUnorderedAccessView* BufferRW, uint32 NumDWordsToClear, uint32 ClearValue, bool bBarriers = true);
+	RENDERCORE_API void FinalizeParameters(FRHICommandList& RHICmdList, FRHIUnorderedAccessView* BufferRW, bool bBarriers = true);
+	
+	// FShader interface.
+	virtual bool Serialize(FArchive& Ar) override
+	{
+		bool bShaderHasOutdatedParameters = FGlobalShader::Serialize(Ar);
+		Ar << ClearBufferCSParams << ClearBufferRW;
+		return bShaderHasOutdatedParameters;
+	}
+	
+	static bool ShouldCompilePermutation(const FGlobalShaderPermutationParameters& Parameters)
+	{
+		return RHISupportsComputeShaders(Parameters.Platform);
+	}
+	
+	//protected:
+	FShaderParameter ClearBufferCSParams;
+	FShaderResourceParameter ClearBufferRW;
+};
diff --git a/Engine/Source/Runtime/RenderCore/Public/LongGPUTask.h b/Engine/Source/Runtime/RenderCore/Public/LongGPUTask.h
new file mode 100644
--- /dev/null
+++ b/Engine/Source/Runtime/RenderCore/Public/LongGPUTask.h
@@ -0,0 +1,10 @@
+// Copyright 1998-2019 Epic Games, Inc. All Rights Reserved.
+
+#pragma once
+
+#include "CoreMinimal.h"
+
+class FRHICommandListImmediate;
+
+extern RENDERCORE_API void IssueScalableLongGPUTask(FRHICommandListImmediate& RHICmdList, int32 NumIteration = -1);
+extern RENDERCORE_API void MeasureLongGPUTaskExecutionTime(FRHICommandListImmediate& RHICmdList);
diff --git a/Engine/Source/Runtime/RenderCore/Public/MediaShaders.h b/Engine/Source/Runtime/RenderCore/Public/MediaShaders.h
new file mode 100644
--- /dev/null
+++ b/Engine/Source/Runtime/RenderCore/Public/MediaShaders.h
@@ -0,0 +1,666 @@
+// Copyright 1998-2019 Epic Games, Inc. All Rights Reserved.
+
+#pragma once
+
+#include "CoreMinimal.h"
+#include "RHI.h"
+#include "RenderResource.h"
+#include "Shader.h"
+#include "GlobalShader.h"
+#include "PipelineStateCache.h"
+
+
+namespace MediaShaders
+{
+	/** Color transform from YUV to sRGB (using values from MSDN). */
+	RENDERCORE_API extern const FMatrix YuvToSrgbDefault;
+
+	/** Color transform from YUV to sRGB (in JPEG color space). */
+	RENDERCORE_API extern const FMatrix YuvToSrgbJpeg;
+
+	/** Color transform from YUV to sRGB (using values from PS4 AvPlayer codec). */
+	RENDERCORE_API extern const FMatrix YuvToSrgbPs4;
+
+	/** Color transform from YUV to sRGB (in Rec. 601 color space). */
+	RENDERCORE_API extern const FMatrix YuvToSrgbRec601;
+
+	/** Color transform from YUV to sRGB (in Rec. 709 color space). */
+	RENDERCORE_API extern const FMatrix YuvToRgbRec709;
+
+	/** Color transform from YUV to RGB (in Rec. 709 color space, RGB full range) */
+	RENDERCORE_API extern const FMatrix YuvToRgbRec709Full;
+
+	/** Color transform from RGB to YUV (in Rec. 709 color space, RGB full range) */
+	RENDERCORE_API extern const FMatrix RgbToYuvRec709Full;
+
+	/** YUV Offset for 8 bit conversion (Computed as 16/255, 128/255, 128/255) */
+	RENDERCORE_API extern const FVector YUVOffset8bits;
+
+	/** YUV Offset for 10 bit conversion (Computed as 64/1023, 512/1023, 512/1023) */
+	RENDERCORE_API extern const FVector YUVOffset10bits;
+
+	/** Combine color transform matrix with yuv offset in a single matrix */
+	RENDERCORE_API FMatrix CombineColorTransformAndOffset(const FMatrix& InMatrix, const FVector& InYUVOffset);
+}
+
+
+/**
+ * Stores media drawing vertices.
+ */
+struct FMediaElementVertex
+{
+	FVector4 Position;
+	FVector2D TextureCoordinate;
+
+	FMediaElementVertex() { }
+
+	FMediaElementVertex(const FVector4& InPosition, const FVector2D& InTextureCoordinate)
+		: Position(InPosition)
+		, TextureCoordinate(InTextureCoordinate)
+	{ }
+};
+
+inline FVertexBufferRHIRef CreateTempMediaVertexBuffer(float ULeft = 0.0f, float URight = 1.0f, float VTop = 0.0f, float VBottom = 1.0f)
+{
+	FRHIResourceCreateInfo CreateInfo;
+	FVertexBufferRHIRef VertexBufferRHI = RHICreateVertexBuffer(sizeof(FMediaElementVertex) * 4, BUF_Volatile, CreateInfo);
+	void* VoidPtr = RHILockVertexBuffer(VertexBufferRHI, 0, sizeof(FMediaElementVertex) * 4, RLM_WriteOnly);
+
+	FMediaElementVertex* Vertices = (FMediaElementVertex*)VoidPtr;
+	Vertices[0].Position.Set(-1.0f, 1.0f, 1.0f, 1.0f); // Top Left
+	Vertices[1].Position.Set(1.0f, 1.0f, 1.0f, 1.0f); // Top Right
+	Vertices[2].Position.Set(-1.0f, -1.0f, 1.0f, 1.0f); // Bottom Left
+	Vertices[3].Position.Set(1.0f, -1.0f, 1.0f, 1.0f); // Bottom Right
+
+	Vertices[0].TextureCoordinate.Set(ULeft, VTop);
+	Vertices[1].TextureCoordinate.Set(URight, VTop);
+	Vertices[2].TextureCoordinate.Set(ULeft, VBottom);
+	Vertices[3].TextureCoordinate.Set(URight, VBottom);
+	RHIUnlockVertexBuffer(VertexBufferRHI);
+
+	return VertexBufferRHI;
+}
+
+/**
+ * The simple element vertex declaration resource type.
+ */
+class FMediaVertexDeclaration
+	: public FRenderResource
+{
+public:
+
+	FVertexDeclarationRHIRef VertexDeclarationRHI;
+
+	virtual ~FMediaVertexDeclaration() { }
+
+	virtual void InitRHI() override
+	{
+		FVertexDeclarationElementList Elements;
+		uint16 Stride = sizeof(FMediaElementVertex);
+		Elements.Add(FVertexElement(0, STRUCT_OFFSET(FMediaElementVertex, Position), VET_Float4, 0, Stride));
+		Elements.Add(FVertexElement(0, STRUCT_OFFSET(FMediaElementVertex, TextureCoordinate), VET_Float2, 1, Stride));
+		VertexDeclarationRHI = PipelineStateCache::GetOrCreateVertexDeclaration(Elements);
+	}
+
+	virtual void ReleaseRHI() override
+	{
+		VertexDeclarationRHI.SafeRelease();
+	}
+};
+
+
+RENDERCORE_API extern TGlobalResource<FMediaVertexDeclaration> GMediaVertexDeclaration;
+
+
+/**
+ * Media vertex shader (shared by all media shaders).
+ */
+class FMediaShadersVS
+	: public FGlobalShader
+{
+	DECLARE_EXPORTED_SHADER_TYPE(FMediaShadersVS, Global, RENDERCORE_API);
+
+public:
+
+	static bool ShouldCompilePermutation(const FGlobalShaderPermutationParameters& Parameters)
+	{
+		return IsFeatureLevelSupported(Parameters.Platform, ERHIFeatureLevel::ES2);
+	}
+
+	/** Default constructor. */
+	FMediaShadersVS() { }
+
+	/** Initialization constructor. */
+	FMediaShadersVS(const ShaderMetaType::CompiledShaderInitializerType& Initializer)
+		: FGlobalShader(Initializer)
+	{ }
+
+	virtual bool Serialize(FArchive& Ar) override
+	{
+		bool bShaderHasOutdatedParameters = FGlobalShader::Serialize(Ar);
+		return bShaderHasOutdatedParameters;
+	}
+};
+
+
+/**
+ * Pixel shader to convert an AYUV texture to RGBA.
+ *
+ * This shader expects a single texture consisting of a N x M array of pixels
+ * in AYUV format. Each pixel is encoded as four consecutive unsigned chars
+ * with the following layout: [V0 U0 Y0 A0][V1 U1 Y1 A1]..
+ */
+class FAYUVConvertPS
+	: public FGlobalShader
+{
+	DECLARE_EXPORTED_SHADER_TYPE(FAYUVConvertPS, Global, RENDERCORE_API);
+
+public:
+
+	static bool ShouldCompilePermutation(const FGlobalShaderPermutationParameters& Parameters)
+	{
+		return IsFeatureLevelSupported(Parameters.Platform, ERHIFeatureLevel::ES2);
+	}
+
+	FAYUVConvertPS() { }
+
+	FAYUVConvertPS(const ShaderMetaType::CompiledShaderInitializerType& Initializer)
+		: FGlobalShader(Initializer)
+	{ }
+
+	virtual bool Serialize(FArchive& Ar) override
+	{
+		bool bShaderHasOutdatedParameters = FGlobalShader::Serialize(Ar);
+		return bShaderHasOutdatedParameters;
+	}
+
+	RENDERCORE_API void SetParameters(FRHICommandList& RHICmdList, TRefCountPtr<FRHITexture2D> AYUVTexture, const FMatrix& ColorTransform, const FVector& YUVOffset, bool SrgbToLinear);
+};
+
+
+/**
+ * Pixel shader to convert a Windows Bitmap texture.
+ *
+ * This shader expects a BMP frame packed into a single texture in PF_B8G8R8A8 format.
+ */
+class FBMPConvertPS
+	: public FGlobalShader
+{
+	DECLARE_EXPORTED_SHADER_TYPE(FBMPConvertPS, Global, RENDERCORE_API);
+
+public:
+
+	static bool ShouldCompilePermutation(const FGlobalShaderPermutationParameters& Parameters)
+	{
+		return IsFeatureLevelSupported(Parameters.Platform, ERHIFeatureLevel::ES2);
+	}
+
+	FBMPConvertPS() { }
+
+	FBMPConvertPS(const ShaderMetaType::CompiledShaderInitializerType& Initializer)
+		: FGlobalShader(Initializer)
+	{ }
+
+	virtual bool Serialize(FArchive& Ar) override
+	{
+		bool bShaderHasOutdatedParameters = FGlobalShader::Serialize(Ar);
+		return bShaderHasOutdatedParameters;
+	}
+
+	RENDERCORE_API void SetParameters(FRHICommandList& RHICmdList, TRefCountPtr<FRHITexture2D> BMPTexture, const FIntPoint& OutputDimensions, bool SrgbToLinear);
+};
+
+
+/**
+ * Pixel shader to convert a NV12 frame to RGBA.
+ *
+ * This shader expects an NV12 frame packed into a single texture in PF_G8 format.
+ *
+ * @see http://www.fourcc.org/yuv.php#NV12
+ */
+class FNV12ConvertPS
+	: public FGlobalShader
+{
+	DECLARE_EXPORTED_SHADER_TYPE(FNV12ConvertPS, Global, RENDERCORE_API);
+
+public:
+
+	static bool ShouldCompilePermutation(const FGlobalShaderPermutationParameters& Parameters)
+	{
+		return IsFeatureLevelSupported(Parameters.Platform, ERHIFeatureLevel::ES2);
+	}
+
+	FNV12ConvertPS() { }
+
+	FNV12ConvertPS(const ShaderMetaType::CompiledShaderInitializerType& Initializer)
+		: FGlobalShader(Initializer)
+	{ }
+
+	virtual bool Serialize(FArchive& Ar) override
+	{
+		bool bShaderHasOutdatedParameters = FGlobalShader::Serialize(Ar);
+		return bShaderHasOutdatedParameters;
+	}
+
+	RENDERCORE_API void SetParameters(FRHICommandList& RHICmdList, TRefCountPtr<FRHITexture2D> NV12Texture, const FIntPoint& OutputDimensions, const FMatrix& ColorTransform, const FVector& YUVOffset, bool SrgbToLinear);
+};
+
+
+/**
+ * Pixel shader to convert a NV21 frame to RGBA.
+ *
+ * This shader expects an NV21 frame packed into a single texture in PF_G8 format.
+ *
+ * @see http://www.fourcc.org/yuv.php#NV21
+ */
+class FNV21ConvertPS
+	: public FGlobalShader
+{
+	DECLARE_EXPORTED_SHADER_TYPE(FNV21ConvertPS, Global, RENDERCORE_API);
+
+public:
+
+	static bool ShouldCompilePermutation(const FGlobalShaderPermutationParameters& Parameters)
+	{
+		return IsFeatureLevelSupported(Parameters.Platform, ERHIFeatureLevel::ES2);
+	}
+
+	FNV21ConvertPS() { }
+
+	FNV21ConvertPS(const ShaderMetaType::CompiledShaderInitializerType& Initializer)
+		: FGlobalShader(Initializer)
+	{ }
+
+	virtual bool Serialize(FArchive& Ar) override
+	{
+		bool bShaderHasOutdatedParameters = FGlobalShader::Serialize(Ar);
+		return bShaderHasOutdatedParameters;
+	}
+
+	RENDERCORE_API void SetParameters(FRHICommandList& RHICmdList, TRefCountPtr<FRHITexture2D> NV21Texture, const FIntPoint& OutputDimensions, const FMatrix& ColorTransform, const FVector& YUVOffset, bool SrgbToLinear);
+};
+
+
+/**
+ * Pixel shader to resize an RGB texture.
+ *
+ * This shader expects an RGB or RGBA frame packed into a single texture
+ * in PF_B8G8R8A8 or PF_FloatRGB format.
+ */
+class FRGBConvertPS
+	: public FGlobalShader
+{
+	DECLARE_EXPORTED_SHADER_TYPE(FRGBConvertPS, Global, RENDERCORE_API);
+
+public:
+
+	static bool ShouldCompilePermutation(const FGlobalShaderPermutationParameters& Parameters)
+	{
+		return IsFeatureLevelSupported(Parameters.Platform, ERHIFeatureLevel::ES2);
+	}
+
+	FRGBConvertPS() { }
+
+	FRGBConvertPS(const ShaderMetaType::CompiledShaderInitializerType& Initializer)
+		: FGlobalShader(Initializer)
+	{ }
+
+	virtual bool Serialize(FArchive& Ar) override
+	{
+		bool bShaderHasOutdatedParameters = FGlobalShader::Serialize(Ar);
+		return bShaderHasOutdatedParameters;
+	}
+
+	RENDERCORE_API void SetParameters(FRHICommandList& RHICmdList, TRefCountPtr<FRHITexture2D> RGBTexture, const FIntPoint& OutputDimensions, bool SrgbToLinear);
+};
+
+
+/**
+ * Pixel shader to convert a PS4 YCbCr texture to RGBA.
+ *
+ * This shader expects a separate chroma and luma plane stored in two textures
+ * in PF_B8G8R8A8 format. The full-size luma plane contains the Y-components.
+ * The half-size chroma plane contains the UV components in the following
+ * memory layout: [U0, V0][U1, V1]
+ * 
+ */
+class FYCbCrConvertPS
+	: public FGlobalShader
+{
+	DECLARE_EXPORTED_SHADER_TYPE(FYCbCrConvertPS, Global, RENDERCORE_API);
+
+public:
+
+	static bool ShouldCompilePermutation(const FGlobalShaderPermutationParameters& Parameters)
+	{
+		return IsFeatureLevelSupported(Parameters.Platform, ERHIFeatureLevel::ES2);
+	}
+
+	FYCbCrConvertPS() { }
+
+	FYCbCrConvertPS(const ShaderMetaType::CompiledShaderInitializerType& Initializer)
+		: FGlobalShader(Initializer)
+	{ }
+
+	virtual bool Serialize(FArchive& Ar) override
+	{
+		bool bShaderHasOutdatedParameters = FGlobalShader::Serialize(Ar);		
+		return bShaderHasOutdatedParameters;
+	}
+
+	RENDERCORE_API void SetParameters(FRHICommandList& RHICmdList, TRefCountPtr<FRHITexture2D> LumaTexture, TRefCountPtr<FRHITexture2D> CbCrTexture, const FMatrix& ColorTransform, const FVector& YUVOffset, bool SrgbToLinear);
+};
+
+class FYCbCrConvertPS_4x4Matrix : public FYCbCrConvertPS
+{
+    DECLARE_EXPORTED_SHADER_TYPE(FYCbCrConvertPS_4x4Matrix, Global, RENDERCORE_API);
+    
+public:
+    
+    static bool ShouldCompilePermutation(const FGlobalShaderPermutationParameters& Parameters)
+    {
+        return IsFeatureLevelSupported(Parameters.Platform, ERHIFeatureLevel::ES2);
+    }
+    
+    FYCbCrConvertPS_4x4Matrix() { }
+    
+    FYCbCrConvertPS_4x4Matrix(const ShaderMetaType::CompiledShaderInitializerType& Initializer)
+    : FYCbCrConvertPS(Initializer)
+    { }
+};
+
+
+/**
+ * Pixel shader to convert a UYVY (Y422, UYNV) frame to RGBA.
+ *
+ * This shader expects a UYVY frame packed into a single texture in PF_B8G8R8A8
+ * format with the following memory layout: [U0, Y0, V1, Y1][U1, Y2, V1, Y3]..
+ *
+ * @see http://www.fourcc.org/yuv.php#UYVY
+ */
+class FUYVYConvertPS
+	: public FGlobalShader
+{
+	DECLARE_EXPORTED_SHADER_TYPE(FUYVYConvertPS, Global, RENDERCORE_API);
+
+public:
+
+	static bool ShouldCompilePermutation(const FGlobalShaderPermutationParameters& Parameters)
+	{
+		return IsFeatureLevelSupported(Parameters.Platform, ERHIFeatureLevel::ES2);
+	}
+
+	FUYVYConvertPS() { }
+
+	FUYVYConvertPS(const ShaderMetaType::CompiledShaderInitializerType& Initializer)
+		: FGlobalShader(Initializer)
+	{ }
+
+	virtual bool Serialize(FArchive& Ar) override
+	{
+		bool bShaderHasOutdatedParameters = FGlobalShader::Serialize(Ar);
+		return bShaderHasOutdatedParameters;
+	}
+
+	RENDERCORE_API void SetParameters(FRHICommandList& RHICmdList, TRefCountPtr<FRHITexture2D> UYVYTexture, const FMatrix& ColorTransform, const FVector& YUVOffset, bool SrgbToLinear);
+};
+
+
+/**
+ * Pixel shader to convert Y, U, and V planes to RGBA.
+ *
+ * This shader expects three textures in PF_G8 format,
+ * one for each plane of Y, U, and V components.
+ */
+class FYUVConvertPS
+	: public FGlobalShader
+{
+	DECLARE_EXPORTED_SHADER_TYPE(FYUVConvertPS, Global, RENDERCORE_API);
+
+public:
+
+	static bool ShouldCompilePermutation(const FGlobalShaderPermutationParameters& Parameters)
+	{
+		return IsFeatureLevelSupported(Parameters.Platform, ERHIFeatureLevel::ES2);
+	}
+
+	FYUVConvertPS() { }
+
+	FYUVConvertPS(const ShaderMetaType::CompiledShaderInitializerType& Initializer)
+		: FGlobalShader(Initializer)
+	{ }
+
+	virtual bool Serialize(FArchive& Ar) override
+	{
+		bool bShaderHasOutdatedParameters = FGlobalShader::Serialize(Ar);
+		return bShaderHasOutdatedParameters;
+	}
+
+	RENDERCORE_API void SetParameters(FRHICommandList& RHICmdList, TRefCountPtr<FRHITexture2D> YTexture, TRefCountPtr<FRHITexture2D> UTexture, TRefCountPtr<FRHITexture2D> VTexture, const FIntPoint& OutputDimensions, const FMatrix& ColorTransform, const FVector& YUVOffset, bool SrgbToLinear);
+};
+
+
+/**
+ * Pixel shader to convert YUV v210 to RGB
+ *
+ * This shader expects a single texture in PF_R32G32B32A32_UINT format.
+ */
+class FYUVv210ConvertPS
+	: public FGlobalShader
+{
+	DECLARE_EXPORTED_SHADER_TYPE(FYUVv210ConvertPS, Global, RENDERCORE_API);
+
+public:
+
+	static bool ShouldCompilePermutation(const FGlobalShaderPermutationParameters& Parameters)
+	{
+		return IsFeatureLevelSupported(Parameters.Platform, ERHIFeatureLevel::ES2);
+	}
+
+	FYUVv210ConvertPS() { }
+
+	FYUVv210ConvertPS(const ShaderMetaType::CompiledShaderInitializerType& Initializer)
+		: FGlobalShader(Initializer)
+	{ }
+
+	virtual bool Serialize(FArchive& Ar) override
+	{
+		bool bShaderHasOutdatedParameters = FGlobalShader::Serialize(Ar);
+		return bShaderHasOutdatedParameters;
+	}
+
+	RENDERCORE_API void SetParameters(FRHICommandList& RHICmdList, TRefCountPtr<FRHITexture2D> YUVTexture, const FIntPoint& OutputDimensions, const FMatrix& ColorTransform, const FVector& YUVOffset, bool SrgbToLinear);
+};
+
+
+
+/**
+ * Pixel shader to convert a YUY2 frame to RGBA.
+ *
+ * This shader expects an YUY2 frame packed into a single texture in PF_B8G8R8A8
+ * format with the following memory layout: [Y0, U0, Y1, V0][Y2, U1, Y3, V1]...
+ *
+ * @see http://www.fourcc.org/yuv.php#YUY2
+ */
+class FYUY2ConvertPS
+	: public FGlobalShader
+{
+	DECLARE_EXPORTED_SHADER_TYPE(FYUY2ConvertPS, Global, RENDERCORE_API);
+
+public:
+
+	static bool ShouldCompilePermutation(const FGlobalShaderPermutationParameters& Parameters)
+	{
+		return IsFeatureLevelSupported(Parameters.Platform, ERHIFeatureLevel::ES2);
+	}
+
+	FYUY2ConvertPS() { }
+
+	FYUY2ConvertPS(const ShaderMetaType::CompiledShaderInitializerType& Initializer)
+		: FGlobalShader(Initializer)
+	{ }
+
+	virtual bool Serialize(FArchive& Ar) override
+	{
+		bool bShaderHasOutdatedParameters = FGlobalShader::Serialize(Ar);
+		return bShaderHasOutdatedParameters;
+	}
+
+	RENDERCORE_API void SetParameters(FRHICommandList& RHICmdList, TRefCountPtr<FRHITexture2D> YUY2Texture, const FIntPoint& OutputDimensions, const FMatrix& ColorTransform, const FVector& YUVOffset, bool SrgbToLinear);
+};
+
+
+/**
+ * Pixel shader to convert a YVYU frame to RGBA.
+ *
+ * This shader expects a YVYU frame packed into a single texture in PF_B8G8R8A8
+ * format with the following memory layout: [Y0, V0, Y1, U0][Y2, V1, Y3, U1]..
+ *
+ * @see http://www.fourcc.org/yuv.php#YVYU
+ */
+class FYVYUConvertPS
+	: public FGlobalShader
+{
+	DECLARE_EXPORTED_SHADER_TYPE(FYVYUConvertPS, Global, RENDERCORE_API);
+
+public:
+
+	static bool ShouldCompilePermutation(const FGlobalShaderPermutationParameters& Parameters)
+	{
+		return IsFeatureLevelSupported(Parameters.Platform, ERHIFeatureLevel::ES2);
+	}
+
+	FYVYUConvertPS() { }
+
+	FYVYUConvertPS(const ShaderMetaType::CompiledShaderInitializerType& Initializer)
+		: FGlobalShader(Initializer)
+	{ }
+
+	virtual bool Serialize(FArchive& Ar) override
+	{
+		bool bShaderHasOutdatedParameters = FGlobalShader::Serialize(Ar);
+		return bShaderHasOutdatedParameters;
+	}
+
+	RENDERCORE_API void SetParameters(FRHICommandList& RHICmdList, TRefCountPtr<FRHITexture2D> YVYUTexture, const FMatrix& ColorTransform, const FVector& YUVOffset, bool SrgbToLinear);
+};
+
+
+/**
+ * Pixel shader to convert RGB 8 bits to UYVY 8 bits
+ *
+ * This shader expects a single texture in PF_B8G8R8A8 format.
+ */
+class FRGB8toUYVY8ConvertPS
+	: public FGlobalShader
+{
+	DECLARE_EXPORTED_SHADER_TYPE(FRGB8toUYVY8ConvertPS, Global, RENDERCORE_API);
+
+public:
+
+	static bool ShouldCompilePermutation(const FGlobalShaderPermutationParameters& Parameters)
+	{
+		return IsFeatureLevelSupported(Parameters.Platform, ERHIFeatureLevel::ES2);
+	}
+
+	FRGB8toUYVY8ConvertPS() { }
+
+	FRGB8toUYVY8ConvertPS(const ShaderMetaType::CompiledShaderInitializerType& Initializer)
+		: FGlobalShader(Initializer)
+	{ }
+
+	virtual bool Serialize(FArchive& Ar) override
+	{
+		bool bShaderHasOutdatedParameters = FGlobalShader::Serialize(Ar);
+		return bShaderHasOutdatedParameters;
+	}
+
+	RENDERCORE_API void SetParameters(FRHICommandList& RHICmdList, TRefCountPtr<FRHITexture2D> RGBATexture, const FMatrix& ColorTransform, const FVector& YUVOffset, bool LinearToSrgb);
+};
+
+
+/**
+ * Pixel shader to convert RGB 10 bits to YUV v210
+ *
+ * This shader expects a single texture in PF_A2B10G10R10 format.
+ */
+class FRGB10toYUVv210ConvertPS
+	: public FGlobalShader
+{
+	DECLARE_EXPORTED_SHADER_TYPE(FRGB10toYUVv210ConvertPS, Global, RENDERCORE_API);
+
+public:
+
+	static bool ShouldCompilePermutation(const FGlobalShaderPermutationParameters& Parameters)
+	{
+		return IsFeatureLevelSupported(Parameters.Platform, ERHIFeatureLevel::ES2);
+	}
+
+	FRGB10toYUVv210ConvertPS() { }
+
+	FRGB10toYUVv210ConvertPS(const ShaderMetaType::CompiledShaderInitializerType& Initializer)
+		: FGlobalShader(Initializer)
+	{ }
+
+	virtual bool Serialize(FArchive& Ar) override
+	{
+		bool bShaderHasOutdatedParameters = FGlobalShader::Serialize(Ar);
+		return bShaderHasOutdatedParameters;
+	}
+
+	RENDERCORE_API void SetParameters(FRHICommandList& RHICmdList, TRefCountPtr<FRHITexture2D> RGBATexture, const FMatrix& ColorTransform, const FVector& YUVOffset, bool LinearToSrgb);
+};
+
+
+/**
+ * Pixel shader to inverse alpha component 
+ *
+ * This shader expects a single texture in RGBA 8 or 10 bits format.
+ */
+class FInvertAlphaPS
+	: public FGlobalShader
+{
+	DECLARE_EXPORTED_SHADER_TYPE(FInvertAlphaPS, Global, RENDERCORE_API);
+
+public:
+
+	static bool ShouldCompilePermutation(const FGlobalShaderPermutationParameters& Parameters)
+	{
+		return IsFeatureLevelSupported(Parameters.Platform, ERHIFeatureLevel::ES2);
+	}
+
+	FInvertAlphaPS() { }
+
+	FInvertAlphaPS(const ShaderMetaType::CompiledShaderInitializerType& Initializer)
+		: FGlobalShader(Initializer)
+	{ }
+
+	RENDERCORE_API void SetParameters(FRHICommandList& RHICmdList, TRefCountPtr<FRHITexture2D> RGBATexture);
+};
+
+
+/**
+ * Pixel shader to set alpha component to 1.0f
+ *
+ * This shader expects a single texture in RGBA 8 or 10 bits format.
+ */
+class FSetAlphaOnePS
+	: public FGlobalShader
+{
+	DECLARE_EXPORTED_SHADER_TYPE(FSetAlphaOnePS, Global, RENDERCORE_API);
+
+public:
+
+	static bool ShouldCompilePermutation(const FGlobalShaderPermutationParameters& Parameters)
+	{
+		return IsFeatureLevelSupported(Parameters.Platform, ERHIFeatureLevel::ES2);
+	}
+
+	FSetAlphaOnePS() { }
+
+	FSetAlphaOnePS(const ShaderMetaType::CompiledShaderInitializerType& Initializer)
+		: FGlobalShader(Initializer)
+	{ }
+
+	RENDERCORE_API void SetParameters(FRHICommandList& RHICmdList, TRefCountPtr<FRHITexture2D> RGBATexture);
+};
diff --git a/Engine/Source/Runtime/RenderCore/Public/OculusShaders.h b/Engine/Source/Runtime/RenderCore/Public/OculusShaders.h
new file mode 100644
--- /dev/null
+++ b/Engine/Source/Runtime/RenderCore/Public/OculusShaders.h
@@ -0,0 +1,138 @@
+// Copyright 1998-2019 Epic Games, Inc. All Rights Reserved.
+
+#pragma once
+
+#include "CoreMinimal.h"
+#include "ShaderParameters.h"
+#include "Shader.h"
+#include "GlobalShader.h"
+#include "ShaderParameterUtils.h"
+#include "RenderResource.h"
+
+class FOculusVertexShader : public FGlobalShader
+{
+	DECLARE_EXPORTED_SHADER_TYPE(FOculusVertexShader, Global, RENDERCORE_API);
+public:
+
+	static bool ShouldCompilePermutation(const FGlobalShaderPermutationParameters& Parameters) { return true; }
+
+	FOculusVertexShader(const ShaderMetaType::CompiledShaderInitializerType& Initializer) :
+		FGlobalShader(Initializer)
+	{}
+	FOculusVertexShader() {}
+};
+
+class FOculusWhiteShader : public FGlobalShader
+{
+	DECLARE_EXPORTED_SHADER_TYPE(FOculusWhiteShader, Global, RENDERCORE_API);
+public:
+
+	static bool ShouldCompilePermutation(const FGlobalShaderPermutationParameters& Parameters) { return true; }
+
+	FOculusWhiteShader(const ShaderMetaType::CompiledShaderInitializerType& Initializer) :
+		FGlobalShader(Initializer)
+	{
+	}
+	
+	FOculusWhiteShader() {}
+};
+
+class FOculusBlackShader : public FGlobalShader
+{
+	DECLARE_EXPORTED_SHADER_TYPE(FOculusBlackShader, Global, RENDERCORE_API);
+public:
+
+	static bool ShouldCompilePermutation(const FGlobalShaderPermutationParameters& Parameters) { return true; }
+
+	FOculusBlackShader(const ShaderMetaType::CompiledShaderInitializerType& Initializer) :
+		FGlobalShader(Initializer)
+	{
+	}
+
+	FOculusBlackShader() {}
+};
+
+class FOculusAlphaInverseShader : public FGlobalShader
+{
+	DECLARE_EXPORTED_SHADER_TYPE(FOculusAlphaInverseShader, Global, RENDERCORE_API);
+public:
+
+	static bool ShouldCompilePermutation(const FGlobalShaderPermutationParameters& Parameters) { return true; }
+
+	FOculusAlphaInverseShader(const ShaderMetaType::CompiledShaderInitializerType& Initializer) :
+		FGlobalShader(Initializer)
+	{
+		InTexture.Bind(Initializer.ParameterMap, TEXT("InTexture"), SPF_Mandatory);
+		InTextureSampler.Bind(Initializer.ParameterMap, TEXT("InTextureSampler"));
+	}
+	FOculusAlphaInverseShader() {}
+
+	RENDERCORE_API void SetParameters(FRHICommandList& RHICmdList, const FTexture* Texture)
+	{
+		SetTextureParameter(RHICmdList, GetPixelShader(), InTexture, InTextureSampler, Texture);
+	}
+
+	RENDERCORE_API void SetParameters(FRHICommandList& RHICmdList, FRHISamplerState* SamplerStateRHI, FRHITexture* TextureRHI)
+	{
+		SetTextureParameter(RHICmdList, GetPixelShader(), InTexture, InTextureSampler, SamplerStateRHI, TextureRHI);
+	}
+
+	virtual bool Serialize(FArchive& Ar) override
+	{
+		bool bShaderHasOutdatedParameters = FGlobalShader::Serialize(Ar);
+		Ar << InTexture;
+		Ar << InTextureSampler;
+		return bShaderHasOutdatedParameters;
+	}
+
+private:
+	FShaderResourceParameter InTexture;
+	FShaderResourceParameter InTextureSampler;
+};
+
+
+/**
+* A pixel shader for rendering a textured screen element.
+*/
+class FOculusCubemapPS : public FGlobalShader
+{
+	DECLARE_EXPORTED_SHADER_TYPE(FOculusCubemapPS, Global, RENDERCORE_API);
+public:
+
+	static bool ShouldCompilePermutation(const FGlobalShaderPermutationParameters& Parameters) { return true; }
+
+	FOculusCubemapPS(const ShaderMetaType::CompiledShaderInitializerType& Initializer) :
+		FGlobalShader(Initializer)
+	{
+		InTexture.Bind(Initializer.ParameterMap, TEXT("InTextureCube"), SPF_Mandatory);
+		InTextureSampler.Bind(Initializer.ParameterMap, TEXT("InTextureSampler"));
+		InFaceIndexParameter.Bind(Initializer.ParameterMap, TEXT("CubeFaceIndex"));
+	}
+	FOculusCubemapPS() {}
+
+	void SetParameters(FRHICommandList& RHICmdList, const FTexture* Texture, int FaceIndex)
+	{
+		SetTextureParameter(RHICmdList, GetPixelShader(), InTexture, InTextureSampler, Texture);
+		SetShaderValue(RHICmdList, GetPixelShader(), InFaceIndexParameter, FaceIndex);
+	}
+
+	void SetParameters(FRHICommandList& RHICmdList, FRHISamplerState* SamplerStateRHI, FRHITexture* TextureRHI, int FaceIndex)
+	{
+		SetTextureParameter(RHICmdList, GetPixelShader(), InTexture, InTextureSampler, SamplerStateRHI, TextureRHI);
+		SetShaderValue(RHICmdList, GetPixelShader(), InFaceIndexParameter, FaceIndex);
+	}
+
+	virtual bool Serialize(FArchive& Ar) override
+	{
+		bool bShaderHasOutdatedParameters = FGlobalShader::Serialize(Ar);
+		Ar << InTexture;
+		Ar << InTextureSampler;
+		Ar << InFaceIndexParameter;
+		return bShaderHasOutdatedParameters;
+	}
+
+private:
+	FShaderResourceParameter InTexture;
+	FShaderResourceParameter InTextureSampler;
+	FShaderParameter		 InFaceIndexParameter;
+};
\ No newline at end of file
diff --git a/Engine/Source/Runtime/RenderCore/Public/OneColorShader.h b/Engine/Source/Runtime/RenderCore/Public/OneColorShader.h
new file mode 100644
--- /dev/null
+++ b/Engine/Source/Runtime/RenderCore/Public/OneColorShader.h
@@ -0,0 +1,195 @@
+// Copyright 1998-2019 Epic Games, Inc. All Rights Reserved.
+
+
+#pragma once
+
+#include "CoreMinimal.h"
+#include "ShaderParameters.h"
+#include "Shader.h"
+#include "GlobalShader.h"
+#include "ShaderParameterUtils.h"
+
+/**
+ * Vertex shader for rendering a single, constant color.
+ */
+template<bool bUsingNDCPositions=true, bool bUsingVertexLayers=false>
+class TOneColorVS : public FGlobalShader
+{
+	DECLARE_EXPORTED_SHADER_TYPE(TOneColorVS, Global, RENDERCORE_API);
+
+	/** Default constructor. */
+	TOneColorVS() {}
+
+public:
+
+	TOneColorVS(const ShaderMetaType::CompiledShaderInitializerType& Initializer)
+		: FGlobalShader(Initializer)
+	{
+		DepthParameter.Bind(Initializer.ParameterMap, TEXT("InputDepth"), SPF_Mandatory);
+	}
+
+	static void ModifyCompilationEnvironment(const FGlobalShaderPermutationParameters& Parameters, FShaderCompilerEnvironment& OutEnvironment)
+	{
+		FGlobalShader::ModifyCompilationEnvironment(Parameters, OutEnvironment);
+		OutEnvironment.SetDefine(TEXT("USING_NDC_POSITIONS"), (uint32)(bUsingNDCPositions ? 1 : 0));
+		OutEnvironment.SetDefine(TEXT("USING_LAYERS"), (uint32)(bUsingVertexLayers ? 1 : 0));
+	}
+
+	void SetDepthParameter(FRHICommandList& RHICmdList, float Depth)
+	{
+		SetShaderValue(RHICmdList, GetVertexShader(), DepthParameter, Depth);
+	}
+
+	// FShader interface.
+	virtual bool Serialize(FArchive& Ar) override
+	{
+		bool bShaderHasOutdatedParameters = FGlobalShader::Serialize(Ar);
+		Ar << DepthParameter;
+		return bShaderHasOutdatedParameters;
+	}
+
+	static bool ShouldCompilePermutation(const FGlobalShaderPermutationParameters& Parameters)
+	{
+		return true;
+	}
+	
+	static const TCHAR* GetSourceFilename()
+	{
+		return TEXT("/Engine/Private/OneColorShader.usf");
+	}
+	
+	static const TCHAR* GetFunctionName()
+	{
+		return TEXT("MainVertexShader");
+	}
+
+private:
+	FShaderParameter DepthParameter;
+};
+
+/**
+ * Pixel shader for rendering a single, constant color.
+ */
+class FOneColorPS : public FGlobalShader
+{
+	DECLARE_EXPORTED_SHADER_TYPE(FOneColorPS,Global,RENDERCORE_API);
+public:
+	
+	FOneColorPS( )	{ }
+	FOneColorPS(const ShaderMetaType::CompiledShaderInitializerType& Initializer)
+	:	FGlobalShader( Initializer )
+	{
+		//ColorParameter.Bind( Initializer.ParameterMap, TEXT("DrawColorMRT"), SPF_Mandatory);
+	}
+
+	virtual void SetColors(FRHICommandList& RHICmdList, const FLinearColor* Colors, int32 NumColors);	
+
+	// FShader interface.
+	virtual bool Serialize(FArchive& Ar) override
+	{
+		bool bShaderHasOutdatedParameters = FGlobalShader::Serialize(Ar);
+		//Ar << ColorParameter;
+		return bShaderHasOutdatedParameters;
+	}
+	static bool ShouldCompilePermutation(const FGlobalShaderPermutationParameters& Parameters)
+	{
+		return true;
+	}
+
+	/** The parameter to use for setting the draw Color. */
+	//FShaderParameter ColorParameter;
+};
+
+/**
+ * Pixel shader for rendering a single, constant color to MRTs.
+ */
+template<int32 NumOutputs>
+class TOneColorPixelShaderMRT : public FOneColorPS
+{
+	DECLARE_EXPORTED_SHADER_TYPE(TOneColorPixelShaderMRT,Global,RENDERCORE_API);
+public:
+	TOneColorPixelShaderMRT( )	{ }
+	TOneColorPixelShaderMRT(const ShaderMetaType::CompiledShaderInitializerType& Initializer)
+	:	FOneColorPS( Initializer )
+	{
+	}
+
+	static bool ShouldCompilePermutation(const FGlobalShaderPermutationParameters& Parameters)
+	{
+		if( NumOutputs > 1 )
+		{
+			return IsFeatureLevelSupported( Parameters.Platform, ERHIFeatureLevel::ES3_1);
+		}
+		return true;
+	}
+
+	static void ModifyCompilationEnvironment(const FGlobalShaderPermutationParameters& Parameters, FShaderCompilerEnvironment& OutEnvironment)
+	{
+		FOneColorPS::ModifyCompilationEnvironment(Parameters, OutEnvironment);
+		OutEnvironment.SetDefine(TEXT("NUM_OUTPUTS"), NumOutputs);
+	}
+};
+
+/**
+ * Compute shader for writing values
+ */
+class FFillTextureCS : public FGlobalShader
+{
+	DECLARE_EXPORTED_SHADER_TYPE(FFillTextureCS,Global,RENDERCORE_API);
+public:
+	FFillTextureCS( )	{ }
+	FFillTextureCS(const ShaderMetaType::CompiledShaderInitializerType& Initializer)
+	:	FGlobalShader( Initializer )
+	{
+		FillValue.Bind( Initializer.ParameterMap, TEXT("FillValue"), SPF_Mandatory);
+		Params0.Bind( Initializer.ParameterMap, TEXT("Params0"), SPF_Mandatory);
+		Params1.Bind( Initializer.ParameterMap, TEXT("Params1"), SPF_Mandatory);
+		Params2.Bind( Initializer.ParameterMap, TEXT("Params2"), SPF_Optional);
+		FillTexture.Bind( Initializer.ParameterMap, TEXT("FillTexture"), SPF_Mandatory);
+	}
+
+	// FShader interface.
+	virtual bool Serialize(FArchive& Ar) override
+	{
+		bool bShaderHasOutdatedParameters = FGlobalShader::Serialize(Ar);
+		Ar << FillValue << Params0 << Params1 << Params2 << FillTexture;
+		return bShaderHasOutdatedParameters;
+	}
+
+	static bool ShouldCompilePermutation(const FGlobalShaderPermutationParameters& Parameters)
+	{
+		return IsFeatureLevelSupported(Parameters.Platform, ERHIFeatureLevel::SM5);
+	}
+
+	FShaderParameter FillValue;
+	FShaderParameter Params0;	// Texture Width,Height (.xy); Use Exclude Rect 1 : 0 (.z)
+	FShaderParameter Params1;	// Include X0,Y0 (.xy) - X1,Y1 (.zw)
+	FShaderParameter Params2;	// ExcludeRect X0,Y0 (.xy) - X1,Y1 (.zw)
+	FShaderResourceParameter FillTexture;
+};
+
+class FLongGPUTaskPS : public FGlobalShader
+{
+	DECLARE_EXPORTED_SHADER_TYPE(FLongGPUTaskPS,Global,RENDERCORE_API);
+public:
+	FLongGPUTaskPS( )	{ }
+	FLongGPUTaskPS(const ShaderMetaType::CompiledShaderInitializerType& Initializer)
+	:	FGlobalShader( Initializer )
+	{
+	}
+	
+	// FShader interface.
+	virtual bool Serialize(FArchive& Ar) override
+	{
+		bool bShaderHasOutdatedParameters = FGlobalShader::Serialize(Ar);
+		return bShaderHasOutdatedParameters;
+	}
+	
+	static bool ShouldCompilePermutation(const FGlobalShaderPermutationParameters& Parameters)
+	{
+		// MLCHANGES BEGIN
+		return true;
+		// MLCHANGES END
+	}
+};
+
diff --git a/Engine/Source/Runtime/RenderCore/Public/RGBAToYUV420Shader.h b/Engine/Source/Runtime/RenderCore/Public/RGBAToYUV420Shader.h
new file mode 100644
--- /dev/null
+++ b/Engine/Source/Runtime/RenderCore/Public/RGBAToYUV420Shader.h
@@ -0,0 +1,74 @@
+// Copyright 1998-2019 Epic Games, Inc. All Rights Reserved.
+
+#pragma once
+
+#include "CoreMinimal.h"
+#include "RHI.h"
+#include "RenderResource.h"
+#include "Shader.h"
+#include "GlobalShader.h"
+#include "Modules/ModuleManager.h"
+
+#if defined(HAS_MORPHEUS) && HAS_MORPHEUS
+
+/**
+ * Pixel shader to convert RGBA to YUV420.
+ *
+ * YUV420 has 8 bit intensity values in the top 2/3 of the texture for every pixel and 8 bit each UV coordinates 
+ * into the YUV colorspace in the bottom 1/3 for every pixel quad.
+ *
+ * This is only used by PS4.
+ */
+class FRGBAToYUV420CS
+	: public FGlobalShader
+{
+	DECLARE_SHADER_TYPE(FRGBAToYUV420CS, Global);
+
+public:
+
+	static bool ShouldCompilePermutation(const FGlobalShaderPermutationParameters& Parameters)
+	{
+		if (Parameters.Platform == EShaderPlatform::SP_PS4)
+		{
+			// we must use a run time check for this because the builds the build machines create will have Morpheus defined,
+			// but a user will not necessarily have the Morpheus files
+			return FModuleManager::Get().ModuleExists(TEXT("Morpheus"));
+		}
+
+		return false;
+	}
+
+	FRGBAToYUV420CS() { }
+
+	FRGBAToYUV420CS(const ShaderMetaType::CompiledShaderInitializerType& Initializer)
+		: FGlobalShader(Initializer)
+	{
+		TargetHeight.Bind(Initializer.ParameterMap, TEXT("TargetHeight"), SPF_Mandatory);
+		ScaleFactorX.Bind(Initializer.ParameterMap, TEXT("ScaleFactorX"), SPF_Mandatory);
+		ScaleFactorY.Bind(Initializer.ParameterMap, TEXT("ScaleFactorY"), SPF_Mandatory);
+		TextureYOffset.Bind(Initializer.ParameterMap, TEXT("TextureYOffset"), SPF_Mandatory);
+		SrcTexture.Bind(Initializer.ParameterMap, TEXT("SrcTexture"), SPF_Mandatory);
+		OutTextureRW.Bind(Initializer.ParameterMap, TEXT("OutTexture"), SPF_Mandatory);
+	}
+
+	virtual bool Serialize(FArchive& Ar) override
+	{
+		bool bShaderHasOutdatedParameters = FGlobalShader::Serialize(Ar);
+		Ar << TargetHeight << ScaleFactorX << ScaleFactorY << TextureYOffset << SrcTexture << OutTextureRW;
+		return bShaderHasOutdatedParameters;
+	}
+
+	void SetParameters(FRHICommandList& RHICmdList, TRefCountPtr<FRHITexture2D> SrcTex, FRHIUnorderedAccessView* OutUAV, float TargetHeight, float ScaleFactorX, float ScaleFactorY, float TextureYOffset);
+	void UnbindBuffers(FRHICommandList& RHICmdList);
+
+protected:
+	FShaderParameter TargetHeight;
+	FShaderParameter ScaleFactorX;
+	FShaderParameter ScaleFactorY;
+	FShaderParameter TextureYOffset;
+	FShaderResourceParameter SrcTexture;
+
+	FShaderResourceParameter OutTextureRW;
+};
+
+#endif //defined(HAS_MORPHEUS) && HAS_MORPHEUS
diff --git a/Engine/Source/Runtime/RenderCore/Public/ResolveShader.h b/Engine/Source/Runtime/RenderCore/Public/ResolveShader.h
new file mode 100644
--- /dev/null
+++ b/Engine/Source/Runtime/RenderCore/Public/ResolveShader.h
@@ -0,0 +1,283 @@
+// Copyright 1998-2019 Epic Games, Inc. All Rights Reserved.
+
+
+#pragma once
+
+#include "CoreMinimal.h"
+#include "ShaderParameters.h"
+#include "Shader.h"
+#include "GlobalShader.h"
+
+struct FDummyResolveParameter {};
+
+class FResolveDepthPS : public FGlobalShader
+{
+	DECLARE_EXPORTED_SHADER_TYPE(FResolveDepthPS, Global, RENDERCORE_API);
+public:
+	
+	typedef FDummyResolveParameter FParameter;
+	
+	static bool ShouldCompilePermutation(const FGlobalShaderPermutationParameters& Parameters) 
+	{ 
+		return GetMaxSupportedFeatureLevel(Parameters.Platform) >= ERHIFeatureLevel::SM5
+			|| IsSimulatedPlatform(Parameters.Platform); // support resolving MSAA depth in mobile emulation
+	}
+
+	static void ModifyCompilationEnvironment(const FGlobalShaderPermutationParameters& Parameters, FShaderCompilerEnvironment& OutEnvironment)
+	{
+		FGlobalShader::ModifyCompilationEnvironment(Parameters, OutEnvironment);
+		OutEnvironment.SetDefine(TEXT("SIMULATED_PLATFORM"), IsSimulatedPlatform(Parameters.Platform) ? 1 : 0);
+	}
+		
+	FResolveDepthPS(const ShaderMetaType::CompiledShaderInitializerType& Initializer):
+	FGlobalShader(Initializer)
+	{
+		UnresolvedSurface.Bind(Initializer.ParameterMap,TEXT("UnresolvedSurface"), SPF_Mandatory);
+	}
+	FResolveDepthPS() {}
+	
+	void SetParameters(FRHICommandList& RHICmdList, FParameter)
+	{
+	}
+
+	virtual bool Serialize(FArchive& Ar) override
+	{
+		bool bShaderHasOutdatedParameters = FGlobalShader::Serialize(Ar);
+		Ar << UnresolvedSurface;
+		return bShaderHasOutdatedParameters;
+	}
+	
+	FShaderResourceParameter UnresolvedSurface;
+};
+
+class FResolveDepth2XPS : public FGlobalShader
+{
+	DECLARE_EXPORTED_SHADER_TYPE(FResolveDepth2XPS, Global, RENDERCORE_API);
+public:
+
+	typedef FDummyResolveParameter FParameter;
+
+	static bool ShouldCompilePermutation(const FGlobalShaderPermutationParameters& Parameters) 
+	{ 
+		return GetMaxSupportedFeatureLevel(Parameters.Platform) >= ERHIFeatureLevel::SM5 
+			|| IsSimulatedPlatform(Parameters.Platform); // support resolving MSAA depth in mobile emulation
+	}
+
+	static void ModifyCompilationEnvironment(const FGlobalShaderPermutationParameters& Parameters, FShaderCompilerEnvironment& OutEnvironment)
+	{
+		FGlobalShader::ModifyCompilationEnvironment(Parameters, OutEnvironment);
+		OutEnvironment.SetDefine(TEXT("DEPTH_RESOLVE_NUM_SAMPLES"), 2);
+		OutEnvironment.SetDefine(TEXT("SIMULATED_PLATFORM"), IsSimulatedPlatform(Parameters.Platform) ? 1 : 0);
+	}
+
+	FResolveDepth2XPS(const ShaderMetaType::CompiledShaderInitializerType& Initializer) :
+		FGlobalShader(Initializer)
+	{
+		UnresolvedSurface.Bind(Initializer.ParameterMap, TEXT("UnresolvedSurface"), SPF_Mandatory);
+	}
+	FResolveDepth2XPS() {}
+
+	void SetParameters(FRHICommandList& RHICmdList, FParameter)
+	{
+	}
+
+	virtual bool Serialize(FArchive& Ar) override
+	{
+		bool bShaderHasOutdatedParameters = FGlobalShader::Serialize(Ar);
+		Ar << UnresolvedSurface;
+		return bShaderHasOutdatedParameters;
+	}
+
+	FShaderResourceParameter UnresolvedSurface;
+};
+
+
+class FResolveDepth4XPS : public FGlobalShader
+{
+	DECLARE_EXPORTED_SHADER_TYPE(FResolveDepth4XPS, Global, RENDERCORE_API);
+public:
+
+	typedef FDummyResolveParameter FParameter;
+
+	static bool ShouldCompilePermutation(const FGlobalShaderPermutationParameters& Parameters) 
+	{ 
+		return GetMaxSupportedFeatureLevel(Parameters.Platform) >= ERHIFeatureLevel::SM5
+			|| IsSimulatedPlatform(Parameters.Platform); // support resolving MSAA depth in mobile emulation
+	}
+
+	static void ModifyCompilationEnvironment(const FGlobalShaderPermutationParameters& Parameters, FShaderCompilerEnvironment& OutEnvironment)
+	{
+		FGlobalShader::ModifyCompilationEnvironment(Parameters, OutEnvironment);
+		OutEnvironment.SetDefine(TEXT("DEPTH_RESOLVE_NUM_SAMPLES"), 4);
+		OutEnvironment.SetDefine(TEXT("SIMULATED_PLATFORM"), IsSimulatedPlatform(Parameters.Platform) ? 1 : 0);
+	}
+
+	FResolveDepth4XPS(const ShaderMetaType::CompiledShaderInitializerType& Initializer) :
+		FGlobalShader(Initializer)
+	{
+		UnresolvedSurface.Bind(Initializer.ParameterMap, TEXT("UnresolvedSurface"), SPF_Mandatory);
+	}
+	FResolveDepth4XPS() {}
+
+	void SetParameters(FRHICommandList& RHICmdList, FParameter)
+	{
+	}
+
+	virtual bool Serialize(FArchive& Ar) override
+	{
+		bool bShaderHasOutdatedParameters = FGlobalShader::Serialize(Ar);
+		Ar << UnresolvedSurface;
+		return bShaderHasOutdatedParameters;
+	}
+
+	FShaderResourceParameter UnresolvedSurface;
+};
+
+
+class FResolveDepth8XPS : public FGlobalShader
+{
+	DECLARE_EXPORTED_SHADER_TYPE(FResolveDepth8XPS, Global, RENDERCORE_API);
+public:
+
+	typedef FDummyResolveParameter FParameter;
+
+	static bool ShouldCache(EShaderPlatform Platform) { return GetMaxSupportedFeatureLevel(Platform) >= ERHIFeatureLevel::SM5; }
+
+	static bool ShouldCompilePermutation(const FGlobalShaderPermutationParameters& Parameters)
+	{
+		return GetMaxSupportedFeatureLevel(Parameters.Platform) >= ERHIFeatureLevel::SM5
+			|| IsSimulatedPlatform(Parameters.Platform); // support resolving MSAA depth in mobile emulation
+	}
+
+	static void ModifyCompilationEnvironment(const FGlobalShaderPermutationParameters& Parameters, FShaderCompilerEnvironment& OutEnvironment)
+	{
+		FGlobalShader::ModifyCompilationEnvironment(Parameters, OutEnvironment);
+		OutEnvironment.SetDefine(TEXT("DEPTH_RESOLVE_NUM_SAMPLES"), 8);
+		OutEnvironment.SetDefine(TEXT("SIMULATED_PLATFORM"), IsSimulatedPlatform(Parameters.Platform) ? 1 : 0);
+	}
+
+	FResolveDepth8XPS(const ShaderMetaType::CompiledShaderInitializerType& Initializer) :
+		FGlobalShader(Initializer)
+	{
+		UnresolvedSurface.Bind(Initializer.ParameterMap, TEXT("UnresolvedSurface"), SPF_Mandatory);
+	}
+	FResolveDepth8XPS() {}
+
+	void SetParameters(FRHICommandList& RHICmdList, FParameter)
+	{
+	}
+
+	virtual bool Serialize(FArchive& Ar) override
+	{
+		bool bShaderHasOutdatedParameters = FGlobalShader::Serialize(Ar);
+		Ar << UnresolvedSurface;
+		return bShaderHasOutdatedParameters;
+	}
+
+	FShaderResourceParameter UnresolvedSurface;
+};
+
+
+class FResolveDepthNonMSPS : public FGlobalShader
+{
+	DECLARE_EXPORTED_SHADER_TYPE(FResolveDepthNonMSPS, Global, RENDERCORE_API);
+public:
+	
+	typedef FDummyResolveParameter FParameter;
+	
+	static bool ShouldCompilePermutation(const FGlobalShaderPermutationParameters& Parameters) { return GetMaxSupportedFeatureLevel(Parameters.Platform) <= ERHIFeatureLevel::ES3_1; }
+	
+	FResolveDepthNonMSPS(const ShaderMetaType::CompiledShaderInitializerType& Initializer):
+	FGlobalShader(Initializer)
+	{
+		UnresolvedSurface.Bind(Initializer.ParameterMap,TEXT("UnresolvedSurfaceNonMS"), SPF_Mandatory);
+	}
+	FResolveDepthNonMSPS() {}
+	
+	void SetParameters(FRHICommandList& RHICmdList, FParameter)
+	{
+	}
+	
+	virtual bool Serialize(FArchive& Ar) override
+	{
+		bool bShaderHasOutdatedParameters = FGlobalShader::Serialize(Ar);
+		Ar << UnresolvedSurface;
+		return bShaderHasOutdatedParameters;
+	}
+	
+	FShaderResourceParameter UnresolvedSurface;
+};
+
+class FResolveSingleSamplePS : public FGlobalShader
+{
+	DECLARE_EXPORTED_SHADER_TYPE(FResolveSingleSamplePS, Global, RENDERCORE_API);
+public:
+	
+	typedef uint32 FParameter;
+	
+	static bool ShouldCompilePermutation(const FGlobalShaderPermutationParameters& Parameters) { return Parameters.Platform == SP_PCD3D_SM5; }
+	
+	FResolveSingleSamplePS(const ShaderMetaType::CompiledShaderInitializerType& Initializer):
+	FGlobalShader(Initializer)
+	{
+		UnresolvedSurface.Bind(Initializer.ParameterMap,TEXT("UnresolvedSurface"), SPF_Mandatory);
+		SingleSampleIndex.Bind(Initializer.ParameterMap,TEXT("SingleSampleIndex"), SPF_Mandatory);
+	}
+	FResolveSingleSamplePS() {}
+	
+	RENDERCORE_API void SetParameters(FRHICommandList& RHICmdList, uint32 SingleSampleIndexValue);
+	
+	virtual bool Serialize(FArchive& Ar) override
+	{
+		bool bShaderHasOutdatedParameters = FGlobalShader::Serialize(Ar);
+		Ar << UnresolvedSurface;
+		Ar << SingleSampleIndex;
+		return bShaderHasOutdatedParameters;
+	}
+	
+	FShaderResourceParameter UnresolvedSurface;
+	FShaderParameter SingleSampleIndex;
+};
+
+/**
+ * A vertex shader for rendering a textured screen element.
+ */
+class FResolveVS : public FGlobalShader
+{
+	DECLARE_EXPORTED_SHADER_TYPE(FResolveVS, Global, RENDERCORE_API);
+public:
+	
+	static bool ShouldCompilePermutation(const FGlobalShaderPermutationParameters& Parameters) { return true; }
+	
+	FResolveVS(const ShaderMetaType::CompiledShaderInitializerType& Initializer)
+		: FGlobalShader(Initializer)
+	{
+		PositionMinMax.Bind(Initializer.ParameterMap, TEXT("PositionMinMax"), SPF_Mandatory);
+		UVMinMax.Bind(Initializer.ParameterMap, TEXT("UVMinMax"), SPF_Mandatory);
+	}
+	FResolveVS() {}
+
+	RENDERCORE_API void SetParameters(FRHICommandList& RHICmdList, FVector4 PositionMinMax, FVector4 UVMinMax);
+
+	virtual bool Serialize(FArchive& Ar) override
+	{
+		bool bShaderHasOutdatedParameters = FGlobalShader::Serialize(Ar);
+		Ar << PositionMinMax;
+		Ar << UVMinMax;
+		return bShaderHasOutdatedParameters;
+	}
+
+	FShaderParameter PositionMinMax;
+	FShaderParameter UVMinMax;
+};
+
+struct FResolveVertexBuffer : public FRenderResource
+{
+	FVertexBufferRHIRef VertexBufferRHI;
+	FVertexDeclarationRHIRef VertexDeclarationRHI;
+
+	virtual void InitDynamicRHI() override;
+	virtual void ReleaseDynamicRHI() override;
+};
+
+extern RENDERCORE_API TGlobalResource<FResolveVertexBuffer> GResolveVertexBuffer;
