Downgraded missing TexCreate_ShaderResource error to be non-fatal.

diff --git a/Engine/Source/Runtime/RenderCore/Private/ShaderParameterStruct.cpp b/Engine/Source/Runtime/RenderCore/Private/ShaderParameterStruct.cpp
--- a/Engine/Source/Runtime/RenderCore/Private/ShaderParameterStruct.cpp
+++ b/Engine/Source/Runtime/RenderCore/Private/ShaderParameterStruct.cpp
@@ -475,7 +475,7 @@ void ValidateShaderParameters(const FShader* Shader, const FShaderParametersMeta
 		{
 			FString MemberName = ParametersMetadata->GetFullMemberCodeName(ParameterBinding.ByteOffset);
 
-			UE_LOG(LogShaders, Fatal,
+			UE_LOG(LogShaders, Error,
 				TEXT("Attempting to set shader %s parameter %s::%s with the RDG texture %s which was not created with TexCreate_ShaderResource"),
 				ShaderClassName, ShaderParemeterStructName, *MemberName, GraphTexture->Name);
 		}
