Fixed crash when creating ray tracing geometry if any section has a null vertex buffer.

(ushell-p4-cherrypick of 9568434 by [REDACTED_PERSON])

diff --git a/Engine/Source/Runtime/RenderCore/Public/RenderResource.h b/Engine/Source/Runtime/RenderCore/Public/RenderResource.h
--- a/Engine/Source/Runtime/RenderCore/Public/RenderResource.h
+++ b/Engine/Source/Runtime/RenderCore/Public/RenderResource.h
@@ -597,7 +597,17 @@ public:
 
 	virtual void InitRHI() override
 	{
-		if (Initializer.IndexBuffer && IsRayTracingEnabled())
+		bool bAllSegmentsAreValid = true;
+		for (const FRayTracingGeometrySegment& Segment : Initializer.Segments)
+		{
+			if (!Segment.VertexBuffer)
+			{
+				bAllSegmentsAreValid = false;
+				break;
+			}
+		}
+
+		if (Initializer.IndexBuffer && bAllSegmentsAreValid && IsRayTracingEnabled())
 		{
 			RayTracingGeometryRHI = RHICreateRayTracingGeometry(Initializer);
 			FRHICommandListExecutor::GetImmediateCommandList().BuildAccelerationStructure(RayTracingGeometryRHI);
