Added support for creating ray tracing shaders through FShaderCodeLibrary

diff --git a/Engine/Source/Runtime/RenderCore/Private/Shader.cpp b/Engine/Source/Runtime/RenderCore/Private/Shader.cpp
--- a/Engine/Source/Runtime/RenderCore/Private/Shader.cpp
+++ b/Engine/Source/Runtime/RenderCore/Private/Shader.cpp
@@ -1012,15 +1012,13 @@ void FShaderResource::InitRHI()
 	{
 		if (GRHISupportsRayTracing)
 		{
-			RayTracingShader = RHICreateRayTracingShader(UncompressedCode, Target.GetFrequency());
+			RayTracingShader = FShaderCodeLibrary::CreateRayTracingShader((EShaderPlatform)Target.Platform, Target.GetFrequency(), OutputHash, UncompressedCode);
 			UE_CLOG((bCodeInSharedLocation && !IsValidRef(RayTracingShader)), LogShaders, Fatal, TEXT("FShaderResource::SerializeShaderCode can't find shader code for: [%s]"), *LegacyShaderPlatformToShaderFormat((EShaderPlatform)Target.Platform).ToString());
 
 			if (Target.Frequency == SF_RayHitGroup)
 			{
 				RayTracingMaterialLibraryIndex = AddToRayTracingLibrary(RayTracingShader);
 			}
-
-			RayTracingShader->SetHash(OutputHash);
 		}
 	}
 #endif // RHI_RAYTRACING
diff --git a/Engine/Source/Runtime/RenderCore/Private/ShaderCodeLibrary.cpp b/Engine/Source/Runtime/RenderCore/Private/ShaderCodeLibrary.cpp
--- a/Engine/Source/Runtime/RenderCore/Private/ShaderCodeLibrary.cpp
+++ b/Engine/Source/Runtime/RenderCore/Private/ShaderCodeLibrary.cpp
@@ -756,6 +756,30 @@ public:
 		return Shader;
 	}
 
+	FRayTracingShaderRHIRef CreateRayTracingShader(EShaderFrequency Frequency, const FSHAHash& Hash) override final
+	{
+		FRayTracingShaderRHIRef Shader;
+
+#if RHI_RAYTRACING
+		int32 Size = 0;
+		bool bWasSync = false;
+		TArray<uint8>* Code = LookupShaderCode(Hash, Size, bWasSync);
+		if (Code)
+		{
+			TArray<uint8> UCode;
+			TArray<uint8>& UncompressedCode = FShaderLibraryHelperUncompressCode(Platform, Size, *Code, UCode);
+			Shader = RHICreateRayTracingShader(UncompressedCode, Frequency);
+			CheckShaderCreation(Shader.GetReference(), Hash);
+			if (bWasSync)
+			{
+				ReleaseShaderCode(Hash);
+			}
+		}
+#endif // RHI_RAYTRACING
+
+		return Shader;
+	}
+
 	class FShaderCodeLibraryIterator : public FRHIShaderLibrary::FShaderLibraryIterator
 	{
 	public:
@@ -1876,6 +1900,22 @@ public:
 		return Result;
 	}
 
+	FRayTracingShaderRHIRef CreateRayTracingShader(EShaderPlatform Platform, EShaderFrequency Frequency, FSHAHash Hash)
+	{
+		FRayTracingShaderRHIRef Result;
+
+#if RHI_RAYTRACING
+		checkSlow(Platform == GetRuntimeShaderPlatform());
+		FRHIShaderLibrary* ShaderCodeArchive = FindShaderLibrary(Hash);
+		if (ShaderCodeArchive)
+		{
+			Result = ((FShaderCodeArchive*)ShaderCodeArchive)->CreateRayTracingShader(Frequency, Hash);
+		}
+#endif // RHI_RAYTRACING
+
+		return Result;
+	}
+
 	TRefCountPtr<FRHIShaderLibrary::FShaderLibraryIterator> CreateIterator(void)
 	{
 		return new FShaderCodeLibraryIterator(ShaderCodeArchiveStack, LibraryMutex);
@@ -2372,6 +2412,25 @@ FComputeShaderRHIRef FShaderCodeLibrary::CreateComputeShader(EShaderPlatform Pla
 	return Shader;
 }
 
+FRayTracingShaderRHIRef FShaderCodeLibrary::CreateRayTracingShader(EShaderPlatform Platform, EShaderFrequency Frequency, FSHAHash Hash, TArray<uint8> const& Code)
+{
+	FRayTracingShaderRHIRef Shader;
+
+#if RHI_RAYTRACING
+	if (FShaderCodeLibraryImpl::Impl)
+	{
+		Shader = FShaderCodeLibraryImpl::Impl->CreateRayTracingShader(Platform, Frequency, Hash);
+	}
+	if (!IsValidRef(Shader))
+	{
+		Shader = RHICreateRayTracingShader(Code, Frequency);
+	}
+	SafeAssignHash(Shader, Hash);
+#endif // RHI_RAYTRACING
+
+	return Shader;
+}
+
 TRefCountPtr<FRHIShaderLibrary::FShaderLibraryIterator> FShaderCodeLibrary::CreateIterator(void)
 {
 	TRefCountPtr<FRHIShaderLibrary::FShaderLibraryIterator> It;
diff --git a/Engine/Source/Runtime/RenderCore/Public/ShaderCodeLibrary.h b/Engine/Source/Runtime/RenderCore/Public/ShaderCodeLibrary.h
--- a/Engine/Source/Runtime/RenderCore/Public/ShaderCodeLibrary.h
+++ b/Engine/Source/Runtime/RenderCore/Public/ShaderCodeLibrary.h
@@ -127,6 +127,7 @@ public:
 	virtual FDomainShaderRHIRef CreateDomainShader(const FSHAHash& Hash) = 0;
 	virtual FGeometryShaderRHIRef CreateGeometryShader(const FSHAHash& Hash) = 0;
 	virtual FComputeShaderRHIRef CreateComputeShader(const FSHAHash& Hash) = 0;
+	virtual FRayTracingShaderRHIRef CreateRayTracingShader(EShaderFrequency Frequency, const FSHAHash& Hash) = 0;
 };
 
 DECLARE_MULTICAST_DELEGATE_TwoParams(FSharedShaderCodeRequest, const FSHAHash&, FArchive*);
@@ -163,6 +164,8 @@ struct RENDERCORE_API FShaderCodeLibrary
 	static FDomainShaderRHIRef CreateDomainShader(EShaderPlatform Platform, FSHAHash Hash, TArray<uint8> const& Code);
 	/** Instantiate or retrieve a compute shader from the cache for the provided code & hash. */
 	static FComputeShaderRHIRef CreateComputeShader(EShaderPlatform Platform, FSHAHash Hash, TArray<uint8> const& Code);
+	/** Instantiate or retrieve a ray tracing shader from the cache for the provided code & hash. */
+	static FRayTracingShaderRHIRef CreateRayTracingShader(EShaderPlatform Platform, EShaderFrequency Frequency, FSHAHash Hash, TArray<uint8> const& Code);
 
     static bool ContainsShaderCode(const FSHAHash& Hash);
     
