Delete Library file handle on close (Non native ShaderCodeLibrary implementation only), unfortunatly this also means we have to wait on any outstanding read tasks.

Plus - Fix rare shutdown check() failure when async file reads are still wainting in the FShaderPipelineCache.

Plus - Put back correct order of shutdown for FShaderCodeLibrary and FShaderPipelineCache as they can have outstanding aync file handles and we'll get a crash if they are closed after the module shutdown has deleted the FPakPrecacher singleton.

diff --git a/Engine/Source/Runtime/RenderCore/Private/ShaderCodeLibrary.cpp b/Engine/Source/Runtime/RenderCore/Private/ShaderCodeLibrary.cpp
--- a/Engine/Source/Runtime/RenderCore/Private/ShaderCodeLibrary.cpp
+++ b/Engine/Source/Runtime/RenderCore/Private/ShaderCodeLibrary.cpp
@@ -424,6 +424,7 @@ public:
 		, LibraryDir(InLibraryDir)
 		, LibraryCodeOffset(0)
 		, LibraryAsyncFileHandle(nullptr)
+		, InFlightAsyncReadRequests(0)
 	{
 		FName PlatformName = LegacyShaderPlatformToShaderFormat(InPlatform);
 		FString DestFilePath = GetCodeArchiveFilename(LibraryDir, LibraryName, PlatformName);
@@ -462,6 +463,33 @@ public:
 
 	virtual ~FShaderCodeArchive()
 	{
+		if(LibraryAsyncFileHandle != nullptr)
+		{
+			UE_LOG(LogShaderLibrary, Display, TEXT("FShaderCodeArchive: Shutting down %s"), *GetName());
+
+			FScopeLock ScopeLock(&ReadRequestLock);
+
+			const int64 OutstandingReads = FPlatformAtomics::AtomicRead(&InFlightAsyncReadRequests);
+			if(OutstandingReads > 0)
+			{
+				const float MaxWaitTimePerRead = 1.f / 60.f;
+				UE_LOG(LogShaderLibrary, Warning, TEXT("FShaderCodeArchive: Library %s has %d inflight requests to LibraryAsyncFileHandle - cancelling and waiting %f seconds each for them to finish."), *GetName(), OutstandingReads, MaxWaitTimePerRead);
+
+				for(auto& Pair : Shaders)
+				{
+					FShaderCodeEntry& Entry = Pair.Value;
+					TSharedPtr<IAsyncReadRequest, ESPMode::ThreadSafe> LocalReadRequest = Entry.ReadRequest.Pin();
+					if(LocalReadRequest.IsValid())
+					{
+						LocalReadRequest->Cancel();
+						LocalReadRequest->WaitCompletion(MaxWaitTimePerRead);
+					}
+				}
+			}
+
+			delete LibraryAsyncFileHandle;
+			LibraryAsyncFileHandle = nullptr;
+		}
 	}
 
 	virtual bool IsLibraryNativeFormat() const { return false; }
@@ -561,6 +589,8 @@ public:
 
 			if (bHasReadRequest)
 			{
+				FPlatformAtomics::InterlockedAdd(&InFlightAsyncReadRequests, 1);
+			
 				FExternalReadCallback ExternalReadCallback = [this, Entry, LocalReadRequest](double ReaminingTime)
 				{
 					return this->OnExternalReadCallback(LocalReadRequest, Entry, ReaminingTime);
@@ -602,6 +632,8 @@ public:
 #if DO_CHECK
 		Entry->bReadCompleted = 1;
 #endif
+
+		FPlatformAtomics::InterlockedAdd(&InFlightAsyncReadRequests, -1);
 		
 		return true;
 	}
@@ -861,6 +893,9 @@ private:
 	// Library file handle for async reads
 	IAsyncReadFileHandle* LibraryAsyncFileHandle;
 	FCriticalSection ReadRequestLock;
+	
+	// A count of the number of LibraryAsync Read Requests in flight
+	volatile int64 InFlightAsyncReadRequests;
 
 	// The shader code present in the library
 	TMap<FSHAHash, FShaderCodeEntry> Shaders;
diff --git a/Engine/Source/Runtime/RenderCore/Private/ShaderPipelineCache.cpp b/Engine/Source/Runtime/RenderCore/Private/ShaderPipelineCache.cpp
--- a/Engine/Source/Runtime/RenderCore/Private/ShaderPipelineCache.cpp
+++ b/Engine/Source/Runtime/RenderCore/Private/ShaderPipelineCache.cpp
@@ -1150,35 +1150,25 @@ FShaderPipelineCache::~FShaderPipelineCache()
 	
 	// The render thread tick should be dead now and we are safe to destroy everything that needs to wait or manual destruction
 
-	for(CompileJob& Entry: ReadTasks )
+	// Close() should have called though to Flush() - all work is now in the shutdown lists
+	for(CompileJob& Entry: ShutdownReadCompileTasks)
 	{
-		check(Entry.ReadRequests);
-		Entry.ReadRequests->BlockingWaitComplete();
-	}
-	
-	for(CompileJob& Entry: ShutdownReadCompileTasks )
-	{
-		check(Entry.ReadRequests);
-		Entry.ReadRequests->BlockingWaitComplete();
-	}
-	
-	for (FPipelineCacheFileFormatPSORead* Entry : FetchTasks)
-	{
-		if(Entry->ReadRequest.IsValid())
+		if(Entry.ReadRequests != nullptr)
 		{
-			Entry->ReadRequest->WaitCompletion(0.0);
+			Entry.ReadRequests->BlockingWaitComplete();
+			delete Entry.ReadRequests;
+			Entry.ReadRequests = nullptr;
 		}
-		delete Entry;
 	}
 	
 	for (FPipelineCacheFileFormatPSORead* Entry : ShutdownFetchTasks)
 	{
 		if(Entry->ReadRequest.IsValid())
 		{
-			Entry->ReadRequest->WaitCompletion(0.0);
+			Entry->ReadRequest->WaitCompletion(0.f);
 		}
 		delete Entry;
-	}	
+	}
 }
 
 bool FShaderPipelineCache::IsTickable() const
