Fixing up header in shadercode library CSV, previously the first entry would append to the same line and result in "OutputHashempty" as a header field. Standard line break appended to correct.

diff --git a/Engine/Source/Runtime/RenderCore/Private/ShaderCodeLibrary.cpp b/Engine/Source/Runtime/RenderCore/Private/ShaderCodeLibrary.cpp
--- a/Engine/Source/Runtime/RenderCore/Private/ShaderCodeLibrary.cpp
+++ b/Engine/Source/Runtime/RenderCore/Private/ShaderCodeLibrary.cpp
@@ -1522,7 +1522,8 @@ struct FEditorShaderStableInfo
 			{
 				TUniquePtr<FArchive> IntermediateFormatAr(IFileManager::Get().CreateFileWriter(*IntermediateFormatPath));
 
-				const FString HeaderText = FStableShaderKeyAndValue::HeaderLine();
+				FString HeaderText = FStableShaderKeyAndValue::HeaderLine();
+				HeaderText += TCHAR('\n');
 				auto HeaderSrc = StringCast<ANSICHAR>(*HeaderText, HeaderText.Len());
 
 				IntermediateFormatAr->Serialize((ANSICHAR*)HeaderSrc.Get(), HeaderSrc.Length() * sizeof(ANSICHAR));
