Support shader pipelines in the shader code library's stable map.

- Added shader pipeline hash as a key attribute.
- Added shaders from shader pipelines (they can have different code if RHI removes unused interpolators between the stages, like D3D does).

diff --git a/Engine/Source/Runtime/Engine/Private/Materials/MaterialShader.cpp b/Engine/Source/Runtime/Engine/Private/Materials/MaterialShader.cpp
--- a/Engine/Source/Runtime/Engine/Private/Materials/MaterialShader.cpp
+++ b/Engine/Source/Runtime/Engine/Private/Materials/MaterialShader.cpp
@@ -2650,6 +2650,13 @@ void FMaterialShaderMap::SaveShaderStableKeys(EShaderPlatform TargetShaderPlatfo
 		FShader* Shader = KeyValue.Value;
 		Shader->SaveShaderStableKeys(TargetShaderPlatform, SaveKeyVal);
 	}
+
+	TArray<FShaderPipeline*> ShadersPipelinesL;
+	GetShaderPipelineList(ShadersPipelinesL);
+	for (FShaderPipeline* Pipeline : ShadersPipelinesL)
+	{
+		Pipeline->SaveShaderStableKeys(TargetShaderPlatform, SaveKeyVal);
+	}
 #endif
 }
 
diff --git a/Engine/Source/Runtime/RenderCore/Private/Shader.cpp b/Engine/Source/Runtime/RenderCore/Private/Shader.cpp
--- a/Engine/Source/Runtime/RenderCore/Private/Shader.cpp
+++ b/Engine/Source/Runtime/RenderCore/Private/Shader.cpp
@@ -1889,6 +1889,38 @@ void FShaderPipeline::CookPipeline(FShaderPipeline* Pipeline)
 #endif
 }
 
+void FShaderPipeline::SaveShaderStableKeys(EShaderPlatform TargetShaderPlatform, const struct FStableShaderKeyAndValue& InSaveKeyVal)
+{
+	// the higher level code can pass SP_NumPlatforms, in which case play it safe and use a platform that we know can remove inteprolators
+	const EShaderPlatform ShaderPlatformThatSupportsRemovingInterpolators = SP_PCD3D_SM5;
+	checkf(RHISupportsShaderPipelines(ShaderPlatformThatSupportsRemovingInterpolators), TEXT("We assumed that shader platform %d supports shaderpipelines while it doesn't"), static_cast<int32>(ShaderPlatformThatSupportsRemovingInterpolators));
+
+	bool bCanHaveUniqueShaders = (TargetShaderPlatform != SP_NumPlatforms) ? PipelineType->ShouldOptimizeUnusedOutputs(TargetShaderPlatform) : PipelineType->ShouldOptimizeUnusedOutputs(ShaderPlatformThatSupportsRemovingInterpolators);
+	if (bCanHaveUniqueShaders)
+	{
+		FStableShaderKeyAndValue SaveKeyVal(InSaveKeyVal);
+		SaveKeyVal.SetPipelineHash(this); // could use PipelineType->GetSourceHash(), but each pipeline instance even of the same type can have unique shaders
+
+		// Using GetShaders() would be more future-proof but would result in more dynamic allocation churn.
+		// Instead, mirroring the logic here
+		VertexShader->SaveShaderStableKeys(TargetShaderPlatform, SaveKeyVal);
+
+		if (PixelShader)
+		{
+			PixelShader->SaveShaderStableKeys(TargetShaderPlatform, SaveKeyVal);
+		}
+		if (GeometryShader)
+		{
+			GeometryShader->SaveShaderStableKeys(TargetShaderPlatform, SaveKeyVal);
+		}
+		if (HullShader)
+		{
+			HullShader->SaveShaderStableKeys(TargetShaderPlatform, SaveKeyVal);
+			DomainShader->SaveShaderStableKeys(TargetShaderPlatform, SaveKeyVal);
+		}
+	}
+}
+
 void DumpShaderStats(EShaderPlatform Platform, EShaderFrequency Frequency)
 {
 #if ALLOW_DEBUG_FILES
diff --git a/Engine/Source/Runtime/RenderCore/Private/ShaderCodeLibrary.cpp b/Engine/Source/Runtime/RenderCore/Private/ShaderCodeLibrary.cpp
--- a/Engine/Source/Runtime/RenderCore/Private/ShaderCodeLibrary.cpp
+++ b/Engine/Source/Runtime/RenderCore/Private/ShaderCodeLibrary.cpp
@@ -214,7 +214,8 @@ void FixupUnsanitizedNames(const FString& Src, TArray<FString>& OutFields)
 		}
 		OutFields.Empty();
 		NewSrc.TrimStartAndEnd().ParseIntoArray(OutFields, TEXT(","), false);
-		check(OutFields.Num() == 11);
+		// allow formats both with and without pipeline hash
+		check(OutFields.Num() == 11 || OutFields.Num() == 12);
 	}
 }
 
@@ -233,19 +234,24 @@ void FStableShaderKeyAndValue::ComputeKeyHash()
 
 	KeyHash = HashCombine(KeyHash, GetTypeHash(VFType));
 	KeyHash = HashCombine(KeyHash, GetTypeHash(PermutationId));
+	KeyHash = HashCombine(KeyHash, GetTypeHash(PipelineHash));
 }
 
 void FStableShaderKeyAndValue::ParseFromString(const FString& Src)
 {
 	TArray<FString> Fields;
 	Src.TrimStartAndEnd().ParseIntoArray(Fields, TEXT(","), false);
-	if (Fields.Num() > 11)
+
+	/* disabled, should not be happening since 1H 2018
+	if (Fields.Num() > 12)
 	{
 		// hack fix for unsanitized names, should not occur anymore.
 		FixupUnsanitizedNames(Src, Fields);
 	}
+	*/
 
-	check(Fields.Num() == 11);
+	// for a while, accept old .scl.csv without pipelinehash
+	check(Fields.Num() == 11 || Fields.Num() == 12);
 
 	int32 Index = 0;
 	ClassNameAndObjectPath.ParseFromString(Fields[Index++]);
@@ -266,6 +272,15 @@ void FStableShaderKeyAndValue::ParseFromString(const FString& Src)
 
 	check(Index == 11);
 
+	if (Fields.Num() == 12)
+	{
+		PipelineHash.FromString(Fields[Index++]);
+	}
+	else
+	{
+		PipelineHash = FSHAHash();
+	}
+
 	ComputeKeyHash();
 }
 
@@ -275,13 +290,16 @@ void FStableShaderKeyAndValue::ParseFromStringCached(const FString& Src, TMap<ui
 	TArray<FString> Fields;
 	Src.TrimStartAndEnd().ParseIntoArray(Fields, TEXT(","), false);
 
+	/* disabled, should not be happening since 1H 2018
 	if (Fields.Num() > 11)
 	{
 		// hack fix for unsanitized names, should not occur anymore.
 		FixupUnsanitizedNames(Src, Fields);
 	}
+	*/
 	
-	check(Fields.Num() == 11);
+	// for a while, accept old .scl.csv without pipelinehash
+	check(Fields.Num() == 11 || Fields.Num() == 12);
 
 	int32 Index = 0;
 	ClassNameAndObjectPath.ParseFromString(Fields[Index++]);
@@ -304,6 +322,15 @@ void FStableShaderKeyAndValue::ParseFromStringCached(const FString& Src, TMap<ui
 
 	check(Index == 11);
 
+	if (Fields.Num() == 12)
+	{
+		PipelineHash.FromString(Fields[Index++]);
+	}
+	else
+	{
+		PipelineHash = FSHAHash();
+	}
+
 	ComputeKeyHash();
 }
 
@@ -345,6 +372,9 @@ void FStableShaderKeyAndValue::ToString(FString& OutResult) const
 	OutResult += Delim;
 
 	OutResult += OutputHash.ToString();
+	OutResult += Delim;
+
+	OutResult += PipelineHash.ToString();
 }
 
 FString FStableShaderKeyAndValue::HeaderLine()
@@ -378,10 +408,66 @@ FString FStableShaderKeyAndValue::HeaderLine()
 	Result += Delim;
 
 	Result += TEXT("OutputHash");
+	Result += Delim;
+
+	Result += TEXT("PipelineHash");
 
 	return Result;
 }
 
+void FStableShaderKeyAndValue::SetPipelineHash(FShaderPipeline* Pipeline)
+{
+	if (LIKELY(Pipeline))
+	{
+		// cache this?
+		FShaderCodeLibraryPipeline LibraryPipeline;
+		LibraryPipeline.Initialize(Pipeline);
+		LibraryPipeline.GetPipelineHash(PipelineHash); 
+	}
+	else
+	{
+		PipelineHash = FSHAHash();
+	}
+}
+
+void FShaderCodeLibraryPipeline::Initialize(FShaderPipeline* Pipeline)
+{
+	check(Pipeline != nullptr);
+
+	if (IsValidRef(Pipeline->VertexShader))
+	{
+		VertexShader = Pipeline->VertexShader->GetOutputHash();
+	}
+	if (IsValidRef(Pipeline->GeometryShader))
+	{
+		GeometryShader = Pipeline->GeometryShader->GetOutputHash();
+	}
+	if (IsValidRef(Pipeline->HullShader))
+	{
+		HullShader = Pipeline->HullShader->GetOutputHash();
+	}
+	if (IsValidRef(Pipeline->DomainShader))
+	{
+		DomainShader = Pipeline->DomainShader->GetOutputHash();
+	}
+	if (IsValidRef(Pipeline->PixelShader))
+	{
+		PixelShader = Pipeline->PixelShader->GetOutputHash();
+	}
+}
+
+void FShaderCodeLibraryPipeline::GetPipelineHash(FSHAHash& Output)
+{
+	FSHA1 Hasher;
+	Hasher.Update(&VertexShader.Hash[0], sizeof(VertexShader.Hash));
+	Hasher.Update(&PixelShader.Hash[0], sizeof(PixelShader.Hash));
+	Hasher.Update(&GeometryShader.Hash[0], sizeof(GeometryShader.Hash));
+	Hasher.Update(&HullShader.Hash[0], sizeof(HullShader.Hash));
+	Hasher.Update(&DomainShader.Hash[0], sizeof(DomainShader.Hash));
+
+	Hasher.Final();
+	Hasher.GetHash(&Output.Hash[0]);
+}
 
 struct FShaderCodeEntry
 {
@@ -990,26 +1076,7 @@ struct FEditorShaderCodeArchive
 		EShaderPlatform ShaderPlatform = ShaderFormatToLegacyShaderPlatform(FormatName);
 
 		FShaderCodeLibraryPipeline LibraryPipeline;
-		if (IsValidRef(Pipeline->VertexShader))
-		{
-			LibraryPipeline.VertexShader = Pipeline->VertexShader->GetOutputHash();
-		}
-		if (IsValidRef(Pipeline->GeometryShader))
-		{
-			LibraryPipeline.GeometryShader = Pipeline->GeometryShader->GetOutputHash();
-		}
-		if (IsValidRef(Pipeline->HullShader))
-		{
-			LibraryPipeline.HullShader = Pipeline->HullShader->GetOutputHash();
-		}
-		if (IsValidRef(Pipeline->DomainShader))
-		{
-			LibraryPipeline.DomainShader = Pipeline->DomainShader->GetOutputHash();
-		}
-		if (IsValidRef(Pipeline->PixelShader))
-		{
-			LibraryPipeline.PixelShader = Pipeline->PixelShader->GetOutputHash();
-		}
+		LibraryPipeline.Initialize(Pipeline);
 		if (!Pipelines.Contains(LibraryPipeline))
 		{
 			Pipelines.Add(LibraryPipeline);
diff --git a/Engine/Source/Runtime/RenderCore/Public/Shader.h b/Engine/Source/Runtime/RenderCore/Public/Shader.h
--- a/Engine/Source/Runtime/RenderCore/Public/Shader.h
+++ b/Engine/Source/Runtime/RenderCore/Public/Shader.h
@@ -1863,6 +1863,9 @@ public:
 	};
 	
 	static void CookPipeline(FShaderPipeline* Pipeline);
+
+	/** Saves stable keys for the shaders in the pipeline */
+	void SaveShaderStableKeys(EShaderPlatform TargetShaderPlatform, const struct FStableShaderKeyAndValue& SaveKeyVal);
 };
 
 inline bool operator<(const FShaderPipeline& Lhs, const FShaderPipeline& Rhs)
diff --git a/Engine/Source/Runtime/RenderCore/Public/ShaderCodeLibrary.h b/Engine/Source/Runtime/RenderCore/Public/ShaderCodeLibrary.h
--- a/Engine/Source/Runtime/RenderCore/Public/ShaderCodeLibrary.h
+++ b/Engine/Source/Runtime/RenderCore/Public/ShaderCodeLibrary.h
@@ -21,7 +21,10 @@ struct RENDERCORE_API FShaderCodeLibraryPipeline
 	FSHAHash HullShader;
 	FSHAHash DomainShader;
 	mutable uint32 Hash;
-	
+
+	/** Fills the hashes from the pipeline stage shaders */
+	void Initialize(FShaderPipeline* Pipeline);
+
 	FShaderCodeLibraryPipeline() : Hash(0) {}
 	
 	friend bool operator ==(const FShaderCodeLibraryPipeline& A,const FShaderCodeLibraryPipeline& B)
@@ -41,6 +44,9 @@ struct RENDERCORE_API FShaderCodeLibraryPipeline
 		}
 		return Key.Hash;
 	}
+
+	/** Computes a longer hash that uniquely identifies the whole pipeline, used in FStableShaderKeyAndValue */
+	void GetPipelineHash(FSHAHash& Output);
 	
 	friend FArchive& operator<<( FArchive& Ar, FShaderCodeLibraryPipeline& Info )
 	{
@@ -75,6 +81,7 @@ struct RENDERCORE_API FStableShaderKeyAndValue
 	FName TargetPlatform;
 	FName VFType;
 	FName PermutationId;
+	FSHAHash PipelineHash;
 
 	uint32 KeyHash;
 
@@ -92,6 +99,9 @@ struct RENDERCORE_API FStableShaderKeyAndValue
 	void ToString(FString& OutResult) const;
 	static FString HeaderLine();
 
+	/** Computes pipeline hash from the passed pipeline. Pass nullptr to clear */
+	void SetPipelineHash(FShaderPipeline* Pipeline);
+
 	friend bool operator ==(const FStableShaderKeyAndValue& A, const FStableShaderKeyAndValue& B)
 	{
 		return
@@ -104,7 +114,8 @@ struct RENDERCORE_API FStableShaderKeyAndValue
 			A.TargetFrequency == B.TargetFrequency &&
 			A.TargetPlatform == B.TargetPlatform &&
 			A.VFType == B.VFType &&
-			A.PermutationId == B.PermutationId;
+			A.PermutationId == B.PermutationId &&
+			A.PipelineHash == B.PipelineHash;
 	}
 
 	friend uint32 GetTypeHash(const FStableShaderKeyAndValue &Key)
