Fixing up mipmap generation using compute for textures/mipmaps not multiple of 8

diff --git a/Engine/Source/Runtime/RenderCore/Private/GenerateMips.cpp b/Engine/Source/Runtime/RenderCore/Private/GenerateMips.cpp
--- a/Engine/Source/Runtime/RenderCore/Private/GenerateMips.cpp
+++ b/Engine/Source/Runtime/RenderCore/Private/GenerateMips.cpp
@@ -125,8 +125,8 @@ void FGenerateMips::Compute(FRHICommandListImmediate& RHIImmCmdList, FRHITexture
 
 		//Dispatch count is the destination's mip texture dimensions, so only the number required is executed.
 		FIntVector GenMipsGroupCount(
-			FMath::Max(DestTextureSizeX / MIPSSHADER_NUMTHREADS, 1),
-			FMath::Max(DestTextureSizeY / MIPSSHADER_NUMTHREADS, 1),
+			FMath::Max((DestTextureSizeX + MIPSSHADER_NUMTHREADS - 1) / MIPSSHADER_NUMTHREADS, 1),
+			FMath::Max((DestTextureSizeY + MIPSSHADER_NUMTHREADS - 1) / MIPSSHADER_NUMTHREADS, 1),
 			1);
 		//Pass added per mip level to be written.
 		ClearUnusedGraphResources(*ComputeShader, PassParameters);
