Adding sRGB compatible mip-generation code to mips generator (compute case)

diff --git a/Engine/Shaders/Private/ComputeGenerateMips.usf b/Engine/Shaders/Private/ComputeGenerateMips.usf
--- a/Engine/Shaders/Private/ComputeGenerateMips.usf
+++ b/Engine/Shaders/Private/ComputeGenerateMips.usf
@@ -6,10 +6,15 @@
 
 #pragma once
 #include "Common.ush"
+#include "GammaCorrectionCommon.ush"
 
 float2 TexelSize;
 Texture2D MipInSRV;
+#if GENMIPS_SRGB
+RWTexture2D<half4> MipOutUAV;
+#else
 RWTexture2D<float4> MipOutUAV;
+#endif
 SamplerState MipSampler;
 
 [numthreads(8, 8, 1)]
@@ -17,7 +22,13 @@ void MainCS(uint3 DT_ID : SV_DispatchThreadID)
 {
     float2 UV = TexelSize * (DT_ID.xy + 0.5f);
 
+#if GENMIPS_SRGB
+    half4 outColor = MipInSRV.SampleLevel(MipSampler, UV, 0);
+
+    MipOutUAV[DT_ID.xy] = half4(LinearToSrgb(outColor.xyz), outColor.w);
+#else
     float4 outColor = MipInSRV.SampleLevel(MipSampler, UV, 0);
 
     MipOutUAV[DT_ID.xy] = outColor;
+#endif
 }
\ No newline at end of file
diff --git a/Engine/Source/Runtime/RenderCore/Private/GenerateMips.cpp b/Engine/Source/Runtime/RenderCore/Private/GenerateMips.cpp
--- a/Engine/Source/Runtime/RenderCore/Private/GenerateMips.cpp
+++ b/Engine/Source/Runtime/RenderCore/Private/GenerateMips.cpp
@@ -4,6 +4,7 @@
 #include "RenderGraph.h"
 #include "RenderGraphUtils.h"
 #include "RenderTargetPool.h"
+#include "ShaderPermutation.h"
 
 #define MIPSSHADER_NUMTHREADS 8
 
@@ -18,6 +19,11 @@ struct FGenerateMipsStruct
 class FGenerateMipsCS : public FGlobalShader
 {
 	DECLARE_GLOBAL_SHADER(FGenerateMipsCS)
+
+public:
+	class FGenMipsSRGB : SHADER_PERMUTATION_BOOL("GENMIPS_SRGB");
+	using FPermutationDomain = TShaderPermutationDomain<FGenMipsSRGB>;
+
 	SHADER_USE_PARAMETER_STRUCT(FGenerateMipsCS, FGlobalShader)
 
 	BEGIN_SHADER_PARAMETER_STRUCT(FParameters, )
@@ -103,7 +109,11 @@ void FGenerateMips::Compute(FRHICommandListImmediate& RHIImmCmdList, FRHITexture
 	//Begin rendergraph for executing the compute shader
 	FRDGBuilder GraphBuilder(RHIImmCmdList);
 	FRDGTextureRef GraphTexture = GraphBuilder.RegisterExternalTexture(GenMipsStruct->RenderTarget, TEXT("GenerateMipsGraphTexture"));
-	TShaderMapRef<FGenerateMipsCS> ComputeShader(GetGlobalShaderMap(ERHIFeatureLevel::SM5));
+
+	//Select compute shader variant (normal vs. sRGB)
+	FGenerateMipsCS::FPermutationDomain PermutationVector;
+	PermutationVector.Set<FGenerateMipsCS::FGenMipsSRGB>(!!(InTexture->GetFlags() & TexCreate_SRGB));
+	TShaderMapRef<FGenerateMipsCS> ComputeShader(GetGlobalShaderMap(ERHIFeatureLevel::SM5), PermutationVector);
 
 	//Loop through each level of the mips that require creation and add a dispatch pass per level,.
 	for (uint8 MipLevel = 1; MipLevel < InTexture->GetNumMips(); MipLevel++)
