UE4 - Temporatily remove srgb permutation to fix compile issue
#rnx

diff --git a/Engine/Source/Runtime/RenderCore/Private/GenerateMips.cpp b/Engine/Source/Runtime/RenderCore/Private/GenerateMips.cpp
--- a/Engine/Source/Runtime/RenderCore/Private/GenerateMips.cpp
+++ b/Engine/Source/Runtime/RenderCore/Private/GenerateMips.cpp
@@ -21,9 +21,6 @@ class FGenerateMipsCS : public FGlobalShader
 	DECLARE_GLOBAL_SHADER(FGenerateMipsCS)
 
 public:
-	class FGenMipsSRGB : SHADER_PERMUTATION_BOOL("GENMIPS_SRGB");
-	using FPermutationDomain = TShaderPermutationDomain<FGenMipsSRGB>;
-
 	SHADER_USE_PARAMETER_STRUCT(FGenerateMipsCS, FGlobalShader)
 
 	BEGIN_SHADER_PARAMETER_STRUCT(FParameters, )
@@ -111,9 +108,7 @@ void FGenerateMips::Compute(FRHICommandListImmediate& RHIImmCmdList, FRHITexture
 	FRDGTextureRef GraphTexture = GraphBuilder.RegisterExternalTexture(GenMipsStruct->RenderTarget, TEXT("GenerateMipsGraphTexture"));
 
 	//Select compute shader variant (normal vs. sRGB)
-	FGenerateMipsCS::FPermutationDomain PermutationVector;
-	PermutationVector.Set<FGenerateMipsCS::FGenMipsSRGB>(!!(InTexture->GetFlags() & TexCreate_SRGB));
-	TShaderMapRef<FGenerateMipsCS> ComputeShader(GetGlobalShaderMap(ERHIFeatureLevel::SM5), PermutationVector);
+	TShaderMapRef<FGenerateMipsCS> ComputeShader(GetGlobalShaderMap(ERHIFeatureLevel::SM5));
 
 	//Loop through each level of the mips that require creation and add a dispatch pass per level,.
 	for (uint8 MipLevel = 1; MipLevel < InTexture->GetNumMips(); MipLevel++)
