UE4 - Proper fix for shader permutation domain
(undo workaround 10011374)
#rnx

diff --git a/Engine/Source/Runtime/RenderCore/Private/GenerateMips.cpp b/Engine/Source/Runtime/RenderCore/Private/GenerateMips.cpp
--- a/Engine/Source/Runtime/RenderCore/Private/GenerateMips.cpp
+++ b/Engine/Source/Runtime/RenderCore/Private/GenerateMips.cpp
@@ -21,6 +21,9 @@ class FGenerateMipsCS : public FGlobalShader
 	DECLARE_GLOBAL_SHADER(FGenerateMipsCS)
 
 public:
+	class FGenMipsSRGB : SHADER_PERMUTATION_BOOL("GENMIPS_SRGB");
+	using FPermutationDomain = TShaderPermutationDomain<FGenMipsSRGB>;
+
 	SHADER_USE_PARAMETER_STRUCT(FGenerateMipsCS, FGlobalShader)
 
 	BEGIN_SHADER_PARAMETER_STRUCT(FParameters, )
@@ -108,7 +111,9 @@ void FGenerateMips::Compute(FRHICommandListImmediate& RHIImmCmdList, FRHITexture
 	FRDGTextureRef GraphTexture = GraphBuilder.RegisterExternalTexture(GenMipsStruct->RenderTarget, TEXT("GenerateMipsGraphTexture"));
 
 	//Select compute shader variant (normal vs. sRGB)
-	TShaderMapRef<FGenerateMipsCS> ComputeShader(GetGlobalShaderMap(ERHIFeatureLevel::SM5));
+	FGenerateMipsCS::FPermutationDomain PermutationVector;
+	PermutationVector.Set<FGenerateMipsCS::FGenMipsSRGB>(!!(InTexture->GetFlags() & TexCreate_SRGB));
+	TShaderMapRef<FGenerateMipsCS> ComputeShader(GetGlobalShaderMap(ERHIFeatureLevel::SM5), PermutationVector);
 
 	//Loop through each level of the mips that require creation and add a dispatch pass per level,.
 	for (uint8 MipLevel = 1; MipLevel < InTexture->GetNumMips(); MipLevel++)
@@ -180,8 +185,10 @@ void FGenerateMips::Execute(FRDGBuilder* GraphBuilder, FRDGTextureRef InGraphTex
 	check(InGraphTexture);
 	check(InSampler);
 
-	TShaderMapRef<FGenerateMipsCS> ComputeShader(GetGlobalShaderMap(ERHIFeatureLevel::SM5));
-	
+	FGenerateMipsCS::FPermutationDomain PermutationVector;
+	PermutationVector.Set<FGenerateMipsCS::FGenMipsSRGB>(!!(InGraphTexture->GetRHI()->GetFlags() & TexCreate_SRGB));
+	TShaderMapRef<FGenerateMipsCS> ComputeShader(GetGlobalShaderMap(ERHIFeatureLevel::SM5), PermutationVector);
+
 	//Loop through each level of the mips that require creation and add a dispatch pass per level,.
 	for (uint8 MipLevel = 1; MipLevel < InGraphTexture->Desc.NumMips; MipLevel++)
 	{
