Fixed boolean version of SetShaderValue for FRHICommandListImmediate.

diff --git a/Engine/Source/Runtime/RenderCore/Public/ShaderParameterUtils.h b/Engine/Source/Runtime/RenderCore/Public/ShaderParameterUtils.h
--- a/Engine/Source/Runtime/RenderCore/Public/ShaderParameterUtils.h
+++ b/Engine/Source/Runtime/RenderCore/Public/ShaderParameterUtils.h
@@ -109,6 +109,20 @@ void SetShaderValue(
 	SetShaderValue(RHICmdList, Shader, Parameter, BoolValue, ElementIndex);
 }
 
+/** Specialization of the above for C++ bool type. */
+template<typename ShaderRHIParamRef>
+void SetShaderValue(
+	FRHICommandListImmediate& RHICmdList,
+	const ShaderRHIParamRef& Shader,
+	const FShaderParameter& Parameter,
+	bool Value,
+	uint32 ElementIndex = 0
+	)
+{
+	const uint32 BoolValue = Value;
+	SetShaderValue(RHICmdList, Shader, Parameter, BoolValue, ElementIndex);
+}
+
 /**
  * Sets the value of a shader parameter array.  Template'd on shader type
  * A template parameter specified the type of the parameter value.
