Making OpenGL RHI report more correctly if compute mip genmeration is needed & allowed compute mip generation code to be used with ES3.1
#rnx

diff --git a/Engine/Source/Runtime/OpenGLDrv/Private/OpenGLDrv.cpp b/Engine/Source/Runtime/OpenGLDrv/Private/OpenGLDrv.cpp
--- a/Engine/Source/Runtime/OpenGLDrv/Private/OpenGLDrv.cpp
+++ b/Engine/Source/Runtime/OpenGLDrv/Private/OpenGLDrv.cpp
@@ -57,7 +57,10 @@ void FOpenGLGPUProfiler::PopEvent()
 
 }
 
-
+bool FOpenGLDynamicRHI::RHIRequiresComputeGenerateMips() const
+{
+	return !FOpenGL::SupportsGenerateMipmap();
+};
 
 void FOpenGLGPUProfiler::BeginFrame(FOpenGLDynamicRHI* InRHI)
 {
diff --git a/Engine/Source/Runtime/OpenGLDrv/Public/OpenGLDrv.h b/Engine/Source/Runtime/OpenGLDrv/Public/OpenGLDrv.h
--- a/Engine/Source/Runtime/OpenGLDrv/Public/OpenGLDrv.h
+++ b/Engine/Source/Runtime/OpenGLDrv/Public/OpenGLDrv.h
@@ -541,6 +541,8 @@ public:
 	virtual FStagingBufferRHIRef RHICreateStagingBuffer() final override;
 	virtual FGPUFenceRHIRef RHICreateGPUFence(const FName &Name) final override;
 
+	virtual bool RHIRequiresComputeGenerateMips() const override;
+
 	void Cleanup();
 
 	void PurgeFramebufferFromCaches(GLuint Framebuffer);
diff --git a/Engine/Source/Runtime/RenderCore/Private/GenerateMips.cpp b/Engine/Source/Runtime/RenderCore/Private/GenerateMips.cpp
--- a/Engine/Source/Runtime/RenderCore/Private/GenerateMips.cpp
+++ b/Engine/Source/Runtime/RenderCore/Private/GenerateMips.cpp
@@ -35,7 +35,7 @@ public:
 
 	static bool ShouldCompilePermutation(const FGlobalShaderPermutationParameters& Parameters)
 	{
-		return IsFeatureLevelSupported(Parameters.Platform, ERHIFeatureLevel::SM5);
+		return RHISupportsComputeShaders(Parameters.Platform);
 	}
 };
 IMPLEMENT_GLOBAL_SHADER(FGenerateMipsCS, "/Engine/Private/ComputeGenerateMips.usf", "MainCS", SF_Compute);
@@ -113,7 +113,7 @@ void FGenerateMips::Compute(FRHICommandListImmediate& RHIImmCmdList, FRHITexture
 	//Select compute shader variant (normal vs. sRGB)
 	FGenerateMipsCS::FPermutationDomain PermutationVector;
 	PermutationVector.Set<FGenerateMipsCS::FGenMipsSRGB>(!!(InTexture->GetFlags() & TexCreate_SRGB));
-	TShaderMapRef<FGenerateMipsCS> ComputeShader(GetGlobalShaderMap(ERHIFeatureLevel::SM5), PermutationVector);
+	TShaderMapRef<FGenerateMipsCS> ComputeShader(GetGlobalShaderMap(GMaxRHIFeatureLevel), PermutationVector);
 
 	//Loop through each level of the mips that require creation and add a dispatch pass per level,.
 	for (uint8 MipLevel = 1; MipLevel < InTexture->GetNumMips(); MipLevel++)
