Fixed bad access of RDG resource outside of pass.

diff --git a/Engine/Source/Runtime/RenderCore/Private/GenerateMips.cpp b/Engine/Source/Runtime/RenderCore/Private/GenerateMips.cpp
--- a/Engine/Source/Runtime/RenderCore/Private/GenerateMips.cpp
+++ b/Engine/Source/Runtime/RenderCore/Private/GenerateMips.cpp
@@ -186,7 +186,7 @@ void FGenerateMips::Execute(FRDGBuilder* GraphBuilder, FRDGTextureRef InGraphTex
 	check(InSampler);
 
 	FGenerateMipsCS::FPermutationDomain PermutationVector;
-	PermutationVector.Set<FGenerateMipsCS::FGenMipsSRGB>(!!(InGraphTexture->GetRHI()->GetFlags() & TexCreate_SRGB));
+	PermutationVector.Set<FGenerateMipsCS::FGenMipsSRGB>(!!(InGraphTexture->Desc.Flags & TexCreate_SRGB));
 	TShaderMapRef<FGenerateMipsCS> ComputeShader(GetGlobalShaderMap(ERHIFeatureLevel::SM5), PermutationVector);
 
 	//Loop through each level of the mips that require creation and add a dispatch pass per level,.
