Allows to read any RDG resources in a GenerateMips pass

diff --git a/Engine/Source/Runtime/RenderCore/Private/RenderGraphBarrierBatcher.cpp b/Engine/Source/Runtime/RenderCore/Private/RenderGraphBarrierBatcher.cpp
--- a/Engine/Source/Runtime/RenderCore/Private/RenderGraphBarrierBatcher.cpp
+++ b/Engine/Source/Runtime/RenderCore/Private/RenderGraphBarrierBatcher.cpp
@@ -113,7 +113,6 @@ FRDGBarrierBatcher::FRDGBarrierBatcher(FRHICommandList& InRHICmdList, const FRDG
 {
 	if (Pass)
 	{
-		bIsGeneratingMips = Pass->IsGenerateMips();
 		Pipeline = Pass->IsCompute() ? FRDGResourceState::EPipeline::Compute : FRDGResourceState::EPipeline::Graphics;
 	}
 }
@@ -177,13 +176,6 @@ void FRDGBarrierBatcher::QueueTransitionTexture(FRDGTexture* Texture, FRDGResour
 {
 	check(Texture);
 
-	// Texture transitions are ignored when generating mips, since the render target binding call or UAV will
-	// perform the subresource transition.
-	if (bIsGeneratingMips)
-	{
-		return;
-	}
-
 	const FRDGResourceState StateBefore = Texture->State;
 	const FRDGResourceState StateAfter(Pass, Pipeline, AccessAfter);
 
@@ -248,6 +240,7 @@ void FRDGBarrierBatcher::QueueTransitionUAV(
 	FRHIUnorderedAccessView* UAV,
 	FRDGParentResource* ParentResource,
 	FRDGResourceState::EAccess AccessAfter,
+	bool bIsGeneratingMips,
 	FRDGResourceState::EPipeline PipelineAfter)
 {
 	check(UAV);
@@ -270,7 +263,7 @@ void FRDGBarrierBatcher::QueueTransitionUAV(
 		// Add transition to the correct batch bucket.
 		{
 			FTransitionParameters TransitionParameters;
-			TransitionParameters.TransitionAccess = GetResourceTransitionAccessForUAV(StateBefore.Access, StateAfter.Access);
+			TransitionParameters.TransitionAccess = GetResourceTransitionAccessForUAV(StateBefore.Access, StateAfter.Access, bIsGeneratingMips);
 			TransitionParameters.TransitionPipeline = GetResourceTransitionPipeline(StateBefore.Pipeline, StateAfter.Pipeline);
 
 			FUAVBatch& UAVBatch = UAVBatchMap.FindOrAdd(TransitionParameters);
@@ -350,7 +343,7 @@ EResourceTransitionAccess FRDGBarrierBatcher::GetResourceTransitionAccess(FRDGRe
 	return AccessAfter == FRDGResourceState::EAccess::Write ? EResourceTransitionAccess::EWritable : EResourceTransitionAccess::EReadable;
 }
 
-EResourceTransitionAccess FRDGBarrierBatcher::GetResourceTransitionAccessForUAV(FRDGResourceState::EAccess AccessBefore, FRDGResourceState::EAccess AccessAfter) const
+EResourceTransitionAccess FRDGBarrierBatcher::GetResourceTransitionAccessForUAV(FRDGResourceState::EAccess AccessBefore, FRDGResourceState::EAccess AccessAfter, bool bIsGeneratingMips) const
 {
 	switch (AccessAfter)
 	{
diff --git a/Engine/Source/Runtime/RenderCore/Private/RenderGraphBarrierBatcher.h b/Engine/Source/Runtime/RenderCore/Private/RenderGraphBarrierBatcher.h
--- a/Engine/Source/Runtime/RenderCore/Private/RenderGraphBarrierBatcher.h
+++ b/Engine/Source/Runtime/RenderCore/Private/RenderGraphBarrierBatcher.h
@@ -25,6 +25,7 @@ public:
 		FRHIUnorderedAccessView* UAV,
 		FRDGParentResource* ParentResource,
 		FRDGResourceState::EAccess AccessAfter,
+		bool bGeneratingMips = false,
 		FRDGResourceState::EPipeline PipelineAfter = FRDGResourceState::EPipeline::MAX);
 
 #if WITH_MGPU
@@ -87,9 +88,8 @@ private:
 	void LogTransition(const FRDGParentResource* Resource, FTransitionParameters Parameters) const;
 
 	EResourceTransitionAccess GetResourceTransitionAccess(FRDGResourceState::EAccess AccessAfter) const;
-	EResourceTransitionAccess GetResourceTransitionAccessForUAV(FRDGResourceState::EAccess AccessBefore, FRDGResourceState::EAccess AccessAfter) const;
+	EResourceTransitionAccess GetResourceTransitionAccessForUAV(FRDGResourceState::EAccess AccessBefore, FRDGResourceState::EAccess AccessAfter, bool bIsGeneratingMips) const;
 
 	const FRDGPass* Pass = nullptr;
 	FRDGResourceState::EPipeline Pipeline = FRDGResourceState::EPipeline::Graphics;
-	bool bIsGeneratingMips = false;
 };
diff --git a/Engine/Source/Runtime/RenderCore/Private/RenderGraphBuilder.cpp b/Engine/Source/Runtime/RenderCore/Private/RenderGraphBuilder.cpp
--- a/Engine/Source/Runtime/RenderCore/Private/RenderGraphBuilder.cpp
+++ b/Engine/Source/Runtime/RenderCore/Private/RenderGraphBuilder.cpp
@@ -733,13 +733,81 @@ void FRDGBuilder::PrepareResourcesForExecute(const FRDGPass* Pass, struct FRHIRe
 	BarrierBatcher.SetNameForTemporalEffect(NameForTemporalEffect);
 #endif
 
-	// NOTE: When generating mips, we don't perform any transitions on textures. They are done implicitly by the RHI.
-	const bool bGeneratingMips = Pass->IsGenerateMips();
-
 	FRDGPassParameterStruct ParameterStruct = Pass->GetParameters();
 
 	const uint32 ParameterCount = ParameterStruct.GetParameterCount();
 
+	// List all RDG texture being read and modified by the pass.
+	TSet<FRDGTextureRef, DefaultKeyFuncs<FRDGTextureRef>, SceneRenderingSetAllocator> ReadTextures;
+	TSet<FRDGTextureRef, DefaultKeyFuncs<FRDGTextureRef>, SceneRenderingSetAllocator> ModifiedTextures;
+	for (uint32 ParameterIndex = 0; ParameterIndex < ParameterCount; ++ParameterIndex)
+	{
+		FRDGPassParameter Parameter = ParameterStruct.GetParameter(ParameterIndex);
+
+		switch (Parameter.GetType())
+		{
+		case UBMT_RDG_TEXTURE:
+		{
+			if (FRDGTextureRef Texture = Parameter.GetAsTexture())
+			{
+				ReadTextures.Add(Texture);
+			}
+		}
+		break;
+		case UBMT_RDG_TEXTURE_SRV:
+		{
+			if (FRDGTextureSRVRef SRV = Parameter.GetAsTextureSRV())
+			{
+				ReadTextures.Add(SRV->GetParent());
+			}
+		}
+		break;
+		case UBMT_RDG_TEXTURE_UAV:
+		{
+			if (FRDGTextureUAVRef UAV = Parameter.GetAsTextureUAV())
+			{
+				ModifiedTextures.Add(UAV->GetParent());
+			}
+		}
+		break;
+		case UBMT_RDG_TEXTURE_COPY_DEST:
+		{
+			if (FRDGTextureRef Texture = Parameter.GetAsTexture())
+			{
+				ModifiedTextures.Add(Texture);
+			}
+		}
+		break;
+		case UBMT_RENDER_TARGET_BINDING_SLOTS:
+		{
+			const FRenderTargetBindingSlots& RenderTargetBindingSlots = Parameter.GetAsRenderTargetBindingSlots();
+			const auto& RenderTargets = RenderTargetBindingSlots.Output;
+			const uint32 RenderTargetCount = RenderTargets.Num();
+
+			for (uint32 RenderTargetIndex = 0; RenderTargetIndex < RenderTargetCount; RenderTargetIndex++)
+			{
+				if (FRDGTextureRef Texture = RenderTargets[RenderTargetIndex].GetTexture())
+				{
+					ModifiedTextures.Add(Texture);
+				}
+				else
+				{
+					break;
+				}
+			}
+
+			if (FRDGTextureRef Texture = RenderTargetBindingSlots.DepthStencil.GetTexture())
+			{
+				if (RenderTargetBindingSlots.DepthStencil.GetDepthStencilAccess().IsAnyWrite())
+					ModifiedTextures.Add(Texture);
+			}
+		}
+		break;
+		default:
+			break;
+		}
+	}
+
 	for (uint32 ParameterIndex = 0; ParameterIndex < ParameterCount; ++ParameterIndex)
 	{
 		FRDGPassParameter Parameter = ParameterStruct.GetParameter(ParameterIndex);
@@ -753,6 +821,7 @@ void FRDGBuilder::PrepareResourcesForExecute(const FRDGPass* Pass, struct FRHIRe
 				check(Texture->PooledRenderTarget);
 				check(Texture->ResourceRHI);
 
+				check(!ModifiedTextures.Contains(Texture));
 				BarrierBatcher.QueueTransitionTexture(Texture, FRDGResourceState::EAccess::Read);
 			}
 		}
@@ -765,7 +834,15 @@ void FRDGBuilder::PrepareResourcesForExecute(const FRDGPass* Pass, struct FRHIRe
 
 				AllocateRHITextureSRVIfNeeded(SRV);
 
-				BarrierBatcher.QueueTransitionTexture(Texture, FRDGResourceState::EAccess::Read);
+				if (ModifiedTextures.Contains(Texture))
+				{
+					// If it is bound to a RT, no need for ERWSubResBarrier.
+					// If it is bound as a UAV, the UAV is going to issue the ERWSubResBarrier.
+				}
+				else
+				{
+					BarrierBatcher.QueueTransitionTexture(Texture, FRDGResourceState::EAccess::Read);
+				}
 			}
 		}
 		break;
@@ -774,7 +851,7 @@ void FRDGBuilder::PrepareResourcesForExecute(const FRDGPass* Pass, struct FRHIRe
 			if (FRDGTextureUAVRef UAV = Parameter.GetAsTextureUAV())
 			{
 				FRDGTextureRef Texture = UAV->Desc.Texture;
-	
+
 				AllocateRHITextureUAVIfNeeded(UAV);
 
 				FRHIUnorderedAccessView* UAVRHI = UAV->GetRHI();
@@ -784,7 +861,9 @@ void FRDGBuilder::PrepareResourcesForExecute(const FRDGPass* Pass, struct FRHIRe
 					OutRPInfo->UAVs[OutRPInfo->NumUAVs++] = UAVRHI;	// Bind UAVs in declaration order
 				}
 
-				BarrierBatcher.QueueTransitionUAV(UAVRHI, Texture, FRDGResourceState::EAccess::Write);
+				bool bGeneratingMips = ReadTextures.Contains(Texture);
+
+				BarrierBatcher.QueueTransitionUAV(UAVRHI, Texture, FRDGResourceState::EAccess::Write, bGeneratingMips);
 			}
 		}
 		break;
@@ -794,6 +873,7 @@ void FRDGBuilder::PrepareResourcesForExecute(const FRDGPass* Pass, struct FRHIRe
 			{
 				AllocateRHITextureIfNeeded(Texture);
 
+				check(!ModifiedTextures.Contains(Texture));
 				BarrierBatcher.QueueTransitionTexture(Texture, FRDGResourceState::EAccess::Write);
 			}
 		}
@@ -959,7 +1039,7 @@ void FRDGBuilder::PrepareResourcesForExecute(const FRDGPass* Pass, struct FRHIRe
 		}
 	}
 
-	OutRPInfo->bGeneratingMips = bGeneratingMips;
+	OutRPInfo->bGeneratingMips = Pass->IsGenerateMips();
 }
 
 void FRDGBuilder::ReleaseRHITextureIfUnreferenced(FRDGTexture* Texture)
@@ -1113,7 +1193,7 @@ void FRDGBuilder::ProcessDeferredInternalResourceQueries()
 
 		for (TMap<FRDGBufferUAVDesc, FUnorderedAccessViewRHIRef, FDefaultSetAllocator, TMapRDGBufferUAVFuncs<FRDGBufferUAVDesc, FUnorderedAccessViewRHIRef>>::TIterator It(Query.Buffer->PooledBuffer->UAVs); It; ++It)
 		{
-			BarrierBatcher.QueueTransitionUAV(It.Value(), Query.Buffer, Query.DestinationAccess, Query.DestinationPipeline);
+			BarrierBatcher.QueueTransitionUAV(It.Value(), Query.Buffer, Query.DestinationAccess, /* bGeneratingMips = */ false, Query.DestinationPipeline);
 		}
 
 		*Query.OutBufferPtr = AllocatedBuffers.FindChecked(Query.Buffer);
diff --git a/Engine/Source/Runtime/RenderCore/Private/RenderGraphValidation.cpp b/Engine/Source/Runtime/RenderCore/Private/RenderGraphValidation.cpp
--- a/Engine/Source/Runtime/RenderCore/Private/RenderGraphValidation.cpp
+++ b/Engine/Source/Runtime/RenderCore/Private/RenderGraphValidation.cpp
@@ -43,11 +43,6 @@ public:
 			// A generate mips pass requires both read / write access to the same texture for barriers to work correctly.
 			if (bIsGenerateMips)
 			{
-				const bool bReadWithoutWrite = TextureAccess.IsAnyRead() && !TextureAccess.IsAnyWrite();
-				ensureMsgf(!bReadWithoutWrite,
-					TEXT("Attempting to use texture '%s' for read but not write access in pass '%s' which is marked with 'GenerateMips'. This kind of pass only supports textures for simultaneous read and write."),
-					Texture->Name, Pass.GetName());
-
 				const bool bWriteWithoutRead = TextureAccess.IsAnyWrite() && !TextureAccess.IsAnyRead();
 				ensureMsgf(!bWriteWithoutRead,
 					TEXT("Attempting to use texture '%s' for write but not read access in pass '%s' which is marked with 'GenerateMips'. This kind of pass only supports textures for simultaneous read and write."),
