Fixes crashes in QAEntry map

diff --git a/Engine/Source/Runtime/RenderCore/Private/RenderGraphBuilder.cpp b/Engine/Source/Runtime/RenderCore/Private/RenderGraphBuilder.cpp
--- a/Engine/Source/Runtime/RenderCore/Private/RenderGraphBuilder.cpp
+++ b/Engine/Source/Runtime/RenderCore/Private/RenderGraphBuilder.cpp
@@ -873,7 +873,7 @@ void FRDGBuilder::PrepareResourcesForExecute(const FRDGPass* Pass, struct FRHIRe
 			{
 				AllocateRHITextureIfNeeded(Texture);
 
-				check(!ModifiedTextures.Contains(Texture));
+				check(!ReadTextures.Contains(Texture));
 				BarrierBatcher.QueueTransitionTexture(Texture, FRDGResourceState::EAccess::Write);
 			}
 		}
