Cherry-pick 10232550: Validates shader parameter structure layout before submitting to RHI

diff --git a/Engine/Shaders/Public/ShaderVersion.ush b/Engine/Shaders/Public/ShaderVersion.ush
--- a/Engine/Shaders/Public/ShaderVersion.ush
+++ b/Engine/Shaders/Public/ShaderVersion.ush
@@ -2,5 +2,4 @@
 // Each time the console command is executed it generates a new GUID. As this file is included
 // in Platform.ush (which should be included in any shader) it allows to invalidate the shader DDC.
 // 
-// If you are merging streams and there is a conflict with this GUID you should make a new GUID rather than taking one or the other.
-// GUID = 054EAB614D124BB38BF05ED86E984C59
\ No newline at end of file
+// GUID = 3FC0C8634C638C1980F7FFBE48ABCB58
\ No newline at end of file
diff --git a/Engine/Source/Runtime/Engine/Private/ShaderCompiler/ShaderCompiler.cpp b/Engine/Source/Runtime/Engine/Private/ShaderCompiler/ShaderCompiler.cpp
--- a/Engine/Source/Runtime/Engine/Private/ShaderCompiler/ShaderCompiler.cpp
+++ b/Engine/Source/Runtime/Engine/Private/ShaderCompiler/ShaderCompiler.cpp
@@ -4110,7 +4110,23 @@ void VerifyGlobalShaders(EShaderPlatform Platform, bool bLoadedFromCacheFile)
 		int32 PermutationCountToCompile = 0;
 		for (int32 PermutationId = 0; PermutationId < GlobalShaderType->GetPermutationCount(); PermutationId++)
 		{
-			if (GlobalShaderType->ShouldCompilePermutation(Platform, PermutationId) && !GlobalShaderMap->HasShader(GlobalShaderType, PermutationId))
+			if (!GlobalShaderType->ShouldCompilePermutation(Platform, PermutationId))
+			{
+				continue;
+			}
+
+			if (FShader* Shader = GlobalShaderMap->GetShader(GlobalShaderType, PermutationId))
+			{
+				// Validate the shader parameter structure early.
+				if (const FShaderParametersMetadata* ParameterStructMetadata = GlobalShaderType->GetRootParametersMetadata())
+				{
+					checkf(
+						Shader->Bindings.StructureLayoutHash == ParameterStructMetadata->GetLayoutHash(),
+						TEXT("Seams shader %s's parameter structure has changed without recompilation of the shader"),
+						GlobalShaderType->GetName());
+				}
+			}
+			else
 			{
 				if (bErrorOnMissing)
 				{
diff --git a/Engine/Source/Runtime/RenderCore/Private/ShaderParameterMetadata.cpp b/Engine/Source/Runtime/RenderCore/Private/ShaderParameterMetadata.cpp
--- a/Engine/Source/Runtime/RenderCore/Private/ShaderParameterMetadata.cpp
+++ b/Engine/Source/Runtime/RenderCore/Private/ShaderParameterMetadata.cpp
@@ -135,7 +135,10 @@ void FShaderParametersMetadata::InitializeAllGlobalStructs()
 {
 	for (TLinkedList<FShaderParametersMetadata*>::TIterator StructIt(FShaderParametersMetadata::GetStructList()); StructIt; StructIt.Next())
 	{
-		StructIt->InitializeLayout();
+		if (!StructIt->bLayoutInitialized)
+		{
+			StructIt->InitializeLayout();
+		}
 	}
 }
 
@@ -146,7 +149,6 @@ void FShaderParametersMetadata::InitializeLayout()
 
 	TArray<FUniformBufferMemberAndOffset> MemberStack;
 	MemberStack.Reserve(Members.Num());
-
 	for (int32 MemberIndex = 0; MemberIndex < Members.Num(); MemberIndex++)
 	{
 		MemberStack.Push(FUniformBufferMemberAndOffset(*this, Members[MemberIndex], 0));
@@ -281,8 +283,48 @@ void FShaderParametersMetadata::InitializeLayout()
 #endif
 	});
 
+	// Compute the hash of the RHI layout.
 	Layout.ComputeHash();
+	
+	// Compute the hash about the entire layout of the structure.
+	{
+		uint32 RootStructureHash = 0;
+		RootStructureHash = HashCombine(RootStructureHash, GetTypeHash(int32(GetSize())));
+
+		for (const FMember& CurrentMember : Members)
+		{
+			EUniformBufferBaseType BaseType = CurrentMember.GetBaseType();
+			const FShaderParametersMetadata* ChildStruct = CurrentMember.GetStructMetadata();
+
+			uint32 MemberHash = 0;
+			MemberHash = HashCombine(MemberHash, GetTypeHash(int32(CurrentMember.GetOffset())));
+			MemberHash = HashCombine(MemberHash, GetTypeHash(uint8(BaseType)));
+			static_assert(EUniformBufferBaseType_NumBits <= 8, "Invalid EUniformBufferBaseType_NumBits");
+			MemberHash = HashCombine(MemberHash, GetTypeHash(CurrentMember.GetName()));
+			MemberHash = HashCombine(MemberHash, GetTypeHash(int32(CurrentMember.GetNumElements())));
+
+			if (BaseType == UBMT_INT32 ||
+				BaseType == UBMT_UINT32 ||
+				BaseType == UBMT_FLOAT32)
+			{
+				MemberHash = HashCombine(MemberHash, GetTypeHash(uint8(CurrentMember.GetNumRows())));
+				MemberHash = HashCombine(MemberHash, GetTypeHash(uint8(CurrentMember.GetNumColumns())));
+			}
+			else if (BaseType == UBMT_INCLUDED_STRUCT || BaseType == UBMT_NESTED_STRUCT)
+			{
+				if (!ChildStruct->bLayoutInitialized)
+				{
+					const_cast<FShaderParametersMetadata*>(ChildStruct)->InitializeLayout();
+				}
+
+				MemberHash = HashCombine(MemberHash, ChildStruct->GetLayoutHash());
+			}
+
+			RootStructureHash = HashCombine(RootStructureHash, MemberHash);
+		}
 
+		LayoutHash = RootStructureHash;
+	}
 	bLayoutInitialized = true;
 }
 
diff --git a/Engine/Source/Runtime/RenderCore/Private/ShaderParameterStruct.cpp b/Engine/Source/Runtime/RenderCore/Private/ShaderParameterStruct.cpp
--- a/Engine/Source/Runtime/RenderCore/Private/ShaderParameterStruct.cpp
+++ b/Engine/Source/Runtime/RenderCore/Private/ShaderParameterStruct.cpp
@@ -262,6 +262,7 @@ void FShaderParameterBindings::BindForLegacyShaderParameters(const FShader* Shad
 		/* MemberPrefix = */ TEXT(""),
 		/* ByteOffset = */ 0);
 
+	StructureLayoutHash = StructMetaData.GetLayoutHash();
 	RootParameterBufferIndex = kInvalidBufferIndex;
 
 	TArray<FString> AllParameterNames;
@@ -313,6 +314,8 @@ void FShaderParameterBindings::BindForRootShaderParameters(const FShader* Shader
 		/* MemberPrefix = */ TEXT(""),
 		/* ByteOffset = */ 0);
 
+	StructureLayoutHash = StructMetaData.GetLayoutHash();
+
 	// Binds the uniform buffer that contains the root shader parameters.
 	{
 		const TCHAR* ShaderBindingName = FShaderParametersMetadata::kRootUniformBufferBindingName;
@@ -428,6 +431,12 @@ void EmitNullShaderParameterFatalError(const FShader* Shader, const FShaderParam
 void ValidateShaderParameters(const FShader* Shader, const FShaderParametersMetadata* ParametersMetadata, const void* Parameters)
 {
 	const FShaderParameterBindings& Bindings = Shader->Bindings;
+
+	checkf(
+		Bindings.StructureLayoutHash == ParametersMetadata->GetLayoutHash(),
+		TEXT("Seams shader %s's parameter structure has changed without recompilation of the shader"),
+		Shader->GetType()->GetName());
+
 	const uint8* Base = reinterpret_cast<const uint8*>(Parameters);
 
 	const TCHAR* ShaderClassName = Shader->GetType()->GetName();
diff --git a/Engine/Source/Runtime/RenderCore/Public/Shader.h b/Engine/Source/Runtime/RenderCore/Public/Shader.h
--- a/Engine/Source/Runtime/RenderCore/Public/Shader.h
+++ b/Engine/Source/Runtime/RenderCore/Public/Shader.h
@@ -808,6 +808,9 @@ public:
 	TArray<FResourceParameter> GraphUAVs;
 	TArray<FParameterStructReference> ParameterReferences;
 
+	// Hash of the shader parameter structure when doing the binding.
+	uint32 StructureLayoutHash = 0;
+
 	// Buffer index of FShaderParametersMetadata::kRootUniformBufferBindingName
 	uint16 RootParameterBufferIndex = FShaderParameterBindings::kInvalidBufferIndex;
 
@@ -822,6 +825,7 @@ public:
 		Ar << ParametersBindingData.GraphSRVs;
 		Ar << ParametersBindingData.GraphUAVs;
 		Ar << ParametersBindingData.ParameterReferences;
+		Ar << ParametersBindingData.StructureLayoutHash;
 		Ar << ParametersBindingData.RootParameterBufferIndex;
 		return Ar;
 	}
diff --git a/Engine/Source/Runtime/RenderCore/Public/ShaderParameterMetadata.h b/Engine/Source/Runtime/RenderCore/Public/ShaderParameterMetadata.h
--- a/Engine/Source/Runtime/RenderCore/Public/ShaderParameterMetadata.h
+++ b/Engine/Source/Runtime/RenderCore/Public/ShaderParameterMetadata.h
@@ -178,6 +178,14 @@ public:
 	/** Initialize all the global shader parameter structs. */
 	static void InitializeAllGlobalStructs();
 
+	/** Returns a hash about the entire layout of the structure. */
+	uint32 GetLayoutHash() const
+	{
+		check(UseCase == EUseCase::ShaderParameterStruct || UseCase == EUseCase::GlobalShaderParameterStruct);
+		check(bLayoutInitialized);
+		return LayoutHash;	
+	}
+
 private:
 	/** Name of the structure type in C++ and shader code. */
 	const TCHAR* const StructTypeName;
@@ -203,6 +211,9 @@ private:
 	/** Whether the layout is actually initialized yet or not. */
 	uint32 bLayoutInitialized : 1;
 
+	/** Hash about the entire memory layout of the structure. */
+	uint32 LayoutHash = 0;
+
 
 	void InitializeLayout();
 
