Cherry-pick 10232551: Invalidates global shader map DDC when shader parameter structure has changed

diff --git a/Engine/Source/Runtime/RenderCore/Private/GlobalShader.cpp b/Engine/Source/Runtime/RenderCore/Private/GlobalShader.cpp
--- a/Engine/Source/Runtime/RenderCore/Private/GlobalShader.cpp
+++ b/Engine/Source/Runtime/RenderCore/Private/GlobalShader.cpp
@@ -120,6 +120,11 @@ void FGlobalShaderMapId::AppendKeyString(FString& KeyString, const TArray<FShade
 		// Add the type's source hash so that we can invalidate cached shaders when .usf changes are made
 		KeyString += ShaderTypeDependency.SourceHash.ToString();
 
+		if (const FShaderParametersMetadata* ParameterStructMetadata = ShaderTypeDependency.ShaderType->GetRootParametersMetadata())
+		{
+			KeyString += FString::Printf(TEXT("%08x"), ParameterStructMetadata->GetLayoutHash());
+		}
+
 		// Add the serialization history to the key string so that we can detect changes to global shader serialization without a corresponding .usf change
 		ShaderTypeDependency.ShaderType->GetSerializationHistory().AppendKeyString(KeyString);
 
@@ -144,6 +149,11 @@ void FGlobalShaderMapId::AppendKeyString(FString& KeyString, const TArray<FShade
 
 		for (const FShaderType* ShaderType : Dependency.ShaderPipelineType->GetStages())
 		{
+			if (const FShaderParametersMetadata* ParameterStructMetadata = ShaderType->GetRootParametersMetadata())
+			{
+				KeyString += FString::Printf(TEXT("%08x"), ParameterStructMetadata->GetLayoutHash());
+			}
+
 			const TMap<const TCHAR*, FCachedUniformBufferDeclaration>& ReferencedUniformBufferStructsCache = ShaderType->GetReferencedUniformBufferStructsCache();
 
 			// Gather referenced uniform buffers
