FGenerateMips: adding code to allow for manual caching of setups to avoid circular refs & leaks; adding code to allow for render-based fallback for mip-generation on Android ES (bringing cleaned up, deadended code from 11.2x to 11.30)
#rnx

diff --git a/Engine/Shaders/Private/ComputeGenerateMips.usf b/Engine/Shaders/Private/ComputeGenerateMips.usf
--- a/Engine/Shaders/Private/ComputeGenerateMips.usf
+++ b/Engine/Shaders/Private/ComputeGenerateMips.usf
@@ -24,11 +24,29 @@ void MainCS(uint3 DT_ID : SV_DispatchThreadID)
 
 #if GENMIPS_SRGB
     half4 outColor = MipInSRV.SampleLevel(MipSampler, UV, 0);
-
-    MipOutUAV[DT_ID.xy] = half4(LinearToSrgb(outColor.xyz), outColor.w);
+    outColor = half4(LinearToSrgb(outColor.xyz), outColor.w);
 #else
     float4 outColor = MipInSRV.SampleLevel(MipSampler, UV, 0);
+#endif
 
+#if GENMIPS_SWIZZLE
+    MipOutUAV[DT_ID.xy] = outColor.zyxw;
+#else
     MipOutUAV[DT_ID.xy] = outColor;
 #endif
-}
\ No newline at end of file
+}
+
+/*==================================================================================================
+	Fallback Shaders to generate mips using classic rendering
+===================================================================================================*/
+
+void MainVS(in float4 InPosition : ATTRIBUTE0, in float2 InUV : ATTRIBUTE1,	out float4 OutPosition : SV_POSITION, out float2 OutUV : TEXCOORD0)
+{
+	OutPosition = InPosition;
+	OutUV = InUV;
+}
+
+void MainPS(float4 InPosition : SV_POSITION,	float2 InUV : TEXCOORD0, out float4 OutColor : SV_Target0)
+{
+    OutColor = MipsShadersUB.Texture.SampleLevel(MipsShadersUB.SamplerP, InUV + MipsShadersUB.HalfTexelSize, MipsShadersUB.Level);
+}
diff --git a/Engine/Source/Runtime/RenderCore/Private/GenerateMips.cpp b/Engine/Source/Runtime/RenderCore/Private/GenerateMips.cpp
--- a/Engine/Source/Runtime/RenderCore/Private/GenerateMips.cpp
+++ b/Engine/Source/Runtime/RenderCore/Private/GenerateMips.cpp
@@ -5,16 +5,28 @@
 #include "RenderGraphUtils.h"
 #include "RenderTargetPool.h"
 #include "ShaderPermutation.h"
+#include "RHIStaticStates.h"
+#include "PipelineStateCache.h"
+#include "ShaderParameterUtils.h"
 
 #define MIPSSHADER_NUMTHREADS 8
 
 //Struct for tracking GenerateMip resources that should only need to be created once.
 struct FGenerateMipsStruct
 {
+	// Compute
 	TRefCountPtr<IPooledRenderTarget> RenderTarget;
+
+	// Rendering
+	FVertexBufferRHIRef VertexBuffer;
+	FVertexDeclarationRHIRef VertexDeclaration;
+
+	// Both
 	FSamplerStateInitializerRHI Sampler;
 };
 
+// -------------------------------------------------------------------------------------------------------------------------------
+
 //Generate mips compute shader declaration
 class FGenerateMipsCS : public FGlobalShader
 {
@@ -22,7 +34,8 @@ class FGenerateMipsCS : public FGlobalShader
 
 public:
 	class FGenMipsSRGB : SHADER_PERMUTATION_BOOL("GENMIPS_SRGB");
-	using FPermutationDomain = TShaderPermutationDomain<FGenMipsSRGB>;
+	class FGenMipsSwizzle : SHADER_PERMUTATION_BOOL("GENMIPS_SWIZZLE");
+	using FPermutationDomain = TShaderPermutationDomain<FGenMipsSRGB, FGenMipsSwizzle>;
 
 	SHADER_USE_PARAMETER_STRUCT(FGenerateMipsCS, FGlobalShader)
 
@@ -40,15 +53,190 @@ public:
 };
 IMPLEMENT_GLOBAL_SHADER(FGenerateMipsCS, "/Engine/Private/ComputeGenerateMips.usf", "MainCS", SF_Compute);
 
+// -------------------------------------------------------------------------------------------------------------------------------
+
+class FMipsShadersVS : public FGlobalShader
+{
+	DECLARE_GLOBAL_SHADER(FMipsShadersVS);
+
+public:
+	static bool ShouldCompilePermutation(const FGlobalShaderPermutationParameters& Parameters)
+	{
+		return IsFeatureLevelSupported(Parameters.Platform, ERHIFeatureLevel::ES2);
+	}
+
+	FMipsShadersVS() { }
+
+	FMipsShadersVS(const ShaderMetaType::CompiledShaderInitializerType& Initializer)
+		: FGlobalShader(Initializer)
+	{ }
+
+	virtual bool Serialize(FArchive& Ar) override
+	{
+		bool bShaderHasOutdatedParameters = FGlobalShader::Serialize(Ar);
+		return bShaderHasOutdatedParameters;
+	}
+};
+IMPLEMENT_SHADER_TYPE(, FMipsShadersVS, TEXT("/Engine/Private/ComputeGenerateMips.usf"), TEXT("MainVS"), SF_Vertex);
+
+// -------------------------------------------------------------------------------------------------------------------------------
+
+BEGIN_GLOBAL_SHADER_PARAMETER_STRUCT(FMipsShadersUB, )
+SHADER_PARAMETER_TEXTURE(Texture2D, Texture)
+SHADER_PARAMETER_SAMPLER(SamplerState, SamplerP)
+SHADER_PARAMETER(FVector2D, HalfTexelSize)
+SHADER_PARAMETER(float, Level)
+END_GLOBAL_SHADER_PARAMETER_STRUCT()
+
+class FMipsShadersPS : public FGlobalShader
+{
+	DECLARE_GLOBAL_SHADER(FMipsShadersPS);
+
+public:
+	static bool ShouldCompilePermutation(const FGlobalShaderPermutationParameters& Parameters)
+	{
+		return IsFeatureLevelSupported(Parameters.Platform, ERHIFeatureLevel::ES2);
+	}
+
+	FMipsShadersPS() { }
+
+	FMipsShadersPS(const ShaderMetaType::CompiledShaderInitializerType& Initializer)
+		: FGlobalShader(Initializer)
+	{ }
+
+	void SetParameters(FRHICommandList& RHICmdList, FTextureRHIRef Texture, FSamplerStateRHIRef SamplerState, const FVector2D & HalfTexelSize, int32 Level)
+	{
+		FMipsShadersUB UB;
+		{
+			UB.SamplerP = SamplerState;
+			UB.Texture = Texture;
+			UB.HalfTexelSize = HalfTexelSize;
+			UB.Level = Level;
+		}
+
+		TUniformBufferRef<FMipsShadersUB> Data = TUniformBufferRef<FMipsShadersUB>::CreateUniformBufferImmediate(UB, UniformBuffer_SingleFrame);
+		SetUniformBufferParameter(RHICmdList, GetPixelShader(), GetUniformBufferParameter<FMipsShadersUB>(), Data);
+	}
+};
+
+IMPLEMENT_GLOBAL_SHADER_PARAMETER_STRUCT(FMipsShadersUB, "MipsShadersUB");
+IMPLEMENT_SHADER_TYPE(, FMipsShadersPS, TEXT("/Engine/Private/ComputeGenerateMips.usf"), TEXT("MainPS"), SF_Pixel);
+
+// -------------------------------------------------------------------------------------------------------------------------------
+
+void FGenerateMips::SetupRendering(FGenerateMipsStruct *GenMipsStruct, FRHITexture* InTexture, const FGenerateMipsParams& InParams)
+{
+	struct FMipsElementVertex
+	{
+		FVector4 Position;
+		FVector2D TextureCoordinate;
+
+		FMipsElementVertex() { }
+
+		FMipsElementVertex(const FVector4& InPosition, const FVector2D& InTextureCoordinate)
+			: Position(InPosition)
+			, TextureCoordinate(InTextureCoordinate)
+		{ }
+	};
+
+	FRHIResourceCreateInfo CreateInfo;
+	GenMipsStruct->VertexBuffer = RHICreateVertexBuffer(sizeof(FMipsElementVertex) * 4, BUF_Volatile, CreateInfo);
+	void* VoidPtr = RHILockVertexBuffer(GenMipsStruct->VertexBuffer, 0, sizeof(FMipsElementVertex) * 4, RLM_WriteOnly);
+
+	FMipsElementVertex* Vertices = (FMipsElementVertex*)VoidPtr;
+	Vertices[0].Position.Set(-1.0f, 1.0f, 1.0f, 1.0f); // Top Left
+	Vertices[1].Position.Set(1.0f, 1.0f, 1.0f, 1.0f); // Top Right
+	Vertices[2].Position.Set(-1.0f, -1.0f, 1.0f, 1.0f); // Bottom Left
+	Vertices[3].Position.Set(1.0f, -1.0f, 1.0f, 1.0f); // Bottom Right
+	Vertices[0].TextureCoordinate.Set(0.0f, 0.0f);
+	Vertices[1].TextureCoordinate.Set(1.0f, 0.0f);
+	Vertices[2].TextureCoordinate.Set(0.0f, 1.0f);
+	Vertices[3].TextureCoordinate.Set(1.0f, 1.0f);
+	RHIUnlockVertexBuffer(GenMipsStruct->VertexBuffer);
+
+	// Vertex Declaration
+	FVertexDeclarationElementList Elements;
+	const uint16 Stride = sizeof(FMipsElementVertex);
+	Elements.Add(FVertexElement(0, STRUCT_OFFSET(FMipsElementVertex, Position), VET_Float4, 0, Stride));
+	Elements.Add(FVertexElement(0, STRUCT_OFFSET(FMipsElementVertex, TextureCoordinate), VET_Float2, 1, Stride));
+	GenMipsStruct->VertexDeclaration = PipelineStateCache::GetOrCreateVertexDeclaration(Elements);
+
+	//Specify the Sampler details based on the input.
+	GenMipsStruct->Sampler.Filter = InParams.Filter;
+	GenMipsStruct->Sampler.AddressU = InParams.AddressU;
+	GenMipsStruct->Sampler.AddressV = InParams.AddressV;
+	GenMipsStruct->Sampler.AddressW = InParams.AddressW;
+}
+
+
+// Generate mips via classic rendering (2D textures only)
+void FGenerateMips::RenderMips(FRHICommandListImmediate& CommandList, FRHITexture* InTexture, const FGenerateMipsParams& InParams, TSharedPtr<FGenerateMipsStruct> * ExternalMipsStructCache)
+{
+	check(IsInRenderingThread());
+	check(InTexture->GetNumMips() > 1);
+
+	FGenerateMipsStruct *GenMipsStruct;
+	FGenerateMipsStruct _LocalGenMipsStruct;
+	if (ExternalMipsStructCache)
+	{
+		if (!*ExternalMipsStructCache)
+		{
+			*ExternalMipsStructCache = MakeShareable(new FGenerateMipsStruct());
+			SetupRendering(ExternalMipsStructCache->Get(), InTexture, InParams);
+		}
+		GenMipsStruct = ExternalMipsStructCache->Get();
+	}
+	else
+	{
+		GenMipsStruct =  &_LocalGenMipsStruct;
+		SetupRendering(GenMipsStruct, InTexture, InParams);
+	}
+
+	// Same basic general state each loop
+	FGraphicsPipelineStateInitializer GraphicsPSOInit;
+	GraphicsPSOInit.DepthStencilState = TStaticDepthStencilState<false, CF_Always>::GetRHI();
+	GraphicsPSOInit.RasterizerState = TStaticRasterizerState<>::GetRHI();
+	GraphicsPSOInit.BlendState = TStaticBlendStateWriteMask<CW_RGBA, CW_NONE, CW_NONE, CW_NONE, CW_NONE, CW_NONE, CW_NONE, CW_NONE>::GetRHI();
+	GraphicsPSOInit.PrimitiveType = PT_TriangleStrip;
+	
+	auto ShaderMap = GetGlobalShaderMap(GMaxRHIFeatureLevel);
+	TShaderMapRef<FMipsShadersVS> VertexShader(ShaderMap);
+	TShaderMapRef<FMipsShadersPS> CopyShader(ShaderMap);
+
+	GraphicsPSOInit.BoundShaderState.VertexDeclarationRHI = GenMipsStruct->VertexDeclaration;
+	GraphicsPSOInit.BoundShaderState.VertexShaderRHI = GETSAFERHISHADER_VERTEX(*VertexShader);
+	GraphicsPSOInit.BoundShaderState.PixelShaderRHI = GETSAFERHISHADER_PIXEL(*CopyShader);
+
+	uint32 NumMips = InTexture->GetNumMips();
+	for (uint32 MipLevel=1; MipLevel<NumMips; ++MipLevel)
+	{
+		int32 Width = InTexture->GetSizeXYZ().X >> MipLevel;
+		int32 Height = InTexture->GetSizeXYZ().Y >> MipLevel;
+
+		FRHIRenderPassInfo RPInfo(InTexture, ERenderTargetActions::DontLoad_Store, nullptr, MipLevel);
+		CommandList.BeginRenderPass(RPInfo, TEXT("GenMipsLevel"));
+		{
+			CommandList.ApplyCachedRenderTargets(GraphicsPSOInit);
+			CommandList.SetViewport(0, 0, 0.0f, Width, Height, 1.0f);
+
+			SetGraphicsPipelineState(CommandList, GraphicsPSOInit);
+			CopyShader->SetParameters(CommandList, InTexture, CommandList.CreateSamplerState(GenMipsStruct->Sampler), FVector2D(0.5f / Width,  0.5f / Height), MipLevel - 1);
+
+			CommandList.SetStreamSource(0, GenMipsStruct->VertexBuffer, 0);
+			CommandList.DrawPrimitive(0, 2, 1);
+		}
+		CommandList.EndRenderPass();
+	}
+}
+
+
 //Initialise the texture for usage with RenderGraph and ComputeGenerateMips shader
-FGenerateMipsStruct* FGenerateMips::SetupTexture(FRHITexture* InTexture,
-	const FGenerateMipsParams& InParams)
+TSharedPtr<FGenerateMipsStruct> FGenerateMips::SetupTexture(FRHITexture* InTexture, const FGenerateMipsParams& InParams)
 {
 	//Currently only 2D textures supported
 	check(InTexture->GetTexture2D());
-	if (!InTexture->GenMipsStruct)
-	{
-		InTexture->GenMipsStruct = MakeShareable(new FGenerateMipsStruct());
+
+	TSharedPtr<FGenerateMipsStruct> GenMipsStruct = MakeShareable(new FGenerateMipsStruct());
 
 		FPooledRenderTargetDesc Desc;
 		Desc.Extent.X = InTexture->GetSizeXYZ().X;
@@ -82,37 +270,43 @@ FGenerateMipsStruct* FGenerateMips::SetupTexture(FRHITexture* InTexture,
 		RenderTexture.UAV = RenderTexture.MipUAVs[0];
 
 		//Create the RenderTarget from the PooledRenderTarget Desc and the new RenderTexture object.
-		GRenderTargetPool.CreateUntrackedElement(Desc, InTexture->GenMipsStruct->RenderTarget, RenderTexture);
+	GRenderTargetPool.CreateUntrackedElement(Desc, GenMipsStruct->RenderTarget, RenderTexture);
 
 		//Specify the Sampler details based on the input.
-		InTexture->GenMipsStruct->Sampler.Filter = InParams.Filter;
-		InTexture->GenMipsStruct->Sampler.AddressU = InParams.AddressU;
-		InTexture->GenMipsStruct->Sampler.AddressV = InParams.AddressV;
-		InTexture->GenMipsStruct->Sampler.AddressW = InParams.AddressW;
-	}
+	GenMipsStruct->Sampler.Filter = InParams.Filter;
+	GenMipsStruct->Sampler.AddressU = InParams.AddressU;
+	GenMipsStruct->Sampler.AddressV = InParams.AddressV;
+	GenMipsStruct->Sampler.AddressW = InParams.AddressW;
 
-	//Return the raw pointer
-	return InTexture->GenMipsStruct.Get();
+	return GenMipsStruct;
 }
 
 //Compute shader execution function for generating mips in real time.
-void FGenerateMips::Compute(FRHICommandListImmediate& RHIImmCmdList, FRHITexture* InTexture)
+void FGenerateMips::Compute(FRHICommandListImmediate& RHIImmCmdList, FRHITexture* InTexture, TSharedPtr<FGenerateMipsStruct> GenMipsStruct)
 {
 	check(IsInRenderingThread());	
 	//Currently only 2D textures supported
 	check(InTexture->GetTexture2D());
 
 	//Ensure the generate mips structure has been initialised correctly.
-	FGenerateMipsStruct* GenMipsStruct = SetupTexture(InTexture);
 	check(GenMipsStruct);
 
 	//Begin rendergraph for executing the compute shader
 	FRDGBuilder GraphBuilder(RHIImmCmdList);
 	FRDGTextureRef GraphTexture = GraphBuilder.RegisterExternalTexture(GenMipsStruct->RenderTarget, TEXT("GenerateMipsGraphTexture"));
 
-	//Select compute shader variant (normal vs. sRGB)
+	//Select compute shader variant (normal vs. sRGB etc.)
+#if PLATFORM_ANDROID
+	const bool bIsUsingVulkan = FAndroidMisc::ShouldUseVulkan();
+	const bool bCanDoSRGB = (GMaxRHIFeatureLevel > ERHIFeatureLevel::ES2);
+#else
+	const bool bIsUsingVulkan = false;
+	const bool bCanDoSRGB = true;
+#endif
 	FGenerateMipsCS::FPermutationDomain PermutationVector;
-	PermutationVector.Set<FGenerateMipsCS::FGenMipsSRGB>(!!(InTexture->GetFlags() & TexCreate_SRGB));
+	PermutationVector.Set<FGenerateMipsCS::FGenMipsSRGB>(!!(InTexture->GetFlags() & TexCreate_SRGB) && bCanDoSRGB);
+	// TEMP: On Vulkan we experience RGB being swizzled around, this little switch here circumvents the issue
+	PermutationVector.Set<FGenerateMipsCS::FGenMipsSwizzle>(bIsUsingVulkan);
 	TShaderMapRef<FGenerateMipsCS> ComputeShader(GetGlobalShaderMap(GMaxRHIFeatureLevel), PermutationVector);
 
 	//Loop through each level of the mips that require creation and add a dispatch pass per level,.
@@ -154,9 +348,9 @@ void FGenerateMips::Compute(FRHICommandListImmediate& RHIImmCmdList, FRHITexture
 	GraphBuilder.Execute();	
 }
 
+
 //Public execute function for calling the generate mips compute shader. Handles everything per platform.
-void FGenerateMips::Execute(FRHICommandListImmediate& RHICmdList, FRHITexture* InTexture,
-	const FGenerateMipsParams& InParams)
+void FGenerateMips::Execute(FRHICommandListImmediate& RHICmdList, FRHITexture* InTexture, const FGenerateMipsParams& InParams, TSharedPtr<FGenerateMipsStruct> * ExternalMipsStructCache, bool bAllowRenderBasedGeneration)
 {
 	//Only executes if mips are required.
 	if (InTexture->GetNumMips() > 1)
@@ -164,20 +358,52 @@ void FGenerateMips::Execute(FRHICommandListImmediate& RHICmdList, FRHITexture* I
 		//Checks whether the platform requires the compute shader. If not,executes that RHI's native generate mips call.
 		if (RHIRequiresComputeGenerateMips())
 		{
+#if PLATFORM_ANDROID
+			if (!FAndroidMisc::ShouldUseVulkan())
+			{
+				if (bAllowRenderBasedGeneration)
+				{
+					RenderMips(RHICmdList, InTexture, InParams, ExternalMipsStructCache);
+					return;
+				}
+				check(!"ES2 & ES3.1 do not support suitable compute features for mip generation currently; use rendering based generation if possible.");
+			}
+#endif
 			//Generate the RenderGraph texture if required.
+
+			// Do we have an external cache for the parameters we will use?
+			if (!ExternalMipsStructCache)
+			{
+				// No: Old path to keep things compatible
+				//THIS WILL CAUSE CIRCULAR REFS BETWEEN THE PARAMETERS AND THE TEXTURE -> LEAK!
 			if (!InTexture->GenMipsStruct)
 			{
-				SetupTexture(InTexture, InParams);
+					InTexture->GenMipsStruct = SetupTexture(InTexture, InParams);
 			}
-			Compute(RHICmdList, InTexture);
+				Compute(RHICmdList, InTexture, InTexture->GenMipsStruct);
 		}
 		else
 		{
+				// Use external location to cache params etc.
+
+				// Already valid?
+				if (!(*ExternalMipsStructCache))
+				{
+					// No, populate...
+					*ExternalMipsStructCache = SetupTexture(InTexture, InParams);
+				}
+				Compute(RHICmdList, InTexture, *ExternalMipsStructCache);
+			}
+		}
+		else
+		{
+			// Fallback to platform's native implementation
 			RHICmdList.GenerateMips(InTexture);
 		}
 	}
 }
 
+
 void FGenerateMips::Execute(FRDGBuilder* GraphBuilder, FRDGTextureRef InGraphTexture, FRHISamplerState* InSampler)
 {
 	check(IsInRenderingThread());
@@ -187,8 +413,9 @@ void FGenerateMips::Execute(FRDGBuilder* GraphBuilder, FRDGTextureRef InGraphTex
 
 	FGenerateMipsCS::FPermutationDomain PermutationVector;
 	PermutationVector.Set<FGenerateMipsCS::FGenMipsSRGB>(!!(InGraphTexture->Desc.Flags & TexCreate_SRGB));
-	TShaderMapRef<FGenerateMipsCS> ComputeShader(GetGlobalShaderMap(ERHIFeatureLevel::SM5), PermutationVector);
-
+	PermutationVector.Set<FGenerateMipsCS::FGenMipsSwizzle>(false);
+	TShaderMapRef<FGenerateMipsCS> ComputeShader(GetGlobalShaderMap(GMaxRHIFeatureLevel), PermutationVector);
+	
 	//Loop through each level of the mips that require creation and add a dispatch pass per level,.
 	for (uint8 MipLevel = 1; MipLevel < InGraphTexture->Desc.NumMips; MipLevel++)
 	{
diff --git a/Engine/Source/Runtime/RenderCore/Public/GenerateMips.h b/Engine/Source/Runtime/RenderCore/Public/GenerateMips.h
--- a/Engine/Source/Runtime/RenderCore/Public/GenerateMips.h
+++ b/Engine/Source/Runtime/RenderCore/Public/GenerateMips.h
@@ -19,12 +19,14 @@ class RENDERCORE_API FGenerateMips
 public:
 	//Public function for executing the generate mips compute shader 
 	//Default sampler is always bilinear clamp
-	static void Execute(FRHICommandListImmediate& RHICmdList, FRHITexture* InTexture,
-		const FGenerateMipsParams& InParams = FGenerateMipsParams());
+	static void Execute(FRHICommandListImmediate& RHICmdList, FRHITexture* InTexture, const FGenerateMipsParams& InParams = FGenerateMipsParams(), TSharedPtr<FGenerateMipsStruct> * ExternalMipsStructCache = nullptr, bool bAllowRenderBasedGeneration = false);
 
 	static void Execute(class FRDGBuilder* GraphBuilder, class FRDGTexture* InGraphTexture, FRHISamplerState* InSampler);
+
 private:
-	static void Compute(FRHICommandListImmediate& RHIImmCmdList, FRHITexture* InTexture);
-	static FGenerateMipsStruct* SetupTexture(FRHITexture* InTexture,
-		const FGenerateMipsParams& InParams = FGenerateMipsParams());
+	static void Compute(FRHICommandListImmediate& RHIImmCmdList, FRHITexture* InTexture, TSharedPtr<FGenerateMipsStruct> GenMipsStruct);
+	static TSharedPtr<FGenerateMipsStruct> SetupTexture(FRHITexture* InTexture, const FGenerateMipsParams& InParams = FGenerateMipsParams());
+
+	static void RenderMips(FRHICommandListImmediate& RHICmdList, FRHITexture* InTexture, const FGenerateMipsParams& InParams, TSharedPtr<FGenerateMipsStruct> * ExternalMipsStructCache);
+	static void SetupRendering(FGenerateMipsStruct *GenMipsStruct, FRHITexture* InTexture, const FGenerateMipsParams& InParams);
 };
\ No newline at end of file
