Android: Moving mips generation to rendering version for Vulkan devices, too
#rnx

diff --git a/Engine/Source/Runtime/RenderCore/Private/GenerateMips.cpp b/Engine/Source/Runtime/RenderCore/Private/GenerateMips.cpp
--- a/Engine/Source/Runtime/RenderCore/Private/GenerateMips.cpp
+++ b/Engine/Source/Runtime/RenderCore/Private/GenerateMips.cpp
@@ -359,15 +359,12 @@ void FGenerateMips::Execute(FRHICommandListImmediate& RHICmdList, FRHITexture* I
 		if (RHIRequiresComputeGenerateMips())
 		{
 #if PLATFORM_ANDROID
-			if (!FAndroidMisc::ShouldUseVulkan())
+			if (bAllowRenderBasedGeneration)
 			{
-				if (bAllowRenderBasedGeneration)
-				{
-					RenderMips(RHICmdList, InTexture, InParams, ExternalMipsStructCache);
-					return;
-				}
-				check(!"ES2 & ES3.1 do not support suitable compute features for mip generation currently; use rendering based generation if possible.");
+				RenderMips(RHICmdList, InTexture, InParams, ExternalMipsStructCache);
+				return;
 			}
+			check(!"Vulkan, ES2 & ES3.1 do not support suitable compute features (output format selection via HLSL crosscompile) for mip generation currently; use rendering based generation if possible.");
 #endif
 			//Generate the RenderGraph texture if required.
 
