Add global default UINT textures
Bind one of these when virtual texture page table is not ready
Should fix non-matching format error on APIs where this matters

diff --git a/Engine/Source/Runtime/Engine/Private/Materials/MaterialUniformExpressions.cpp b/Engine/Source/Runtime/Engine/Private/Materials/MaterialUniformExpressions.cpp
--- a/Engine/Source/Runtime/Engine/Private/Materials/MaterialUniformExpressions.cpp
+++ b/Engine/Source/Runtime/Engine/Private/Materials/MaterialUniformExpressions.cpp
@@ -848,10 +848,10 @@ void FUniformExpressionSet::FillUniformBuffer(const FMaterialRenderContext& Mate
 			else
 			{
 				// Don't have valid resources to bind for this VT, so make sure something is bound
-				*ResourceTablePageTexture0Ptr = GWhiteTexture->TextureRHI;
+				*ResourceTablePageTexture0Ptr = GBlackUintTexture->TextureRHI;
 				if (ResourceTablePageTexture1Ptr != nullptr)
 				{
-					*ResourceTablePageTexture1Ptr = GWhiteTexture->TextureRHI;
+					*ResourceTablePageTexture1Ptr = GBlackUintTexture->TextureRHI;
 				}
 			}
 		}
diff --git a/Engine/Source/Runtime/RenderCore/Private/RenderUtils.cpp b/Engine/Source/Runtime/RenderCore/Private/RenderUtils.cpp
--- a/Engine/Source/Runtime/RenderCore/Private/RenderUtils.cpp
+++ b/Engine/Source/Runtime/RenderCore/Private/RenderUtils.cpp
@@ -593,6 +593,89 @@ public:
 };
 FTexture* GBlackCubeArrayTexture = new TGlobalResource<FBlackCubeArrayTexture>;
 
+/**
+ * A UINT 1x1 texture.
+ */
+template <EPixelFormat Format, uint32 R = 0, uint32 G = 0, uint32 B = 0, uint32 A = 0>
+class FUintTexture : public FTextureWithSRV
+{
+public:
+	// FResource interface.
+	virtual void InitRHI() override
+	{
+		// Create the texture RHI.  		
+		FRHIResourceCreateInfo CreateInfo(TEXT("UintTexture"));
+		FTexture2DRHIRef Texture2D = RHICreateTexture2D(1, 1, Format, 1, 1, TexCreate_ShaderResource, CreateInfo);
+		TextureRHI = Texture2D;
+
+		// Write the contents of the texture.
+		uint32 DestStride;
+		void* DestBuffer = RHILockTexture2D(Texture2D, 0, RLM_WriteOnly, DestStride, false);
+		WriteData(DestBuffer);
+		RHIUnlockTexture2D(Texture2D, 0, false);
+
+		// Create the sampler state RHI resource.
+		FSamplerStateInitializerRHI SamplerStateInitializer(SF_Point, AM_Wrap, AM_Wrap, AM_Wrap);
+		SamplerStateRHI = RHICreateSamplerState(SamplerStateInitializer);
+
+		// Create a view of the texture
+		ShaderResourceViewRHI = RHICreateShaderResourceView(TextureRHI, 0u);
+	}
+
+	/** Returns the width of the texture in pixels. */
+	virtual uint32 GetSizeX() const override
+	{
+		return 1;
+	}
+
+	/** Returns the height of the texture in pixels. */
+	virtual uint32 GetSizeY() const override
+	{
+		return 1;
+	}
+
+protected:
+	static int32 GetNumChannels()
+	{
+		return GPixelFormats[Format].NumComponents;
+	}
+
+	static int32 GetBytesPerChannel()
+	{
+		return GPixelFormats[Format].BlockBytes / GPixelFormats[Format].NumComponents;
+	}
+
+	template<typename T>
+	static void DoWriteData(T* DataPtr)
+	{
+		T Values[] = { R, G, B, A };
+		for (int32 i = 0; i < GetNumChannels(); ++i)
+		{
+			DataPtr[i] = Values[i];
+		}
+	}
+
+	static void WriteData(void* DataPtr)
+	{
+		switch (GetBytesPerChannel())
+		{
+		case 1: 
+			DoWriteData((uint8*)DataPtr);
+			return;
+		case 2:
+			DoWriteData((uint16*)DataPtr);
+			return;
+		case 4:
+			DoWriteData((uint32*)DataPtr);
+			return;
+		}
+		// Unsupported format
+		check(0);
+	}
+};
+
+FTexture* GBlackUintTexture = new TGlobalResource< FUintTexture<PF_R32G32B32A32_UINT> >;
+
 /*
 	3 XYZ packed in 4 bytes. (11:11:10 for X:Y:Z)
 */
diff --git a/Engine/Source/Runtime/RenderCore/Public/RenderUtils.h b/Engine/Source/Runtime/RenderCore/Public/RenderUtils.h
--- a/Engine/Source/Runtime/RenderCore/Public/RenderUtils.h
+++ b/Engine/Source/Runtime/RenderCore/Public/RenderUtils.h
@@ -136,6 +136,9 @@ extern RENDERCORE_API class FTexture* GBlackVolumeTexture;
 /** A global black volume texture, with alpha=1. */
 extern RENDERCORE_API class FTexture* GBlackAlpha1VolumeTexture;
 
+/** A global black texture<uint> */
+extern RENDERCORE_API class FTexture* GBlackUintTexture;
+
 /** A global black volume texture<uint>  */
 extern RENDERCORE_API class FTexture* GBlackUintVolumeTexture;
 
