Added missing forward class reference.

diff --git a/Engine/Source/Runtime/RenderCore/Public/VertexFactory.h b/Engine/Source/Runtime/RenderCore/Public/VertexFactory.h
--- a/Engine/Source/Runtime/RenderCore/Public/VertexFactory.h
+++ b/Engine/Source/Runtime/RenderCore/Public/VertexFactory.h
@@ -640,7 +640,7 @@ public:
 		ERHIFeatureLevel::Type FeatureLevel,
 		const FVertexFactory* VertexFactory,
 		const struct FMeshBatchElement& BatchElement,
-		FMeshDrawSingleShaderBindings& ShaderBindings,
+		class FMeshDrawSingleShaderBindings& ShaderBindings,
 		FVertexInputStreamArray& VertexStreams
 		) const
 	{
