Remove redundant hashing from CL 6396547.

#rnx

diff --git a/Engine/Source/Runtime/RenderCore/Public/Shader.h b/Engine/Source/Runtime/RenderCore/Public/Shader.h
--- a/Engine/Source/Runtime/RenderCore/Public/Shader.h
+++ b/Engine/Source/Runtime/RenderCore/Public/Shader.h
@@ -225,13 +225,6 @@ public:
 		checkf(BufferIndex == InBufferIndex, TEXT("Tweak FShaderLooseParameterBufferInfo type sizes"));
 	}
 
-	inline uint32 GetHash() const
-	{
-		uint32 Hash = FCrc::TypeCrc32(BufferIndex, 0);
-		Hash = FCrc::TypeCrc32(BufferSize, Hash);
-		return Hash;
-	}
-
 	friend FArchive& operator<<(FArchive& Ar,FShaderLooseParameterBufferInfo& Info)
 	{
 		Ar << Info.BufferIndex;
@@ -265,16 +258,6 @@ public:
 		return Ar;
 	}
 
-	inline uint32 GetHash() const
-	{
-		uint32 Hash = 0;
-		for (const FShaderLooseParameterBufferInfo& ShaderLooseParameterBufferInfo: LooseParameterBuffers)
-		{
-			Hash = HashCombine(ShaderLooseParameterBufferInfo.GetHash(), Hash);
-		}
-		return Hash;
-	}
-
 	inline bool operator==(const FShaderParameterMapInfo& Rhs) const
 	{
 		return UniformBuffers == Rhs.UniformBuffers
diff --git a/Engine/Source/Runtime/Renderer/Public/MeshDrawShaderBindings.h b/Engine/Source/Runtime/Renderer/Public/MeshDrawShaderBindings.h
--- a/Engine/Source/Runtime/Renderer/Public/MeshDrawShaderBindings.h
+++ b/Engine/Source/Runtime/Renderer/Public/MeshDrawShaderBindings.h
@@ -28,8 +28,7 @@ public:
 	inline uint32 GetHash() const
 	{
 		uint32 LocalFrequency = Frequency;
-		uint32 Hash = FCrc::TypeCrc32(LocalFrequency, ParameterMapInfo.GetHash());
-		return Hash;
+		return FCrc::TypeCrc32(LocalFrequency, 0);
 	}
 
 	bool operator==(const FMeshDrawShaderBindingsLayout& Rhs) const
