addition of cvar to limit the max amount of pre-compile time for PSOs, defaulting to off for the time being

diff --git a/Engine/Source/Runtime/RenderCore/Private/ShaderPipelineCache.cpp b/Engine/Source/Runtime/RenderCore/Private/ShaderPipelineCache.cpp
--- a/Engine/Source/Runtime/RenderCore/Private/ShaderPipelineCache.cpp
+++ b/Engine/Source/Runtime/RenderCore/Private/ShaderPipelineCache.cpp
@@ -146,6 +146,13 @@ static TAutoConsoleVariable<int32> CVarPSOFileCacheMinBindCount(
                                                                 ECVF_Default | ECVF_RenderThreadSafe
                                                                 );
 
+static TAutoConsoleVariable<float> CVarPSOFileCacheMaxPrecompileTime(
+																TEXT("r.ShaderPipelineCache.MaxPrecompileTime"),
+																(float)0.f,
+																TEXT("The maximum time to allow a PSO to be precompiled.  if greather than 0, the amount of wall time we will allow pre-compile of PSOs and then switch to background processing."),
+																ECVF_Default | ECVF_RenderThreadSafe
+																);
+
 static bool GetShaderPipelineCacheSaveBoundPSOLog()
 {
 	static bool bOnce = false;
@@ -481,9 +488,17 @@ uint32 FShaderPipelineCache::NumPrecompilesRemaining()
 	
 	if (ShaderPipelineCache)
 	{
-        int64 NumActiveTasksRemaining = FMath::Max(0ll, FPlatformAtomics::AtomicRead(&ShaderPipelineCache->TotalActiveTasks));
-        int64 NumWaitingTasksRemaining = FMath::Max(0ll, FPlatformAtomics::AtomicRead(&ShaderPipelineCache->TotalWaitingTasks));
-		NumRemaining = (uint32)(NumActiveTasksRemaining + NumWaitingTasksRemaining);
+		if (CVarPSOFileCacheMaxPrecompileTime.GetValueOnAnyThread() > 0.0 && FPlatformAtomics::AtomicRead(&ShaderPipelineCache->TotalPrecompileTasks) > 0)
+		{
+			float Mult = ShaderPipelineCache->TotalPrecompileWallTime / CVarPSOFileCacheMaxPrecompileTime.GetValueOnAnyThread();
+			NumRemaining = FMath::Max(0.f, 1.f - Mult) * FPlatformAtomics::AtomicRead(&ShaderPipelineCache->TotalPrecompileTasks);
+		}
+		else
+		{
+        	int64 NumActiveTasksRemaining = FMath::Max(0ll, FPlatformAtomics::AtomicRead(&ShaderPipelineCache->TotalActiveTasks));
+        	int64 NumWaitingTasksRemaining = FMath::Max(0ll, FPlatformAtomics::AtomicRead(&ShaderPipelineCache->TotalWaitingTasks));
+			NumRemaining = (uint32)(NumActiveTasksRemaining + NumWaitingTasksRemaining);
+		}
 	}
 
 	return NumRemaining;
@@ -500,6 +515,11 @@ uint32 FShaderPipelineCache::NumPrecompilesActive()
         {
             NumRemaining = (uint32) NumTasksRemaining;
         }
+		if (CVarPSOFileCacheMaxPrecompileTime.GetValueOnAnyThread() > 0.0 && FPlatformAtomics::AtomicRead(&ShaderPipelineCache->TotalPrecompileTasks) > 0)
+		{
+			float Mult = ShaderPipelineCache->TotalPrecompileWallTime / CVarPSOFileCacheMaxPrecompileTime.GetValueOnAnyThread();
+			NumRemaining = FMath::Max(0.f, 1.f - Mult) * FPlatformAtomics::AtomicRead(&ShaderPipelineCache->TotalPrecompileTasks);
+		}
     }
     
     return NumRemaining;
@@ -1105,6 +1125,8 @@ FShaderPipelineCache::FShaderPipelineCache(EShaderPlatform Platform)
 , LastAutoSaveTime(0)
 , LastAutoSaveTimeLogBoundPSO(0)
 , LastAutoSaveNum(-1)
+, TotalPrecompileWallTime(0.0)
+, TotalPrecompileTasks(0)
 {
 	SET_DWORD_STAT(STAT_ShaderPipelineTaskCount, 0);
     SET_DWORD_STAT(STAT_ShaderPipelineWaitingTaskCount, 0);
@@ -1190,22 +1212,32 @@ void FShaderPipelineCache::Tick( float DeltaTime )
 		LastPrecompileRHIFence = nullptr;
 	}
 
-    if (FPlatformAtomics::AtomicRead(&TotalWaitingTasks) == 0 && FPlatformAtomics::AtomicRead(&TotalActiveTasks) == 0 && FPlatformAtomics::AtomicRead(&TotalCompleteTasks) != 0 && !LastPrecompileRHIFence.GetReference())
+	if (PrecompileStartTime > 0.0)
+	{
+		TotalPrecompileWallTime = float(FPlatformTime::Seconds() - PrecompileStartTime);
+	}
+
+	// HACK: note for the time being until the code further upstream properly uses the OnPrecompilationComplete delegate, we need to provide some delays to make sure it gets the messaging properly (the -10.0f and the -20.0f below).
+	if (((FPlatformAtomics::AtomicRead(&TotalWaitingTasks) == 0 && FPlatformAtomics::AtomicRead(&TotalActiveTasks) == 0 && FPlatformAtomics::AtomicRead(&TotalCompleteTasks) != 0) || (CVarPSOFileCacheMaxPrecompileTime.GetValueOnAnyThread() > 0.0 && TotalPrecompileWallTime - 10.0f > CVarPSOFileCacheMaxPrecompileTime.GetValueOnAnyThread() && TotalPrecompileTasks > 0)) && !LastPrecompileRHIFence.GetReference())
     {
-		float WallTime = 0.0f;
-		if (PrecompileStartTime > 0.0)
-		{
-			WallTime = float(FPlatformTime::Seconds() - PrecompileStartTime);
-		}
-		UE_LOG(LogRHI, Warning, TEXT("FShaderPipelineCache completed %u tasks in %.2fs (%.2fs wall time since intial open)."), (uint32)TotalCompleteTasks, FPlatformTime::ToSeconds64(TotalPrecompileTime), WallTime);
+		UE_LOG(LogRHI, Warning, TEXT("FShaderPipelineCache completed %u tasks in %.2fs (%.2fs wall time since intial open)."), (uint32)TotalCompleteTasks, FPlatformTime::ToSeconds64(TotalPrecompileTime), TotalPrecompileWallTime);
         if (OnPrecompilationComplete.IsBound())
         {
             OnPrecompilationComplete.Broadcast((uint32)TotalCompleteTasks, FPlatformTime::ToSeconds64(TotalPrecompileTime), ShaderCachePrecompileContext);
         }
-		FPipelineFileCache::PreCompileComplete();
+		if (CVarPSOFileCacheMaxPrecompileTime.GetValueOnAnyThread() > 0.0 && TotalPrecompileWallTime - 20.0f > CVarPSOFileCacheMaxPrecompileTime.GetValueOnAnyThread() && TotalPrecompileTasks > 0)
+		{
+			PrecompileStartTime = 0.0;
+			SetBatchMode(BatchMode::Background);
+			TotalPrecompileTasks = 0;
+		}
+		else
+		{
+			FPipelineFileCache::PreCompileComplete();
 		
-        FPlatformAtomics::InterlockedExchange(&TotalCompleteTasks, 0);
-        FPlatformAtomics::InterlockedExchange(&TotalPrecompileTime, 0);
+        	FPlatformAtomics::InterlockedExchange(&TotalCompleteTasks, 0);
+        	FPlatformAtomics::InterlockedExchange(&TotalPrecompileTime, 0);
+		}
     }
 	
 	if (ReadyForAutoSave())
@@ -1446,7 +1478,8 @@ bool FShaderPipelineCache::Open(FString const& Name, EShaderPlatform Platform)
 			}
 
 			PreFetchedTasks = LocalPreFetchedTasks;
-
+			TotalPrecompileTasks = PreFetchedTasks.Num();
+			
 			UE_LOG(LogRHI, Display, TEXT("Opened pipeline cache and enqueued %d of %d tasks for precompile with BatchSize %d and BatchTime %f."), Count, PreFetchedTasks.Num(), BatchSize, BatchTime);
 		}
 		else
diff --git a/Engine/Source/Runtime/RenderCore/Public/ShaderPipelineCache.h b/Engine/Source/Runtime/RenderCore/Public/ShaderPipelineCache.h
--- a/Engine/Source/Runtime/RenderCore/Public/ShaderPipelineCache.h
+++ b/Engine/Source/Runtime/RenderCore/Public/ShaderPipelineCache.h
@@ -267,4 +267,6 @@ private:
 	int32 LastAutoSaveNum;
 	
 	TSet<uint64> CompletedMasks;
+	float TotalPrecompileWallTime;
+	int64 TotalPrecompileTasks;
 };
