Minor shader loading optimization

diff --git a/Engine/Source/Runtime/RenderCore/Public/Shader.h b/Engine/Source/Runtime/RenderCore/Public/Shader.h
--- a/Engine/Source/Runtime/RenderCore/Public/Shader.h
+++ b/Engine/Source/Runtime/RenderCore/Public/Shader.h
@@ -115,7 +115,7 @@ public:
 
 	friend inline uint32 GetTypeHash( const FShaderResourceId& Id )
 	{
-		return FCrc::MemCrc_DEPRECATED((const void*)&Id.OutputHash, sizeof(Id.OutputHash));
+		return GetTypeHash(Id.OutputHash);
 	}
 
 	friend bool operator==(const FShaderResourceId& X, const FShaderResourceId& Y)
