Nomaster fix

diff --git a/Engine/Source/Runtime/RenderCore/Public/VertexFactory.h b/Engine/Source/Runtime/RenderCore/Public/VertexFactory.h
--- a/Engine/Source/Runtime/RenderCore/Public/VertexFactory.h
+++ b/Engine/Source/Runtime/RenderCore/Public/VertexFactory.h
@@ -16,6 +16,7 @@
 #include "Misc/EnumClassFlags.h"
 
 class FMaterial;
+class FMeshDrawSingleShaderBindings;
 
 #if PLATFORM_SUPPORTS_PRAGMA_PACK
 	#pragma pack (push,4)
