Change deprecation style from macro / compiler warning + linker error to static assert

diff --git a/Engine/Source/Runtime/RenderCore/Public/Shader.h b/Engine/Source/Runtime/RenderCore/Public/Shader.h
--- a/Engine/Source/Runtime/RenderCore/Public/Shader.h
+++ b/Engine/Source/Runtime/RenderCore/Public/Shader.h
@@ -420,8 +420,12 @@ public:
 	virtual void ReleaseRHI();
 
 
-	UE_DEPRECATED(4.24, "Save return value as TRefCountPtr<FShaderResource> and use FindById() instead")
-	static FShaderResource* FindShaderResourceById(const FShaderResourceId& Id);
+	template<class T>
+	static FShaderResource* FindShaderResourceById(T Id)
+	{
+		static_assert(sizeof(T) == 0, "Deprecated. Please save return value as TRefCountPtr<FShaderResource> and switch to FindById().");
+		return nullptr;
+	}
 
 	/** Finds a matching shader resource in memory if possible. */
 	RENDERCORE_API static TRefCountPtr<FShaderResource> FindById(const FShaderResourceId& Id);
@@ -429,8 +433,12 @@ public:
 	/** Finds a matching shader resource in memory or clone the temporary and register the new resource */
 	RENDERCORE_API static TRefCountPtr<FShaderResource> FindOrClone(FShaderResource&& Temp);
 
-	UE_DEPRECATED(4.24, "Save return value as TRefCountPtr<FShaderResource> and use FindOrCreate() instead")
-	static FShaderResource* FindOrCreateShaderResource(const FShaderCompilerOutput& Output, class FShaderType* SpecificType, int32 SpecificPermutationId);
+	template<class T>
+	static FShaderResource* FindOrCreateShaderResource(const FShaderCompilerOutput& Output, class FShaderType* SpecificType, T SpecificPermutationId)
+	{
+		static_assert(sizeof(T) == 0, "Deprecated. Please save return value as TRefCountPtr<FShaderResource> and switch to FindById().");
+		return nullptr;
+	}
 
 	/** 
 	 * Finds a matching shader resource in memory or creates a new one with the given compiler output.  
