Optimize AllowDebugViewmodes (TLightMapDensityDS::ShouldCompilePermutation) by replacing lots of temporary string allocations and individual string lookups with a global precalculated u32 GRequiresCookedDataPlatformMask.

diff --git a/Engine/Source/Runtime/RenderCore/Private/RenderUtils.cpp b/Engine/Source/Runtime/RenderCore/Private/RenderUtils.cpp
--- a/Engine/Source/Runtime/RenderCore/Private/RenderUtils.cpp
+++ b/Engine/Source/Runtime/RenderCore/Private/RenderUtils.cpp
@@ -950,6 +950,11 @@ static_assert(SP_NumPlatforms <= sizeof(GSelectiveBasePassOutputsPlatformMask) *
 RENDERCORE_API uint64 GDistanceFieldsPlatformMask = 0;
 static_assert(SP_NumPlatforms <= sizeof(GDistanceFieldsPlatformMask) * 8, "GDistanceFieldsPlatformMask must be large enough to support all shader platforms");
 
+#if WITH_EDITOR
+RENDERCORE_API uint32 GRequiresCookedDataPlatformMask = 0;
+static_assert(SP_NumPlatforms <= sizeof(GRequiresCookedDataPlatformMask) * 8, "GRequiresCookedDataPlatformMask must be large enough to support all shader platforms");
+#endif
+
 RENDERCORE_API void RenderUtilsInit()
 {
 	if (GUseForwardShading)
@@ -1038,6 +1043,15 @@ RENDERCORE_API void RenderUtilsInit()
 				{
 					GDistanceFieldsPlatformMask &= ~Mask;
 				}
+
+				if (TargetPlatform->RequiresCookedData())
+				{
+					GRequiresCookedDataPlatformMask |= Mask;
+				}
+				else
+				{
+					GRequiresCookedDataPlatformMask &= ~Mask;
+				}
 			}
 		}
 	}
@@ -1128,4 +1142,4 @@ RENDERCORE_API void QuantizeSceneBufferSize(const FIntPoint& InBufferSize, FIntP
 	const uint32 Mask = ~(DividableBy - 1);
 	OutBufferSize.X = (InBufferSize.X + DividableBy - 1) & Mask;
 	OutBufferSize.Y = (InBufferSize.Y + DividableBy - 1) & Mask;
-}
\ No newline at end of file
+}
diff --git a/Engine/Source/Runtime/RenderCore/Private/ShaderCore.cpp b/Engine/Source/Runtime/RenderCore/Private/ShaderCore.cpp
--- a/Engine/Source/Runtime/RenderCore/Private/ShaderCore.cpp
+++ b/Engine/Source/Runtime/RenderCore/Private/ShaderCore.cpp
@@ -18,6 +18,7 @@
 #include "Interfaces/IShaderFormatModule.h"
 #include "RHIShaderFormatDefinitions.inl"
 #if WITH_EDITOR
+#include "RenderUtils.h"
 #include "Misc/CoreMisc.h"
 #include "Interfaces/ITargetPlatform.h"
 #include "Interfaces/ITargetPlatformManagerModule.h"
@@ -259,8 +260,7 @@ bool AllowDebugViewmodes(EShaderPlatform Platform)
 	const int32 ForceDebugViewValue = CVarForceDebugViewModes.GetValueOnAnyThread();
 	bool bForceEnable = ForceDebugViewValue == 1;
 	bool bForceDisable = ForceDebugViewValue == 2;
-	ITargetPlatform* TargetPlatform = GetTargetPlatformManager()->FindTargetPlatform(ShaderPlatformToPlatformName(Platform).ToString());
-	return (!bForceDisable) && (bForceEnable || !TargetPlatform || !TargetPlatform->RequiresCookedData());
+	return (!bForceDisable) && (bForceEnable || !RequiresCookedData(Platform));
 #else
 	return AllowDebugViewmodes();
 #endif
diff --git a/Engine/Source/Runtime/RenderCore/Public/RenderUtils.h b/Engine/Source/Runtime/RenderCore/Public/RenderUtils.h
--- a/Engine/Source/Runtime/RenderCore/Public/RenderUtils.h
+++ b/Engine/Source/Runtime/RenderCore/Public/RenderUtils.h
@@ -481,6 +481,15 @@ inline bool IsUsingDistanceFields(EShaderPlatform Platform)
 	return !!(GDistanceFieldsPlatformMask & (1ull << Platform));
 }
 
+#if WITH_EDITOR
+/** Returns whether cooked data is required for a given shader platform */
+inline bool RequiresCookedData(EShaderPlatform Platform)
+{
+	extern RENDERCORE_API uint32 GRequiresCookedDataPlatformMask;
+	return !!(GRequiresCookedDataPlatformMask & (1u << Platform));
+}
+#endif
+
 inline bool IsUsingPerPixelDBufferMask(EShaderPlatform Platform)
 {
 	switch (Platform)
