Fix for error C2338: GRequiresCookedDataPlatformMask must be large enough to support all shader platforms (introduced in cl 6939155)

diff --git a/Engine/Source/Runtime/RenderCore/Private/RenderUtils.cpp b/Engine/Source/Runtime/RenderCore/Private/RenderUtils.cpp
--- a/Engine/Source/Runtime/RenderCore/Private/RenderUtils.cpp
+++ b/Engine/Source/Runtime/RenderCore/Private/RenderUtils.cpp
@@ -951,7 +951,7 @@ RENDERCORE_API uint64 GDistanceFieldsPlatformMask = 0;
 static_assert(SP_NumPlatforms <= sizeof(GDistanceFieldsPlatformMask) * 8, "GDistanceFieldsPlatformMask must be large enough to support all shader platforms");
 
 #if WITH_EDITOR
-RENDERCORE_API uint32 GRequiresCookedDataPlatformMask = 0;
+RENDERCORE_API uint64 GRequiresCookedDataPlatformMask = 0;
 static_assert(SP_NumPlatforms <= sizeof(GRequiresCookedDataPlatformMask) * 8, "GRequiresCookedDataPlatformMask must be large enough to support all shader platforms");
 #endif
 
diff --git a/Engine/Source/Runtime/RenderCore/Public/RenderUtils.h b/Engine/Source/Runtime/RenderCore/Public/RenderUtils.h
--- a/Engine/Source/Runtime/RenderCore/Public/RenderUtils.h
+++ b/Engine/Source/Runtime/RenderCore/Public/RenderUtils.h
@@ -485,8 +485,8 @@ inline bool IsUsingDistanceFields(EShaderPlatform Platform)
 /** Returns whether cooked data is required for a given shader platform */
 inline bool RequiresCookedData(EShaderPlatform Platform)
 {
-	extern RENDERCORE_API uint32 GRequiresCookedDataPlatformMask;
-	return !!(GRequiresCookedDataPlatformMask & (1u << Platform));
+	extern RENDERCORE_API uint64 GRequiresCookedDataPlatformMask;
+	return !!(GRequiresCookedDataPlatformMask & (1ull << Platform));
 }
 #endif
 
