Re-add GRequiresCookedDataPlatformsMask that was lost in integration.

diff --git a/Engine/Source/Runtime/RenderCore/Private/RenderUtils.cpp b/Engine/Source/Runtime/RenderCore/Private/RenderUtils.cpp
--- a/Engine/Source/Runtime/RenderCore/Private/RenderUtils.cpp
+++ b/Engine/Source/Runtime/RenderCore/Private/RenderUtils.cpp
@@ -956,6 +956,11 @@ static_assert(SP_NumPlatforms <= sizeof(GSelectiveBasePassOutputsPlatformMask) *
 RENDERCORE_API uint64 GDistanceFieldsPlatformMask = 0;
 static_assert(SP_NumPlatforms <= sizeof(GDistanceFieldsPlatformMask) * 8, "GDistanceFieldsPlatformMask must be large enough to support all shader platforms");
 
+#if WITH_EDITOR
+RENDERCORE_API uint64 GRequiresCookedDataPlatformMask = 0;
+static_assert(SP_NumPlatforms <= sizeof(GRequiresCookedDataPlatformMask) * 8, "GRequiresCookedDataPlatformMask must be large enough to support all shader platforms");
+#endif
+
 RENDERCORE_API void RenderUtilsInit()
 {
 	if (GUseForwardShading)
@@ -1044,6 +1049,15 @@ RENDERCORE_API void RenderUtilsInit()
 				{
 					GDistanceFieldsPlatformMask &= ~Mask;
 				}
+
+				if (TargetPlatform->RequiresCookedData())
+				{
+					GRequiresCookedDataPlatformMask |= Mask;
+				}
+				else
+				{
+					GRequiresCookedDataPlatformMask &= ~Mask;
+				}
 			}
 		}
 	}
