Further fixup work for circular ref / leak when using FGenerateMips::Execute() - parameters are now regenerate each time unless an external cache for them is provided by the caller

diff --git a/Engine/Source/Runtime/Engine/Private/TextureRenderTarget2D.cpp b/Engine/Source/Runtime/Engine/Private/TextureRenderTarget2D.cpp
--- a/Engine/Source/Runtime/Engine/Private/TextureRenderTarget2D.cpp
+++ b/Engine/Source/Runtime/Engine/Private/TextureRenderTarget2D.cpp
@@ -567,12 +567,12 @@ void FTextureRenderTarget2DResource::UpdateDeferredResource( FRHICommandListImme
 		RHICmdList.EndRenderPass();
 	}
  
-	// #todo-renderpasses must generate mips outside of a renderpass?
+	// #todo-renderpasses must generate mips outside of a renderpass (right now compute only is used; but this is not possible on all platforms) ?
 	if (Owner->bAutoGenerateMips)
 	{
 		/**Convert the input values from the editor to a compatible format for FSamplerStateInitializerRHI. 
 			Ensure default sampler is Bilinear clamp*/
-		FGenerateMips::Execute(RHICmdList, RenderTargetTextureRHI, FGenerateMipsParams{
+		FGenerateMips::Execute(RHICmdList, RenderTargetTextureRHI, CachedMipsGenParams, FGenerateMipsParams{
 			Owner->MipsSamplerFilter == TF_Nearest ? SF_Point : (Owner->MipsSamplerFilter == TF_Trilinear ? SF_Trilinear : SF_Bilinear),
 			Owner->MipsAddressU == TA_Wrap ? AM_Wrap : (Owner->MipsAddressU == TA_Mirror ? AM_Mirror : AM_Clamp),
 			Owner->MipsAddressV == TA_Wrap ? AM_Wrap : (Owner->MipsAddressV == TA_Mirror ? AM_Mirror : AM_Clamp)});
diff --git a/Engine/Source/Runtime/Engine/Public/TextureResource.h b/Engine/Source/Runtime/Engine/Public/TextureResource.h
--- a/Engine/Source/Runtime/Engine/Public/TextureResource.h
+++ b/Engine/Source/Runtime/Engine/Public/TextureResource.h
@@ -703,6 +703,8 @@ private:
 	FLinearColor ClearColor;
 	EPixelFormat Format;
 	int32 TargetSizeX,TargetSizeY;
+	/** cached params etc. for use with mip generator */
+	TSharedPtr<FGenerateMipsStruct> CachedMipsGenParams;
 };
 
 /**
diff --git a/Engine/Source/Runtime/MediaAssets/Private/Misc/MediaTextureResource.cpp b/Engine/Source/Runtime/MediaAssets/Private/Misc/MediaTextureResource.cpp
--- a/Engine/Source/Runtime/MediaAssets/Private/Misc/MediaTextureResource.cpp
+++ b/Engine/Source/Runtime/MediaAssets/Private/Misc/MediaTextureResource.cpp
@@ -313,7 +313,7 @@ void FMediaTextureResource::Render(const FRenderParams& Params)
 #else
 				RHICmdList.TransitionResource(EResourceTransitionAccess::ERWBarrier, OutputTarget);
 #endif
-				FGenerateMips::Execute(RHICmdList, OutputTarget, FGenerateMipsParams{ SF_Bilinear, AM_Clamp, AM_Clamp, AM_Clamp }, &CachedMipsGenParams, true);
+				FGenerateMips::Execute(RHICmdList, OutputTarget, CachedMipsGenParams, FGenerateMipsParams{ SF_Bilinear, AM_Clamp, AM_Clamp, AM_Clamp }, true);
 				RHICmdList.TransitionResource(EResourceTransitionAccess::EReadable, OutputTarget);
 			}
 
diff --git a/Engine/Source/Runtime/RHI/Public/RHIResources.h b/Engine/Source/Runtime/RHI/Public/RHIResources.h
--- a/Engine/Source/Runtime/RHI/Public/RHIResources.h
+++ b/Engine/Source/Runtime/RHI/Public/RHIResources.h
@@ -636,7 +636,6 @@ public:
 	bool IsMultisampled() const { return NumSamples > 1; }		
 
 	FRHIResourceInfo ResourceInfo;
-	TSharedPtr<FGenerateMipsStruct> GenMipsStruct;
 
 	/** sets the last time this texture was cached in a resource table. */
 	FORCEINLINE_DEBUGGABLE void SetLastRenderTime(float InLastRenderTime)
diff --git a/Engine/Source/Runtime/RenderCore/Private/GenerateMips.cpp b/Engine/Source/Runtime/RenderCore/Private/GenerateMips.cpp
--- a/Engine/Source/Runtime/RenderCore/Private/GenerateMips.cpp
+++ b/Engine/Source/Runtime/RenderCore/Private/GenerateMips.cpp
@@ -350,7 +350,7 @@ void FGenerateMips::Compute(FRHICommandListImmediate& RHIImmCmdList, FRHITexture
 
 
 //Public execute function for calling the generate mips compute shader. Handles everything per platform.
-void FGenerateMips::Execute(FRHICommandListImmediate& RHICmdList, FRHITexture* InTexture, const FGenerateMipsParams& InParams, TSharedPtr<FGenerateMipsStruct> * ExternalMipsStructCache, bool bAllowRenderBasedGeneration)
+void FGenerateMips::ExecuteInternal(FRHICommandListImmediate& RHICmdList, FRHITexture* InTexture, const FGenerateMipsParams& InParams, TSharedPtr<FGenerateMipsStruct> * ExternalMipsStructCache, bool bAllowRenderBasedGeneration)
 {
 	//Only executes if mips are required.
 	if (InTexture->GetNumMips() > 1)
@@ -371,16 +371,11 @@ void FGenerateMips::Execute(FRHICommandListImmediate& RHICmdList, FRHITexture* I
 			// Do we have an external cache for the parameters we will use?
 			if (!ExternalMipsStructCache)
 			{
-				// No: Old path to keep things compatible
-				//THIS WILL CAUSE CIRCULAR REFS BETWEEN THE PARAMETERS AND THE TEXTURE -> LEAK!
-			if (!InTexture->GenMipsStruct)
-			{
-					InTexture->GenMipsStruct = SetupTexture(InTexture, InParams);
+				// No: Old path to keep things compatible. This will not cache any setup, so some performance penalty might apply!
+				Compute(RHICmdList, InTexture, SetupTexture(InTexture, InParams));
 			}
-				Compute(RHICmdList, InTexture, InTexture->GenMipsStruct);
-		}
-		else
-		{
+			else
+			{
 				// Use external location to cache params etc.
 
 				// Already valid?
diff --git a/Engine/Source/Runtime/RenderCore/Public/GenerateMips.h b/Engine/Source/Runtime/RenderCore/Public/GenerateMips.h
--- a/Engine/Source/Runtime/RenderCore/Public/GenerateMips.h
+++ b/Engine/Source/Runtime/RenderCore/Public/GenerateMips.h
@@ -17,13 +17,28 @@ struct FGenerateMipsParams
 class RENDERCORE_API FGenerateMips
 {
 public:
-	//Public function for executing the generate mips compute shader 
-	//Default sampler is always bilinear clamp
-	static void Execute(FRHICommandListImmediate& RHICmdList, FRHITexture* InTexture, const FGenerateMipsParams& InParams = FGenerateMipsParams(), TSharedPtr<FGenerateMipsStruct> * ExternalMipsStructCache = nullptr, bool bAllowRenderBasedGeneration = false);
+	/**
+	 * Public function for executing the generate mips compute shader 
+	 * @param RHICmdList RHI command list to use to schedule generation of mips
+	 * @param InTexture Texture to generate mips for
+	 * @param InParams Parameters influencing mips generation (Default sampler is always bilinear clamp)
+	 * @param ExternalMipsStructCache Shared pointer to provide a location to cache internal parameters and resource needed to generate mips (provide if generation is triggered repeatedly)
+	 * @param bAllowRenderBasedGeneration Set to true if generation via classic rendering is allowable (some platforms do not support compute based generation, but some use cases may disallow this)
+	 */
+	static void Execute(FRHICommandListImmediate& RHICmdList, FRHITexture* InTexture, TSharedPtr<FGenerateMipsStruct> & ExternalMipsStructCache, const FGenerateMipsParams& InParams = FGenerateMipsParams(), bool bAllowRenderBasedGeneration = false)
+	{
+		ExecuteInternal(RHICmdList, InTexture, InParams, &ExternalMipsStructCache, bAllowRenderBasedGeneration);
+	}
+	static void Execute(FRHICommandListImmediate& RHICmdList, FRHITexture* InTexture, const FGenerateMipsParams& InParams = FGenerateMipsParams(), bool bAllowRenderBasedGeneration = false)
+	{
+		ExecuteInternal(RHICmdList, InTexture, InParams, nullptr, bAllowRenderBasedGeneration);
+	}
 
 	static void Execute(class FRDGBuilder* GraphBuilder, class FRDGTexture* InGraphTexture, FRHISamplerState* InSampler);
 
 private:
+	static void ExecuteInternal(FRHICommandListImmediate& RHICmdList, FRHITexture* InTexture, const FGenerateMipsParams& InParams, TSharedPtr<FGenerateMipsStruct> * ExternalMipsStructCache, bool bAllowRenderBasedGeneration);
+
 	static void Compute(FRHICommandListImmediate& RHIImmCmdList, FRHITexture* InTexture, TSharedPtr<FGenerateMipsStruct> GenMipsStruct);
 	static TSharedPtr<FGenerateMipsStruct> SetupTexture(FRHITexture* InTexture, const FGenerateMipsParams& InParams = FGenerateMipsParams());
 
