Consistently use Buffer rather than Texture for VT feedback buffer

diff --git a/Engine/Source/Runtime/Engine/Private/SceneManagement.cpp b/Engine/Source/Runtime/Engine/Private/SceneManagement.cpp
--- a/Engine/Source/Runtime/Engine/Private/SceneManagement.cpp
+++ b/Engine/Source/Runtime/Engine/Private/SceneManagement.cpp
@@ -761,7 +761,7 @@ FViewUniformShaderParameters::FViewUniformShaderParameters()
 	{
 		LightmapSceneData = GBlackTextureWithSRV->ShaderResourceViewRHI;
 	}
-	VTFeedbackBuffer = GBlackTextureWithUAV->UnorderedAccessViewRHI;
+	VTFeedbackBuffer = GEmptyVertexBufferWithUAV->UnorderedAccessViewRHI;
 	QuadOverdraw = GBlackTextureWithUAV->UnorderedAccessViewRHI;
 }
 
diff --git a/Engine/Source/Runtime/Engine/Public/SceneView.h b/Engine/Source/Runtime/Engine/Public/SceneView.h
--- a/Engine/Source/Runtime/Engine/Public/SceneView.h
+++ b/Engine/Source/Runtime/Engine/Public/SceneView.h
@@ -812,7 +812,7 @@ BEGIN_GLOBAL_SHADER_PARAMETER_STRUCT_WITH_CONSTRUCTOR(FViewUniformShaderParamete
 	SHADER_PARAMETER_TEXTURE(Texture3D, CameraAerialPerspectiveVolume)
 	SHADER_PARAMETER_SAMPLER(SamplerState, CameraAerialPerspectiveVolumeSampler)
 
-	SHADER_PARAMETER_UAV(RWTexture2D<uint>, VTFeedbackBuffer)
+	SHADER_PARAMETER_UAV(RWBuffer<uint>, VTFeedbackBuffer)
 	SHADER_PARAMETER_UAV(RWTexture2D<uint>, QuadOverdraw)
 
 END_GLOBAL_SHADER_PARAMETER_STRUCT()
diff --git a/Engine/Source/Runtime/RenderCore/Private/RenderUtils.cpp b/Engine/Source/Runtime/RenderCore/Private/RenderUtils.cpp
--- a/Engine/Source/Runtime/RenderCore/Private/RenderUtils.cpp
+++ b/Engine/Source/Runtime/RenderCore/Private/RenderUtils.cpp
@@ -244,11 +244,29 @@ public:
 	}
 };
 
+class FEmptyVertexBuffer : public FVertexBufferWithSRV
+{
+public:
+	virtual void InitRHI() override
+	{
+		// Create the texture RHI.  		
+		FRHIResourceCreateInfo CreateInfo(TEXT("EmptyVertexBuffer"));
+		
+		VertexBufferRHI = RHICreateVertexBuffer(16u, BUF_Static | BUF_ShaderResource | BUF_UnorderedAccess, CreateInfo);
+
+		// Create a view of the buffer
+		ShaderResourceViewRHI = RHICreateShaderResourceView(VertexBufferRHI, 4u, PF_R32_UINT);
+		UnorderedAccessViewRHI = RHICreateUnorderedAccessView(VertexBufferRHI, PF_R32_UINT);
+	}
+};
+
 FTextureWithSRV* GWhiteTextureWithSRV = new TGlobalResource<FColoredTexture<255,255,255,255> >;
 FTextureWithSRV* GBlackTextureWithSRV = new TGlobalResource<FColoredTexture<0,0,0,255> >;
 FTexture* GWhiteTexture = GWhiteTextureWithSRV;
 FTexture* GBlackTexture = GBlackTextureWithSRV;
-FTextureWithSRV* GBlackTextureWithUAV = new TGlobalResource<FColoredTexture<0,0,0,0,1> >;
+FTextureWithSRV* GBlackTextureWithUAV = new TGlobalResource<FColoredTexture<0,0,0,0,true> >;
+
+FVertexBufferWithSRV* GEmptyVertexBufferWithUAV = new TGlobalResource<FEmptyVertexBuffer>;
 
 /**
  * Bulk data interface for providing a single black color used to initialize a
diff --git a/Engine/Source/Runtime/RenderCore/Public/RenderResource.h b/Engine/Source/Runtime/RenderCore/Public/RenderResource.h
--- a/Engine/Source/Runtime/RenderCore/Public/RenderResource.h
+++ b/Engine/Source/Runtime/RenderCore/Public/RenderResource.h
@@ -475,6 +475,23 @@ public:
 	virtual FString GetFriendlyName() const override { return TEXT("FVertexBuffer"); }
 };
 
+class RENDERCORE_API FVertexBufferWithSRV : public FVertexBuffer
+{
+public:
+	/** SRV that views the entire texture */
+	FShaderResourceViewRHIRef ShaderResourceViewRHI;
+
+	/** *optional* UAV that views the entire texture */
+	FUnorderedAccessViewRHIRef UnorderedAccessViewRHI;
+
+	virtual void ReleaseRHI() override
+	{
+		ShaderResourceViewRHI.SafeRelease();
+		UnorderedAccessViewRHI.SafeRelease();
+		FVertexBuffer::ReleaseRHI();
+	}
+};
+
 /**
 * A vertex buffer with a single color component.  This is used on meshes that don't have a color component
 * to keep from needing a separate vertex factory to handle this case.
diff --git a/Engine/Source/Runtime/RenderCore/Public/RenderUtils.h b/Engine/Source/Runtime/RenderCore/Public/RenderUtils.h
--- a/Engine/Source/Runtime/RenderCore/Public/RenderUtils.h
+++ b/Engine/Source/Runtime/RenderCore/Public/RenderUtils.h
@@ -128,6 +128,8 @@ extern RENDERCORE_API class FTexture* GBlackTexture;
 extern RENDERCORE_API class FTextureWithSRV* GBlackTextureWithSRV;
 extern RENDERCORE_API class FTextureWithSRV* GBlackTextureWithUAV;
 
+extern RENDERCORE_API class FVertexBufferWithSRV* GEmptyVertexBufferWithUAV;
+
 /** A global black array texture. */
 extern RENDERCORE_API class FTexture* GBlackArrayTexture;
 
diff --git a/Engine/Source/Runtime/Renderer/Private/PostProcess/SceneRenderTargets.cpp b/Engine/Source/Runtime/Renderer/Private/PostProcess/SceneRenderTargets.cpp
--- a/Engine/Source/Runtime/Renderer/Private/PostProcess/SceneRenderTargets.cpp
+++ b/Engine/Source/Runtime/Renderer/Private/PostProcess/SceneRenderTargets.cpp
@@ -871,7 +871,7 @@ FUnorderedAccessViewRHIRef FSceneRenderTargets::GetVirtualTextureFeedbackUAV()
 {
 	if (!VirtualTextureFeedback.FeedbackBufferUAV)
 	{
-		return GBlackTextureWithUAV->UnorderedAccessViewRHI;
+		return GEmptyVertexBufferWithUAV->UnorderedAccessViewRHI;
 	}
 	
 	return VirtualTextureFeedback.FeedbackBufferUAV;
diff --git a/Engine/Source/Runtime/SlateRHIRenderer/Private/SlateRHIRenderingPolicy.cpp b/Engine/Source/Runtime/SlateRHIRenderer/Private/SlateRHIRenderingPolicy.cpp
--- a/Engine/Source/Runtime/SlateRHIRenderer/Private/SlateRHIRenderingPolicy.cpp
+++ b/Engine/Source/Runtime/SlateRHIRenderer/Private/SlateRHIRenderingPolicy.cpp
@@ -224,7 +224,7 @@ static FSceneView* CreateSceneView( FSceneViewFamilyContext* ViewFamilyContext,
 
 	// Create the view's uniform buffer.
 	FViewUniformShaderParameters ViewUniformShaderParameters;
-	ViewUniformShaderParameters.VTFeedbackBuffer = GBlackTextureWithUAV->UnorderedAccessViewRHI; 
+	ViewUniformShaderParameters.VTFeedbackBuffer = GEmptyVertexBufferWithUAV->UnorderedAccessViewRHI;
 
 	View->SetupCommonViewUniformBufferParameters(
 		ViewUniformShaderParameters,
