ARRAY_COUNT -> UE_ARRAY_COUNT

diff --git a/Engine/Source/Runtime/RenderCore/Private/CommonRenderResources.cpp b/Engine/Source/Runtime/RenderCore/Private/CommonRenderResources.cpp
--- a/Engine/Source/Runtime/RenderCore/Private/CommonRenderResources.cpp
+++ b/Engine/Source/Runtime/RenderCore/Private/CommonRenderResources.cpp
@@ -53,7 +53,7 @@ void FScreenRectangleIndexBuffer::InitRHI()
 	const uint16 Indices[] = { 0, 1, 2, 2, 1, 3, 0, 4, 5 };
 
 	TResourceArray<uint16, INDEXBUFFER_ALIGNMENT> IndexBuffer;
-	uint32 NumIndices = ARRAY_COUNT(Indices);
+	uint32 NumIndices = UE_ARRAY_COUNT(Indices);
 	IndexBuffer.AddUninitialized(NumIndices);
 	FMemory::Memcpy(IndexBuffer.GetData(), Indices, NumIndices * sizeof(uint16));
 
diff --git a/Engine/Source/Runtime/RenderCore/Private/ProfilingDebugging/RealtimeGPUProfiler.cpp b/Engine/Source/Runtime/RenderCore/Private/ProfilingDebugging/RealtimeGPUProfiler.cpp
--- a/Engine/Source/Runtime/RenderCore/Private/ProfilingDebugging/RealtimeGPUProfiler.cpp
+++ b/Engine/Source/Runtime/RenderCore/Private/ProfilingDebugging/RealtimeGPUProfiler.cpp
@@ -61,7 +61,7 @@ void TDrawEvent<TRHICmdList>::Start(TRHICmdList& InRHICmdList, FColor Color, con
 		va_start(ptr, Fmt);
 		TCHAR TempStr[256];
 		// Build the string in the temp buffer
-		FCString::GetVarArgs(TempStr, ARRAY_COUNT(TempStr), Fmt, ptr);
+		FCString::GetVarArgs(TempStr, UE_ARRAY_COUNT(TempStr), Fmt, ptr);
 		InRHICmdList.PushEvent(TempStr, Color);
 		RHICmdList = &InRHICmdList;
 		va_end(ptr);
@@ -89,7 +89,7 @@ void FDrawEventRHIExecute::Start(IRHIComputeContext& InRHICommandContext, FColor
 		va_start(ptr, Fmt);
 		TCHAR TempStr[256];
 		// Build the string in the temp buffer
-		FCString::GetVarArgs(TempStr, ARRAY_COUNT(TempStr), Fmt, ptr);
+		FCString::GetVarArgs(TempStr, UE_ARRAY_COUNT(TempStr), Fmt, ptr);
 		RHICommandContext = &InRHICommandContext;
 		RHICommandContext->RHIPushEvent(TempStr, Color);
 		va_end(ptr);
diff --git a/Engine/Source/Runtime/RenderCore/Private/RenderGraphEvent.cpp b/Engine/Source/Runtime/RenderCore/Private/RenderGraphEvent.cpp
--- a/Engine/Source/Runtime/RenderCore/Private/RenderGraphEvent.cpp
+++ b/Engine/Source/Runtime/RenderCore/Private/RenderGraphEvent.cpp
@@ -15,7 +15,7 @@ FRDGEventName::FRDGEventName(const TCHAR* InEventFormat, ...)
 		va_start(VAList, InEventFormat);
 		TCHAR TempStr[256];
 		// Build the string in the temp buffer
-		FCString::GetVarArgs(TempStr, ARRAY_COUNT(TempStr), InEventFormat, VAList);
+		FCString::GetVarArgs(TempStr, UE_ARRAY_COUNT(TempStr), InEventFormat, VAList);
 		va_end(VAList);
 
 		FormatedEventName = TempStr;
diff --git a/Engine/Source/Runtime/RenderCore/Private/RenderUtils.cpp b/Engine/Source/Runtime/RenderCore/Private/RenderUtils.cpp
--- a/Engine/Source/Runtime/RenderCore/Private/RenderUtils.cpp
+++ b/Engine/Source/Runtime/RenderCore/Private/RenderUtils.cpp
@@ -153,7 +153,7 @@ static struct FValidatePixelFormats
 {
 	FValidatePixelFormats()
 	{
-		for (int32 X = 0; X < ARRAY_COUNT(GPixelFormats); ++X)
+		for (int32 X = 0; X < UE_ARRAY_COUNT(GPixelFormats); ++X)
 		{
 			// Make sure GPixelFormats has an entry for every unreal format
 			check(X == GPixelFormats[X].UnrealFormat);
@@ -1099,7 +1099,7 @@ public:
 	{
 		TResourceArray<uint16, INDEXBUFFER_ALIGNMENT> Indices;
 		
-		int32 NumIndices = ARRAY_COUNT(GCubeIndices);
+		int32 NumIndices = UE_ARRAY_COUNT(GCubeIndices);
 		Indices.AddUninitialized(NumIndices);
 		FMemory::Memcpy(Indices.GetData(), GCubeIndices, NumIndices * sizeof(uint16));
 
diff --git a/Engine/Source/Runtime/RenderCore/Private/ShaderCore.cpp b/Engine/Source/Runtime/RenderCore/Private/ShaderCore.cpp
--- a/Engine/Source/Runtime/RenderCore/Private/ShaderCore.cpp
+++ b/Engine/Source/Runtime/RenderCore/Private/ShaderCore.cpp
@@ -159,7 +159,7 @@ public:
 
 	FSHAHash* FindHash(EShaderPlatform ShaderPlatform, const FString& VirtualFilePath)
 	{
-		check(ShaderPlatform < ARRAY_COUNT(Platforms));
+		check(ShaderPlatform < UE_ARRAY_COUNT(Platforms));
 		checkf(bInitialized, TEXT("GShaderHashCache::Initialize needs to be called before GShaderHashCache::FindHash."));
 
 		return Platforms[ShaderPlatform].ShaderHashCache.Find(VirtualFilePath);
@@ -167,7 +167,7 @@ public:
 
 	FSHAHash& AddHash(EShaderPlatform ShaderPlatform, const FString& VirtualFilePath)
 	{
-		check(ShaderPlatform < ARRAY_COUNT(Platforms));
+		check(ShaderPlatform < UE_ARRAY_COUNT(Platforms));
 		checkf(bInitialized, TEXT("GShaderHashCache::Initialize needs to be called before GShaderHashCache::AddHash."));
 
 		return Platforms[ShaderPlatform].ShaderHashCache.Add(VirtualFilePath, FSHAHash());
diff --git a/Engine/Source/Runtime/RenderCore/Public/ShaderParameters.h b/Engine/Source/Runtime/RenderCore/Public/ShaderParameters.h
--- a/Engine/Source/Runtime/RenderCore/Public/ShaderParameters.h
+++ b/Engine/Source/Runtime/RenderCore/Public/ShaderParameters.h
@@ -286,31 +286,31 @@ struct FRayTracingShaderBindingsWriter : FRayTracingShaderBindings
 
 	void SetTexture(uint16 BaseIndex, FRHITexture* Value)
 	{
-		checkSlow(BaseIndex < ARRAY_COUNT(Textures));
+		checkSlow(BaseIndex < UE_ARRAY_COUNT(Textures));
 		Textures[BaseIndex] = Value;
 	}
 
 	void SetSRV(uint16 BaseIndex, FRHIShaderResourceView* Value)
 	{
-		checkSlow(BaseIndex < ARRAY_COUNT(SRVs));
+		checkSlow(BaseIndex < UE_ARRAY_COUNT(SRVs));
 		SRVs[BaseIndex] = Value;
 	}
 
 	void SetSampler(uint16 BaseIndex, FRHISamplerState* Value)
 	{
-		checkSlow(BaseIndex < ARRAY_COUNT(Samplers));
+		checkSlow(BaseIndex < UE_ARRAY_COUNT(Samplers));
 		Samplers[BaseIndex] = Value;
 	}
 
 	void SetUAV(uint16 BaseIndex, FRHIUnorderedAccessView* Value)
 	{
-		checkSlow(BaseIndex < ARRAY_COUNT(UAVs));
+		checkSlow(BaseIndex < UE_ARRAY_COUNT(UAVs));
 		UAVs[BaseIndex] = Value;
 	}
 
 	void SetUniformBuffer(uint16 BaseIndex, FRHIUniformBuffer* Value)
 	{
-		checkSlow(BaseIndex < ARRAY_COUNT(UniformBuffers));
+		checkSlow(BaseIndex < UE_ARRAY_COUNT(UniformBuffers));
 		UniformBuffers[BaseIndex] = Value;
 	}
 };
