Build fix VS2019 strict

diff --git a/Engine/Source/Runtime/RenderCore/Private/ShaderCodeLibrary.cpp b/Engine/Source/Runtime/RenderCore/Private/ShaderCodeLibrary.cpp
--- a/Engine/Source/Runtime/RenderCore/Private/ShaderCodeLibrary.cpp
+++ b/Engine/Source/Runtime/RenderCore/Private/ShaderCodeLibrary.cpp
@@ -666,7 +666,7 @@ public:
 				{
 					FScopeLock ScopeLock(&ReadRequestLock);
 					CodeNumRefs = Entry->NumRefs++;
-					bIsLoaded = Entry->LoadedCode.Num();
+					bIsLoaded = Entry->LoadedCode.Num() > 0;
 					if (!bIsLoaded)
 					{
 						Entry->LoadedCode.SetNumUninitialized(Entry->Size);
