Fix use of IoDispatcher in ShaderCodeLibrary

diff --git a/Engine/Source/Runtime/Core/Private/IO/IoDispatcher.cpp b/Engine/Source/Runtime/Core/Private/IO/IoDispatcher.cpp
--- a/Engine/Source/Runtime/Core/Private/IO/IoDispatcher.cpp
+++ b/Engine/Source/Runtime/Core/Private/IO/IoDispatcher.cpp
@@ -668,6 +668,12 @@ FIoDispatcher::GetSizeForChunk(const FIoChunkId& ChunkId) const
 	return Impl->GetSizeForChunk(ChunkId);
 }
 
+bool
+FIoDispatcher::IsInitialized()
+{
+	return GIoDispatcher.IsValid();
+}
+
 FIoStatus
 FIoDispatcher::Initialize(const FString& Directory)
 {
diff --git a/Engine/Source/Runtime/Core/Public/IO/IoDispatcher.h b/Engine/Source/Runtime/Core/Public/IO/IoDispatcher.h
--- a/Engine/Source/Runtime/Core/Public/IO/IoDispatcher.h
+++ b/Engine/Source/Runtime/Core/Public/IO/IoDispatcher.h
@@ -704,6 +704,7 @@ public:
 	FIoDispatcher(const FIoDispatcher&) = default;
 	FIoDispatcher& operator=(const FIoDispatcher&) = delete;
 
+	static CORE_API bool IsInitialized();
 	static CORE_API FIoStatus Initialize(const FString& Directory);
 	static CORE_API void Shutdown();
 	static CORE_API FIoDispatcher& Get();
diff --git a/Engine/Source/Runtime/RenderCore/Private/ShaderCodeLibrary.cpp b/Engine/Source/Runtime/RenderCore/Private/ShaderCodeLibrary.cpp
--- a/Engine/Source/Runtime/RenderCore/Private/ShaderCodeLibrary.cpp
+++ b/Engine/Source/Runtime/RenderCore/Private/ShaderCodeLibrary.cpp
@@ -552,7 +552,10 @@ public:
 			UE_LOG(LogShaderLibrary, Display, TEXT("Using %s for material shader code. Total %d unique shaders."), *DestFilePath, Shaders.Num());
 
 #if RHI_IO_DISPATCHER
-			LibraryChunkId = FIoDispatcher::Get().OpenFileChunk(DestFilePath).ValueOrDie();
+			if (FIoDispatcher::IsInitialized())
+			{
+				LibraryChunkId = FIoDispatcher::Get().OpenFileChunk(DestFilePath).ValueOrDie();
+			}
 #endif
 		}
 	}
