Copying Private-LoadTimes-4.24 stream to Main. Biggest changes are in Materials/Shader memory freezing.

diff --git a/Engine/Source/Runtime/RenderCore/Private/RenderResource.cpp b/Engine/Source/Runtime/RenderCore/Private/RenderResource.cpp
--- a/Engine/Source/Runtime/RenderCore/Private/RenderResource.cpp
+++ b/Engine/Source/Runtime/RenderCore/Private/RenderResource.cpp
@@ -8,6 +8,8 @@
 #include "Misc/ScopedEvent.h"
 #include "Misc/App.h"
 #include "RenderingThread.h"
+#include "ProfilingDebugging/LoadTimeTracker.h"
+
 
 /** Whether to enable mip-level fading or not: +1.0f if enabled, -1.0f if disabled. */
 float GEnableMipLevelFading = 1.0f;
@@ -26,33 +28,28 @@ FAutoConsoleVariableRef CVarReadBufferNumFramesUnusedThresold(
 	GGlobalBufferNumFramesUnusedThresold ,
 	TEXT("Number of frames after which unused global resource allocations will be discarded. Set 0 to ignore. (default=30)"));
 
-TLinkedList<FRenderResource*>*& FRenderResource::GetResourceList()
+FThreadSafeCounter FRenderResource::ResourceListIterationActive;
+
+TArray<int32>& GetFreeIndicesList()
+{
+	static TArray<int32> FreeIndicesList;
+	return FreeIndicesList;
+}
+
+TArray<FRenderResource*>& FRenderResource::GetResourceList()
 {
-	static TLinkedList<FRenderResource*>* FirstResourceLink = NULL;
-	return FirstResourceLink;
+	static TArray<FRenderResource*> RenderResourceList;
+	return RenderResourceList;
 }
 
 /** Initialize all resources initialized before the RHI was initialized */
 void FRenderResource::InitPreRHIResources()
 {	
 	// Notify all initialized FRenderResources that there's a valid RHI device to create their RHI resources for now.
-	for (TLinkedList<FRenderResource*>::TIterator ResourceIt(FRenderResource::GetResourceList()); ResourceIt; ResourceIt.Next())
-	{
-		ResourceIt->InitRHI();
-	}
-	// Dynamic resources can have dependencies on static resources (with uniform buffers) and must initialized last!
-	for (TLinkedList<FRenderResource*>::TIterator ResourceIt(FRenderResource::GetResourceList()); ResourceIt; ResourceIt.Next())
-	{
-		ResourceIt->InitDynamicRHI();
-	}
+	FRenderResource::InitRHIForAllResources();
 
 #if !PLATFORM_NEEDS_RHIRESOURCELIST
-	while (GetResourceList())
-	{
-		TLinkedList<FRenderResource*>* CurrentHead = GetResourceList();
-		CurrentHead->Unlink();
-		delete CurrentHead;
-	}
+	FRenderResource::GetResourceList().Empty();
 #endif
 }
 
@@ -61,10 +58,8 @@ void FRenderResource::ChangeFeatureLevel(ERHIFeatureLevel::Type NewFeatureLevel)
 	ENQUEUE_RENDER_COMMAND(FRenderResourceChangeFeatureLevel)(
 		[NewFeatureLevel](FRHICommandList& RHICmdList)
 	{
-		for (TLinkedList<FRenderResource*>::TIterator It(FRenderResource::GetResourceList()); It; It.Next())
+		FRenderResource::ForAllResources([NewFeatureLevel](FRenderResource* Resource)
 		{
-			FRenderResource* Resource = *It;
-
 			// Only resources configured for a specific feature level need to be updated
 			if (Resource->HasValidFeatureLevel() && (Resource->FeatureLevel != NewFeatureLevel))
 			{
@@ -74,34 +69,43 @@ void FRenderResource::ChangeFeatureLevel(ERHIFeatureLevel::Type NewFeatureLevel)
 				Resource->InitDynamicRHI();
 				Resource->InitRHI();
 			}
-		}
+		});
 	});
 }
 
 void FRenderResource::InitResource()
 {
 	check(IsInRenderingThread());
-	if(!bInitialized)
+	if (ListIndex == INDEX_NONE)
 	{
-#if PLATFORM_NEEDS_RHIRESOURCELIST
-		ResourceLink = TLinkedList<FRenderResource*>(this);
-		ResourceLink.LinkHead(GetResourceList());
-#endif
-		if(GIsRHIInitialized)
+		check(!IsEngineExitRequested());
+
+		TArray<FRenderResource*>& ResourceList = GetResourceList();
+		TArray<int32>& FreeIndicesList = GetFreeIndicesList();
+
+		// If resource list is currently being iterated, new resources must be added to the end of the list, to ensure they're processed during the iteration
+		// Otherwise empty slots in the list may be re-used for new resources
+		int32 LocalListIndex = INDEX_NONE;
+		if (FreeIndicesList.Num() > 0 && ResourceListIterationActive.GetValue() == 0)
 		{
-			CSV_SCOPED_TIMING_STAT_EXCLUSIVE(InitRenderResource);
-			InitDynamicRHI();
-			InitRHI();
+			LocalListIndex = FreeIndicesList.Pop();
+			check(ResourceList[LocalListIndex] == nullptr);
+			ResourceList[LocalListIndex] = this;
 		}
 		else
 		{
-#if !PLATFORM_NEEDS_RHIRESOURCELIST
-		TLinkedList<FRenderResource*>* ListEntry = new TLinkedList<FRenderResource*>(this);
-		ListEntry->LinkHead(GetResourceList());
-#endif
+			LocalListIndex = ResourceList.Add(this);
 		}
-		FPlatformMisc::MemoryBarrier(); // there are some multithreaded reads of bInitialized
-		bInitialized = true;
+
+		if (GIsRHIInitialized)
+		{
+			CSV_SCOPED_TIMING_STAT_EXCLUSIVE(InitRenderResource);
+			InitDynamicRHI();
+			InitRHI();
+		}
+
+		FPlatformMisc::MemoryBarrier(); // there are some multithreaded reads of ListIndex
+		ListIndex = LocalListIndex;
 	}
 }
 
@@ -110,17 +114,19 @@ void FRenderResource::ReleaseResource()
 	if ( !GIsCriticalError )
 	{
 		check(IsInRenderingThread());
-		if(bInitialized)
+		if(ListIndex != INDEX_NONE)
 		{
 			if(GIsRHIInitialized)
 			{
 				ReleaseRHI();
 				ReleaseDynamicRHI();
 			}
-#if PLATFORM_NEEDS_RHIRESOURCELIST
-			ResourceLink.Unlink();
-#endif
-			bInitialized = false;
+
+			TArray<FRenderResource*>& ResourceList = GetResourceList();
+			TArray<int32>& FreeIndicesList = GetFreeIndicesList();
+			ResourceList[ListIndex] = nullptr;
+			FreeIndicesList.Add(ListIndex);
+			ListIndex = INDEX_NONE;
 		}
 	}
 }
@@ -128,7 +134,7 @@ void FRenderResource::ReleaseResource()
 void FRenderResource::UpdateRHI()
 {
 	check(IsInRenderingThread());
-	if(bInitialized && GIsRHIInitialized)
+	if(IsInitialized() && GIsRHIInitialized)
 	{
 		ReleaseRHI();
 		ReleaseDynamicRHI();
@@ -137,57 +143,9 @@ void FRenderResource::UpdateRHI()
 	}
 }
 
-void FRenderResource::InitResourceFromPossiblyParallelRendering()
-{
-	check(IsInParallelRenderingThread());
-
-	if (IsInRenderingThread())
-	{
-		InitResource();
-	}
-	else
-	{
-		class FInitResourceRenderThreadTask
-		{
-			FRenderResource& Resource;
-			FScopedEvent& Event;
-		public:
-
-			FInitResourceRenderThreadTask(FRenderResource& InResource, FScopedEvent& InEvent)
-				: Resource(InResource)
-				, Event(InEvent)
-			{
-			}
-
-			static FORCEINLINE TStatId GetStatId()
-			{
-				RETURN_QUICK_DECLARE_CYCLE_STAT(FInitResourceRenderThreadTask, STATGROUP_TaskGraphTasks);
-			}
-
-			static FORCEINLINE ENamedThreads::Type GetDesiredThread()
-			{
-				return ENamedThreads::GetRenderThread_Local();
-			}
-
-			static FORCEINLINE ESubsequentsMode::Type GetSubsequentsMode() { return ESubsequentsMode::FireAndForget; }
-
-			void DoTask(ENamedThreads::Type CurrentThread, const FGraphEventRef& MyCompletionGraphEvent)
-			{
-				Resource.InitResource();
-				Event.Trigger();
-			}
-		};
-		{
-			FScopedEvent Event;
-			TGraphTask<FInitResourceRenderThreadTask>::CreateTask().ConstructAndDispatchWhenReady(*this, Event);
-		}
-	}
-}
-
-
 FRenderResource::~FRenderResource()
 {
-	if (bInitialized && !GIsCriticalError)
+	if (IsInitialized() && !GIsCriticalError)
 	{
 		// Deleting an initialized FRenderResource will result in a crash later since it is still linked
 		UE_LOG(LogRendererCore, Fatal,TEXT("A FRenderResource was deleted without being released first!"));
@@ -334,6 +292,7 @@ void FTextureReference::InvalidateLastRenderTime()
 
 void FTextureReference::InitRHI()
 {
+	SCOPED_LOADTIMER(FTextureReference_InitRHI);
 	TextureReferenceRHI = RHICreateTextureReference(&LastRenderTimeRHI);
 }
 	
diff --git a/Engine/Source/Runtime/RenderCore/Private/RenderTargetPool.cpp b/Engine/Source/Runtime/RenderCore/Private/RenderTargetPool.cpp
--- a/Engine/Source/Runtime/RenderCore/Private/RenderTargetPool.cpp
+++ b/Engine/Source/Runtime/RenderCore/Private/RenderTargetPool.cpp
@@ -1218,26 +1218,26 @@ uint32 FPooledRenderTarget::ComputeMemorySize() const
 	{
 		if (Desc.Is2DTexture())
 		{
-			Size += RHIComputeMemorySize((const FTexture2DRHIRef&)RenderTargetItem.TargetableTexture);
+			Size += RHIComputeMemorySize(RenderTargetItem.TargetableTexture);
 			if (RenderTargetItem.ShaderResourceTexture != RenderTargetItem.TargetableTexture)
 			{
-				Size += RHIComputeMemorySize((const FTexture2DRHIRef&)RenderTargetItem.ShaderResourceTexture);
+				Size += RHIComputeMemorySize(RenderTargetItem.ShaderResourceTexture);
 			}
 		}
 		else if (Desc.Is3DTexture())
 		{
-			Size += RHIComputeMemorySize((const FTexture3DRHIRef&)RenderTargetItem.TargetableTexture);
+			Size += RHIComputeMemorySize(RenderTargetItem.TargetableTexture);
 			if (RenderTargetItem.ShaderResourceTexture != RenderTargetItem.TargetableTexture)
 			{
-				Size += RHIComputeMemorySize((const FTexture3DRHIRef&)RenderTargetItem.ShaderResourceTexture);
+				Size += RHIComputeMemorySize(RenderTargetItem.ShaderResourceTexture);
 			}
 		}
 		else
 		{
-			Size += RHIComputeMemorySize((const FTextureCubeRHIRef&)RenderTargetItem.TargetableTexture);
+			Size += RHIComputeMemorySize(RenderTargetItem.TargetableTexture);
 			if (RenderTargetItem.ShaderResourceTexture != RenderTargetItem.TargetableTexture)
 			{
-				Size += RHIComputeMemorySize((const FTextureCubeRHIRef&)RenderTargetItem.ShaderResourceTexture);
+				Size += RHIComputeMemorySize(RenderTargetItem.ShaderResourceTexture);
 			}
 		}
 	}
diff --git a/Engine/Source/Runtime/RenderCore/Private/RenderUtils.cpp b/Engine/Source/Runtime/RenderCore/Private/RenderUtils.cpp
--- a/Engine/Source/Runtime/RenderCore/Private/RenderUtils.cpp
+++ b/Engine/Source/Runtime/RenderCore/Private/RenderUtils.cpp
@@ -1008,27 +1008,27 @@ RENDERCORE_API FVertexDeclarationRHIRef& GetVertexDeclarationFVector2()
 	return GVector2VertexDeclaration.VertexDeclarationRHI;
 }
 
-RENDERCORE_API bool PlatformSupportsSimpleForwardShading(EShaderPlatform Platform)
+RENDERCORE_API bool PlatformSupportsSimpleForwardShading(const FStaticShaderPlatform Platform)
 {
 	static const auto SupportSimpleForwardShadingCVar = IConsoleManager::Get().FindTConsoleVariableDataInt(TEXT("r.SupportSimpleForwardShading"));
 	// Scalability feature only needed / used on PC
 	return IsPCPlatform(Platform) && SupportSimpleForwardShadingCVar->GetValueOnAnyThread() != 0;
 }
 
-RENDERCORE_API bool IsSimpleForwardShadingEnabled(EShaderPlatform Platform)
+RENDERCORE_API bool IsSimpleForwardShadingEnabled(const FStaticShaderPlatform Platform)
 {
 	static const auto CVar = IConsoleManager::Get().FindTConsoleVariableDataInt(TEXT("r.SimpleForwardShading"));
 	return CVar->GetValueOnAnyThread() != 0 && PlatformSupportsSimpleForwardShading(Platform);
 }
 
-RENDERCORE_API bool MobileSupportsGPUScene(EShaderPlatform Platform)
+RENDERCORE_API bool MobileSupportsGPUScene(const FStaticShaderPlatform Platform)
 {
 	// make it shader platform setting?
 	static TConsoleVariableData<int32>* CVar = IConsoleManager::Get().FindTConsoleVariableDataInt(TEXT("r.Mobile.SupportGPUScene"));
 	return (CVar && CVar->GetValueOnAnyThread() != 0) ? true : false;
 }
 
-RENDERCORE_API bool GPUSceneUseTexture2D(EShaderPlatform Platform)
+RENDERCORE_API bool GPUSceneUseTexture2D(const FStaticShaderPlatform Platform)
 {
 	if (IsMobilePlatform(Platform))
 	{
@@ -1045,7 +1045,7 @@ RENDERCORE_API bool GPUSceneUseTexture2D(EShaderPlatform Platform)
 	return false;
 }
 
-RENDERCORE_API bool MaskedInEarlyPass(EShaderPlatform Platform)
+RENDERCORE_API bool MaskedInEarlyPass(const FStaticShaderPlatform Platform)
 {
 	static IConsoleVariable* CVarMobileEarlyZPassOnlyMaterialMasking = IConsoleManager::Get().FindConsoleVariable(TEXT("r.Mobile.EarlyZPassOnlyMaterialMasking"));
 	static IConsoleVariable* CVarEarlyZPassOnlyMaterialMasking = IConsoleManager::Get().FindConsoleVariable(TEXT("r.EarlyZPassOnlyMaterialMasking"));
@@ -1059,7 +1059,7 @@ RENDERCORE_API bool MaskedInEarlyPass(EShaderPlatform Platform)
 	}
 }
 
-RENDERCORE_API bool AllowPixelDepthOffset(EShaderPlatform Platform)
+RENDERCORE_API bool AllowPixelDepthOffset(const FStaticShaderPlatform Platform)
 {
 	if (IsMobilePlatform(Platform))
 	{
@@ -1284,7 +1284,7 @@ RENDERCORE_API void QuantizeSceneBufferSize(const FIntPoint& InBufferSize, FIntP
 	OutBufferSize.Y = (InBufferSize.Y + DividableBy - 1) & Mask;
 }
 
-RENDERCORE_API bool UseVirtualTexturing(ERHIFeatureLevel::Type InFeatureLevel, const ITargetPlatform* TargetPlatform)
+RENDERCORE_API bool UseVirtualTexturing(const FStaticFeatureLevel InFeatureLevel, const ITargetPlatform* TargetPlatform)
 {
 #if !PLATFORM_SUPPORTS_VIRTUAL_TEXTURE_STREAMING
 	if (GIsEditor == false)
diff --git a/Engine/Source/Runtime/RenderCore/Private/ResolveShader.cpp b/Engine/Source/Runtime/RenderCore/Private/ResolveShader.cpp
--- a/Engine/Source/Runtime/RenderCore/Private/ResolveShader.cpp
+++ b/Engine/Source/Runtime/RenderCore/Private/ResolveShader.cpp
@@ -14,7 +14,7 @@ IMPLEMENT_SHADER_TYPE(, FResolveVS, TEXT("/Engine/Private/ResolveVertexShader.us
 
 void FResolveSingleSamplePS::SetParameters(FRHICommandList& RHICmdList, uint32 SingleSampleIndexValue)
 {
-	SetShaderValue(RHICmdList, GetPixelShader(),SingleSampleIndex,SingleSampleIndexValue);
+	SetShaderValue(RHICmdList, RHICmdList.GetBoundPixelShader(),SingleSampleIndex,SingleSampleIndexValue);
 }
 
 void FResolveVS::SetParameters(FRHICommandList& RHICmdList, const FResolveRect& SrcBounds, const FResolveRect& DstBounds, uint32 DstSurfaceWidth, uint32 DstSurfaceHeight)
@@ -29,6 +29,6 @@ void FResolveVS::SetParameters(FRHICommandList& RHICmdList, const FResolveRect&
 	const float MaxX = -1.f + DstBounds.X2 / ((float)DstSurfaceWidth * 0.5f);
 	const float MaxY = +1.f - DstBounds.Y2 / ((float)DstSurfaceHeight * 0.5f);
 
-	SetShaderValue(RHICmdList, GetVertexShader(), PositionMinMax, FVector4(MinX, MinY, MaxX, MaxY));
-	SetShaderValue(RHICmdList, GetVertexShader(), UVMinMax, FVector4(MinU, MinV, MaxU, MaxV));
+	SetShaderValue(RHICmdList, RHICmdList.GetBoundVertexShader(), PositionMinMax, FVector4(MinX, MinY, MaxX, MaxY));
+	SetShaderValue(RHICmdList, RHICmdList.GetBoundVertexShader(), UVMinMax, FVector4(MinU, MinV, MaxU, MaxV));
 }
diff --git a/Engine/Source/Runtime/RenderCore/Private/Shader.cpp b/Engine/Source/Runtime/RenderCore/Private/Shader.cpp
--- a/Engine/Source/Runtime/RenderCore/Private/Shader.cpp
+++ b/Engine/Source/Runtime/RenderCore/Private/Shader.cpp
@@ -21,6 +21,7 @@
 #include "UObject/RenderingObjectVersion.h"
 #include "UObject/FortniteMainBranchObjectVersion.h"
 #include "Misc/ScopeRWLock.h"
+#include "ProfilingDebugging/LoadTimeTracker.h"
 
 #if WITH_EDITORONLY_DATA
 #include "Interfaces/IShaderFormat.h"
@@ -28,6 +29,18 @@
 
 DEFINE_LOG_CATEGORY(LogShaders);
 
+IMPLEMENT_TYPE_LAYOUT(FShader);
+IMPLEMENT_TYPE_LAYOUT(FShaderParameterBindings);
+IMPLEMENT_TYPE_LAYOUT(FShaderMapContent);
+IMPLEMENT_TYPE_LAYOUT(FShaderTypeDependency);
+IMPLEMENT_TYPE_LAYOUT(FShaderPipeline);
+IMPLEMENT_TYPE_LAYOUT(FShaderParameterInfo);
+IMPLEMENT_TYPE_LAYOUT(FShaderLooseParameterBufferInfo);
+IMPLEMENT_TYPE_LAYOUT(FShaderParameterMapInfo);
+
+IMPLEMENT_EXPORTED_INTRINSIC_TYPE_LAYOUT(TIndexedPtr<FShaderType>);
+IMPLEMENT_EXPORTED_INTRINSIC_TYPE_LAYOUT(TIndexedPtr<FVertexFactoryType>);
+
 RENDERCORE_API bool UsePreExposure(EShaderPlatform Platform)
 {
 	// Mobile platforms are excluded because they use a different pre-exposure logic in MobileBasePassPixelShader.usf
@@ -35,8 +48,6 @@ RENDERCORE_API bool UsePreExposure(EShaderPlatform Platform)
 	return CVarUsePreExposure->GetValueOnAnyThread() != 0 && !IsMobilePlatform(Platform) && IsFeatureLevelSupported(Platform, ERHIFeatureLevel::SM5);
 }
 
-static const FName ShaderCompressionFormat = NAME_Zlib;
-
 static TAutoConsoleVariable<int32> CVarUsePipelines(
 	TEXT("r.ShaderPipelines"),
 	1,
@@ -51,7 +62,6 @@ static TAutoConsoleVariable<int32> CVarSkipShaderCompression(
 
 static TLinkedList<FShaderType*>*			GShaderTypeList = nullptr;
 static TLinkedList<FShaderPipelineType*>*	GShaderPipelineList = nullptr;
-static TMap<FName, FShaderType*>*			GShaderNameToTypeMap = nullptr;
 
 static FSHAHash ShaderSourceDefaultHash; //will only be read (never written) for the cooking case
 
@@ -146,27 +156,48 @@ void FShaderParameterMap::UpdateHash(FSHA1& HashState) const
 
 bool FShaderType::bInitializedSerializationHistory = false;
 
+static TArray<FShaderType*>& GetSortedShaderTypes(FShaderType::EShaderTypeForDynamicCast Type)
+{
+	static TArray<FShaderType*> SortedTypes[(uint32)FShaderType::EShaderTypeForDynamicCast::NumShaderTypes];
+	return SortedTypes[(uint32)Type];
+}
+
 FShaderType::FShaderType(
 	EShaderTypeForDynamicCast InShaderTypeForDynamicCast,
+	FTypeLayoutDesc& InTypeLayout,
 	const TCHAR* InName,
 	const TCHAR* InSourceFilename,
 	const TCHAR* InFunctionName,
 	uint32 InFrequency,
 	int32 InTotalPermutationCount,
 	ConstructSerializedType InConstructSerializedRef,
+	ConstructCompiledType InConstructCompiledRef,
+	ModifyCompilationEnvironmentType InModifyCompilationEnvironmentRef,
+	ShouldCompilePermutationType InShouldCompilePermutationRef,
+	ValidateCompiledResultType InValidateCompiledResultRef,
+	uint32 InTypeSize,
 	const FShaderParametersMetadata* InRootParametersMetadata
 	):
 	ShaderTypeForDynamicCast(InShaderTypeForDynamicCast),
+	TypeLayout(&InTypeLayout),
 	Name(InName),
 	TypeName(InName),
+	HashedName(TypeName),
 	SourceFilename(InSourceFilename),
 	FunctionName(InFunctionName),
 	Frequency(InFrequency),
+	TypeSize(InTypeSize),
 	TotalPermutationCount(InTotalPermutationCount),
 	ConstructSerializedRef(InConstructSerializedRef),
+	ConstructCompiledRef(InConstructCompiledRef),
+	ModifyCompilationEnvironmentRef(InModifyCompilationEnvironmentRef),
+	ShouldCompilePermutationRef(InShouldCompilePermutationRef),
+	ValidateCompiledResultRef(InValidateCompiledResultRef),
 	RootParametersMetadata(InRootParametersMetadata),
 	GlobalListLink(this)
 {
+	FTypeLayoutDesc::Register(InTypeLayout);
+
 	bCachedUniformBufferStructDeclarations = false;
 
 	// This will trigger if an IMPLEMENT_SHADER_TYPE was in a module not loaded before InitializeShaderTypes
@@ -181,17 +212,22 @@ FShaderType::FShaderType(
 
 	// register this shader type
 	GlobalListLink.LinkHead(GetTypeList());
-	GetNameToTypeMap().Add(TypeName, this);
+	GetNameToTypeMap().Add(HashedName, this);
 
-	// Assign the shader type the next unassigned hash index.
-	static uint32 NextHashIndex = 0;
-	HashIndex = NextHashIndex++;
+	TArray<FShaderType*>& SortedTypes = GetSortedShaderTypes(InShaderTypeForDynamicCast);
+	const int32 SortedIndex = Algo::LowerBoundBy(SortedTypes, HashedName, [](const FShaderType* InType) { return InType->GetHashedName(); });
+	SortedTypes.Insert(this, SortedIndex);
 }
 
 FShaderType::~FShaderType()
 {
 	GlobalListLink.Unlink();
-	GetNameToTypeMap().Remove(TypeName);
+	GetNameToTypeMap().Remove(HashedName);
+
+	TArray<FShaderType*>& SortedTypes = GetSortedShaderTypes(ShaderTypeForDynamicCast);
+	const int32 SortedIndex = Algo::BinarySearchBy(SortedTypes, HashedName, [](const FShaderType* InType) { return InType->GetHashedName(); });
+	check(SortedIndex != INDEX_NONE);
+	SortedTypes.RemoveAt(SortedIndex);
 }
 
 TLinkedList<FShaderType*>*& FShaderType::GetTypeList()
@@ -214,13 +250,12 @@ FShaderType* FShaderType::GetShaderTypeByName(const TCHAR* Name)
 	return nullptr;
 }
 
-TArray<FShaderType*> FShaderType::GetShaderTypesByFilename(const TCHAR* Filename)
+TArray<const FShaderType*> FShaderType::GetShaderTypesByFilename(const TCHAR* Filename)
 {
-	TArray<FShaderType*> OutShaders;
+	TArray<const FShaderType*> OutShaders;
 	for(TLinkedList<FShaderType*>::TIterator It(GetTypeList()); It; It.Next())
 	{
-		FShaderType* Type = *It;
-
+		const FShaderType* Type = *It;
 		if (FPlatformString::Strcmp(Filename, Type->GetShaderFilename()) == 0)
 		{
 			OutShaders.Add(Type);
@@ -229,64 +264,17 @@ TArray<FShaderType*> FShaderType::GetShaderTypesByFilename(const TCHAR* Filename
 	return OutShaders;
 }
 
-TMap<FName, FShaderType*>& FShaderType::GetNameToTypeMap()
+TMap<FHashedName, FShaderType*>& FShaderType::GetNameToTypeMap()
 {
-	if(!GShaderNameToTypeMap)
-	{
-		GShaderNameToTypeMap = new TMap<FName, FShaderType*>();
-	}
-	return *GShaderNameToTypeMap;
+	static TMap<FHashedName, FShaderType*> ShaderNameToTypeMap;
+	return ShaderNameToTypeMap;
 }
 
-inline bool FShaderType::GetOutdatedCurrentType(TArray<FShaderType*>& OutdatedShaderTypes, TArray<const FVertexFactoryType*>& OutdatedFactoryTypes) const
+const TArray<FShaderType*>& FShaderType::GetSortedTypes(EShaderTypeForDynamicCast Type)
 {
-	bool bOutdated = false;
-#if WITH_EDITOR
-	for (TMap<FShaderKey, FShader*>::TConstIterator ShaderIt(ShaderIdMap);ShaderIt;++ShaderIt)
-	{
-		FShader* Shader = ShaderIt.Value();
-		const FVertexFactoryParameterRef* VFParameterRef = Shader->GetVertexFactoryParameterRef();
-		const FSHAHash& SavedHash = Shader->GetHash();
-		const FSHAHash& CurrentHash = GetSourceHash(Shader->GetShaderPlatform());
-		const bool bOutdatedShader = SavedHash != CurrentHash;
-		const bool bOutdatedVertexFactory =
-			VFParameterRef && VFParameterRef->GetVertexFactoryType() && VFParameterRef->GetVertexFactoryType()->GetSourceHash(VFParameterRef->GetShaderPlatform()) != VFParameterRef->GetHash();
-
-		if (bOutdatedShader)
-		{
-			OutdatedShaderTypes.AddUnique(Shader->Type);
-			bOutdated = true;
-		}
-
-		if (bOutdatedVertexFactory)
-		{
-			OutdatedFactoryTypes.AddUnique(VFParameterRef->GetVertexFactoryType());
-			bOutdated = true;
-		}
-	}
-#endif // WITH_EDITOR
-	return bOutdated;
+	return GetSortedShaderTypes(Type);
 }
 
-void FShaderType::GetOutdatedTypes(TArray<FShaderType*>& OutdatedShaderTypes, TArray<const FVertexFactoryType*>& OutdatedFactoryTypes)
-{
-	for(TLinkedList<FShaderType*>::TIterator It(GetTypeList()); It; It.Next())
-	{
-		FShaderType* Type = *It;
-		Type->GetOutdatedCurrentType(OutdatedShaderTypes, OutdatedFactoryTypes);
-	}
-
-	for (int32 TypeIndex = 0; TypeIndex < OutdatedShaderTypes.Num(); TypeIndex++)
-	{
-		UE_LOG(LogShaders, Warning, TEXT("		Recompiling %s"), OutdatedShaderTypes[TypeIndex]->GetName());
-	}
-	for (int32 TypeIndex = 0; TypeIndex < OutdatedFactoryTypes.Num(); TypeIndex++)
-	{
-		UE_LOG(LogShaders, Warning, TEXT("		Recompiling %s"), OutdatedFactoryTypes[TypeIndex]->GetName());
-	}
-}
-
-
 FArchive& operator<<(FArchive& Ar,FShaderType*& Ref)
 {
 	if(Ar.IsSaving())
@@ -319,18 +307,29 @@ FArchive& operator<<(FArchive& Ar,FShaderType*& Ref)
 	return Ar;
 }
 
+FShader* FShaderType::ConstructForDeserialization() const
+{
+	return (*ConstructSerializedRef)();
+}
 
-FShader* FShaderType::FindShaderByKey(const FShaderKey& Key)
+FShader* FShaderType::ConstructCompiled(const FShader::CompiledShaderInitializerType& Initializer) const
 {
-	check(IsInGameThread());
-	FShader* Result = ShaderIdMap.FindRef(Key);
-	checkSlow(!Result || Result->GetKey() == Key);
-	return Result;
+	return (*ConstructCompiledRef)(Initializer);
 }
 
-FShader* FShaderType::ConstructForDeserialization() const
+bool FShaderType::ShouldCompilePermutation(const FShaderPermutationParameters& Parameters) const
 {
-	return (*ConstructSerializedRef)();
+	return (*ShouldCompilePermutationRef)(Parameters);
+}
+
+void FShaderType::ModifyCompilationEnvironment(const FShaderPermutationParameters& Parameters, FShaderCompilerEnvironment& OutEnvironment) const
+{
+	(*ModifyCompilationEnvironmentRef)(Parameters, OutEnvironment);
+}
+
+bool FShaderType::ValidateCompiledResult(EShaderPlatform Platform, const FShaderParameterMap& ParameterMap, TArray<FString>& OutError) const
+{
+	return (*ValidateCompiledResultRef)(Platform, ParameterMap, OutError);
 }
 
 const FSHAHash& FShaderType::GetSourceHash(EShaderPlatform ShaderPlatform) const
@@ -353,24 +352,6 @@ void FShaderType::Initialize(const TMap<FString, TArray<const TCHAR*> >& ShaderF
 			UniqueShaderTypes.Add(Type);
 #endif
 			GenerateReferencedUniformBuffers(Type->SourceFilename, Type->Name, ShaderFileToUniformBufferVariables, Type->ReferencedUniformBufferStructsCache);
-
-			// Cache serialization history for each shader type
-			// This history is used to detect when shader serialization changes without a corresponding .usf change
-			{
-				// Construct a temporary shader, which is initialized to safe values for serialization
-				FShader* TempShader = Type->ConstructForDeserialization();
-				check(TempShader != NULL);
-				TempShader->Type = Type;
-
-				// Serialize the temp shader to memory and record the number and sizes of serializations
-				TArray<uint8> TempData;
-				FMemoryWriter Ar(TempData, true);
-				FShaderSaveArchive SaveArchive(Ar, Type->SerializationHistory);
-				TempShader->SerializeBase(SaveArchive, false, false);
-
-				// Destroy the temporary shader
-				delete TempShader;
-			}
 		}
 	
 #if UE_BUILD_DEBUG
@@ -388,838 +369,75 @@ void FShaderType::Initialize(const TMap<FString, TArray<const TCHAR*> >& ShaderF
 
 void FShaderType::Uninitialize()
 {
-	for(TLinkedList<FShaderType*>::TIterator It(FShaderType::GetTypeList()); It; It.Next())
-	{
-		FShaderType* Type = *It;
-		Type->SerializationHistory = FSerializationHistory();
-	}
-
 	bInitializedSerializationHistory = false;
 }
 
-struct FShaderResourceKeyFuncs
-{
-	typedef FShaderResource* KeyType;
-	typedef FShaderResource* KeyInitType;
-	typedef FShaderResource* ElementInitType;
-
-	enum { bAllowDuplicateKeys = false };
-
-	static KeyInitType GetSetKey(FShaderResource* Element)
-	{
-		return Element;
-	}
-
-	static bool Matches(FShaderResource* A, FShaderResource* B)
-	{
-		return A == B || A->GetId() == B->GetId();
-	}
-
-	static bool Matches(FShaderResource* Resource, const FShaderResourceId& Id)
-	{
-		return Resource->GetId() == Id;
-	}
-
-	static uint32 GetKeyHash(FShaderResource* Key)
-	{
-		return GetTypeHash(Key->GetOutputHash());
-	}
-
-	static uint32 GetKeyHash(const FShaderResourceId& Key)
-	{
-		return GetTypeHash(Key.OutputHash);
-	}
-};
-
-
-class FShaderResourceSet 
-{
-public:
-	TRefCountPtr<FShaderResource> Find(const FShaderResourceId& Id) const
-	{
-		uint32 IdHash = FShaderResourceKeyFuncs::GetKeyHash(Id);
-
-		FRWScopeLock Locker(Lock, SLT_ReadOnly);
-		
-		if (FShaderResource*const* ResourceIt = Set.FindByHash(IdHash, Id))
-		{
-			return TRefCountPtr<FShaderResource>(*ResourceIt);
-		}
-
-		return TRefCountPtr<FShaderResource>();
-	}
-
-	/** 
-	 * Add resource and return it as smart pointer unless an identical resource exists.
-	 * In that case the resource is deleted and the existing resource is returned.
-	 */
-	TRefCountPtr<FShaderResource> AddOrDelete(FShaderResource* Resource)
-	{
-		TRefCountPtr<FShaderResource> OutResource(Resource);
-
-		{
-			FRWScopeLock Locker(Lock, SLT_Write);
-		
-			bool bAlreadyInSet;
-			FSetElementId ElementId = Set.Add(Resource, &bAlreadyInSet);
-			if (bAlreadyInSet)
-			{
-				return TRefCountPtr<FShaderResource>(Set[ElementId]);
-			}
-		}
-
-		return OutResource;
-	}
-
-	/** Returns true if resource was removed before someone created a new reference to it */
-	bool TryRemove(FShaderResource& Resource)
-	{
-		FRWScopeLock Locker(Lock, SLT_Write);
-
-		if (Resource.NumRefs == 0)
-		{
-			Set.Remove(&Resource);
-			return true;
-		}
-
-		// Another thread found the resource after ref-count was decremented to zero
-		return false;
-	}
-
-	void GetAllIds(TArray<FShaderResourceId>& Out) const
-	{
-		FRWScopeLock Locker(Lock, SLT_ReadOnly);
-
-		Out.Reserve(Set.Num());
-		for (FShaderResource* Resource : Set)
-		{
-			Out.Add(Resource->GetId());
-			Out.Last().MakeSelfContained();			
-		}
-	}
-
-private:
-	mutable FRWLock Lock;
-	TSet<FShaderResource*, FShaderResourceKeyFuncs> Set;
-};
-
-static FShaderResourceSet gLoadedShaderResources;
-#if RHI_RAYTRACING
-TArray<uint32> FShaderResource::GlobalUnusedIndicies;
-TArray<FRHIRayTracingShader*> FShaderResource::GlobalRayTracingMaterialLibrary;
-FCriticalSection FShaderResource::GlobalRayTracingMaterialLibraryCS;
-
-void FShaderResource::GetRayTracingMaterialLibrary(TArray<FRHIRayTracingShader*>& RayTracingMaterials, FRHIRayTracingShader* DefaultShader)
-{
-	FScopeLock Lock(&GlobalRayTracingMaterialLibraryCS);
-	RayTracingMaterials = GlobalRayTracingMaterialLibrary;
-
-	for (uint32 Index : GlobalUnusedIndicies)
-	{
-		RayTracingMaterials[Index] = DefaultShader;
-	}
-}
-
-uint32 FShaderResource::AddToRayTracingLibrary(FRHIRayTracingShader* Shader)
-{
-	FScopeLock Lock(&GlobalRayTracingMaterialLibraryCS);
-
-	if (GlobalUnusedIndicies.Num() != 0)
-	{
-		uint32 Index = GlobalUnusedIndicies.Pop(false);
-		checkSlow(GlobalRayTracingMaterialLibrary[Index] == nullptr);
-		GlobalRayTracingMaterialLibrary[Index] = Shader;
-		return Index;
-	}
-	else
-	{
-		GlobalRayTracingMaterialLibrary.Add(Shader);
-		return GlobalRayTracingMaterialLibrary.Num() - 1;
-	}
-}
-
-void FShaderResource::RemoveFromRayTracingLibrary(uint32 Index)
-{
-	FScopeLock Lock(&GlobalRayTracingMaterialLibraryCS);
-	GlobalUnusedIndicies.Push(Index);
-	GlobalRayTracingMaterialLibrary[Index] = nullptr;
-}
-#endif
-
-FShaderResource::FShaderResource()
-	: SpecificPermutationId(0)
-	, NumRefs(0)
-	, NumInstructions(0)
-#if WITH_EDITORONLY_DATA
-	, NumTextureSamplers(0)
-#endif
-	, bCodeInSharedLocation(false)
-	, bCodeInSharedLocationRequested(false)
-{
-	INC_DWORD_STAT_BY(STAT_Shaders_NumShaderResourcesLoaded, 1);
-}
-
-//#todo-RemoveStreamOut
-FShaderResource::FShaderResource(const FShaderCompilerOutput& Output, int32 InSpecificPermutationId) 
-	: SpecificPermutationId(InSpecificPermutationId)
-	, NumRefs(0)
-	, NumInstructions(Output.NumInstructions)
-#if WITH_EDITORONLY_DATA
-	, NumTextureSamplers(Output.NumTextureSamplers)
-#endif
-	, bCodeInSharedLocation(false)
-	, bCodeInSharedLocationRequested(false)
-	
-{
-	BuildParameterMapInfo(Output.ParameterMap.GetParameterMap());
-
-	check(!(SpecificPermutationId != 0));
-
-	Target = Output.Target;
-	CompressCode(Output.ShaderCode.GetReadAccess());
-
-	check(Code.Num() > 0);
-
-	OutputHash = Output.OutputHash;
-	checkSlow(OutputHash != FSHAHash());
-
-#if WITH_EDITORONLY_DATA
-	PlatformDebugData = Output.PlatformDebugData;
-#endif
-		
-	INC_DWORD_STAT_BY_FName(GetMemoryStatType((EShaderFrequency)Target.Frequency).GetName(), Code.Num());
-	INC_DWORD_STAT_BY(STAT_Shaders_ShaderResourceMemory, GetSizeBytes());
-	INC_DWORD_STAT_BY(STAT_Shaders_NumShaderResourcesLoaded, 1);
-}
-
-
-FShaderResource::~FShaderResource()
-{
-	check(NumRefs == 0);
-
-	DEC_DWORD_STAT_BY_FName(GetMemoryStatType((EShaderFrequency)Target.Frequency).GetName(), Code.Num());
-	DEC_DWORD_STAT_BY(STAT_Shaders_ShaderResourceMemory, GetSizeBytes());
-	DEC_DWORD_STAT_BY(STAT_Shaders_NumShaderResourcesLoaded, 1);
-}
-
-void FShaderResource::BuildParameterMapInfo(const TMap<FString, FParameterAllocation>& ParameterMap)
-{
-	for (int32 ParameterTypeIndex = 0; ParameterTypeIndex < (int32)EShaderParameterType::Num; ParameterTypeIndex++)
-	{
-		EShaderParameterType CurrentParameterType = (EShaderParameterType)ParameterTypeIndex;
-
-		if (CurrentParameterType == EShaderParameterType::LooseData)
-		{
-			for (TMap<FString, FParameterAllocation>::TConstIterator ParameterIt(ParameterMap); ParameterIt; ++ParameterIt)
-			{
-				const FParameterAllocation& ParamValue = ParameterIt.Value();
-
-				if (ParamValue.Type == CurrentParameterType)
-				{
-					bool bAddedToExistingBuffer = false;
-
-					for (int32 LooseParameterBufferIndex = 0; LooseParameterBufferIndex < ParameterMapInfo.LooseParameterBuffers.Num(); LooseParameterBufferIndex++)
-					{
-						FShaderLooseParameterBufferInfo& LooseParameterBufferInfo = ParameterMapInfo.LooseParameterBuffers[LooseParameterBufferIndex];
-
-						if (LooseParameterBufferInfo.BufferIndex == ParamValue.BufferIndex)
-						{
-							FShaderParameterInfo ParameterInfo(ParamValue.BaseIndex, ParamValue.Size);
-							LooseParameterBufferInfo.Parameters.Add(ParameterInfo);
-							LooseParameterBufferInfo.BufferSize += ParamValue.Size;
-							bAddedToExistingBuffer = true;
-						}
-					}
-
-					if (!bAddedToExistingBuffer)
-					{
-						FShaderLooseParameterBufferInfo NewParameterBufferInfo(ParamValue.BufferIndex, ParamValue.Size);
-
-						FShaderParameterInfo ParameterInfo(ParamValue.BaseIndex, ParamValue.Size);
-						NewParameterBufferInfo.Parameters.Add(ParameterInfo);
-
-						ParameterMapInfo.LooseParameterBuffers.Add(NewParameterBufferInfo);
-					}
-				}
-			}
-		}
-		else if (CurrentParameterType != EShaderParameterType::UAV)
-		{
-			int32 NumParameters = 0;
-
-			for (TMap<FString, FParameterAllocation>::TConstIterator ParameterIt(ParameterMap); ParameterIt; ++ParameterIt)
-			{
-				const FParameterAllocation& ParamValue = ParameterIt.Value();
-
-				if (ParamValue.Type == CurrentParameterType)
-				{
-					NumParameters++;
-				}
-			}
-
-			TArray<FShaderParameterInfo>* ParameterInfoArray = &ParameterMapInfo.UniformBuffers;
-
-			if (CurrentParameterType == EShaderParameterType::Sampler)
-			{
-				ParameterInfoArray = &ParameterMapInfo.TextureSamplers;
-			}
-			else if (CurrentParameterType == EShaderParameterType::SRV)
-			{
-				ParameterInfoArray = &ParameterMapInfo.SRVs;
-			}
-			else
-			{
-				check(CurrentParameterType == EShaderParameterType::UniformBuffer);
-			}
-
-			ParameterInfoArray->Empty(NumParameters);
-		
-			for (TMap<FString, FParameterAllocation>::TConstIterator ParameterIt(ParameterMap); ParameterIt; ++ParameterIt)
-			{
-				const FParameterAllocation& ParamValue = ParameterIt.Value();
-
-				if (ParamValue.Type == CurrentParameterType)
-				{
-					const uint16 BaseIndex = CurrentParameterType == EShaderParameterType::UniformBuffer ? ParamValue.BufferIndex : ParamValue.BaseIndex;
-					FShaderParameterInfo ParameterInfo(BaseIndex, ParamValue.Size);
-					ParameterInfoArray->Add(ParameterInfo);
-				}
-			}
-		}
-	}
-}
-
-void FShaderResource::UncompressCode(TArray<uint8>& UncompressedCode) const
-{
-	if (Code.Num() != UncompressedCodeSize)
-	{
-		UncompressedCode.SetNum(UncompressedCodeSize);
-		auto bSucceed = FCompression::UncompressMemory(ShaderCompressionFormat, UncompressedCode.GetData(), UncompressedCodeSize, Code.GetData(), Code.Num());
-		check(bSucceed);
-	}
-	else
-	{
-		UncompressedCode = Code;
-	}
-}
-
-void FShaderResource::CompressCode(const TArray<uint8>& UncompressedCode)
-{
-	UncompressedCodeSize = UncompressedCode.Num();
-	Code = UncompressedCode;
-	int32 CompressedSize = Code.Num();
-
-#if !(UE_BUILD_SHIPPING || UE_BUILD_TEST)
-	if (!CVarSkipShaderCompression.GetValueOnAnyThread())
-	{
-#endif
-
-		if (FCompression::CompressMemory(ShaderCompressionFormat, Code.GetData(), CompressedSize, UncompressedCode.GetData(), UncompressedCode.Num()))
-		{
-			Code.SetNum(CompressedSize);
-		}
-
-#if !(UE_BUILD_SHIPPING || UE_BUILD_TEST)
-	}
-#endif
-
-	Code.Shrink();
-}
-
-// Note: this is derived data.  Bump guid in ShaderVersion.ush if changing the format, no backwards compat is necessary
-void FShaderResource::Serialize(FArchive& Ar, bool bLoadedByCookedMaterial)
-{
-	check(!(SpecificPermutationId != 0));
-
-	Ar.UsingCustomVersion(FRenderingObjectVersion::GUID);
-
-	//#todo-RemoveStreamOut
-	FShaderType* SpecificType = nullptr;
-	Ar << SpecificType;
-	if (Ar.CustomVer(FRenderingObjectVersion::GUID) >= FRenderingObjectVersion::ShaderPermutationId)
-	{
-		Ar << SpecificPermutationId;
-	}
-	Ar << Target;
-
-	if (Ar.CustomVer(FRenderingObjectVersion::GUID) < FRenderingObjectVersion::ShaderResourceCodeSharing)
-	{
-		Ar << Code;
-	}
-	Ar << OutputHash;
-	Ar << NumInstructions;
-
-#if WITH_EDITORONLY_DATA
-	if ((!Ar.IsCooking() || Ar.CookingTarget()->HasEditorOnlyData()) && !bLoadedByCookedMaterial)
-	{
-		Ar << NumTextureSamplers;
-	}
-#endif // WITH_EDITORONLY_DATA
-
-	Ar << ParameterMapInfo;
-
-	if (Ar.UE4Ver() >= VER_UE4_COMPRESSED_SHADER_RESOURCES)
-	{
-		Ar << UncompressedCodeSize;
-	}
-
-	if (Ar.CustomVer(FRenderingObjectVersion::GUID) >= FRenderingObjectVersion::ShaderResourceCodeSharing)
-	{
-		SerializeShaderCode(Ar);
-	}
-	
-#if WITH_EDITORONLY_DATA
-	if (!bLoadedByCookedMaterial)
-	{
-		SerializePlatformDebugData(Ar);
-	}
-#endif
-
-	if (Ar.IsLoading())
-	{
-		INC_DWORD_STAT_BY_FName(GetMemoryStatType((EShaderFrequency)Target.Frequency).GetName(), (int64)Code.Num());
-		INC_DWORD_STAT_BY(STAT_Shaders_ShaderResourceMemory, GetSizeBytes());
-	}
-}
-
-void FShaderResource::SerializeShaderCode(FArchive& Ar)
-{
-	// To not pollute the DDC we don't change the state of this object in memory, just the state of the object in the serialised archive.
-	bool bCodeShared = bCodeInSharedLocation;
-	
-#if WITH_EDITOR
-	// in case shader code sharing is enabled, code will be saved outside of material asset
-	if(Ar.IsSaving() && Ar.IsCooking() && Ar.IsPersistent() && !Ar.IsObjectReferenceCollector() && !bCodeInSharedLocation)
-	{
-		bCodeShared = FShaderCodeLibrary::AddShaderCode((EShaderPlatform)Target.Platform, (EShaderFrequency)Target.Frequency, OutputHash, Code, UncompressedCodeSize);
-	}
-#endif
-	
-	Ar << bCodeShared;
-	
-	if (Ar.IsLoading())
-	{
-		bCodeInSharedLocation = bCodeShared;
-
-		if (bCodeInSharedLocation)
-		{
-			if (!GRHILazyShaderCodeLoading)
-			{
-				if (FShaderCodeLibrary::RequestShaderCode(OutputHash, &Ar))
-				{
-					bCodeInSharedLocationRequested = true;
-				}
-			}
-			else
-			{
-				FShaderCodeLibrary::LazyRequestShaderCode(OutputHash, &Ar);
-			}
-		}
-	}
-
-	if (!bCodeShared)
-	{
-		Ar << Code;
-	}
-}
-
-#if WITH_EDITORONLY_DATA
-void FShaderResource::SerializePlatformDebugData(FArchive& Ar)
-{
-#if WITH_ENGINE
-	if (Ar.IsCooking())
-	{
-		// Notify the platform shader format that this particular shader is being used in the cook.
-		// We discard this data in cooked builds unless Ar.CookingTarget()->HasEditorOnlyData() is true.
-		if (PlatformDebugData.Num())
-		{
-			TArray<FName> ShaderFormatNames;
-			Ar.CookingTarget()->GetAllTargetedShaderFormats(ShaderFormatNames);
-
-			for (FName FormatName : ShaderFormatNames)
-			{
-				const IShaderFormat* ShaderFormat = GetTargetPlatformManagerRef().FindShaderFormat(FormatName);
-				if (ShaderFormat)
-				{
-					ShaderFormat->NotifyShaderCooked(PlatformDebugData, FormatName);
-				}
-			}
-		}
-	}
-
-	if (!Ar.IsCooking() || Ar.CookingTarget()->HasEditorOnlyData())
-#endif
-	{
-		// Always serialize if we're not cooking, the cook target requires editor only data, or we don't have the engine (i.e. we're SCW).
-		Ar << PlatformDebugData;
-	}
-}
-#endif
-
-void FShaderResource::AddRef()
-{
-	FPlatformAtomics::InterlockedIncrement((volatile int32*)&NumRefs);
-}
-
-
-void FShaderResource::Release()
-{
-	check(NumRefs > 0);
-	if (FPlatformAtomics::InterlockedDecrement((volatile int32*)&NumRefs) == 0 && gLoadedShaderResources.TryRemove(*this))
-	{
-		// Send a release message to the rendering thread when the shader loses its last reference.
-		BeginReleaseResource(this);
-		BeginCleanup(this);
-
-		if (bCodeInSharedLocation)
-		{
-			if (bCodeInSharedLocationRequested)
-			{
-				FShaderCodeLibrary::ReleaseShaderCode(OutputHash);
-			}
-			else
-			{
-				FShaderCodeLibrary::LazyReleaseShaderCode(OutputHash);
-			}
-		}
-	}
-}
-
-TRefCountPtr<FShaderResource> FShaderResource::FindById(const FShaderResourceId& Id)
-{
-	return gLoadedShaderResources.Find(Id);
-}
-
-TRefCountPtr<FShaderResource> FShaderResource::FindOrClone(FShaderResource&& Temp)
-{
-	if (TRefCountPtr<FShaderResource> ExistingResource = FindById(Temp.GetId()))
-	{
-		return ExistingResource;
-	}
-
-	FShaderResource* Clone = new FShaderResource(MoveTemp(Temp));
-	return gLoadedShaderResources.AddOrDelete(Clone);
-}
-
-//#todo-RemoveStreamOut
-TRefCountPtr<FShaderResource> FShaderResource::FindOrCreate(const FShaderCompilerOutput& Output, int32 InSpecificPermutationId)
-{
-	const FShaderResourceId ResourceId(Output.Target, Output.OutputHash, InSpecificPermutationId);
-
-	if (TRefCountPtr<FShaderResource> ExistingResource = FindById(ResourceId))
-	{
-		return ExistingResource;
-	}
-
-	FShaderResource* NewResource = new FShaderResource(Output, InSpecificPermutationId);
-	return gLoadedShaderResources.AddOrDelete(NewResource);
-}
-
-void FShaderResource::GetAllShaderResourceId(TArray<FShaderResourceId>& Ids)
-{
-	gLoadedShaderResources.GetAllIds(Ids);
-}
-
-bool FShaderResource::ArePlatformsCompatible(EShaderPlatform CurrentPlatform, EShaderPlatform TargetPlatform)
-{
-	bool bFeatureLevelCompatible = CurrentPlatform == TargetPlatform;
-	
-	if (!bFeatureLevelCompatible && IsPCPlatform(CurrentPlatform) && IsPCPlatform(TargetPlatform) )
-	{
-		bFeatureLevelCompatible = GetMaxSupportedFeatureLevel(CurrentPlatform) >= GetMaxSupportedFeatureLevel(TargetPlatform);
-		
-		bool const bIsTargetD3D = TargetPlatform == SP_PCD3D_SM5 ||
-		TargetPlatform == SP_PCD3D_ES3_1;
-		
-		bool const bIsCurrentPlatformD3D = CurrentPlatform == SP_PCD3D_SM5 ||
-		CurrentPlatform == SP_PCD3D_ES3_1;
-		
-		// For Metal in Editor we can switch feature-levels, but not in cooked projects when using Metal shader librariss.
-		bool const bIsCurrentMetal = IsMetalPlatform(CurrentPlatform);
-		bool const bIsTargetMetal = IsMetalPlatform(TargetPlatform);
-		bool const bIsMetalCompatible = (bIsCurrentMetal == bIsTargetMetal) 
-#if !WITH_EDITOR	// Static analysis doesn't like (|| WITH_EDITOR)
-			&& (!IsMetalPlatform(CurrentPlatform) || (CurrentPlatform == TargetPlatform))
-#endif
-			;
-		
-		bool const bIsCurrentOpenGL = IsOpenGLPlatform(CurrentPlatform);
-		bool const bIsTargetOpenGL = IsOpenGLPlatform(TargetPlatform);
-		
-		bFeatureLevelCompatible = bFeatureLevelCompatible && (bIsCurrentPlatformD3D == bIsTargetD3D && bIsMetalCompatible && bIsCurrentOpenGL == bIsTargetOpenGL);
-	}
-
-	return bFeatureLevelCompatible;
-}
-
-FSHAHash &FShaderResource::FilterShaderSourceHashForSerialization(const FArchive& Ar, FSHAHash &HashToSerialize)
-{
-#if KEEP_SHADER_SOURCE_HASHES
-	return (!Ar.IsCooking()) ? HashToSerialize : ShaderSourceDefaultHash;
-#else
-	return ShaderSourceDefaultHash;
-#endif
-}
-
-static void SafeAssignHash(FRHIShader* InShader, const FSHAHash& Hash)
-{
-	if (InShader)
-	{
-		InShader->SetHash(Hash);
-	}
-}
-
-void FShaderResource::InitRHI()
-{
-	checkf(bCodeInSharedLocation || Code.Num() > 0, TEXT("FShaderResource::InitRHI was called with empty bytecode, which can happen if the resource is initialized multiple times on platforms with no editor data."));
-
-	// we can't have this called on the wrong platform's shaders
-	if (!ArePlatformsCompatible(GMaxRHIShaderPlatform, (EShaderPlatform)Target.Platform))
-	{
-		if (FPlatformProperties::RequiresCookedData())
-		{
-			UE_LOG(LogShaders, Fatal, TEXT("FShaderResource::InitRHI got platform %s but it is not compatible with %s"), 
-				*LegacyShaderPlatformToShaderFormat((EShaderPlatform)Target.Platform).ToString(), *LegacyShaderPlatformToShaderFormat(GMaxRHIShaderPlatform).ToString());
-		}
-		return;
-	}
-
-	TArray<uint8> UncompressedCode;
-	if (!bCodeInSharedLocation)
-	{
-		UncompressCode(UncompressedCode);
-	}
-
-	INC_DWORD_STAT_BY(STAT_Shaders_NumShadersUsedForRendering, 1);
-	SCOPE_CYCLE_COUNTER(STAT_Shaders_RTShaderLoadTime);
-
-	if(Target.Frequency == SF_Vertex)
-	{
-		Shader = FShaderCodeLibrary::CreateVertexShader((EShaderPlatform)Target.Platform, OutputHash, UncompressedCode);
-		UE_CLOG((bCodeInSharedLocation && !IsValidRef(Shader)), LogShaders, Fatal, TEXT("FShaderResource::SerializeShaderCode can't find shader code for: [%s]"), *LegacyShaderPlatformToShaderFormat((EShaderPlatform)Target.Platform).ToString());
-	}
-	else if(Target.Frequency == SF_Pixel)
-	{
-		Shader = FShaderCodeLibrary::CreatePixelShader((EShaderPlatform)Target.Platform, OutputHash, UncompressedCode);
-		UE_CLOG((bCodeInSharedLocation && !IsValidRef(Shader)), LogShaders, Fatal, TEXT("FShaderResource::SerializeShaderCode can't find shader code for: [%s]"), *LegacyShaderPlatformToShaderFormat((EShaderPlatform)Target.Platform).ToString());
-	}
-	else if(Target.Frequency == SF_Hull)
-	{
-		Shader = FShaderCodeLibrary::CreateHullShader((EShaderPlatform)Target.Platform, OutputHash, UncompressedCode);
-		UE_CLOG((bCodeInSharedLocation && !IsValidRef(Shader)), LogShaders, Fatal, TEXT("FShaderResource::SerializeShaderCode can't find shader code for: [%s]"), *LegacyShaderPlatformToShaderFormat((EShaderPlatform)Target.Platform).ToString());
-	}
-	else if(Target.Frequency == SF_Domain)
-	{
-		Shader = FShaderCodeLibrary::CreateDomainShader((EShaderPlatform)Target.Platform, OutputHash, UncompressedCode);
-		UE_CLOG((bCodeInSharedLocation && !IsValidRef(Shader)), LogShaders, Fatal, TEXT("FShaderResource::SerializeShaderCode can't find shader code for: [%s]"), *LegacyShaderPlatformToShaderFormat((EShaderPlatform)Target.Platform).ToString());
-	}
-	else if(Target.Frequency == SF_Geometry)
-	{
-		Shader = FShaderCodeLibrary::CreateGeometryShader((EShaderPlatform)Target.Platform, OutputHash, UncompressedCode);
-		UE_CLOG((bCodeInSharedLocation && !IsValidRef(Shader)), LogShaders, Fatal, TEXT("FShaderResource::SerializeShaderCode can't find shader code for: [%s]"), *LegacyShaderPlatformToShaderFormat((EShaderPlatform)Target.Platform).ToString());
-	}
-	else if(Target.Frequency == SF_Compute)
-	{
-		Shader = FShaderCodeLibrary::CreateComputeShader((EShaderPlatform)Target.Platform, OutputHash, UncompressedCode);
-		UE_CLOG((bCodeInSharedLocation && !IsValidRef(Shader)), LogShaders, Fatal, TEXT("FShaderResource::SerializeShaderCode can't find shader code for: [%s]"), *LegacyShaderPlatformToShaderFormat((EShaderPlatform)Target.Platform).ToString());
-	}
-#if RHI_RAYTRACING
-	else if (Target.Frequency == SF_RayGen || Target.Frequency == SF_RayMiss || Target.Frequency == SF_RayHitGroup || Target.Frequency == SF_RayCallable)
-	{
-		if (GRHISupportsRayTracing)
-		{
-			RayTracingShader = FShaderCodeLibrary::CreateRayTracingShader((EShaderPlatform)Target.Platform, Target.GetFrequency(), OutputHash, UncompressedCode);
-			UE_CLOG((bCodeInSharedLocation && !IsValidRef(RayTracingShader)), LogShaders, Fatal, TEXT("FShaderResource::SerializeShaderCode can't find shader code for: [%s]"), *LegacyShaderPlatformToShaderFormat((EShaderPlatform)Target.Platform).ToString());
-
-			if (Target.Frequency == SF_RayHitGroup)
-			{
-				RayTracingMaterialLibraryIndex = AddToRayTracingLibrary(RayTracingShader);
-			}
-		}
-	}
-#endif // RHI_RAYTRACING
-	else
-	{
-		checkNoEntry(); // Unexpected shader target frequency
-	}
-
-	if (!FPlatformProperties::HasEditorOnlyData())
-	{
-		DEC_DWORD_STAT_BY_FName(GetMemoryStatType((EShaderFrequency)Target.Frequency).GetName(), Code.Num());
-		DEC_DWORD_STAT_BY(STAT_Shaders_ShaderResourceMemory, Code.GetAllocatedSize());
-		Code.Empty();
-		
-		if (bCodeInSharedLocation)
-		{
-			if (bCodeInSharedLocationRequested)
-			{
-				FShaderCodeLibrary::ReleaseShaderCode(OutputHash);
-			}
-			else
-			{
-				FShaderCodeLibrary::LazyReleaseShaderCode(OutputHash);
-			}
-		}
-		bCodeInSharedLocation = false;
-		bCodeInSharedLocationRequested = false;
-	}
-}
-
-
-void FShaderResource::ReleaseRHI()
-{
-	DEC_DWORD_STAT_BY(STAT_Shaders_NumShadersUsedForRendering, 1);
-
-#if RHI_RAYTRACING
-	if (IsInitialized() && RayTracingMaterialLibraryIndex != UINT_MAX)
-	{
-		RemoveFromRayTracingLibrary(RayTracingMaterialLibraryIndex);
-		RayTracingMaterialLibraryIndex = UINT_MAX;
-	}
-#endif
-
-	Shader.SafeRelease();
-
-#if RHI_RAYTRACING
-	RayTracingShader.SafeRelease();
-#endif // RHI_RAYTRACING
-}
-
-void FShaderResource::InitializeShaderRHI() 
-{ 
-	if (!IsInitialized())
-	{
-		STAT(double ShaderInitializationTime = 0);
-		{
-			SCOPE_CYCLE_COUNTER(STAT_Shaders_FrameRTShaderInitForRenderingTime);
-			SCOPE_SECONDS_COUNTER(ShaderInitializationTime);
-
-			InitResourceFromPossiblyParallelRendering();
-		}
-
-		INC_FLOAT_STAT_BY(STAT_Shaders_TotalRTShaderInitForRenderingTime,(float)ShaderInitializationTime);
-	}
-
-	checkSlow(IsInitialized());
-}
-
-FShaderResourceId FShaderResource::GetId() const
+int32 FShaderMapPointerTable::AddIndexedPointer(const FTypeLayoutDesc& TypeDesc, void* Ptr)
 {
-	return FShaderResourceId(Target, OutputHash, SpecificPermutationId);
-}
-
-FShaderId::FShaderId(const FSHAHash& InMaterialShaderMapHash, const FShaderPipelineType* InShaderPipeline, FVertexFactoryType* InVertexFactoryType, FShaderType* InShaderType, int32 InPermutationId, FShaderTarget InTarget)
-	: ShaderType(InShaderType)
-	, VertexFactoryType(InVertexFactoryType)
-	, ShaderPipeline(InShaderPipeline)
-	, MaterialShaderMapHash(InMaterialShaderMapHash)
-#if KEEP_SHADER_SOURCE_HASHES
-	, SourceHash(InShaderType->GetSourceHash(InTarget.GetPlatform()))
-#endif
-	, Target(InTarget)
-	, PermutationId(InPermutationId)
-{
-#if KEEP_SHADER_SOURCE_HASHES
-	if (InVertexFactoryType)
-	{
-		VFSourceHash = InVertexFactoryType->GetSourceHash(InTarget.GetPlatform());
-	}
-#endif
+	int32 Index = INDEX_NONE;
+	if (ShaderTypes.TryAddIndexedPtr(TypeDesc, Ptr, Index)) return Index;
+	if (VFTypes.TryAddIndexedPtr(TypeDesc, Ptr, Index)) return Index;
+	return Index;
 }
 
-FArchive& operator<<(FArchive& Ar, FShaderKey& Ref)
+void* FShaderMapPointerTable::GetIndexedPointer(const FTypeLayoutDesc& TypeDesc, uint32 i) const
 {
-	uint32 Platform = Ref.Platform;
-	Ar
-		<< Ref.VertexFactoryType
-		<< Ref.ShaderPipeline
-		<< Ref.MaterialShaderMapHash
-		<< Ref.PermutationId
-		<< Platform;
-	Ref.Platform = Platform;
-	return Ar;
+	void* Ptr = nullptr;
+	if (ShaderTypes.TryGetIndexedPtr(TypeDesc, i, Ptr)) return Ptr;
+	if (VFTypes.TryGetIndexedPtr(TypeDesc, i, Ptr)) return Ptr;
+	return Ptr;
 }
 
-FSelfContainedShaderId::FSelfContainedShaderId() :
-	Target(FShaderTarget(SF_NumFrequencies, SP_NumPlatforms))
-{}
-
-FSelfContainedShaderId::FSelfContainedShaderId(const FShaderId& InShaderId)
+void FShaderMapPointerTable::SaveToArchive(FArchive& Ar, void* FrozenContent, bool bInlineShaderResources) const
 {
-	MaterialShaderMapHash = InShaderId.MaterialShaderMapHash;
-	ShaderPipelineName = InShaderId.ShaderPipeline ? InShaderId.ShaderPipeline->GetName() : TEXT("");
-	ShaderTypeName = InShaderId.ShaderType->GetName();
-	PermutationId = InShaderId.PermutationId;
-	Target = InShaderId.Target;
-#if KEEP_SHADER_SOURCE_HASHES
-	SourceHash = InShaderId.SourceHash;
-	VFSourceHash = InShaderId.VFSourceHash;
-#endif
-	SerializationHistory = InShaderId.ShaderType->GetSerializationHistory();
+	int32 NumTypes = ShaderTypes.Num();
+	int32 NumVFTypes = VFTypes.Num();
+
+	Ar << NumTypes;
+	Ar << NumVFTypes;
 
-	if (InShaderId.VertexFactoryType)
+	for (int32 TypeIndex = 0; TypeIndex < NumTypes; ++TypeIndex)
 	{
-		const FSerializationHistory* InVFSerializationHistory =InShaderId.VertexFactoryType->GetSerializationHistory(InShaderId.Target.GetFrequency());
-		VertexFactoryTypeName = InShaderId.VertexFactoryType->GetName();
-		if (InVFSerializationHistory)
-		{
-			VFSerializationHistory = *InVFSerializationHistory;
-		}
+		const FShaderType* Type = ShaderTypes.GetIndexedPointer(TypeIndex);
+		FHashedName TypeName = Type->GetHashedName();
+		Ar << TypeName;
 	}
-}
 
-bool FSelfContainedShaderId::IsValid()
-{
-	FShaderType** TypePtr = FShaderType::GetNameToTypeMap().Find(FName(*ShaderTypeName));
-	if (TypePtr
-#if KEEP_SHADER_SOURCE_HASHES
-		&& SourceHash == (*TypePtr)->GetSourceHash(Target.GetPlatform())
-#endif
-		&& SerializationHistory == (*TypePtr)->GetSerializationHistory())
+	for (int32 VFTypeIndex = 0; VFTypeIndex < NumVFTypes; ++VFTypeIndex)
 	{
-		FVertexFactoryType* VFTypePtr = FVertexFactoryType::GetVFByName(VertexFactoryTypeName);
-
-		if (VertexFactoryTypeName == TEXT("") 
-			|| (VFTypePtr
-#if KEEP_SHADER_SOURCE_HASHES
-				&& VFSourceHash == VFTypePtr->GetSourceHash(Target.GetPlatform())
-#endif
-				&& VFSerializationHistory == *VFTypePtr->GetSerializationHistory(Target.GetFrequency())))
-		{
-			return true;
-		}
+		const FVertexFactoryType* VFType = VFTypes.GetIndexedPointer(VFTypeIndex);
+		FHashedName TypeName = VFType->GetHashedName();
+		Ar << TypeName;
 	}
-
-	return false;
 }
 
-FArchive& operator<<(FArchive& Ar,class FSelfContainedShaderId& Ref)
+void FShaderMapPointerTable::LoadFromArchive(FArchive& Ar, void* FrozenContent, bool bInlineShaderResources, bool bLoadedByCookedMaterial)
 {
-	Ar.UsingCustomVersion(FRenderingObjectVersion::GUID);
+	SCOPED_LOADTIMER(FShaderMapPointerTable_LoadFromArchive);
 
-#if KEEP_SHADER_SOURCE_HASHES
-	FSHAHash& VFHash = Ref.VFSourceHash;
-	FSHAHash& Hash = Ref.SourceHash;
-#else
-	FSHAHash VFHash, Hash;
-#endif
-
-	Ar << Ref.MaterialShaderMapHash 
-		<< Ref.VertexFactoryTypeName
-		<< Ref.ShaderPipelineName
-		<< FShaderResource::FilterShaderSourceHashForSerialization(Ar, VFHash)
-		<< Ref.VFSerializationHistory
-		<< Ref.ShaderTypeName
-		<< FShaderResource::FilterShaderSourceHashForSerialization(Ar, Hash)
-		<< Ref.SerializationHistory
-		<< Ref.Target;
+	int32 NumTypes = 0;
+	int32 NumVFTypes = 0;
+		
+	Ar << NumTypes;
+	Ar << NumVFTypes;
 
-	if (Ar.CustomVer(FRenderingObjectVersion::GUID) >= FRenderingObjectVersion::ShaderPermutationId)
+	ShaderTypes.Empty(NumTypes);
+	for (int32 TypeIndex = 0; TypeIndex < NumTypes; ++TypeIndex)
 	{
-		Ar << Ref.PermutationId;
+		FHashedName TypeName;
+		Ar << TypeName;
+		FShaderType* Type = FindShaderTypeByName(TypeName);
+		ShaderTypes.LoadIndexedPointer(Type);
 	}
 
-	return Ar;
+	VFTypes.Empty(NumVFTypes);
+	for (int32 VFTypeIndex = 0; VFTypeIndex < NumVFTypes; ++VFTypeIndex)
+	{
+		FHashedName TypeName;
+		Ar << TypeName;
+		FVertexFactoryType* VFType = FVertexFactoryType::GetVFByName(TypeName);
+		VFTypes.LoadIndexedPointer(VFType);
+	}
 }
 
 /** 
@@ -1227,308 +445,256 @@ FArchive& operator<<(FArchive& Ar,class FSelfContainedShaderId& Ref)
  * This still needs to initialize members to safe values since FShaderType::GenerateSerializationHistory uses this constructor.
  */
 FShader::FShader()
-	: ShaderPipeline(nullptr)
-	, VFType(nullptr)
-	, Type(nullptr)
+	// set to undefined (currently shared with SF_Vertex)
+	: Target((EShaderFrequency)0, GShaderPlatformForFeatureLevel[GMaxRHIFeatureLevel])
+	, ResourceIndex(INDEX_NONE)
 	, PermutationId(0)
-	, NumRefs(0)
+#if WITH_EDITORONLY_DATA
+	, NumInstructions(0u)
+	, NumTextureSamplers(0u)
+	, CodeSize(0u)
+#endif // WITH_EDITORONLY_DATA
 {
-	// set to undefined (currently shared with SF_Vertex)
-	Target.Frequency = 0;
-	Target.Platform = GShaderPlatformForFeatureLevel[GMaxRHIFeatureLevel];
 }
 
 /**
  * Construct a shader from shader compiler output.
  */
 FShader::FShader(const CompiledShaderInitializerType& Initializer)
-	: MaterialShaderMapHash(Initializer.MaterialShaderMapHash)
-	, Resource(Initializer.Resource)
-	, ShaderPipeline(Initializer.ShaderPipeline)
+	: Type(Initializer.Type)
 	, VFType(Initializer.VertexFactoryType)
-	, Type(Initializer.Type)
-	, PermutationId(Initializer.PermutationId)
+	, TypeName(Initializer.Type->GetFName())
 	, Target(Initializer.Target)
-	, NumRefs(0)
+	, ResourceIndex(Initializer.ResourceIndex)
+	, PermutationId(Initializer.PermutationId)
+#if WITH_EDITORONLY_DATA
+	, NumInstructions(Initializer.NumInstructions)
+	, NumTextureSamplers(Initializer.NumTextureSamplers)
+	, CodeSize(Initializer.CodeSize)
+#endif // WITH_EDITORONLY_DATA
 {
-	check(Type);
-
-#if KEEP_SHADER_SOURCE_HASHES
+#if WITH_EDITORONLY_DATA
+	checkSlow(Initializer.OutputHash != FSHAHash());
+	
 	OutputHash = Initializer.OutputHash;
-	checkSlow(OutputHash != FSHAHash());
 
 	// Store off the source hash that this shader was compiled with
 	// This will be used as part of the shader key in order to identify when shader files have been changed and a recompile is needed
-	SourceHash = Type->GetSourceHash(Target.GetPlatform());
+	SourceHash = Initializer.Type->GetSourceHash(Initializer.Target.GetPlatform());
 
-	if (VFType)
+	if (Initializer.VertexFactoryType)
 	{
 		// Store off the VF source hash that this shader was compiled with
-		VFSourceHash = VFType->GetSourceHash(Target.GetPlatform());
+		VFSourceHash = Initializer.VertexFactoryType->GetSourceHash(Initializer.Target.GetPlatform());
 	}
-#endif
+#endif // WITH_EDITORONLY_DATA
+
+	BuildParameterMapInfo(Initializer.ParameterMap.GetParameterMap());
 
 	// Bind uniform buffer parameters automatically 
 	for (TLinkedList<FShaderParametersMetadata*>::TIterator StructIt(FShaderParametersMetadata::GetStructList()); StructIt; StructIt.Next())
 	{
 		if (Initializer.ParameterMap.ContainsParameterAllocation(StructIt->GetShaderVariableName()))
 		{
-			UniformBufferParameterStructs.Add(*StructIt);
+			UniformBufferParameterStructs.Add(StructIt->GetShaderVariableHashedName());
 			FShaderUniformBufferParameter& Parameter = UniformBufferParameters.AddDefaulted_GetRef();
 			Parameter.Bind(Initializer.ParameterMap, StructIt->GetShaderVariableName(), SPF_Mandatory);
 		}
 	}
 
-	
-	check(Resource && Resource->Target == Target);
-	checkSlow(FShaderResource::FindById(Resource->GetId()) == Resource);
-
 	// Register the shader now that it is valid, so that it can be reused
-	Register(false);
+	//Register(false);
 }
 
-
 FShader::~FShader()
 {
-	check(NumRefs == 0);
-}
-
-
-const FSHAHash& FShader::GetHash() const 
-{ 
-#if KEEP_SHADER_SOURCE_HASHES
-	return SourceHash;
-#else
-	return ShaderSourceDefaultHash;
-#endif
-}
-
-EShaderPlatform FShader::GetShaderPlatform() const
-{
-	return Target.GetPlatform();
 }
 
-
-bool FShader::SerializeBase(FArchive& Ar, bool bShadersInline, bool bLoadedByCookedMaterial)
+void FShader::BuildParameterMapInfo(const TMap<FString, FParameterAllocation>& ParameterMap)
 {
-	Ar.UsingCustomVersion(FFortniteMainBranchObjectVersion::GUID);
-
-	Serialize(Ar);
+	for (int32 ParameterTypeIndex = 0; ParameterTypeIndex < (int32)EShaderParameterType::Num; ParameterTypeIndex++)
+	{
+		EShaderParameterType CurrentParameterType = (EShaderParameterType)ParameterTypeIndex;
 
-	Ar.UsingCustomVersion(FRenderingObjectVersion::GUID);
+		if (CurrentParameterType == EShaderParameterType::LooseData)
+		{
+			for (TMap<FString, FParameterAllocation>::TConstIterator ParameterIt(ParameterMap); ParameterIt; ++ParameterIt)
+			{
+				const FParameterAllocation& ParamValue = ParameterIt.Value();
 
-#if KEEP_SHADER_SOURCE_HASHES
-	FSHAHash& VFHash = VFSourceHash;
-	FSHAHash& Hash = SourceHash;
-#else
-	FSHAHash VFHash, Hash, OutputHash;
-#endif
+				if (ParamValue.Type == CurrentParameterType)
+				{
+					bool bAddedToExistingBuffer = false;
 
-	Ar << OutputHash;
-	Ar << MaterialShaderMapHash;
-	Ar << ShaderPipeline;
-	Ar << VFType;
-	Ar << FShaderResource::FilterShaderSourceHashForSerialization(Ar, VFHash);
-	Ar << Type;
-	if (Ar.CustomVer(FRenderingObjectVersion::GUID) >= FRenderingObjectVersion::ShaderPermutationId)
-	{
-		Ar << PermutationId;
-	}
-	Ar << FShaderResource::FilterShaderSourceHashForSerialization(Ar, Hash);
-	Ar << Target;
+					for (int32 LooseParameterBufferIndex = 0; LooseParameterBufferIndex < ParameterMapInfo.LooseParameterBuffers.Num(); LooseParameterBufferIndex++)
+					{
+						FShaderLooseParameterBufferInfo& LooseParameterBufferInfo = ParameterMapInfo.LooseParameterBuffers[LooseParameterBufferIndex];
 
-	const bool SerializeStructTypeNameAsStrings = Ar.CustomVer(FFortniteMainBranchObjectVersion::GUID) < FFortniteMainBranchObjectVersion::MaterialInstanceSerializeOptimization_ShaderFName;
+						if (LooseParameterBufferInfo.BaseIndex == ParamValue.BufferIndex)
+						{
+							FShaderParameterInfo ParameterInfo(ParamValue.BaseIndex, ParamValue.Size);
+							LooseParameterBufferInfo.Parameters.Add(ParameterInfo);
+							LooseParameterBufferInfo.Size += ParamValue.Size;
+							bAddedToExistingBuffer = true;
+						}
+					}
 
-	// TODO(RDG): Kill that once all shaders are refactored.
-	if (Ar.IsLoading())
-	{
-		int32 NumUniformParameters;
-		Ar << NumUniformParameters;
+					if (!bAddedToExistingBuffer)
+					{
+						FShaderLooseParameterBufferInfo NewParameterBufferInfo(ParamValue.BufferIndex, ParamValue.Size);
 
-		UniformBufferParameterStructs.Empty(NumUniformParameters);
-		UniformBufferParameters.Empty(NumUniformParameters);
+						FShaderParameterInfo ParameterInfo(ParamValue.BaseIndex, ParamValue.Size);
+						NewParameterBufferInfo.Parameters.Add(ParameterInfo);
 
-		for (int32 ParameterIndex = 0; ParameterIndex < NumUniformParameters; ParameterIndex++)
+						ParameterMapInfo.LooseParameterBuffers.Add(NewParameterBufferInfo);
+					}
+				}
+			}
+		}
+		else if (CurrentParameterType != EShaderParameterType::UAV)
 		{
-			FShaderParametersMetadata* Struct = nullptr;
+			int32 NumParameters = 0;
 
-			if (SerializeStructTypeNameAsStrings)
-			{
-				FString StructName;
-				Ar << StructName;
-				Struct = FindUniformBufferStructByName(*StructName);
-				checkf(Struct, TEXT("Uniform Buffer Struct %s no longer exists, which shader of type %s was compiled with.  Modify ShaderVersion.ush to invalidate old shaders."), *StructName, Type->GetName());
-			}
-			else
+			for (TMap<FString, FParameterAllocation>::TConstIterator ParameterIt(ParameterMap); ParameterIt; ++ParameterIt)
 			{
-				FName StructFName;
-				Ar << StructFName;
-				Struct = FindUniformBufferStructByFName(StructFName);
-				checkf(Struct, TEXT("Uniform Buffer Struct %s no longer exists, which shader of type %s was compiled with.  Modify ShaderVersion.ush to invalidate old shaders."), *StructFName.ToString(), Type->GetName());
-			}
+				const FParameterAllocation& ParamValue = ParameterIt.Value();
 
-			UniformBufferParameterStructs.Add(Struct);
-			FShaderUniformBufferParameter& Parameter = UniformBufferParameters.AddDefaulted_GetRef();
-			Ar << Parameter;
-		}
-	}
-	else
-	{
-		int32 NumUniformParameters = UniformBufferParameters.Num();
-		Ar << NumUniformParameters;
+				if (ParamValue.Type == CurrentParameterType)
+				{
+					NumParameters++;
+				}
+			}
 
-		for (int32 StructIndex = 0; StructIndex < NumUniformParameters; StructIndex++)
-		{
-			const TCHAR* StructName = UniformBufferParameterStructs[StructIndex]->GetStructTypeName();
+			TMemoryImageArray<FShaderParameterInfo>* ParameterInfoArray = &ParameterMapInfo.UniformBuffers;
 
-			if (SerializeStructTypeNameAsStrings)
+			if (CurrentParameterType == EShaderParameterType::Sampler)
+			{
+				ParameterInfoArray = &ParameterMapInfo.TextureSamplers;
+			}
+			else if (CurrentParameterType == EShaderParameterType::SRV)
 			{
-				FString StructNameCopy(StructName);
-				Ar << StructNameCopy;
+				ParameterInfoArray = &ParameterMapInfo.SRVs;
 			}
 			else
 			{
-				FName StructFName(StructName);
-				Ar << StructFName;
+				check(CurrentParameterType == EShaderParameterType::UniformBuffer);
 			}
 
-			Ar << UniformBufferParameters[StructIndex];
-		}
-	}
-
-	if (bShadersInline)
-	{
-		// Save the shader resource if we are inlining shaders
-		if (Ar.IsSaving())
-		{
-			check(Resource->Target == Target);
-			Resource->Serialize(Ar, false);
-		}
-
-		if (Ar.IsLoading())
-		{
-			// Load the inlined shader resource
-			FShaderResource ResourceTemp;
-			ResourceTemp.Serialize(Ar, bLoadedByCookedMaterial);
-
-			Resource = FShaderResource::FindOrClone(MoveTemp(ResourceTemp));
+			ParameterInfoArray->Empty(NumParameters);
+		
+			for (TMap<FString, FParameterAllocation>::TConstIterator ParameterIt(ParameterMap); ParameterIt; ++ParameterIt)
+			{
+				const FParameterAllocation& ParamValue = ParameterIt.Value();
 
-			checkSlow(OutputHash == Resource->OutputHash);
+				if (ParamValue.Type == CurrentParameterType)
+				{
+					const uint16 BaseIndex = CurrentParameterType == EShaderParameterType::UniformBuffer ? ParamValue.BufferIndex : ParamValue.BaseIndex;
+					FShaderParameterInfo ParameterInfo(BaseIndex, ParamValue.Size);
+					ParameterInfoArray->Add(ParameterInfo);
+				}
+			}
 		}
 	}
-	else
+	for (FShaderLooseParameterBufferInfo& Info : ParameterMapInfo.LooseParameterBuffers)
 	{
-		// if saving, there's nothing to, the required data is already saved above to look it up at load time
-		if (Ar.IsLoading())
-		{
-			// generate a resource id
-			FShaderResourceId ResourceId(Target, OutputHash, 0);
-
-			// use it to look up in the registered resource map
-			Resource = FShaderResource::FindById(ResourceId);
-			check(Resource && Resource->Target == Target);
-		}
+		Info.Parameters.Sort();
 	}
-
-	Ar << Bindings;
-
-	return false;
+	ParameterMapInfo.LooseParameterBuffers.Sort();
+	ParameterMapInfo.UniformBuffers.Sort();
+	ParameterMapInfo.TextureSamplers.Sort();
+	ParameterMapInfo.SRVs.Sort();
 }
 
-void FShader::AddRef()
-{	
-	++NumRefs;
-	if (NumRefs == 1)
-	{
-		INC_DWORD_STAT_BY(STAT_Shaders_ShaderMemory, GetSizeBytes());
-		INC_DWORD_STAT_BY(STAT_Shaders_NumShadersLoaded,1);
-	}
+const FSHAHash& FShader::GetOutputHash() const
+{
+#if WITH_EDITORONLY_DATA
+	return OutputHash;
+#endif
+	return ShaderSourceDefaultHash;
 }
 
-void FShader::Release()
+const FSHAHash& FShader::GetHash() const 
 {
-	if(--NumRefs == 0)
-	{
-		DEC_DWORD_STAT_BY(STAT_Shaders_ShaderMemory, GetSizeBytes());
-		DEC_DWORD_STAT_BY(STAT_Shaders_NumShadersLoaded,1);
-
-		// Deregister the shader now to eliminate references to it by the type's ShaderIdMap
-		Deregister();
-		BeginCleanup(this);
-	}
+#if WITH_EDITORONLY_DATA
+	return SourceHash;
+#endif
+	return ShaderSourceDefaultHash;
 }
 
-
-void FShader::Register(bool bLoadedByCookedMaterial)
+const FSHAHash& FShader::GetVertexFactoryHash() const
 {
-	const FShaderKey ShaderKey = GetKey();
-	check(ShaderKey.MaterialShaderMapHash != FSHAHash());
-#if KEEP_SHADER_SOURCE_HASHES
-	check(SourceHash != FSHAHash() || FPlatformProperties::RequiresCookedData() || bLoadedByCookedMaterial);
+#if WITH_EDITORONLY_DATA
+	return VFSourceHash;
 #endif
-	check(Resource);
-	Type->AddToShaderMap(ShaderKey, this);
+	return ShaderSourceDefaultHash;
 }
 
-void FShader::Deregister()
+const FTypeLayoutDesc& GetTypeLayoutDesc(const FPointerTableBase* PtrTable, const FShader& Shader)
 {
-	Type->RemoveFromShaderMap(GetKey());
+	const FShaderType* Type = Shader.GetType(PtrTable);
+	return Type->GetLayout();
 }
 
-FShaderId FShader::GetId() const
+const FShaderParametersMetadata* FShader::FindAutomaticallyBoundUniformBufferStruct(int32 BaseIndex) const
 {
-	FShaderId ShaderId;
-	ShaderId.MaterialShaderMapHash = MaterialShaderMapHash;
-	ShaderId.ShaderPipeline = ShaderPipeline;
-	ShaderId.VertexFactoryType = VFType;
-	ShaderId.ShaderType = Type;
-	ShaderId.PermutationId = PermutationId;
-#if KEEP_SHADER_SOURCE_HASHES
-	ShaderId.SourceHash = SourceHash;
-	ShaderId.VFSourceHash = VFSourceHash;
-#endif
-	ShaderId.Target = Target;
-	return ShaderId;
+	for (int32 i = 0; i < UniformBufferParameters.Num(); i++)
+	{
+		if (UniformBufferParameters[i].GetBaseIndex() == BaseIndex)
+		{
+			FShaderParametersMetadata** Parameters = FShaderParametersMetadata::GetNameStructMap().Find(UniformBufferParameterStructs[i]);
+			return Parameters ? *Parameters : nullptr;
+		}
+	}
+
+	return nullptr;
 }
 
-void FShader::DumpDebugInfo()
+void FShader::DumpDebugInfo(const FShaderMapPointerTable& InPtrTable)
 {
-	UE_LOG(LogConsoleResponse, Display, TEXT("      FShader  :MaterialShaderMapHash %s"), *MaterialShaderMapHash.ToString());
-	UE_LOG(LogConsoleResponse, Display, TEXT("               :Target %s"), GetShaderFrequencyString((EShaderFrequency)Target.Frequency));
-	UE_LOG(LogConsoleResponse, Display, TEXT("               :Target %s"), *LegacyShaderPlatformToShaderFormat(EShaderPlatform(Target.Platform)).ToString());
-	UE_LOG(LogConsoleResponse, Display, TEXT("               :VFType %s"), VFType ? VFType->GetName() : TEXT("null"));
-	UE_LOG(LogConsoleResponse, Display, TEXT("               :Type %s"), Type->GetName());
+	FVertexFactoryType* VertexFactoryType = GetVertexFactoryType(InPtrTable);
+
+	UE_LOG(LogConsoleResponse, Display, TEXT("      FShader  :Frequency %s"), GetShaderFrequencyString(GetFrequency()));
+	UE_LOG(LogConsoleResponse, Display, TEXT("               :Target %s"), *LegacyShaderPlatformToShaderFormat(GetShaderPlatform()).ToString());
+	UE_LOG(LogConsoleResponse, Display, TEXT("               :VFType %s"), VertexFactoryType ? VertexFactoryType->GetName() : TEXT("null"));
+	UE_LOG(LogConsoleResponse, Display, TEXT("               :Type %s"), GetType(InPtrTable)->GetName());
 	UE_LOG(LogConsoleResponse, Display, TEXT("               :PermutationId %d"), PermutationId);
-#if KEEP_SHADER_SOURCE_HASHES
-	UE_LOG(LogConsoleResponse, Display, TEXT("               :SourceHash %s"), *SourceHash.ToString());
-	UE_LOG(LogConsoleResponse, Display, TEXT("               :VFSourceHash %s"), *VFSourceHash.ToString());
-	UE_LOG(LogConsoleResponse, Display, TEXT("               :OutputHash %s"), *OutputHash.ToString());
-#endif
+	UE_LOG(LogConsoleResponse, Display, TEXT("               :SourceHash %s"), *GetHash().ToString());
+	UE_LOG(LogConsoleResponse, Display, TEXT("               :VFSourceHash %s"), *GetVertexFactoryHash().ToString());
+	UE_LOG(LogConsoleResponse, Display, TEXT("               :OutputHash %s"), *GetOutputHash().ToString());
 }
 
-void FShader::SaveShaderStableKeys(EShaderPlatform TargetShaderPlatform, const FStableShaderKeyAndValue& InSaveKeyVal)
-{
 #if WITH_EDITOR
-	if ((TargetShaderPlatform == EShaderPlatform::SP_NumPlatforms || EShaderPlatform(Target.Platform) == TargetShaderPlatform) 
+void FShader::SaveShaderStableKeys(const FShaderMapPointerTable& InPtrTable, EShaderPlatform TargetShaderPlatform, const FStableShaderKeyAndValue& InSaveKeyVal)
+{
+	if ((TargetShaderPlatform == EShaderPlatform::SP_NumPlatforms || GetShaderPlatform() == TargetShaderPlatform) 
 		&& FShaderCodeLibrary::NeedsShaderStableKeys(TargetShaderPlatform))
 	{
+		FShaderType* ShaderType = GetType(InPtrTable);
+		FVertexFactoryType* VertexFactoryType = GetVertexFactoryType(InPtrTable);
+
 		FStableShaderKeyAndValue SaveKeyVal(InSaveKeyVal);
-		SaveKeyVal.TargetFrequency = FName(GetShaderFrequencyString((EShaderFrequency)Target.Frequency));
-		SaveKeyVal.TargetPlatform = FName(*LegacyShaderPlatformToShaderFormat(EShaderPlatform(Target.Platform)).ToString());
-		SaveKeyVal.VFType = FName(VFType ? VFType->GetName() : TEXT("null"));
+		SaveKeyVal.TargetFrequency = FName(GetShaderFrequencyString(GetFrequency()));
+		SaveKeyVal.TargetPlatform = LegacyShaderPlatformToShaderFormat(GetShaderPlatform());
+		SaveKeyVal.VFType = FName(VertexFactoryType ? VertexFactoryType->GetName() : TEXT("null"));
 		SaveKeyVal.PermutationId = FName(*FString::Printf(TEXT("Perm_%d"), PermutationId));
-		SaveKeyVal.OutputHash = OutputHash;
-		if (Type)
+		SaveKeyVal.OutputHash = GetOutputHash();
+		if (ShaderType)
 		{
-			Type->GetShaderStableKeyParts(SaveKeyVal);
+			ShaderType->GetShaderStableKeyParts(SaveKeyVal);
 		}
-		FShaderCodeLibrary::AddShaderStableKeyValue(EShaderPlatform(Target.Platform), SaveKeyVal);
+		FShaderCodeLibrary::AddShaderStableKeyValue(GetShaderPlatform(), SaveKeyVal);
 	}
-#endif
 }
+#endif // WITH_EDITOR
 
 bool FShaderPipelineType::bInitialized = false;
 
+static TArray<FShaderPipelineType*>& GetSortedShaderPipelineTypes(FShaderType::EShaderTypeForDynamicCast Type)
+{
+	static TArray<FShaderPipelineType*> SortedTypes[(uint32)FShaderType::EShaderTypeForDynamicCast::NumShaderTypes];
+	return SortedTypes[(uint32)Type];
+}
+
 FShaderPipelineType::FShaderPipelineType(
 	const TCHAR* InName,
 	const FShaderType* InVertexShader,
@@ -1539,6 +705,8 @@ FShaderPipelineType::FShaderPipelineType(
 	bool bInShouldOptimizeUnusedOutputs) :
 	Name(InName),
 	TypeName(Name),
+	HashedName(TypeName),
+	HashedPrimaryShaderFilename(InVertexShader->GetShaderFilename()),
 	GlobalListLink(this),
 	bShouldOptimizeUnusedOutputs(bInShouldOptimizeUnusedOutputs)
 {
@@ -1555,16 +723,20 @@ FShaderPipelineType::FShaderPipelineType(
 
 	if (InPixelShader)
 	{
+		check(InPixelShader->GetTypeForDynamicCast() == InVertexShader->GetTypeForDynamicCast());
 		Stages.Add(InPixelShader);
 		AllStages[SF_Pixel] = InPixelShader;
 	}
 	if (InGeometryShader)
 	{
+		check(InGeometryShader->GetTypeForDynamicCast() == InVertexShader->GetTypeForDynamicCast());
 		Stages.Add(InGeometryShader);
 		AllStages[SF_Geometry] = InGeometryShader;
 	}
 	if (InDomainShader)
 	{
+		check(InDomainShader->GetTypeForDynamicCast() == InVertexShader->GetTypeForDynamicCast());
+		check(InHullShader->GetTypeForDynamicCast() == InVertexShader->GetTypeForDynamicCast());
 		Stages.Add(InDomainShader);
 		AllStages[SF_Domain] = InDomainShader;
 
@@ -1587,7 +759,11 @@ FShaderPipelineType::FShaderPipelineType(
 	HashIndex = TypeHashCounter;
 
 	GlobalListLink.LinkHead(GetTypeList());
-	GetNameToTypeMap().Add(TypeName, this);
+	GetNameToTypeMap().Add(HashedName, this);
+
+	TArray<FShaderPipelineType*>& SortedTypes = GetSortedShaderPipelineTypes(InVertexShader->GetTypeForDynamicCast());
+	const int32 SortedIndex = Algo::LowerBoundBy(SortedTypes, HashedName, [](const FShaderPipelineType* InType) { return InType->GetHashedName(); });
+	SortedTypes.Insert(this, SortedIndex);
 
 	// This will trigger if an IMPLEMENT_SHADER_TYPE was in a module not loaded before InitializeShaderTypes
 	// Shader types need to be implemented in modules that are loaded before that
@@ -1596,18 +772,19 @@ FShaderPipelineType::FShaderPipelineType(
 
 FShaderPipelineType::~FShaderPipelineType()
 {
-	GetNameToTypeMap().Remove(TypeName);
+	GetNameToTypeMap().Remove(HashedName);
 	GlobalListLink.Unlink();
+
+	TArray<FShaderPipelineType*>& SortedTypes = GetSortedShaderPipelineTypes(AllStages[SF_Vertex]->GetTypeForDynamicCast());
+	const int32 SortedIndex = Algo::BinarySearchBy(SortedTypes, HashedName, [](const FShaderPipelineType* InType) { return InType->GetHashedName(); });
+	check(SortedIndex != INDEX_NONE);
+	SortedTypes.RemoveAt(SortedIndex);
 }
 
-TMap<FName, FShaderPipelineType*>& FShaderPipelineType::GetNameToTypeMap()
+TMap<FHashedName, FShaderPipelineType*>& FShaderPipelineType::GetNameToTypeMap()
 {
-	static TMap<FName, FShaderPipelineType*>* GShaderPipelineNameToTypeMap = NULL;
-	if (!GShaderPipelineNameToTypeMap)
-	{
-		GShaderPipelineNameToTypeMap = new TMap<FName, FShaderPipelineType*>();
-	}
-	return *GShaderPipelineNameToTypeMap;
+	static TMap<FHashedName, FShaderPipelineType*> GShaderPipelineNameToTypeMap;
+	return GShaderPipelineNameToTypeMap;
 }
 
 TLinkedList<FShaderPipelineType*>*& FShaderPipelineType::GetTypeList()
@@ -1615,6 +792,11 @@ TLinkedList<FShaderPipelineType*>*& FShaderPipelineType::GetTypeList()
 	return GShaderPipelineList;
 }
 
+const TArray<FShaderPipelineType*>& FShaderPipelineType::GetSortedTypes(FShaderType::EShaderTypeForDynamicCast Type)
+{
+	return GetSortedShaderPipelineTypes(Type);
+}
+
 TArray<const FShaderPipelineType*> FShaderPipelineType::GetShaderPipelineTypesByFilename(const TCHAR* Filename)
 {
 	TArray<const FShaderPipelineType*> PipelineTypes;
@@ -1703,42 +885,10 @@ void FShaderPipelineType::Uninitialize()
 	bInitialized = false;
 }
 
-void FShaderPipelineType::GetOutdatedTypes(TArray<FShaderType*>& OutdatedShaderTypes, TArray<const FShaderPipelineType*>& OutdatedShaderPipelineTypes, TArray<const FVertexFactoryType*>& OutdatedFactoryTypes)
-{
-	for (TLinkedList<FShaderPipelineType*>::TIterator It(FShaderPipelineType::GetTypeList()); It; It.Next())
-	{
-		const auto* PipelineType = *It;
-		auto& Stages = PipelineType->GetStages();
-		bool bOutdated = false;
-		for (const FShaderType* ShaderType : Stages)
-		{
-			bOutdated = ShaderType->GetOutdatedCurrentType(OutdatedShaderTypes, OutdatedFactoryTypes) || bOutdated;
-		}
-
-		if (bOutdated)
-		{
-			OutdatedShaderPipelineTypes.AddUnique(PipelineType);
-		}
-	}
-
-	for (int32 TypeIndex = 0; TypeIndex < OutdatedShaderPipelineTypes.Num(); TypeIndex++)
-	{
-		UE_LOG(LogShaders, Warning, TEXT("		Recompiling Pipeline %s"), OutdatedShaderPipelineTypes[TypeIndex]->GetName());
-	}
-}
-
-const FShaderPipelineType* FShaderPipelineType::GetShaderPipelineTypeByName(FName Name)
+const FShaderPipelineType* FShaderPipelineType::GetShaderPipelineTypeByName(const FHashedName& Name)
 {
-	for (TLinkedList<FShaderPipelineType*>::TIterator It(GetTypeList()); It; It.Next())
-	{
-		const FShaderPipelineType* Type = *It;
-		if (Name == Type->GetFName())
-		{
-			return Type;
-		}
-	}
-
-	return nullptr;
+	FShaderPipelineType** FoundType = GetNameToTypeMap().Find(Name);
+	return FoundType ? *FoundType : nullptr;
 }
 
 const FSHAHash& FShaderPipelineType::GetSourceHash(EShaderPlatform ShaderPlatform) const
@@ -1759,7 +909,7 @@ FShaderPipeline::FShaderPipeline(
 	FShader* InDomainShader,
 	FShader* InGeometryShader,
 	FShader* InPixelShader) :
-	PipelineType(InPipelineType),
+	TypeName(InPipelineType->GetHashedName()),
 	VertexShader(InVertexShader),
 	HullShader(InHullShader),
 	DomainShader(InDomainShader),
@@ -1767,56 +917,11 @@ FShaderPipeline::FShaderPipeline(
 	PixelShader(InPixelShader)
 {
 	check(InPipelineType);
-	Validate();
+	Validate(InPipelineType);
 }
 
 FShaderPipeline::FShaderPipeline(const FShaderPipelineType* InPipelineType, const TArray<FShader*>& InStages) :
-	PipelineType(InPipelineType),
-	VertexShader(nullptr),
-	HullShader(nullptr),
-	DomainShader(nullptr),
-	GeometryShader(nullptr),
-	PixelShader(nullptr)
-{
-	check(InPipelineType);
-	for (FShader* Shader : InStages)
-	{
-		if (Shader)
-		{
-			switch (Shader->GetType()->GetFrequency())
-			{
-			case SF_Vertex:
-				check(!VertexShader);
-				VertexShader = Shader;
-				break;
-			case SF_Pixel:
-				check(!PixelShader);
-				PixelShader = Shader;
-				break;
-			case SF_Hull:
-				check(!HullShader);
-				HullShader = Shader;
-				break;
-			case SF_Domain:
-				check(!DomainShader);
-				DomainShader = Shader;
-				break;
-			case SF_Geometry:
-				check(!GeometryShader);
-				GeometryShader = Shader;
-				break;
-			default:
-				checkf(0, TEXT("Invalid stage %u found!"), (uint32)Shader->GetType()->GetFrequency());
-				break;
-			}
-		}
-	}
-
-	Validate();
-}
-
-FShaderPipeline::FShaderPipeline(const FShaderPipelineType* InPipelineType, const TArray< TRefCountPtr<FShader> >& InStages) :
-	PipelineType(InPipelineType),
+	TypeName(InPipelineType->GetHashedName()),
 	VertexShader(nullptr),
 	HullShader(nullptr),
 	DomainShader(nullptr),
@@ -1828,7 +933,7 @@ FShaderPipeline::FShaderPipeline(const FShaderPipelineType* InPipelineType, cons
 	{
 		if (Shader)
 		{
-			switch (Shader->GetType()->GetFrequency())
+			switch (Shader->GetTypeUnfrozen()->GetFrequency())
 			{
 			case SF_Vertex:
 				check(!VertexShader);
@@ -1851,13 +956,13 @@ FShaderPipeline::FShaderPipeline(const FShaderPipelineType* InPipelineType, cons
 				GeometryShader = Shader;
 				break;
 			default:
-				checkf(0, TEXT("Invalid stage %u found!"), (uint32)Shader->GetType()->GetFrequency());
+				checkf(0, TEXT("Invalid stage %u found!"), (uint32)Shader->GetTypeUnfrozen()->GetFrequency());
 				break;
 			}
 		}
 	}
 
-	Validate();
+	Validate(InPipelineType);
 }
 
 FShaderPipeline::~FShaderPipeline()
@@ -1870,26 +975,27 @@ FShaderPipeline::~FShaderPipeline()
 	PixelShader = nullptr;
 }
 
-void FShaderPipeline::Validate()
+void FShaderPipeline::Validate(const FShaderPipelineType* InPipelineType) const
 {
-	for (const FShaderType* Stage : PipelineType->GetStages())
+	check(InPipelineType->GetHashedName() == TypeName);
+	for (const FShaderType* Stage : InPipelineType->GetStages())
 	{
 		switch (Stage->GetFrequency())
 		{
 		case SF_Vertex:
-			check(VertexShader && VertexShader->GetType() == Stage);
+			check(VertexShader && VertexShader->GetTypeUnfrozen() == Stage);
 			break;
 		case SF_Pixel:
-			check(PixelShader && PixelShader->GetType() == Stage);
+			check(PixelShader && PixelShader->GetTypeUnfrozen() == Stage);
 			break;
 		case SF_Hull:
-			check(HullShader && HullShader->GetType() == Stage);
+			check(HullShader && HullShader->GetTypeUnfrozen() == Stage);
 			break;
 		case SF_Domain:
-			check(DomainShader && DomainShader->GetType() == Stage);
+			check(DomainShader && DomainShader->GetTypeUnfrozen() == Stage);
 			break;
 		case SF_Geometry:
-			check(GeometryShader && GeometryShader->GetType() == Stage);
+			check(GeometryShader && GeometryShader->GetTypeUnfrozen() == Stage);
 			break;
 		default:
 			// Can never happen :)
@@ -1898,44 +1004,44 @@ void FShaderPipeline::Validate()
 	}
 }
 
-void FShaderPipeline::CookPipeline(FShaderPipeline* Pipeline)
-{
 #if WITH_EDITOR
-	FShaderCodeLibrary::AddShaderPipeline(Pipeline);
-#endif
-}
-
-void FShaderPipeline::SaveShaderStableKeys(EShaderPlatform TargetShaderPlatform, const struct FStableShaderKeyAndValue& InSaveKeyVal)
+void FShaderPipeline::SaveShaderStableKeys(const FShaderMapPointerTable& InPtrTable, EShaderPlatform TargetShaderPlatform, const struct FStableShaderKeyAndValue& InSaveKeyVal) const
 {
 	// the higher level code can pass SP_NumPlatforms, in which case play it safe and use a platform that we know can remove inteprolators
 	const EShaderPlatform ShaderPlatformThatSupportsRemovingInterpolators = SP_PCD3D_SM5;
 	checkf(RHISupportsShaderPipelines(ShaderPlatformThatSupportsRemovingInterpolators), TEXT("We assumed that shader platform %d supports shaderpipelines while it doesn't"), static_cast<int32>(ShaderPlatformThatSupportsRemovingInterpolators));
 
+	FShaderPipelineType** FoundPipelineType = FShaderPipelineType::GetNameToTypeMap().Find(TypeName);
+	check(FoundPipelineType);
+	FShaderPipelineType* PipelineType = *FoundPipelineType;
+
 	bool bCanHaveUniqueShaders = (TargetShaderPlatform != SP_NumPlatforms) ? PipelineType->ShouldOptimizeUnusedOutputs(TargetShaderPlatform) : PipelineType->ShouldOptimizeUnusedOutputs(ShaderPlatformThatSupportsRemovingInterpolators);
 	if (bCanHaveUniqueShaders)
 	{
 		FStableShaderKeyAndValue SaveKeyVal(InSaveKeyVal);
+
 		SaveKeyVal.SetPipelineHash(this); // could use PipelineType->GetSourceHash(), but each pipeline instance even of the same type can have unique shaders
 
 		// Using GetShaders() would be more future-proof but would result in more dynamic allocation churn.
 		// Instead, mirroring the logic here
-		VertexShader->SaveShaderStableKeys(TargetShaderPlatform, SaveKeyVal);
+		VertexShader->SaveShaderStableKeys(InPtrTable, TargetShaderPlatform, SaveKeyVal);
 
 		if (PixelShader)
 		{
-			PixelShader->SaveShaderStableKeys(TargetShaderPlatform, SaveKeyVal);
+			PixelShader->SaveShaderStableKeys(InPtrTable, TargetShaderPlatform, SaveKeyVal);
 		}
 		if (GeometryShader)
 		{
-			GeometryShader->SaveShaderStableKeys(TargetShaderPlatform, SaveKeyVal);
+			GeometryShader->SaveShaderStableKeys(InPtrTable, TargetShaderPlatform, SaveKeyVal);
 		}
 		if (HullShader)
 		{
-			HullShader->SaveShaderStableKeys(TargetShaderPlatform, SaveKeyVal);
-			DomainShader->SaveShaderStableKeys(TargetShaderPlatform, SaveKeyVal);
+			HullShader->SaveShaderStableKeys(InPtrTable, TargetShaderPlatform, SaveKeyVal);
+			DomainShader->SaveShaderStableKeys(InPtrTable, TargetShaderPlatform, SaveKeyVal);
 		}
 	}
 }
+#endif // WITH_EDITOR
 
 void DumpShaderStats(EShaderPlatform Platform, EShaderFrequency Frequency)
 {
@@ -1972,7 +1078,8 @@ void DumpShaderStats(EShaderPlatform Platform, EShaderFrequency Frequency)
 			int32 NumShaders					= 0;
 			int32 NumPipelines = 0;
 			int32 NumSharedPipelines = 0;
-			for (TMap<FShaderKey,FShader*>::TConstIterator ShaderIt(Type->ShaderIdMap);ShaderIt;++ShaderIt)
+#if 0
+			for (TMap<FShaderId,FShader*>::TConstIterator ShaderIt(Type->ShaderIdMap);ShaderIt;++ShaderIt)
 			{
 				const FShader* Shader = ShaderIt.Value();
 				// Skip shaders that don't match frequency.
@@ -1991,6 +1098,7 @@ void DumpShaderStats(EShaderPlatform Platform, EShaderFrequency Frequency)
 				NumShaders++;
 			}
 			AverageNumInstructions = (float)NumInitializedInstructions / (float)Type->GetNumShaders();
+#endif
 			
 			for (TLinkedList<FShaderPipelineType*>::TConstIterator PipelineIt(FShaderPipelineType::GetTypeList()); PipelineIt; PipelineIt.Next())
 			{
@@ -2106,7 +1214,7 @@ void DumpShaderPipelineStats(EShaderPlatform Platform)
 #endif
 }
 
-FShaderType* FindShaderTypeByName(FName ShaderTypeName)
+FShaderType* FindShaderTypeByName(const FHashedName& ShaderTypeName)
 {
 	FShaderType** FoundShader = FShaderType::GetNameToTypeMap().Find(ShaderTypeName);
 	if (FoundShader)
@@ -2117,7 +1225,6 @@ FShaderType* FindShaderTypeByName(FName ShaderTypeName)
 	return nullptr;
 }
 
-
 void DispatchComputeShader(
 	FRHICommandList& RHICmdList,
 	FShader* Shader,
@@ -2361,7 +1468,7 @@ void ShaderMapAppendKeyString(EShaderPlatform Platform, FString& KeyString)
 	{
 		ShaderFormat->AppendToKeyString(KeyString);
 	}
-
+	
 	// Encode the Metal standard into the shader compile options so that they recompile if the settings change.
 	if (IsMetalPlatform(Platform))
 	{
@@ -2432,7 +1539,7 @@ void ShaderMapAppendKeyString(EShaderPlatform Platform, FString& KeyString)
 			}
 		}
 	}
-	
+
 	if (IsOpenGLPlatform(Platform))
 	{
 		static const auto CVar = IConsoleManager::Get().FindConsoleVariable(TEXT("r.OpenGL.ForceDXC"));
diff --git a/Engine/Source/Runtime/RenderCore/Private/ShaderCodeArchive.cpp b/Engine/Source/Runtime/RenderCore/Private/ShaderCodeArchive.cpp
new file mode 100644
--- /dev/null
+++ b/Engine/Source/Runtime/RenderCore/Private/ShaderCodeArchive.cpp
@@ -0,0 +1,316 @@
+#include "ShaderCodeArchive.h"
+#include "ShaderCodeLibrary.h"
+#include "Shader.h"
+#include "Stats/Stats.h"
+#include "ProfilingDebugging/LoadTimeTracker.h"
+#include "Misc/MemStack.h"
+
+int32 GShaderCodeLibraryAsyncLoadingPriority = int32(AIOP_Normal);
+static FAutoConsoleVariableRef CVarShaderCodeLibraryAsyncLoadingPriority(
+	TEXT("r.ShaderCodeLibrary.DefaultAsyncIOPriority"),
+	GShaderCodeLibraryAsyncLoadingPriority,
+	TEXT(""),
+	ECVF_Default
+);
+
+static const FName ShaderLibraryCompressionFormat = NAME_LZ4;
+
+int32 FSerializedShaderArchive::FindShaderMapWithKey(const FSHAHash& Hash, uint32 Key) const
+{
+	for (uint32 Index = ShaderMapHashTable.First(Key); ShaderMapHashTable.IsValid(Index); Index = ShaderMapHashTable.Next(Index))
+	{
+		if (ShaderMapHashes[Index] == Hash)
+		{
+			return Index;
+		}
+	}
+	return INDEX_NONE;
+}
+
+int32 FSerializedShaderArchive::FindShaderMap(const FSHAHash& Hash) const
+{
+	const uint32 Key = GetTypeHash(Hash);
+	return FindShaderMapWithKey(Hash, Key);
+}
+
+bool FSerializedShaderArchive::FindOrAddShaderMap(const FSHAHash& Hash, int32& OutIndex)
+{
+	const uint32 Key = GetTypeHash(Hash);
+	int32 Index = FindShaderMapWithKey(Hash, Key);
+	bool bAdded = false;
+	if (Index == INDEX_NONE)
+	{
+		Index = ShaderMapHashes.Add(Hash);
+		ShaderMapEntries.AddDefaulted();
+		check(ShaderMapEntries.Num() == ShaderMapHashes.Num());
+		ShaderMapHashTable.Add(Key, Index);
+		bAdded = true;
+	}
+
+	OutIndex = Index;
+	return bAdded;
+}
+
+int32 FSerializedShaderArchive::FindShaderWithKey(const FSHAHash& Hash, uint32 Key) const
+{
+	for (uint32 Index = ShaderHashTable.First(Key); ShaderHashTable.IsValid(Index); Index = ShaderHashTable.Next(Index))
+	{
+		if (ShaderHashes[Index] == Hash)
+		{
+			return Index;
+		}
+	}
+	return INDEX_NONE;
+}
+
+int32 FSerializedShaderArchive::FindShader(const FSHAHash& Hash) const
+{
+	const uint32 Key = GetTypeHash(Hash);
+	return FindShaderWithKey(Hash, Key);
+}
+
+bool FSerializedShaderArchive::FindOrAddShader(const FSHAHash& Hash, int32& OutIndex)
+{
+	const uint32 Key = GetTypeHash(Hash);
+	int32 Index = FindShaderWithKey(Hash, Key);
+	bool bAdded = false;
+	if (Index == INDEX_NONE)
+	{
+		Index = ShaderHashes.Add(Hash);
+		ShaderEntries.AddDefaulted();
+		check(ShaderEntries.Num() == ShaderHashes.Num());
+		ShaderHashTable.Add(Key, Index);
+		bAdded = true;
+	}
+
+	OutIndex = Index;
+	return bAdded;
+}
+
+void FSerializedShaderArchive::DecompressShader(int32 Index, const TArray<TArray<uint8>>& ShaderCode, TArray<uint8>& OutDecompressedShader) const
+{
+	const FShaderCodeEntry& Entry = ShaderEntries[Index];
+	OutDecompressedShader.SetNum(Entry.UncompressedSize, false);
+	if (Entry.Size == Entry.UncompressedSize)
+	{
+		FMemory::Memcpy(OutDecompressedShader.GetData(), ShaderCode[Index].GetData(), Entry.UncompressedSize);
+	}
+	else
+	{
+		bool bSucceed = FCompression::UncompressMemory(ShaderLibraryCompressionFormat, OutDecompressedShader.GetData(), Entry.UncompressedSize, ShaderCode[Index].GetData(), Entry.Size);
+		check(bSucceed);
+	}
+}
+
+void FSerializedShaderArchive::Finalize()
+{
+	// Set the correct offsets
+	{
+		uint64 Offset = 0u;
+		for (FShaderCodeEntry& Entry : ShaderEntries)
+		{
+			Entry.Offset = Offset;
+			Offset += Entry.Size;
+		}
+	}
+
+	PreloadEntries.Empty();
+	for (FShaderMapEntry& ShaderMapEntry : ShaderMapEntries)
+	{
+		check(ShaderMapEntry.NumShaders > 0u);
+		TArray<FFileCachePreloadEntry> SortedPreloadEntries;
+		SortedPreloadEntries.Empty(ShaderMapEntry.NumShaders);
+		for (uint32 i = 0; i < ShaderMapEntry.NumShaders; ++i)
+		{
+			const int32 ShaderIndex = ShaderIndices[ShaderMapEntry.ShaderIndicesOffset + i];
+			const FShaderCodeEntry& ShaderEntry = ShaderEntries[ShaderIndex];
+			SortedPreloadEntries.Add(FFileCachePreloadEntry(ShaderEntry.Offset, ShaderEntry.Size));
+		}
+		SortedPreloadEntries.Sort([](const FFileCachePreloadEntry& Lhs, const FFileCachePreloadEntry& Rhs) { return Lhs.Offset < Rhs.Offset; });
+		SortedPreloadEntries.Add(FFileCachePreloadEntry(INT64_MAX, 0));
+
+		ShaderMapEntry.FirstPreloadIndex = PreloadEntries.Num();
+		FFileCachePreloadEntry CurrentPreloadEntry = SortedPreloadEntries[0];
+		for (uint32 PreloadIndex = 1; PreloadIndex <= ShaderMapEntry.NumShaders; ++PreloadIndex)
+		{
+			const FFileCachePreloadEntry& PreloadEntry = SortedPreloadEntries[PreloadIndex];
+			const int64 Gap = PreloadEntry.Offset - CurrentPreloadEntry.Offset - CurrentPreloadEntry.Size;
+			if (Gap > 1024)
+			{
+				++ShaderMapEntry.NumPreloadEntries;
+				PreloadEntries.Add(CurrentPreloadEntry);
+				CurrentPreloadEntry = PreloadEntry;
+			}
+			else
+			{
+				CurrentPreloadEntry.Size = PreloadEntry.Offset + PreloadEntry.Size - CurrentPreloadEntry.Offset;
+			}
+		}
+		check(ShaderMapEntry.NumPreloadEntries > 0u);
+		check(CurrentPreloadEntry.Size == 0);
+	}
+}
+
+void FSerializedShaderArchive::Serialize(FArchive& Ar)
+{
+	Ar << ShaderMapHashes;
+	Ar << ShaderHashes;
+	Ar << ShaderMapEntries;
+	Ar << ShaderEntries;
+	Ar << PreloadEntries;
+	Ar << ShaderIndices;
+
+	check(ShaderHashes.Num() == ShaderEntries.Num());
+	check(ShaderMapHashes.Num() == ShaderMapEntries.Num());
+
+	if (Ar.IsLoading())
+	{
+		{
+			const uint32 HashSize = FMath::Min<uint32>(0x10000, 1u << FMath::CeilLogTwo(ShaderMapHashes.Num()));
+			ShaderMapHashTable.Initialize(HashSize, ShaderMapHashes.Num());
+			for (int32 Index = 0; Index < ShaderMapHashes.Num(); ++Index)
+			{
+				const uint32 Key = GetTypeHash(ShaderMapHashes[Index]);
+				ShaderMapHashTable.Add(Key, Index);
+			}
+		}
+		{
+			const uint32 HashSize = FMath::Min<uint32>(0x10000, 1u << FMath::CeilLogTwo(ShaderHashes.Num()));
+			ShaderHashTable.Initialize(HashSize, ShaderHashes.Num());
+			for (int32 Index = 0; Index < ShaderHashes.Num(); ++Index)
+			{
+				const uint32 Key = GetTypeHash(ShaderHashes[Index]);
+				ShaderHashTable.Add(Key, Index);
+			}
+		}
+	}
+}
+
+FShaderCodeArchive* FShaderCodeArchive::Create(EShaderPlatform InPlatform, FArchive& Ar, const FString& InDestFilePath, const FString& InLibraryDir, const FString& InLibraryName)
+{
+	FShaderCodeArchive* Library = new FShaderCodeArchive(InPlatform, InLibraryDir, InLibraryName);
+	Ar << Library->SerializedShaders;
+	Library->LibraryCodeOffset = Ar.Tell();
+
+	// Open library for async reads
+	Library->FileCacheHandle = IFileCacheHandle::CreateFileCacheHandle(*InDestFilePath);
+
+	UE_LOG(LogShaderLibrary, Display, TEXT("Using %s for material shader code. Total %d unique shaders."), *InDestFilePath, Library->SerializedShaders.ShaderEntries.Num());
+
+	INC_DWORD_STAT_BY(STAT_Shaders_ShaderResourceMemory, Library->GetSizeBytes());
+
+	return Library;
+}
+
+FShaderCodeArchive::FShaderCodeArchive(EShaderPlatform InPlatform, const FString& InLibraryDir, const FString& InLibraryName)
+	: FRHIShaderLibrary(InPlatform, InLibraryName)
+	, LibraryDir(InLibraryDir)
+	, LibraryCodeOffset(0)
+	, FileCacheHandle(nullptr)
+{
+}
+
+FShaderCodeArchive::~FShaderCodeArchive()
+{
+	DEC_DWORD_STAT_BY(STAT_Shaders_ShaderResourceMemory, GetSizeBytes());
+	if (FileCacheHandle)
+	{
+		delete FileCacheHandle;
+	}
+}
+
+IMemoryReadStreamRef FShaderCodeArchive::ReadShaderCode(int32 Index)
+{
+	SCOPED_LOADTIMER(FShaderCodeArchive_ReadShaderCode);
+
+	const FShaderCodeEntry& Entry = SerializedShaders.ShaderEntries[Index];
+
+	FGraphEventArray ReadCompleteEvents;
+	IMemoryReadStreamRef LoadedCode = FileCacheHandle->ReadData(ReadCompleteEvents, LibraryCodeOffset + Entry.Offset, Entry.Size, AIOP_CriticalPath);
+	if (ReadCompleteEvents.Num() > 0)
+	{
+		FTaskGraphInterface::Get().WaitUntilTasksComplete(ReadCompleteEvents);
+	}
+
+	return LoadedCode;
+}
+
+FGraphEventRef FShaderCodeArchive::PreloadShader(int32 ShaderIndex)
+{
+	const FShaderCodeEntry& ShaderEntry = SerializedShaders.ShaderEntries[ShaderIndex];
+
+	const EAsyncIOPriorityAndFlags IOPriority = (EAsyncIOPriorityAndFlags)GShaderCodeLibraryAsyncLoadingPriority;
+	const FFileCachePreloadEntry PreloadEntry(ShaderEntry.Offset, ShaderEntry.Size);
+	return FileCacheHandle->PreloadData(&PreloadEntry, 1, IOPriority);
+}
+
+FGraphEventRef FShaderCodeArchive::PreloadShaderMap(int32 ShaderMapIndex)
+{
+	const FShaderMapEntry& ShaderMapEntry = SerializedShaders.ShaderMapEntries[ShaderMapIndex];
+
+	const EAsyncIOPriorityAndFlags IOPriority = (EAsyncIOPriorityAndFlags)GShaderCodeLibraryAsyncLoadingPriority;
+	return FileCacheHandle->PreloadData(&SerializedShaders.PreloadEntries[ShaderMapEntry.FirstPreloadIndex], ShaderMapEntry.NumPreloadEntries, IOPriority);
+}
+
+TRefCountPtr<FRHIShader> FShaderCodeArchive::CreateShader(int32 Index)
+{
+	TRefCountPtr<FRHIShader> Shader;
+
+	IMemoryReadStreamRef Code = ReadShaderCode(Index);
+	if (Code)
+	{
+		FMemStackBase& MemStack = FMemStack::Get();
+		const FShaderCodeEntry& ShaderEntry = SerializedShaders.ShaderEntries[Index];
+		check(ShaderEntry.Size == Code->GetSize());
+		const uint8* ShaderCode = nullptr;
+
+		FMemMark Mark(MemStack);
+		if (ShaderEntry.UncompressedSize != ShaderEntry.Size)
+		{
+			void* UncompressedCode = MemStack.Alloc(ShaderEntry.UncompressedSize, 16);
+			const bool bDecompressResult = FCompression::UncompressMemoryStream(ShaderLibraryCompressionFormat, UncompressedCode, ShaderEntry.UncompressedSize, Code, 0, ShaderEntry.Size);
+			check(bDecompressResult);
+			ShaderCode = (uint8*)UncompressedCode;
+		}
+		else
+		{
+			int64 ReadSize = 0;
+			ShaderCode = (uint8*)Code->Read(ReadSize, 0, ShaderEntry.UncompressedSize);
+			if (ReadSize != ShaderEntry.UncompressedSize)
+			{
+				// Unable to read contiguous block of code, need to copy to temp buffer
+				void* UncompressedCode = MemStack.Alloc(ShaderEntry.UncompressedSize, 16);
+				Code->CopyTo(UncompressedCode, 0, ShaderEntry.UncompressedSize);
+				ShaderCode = (uint8*)UncompressedCode;
+			}
+		}
+
+		const auto ShaderCodeView = MakeArrayView(ShaderCode, ShaderEntry.UncompressedSize);
+		const FSHAHash& ShaderHash = SerializedShaders.ShaderHashes[Index];
+		switch (ShaderEntry.Frequency)
+		{
+		case SF_Vertex: Shader = RHICreateVertexShader(ShaderCodeView, ShaderHash); CheckShaderCreation(Shader, Index); break;
+		case SF_Pixel: Shader = RHICreatePixelShader(ShaderCodeView, ShaderHash); CheckShaderCreation(Shader, Index); break;
+		case SF_Geometry: Shader = RHICreateGeometryShader(ShaderCodeView, ShaderHash); CheckShaderCreation(Shader, Index); break;
+		case SF_Hull: Shader = RHICreateHullShader(ShaderCodeView, ShaderHash); CheckShaderCreation(Shader, Index); break;
+		case SF_Domain: Shader = RHICreateDomainShader(ShaderCodeView, ShaderHash); CheckShaderCreation(Shader, Index); break;
+		case SF_Compute: Shader = RHICreateComputeShader(ShaderCodeView, ShaderHash); CheckShaderCreation(Shader, Index); break;
+#if RHI_RAYTRACING
+		case SF_RayGen: case SF_RayMiss: case SF_RayHitGroup: case SF_RayCallable:
+			if (GRHISupportsRayTracing)
+			{
+				Shader = RHICreateRayTracingShader(ShaderCodeView, ShaderHash, ShaderEntry.GetFrequency());
+				CheckShaderCreation(Shader, Index);
+			}
+			break;
+#endif // RHI_RAYTRACING
+		default: checkNoEntry(); break;
+		}
+
+		if (Shader)
+		{
+			Shader->SetHash(ShaderHash);
+		}
+	}
+	return Shader;
+}
