Copying Private-LoadTimes-4.24 stream to Main. Biggest changes are in Materials/Shader memory freezing.

diff --git a/Engine/Source/Runtime/RenderCore/Private/ShaderCodeLibrary.cpp b/Engine/Source/Runtime/RenderCore/Private/ShaderCodeLibrary.cpp
--- a/Engine/Source/Runtime/RenderCore/Private/ShaderCodeLibrary.cpp
+++ b/Engine/Source/Runtime/RenderCore/Private/ShaderCodeLibrary.cpp
@@ -5,6 +5,7 @@ ShaderCodeLibrary.cpp: Bound shader state cache implementation.
 =============================================================================*/
 
 #include "ShaderCodeLibrary.h"
+#include "ShaderCodeArchive.h"
 #include "Shader.h"
 #include "Algo/Replace.h"
 #include "Containers/StringView.h"
@@ -22,13 +23,14 @@ ShaderCodeLibrary.cpp: Bound shader state cache implementation.
 #include "PipelineFileCache.h"
 #include "Interfaces/IPluginManager.h"
 #include "Hash/CityHash.h"
-#include "Containers/SortedMap.h"
+#include "Containers/HashTable.h"
+#include "FileCache/FileCache.h"
 
-#include "Interfaces/IShaderFormatArchive.h"
 #include "ShaderPipelineCache.h"
 #include "Misc/FileHelper.h"
 #include "Misc/ConfigCacheIni.h"
 #include "Misc/CommandLine.h"
+#include "ProfilingDebugging/LoadTimeTracker.h"
 
 #if WITH_EDITORONLY_DATA
 #include "Modules/ModuleManager.h"
@@ -43,31 +45,38 @@ ShaderCodeLibrary.cpp: Bound shader state cache implementation.
 
 DEFINE_LOG_CATEGORY(LogShaderLibrary);
 
-static const FName ShaderLibraryCompressionFormat = NAME_Zlib;
-
-static uint32 GShaderCodeArchiveVersion = 1;
+static uint32 GShaderCodeArchiveVersion = 2;
 static uint32 GShaderPipelineArchiveVersion = 1;
 
 static FString ShaderExtension = TEXT(".ushaderbytecode");
 static FString StableExtension = TEXT(".scl.csv");
 static FString PipelineExtension = TEXT(".ushaderpipelines");
 
-int32 GShaderCodeLibraryAsyncLoadingPriority = int32(AIOP_Normal);
-static FAutoConsoleVariableRef CVarShaderCodeLibraryAsyncLoadingPriority(
-	TEXT("r.ShaderCodeLibrary.DefaultAsyncIOPriority"),
-	GShaderCodeLibraryAsyncLoadingPriority,
-	TEXT(""),
-	ECVF_Default
-);
-
 int32 GShaderCodeLibrarySeperateLoadingCache = 0;
 static FAutoConsoleVariableRef CVarShaderCodeLibrarySeperateLoadingCache(
 	TEXT("r.ShaderCodeLibrary.SeperateLoadingCache"),
-	GShaderCodeLibraryAsyncLoadingPriority,
+	GShaderCodeLibrarySeperateLoadingCache,
 	TEXT("if > 0, each shader code library has it's own loading cache."),
 	ECVF_Default
 );
 
+class FShaderMapResource_SharedCode final : public FShaderMapResource
+{
+public:
+	FShaderMapResource_SharedCode(class FShaderLibraryInstance* InLibraryInstance, int32 InShaderMapIndex);
+	virtual ~FShaderMapResource_SharedCode();
+
+	// FRenderResource interface.
+	virtual void ReleaseRHI() override;
+
+	// FShaderMapResource interface
+	virtual TRefCountPtr<FRHIShader> CreateRHIShader(int32 ShaderIndex) override;
+	virtual bool TryRelease() override;
+	virtual uint32 GetSizeBytes() const override { return sizeof(*this) + GetAllocatedSize(); }
+
+	class FShaderLibraryInstance* LibraryInstance;
+	int32 ShaderMapIndex;
+};
 
 static FString GetCodeArchiveFilename(const FString& BaseDir, const FString& LibraryName, FName Platform)
 {
@@ -94,38 +103,6 @@ static FString GetShaderDebugFolder(const FString& BaseDir, const FString& Libra
 	return BaseDir / FString::Printf(TEXT("ShaderDebug-%s-"), *LibraryName) + Platform.ToString();
 }
 
-static TArray<uint8>& FShaderLibraryHelperUncompressCode(EShaderPlatform Platform, int32 UncompressedSize, TArray<uint8>& Code, TArray<uint8>& UncompressedCode)
-{
-	if (Code.Num() != UncompressedSize)
-	{
-		UncompressedCode.SetNum(UncompressedSize);
-		bool bSucceed = FCompression::UncompressMemory(ShaderLibraryCompressionFormat, UncompressedCode.GetData(), UncompressedSize, Code.GetData(), Code.Num());
-		check(bSucceed);
-		return UncompressedCode;
-	}
-	else
-	{
-		return Code;
-	}
-}
-
-static void FShaderLibraryHelperCompressCode(EShaderPlatform Platform, const TArray<uint8>& UncompressedCode, TArray<uint8>& CompressedCode)
-{
-		int32 CompressedSize = UncompressedCode.Num() * 4.f / 3.f;
-		CompressedCode.SetNumUninitialized(CompressedSize); // Allocate large enough buffer for compressed code
-
-		if (FCompression::CompressMemory(ShaderLibraryCompressionFormat, CompressedCode.GetData(), CompressedSize, UncompressedCode.GetData(), UncompressedCode.Num()))
-		{
-			CompressedCode.SetNum(CompressedSize);
-		}
-		else
-		{
-			CompressedCode = UncompressedCode;
-		}
-		CompressedCode.Shrink();
-}
-
-
 FORCEINLINE FName ParseFNameCached(const FStringView& Src, TMap<uint32,FName>& NameCache)
 {
 	uint32 SrcHash = CityHash32(reinterpret_cast<const char*>(Src.GetData()), Src.Len() * sizeof(TCHAR));
@@ -504,7 +481,7 @@ FString FStableShaderKeyAndValue::HeaderLine()
 	return Result;
 }
 
-void FStableShaderKeyAndValue::SetPipelineHash(FShaderPipeline* Pipeline)
+void FStableShaderKeyAndValue::SetPipelineHash(const FShaderPipeline* Pipeline)
 {
 	if (LIKELY(Pipeline))
 	{
@@ -519,27 +496,27 @@ void FStableShaderKeyAndValue::SetPipelineHash(FShaderPipeline* Pipeline)
 	}
 }
 
-void FShaderCodeLibraryPipeline::Initialize(FShaderPipeline* Pipeline)
+void FShaderCodeLibraryPipeline::Initialize(const FShaderPipeline* Pipeline)
 {
 	check(Pipeline != nullptr);
 
-	if (IsValidRef(Pipeline->VertexShader))
+	if (Pipeline->VertexShader.IsValid())
 	{
 		VertexShader = Pipeline->VertexShader->GetOutputHash();
 	}
-	if (IsValidRef(Pipeline->GeometryShader))
+	if (Pipeline->GeometryShader.IsValid())
 	{
 		GeometryShader = Pipeline->GeometryShader->GetOutputHash();
 	}
-	if (IsValidRef(Pipeline->HullShader))
+	if (Pipeline->HullShader.IsValid())
 	{
 		HullShader = Pipeline->HullShader->GetOutputHash();
 	}
-	if (IsValidRef(Pipeline->DomainShader))
+	if (Pipeline->DomainShader.IsValid())
 	{
 		DomainShader = Pipeline->DomainShader->GetOutputHash();
 	}
-	if (IsValidRef(Pipeline->PixelShader))
+	if (Pipeline->PixelShader.IsValid())
 	{
 		PixelShader = Pipeline->PixelShader->GetOutputHash();
 	}
@@ -558,547 +535,267 @@ void FShaderCodeLibraryPipeline::GetPipelineHash(FSHAHash& Output)
 	Hasher.GetHash(&Output.Hash[0]);
 }
 
-struct FShaderCodeEntry
-{
-	// Serialized
-	uint32 Size;
-	uint64 Offset;
-	uint32 UncompressedSize;
-	uint8 Frequency;
-
-	// Transient
-	TArray<uint8> LoadedCode;
-	int32 NumRefs;
-	TWeakPtr<IAsyncReadRequest, ESPMode::ThreadSafe> ReadRequest;
-#if DO_CHECK	
-	volatile int32 bReadCompleted;
-#endif
-
-	FShaderCodeEntry()
-		: Size(0)
-		, Offset(0)
-		, UncompressedSize(0)
-		, Frequency(0)
-		, NumRefs(0)
-#if DO_CHECK
-		, bReadCompleted(0)
-#endif
-	{}
-};
-
-static FArchive& operator <<(FArchive& Ar, FShaderCodeEntry& Ref)
-{
-	return Ar << Ref.Offset << Ref.Size << Ref.UncompressedSize << Ref.Frequency;
-}
-
-class FShaderCodeArchive final : public FShaderFactoryInterface
+class FShaderLibraryInstance
 {
 public:
-	FShaderCodeArchive(EShaderPlatform InPlatform, const FString& InLibraryDir, const FString& InLibraryName)
-		: FShaderFactoryInterface(InPlatform, InLibraryName)
-		, LibraryDir(InLibraryDir)
-		, LibraryCodeOffset(0)
-		, LibraryAsyncFileHandle(nullptr)
-		, InFlightAsyncReadRequests(0)
+	static FShaderLibraryInstance* Create(EShaderPlatform InShaderPlatform, const FString& ShaderCodeDir, FString const& InLibraryName)
 	{
-		FName PlatformName = LegacyShaderPlatformToShaderFormat(InPlatform);
-		FString DestFilePath = GetCodeArchiveFilename(LibraryDir, LibraryName, PlatformName);
-
-		FArchive* Ar = IFileManager::Get().CreateFileReader(*DestFilePath);
-		if (Ar)
+		FRHIShaderLibraryRef Library;
+		if (RHISupportsNativeShaderLibraries(InShaderPlatform))
 		{
-			uint32 Version = 0;
-			*Ar << Version;
-
-			if (Version == GShaderCodeArchiveVersion)
-			{
-				*Ar << Shaders;
-				LibraryCodeOffset = Ar->Tell();
-			}
-			Ar->Close();
-			delete Ar;
-
-
-			bool ShaderCodeLibrarySeperateLoadingCacheCommandLineOverride = FParse::Param(FCommandLine::Get(), TEXT("ShaderCodeLibrarySeperateLoadingCache"));;
-			if (GShaderCodeLibrarySeperateLoadingCache || ShaderCodeLibrarySeperateLoadingCacheCommandLineOverride)
-			{
-
-				TArray<TArray<FString>> FilesToMakeUnique;
-				FilesToMakeUnique.AddDefaulted(1);
-				FilesToMakeUnique[0].Add(DestFilePath);
-				FPlatformFileManager::Get().GetPlatformFile().MakeUniquePakFilesForTheseFiles(FilesToMakeUnique);
-			}
-
-			// Open library for async reads
-			LibraryAsyncFileHandle = FPlatformFileManager::Get().GetPlatformFile().OpenAsyncRead(*DestFilePath);
-
-			UE_LOG(LogShaderLibrary, Display, TEXT("Using %s for material shader code. Total %d unique shaders."), *DestFilePath, Shaders.Num());
+			Library = RHICreateShaderLibrary(InShaderPlatform, ShaderCodeDir, InLibraryName);
 		}
-	}
 
-	virtual ~FShaderCodeArchive()
-	{
-		if(LibraryAsyncFileHandle != nullptr)
+		if (!Library)
 		{
-			UE_LOG(LogShaderLibrary, Display, TEXT("FShaderCodeArchive: Shutting down %s"), *GetName());
-
-			FScopeLock ScopeLock(&ReadRequestLock);
-
-			const int64 OutstandingReads = FPlatformAtomics::AtomicRead(&InFlightAsyncReadRequests);
-			if(OutstandingReads > 0)
+			const FName PlatformName = LegacyShaderPlatformToShaderFormat(InShaderPlatform);
+			const FString DestFilePath = GetCodeArchiveFilename(ShaderCodeDir, InLibraryName, PlatformName);
+			TUniquePtr<FArchive> Ar(IFileManager::Get().CreateFileReader(*DestFilePath));
+			if (Ar)
 			{
-				const float MaxWaitTimePerRead = 1.f / 60.f;
-				UE_LOG(LogShaderLibrary, Warning, TEXT("FShaderCodeArchive: Library %s has %d inflight requests to LibraryAsyncFileHandle - cancelling and waiting %f seconds each for them to finish."), *GetName(), OutstandingReads, MaxWaitTimePerRead);
-
-				for(auto& Pair : Shaders)
+				uint32 Version = 0;
+				*Ar << Version;
+				if (Version == GShaderCodeArchiveVersion)
 				{
-					FShaderCodeEntry& Entry = Pair.Value;
-					TSharedPtr<IAsyncReadRequest, ESPMode::ThreadSafe> LocalReadRequest = Entry.ReadRequest.Pin();
-					if(LocalReadRequest.IsValid())
+					Library = FShaderCodeArchive::Create(InShaderPlatform, *Ar, DestFilePath, ShaderCodeDir, InLibraryName);
+					if (Library)
 					{
-						LocalReadRequest->Cancel();
-						LocalReadRequest->WaitCompletion(MaxWaitTimePerRead);
+						bool ShaderCodeLibrarySeperateLoadingCacheCommandLineOverride = FParse::Param(FCommandLine::Get(), TEXT("ShaderCodeLibrarySeperateLoadingCache"));;
+						if (GShaderCodeLibrarySeperateLoadingCache || ShaderCodeLibrarySeperateLoadingCacheCommandLineOverride)
+						{
+							TArray<TArray<FString>> FilesToMakeUnique;
+							FilesToMakeUnique.AddDefaulted(1);
+							FilesToMakeUnique[0].Add(DestFilePath);
+							FPlatformFileManager::Get().GetPlatformFile().MakeUniquePakFilesForTheseFiles(FilesToMakeUnique);
+						}
 					}
 				}
 			}
+		}
 
-			delete LibraryAsyncFileHandle;
-			LibraryAsyncFileHandle = nullptr;
+		if (!Library)
+		{
+			return nullptr;
 		}
-	}
 
-	virtual bool IsLibraryNativeFormat() const { return false; }
+		FShaderLibraryInstance* Instance = new FShaderLibraryInstance();
+		Instance->Library = Library;
 
-	TArray<uint8>* LookupShaderCode(const FSHAHash& Hash, int32& OutSize, bool& bWasSynchronous)
-	{
-		FShaderCodeEntry* Entry = Shaders.Find(Hash);
-		bWasSynchronous = false;
-		if (Entry)
-		{
-			FScopeLock ScopeLock(&ReadRequestLock);
+		const int32 NumResources = Library->GetNumShaderMaps();
+		Instance->Resources.AddZeroed(NumResources);
 
-			check(Entry->NumRefs >= 0);
-			if (Entry->NumRefs == 0 && Entry->LoadedCode.Num() == 0)
-			{
-				static FThreadSafeCounter SyncCount;
-				SyncCount.Increment();
-				// Someone has asked for a shader without previously invoking RequestEntry, we cannot afford to crash because this happens all too frequently.
-				double StartTime = FPlatformTime::Seconds();
-				bool bFound = RequestEntryInternal(Hash, nullptr, true);
-				check(bFound);
-				float ThisTimeMS = (FPlatformTime::Seconds() - StartTime) * 1000.0;
-				UE_LOG(LogShaderLibrary, Warning, TEXT("Took %6.2fms (%d total sync shader loads) to synchronously load shader %s from library: %s"), ThisTimeMS, SyncCount.GetValue(), *Hash.ToString(), *GetName());
-				bWasSynchronous = bFound;
-
-				check(Entry->NumRefs > 0);
-				check(Entry->LoadedCode.Num() != 0);
-				check(Entry->bReadCompleted == 1);
-			}
+		const int32 NumShaders = Library->GetNumShaders();
+		Instance->RHIShaders.AddZeroed(NumShaders);
 
-			check(Entry->NumRefs > 0);
-			check(Entry->LoadedCode.Num() != 0);
-			check(Entry->bReadCompleted == 1);
+		INC_DWORD_STAT_BY(STAT_Shaders_ShaderResourceMemory, Instance->GetSizeBytes());
 
-			OutSize = Entry->UncompressedSize;
-			return &Entry->LoadedCode;
-		}
-		return nullptr;
+		return Instance;
 	}
 
-	virtual bool ContainsEntry(const FSHAHash& Hash) final override
+	~FShaderLibraryInstance()
 	{
-		FShaderCodeEntry* Entry = Shaders.Find(Hash);
-		return (Entry != nullptr);
+		DEC_DWORD_STAT_BY(STAT_Shaders_ShaderResourceMemory, GetSizeBytes());
 	}
 
-	virtual bool RequestEntry(const FSHAHash& Hash, FArchive* Ar) final override
-	{
-		return RequestEntryInternal(Hash, Ar, false);
-	}
-	virtual bool RequestEntry(const FSHAHash& Hash, TArray<uint8>& OutRaw) final override
-	{
-		int32 Size = -1;
-		bool bWasSync = false;
-		TArray<uint8>* Code = LookupShaderCode(Hash, Size, bWasSync);
-		if (Code)
-		{
-			TArray<uint8>& UncompressedCode = FShaderLibraryHelperUncompressCode(Platform, Size, *Code, OutRaw);
-			if (bWasSync)
-			{
-				ReleaseShaderCode(Hash);
-			}
-			return true;
-		}
-		return false;
-	}
+	EShaderPlatform GetPlatform() const { return Library->GetPlatform(); }
+	const int32 GetNumResources() const { return Resources.Num(); }
+	const int32 GetNumShaders() const { return RHIShaders.Num(); }
 
-	bool RequestEntryInternal(const FSHAHash& Hash, FArchive* Ar, bool bHiPriSync)
+	uint32 GetSizeBytes() const
 	{
-		FShaderCodeEntry* Entry = Shaders.Find(Hash);
-		if (Entry)
-		{
-			FScopeLock ScopeLock(&ReadRequestLock);
-
-			int32 CodeNumRefs = Entry->NumRefs++;
-			TSharedPtr<IAsyncReadRequest, ESPMode::ThreadSafe> LocalReadRequest = Entry->ReadRequest.Pin();
-			bool bHasReadRequest = LocalReadRequest.IsValid();
-
-			if (CodeNumRefs == 0 && !bHasReadRequest)
-			{
-				// should not have allocated mem for code if there is no active read request
-				check(Entry->LoadedCode.Num() == 0);
-
-				int64 ReadSize = Entry->Size;
-				int64 ReadOffset = LibraryCodeOffset + Entry->Offset;
-				Entry->LoadedCode.SetNumUninitialized(ReadSize);
-				
-				INC_DWORD_STAT_BY_FName(GetMemoryStatType((EShaderFrequency)Entry->Frequency).GetName(), ReadSize);
-				INC_DWORD_STAT_BY(STAT_Shaders_ShaderResourceMemory, sizeof(FShaderCodeEntry) + ReadSize);
-				
-				EAsyncIOPriorityAndFlags IOPriority = bHiPriSync ? AIOP_CriticalPath : (EAsyncIOPriorityAndFlags)GShaderCodeLibraryAsyncLoadingPriority;
-				LocalReadRequest = MakeShareable(LibraryAsyncFileHandle->ReadRequest(ReadOffset, ReadSize, IOPriority, nullptr, Entry->LoadedCode.GetData()));
-
-				Entry->ReadRequest = LocalReadRequest;
-				bHasReadRequest = true;
-			}
-
-			if (bHasReadRequest)
-			{
-				FPlatformAtomics::InterlockedAdd(&InFlightAsyncReadRequests, 1);
-			
-				FExternalReadCallback ExternalReadCallback = [this, Entry, LocalReadRequest](double ReaminingTime)
-				{
-					return this->OnExternalReadCallback(LocalReadRequest, Entry, ReaminingTime);
-				};
-
-				if (!Ar || !Ar->AttachExternalReadDependency(ExternalReadCallback))
-				{
-					// Archive does not support async loading 
-					// do a blocking load
-					ExternalReadCallback(0.0);
-					// Should be loaded now
-					check(Entry->LoadedCode.Num() != 0);
-					check(Entry->bReadCompleted == 1);
-				}
-			}
-			else
-			{
-				// already loaded
-				check(Entry->LoadedCode.Num() != 0);
-				check(Entry->bReadCompleted == 1);
-			}
-
-			return true;
-		}
-
-		return false;
+		return sizeof(*this) + RHIShaders.GetAllocatedSize() + Resources.GetAllocatedSize();
 	}
 
-	bool OnExternalReadCallback(const TSharedPtr<IAsyncReadRequest, ESPMode::ThreadSafe>& AsyncReadRequest, FShaderCodeEntry* Entry, double RemainingTime)
+	const int32 GetNumShadersForShaderMap(int32 ShaderMapIndex) const
 	{
-		if (RemainingTime < 0.0 && !AsyncReadRequest->PollCompletion())
-		{
-			return false;
-		}
-		else if (RemainingTime >= 0.0 && !AsyncReadRequest->WaitCompletion(RemainingTime))
-		{
-			return false;
-		}
-#if DO_CHECK
-		Entry->bReadCompleted = 1;
-#endif
-
-		FPlatformAtomics::InterlockedAdd(&InFlightAsyncReadRequests, -1);
-		
-		return true;
+		return Library->GetNumShadersForShaderMap(ShaderMapIndex);
 	}
 
-	void ReleaseShaderCode(const FSHAHash& Hash)
+	void PreloadShader(int32 ShaderIndex, FArchive* Ar)
 	{
-		FShaderCodeEntry* Entry = Shaders.Find(Hash);
-		if (Entry)
+		SCOPED_LOADTIMER(FShaderLibraryInstance_PreloadShader);
+		FGraphEventRef Event = Library->PreloadShader(ShaderIndex);
+		if (Ar && Event.IsValid())
 		{
-			FScopeLock ScopeLock(&ReadRequestLock);
-
-			Entry->NumRefs--;
-			if (Entry->NumRefs == 0)
+			FExternalReadCallback ExternalReadCallback = [this, Event](double ReaminingTime)
 			{
-				DEC_DWORD_STAT_BY_FName(GetMemoryStatType((EShaderFrequency)Entry->Frequency).GetName(), Entry->LoadedCode.Num());
-				DEC_DWORD_STAT_BY(STAT_Shaders_ShaderResourceMemory, sizeof(FShaderCodeEntry) + Entry->LoadedCode.Num());
-			
-				// should not attempt to release shader code while it's loading
-				check(Entry->ReadRequest.IsValid() == false);
-
-				Entry->LoadedCode.Empty();
-#if DO_CHECK
-				Entry->bReadCompleted = 0;
-#endif
-			}
+				return this->OnExternalReadCallback(Event, ReaminingTime);
+			};
+			Ar->AttachExternalReadDependency(ExternalReadCallback);
 		}
 	}
 
-	FPixelShaderRHIRef CreatePixelShader(const FSHAHash& Hash) override final
+	void PreloadShaderMap(int32 ShaderMapIndex, FArchive* Ar)
 	{
-		FPixelShaderRHIRef Shader;
-
-		int32 Size = 0;
-		bool bWasSync = false;
-		TArray<uint8>* Code = LookupShaderCode(Hash, Size, bWasSync);
-		if (Code)
+		SCOPED_LOADTIMER(FShaderLibraryInstance_PreloadShaderMap);
+		FGraphEventRef Event = Library->PreloadShaderMap(ShaderMapIndex);
+		if (Ar && Event.IsValid())
 		{
-			TArray<uint8> UCode;
-			TArray<uint8>& UncompressedCode = FShaderLibraryHelperUncompressCode(Platform, Size, *Code, UCode);
-			Shader = RHICreatePixelShader(UncompressedCode);
-			CheckShaderCreation(Shader.GetReference(), Hash);
-			if (bWasSync)
+			FExternalReadCallback ExternalReadCallback = [this, Event](double ReaminingTime)
 			{
-				ReleaseShaderCode(Hash);
-			}
+				return this->OnExternalReadCallback(Event, ReaminingTime);
+			};
+			Ar->AttachExternalReadDependency(ExternalReadCallback);
 		}
-		return Shader;
 	}
 
-	FVertexShaderRHIRef CreateVertexShader(const FSHAHash& Hash) override final
+	TRefCountPtr<FShaderMapResource_SharedCode> GetResource(int32 ShaderMapIndex)
 	{
-		FVertexShaderRHIRef Shader;
-
-		int32 Size = 0;
-		bool bWasSync = false;
-		TArray<uint8>* Code = LookupShaderCode(Hash, Size, bWasSync);
-		if (Code)
-		{
-			TArray<uint8> UCode;
-			TArray<uint8>& UncompressedCode = FShaderLibraryHelperUncompressCode(Platform, Size, *Code, UCode);
-			Shader = RHICreateVertexShader(UncompressedCode);
-			CheckShaderCreation(Shader.GetReference(), Hash);
-			if (bWasSync)
-			{
-				ReleaseShaderCode(Hash);
-			}
-		}
-		return Shader;
+		FRWScopeLock Locker(ResourceLock, SLT_ReadOnly);
+		return Resources[ShaderMapIndex];
 	}
 
-	FHullShaderRHIRef CreateHullShader(const FSHAHash& Hash) override final
+	TRefCountPtr<FShaderMapResource_SharedCode> AddOrDeleteResource(FShaderMapResource_SharedCode* Resource)
 	{
-		FHullShaderRHIRef Shader;
+		const int32 ShaderMapIndex = Resource->ShaderMapIndex;
+		TRefCountPtr<FShaderMapResource_SharedCode> OutResource(Resource);
 
-		int32 Size = 0;
-		bool bWasSync = false;
-		TArray<uint8>* Code = LookupShaderCode(Hash, Size, bWasSync);
-		if (Code)
+		FRWScopeLock Locker(ResourceLock, SLT_Write);
+		FShaderMapResource_SharedCode* PrevResource = Resources[ShaderMapIndex];
+		if (!PrevResource)
 		{
-			TArray<uint8> UCode;
-			TArray<uint8>& UncompressedCode = FShaderLibraryHelperUncompressCode(Platform, Size, *Code, UCode);
-			Shader = RHICreateHullShader(UncompressedCode);
-			CheckShaderCreation(Shader.GetReference(), Hash);
-			if (bWasSync)
-			{
-				ReleaseShaderCode(Hash);
-			}
+			Resources[ShaderMapIndex] = Resource;
 		}
-		return Shader;
-	}
-
-	FDomainShaderRHIRef CreateDomainShader(const FSHAHash& Hash) override final
-	{
-		FDomainShaderRHIRef Shader;
-
-		int32 Size = 0;
-		bool bWasSync = false;
-		TArray<uint8>* Code = LookupShaderCode(Hash, Size, bWasSync);
-		if (Code)
+		else
 		{
-			TArray<uint8> UCode;
-			TArray<uint8>& UncompressedCode = FShaderLibraryHelperUncompressCode(Platform, Size, *Code, UCode);
-			Shader = RHICreateDomainShader(UncompressedCode);
-			CheckShaderCreation(Shader.GetReference(), Hash);
-			if (bWasSync)
-			{
-				ReleaseShaderCode(Hash);
-			}
+			OutResource = PrevResource;
 		}
-		return Shader;
+		return OutResource;
 	}
 
-	FGeometryShaderRHIRef CreateGeometryShader(const FSHAHash& Hash) override final
+	bool TryRemoveResource(FShaderMapResource_SharedCode* Resource)
 	{
-		FGeometryShaderRHIRef Shader;
+		FRWScopeLock Locker(ResourceLock, SLT_Write);
 
-		int32 Size = 0;
-		bool bWasSync = false;
-		TArray<uint8>* Code = LookupShaderCode(Hash, Size, bWasSync);
-		if (Code)
+		if (Resource->GetNumRefs() == 0)
 		{
-			TArray<uint8> UCode;
-			TArray<uint8>& UncompressedCode = FShaderLibraryHelperUncompressCode(Platform, Size, *Code, UCode);
-			Shader = RHICreateGeometryShader(UncompressedCode);
-			CheckShaderCreation(Shader.GetReference(), Hash);
-			if (bWasSync)
-			{
-				ReleaseShaderCode(Hash);
-			}
+			const int32 ShaderMapIndex = Resource->ShaderMapIndex;
+			check(Resources[ShaderMapIndex] == Resource);
+			Resources[ShaderMapIndex] = nullptr;
+			return true;
 		}
-		return Shader;
+
+		// Another thread found the resource after ref-count was decremented to zero
+		return false;
 	}
 
-	FComputeShaderRHIRef CreateComputeShader(const FSHAHash& Hash) override final
+	TRefCountPtr<FRHIShader> GetOrCreateShader(int32 ShaderIndex)
 	{
-		FComputeShaderRHIRef Shader;
-
-		int32 Size = 0;
-		bool bWasSync = false;
-		TArray<uint8>* Code = LookupShaderCode(Hash, Size, bWasSync);
-		if (Code)
+		const int32 LockIndex = ShaderIndex % NumShaderLocks;
+		TRefCountPtr<FRHIShader> Shader;
 		{
-			TArray<uint8> UCode;
-			TArray<uint8>& UncompressedCode = FShaderLibraryHelperUncompressCode(Platform, Size, *Code, UCode);
-			Shader = RHICreateComputeShader(UncompressedCode);
-			CheckShaderCreation(Shader.GetReference(), Hash);
-			if (bWasSync)
-			{
-				ReleaseShaderCode(Hash);
-			}
+			FRWScopeLock Locker(ShaderLocks[LockIndex], SLT_ReadOnly);
+			Shader = RHIShaders[ShaderIndex];
+		}
+		if (!Shader)
+		{
+			FRWScopeLock Locker(ShaderLocks[LockIndex], SLT_Write);
+			Shader = Library->CreateShader(ShaderIndex);
+			RHIShaders[ShaderIndex] = Shader;
 		}
 		return Shader;
 	}
 
-	FRayTracingShaderRHIRef CreateRayTracingShader(EShaderFrequency Frequency, const FSHAHash& Hash) override final
+	void ReleaseShader(int32 ShaderIndex)
 	{
-		FRayTracingShaderRHIRef Shader;
-
-#if RHI_RAYTRACING
-		int32 Size = 0;
-		bool bWasSync = false;
-		TArray<uint8>* Code = LookupShaderCode(Hash, Size, bWasSync);
-		if (Code)
+		const int32 LockIndex = ShaderIndex % NumShaderLocks;
+		FRWScopeLock Locker(ShaderLocks[LockIndex], SLT_Write);
+		FRHIShader* Shader = RHIShaders[ShaderIndex];
+		if(Shader)
 		{
-			TArray<uint8> UCode;
-			TArray<uint8>& UncompressedCode = FShaderLibraryHelperUncompressCode(Platform, Size, *Code, UCode);
-			Shader = RHICreateRayTracingShader(UncompressedCode, Frequency);
-			CheckShaderCreation(Shader.GetReference(), Hash);
-			if (bWasSync)
+			// The library instance is holding one ref
+			// External caller of this method must be holding a ref as well, so there must be at least 2 refs
+			// If those are the only 2 refs, we release the ref held by the library instance, to allow the shader to be destroyed once caller releases its ref
+			const uint32 NumRefs = Shader->GetRefCount();
+			check(NumRefs > 1u);
+			if(NumRefs == 2u)
 			{
-				ReleaseShaderCode(Hash);
+				RHIShaders[ShaderIndex].SafeRelease();
 			}
 		}
-#endif // RHI_RAYTRACING
-
-		return Shader;
 	}
 
-	class FShaderCodeLibraryIterator : public FRHIShaderLibrary::FShaderLibraryIterator
-	{
-	public:
-		FShaderCodeLibraryIterator(FShaderCodeArchive* Owner, EShaderPlatform Plat, TMap<FSHAHash, FShaderCodeEntry>::TIterator It)
-			: FRHIShaderLibrary::FShaderLibraryIterator(Owner)
-			, Platform(Plat)
-			, IteratorImpl(It)
-		{}
-
-		virtual bool IsValid() const final override
-		{
-			return !!IteratorImpl;
-		}
-
-		virtual FRHIShaderLibrary::FShaderLibraryEntry operator*() const final override
-		{
-			FRHIShaderLibrary::FShaderLibraryEntry Entry;
-			TPair<FSHAHash, FShaderCodeEntry> const& Pair = *IteratorImpl;
-			Entry.Hash = Pair.Key;
-			Entry.Frequency = (EShaderFrequency)Pair.Value.Frequency;
-			Entry.Platform = Platform;
-			return Entry;
-		}
-
-		virtual FRHIShaderLibrary::FShaderLibraryIterator& operator++() final override
-		{
-			++IteratorImpl;
-			return *this;
-		}
+public:
+	FRHIShaderLibraryRef Library;
 
-	private:
-		EShaderPlatform Platform;
-		TMap<FSHAHash, FShaderCodeEntry>::TIterator IteratorImpl;
-	};
+private:
+	static const int32 NumShaderLocks = 32;
 
-	virtual TRefCountPtr<FRHIShaderLibrary::FShaderLibraryIterator> CreateIterator(void) override final
-	{
-		return new FShaderCodeLibraryIterator(this, Platform, Shaders.CreateIterator());
-	}
+	FShaderLibraryInstance() {}
 
-	TSet<FShaderCodeLibraryPipeline> const* GetShaderPipelines(EShaderPlatform InPlatform)
+	bool OnExternalReadCallback(const FGraphEventRef& Event, double RemainingTime)
 	{
-		if (Pipelines.Num() == 0)
+		if (Event)
 		{
-			FName PlatformName = LegacyShaderPlatformToShaderFormat(InPlatform);
-			FString DestFilePath = GetPipelinesArchiveFilename(LibraryDir, LibraryName, PlatformName);
-
-			FArchive* Ar = IFileManager::Get().CreateFileReader(*DestFilePath);
-			if (Ar)
+			if (RemainingTime < 0.0)
 			{
-				uint32 Version = 0;
-				*Ar << Version;
-
-				if (Version == GShaderPipelineArchiveVersion)
-				{
-					*Ar << Pipelines;
-				}
-
-				Ar->Close();
-				delete Ar;
+				return Event->IsComplete();
 			}
+			FTaskGraphInterface::Get().WaitUntilTaskCompletes(Event);
 		}
-
-		return &Pipelines;
-	}
-
-	virtual uint32 GetShaderCount(void) const override final
-	{
-		return Shaders.Num();
+		return true;
 	}
 
-private:
-	// Library directory
-	FString LibraryDir;
+	TArray<TRefCountPtr<FRHIShader>> RHIShaders;
+	TArray<FShaderMapResource_SharedCode*> Resources;
+	FRWLock ShaderLocks[NumShaderLocks];
+	FRWLock ResourceLock;
+};
 
-	// Offset at where shader code starts in a code library
-	int64 LibraryCodeOffset;
+FShaderMapResource_SharedCode::FShaderMapResource_SharedCode(FShaderLibraryInstance* InLibraryInstance, int32 InShaderMapIndex)
+	: FShaderMapResource(InLibraryInstance->GetPlatform(), InLibraryInstance->GetNumShadersForShaderMap(InShaderMapIndex))
+	, LibraryInstance(InLibraryInstance)
+	, ShaderMapIndex(InShaderMapIndex)
+{
+}
 
-	// Library file handle for async reads
-	IAsyncReadFileHandle* LibraryAsyncFileHandle;
-	FCriticalSection ReadRequestLock;
+FShaderMapResource_SharedCode::~FShaderMapResource_SharedCode()
+{
 	
-	// A count of the number of LibraryAsync Read Requests in flight
-	volatile int64 InFlightAsyncReadRequests;
+}
 
-	// The shader code present in the library (note: in the editor this is a TSortedMap, but the classes are guaranteed to be serialiation-compatible)
-	TMap<FSHAHash, FShaderCodeEntry> Shaders;
+TRefCountPtr<FRHIShader> FShaderMapResource_SharedCode::CreateRHIShader(int32 ShaderIndex)
+{
+	SCOPED_LOADTIMER(FShaderMapResource_SharedCode_InitRHI);
 
-	// De-serialised pipeline map
-	TSet<FShaderCodeLibraryPipeline> Pipelines;
+	const int32 LibraryShaderIndex = LibraryInstance->Library->GetShaderIndex(ShaderMapIndex, ShaderIndex);
+	return LibraryInstance->GetOrCreateShader(LibraryShaderIndex);
+}
 
-	FORCENOINLINE void CheckShaderCreation(void* ShaderPtr, const FSHAHash& Hash)
+void FShaderMapResource_SharedCode::ReleaseRHI()
+{
+	const int32 NumShaders = GetNumShaders();
+	for (int32 i = 0; i < NumShaders; ++i)
 	{
-#if CHECK_SHADER_CREATION
-		if (!ShaderPtr)
+		if (HasShader(i))
 		{
-			FSHAHash DebugCopy;
-			FMemory::Memcpy(DebugCopy.Hash, Hash.Hash, sizeof(Hash.Hash));
-			UE_LOG(LogShaderLibrary, Fatal, TEXT("Failed to create shader %s, %s, %s"), *DebugCopy.ToString(), *LibraryName, *LibraryDir);
+			const int32 ShaderIndex = LibraryInstance->Library->GetShaderIndex(ShaderMapIndex, i);
+			LibraryInstance->ReleaseShader(ShaderIndex);
 		}
-#endif
 	}
-};
 
-#if WITH_EDITOR
-template <>
-struct TLess<FSHAHash>
+	FShaderMapResource::ReleaseRHI();
+}
+
+bool FShaderMapResource_SharedCode::TryRelease()
 {
-	FORCEINLINE bool operator()(const FSHAHash& Lhs, const FSHAHash& Rhs) const
+	if (LibraryInstance->TryRemoveResource(this))
 	{
-		return FMemory::Memcmp(&Lhs.Hash, &Rhs.Hash, sizeof(FSHAHash::Hash)) < 0;
+		return true;
 	}
+
+	return false;
+}
+
+#if WITH_EDITOR
+struct FShaderCodeStats
+{
+	int64 ShadersSize;
+	int64 ShadersUniqueSize;
+	int32 NumShaders;
+	int32 NumUniqueShaders;
+	//int32 NumPipelines;
+	//int32 NumUniquePipelines;
 };
 
 struct FEditorShaderCodeArchive
@@ -1109,6 +806,9 @@ struct FEditorShaderCodeArchive
 	{
 		Format = GetTargetPlatformManagerRef().FindShaderFormat(InFormat);
 		check(Format);
+
+		SerializedShaders.ShaderHashTable.Initialize(0x10000);
+		SerializedShaders.ShaderMapHashTable.Initialize(0x10000);
 	}
 
 	~FEditorShaderCodeArchive() {}
@@ -1123,9 +823,8 @@ struct FEditorShaderCodeArchive
 		check(LibraryName.Len() == 0);
 		check(Name.Len() > 0);
 		LibraryName = Name;
-		Offset = 0;
-		Shaders.Empty();
-		Pipelines.Empty();
+		SerializedShaders.Empty();
+		ShaderCode.Empty();
 	}
 
 	void CloseLibrary(FString const& Name)
@@ -1134,53 +833,107 @@ struct FEditorShaderCodeArchive
 		LibraryName = TEXT("");
 	}
 
-	bool HasShader(const FSHAHash& Hash) const
+	bool HasShaderMap(const FSHAHash& Hash) const
 	{
-		return Shaders.Contains(Hash);
+		return SerializedShaders.FindShaderMap(Hash) != INDEX_NONE;
 	}
 
-	bool AddShader(uint8 Frequency, const FSHAHash& Hash, TArray<uint8> const& InCode, int32 const UncompressedSize)
+	int32 AddShaderCode(FShaderCodeStats& CodeStats, const FShaderMapResourceCode* Code)
 	{
-		bool bAdd = false;
-		if (!Shaders.Contains(Hash))
+		int32 ShaderMapIndex = INDEX_NONE;
+		if (SerializedShaders.FindOrAddShaderMap(Code->ResourceHash, ShaderMapIndex))
 		{
-#if DO_CHECK
-			uint8 Count = 0;
-			for (uint8 i : InCode)
+			const int32 NumShaders = Code->ShaderEntries.Num();
+			FShaderMapEntry& ShaderMapEntry = SerializedShaders.ShaderMapEntries[ShaderMapIndex];
+			ShaderMapEntry.NumShaders = NumShaders;
+			ShaderMapEntry.ShaderIndicesOffset = SerializedShaders.ShaderIndices.AddZeroed(NumShaders);
+
+			for(int32 i = 0; i < NumShaders; ++i)
 			{
-				Count |= i;
+				int32 ShaderIndex = INDEX_NONE;
+				if (SerializedShaders.FindOrAddShader(Code->ShaderHashes[i], ShaderIndex))
+				{
+					const FShaderMapResourceCode::FShaderEntry& SourceShaderEntry = Code->ShaderEntries[i];
+					FShaderCodeEntry& SerializedShaderEntry = SerializedShaders.ShaderEntries[ShaderIndex];
+					SerializedShaderEntry.Frequency = SourceShaderEntry.Frequency;
+					SerializedShaderEntry.Size = SourceShaderEntry.CompressedSize;
+					SerializedShaderEntry.UncompressedSize = SourceShaderEntry.UncompressedSize;
+					ShaderCode.Emplace(Code->ShaderCode.GetData() + SourceShaderEntry.Offset, SourceShaderEntry.CompressedSize);
+					check(ShaderCode.Num() == SerializedShaders.ShaderEntries.Num());
+
+					CodeStats.NumUniqueShaders++;
+					CodeStats.ShadersUniqueSize += SourceShaderEntry.CompressedSize;
+				}
+				SerializedShaders.ShaderIndices[ShaderMapEntry.ShaderIndicesOffset + i] = ShaderIndex;
 			}
-			check(Count > 0);
-#endif
+		}
+		return ShaderMapIndex;
+	}
 
-			FShaderCodeEntry Entry;
-			Entry.Size = InCode.Num();
-			Entry.Offset = Offset;
-			Entry.UncompressedSize = UncompressedSize;
-			Entry.Frequency = Frequency;
-			Entry.LoadedCode = InCode;
+	int32 AddShaderCode(int32 OtherShaderMapIndex, const FEditorShaderCodeArchive& OtherArchive)
+	{
+		int32 ShaderMapIndex = 0;
+		if (SerializedShaders.FindOrAddShaderMap(OtherArchive.SerializedShaders.ShaderMapHashes[OtherShaderMapIndex], ShaderMapIndex))
+		{
+			const FShaderMapEntry& PrevShaderMapEntry = OtherArchive.SerializedShaders.ShaderMapEntries[OtherShaderMapIndex];
+			FShaderMapEntry& ShaderMapEntry = SerializedShaders.ShaderMapEntries[ShaderMapIndex];
+			ShaderMapEntry.NumShaders = PrevShaderMapEntry.NumShaders;
+			ShaderMapEntry.ShaderIndicesOffset = SerializedShaders.ShaderIndices.AddZeroed(ShaderMapEntry.NumShaders);
 
-			Offset += Entry.Size;
+			for (uint32 i = 0; i < ShaderMapEntry.NumShaders; ++i)
+			{
+				const int32 OtherShaderIndex = OtherArchive.SerializedShaders.ShaderIndices[PrevShaderMapEntry.ShaderIndicesOffset + i];
+				int32 ShaderIndex = 0;
+				if (SerializedShaders.FindOrAddShader(OtherArchive.SerializedShaders.ShaderHashes[OtherShaderIndex], ShaderIndex))
+				{
+					const FShaderCodeEntry& OtherShaderEntry = OtherArchive.SerializedShaders.ShaderEntries[OtherShaderIndex];
+					SerializedShaders.ShaderEntries[ShaderIndex] = OtherShaderEntry;
 
-			Shaders.Add(Hash, Entry);
-			bAdd = true;
+					ShaderCode.Add(OtherArchive.ShaderCode[OtherShaderIndex]);
+					check(ShaderCode.Num() == SerializedShaders.ShaderEntries.Num());
+				}
+				SerializedShaders.ShaderIndices[ShaderMapEntry.ShaderIndicesOffset + i] = ShaderIndex;
+			}
 		}
-		return bAdd;
+		return ShaderMapIndex;
 	}
 
-	bool AddPipeline(FShaderPipeline* Pipeline)
+	int32 AddShaderCode(int32 OtherShaderMapIndex,
+		const FSerializedShaderArchive& OtherShaders,
+		int64 OtherShaderCodeOffset,
+		FArchive& Ar)
 	{
-		check(LibraryName.Len() != 0);
-		EShaderPlatform ShaderPlatform = ShaderFormatToLegacyShaderPlatform(FormatName);
-
-		FShaderCodeLibraryPipeline LibraryPipeline;
-		LibraryPipeline.Initialize(Pipeline);
-		if (!Pipelines.Contains(LibraryPipeline))
+		int32 ShaderMapIndex = 0;
+		if (SerializedShaders.FindOrAddShaderMap(OtherShaders.ShaderMapHashes[OtherShaderMapIndex], ShaderMapIndex))
 		{
-			Pipelines.Add(LibraryPipeline);
-			return true;
+			const FShaderMapEntry& PrevShaderMapEntry = OtherShaders.ShaderMapEntries[OtherShaderMapIndex];
+			FShaderMapEntry& ShaderMapEntry = SerializedShaders.ShaderMapEntries[ShaderMapIndex];
+			ShaderMapEntry.NumShaders = PrevShaderMapEntry.NumShaders;
+			ShaderMapEntry.ShaderIndicesOffset = SerializedShaders.ShaderIndices.AddZeroed(ShaderMapEntry.NumShaders);
+
+			for (uint32 i = 0; i < ShaderMapEntry.NumShaders; ++i)
+			{
+				const int32 OtherShaderIndex = OtherShaders.ShaderIndices[PrevShaderMapEntry.ShaderIndicesOffset + i];
+				int32 ShaderIndex = 0;
+				if (SerializedShaders.FindOrAddShader(OtherShaders.ShaderHashes[OtherShaderIndex], ShaderIndex))
+				{
+					const FShaderCodeEntry& OtherShaderEntry = OtherShaders.ShaderEntries[OtherShaderIndex];
+					SerializedShaders.ShaderEntries[ShaderIndex] = OtherShaderEntry;
+
+					TArray<uint8>& Code = ShaderCode.AddDefaulted_GetRef();
+					check(ShaderCode.Num() == SerializedShaders.GetNumShaders());
+
+					// Read shader code from archive and add shader to set
+					const int64 ReadSize = OtherShaderEntry.Size;
+					const int64 ReadOffset = OtherShaderCodeOffset + OtherShaderEntry.Offset;
+					Code.SetNumUninitialized(ReadSize);
+					Ar.Seek(ReadOffset);
+					Ar.Serialize(Code.GetData(), ReadSize);
+				}
+				SerializedShaders.ShaderIndices[ShaderMapEntry.ShaderIndicesOffset + i] = ShaderIndex;
+			}
 		}
-		return false;
+		return ShaderMapIndex;
 	}
 
 	bool LoadExistingShaderCodeLibrary(FString const& MetaDataDir)
@@ -1195,15 +948,19 @@ struct FEditorShaderCodeArchive
 			if (ArchiveVersion == GShaderCodeArchiveVersion)
 			{
 				// Read shader library
-				*PrevCookedAr << Shaders;
-				for (auto& Entry : Shaders)
+				*PrevCookedAr << SerializedShaders;
+
+				ShaderCode.AddDefaulted(SerializedShaders.ShaderEntries.Num());
+				for(int32 Index = 0; Index < ShaderCode.Num(); ++Index)
 				{
-					Entry.Value.LoadedCode.SetNumUninitialized(Entry.Value.Size);
-					PrevCookedAr->Serialize(Entry.Value.LoadedCode.GetData(), Entry.Value.Size);
+					const FShaderCodeEntry& Entry = SerializedShaders.ShaderEntries[Index];
+					TArray<uint8>& Code = ShaderCode[Index];
+					Code.SetNumUninitialized(Entry.Size);
+					PrevCookedAr->Serialize(Code.GetData(), Entry.Size);
 					bOK = !PrevCookedAr->GetError();
 					if (!bOK)
 					{
-						UE_LOG(LogShaderLibrary, Error, TEXT("Failed to deserialize shader code for %s from %s"), *Entry.Key.ToString(), *IntermediateFormatPath);
+						UE_LOG(LogShaderLibrary, Error, TEXT("Failed to deserialize shader code for %s from %s"), *SerializedShaders.ShaderHashes[Index].ToString(), *IntermediateFormatPath);
 						break;
 					}
 				}
@@ -1223,6 +980,7 @@ struct FEditorShaderCodeArchive
 			UE_LOG(LogShaderLibrary, Error, TEXT("Failed to open shader code library from %s"), *IntermediateFormatPath);
 		}
 		
+#if 0
 		if (bOK)
 		{
 			FString PipelinesPath = GetPipelinesArchiveFilename(MetaDataDir / TEXT("ShaderLibrarySource"), LibraryName, FormatName);
@@ -1245,6 +1003,7 @@ struct FEditorShaderCodeArchive
 				delete PipelinesArchive;
 			}
 		}
+#endif
 		return bOK;
 	}
 
@@ -1270,32 +1029,14 @@ struct FEditorShaderCodeArchive
 
 					if (Version == GShaderCodeArchiveVersion)
 					{
-						// note: we are likely loading a TSortedMap and not a TMap, but the classes are guaranteed to be serialization-compatible
-						TMap<FSHAHash, FShaderCodeEntry> PrevCookedShaders;
+						FSerializedShaderArchive PrevCookedShaders;
 
 						*PrevCookedAr << PrevCookedShaders;
 						int64 PrevCookedShadersCodeStart = PrevCookedAr->Tell();
 
-						for (TMap<FSHAHash, FShaderCodeEntry>::TIterator It(PrevCookedShaders); It; ++It)
+						for(int32 PrevShaderMapIndex = 0; PrevShaderMapIndex < PrevCookedShaders.ShaderMapEntries.Num(); ++PrevShaderMapIndex)
 						{
-							FSHAHash& Hash = It.Key();
-
-							if (!Shaders.Contains(Hash))
-							{
-								// Shader not in list - lazy load shader code
-								FShaderCodeEntry& CodeEntry = It.Value();
-
-								int64 ReadSize = CodeEntry.Size;
-								int64 ReadOffset = PrevCookedShadersCodeStart + CodeEntry.Offset;
-
-								CodeEntry.LoadedCode.SetNumUninitialized(ReadSize);
-
-								// Read shader code from archive and add shader to set
-								PrevCookedAr->Seek(ReadOffset);
-								PrevCookedAr->Serialize(CodeEntry.LoadedCode.GetData(), ReadSize);
-
-								AddShader(CodeEntry.Frequency, Hash, CodeEntry.LoadedCode, CodeEntry.UncompressedSize);
-							}
+							AddShaderCode(PrevShaderMapIndex, PrevCookedShaders, PrevCookedShadersCodeStart, *PrevCookedAr);
 						}
 					}
 
@@ -1308,6 +1049,7 @@ struct FEditorShaderCodeArchive
 		TArray<FString> PipelineFiles;
 		IFileManager::Get().FindFiles(PipelineFiles, *ShaderIntermediateLocation, *PipelineExtension);
 
+#if 0
 		for (const FString& ShaderFileName : PipelineFiles)
 		{
 			if (ShaderFileName.Contains(LibraryName + TEXT("-") + FormatName.ToString() + TEXT(".")))
@@ -1334,6 +1076,7 @@ struct FEditorShaderCodeArchive
 				}
 			}
 		}
+#endif
 	}
 
 	bool Finalize(FString OutputDir, const FString& MetaOutputDir, bool bNativeFormat, bool bMasterCooker)
@@ -1350,17 +1093,8 @@ struct FEditorShaderCodeArchive
 		EShaderPlatform Platform = ShaderFormatToLegacyShaderPlatform(FormatName);
 
 		// Shader library
-		if (bSuccess && Shaders.Num() > 0)
+		if (bSuccess && SerializedShaders.GetNumShaderMaps() > 0)
 		{
-			// recreate the shader offsets so they match the sorted order
-			uint64 NewOffset = 0;
-			for (TSortedMap<FSHAHash, FShaderCodeEntry, FDefaultAllocator, TLess<FSHAHash>>::TIterator It(Shaders); It; ++It)
-			{
-				It.Value().Offset = NewOffset;
-				NewOffset += It.Value().Size;
-			}
-			checkf(NewOffset == Offset, TEXT("Reordering of shaders at the very end did not end at the original offset - internal error."));
-
 			// Write to a intermediate file
 			FString IntermediateFormatPath = GetShaderCodeFilename(FPaths::ProjectSavedDir() / TEXT("Shaders") / FormatName.ToString(), LibraryName, FormatName);
 			FArchive* FileWriter = IFileManager::Get().CreateFileWriter(*IntermediateFormatPath, FILEWRITE_NoFail);
@@ -1369,13 +1103,15 @@ struct FEditorShaderCodeArchive
 			{
 				check(Format);
 
+				SerializedShaders.Finalize();
+
 				*FileWriter << GShaderCodeArchiveVersion;
 
                 // Write shader library
-                *FileWriter << Shaders;
-                for (auto& Pair : Shaders)
+                *FileWriter << SerializedShaders;
+                for (auto& Code : ShaderCode)
                 {
-                    FileWriter->Serialize(Pair.Value.LoadedCode.GetData(), Pair.Value.Size);
+                    FileWriter->Serialize(Code.GetData(), Code.Num());
                 }
 
 				FileWriter->Close();
@@ -1408,6 +1144,7 @@ struct FEditorShaderCodeArchive
 			}
 		}
 
+#if 0
 		// Pipelines
 		if (bSuccess && Pipelines.Num() > 0)
 		{
@@ -1417,8 +1154,6 @@ struct FEditorShaderCodeArchive
 
 			*FileWriter << GShaderPipelineArchiveVersion;
 
-			Pipelines.Sort([](const FShaderCodeLibraryPipeline& A, const FShaderCodeLibraryPipeline& B){ return A.Hash <= B.Hash; });
-
 			*FileWriter << Pipelines;
 
 			FileWriter->Close();
@@ -1449,19 +1184,18 @@ struct FEditorShaderCodeArchive
                 }
 			}
 		}
+#endif
 
 		return bSuccess;
 	}
 
 	bool PackageNativeShaderLibrary(const FString& ShaderCodeDir)
 	{
-		if (Shaders.Num() == 0)
+		if (SerializedShaders.GetNumShaders() == 0)
 		{
 			return true;
 		}
 
-		bool bOK = false;
-
 		FString IntermediateFormatPath = GetShaderDebugFolder(FPaths::ProjectSavedDir() / TEXT("Shaders") / FormatName.ToString(), LibraryName, FormatName);
 		FString TempPath = IntermediateFormatPath / TEXT("NativeLibrary");
 
@@ -1469,49 +1203,17 @@ struct FEditorShaderCodeArchive
 		IFileManager::Get().MakeDirectory(*ShaderCodeDir, true);
 
 		EShaderPlatform Platform = ShaderFormatToLegacyShaderPlatform(FormatName);
-		IShaderFormatArchive* Archive = Format->CreateShaderArchive(LibraryName, FormatName, TempPath);
-		if (Archive)
+		const bool bOK = Format->CreateShaderArchive(LibraryName, FormatName, TempPath, ShaderCodeDir, IntermediateFormatPath, SerializedShaders, ShaderCode, nullptr);
+		if (bOK)
 		{
-			bOK = true;
-
-			// Add the shaders to the archive.
-			for (auto& Pair : Shaders)
+			// Delete Shader code library / pipelines as we now have native versions
 			{
-				FSHAHash& Hash = Pair.Key;
-				FShaderCodeEntry& Entry = Pair.Value;
-
-				TArray<uint8> UCode;
-				TArray<uint8>& UncompressedCode = FShaderLibraryHelperUncompressCode(Platform, Entry.UncompressedSize, Entry.LoadedCode, UCode);
-
-				if (Format->CanStripShaderCode(true))
-				{
-					if (!Format->StripShaderCode(UncompressedCode, IntermediateFormatPath, true))
-					{
-						bOK = false;
-						break;
-					}
-				}
-
-				if (!Archive->AddShader(Entry.Frequency, Hash, UncompressedCode))
-				{
-					bOK = false;
-					break;
-				}
+				FString OutputFilePath = GetCodeArchiveFilename(ShaderCodeDir, LibraryName, FormatName);
+				IFileManager::Get().Delete(*OutputFilePath);
 			}
-
-			if (bOK)
 			{
-				bOK = Archive->Finalize(ShaderCodeDir, IntermediateFormatPath, nullptr);
-
-				// Delete Shader code library / pipelines as we now have native versions
-				{
-					FString OutputFilePath = GetCodeArchiveFilename(ShaderCodeDir, LibraryName, FormatName);
-					IFileManager::Get().Delete(*OutputFilePath);
-				}
-				{
-					FString OutputFilePath = GetPipelinesArchiveFilename(ShaderCodeDir, LibraryName, FormatName);
-					IFileManager::Get().Delete(*OutputFilePath);
-				}
+				FString OutputFilePath = GetPipelinesArchiveFilename(ShaderCodeDir, LibraryName, FormatName);
+				IFileManager::Get().Delete(*OutputFilePath);
 			}
 		}
 
@@ -1524,14 +1226,15 @@ struct FEditorShaderCodeArchive
 
 	void MakePatchLibrary(TArray<FEditorShaderCodeArchive*> const& OldLibraries, FEditorShaderCodeArchive const& NewLibrary)
 	{
-		for(auto const& Pair : NewLibrary.Shaders)
+		for(int32 NewShaderMapIndex = 0; NewShaderMapIndex < NewLibrary.SerializedShaders.ShaderMapHashes.Num(); ++NewShaderMapIndex)
 		{
-			if (!HasShader(Pair.Key))
+			const FSHAHash& Hash = NewLibrary.SerializedShaders.ShaderMapHashes[NewShaderMapIndex];
+			if (!HasShaderMap(Hash))
 			{
 				bool bInPreviousPatch = false;
 				for (FEditorShaderCodeArchive const* OldLibrary : OldLibraries)
 				{
-					bInPreviousPatch |= OldLibrary->HasShader(Pair.Key);
+					bInPreviousPatch |= OldLibrary->HasShaderMap(Hash);
 					if (bInPreviousPatch)
 					{
 						break;
@@ -1539,15 +1242,10 @@ struct FEditorShaderCodeArchive
 				}
 				if (!bInPreviousPatch)
 				{
-					FShaderCodeEntry Entry = Pair.Value;
-					Entry.Offset = Offset;
-					Offset += Entry.Size;
-					Shaders.Add(Pair.Key, Entry);
+					AddShaderCode(NewShaderMapIndex, NewLibrary);
 				}
 			}
 		}
-		
-		Pipelines = NewLibrary.Pipelines;
 	}
 	
 	static bool CreatePatchLibrary(FName FormatName, FString const& LibraryName, TArray<FString> const& OldMetaDataDirs, FString const& NewMetaDataDir, FString const& OutDir, bool bNativeFormat)
@@ -1571,7 +1269,7 @@ struct FEditorShaderCodeArchive
 			FEditorShaderCodeArchive OutLibrary(FormatName);
 			OutLibrary.OpenLibrary(LibraryName);
 			OutLibrary.MakePatchLibrary(OldLibraries, NewLibrary);
-			bOK = OutLibrary.Offset > 0;
+			bOK = OutLibrary.SerializedShaders.GetNumShaderMaps() > 0;
 			if (bOK)
 			{
 				FString Empty;
@@ -1604,9 +1302,8 @@ struct FEditorShaderCodeArchive
 private:
 	FName FormatName;
 	FString LibraryName;
-	TSortedMap<FSHAHash, FShaderCodeEntry, FDefaultAllocator, TLess<FSHAHash>> Shaders;
-	TSet<FShaderCodeLibraryPipeline> Pipelines;
-	uint64 Offset;
+	FSerializedShaderArchive SerializedShaders;
+	TArray<TArray<uint8>> ShaderCode;
 	const IShaderFormat* Format;
 };
 
@@ -1742,23 +1439,13 @@ private:
 	FString LibraryName;
 	TSet<FStableShaderKeyAndValue> StableMap;
 };
-
-struct FShaderCodeStats
-{
-	int64 ShadersSize;
-	int64 ShadersUniqueSize;
-	int32 NumShaders;
-	int32 NumUniqueShaders;
-	int32 NumPipelines;
-	int32 NumUniquePipelines;
-};
 #endif //WITH_EDITOR
 
 class FShaderCodeLibraryImpl
 {
 	// At runtime, shader code collection for current shader platform
-	TArray<FRHIShaderLibraryRef> ShaderCodeArchiveStack;
-	TSet<FShaderCodeLibraryPipeline> Pipelines;
+	TArray<TUniquePtr<FShaderLibraryInstance>> ShaderCodeArchiveStack;
+
 	EShaderPlatform ShaderPlatform;
 	uint64 ShaderCount;
 	FRWLock LibraryMutex;
@@ -1778,52 +1465,6 @@ class FShaderCodeLibraryImpl
 	bool bSupportsPipelines;
 	bool bNativeFormat;
 
-	class FShaderCodeLibraryIterator : public FRHIShaderLibrary::FShaderLibraryIterator
-	{
-	public:
-		FShaderCodeLibraryIterator(TArray<FRHIShaderLibraryRef>& Stack, FRWLock& InLibraryMutex)
-			: FShaderLibraryIterator(nullptr)
-			, LibraryMutex(InLibraryMutex, SLT_ReadOnly)
-			, IteratorImpl(Stack.CreateIterator())
-		{
-			if (Stack.Num() > 0)
-			{
-				Current = (*IteratorImpl)->CreateIterator();
-				ShaderLibrarySource = *IteratorImpl;
-			}
-		}
-
-		virtual bool IsValid() const final override
-		{
-			return IsValidRef(Current) && Current->IsValid();
-		}
-
-		virtual FRHIShaderLibrary::FShaderLibraryEntry operator*() const final override
-		{
-			check(IsValid());
-			return *(*Current);
-		}
-		virtual FShaderLibraryIterator& operator++() final override
-		{
-			++(*Current);
-			if (!Current->IsValid())
-			{
-				++IteratorImpl;
-				if (!!IteratorImpl)
-				{
-					Current = (*IteratorImpl)->CreateIterator();
-					ShaderLibrarySource = *IteratorImpl;
-				}
-			}
-			return *this;
-		}
-
-	private:
-		FRWScopeLock LibraryMutex;
-		TArray<FRHIShaderLibraryRef>::TIterator IteratorImpl;
-		TRefCountPtr<FRHIShaderLibrary::FShaderLibraryIterator> Current;
-	};
-
 public:
 	static FShaderCodeLibraryImpl* Impl;
 
@@ -1909,8 +1550,8 @@ public:
 			FRWScopeLock(LibraryMutex, SLT_Write);
 			for (uint32 i = ShaderCodeArchiveStack.Num(); i > 0; i--)
 			{
-				FRHIShaderLibrary* ShaderCodeArchive = ShaderCodeArchiveStack[i - 1];
-				if (ShaderCodeArchive->GetName() == Name)
+				FShaderLibraryInstance* LibraryInstance = ShaderCodeArchiveStack[i - 1].Get();
+				if (LibraryInstance->Library->GetName() == Name)
 				{
 					ShaderCodeArchiveStack.RemoveAt(i - 1);
 					break;
@@ -1941,258 +1582,130 @@ public:
 	{
 		check(ShaderPlatform == SP_NumPlatforms || InShaderPlatform == ShaderPlatform);
 
-		ShaderPlatform = InShaderPlatform;
-
-		FRHIShaderLibraryRef ShaderCodeArchive = new FShaderCodeArchive(ShaderPlatform, ShaderCodeDir, Library);
-		if (ShaderCodeArchive->GetShaderCount() > 0)
+		FShaderLibraryInstance* LibraryInstance = FShaderLibraryInstance::Create(InShaderPlatform, ShaderCodeDir, Library);
+		if(!LibraryInstance)
 		{
-			bSupportsPipelines = true;
-			UE_LOG(LogShaderLibrary, Display, TEXT("Cooked Context: Using Shared Shader Library %s"), *Library);
+			UE_LOG(LogShaderLibrary, Display, TEXT("Cooked Context: No Shared Shader Library for: %s and native library not supported."), *Library);
+			return false;
 		}
-		else if (RHISupportsNativeShaderLibraries(ShaderPlatform))
-		{
-			ShaderCodeArchive = RHICreateShaderLibrary(ShaderPlatform, ShaderCodeDir, Library);
 
-			if (ShaderCodeArchive.IsValid())
-			{
-				bNativeFormat = true;
+		ShaderPlatform = InShaderPlatform;
 
-				UE_LOG(LogShaderLibrary, Display, TEXT("Cooked Context: Loaded Native Shared Shader Library %s"), *Library);
-			}
-			else
-			{
-				UE_LOG(LogShaderLibrary, Display, TEXT("Cooked Context: No Native Shared Shader Library for %s"), *Library);
-			}
+		if (LibraryInstance->Library->IsNativeLibrary())
+		{
+			UE_LOG(LogShaderLibrary, Display, TEXT("Cooked Context: Loaded Native Shared Shader Library %s"), *Library);
 		}
 		else
 		{
-			UE_LOG(LogShaderLibrary, Display, TEXT("Cooked Context: No Shared Shader Library for: %s and native library not supported."), *Library);
+			UE_LOG(LogShaderLibrary, Display, TEXT("Cooked Context: Using Shared Shader Library %s"), *Library);
 		}
 
-		bool const bOK = IsValidRef(ShaderCodeArchive);
-		if (bOK)
-		{
-			FRWScopeLock(LibraryMutex, SLT_Write);
-
-			ShaderCodeArchiveStack.Add(ShaderCodeArchive);
-
-			ShaderCount += ShaderCodeArchive->GetShaderCount();
-
-			if (bSupportsPipelines && !bNativeFormat)
-			{
-				TSet<FShaderCodeLibraryPipeline> const* NewPipelines = ((FShaderCodeArchive*)ShaderCodeArchive.GetReference())->GetShaderPipelines(ShaderPlatform);
-				if (NewPipelines)
-				{
-					Pipelines.Append(*NewPipelines);
-				}
-			}
-		}
-		return bOK;
+		FRWScopeLock(LibraryMutex, SLT_Write);
+		ShaderCodeArchiveStack.Emplace(LibraryInstance);
+		ShaderCount += LibraryInstance->GetNumShaders();
+		return true;
 	}
 
-	FVertexShaderRHIRef CreateVertexShader(EShaderPlatform Platform, FSHAHash Hash)
+	uint32 GetShaderCount(void)
 	{
-		checkSlow(Platform == GetRuntimeShaderPlatform());
-
-		FVertexShaderRHIRef Result;
-		FRHIShaderLibrary* ShaderCodeArchive = FindShaderLibrary(Hash);
-		if (ShaderCodeArchive)
-		{
-			if (bNativeFormat || GRHILazyShaderCodeLoading)
-			{
-				Result = RHICreateVertexShader(ShaderCodeArchive, Hash);
-			}
-			else
-			{
-				Result = ((FShaderCodeArchive*)ShaderCodeArchive)->CreateVertexShader(Hash);
-			}
-		}
-		return Result;
+		return ShaderCount;
 	}
 
-	FPixelShaderRHIRef CreatePixelShader(EShaderPlatform Platform, FSHAHash Hash)
+	EShaderPlatform GetRuntimeShaderPlatform(void)
 	{
-		checkSlow(Platform == GetRuntimeShaderPlatform());
-
-		FPixelShaderRHIRef Result;
-		FRHIShaderLibrary* ShaderCodeArchive = FindShaderLibrary(Hash);
-		if (ShaderCodeArchive)
-		{
-			if (bNativeFormat || GRHILazyShaderCodeLoading)
-			{
-				Result = RHICreatePixelShader(ShaderCodeArchive, Hash);
-			}
-			else
-			{
-				Result = ((FShaderCodeArchive*)ShaderCodeArchive)->CreatePixelShader(Hash);
-			}
-		}
-		return Result;
+		return ShaderPlatform;
 	}
 
-	FGeometryShaderRHIRef CreateGeometryShader(EShaderPlatform Platform, FSHAHash Hash)
+	FShaderLibraryInstance* FindShaderLibraryForShaderMap(const FSHAHash& Hash, int32& OutShaderMapIndex)
 	{
-		checkSlow(Platform == GetRuntimeShaderPlatform());
+		FRWScopeLock(LibraryMutex, SLT_ReadOnly);
 
-		FGeometryShaderRHIRef Result;
-		FRHIShaderLibrary* ShaderCodeArchive = FindShaderLibrary(Hash);
-		if (ShaderCodeArchive)
+		// Search in library opened order
+		for (const TUniquePtr<FShaderLibraryInstance>& Instance : ShaderCodeArchiveStack)
 		{
-			if (bNativeFormat || GRHILazyShaderCodeLoading)
+			const int32 ShaderMapIndex = Instance->Library->FindShaderMapIndex(Hash);
+			if(ShaderMapIndex != INDEX_NONE)
 			{
-				Result = RHICreateGeometryShader(ShaderCodeArchive, Hash);
-			}
-			else
-			{
-				Result = ((FShaderCodeArchive*)ShaderCodeArchive)->CreateGeometryShader(Hash);
+				OutShaderMapIndex = ShaderMapIndex;
+				return Instance.Get();
 			}
 		}
-		return Result;
+		return nullptr;
 	}
 
-	FHullShaderRHIRef CreateHullShader(EShaderPlatform Platform, FSHAHash Hash)
+	FShaderLibraryInstance* FindShaderLibraryForShader(const FSHAHash& Hash, int32& OutShaderIndex)
 	{
-		checkSlow(Platform == GetRuntimeShaderPlatform());
+		FRWScopeLock(LibraryMutex, SLT_ReadOnly);
 
-		FHullShaderRHIRef Result;
-		FRHIShaderLibrary* ShaderCodeArchive = FindShaderLibrary(Hash);
-		if (ShaderCodeArchive)
+		// Search in library opened order
+		for (const TUniquePtr<FShaderLibraryInstance>& Instance : ShaderCodeArchiveStack)
 		{
-			if (bNativeFormat || GRHILazyShaderCodeLoading)
+			const int32 ShaderIndex = Instance->Library->FindShaderIndex(Hash);
+			if (ShaderIndex != INDEX_NONE)
 			{
-				Result = RHICreateHullShader(ShaderCodeArchive, Hash);
-			}
-			else
-			{
-				Result = ((FShaderCodeArchive*)ShaderCodeArchive)->CreateHullShader(Hash);
+				OutShaderIndex = ShaderIndex;
+				return Instance.Get();
 			}
 		}
-		return Result;
+		return nullptr;
 	}
 
-	FDomainShaderRHIRef CreateDomainShader(EShaderPlatform Platform, FSHAHash Hash)
+	TRefCountPtr<FShaderMapResource_SharedCode> LoadResource(const FSHAHash& Hash, FArchive* Ar)
 	{
-		checkSlow(Platform == GetRuntimeShaderPlatform());
-
-		FDomainShaderRHIRef Result;
-		FRHIShaderLibrary* ShaderCodeArchive = FindShaderLibrary(Hash);
-		if (ShaderCodeArchive)
+		int32 ShaderMapIndex = INDEX_NONE;
+		FShaderLibraryInstance* LibraryInstance = FindShaderLibraryForShaderMap(Hash, ShaderMapIndex);
+		if (LibraryInstance)
 		{
-			if (bNativeFormat || GRHILazyShaderCodeLoading)
-			{
-				Result = RHICreateDomainShader(ShaderCodeArchive, Hash);
-			}
-			else
-			{
-				Result = ((FShaderCodeArchive*)ShaderCodeArchive)->CreateDomainShader(Hash);
-			}
-		}
-		return Result;
-	}
+			SCOPED_LOADTIMER(LoadShaderResource_Internal);
 
-	FComputeShaderRHIRef CreateComputeShader(EShaderPlatform Platform, FSHAHash Hash)
-	{
-		checkSlow(Platform == GetRuntimeShaderPlatform());
-
-		FComputeShaderRHIRef Result;
-		FRHIShaderLibrary* ShaderCodeArchive = FindShaderLibrary(Hash);
-		if (ShaderCodeArchive)
-		{
-			if (bNativeFormat || GRHILazyShaderCodeLoading)
+			TRefCountPtr<FShaderMapResource_SharedCode> Resource = LibraryInstance->GetResource(ShaderMapIndex);
+			if (!Resource)
 			{
-				Result = RHICreateComputeShader(ShaderCodeArchive, Hash);
+				SCOPED_LOADTIMER(LoadShaderResource_AddOrDeleteResource);
+				Resource = LibraryInstance->AddOrDeleteResource(new FShaderMapResource_SharedCode(LibraryInstance, ShaderMapIndex));
 			}
-			else
+
+			if (!GRHILazyShaderCodeLoading)
 			{
-				Result = ((FShaderCodeArchive*)ShaderCodeArchive)->CreateComputeShader(Hash);
+				LibraryInstance->PreloadShaderMap(ShaderMapIndex, Ar);
 			}
-		}
-		return Result;
-	}
-
-	FRayTracingShaderRHIRef CreateRayTracingShader(EShaderPlatform Platform, EShaderFrequency Frequency, FSHAHash Hash)
-	{
-		FRayTracingShaderRHIRef Result;
 
-#if RHI_RAYTRACING
-		checkSlow(Platform == GetRuntimeShaderPlatform());
-		FRHIShaderLibrary* ShaderCodeArchive = FindShaderLibrary(Hash);
-		if (ShaderCodeArchive)
-		{
-			Result = ((FShaderCodeArchive*)ShaderCodeArchive)->CreateRayTracingShader(Frequency, Hash);
+			return Resource;
 		}
-#endif // RHI_RAYTRACING
-
-		return Result;
-	}
-
-	TRefCountPtr<FRHIShaderLibrary::FShaderLibraryIterator> CreateIterator(void)
-	{
-		return new FShaderCodeLibraryIterator(ShaderCodeArchiveStack, LibraryMutex);
-	}
-
-	uint32 GetShaderCount(void)
-	{
-		return ShaderCount;
-	}
 
-	EShaderPlatform GetRuntimeShaderPlatform(void)
-	{
-		return ShaderPlatform;
+		return TRefCountPtr<FShaderMapResource_SharedCode>();
 	}
 
-	TSet<FShaderCodeLibraryPipeline> const* GetShaderPipelines(EShaderPlatform Platform)
+	TRefCountPtr<FRHIShader> CreateShader(EShaderFrequency Frequency, const FSHAHash& Hash)
 	{
-		if (bSupportsPipelines)
+		int32 ShaderIndex = INDEX_NONE;
+		FShaderLibraryInstance* LibraryInstance = FindShaderLibraryForShader(Hash, ShaderIndex);
+		if (LibraryInstance)
 		{
-			FRWScopeLock(LibraryMutex, SLT_ReadOnly);
-			checkSlow(Platform == GetRuntimeShaderPlatform());
-			return &Pipelines;
+			TRefCountPtr<FRHIShader> Shader = LibraryInstance->GetOrCreateShader(ShaderIndex);
+			check(Shader->GetFrequency() == Frequency);
+			return Shader;
 		}
-		return nullptr;
+		return TRefCountPtr<FRHIShader>();
 	}
 
-	FRHIShaderLibrary* FindShaderLibrary(const FSHAHash& Hash)
+	bool PreloadShader(const FSHAHash& Hash, FArchive* Ar)
 	{
-		FRWScopeLock(LibraryMutex, SLT_ReadOnly);
-
-		// Search in library opened order
-		for (FRHIShaderLibrary* ShaderCodeArchive : ShaderCodeArchiveStack)
+		int32 ShaderIndex = INDEX_NONE;
+		FShaderLibraryInstance* LibraryInstance = FindShaderLibraryForShader(Hash, ShaderIndex);
+		if (LibraryInstance)
 		{
-			if (ShaderCodeArchive->ContainsEntry(Hash))
-			{
-				return ShaderCodeArchive;
-			}
+			LibraryInstance->PreloadShader(ShaderIndex, Ar);
+			return true;
 		}
-		return nullptr;
+		return false;
 	}
 
 	bool ContainsShaderCode(const FSHAHash& Hash)
 	{
-		FRHIShaderLibrary* ShaderCodeArchive = FindShaderLibrary(Hash);
-		if (ShaderCodeArchive)
-			return true;
-		else
-			return false;
-	}
-
-	bool RequestShaderCode(const FSHAHash& Hash, FArchive* Ar)
-	{
-		FRHIShaderLibrary* ShaderCodeArchive = FindShaderLibrary(Hash);
-		if (ShaderCodeArchive)
-			return ShaderCodeArchive->RequestEntry(Hash, Ar);
-		else
-			return false;
-	}
-
-	void ReleaseShaderCode(const FSHAHash& Hash)
-	{
-		if (!bNativeFormat)
-		{
-			FRHIShaderLibrary* ShaderCodeArchive = FindShaderLibrary(Hash);
-			if (ShaderCodeArchive)
-				((FShaderCodeArchive*)ShaderCodeArchive)->ReleaseShaderCode(Hash);
-		}
+		int32 ShaderIndex = INDEX_NONE;
+		FShaderLibraryInstance* LibraryInstance = FindShaderLibraryForShader(Hash, ShaderIndex);
+		return LibraryInstance != nullptr;
 	}
 
 #if WITH_EDITOR
@@ -2249,21 +1762,17 @@ public:
 		return (bShaderFormatsThatNeedStableKeys & (uint64_t(1u) << (uint32_t) Platform)) != 0;
 	}
 
-	void AddShaderCode(EShaderPlatform Platform, EShaderFrequency Frequency, const FSHAHash& Hash, const TArray<uint8>& InCode, uint32 const UncompressedSize)
+	void AddShaderCode(EShaderPlatform Platform, const FShaderMapResourceCode* Code)
 	{
 		FScopeLock ScopeLock(&ShaderCodeCS);
 		FShaderCodeStats& CodeStats = EditorShaderCodeStats[Platform];
-		CodeStats.NumShaders++;
-		CodeStats.ShadersSize += InCode.Num();
+		CodeStats.NumShaders += Code->ShaderEntries.Num();
+		CodeStats.ShadersSize += Code->ShaderCode.Num();
 
 		FEditorShaderCodeArchive* CodeArchive = EditorShaderCodeArchive[Platform];
 		check(CodeArchive);
 
-		if (CodeArchive->AddShader((uint8)Frequency, Hash, InCode, UncompressedSize))
-		{
-			CodeStats.NumUniqueShaders++;
-			CodeStats.ShadersUniqueSize += InCode.Num();
-		}
+		CodeArchive->AddShaderCode(CodeStats, Code);
 	}
 
 	void AddShaderStableKeyValue(EShaderPlatform InShaderPlatform, FStableShaderKeyAndValue& StableKeyValue)
@@ -2280,43 +1789,6 @@ public:
 		StableArchive->AddShader(StableKeyValue);
 	}
 
-	bool AddShaderPipeline(FShaderPipeline* Pipeline)
-	{
-		check(Pipeline);
-
-		EShaderPlatform SPlatform = SP_NumPlatforms;
-		for (uint8 Freq = 0; Freq < SF_Compute; Freq++)
-		{
-			FShader* Shader = Pipeline->GetShader((EShaderFrequency)Freq);
-			if (Shader)
-			{
-				if (SPlatform == SP_NumPlatforms)
-				{
-					SPlatform = (EShaderPlatform)Shader->GetTarget().Platform;
-				}
-				else
-				{
-					check(SPlatform == (EShaderPlatform)Shader->GetTarget().Platform);
-				}
-			}
-		}
-
-		FScopeLock ScopeLock(&ShaderCodeCS);
-		FShaderCodeStats& CodeStats = EditorShaderCodeStats[SPlatform];
-		CodeStats.NumPipelines++;
-
-		FEditorShaderCodeArchive* CodeArchive = EditorShaderCodeArchive[SPlatform];
-		check(CodeArchive);
-
-		bool bAdded = false;
-		if (EditorArchivePipelines[SPlatform] && ((FEditorShaderCodeArchive*)CodeArchive)->AddPipeline(Pipeline))
-		{
-			CodeStats.NumUniquePipelines++;
-			bAdded = true;
-		}
-		return bAdded;
-	}
-
 	bool SaveShaderCode(const FString& ShaderCodeDir, const FString& MetaOutputDir, const TArray<FName>& ShaderFormats, bool bMaster, TArray<FString>& OutSCLCSVPath)
 	{
 		bool bOk = ShaderFormats.Num() > 0;
@@ -2392,7 +1864,6 @@ public:
 };
 
 static FSharedShaderCodeRequest OnSharedShaderCodeRequest;
-static FSharedShaderCodeRelease OnSharedShaderCodeRelease;
 
 FShaderCodeLibraryImpl* FShaderCodeLibraryImpl::Impl = nullptr;
 
@@ -2493,162 +1964,79 @@ bool FShaderCodeLibrary::ContainsShaderCode(const FSHAHash& Hash)
 	return false;
 }
 
-bool FShaderCodeLibrary::RequestShaderCode(const FSHAHash& Hash, FArchive* Ar)
+TRefCountPtr<FShaderMapResource> FShaderCodeLibrary::LoadResource(const FSHAHash& Hash, FArchive* Ar)
 {
 	if (FShaderCodeLibraryImpl::Impl)
 	{
+		SCOPED_LOADTIMER(FShaderCodeLibrary_LoadResource);
 		OnSharedShaderCodeRequest.Broadcast(Hash, Ar);
-		return FShaderCodeLibraryImpl::Impl->RequestShaderCode(Hash, Ar);
+		return TRefCountPtr<FShaderMapResource>(FShaderCodeLibraryImpl::Impl->LoadResource(Hash, Ar));
 	}
-	return false;
+	return TRefCountPtr<FShaderMapResource>();
 }
 
-bool FShaderCodeLibrary::LazyRequestShaderCode(const FSHAHash& Hash, FArchive* Ar)
+bool FShaderCodeLibrary::PreloadShader(const FSHAHash& Hash, FArchive* Ar)
 {
 	if (FShaderCodeLibraryImpl::Impl)
 	{
 		OnSharedShaderCodeRequest.Broadcast(Hash, Ar);
-		return true;
+		return FShaderCodeLibraryImpl::Impl->PreloadShader(Hash, Ar);
 	}
 	return false;
 }
 
-void FShaderCodeLibrary::ReleaseShaderCode(const FSHAHash& Hash)
+FVertexShaderRHIRef FShaderCodeLibrary::CreateVertexShader(EShaderPlatform Platform, const FSHAHash& Hash)
 {
 	if (FShaderCodeLibraryImpl::Impl)
 	{
-		OnSharedShaderCodeRelease.Broadcast(Hash);
-		return FShaderCodeLibraryImpl::Impl->ReleaseShaderCode(Hash);
+		return FVertexShaderRHIRef(FShaderCodeLibraryImpl::Impl->CreateShader(SF_Vertex, Hash));
 	}
+	return nullptr;
 }
 
-void FShaderCodeLibrary::LazyReleaseShaderCode(const FSHAHash& Hash)
+FPixelShaderRHIRef FShaderCodeLibrary::CreatePixelShader(EShaderPlatform Platform, const FSHAHash& Hash)
 {
 	if (FShaderCodeLibraryImpl::Impl)
 	{
-		OnSharedShaderCodeRelease.Broadcast(Hash);
+		return FPixelShaderRHIRef(FShaderCodeLibraryImpl::Impl->CreateShader(SF_Pixel, Hash));
 	}
+	return nullptr;
 }
 
-FVertexShaderRHIRef FShaderCodeLibrary::CreateVertexShader(EShaderPlatform Platform, FSHAHash Hash, TArray<uint8> const& Code)
+FHullShaderRHIRef FShaderCodeLibrary::CreateHullShader(EShaderPlatform Platform, const FSHAHash& Hash)
 {
-	FVertexShaderRHIRef Shader;
 	if (FShaderCodeLibraryImpl::Impl)
 	{
-		Shader = FShaderCodeLibraryImpl::Impl->CreateVertexShader(Platform, Hash);
+		return FHullShaderRHIRef(FShaderCodeLibraryImpl::Impl->CreateShader(SF_Hull, Hash));
 	}
-	if (!IsValidRef(Shader))
-	{
-		Shader = RHICreateVertexShader(Code);
-	}
-	SafeAssignHash(Shader, Hash);
-	return Shader;
-}
-
-FPixelShaderRHIRef FShaderCodeLibrary::CreatePixelShader(EShaderPlatform Platform, FSHAHash Hash, TArray<uint8> const& Code)
-{
-	FPixelShaderRHIRef Shader;
-	if (FShaderCodeLibraryImpl::Impl)
-	{
-		Shader = FShaderCodeLibraryImpl::Impl->CreatePixelShader(Platform, Hash);
-	}
-	if (!IsValidRef(Shader))
-	{
-		Shader = RHICreatePixelShader(Code);
-	}
-	SafeAssignHash(Shader, Hash);
-	return Shader;
+	return nullptr;
 }
 
-FGeometryShaderRHIRef FShaderCodeLibrary::CreateGeometryShader(EShaderPlatform Platform, FSHAHash Hash, TArray<uint8> const& Code)
+FDomainShaderRHIRef FShaderCodeLibrary::CreateDomainShader(EShaderPlatform Platform, const FSHAHash& Hash)
 {
-	FGeometryShaderRHIRef Shader;
 	if (FShaderCodeLibraryImpl::Impl)
 	{
-		Shader = FShaderCodeLibraryImpl::Impl->CreateGeometryShader(Platform, Hash);
-	}
-	if (!IsValidRef(Shader))
-	{
-		Shader = RHICreateGeometryShader(Code);
+		return FDomainShaderRHIRef(FShaderCodeLibraryImpl::Impl->CreateShader(SF_Domain, Hash));
 	}
-	SafeAssignHash(Shader, Hash);
-	return Shader;
+	return nullptr;
 }
 
-FHullShaderRHIRef FShaderCodeLibrary::CreateHullShader(EShaderPlatform Platform, FSHAHash Hash, TArray<uint8> const& Code)
+FGeometryShaderRHIRef FShaderCodeLibrary::CreateGeometryShader(EShaderPlatform Platform, const FSHAHash& Hash)
 {
-	FHullShaderRHIRef Shader;
 	if (FShaderCodeLibraryImpl::Impl)
 	{
-		Shader = FShaderCodeLibraryImpl::Impl->CreateHullShader(Platform, Hash);
-	}
-	if (!IsValidRef(Shader))
-	{
-		Shader = RHICreateHullShader(Code);
+		return FGeometryShaderRHIRef(FShaderCodeLibraryImpl::Impl->CreateShader(SF_Geometry, Hash));
 	}
-	SafeAssignHash(Shader, Hash);
-	return Shader;
+	return nullptr;
 }
 
-FDomainShaderRHIRef FShaderCodeLibrary::CreateDomainShader(EShaderPlatform Platform, FSHAHash Hash, TArray<uint8> const& Code)
+FComputeShaderRHIRef FShaderCodeLibrary::CreateComputeShader(EShaderPlatform Platform, const FSHAHash& Hash)
 {
-	FDomainShaderRHIRef Shader;
 	if (FShaderCodeLibraryImpl::Impl)
 	{
-		Shader = FShaderCodeLibraryImpl::Impl->CreateDomainShader(Platform, Hash);
-	}
-	if (!IsValidRef(Shader))
-	{
-		Shader = RHICreateDomainShader(Code);
+		return FComputeShaderRHIRef(FShaderCodeLibraryImpl::Impl->CreateShader(SF_Compute, Hash));
 	}
-	SafeAssignHash(Shader, Hash);
-	return Shader;
-}
-
-FComputeShaderRHIRef FShaderCodeLibrary::CreateComputeShader(EShaderPlatform Platform, FSHAHash Hash, TArray<uint8> const& Code)
-{
-	FComputeShaderRHIRef Shader;
-	if (FShaderCodeLibraryImpl::Impl)
-	{
-		Shader = FShaderCodeLibraryImpl::Impl->CreateComputeShader(Platform, Hash);
-	}
-	if (!IsValidRef(Shader))
-	{
-		Shader = RHICreateComputeShader(Code);
-	}
-	SafeAssignHash(Shader, Hash);
-	FPipelineFileCache::CacheComputePSO(GetTypeHash(Shader.GetReference()), Shader.GetReference());
-	Shader->SetStats(FPipelineFileCache::RegisterPSOStats(GetTypeHash(Shader.GetReference())));
-	return Shader;
-}
-
-FRayTracingShaderRHIRef FShaderCodeLibrary::CreateRayTracingShader(EShaderPlatform Platform, EShaderFrequency Frequency, FSHAHash Hash, TArray<uint8> const& Code)
-{
-	FRayTracingShaderRHIRef Shader;
-
-#if RHI_RAYTRACING
-	if (FShaderCodeLibraryImpl::Impl)
-	{
-		Shader = FShaderCodeLibraryImpl::Impl->CreateRayTracingShader(Platform, Frequency, Hash);
-	}
-	if (!IsValidRef(Shader))
-	{
-		Shader = RHICreateRayTracingShader(Code, Frequency);
-	}
-	SafeAssignHash(Shader, Hash);
-#endif // RHI_RAYTRACING
-
-	return Shader;
-}
-
-TRefCountPtr<FRHIShaderLibrary::FShaderLibraryIterator> FShaderCodeLibrary::CreateIterator(void)
-{
-	TRefCountPtr<FRHIShaderLibrary::FShaderLibraryIterator> It;
-	if (FShaderCodeLibraryImpl::Impl)
-	{
-		It = FShaderCodeLibraryImpl::Impl->CreateIterator();
-	}
-	return It;
+	return nullptr;
 }
 
 uint32 FShaderCodeLibrary::GetShaderCount(void)
@@ -2661,16 +2049,6 @@ uint32 FShaderCodeLibrary::GetShaderCount(void)
 	return Num;
 }
 
-TSet<FShaderCodeLibraryPipeline> const* FShaderCodeLibrary::GetShaderPipelines(EShaderPlatform Platform)
-{
-	TSet<FShaderCodeLibraryPipeline> const* Pipelines = nullptr;
-	if (FShaderCodeLibraryImpl::Impl)
-	{
-		Pipelines = FShaderCodeLibraryImpl::Impl->GetShaderPipelines(Platform);
-	}
-	return Pipelines;
-}
-
 EShaderPlatform FShaderCodeLibrary::GetRuntimeShaderPlatform(void)
 {
 	EShaderPlatform Platform = SP_NumPlatforms;
@@ -2721,12 +2099,12 @@ void FShaderCodeLibrary::CookShaderFormats(TArray<TTuple<FName,bool>> const& Sha
 	}
 }
 
-bool FShaderCodeLibrary::AddShaderCode(EShaderPlatform ShaderPlatform, EShaderFrequency Frequency, const FSHAHash& Hash, const TArray<uint8>& InCode, uint32 const UncompressedSize)
+bool FShaderCodeLibrary::AddShaderCode(EShaderPlatform ShaderPlatform, const FShaderMapResourceCode* Code)
 {
 #if WITH_EDITOR
 	if (FShaderCodeLibraryImpl::Impl)
 	{
-		FShaderCodeLibraryImpl::Impl->AddShaderCode(ShaderPlatform, Frequency, Hash, InCode, UncompressedSize);
+		FShaderCodeLibraryImpl::Impl->AddShaderCode(ShaderPlatform, Code);
 		return true;
 	}
 #endif// WITH_EDITOR
@@ -2755,19 +2133,6 @@ void FShaderCodeLibrary::AddShaderStableKeyValue(EShaderPlatform ShaderPlatform,
 #endif// WITH_EDITOR
 }
 
-bool FShaderCodeLibrary::AddShaderPipeline(FShaderPipeline* Pipeline)
-{
-#if WITH_EDITOR
-	if (FShaderCodeLibraryImpl::Impl && Pipeline)
-	{
-		FShaderCodeLibraryImpl::Impl->AddShaderPipeline(Pipeline);
-		return true;
-	}
-#endif// WITH_EDITOR
-
-	return false;
-}
-
 bool FShaderCodeLibrary::SaveShaderCodeMaster(const FString& OutputDir, const FString& MetaOutputDir, const TArray<FName>& ShaderFormats, TArray<FString>& OutSCLCSVPath)
 {
 	if (FShaderCodeLibraryImpl::Impl)
@@ -2858,13 +2223,3 @@ void FShaderCodeLibrary::UnregisterSharedShaderCodeRequestDelegate_Handle(FDeleg
 	OnSharedShaderCodeRequest.Remove(Handle);
 }
 
-FDelegateHandle FShaderCodeLibrary::RegisterSharedShaderCodeReleaseDelegate_Handle(const FSharedShaderCodeRelease::FDelegate& Delegate)
-{
-	return OnSharedShaderCodeRelease.Add(Delegate);
-
-}
-
-void FShaderCodeLibrary::UnregisterSharedShaderCodeReleaseDelegate_Handle(FDelegateHandle Handle)
-{
-	OnSharedShaderCodeRelease.Remove(Handle);
-}
diff --git a/Engine/Source/Runtime/RenderCore/Private/ShaderCore.cpp b/Engine/Source/Runtime/RenderCore/Private/ShaderCore.cpp
--- a/Engine/Source/Runtime/RenderCore/Private/ShaderCore.cpp
+++ b/Engine/Source/Runtime/RenderCore/Private/ShaderCore.cpp
@@ -93,7 +93,6 @@ DEFINE_STAT(STAT_Shaders_TotalRTShaderInitForRenderingTime);
 DEFINE_STAT(STAT_Shaders_FrameRTShaderInitForRenderingTime);
 DEFINE_STAT(STAT_Shaders_ShaderMemory);
 DEFINE_STAT(STAT_Shaders_ShaderResourceMemory);
-DEFINE_STAT(STAT_Shaders_ShaderMapMemory);
 
 DEFINE_STAT(STAT_Shaders_NumShadersRegistered);
 DEFINE_STAT(STAT_Shaders_NumShadersDuplicated);
@@ -377,12 +376,12 @@ void ValidateStaticUniformBuffer(FRHIUniformBuffer* UniformBuffer, FUniformBuffe
 			checkf(
 				ExpectedStructMetadata,
 				TEXT("Shader is requesting uniform buffer '%s' at slot %s with hash '%u', but a reverse lookup of the hash can't find it. The shader cache may be out of date."),
-				*Layout.GetDebugName().ToString(), *SlotRegistry.GetDebugDescription(Slot), ExpectedHash);
+				*Layout.GetDebugName(), *SlotRegistry.GetDebugDescription(Slot), ExpectedHash);
 
 			checkf(
 				false,
 				TEXT("Shader attempted to bind uniform buffer '%s' at slot %s with hash '%u', but the shader expected '%s' with hash '%u'."),
-				*Layout.GetDebugName().ToString(), *SlotRegistry.GetDebugDescription(Slot), ExpectedHash, ExpectedStructMetadata->GetShaderVariableName(), Layout.GetHash());
+				*Layout.GetDebugName(), *SlotRegistry.GetDebugDescription(Slot), ExpectedHash, ExpectedStructMetadata->GetShaderVariableName(), Layout.GetHash());
 		}
 	}
 #endif
diff --git a/Engine/Source/Runtime/RenderCore/Private/ShaderMap.cpp b/Engine/Source/Runtime/RenderCore/Private/ShaderMap.cpp
new file mode 100644
--- /dev/null
+++ b/Engine/Source/Runtime/RenderCore/Private/ShaderMap.cpp
@@ -0,0 +1,551 @@
+// Copyright 1998-2019 Epic Games, Inc. All Rights Reserved.
+
+/*=============================================================================
+	Shader.cpp: Shader implementation.
+=============================================================================*/
+
+#include "Shader.h"
+#include "Misc/CoreMisc.h"
+#include "Stats/StatsMisc.h"
+#include "Serialization/MemoryWriter.h"
+#include "VertexFactory.h"
+#include "ProfilingDebugging/DiagnosticTable.h"
+#include "Interfaces/ITargetPlatform.h"
+#include "Interfaces/ITargetPlatformManagerModule.h"
+#include "Interfaces/IShaderFormat.h"
+#include "ShaderCodeLibrary.h"
+#include "ShaderCore.h"
+#include "RenderUtils.h"
+#include "Misc/ConfigCacheIni.h"
+#include "Misc/ScopeLock.h"
+#include "UObject/RenderingObjectVersion.h"
+#include "UObject/FortniteMainBranchObjectVersion.h"
+#include "Misc/ScopeRWLock.h"
+#include "ProfilingDebugging/LoadTimeTracker.h"
+
+#if WITH_EDITORONLY_DATA
+#include "Interfaces/IShaderFormat.h"
+#endif
+
+FShaderMapBase::FShaderMapBase(const FTypeLayoutDesc& InContentTypeLayout)
+	: ContentTypeLayout(InContentTypeLayout)
+	, PointerTable(nullptr)
+	, Content(nullptr)
+	, FrozenContentSize(0u)
+	, NumFrozenShaders(0u)
+{}
+
+FShaderMapBase::~FShaderMapBase()
+{
+	DestroyContent();
+	if (PointerTable)
+	{
+		delete PointerTable;
+	}
+}
+
+FShaderMapResourceCode* FShaderMapBase::GetResourceCode()
+{
+	if (!Code)
+	{
+		Code = new FShaderMapResourceCode();
+	}
+	return Code;
+}
+
+void FShaderMapBase::AssignContent(FShaderMapContent* InContent)
+{
+	check(!Content);
+	check(!PointerTable);
+	Content = InContent;
+	PointerTable = CreatePointerTable();
+}
+
+void FShaderMapBase::FinalizeContent()
+{
+	if (Content && FrozenContentSize == 0u)
+	{
+		Content->Validate(*this);
+
+		FMemoryImage MemoryImage;
+		MemoryImage.TargetLayoutParameters.InitializeForCurrent();
+		MemoryImage.PointerTable = PointerTable;
+		FMemoryImageWriter Writer(MemoryImage);
+
+		Writer.WriteObject(Content, ContentTypeLayout);
+
+		FMemoryImageResult MemoryImageResult;
+		MemoryImage.Flatten(MemoryImageResult, true);
+
+		DestroyContent();
+
+		{
+			FrozenContentSize = MemoryImageResult.Bytes.Num();
+			check(FrozenContentSize > 0u);
+			void* ContentMemory = FMemory::Malloc(FrozenContentSize);
+			FMemory::Memcpy(ContentMemory, MemoryImageResult.Bytes.GetData(), FrozenContentSize);
+			Content = static_cast<FShaderMapContent*>(ContentMemory);
+			MemoryImageResult.ApplyPatches(Content);
+			NumFrozenShaders = Content->GetNumShaders();
+		}
+
+		INC_DWORD_STAT_BY(STAT_Shaders_ShaderMemory, FrozenContentSize);
+		INC_DWORD_STAT_BY(STAT_Shaders_NumShadersLoaded, NumFrozenShaders);
+	}
+
+	check(Code);
+	Code->Finalize();
+	Resource = new FShaderMapResource_InlineCode(GetShaderPlatform(), Code);
+	BeginInitResource(Resource);
+
+	INC_DWORD_STAT_BY(STAT_Shaders_ShaderResourceMemory, Resource->GetSizeBytes());
+}
+
+void FShaderMapBase::UnfreezeContent()
+{
+	if (Content && FrozenContentSize > 0u)
+	{
+		void* UnfrozenMemory = FMemory::Malloc(ContentTypeLayout.Size, ContentTypeLayout.Alignment);
+
+		FMemoryUnfreezeContent Context;
+		Context.PrevPointerTable = PointerTable;
+		Context.UnfreezeObject(Content, ContentTypeLayout, UnfrozenMemory);
+
+		DestroyContent();
+
+		Content = static_cast<FShaderMapContent*>(UnfrozenMemory);
+	}
+}
+
+#define CHECK_SHADERMAP_DEPENDENCIES !(UE_BUILD_SHIPPING || UE_BUILD_TEST)
+
+void FShaderMapBase::Serialize(FArchive& Ar, bool bInlineShaderResources, bool bLoadedByCookedMaterial)
+{
+	LLM_SCOPE(ELLMTag::Shaders);
+	if (Ar.IsSaving())
+	{
+		check(Content);
+		Content->Validate(*this);
+
+		FShaderMapPointerTable* SavePointerTable = CreatePointerTable();
+
+		FMemoryImage MemoryImage;
+		MemoryImage.PrevPointerTable = PointerTable;
+		MemoryImage.PointerTable = SavePointerTable;
+		MemoryImage.TargetLayoutParameters.InitializeForArchive(Ar);
+
+		FMemoryImageWriter Writer(MemoryImage);
+
+		Writer.WriteObject(Content, ContentTypeLayout);
+
+		FMemoryImageResult MemoryImageResult;
+		MemoryImage.Flatten(MemoryImageResult, true);
+
+		void* SaveFrozenContent = MemoryImageResult.Bytes.GetData();
+		uint32 SaveFrozenContentSize = MemoryImageResult.Bytes.Num();
+		Ar << SaveFrozenContentSize;
+		Ar.Serialize(SaveFrozenContent, SaveFrozenContentSize);
+		MemoryImageResult.SaveToArchive(Ar);
+		SavePointerTable->SaveToArchive(Ar, SaveFrozenContent, bInlineShaderResources);
+		delete SavePointerTable;
+
+		int32 NumDependencies = MemoryImage.TypeDependencies.Num();
+		Ar << NumDependencies;
+		for (const FTypeLayoutDesc* DependencyTypeDesc : MemoryImage.TypeDependencies)
+		{
+			uint64 NameHash = DependencyTypeDesc->NameHash;
+			FSHAHash LayoutHash;
+			uint32 LayoutSize = Freeze::HashLayout(*DependencyTypeDesc, MemoryImage.TargetLayoutParameters, LayoutHash);
+			Ar << NameHash;
+			Ar << LayoutSize;
+			Ar << LayoutHash;
+		}
+
+		bool bShareCode = false;
+#if WITH_EDITOR
+		bShareCode = FShaderCodeLibrary::IsEnabled() && Ar.IsCooking();
+#endif // WITH_EDITOR
+		Ar << bShareCode;
+#if WITH_EDITOR
+		if (bShareCode)
+		{
+			FSHAHash ResourceHash = Code->ResourceHash;
+			Ar << ResourceHash;
+			FShaderCodeLibrary::AddShaderCode(GetShaderPlatform(), Code);
+		}
+		else
+#endif // WITH_EDITOR
+		{
+			Code->Serialize(Ar);
+		}
+	}
+	else
+	{
+		check(!PointerTable);
+		PointerTable = CreatePointerTable();
+
+		Ar << FrozenContentSize;
+		void* ContentMemory = FMemory::Malloc(FrozenContentSize);
+		Ar.Serialize(ContentMemory, FrozenContentSize);
+		Content = static_cast<FShaderMapContent*>(ContentMemory);
+		NumFrozenShaders = Content->GetNumShaders();
+		FMemoryImageResult::ApplyPatchesFromArchive(Content, Ar);
+		PointerTable->LoadFromArchive(Ar, Content, bInlineShaderResources, bLoadedByCookedMaterial);
+
+		int32 NumDependencies = 0;
+		Ar << NumDependencies;
+		if(NumDependencies > 0)
+		{
+#if CHECK_SHADERMAP_DEPENDENCIES
+			FPlatformTypeLayoutParameters LayoutParams;
+			LayoutParams.InitializeForCurrent();
+#endif // CHECK_SHADERMAP_DEPENDENCIES
+
+			// Waste a bit of time even in shipping builds skipping over this stuff
+			// Could add a cook-time option to exclude dependencies completely
+			for (int32 i = 0u; i < NumDependencies; ++i)
+			{
+				uint64 NameHash = 0u;
+				uint32 SavedLayoutSize = 0u;
+				FSHAHash SavedLayoutHash;
+				Ar << NameHash;
+				Ar << SavedLayoutSize;
+				Ar << SavedLayoutHash;
+#if CHECK_SHADERMAP_DEPENDENCIES
+				const FTypeLayoutDesc* DependencyType = FTypeLayoutDesc::Find(NameHash);
+				if (DependencyType)
+				{
+					FSHAHash CheckLayoutHash;
+					const uint32 CheckLayoutSize = Freeze::HashLayout(*DependencyType, LayoutParams, CheckLayoutHash);
+					checkf(CheckLayoutSize == SavedLayoutSize, TEXT("Mismatch size for type %s"), DependencyType->Name);
+					checkf(CheckLayoutHash == SavedLayoutHash, TEXT("Mismatch hash for type %s"), DependencyType->Name);
+				}
+#endif // CHECK_SHADERMAP_DEPENDENCIES
+			}
+		}
+
+		bool bShareCode = false;
+		Ar << bShareCode;
+		if (bShareCode)
+		{
+			FSHAHash ResourceHash;
+			Ar << ResourceHash;
+			Resource = FShaderCodeLibrary::LoadResource(ResourceHash, &Ar);
+		}
+		else
+		{
+			Code = new FShaderMapResourceCode();
+			Code->Serialize(Ar);
+			Resource = new FShaderMapResource_InlineCode(GetShaderPlatform(), Code);
+		}
+		
+		BeginInitResource(Resource);
+
+		INC_DWORD_STAT_BY(STAT_Shaders_ShaderResourceMemory, Resource->GetSizeBytes());
+		INC_DWORD_STAT_BY(STAT_Shaders_ShaderMemory, FrozenContentSize);
+		INC_DWORD_STAT_BY(STAT_Shaders_NumShadersLoaded, NumFrozenShaders);
+	}
+}
+
+void FShaderMapBase::DestroyContent()
+{
+	if (Content)
+	{
+		DEC_DWORD_STAT_BY(STAT_Shaders_ShaderMemory, FrozenContentSize);
+		DEC_DWORD_STAT_BY(STAT_Shaders_NumShadersLoaded, NumFrozenShaders);
+
+		InternalDeleteObjectFromLayout(Content, ContentTypeLayout, FrozenContentSize > 0u);
+
+		FrozenContentSize = 0u;
+		NumFrozenShaders = 0u;
+		Content = nullptr;
+	}
+}
+
+FShader* FShaderMapContent::GetShader(const FHashedName& TypeName, int32 PermutationId) const
+{
+	TRACE_CPUPROFILER_EVENT_SCOPE(FShaderMapContent::GetShader);
+	const int32 Index = Algo::BinarySearchBy(Shaders, FShaderMapKey(TypeName, PermutationId), FProjectShaderToKey());
+	return (Index != INDEX_NONE) ? Shaders[Index].Get() : nullptr;
+}
+
+void FShaderMapContent::AddShader(FShader* Shader)
+{
+	check(!Shader->IsFrozen());
+	checkSlow(!HasShader(Shader->GetTypeName(), Shader->GetPermutationId()));
+	const int32 Index = Algo::LowerBoundBy(Shaders, FShaderMapKey(*Shader), FProjectShaderToKey());
+	Shaders.Insert(Shader, Index);
+}
+
+FShader* FShaderMapContent::FindOrAddShader(FShader* Shader)
+{
+	check(!Shader->IsFrozen());
+	const int32 Index = Algo::LowerBoundBy(Shaders, FShaderMapKey(*Shader), FProjectShaderToKey());
+	if (Index < Shaders.Num())
+	{
+		FShader* PrevShader = Shaders[Index];
+		if (PrevShader->GetTypeName() == Shader->GetTypeName() && PrevShader->GetPermutationId() == Shader->GetPermutationId())
+		{
+			DeleteObjectFromLayout(Shader);
+			return PrevShader;
+		}
+	}
+
+	Shaders.Insert(Shader, Index);
+	return Shader;
+}
+
+void FShaderMapContent::AddShaderPipeline(FShaderPipeline* Pipeline)
+{
+	checkSlow(!HasShaderPipeline(Pipeline->TypeName));
+	const int32 Index = Algo::LowerBoundBy(ShaderPipelines, Pipeline->TypeName, FProjectShaderPipelineToKey());
+	ShaderPipelines.Insert(Pipeline, Index);
+}
+
+FShaderPipeline* FShaderMapContent::FindOrAddShaderPipeline(FShaderPipeline* Pipeline)
+{
+	const int32 Index = Algo::LowerBoundBy(ShaderPipelines, Pipeline->TypeName, FProjectShaderPipelineToKey());
+	if (Index < ShaderPipelines.Num())
+	{
+		FShaderPipeline* PrevShaderPipeline = ShaderPipelines[Index];
+		if (PrevShaderPipeline->TypeName == Pipeline->TypeName)
+		{
+			delete Pipeline;
+			return PrevShaderPipeline;
+		}
+	}
+
+	ShaderPipelines.Insert(Pipeline, Index);
+	return Pipeline;
+}
+
+/**
+ * Removes the shader of the given type from the shader map
+ * @param Type Shader type to remove the entry for
+ */
+void FShaderMapContent::RemoveShaderTypePermutaion(const FShaderType* Type, int32 PermutationId)
+{
+	const int32 Index = Algo::BinarySearchBy(Shaders, FShaderMapKey(Type->GetHashedName(), PermutationId), FProjectShaderToKey());
+	if (Index != INDEX_NONE)
+	{
+		Shaders.RemoveAt(Index, 1, false);
+	}
+}
+
+void FShaderMapContent::RemoveShaderPipelineType(const FShaderPipelineType* ShaderPipelineType)
+{
+	const int32 Index = Algo::BinarySearchBy(ShaderPipelines, ShaderPipelineType->GetHashedName(), FProjectShaderPipelineToKey());
+	if (Index != INDEX_NONE)
+	{
+		FShaderPipeline* Pipeline = ShaderPipelines[Index];
+		delete Pipeline;
+		ShaderPipelines.RemoveAt(Index, 1, false);
+	}
+}
+
+void FShaderMapContent::GetShaderList(const FShaderMapBase& InShaderMap, const FSHAHash& InMaterialShaderMapHash, TMap<FShaderId, TShaderRef<FShader>>& OutShaders) const
+{
+	for (FShader* Shader : Shaders)
+	{
+		const FShaderId ShaderId(
+			Shader->GetType(InShaderMap.GetPointerTable()),
+			InMaterialShaderMapHash,
+			FHashedName(),
+			Shader->GetVertexFactoryType(InShaderMap.GetPointerTable()),
+			Shader->GetPermutationId(),
+			GetShaderPlatform());
+
+		OutShaders.Add(ShaderId, TShaderRef<FShader>(Shader, InShaderMap));
+	}
+
+	for (const FShaderPipeline* ShaderPipeline : ShaderPipelines)
+	{
+		for (const TShaderRef<FShader>& Shader : ShaderPipeline->GetShaders(InShaderMap))
+		{
+			const FShaderId ShaderId(
+				Shader.GetType(),
+				InMaterialShaderMapHash,
+				ShaderPipeline->TypeName,
+				Shader.GetVertexFactoryType(),
+				Shader->GetPermutationId(),
+				GetShaderPlatform());
+
+			OutShaders.Add(ShaderId, Shader);
+		}
+	}
+}
+
+void FShaderMapContent::GetShaderList(const FShaderMapBase& InShaderMap, TMap<FHashedName, TShaderRef<FShader>>& OutShaders) const
+{
+	for (FShader* Shader : Shaders)
+	{
+		OutShaders.Add(Shader->GetTypeName(), TShaderRef<FShader>(Shader, InShaderMap));
+	}
+
+	for (const FShaderPipeline* ShaderPipeline : ShaderPipelines)
+	{
+		for (const TShaderRef<FShader>& Shader : ShaderPipeline->GetShaders(InShaderMap))
+		{
+			OutShaders.Add(Shader->GetTypeName(), Shader);
+		}
+	}
+}
+
+void FShaderMapContent::GetShaderPipelineList(const FShaderMapBase& InShaderMap, TArray<FShaderPipelineRef>& OutShaderPipelines, FShaderPipeline::EFilter Filter) const
+{
+	for (FShaderPipeline* Pipeline : ShaderPipelines)
+	{
+		const FShaderPipelineType* PipelineType = FShaderPipelineType::GetShaderPipelineTypeByName(Pipeline->TypeName);
+		if (PipelineType->ShouldOptimizeUnusedOutputs(Platform) && Filter == FShaderPipeline::EOnlyShared)
+		{
+			continue;
+		}
+		else if (!PipelineType->ShouldOptimizeUnusedOutputs(Platform) && Filter == FShaderPipeline::EOnlyUnique)
+		{
+			continue;
+		}
+		OutShaderPipelines.Add(FShaderPipelineRef(Pipeline, InShaderMap));
+	}
+}
+
+void FShaderMapContent::Validate(const FShaderMapBase& InShaderMap)
+{
+	for (FShader* Shader : Shaders)
+	{
+		checkf(Shader->GetResourceIndex() != INDEX_NONE, TEXT("Missing resource for %s"), Shader->GetType(InShaderMap.GetPointerTable())->GetName());
+	}
+
+	/*for(FShaderPipeline* Pipeline : ShaderPipelines)
+	{
+		for(const TShaderRef<FShader>& Shader : Pipeline->GetShaders(InPtrTable))
+		{
+			checkf(Shader.GetResource(), TEXT("Missing resource for %s"), Shader.GetType()->GetName());
+		}
+	}*/
+}
+
+#if WITH_EDITOR
+static bool CheckOutdatedShaderType(EShaderPlatform Platform, const TShaderRef<FShader>& Shader, TArray<const FShaderType*>& OutdatedShaderTypes, TArray<const FVertexFactoryType*>& OutdatedFactoryTypes)
+{
+	const FShaderType* Type = Shader.GetType();
+	const bool bOutdatedShader = Type->GetSourceHash(Platform) != Shader->GetHash();
+
+	const FVertexFactoryType* VFType = Shader.GetVertexFactoryType();
+	const bool bOutdatedVertexFactory = VFType && VFType->GetSourceHash(Platform) != Shader->GetVertexFactoryHash();
+
+	if (bOutdatedShader)
+	{
+		OutdatedShaderTypes.AddUnique(Type);
+	}
+	if (bOutdatedVertexFactory)
+	{
+		OutdatedFactoryTypes.AddUnique(VFType);
+	}
+
+	return bOutdatedShader || bOutdatedVertexFactory;
+}
+
+void FShaderMapContent::GetOutdatedTypes(const FShaderMapBase& InShaderMap, TArray<const FShaderType*>& OutdatedShaderTypes, TArray<const FShaderPipelineType*>& OutdatedShaderPipelineTypes, TArray<const FVertexFactoryType*>& OutdatedFactoryTypes) const
+{
+	for (FShader* Shader : Shaders)
+	{
+		CheckOutdatedShaderType(GetShaderPlatform(), TShaderRef<FShader>(Shader, InShaderMap), OutdatedShaderTypes, OutdatedFactoryTypes);
+	}
+
+	for (const FShaderPipeline* Pipeline : ShaderPipelines)
+	{
+		for (const TShaderRef<FShader>& Shader : Pipeline->GetShaders(InShaderMap))
+		{
+			if (CheckOutdatedShaderType(GetShaderPlatform(), Shader, OutdatedShaderTypes, OutdatedFactoryTypes))
+			{
+				const FShaderPipelineType* PipelineType = FShaderPipelineType::GetShaderPipelineTypeByName(Pipeline->TypeName);
+				check(PipelineType);
+				OutdatedShaderPipelineTypes.AddUnique(PipelineType);
+			}
+		}
+	}
+}
+
+void FShaderMapContent::SaveShaderStableKeys(const FShaderMapBase& InShaderMap, EShaderPlatform TargetShaderPlatform, const struct FStableShaderKeyAndValue& SaveKeyVal)
+{
+	for (FShader* Shader : Shaders)
+	{
+		Shader->SaveShaderStableKeys(InShaderMap.GetPointerTable(), TargetShaderPlatform, SaveKeyVal);
+	}
+
+	for (const FShaderPipeline* Pipeline : ShaderPipelines)
+	{
+		Pipeline->SaveShaderStableKeys(InShaderMap.GetPointerTable(), TargetShaderPlatform, SaveKeyVal);
+	}
+}
+
+uint32 FShaderMapContent::GetMaxTextureSamplersShaderMap(const FShaderMapBase& InShaderMap) const
+{
+	uint32 MaxTextureSamplers = 0;
+
+	for (FShader* Shader : Shaders)
+	{
+		MaxTextureSamplers = FMath::Max(MaxTextureSamplers, Shader->GetNumTextureSamplers());
+	}
+
+	for (FShaderPipeline* Pipeline : ShaderPipelines)
+	{
+		for (const TShaderRef<FShader>& Shader : Pipeline->GetShaders(InShaderMap))
+		{
+			MaxTextureSamplers = FMath::Max(MaxTextureSamplers, Shader->GetNumTextureSamplers());
+		}
+	}
+
+	return MaxTextureSamplers;
+}
+#endif // WITH_EDITOR
+
+uint32 FShaderMapContent::GetNumShaders() const
+{
+	uint32 NumShaders = Shaders.Num();
+	for (FShaderPipeline* Pipeline : ShaderPipelines)
+	{
+		NumShaders += Pipeline->GetNumShaders();
+	}
+	return NumShaders;
+}
+
+uint32 FShaderMapContent::GetMaxNumInstructionsForShader(const FShaderMapBase& InShaderMap, FShaderType* ShaderType) const
+{
+	uint32 MaxNumInstructions = 0;
+	FShader* Shader = GetShader(ShaderType);
+	if (Shader)
+	{
+		MaxNumInstructions = FMath::Max(MaxNumInstructions, Shader->GetNumInstructions());
+	}
+
+	for (FShaderPipeline* Pipeline : ShaderPipelines)
+	{
+		FShader* PipelineShader = Pipeline->GetShader(ShaderType->GetFrequency());
+		if (PipelineShader)
+		{
+			MaxNumInstructions = FMath::Max(MaxNumInstructions, PipelineShader->GetNumInstructions());
+		}
+	}
+
+	return MaxNumInstructions;
+}
+
+void FShaderMapContent::Empty()
+{
+	EmptyShaderPipelines();
+	for (int32 i = 0; i < Shaders.Num(); ++i)
+	{
+		Shaders[i].SafeDelete();
+	}
+	Shaders.Empty();
+}
+
+void FShaderMapContent::EmptyShaderPipelines()
+{
+	for (TMemoryImagePtr<FShaderPipeline>& Pipeline : ShaderPipelines)
+	{
+		Pipeline.SafeDelete();
+	}
+	ShaderPipelines.Empty();
+}
\ No newline at end of file
diff --git a/Engine/Source/Runtime/RenderCore/Private/ShaderParameterMetadata.cpp b/Engine/Source/Runtime/RenderCore/Private/ShaderParameterMetadata.cpp
--- a/Engine/Source/Runtime/RenderCore/Private/ShaderParameterMetadata.cpp
+++ b/Engine/Source/Runtime/RenderCore/Private/ShaderParameterMetadata.cpp
@@ -44,9 +44,9 @@ TLinkedList<FShaderParametersMetadata*>*& FShaderParametersMetadata::GetStructLi
 	return GUniformStructList;
 }
 
-TMap<FName, FShaderParametersMetadata*>& FShaderParametersMetadata::GetNameStructMap()
+TMap<FHashedName, FShaderParametersMetadata*>& FShaderParametersMetadata::GetNameStructMap()
 {
-	static TMap<FName, FShaderParametersMetadata*> NameStructMap;
+	static TMap<FHashedName, FShaderParametersMetadata*> NameStructMap;
 	return NameStructMap;
 }
 
@@ -81,7 +81,7 @@ public:
 
 FShaderParametersMetadata::FShaderParametersMetadata(
 	EUseCase InUseCase,
-	const FName& InLayoutName,
+	const TCHAR* InLayoutName,
 	const TCHAR* InStructTypeName,
 	const TCHAR* InShaderVariableName,
 	const TCHAR* InStaticSlotName,
@@ -90,6 +90,7 @@ FShaderParametersMetadata::FShaderParametersMetadata(
 	: StructTypeName(InStructTypeName)
 	, ShaderVariableName(InShaderVariableName)
 	, StaticSlotName(InStaticSlotName)
+	, ShaderVariableHashedName(InShaderVariableName)
 	, Size(InSize)
 	, UseCase(InUseCase)
 	, Layout(InLayoutName)
@@ -117,7 +118,8 @@ FShaderParametersMetadata::FShaderParametersMetadata(
 		FName StructTypeFName(StructTypeName);
 		// Verify that during FName creation there's no case conversion
 		checkSlow(FCString::Strcmp(StructTypeName, *StructTypeFName.GetPlainNameString()) == 0);
-		GetNameStructMap().Add(FName(StructTypeFName), this);
+		GetNameStructMap().Add(ShaderVariableHashedName, this);
+
 
 #if VALIDATE_UNIFORM_BUFFER_UNIQUE_NAME
 		FName ShaderVariableFName(ShaderVariableName);
@@ -152,7 +154,7 @@ FShaderParametersMetadata::~FShaderParametersMetadata()
 	if (UseCase == EUseCase::UniformBuffer)
 	{
 		GlobalListLink.Unlink();
-		GetNameStructMap().Remove(FName(StructTypeName, FNAME_Find));
+		GetNameStructMap().Remove(ShaderVariableHashedName);
 
 #if VALIDATE_UNIFORM_BUFFER_UNIQUE_NAME
 		GlobalShaderVariableToStructMap.Remove(FName(ShaderVariableName, FNAME_Find));
diff --git a/Engine/Source/Runtime/RenderCore/Private/ShaderParameterStruct.cpp b/Engine/Source/Runtime/RenderCore/Private/ShaderParameterStruct.cpp
--- a/Engine/Source/Runtime/RenderCore/Private/ShaderParameterStruct.cpp
+++ b/Engine/Source/Runtime/RenderCore/Private/ShaderParameterStruct.cpp
@@ -178,7 +178,7 @@ struct FShaderParameterStructBindingContext
 					if (uint32(BoundSize) > ByteSize)
 					{
 						UE_LOG(LogShaders, Fatal, TEXT("The size required to bind shader %s's (Permutation Id %d) struct %s parameter %s is %i bytes, smaller than %s's %i bytes."),
-							Shader->GetType()->GetName(), Shader->GetPermutationId(), StructMetaData.GetStructTypeName(),
+							Shader->GetTypeUnfrozen()->GetName(), Shader->GetPermutationId(), StructMetaData.GetStructTypeName(),
 							*ElementShaderBindingName, BoundSize, *CppName, ByteSize);
 					}
 
@@ -210,7 +210,7 @@ struct FShaderParameterStructBindingContext
 								TEXT("Error with shader %s's (Permutation Id %d) parameter %s is %i bytes, cpp name = %s.")
 								TEXT("The shader compiler should give precisely which elements of an array did not get compiled out, ")
 								TEXT("for optimal automatic render graph pass dependency with ClearUnusedGraphResources()."),
-								Shader->GetType()->GetName(), Shader->GetPermutationId(),
+								Shader->GetTypeUnfrozen()->GetName(), Shader->GetPermutationId(),
 								*ElementShaderBindingName, BoundSize, *CppName);
 						}
 					}
@@ -241,10 +241,11 @@ struct FShaderParameterStructBindingContext
 
 void FShaderParameterBindings::BindForLegacyShaderParameters(const FShader* Shader, const FShaderParameterMap& ParametersMap, const FShaderParametersMetadata& StructMetaData, bool bShouldBindEverything)
 {
+	const FShaderType* Type = Shader->GetTypeUnfrozen();
 	checkf(StructMetaData.GetSize() < (1 << (sizeof(uint16) * 8)), TEXT("Shader parameter structure can only have a size < 65536 bytes."));
 	check(this == &Shader->Bindings);
 	
-	switch (Shader->GetType()->GetFrequency())
+	switch (Type->GetFrequency())
 	{
 	case SF_Vertex:
 	case SF_Hull:
@@ -276,7 +277,7 @@ void FShaderParameterBindings::BindForLegacyShaderParameters(const FShader* Shad
 	if (bShouldBindEverything && BindingContext.ShaderGlobalScopeBindings.Num() != AllParameterNames.Num())
 	{
 		FString ErrorString = FString::Printf(
-			TEXT("Shader %s has unbound parameters not represented in the parameter struct:"), Shader->GetType()->GetName());
+			TEXT("Shader %s has unbound parameters not represented in the parameter struct:"), Type->GetName());
 
 		for (const FString& GlobalParameterName : AllParameterNames)
 		{
@@ -292,13 +293,14 @@ void FShaderParameterBindings::BindForLegacyShaderParameters(const FShader* Shad
 
 void FShaderParameterBindings::BindForRootShaderParameters(const FShader* Shader, const FShaderParameterMap& ParametersMap)
 {
+	const FShaderType* Type = Shader->GetTypeUnfrozen();
 	check(this == &Shader->Bindings);
-	check(Shader->GetType()->GetRootParametersMetadata());
+	check(Type->GetRootParametersMetadata());
 
-	const FShaderParametersMetadata& StructMetaData = *Shader->GetType()->GetRootParametersMetadata();
+	const FShaderParametersMetadata& StructMetaData = *Type->GetRootParametersMetadata();
 	checkf(StructMetaData.GetSize() < (1 << (sizeof(uint16) * 8)), TEXT("Shader parameter structure can only have a size < 65536 bytes."));
 
-	switch (Shader->GetType()->GetFrequency())
+	switch (Type->GetFrequency())
 	{
 	case SF_RayGen:
 	case SF_RayMiss:
@@ -342,7 +344,7 @@ void FShaderParameterBindings::BindForRootShaderParameters(const FShader* Shader
 	if (BindingContext.ShaderGlobalScopeBindings.Num() != AllParameterNames.Num())
 	{
 		FString ErrorString = FString::Printf(
-			TEXT("Shader %s has unbound parameters not represented in the parameter struct:"), Shader->GetType()->GetName());
+			TEXT("Shader %s has unbound parameters not represented in the parameter struct:"), Type->GetName());
 
 		for (const FString& GlobalParameterName : AllParameterNames)
 		{
@@ -419,11 +421,11 @@ bool FDepthStencilBinding::Validate() const
 	return true;
 }
 
-void EmitNullShaderParameterFatalError(const FShader* Shader, const FShaderParametersMetadata* ParametersMetadata, uint16 MemberOffset)
+void EmitNullShaderParameterFatalError(const TShaderRef<FShader>& Shader, const FShaderParametersMetadata* ParametersMetadata, uint16 MemberOffset)
 {
 	FString MemberName = ParametersMetadata->GetFullMemberCodeName(MemberOffset);
 
-	const TCHAR* ShaderClassName = Shader->GetType()->GetName();
+	const TCHAR* ShaderClassName = Shader.GetType()->GetName();
 
 	UE_LOG(LogShaders, Fatal,
 		TEXT("%s's required shader parameter %s::%s was not set."),
@@ -434,18 +436,18 @@ void EmitNullShaderParameterFatalError(const FShader* Shader, const FShaderParam
 
 #if DO_CHECK
 
-void ValidateShaderParameters(const FShader* Shader, const FShaderParametersMetadata* ParametersMetadata, const void* Parameters)
+void ValidateShaderParameters(const TShaderRef<FShader>& Shader, const FShaderParametersMetadata* ParametersMetadata, const void* Parameters)
 {
 	const FShaderParameterBindings& Bindings = Shader->Bindings;
 
 	checkf(
 		Bindings.StructureLayoutHash == ParametersMetadata->GetLayoutHash(),
 		TEXT("Seams shader %s's parameter structure has changed without recompilation of the shader"),
-		Shader->GetType()->GetName());
+		Shader->GetTypeUnfrozen()->GetName());
 
 	const uint8* Base = reinterpret_cast<const uint8*>(Parameters);
 
-	const TCHAR* ShaderClassName = Shader->GetType()->GetName();
+	const TCHAR* ShaderClassName = Shader.GetType()->GetName();
 	const TCHAR* ShaderParemeterStructName = ParametersMetadata->GetStructTypeName();
 
 	// Textures
