Copying Private-LoadTimes-4.24 stream to Main. Biggest changes are in Materials/Shader memory freezing.

diff --git a/Engine/Source/Runtime/RenderCore/Private/ShaderParameters.cpp b/Engine/Source/Runtime/RenderCore/Private/ShaderParameters.cpp
--- a/Engine/Source/Runtime/RenderCore/Private/ShaderParameters.cpp
+++ b/Engine/Source/Runtime/RenderCore/Private/ShaderParameters.cpp
@@ -12,12 +12,13 @@
 #include "VertexFactory.h"
 #include "ShaderCodeLibrary.h"
 
+IMPLEMENT_TYPE_LAYOUT(FShaderParameter);
+IMPLEMENT_TYPE_LAYOUT(FShaderResourceParameter);
+IMPLEMENT_TYPE_LAYOUT(FRWShaderParameter);
+IMPLEMENT_TYPE_LAYOUT(FShaderUniformBufferParameter);
+
 void FShaderParameter::Bind(const FShaderParameterMap& ParameterMap,const TCHAR* ParameterName,EShaderParameterFlags Flags)
 {
-#if UE_BUILD_DEBUG
-	bInitialized = true;
-#endif
-
 	if (!ParameterMap.FindParameterAllocation(ParameterName,BufferIndex,BaseIndex,NumBytes) && Flags == SPF_Mandatory)
 	{
 		if (!UE_LOG_ACTIVE(LogShaders, Log))
@@ -36,13 +37,6 @@ void FShaderParameter::Bind(const FShaderParameterMap& ParameterMap,const TCHAR*
 
 FArchive& operator<<(FArchive& Ar,FShaderParameter& P)
 {
-#if UE_BUILD_DEBUG
-	if (Ar.IsLoading())
-	{
-		P.bInitialized = true;
-	}
-#endif
-
 	uint16& PBufferIndex = P.BufferIndex;
 	return Ar << P.BaseIndex << P.NumBytes << PBufferIndex;
 }
@@ -51,10 +45,6 @@ void FShaderResourceParameter::Bind(const FShaderParameterMap& ParameterMap,cons
 {
 	uint16 UnusedBufferIndex = 0;
 
-#if UE_BUILD_DEBUG
-	bInitialized = true;
-#endif
-
 	if(!ParameterMap.FindParameterAllocation(ParameterName,UnusedBufferIndex,BaseIndex,NumResources) && Flags == SPF_Mandatory)
 	{
 		if (!UE_LOG_ACTIVE(LogShaders, Log))
@@ -73,13 +63,6 @@ void FShaderResourceParameter::Bind(const FShaderParameterMap& ParameterMap,cons
 
 FArchive& operator<<(FArchive& Ar,FShaderResourceParameter& P)
 {
-#if UE_BUILD_DEBUG
-	if (Ar.IsLoading())
-	{
-		P.bInitialized = true;
-	}
-#endif
-
 	return Ar << P.BaseIndex << P.NumResources;
 }
 
@@ -103,13 +86,9 @@ void FShaderUniformBufferParameter::Bind(const FShaderParameterMap& ParameterMap
 	uint16 UnusedBaseIndex = 0;
 	uint16 UnusedNumBytes = 0;
 
-#if UE_BUILD_DEBUG
-	bInitialized = true;
-#endif
-
 	if(!ParameterMap.FindParameterAllocation(ParameterName,BaseIndex,UnusedBaseIndex,UnusedNumBytes))
 	{
-		bIsBound = false;
+		BaseIndex = 0xffff;
 		if(Flags == SPF_Mandatory)
 		{
 			if (!UE_LOG_ACTIVE(LogShaders, Log))
@@ -127,7 +106,7 @@ void FShaderUniformBufferParameter::Bind(const FShaderParameterMap& ParameterMap
 	}
 	else
 	{
-		bIsBound = true;
+		check(IsBound());
 	}
 }
 
@@ -411,6 +390,7 @@ void FShaderType::DumpDebugInfo()
 		break;
 	}
 
+#if 0
 	UE_LOG(LogConsoleResponse, Display, TEXT("  --- %d shaders"), ShaderIdMap.Num());
 	int32 Index = 0;
 	for (auto& KeyValue : ShaderIdMap)
@@ -420,7 +400,7 @@ void FShaderType::DumpDebugInfo()
 		Shader->DumpDebugInfo();
 		Index++;
 	}
-
+#endif
 }
 
 void FShaderType::GetShaderStableKeyParts(FStableShaderKeyAndValue& SaveKeyVal)
@@ -449,24 +429,6 @@ void FShaderType::GetShaderStableKeyParts(FStableShaderKeyAndValue& SaveKeyVal)
 #endif
 }
 
-void FShaderType::SaveShaderStableKeys(EShaderPlatform TargetShaderPlatform)
-{
-#if WITH_EDITOR
-	FStableShaderKeyAndValue SaveKeyVal;
-	if (ShaderIdMap.Num())
-	{
-		for (auto& KeyValue : ShaderIdMap)
-		{
-			FShader* Shader = KeyValue.Value;
-			check(Shader);
-			check(Shader->Type == this);
-			Shader->SaveShaderStableKeys(TargetShaderPlatform, SaveKeyVal);
-		}
-	}
-#endif
-}
-
-
 void FVertexFactoryType::AddReferencedUniformBufferIncludes(FShaderCompilerEnvironment& OutEnvironment, FString& OutSourceFilePrefix, EShaderPlatform Platform)
 {
 	// Cache uniform buffer struct declarations referenced by this shader type's files
diff --git a/Engine/Source/Runtime/RenderCore/Private/ShaderPipelineCache.cpp b/Engine/Source/Runtime/RenderCore/Private/ShaderPipelineCache.cpp
--- a/Engine/Source/Runtime/RenderCore/Private/ShaderPipelineCache.cpp
+++ b/Engine/Source/Runtime/RenderCore/Private/ShaderPipelineCache.cpp
@@ -584,8 +584,6 @@ bool FShaderPipelineCache::Precompile(FRHICommandListImmediate& RHICmdList, ESha
 	
 	if(PSO.Verify())
 	{
-		TArray<uint8> DummyCode;
-		
 		if(FPipelineCacheFileFormatPSO::DescriptorType::Graphics == PSO.Type)
 		{
 			FGraphicsPipelineStateInitializer GraphicsInitializer;
@@ -596,7 +594,7 @@ bool FShaderPipelineCache::Precompile(FRHICommandListImmediate& RHICmdList, ESha
 			FVertexShaderRHIRef VertexShader;
 			if (PSO.GraphicsDesc.VertexShader != FSHAHash())
 			{
-				VertexShader = FShaderCodeLibrary::CreateVertexShader(Platform, PSO.GraphicsDesc.VertexShader, DummyCode);
+				VertexShader = FShaderCodeLibrary::CreateVertexShader(Platform, PSO.GraphicsDesc.VertexShader);
 				GraphicsInitializer.BoundShaderState.VertexShaderRHI = VertexShader;
 			}
 
@@ -604,21 +602,21 @@ bool FShaderPipelineCache::Precompile(FRHICommandListImmediate& RHICmdList, ESha
 			FHullShaderRHIRef HullShader;
 			if (PSO.GraphicsDesc.HullShader != FSHAHash())
 			{
-				HullShader = FShaderCodeLibrary::CreateHullShader(Platform, PSO.GraphicsDesc.HullShader, DummyCode);
+				HullShader = FShaderCodeLibrary::CreateHullShader(Platform, PSO.GraphicsDesc.HullShader);
 				GraphicsInitializer.BoundShaderState.HullShaderRHI = HullShader;
 			}
 
 			FDomainShaderRHIRef DomainShader;
 			if (PSO.GraphicsDesc.DomainShader != FSHAHash())
 			{
-				DomainShader = FShaderCodeLibrary::CreateDomainShader(Platform, PSO.GraphicsDesc.DomainShader, DummyCode);
+				DomainShader = FShaderCodeLibrary::CreateDomainShader(Platform, PSO.GraphicsDesc.DomainShader);
 				GraphicsInitializer.BoundShaderState.DomainShaderRHI = DomainShader;
 			}
 	#endif
 			FPixelShaderRHIRef FragmentShader;
 			if (PSO.GraphicsDesc.FragmentShader != FSHAHash())
 			{
-				FragmentShader = FShaderCodeLibrary::CreatePixelShader(Platform, PSO.GraphicsDesc.FragmentShader, DummyCode);
+				FragmentShader = FShaderCodeLibrary::CreatePixelShader(Platform, PSO.GraphicsDesc.FragmentShader);
 				GraphicsInitializer.BoundShaderState.PixelShaderRHI = FragmentShader;
 			}
 
@@ -626,7 +624,7 @@ bool FShaderPipelineCache::Precompile(FRHICommandListImmediate& RHICmdList, ESha
 			FGeometryShaderRHIRef GeometryShader;
 			if (PSO.GraphicsDesc.GeometryShader != FSHAHash())
 			{
-				GeometryShader = FShaderCodeLibrary::CreateGeometryShader(Platform, PSO.GraphicsDesc.GeometryShader, DummyCode);
+				GeometryShader = FShaderCodeLibrary::CreateGeometryShader(Platform, PSO.GraphicsDesc.GeometryShader);
 				GraphicsInitializer.BoundShaderState.GeometryShaderRHI = GeometryShader;
 			}
 	#endif
@@ -671,7 +669,7 @@ bool FShaderPipelineCache::Precompile(FRHICommandListImmediate& RHICmdList, ESha
 		}
 		else if(FPipelineCacheFileFormatPSO::DescriptorType::Compute == PSO.Type)
 		{
-			FComputeShaderRHIRef ComputeInitializer = FShaderCodeLibrary::CreateComputeShader(GMaxRHIShaderPlatform, PSO.ComputeDesc.ComputeShader, DummyCode);
+			FComputeShaderRHIRef ComputeInitializer = FShaderCodeLibrary::CreateComputeShader(GMaxRHIShaderPlatform, PSO.ComputeDesc.ComputeShader);
 			if(ComputeInitializer.IsValid())
 			{
 				FComputePipelineState* ComputeResult = PipelineStateCache::GetAndOrCreateComputePipelineState(RHICmdList, ComputeInitializer);
@@ -701,16 +699,6 @@ bool FShaderPipelineCache::Precompile(FRHICommandListImmediate& RHICmdList, ESha
 	return bOk;
 }
 
-inline bool InternalRequestShaderCode(FSHAHash Shader, FShaderPipelineCacheArchive* ReadRequest, TSet<FSHAHash>& ShaderCodeRequests)
-{
-	const bool bResult = FShaderCodeLibrary::RequestShaderCode(Shader, ReadRequest);
-	if(bResult)
-	{
-		ShaderCodeRequests.Add(Shader);
-	}
-	return bResult;
-}
-
 void FShaderPipelineCache::PreparePipelineBatch(TDoubleLinkedList<FPipelineCacheFileFormatPSORead*>& PipelineBatch)
 {
 	TDoubleLinkedList<FPipelineCacheFileFormatPSORead*>::TDoubleLinkedListNode* CurrentNode = PipelineBatch.GetHead();
@@ -791,31 +779,31 @@ void FShaderPipelineCache::PreparePipelineBatch(TDoubleLinkedList<FPipelineCache
                 // If everything is OK then we can issue reads of the actual shader code
 				if (bOK && PSO.GraphicsDesc.VertexShader != FSHAHash())
 				{
-					bOK &= InternalRequestShaderCode(PSO.GraphicsDesc.VertexShader, AsyncJob.ReadRequests, AsyncJob.ShaderCodeReads);
+					bOK &= FShaderCodeLibrary::PreloadShader(PSO.GraphicsDesc.VertexShader, AsyncJob.ReadRequests);
                     UE_CLOG(!bOK, LogRHI, Verbose, TEXT("Failed to read VertexShader shader: %s"), *(PSO.GraphicsDesc.VertexShader.ToString()));
 				}
 		
 				if (bOK && PSO.GraphicsDesc.HullShader != EmptySHA)
 				{
-					bOK &= InternalRequestShaderCode(PSO.GraphicsDesc.HullShader, AsyncJob.ReadRequests, AsyncJob.ShaderCodeReads);
+					bOK &= FShaderCodeLibrary::PreloadShader(PSO.GraphicsDesc.HullShader, AsyncJob.ReadRequests);
                     UE_CLOG(!bOK, LogRHI, Verbose, TEXT("Failed to read HullShader shader: %s"), *(PSO.GraphicsDesc.HullShader.ToString()));
 				}
 		
 				if (bOK && PSO.GraphicsDesc.DomainShader != EmptySHA)
 				{
-					bOK &= InternalRequestShaderCode(PSO.GraphicsDesc.DomainShader, AsyncJob.ReadRequests, AsyncJob.ShaderCodeReads);
+					bOK &= FShaderCodeLibrary::PreloadShader(PSO.GraphicsDesc.DomainShader, AsyncJob.ReadRequests);
                     UE_CLOG(!bOK, LogRHI, Verbose, TEXT("Failed to read DomainShader shader: %s"), *(PSO.GraphicsDesc.DomainShader.ToString()));
 				}
 		
 				if (bOK && PSO.GraphicsDesc.FragmentShader != EmptySHA)
 				{
-					bOK &= InternalRequestShaderCode(PSO.GraphicsDesc.FragmentShader, AsyncJob.ReadRequests, AsyncJob.ShaderCodeReads);
+					bOK &= FShaderCodeLibrary::PreloadShader(PSO.GraphicsDesc.FragmentShader, AsyncJob.ReadRequests);
                     UE_CLOG(!bOK, LogRHI, Verbose, TEXT("Failed to read FragmentShader shader: %s"), *(PSO.GraphicsDesc.FragmentShader.ToString()));
 				}
 		
 				if (bOK && PSO.GraphicsDesc.GeometryShader != EmptySHA)
 				{
-					bOK &= InternalRequestShaderCode(PSO.GraphicsDesc.GeometryShader, AsyncJob.ReadRequests, AsyncJob.ShaderCodeReads);
+					bOK &= FShaderCodeLibrary::PreloadShader(PSO.GraphicsDesc.GeometryShader, AsyncJob.ReadRequests);
                     UE_CLOG(!bOK, LogRHI, Verbose, TEXT("Failed to read GeometryShader shader: %s"), *(PSO.GraphicsDesc.GeometryShader.ToString()));
 				}
 			}
@@ -824,7 +812,7 @@ void FShaderPipelineCache::PreparePipelineBatch(TDoubleLinkedList<FPipelineCache
 				if (PSO.ComputeDesc.ComputeShader != EmptySHA)
 				{
                     RequiredShaders.Add(PSO.ComputeDesc.ComputeShader);
-					bOK &= InternalRequestShaderCode(PSO.ComputeDesc.ComputeShader, AsyncJob.ReadRequests, AsyncJob.ShaderCodeReads);
+					bOK &= FShaderCodeLibrary::PreloadShader(PSO.ComputeDesc.ComputeShader, AsyncJob.ReadRequests);
                     UE_CLOG(!bOK, LogRHI, Verbose, TEXT("Failed to find ComputeShader shader: %s"), *(PSO.ComputeDesc.ComputeShader.ToString()));
 				}
 				else
@@ -927,29 +915,6 @@ void FShaderPipelineCache::PrecompilePipelineBatch()
 		Precompile(RHICmdList, GMaxRHIShaderPlatform, CompileTask.PSO);
 		CompiledHashes.Add(GetTypeHash(CompileTask.PSO));
 		
-		// Free our code references that were actually made
-		if(CompileTask.ShaderCodeReads.Num())
-		{
-			if(GRHILazyShaderCodeLoading && IsRunningRHIInSeparateThread())
-			{
-				TSet<FSHAHash> ShaderCodeReads = CompileTask.ShaderCodeReads;
-				RHICmdList.EnqueueLambda([ShaderCodeReads](FRHICommandListImmediate&)
-				{
-					for (const FSHAHash& DelHash : ShaderCodeReads)
-					{
-						FShaderCodeLibrary::ReleaseShaderCode(DelHash);
-					}
-				});
-			}
-			else
-			{
-				for (const FSHAHash& DelHash : CompileTask.ShaderCodeReads)
-				{
-					FShaderCodeLibrary::ReleaseShaderCode(DelHash);
-				}
-			}
-		}
-		
 		delete CompileTask.ReadRequests;
 		CompileTask.ReadRequests = nullptr;
 		
@@ -1022,12 +987,6 @@ void FShaderPipelineCache::PollShutdownItems()
 			check(ShutdownReadCompileTasks[i].ReadRequests);
 			if (ShutdownReadCompileTasks[i].ReadRequests->PollExternalReadDependencies())
 			{
-				// Free our code references that were actually made - maybe not all were
-				for (const FSHAHash& DelHash : ShutdownReadCompileTasks[i].ShaderCodeReads)
-				{
-					FShaderCodeLibrary::ReleaseShaderCode(DelHash);
-				}
-				
                 delete ShutdownReadCompileTasks[i].ReadRequests;
 				ShutdownReadCompileTasks[i].ReadRequests = nullptr;
 				
diff --git a/Engine/Source/Runtime/RenderCore/Private/ShaderResource.cpp b/Engine/Source/Runtime/RenderCore/Private/ShaderResource.cpp
new file mode 100644
--- /dev/null
+++ b/Engine/Source/Runtime/RenderCore/Private/ShaderResource.cpp
@@ -0,0 +1,343 @@
+// Copyright 1998-2019 Epic Games, Inc. All Rights Reserved.
+
+/*=============================================================================
+	ShaderResource.cpp: ShaderResource implementation.
+=============================================================================*/
+
+#include "Shader.h"
+#include "Misc/CoreMisc.h"
+#include "Stats/StatsMisc.h"
+#include "Serialization/MemoryWriter.h"
+#include "VertexFactory.h"
+#include "ProfilingDebugging/DiagnosticTable.h"
+#include "Interfaces/ITargetPlatform.h"
+#include "Interfaces/ITargetPlatformManagerModule.h"
+#include "Interfaces/IShaderFormat.h"
+#include "ShaderCodeLibrary.h"
+#include "ShaderCore.h"
+#include "RenderUtils.h"
+#include "Misc/ConfigCacheIni.h"
+#include "Misc/ScopeLock.h"
+#include "UObject/RenderingObjectVersion.h"
+#include "UObject/FortniteMainBranchObjectVersion.h"
+#include "ProfilingDebugging/LoadTimeTracker.h"
+#include "Misc/MemStack.h"
+
+#if WITH_EDITORONLY_DATA
+#include "Interfaces/IShaderFormat.h"
+#endif
+
+static const FName ShaderCompressionFormat = NAME_LZ4;
+
+bool FShaderMapResource::ArePlatformsCompatible(EShaderPlatform CurrentPlatform, EShaderPlatform TargetPlatform)
+{
+	bool bFeatureLevelCompatible = CurrentPlatform == TargetPlatform;
+
+	if (!bFeatureLevelCompatible && IsPCPlatform(CurrentPlatform) && IsPCPlatform(TargetPlatform))
+	{
+		bFeatureLevelCompatible = GetMaxSupportedFeatureLevel(CurrentPlatform) >= GetMaxSupportedFeatureLevel(TargetPlatform);
+
+		bool const bIsTargetD3D = TargetPlatform == SP_PCD3D_SM5 ||
+			TargetPlatform == SP_PCD3D_ES3_1;
+
+		bool const bIsCurrentPlatformD3D = CurrentPlatform == SP_PCD3D_SM5 ||
+			TargetPlatform == SP_PCD3D_ES3_1;
+
+		// For Metal in Editor we can switch feature-levels, but not in cooked projects when using Metal shader librariss.
+		bool const bIsCurrentMetal = IsMetalPlatform(CurrentPlatform);
+		bool const bIsTargetMetal = IsMetalPlatform(TargetPlatform);
+		bool const bIsMetalCompatible = (bIsCurrentMetal == bIsTargetMetal)
+#if !WITH_EDITOR	// Static analysis doesn't like (|| WITH_EDITOR)
+			&& (!IsMetalPlatform(CurrentPlatform) || (CurrentPlatform == TargetPlatform))
+#endif
+			;
+
+		bool const bIsCurrentOpenGL = IsOpenGLPlatform(CurrentPlatform);
+		bool const bIsTargetOpenGL = IsOpenGLPlatform(TargetPlatform);
+
+		bFeatureLevelCompatible = bFeatureLevelCompatible && (bIsCurrentPlatformD3D == bIsTargetD3D && bIsMetalCompatible && bIsCurrentOpenGL == bIsTargetOpenGL);
+	}
+
+	return bFeatureLevelCompatible;
+}
+
+#if RHI_RAYTRACING
+static TArray<uint32> GlobalUnusedIndicies;
+static TArray<FRHIRayTracingShader*> GlobalRayTracingMaterialLibrary;
+static FCriticalSection GlobalRayTracingMaterialLibraryCS;
+
+void FShaderMapResource::GetRayTracingMaterialLibrary(TArray<FRHIRayTracingShader*>& RayTracingMaterials, FRHIRayTracingShader* DefaultShader)
+{
+	FScopeLock Lock(&GlobalRayTracingMaterialLibraryCS);
+	RayTracingMaterials = GlobalRayTracingMaterialLibrary;
+
+	for (uint32 Index : GlobalUnusedIndicies)
+	{
+		RayTracingMaterials[Index] = DefaultShader;
+	}
+}
+
+static uint32 AddToRayTracingLibrary(FRHIRayTracingShader* Shader)
+{
+	FScopeLock Lock(&GlobalRayTracingMaterialLibraryCS);
+
+	if (GlobalUnusedIndicies.Num() != 0)
+	{
+		uint32 Index = GlobalUnusedIndicies.Pop(false);
+		checkSlow(GlobalRayTracingMaterialLibrary[Index] == nullptr);
+		GlobalRayTracingMaterialLibrary[Index] = Shader;
+		return Index;
+	}
+	else
+	{
+		GlobalRayTracingMaterialLibrary.Add(Shader);
+		return GlobalRayTracingMaterialLibrary.Num() - 1;
+	}
+}
+
+static void RemoveFromRayTracingLibrary(uint32 Index)
+{
+	if (Index != ~0u)
+	{
+		FScopeLock Lock(&GlobalRayTracingMaterialLibraryCS);
+		GlobalUnusedIndicies.Push(Index);
+		GlobalRayTracingMaterialLibrary[Index] = nullptr;
+	}
+}
+#endif // RHI_RAYTRACING
+
+static void ApplyResourceStats(FShaderMapResourceCode& Resource)
+{
+#if STATS
+	INC_DWORD_STAT_BY(STAT_Shaders_ShaderResourceMemory, Resource.GetSizeBytes());
+	for (FShaderMapResourceCode::FShaderEntry& Shader : Resource.ShaderEntries)
+	{
+		INC_DWORD_STAT_BY_FName(GetMemoryStatType(Shader.Frequency).GetName(), Shader.CompressedSize);
+	}
+#endif // STATS
+}
+
+static void RemoveResourceStats(FShaderMapResourceCode& Resource)
+{
+#if STATS
+	DEC_DWORD_STAT_BY(STAT_Shaders_ShaderResourceMemory, Resource.GetSizeBytes());
+	for (FShaderMapResourceCode::FShaderEntry& Shader : Resource.ShaderEntries)
+	{
+		DEC_DWORD_STAT_BY_FName(GetMemoryStatType(Shader.Frequency).GetName(), Shader.CompressedSize);
+	}
+#endif // STATS
+}
+
+FShaderMapResourceCode::~FShaderMapResourceCode()
+{
+	RemoveResourceStats(*this);
+}
+
+void FShaderMapResourceCode::Finalize()
+{
+	FSHA1 Hasher;
+	Hasher.Update((uint8*)ShaderHashes.GetData(), ShaderHashes.Num() * sizeof(FSHAHash));
+	Hasher.Final();
+	Hasher.GetHash(ResourceHash.Hash);
+	ApplyResourceStats(*this);
+}
+
+void FShaderMapResourceCode::Serialize(FArchive& Ar)
+{
+	Ar << ResourceHash;
+	Ar << ShaderHashes;
+	Ar << ShaderEntries;
+	Ar << ShaderCode;
+	check(ShaderEntries.Num() == ShaderHashes.Num());
+	ApplyResourceStats(*this);
+}
+
+FShaderMapResourceBuilder::FShaderMapResourceBuilder(FShaderMapResourceCode* InCode) : ShaderHashTable(1024, 256), Code(InCode)
+{
+	check(InCode->ShaderEntries.Num() == InCode->ShaderHashes.Num());
+	for(int32 ShaderIndex = 0; ShaderIndex < InCode->ShaderHashes.Num(); ++ShaderIndex)
+	{
+		const uint32 Key = GetTypeHash(InCode->ShaderHashes[ShaderIndex]);
+		ShaderHashTable.Add(Key, ShaderIndex);
+	}
+}
+
+int32 FShaderMapResourceBuilder::FindCode(const FSHAHash& InHash, uint32 InKey) const
+{
+	for (int32 Index = ShaderHashTable.First(InKey); ShaderHashTable.IsValid(Index); Index = ShaderHashTable.Next(Index))
+	{
+		if (Code->ShaderHashes[Index] == InHash)
+		{
+			return Index;
+		}
+	}
+	return INDEX_NONE;
+}
+
+int32 FShaderMapResourceBuilder::FindOrAddCode(EShaderFrequency InFrequency, const FSHAHash& InHash, const TConstArrayView<uint8>& InCode)
+{
+	const uint32 Key = GetTypeHash(InHash);
+	int32 Index = FindCode(InHash, Key);
+	if(Index == INDEX_NONE)
+	{
+		int32 CompressedSize = InCode.Num();
+		TArray<uint8> CompressedCode;
+		CompressedCode.AddUninitialized(CompressedSize);
+		if(FCompression::CompressMemory(ShaderCompressionFormat, CompressedCode.GetData(), CompressedSize, InCode.GetData(), InCode.Num()))
+		{
+			CompressedCode.SetNum(CompressedSize, false);
+		}
+		else
+		{
+			CompressedCode = InCode;
+		}
+
+		Index = Code->ShaderEntries.Num();
+		ShaderHashTable.Add(Key, Index);
+
+		Code->ShaderHashes.Add(InHash);
+		FShaderMapResourceCode::FShaderEntry& Entry = Code->ShaderEntries.AddDefaulted_GetRef();
+		Entry.CompressedSize = CompressedSize;
+		Entry.UncompressedSize = InCode.Num();
+		Entry.Frequency = InFrequency;
+		Entry.Offset = Code->ShaderCode.Num();
+
+		Code->ShaderCode.Append(CompressedCode);
+	}
+
+	return Index;
+}
+
+FShaderMapResource::FShaderMapResource(EShaderPlatform InPlatform, int32 NumShaders)
+	: Platform(InPlatform)
+	, NumRefs(0)
+{
+	RHIShaders.AddZeroed(NumShaders);
+#if RHI_RAYTRACING
+	RayTracingMaterialLibraryIndices.AddUninitialized(NumShaders);
+	FMemory::Memset(RayTracingMaterialLibraryIndices.GetData(), 0xff, NumShaders * RayTracingMaterialLibraryIndices.GetTypeSize());
+#endif // RHI_RAYTRACING
+}
+
+FShaderMapResource::~FShaderMapResource()
+{
+	check(NumRefs == 0);
+}
+
+void FShaderMapResource::AddRef()
+{
+	FPlatformAtomics::InterlockedIncrement((volatile int32*)&NumRefs);
+}
+
+void FShaderMapResource::Release()
+{
+	check(NumRefs > 0);
+	if (FPlatformAtomics::InterlockedDecrement((volatile int32*)&NumRefs) == 0 && TryRelease())
+	{
+		// Send a release message to the rendering thread when the shader loses its last reference.
+		BeginReleaseResource(this);
+		BeginCleanup(this);
+
+		DEC_DWORD_STAT_BY(STAT_Shaders_ShaderResourceMemory, GetSizeBytes());
+	}
+}
+
+void FShaderMapResource::ReleaseRHI()
+{
+#if RHI_RAYTRACING
+	for (int32 Index : RayTracingMaterialLibraryIndices)
+	{
+		RemoveFromRayTracingLibrary(Index);
+	}
+	RayTracingMaterialLibraryIndices.Empty();
+#endif // RHI_RAYTRACING
+
+	RHIShaders.Empty();
+}
+
+void FShaderMapResource::BeginCreateAllShaders()
+{
+	FShaderMapResource* Resource = this;
+	ENQUEUE_RENDER_COMMAND(InitCommand)(
+		[Resource](FRHICommandListImmediate& RHICmdList)
+	{
+		for (int32 ShaderIndex = 0; ShaderIndex < Resource->GetNumShaders(); ++ShaderIndex)
+		{
+			Resource->GetShader(ShaderIndex);
+		}
+	});
+}
+
+void FShaderMapResource::CreateShader(int32 ShaderIndex)
+{
+	check(IsInParallelRenderingThread());
+	check(!RHIShaders[ShaderIndex]);
+
+	TRefCountPtr<FRHIShader> RHIShader = CreateRHIShader(ShaderIndex);
+#if RHI_RAYTRACING
+	if (GRHISupportsRayTracing && RHIShader.IsValid() && RHIShader->GetFrequency() == SF_RayHitGroup)
+	{
+		RayTracingMaterialLibraryIndices[ShaderIndex] = AddToRayTracingLibrary(static_cast<FRHIRayTracingShader*>(RHIShader.GetReference()));
+	}
+#endif // RHI_RAYTRACING
+	RHIShaders[ShaderIndex] = MoveTemp(RHIShader);
+}
+
+TRefCountPtr<FRHIShader> FShaderMapResource_InlineCode::CreateRHIShader(int32 ShaderIndex)
+{
+	// we can't have this called on the wrong platform's shaders
+	if (!ArePlatformsCompatible(GMaxRHIShaderPlatform, GetPlatform()))
+	{
+		if (FPlatformProperties::RequiresCookedData())
+		{
+			UE_LOG(LogShaders, Fatal, TEXT("FShaderMapResource_InlineCode::InitRHI got platform %s but it is not compatible with %s"),
+				*LegacyShaderPlatformToShaderFormat(GetPlatform()).ToString(), *LegacyShaderPlatformToShaderFormat(GMaxRHIShaderPlatform).ToString());
+		}
+		return TRefCountPtr<FRHIShader>();
+	}
+
+	FMemStackBase& MemStack = FMemStack::Get();
+	const FShaderMapResourceCode::FShaderEntry& ShaderEntry = Code->ShaderEntries[ShaderIndex];
+	const uint8* ShaderCode = Code->ShaderCode.GetData() + ShaderEntry.Offset;
+
+	FMemMark Mark(MemStack);
+	if (ShaderEntry.CompressedSize != ShaderEntry.UncompressedSize)
+	{
+		void* UncompressedCode = MemStack.Alloc(ShaderEntry.UncompressedSize, 16);
+		auto bSucceed = FCompression::UncompressMemory(ShaderCompressionFormat, UncompressedCode, ShaderEntry.UncompressedSize, ShaderCode, ShaderEntry.CompressedSize);
+		check(bSucceed);
+		ShaderCode = (uint8*)UncompressedCode;
+	}
+
+	const auto ShaderCodeView = MakeArrayView(ShaderCode, ShaderEntry.UncompressedSize);
+	const FSHAHash& ShaderHash = Code->ShaderHashes[ShaderIndex];
+	const EShaderFrequency Frequency = ShaderEntry.Frequency;
+
+	TRefCountPtr<FRHIShader> RHIShader;
+	switch (Frequency)
+	{
+	case SF_Vertex: RHIShader = RHICreateVertexShader(ShaderCodeView, ShaderHash); break;
+	case SF_Pixel: RHIShader = RHICreatePixelShader(ShaderCodeView, ShaderHash); break;
+	case SF_Hull: RHIShader = RHICreateHullShader(ShaderCodeView, ShaderHash); break;
+	case SF_Domain: RHIShader = RHICreateDomainShader(ShaderCodeView, ShaderHash); break;
+	case SF_Geometry: RHIShader = RHICreateGeometryShader(ShaderCodeView, ShaderHash); break;
+	case SF_Compute: RHIShader = RHICreateComputeShader(ShaderCodeView, ShaderHash); break;
+#if RHI_RAYTRACING
+	case SF_RayGen: case SF_RayMiss: case SF_RayHitGroup: case SF_RayCallable:
+		if (GRHISupportsRayTracing)
+		{
+			RHIShader = RHICreateRayTracingShader(ShaderCodeView, ShaderHash, Frequency);
+		}
+		break;
+#endif // RHI_RAYTRACING
+	default:
+		checkNoEntry();
+		break;
+	}
+
+	if (RHIShader)
+	{
+		RHIShader->SetHash(ShaderHash);
+	}
+	return RHIShader;
+}
diff --git a/Engine/Source/Runtime/RenderCore/Private/UnifiedBuffer.cpp b/Engine/Source/Runtime/RenderCore/Private/UnifiedBuffer.cpp
--- a/Engine/Source/Runtime/RenderCore/Private/UnifiedBuffer.cpp
+++ b/Engine/Source/Runtime/RenderCore/Private/UnifiedBuffer.cpp
@@ -12,7 +12,7 @@
 
 class FByteBufferShader : public FGlobalShader
 {
-	DECLARE_GLOBAL_SHADER( FByteBufferShader );
+	DECLARE_INLINE_TYPE_LAYOUT( FByteBufferShader, NonVirtual );
 
 	FByteBufferShader() {}
 	FByteBufferShader( const ShaderMetaType::CompiledShaderInitializerType& Initializer )
@@ -42,7 +42,7 @@ class FByteBufferShader : public FGlobalShader
 				&& FDataDrivenShaderPlatformInfo::GetInfo(Parameters.Platform).bSupportsByteBufferComputeShaders;
 				*/
 			// TODO: Workaround for FDataDrivenShaderPlatformInfo::GetInfo not being properly filled out yet.
-			return FDataDrivenShaderPlatformInfo::GetInfo(Parameters.Platform).bSupportsByteBufferComputeShaders || Parameters.Platform == SP_PS4 || Parameters.Platform == SP_PCD3D_SM5 || Parameters.Platform == SP_XBOXONE_D3D12;
+			return FDataDrivenShaderPlatformInfo::GetSupportsByteBufferComputeShaders(Parameters.Platform) || Parameters.Platform == SP_PS4 || Parameters.Platform == SP_PCD3D_SM5 || Parameters.Platform == SP_XBOXONE_D3D12;
 		}
 	}
 
diff --git a/Engine/Source/Runtime/RenderCore/Private/VertexFactory.cpp b/Engine/Source/Runtime/RenderCore/Private/VertexFactory.cpp
--- a/Engine/Source/Runtime/RenderCore/Private/VertexFactory.cpp
+++ b/Engine/Source/Runtime/RenderCore/Private/VertexFactory.cpp
@@ -9,11 +9,25 @@
 #include "UObject/DebugSerializationFlags.h"
 #include "PipelineStateCache.h"
 
-uint32 FVertexFactoryType::NextHashIndex = 0;
+IMPLEMENT_TYPE_LAYOUT(FVertexFactoryShaderParameters);
+
+uint32 FVertexFactoryType::NumVertexFactories = 0;
 bool FVertexFactoryType::bInitializedSerializationHistory = false;
 
 static TLinkedList<FVertexFactoryType*>* GVFTypeList = nullptr;
 
+static TArray<FVertexFactoryType*>& GetSortedMaterialVertexFactoryTypes()
+{
+	static TArray<FVertexFactoryType*> Types;
+	return Types;
+}
+
+static TMap<FHashedName, FVertexFactoryType*>& GetVFTypeMap()
+{
+	static TMap<FHashedName, FVertexFactoryType*> VTTypeMap;
+	return VTTypeMap;
+}
+
 /**
  * @return The global shader factory list.
  */
@@ -22,20 +36,18 @@ TLinkedList<FVertexFactoryType*>*& FVertexFactoryType::GetTypeList()
 	return GVFTypeList;
 }
 
+const TArray<FVertexFactoryType*>& FVertexFactoryType::GetSortedMaterialTypes()
+{
+	return GetSortedMaterialVertexFactoryTypes();
+}
+
 /**
  * Finds a FVertexFactoryType by name.
  */
-FVertexFactoryType* FVertexFactoryType::GetVFByName(const FString& VFName)
+FVertexFactoryType* FVertexFactoryType::GetVFByName(const FHashedName& VFName)
 {
-	for(TLinkedList<FVertexFactoryType*>::TIterator It(GetTypeList()); It; It.Next())
-	{
-		FString CurrentVFName = FString(It->GetName());
-		if (CurrentVFName == VFName)
-		{
-			return *It;
-		}
-	}
-	return NULL;
+	FVertexFactoryType** Type = GetVFTypeMap().Find(VFName);
+	return Type ? *Type : nullptr;
 }
 
 void FVertexFactoryType::Initialize(const TMap<FString, TArray<const TCHAR*> >& ShaderFileToUniformBufferVariables)
@@ -48,22 +60,6 @@ void FVertexFactoryType::Initialize(const TMap<FString, TArray<const TCHAR*> >&
 		{
 			FVertexFactoryType* Type = *It;
 			GenerateReferencedUniformBuffers(Type->ShaderFilename, Type->Name, ShaderFileToUniformBufferVariables, Type->ReferencedUniformBufferStructsCache);
-
-			for (int32 Frequency = 0; Frequency < SF_NumFrequencies; Frequency++)
-			{
-				// Construct a temporary shader parameter instance, which is initialized to safe values for serialization
-				FVertexFactoryShaderParameters* Parameters = Type->CreateShaderParameters((EShaderFrequency)Frequency);
-
-				if (Parameters)
-				{
-					// Serialize the temp shader to memory and record the number and sizes of serializations
-					TArray<uint8> TempData;
-					FMemoryWriter Ar(TempData, true);
-					FShaderSaveArchive SaveArchive(Ar, Type->SerializationHistory[Frequency]);
-					Parameters->Serialize(SaveArchive);
-					delete Parameters;
-				}
-			}
 		}
 	}
 
@@ -72,16 +68,6 @@ void FVertexFactoryType::Initialize(const TMap<FString, TArray<const TCHAR*> >&
 
 void FVertexFactoryType::Uninitialize()
 {
-	for(TLinkedList<FVertexFactoryType*>::TIterator It(FVertexFactoryType::GetTypeList()); It; It.Next())
-	{
-		FVertexFactoryType* Type = *It;
-
-		for (int32 Frequency = 0; Frequency < SF_NumFrequencies; Frequency++)
-		{
-			Type->SerializationHistory[Frequency] = FSerializationHistory();
-		}
-	}
-
 	bInitializedSerializationHistory = false;
 }
 
@@ -96,6 +82,8 @@ FVertexFactoryType::FVertexFactoryType(
 	bool bInSupportsCachingMeshDrawCommands,
 	bool bInSupportsPrimitiveIdStream,
 	ConstructParametersType InConstructParameters,
+	GetParameterTypeLayoutType InGetParameterTypeLayout,
+	GetParameterTypeElementShaderBindingsType InGetParameterTypeElementShaderBindings,
 	ShouldCacheType InShouldCache,
 	ModifyCompilationEnvironmentType InModifyCompilationEnvironment,
 	ValidateCompiledResultType InValidateCompiledResult,
@@ -104,6 +92,7 @@ FVertexFactoryType::FVertexFactoryType(
 	Name(InName),
 	ShaderFilename(InShaderFilename),
 	TypeName(InName),
+	HashedName(TypeName),
 	bUsedWithMaterials(bInUsedWithMaterials),
 	bSupportsStaticLighting(bInSupportsStaticLighting),
 	bSupportsDynamicLighting(bInSupportsDynamicLighting),
@@ -112,6 +101,8 @@ FVertexFactoryType::FVertexFactoryType(
 	bSupportsCachingMeshDrawCommands(bInSupportsCachingMeshDrawCommands),
 	bSupportsPrimitiveIdStream(bInSupportsPrimitiveIdStream),
 	ConstructParameters(InConstructParameters),
+	GetParameterTypeLayout(InGetParameterTypeLayout),
+	GetParameterTypeElementShaderBindings(InGetParameterTypeElementShaderBindings),
 	ShouldCacheRef(InShouldCache),
 	ModifyCompilationEnvironmentRef(InModifyCompilationEnvironment),
 	ValidateCompiledResultRef(InValidateCompiledResult),
@@ -133,14 +124,33 @@ FVertexFactoryType::FVertexFactoryType(
 
 	// Add this vertex factory type to the global list.
 	GlobalListLink.LinkHead(GetTypeList());
+	GetVFTypeMap().Add(HashedName, this);
 
-	// Assign the vertex factory type the next unassigned hash index.
-	HashIndex = NextHashIndex++;
+	if (bUsedWithMaterials)
+	{
+		TArray<FVertexFactoryType*>& SortedTypes = GetSortedMaterialVertexFactoryTypes();
+		const int32 SortedIndex = Algo::LowerBoundBy(SortedTypes, HashedName, [](const FVertexFactoryType* InType) { return InType->GetHashedName(); });
+		SortedTypes.Insert(this, SortedIndex);
+	}
+
+	++NumVertexFactories;
 }
 
 FVertexFactoryType::~FVertexFactoryType()
 {
 	GlobalListLink.Unlink();
+	verify(GetVFTypeMap().Remove(HashedName) == 1);
+
+	if (bUsedWithMaterials)
+	{
+		TArray<FVertexFactoryType*>& SortedTypes = GetSortedMaterialVertexFactoryTypes();
+		const int32 SortedIndex = Algo::BinarySearchBy(SortedTypes, HashedName, [](const FVertexFactoryType* InType) { return InType->GetHashedName(); });
+		check(SortedIndex != INDEX_NONE);
+		SortedTypes.RemoveAt(SortedIndex);
+	}
+
+	check(NumVertexFactories > 0u);
+	--NumVertexFactories;
 }
 
 /** Calculates a Hash based on this vertex factory type's source code and includes */
@@ -165,17 +175,9 @@ FArchive& operator<<(FArchive& Ar,FVertexFactoryType*& TypeRef)
 	return Ar;
 }
 
-FVertexFactoryType* FindVertexFactoryType(FName TypeName)
+FVertexFactoryType* FindVertexFactoryType(const FHashedName& TypeName)
 {
-	// Search the global vertex factory list for a type with a matching name.
-	for(TLinkedList<FVertexFactoryType*>::TIterator VertexFactoryTypeIt(FVertexFactoryType::GetTypeList());VertexFactoryTypeIt;VertexFactoryTypeIt.Next())
-	{
-		if(VertexFactoryTypeIt->GetFName() == TypeName)
-		{
-			return *VertexFactoryTypeIt;
-		}
-	}
-	return NULL;
+	return FVertexFactoryType::GetVFByName(TypeName);
 }
 
 void FVertexFactory::GetStreams(ERHIFeatureLevel::Type InFeatureLevel, EVertexInputStreamType VertexStreamType, FVertexInputStreamArray& OutVertexStreams) const
@@ -316,107 +318,6 @@ void FVertexFactory::InitDeclaration(const FVertexDeclarationElementList& Elemen
 	}
 }
 
-FVertexFactoryParameterRef::FVertexFactoryParameterRef(FVertexFactoryType* InVertexFactoryType,const FShaderParameterMap& ParameterMap, EShaderFrequency InShaderFrequency, EShaderPlatform InShaderPlatform)
-: Parameters(NULL)
-, VertexFactoryType(InVertexFactoryType)
-, ShaderFrequency(InShaderFrequency)
-, ShaderPlatform(InShaderPlatform)
-{
-	Parameters = VertexFactoryType->CreateShaderParameters(InShaderFrequency);
-#if KEEP_SHADER_SOURCE_HASHES
-	VFHash = GetShaderFileHash(VertexFactoryType->GetShaderFilename(), InShaderPlatform);
-#endif
-
-	if(Parameters)
-	{
-		Parameters->Bind(ParameterMap);
-	}
-}
-
-bool operator<<(FArchive& Ar,FVertexFactoryParameterRef& Ref)
-{
-	bool bShaderHasOutdatedParameters = false;
-
-	Ar << Ref.VertexFactoryType;
-
-	uint8 ShaderFrequencyByte = Ref.ShaderFrequency;
-	Ar << ShaderFrequencyByte;
-	if(Ar.IsLoading())
-	{
-		Ref.ShaderFrequency = (EShaderFrequency)ShaderFrequencyByte;
-	}
-
-	uint8 ShaderPlatformByte = Ref.ShaderPlatform;
-	Ar << ShaderPlatformByte;
-	if (Ar.IsLoading())
-	{
-		Ref.ShaderPlatform = (EShaderPlatform)ShaderPlatformByte;
-	}
-
-#if KEEP_SHADER_SOURCE_HASHES
-	FSHAHash& VFHash = Ref.VFHash;
-#else
-	FSHAHash VFHash;
-#endif
-	Ar << FShaderResource::FilterShaderSourceHashForSerialization(Ar, VFHash);
-
-
-	if (Ar.IsLoading())
-	{
-		delete Ref.Parameters;
-		if (Ref.VertexFactoryType)
-		{
-			Ref.Parameters = Ref.VertexFactoryType->CreateShaderParameters(Ref.ShaderFrequency);
-		}
-		else
-		{
-			bShaderHasOutdatedParameters = true;
-			Ref.Parameters = NULL;
-		}
-	}
-
-	// Need to be able to skip over parameters for no longer existing vertex factories.
-	int64 SkipOffset = Ar.Tell();
-	{
-		FArchive::FScopeSetDebugSerializationFlags S(Ar, DSF_IgnoreDiff);
-		// Write placeholder.
-		Ar << SkipOffset;
-	}
-
-
-	if(Ref.Parameters)
-	{
-		Ref.Parameters->Serialize(Ar);
-	}
-	else if(Ar.IsLoading())
-	{
-		Ar.Seek( SkipOffset );
-	}
-
-	if( Ar.IsSaving() )
-	{
-		int64 EndOffset = Ar.Tell();
-		Ar.Seek( SkipOffset );
-		Ar << EndOffset;
-		Ar.Seek( EndOffset );
-	}
-
-	return bShaderHasOutdatedParameters;
-}
-
-#if KEEP_SHADER_SOURCE_HASHES
-/** Returns the hash of the vertex factory shader file that this shader was compiled with. */
-const FSHAHash& FVertexFactoryParameterRef::GetHash() const 
-{ 
-	return VFHash;
-}
-#endif
-
-/** Returns the shader platform that this shader was compiled with. */
-EShaderPlatform FVertexFactoryParameterRef::GetShaderPlatform() const
-{
-	return ShaderPlatform;
-}
 void FPrimitiveIdDummyBuffer::InitRHI() 
 {
 	// create a static vertex buffer
diff --git a/Engine/Source/Runtime/RenderCore/Private/VisualizeTexture.cpp b/Engine/Source/Runtime/RenderCore/Private/VisualizeTexture.cpp
--- a/Engine/Source/Runtime/RenderCore/Private/VisualizeTexture.cpp
+++ b/Engine/Source/Runtime/RenderCore/Private/VisualizeTexture.cpp
@@ -296,7 +296,7 @@ void FVisualizeTexture::CreateContentCapturePass(FRDGBuilder& GraphBuilder, cons
 			GPixelFormats[SrcDesc.Format].Name,
 			SrcDesc.Extent.X, SrcDesc.Extent.Y,
 			*ExtendedDrawEvent),
-		*PixelShader,
+		PixelShader,
 		PassParameters,
 		FIntRect(0, 0, RTExtent.X, RTExtent.Y));
 
diff --git a/Engine/Source/Runtime/RenderCore/Public/ClearReplacementShaders.h b/Engine/Source/Runtime/RenderCore/Public/ClearReplacementShaders.h
--- a/Engine/Source/Runtime/RenderCore/Public/ClearReplacementShaders.h
+++ b/Engine/Source/Runtime/RenderCore/Public/ClearReplacementShaders.h
@@ -33,6 +33,7 @@ template <EClearReplacementValueType ValueType, uint32 NumChannels, bool bZeroOu
 struct TClearReplacementBase : public FGlobalShader
 {
 	static_assert(NumChannels >= 1 && NumChannels <= 4, "Only 1 to 4 channels are supported.");
+	DECLARE_INLINE_TYPE_LAYOUT(TClearReplacementBase, NonVirtual);
 
 protected:
 	TClearReplacementBase() {}
@@ -93,34 +94,18 @@ public:
 			break;
 		}
 	}
-
-	virtual bool Serialize(FArchive& Ar) override
-	{
-		bool bShaderHasOutdatedParameters = FGlobalShader::Serialize(Ar);
-
-		if (!bZeroOutput)
-		{
-			Ar << ClearValueParam;
-		}
-		if (bEnableBounds)
-		{
-			Ar << MinBoundsParam;
-			Ar << MaxBoundsParam;
-		}
-
-		return bShaderHasOutdatedParameters;
-	}
-
+	
 	template <typename T = const FShaderParameter&>	inline typename TEnableIf<!bZeroOutput, T>::Type GetClearValueParam() const { return ClearValueParam; }
 	template <typename T = const FShaderParameter&> inline typename TEnableIf<bEnableBounds, T>::Type GetMinBoundsParam() const { return MinBoundsParam;  }
 	template <typename T = const FShaderParameter&> inline typename TEnableIf<bEnableBounds, T>::Type GetMaxBoundsParam() const { return MaxBoundsParam;  }
-
+	
 private:
-	FShaderParameter ClearValueParam;
-	FShaderParameter MinBoundsParam;
-	FShaderParameter MaxBoundsParam;
-};
+	LAYOUT_FIELD(FShaderParameter, ClearValueParam);
+	LAYOUT_FIELD(FShaderParameter, MinBoundsParam);
+	LAYOUT_FIELD(FShaderParameter, MaxBoundsParam);
 
+};
+	
 namespace ClearReplacementCS
 {
 	template <EClearReplacementResourceType> struct TThreadGroupSize {};
@@ -139,14 +124,14 @@ public:
 	static constexpr uint32 ThreadGroupSizeX = ClearReplacementCS::TThreadGroupSize<ResourceType>::X;
 	static constexpr uint32 ThreadGroupSizeY = ClearReplacementCS::TThreadGroupSize<ResourceType>::Y;
 	static constexpr uint32 ThreadGroupSizeZ = ClearReplacementCS::TThreadGroupSize<ResourceType>::Z;
-
+	
 	TClearReplacementCS() {}
 	TClearReplacementCS(const ShaderMetaType::CompiledShaderInitializerType& Initializer)
 		: BaseType(Initializer)
 	{
 		ClearResourceParam.Bind(Initializer.ParameterMap, TEXT("ClearResource"), SPF_Mandatory);
 	}
-
+	
 	static bool ShouldCompilePermutation(const FGlobalShaderPermutationParameters& Parameters)
 	{
 		return IsFeatureLevelSupported(Parameters.Platform, ERHIFeatureLevel::SM5)
@@ -156,35 +141,26 @@ public:
 	static inline void ModifyCompilationEnvironment(const FGlobalShaderPermutationParameters& Parameters, FShaderCompilerEnvironment& OutEnvironment)
 	{
 		BaseType::ModifyCompilationEnvironment(Parameters, OutEnvironment);
-
+	
 		OutEnvironment.SetDefine(TEXT("THREADGROUPSIZE_X"), ThreadGroupSizeX);
 		OutEnvironment.SetDefine(TEXT("THREADGROUPSIZE_Y"), ThreadGroupSizeY);
 		OutEnvironment.SetDefine(TEXT("THREADGROUPSIZE_Z"), ThreadGroupSizeZ);
 		OutEnvironment.SetDefine(TEXT("RESOURCE_TYPE"), uint32(ResourceType));
 	}
-
+	
 	static const TCHAR* GetFunctionName() { return TEXT("ClearCS"); }
 
-	virtual bool Serialize(FArchive& Ar) override
-	{
-		bool bShaderHasOutdatedParameters = BaseType::Serialize(Ar);
-
-		Ar << ClearResourceParam;
-
-		return bShaderHasOutdatedParameters;
-	}
-
 	UE_DEPRECATED(4.25, "TClearReplacementCS::SetResource is deprecated. Call GetClearResourceParam() and bind the UAV manually instead. Be sure to handle any necessary resource transitions.")
 	inline void SetResource(FRHICommandList& RHICmdList, FRHIUnorderedAccessView* UAV)
 	{
 		RHICmdList.TransitionResource(EResourceTransitionAccess::ERWBarrier, EResourceTransitionPipeline::EGfxToCompute, UAV);
-		SetUAVParameter(RHICmdList, BaseType::GetComputeShader(), ClearResourceParam, UAV);
+		SetUAVParameter(RHICmdList, RHICmdList.GetBoundComputeShader(), ClearResourceParam, UAV);
 	}
 
 	UE_DEPRECATED(4.25, "TClearReplacementCS::FinalizeResource is deprecated. Call GetClearResourceParam() and bind the UAV manually instead. Be sure to handle any necessary resource transitions.")
 	inline void FinalizeResource(FRHICommandList& RHICmdList, FRHIUnorderedAccessView* UAV)
 	{
-		SetUAVParameter(RHICmdList, BaseType::GetComputeShader(), ClearResourceParam, nullptr);
+		SetUAVParameter(RHICmdList, RHICmdList.GetBoundComputeShader(), ClearResourceParam, nullptr);
 		RHICmdList.TransitionResource(EResourceTransitionAccess::ERWBarrier, EResourceTransitionPipeline::EComputeToCompute, UAV);
 	}
 
@@ -192,7 +168,7 @@ public:
 	inline uint32 GetResourceParamIndex() const { return ClearResourceParam.GetBaseIndex(); }
 
 private:
-	FShaderResourceParameter ClearResourceParam;
+	LAYOUT_FIELD(FShaderResourceParameter, ClearResourceParam);
 };
 
 template <bool bEnableDepth, typename BaseType>
@@ -225,18 +201,6 @@ public:
 
 	static const TCHAR* GetFunctionName() { return TEXT("ClearVS"); }
 
-	virtual bool Serialize(FArchive& Ar) override
-	{
-		bool bShaderHasOutdatedParameters = BaseType::Serialize(Ar);
-
-		if (bEnableDepth)
-		{
-			Ar << DepthParam;
-		}
-
-		return bShaderHasOutdatedParameters;
-	}
-
 	template <typename T = const FShaderParameter&>
 	inline typename TEnableIf<bEnableDepth, T>::Type GetDepthParam() const
 	{
@@ -244,7 +208,7 @@ public:
 	}
 
 private:
-	FShaderParameter DepthParam;
+	LAYOUT_FIELD(FShaderParameter, DepthParam);
 };
 
 template <bool b128BitOutput, typename BaseType>
@@ -302,7 +266,7 @@ typedef TClearReplacementPS<false, FClearReplacementBase_Float4>             FCl
 typedef TClearReplacementPS<true,  FClearReplacementBase_Float4>             FClearReplacementPS_128;
 // Simple pixel shader which outputs zero to MRT0
 typedef TClearReplacementPS<false, FClearReplacementBase_Float4_Zero>        FClearReplacementPS_Zero;
-
+	
 // Compute shaders for clearing each resource type, with a min/max bounds enabled.
 typedef TClearReplacementCS<EClearReplacementResourceType::Buffer,         FClearReplacementBase_Uint_Bounds>   FClearReplacementCS_Buffer_Uint_Bounds;
 typedef TClearReplacementCS<EClearReplacementResourceType::Buffer,         FClearReplacementBase_Float4_Bounds> FClearReplacementCS_Buffer_Float4_Bounds;
@@ -315,11 +279,11 @@ typedef TClearReplacementCS<EClearReplacementResourceType::Buffer,         FClea
 typedef TClearReplacementCS<EClearReplacementResourceType::Texture2DArray, FClearReplacementBase_Uint_Zero>     FClearReplacementCS_Texture2DArray_Uint_Zero;
 typedef TClearReplacementCS<EClearReplacementResourceType::Buffer,         FClearReplacementBase_Uint>          FClearReplacementCS_Buffer_Uint;
 typedef TClearReplacementCS<EClearReplacementResourceType::Texture2DArray, FClearReplacementBase_Uint>          FClearReplacementCS_Texture2DArray_Uint;
-
+	
 typedef TClearReplacementCS<EClearReplacementResourceType::Texture3D,      FClearReplacementBase_Float4>        FClearReplacementCS_Texture3D_Float4;
 typedef TClearReplacementCS<EClearReplacementResourceType::Texture2D,      FClearReplacementBase_Float4>        FClearReplacementCS_Texture2D_Float4;
 typedef TClearReplacementCS<EClearReplacementResourceType::Texture2DArray, FClearReplacementBase_Float4>        FClearReplacementCS_Texture2DArray_Float4;
-
+	
 // Used by ClearUAV_T in ClearQuad.cpp
 typedef TClearReplacementCS<EClearReplacementResourceType::Texture3D,      FClearReplacementBase_Uint4>         FClearReplacementCS_Texture3D_Uint4;
 typedef TClearReplacementCS<EClearReplacementResourceType::Texture2D,      FClearReplacementBase_Uint4>         FClearReplacementCS_Texture2D_Uint4;
@@ -347,7 +311,7 @@ inline void ClearUAVShader_T(FRHICommandList& RHICmdList, FRHIUnorderedAccessVie
 	typedef TClearReplacementCS<ResourceType, TClearReplacementBase<ValueType, NumChannels, false, true>> FClearShader;
 
 	TShaderMapRef<FClearShader> ComputeShader(GetGlobalShaderMap(GMaxRHIFeatureLevel));
-	FRHIComputeShader* ShaderRHI = ComputeShader->GetComputeShader();
+	FRHIComputeShader* ShaderRHI = ComputeShader.GetComputeShader();
 	RHICmdList.SetComputeShader(ShaderRHI);
 
 	SetShaderValue(RHICmdList, ShaderRHI, ComputeShader->GetClearValueParam(), ClearValues);
diff --git a/Engine/Source/Runtime/RenderCore/Public/CommonRenderResources.h b/Engine/Source/Runtime/RenderCore/Public/CommonRenderResources.h
--- a/Engine/Source/Runtime/RenderCore/Public/CommonRenderResources.h
+++ b/Engine/Source/Runtime/RenderCore/Public/CommonRenderResources.h
@@ -15,6 +15,7 @@
 /** The vertex data used to filter a texture. */
 struct FFilterVertex
 {
+public:
 	FVector4 Position;
 	FVector2D UV;
 };
diff --git a/Engine/Source/Runtime/RenderCore/Public/CopyTextureShaders.h b/Engine/Source/Runtime/RenderCore/Public/CopyTextureShaders.h
--- a/Engine/Source/Runtime/RenderCore/Public/CopyTextureShaders.h
+++ b/Engine/Source/Runtime/RenderCore/Public/CopyTextureShaders.h
@@ -23,10 +23,22 @@ namespace CopyTextureCS
 {
 	template <ECopyTextureResourceType> struct TThreadGroupSize              { static constexpr int32 X = 8, Y = 8, Z = 1; };
 	template <> struct TThreadGroupSize<ECopyTextureResourceType::Texture3D> { static constexpr int32 X = 4, Y = 4, Z = 4; };
+
+	struct DispatchContext
+	{
+		uint32 ThreadGroupSizeX = 0u;
+		uint32 ThreadGroupSizeY = 0u;
+		uint32 ThreadGroupSizeZ = 0u;
+		ECopyTextureResourceType SrcType;
+		ECopyTextureResourceType DstType;
+		ECopyTextureValueType ValueType;
+		uint32 NumChannels;
+	};
 }
 
 class FCopyTextureCS : public FGlobalShader
 {
+	DECLARE_EXPORTED_TYPE_LAYOUT(FCopyTextureCS, RENDERCORE_API, NonVirtual);
 protected:
 	FCopyTextureCS() {}
 	FCopyTextureCS(const ShaderMetaType::CompiledShaderInitializerType& Initializer)
@@ -39,44 +51,50 @@ protected:
 		DstResourceParam.Bind(Initializer.ParameterMap, TEXT("DstResource"), SPF_Mandatory);
 	}
 
-	virtual bool Serialize(FArchive& Ar) override
-	{
-		bool bShaderHasOutdatedParameters = FGlobalShader::Serialize(Ar);
-
-		Ar << DstOffsetParam;
-		Ar << SrcOffsetParam;
-		Ar << DimensionsParam;
-		Ar << SrcResourceParam;
-		Ar << DstResourceParam;
-
-		return bShaderHasOutdatedParameters;
-	}
-
 public:
 	inline const FShaderResourceParameter& GetSrcResourceParam() { return SrcResourceParam; }
 	inline const FShaderResourceParameter& GetDstResourceParam() { return DstResourceParam; }
 
-	virtual void Dispatch(
+	void Dispatch(
 		FRHIComputeCommandList& RHICmdList,
+		const CopyTextureCS::DispatchContext& Context,
 		FIntVector const& SrcOffset,
 		FIntVector const& DstOffset,
 		FIntVector const& Dimensions
-	) = 0;
+	)
+	{
+		check(SrcOffset.GetMin() >= 0 && DstOffset.GetMin() >= 0 && Dimensions.GetMin() >= 0);
+		check(Context.DstType != ECopyTextureResourceType::Texture2D || Dimensions.Z <= 1);
+
+		FRHIComputeShader* ShaderRHI = RHICmdList.GetBoundComputeShader();
+		SetShaderValue(RHICmdList, ShaderRHI, SrcOffsetParam, SrcOffset);
+		SetShaderValue(RHICmdList, ShaderRHI, DstOffsetParam, DstOffset);
+		SetShaderValue(RHICmdList, ShaderRHI, DimensionsParam, Dimensions);
+
+		RHICmdList.DispatchComputeShader(
+			FMath::DivideAndRoundUp(uint32(Dimensions.X), Context.ThreadGroupSizeX),
+			FMath::DivideAndRoundUp(uint32(Dimensions.Y), Context.ThreadGroupSizeY),
+			FMath::DivideAndRoundUp(uint32(Dimensions.Z), Context.ThreadGroupSizeZ)
+		);
+	}
 
-	static inline FCopyTextureCS* SelectShader(TShaderMap<FGlobalShaderType>* GlobalShaderMap, ECopyTextureResourceType SrcType, ECopyTextureResourceType DstType, ECopyTextureValueType ValueType);
+	static inline TShaderRef<FCopyTextureCS> SelectShader(FGlobalShaderMap* GlobalShaderMap, ECopyTextureResourceType SrcType, ECopyTextureResourceType DstType, ECopyTextureValueType ValueType, CopyTextureCS::DispatchContext& OutContext);
 
 protected:
+	template<typename ShaderType>
+	static inline TShaderRef<FCopyTextureCS> SelectShader(FGlobalShaderMap* GlobalShaderMap, CopyTextureCS::DispatchContext& OutContext);
+
 	template <ECopyTextureResourceType SrcType>
-	static inline FCopyTextureCS* SelectShader(TShaderMap<FGlobalShaderType>* GlobalShaderMap, ECopyTextureResourceType DstType, ECopyTextureValueType ValueType);
+	static inline TShaderRef<FCopyTextureCS> SelectShader(FGlobalShaderMap* GlobalShaderMap, ECopyTextureResourceType DstType, ECopyTextureValueType ValueType, CopyTextureCS::DispatchContext& OutContext);
 
 	template <ECopyTextureResourceType SrcType, ECopyTextureResourceType DstType>
-	static inline FCopyTextureCS* SelectShader(TShaderMap<FGlobalShaderType>* GlobalShaderMap, ECopyTextureValueType ValueType);
+	static inline TShaderRef<FCopyTextureCS> SelectShader(FGlobalShaderMap* GlobalShaderMap, ECopyTextureValueType ValueType, CopyTextureCS::DispatchContext& OutContext);
 
-	FShaderParameter DstOffsetParam;
-	FShaderParameter SrcOffsetParam;
-	FShaderParameter DimensionsParam;
-	FShaderResourceParameter SrcResourceParam;
-	FShaderResourceParameter DstResourceParam;
+	LAYOUT_FIELD(FShaderParameter, DstOffsetParam);
+	LAYOUT_FIELD(FShaderParameter, SrcOffsetParam);
+	LAYOUT_FIELD(FShaderParameter, DimensionsParam);
+	LAYOUT_FIELD(FShaderResourceParameter, SrcResourceParam);
+	LAYOUT_FIELD(FShaderResourceParameter, DstResourceParam);
 };
 
 template <ECopyTextureResourceType SrcType, ECopyTextureResourceType DstType, ECopyTextureValueType ValueType, uint32 NumChannels>
@@ -148,59 +166,57 @@ public:
 		}
 	}
 
-	virtual void Dispatch(
-		FRHIComputeCommandList& RHICmdList,
-		FIntVector const& SrcOffset,
-		FIntVector const& DstOffset,
-		FIntVector const& Dimensions) override final
+	void GetDispatchContext(CopyTextureCS::DispatchContext& OutContext)
 	{
-		check(SrcOffset.GetMin() >= 0 && DstOffset.GetMin() >= 0 && Dimensions.GetMin() >= 0);
-		check(DstType != ECopyTextureResourceType::Texture2D || Dimensions.Z <= 1);
-
-		FRHIComputeShader* ShaderRHI = GetComputeShader();
-		SetShaderValue(RHICmdList, ShaderRHI, SrcOffsetParam, SrcOffset);
-		SetShaderValue(RHICmdList, ShaderRHI, DstOffsetParam, DstOffset);
-		SetShaderValue(RHICmdList, ShaderRHI, DimensionsParam, Dimensions);
-
-		RHICmdList.DispatchComputeShader(
-			FMath::DivideAndRoundUp(uint32(Dimensions.X), ThreadGroupSizeX),
-			FMath::DivideAndRoundUp(uint32(Dimensions.Y), ThreadGroupSizeY),
-			FMath::DivideAndRoundUp(uint32(Dimensions.Z), ThreadGroupSizeZ)
-		);
+		OutContext.ThreadGroupSizeX = ThreadGroupSizeX;
+		OutContext.ThreadGroupSizeY = ThreadGroupSizeY;
+		OutContext.ThreadGroupSizeZ = ThreadGroupSizeZ;
+		OutContext.SrcType = SrcType;
+		OutContext.DstType = DstType;
+		OutContext.ValueType = ValueType;
+		OutContext.NumChannels = NumChannels;
 	}
 };
 
+template <typename ShaderType>
+inline TShaderRef<FCopyTextureCS> FCopyTextureCS::SelectShader(FGlobalShaderMap* GlobalShaderMap, CopyTextureCS::DispatchContext& OutContext)
+{
+	TShaderRef<ShaderType> Shader = GlobalShaderMap->GetShader<ShaderType>();
+	Shader->GetDispatchContext(OutContext);
+	return Shader;
+}
+
 template <ECopyTextureResourceType SrcType, ECopyTextureResourceType DstType>
-inline FCopyTextureCS* FCopyTextureCS::SelectShader(TShaderMap<FGlobalShaderType>* GlobalShaderMap, ECopyTextureValueType ValueType)
+inline TShaderRef<FCopyTextureCS> FCopyTextureCS::SelectShader(FGlobalShaderMap* GlobalShaderMap, ECopyTextureValueType ValueType, CopyTextureCS::DispatchContext& OutContext)
 {
 	switch (ValueType)
 	{
 	default: checkNoEntry();
-	case ECopyTextureValueType::Float:  return GlobalShaderMap->GetShader<TCopyResourceCS<SrcType, DstType, ECopyTextureValueType::Float,  4>>();
-	case ECopyTextureValueType::Int32:  return GlobalShaderMap->GetShader<TCopyResourceCS<SrcType, DstType, ECopyTextureValueType::Int32,  4>>();
-	case ECopyTextureValueType::Uint32: return GlobalShaderMap->GetShader<TCopyResourceCS<SrcType, DstType, ECopyTextureValueType::Uint32, 4>>();
+	case ECopyTextureValueType::Float:  return SelectShader<TCopyResourceCS<SrcType, DstType, ECopyTextureValueType::Float,  4>>(GlobalShaderMap, OutContext);
+	case ECopyTextureValueType::Int32:  return SelectShader<TCopyResourceCS<SrcType, DstType, ECopyTextureValueType::Int32,  4>>(GlobalShaderMap, OutContext);
+	case ECopyTextureValueType::Uint32: return SelectShader<TCopyResourceCS<SrcType, DstType, ECopyTextureValueType::Uint32, 4>>(GlobalShaderMap, OutContext);
 	}
 }
 
 template <ECopyTextureResourceType SrcType>
-inline FCopyTextureCS* FCopyTextureCS::SelectShader(TShaderMap<FGlobalShaderType>* GlobalShaderMap, ECopyTextureResourceType DstType, ECopyTextureValueType ValueType)
+inline TShaderRef<FCopyTextureCS> FCopyTextureCS::SelectShader(FGlobalShaderMap* GlobalShaderMap, ECopyTextureResourceType DstType, ECopyTextureValueType ValueType, CopyTextureCS::DispatchContext& OutContext)
 {
 	switch (DstType)
 	{
 	default: checkNoEntry();
-	case ECopyTextureResourceType::Texture2D:	   return FCopyTextureCS::SelectShader<SrcType, ECopyTextureResourceType::Texture2D     >(GlobalShaderMap, ValueType);
-	case ECopyTextureResourceType::Texture2DArray: return FCopyTextureCS::SelectShader<SrcType, ECopyTextureResourceType::Texture2DArray>(GlobalShaderMap, ValueType);
-	case ECopyTextureResourceType::Texture3D:      return FCopyTextureCS::SelectShader<SrcType, ECopyTextureResourceType::Texture3D     >(GlobalShaderMap, ValueType);
+	case ECopyTextureResourceType::Texture2D:	   return FCopyTextureCS::SelectShader<SrcType, ECopyTextureResourceType::Texture2D     >(GlobalShaderMap, ValueType, OutContext);
+	case ECopyTextureResourceType::Texture2DArray: return FCopyTextureCS::SelectShader<SrcType, ECopyTextureResourceType::Texture2DArray>(GlobalShaderMap, ValueType, OutContext);
+	case ECopyTextureResourceType::Texture3D:      return FCopyTextureCS::SelectShader<SrcType, ECopyTextureResourceType::Texture3D     >(GlobalShaderMap, ValueType, OutContext);
 	}
 }
 
-inline FCopyTextureCS* FCopyTextureCS::SelectShader(TShaderMap<FGlobalShaderType>* GlobalShaderMap, ECopyTextureResourceType SrcType, ECopyTextureResourceType DstType, ECopyTextureValueType ValueType)
+inline TShaderRef<FCopyTextureCS> FCopyTextureCS::SelectShader(FGlobalShaderMap* GlobalShaderMap, ECopyTextureResourceType SrcType, ECopyTextureResourceType DstType, ECopyTextureValueType ValueType, CopyTextureCS::DispatchContext& OutContext)
 {
 	switch (SrcType)
 	{
 	default: checkNoEntry();
-	case ECopyTextureResourceType::Texture2D:	   return FCopyTextureCS::SelectShader<ECopyTextureResourceType::Texture2D     >(GlobalShaderMap, DstType, ValueType);
-	case ECopyTextureResourceType::Texture2DArray: return FCopyTextureCS::SelectShader<ECopyTextureResourceType::Texture2DArray>(GlobalShaderMap, DstType, ValueType);
-	case ECopyTextureResourceType::Texture3D:      return FCopyTextureCS::SelectShader<ECopyTextureResourceType::Texture3D     >(GlobalShaderMap, DstType, ValueType);
+	case ECopyTextureResourceType::Texture2D:	   return FCopyTextureCS::SelectShader<ECopyTextureResourceType::Texture2D     >(GlobalShaderMap, DstType, ValueType, OutContext);
+	case ECopyTextureResourceType::Texture2DArray: return FCopyTextureCS::SelectShader<ECopyTextureResourceType::Texture2DArray>(GlobalShaderMap, DstType, ValueType, OutContext);
+	case ECopyTextureResourceType::Texture3D:      return FCopyTextureCS::SelectShader<ECopyTextureResourceType::Texture3D     >(GlobalShaderMap, DstType, ValueType, OutContext);
 	}
 }
diff --git a/Engine/Source/Runtime/RenderCore/Public/GlobalShader.h b/Engine/Source/Runtime/RenderCore/Public/GlobalShader.h
--- a/Engine/Source/Runtime/RenderCore/Public/GlobalShader.h
+++ b/Engine/Source/Runtime/RenderCore/Public/GlobalShader.h
@@ -14,6 +14,7 @@
 
 class FShaderCommonCompileJob;
 class FShaderCompileJob;
+class TargetPlatform;
 
 /** Used to identify the global shader map in compile queues. */
 extern RENDERCORE_API const int32 GlobalShaderMapId;
@@ -27,7 +28,7 @@ public:
 	RENDERCORE_API FGlobalShaderMapId(EShaderPlatform Platform);
 
 	/** Append to a string that will be used as a DDC key. */
-	RENDERCORE_API void AppendKeyString(FString& KeyString, const TArray<FShaderTypeDependency>& Dependencies) const;
+	RENDERCORE_API void AppendKeyString(FString& KeyString, const TArray<FShaderTypeDependency>& Dependencies, const ITargetPlatform* TargetPlatform) const;
 
 	RENDERCORE_API const TMap<FString, TArray<FShaderTypeDependency>>& GetShaderFilenameToDependeciesMap() const { return ShaderFilenameToDependenciesMap; }
 
@@ -39,19 +40,7 @@ private:
 	TArray<FShaderPipelineTypeDependency> ShaderPipelineTypeDependencies;
 };
 
-struct FGlobalShaderPermutationParameters
-{
-	// Shader platform to compile to.
-	const EShaderPlatform Platform;
-
-	/** Unique permutation identifier of the global shader type. */
-	const int32 PermutationId;
-
-	FGlobalShaderPermutationParameters(EShaderPlatform InPlatform, int32 InPermutationId)
-		: Platform(InPlatform)
-		, PermutationId(InPermutationId)
-	{ }
-};
+using FGlobalShaderPermutationParameters = FShaderPermutationParameters;
 
 /**
  * A shader meta type for the simplest shaders; shaders which are not material or vertex factory linked.
@@ -63,12 +52,9 @@ class FGlobalShaderType : public FShaderType
 public:
 
 	typedef FShader::CompiledShaderInitializerType CompiledShaderInitializerType;
-	typedef FShader* (*ConstructCompiledType)(const CompiledShaderInitializerType&);
-	typedef bool (*ShouldCompilePermutationType)(const FGlobalShaderPermutationParameters&);
-	typedef bool(*ValidateCompiledResultType)(EShaderPlatform, const FShaderParameterMap&, TArray<FString>&);
-	typedef void (*ModifyCompilationEnvironmentType)(const FGlobalShaderPermutationParameters&, FShaderCompilerEnvironment&);
 
 	FGlobalShaderType(
+		FTypeLayoutDesc& InTypeLayout,
 		const TCHAR* InName,
 		const TCHAR* InSourceFilename,
 		const TCHAR* InFunctionName,
@@ -79,13 +65,17 @@ public:
 		ModifyCompilationEnvironmentType InModifyCompilationEnvironmentRef,
 		ShouldCompilePermutationType InShouldCompilePermutationRef,
 		ValidateCompiledResultType InValidateCompiledResultRef,
+		uint32 InTypeSize,
 		const FShaderParametersMetadata* InRootParametersMetadata = nullptr
 		):
-		FShaderType(EShaderTypeForDynamicCast::Global, InName, InSourceFilename, InFunctionName, InFrequency, InTotalPermutationCount, InConstructSerializedRef, InRootParametersMetadata),
-		ConstructCompiledRef(InConstructCompiledRef),
-		ShouldCompilePermutationRef(InShouldCompilePermutationRef),
-		ValidateCompiledResultRef(InValidateCompiledResultRef),
-		ModifyCompilationEnvironmentRef(InModifyCompilationEnvironmentRef)
+		FShaderType(EShaderTypeForDynamicCast::Global, InTypeLayout, InName, InSourceFilename, InFunctionName, InFrequency, InTotalPermutationCount,
+			InConstructSerializedRef,
+			InConstructCompiledRef,
+			InModifyCompilationEnvironmentRef,
+			InShouldCompilePermutationRef,
+			InValidateCompiledResultRef,
+			InTypeSize,
+			InRootParametersMetadata)
 	{
 		checkf(FPaths::GetExtension(InSourceFilename) == TEXT("usf"),
 			TEXT("Incorrect virtual shader path extension for global shader '%s': Only .usf files should be "
@@ -100,7 +90,7 @@ public:
 	 */
 	bool ShouldCompilePermutation(EShaderPlatform Platform, int32 PermutationId) const
 	{
-		return (*ShouldCompilePermutationRef)(FGlobalShaderPermutationParameters(Platform, PermutationId));
+		return FShaderType::ShouldCompilePermutation(FGlobalShaderPermutationParameters(Platform, PermutationId));
 	}
 
 	/**
@@ -111,28 +101,135 @@ public:
 	void SetupCompileEnvironment(EShaderPlatform Platform, int32 PermutationId, FShaderCompilerEnvironment& Environment)
 	{
 		// Allow the shader type to modify its compile environment.
-		(*ModifyCompilationEnvironmentRef)(FGlobalShaderPermutationParameters(Platform, PermutationId), Environment);
+		ModifyCompilationEnvironment(FGlobalShaderPermutationParameters(Platform, PermutationId), Environment);
 	}
+};
 
-	/**
-	* Checks if the shader type should pass compilation for a particular set of parameters.
-	* @param Platform - Platform to validate.
-	* @param ParameterMap - Shader parameters to validate.
-	* @param OutError - List for appending validation errors.
-	*/
-	bool ValidateCompiledResult(EShaderPlatform Platform, const FShaderParameterMap& ParameterMap, TArray<FString>& OutError) const
+class RENDERCORE_API FGlobalShaderMapContent : public FShaderMapContent
+{
+	using Super = FShaderMapContent;
+	friend class FGlobalShaderMap;
+	friend class FGlobalShaderMapSection;
+	DECLARE_TYPE_LAYOUT(FGlobalShaderMapContent, NonVirtual);
+public:
+	const FHashedName& GetHashedSourceFilename() const { return HashedSourceFilename; }
+
+private:
+	inline FGlobalShaderMapContent(EShaderPlatform InPlatform, const FHashedName& InHashedSourceFilename)
+		: Super(InPlatform)
+		, HashedSourceFilename(InHashedSourceFilename)
+	{}
+
+	LAYOUT_FIELD(FHashedName, HashedSourceFilename);
+};
+
+class RENDERCORE_API FGlobalShaderMapSection : public TShaderMap<FGlobalShaderMapContent, FShaderMapPointerTable>
+{
+	using Super = TShaderMap<FGlobalShaderMapContent, FShaderMapPointerTable>;
+	friend class FGlobalShaderMap;
+public:
+	static FGlobalShaderMapSection* CreateFromArchive(FArchive& Ar);
+
+	~FGlobalShaderMapSection() { ReleaseResourceBuilder(); }
+
+	void Serialize(FArchive& Ar);
+
+	FShaderMapResourceBuilder& GetResourceBuilder()
+	{
+		if (ResourceBuilder == nullptr)
+		{
+			ResourceBuilder = new FShaderMapResourceBuilder(GetResourceCode());
+		}
+		return *ResourceBuilder;
+	}
+
+	void ReleaseResourceBuilder()
 	{
-		return (*ValidateCompiledResultRef)(Platform, ParameterMap, OutError);
+		if (ResourceBuilder)
+		{
+			delete ResourceBuilder;
+			ResourceBuilder = nullptr;
+		}
 	}
 
 private:
-	ConstructCompiledType ConstructCompiledRef;
-	ShouldCompilePermutationType ShouldCompilePermutationRef;
-	ValidateCompiledResultType ValidateCompiledResultRef;
-	ModifyCompilationEnvironmentType ModifyCompilationEnvironmentRef;
+	inline FGlobalShaderMapSection() {}
+
+	inline FGlobalShaderMapSection(EShaderPlatform InPlatform, const FHashedName& InHashedSourceFilename)
+	{
+		AssignContent(new FGlobalShaderMapContent(InPlatform, InHashedSourceFilename));
+	}
+
+	TShaderRef<FShader> GetShader(FShaderType* ShaderType, int32 PermutationId = 0) const;
+	FShaderPipelineRef GetShaderPipeline(const FShaderPipelineType* PipelineType) const;
+
+	FShaderMapResourceBuilder* ResourceBuilder = nullptr;
 };
 
-extern RENDERCORE_API TShaderMap<FGlobalShaderType>* GGlobalShaderMap[SP_NumPlatforms];
+class RENDERCORE_API FGlobalShaderMap
+{
+public:
+	explicit FGlobalShaderMap(EShaderPlatform InPlatform);
+	~FGlobalShaderMap();
+
+	TShaderRef<FShader> GetShader(FShaderType* ShaderType, int32 PermutationId = 0) const;
+	FShaderPipelineRef GetShaderPipeline(const FShaderPipelineType* PipelineType) const;
+
+	template<typename ShaderType>
+	TShaderRef<ShaderType> GetShader(int32 PermutationId = 0) const
+	{
+		TShaderRef<FShader> Shader = GetShader(&ShaderType::StaticType, PermutationId);
+		checkf(Shader.IsValid(), TEXT("Failed to find shader type %s in Platform %s"), ShaderType::StaticType.GetName(), *LegacyShaderPlatformToShaderFormat(Platform).ToString());
+		return TShaderRef<ShaderType>::Cast(Shader);
+	}
+
+	/** Finds the shader with the given type.  Asserts on failure. */
+	template<typename ShaderType>
+	TShaderRef<ShaderType> GetShader(const typename ShaderType::FPermutationDomain& PermutationVector) const
+	{
+		return GetShader<ShaderType>(PermutationVector.ToDimensionValueId());
+	}
+
+	bool HasShader(FShaderType* Type, int32 PermutationId) const
+	{
+		return GetShader(Type, PermutationId).IsValid();
+	}
+
+	bool HasShaderPipeline(const FShaderPipelineType* ShaderPipelineType) const
+	{
+		return GetShaderPipeline(ShaderPipelineType).IsValid();
+	}
+
+	bool IsEmpty() const;
+
+	void Empty();
+
+	FShader* FindOrAddShader(const FShaderType* ShaderType, FShader* Shader);
+	FShaderPipeline* FindOrAddShaderPipeline(const FShaderPipelineType* ShaderPipelineType, FShaderPipeline* ShaderPipeline);
+
+	void RemoveShaderTypePermutaion(const FShaderType* Type, int32 PermutationId);
+	void RemoveShaderPipelineType(const FShaderPipelineType* ShaderPipelineType);
+
+	void AddSection(FGlobalShaderMapSection* InSection);
+	FGlobalShaderMapSection* FindSection(const FHashedName& HashedShaderFilename);
+	FGlobalShaderMapSection* FindOrAddSection(const FShaderType* ShaderType);
+	
+	void LoadFromGlobalArchive(FArchive& Ar);
+	void SaveToGlobalArchive(FArchive& Ar);
+
+	void BeginCreateAllShaders();
+
+#if WITH_EDITOR
+	void GetOutdatedTypes(TArray<const FShaderType*>& OutdatedShaderTypes, TArray<const FShaderPipelineType*>& OutdatedShaderPipelineTypes, TArray<const FVertexFactoryType*>& OutdatedFactoryTypes) const;
+	void SaveShaderStableKeys(EShaderPlatform TargetShaderPlatform);
+#endif // WITH_EDITOR
+
+private:
+	TMap<FHashedName, FGlobalShaderMapSection*> SectionMap;
+	EShaderPlatform Platform;
+};
+
+extern RENDERCORE_API FGlobalShaderMap* GGlobalShaderMap[SP_NumPlatforms];
 
 
 /**
@@ -142,8 +239,9 @@ extern RENDERCORE_API TShaderMap<FGlobalShaderType>* GGlobalShaderMap[SP_NumPlat
  */
 class FGlobalShader : public FShader
 {
-	DECLARE_SHADER_TYPE(FGlobalShader,Global);
+	DECLARE_EXPORTED_TYPE_LAYOUT(FGlobalShader, RENDERCORE_API, NonVirtual);
 public:
+	using ShaderMetaType = FGlobalShaderType;
 
 	FGlobalShader() : FShader() {}
 
@@ -156,11 +254,8 @@ public:
 		SetUniformBufferParameter(RHICmdList, ShaderRHI, ViewUniformBufferParameter, ViewUniformBuffer);
 	}
 
-	typedef void (*ModifyCompilationEnvironmentType)(EShaderPlatform, FShaderCompilerEnvironment&);
-
-	// FShader interface.
-	RENDERCORE_API static void ModifyCompilationEnvironment(const FGlobalShaderPermutationParameters& Parameters, FShaderCompilerEnvironment& OutEnvironment) {}
-	RENDERCORE_API static bool ValidateCompiledResult(EShaderPlatform Platform, const FShaderParameterMap& ParameterMap, TArray<FString>& OutErrors) { return true; }
+	
+	
 };
 
 /**
@@ -199,13 +294,13 @@ extern RENDERCORE_API void BackupGlobalShaderMap(FGlobalShaderBackupData& OutGlo
 extern RENDERCORE_API void RestoreGlobalShaderMap(const FGlobalShaderBackupData& GlobalShaderData);
 
 /**
- * Accesses the global shader map.  This is a global TShaderMap<FGlobalShaderType> which contains an instance of each global shader type.
+ * Accesses the global shader map.  This is a global FGlobalShaderMap which contains an instance of each global shader type.
  *
  * @param Platform Which platform's global shader map to use
  * @param bRefreshShaderMap If true, the existing global shader map will be tossed first
  * @return A reference to the global shader map.
  */
-extern RENDERCORE_API TShaderMap<FGlobalShaderType>* GetGlobalShaderMap(EShaderPlatform Platform);
+extern RENDERCORE_API FGlobalShaderMap* GetGlobalShaderMap(EShaderPlatform Platform);
 
 /**
   * Overload for the above GetGlobalShaderMap which takes a feature level and translates to the appropriate shader platform
@@ -215,7 +310,7 @@ extern RENDERCORE_API TShaderMap<FGlobalShaderType>* GetGlobalShaderMap(EShaderP
   * @return A reference to the global shader map.
   *
   **/
-inline TShaderMap<FGlobalShaderType>* GetGlobalShaderMap(ERHIFeatureLevel::Type FeatureLevel) 
+inline FGlobalShaderMap* GetGlobalShaderMap(ERHIFeatureLevel::Type FeatureLevel)
 { 
 	return GetGlobalShaderMap(GShaderPlatformForFeatureLevel[FeatureLevel]); 
 }
@@ -249,21 +344,21 @@ inline TShaderMap<FGlobalShaderType>* GetGlobalShaderMap(ERHIFeatureLevel::Type
  *		// ...
  * };
  */
-#define DECLARE_GLOBAL_SHADER(ShaderClass) \
+/*#define DECLARE_GLOBAL_SHADER(ShaderClass) \
 	public: \
 	\
 	using ShaderMetaType = FGlobalShaderType; \
 	\
+	using ShaderMapType = FGlobalShaderMap; \
+	\
 	static ShaderMetaType StaticType; \
 	\
 	static FShader* ConstructSerializedInstance() { return new ShaderClass(); } \
 	static FShader* ConstructCompiledInstance(const ShaderMetaType::CompiledShaderInitializerType& Initializer) \
 	{ return new ShaderClass(Initializer); } \
 	\
-	virtual uint32 GetTypeSize() const override { return sizeof(*this); } \
-	\
 	static void ModifyCompilationEnvironmentImpl( \
-		const FGlobalShaderPermutationParameters& Parameters, \
+		const FShaderPermutationParameters& Parameters, \
 		FShaderCompilerEnvironment& OutEnvironment) \
 	{ \
 		FPermutationDomain PermutationVector(Parameters.PermutationId); \
@@ -280,8 +375,19 @@ inline TShaderMap<FGlobalShaderType>* GetGlobalShaderMap(ERHIFeatureLevel::Type
 		ShaderClass::FPermutationDomain::PermutationCount, \
 		ShaderClass::ConstructSerializedInstance, \
 		ShaderClass::ConstructCompiledInstance, \
+		ShaderClass::ConstructEditorContent, \
+		ShaderClass::DestroyInstance, \
+		ShaderClass::DestroyEditorContent, \
 		ShaderClass::ModifyCompilationEnvironmentImpl, \
 		ShaderClass::ShouldCompilePermutation, \
 		ShaderClass::ValidateCompiledResult, \
+		ShaderClass::FreezeMemoryImageImpl, \
+		ShaderClass::FreezeMemoryImageEditorImpl, \
+		sizeof(ShaderClass), \
+		sizeof(ShaderClass::FEditorContent), \
 		ShaderClass::GetRootParametersMetadata() \
 		)
+		*/
+
+#define DECLARE_GLOBAL_SHADER(ShaderClass) DECLARE_SHADER_TYPE(ShaderClass, Global)
+#define IMPLEMENT_GLOBAL_SHADER(ShaderClass,SourceFilename,FunctionName,Frequency) IMPLEMENT_SHADER_TYPE(,ShaderClass,TEXT(SourceFilename),TEXT(FunctionName),Frequency)
\ No newline at end of file
diff --git a/Engine/Source/Runtime/RenderCore/Public/MediaShaders.h b/Engine/Source/Runtime/RenderCore/Public/MediaShaders.h
--- a/Engine/Source/Runtime/RenderCore/Public/MediaShaders.h
+++ b/Engine/Source/Runtime/RenderCore/Public/MediaShaders.h
@@ -134,12 +134,6 @@ public:
 	FMediaShadersVS(const ShaderMetaType::CompiledShaderInitializerType& Initializer)
 		: FGlobalShader(Initializer)
 	{ }
-
-	virtual bool Serialize(FArchive& Ar) override
-	{
-		bool bShaderHasOutdatedParameters = FGlobalShader::Serialize(Ar);
-		return bShaderHasOutdatedParameters;
-	}
 };
 
 
@@ -168,12 +162,6 @@ public:
 		: FGlobalShader(Initializer)
 	{ }
 
-	virtual bool Serialize(FArchive& Ar) override
-	{
-		bool bShaderHasOutdatedParameters = FGlobalShader::Serialize(Ar);
-		return bShaderHasOutdatedParameters;
-	}
-
 	RENDERCORE_API void SetParameters(FRHICommandList& RHICmdList, TRefCountPtr<FRHITexture2D> AYUVTexture, const FMatrix& ColorTransform, const FVector& YUVOffset, bool SrgbToLinear);
 };
 
@@ -201,12 +189,6 @@ public:
 		: FGlobalShader(Initializer)
 	{ }
 
-	virtual bool Serialize(FArchive& Ar) override
-	{
-		bool bShaderHasOutdatedParameters = FGlobalShader::Serialize(Ar);
-		return bShaderHasOutdatedParameters;
-	}
-
 	RENDERCORE_API void SetParameters(FRHICommandList& RHICmdList, TRefCountPtr<FRHITexture2D> BMPTexture, const FIntPoint& OutputDimensions, bool SrgbToLinear);
 };
 
@@ -236,12 +218,6 @@ public:
 		: FGlobalShader(Initializer)
 	{ }
 
-	virtual bool Serialize(FArchive& Ar) override
-	{
-		bool bShaderHasOutdatedParameters = FGlobalShader::Serialize(Ar);
-		return bShaderHasOutdatedParameters;
-	}
-
 	RENDERCORE_API void SetParameters(FRHICommandList& RHICmdList, TRefCountPtr<FRHITexture2D> NV12Texture, const FIntPoint& OutputDimensions, const FMatrix& ColorTransform, const FVector& YUVOffset, bool SrgbToLinear);
 };
 
@@ -271,12 +247,6 @@ public:
 		: FGlobalShader(Initializer)
 	{ }
 
-	virtual bool Serialize(FArchive& Ar) override
-	{
-		bool bShaderHasOutdatedParameters = FGlobalShader::Serialize(Ar);
-		return bShaderHasOutdatedParameters;
-	}
-
 	RENDERCORE_API void SetParameters(FRHICommandList& RHICmdList, TRefCountPtr<FRHITexture2D> NV21Texture, const FIntPoint& OutputDimensions, const FMatrix& ColorTransform, const FVector& YUVOffset, bool SrgbToLinear);
 };
 
@@ -305,11 +275,6 @@ public:
 		: FGlobalShader(Initializer)
 	{ }
 
-	virtual bool Serialize(FArchive& Ar) override
-	{
-		bool bShaderHasOutdatedParameters = FGlobalShader::Serialize(Ar);
-		return bShaderHasOutdatedParameters;
-	}
 
 	RENDERCORE_API void SetParameters(FRHICommandList& RHICmdList, TRefCountPtr<FRHITexture2D> RGBTexture, const FIntPoint& OutputDimensions, bool SrgbToLinear);
 };
@@ -342,12 +307,6 @@ public:
 		: FGlobalShader(Initializer)
 	{ }
 
-	virtual bool Serialize(FArchive& Ar) override
-	{
-		bool bShaderHasOutdatedParameters = FGlobalShader::Serialize(Ar);		
-		return bShaderHasOutdatedParameters;
-	}
-
 	RENDERCORE_API void SetParameters(FRHICommandList& RHICmdList, TRefCountPtr<FRHITexture2D> LumaTexture, TRefCountPtr<FRHITexture2D> CbCrTexture, const FMatrix& ColorTransform, const FVector& YUVOffset, bool SrgbToLinear);
 };
 
@@ -396,12 +355,6 @@ public:
 		: FGlobalShader(Initializer)
 	{ }
 
-	virtual bool Serialize(FArchive& Ar) override
-	{
-		bool bShaderHasOutdatedParameters = FGlobalShader::Serialize(Ar);
-		return bShaderHasOutdatedParameters;
-	}
-
 	RENDERCORE_API void SetParameters(FRHICommandList& RHICmdList, TRefCountPtr<FRHITexture2D> UYVYTexture, const FMatrix& ColorTransform, const FVector& YUVOffset, bool SrgbToLinear);
 };
 
@@ -430,11 +383,6 @@ public:
 		: FGlobalShader(Initializer)
 	{ }
 
-	virtual bool Serialize(FArchive& Ar) override
-	{
-		bool bShaderHasOutdatedParameters = FGlobalShader::Serialize(Ar);
-		return bShaderHasOutdatedParameters;
-	}
 
 	RENDERCORE_API void SetParameters(FRHICommandList& RHICmdList, TRefCountPtr<FRHITexture2D> YTexture, TRefCountPtr<FRHITexture2D> UTexture, TRefCountPtr<FRHITexture2D> VTexture, const FIntPoint& OutputDimensions, const FMatrix& ColorTransform, const FVector& YUVOffset, bool SrgbToLinear);
 };
@@ -463,12 +411,6 @@ public:
 		: FGlobalShader(Initializer)
 	{ }
 
-	virtual bool Serialize(FArchive& Ar) override
-	{
-		bool bShaderHasOutdatedParameters = FGlobalShader::Serialize(Ar);
-		return bShaderHasOutdatedParameters;
-	}
-
 	RENDERCORE_API void SetParameters(FRHICommandList& RHICmdList, TRefCountPtr<FRHITexture2D> YUVTexture, const FIntPoint& OutputDimensions, const FMatrix& ColorTransform, const FVector& YUVOffset, bool SrgbToLinear);
 };
 
@@ -500,12 +442,6 @@ public:
 		: FGlobalShader(Initializer)
 	{ }
 
-	virtual bool Serialize(FArchive& Ar) override
-	{
-		bool bShaderHasOutdatedParameters = FGlobalShader::Serialize(Ar);
-		return bShaderHasOutdatedParameters;
-	}
-
 	RENDERCORE_API void SetParameters(FRHICommandList& RHICmdList, TRefCountPtr<FRHITexture2D> YUY2Texture, const FIntPoint& OutputDimensions, const FMatrix& ColorTransform, const FVector& YUVOffset, bool SrgbToLinear);
 };
 
@@ -536,12 +472,6 @@ public:
 		: FGlobalShader(Initializer)
 	{ }
 
-	virtual bool Serialize(FArchive& Ar) override
-	{
-		bool bShaderHasOutdatedParameters = FGlobalShader::Serialize(Ar);
-		return bShaderHasOutdatedParameters;
-	}
-
 	RENDERCORE_API void SetParameters(FRHICommandList& RHICmdList, TRefCountPtr<FRHITexture2D> YVYUTexture, const FMatrix& ColorTransform, const FVector& YUVOffset, bool SrgbToLinear);
 };
 
@@ -569,12 +499,6 @@ public:
 		: FGlobalShader(Initializer)
 	{ }
 
-	virtual bool Serialize(FArchive& Ar) override
-	{
-		bool bShaderHasOutdatedParameters = FGlobalShader::Serialize(Ar);
-		return bShaderHasOutdatedParameters;
-	}
-
 	RENDERCORE_API void SetParameters(FRHICommandList& RHICmdList, TRefCountPtr<FRHITexture2D> RGBATexture, const FMatrix& ColorTransform, const FVector& YUVOffset, bool LinearToSrgb);
 };
 
@@ -602,12 +526,6 @@ public:
 		: FGlobalShader(Initializer)
 	{ }
 
-	virtual bool Serialize(FArchive& Ar) override
-	{
-		bool bShaderHasOutdatedParameters = FGlobalShader::Serialize(Ar);
-		return bShaderHasOutdatedParameters;
-	}
-
 	RENDERCORE_API void SetParameters(FRHICommandList& RHICmdList, TRefCountPtr<FRHITexture2D> RGBATexture, const FMatrix& ColorTransform, const FVector& YUVOffset, bool LinearToSrgb);
 };
 
diff --git a/Engine/Source/Runtime/RenderCore/Public/OculusShaders.h b/Engine/Source/Runtime/RenderCore/Public/OculusShaders.h
--- a/Engine/Source/Runtime/RenderCore/Public/OculusShaders.h
+++ b/Engine/Source/Runtime/RenderCore/Public/OculusShaders.h
@@ -69,25 +69,17 @@ public:
 
 	RENDERCORE_API void SetParameters(FRHICommandList& RHICmdList, const FTexture* Texture)
 	{
-		SetTextureParameter(RHICmdList, GetPixelShader(), InTexture, InTextureSampler, Texture);
+		SetTextureParameter(RHICmdList, RHICmdList.GetBoundPixelShader(), InTexture, InTextureSampler, Texture);
 	}
 
 	RENDERCORE_API void SetParameters(FRHICommandList& RHICmdList, FRHISamplerState* SamplerStateRHI, FRHITexture* TextureRHI)
 	{
-		SetTextureParameter(RHICmdList, GetPixelShader(), InTexture, InTextureSampler, SamplerStateRHI, TextureRHI);
-	}
-
-	virtual bool Serialize(FArchive& Ar) override
-	{
-		bool bShaderHasOutdatedParameters = FGlobalShader::Serialize(Ar);
-		Ar << InTexture;
-		Ar << InTextureSampler;
-		return bShaderHasOutdatedParameters;
+		SetTextureParameter(RHICmdList, RHICmdList.GetBoundPixelShader(), InTexture, InTextureSampler, SamplerStateRHI, TextureRHI);
 	}
 
 private:
-	FShaderResourceParameter InTexture;
-	FShaderResourceParameter InTextureSampler;
+	LAYOUT_FIELD(FShaderResourceParameter, InTexture);
+	LAYOUT_FIELD(FShaderResourceParameter, InTextureSampler);
 };
 
 
@@ -112,27 +104,18 @@ public:
 
 	void SetParameters(FRHICommandList& RHICmdList, const FTexture* Texture, int FaceIndex)
 	{
-		SetTextureParameter(RHICmdList, GetPixelShader(), InTexture, InTextureSampler, Texture);
-		SetShaderValue(RHICmdList, GetPixelShader(), InFaceIndexParameter, FaceIndex);
+		SetTextureParameter(RHICmdList, RHICmdList.GetBoundPixelShader(), InTexture, InTextureSampler, Texture);
+		SetShaderValue(RHICmdList, RHICmdList.GetBoundPixelShader(), InFaceIndexParameter, FaceIndex);
 	}
 
 	void SetParameters(FRHICommandList& RHICmdList, FRHISamplerState* SamplerStateRHI, FRHITexture* TextureRHI, int FaceIndex)
 	{
-		SetTextureParameter(RHICmdList, GetPixelShader(), InTexture, InTextureSampler, SamplerStateRHI, TextureRHI);
-		SetShaderValue(RHICmdList, GetPixelShader(), InFaceIndexParameter, FaceIndex);
-	}
-
-	virtual bool Serialize(FArchive& Ar) override
-	{
-		bool bShaderHasOutdatedParameters = FGlobalShader::Serialize(Ar);
-		Ar << InTexture;
-		Ar << InTextureSampler;
-		Ar << InFaceIndexParameter;
-		return bShaderHasOutdatedParameters;
+		SetTextureParameter(RHICmdList, RHICmdList.GetBoundPixelShader(), InTexture, InTextureSampler, SamplerStateRHI, TextureRHI);
+		SetShaderValue(RHICmdList, RHICmdList.GetBoundPixelShader(), InFaceIndexParameter, FaceIndex);
 	}
 
 private:
-	FShaderResourceParameter InTexture;
-	FShaderResourceParameter InTextureSampler;
-	FShaderParameter		 InFaceIndexParameter;
+	LAYOUT_FIELD(FShaderResourceParameter, InTexture);
+	LAYOUT_FIELD(FShaderResourceParameter, InTextureSampler);
+	LAYOUT_FIELD(FShaderParameter, InFaceIndexParameter);
 };
\ No newline at end of file
diff --git a/Engine/Source/Runtime/RenderCore/Public/OneColorShader.h b/Engine/Source/Runtime/RenderCore/Public/OneColorShader.h
--- a/Engine/Source/Runtime/RenderCore/Public/OneColorShader.h
+++ b/Engine/Source/Runtime/RenderCore/Public/OneColorShader.h
@@ -37,15 +37,7 @@ public:
 
 	void SetDepthParameter(FRHICommandList& RHICmdList, float Depth)
 	{
-		SetShaderValue(RHICmdList, GetVertexShader(), DepthParameter, Depth);
-	}
-
-	// FShader interface.
-	virtual bool Serialize(FArchive& Ar) override
-	{
-		bool bShaderHasOutdatedParameters = FGlobalShader::Serialize(Ar);
-		Ar << DepthParameter;
-		return bShaderHasOutdatedParameters;
+		SetShaderValue(RHICmdList, RHICmdList.GetBoundVertexShader(), DepthParameter, Depth);
 	}
 
 	static bool ShouldCompilePermutation(const FGlobalShaderPermutationParameters& Parameters)
@@ -64,7 +56,7 @@ public:
 	}
 
 private:
-	FShaderParameter DepthParameter;
+	LAYOUT_FIELD(FShaderParameter, DepthParameter);
 };
 
 /**
@@ -82,22 +74,15 @@ public:
 		//ColorParameter.Bind( Initializer.ParameterMap, TEXT("DrawColorMRT"), SPF_Mandatory);
 	}
 
-	virtual void SetColors(FRHICommandList& RHICmdList, const FLinearColor* Colors, int32 NumColors);	
+	RENDERCORE_API void SetColors(FRHICommandList& RHICmdList, const FLinearColor* Colors, int32 NumColors);
 
-	// FShader interface.
-	virtual bool Serialize(FArchive& Ar) override
-	{
-		bool bShaderHasOutdatedParameters = FGlobalShader::Serialize(Ar);
-		//Ar << ColorParameter;
-		return bShaderHasOutdatedParameters;
-	}
 	static bool ShouldCompilePermutation(const FGlobalShaderPermutationParameters& Parameters)
 	{
 		return true;
 	}
 
 	/** The parameter to use for setting the draw Color. */
-	//FShaderParameter ColorParameter;
+	//LAYOUT_FIELD(FShaderParameter, ColorParameter);
 };
 
 /**
@@ -148,24 +133,16 @@ public:
 		FillTexture.Bind( Initializer.ParameterMap, TEXT("FillTexture"), SPF_Mandatory);
 	}
 
-	// FShader interface.
-	virtual bool Serialize(FArchive& Ar) override
-	{
-		bool bShaderHasOutdatedParameters = FGlobalShader::Serialize(Ar);
-		Ar << FillValue << Params0 << Params1 << Params2 << FillTexture;
-		return bShaderHasOutdatedParameters;
-	}
-
 	static bool ShouldCompilePermutation(const FGlobalShaderPermutationParameters& Parameters)
 	{
 		return IsFeatureLevelSupported(Parameters.Platform, ERHIFeatureLevel::SM5);
 	}
 
-	FShaderParameter FillValue;
-	FShaderParameter Params0;	// Texture Width,Height (.xy); Use Exclude Rect 1 : 0 (.z)
-	FShaderParameter Params1;	// Include X0,Y0 (.xy) - X1,Y1 (.zw)
-	FShaderParameter Params2;	// ExcludeRect X0,Y0 (.xy) - X1,Y1 (.zw)
-	FShaderResourceParameter FillTexture;
+	LAYOUT_FIELD(FShaderParameter, FillValue);
+	LAYOUT_FIELD(FShaderParameter, Params0);	// Texture Width,Height (.xy); Use Exclude Rect 1 : 0 (.z)
+	LAYOUT_FIELD(FShaderParameter, Params1);	// Include X0,Y0 (.xy) - X1,Y1 (.zw)
+	LAYOUT_FIELD(FShaderParameter, Params2);	// ExcludeRect X0,Y0 (.xy) - X1,Y1 (.zw)
+	LAYOUT_FIELD(FShaderResourceParameter, FillTexture);
 };
 
 class FLongGPUTaskPS : public FGlobalShader
@@ -177,14 +154,7 @@ public:
 	:	FGlobalShader( Initializer )
 	{
 	}
-	
-	// FShader interface.
-	virtual bool Serialize(FArchive& Ar) override
-	{
-		bool bShaderHasOutdatedParameters = FGlobalShader::Serialize(Ar);
-		return bShaderHasOutdatedParameters;
-	}
-	
+
 	static bool ShouldCompilePermutation(const FGlobalShaderPermutationParameters& Parameters)
 	{
 		// MLCHANGES BEGIN
diff --git a/Engine/Source/Runtime/RenderCore/Public/PackedNormal.h b/Engine/Source/Runtime/RenderCore/Public/PackedNormal.h
--- a/Engine/Source/Runtime/RenderCore/Public/PackedNormal.h
+++ b/Engine/Source/Runtime/RenderCore/Public/PackedNormal.h
@@ -4,6 +4,7 @@
 #pragma once
 
 #include "CoreMinimal.h"
+#include "Serialization/MemoryLayout.h"
 
 /** A normal vector, quantized and packed into 32-bits. */
 struct FPackedNormal
@@ -51,6 +52,8 @@ struct FPackedNormal
 	friend RENDERCORE_API FArchive& operator<<(FArchive& Ar, FPackedNormal& N);
 };
 
+DECLARE_INTRINSIC_TYPE_LAYOUT(FPackedNormal);
+
 class FDeprecatedSerializedPackedNormal
 {
 	union
@@ -285,6 +288,8 @@ struct FPackedRGBA16N
 	friend RENDERCORE_API FArchive& operator<<(FArchive& Ar, FPackedRGBA16N& N);
 };
 
+DECLARE_INTRINSIC_TYPE_LAYOUT(FPackedRGBA16N);
+
 FORCEINLINE void FPackedRGBA16N::operator=(const FVector& InVector)
 {
 	const float Scale = MAX_int16;
diff --git a/Engine/Source/Runtime/RenderCore/Public/PixelShaderUtils.h b/Engine/Source/Runtime/RenderCore/Public/PixelShaderUtils.h
--- a/Engine/Source/Runtime/RenderCore/Public/PixelShaderUtils.h
+++ b/Engine/Source/Runtime/RenderCore/Public/PixelShaderUtils.h
@@ -41,16 +41,16 @@ struct RENDERCORE_API FPixelShaderUtils
 	/** Initialize a pipeline state object initializer with almost all the basics required to do a full viewport pass. */
 	static void InitFullscreenPipelineState(
 		FRHICommandList& RHICmdList,
-		const TShaderMap<FGlobalShaderType>* GlobalShaderMap,
-		const FShader* PixelShader,
+		const FGlobalShaderMap* GlobalShaderMap,
+		const TShaderRef<FShader>& PixelShader,
 		FGraphicsPipelineStateInitializer& GraphicsPSOInit);
 
 	/** Dispatch a full screen pixel shader to rhi command list with its parameters. */
 	template<typename TShaderClass>
 	static inline void DrawFullscreenPixelShader(
 		FRHICommandList& RHICmdList, 
-		const TShaderMap<FGlobalShaderType>* GlobalShaderMap,
-		const TShaderClass* PixelShader,
+		const FGlobalShaderMap* GlobalShaderMap,
+		const TShaderRef<TShaderClass>& PixelShader,
 		const typename TShaderClass::FParameters& Parameters,
 		const FIntRect& Viewport,
 		FRHIBlendState* BlendState = nullptr,
@@ -58,7 +58,7 @@ struct RENDERCORE_API FPixelShaderUtils
 		FRHIDepthStencilState* DepthStencilState = nullptr,
 		uint32 StencilRef = 0)
 	{
-		check(PixelShader);
+		check(PixelShader.IsValid());
 		RHICmdList.SetViewport(Viewport.Min.X, Viewport.Min.Y, 0.0f, Viewport.Max.X, Viewport.Max.Y, 1.0f);
 		
 		FGraphicsPipelineStateInitializer GraphicsPSOInit;
@@ -70,7 +70,7 @@ struct RENDERCORE_API FPixelShaderUtils
 		SetGraphicsPipelineState(RHICmdList, GraphicsPSOInit);
 		RHICmdList.SetStencilRef(StencilRef);
 
-		SetShaderParameters(RHICmdList, PixelShader, PixelShader->GetPixelShader(), Parameters);
+		SetShaderParameters(RHICmdList, PixelShader, PixelShader.GetPixelShader(), Parameters);
 
 		DrawFullscreenTriangle(RHICmdList);
 	}
@@ -79,9 +79,9 @@ struct RENDERCORE_API FPixelShaderUtils
 	template<typename TShaderClass>
 	static inline void AddFullscreenPass(
 		FRDGBuilder& GraphBuilder,
-		const TShaderMap<FGlobalShaderType>* GlobalShaderMap,
+		const FGlobalShaderMap* GlobalShaderMap,
 		FRDGEventName&& PassName,
-		const TShaderClass* PixelShader,
+		const TShaderRef<TShaderClass>& PixelShader,
 		typename TShaderClass::FParameters* Parameters,
 		const FIntRect& Viewport,
 		FRHIBlendState* BlendState = nullptr,
@@ -89,7 +89,7 @@ struct RENDERCORE_API FPixelShaderUtils
 		FRHIDepthStencilState* DepthStencilState = nullptr,
 		uint32 StencilRef = 0)
 	{
-		check(PixelShader);
+		check(PixelShader.IsValid());
 		ClearUnusedGraphResources(PixelShader, Parameters);
 
 		GraphBuilder.AddPass(
@@ -107,9 +107,9 @@ struct RENDERCORE_API FPixelShaderUtils
 	template<typename TPixelShaderClass, typename TPassParameters>
 	static inline void AddRasterizeToRectsPass(
 		FRDGBuilder& GraphBuilder,
-		const TShaderMap<FGlobalShaderType>* GlobalShaderMap,
+		const FGlobalShaderMap* GlobalShaderMap,
 		FRDGEventName&& PassName,
-		const TPixelShaderClass* PixelShader,
+		const TShaderRef<TPixelShaderClass>& PixelShader,
 		TPassParameters* Parameters,
 		const FIntPoint& ViewportSize,
 		FRDGBufferSRVRef RectMinMaxBufferSRV,
@@ -145,15 +145,15 @@ struct RENDERCORE_API FPixelShaderUtils
 			GraphicsPSOInit.PrimitiveType = GRHISupportsRectTopology ? PT_RectList : PT_TriangleList;
 
 			GraphicsPSOInit.BoundShaderState.VertexDeclarationRHI = GTileVertexDeclaration.VertexDeclarationRHI;
-			GraphicsPSOInit.BoundShaderState.VertexShaderRHI = VertexShader->GetVertexShader();
-			GraphicsPSOInit.BoundShaderState.PixelShaderRHI = PixelShader->GetPixelShader();
+			GraphicsPSOInit.BoundShaderState.VertexShaderRHI = VertexShader.GetVertexShader();
+			GraphicsPSOInit.BoundShaderState.PixelShaderRHI = PixelShader.GetPixelShader();
 
 			SetGraphicsPipelineState(RHICmdList, GraphicsPSOInit);
 
 			RHICmdList.SetStencilRef(StencilRef);
 
-			SetShaderParameters(RHICmdList, VertexShader, VertexShader->GetVertexShader(), Parameters->VS);
-			SetShaderParameters(RHICmdList, PixelShader, PixelShader->GetPixelShader(), Parameters->PS);
+			SetShaderParameters(RHICmdList, VertexShader, VertexShader.GetVertexShader(), Parameters->VS);
+			SetShaderParameters(RHICmdList, PixelShader, PixelShader.GetPixelShader(), Parameters->PS);
 
 			const uint32 NumPrimitives = GRHISupportsRectTopology ? 1 : 2;
 			const uint32 NumInstances = Parameters->VS.NumRects;
diff --git a/Engine/Source/Runtime/RenderCore/Public/RayGenShaderUtils.h b/Engine/Source/Runtime/RenderCore/Public/RayGenShaderUtils.h
--- a/Engine/Source/Runtime/RenderCore/Public/RayGenShaderUtils.h
+++ b/Engine/Source/Runtime/RenderCore/Public/RayGenShaderUtils.h
@@ -19,7 +19,7 @@ struct RENDERCORE_API FRayGenShaderUtils
 	static inline void AddRayTraceDispatchPass(
 		FRDGBuilder& GraphBuilder,
 		FRDGEventName&& PassName,
-		const TShaderClass* RayGenerationShader,
+		const TShaderRef<TShaderClass>& RayGenerationShader,
 		typename TShaderClass::FParameters* Parameters,
 		FIntPoint Resolution)
 	{
@@ -35,11 +35,11 @@ struct RENDERCORE_API FRayGenShaderUtils
 			SetShaderParameters(GlobalResources, RayGenerationShader, *Parameters);
 
 			FRayTracingPipelineStateInitializer Initializer;
-			FRHIRayTracingShader* RayGenShaderTable[] = { RayGenerationShader->GetRayTracingShader() };
+			FRHIRayTracingShader* RayGenShaderTable[] = { RayGenerationShader.GetRayTracingShader() };
 			Initializer.SetRayGenShaderTable(RayGenShaderTable);
 
 			FRayTracingPipelineState* Pipeline = PipelineStateCache::GetAndOrCreateRayTracingPipelineState(RHICmdList, Initializer);
-			RHICmdList.RayTraceDispatch(Pipeline, RayGenerationShader->GetRayTracingShader(), GlobalResources, Resolution.X, Resolution.Y);
+			RHICmdList.RayTraceDispatch(Pipeline, RayGenerationShader.GetRayTracingShader(), GlobalResources, Resolution.X, Resolution.Y);
 		});
 	}
 };
diff --git a/Engine/Source/Runtime/RenderCore/Public/RenderGraphUtils.h b/Engine/Source/Runtime/RenderCore/Public/RenderGraphUtils.h
--- a/Engine/Source/Runtime/RenderCore/Public/RenderGraphUtils.h
+++ b/Engine/Source/Runtime/RenderCore/Public/RenderGraphUtils.h
@@ -8,6 +8,8 @@
 #include "ShaderParameterStruct.h"
 #include "ShaderParameterMacros.h"
 
+class FGlobalShaderMap;
+
 /** An empty shader parameter structure ready to be used anywhere. */
 BEGIN_SHADER_PARAMETER_STRUCT(FEmptyShaderParameters, RENDERCORE_API)
 END_SHADER_PARAMETER_STRUCT()
@@ -43,7 +45,7 @@ extern RENDERCORE_API void ClearUnusedGraphResourcesImpl(
 
 template <typename TShaderClass>
 void ClearUnusedGraphResources(
-	const TShaderClass* Shader,
+	const TShaderRef<TShaderClass>& Shader,
 	typename TShaderClass::FParameters* InoutParameters,
 	std::initializer_list<FRDGResourceRef> ExcludeList = {})
 {
@@ -60,8 +62,8 @@ void ClearUnusedGraphResources(
 
 template <typename TShaderClassA, typename TShaderClassB, typename TPassParameterStruct>
 void ClearUnusedGraphResources(
-	const TShaderClassA* ShaderA,
-	const TShaderClassB* ShaderB,
+	const TShaderRef<TShaderClassA>& ShaderA,
+	const TShaderRef<TShaderClassB>& ShaderB,
 	TPassParameterStruct* InoutParameters,
 	std::initializer_list<FRDGResourceRef> ExcludeList = {})
 {
@@ -132,9 +134,9 @@ struct RENDERCORE_API FComputeShaderUtils
 
 	/** Dispatch a compute shader to rhi command list with its parameters. */
 	template<typename TShaderClass>
-	static void Dispatch(FRHICommandList& RHICmdList, const TShaderClass* ComputeShader, const typename TShaderClass::FParameters& Parameters, FIntVector GroupCount)
+	static void Dispatch(FRHICommandList& RHICmdList, const TShaderRef<TShaderClass>& ComputeShader, const typename TShaderClass::FParameters& Parameters, FIntVector GroupCount)
 	{
-		FRHIComputeShader* ShaderRHI = ComputeShader->GetComputeShader();
+		FRHIComputeShader* ShaderRHI = ComputeShader.GetComputeShader();
 		RHICmdList.SetComputeShader(ShaderRHI);
 		SetShaderParameters(RHICmdList, ComputeShader, ShaderRHI, Parameters);
 		RHICmdList.DispatchComputeShader(GroupCount.X, GroupCount.Y, GroupCount.Z);
@@ -145,12 +147,12 @@ struct RENDERCORE_API FComputeShaderUtils
 	template<typename TShaderClass>
 	static void DispatchIndirect(
 		FRHICommandList& RHICmdList,
-		const TShaderClass* ComputeShader,
+		const TShaderRef<TShaderClass>& ComputeShader,
 		const typename TShaderClass::FParameters& Parameters,
 		FRHIVertexBuffer* IndirectArgsBuffer,
 		uint32 IndirectArgOffset)
 	{
-		FRHIComputeShader* ShaderRHI = ComputeShader->GetComputeShader();
+		FRHIComputeShader* ShaderRHI = ComputeShader.GetComputeShader();
 		RHICmdList.SetComputeShader(ShaderRHI);
 		SetShaderParameters(RHICmdList, ComputeShader, ShaderRHI, Parameters);
 		RHICmdList.DispatchIndirectComputeShader(IndirectArgsBuffer, IndirectArgOffset);
@@ -178,7 +180,7 @@ struct RENDERCORE_API FComputeShaderUtils
 	static void AddPass(
 		FRDGBuilder& GraphBuilder,
 		FRDGEventName&& PassName,
-		const TShaderClass* ComputeShader,
+		const TShaderRef<TShaderClass>& ComputeShader,
 		typename TShaderClass::FParameters* Parameters,
 		FIntVector GroupCount)
 	{
@@ -199,7 +201,7 @@ struct RENDERCORE_API FComputeShaderUtils
 	static void AddPass(
 		FRDGBuilder& GraphBuilder,
 		FRDGEventName&& PassName,
-		const TShaderClass* ComputeShader,
+		const TShaderRef<TShaderClass>& ComputeShader,
 		typename TShaderClass::FParameters* Parameters,
 		FRDGBufferRef IndirectArgsBuffer,
 		uint32 IndirectArgOffset)
@@ -222,10 +224,9 @@ struct RENDERCORE_API FComputeShaderUtils
 		});
 	}
 
-	static void ClearUAV(FRDGBuilder& GraphBuilder, TShaderMap<FGlobalShaderType>* ShaderMap, FRDGBufferUAVRef UAV, uint32 ClearValue);
-	static void ClearUAV(FRDGBuilder& GraphBuilder, TShaderMap<FGlobalShaderType>* ShaderMap, FRDGBufferUAVRef UAV, FVector4 ClearValue);
+	static void ClearUAV(FRDGBuilder& GraphBuilder, FGlobalShaderMap* ShaderMap, FRDGBufferUAVRef UAV, uint32 ClearValue);
+	static void ClearUAV(FRDGBuilder& GraphBuilder, FGlobalShaderMap* ShaderMap, FRDGBufferUAVRef UAV, FVector4 ClearValue);
 };
-
 /** Adds a render graph pass to copy a region from one texture to another. Uses RHICopyTexture under the hood.
  *  Formats of the two textures must match. The output and output texture regions be within the respective extents.
  */
diff --git a/Engine/Source/Runtime/RenderCore/Public/RenderResource.h b/Engine/Source/Runtime/RenderCore/Public/RenderResource.h
--- a/Engine/Source/Runtime/RenderCore/Public/RenderResource.h
+++ b/Engine/Source/Runtime/RenderCore/Public/RenderResource.h
@@ -10,32 +10,60 @@
 #include "Containers/List.h"
 #include "RHI.h"
 #include "RenderCore.h"
+#include "Serialization/MemoryLayout.h"
 
 /** Number of frames after which unused global resource allocations will be discarded. */
 extern int32 GGlobalBufferNumFramesUnusedThresold;
 
 /**
  * A rendering resource which is owned by the rendering thread.
+ * NOTE - Adding new virtual methods to this class may require stubs added to FViewport/FDummyViewport, otherwise certain modules may have link errors
  */
 class RENDERCORE_API FRenderResource
 {
 public:
+	template<typename FunctionType>
+	static void ForAllResources(const FunctionType& Function)
+	{
+		const TArray<FRenderResource*>& ResourceList = GetResourceList();
+		ResourceListIterationActive.Increment();
+		for (int32 Index = 0; Index < ResourceList.Num(); ++Index)
+		{
+			FRenderResource* Resource = ResourceList[Index];
+			if (Resource)
+			{
+				checkSlow(Resource->ListIndex == Index);
+				Function(Resource);
+			}
+		}
+		ResourceListIterationActive.Decrement();
+	}
 
-	/** @return The global initialized resource list. */
-	static TLinkedList<FRenderResource*>*& GetResourceList();
+	static void InitRHIForAllResources()
+	{
+		ForAllResources([](FRenderResource* Resource) { Resource->InitRHI(); });
+		// Dynamic resources can have dependencies on static resources (with uniform buffers) and must initialized last!
+		ForAllResources([](FRenderResource* Resource) { Resource->InitDynamicRHI(); });
+	}
+
+	static void ReleaseRHIForAllResources()
+	{
+		ForAllResources([](FRenderResource* Resource) { check(Resource->IsInitialized()); Resource->ReleaseRHI(); });
+		ForAllResources([](FRenderResource* Resource) { Resource->ReleaseDynamicRHI(); });
+	}
 
 	static void ChangeFeatureLevel(ERHIFeatureLevel::Type NewFeatureLevel);
 
 	/** Default constructor. */
 	FRenderResource()
-		: FeatureLevel(ERHIFeatureLevel::Num)
-		, bInitialized(false)
+		: ListIndex(INDEX_NONE)
+		, FeatureLevel(ERHIFeatureLevel::Num)
 	{}
 
 	/** Constructor when we know what feature level this resource should support */
 	FRenderResource(ERHIFeatureLevel::Type InFeatureLevel)
-		: FeatureLevel(InFeatureLevel)
-		, bInitialized(false)
+		: ListIndex(INDEX_NONE)
+		, FeatureLevel(InFeatureLevel)
 	{}
 
 	/** Destructor used to catch unreleased resources. */
@@ -89,32 +117,31 @@ public:
 	 */
 	void UpdateRHI();
 
-	// Probably temporary code that sends a task back to renderthread_local and blocks waiting for it to call InitResource
-	void InitResourceFromPossiblyParallelRendering();
-
 	/** @return The resource's friendly name.  Typically a UObject name. */
 	virtual FString GetFriendlyName() const { return TEXT("undefined"); }
 
 	// Accessors.
-	FORCEINLINE bool IsInitialized() const { return bInitialized; }
+	FORCEINLINE bool IsInitialized() const { return ListIndex != INDEX_NONE; }
 
 	/** Initialize all resources initialized before the RHI was initialized */
 	static void InitPreRHIResources();
 
+private:
+	/** @return The global initialized resource list. */
+	static TArray<FRenderResource*>& GetResourceList();
+
+	static FThreadSafeCounter ResourceListIterationActive;
+
+	int32 ListIndex;
+
 protected:
 	// This is used during mobile editor preview refactor, this will eventually be replaced with a parameter to InitRHI() etc..
-	ERHIFeatureLevel::Type GetFeatureLevel() const { return FeatureLevel == ERHIFeatureLevel::Num ? GMaxRHIFeatureLevel : FeatureLevel; }
+	void SetFeatureLevel(const FStaticFeatureLevel InFeatureLevel) { FeatureLevel = (ERHIFeatureLevel::Type)InFeatureLevel; }
+	const FStaticFeatureLevel GetFeatureLevel() const { return FeatureLevel == ERHIFeatureLevel::Num ? FStaticFeatureLevel(GMaxRHIFeatureLevel) : FeatureLevel; }
 	FORCEINLINE bool HasValidFeatureLevel() const { return FeatureLevel < ERHIFeatureLevel::Num; }
 
-	ERHIFeatureLevel::Type FeatureLevel;
-
 private:
-	/** True if the resource has been initialized. */
-	bool bInitialized;
-
-	#if PLATFORM_NEEDS_RHIRESOURCELIST
-	TLinkedList<FRenderResource*> ResourceLink;
-	#endif
+	TEnumAsByte<ERHIFeatureLevel::Type> FeatureLevel;
 };
 
 /**
