Copying Private-LoadTimes-4.24 stream to Main. Biggest changes are in Materials/Shader memory freezing.

diff --git a/Engine/Source/Runtime/RenderCore/Public/RenderUtils.h b/Engine/Source/Runtime/RenderCore/Public/RenderUtils.h
--- a/Engine/Source/Runtime/RenderCore/Public/RenderUtils.h
+++ b/Engine/Source/Runtime/RenderCore/Public/RenderUtils.h
@@ -7,6 +7,7 @@
 #include "RHI.h"
 #include "PackedNormal.h"
 #include "RenderResource.h"
+#include "RHIDefinitions.h"
 
 extern RENDERCORE_API void RenderUtilsInit();
 
@@ -460,20 +461,20 @@ RENDERCORE_API FVertexDeclarationRHIRef& GetVertexDeclarationFVector3();
 
 RENDERCORE_API FVertexDeclarationRHIRef& GetVertexDeclarationFVector2();
 
-RENDERCORE_API bool PlatformSupportsSimpleForwardShading(EShaderPlatform Platform);
+RENDERCORE_API bool PlatformSupportsSimpleForwardShading(const FStaticShaderPlatform Platform);
 
-RENDERCORE_API bool IsSimpleForwardShadingEnabled(EShaderPlatform Platform);
+RENDERCORE_API bool IsSimpleForwardShadingEnabled(const FStaticShaderPlatform Platform);
 
-RENDERCORE_API bool MobileSupportsGPUScene(EShaderPlatform Platform);
+RENDERCORE_API bool MobileSupportsGPUScene(const FStaticShaderPlatform Platform);
 
-RENDERCORE_API bool GPUSceneUseTexture2D(EShaderPlatform Platform);
+RENDERCORE_API bool GPUSceneUseTexture2D(const FStaticShaderPlatform Platform);
 
-RENDERCORE_API bool MaskedInEarlyPass(EShaderPlatform Platform);
+RENDERCORE_API bool MaskedInEarlyPass(const FStaticShaderPlatform Platform);
 
-RENDERCORE_API bool AllowPixelDepthOffset(EShaderPlatform Platform);
+RENDERCORE_API bool AllowPixelDepthOffset(const FStaticShaderPlatform Platform);
 
 /** Returns if ForwardShading is enabled. Only valid for the current platform (otherwise call ITargetPlatform::UsesForwardShading()). */
-inline bool IsForwardShadingEnabled(EShaderPlatform Platform)
+inline bool IsForwardShadingEnabled(const FStaticShaderPlatform Platform)
 {
 	extern RENDERCORE_API uint64 GForwardShadingPlatformMask;
 	return !!(GForwardShadingPlatformMask & (1ull << Platform))
@@ -482,46 +483,46 @@ inline bool IsForwardShadingEnabled(EShaderPlatform Platform)
 }
 
 /** Returns if ForwardShading or SimpleForwardShading is enabled. Only valid for the current platform. */
-inline bool IsAnyForwardShadingEnabled(EShaderPlatform Platform)
+inline bool IsAnyForwardShadingEnabled(const FStaticShaderPlatform Platform)
 {
 	return IsForwardShadingEnabled(Platform) || IsSimpleForwardShadingEnabled(Platform);
 }
 
 /** Returns if the GBuffer is used. Only valid for the current platform. */
-inline bool IsUsingGBuffers(EShaderPlatform Platform)
+inline bool IsUsingGBuffers(const FStaticShaderPlatform Platform)
 {
 	return !IsAnyForwardShadingEnabled(Platform);
 }
 
 /** Returns whether DBuffer decals are enabled for a given shader platform */
-inline bool IsUsingDBuffers(EShaderPlatform Platform)
+inline bool IsUsingDBuffers(const FStaticShaderPlatform Platform)
 {
 	extern RENDERCORE_API uint64 GDBufferPlatformMask;
 	return !!(GDBufferPlatformMask & (1ull << Platform));
 }
 
 /** Returns whether the base pass should output to the velocity buffer is enabled for a given shader platform */
-inline bool IsUsingBasePassVelocity(EShaderPlatform Platform)
+inline bool IsUsingBasePassVelocity(const FStaticShaderPlatform Platform)
 {
 	extern RENDERCORE_API uint64 GBasePassVelocityPlatformMask;
 	return !!(GBasePassVelocityPlatformMask & (1ull << Platform));
 }
 
 /** Returns whether the base pass should use selective outputs for a given shader platform */
-inline bool IsUsingSelectiveBasePassOutputs(EShaderPlatform Platform)
+inline bool IsUsingSelectiveBasePassOutputs(const FStaticShaderPlatform Platform)
 {
 	extern RENDERCORE_API uint64 GSelectiveBasePassOutputsPlatformMask;
 	return !!(GSelectiveBasePassOutputsPlatformMask & (1ull << Platform));
 }
 
 /** Returns whether distance fields are enabled for a given shader platform */
-inline bool IsUsingDistanceFields(EShaderPlatform Platform)
+inline bool IsUsingDistanceFields(const FStaticShaderPlatform Platform)
 {
 	extern RENDERCORE_API uint64 GDistanceFieldsPlatformMask;
 	return !!(GDistanceFieldsPlatformMask & (1ull << Platform));
 }
 
-inline bool IsUsingPerPixelDBufferMask(EShaderPlatform Platform)
+inline bool IsUsingPerPixelDBufferMask(const FStaticShaderPlatform Platform)
 {
 	switch (Platform)
 	{
@@ -534,7 +535,7 @@ inline bool IsUsingPerPixelDBufferMask(EShaderPlatform Platform)
 	}
 }
 
-inline bool UseGPUScene(EShaderPlatform Platform, ERHIFeatureLevel::Type FeatureLevel)
+inline bool UseGPUScene(const FStaticShaderPlatform Platform, const FStaticFeatureLevel FeatureLevel)
 {
 	if (FeatureLevel == ERHIFeatureLevel::ES3_1)
 	{
@@ -545,7 +546,9 @@ inline bool UseGPUScene(EShaderPlatform Platform, ERHIFeatureLevel::Type Feature
 	return FeatureLevel >= ERHIFeatureLevel::SM5 
 		//@todo - support GPU Scene management compute shaders on these platforms to get dynamic instancing speedups on the Rendering Thread and RHI Thread
 		&& !IsOpenGLPlatform(Platform)
-		&& !IsSwitchPlatform(Platform);
+		&& !IsSwitchPlatform(Platform)
+		// we only check DDSPI for platforms that have been read in - IsValid() can go away once ALL platforms are converted over to this system
+		&& (!FDataDrivenShaderPlatformInfo::IsValid(Platform) || FDataDrivenShaderPlatformInfo::GetSupportsGPUScene(Platform));
 }
 
 /** Unit cube vertex buffer (VertexDeclarationFVector4) */
@@ -564,4 +567,4 @@ RENDERCORE_API void QuantizeSceneBufferSize(const FIntPoint& InBufferSize, FIntP
 /**
 *	Checks if virtual texturing enabled and supported
 */
-RENDERCORE_API bool UseVirtualTexturing(ERHIFeatureLevel::Type InFeatureLevel, const class ITargetPlatform* TargetPlatform = nullptr);
+RENDERCORE_API bool UseVirtualTexturing(const FStaticFeatureLevel InFeatureLevel, const class ITargetPlatform* TargetPlatform = nullptr);
diff --git a/Engine/Source/Runtime/RenderCore/Public/RendererInterface.h b/Engine/Source/Runtime/RenderCore/Public/RendererInterface.h
--- a/Engine/Source/Runtime/RenderCore/Public/RendererInterface.h
+++ b/Engine/Source/Runtime/RenderCore/Public/RendererInterface.h
@@ -27,6 +27,9 @@ class FSceneTexturesUniformParameters;
 class FGlobalDistanceFieldParameterData;
 struct FMeshBatch;
 struct FSynthBenchmarkResults;
+class FShader;
+class FShaderMapPointerTable;
+template<typename ShaderType, typename PointerTableType> class TShaderRefBase;
 
 // Shortcut for the allocator used by scene rendering.
 typedef TMemStackAllocator<> SceneRenderingAllocator;
@@ -722,7 +725,7 @@ public:
 		float SizeV,
 		FIntPoint TargetSize,
 		FIntPoint TextureSize,
-		class FShader* VertexShader,
+		const TShaderRefBase<FShader, FShaderMapPointerTable>& VertexShader,
 		EDrawRectangleFlags Flags = EDRF_Default
 		) = 0;
 
diff --git a/Engine/Source/Runtime/RenderCore/Public/RenderingThread.h b/Engine/Source/Runtime/RenderCore/Public/RenderingThread.h
--- a/Engine/Source/Runtime/RenderCore/Public/RenderingThread.h
+++ b/Engine/Source/Runtime/RenderCore/Public/RenderingThread.h
@@ -11,6 +11,7 @@
 #include "Async/TaskGraphInterfaces.h"
 #include "Templates/Atomic.h"
 #include "Trace/Trace.h"
+#include "Serialization/MemoryLayout.h"
 
 class FRHICommandListImmediate;
 
diff --git a/Engine/Source/Runtime/RenderCore/Public/ResolveShader.h b/Engine/Source/Runtime/RenderCore/Public/ResolveShader.h
--- a/Engine/Source/Runtime/RenderCore/Public/ResolveShader.h
+++ b/Engine/Source/Runtime/RenderCore/Public/ResolveShader.h
@@ -40,14 +40,7 @@ public:
 	{
 	}
 
-	virtual bool Serialize(FArchive& Ar) override
-	{
-		bool bShaderHasOutdatedParameters = FGlobalShader::Serialize(Ar);
-		Ar << UnresolvedSurface;
-		return bShaderHasOutdatedParameters;
-	}
-	
-	FShaderResourceParameter UnresolvedSurface;
+	LAYOUT_FIELD(FShaderResourceParameter, UnresolvedSurface);
 };
 
 class FResolveDepth2XPS : public FGlobalShader
@@ -81,14 +74,7 @@ public:
 	{
 	}
 
-	virtual bool Serialize(FArchive& Ar) override
-	{
-		bool bShaderHasOutdatedParameters = FGlobalShader::Serialize(Ar);
-		Ar << UnresolvedSurface;
-		return bShaderHasOutdatedParameters;
-	}
-
-	FShaderResourceParameter UnresolvedSurface;
+	LAYOUT_FIELD(FShaderResourceParameter, UnresolvedSurface);
 };
 
 
@@ -123,14 +109,7 @@ public:
 	{
 	}
 
-	virtual bool Serialize(FArchive& Ar) override
-	{
-		bool bShaderHasOutdatedParameters = FGlobalShader::Serialize(Ar);
-		Ar << UnresolvedSurface;
-		return bShaderHasOutdatedParameters;
-	}
-
-	FShaderResourceParameter UnresolvedSurface;
+	LAYOUT_FIELD(FShaderResourceParameter, UnresolvedSurface);
 };
 
 
@@ -167,14 +146,7 @@ public:
 	{
 	}
 
-	virtual bool Serialize(FArchive& Ar) override
-	{
-		bool bShaderHasOutdatedParameters = FGlobalShader::Serialize(Ar);
-		Ar << UnresolvedSurface;
-		return bShaderHasOutdatedParameters;
-	}
-
-	FShaderResourceParameter UnresolvedSurface;
+	LAYOUT_FIELD(FShaderResourceParameter, UnresolvedSurface);
 };
 
 
@@ -197,15 +169,8 @@ public:
 	void SetParameters(FRHICommandList& RHICmdList, FParameter)
 	{
 	}
-	
-	virtual bool Serialize(FArchive& Ar) override
-	{
-		bool bShaderHasOutdatedParameters = FGlobalShader::Serialize(Ar);
-		Ar << UnresolvedSurface;
-		return bShaderHasOutdatedParameters;
-	}
-	
-	FShaderResourceParameter UnresolvedSurface;
+
+	LAYOUT_FIELD(FShaderResourceParameter, UnresolvedSurface);
 };
 
 class FResolveSingleSamplePS : public FGlobalShader
@@ -227,16 +192,8 @@ public:
 	
 	RENDERCORE_API void SetParameters(FRHICommandList& RHICmdList, uint32 SingleSampleIndexValue);
 	
-	virtual bool Serialize(FArchive& Ar) override
-	{
-		bool bShaderHasOutdatedParameters = FGlobalShader::Serialize(Ar);
-		Ar << UnresolvedSurface;
-		Ar << SingleSampleIndex;
-		return bShaderHasOutdatedParameters;
-	}
-	
-	FShaderResourceParameter UnresolvedSurface;
-	FShaderParameter SingleSampleIndex;
+	LAYOUT_FIELD(FShaderResourceParameter, UnresolvedSurface);
+	LAYOUT_FIELD(FShaderParameter, SingleSampleIndex);
 };
 
 /**
@@ -259,14 +216,6 @@ public:
 
 	RENDERCORE_API void SetParameters(FRHICommandList& RHICmdList, const FResolveRect& SrcBounds, const FResolveRect& DstBounds, uint32 DstSurfaceWidth, uint32 DstSurfaceHeight);
 
-	virtual bool Serialize(FArchive& Ar) override
-	{
-		bool bShaderHasOutdatedParameters = FGlobalShader::Serialize(Ar);
-		Ar << PositionMinMax;
-		Ar << UVMinMax;
-		return bShaderHasOutdatedParameters;
-	}
-
-	FShaderParameter PositionMinMax;
-	FShaderParameter UVMinMax;
+	LAYOUT_FIELD(FShaderParameter, PositionMinMax);
+	LAYOUT_FIELD(FShaderParameter, UVMinMax);
 };
