Copying Private-LoadTimes-4.24 stream to Main. Biggest changes are in Materials/Shader memory freezing.

diff --git a/Engine/Source/Runtime/RenderCore/Public/Shader.h b/Engine/Source/Runtime/RenderCore/Public/Shader.h
--- a/Engine/Source/Runtime/RenderCore/Public/Shader.h
+++ b/Engine/Source/Runtime/RenderCore/Public/Shader.h
@@ -8,14 +8,17 @@
 
 #include "CoreMinimal.h"
 #include "Containers/List.h"
+#include "Containers/HashTable.h"
 #include "Misc/SecureHash.h"
 #include "RenderResource.h"
 #include "ShaderParameters.h"
 #include "RenderingThread.h"
 #include "ShaderCore.h"
 #include "ShaderPermutation.h"
+#include "ShaderCodeLibrary.h" // TODO - move to cpp
 #include "Serialization/ArchiveProxy.h"
 #include "UObject/RenderingObjectVersion.h"
+#include "Serialization/MemoryImage.h"
 
 // For FShaderUniformBufferParameter
 
@@ -29,15 +32,12 @@ class FNiagaraShaderType;
 class FOpenColorIOShaderType;
 class FMeshMaterialShaderType;
 class FShader;
+class FShaderMapBase;
 class FShaderPipelineType;
 class FShaderType;
-class FVertexFactoryParameterRef;
 class FVertexFactoryType;
 class FShaderParametersMetadata;
-
-/** By default most shader source hashes are stripped at cook time so can be discarded
-	to save memory. See implementation of FilterShaderSourceHashForSerialization. */
-#define KEEP_SHADER_SOURCE_HASHES	(WITH_EDITOR)
+class FShaderMapPointerTable;
 
 /** Define a shader permutation uniquely according to its type, and permutation id.*/
 template<typename MetaShaderType>
@@ -95,82 +95,17 @@ public:
 	}
 };
 
-
-/** 
- * Uniquely identifies an FShaderResource.  
- * Used to link FShaders to FShaderResources on load. 
- */
+// TODO - remove me
 class FShaderResourceId
 {
-public:
-
-	FShaderResourceId() {}
-
-	//#todo-RemoveStreamOut
-	FShaderResourceId(const FShaderTarget& InTarget, const FSHAHash& InOutputHash, int32 InSpecificPermutationId) :
-		OutputHash(InOutputHash),
-		Target(InTarget),
-		SpecificPermutationId(InSpecificPermutationId)
-	{
-		check(InSpecificPermutationId == 0);
-	}
-
-	friend inline uint32 GetTypeHash( const FShaderResourceId& Id )
-	{
-		return GetTypeHash(Id.OutputHash);
-	}
-
-	friend bool operator==(const FShaderResourceId& X, const FShaderResourceId& Y)
-	{
-		return X.Target == Y.Target 
-			&& X.OutputHash == Y.OutputHash 
-			&& X.SpecificPermutationId == Y.SpecificPermutationId;
-	}
-
-	friend bool operator!=(const FShaderResourceId& X, const FShaderResourceId& Y)
-	{
-		return !(X == Y);
-	}
-
-	friend FArchive& operator<<(FArchive& Ar, FShaderResourceId& Id)
-	{
-		Ar.UsingCustomVersion(FRenderingObjectVersion::GUID);
-
-		Ar << Id.Target << Id.OutputHash;
-
-		//#todo-RemoveStreamOut
-		FString SpecificShaderTypeStorage = TEXT("");
-		Ar << SpecificShaderTypeStorage;
-
-		if (Ar.CustomVer(FRenderingObjectVersion::GUID) >= FRenderingObjectVersion::ShaderPermutationId)
-		{
-			Ar << Id.SpecificPermutationId;
-		}
-
-		return Ar;
-	}
-
-	/** Copy string data that belongs to FShaderResource for cases where this id can outlive the resource. */
-	void MakeSelfContained()
-	{
-	}
-
-	/** Hash of the compiled shader output, which is used to create the FShaderResource. */
-	FSHAHash OutputHash;
-
-	/** Target platform and frequency. */
-	FShaderTarget Target;
-
-	/** Specific permutation identifier of the shader when SpecificShaderTypeName is non null, ignored otherwise. */
-	//#todo-RemoveStreamOut
-	int32 SpecificPermutationId;
 };
 
 class FShaderParameterInfo
 {
+	DECLARE_TYPE_LAYOUT(FShaderParameterInfo, NonVirtual);
 public:
-	uint16 BaseIndex;
-	uint16 Size;
+	LAYOUT_FIELD(uint16, BaseIndex);
+	LAYOUT_FIELD(uint16, Size);
 
 	FShaderParameterInfo() {}
 
@@ -181,7 +116,7 @@ public:
 		checkf(BaseIndex == InBaseIndex && Size == InSize, TEXT("Tweak FShaderParameterInfo type sizes"));
 	}
 
-	friend FArchive& operator<<(FArchive& Ar,FShaderParameterInfo& Info)
+	friend FArchive& operator<<(FArchive& Ar, FShaderParameterInfo& Info)
 	{
 		Ar << Info.BaseIndex;
 		Ar << Info.Size;
@@ -193,47 +128,60 @@ public:
 		return BaseIndex == Rhs.BaseIndex
 			&& Size == Rhs.Size;
 	}
+
+	inline bool operator<(const FShaderParameterInfo& Rhs) const
+	{
+		return BaseIndex < Rhs.BaseIndex;
+	}
 };
 
 class FShaderLooseParameterBufferInfo
 {
+	DECLARE_TYPE_LAYOUT(FShaderLooseParameterBufferInfo, NonVirtual);
 public:
-	uint16 BufferIndex;
-	uint16 BufferSize;
-	TArray<FShaderParameterInfo> Parameters;
+	LAYOUT_FIELD(uint16, BaseIndex);
+	LAYOUT_FIELD(uint16, Size);
+
+	LAYOUT_FIELD(TMemoryImageArray<FShaderParameterInfo>, Parameters);
 
 	FShaderLooseParameterBufferInfo() {}
 
 	FShaderLooseParameterBufferInfo(uint16 InBufferIndex, uint16 InBufferSize)
 	{
-		BufferIndex = InBufferIndex;
-		BufferSize = InBufferSize;
-		checkf(BufferIndex == InBufferIndex, TEXT("Tweak FShaderLooseParameterBufferInfo type sizes"));
+		BaseIndex = InBufferIndex;
+		Size = InBufferSize;
+		checkf(BaseIndex == InBufferIndex, TEXT("Tweak FShaderLooseParameterBufferInfo type sizes"));
 	}
 
 	friend FArchive& operator<<(FArchive& Ar,FShaderLooseParameterBufferInfo& Info)
 	{
-		Ar << Info.BufferIndex;
-		Ar << Info.BufferSize;
+		Ar << Info.BaseIndex;
+		Ar << Info.Size;
 		Ar << Info.Parameters;
 		return Ar;
 	}
 
 	inline bool operator==(const FShaderLooseParameterBufferInfo& Rhs) const
 	{
-		return BufferIndex == Rhs.BufferIndex
-			&& BufferSize == Rhs.BufferSize
+		return BaseIndex == Rhs.BaseIndex
+			&& Size == Rhs.Size
 			&& Parameters == Rhs.Parameters;
 	}
+
+	inline bool operator<(const FShaderLooseParameterBufferInfo& Rhs) const
+	{
+		return BaseIndex < Rhs.BaseIndex;
+	}
 };
 
 class FShaderParameterMapInfo
 {
+	DECLARE_TYPE_LAYOUT(FShaderParameterMapInfo, NonVirtual);
 public:
-	TArray<FShaderParameterInfo> UniformBuffers;
-	TArray<FShaderParameterInfo> TextureSamplers;
-	TArray<FShaderParameterInfo> SRVs;
-	TArray<FShaderLooseParameterBufferInfo> LooseParameterBuffers;
+	LAYOUT_FIELD(TMemoryImageArray<FShaderParameterInfo>, UniformBuffers);
+	LAYOUT_FIELD(TMemoryImageArray<FShaderParameterInfo>, TextureSamplers);
+	LAYOUT_FIELD(TMemoryImageArray<FShaderParameterInfo>, SRVs);
+	LAYOUT_FIELD(TMemoryImageArray<FShaderLooseParameterBufferInfo>, LooseParameterBuffers);
 
 	friend FArchive& operator<<(FArchive& Ar,FShaderParameterMapInfo& Info)
 	{
@@ -253,263 +201,166 @@ public:
 	}
 };
 
-/** 
- * Compiled shader bytecode and its corresponding RHI resource. 
- * This can be shared by multiple FShaders with identical compiled output.
- */
-class FShaderResource : public FRenderResource, public FDeferredCleanupInterface
+class RENDERCORE_API FShaderMapResource : public FRenderResource, public FDeferredCleanupInterface
 {
-	friend class FShader;
-	friend class FShaderResourceSet;
 public:
+	static bool ArePlatformsCompatible(EShaderPlatform CurrentPlatform, EShaderPlatform TargetPlatform);
 
-	/** Constructor used for deserialization. */
-	RENDERCORE_API FShaderResource();
-	FShaderResource(FShaderResource&&) = default;
-
-	/** Constructor used when creating a new shader resource from compiled output. */
-	//#todo-RemoveStreamOut
-	FShaderResource(const FShaderCompilerOutput& Output, int32 InSpecificPermutationId);
+	EShaderPlatform GetPlatform() const { return Platform; }
 
-	RENDERCORE_API ~FShaderResource();
-
-	RENDERCORE_API void Serialize(FArchive& Ar, bool bLoadedByCookedMaterial);
+	void AddRef();
+	void Release();
+	inline int32 GetNumRefs() const { return NumRefs; }
 
-	// Reference counting.
-	RENDERCORE_API void AddRef();
-	RENDERCORE_API void Release();
+	// FRenderResource interface.
+	virtual void ReleaseRHI();
 
-	/** @return the shader's vertex shader */
-	FORCEINLINE FRHIVertexShader* GetVertexShader()
-	{
-		checkSlow(Target.Frequency == SF_Vertex);
-		if (!IsInitialized())
-		{
-			InitializeShaderRHI();
-		}
-		return (FRHIVertexShader*)Shader.GetReference();
-	}
-	/** @return the shader's pixel shader */
-	FORCEINLINE FRHIPixelShader* GetPixelShader()
-	{
-		checkSlow(Target.Frequency == SF_Pixel);
-		if (!IsInitialized())
-		{
-			InitializeShaderRHI();
-		}
-		return (FRHIPixelShader*)Shader.GetReference();
-	}
-	/** @return the shader's hull shader */
-	FORCEINLINE FRHIHullShader* GetHullShader()
-	{
-		checkSlow(Target.Frequency == SF_Hull);
-		if (!IsInitialized())
-		{
-			InitializeShaderRHI();
-		}
-		return (FRHIHullShader*)Shader.GetReference();
-	}
-	/** @return the shader's domain shader */
-	FORCEINLINE FRHIDomainShader* GetDomainShader()
+	inline int32 GetNumShaders() const
 	{
-		checkSlow(Target.Frequency == SF_Domain);
-		if (!IsInitialized())
-		{
-			InitializeShaderRHI();
-		}
-		return (FRHIDomainShader*)Shader.GetReference();
+		return RHIShaders.Num();
 	}
-	/** @return the shader's geometry shader */
-	FORCEINLINE FRHIGeometryShader* GetGeometryShader()
+
+	inline bool HasShader(int32 ShaderIndex) const
 	{
-		checkSlow(Target.Frequency == SF_Geometry);
-		if (!IsInitialized())
-		{
-			InitializeShaderRHI();
-		}
-		return (FRHIGeometryShader*)Shader.GetReference();
+		return RHIShaders[ShaderIndex].IsValid();
 	}
-	/** @return the shader's compute shader */
-	FORCEINLINE FRHIComputeShader* GetComputeShader()
+
+	inline FRHIShader* GetShader(int32 ShaderIndex)
 	{
-		checkSlow(Target.Frequency == SF_Compute);
-		if (!IsInitialized())
+		if (!RHIShaders[ShaderIndex])
 		{
-			InitializeShaderRHI();
+			CreateShader(ShaderIndex);
 		}
-		return (FRHIComputeShader*)Shader.GetReference();
+		return RHIShaders[ShaderIndex];
 	}
 
-#if RHI_RAYTRACING
-	inline FRHIRayTracingShader* GetRayTracingShader()
-	{
-		checkSlow(Target.Frequency == SF_RayGen
-			   || Target.Frequency == SF_RayMiss
-			   || Target.Frequency == SF_RayHitGroup
-			   || Target.Frequency == SF_RayCallable);
+	void BeginCreateAllShaders();
 
-		if (!IsInitialized())
-		{
-			InitializeShaderRHI();
-		}
-		return RayTracingShader;
-	}
+#if RHI_RAYTRACING
+	static void GetRayTracingMaterialLibrary(TArray<FRHIRayTracingShader*>& RayTracingMaterials, FRHIRayTracingShader* DefaultShader);
 
-	inline uint32 GetRayTracingMaterialLibraryIndex()
+	inline uint32 GetRayTracingMaterialLibraryIndex(int32 ShaderIndex)
 	{
-		checkSlow(Target.Frequency == SF_RayHitGroup);
-
-		if (!IsInitialized())
+		if (!RHIShaders[ShaderIndex])
 		{
-			InitializeShaderRHI();
+			CreateShader(ShaderIndex);
 		}
-		return RayTracingMaterialLibraryIndex;
+		return RayTracingMaterialLibraryIndices[ShaderIndex];
 	}
-
-	RENDERCORE_API static void GetRayTracingMaterialLibrary(TArray<FRHIRayTracingShader*>& RayTracingMaterials, FRHIRayTracingShader* DefaultShader);
-
-private:
-	RENDERCORE_API static uint32 AddToRayTracingLibrary(FRHIRayTracingShader* Shader);
-	RENDERCORE_API static void RemoveFromRayTracingLibrary(uint32 Index);
-
-	static TArray<uint32> GlobalUnusedIndicies;
-	static TArray<FRHIRayTracingShader*> GlobalRayTracingMaterialLibrary;
-	static FCriticalSection GlobalRayTracingMaterialLibraryCS;
-
-public:
 #endif // RHI_RAYTRACING
 
-	RENDERCORE_API FShaderResourceId GetId() const;
-	
-	const FSHAHash& GetOutputHash() const
-	{
-		return OutputHash;
-	}
+	virtual uint32 GetSizeBytes() const = 0;
+
+protected:
+	explicit FShaderMapResource(EShaderPlatform InPlatform, int32 NumShaders);
+	virtual ~FShaderMapResource();
 
-	uint32 GetSizeBytes() const
+	uint32 GetAllocatedSize() const
 	{
-		return Code.GetAllocatedSize() + sizeof(FShaderResource);
+		uint32 Size = RHIShaders.GetAllocatedSize();
+#if RHI_RAYTRACING
+		Size += RayTracingMaterialLibraryIndices.GetAllocatedSize();
+#endif
+		return Size;
 	}
 
-	// FRenderResource interface.
-	virtual void InitRHI();
-	virtual void ReleaseRHI();
+	void CreateShader(int32 ShaderIndex);
 
+	virtual TRefCountPtr<FRHIShader> CreateRHIShader(int32 ShaderIndex) = 0;
+	virtual bool TryRelease() { return true; }
 
-	template<class T>
-	static FShaderResource* FindShaderResourceById(T Id)
-	{
-		static_assert(sizeof(T) == 0, "Deprecated. Please save return value as TRefCountPtr<FShaderResource> and switch to FindById().");
-		return nullptr;
-	}
+private:
+	TArray<TRefCountPtr<FRHIShader>> RHIShaders;
+#if RHI_RAYTRACING
+	TArray<uint32> RayTracingMaterialLibraryIndices;
+#endif // RHI_RAYTRACING
 
-	/** Finds a matching shader resource in memory if possible. */
-	RENDERCORE_API static TRefCountPtr<FShaderResource> FindById(const FShaderResourceId& Id);
+	EShaderPlatform Platform;
 
-	/** Finds a matching shader resource in memory or clone the temporary and register the new resource */
-	RENDERCORE_API static TRefCountPtr<FShaderResource> FindOrClone(FShaderResource&& Temp);
+	/** The number of references to this shader. */
+	int32 NumRefs;
+};
 
-	//#todo-RemoveStreamOut
-	template<class T>
-	static FShaderResource* FindOrCreateShaderResource(const FShaderCompilerOutput& Output, T SpecificPermutationId)
+class FShaderMapResourceCode : public FThreadSafeRefCountedObject
+{
+public:
+	struct FShaderEntry
 	{
-		static_assert(sizeof(T) == 0, "Deprecated. Please save return value as TRefCountPtr<FShaderResource> and switch to FindById().");
-		return nullptr;
-	}
+		uint32 Offset;
+		uint32 UncompressedSize;
+		uint32 CompressedSize;
+		EShaderFrequency Frequency;
 
-	/** 
-	 * Finds a matching shader resource in memory or creates a new one with the given compiler output.  
-	 */
-	 //#todo-RemoveStreamOut
-	RENDERCORE_API static TRefCountPtr<FShaderResource> FindOrCreate(const FShaderCompilerOutput& Output, int32 SpecificPermutationId);
+		friend FArchive& operator<<(FArchive& Ar, FShaderEntry& Entry)
+		{
+			uint8 Freq = Entry.Frequency;
+			Ar << Entry.Offset << Entry.UncompressedSize << Entry.CompressedSize << Freq;
+			Entry.Frequency = (EShaderFrequency)Freq;
+			return Ar;
+		}
+	};
 
-	/** Return a list of all shader Ids currently known */
-	RENDERCORE_API static void GetAllShaderResourceId(TArray<FShaderResourceId>& Ids);
+	RENDERCORE_API ~FShaderMapResourceCode();
 
-	/** Returns true if and only if TargetPlatform is compatible for use with CurrentPlatform. */
-	RENDERCORE_API static bool ArePlatformsCompatible(EShaderPlatform CurrentPlatform, EShaderPlatform TargetPlatform);
+	RENDERCORE_API void Finalize();
 
-	void GetShaderCode(TArray<uint8>& OutCode) const
+	RENDERCORE_API void Serialize(FArchive& Ar);
+
+	inline const uint32 GetSizeBytes() const
 	{
-		UncompressCode(OutCode);
+		return sizeof(*this) + ShaderHashes.GetAllocatedSize() + ShaderEntries.GetAllocatedSize() + ShaderCode.GetAllocatedSize();
 	}
 
-	/**
-	* Passes back a zeroed out hash to serialize when saving out cooked data.
-	* The goal here is to ensure that source hash changes do not cause widespread binary differences in cooked data, resulting in bloated patch diffs.
-	*/
-	RENDERCORE_API static FSHAHash& FilterShaderSourceHashForSerialization(const FArchive& Ar, FSHAHash &HashToSerialize);
-
-private:
-	// compression functions
-	void UncompressCode(TArray<uint8>& UncompressedCode) const;
-	void CompressCode(const TArray<uint8>& UncompressedCode);
+	FSHAHash ResourceHash;
+	TArray<FSHAHash> ShaderHashes;
+	TArray<FShaderEntry> ShaderEntries;
+	TArray<uint8> ShaderCode;
+};
 	
+class RENDERCORE_API FShaderMapResource_InlineCode : public FShaderMapResource
+{
+public:
+	FShaderMapResource_InlineCode(EShaderPlatform InPlatform, FShaderMapResourceCode* InCode)
+		: FShaderMapResource(InPlatform, InCode->ShaderEntries.Num())
+		, Code(InCode)
+	{}
 
-	/** Conditionally serialize shader code. */
-	void SerializeShaderCode(FArchive& Ar);
-
-#if WITH_EDITORONLY_DATA
-	/** Conditionally serialize platform debug data. */
-	void SerializePlatformDebugData(FArchive& Ar);
-#endif
-
-	/**
-	 * Hash of the compiled bytecode and the generated parameter map.
-	 * This is used to find existing shader resources in memory or the DDC.
-	 */
-	FSHAHash OutputHash;
-
-	/** Compiled bytecode. */
-	TArray<uint8> Code;
-
-	/** Target platform and frequency. */
-	FShaderTarget Target;
-
-	/** Reference to the RHI shader. References the matching shader type of Target.Frequency. */
-	TRefCountPtr<FRHIShader> Shader;
-
-#if RHI_RAYTRACING
-	FRayTracingShaderRHIRef RayTracingShader;
-	uint32 RayTracingMaterialLibraryIndex = UINT_MAX;
-#endif // RHI_RAYTRACING
-
-#if WITH_EDITORONLY_DATA
-	/** Platform specific debug data output by the shader compiler. Discarded in cooked builds. */
-	TArray<uint8> PlatformDebugData;
-#endif
-
-	/** Original bytecode size, before compression */
-	uint32 UncompressedCodeSize = 0;
+	// FShaderMapResource interface
+	virtual TRefCountPtr<FRHIShader> CreateRHIShader(int32 ShaderIndex) override;
+	virtual uint32 GetSizeBytes() const override { return sizeof(*this) + GetAllocatedSize(); }
 
-	//#todo-RemoveStreamOut
-	/** Specific permutation identifier of the shader when SpecificType is non null, ignored otherwise. */
-	int32 SpecificPermutationId;
+	TRefCountPtr<FShaderMapResourceCode> Code;
+};
 
-	/** The number of references to this shader. */
-	mutable int32 NumRefs;
+class FShaderMapResourceBuilder
+{
+public:
+	RENDERCORE_API explicit FShaderMapResourceBuilder(FShaderMapResourceCode* InCode);
 
-	/** The number of instructions the shader takes to execute. */
-	uint32 NumInstructions;
+	inline int32 FindOrAddCode(const FShaderCompilerOutput& Output)
+	{
+		return FindOrAddCode(Output.Target.GetFrequency(), Output.OutputHash, Output.ShaderCode.GetReadAccess());
+	}
 
-#if WITH_EDITORONLY_DATA
-	/** Number of texture samplers the shader uses. */
-	uint32 NumTextureSamplers;
-#endif
+	RENDERCORE_API int32 FindCode(const FSHAHash& InHash, uint32 InKey) const;
+	RENDERCORE_API int32 FindOrAddCode(EShaderFrequency InFrequency, const FSHAHash& InHash, const TConstArrayView<uint8>& InCode);
 
-	FShaderParameterMapInfo ParameterMapInfo;
-	/** Whether the shader code is stored in a shader library. */
-	bool bCodeInSharedLocation;
-	/** Whether the shader code was requested (and hence if we need to drop the ref later). */
-	bool bCodeInSharedLocationRequested;
+private:
+	FHashTable ShaderHashTable;
+	TRefCountPtr<FShaderMapResourceCode> Code;
+};
 
-	/** Initialize the shader RHI resources. */
-	RENDERCORE_API void InitializeShaderRHI();
+class RENDERCORE_API FShaderMapPointerTable : public FPointerTableBase
+{
+public:
+	virtual int32 AddIndexedPointer(const FTypeLayoutDesc& TypeDesc, void* Ptr) override;
+	virtual void* GetIndexedPointer(const FTypeLayoutDesc& TypeDesc, uint32 i) const override;
 
-	void BuildParameterMapInfo(const TMap<FString, FParameterAllocation>& ParameterMap);
+	virtual void SaveToArchive(FArchive& Ar, void* FrozenContent, bool bInlineShaderResources) const;
+	virtual void LoadFromArchive(FArchive& Ar, void* FrozenContent, bool bInlineShaderResources, bool bLoadedByCookedMaterial);
 
-	/** Critical section for ShaderResourceIdMap. */
-	static FCriticalSection ShaderResourceIdMapCritical;
+	TPtrTable<FShaderType> ShaderTypes;
+	TPtrTable<FVertexFactoryType> VFTypes;
 };
 
 /** Encapsulates information about a shader's serialization behavior, used to detect when C++ serialization changes to auto-recompile. */
@@ -629,65 +480,41 @@ public:
 class FShaderId
 {
 public:
-	/** Shader type */
-	FShaderType* ShaderType;
-
-	/**
-	 * Vertex factory type that the shader was created for,
-	 * This is needed in the Id since a single shader type will be compiled for multiple vertex factories within a material shader map.
-	 * Will be NULL for global shaders.
-	 */
-	FVertexFactoryType* VertexFactoryType;
-
-	/** Shader Pipeline linked to this shader, needed since a single shader might be used on different Pipelines. */
-	const FShaderPipelineType* ShaderPipeline;
+	inline FShaderId() {}
+	inline FShaderId(const FShaderType* InType, const FSHAHash& InMaterialShaderMapHash, const FHashedName& InShaderPipeline, const FVertexFactoryType* InVertexFactoryType, int32 InPermutationId, EShaderPlatform InPlatform)
+		: Type(InType)
+		, VFType(InVertexFactoryType)
+		, ShaderPipelineName(InShaderPipeline)
+		, MaterialShaderMapHash(InMaterialShaderMapHash)
+		, PermutationId(InPermutationId)
+		, Platform(InPlatform)
+	{}
 
-	/** 
-	 * Hash of the material shader map Id, since this shader depends on the generated material code from that shader map.
-	 * A hash is used instead of the full shader map Id to shorten the key length, even though this will result in a hash being hashed when we make a DDC key. 
-	 */ 
+	const FShaderType* Type;
+	const FVertexFactoryType* VFType;
+	FHashedName ShaderPipelineName;
 	FSHAHash MaterialShaderMapHash;
-
-#if KEEP_SHADER_SOURCE_HASHES
-	/** Used to detect changes to the vertex factory source files. */
-	FSHAHash VFSourceHash;
-
-	/** Used to detect changes to the shader source files. */
-	FSHAHash SourceHash;
-#endif
-
-	/** Shader platform and frequency. */
-	FShaderTarget Target;
-
-	/** Unique permutation identifier within the ShaderType. */
 	int32 PermutationId;
-
-	/** Creates empty Id */
-	inline FShaderId() : ShaderType(nullptr), VertexFactoryType(nullptr), ShaderPipeline(nullptr), PermutationId(0) {}
-
-	/** Creates an Id for the given material, vertex factory, shader type and target. */
-	RENDERCORE_API FShaderId(const FSHAHash& InMaterialShaderMapHash, const FShaderPipelineType* InShaderPipeline, FVertexFactoryType* InVertexFactoryType, FShaderType* InShaderType, int32 PermutationId, FShaderTarget InTarget);
+	uint32 Platform : SP_NumBits;
 
 	friend inline uint32 GetTypeHash( const FShaderId& Id )
 	{
-		return
-			HashCombine(
-				HashCombine(*(uint32*)&Id.MaterialShaderMapHash, GetTypeHash(Id.Target)),
-				HashCombine(GetTypeHash(Id.VertexFactoryType), uint32(Id.PermutationId)));
+		return HashCombine(
+			GetTypeHash(Id.Type),
+			HashCombine(GetTypeHash(Id.VFType),
+				HashCombine(GetTypeHash(Id.ShaderPipelineName),
+					HashCombine(GetTypeHash(Id.MaterialShaderMapHash),
+						HashCombine(GetTypeHash(Id.PermutationId), GetTypeHash(Id.Platform))))));
 	}
 
 	friend bool operator==(const FShaderId& X, const FShaderId& Y)
 	{
-		return X.MaterialShaderMapHash == Y.MaterialShaderMapHash
-			&& X.ShaderPipeline == Y.ShaderPipeline
-			&& X.VertexFactoryType == Y.VertexFactoryType
-			&& X.ShaderType == Y.ShaderType
+		return X.Type == Y.Type
+			&& X.MaterialShaderMapHash == Y.MaterialShaderMapHash
+			&& X.ShaderPipelineName == Y.ShaderPipelineName
+			&& X.VFType == Y.VFType
 			&& X.PermutationId == Y.PermutationId 
-#if KEEP_SHADER_SOURCE_HASHES
-			&& X.SourceHash == Y.SourceHash 
-			&& X.VFSourceHash == Y.VFSourceHash
-#endif
-			&& X.Target == Y.Target;
+			&& X.Platform == Y.Platform;
 	}
 
 	friend bool operator!=(const FShaderId& X, const FShaderId& Y)
@@ -696,187 +523,137 @@ public:
 	}
 };
 
-/** Self contained version of FShaderId, which is useful for serializing. */
-class FSelfContainedShaderId
-{
-public:
-
-	/** 
-	 * Hash of the material shader map Id, since this shader depends on the generated material code from that shader map.
-	 * A hash is used instead of the full shader map Id to shorten the key length, even though this will result in a hash being hashed when we make a DDC key. 
-	 */ 
-	FSHAHash MaterialShaderMapHash;
-
-#if KEEP_SHADER_SOURCE_HASHES
-	/** Used to detect changes to the vertex factory source files. */
-	FSHAHash VFSourceHash;
-
-	/** Used to detect changes to the shader source files. */
-	FSHAHash SourceHash;
-#endif
-
-	/** 
-	 * Name of the vertex factory type that the shader was created for, 
-	 * This is needed in the Id since a single shader type will be compiled for multiple vertex factories within a material shader map.
-	 * Will be the empty string for global shaders.
-	 */
-	FString VertexFactoryTypeName;
-
-	// Required to differentiate amongst unique shaders in the global map per Type
-	FString ShaderPipelineName;
-
-	/** Used to detect changes to the vertex factory parameter class serialization, or empty for global shaders. */
-	FSerializationHistory VFSerializationHistory;
-
-	/** Shader type name */
-	FString ShaderTypeName;
-
-	/** Unique permutation identifier within the ShaderType. */
-	int32 PermutationId;
-
-	/** Used to detect changes to the shader serialization. */
-	FSerializationHistory SerializationHistory;
-
-	/** Shader platform and frequency. */
-	FShaderTarget Target;
-
-	RENDERCORE_API FSelfContainedShaderId();
-
-	RENDERCORE_API FSelfContainedShaderId(const FShaderId& InShaderId);
-
-	RENDERCORE_API bool IsValid();
-
-	RENDERCORE_API friend FArchive& operator<<(FArchive& Ar,class FSelfContainedShaderId& Ref);
-};
-
 class FMaterial;
 
-
 /**
  * Stores all shader parameter bindings and their corresponding offset and size in the shader's parameters struct.
  */
 class RENDERCORE_API FShaderParameterBindings
 {
 public:
-	static constexpr uint16 kInvalidBufferIndex = 0xFFFF;
-
 	struct FParameter
 	{
-		uint16 BufferIndex;
-		uint16 BaseIndex;
-		uint16 ByteOffset;
-		uint16 ByteSize;
+		DECLARE_INLINE_TYPE_LAYOUT(FParameter, NonVirtual);
 
-		friend FArchive& operator<<(FArchive& Ar, FParameter& ParameterBindingData)
-		{
-			Ar << ParameterBindingData.BufferIndex << ParameterBindingData.BaseIndex << ParameterBindingData.ByteOffset << ParameterBindingData.ByteSize;
-			return Ar;
-		}
+		LAYOUT_FIELD(uint16, BufferIndex);
+		LAYOUT_FIELD(uint16, BaseIndex);
+		LAYOUT_FIELD(uint16, ByteOffset);
+		LAYOUT_FIELD(uint16, ByteSize);
 	};
 
 	struct FResourceParameter
 	{
-		uint16 BaseIndex;
-		uint16 ByteOffset;
+		DECLARE_INLINE_TYPE_LAYOUT(FResourceParameter, NonVirtual);
+
+		LAYOUT_FIELD(uint16, BaseIndex);
+		LAYOUT_FIELD(uint16, ByteOffset);
 
-		friend FArchive& operator<<(FArchive& Ar, FResourceParameter& ParameterBindingData)
-		{
-			Ar << ParameterBindingData.BaseIndex << ParameterBindingData.ByteOffset;
-			return Ar;
-		}
 	};
 
 	struct FParameterStructReference
 	{
-		uint16 BufferIndex;
-		uint16 ByteOffset;
+		DECLARE_INLINE_TYPE_LAYOUT(FParameterStructReference, NonVirtual);
+
+		LAYOUT_FIELD(uint16, BufferIndex);
+		LAYOUT_FIELD(uint16, ByteOffset);
 
-		friend FArchive& operator<<(FArchive& Ar, FParameterStructReference& ParameterBindingData)
-		{
-			Ar << ParameterBindingData.BufferIndex << ParameterBindingData.ByteOffset;
-			return Ar;
-		}
 	};
 
-	TArray<FParameter> Parameters;
-	TArray<FResourceParameter> Textures;
-	TArray<FResourceParameter> SRVs;
-	TArray<FResourceParameter> UAVs;
-	TArray<FResourceParameter> Samplers;
-	TArray<FResourceParameter> GraphTextures;
-	TArray<FResourceParameter> GraphSRVs;
-	TArray<FResourceParameter> GraphUAVs;
-	TArray<FParameterStructReference> ParameterReferences;
+	DECLARE_TYPE_LAYOUT(FShaderParameterBindings, NonVirtual);
+public:
+	static constexpr uint16 kInvalidBufferIndex = 0xFFFF;
+
+	void BindForLegacyShaderParameters(const FShader* Shader, const FShaderParameterMap& ParameterMaps, const FShaderParametersMetadata& StructMetaData, bool bShouldBindEverything = false);
+	void BindForRootShaderParameters(const FShader* Shader, const FShaderParameterMap& ParameterMaps);
+
+	LAYOUT_FIELD(TMemoryImageArray<FParameter>, Parameters);
+	LAYOUT_FIELD(TMemoryImageArray<FResourceParameter>, Textures);
+	LAYOUT_FIELD(TMemoryImageArray<FResourceParameter>, SRVs);
+	LAYOUT_FIELD(TMemoryImageArray<FResourceParameter>, UAVs);
+	LAYOUT_FIELD(TMemoryImageArray<FResourceParameter>, Samplers);
+	LAYOUT_FIELD(TMemoryImageArray<FResourceParameter>, GraphTextures);
+	LAYOUT_FIELD(TMemoryImageArray<FResourceParameter>, GraphSRVs);
+	LAYOUT_FIELD(TMemoryImageArray<FResourceParameter>, GraphUAVs);
+	LAYOUT_FIELD(TMemoryImageArray<FParameterStructReference>, ParameterReferences);
 
 	// Hash of the shader parameter structure when doing the binding.
-	uint32 StructureLayoutHash = 0;
+	LAYOUT_FIELD_INITIALIZED(uint32, StructureLayoutHash, 0);
 
 	// Buffer index of FShaderParametersMetadata::kRootUniformBufferBindingName
-	uint16 RootParameterBufferIndex = FShaderParameterBindings::kInvalidBufferIndex;
-
-	friend FArchive& operator<<(FArchive& Ar, FShaderParameterBindings& ParametersBindingData)
-	{
-		Ar << ParametersBindingData.Parameters;
-		Ar << ParametersBindingData.Textures;
-		Ar << ParametersBindingData.SRVs;
-		Ar << ParametersBindingData.UAVs;
-		Ar << ParametersBindingData.Samplers;
-		Ar << ParametersBindingData.GraphTextures;
-		Ar << ParametersBindingData.GraphSRVs;
-		Ar << ParametersBindingData.GraphUAVs;
-		Ar << ParametersBindingData.ParameterReferences;
-		Ar << ParametersBindingData.StructureLayoutHash;
-		Ar << ParametersBindingData.RootParameterBufferIndex;
-		return Ar;
-	}
+	LAYOUT_FIELD_INITIALIZED(uint16, RootParameterBufferIndex, FShaderParameterBindings::kInvalidBufferIndex);
 
-	void BindForLegacyShaderParameters(const FShader* Shader, const FShaderParameterMap& ParameterMaps, const FShaderParametersMetadata& StructMetaData, bool bShouldBindEverything = false);
-	void BindForRootShaderParameters(const FShader* Shader, const FShaderParameterMap& ParameterMaps);
 }; // FShaderParameterBindings
 
-
-/** A compiled shader and its parameter bindings. */
-class RENDERCORE_API FShader : public FDeferredCleanupInterface
+struct FShaderPermutationParameters
 {
-	friend class FShaderType;
-public:
+	// Shader platform to compile to.
+	const EShaderPlatform Platform;
 
-	struct CompiledShaderInitializerType
-	{
-		FShaderType* Type;
-		FShaderTarget Target;
-		const TArray<uint8>& Code;
-		const FShaderParameterMap& ParameterMap;
-		const FSHAHash& OutputHash;
-		TRefCountPtr<FShaderResource> Resource;
-		FSHAHash MaterialShaderMapHash;
-		const FShaderPipelineType* ShaderPipeline;
-		FVertexFactoryType* VertexFactoryType;
-		int32 PermutationId;
-
-		CompiledShaderInitializerType(
-			FShaderType* InType,
-			int32 InPermutationId,
-			const FShaderCompilerOutput& CompilerOutput,
-			TRefCountPtr<FShaderResource>&& InResource,
-			const FSHAHash& InMaterialShaderMapHash,
-			const FShaderPipelineType* InShaderPipeline,
-			FVertexFactoryType* InVertexFactoryType
-			):
-			Type(InType),
-			Target(CompilerOutput.Target),
-			Code(CompilerOutput.ShaderCode.GetReadAccess()),
-			ParameterMap(CompilerOutput.ParameterMap),
-			OutputHash(CompilerOutput.OutputHash),
-			Resource(InResource),
-			MaterialShaderMapHash(InMaterialShaderMapHash),
-			ShaderPipeline(InShaderPipeline),
-			VertexFactoryType(InVertexFactoryType),
-			PermutationId(InPermutationId)
-		{}
-	};
+	/** Unique permutation identifier of the material shader type. */
+	const int32 PermutationId;
 
-	/** 
+	explicit FShaderPermutationParameters(EShaderPlatform InPlatform, int32 InPermutationId = 0)
+		: Platform(InPlatform)
+		, PermutationId(InPermutationId)
+	{}
+};
+
+struct FShadereCompiledShaderInitializerType
+{
+	FShaderType* Type;
+	FShaderTarget Target;
+	const TArray<uint8>& Code;
+	const FShaderParameterMap& ParameterMap;
+	const FSHAHash& OutputHash;
+	FSHAHash MaterialShaderMapHash;
+	const FShaderPipelineType* ShaderPipeline;
+	FVertexFactoryType* VertexFactoryType;
+	int32 ResourceIndex;
+	uint32 NumInstructions;
+	uint32 NumTextureSamplers;
+	uint32 CodeSize;
+	int32 PermutationId;
+
+	FShadereCompiledShaderInitializerType(
+		FShaderType* InType,
+		int32 InPermutationId,
+		const FShaderCompilerOutput& CompilerOutput,
+		int32 InResourceIndex,
+		const FSHAHash& InMaterialShaderMapHash,
+		const FShaderPipelineType* InShaderPipeline,
+		FVertexFactoryType* InVertexFactoryType
+	) :
+		Type(InType),
+		Target(CompilerOutput.Target),
+		Code(CompilerOutput.ShaderCode.GetReadAccess()),
+		ParameterMap(CompilerOutput.ParameterMap),
+		OutputHash(CompilerOutput.OutputHash),
+		MaterialShaderMapHash(InMaterialShaderMapHash),
+		ShaderPipeline(InShaderPipeline),
+		VertexFactoryType(InVertexFactoryType),
+		ResourceIndex(InResourceIndex),
+		NumInstructions(CompilerOutput.NumInstructions),
+		NumTextureSamplers(CompilerOutput.NumTextureSamplers),
+		CodeSize(CompilerOutput.ShaderCode.GetShaderCodeSize()),
+		PermutationId(InPermutationId)
+	{}
+};
+
+DECLARE_EXPORTED_TEMPLATE_INTRINSIC_TYPE_LAYOUT(template<>, TIndexedPtr<FShaderType>, RENDERCORE_API);
+DECLARE_EXPORTED_TEMPLATE_INTRINSIC_TYPE_LAYOUT(template<>, TIndexedPtr<FVertexFactoryType>, RENDERCORE_API);
+
+/** A compiled shader and its parameter bindings. */
+class RENDERCORE_API FShader
+{
+	friend class FShaderType;
+	DECLARE_TYPE_LAYOUT(FShader, NonVirtual);
+public:
+	using FPermutationDomain = FShaderPermutationNone;
+	using FPermutationParameters = FShaderPermutationParameters;
+	using CompiledShaderInitializerType = FShadereCompiledShaderInitializerType;
+	using ShaderMetaType = FShaderType;
+
+	/** 
 	 * Used to construct a shader for deserialization.
 	 * This still needs to initialize members to safe values since FShaderType::GenerateSerializationHistory uses this constructor.
 	 */
@@ -887,138 +664,48 @@ public:
 	 */
 	FShader(const CompiledShaderInitializerType& Initializer);
 
-	virtual ~FShader();
-
-/*
-	/ ** Can be overridden by FShader subclasses to modify their compile environment just before compilation occurs. * /
-	static void ModifyCompilationEnvironment(EShaderPlatform Platform, const FMaterial* Material, FShaderCompilerEnvironment& OutEnvironment){}
-*/
+	~FShader();
 
-	/** Serializes the shader. */
-	bool SerializeBase(FArchive& Ar, bool bShadersInline, bool bLoadedByCookedMaterial);
+	/** Can be overridden by FShader subclasses to modify their compile environment just before compilation occurs. */
+	static void ModifyCompilationEnvironment(const FShaderPermutationParameters&, FShaderCompilerEnvironment&) {}
 
-	virtual bool Serialize(FArchive& Ar) { return false; }
+	/** Can be overridden by FShader subclasses to determine whether a specific permutation should be compiled. */
+	static bool ShouldCompilePermutation(const FShaderPermutationParameters&) { return true; }
 
-	// Reference counting.
-	void AddRef();
-	void Release();
-
-	/** Registers this shader for lookup by ID. */
-	void Register(bool bLoadedByCookedMaterial);
-
-	/** Removes this shader from the ID lookup map. */
-	void Deregister();
+	/** Can be overridden by FShader subclasses to determine whether compilation is valid. */
+	static bool ValidateCompiledResult(EShaderPlatform InPlatform, const FShaderParameterMap& InParameterMap, TArray<FString>& OutError) { return true; }
 
 	/** Returns the hash of the shader file that this shader was compiled with. */
 	const FSHAHash& GetHash() const;
 
-	/** Returns the shader platform of the shader file that this shader was compiled with. */
-	EShaderPlatform GetShaderPlatform() const;
+	const FSHAHash& GetVertexFactoryHash() const;
 	
-	/** @return If the shader is linked with a vertex factory, returns the vertex factory's parameter object. */
-	virtual const FVertexFactoryParameterRef* GetVertexFactoryParameterRef() const { return NULL; }
-
-	/** @return the shader's vertex shader */
-	inline FRHIVertexShader* GetVertexShader() const
-	{
-		return Resource->GetVertexShader();
-	}
-	/** @return the shader's pixel shader */
-	inline FRHIPixelShader* GetPixelShader() const
-	{
-		return Resource->GetPixelShader();
-	}
-	/** @return the shader's hull shader */
-	inline FRHIHullShader* GetHullShader() const
-	{
-		return Resource->GetHullShader();
-	}
-	/** @return the shader's domain shader */
-	inline FRHIDomainShader* GetDomainShader() const
-	{
-		return Resource->GetDomainShader();
-	}
-	/** @return the shader's geometry shader */
-	inline FRHIGeometryShader* GetGeometryShader() const
-	{
-		return Resource->GetGeometryShader();
-	}
-	/** @return the shader's compute shader */
-	inline FRHIComputeShader* GetComputeShader() const
-	{
-		return Resource->GetComputeShader();
-	}
-
-#if RHI_RAYTRACING
-	inline FRHIRayTracingShader* GetRayTracingShader() const
-	{
-		return Resource->GetRayTracingShader();
-	}
-
-	inline uint32 GetRayTracingMaterialLibraryIndex() const
-	{
-		return Resource->GetRayTracingMaterialLibraryIndex();
-	}
-#endif // RHI_RAYTRACING
+	const FSHAHash& GetOutputHash() const;
 
 	// Accessors.
-	inline FShaderType* GetType() const { return Type; }
+	inline FShaderType* GetType(const FShaderMapPointerTable& InPointerTable) const { return Type.Get(InPointerTable.ShaderTypes); }
+	inline FShaderType* GetType(const FPointerTableBase* InPointerTable) const { return Type.Get(InPointerTable); }
+	inline FVertexFactoryType* GetVertexFactoryType(const FShaderMapPointerTable& InPointerTable) const { return VFType.Get(InPointerTable.VFTypes); }
+	inline FVertexFactoryType* GetVertexFactoryType(const FPointerTableBase* InPointerTable) const { return VFType.Get(InPointerTable); }
+	inline FShaderType* GetTypeUnfrozen() const { return Type.GetUnfrozen(); }
+	inline const FHashedName& GetTypeName() const { return TypeName; }
+	inline int32 GetResourceIndex() const { return ResourceIndex; }
 	inline int32 GetPermutationId() const { return PermutationId; }
-	inline uint32 GetNumInstructions() const { return Resource->NumInstructions; }
-	inline void SetNumInstructions(uint32 Num) { Resource->NumInstructions = Num; }
-#if WITH_EDITOR
-	inline uint32 GetNumTextureSamplers() const { return Resource->NumTextureSamplers; }
-#endif
-	inline const TArray<uint8>& GetCode() const { return Resource->Code; }
+	inline EShaderPlatform GetShaderPlatform() const { return Target.GetPlatform(); }
+	inline EShaderFrequency GetFrequency() const { return Target.GetFrequency(); }
 	inline const FShaderTarget GetTarget() const { return Target; }
-	inline FSHAHash GetOutputHash() const
-	{
-#if KEEP_SHADER_SOURCE_HASHES
-		return OutputHash;
+	inline bool IsFrozen() const { return Type.IsFrozen(); }
+
+#if WITH_EDITORONLY_DATA
+	inline uint32 GetNumInstructions() const { return NumInstructions; }
+	inline uint32 GetNumTextureSamplers() const { return NumTextureSamplers; }
+	inline uint32 GetCodeSize() const { return CodeSize; }
+	inline void SetNumInstructions(uint32 Value) { NumInstructions = Value; }
 #else
-		check(Resource);
-		return Resource->OutputHash;
+	inline uint32 GetNumInstructions() const { return 0u; }
+	inline uint32 GetNumTextureSamplers() const { return 0u; }
+	inline uint32 GetCodeSize() const { return 0u; }
 #endif
-	}
-	FShaderId GetId() const;
-	inline FShaderKey GetKey() const { return FShaderKey(MaterialShaderMapHash, ShaderPipeline, VFType, PermutationId, Target.GetPlatform()); }
-	inline FVertexFactoryType* GetVertexFactoryType() const { return VFType; }
-	inline int32 GetNumRefs() const { return NumRefs; }
-	const FShaderParameterMapInfo& GetParameterMapInfo() const { return Resource->ParameterMapInfo; }
-
-	inline FShaderResourceId GetResourceId() const
-	{
-		return Resource->GetId();
-	}
-
-	inline uint32 GetSizeBytes() const
-	{
-		return GetTypeSize() + GetAllocatedSize();
-	}
-	
-	/** Returns the size of the concrete type of this shader. */
-	virtual uint32 GetTypeSize() const
-	{
-		return sizeof(*this);
-	}
-
-	/** Returns the size of all allocations owned by this shader, e.g. TArrays. */
-	virtual uint32 GetAllocatedSize() const
-	{
-		return UniformBufferParameters.GetAllocatedSize() + UniformBufferParameterStructs.GetAllocatedSize();
-	}
-
-	uint32 GetResourceSizeBytes() const
-	{
-		return Resource->GetSizeBytes();
-	}
-
-	void SetResource(FShaderResource* InResource);
-
-	void BeginInitializeResources()
-	{
-		BeginInitResource(Resource);
-	}
 
 	/** Finds an automatically bound uniform buffer matching the given uniform buffer type if one exists, or returns an unbound parameter. */
 	template<typename UniformBufferStructType>
@@ -1031,10 +718,11 @@ public:
 	/** Finds an automatically bound uniform buffer matching the given uniform buffer struct if one exists, or returns an unbound parameter. */
 	FORCEINLINE_DEBUGGABLE const FShaderUniformBufferParameter& GetUniformBufferParameter(const FShaderParametersMetadata* SearchStruct) const
 	{
+		const FHashedName SearchName = SearchStruct->GetShaderVariableHashedName();
 		int32 FoundIndex = INDEX_NONE;
 		for (int32 StructIndex = 0, Count = UniformBufferParameterStructs.Num(); StructIndex < Count; StructIndex++)
 		{
-			if (UniformBufferParameterStructs[StructIndex] == SearchStruct)
+			if (UniformBufferParameterStructs[StructIndex] == SearchName)
 			{
 				FoundIndex = StructIndex;
 				break;
@@ -1052,32 +740,17 @@ public:
 			// There's no good way to distinguish not being bound due to temporary debugging / compiler optimizations or an actual code bug,
 			// Hence failing silently instead of an error message
 			static FShaderUniformBufferParameter UnboundParameter;
-			UnboundParameter.SetInitialized();
 			return UnboundParameter;
 		}
 	}
 
-	const FShaderParametersMetadata* FindAutomaticallyBoundUniformBufferStruct(int32 BaseIndex) const
-	{
-		for (int32 i = 0; i < UniformBufferParameters.Num(); i++)
-		{
-			if (UniformBufferParameters[i].GetBaseIndex() == BaseIndex)
-			{
-				return UniformBufferParameterStructs[i];
-			}
-		}
+	const FShaderParametersMetadata* FindAutomaticallyBoundUniformBufferStruct(int32 BaseIndex) const;
 
-		return nullptr;
-	}
-
-	/** Gets the shader. */
-	inline FShader* GetShader()
-	{
-		return this;
-	}
+	void DumpDebugInfo(const FShaderMapPointerTable& InPtrTable);
 
-	void DumpDebugInfo();
-	void SaveShaderStableKeys(EShaderPlatform TargetShaderPlatform, const struct FStableShaderKeyAndValue& SaveKeyVal);
+#if WITH_EDITOR
+	void SaveShaderStableKeys(const FShaderMapPointerTable& InPtrTable, EShaderPlatform TargetShaderPlatform, const struct FStableShaderKeyAndValue& SaveKeyVal);
+#endif // WITH_EDITOR
 
 	/** Returns the meta data for the root shader parameter struct. */
 	static inline const FShaderParametersMetadata* GetRootParametersMetadata()
@@ -1085,56 +758,200 @@ public:
 		return nullptr;
 	}
 
-protected:
+private:
+	void BuildParameterMapInfo(const TMap<FString, FParameterAllocation>& ParameterMap);
 
-	/** Indexed the same as UniformBufferParameters.  Packed densely for coherent traversal. */
-	TArray<const FShaderParametersMetadata*> UniformBufferParameterStructs;
-	TArray<FShaderUniformBufferParameter> UniformBufferParameters;
+public:
+	LAYOUT_FIELD(FShaderParameterBindings, Bindings);
+	LAYOUT_FIELD(FShaderParameterMapInfo, ParameterMapInfo);
 
-private:
-	/** Hash of the material shader map this shader belongs to, stored so that an FShaderId can be constructed from this shader. */
-	FSHAHash MaterialShaderMapHash;
+protected:
+	LAYOUT_FIELD(TMemoryImageArray<FHashedName>, UniformBufferParameterStructs);
+	LAYOUT_FIELD(TMemoryImageArray<FShaderUniformBufferParameter>, UniformBufferParameters);
 
-#if KEEP_SHADER_SOURCE_HASHES
 	/**
 	* Hash of the compiled output from this shader and the resulting parameter map.
 	* This is used to find a matching resource.
 	*/
-	FSHAHash OutputHash;
+	LAYOUT_FIELD_EDITORONLY(FSHAHash, OutputHash);
 
 	/** Vertex factory source hash, stored so that an FShaderId can be constructed from this shader. */
-	FSHAHash VFSourceHash;
+	LAYOUT_FIELD_EDITORONLY(FSHAHash, VFSourceHash);
 
 	/** Hash of this shader's source files generated at compile time, and stored to allow creating an FShaderId. */
-	FSHAHash SourceHash;
-#endif
-
-	/** Reference to the shader resource, which stores the compiled bytecode and the RHI shader resource. */
-	TRefCountPtr<FShaderResource> Resource;
+	LAYOUT_FIELD_EDITORONLY(FSHAHash, SourceHash);
 
-	/** Shader pipeline this shader belongs to, stored so that an FShaderId can be constructed from this shader. */
-	const FShaderPipelineType* ShaderPipeline;
+private:
+	LAYOUT_FIELD(TIndexedPtr<FShaderType>, Type);
 
-	/** Vertex factory type this shader was created for, stored so that an FShaderId can be constructed from this shader. */
-	FVertexFactoryType* VFType;
+	LAYOUT_FIELD(TIndexedPtr<FVertexFactoryType>, VFType);
 
 	/** Shader Type metadata for this shader. */
-	FShaderType* Type;
+	LAYOUT_FIELD(FHashedName, TypeName);
+
+	/** Target platform and frequency. */
+	LAYOUT_FIELD(FShaderTarget, Target);
+
+	/** Index of this shader within the FShaderMapResource */
+	LAYOUT_FIELD(int32, ResourceIndex);
 
 	/** Unique permutation identifier of the shader in the shader type. */
-	int32 PermutationId;
+	LAYOUT_FIELD(int32, PermutationId);
 
-	/** Target platform and frequency. */
-	FShaderTarget Target;
+	/** The number of instructions the shader takes to execute. */
+	LAYOUT_FIELD_EDITORONLY(uint32, NumInstructions);
 
-	/** The number of references to this shader. */
-	mutable uint32 NumRefs;
+	/** Number of texture samplers the shader uses. */
+	LAYOUT_FIELD_EDITORONLY(uint32, NumTextureSamplers);
+
+	/** Size of shader's compiled code */
+	LAYOUT_FIELD_EDITORONLY(uint32, CodeSize);
+};
+
+RENDERCORE_API const FTypeLayoutDesc& GetTypeLayoutDesc(const FPointerTableBase* PtrTable, const FShader& Shader);
 
+template<typename ShaderType, typename PointerTableType>
+class TShaderRefBase
+{
 public:
-	/** Shader parameter bindings. */
-	FShaderParameterBindings Bindings;
+	TShaderRefBase() : ShaderContent(nullptr), ShaderMap(nullptr) {}
+	TShaderRefBase(ShaderType* InShader, const FShaderMapBase& InShaderMap) : ShaderContent(InShader), ShaderMap(&InShaderMap) {}
+	TShaderRefBase(const TShaderRefBase&) = default;
+
+	template<typename OtherShaderType, typename OtherPointerTableType>
+	TShaderRefBase(const TShaderRefBase<OtherShaderType, OtherPointerTableType>& Rhs) : ShaderContent(Rhs.GetShader()), ShaderMap(Rhs.GetShaderMap()) {}
+
+	TShaderRefBase& operator=(const TShaderRefBase&) = default;
+
+	template<typename OtherShaderType, typename OtherPointerTableType>
+	TShaderRefBase& operator=(const TShaderRefBase<OtherShaderType, OtherPointerTableType>& Rhs)
+	{
+		ShaderContent = Rhs.GetShader();
+		ShaderMap = Rhs.GetShaderMap();
+		return *this;
+	}
+
+	template<typename OtherShaderType, typename OtherPointerTableType>
+	static TShaderRefBase<ShaderType, PointerTableType> Cast(const TShaderRefBase<OtherShaderType, OtherPointerTableType>& Rhs)
+	{
+		return TShaderRefBase<ShaderType, PointerTableType>(static_cast<ShaderType*>(Rhs.GetShader()), Rhs.GetShaderMapChecked());
+	}
+
+	template<typename OtherShaderType, typename OtherPointerTableType>
+	static TShaderRefBase<ShaderType, PointerTableType> ReinterpretCast(const TShaderRefBase<OtherShaderType, OtherPointerTableType>& Rhs)
+	{
+		return TShaderRefBase<ShaderType, PointerTableType>(reinterpret_cast<ShaderType*>(Rhs.GetShader()), Rhs.GetShaderMapChecked());
+	}
+
+	inline bool IsValid() const { return ShaderContent != nullptr; }
+	inline bool IsNull() const { return ShaderContent == nullptr; }
+
+	inline void Reset() { ShaderContent = nullptr; ShaderMap = nullptr; }
+
+	inline ShaderType* GetShader() const { return ShaderContent; }
+	inline const FShaderMapBase* GetShaderMap() const { return ShaderMap; }
+	inline const FShaderMapBase& GetShaderMapChecked() const { check(ShaderMap); return *ShaderMap; }
+	inline FShaderType* GetType() const { return ShaderContent->GetType(GetPointerTable()); }
+	inline FVertexFactoryType* GetVertexFactoryType() const { return ShaderContent->GetVertexFactoryType(GetPointerTable()); }
+	inline FShaderMapResource& GetResourceChecked() const { FShaderMapResource* Resource = GetResource(); check(Resource); return *Resource; }
+	const PointerTableType& GetPointerTable() const;
+	FShaderMapResource* GetResource() const;
+
+	inline ShaderType* operator->() const { return ShaderContent; }
+
+	inline FRHIShader* GetRHIShaderBase(EShaderFrequency Frequency) const
+	{
+		FRHIShader* RHIShader = nullptr;
+		if(ShaderContent)
+		{
+			checkSlow(ShaderContent->GetFrequency() == Frequency);
+			RHIShader = GetResourceChecked().GetShader(ShaderContent->GetResourceIndex());
+			checkSlow(RHIShader->GetFrequency() == Frequency);
+		}
+		return RHIShader;
+	}
+
+	/** @return the shader's vertex shader */
+	inline FRHIVertexShader* GetVertexShader() const
+	{
+		return static_cast<FRHIVertexShader*>(GetRHIShaderBase(SF_Vertex));
+	}
+	/** @return the shader's pixel shader */
+	inline FRHIPixelShader* GetPixelShader() const
+	{
+		return static_cast<FRHIPixelShader*>(GetRHIShaderBase(SF_Pixel));
+	}
+	/** @return the shader's hull shader */
+	inline FRHIHullShader* GetHullShader() const
+	{
+		return static_cast<FRHIHullShader*>(GetRHIShaderBase(SF_Hull));
+	}
+	/** @return the shader's domain shader */
+	inline FRHIDomainShader* GetDomainShader() const
+	{
+		return static_cast<FRHIDomainShader*>(GetRHIShaderBase(SF_Domain));
+	}
+	/** @return the shader's geometry shader */
+	inline FRHIGeometryShader* GetGeometryShader() const
+	{
+		return static_cast<FRHIGeometryShader*>(GetRHIShaderBase(SF_Geometry));
+	}
+	/** @return the shader's compute shader */
+	inline FRHIComputeShader* GetComputeShader() const
+	{
+		return static_cast<FRHIComputeShader*>(GetRHIShaderBase(SF_Compute));
+	}
+
+#if RHI_RAYTRACING
+	inline FRHIRayTracingShader* GetRayTracingShader() const
+	{
+		FRHIRayTracingShader* RHIShader = nullptr;
+		if(ShaderContent)
+		{
+			const EShaderFrequency Frequency = ShaderContent->GetFrequency();
+			checkSlow(Frequency == SF_RayGen
+				|| Frequency == SF_RayMiss
+				|| Frequency == SF_RayHitGroup
+				|| Frequency == SF_RayCallable);
+			RHIShader = static_cast<FRHIRayTracingShader*>(GetResourceChecked().GetShader(ShaderContent->GetResourceIndex()));
+			checkSlow(RHIShader->GetFrequency() == Frequency);
+		}
+		return RHIShader;
+	}
+
+	inline uint32 GetRayTracingMaterialLibraryIndex() const
+	{
+		checkSlow(ShaderContent);
+		checkSlow(ShaderContent->GetFrequency() == SF_RayHitGroup);
+		return GetResourceChecked().GetRayTracingMaterialLibraryIndex(ShaderContent->GetResourceIndex());
+	}
+#endif // RHI_RAYTRACING
+
+private:
+	ShaderType* ShaderContent;
+	const FShaderMapBase* ShaderMap;
 };
 
+template<typename ShaderType0, typename ShaderType1, typename PointerTableType>
+inline bool operator==(const TShaderRefBase<ShaderType0, PointerTableType>& Lhs, const TShaderRefBase<ShaderType1, PointerTableType>& Rhs)
+{
+	if (Lhs.GetShader() == Rhs.GetShader())
+	{
+		check(Lhs.GetShaderMap() == Rhs.GetShaderMap());
+		return true;
+	}
+	return false;
+}
+
+template<typename ShaderType0, typename ShaderType1, typename PointerTableType>
+inline bool operator!=(const TShaderRefBase<ShaderType0, PointerTableType>& Lhs, const TShaderRefBase<ShaderType1, PointerTableType>& Rhs)
+{
+	return !operator==(Lhs, Rhs);
+}
+
+template<typename ShaderType>
+using TShaderRef = TShaderRefBase<ShaderType, FShaderMapPointerTable>;
+
 /**
  * An object which is used to serialize/deserialize, compile, and cache a particular shader class.
  *
@@ -1150,25 +967,26 @@ public:
 		Material,
 		MeshMaterial,
 		Niagara,
-		OCIO
+		OCIO,
+		NumShaderTypes,
 	};
 
 	typedef class FShader* (*ConstructSerializedType)();
+	typedef FShader* (*ConstructCompiledType)(const FShader::CompiledShaderInitializerType& Initializer);
+	typedef bool (*ShouldCompilePermutationType)(const FShaderPermutationParameters&);
+	typedef void (*ModifyCompilationEnvironmentType)(const FShaderPermutationParameters&, FShaderCompilerEnvironment&);
+	typedef bool (*ValidateCompiledResultType)(EShaderPlatform, const FShaderParameterMap&, TArray<FString>&);
 
 	/** @return The global shader factory list. */
 	static TLinkedList<FShaderType*>*& GetTypeList();
 
 	static FShaderType* GetShaderTypeByName(const TCHAR* Name);
-	static TArray<FShaderType*> GetShaderTypesByFilename(const TCHAR* Filename);
+	static TArray<const FShaderType*> GetShaderTypesByFilename(const TCHAR* Filename);
 
 	/** @return The global shader name to type map */
-	static TMap<FName, FShaderType*>& GetNameToTypeMap();
-
-	/** Gets a list of FShaderTypes whose source file no longer matches what that type was compiled with */
-	static void GetOutdatedTypes(TArray<FShaderType*>& OutdatedShaderTypes, TArray<const FVertexFactoryType*>& OutdatedFactoryTypes);
+	static TMap<FHashedName, FShaderType*>& GetNameToTypeMap();
 
-	/** Returns true if the source file no longer matches what that type was compiled with */
-	bool GetOutdatedCurrentType(TArray<FShaderType*>& OutdatedShaderTypes, TArray<const FVertexFactoryType*>& OutdatedFactoryTypes) const;
+	static const TArray<FShaderType*>& GetSortedTypes(EShaderTypeForDynamicCast Type);
 
 	/** Initialize FShaderType static members, this must be called before any shader types are created. */
 	static void Initialize(const TMap<FString, TArray<const TCHAR*> >& ShaderFileToUniformBufferVariables);
@@ -1179,24 +997,36 @@ public:
 	/** Minimal initialization constructor. */
 	FShaderType(
 		EShaderTypeForDynamicCast InShaderTypeForDynamicCast,
+		FTypeLayoutDesc& InTypeLayout,
 		const TCHAR* InName,
 		const TCHAR* InSourceFilename,
 		const TCHAR* InFunctionName,
 		uint32 InFrequency,
 		int32 TotalPermutationCount,
 		ConstructSerializedType InConstructSerializedRef,
+		ConstructCompiledType InConstructCompiledRef,
+		ModifyCompilationEnvironmentType InModifyCompilationEnvironmentRef,
+		ShouldCompilePermutationType InShouldCompilePermutationRef,
+		ValidateCompiledResultType InValidateCompiledResultRef,
+		uint32 InTypeSize,
 		const FShaderParametersMetadata* InRootParametersMetadata);
 
 	virtual ~FShaderType();
 
-	/**
-	 * Finds a shader of this type by Key.
-	 * @return NULL if no shader with the specified Key was found.
-	 */
-	FShader* FindShaderByKey(const FShaderKey& Key);
-
 	/** Constructs a new instance of the shader type for deserialization. */
 	FShader* ConstructForDeserialization() const;
+	FShader* ConstructCompiled(const FShader::CompiledShaderInitializerType& Initializer) const;
+
+	bool ShouldCompilePermutation(const FShaderPermutationParameters& Parameters) const;
+	void ModifyCompilationEnvironment(const FShaderPermutationParameters& Parameters, FShaderCompilerEnvironment& OutEnvironment) const;
+
+	/**
+	* Checks if the shader type should pass compilation for a particular set of parameters.
+	* @param Platform - Platform to validate.
+	* @param ParameterMap - Shader parameters to validate.
+	* @param OutError - List for appending validation errors.
+	*/
+	bool ValidateCompiledResult(EShaderPlatform Platform, const FShaderParameterMap& ParameterMap, TArray<FString>& OutError) const;
 
 	/** Calculates a Hash based on this shader type's source code and includes */
 	const FSHAHash& GetSourceHash(EShaderPlatform ShaderPlatform) const;
@@ -1207,7 +1037,7 @@ public:
 	/** Hashes a pointer to a shader type. */
 	friend uint32 GetTypeHash(FShaderType* Ref)
 	{
-		return Ref ? Ref->HashIndex : 0;
+		return Ref ? GetTypeHash(Ref->HashedName) : 0u;
 	}
 
 	// Dynamic casts.
@@ -1252,7 +1082,16 @@ public:
 		return (ShaderTypeForDynamicCast == EShaderTypeForDynamicCast::OCIO) ? reinterpret_cast<FOpenColorIOShaderType*>(this) : nullptr;
 	}
 
+	inline EShaderTypeForDynamicCast GetTypeForDynamicCast() const
+	{
+		return ShaderTypeForDynamicCast;
+	}
+
 	// Accessors.
+	inline const FTypeLayoutDesc& GetLayout() const
+	{
+		return *TypeLayout;
+	}
 	inline EShaderFrequency GetFrequency() const
 	{ 
 		return (EShaderFrequency)Frequency; 
@@ -1265,6 +1104,10 @@ public:
 	{
 		return TypeName;
 	}
+	inline const FHashedName& GetHashedName() const
+	{
+		return HashedName;
+	}
 	inline const TCHAR* GetShaderFilename() const
 	{ 
 		return SourceFilename; 
@@ -1273,19 +1116,20 @@ public:
 	{
 		return FunctionName;
 	}
-	inline int32 GetNumShaders() const
+	inline uint32 GetTypeSize() const
 	{
-		return ShaderIdMap.Num();
+		return TypeSize;
 	}
 
-	inline int32 GetPermutationCount() const
+	inline int32 GetNumShaders() const
 	{
-		return TotalPermutationCount;
+		// TODO count this
+		return 0;
 	}
 
-	inline const FSerializationHistory& GetSerializationHistory() const
+	inline int32 GetPermutationCount() const
 	{
-		return SerializationHistory;
+		return TotalPermutationCount;
 	}
 
 	inline const TMap<const TCHAR*, FCachedUniformBufferDeclaration>& GetReferencedUniformBufferStructsCache() const
@@ -1309,49 +1153,32 @@ public:
 		bCachedUniformBufferStructDeclarations = false;
 	}
 
-	void AddToShaderMap(const FShaderKey& Key, FShader* Shader)
-	{
-		check(IsInGameThread());
-		ShaderIdMap.Add(Key, Shader);
-	}
-
-	inline void RemoveFromShaderMap(const FShaderKey& Key)
-	{
-		check(IsInGameThread());
-		ShaderIdMap.Remove(Key);
-	}
-
 	void DumpDebugInfo();
-	void SaveShaderStableKeys(EShaderPlatform TargetShaderPlatform);
 	void GetShaderStableKeyParts(struct FStableShaderKeyAndValue& SaveKeyVal);
+
 private:
 	EShaderTypeForDynamicCast ShaderTypeForDynamicCast;
-	uint32 HashIndex;
+	const FTypeLayoutDesc* TypeLayout;
 	const TCHAR* Name;
 	FName TypeName;
+	FHashedName HashedName;
 	const TCHAR* SourceFilename;
 	const TCHAR* FunctionName;
 	uint32 Frequency;
+	uint32 TypeSize;
 	int32 TotalPermutationCount;
 
 	ConstructSerializedType ConstructSerializedRef;
+	ConstructCompiledType ConstructCompiledRef;
+	ModifyCompilationEnvironmentType ModifyCompilationEnvironmentRef;
+	ShouldCompilePermutationType ShouldCompilePermutationRef;
+	ValidateCompiledResultType ValidateCompiledResultRef;
 	const FShaderParametersMetadata* const RootParametersMetadata;
 
-	/** A map from shader ID to shader.  A shader will be removed from it when deleted, so this doesn't need to use a TRefCountPtr. */
-	TMap<FShaderKey,FShader*> ShaderIdMap;
-
 	TLinkedList<FShaderType*> GlobalListLink;
 
-	// DumpShaderStats needs to access ShaderIdMap.
 	friend void RENDERCORE_API DumpShaderStats( EShaderPlatform Platform, EShaderFrequency Frequency );
 
-
-	/** 
-	 * Stores a history of serialization sizes for this shader type. 
-	 * This is used to invalidate shaders when serialization changes.
-	 */
-	FSerializationHistory SerializationHistory;
-
 	/** Tracks whether serialization history for all shader types has been initialized. */
 	static bool bInitializedSerializationHistory;
 
@@ -1368,6 +1195,28 @@ protected:
 
 };
 
+#define SHADER_DECLARE_VTABLE(ShaderClass) \
+	static FShader* ConstructSerializedInstance() { return new ShaderClass(); } \
+	static FShader* ConstructCompiledInstance(const typename FShader::CompiledShaderInitializerType& Initializer) \
+	{ return new ShaderClass(static_cast<const typename ShaderMetaType::CompiledShaderInitializerType&>(Initializer)); }\
+	static void ModifyCompilationEnvironmentImpl(const FShaderPermutationParameters& Parameters, FShaderCompilerEnvironment& OutEnvironment) \
+	{ \
+		const typename ShaderClass::FPermutationDomain PermutationVector(Parameters.PermutationId); \
+		PermutationVector.ModifyCompilationEnvironment(OutEnvironment); \
+		ShaderClass::ModifyCompilationEnvironment(static_cast<const typename ShaderClass::FPermutationParameters&>(Parameters), OutEnvironment); \
+	} \
+	static bool ShouldCompilePermutationImpl(const FShaderPermutationParameters& Parameters) \
+	{ return ShaderClass::ShouldCompilePermutation(static_cast<const typename ShaderClass::FPermutationParameters&>(Parameters)); }
+
+#define INTERNAL_DECLARE_SHADER_TYPE_COMMON(ShaderClass,ShaderMetaTypeShortcut,RequiredAPI) \
+	public: \
+	using ShaderMetaType = F##ShaderMetaTypeShortcut##ShaderType; \
+	using ShaderMapType = F##ShaderMetaTypeShortcut##ShaderMap; \
+	\
+	static RequiredAPI ShaderMetaType StaticType; \
+	\
+	SHADER_DECLARE_VTABLE(ShaderClass)
+
 /**
  * A macro to declare a new shader type.  This should be called in the class body of the new shader type.
  * @param ShaderClass - The name of the class representing an instance of the shader type.
@@ -1375,68 +1224,72 @@ protected:
  *	controls 
  */
 #define DECLARE_EXPORTED_SHADER_TYPE(ShaderClass,ShaderMetaTypeShortcut,RequiredAPI, ...) \
-	public: \
-	using FPermutationDomain = FShaderPermutationNone; \
-	using ShaderMetaType = F##ShaderMetaTypeShortcut##ShaderType; \
-	\
-	static RequiredAPI ShaderMetaType StaticType; \
-	\
-	static FShader* ConstructSerializedInstance() { return new ShaderClass(); } \
-	static FShader* ConstructCompiledInstance(const ShaderMetaType::CompiledShaderInitializerType& Initializer) \
-	{ return new ShaderClass(Initializer); } \
-	\
-	virtual uint32 GetTypeSize() const override { return sizeof(*this); }
+	INTERNAL_DECLARE_SHADER_TYPE_COMMON(ShaderClass, ShaderMetaTypeShortcut, RequiredAPI); \
+	DECLARE_EXPORTED_TYPE_LAYOUT(ShaderClass, RequiredAPI, NonVirtual); \
+	public:
 
 #define DECLARE_SHADER_TYPE(ShaderClass,ShaderMetaTypeShortcut,...) \
 	DECLARE_EXPORTED_SHADER_TYPE(ShaderClass,ShaderMetaTypeShortcut,, ##__VA_ARGS__)
 
+#define DECLARE_SHADER_TYPE_EXPLICIT_BASES(ShaderClass,ShaderMetaTypeShortcut,...) \
+	INTERNAL_DECLARE_SHADER_TYPE_COMMON(ShaderClass, ShaderMetaTypeShortcut,); \
+	DECLARE_EXPORTED_TYPE_LAYOUT_EXPLICIT_BASES(ShaderClass,, NonVirtual, __VA_ARGS__); \
+	public:
+
+#define SHADER_TYPE_VTABLE(ShaderClass) \
+	ShaderClass::ConstructSerializedInstance, \
+	ShaderClass::ConstructCompiledInstance, \
+	ShaderClass::ModifyCompilationEnvironmentImpl, \
+	ShaderClass::ShouldCompilePermutationImpl, \
+	ShaderClass::ValidateCompiledResult
+
 #if !UE_BUILD_DOCS
 /** A macro to implement a shader type. */
 #define IMPLEMENT_SHADER_TYPE(TemplatePrefix,ShaderClass,SourceFilename,FunctionName,Frequency) \
+	IMPLEMENT_UNREGISTERED_TEMPLATE_TYPE_LAYOUT(TemplatePrefix, ShaderClass); \
 	TemplatePrefix \
 	ShaderClass::ShaderMetaType ShaderClass::StaticType( \
+		ShaderClass::StaticGetTypeLayout(), \
 		TEXT(#ShaderClass), \
 		SourceFilename, \
 		FunctionName, \
 		Frequency, \
-		1, \
-		ShaderClass::ConstructSerializedInstance, \
-		ShaderClass::ConstructCompiledInstance, \
-		ShaderClass::ModifyCompilationEnvironment, \
-		ShaderClass::ShouldCompilePermutation, \
-		ShaderClass::ValidateCompiledResult \
+		ShaderClass::FPermutationDomain::PermutationCount, \
+		SHADER_TYPE_VTABLE(ShaderClass), \
+		sizeof(ShaderClass), \
+		ShaderClass::GetRootParametersMetadata() \
 		);
 
 /** A macro to implement a shader type. Shader name is got from GetDebugName(), which is helpful for templated shaders. */
 #define IMPLEMENT_SHADER_TYPE_WITH_DEBUG_NAME(TemplatePrefix,ShaderClass,SourceFilename,FunctionName,Frequency) \
+	IMPLEMENT_UNREGISTERED_TEMPLATE_TYPE_LAYOUT(TemplatePrefix, ShaderClass); \
 	TemplatePrefix \
 	typename ShaderClass::ShaderMetaType ShaderClass::StaticType( \
+		ShaderClass::StaticGetTypeLayout(), \
 		ShaderClass::GetDebugName(), \
 		SourceFilename, \
 		FunctionName, \
 		Frequency, \
-		1, \
-		ShaderClass::ConstructSerializedInstance, \
-		ShaderClass::ConstructCompiledInstance, \
-		ShaderClass::ModifyCompilationEnvironment, \
-		ShaderClass::ShouldCompilePermutation, \
-		ShaderClass::ValidateCompiledResult \
+		ShaderClass::FPermutationDomain::PermutationCount, \
+		SHADER_TYPE_VTABLE(ShaderClass), \
+		sizeof(ShaderClass), \
+		ShaderClass::GetRootParametersMetadata() \
 		);
 
 /** A macro to implement a templated shader type, the function name and the source filename comes from the class. */
 #define IMPLEMENT_SHADER_TYPE2_WITH_TEMPLATE_PREFIX(TemplatePrefix,ShaderClass,Frequency) \
+	IMPLEMENT_UNREGISTERED_TEMPLATE_TYPE_LAYOUT(TemplatePrefix, ShaderClass); \
 	TemplatePrefix \
 	ShaderClass::ShaderMetaType ShaderClass::StaticType( \
+	ShaderClass::StaticGetTypeLayout(), \
 	TEXT(#ShaderClass), \
 	ShaderClass::GetSourceFilename(), \
 	ShaderClass::GetFunctionName(), \
 	Frequency, \
-	1, \
-	ShaderClass::ConstructSerializedInstance, \
-	ShaderClass::ConstructCompiledInstance, \
-	ShaderClass::ModifyCompilationEnvironment, \
-	ShaderClass::ShouldCompilePermutation, \
-	ShaderClass::ValidateCompiledResult \
+	ShaderClass::FPermutationDomain::PermutationCount, \
+	SHADER_TYPE_VTABLE(ShaderClass), \
+	sizeof(ShaderClass), \
+	ShaderClass::GetRootParametersMetadata() \
 	);
 
 #define IMPLEMENT_SHADER_TYPE2(ShaderClass,Frequency) \
@@ -1444,33 +1297,33 @@ protected:
 
 /** todo: this should replace IMPLEMENT_SHADER_TYPE */
 #define IMPLEMENT_SHADER_TYPE3(ShaderClass,Frequency) \
+	IMPLEMENT_UNREGISTERED_TEMPLATE_TYPE_LAYOUT(,ShaderClass); \
 	ShaderClass::ShaderMetaType ShaderClass::StaticType( \
+	ShaderClass::StaticGetTypeLayout(), \
 	TEXT(#ShaderClass), \
 	ShaderClass::GetSourceFilename(), \
 	ShaderClass::GetFunctionName(), \
 	Frequency, \
-	1, \
-	ShaderClass::ConstructSerializedInstance, \
-	ShaderClass::ConstructCompiledInstance, \
-	ShaderClass::ModifyCompilationEnvironment, \
-	ShaderClass::ShouldCompilePermutation, \
-	ShaderClass::ValidateCompiledResult \
+	ShaderClass::FPermutationDomain::PermutationCount, \
+	SHADER_TYPE_VTABLE(ShaderClass), \
+	sizeof(ShaderClass), \
+	ShaderClass::GetRootParametersMetadata() \
 	);
-#endif
+#endif // !UE_BUILD_DOCS
 
 #define IMPLEMENT_SHADER_TYPE4_WITH_TEMPLATE_PREFIX(TemplatePrefix,RequiredAPI,ShaderClass,Frequency) \
+	IMPLEMENT_UNREGISTERED_TEMPLATE_TYPE_LAYOUT(TemplatePrefix, ShaderClass); \
 	TemplatePrefix RequiredAPI \
 	ShaderClass::ShaderMetaType ShaderClass::StaticType( \
+	ShaderClass::StaticGetTypeLayout(), \
 	TEXT(#ShaderClass), \
 	ShaderClass::GetSourceFilename(), \
 	ShaderClass::GetFunctionName(), \
 	Frequency, \
-	1, \
-	ShaderClass::ConstructSerializedInstance, \
-	ShaderClass::ConstructCompiledInstance, \
-	ShaderClass::ModifyCompilationEnvironment, \
-	ShaderClass::ShouldCompilePermutation, \
-	ShaderClass::ValidateCompiledResult \
+	ShaderClass::FPermutationDomain::PermutationCount, \
+	SHADER_TYPE_VTABLE(ShaderClass), \
+	sizeof(ShaderClass), \
+	ShaderClass::GetRootParametersMetadata() \
 	);
 
 // Binding of a set of shader stages in a single pipeline
@@ -1501,15 +1354,19 @@ public:
 
 	FORCEINLINE FName GetFName() const { return TypeName; }
 	FORCEINLINE TCHAR const* GetName() const { return Name; }
+	FORCEINLINE const FHashedName& GetHashedName() const { return HashedName; }
+	FORCEINLINE const FHashedName& GetHashedPrimaryShaderFilename() const { return HashedPrimaryShaderFilename; }
 
 	// Returns an array of valid stages, sorted from PS->GS->DS->HS->VS, no gaps if missing stages
 	FORCEINLINE const TArray<const FShaderType*>& GetStages() const { return Stages; }
 
 	static TLinkedList<FShaderPipelineType*>*& GetTypeList();
 
+	static const TArray<FShaderPipelineType*>& GetSortedTypes(FShaderType::EShaderTypeForDynamicCast Type);
+
 	/** @return The global shader pipeline name to type map */
-	static TMap<FName, FShaderPipelineType*>& GetNameToTypeMap();
-	static const FShaderPipelineType* GetShaderPipelineTypeByName(FName Name);
+	static TMap<FHashedName, FShaderPipelineType*>& GetNameToTypeMap();
+	static const FShaderPipelineType* GetShaderPipelineTypeByName(const FHashedName& Name);
 
 	/** Initialize static members, this must be called before any shader types are created. */
 	static void Initialize();
@@ -1534,15 +1391,14 @@ public:
 		return bShouldOptimizeUnusedOutputs && RHISupportsShaderPipelines(Platform);
 	}
 
-	/** Gets a list of FShaderTypes & PipelineTypes whose source file no longer matches what that type was compiled with */
-	static void GetOutdatedTypes(TArray<FShaderType*>& OutdatedShaderTypes, TArray<const FShaderPipelineType*>& ShaderPipelineTypesToFlush, TArray<const FVertexFactoryType*>& OutdatedFactoryTypes);
-
 	/** Calculates a Hash based on this shader pipeline type stages' source code and includes */
 	const FSHAHash& GetSourceHash(EShaderPlatform ShaderPlatform) const;
 
 protected:
 	const TCHAR* const Name;
 	FName TypeName;
+	FHashedName HashedName;
+	FHashedName HashedPrimaryShaderFilename;
 
 	// Pipeline Stages, ordered from lowest (usually PS) to highest (VS). Guaranteed at least one stage (for VS).
 	TArray<const FShaderType*> Stages;
@@ -1587,48 +1443,29 @@ protected:
 /** Encapsulates a dependency on a shader type and saved state from that shader type. */
 class FShaderTypeDependency
 {
+	DECLARE_EXPORTED_TYPE_LAYOUT(FShaderTypeDependency, RENDERCORE_API, NonVirtual);
 public:
 
 	FShaderTypeDependency()
-		: ShaderType(nullptr)
-		, PermutationId(0)
+		: PermutationId(0)
 	{}
 
 	FShaderTypeDependency(FShaderType* InShaderType, EShaderPlatform ShaderPlatform)
-		: ShaderType(InShaderType)
+		: ShaderTypeName(InShaderType->GetHashedName())
 		, PermutationId(0)
 	{
-#if KEEP_SHADER_SOURCE_HASHES
-		if (ShaderType)
+		if (InShaderType)
 		{
-			SourceHash = ShaderType->GetSourceHash(ShaderPlatform);
+			SourceHash = InShaderType->GetSourceHash(ShaderPlatform);
 		}
-#endif
 	}
 
-	/** Shader type */
-	FShaderType* ShaderType;
-
-	/** Unique permutation identifier of the global shader type. */
-	int32 PermutationId;
-
-#if KEEP_SHADER_SOURCE_HASHES
-	/** Used to detect changes to the shader source files. */
-	FSHAHash SourceHash;
-#endif
-
 	friend FArchive& operator<<(FArchive& Ar,class FShaderTypeDependency& Ref)
 	{
 		Ar.UsingCustomVersion(FRenderingObjectVersion::GUID);
 
-		Ar << Ref.ShaderType;
-
-#if KEEP_SHADER_SOURCE_HASHES
-		FSHAHash& Hash = Ref.SourceHash;
-#else
-		FSHAHash Hash;
-#endif
-		Ar << FShaderResource::FilterShaderSourceHashForSerialization(Ar, Hash);
+		Ar << Ref.ShaderTypeName;
+		Ar << Ref.SourceHash;
 
 		if (Ar.CustomVer(FRenderingObjectVersion::GUID) >= FRenderingObjectVersion::ShaderPermutationId)
 		{
@@ -1640,66 +1477,54 @@ public:
 
 	bool operator==(const FShaderTypeDependency& Reference) const
 	{
-#if KEEP_SHADER_SOURCE_HASHES
-		return ShaderType == Reference.ShaderType && PermutationId == Reference.PermutationId && SourceHash == Reference.SourceHash;
-#else
-		return ShaderType == Reference.ShaderType && PermutationId == Reference.PermutationId;
-#endif
+		return ShaderTypeName == Reference.ShaderTypeName && PermutationId == Reference.PermutationId && SourceHash == Reference.SourceHash;
 	}
 
 	bool operator!=(const FShaderTypeDependency& Reference) const
 	{
 		return !(*this == Reference);
 	}
+
+	/** Shader type */
+	LAYOUT_FIELD(FHashedName, ShaderTypeName);
+
+	/** Unique permutation identifier of the global shader type. */
+	LAYOUT_FIELD(int32, PermutationId);
+
+	/** Used to detect changes to the shader source files. This is always present, as this type is sometimes frozen. */
+	LAYOUT_FIELD(FSHAHash, SourceHash);
 };
 
 
 class FShaderPipelineTypeDependency
 {
 public:
-	FShaderPipelineTypeDependency() :
-		ShaderPipelineType(nullptr)
-	{}
-
+	FShaderPipelineTypeDependency() {}
 	FShaderPipelineTypeDependency(const FShaderPipelineType* InShaderPipelineType, EShaderPlatform ShaderPlatform) :
-		ShaderPipelineType(InShaderPipelineType)
+		ShaderPipelineTypeName(InShaderPipelineType->GetHashedName())
 	{
-#if KEEP_SHADER_SOURCE_HASHES
-		if (ShaderPipelineType)
+		if (InShaderPipelineType)
 		{
-			StagesSourceHash = ShaderPipelineType->GetSourceHash(ShaderPlatform);
+			StagesSourceHash = InShaderPipelineType->GetSourceHash(ShaderPlatform);
 		}
-#endif
 	}
 
 	/** Shader Pipeline type */
-	const FShaderPipelineType* ShaderPipelineType;
+	FHashedName ShaderPipelineTypeName;
 
-#if KEEP_SHADER_SOURCE_HASHES
 	/** Used to detect changes to the shader source files. */
 	FSHAHash StagesSourceHash;
-#endif
 
 	friend FArchive& operator<<(FArchive& Ar, class FShaderPipelineTypeDependency& Ref)
 	{
-		Ar << Ref.ShaderPipelineType;
-
-#if KEEP_SHADER_SOURCE_HASHES
-		FSHAHash& Hash = Ref.StagesSourceHash;
-#else
-		FSHAHash Hash;
-#endif
-		Ar << FShaderResource::FilterShaderSourceHashForSerialization(Ar, Hash);
+		Ar << Ref.ShaderPipelineTypeName;
+		Ar << Ref.StagesSourceHash;
 		return Ar;
 	}
 
 	bool operator==(const FShaderPipelineTypeDependency& Reference) const
 	{
-#if KEEP_SHADER_SOURCE_HASHES	
-		return ShaderPipelineType == Reference.ShaderPipelineType && StagesSourceHash == Reference.StagesSourceHash;
-#else
-		return ShaderPipelineType == Reference.ShaderPipelineType;
-#endif
+		return ShaderPipelineTypeName == Reference.ShaderPipelineTypeName && StagesSourceHash == Reference.StagesSourceHash;
 	}
 
 	bool operator!=(const FShaderPipelineTypeDependency& Reference) const
@@ -1761,14 +1586,8 @@ public:
 // A Shader Pipeline instance with compiled stages
 class RENDERCORE_API FShaderPipeline
 {
+	DECLARE_TYPE_LAYOUT(FShaderPipeline, NonVirtual);
 public:
-	const FShaderPipelineType* PipelineType;
-	TRefCountPtr<FShader> VertexShader;
-	TRefCountPtr<FShader> HullShader;
-	TRefCountPtr<FShader> DomainShader;
-	TRefCountPtr<FShader> GeometryShader;
-	TRefCountPtr<FShader> PixelShader;
-
 	FShaderPipeline(
 		const FShaderPipelineType* InPipelineType,
 		FShader* InVertexShader,
@@ -1778,35 +1597,45 @@ public:
 		FShader* InPixelShader);
 
 	FShaderPipeline(const FShaderPipelineType* InPipelineType, const TArray<FShader*>& InStages);
-	FShaderPipeline(const FShaderPipelineType* InPipelineType, const TArray< TRefCountPtr<FShader> >& InStages);
 
 	~FShaderPipeline();
 
+	inline uint32 GetNumShaders() const
+	{
+		uint32 NumShaders = 0u;
+		if (VertexShader.IsValid()) ++NumShaders;
+		if (HullShader.IsValid()) ++NumShaders;
+		if (DomainShader.IsValid()) ++NumShaders;
+		if (GeometryShader.IsValid()) ++NumShaders;
+		if (PixelShader.IsValid()) ++NumShaders;
+		return NumShaders;
+	}
+
 	// Find a shader inside the pipeline
 	template<typename ShaderType>
-	ShaderType* GetShader()
+	ShaderType* GetShader(const FShaderMapPointerTable& InPtrTable)
 	{
-		if (PixelShader && PixelShader->GetType() == &ShaderType::StaticType)
+		if (PixelShader && PixelShader->GetType(InPtrTable) == &ShaderType::StaticType)
 		{
-			return (ShaderType*)PixelShader.GetReference();
+			return (ShaderType*)PixelShader.Get();
 		}
-		else if (VertexShader && VertexShader->GetType() == &ShaderType::StaticType)
+		else if (VertexShader && VertexShader->GetType(InPtrTable) == &ShaderType::StaticType)
 		{
-			return (ShaderType*)VertexShader.GetReference();
+			return (ShaderType*)VertexShader.Get();
 		}
-		else if (GeometryShader && GeometryShader->GetType() == &ShaderType::StaticType)
+		else if (GeometryShader && GeometryShader->GetType(InPtrTable) == &ShaderType::StaticType)
 		{
-			return (ShaderType*)GeometryShader.GetReference();
+			return (ShaderType*)GeometryShader.Get();
 		}
 		else if (HullShader)
 		{
-			if (HullShader->GetType() == &ShaderType::StaticType)
+			if (HullShader->GetType(InPtrTable) == &ShaderType::StaticType)
 			{
-				return (ShaderType*)HullShader.GetReference();
+				return (ShaderType*)HullShader.Get();
 			}
-			else if (DomainShader && DomainShader->GetType() == &ShaderType::StaticType)
+			else if (DomainShader && DomainShader->GetType(InPtrTable) == &ShaderType::StaticType)
 			{
-				return (ShaderType*)DomainShader.GetReference();
+				return (ShaderType*)DomainShader.Get();
 			}
 		}
 
@@ -1817,11 +1646,11 @@ public:
 	{
 		switch (Frequency)
 		{
-		case SF_Vertex: return VertexShader.GetReference();
-		case SF_Domain: return DomainShader.GetReference();
-		case SF_Hull: return HullShader.GetReference();
-		case SF_Geometry: return GeometryShader.GetReference();
-		case SF_Pixel: return PixelShader.GetReference();
+		case SF_Vertex: return VertexShader;
+		case SF_Domain: return DomainShader;
+		case SF_Hull: return HullShader;
+		case SF_Geometry: return GeometryShader;
+		case SF_Pixel: return PixelShader;
 		default: check(0);
 		}
 
@@ -1832,134 +1661,159 @@ public:
 	{
 		switch (Frequency)
 		{
-		case SF_Vertex: return VertexShader.GetReference();
-		case SF_Domain: return DomainShader.GetReference();
-		case SF_Hull: return HullShader.GetReference();
-		case SF_Geometry: return GeometryShader.GetReference();
-		case SF_Pixel: return PixelShader.GetReference();
+		case SF_Vertex: return VertexShader;
+		case SF_Domain: return DomainShader;
+		case SF_Hull: return HullShader;
+		case SF_Geometry: return GeometryShader;
+		case SF_Pixel: return PixelShader;
 		default: check(0);
 		}
 
 		return nullptr;
 	}
 
-	inline TArray<FShader*> GetShaders() const
+	inline TArray<TShaderRef<FShader>> GetShaders(const FShaderMapBase& InShaderMap) const
 	{
-		TArray<FShader*> Shaders;
+		TArray<TShaderRef<FShader>> Shaders;
 
 		if (PixelShader)
 		{
-			Shaders.Add(PixelShader.GetReference());
+			Shaders.Add(TShaderRef<FShader>(PixelShader, InShaderMap));
 		}
 		if (GeometryShader)
 		{
-			Shaders.Add(GeometryShader.GetReference());
+			Shaders.Add(TShaderRef<FShader>(GeometryShader, InShaderMap));
 		}
 		if (HullShader)
 		{
-			Shaders.Add(DomainShader.GetReference());
-			Shaders.Add(HullShader.GetReference());
+			Shaders.Add(TShaderRef<FShader>(DomainShader, InShaderMap));
+			Shaders.Add(TShaderRef<FShader>(HullShader, InShaderMap));
 		}
 
-		Shaders.Add(VertexShader.GetReference());
+		Shaders.Add(TShaderRef<FShader>(VertexShader, InShaderMap));
+
+		return Shaders;
+	}
+
+	void Validate(const FShaderPipelineType* InPipelineType) const;
+
+	enum EFilter
+	{
+		EAll,			// All pipelines
+		EOnlyShared,	// Only pipelines with shared shaders
+		EOnlyUnique,	// Only pipelines with unique shaders
+	};
+
+	/** Saves stable keys for the shaders in the pipeline */
+#if WITH_EDITOR
+	void SaveShaderStableKeys(const FShaderMapPointerTable& InPtrTable, EShaderPlatform TargetShaderPlatform, const struct FStableShaderKeyAndValue& SaveKeyVal) const;
+#endif // WITH_EDITOR
+
+	LAYOUT_FIELD(FHashedName, TypeName);
+	LAYOUT_FIELD(TMemoryImagePtr<FShader>, VertexShader);
+	LAYOUT_FIELD(TMemoryImagePtr<FShader>, HullShader);
+	LAYOUT_FIELD(TMemoryImagePtr<FShader>, DomainShader);
+	LAYOUT_FIELD(TMemoryImagePtr<FShader>, GeometryShader);
+	LAYOUT_FIELD(TMemoryImagePtr<FShader>, PixelShader);
+};
+
+inline bool operator<(const FShaderPipeline& Lhs, const FShaderPipeline& Rhs)
+{
+	return Lhs.TypeName.GetHash() < Rhs.TypeName.GetHash();
+}
+
+class FShaderPipelineRef
+{
+public:
+	FShaderPipelineRef() : ShaderPipeline(nullptr), ShaderMap(nullptr) {}
+	FShaderPipelineRef(FShaderPipeline* InPipeline, const FShaderMapBase& InShaderMap) : ShaderPipeline(InPipeline), ShaderMap(&InShaderMap) {}
+
+	inline bool IsValid() const { return ShaderPipeline != nullptr; }
+	inline bool IsNull() const { return ShaderPipeline == nullptr; }
+
+	template<typename ShaderType>
+	TShaderRef<ShaderType> GetShader() const
+	{
+		return TShaderRef<ShaderType>(ShaderPipeline->GetShader<ShaderType>(GetPointerTable()), *ShaderMap);
+	}
 
-		return Shaders;
+	TShaderRef<FShader> GetShader(EShaderFrequency Frequency) const
+	{
+		return TShaderRef<FShader>(ShaderPipeline->GetShader(Frequency), *ShaderMap);
 	}
 
-	inline uint32 GetSizeBytes() const
+	inline TArray<TShaderRef<FShader>> GetShaders() const
 	{
-		return sizeof(*this);
+		return ShaderPipeline->GetShaders(*ShaderMap);
 	}
 
-	void Validate();
+	inline FShaderPipeline* GetPipeline() const { return ShaderPipeline; }
+	FShaderMapResource* GetResource() const;
+	const FShaderMapPointerTable& GetPointerTable() const;
 
-	enum EFilter
-	{
-		EAll,			// All pipelines
-		EOnlyShared,	// Only pipelines with shared shaders
-		EOnlyUnique,	// Only pipelines with unique shaders
-	};
-	
-	static void CookPipeline(FShaderPipeline* Pipeline);
+	inline FShaderPipeline* operator->() const { check(ShaderPipeline); return ShaderPipeline; }
 
-	/** Saves stable keys for the shaders in the pipeline */
-	void SaveShaderStableKeys(EShaderPlatform TargetShaderPlatform, const struct FStableShaderKeyAndValue& SaveKeyVal);
+private:
+	FShaderPipeline* ShaderPipeline;
+	const FShaderMapBase* ShaderMap;
 };
 
-inline bool operator<(const FShaderPipeline& Lhs, const FShaderPipeline& Rhs)
+struct FShaderMapKey
 {
-	FCompareShaderPipelineNameTypes Comparator;
-	return Comparator(*Lhs.PipelineType, *Rhs.PipelineType);
-}
+	FShaderMapKey() {}
+	FShaderMapKey(const FHashedName& InTypeName, int32 InPermutationId) : TypeName(InTypeName), PermutationId(InPermutationId) {}
+	explicit FShaderMapKey(const FShader& InShader) : FShaderMapKey(InShader.GetTypeName(), InShader.GetPermutationId()) {}
 
-/** A collection of shaders of different types, but the same meta type. */
-template<typename ShaderMetaType>
-class TShaderMap
-{
-	/** Container for serialized shader pipeline stages to be registered on the game thread */
-	struct FSerializedShaderPipeline
+	friend inline bool operator<(const FShaderMapKey& Lhs, const FShaderMapKey& Rhs)
 	{
-		const FShaderPipelineType* ShaderPipelineType;
-		TArray< TRefCountPtr<FShader> > ShaderStages;
-		FSerializedShaderPipeline()
-			: ShaderPipelineType(nullptr)
+		if (Lhs.TypeName != Rhs.TypeName)
 		{
+			return Lhs.TypeName < Rhs.TypeName;
 		}
-	};
+		return Lhs.PermutationId < Rhs.PermutationId;
+	}
 
-	/** List of serialized shaders to be processed and registered on the game thread */
-	TArray<FShader*> SerializedShaders;
-	/** List of serialized shader pipeline stages to be processed and registered on the game thread */
-	TArray<FSerializedShaderPipeline*> SerializedShaderPipelines;
-protected:
-	/** The platform this shader map was compiled with */
-	EShaderPlatform Platform;
-private:
-	/** Flag that makes sure this shader map isn't used until all shaders have been registerd */
-	bool bHasBeenRegistered;
+	FHashedName TypeName;
+	int32 PermutationId;
+};
 
+/** A collection of shaders of different types */
+class RENDERCORE_API FShaderMapContent
+{
+	DECLARE_TYPE_LAYOUT(FShaderMapContent, NonVirtual);
 public:
+	struct FProjectShaderToKey
+	{
+		inline FShaderMapKey operator()(const FShader* InShader) { return FShaderMapKey(*InShader); }
+	};
 
-	using FShaderPrimaryKey = TShaderTypePermutation<FShaderType>;
-
-	/** Used to compare two shader types by name. */
-	class FCompareShaderPrimaryKey
+	struct FProjectShaderPipelineToKey
 	{
-	public:
-		FORCEINLINE bool operator()(const FShaderPrimaryKey& A, const FShaderPrimaryKey& B) const
-		{
-			int32 AL = FCString::Strlen(A.Type->GetName());
-			int32 BL = FCString::Strlen(B.Type->GetName());
-			if (AL == BL)
-			{
-				return FCString::Strncmp(A.Type->GetName(), B.Type->GetName(), AL) > 0 || A.PermutationId > B.PermutationId;
-			}
-			return AL > BL;
-		}
+		inline FHashedName operator()(const FShaderPipeline* InShaderPipeline) { return InShaderPipeline->TypeName; }
 	};
 
 	/** Default constructor. */
-	TShaderMap(EShaderPlatform InPlatform)
+	explicit FShaderMapContent(EShaderPlatform InPlatform)
 		: Platform(InPlatform)
-		, bHasBeenRegistered(true)
 	{}
 
 	/** Destructor ensures pipelines cleared up. */
-	virtual ~TShaderMap()
+	~FShaderMapContent()
 	{
 		Empty();
 	}
 
 	EShaderPlatform GetShaderPlatform() const { return Platform; }
 
+	void Validate(const FShaderMapBase& InShaderMap);
+
 	/** Finds the shader with the given type.  Asserts on failure. */
 	template<typename ShaderType>
 	ShaderType* GetShader(int32 PermutationId = 0) const
 	{
-		check(bHasBeenRegistered);
-		const TRefCountPtr<FShader>* ShaderRef = Shaders.Find(FShaderPrimaryKey(&ShaderType::StaticType, PermutationId));
-		checkf(ShaderRef != NULL && *ShaderRef != nullptr, TEXT("Failed to find shader type %s in Platform %s"), ShaderType::StaticType.GetName(), *LegacyShaderPlatformToShaderFormat(Platform).ToString());
-		return (ShaderType*)((*ShaderRef)->GetShader());
+		FShader* Shader = GetShader(&ShaderType::StaticType, PermutationId);
+		checkf(Shader != nullptr, TEXT("Failed to find shader type %s in Platform %s"), ShaderType::StaticType.GetName(), *LegacyShaderPlatformToShaderFormat(Platform).ToString());
+		return static_cast<ShaderType*>(Shader);
 	}
 
 	/** Finds the shader with the given type.  Asserts on failure. */
@@ -1972,474 +1826,209 @@ public:
 	/** Finds the shader with the given type.  May return NULL. */
 	FShader* GetShader(FShaderType* ShaderType, int32 PermutationId = 0) const
 	{
-		check(bHasBeenRegistered);
-		const TRefCountPtr<FShader>* ShaderRef = Shaders.Find(FShaderPrimaryKey(ShaderType, PermutationId));
-		return ShaderRef ? (*ShaderRef)->GetShader() : nullptr;
+		return GetShader(ShaderType->GetHashedName(), PermutationId);
 	}
 
+	/** Finds the shader with the given type name.  May return NULL. */
+	FShader* GetShader(const FHashedName& TypeName, int32 PermutationId = 0) const;
+
 	/** Finds the shader with the given type. */
+	bool HasShader(const FHashedName& TypeName, int32 PermutationId) const
+	{
+		const FShader* Shader = GetShader(TypeName, PermutationId);
+		return Shader != nullptr;
+	}
+
 	bool HasShader(FShaderType* Type, int32 PermutationId) const
 	{
-		check(bHasBeenRegistered);
-		const TRefCountPtr<FShader>* ShaderRef = Shaders.Find(FShaderPrimaryKey(Type, PermutationId));
-		return ShaderRef != nullptr && *ShaderRef != nullptr;
+		return HasShader(Type->GetHashedName(), PermutationId);
 	}
 
-	inline const TMap<FShaderPrimaryKey,TRefCountPtr<FShader> >& GetShaders() const
+	inline const TMemoryImageArray<TMemoryImagePtr<FShader>>& GetShaders() const
 	{
-		check(bHasBeenRegistered);
 		return Shaders;
 	}
 
-	void AddShader(FShaderType* Type, int32 PermutationId, FShader* Shader)
+	inline const TMemoryImageArray<TMemoryImagePtr<FShaderPipeline>>& GetShaderPipelines() const
 	{
-		check(Type);
-		Shaders.Add(FShaderPrimaryKey(Type, PermutationId), Shader);
+		return ShaderPipelines;
 	}
 
+	void AddShader(FShader* Shader);
+
+	FShader* FindOrAddShader(FShader* Shader);
+
+	void AddShaderPipeline(FShaderPipeline* Pipeline);
+
+	FShaderPipeline* FindOrAddShaderPipeline(FShaderPipeline* Pipeline);
+
 	/**
 	 * Removes the shader of the given type from the shader map
 	 * @param Type Shader type to remove the entry for 
 	 */
-	void RemoveShaderTypePermutaion(FShaderType* Type, int32 PermutationId)
-	{
-		Shaders.Remove(FShaderPrimaryKey(Type, PermutationId));
-	}
-
+	void RemoveShaderTypePermutaion(const FShaderType* Type, int32 PermutationId);
 
-	void RemoveShaderPipelineType(const FShaderPipelineType* ShaderPipelineType)
-	{
-		FShaderPipeline** Found = ShaderPipelines.Find(ShaderPipelineType);
-		if (Found)
-		{
-			if (*Found)
-			{
-				delete *Found;
-			}
-			ShaderPipelines.Remove(ShaderPipelineType);
-		}
-	}
+	void RemoveShaderPipelineType(const FShaderPipelineType* ShaderPipelineType);
 
 	/** Builds a list of the shaders in a shader map. */
-	void GetShaderList(TMap<FShaderId, FShader*>& OutShaders) const
-	{
-		check(bHasBeenRegistered);
-		for (typename TMap<FShaderPrimaryKey, TRefCountPtr<FShader>>::TConstIterator ShaderIt(Shaders); ShaderIt; ++ShaderIt)
-		{
-			if (ShaderIt.Value())
-			{
-				OutShaders.Add(ShaderIt.Value()->GetId(), ShaderIt.Value());
-			}
-		}
-	}
+	void GetShaderList(const FShaderMapBase& InShaderMap, const FSHAHash& InMaterialShaderMapHash, TMap<FShaderId, TShaderRef<FShader>>& OutShaders) const;
 
 	/** Builds a list of the shaders in a shader map. Key is FShaderType::TypeName */
-	void GetShaderList(TMap<FName, FShader*>& OutShaders) const
-	{
-		check(bHasBeenRegistered);
-		for (TMap<FShaderPrimaryKey, TRefCountPtr<FShader> >::TConstIterator ShaderIt(Shaders); ShaderIt; ++ShaderIt)
-		{
-			if (ShaderIt.Value())
-			{
-				OutShaders.Add(ShaderIt.Value()->GetType()->GetFName(), ShaderIt.Value());
-			}
-		}
-	}
-	uint32 GetShaderNum() const
-	{
-		return Shaders.Num();
-	}
+	void GetShaderList(const FShaderMapBase& InShaderMap, TMap<FHashedName, TShaderRef<FShader>>& OutShaders) const;
 
 	/** Builds a list of the shader pipelines in a shader map. */
-	void GetShaderPipelineList(TArray<FShaderPipeline*>& OutShaderPipelines, FShaderPipeline::EFilter Filter) const
-	{
-		check(bHasBeenRegistered);
-		for (auto Pair : ShaderPipelines)
-		{
-			FShaderPipeline* Pipeline = Pair.Value;
-			if (Pipeline->PipelineType->ShouldOptimizeUnusedOutputs(Platform) && Filter == FShaderPipeline::EOnlyShared)
-			{
-				continue;
-			}
-			else if (!Pipeline->PipelineType->ShouldOptimizeUnusedOutputs(Platform) && Filter == FShaderPipeline::EOnlyUnique)
-			{
-				continue;
-			}
-			OutShaderPipelines.Add(Pipeline);
-		}
-	}
+	void GetShaderPipelineList(const FShaderMapBase& InShaderMap, TArray<FShaderPipelineRef>& OutShaderPipelines, FShaderPipeline::EFilter Filter) const;
 
 #if WITH_EDITOR
-	uint32 GetMaxTextureSamplersShaderMap() const
-	{
-		check(bHasBeenRegistered);
-		uint32 MaxTextureSamplers = 0;
-
-		for (typename TMap<FShaderPrimaryKey,TRefCountPtr<FShader> >::TConstIterator ShaderIt(Shaders);ShaderIt;++ShaderIt)
-		{
-			if (ShaderIt.Value())
-			{
-				MaxTextureSamplers = FMath::Max(MaxTextureSamplers, ShaderIt.Value()->GetNumTextureSamplers());
-			}
-		}
+	uint32 GetMaxTextureSamplersShaderMap(const FShaderMapBase& InShaderMap) const;
 
-		for (auto Pair : ShaderPipelines)
-		{
-			const FShaderPipeline* Pipeline = Pair.Value;
-			for (const FShaderType* ShaderType : Pair.Key->GetStages())
-			{
-				MaxTextureSamplers = FMath::Max(MaxTextureSamplers, Pipeline->GetShader(ShaderType->GetFrequency())->GetNumTextureSamplers());
-			}
-		}
+	void GetOutdatedTypes(const FShaderMapBase& InShaderMap, TArray<const FShaderType*>& OutdatedShaderTypes, TArray<const FShaderPipelineType*>& OutdatedShaderPipelineTypes, TArray<const FVertexFactoryType*>& OutdatedFactoryTypes) const;
 
-		return MaxTextureSamplers;
-	}
+	void SaveShaderStableKeys(const FShaderMapBase& InShaderMap, EShaderPlatform TargetShaderPlatform, const struct FStableShaderKeyAndValue& SaveKeyVal);
 #endif // WITH_EDITOR
 
-	inline void SerializeShaderForSaving(FShader* CurrentShader, FArchive& Ar, bool bHandleShaderKeyChanges, bool bInlineShaderResource)
-	{
-		int64 SkipOffset = Ar.Tell();
-
-		{
-#if WITH_EDITOR
-			FArchive::FScopeSetDebugSerializationFlags S(Ar, DSF_IgnoreDiff);
-#endif
-			// Serialize a placeholder value, we will overwrite this with an offset to the end of the shader
-			Ar << SkipOffset;
-		}
-
-		if (bHandleShaderKeyChanges)
-		{
-			FSelfContainedShaderId SelfContainedKey = CurrentShader->GetId();
-			Ar << SelfContainedKey;
-		}
-
-		CurrentShader->SerializeBase(Ar, bInlineShaderResource, false);
-
-		// Get the offset to the end of the shader's serialized data
-		int64 EndOffset = Ar.Tell();
-		// Seek back to the placeholder and write the end offset
-		// This allows us to skip over the shader's serialized data at load time without knowing how to deserialize it
-		// Which can happen with shaders that were available at cook time, but not on the target platform (shaders in editor module for example)
-		Ar.Seek(SkipOffset);
-		Ar << EndOffset;
-		Ar.Seek(EndOffset);
-	}
-
-	inline FShader* SerializeShaderForLoad(FShaderType* Type, FArchive& Ar, bool bHandleShaderKeyChanges, bool bInlineShaderResource, bool bLoadedByCookedMaterial)
-	{
-		int64 EndOffset = 0;
-		Ar << EndOffset;
-
-		FSelfContainedShaderId SelfContainedKey;
-
-		if (bHandleShaderKeyChanges)
-		{
-			Ar << SelfContainedKey;
-		}
-
-		FShader* Shader = nullptr;
-		if (Type
-			// If we are handling shader key changes, only create the shader if the serialized key matches the key the shader would have if created
-			// This allows serialization changes between the save and load to be safely handled
-			&& (!bHandleShaderKeyChanges || SelfContainedKey.IsValid()))
-		{
-			Shader = Type->ConstructForDeserialization();
-			check(Shader != nullptr);
-			Shader->SerializeBase(Ar, bInlineShaderResource, bLoadedByCookedMaterial);
-		}
-		else
-		{
-			// Skip over this shader's serialized data if the type doesn't exist
-			// This can happen with shader types in modules that were loaded during cooking but not at run time (editor)
-			Ar.Seek(EndOffset);
-		}
-		return Shader;
-	}
-
-	/** 
-	 * Used to serialize a shader map inline in a material in a package. 
-	 * @param bInlineShaderResource - whether to inline the shader resource's serializations
-	 * @param bHandleShaderKeyChanges - whether to serialize the data necessary to detect and gracefully handle shader key changes between saving and loading
-	 * @param DependenciesToSave - array of specific ShaderTypeDepencies which should be saved.
-	 */
-	void SerializeInline(FArchive& Ar, bool bInlineShaderResource, bool bHandleShaderKeyChanges, bool bLoadedByCookedMaterial, const TArray<FShaderPrimaryKey>* ShaderKeysToSave = nullptr)
-	{
-		if (Ar.IsSaving())
-		{
-			TArray<FShaderPrimaryKey> SortedShaderKeys;
-
-			if (ShaderKeysToSave)
-			{
-				SortedShaderKeys = *ShaderKeysToSave;
-			}
-			else
-			{
-				Shaders.GenerateKeyArray(SortedShaderKeys);
-			}
-
-			int32 NumShaders = SortedShaderKeys.Num();
-			Ar << NumShaders;
-
-			// Sort the shaders by type name before saving, to make sure the saved result is binary equivalent to what is generated on other machines, 
-			// Which is a requirement of the Derived Data Cache.
-			SortedShaderKeys.Sort(FCompareShaderPrimaryKey());
-
-			for (FShaderPrimaryKey Key : SortedShaderKeys)
-			{
-				FShaderType* Type = Key.Type;
-				check(Type);
-				checkSlow(FName(Type->GetName()) != NAME_None);
-				Ar << Type;
-				FShader* CurrentShader = Shaders.FindChecked(Key);
-				SerializeShaderForSaving(CurrentShader, Ar, bHandleShaderKeyChanges, bInlineShaderResource);
-			}
-
-			TArray<FShaderPipeline*> SortedPipelines;
-			GetShaderPipelineList(SortedPipelines, FShaderPipeline::EAll);
-			int32 NumPipelines = SortedPipelines.Num();
-			Ar << NumPipelines;
-
-			checkf(!ShaderKeysToSave || NumPipelines == 0, TEXT("ShaderPipelines currently not supported for specific list of shader keys."));
-
-			// Sort the shader pipelines by type name before saving, to make sure the saved result is binary equivalent to what is generated on other machines, Which is a requirement of the Derived Data Cache.
-			SortedPipelines.Sort();
-			for (FShaderPipeline* CurrentPipeline : SortedPipelines)
-			{
-				const FShaderPipelineType* PipelineType = CurrentPipeline->PipelineType;
-				Ar << PipelineType;
-
-				auto& PipelineStages = PipelineType->GetStages();
-				int32 NumStages = PipelineStages.Num();
-				Ar << NumStages;
-				for (int32 Index = 0; Index < NumStages; ++Index)
-				{
-					auto* Shader = CurrentPipeline->GetShader(PipelineStages[Index]->GetFrequency());
-					FShaderType* Type = Shader->GetType();
-					Ar << Type;
-					SerializeShaderForSaving(Shader, Ar, bHandleShaderKeyChanges, bInlineShaderResource);
-				}
-#if WITH_EDITORONLY_DATA
-				if(Ar.IsCooking())
-				{
-					FShaderPipeline::CookPipeline(CurrentPipeline);
-				}
-#endif
-			}
-		}
-
-		if (Ar.IsLoading())
-		{
-			// Mark as unregistered - about to load new shaders that need to be registered later 
-			// on the game thread.
-			bHasBeenRegistered = false;
-
-			int32 NumShaders = 0;
-			Ar << NumShaders;
-
-			SerializedShaders.Reserve(NumShaders);
-			for (int32 ShaderIndex = 0; ShaderIndex < NumShaders; ShaderIndex++)
-			{
-				FShaderType* Type = nullptr;
-				Ar << Type;
-
-				FShader* Shader = SerializeShaderForLoad(Type, Ar, bHandleShaderKeyChanges, bInlineShaderResource, bLoadedByCookedMaterial);
-				if (Shader)
-				{
-					SerializedShaders.Add(Shader);
-				}
-			}
-
-			int32 NumPipelines = 0;
-			Ar << NumPipelines;
-			for (int32 PipelineIndex = 0; PipelineIndex < NumPipelines; ++PipelineIndex)
-			{
-				const FShaderPipelineType* ShaderPipelineType = nullptr;
-				Ar << ShaderPipelineType;
-				int32 NumStages = 0;
-				Ar << NumStages;
-				// Make a list of references so they can be deleted when going out of scope if needed
-				TArray< TRefCountPtr<FShader> > ShaderStages;
-				for (int32 Index = 0; Index < NumStages; ++Index)
-				{
-					FShaderType* Type = nullptr;
-					Ar << Type;
-					FShader* Shader = SerializeShaderForLoad(Type, Ar, bHandleShaderKeyChanges, bInlineShaderResource, bLoadedByCookedMaterial);
-					if (Shader)
-					{
-						ShaderStages.Add(Shader);
-					}
-				}
-
-				// ShaderPipelineType can be nullptr if the pipeline existed but now is gone!
-				if (ShaderPipelineType && ShaderStages.Num() == ShaderPipelineType->GetStages().Num())
-				{
-					FSerializedShaderPipeline* SerializedPipeline = new FSerializedShaderPipeline();
-					SerializedPipeline->ShaderPipelineType = ShaderPipelineType;
-					SerializedPipeline->ShaderStages = MoveTemp(ShaderStages);
-					SerializedShaderPipelines.Add(SerializedPipeline);
-				}
-			}
-		}
-	}
-
-	/** Registered all shaders that have been serialized (maybe) on another thread */
-	virtual void RegisterSerializedShaders(bool bCookedMaterial)
-	{
-		bHasBeenRegistered = true;
-		check(IsInGameThread());
-		for (FShader* Shader : SerializedShaders)
-		{
-			FShaderType* Type = Shader->GetType();
-			FShader* ExistingShader = Type->FindShaderByKey(Shader->GetKey());
-
-			if (ExistingShader != nullptr)
-			{
-				INC_DWORD_STAT(STAT_Shaders_NumShadersDuplicated);
-				delete Shader;
-				Shader = ExistingShader;
-			}
-			else
-			{
-				// Register the shader now that it is valid, so that it can be reused
-				INC_DWORD_STAT(STAT_Shaders_NumShadersRegistered);
-				Shader->Register(bCookedMaterial);
-			}
-			AddShader(Shader->GetType(), Shader->GetPermutationId(), Shader);
-		}
-		SerializedShaders.Empty();
-
-		for (FSerializedShaderPipeline* SerializedPipeline : SerializedShaderPipelines)
-		{
-			FShaderPipeline* ShaderPipeline = new FShaderPipeline(SerializedPipeline->ShaderPipelineType, SerializedPipeline->ShaderStages);
-			AddShaderPipeline(SerializedPipeline->ShaderPipelineType, ShaderPipeline);
-
-			delete SerializedPipeline;
-		}
-		SerializedShaderPipelines.Empty();
-	}
-
-	/** Discards serialized shaders when they are not going to be used for anything (NullRHI) */
-	virtual void DiscardSerializedShaders()
-	{
-		for (FShader* Shader : SerializedShaders)
-		{
-			delete Shader;
-		}
-		SerializedShaders.Empty();
-
-		for (FSerializedShaderPipeline* SerializedPipeline : SerializedShaderPipelines)
-		{
-			delete SerializedPipeline;
-		}
-		SerializedShaderPipelines.Empty();
-	}
-
 	/** @return true if the map is empty */
 	inline bool IsEmpty() const
 	{
-		check(bHasBeenRegistered);
 		return Shaders.Num() == 0;
 	}
 
 	/** @return The number of shaders in the map. */
-	inline uint32 GetNumShaders() const
-	{
-		check(bHasBeenRegistered);
-		return Shaders.Num();
-	}
+	uint32 GetNumShaders() const;
 
 	/** @return The number of shader pipelines in the map. */
 	inline uint32 GetNumShaderPipelines() const
 	{
-		check(bHasBeenRegistered);
 		return ShaderPipelines.Num();
 	}
 
 	/** clears out all shaders and deletes shader pipelines held in the map */
-	void Empty()
-	{
-		Shaders.Empty();
-		EmptyShaderPipelines();
-	}
+	void Empty();
 
-	inline FShaderPipeline* GetShaderPipeline(const FShaderPipelineType* PipelineType)
+	inline FShaderPipeline* GetShaderPipeline(const FHashedName& PipelineTypeName) const
 	{
-		check(bHasBeenRegistered);
-		FShaderPipeline** Found = ShaderPipelines.Find(PipelineType);
-		return Found ? *Found : nullptr;
+		const int32 Index = Algo::BinarySearchBy(ShaderPipelines, PipelineTypeName, FProjectShaderPipelineToKey());
+		return (Index != INDEX_NONE) ? ShaderPipelines[Index].Get() : nullptr;
 	}
 
 	inline FShaderPipeline* GetShaderPipeline(const FShaderPipelineType* PipelineType) const
 	{
-		check(bHasBeenRegistered);
-		FShaderPipeline* const* Found = ShaderPipelines.Find(PipelineType);
-		return Found ? *Found : nullptr;
+		return GetShaderPipeline(PipelineType->GetHashedName());
 	}
 
-	// Returns nullptr if not found
-	inline bool HasShaderPipeline(const FShaderPipelineType* PipelineType) const
-	{
-		check(bHasBeenRegistered);
-		return (GetShaderPipeline(PipelineType) != nullptr);
-	}
+	inline bool HasShaderPipeline(const FHashedName& PipelineTypeName) const { return GetShaderPipeline(PipelineTypeName) != nullptr; }
+	inline bool HasShaderPipeline(const FShaderPipelineType* PipelineType) const { return (GetShaderPipeline(PipelineType) != nullptr); }
+
+	uint32 GetMaxNumInstructionsForShader(const FShaderMapBase& InShaderMap, FShaderType* ShaderType) const;
+
+protected:
+	void EmptyShaderPipelines();
+
+	LAYOUT_FIELD(TMemoryImageArray<TMemoryImagePtr<FShader>>, Shaders);
+	LAYOUT_FIELD(TMemoryImageArray<TMemoryImagePtr<FShaderPipeline>>, ShaderPipelines);
+	/** The platform this shader map was compiled with */
+	LAYOUT_FIELD(TEnumAsByte<EShaderPlatform>, Platform);
+};
+
+class RENDERCORE_API FShaderMapBase
+{
+public:
+	virtual ~FShaderMapBase();
+
+	FShaderMapResourceCode* GetResourceCode();
 
-	inline void AddShaderPipeline(const FShaderPipelineType* Type, FShaderPipeline* ShaderPipeline)
+	inline FShaderMapResource* GetResource() const { return Resource; }
+	inline FShaderMapResource* GetResourceChecked() const { check(Resource); return Resource; }
+	inline const FShaderMapPointerTable& GetPointerTable() const { check(PointerTable); return *PointerTable; }
+	inline const FShaderMapContent* GetContent() const { return Content; }
+	inline FShaderMapContent* GetMutableContent()
 	{
-		check(bHasBeenRegistered);
-		check(Type);
-		check(!ShaderPipeline || ShaderPipeline->PipelineType == Type);
-		ShaderPipelines.Add(Type, ShaderPipeline);
+		UnfreezeContent();
+		return Content;
 	}
 
-	uint32 GetMaxNumInstructionsForShader(FShaderType* ShaderType) const
-	{
-		check(bHasBeenRegistered);
-		uint32 MaxNumInstructions = 0;
-		const TRefCountPtr<FShader>* FoundShader = Shaders.Find(FShaderPrimaryKey(ShaderType, 0));
-		if (FoundShader && *FoundShader)
-		{
-			MaxNumInstructions = FMath::Max(MaxNumInstructions, (*FoundShader)->GetNumInstructions());
-		}
+	inline EShaderPlatform GetShaderPlatform() const { return Content ? Content->GetShaderPlatform() : SP_NumPlatforms; }
+	inline uint32 GetFrozenContentSize() const { return FrozenContentSize; }
 
-		for (auto& Pair : ShaderPipelines)
-		{
-			FShaderPipeline* Pipeline = Pair.Value;
-			auto* Shader = Pipeline->GetShader(ShaderType->GetFrequency());
-			if (Shader)
-			{
-				MaxNumInstructions = FMath::Max(MaxNumInstructions, Shader->GetNumInstructions());
-			}
-		}
+	void AssignContent(FShaderMapContent* InContent);
+	void FinalizeContent();
+	void UnfreezeContent();
+	void Serialize(FArchive& Ar, bool bInlineShaderResources, bool bLoadedByCookedMaterial);
 
-		return MaxNumInstructions;
+#if WITH_EDITOR
+	inline void GetOutdatedTypes(TArray<const FShaderType*>& OutdatedShaderTypes, TArray<const FShaderPipelineType*>& OutdatedShaderPipelineTypes, TArray<const FVertexFactoryType*>& OutdatedFactoryTypes) const
+	{
+		Content->GetOutdatedTypes(*this, OutdatedShaderTypes, OutdatedShaderPipelineTypes, OutdatedFactoryTypes);
 	}
-
-protected:
-	inline void EmptyShaderPipelines()
+	void SaveShaderStableKeys(EShaderPlatform TargetShaderPlatform, const struct FStableShaderKeyAndValue& SaveKeyVal)
 	{
-		for (auto& Pair : ShaderPipelines)
-		{
-			if (FShaderPipeline* Pipeline = Pair.Value)
-			{
-				delete Pipeline;
-				Pipeline = nullptr;
-			}
-		}
-		ShaderPipelines.Empty();
+		Content->SaveShaderStableKeys(*this, TargetShaderPlatform, SaveKeyVal);
 	}
+#endif // WITH_EDITOR
+
+protected:
+	explicit FShaderMapBase(const FTypeLayoutDesc& InContentTypeLayout);
+	void DestroyContent();
+
+protected:
+	virtual FShaderMapPointerTable* CreatePointerTable() const = 0;
+
+private:
+	const FTypeLayoutDesc& ContentTypeLayout;
+	TRefCountPtr<FShaderMapResource> Resource;
+	TRefCountPtr<FShaderMapResourceCode> Code;
+	FShaderMapPointerTable* PointerTable;
+	FShaderMapContent* Content;
+	uint32 FrozenContentSize;
+	uint32 NumFrozenShaders;
+};
+
+template<typename ContentType, typename PointerTableType = FShaderMapPointerTable>
+class TShaderMap : public FShaderMapBase
+{
+public:
+	inline const PointerTableType& GetPointerTable() const { return static_cast<const PointerTableType&>(FShaderMapBase::GetPointerTable()); }
+	inline const ContentType* GetContent() const { return static_cast<const ContentType*>(FShaderMapBase::GetContent()); }
+	inline ContentType* GetMutableContent() { return static_cast<ContentType*>(FShaderMapBase::GetMutableContent()); }
 
-	TMap<FShaderPrimaryKey, TRefCountPtr<FShader> > Shaders;
-	TMap<const FShaderPipelineType*, FShaderPipeline*> ShaderPipelines;
+protected:
+	TShaderMap() : FShaderMapBase(StaticGetTypeLayoutDesc<ContentType>()) {}
+	virtual FShaderMapPointerTable* CreatePointerTable() const final override { return new PointerTableType(); }
 };
 
+template<typename ShaderType, typename PointerTableType>
+inline const PointerTableType& TShaderRefBase<ShaderType, PointerTableType>::GetPointerTable() const
+{
+	checkSlow(ShaderMap);
+	return static_cast<const PointerTableType&>(ShaderMap->GetPointerTable());
+}
+
+template<typename ShaderType, typename PointerTableType>
+inline FShaderMapResource* TShaderRefBase<ShaderType, PointerTableType>::GetResource() const
+{
+	checkSlow(ShaderMap);
+	return ShaderMap->GetResource();
+}
+
+inline const FShaderMapPointerTable& FShaderPipelineRef::GetPointerTable() const
+{
+	checkSlow(ShaderMap);
+	return ShaderMap->GetPointerTable();
+}
+
+inline FShaderMapResource* FShaderPipelineRef::GetResource() const
+{
+	checkSlow(ShaderMap);
+	return ShaderMap->GetResource();
+}
+
 /** A reference which is initialized with the requested shader type from a shader map. */
 template<typename ShaderType>
-class TShaderMapRef
+class TShaderMapRef : public TShaderRef<ShaderType>
 {
 public:
-	TShaderMapRef(const TShaderMap<typename ShaderType::ShaderMetaType>* ShaderIndex)
-		: Shader(ShaderIndex->template GetShader<ShaderType>(/* PermutationId = */ 0)) // gcc3 needs the template quantifier so it knows the < is not a less-than
+	TShaderMapRef(const typename ShaderType::ShaderMapType* ShaderIndex)
+		: TShaderRef<ShaderType>(ShaderIndex->template GetShader<ShaderType>(/* PermutationId = */ 0)) // gcc3 needs the template quantifier so it knows the < is not a less-than
 	{
 		static_assert(
 			TIsSame<typename ShaderType::FPermutationDomain, FShaderPermutationNone>::Value,
@@ -2447,47 +2036,20 @@ public:
 	}
 
 	TShaderMapRef(
-		const TShaderMap<typename ShaderType::ShaderMetaType>* ShaderIndex,
+		const typename ShaderType::ShaderMapType* ShaderIndex,
 		const typename ShaderType::FPermutationDomain& PermutationVector)
-		: Shader(ShaderIndex->template GetShader<ShaderType>(PermutationVector.ToDimensionValueId())) // gcc3 needs the template quantifier so it knows the < is not a less-than
+		: TShaderRef<ShaderType>(ShaderIndex->template GetShader<ShaderType>(PermutationVector.ToDimensionValueId())) // gcc3 needs the template quantifier so it knows the < is not a less-than
 	{ }
-
-	FORCEINLINE ShaderType* operator->() const
-	{
-		return Shader;
-	}
-
-	FORCEINLINE ShaderType* operator*() const
-	{
-		return Shader;
-	}
-
-private:
-	ShaderType* Shader;
 };
 
 /** A reference to an optional shader, initialized with a shader type from a shader map if it is available or nullptr if it is not. */
 template<typename ShaderType>
-class TOptionalShaderMapRef
+class TOptionalShaderMapRef : public TShaderRef<ShaderType>
 {
 public:
-	TOptionalShaderMapRef(const TShaderMap<typename ShaderType::ShaderMetaType>* ShaderIndex):
-	Shader((ShaderType*)ShaderIndex->GetShader(&ShaderType::StaticType)) // gcc3 needs the template quantifier so it knows the < is not a less-than
+	TOptionalShaderMapRef(const typename ShaderType::ShaderMapType* ShaderIndex):
+		TShaderRef<ShaderType>(TShaderRef<ShaderType>::Cast(ShaderIndex->GetShader(&ShaderType::StaticType))) // gcc3 needs the template quantifier so it knows the < is not a less-than
 	{}
-	FORCEINLINE bool IsValid() const
-	{
-		return Shader != nullptr;
-	}
-	FORCEINLINE ShaderType* operator->() const
-	{
-		return Shader;
-	}
-	FORCEINLINE ShaderType* operator*() const
-	{
-		return Shader;
-	}
-private:
-	ShaderType* Shader;
 };
 
 /** Tracks state when traversing a FSerializationHistory. */
@@ -2668,7 +2230,7 @@ extern RENDERCORE_API void DumpShaderPipelineStats(EShaderPlatform Platform);
  * @param ShaderTypeName - The name of the shader type to find.
  * @return The shader type, or NULL if none matched.
  */
-extern RENDERCORE_API FShaderType* FindShaderTypeByName(FName ShaderTypeName);
+extern RENDERCORE_API FShaderType* FindShaderTypeByName(const FHashedName& ShaderTypeName);
 
 /** Helper function to dispatch a compute shader while checking that parameters have been set correctly. */
 extern RENDERCORE_API void DispatchComputeShader(
@@ -2692,5 +2254,34 @@ extern RENDERCORE_API void DispatchIndirectComputeShader(
 	FRHIVertexBuffer* ArgumentBuffer,
 	uint32 ArgumentOffset);
 
+inline void DispatchComputeShader(
+	FRHICommandList& RHICmdList,
+	const TShaderRef<FShader>& Shader,
+	uint32 ThreadGroupCountX,
+	uint32 ThreadGroupCountY,
+	uint32 ThreadGroupCountZ)
+{
+	DispatchComputeShader(RHICmdList, Shader.GetShader(), ThreadGroupCountX, ThreadGroupCountY, ThreadGroupCountZ);
+}
+
+inline void DispatchComputeShader(
+	FRHIAsyncComputeCommandListImmediate& RHICmdList,
+	const TShaderRef<FShader>& Shader,
+	uint32 ThreadGroupCountX,
+	uint32 ThreadGroupCountY,
+	uint32 ThreadGroupCountZ)
+{
+	DispatchComputeShader(RHICmdList, Shader.GetShader(), ThreadGroupCountX, ThreadGroupCountY, ThreadGroupCountZ);
+}
+
+inline void DispatchIndirectComputeShader(
+	FRHICommandList& RHICmdList,
+	const TShaderRef<FShader>& Shader,
+	FRHIVertexBuffer* ArgumentBuffer,
+	uint32 ArgumentOffset)
+{
+	DispatchIndirectComputeShader(RHICmdList, Shader.GetShader(), ArgumentBuffer, ArgumentOffset);
+}
+
 /** Appends to KeyString for all shaders. */
 extern RENDERCORE_API void ShaderMapAppendKeyString(EShaderPlatform Platform, FString& KeyString);
