Copying Private-LoadTimes-4.24 stream to Main. Biggest changes are in Materials/Shader memory freezing.

diff --git a/Engine/Source/Runtime/RenderCore/Public/ShaderCodeArchive.h b/Engine/Source/Runtime/RenderCore/Public/ShaderCodeArchive.h
new file mode 100644
--- /dev/null
+++ b/Engine/Source/Runtime/RenderCore/Public/ShaderCodeArchive.h
@@ -0,0 +1,173 @@
+#pragma once
+
+#include "CoreMinimal.h"
+#include "RHI.h"
+#include "FileCache/FileCache.h"
+#include "Containers/HashTable.h"
+
+struct FShaderMapEntry
+{
+	uint32 ShaderIndicesOffset = 0u;
+	uint32 NumShaders = 0u;
+	uint32 FirstPreloadIndex = 0u;
+	uint32 NumPreloadEntries = 0u;
+
+	friend FArchive& operator <<(FArchive& Ar, FShaderMapEntry& Ref)
+	{
+		return Ar << Ref.ShaderIndicesOffset << Ref.NumShaders << Ref.FirstPreloadIndex << Ref.NumPreloadEntries;
+	}
+};
+
+static FArchive& operator <<(FArchive& Ar, FFileCachePreloadEntry& Ref)
+{
+	return Ar << Ref.Offset << Ref.Size;
+}
+
+struct FShaderCodeEntry
+{
+	uint64 Offset = 0;
+	uint32 Size = 0;
+	uint32 UncompressedSize = 0;
+	uint8 Frequency;
+
+	EShaderFrequency GetFrequency() const
+	{
+		return (EShaderFrequency)Frequency;
+	}
+
+	friend FArchive& operator <<(FArchive& Ar, FShaderCodeEntry& Ref)
+	{
+		return Ar << Ref.Offset << Ref.Size << Ref.UncompressedSize << Ref.Frequency;
+	}
+};
+
+// Portion of shader code archive that's serialize to disk
+class RENDERCORE_API FSerializedShaderArchive
+{
+public:
+	TArray<FSHAHash> ShaderMapHashes;
+	TArray<FSHAHash> ShaderHashes;
+	TArray<FShaderMapEntry> ShaderMapEntries;
+	TArray<FShaderCodeEntry> ShaderEntries;
+	TArray<FFileCachePreloadEntry> PreloadEntries;
+	TArray<uint32> ShaderIndices;
+	FHashTable ShaderMapHashTable;
+	FHashTable ShaderHashTable;
+
+	FSerializedShaderArchive() : ShaderMapHashTable(0u), ShaderHashTable(0u) {}
+
+	uint32 GetAllocatedSize() const
+	{
+		return ShaderHashes.GetAllocatedSize() +
+			ShaderEntries.GetAllocatedSize() +
+			ShaderMapHashes.GetAllocatedSize() +
+			ShaderMapEntries.GetAllocatedSize() +
+			PreloadEntries.GetAllocatedSize() +
+			ShaderIndices.GetAllocatedSize();
+	}
+
+	void Empty()
+	{
+		ShaderHashes.Empty();
+		ShaderEntries.Empty();
+		ShaderMapHashes.Empty();
+		ShaderMapEntries.Empty();
+		PreloadEntries.Empty();
+		ShaderIndices.Empty();
+		ShaderMapHashTable.Clear();
+		ShaderHashTable.Clear();
+	}
+
+	int32 GetNumShaderMaps() const
+	{
+		return ShaderMapEntries.Num();
+	}
+
+	int32 GetNumShaders() const
+	{
+		return ShaderEntries.Num();
+	}
+
+	int32 FindShaderMapWithKey(const FSHAHash& Hash, uint32 Key) const;
+	int32 FindShaderMap(const FSHAHash& Hash) const;
+	bool FindOrAddShaderMap(const FSHAHash& Hash, int32& OutIndex);
+
+	int32 FindShaderWithKey(const FSHAHash& Hash, uint32 Key) const;
+	int32 FindShader(const FSHAHash& Hash) const;
+	bool FindOrAddShader(const FSHAHash& Hash, int32& OutIndex);
+
+	void DecompressShader(int32 Index, const TArray<TArray<uint8>>& ShaderCode, TArray<uint8>& OutDecompressedShader) const;
+
+	void Finalize();
+	void Serialize(FArchive& Ar);
+
+	friend FArchive& operator<<(FArchive& Ar, FSerializedShaderArchive& Ref)
+	{
+		Ref.Serialize(Ar);
+		return Ar;
+	}
+};
+
+class FShaderCodeArchive : public FRHIShaderLibrary
+{
+public:
+	static FShaderCodeArchive* Create(EShaderPlatform InPlatform, FArchive& Ar, const FString& InDestFilePath, const FString& InLibraryDir, const FString& InLibraryName);
+
+	virtual ~FShaderCodeArchive();
+
+	virtual bool IsNativeLibrary() const override { return false; }
+
+	uint32 GetSizeBytes() const
+	{
+		return sizeof(*this) +
+			SerializedShaders.GetAllocatedSize() +
+			SerializedShaders.GetAllocatedSize();
+	}
+
+	virtual int32 GetNumShaders() const override { return SerializedShaders.ShaderEntries.Num(); }
+	virtual int32 GetNumShaderMaps() const override { return SerializedShaders.ShaderMapEntries.Num(); }
+	virtual int32 GetNumShadersForShaderMap(int32 ShaderMapIndex) const override { return SerializedShaders.ShaderMapEntries[ShaderMapIndex].NumShaders; }
+
+	virtual int32 GetShaderIndex(int32 ShaderMapIndex, int32 i) const override
+	{
+		const FShaderMapEntry& ShaderMapEntry = SerializedShaders.ShaderMapEntries[ShaderMapIndex];
+		return SerializedShaders.ShaderIndices[ShaderMapEntry.ShaderIndicesOffset + i];
+	}
+
+	virtual int32 FindShaderMapIndex(const FSHAHash& Hash) override
+	{
+		return SerializedShaders.FindShaderMap(Hash);
+	}
+
+	virtual int32 FindShaderIndex(const FSHAHash& Hash) override
+	{
+		return SerializedShaders.FindShader(Hash);
+	}
+
+	virtual FGraphEventRef PreloadShader(int32 ShaderIndex) override;
+
+	virtual FGraphEventRef PreloadShaderMap(int32 ShaderMapIndex) override;
+
+	virtual TRefCountPtr<FRHIShader> CreateShader(int32 Index) override;
+
+protected:
+	FShaderCodeArchive(EShaderPlatform InPlatform, const FString& InLibraryDir, const FString& InLibraryName);
+
+	IMemoryReadStreamRef ReadShaderCode(int32 Index);
+
+	FORCENOINLINE void CheckShaderCreation(void* ShaderPtr, int32 Index)
+	{
+	}
+
+	// Library directory
+	FString LibraryDir;
+
+	// Offset at where shader code starts in a code library
+	int64 LibraryCodeOffset;
+
+	// Library file handle for async reads
+	IFileCacheHandle* FileCacheHandle;
+
+	// The shader code present in the library
+	FSerializedShaderArchive SerializedShaders;
+};
diff --git a/Engine/Source/Runtime/RenderCore/Public/ShaderCodeLibrary.h b/Engine/Source/Runtime/RenderCore/Public/ShaderCodeLibrary.h
--- a/Engine/Source/Runtime/RenderCore/Public/ShaderCodeLibrary.h
+++ b/Engine/Source/Runtime/RenderCore/Public/ShaderCodeLibrary.h
@@ -7,6 +7,7 @@
 #pragma once
 
 #include "CoreMinimal.h"
+#include "Containers/ArrayView.h"
 #include "RHI.h"
 
 DECLARE_LOG_CATEGORY_EXTERN(LogShaderLibrary, Log, All);
@@ -14,6 +15,8 @@ DECLARE_LOG_CATEGORY_EXTERN(LogShaderLibrary, Log, All);
 class FAnsiStringBuilderBase;
 class FShaderPipeline;
 class FStringView;
+class FShaderMapResource;
+class FShaderMapResourceCode;
 
 struct RENDERCORE_API FShaderCodeLibraryPipeline
 {
@@ -23,9 +26,9 @@ struct RENDERCORE_API FShaderCodeLibraryPipeline
 	FSHAHash HullShader;
 	FSHAHash DomainShader;
 	mutable uint32 Hash;
-
+	
 	/** Fills the hashes from the pipeline stage shaders */
-	void Initialize(FShaderPipeline* Pipeline);
+	void Initialize(const FShaderPipeline* Pipeline);
 
 	FShaderCodeLibraryPipeline() : Hash(0) {}
 	
@@ -105,7 +108,7 @@ struct RENDERCORE_API FStableShaderKeyAndValue
 	static FString HeaderLine();
 
 	/** Computes pipeline hash from the passed pipeline. Pass nullptr to clear */
-	void SetPipelineHash(FShaderPipeline* Pipeline);
+	void SetPipelineHash(const FShaderPipeline* Pipeline);
 
 	friend bool operator ==(const FStableShaderKeyAndValue& A, const FStableShaderKeyAndValue& B)
 	{
@@ -130,22 +133,6 @@ struct RENDERCORE_API FStableShaderKeyAndValue
 
 };
 
-class FShaderFactoryInterface : public FRHIShaderLibrary
-{
-public:
-	FShaderFactoryInterface(EShaderPlatform InPlatform, FString const& Name) : FRHIShaderLibrary(InPlatform, Name) {}
-	
-	virtual bool IsNativeLibrary() const override final {return false;}
-	
-	virtual FPixelShaderRHIRef CreatePixelShader(const FSHAHash& Hash) = 0; 
-	virtual FVertexShaderRHIRef CreateVertexShader(const FSHAHash& Hash) = 0;
-	virtual FHullShaderRHIRef CreateHullShader(const FSHAHash& Hash) = 0;
-	virtual FDomainShaderRHIRef CreateDomainShader(const FSHAHash& Hash) = 0;
-	virtual FGeometryShaderRHIRef CreateGeometryShader(const FSHAHash& Hash) = 0;
-	virtual FComputeShaderRHIRef CreateComputeShader(const FSHAHash& Hash) = 0;
-	virtual FRayTracingShaderRHIRef CreateRayTracingShader(EShaderFrequency Frequency, const FSHAHash& Hash) = 0;
-};
-
 DECLARE_MULTICAST_DELEGATE_TwoParams(FSharedShaderCodeRequest, const FSHAHash&, FArchive*);
 DECLARE_MULTICAST_DELEGATE_OneParam(FSharedShaderCodeRelease, const FSHAHash&);
 
@@ -167,59 +154,25 @@ struct RENDERCORE_API FShaderCodeLibrary
 	// For cooking, after this point any AddShaderCode/AddShaderPipeline calls will be invalid until OpenLibrary is called again.
 	// At runtime this will release the library data and further requests for shaders from this library will fail.
 	static void CloseLibrary(FString const& Name);
-	
-	/** Instantiate or retrieve a vertex shader from the cache for the provided code & hash. */
-	static FVertexShaderRHIRef CreateVertexShader(EShaderPlatform Platform, FSHAHash Hash, TArray<uint8> const& Code);
-	/** Instantiate or retrieve a pixel shader from the cache for the provided code & hash. */
-	static FPixelShaderRHIRef CreatePixelShader(EShaderPlatform Platform, FSHAHash Hash, TArray<uint8> const& Code);
-	/** Instantiate or retrieve a geometry shader from the cache for the provided code & hash. */
-	static FGeometryShaderRHIRef CreateGeometryShader(EShaderPlatform Platform, FSHAHash Hash, TArray<uint8> const& Code);
-	/** Instantiate or retrieve a hull shader from the cache for the provided code & hash. */
-	static FHullShaderRHIRef CreateHullShader(EShaderPlatform Platform, FSHAHash Hash, TArray<uint8> const& Code);
-	/** Instantiate or retrieve a domain shader from the cache for the provided code & hash. */
-	static FDomainShaderRHIRef CreateDomainShader(EShaderPlatform Platform, FSHAHash Hash, TArray<uint8> const& Code);
-	/** Instantiate or retrieve a compute shader from the cache for the provided code & hash. */
-	static FComputeShaderRHIRef CreateComputeShader(EShaderPlatform Platform, FSHAHash Hash, TArray<uint8> const& Code);
-	/** Instantiate or retrieve a ray tracing shader from the cache for the provided code & hash. */
-	static FRayTracingShaderRHIRef CreateRayTracingShader(EShaderPlatform Platform, EShaderFrequency Frequency, FSHAHash Hash, TArray<uint8> const& Code);
 
     static bool ContainsShaderCode(const FSHAHash& Hash);
-    
-	// Place a request to preload shader code
-	// Blocking call if no Archive is provided or Archive is not a type of FLinkerLoad
-	// Shader code preload will be finished before owning UObject PostLoad call
-	static bool RequestShaderCode(const FSHAHash& Hash, FArchive* Ar);
-
-	// Note that we skipped preloading shader code.
-	// All this does is call the delegate so that other folks are aware that it was lazy
-	static bool LazyRequestShaderCode(const FSHAHash& Hash, FArchive* Ar);
-
-	// Get the raw payload synchronously
-	// This does NOT require a ReleaseShaderCode; because it is synchronous
-	// This also does not fire any delegates...which are used to load binary programs (we are calling this because we failed to find a binary program)
-	static bool RequestShaderCode(const FSHAHash& Hash, TArray<uint8>& OutRaw);
-
-	// Request to release shader code
-	// Must match RequestShaderCode call
-	// Invalid to call before owning UObject PostLoad call
-	static void ReleaseShaderCode(const FSHAHash& Hash);
-
-	// Request to release shader code that we lazy loaded
-	// Must match LazyRequestShaderCode call
-	// All this does is call the delegate so that other folks are aware that it was lazy
-	static void LazyReleaseShaderCode(const FSHAHash& Hash);
-
-	// Create an iterator over all the shaders in the library
-	static TRefCountPtr<FRHIShaderLibrary::FShaderLibraryIterator> CreateIterator(void);
-	
+
+	static TRefCountPtr<FShaderMapResource> LoadResource(const FSHAHash& Hash, FArchive* Ar);
+
+	static bool PreloadShader(const FSHAHash& Hash, FArchive* Ar);
+
+	static FVertexShaderRHIRef CreateVertexShader(EShaderPlatform Platform, const FSHAHash& Hash);
+	static FPixelShaderRHIRef CreatePixelShader(EShaderPlatform Platform, const FSHAHash& Hash);
+	static FHullShaderRHIRef CreateHullShader(EShaderPlatform Platform, const FSHAHash& Hash);
+	static FDomainShaderRHIRef CreateDomainShader(EShaderPlatform Platform, const FSHAHash& Hash);
+	static FGeometryShaderRHIRef CreateGeometryShader(EShaderPlatform Platform, const FSHAHash& Hash);
+	static FComputeShaderRHIRef CreateComputeShader(EShaderPlatform Platform, const FSHAHash& Hash);
+
 	// Total number of shader entries in the library
 	static uint32 GetShaderCount(void);
 	
 	// The shader platform that the library manages - at runtime this will only be one
 	static EShaderPlatform GetRuntimeShaderPlatform(void);
-	
-	// Get the shader pipelines in the library - only ever valid for OpenGL which can link without full PSO state
-	static TSet<FShaderCodeLibraryPipeline> const* GetShaderPipelines(EShaderPlatform Platform);
 
 #if WITH_EDITOR
 	// Initialize the library cooker
@@ -233,7 +186,7 @@ struct RENDERCORE_API FShaderCodeLibrary
     static void CookShaderFormats(TArray<TTuple<FName,bool>> const& ShaderFormats);
 	
 	// At cook time, add shader code to collection
-	static bool AddShaderCode(EShaderPlatform ShaderPlatform, EShaderFrequency Frequency, const FSHAHash& Hash, const TArray<uint8>& InCode, uint32 const UncompressedSize);
+	static bool AddShaderCode(EShaderPlatform ShaderPlatform, const FShaderMapResourceCode* Code);
 
 	// We check this early in the callstack to avoid creating a bunch of FName and keys and things we will never save anyway. 
 	// Pass the shader platform to check or EShaderPlatform::SP_NumPlatforms to check if any of the registered types require
@@ -243,9 +196,6 @@ struct RENDERCORE_API FShaderCodeLibrary
 	// At cook time, add the human readable key value information
 	static void AddShaderStableKeyValue(EShaderPlatform ShaderPlatform, FStableShaderKeyAndValue& StableKeyValue);
 
-	// At cook time, add shader pipeline to collection
-	static bool AddShaderPipeline(FShaderPipeline* Pipeline);
-	
 	// Save collected shader code to a file for each specified shader platform, collating all child cooker results.
 	static bool SaveShaderCodeMaster(const FString& OutputDir, const FString& MetaOutputDir, const TArray<FName>& ShaderFormats, TArray<FString>& OutSCLCSVPath);
 	
@@ -268,8 +218,4 @@ struct RENDERCORE_API FShaderCodeLibrary
 	// Delegate called whenever shader code is requested.
 	static FDelegateHandle RegisterSharedShaderCodeRequestDelegate_Handle(const FSharedShaderCodeRequest::FDelegate& Delegate);
 	static void UnregisterSharedShaderCodeRequestDelegate_Handle(FDelegateHandle Handle);
-
-	// Delegate called whenever shader code is released.
-	static FDelegateHandle RegisterSharedShaderCodeReleaseDelegate_Handle(const FSharedShaderCodeRelease::FDelegate& Delegate);
-	static void UnregisterSharedShaderCodeReleaseDelegate_Handle(FDelegateHandle Handle);
 };
diff --git a/Engine/Source/Runtime/RenderCore/Public/ShaderCore.h b/Engine/Source/Runtime/RenderCore/Public/ShaderCore.h
--- a/Engine/Source/Runtime/RenderCore/Public/ShaderCore.h
+++ b/Engine/Source/Runtime/RenderCore/Public/ShaderCore.h
@@ -69,7 +69,6 @@ DECLARE_FLOAT_ACCUMULATOR_STAT_EXTERN(TEXT("Total RT Shader Init Time"),STAT_Sha
 DECLARE_CYCLE_STAT_EXTERN(TEXT("Frame RT Shader Init Time"),STAT_Shaders_FrameRTShaderInitForRenderingTime,STATGROUP_Shaders, RENDERCORE_API);
 DECLARE_MEMORY_STAT_EXTERN(TEXT("Shader Memory"),STAT_Shaders_ShaderMemory,STATGROUP_Shaders, RENDERCORE_API);
 DECLARE_MEMORY_STAT_EXTERN(TEXT("Shader Resource Mem"),STAT_Shaders_ShaderResourceMemory,STATGROUP_Shaders, RENDERCORE_API);
-DECLARE_MEMORY_STAT_EXTERN(TEXT("Shader MapMemory"),STAT_Shaders_ShaderMapMemory,STATGROUP_Shaders, RENDERCORE_API);
 
 DECLARE_DWORD_ACCUMULATOR_STAT_EXTERN(TEXT("Num Shaders Registered"), STAT_Shaders_NumShadersRegistered, STATGROUP_Shaders, RENDERCORE_API);
 DECLARE_DWORD_ACCUMULATOR_STAT_EXTERN(TEXT("Num Shaders Duplicated"), STAT_Shaders_NumShadersDuplicated, STATGROUP_Shaders, RENDERCORE_API);
@@ -160,6 +159,7 @@ struct FShaderTarget
 		return ((Target.Frequency << SP_NumBits) | Target.Platform);
 	}
 };
+DECLARE_INTRINSIC_TYPE_LAYOUT(FShaderTarget);
 
 static_assert(sizeof(FShaderTarget) == sizeof(uint32), "FShaderTarget is expected to be bit-packed into a single uint32.");
 
diff --git a/Engine/Source/Runtime/RenderCore/Public/ShaderParameterMacros.h b/Engine/Source/Runtime/RenderCore/Public/ShaderParameterMacros.h
--- a/Engine/Source/Runtime/RenderCore/Public/ShaderParameterMacros.h
+++ b/Engine/Source/Runtime/RenderCore/Public/ShaderParameterMacros.h
@@ -658,7 +658,7 @@ struct TShaderParameterStructTypeInfo<StructType[InNumElements]>
 #define INTERNAL_SHADER_PARAMETER_GET_STRUCT_METADATA(StructTypeName) \
 	static FShaderParametersMetadata StaticStructMetadata(\
 		FShaderParametersMetadata::EUseCase::ShaderParameterStruct, \
-		FName(TEXT(#StructTypeName)), \
+		TEXT(#StructTypeName), \
 		TEXT(#StructTypeName), \
 		nullptr, \
 		nullptr, \
@@ -738,6 +738,7 @@ extern RENDERCORE_API FShaderParametersMetadata* FindUniformBufferStructByFName(
 /** Finds the FShaderParameterMetadata corresponding to the given uniform buffer layout hash, or null if not found. */
 extern RENDERCORE_API FShaderParametersMetadata* FindUniformBufferStructByLayoutHash(uint32 Hash);
 
+extern RENDERCORE_API FShaderParametersMetadata* FindUniformBufferStructByShaderVariableName(const FHashedName& Name);
 
 /** Begins & ends a shader parameter structure.
  *
@@ -786,7 +787,7 @@ extern RENDERCORE_API FShaderParametersMetadata* FindUniformBufferStructByLayout
 #define IMPLEMENT_UNIFORM_BUFFER_STRUCT(StructTypeName,ShaderVariableName) \
 	FShaderParametersMetadata StructTypeName::StaticStructMetadata( \
 	FShaderParametersMetadata::EUseCase::UniformBuffer, \
-	FName(TEXT(#StructTypeName)), \
+	TEXT(#StructTypeName), \
 	TEXT(#StructTypeName), \
 	TEXT(ShaderVariableName), \
 	nullptr, \
@@ -812,7 +813,7 @@ extern RENDERCORE_API FShaderParametersMetadata* FindUniformBufferStructByLayout
 #define IMPLEMENT_STATIC_UNIFORM_BUFFER_STRUCT(StructTypeName,ShaderVariableName,StaticSlotName) \
 	FShaderParametersMetadata StructTypeName::StaticStructMetadata( \
 	FShaderParametersMetadata::EUseCase::UniformBuffer, \
-	FName(TEXT(#StructTypeName)), \
+	TEXT(#StructTypeName), \
 	TEXT(#StructTypeName), \
 	TEXT(ShaderVariableName), \
 	TEXT(#StaticSlotName), \
diff --git a/Engine/Source/Runtime/RenderCore/Public/ShaderParameterMetadata.h b/Engine/Source/Runtime/RenderCore/Public/ShaderParameterMetadata.h
--- a/Engine/Source/Runtime/RenderCore/Public/ShaderParameterMetadata.h
+++ b/Engine/Source/Runtime/RenderCore/Public/ShaderParameterMetadata.h
@@ -10,6 +10,7 @@
 #include "Containers/List.h"
 #include "Containers/StaticArray.h"
 #include "RHI.h"
+#include "Serialization/MemoryLayout.h"
 
 namespace EShaderPrecisionModifier
 {
@@ -185,7 +186,7 @@ public:
 	/** Initialization constructor. */
 	FShaderParametersMetadata(
 		EUseCase UseCase,
-		const FName& InLayoutName,
+		const TCHAR* InLayoutName,
 		const TCHAR* InStructTypeName,
 		const TCHAR* InShaderVariableName,
 		const TCHAR* InStaticSlotName,
@@ -200,6 +201,7 @@ public:
 
 	const TCHAR* GetStructTypeName() const { return StructTypeName; }
 	const TCHAR* GetShaderVariableName() const { return ShaderVariableName; }
+	const FHashedName& GetShaderVariableHashedName() const { return ShaderVariableHashedName; }
 	const TCHAR* GetStaticSlotName() const { return StaticSlotName; }
 
 	bool HasStaticSlot() const { return StaticSlotName != nullptr; }
@@ -225,7 +227,7 @@ public:
 
 	static TLinkedList<FShaderParametersMetadata*>*& GetStructList();
 	/** Speed up finding the uniform buffer by its name */
-	static TMap<FName, FShaderParametersMetadata*>& GetNameStructMap();
+	static TMap<FHashedName, FShaderParametersMetadata*>& GetNameStructMap();
 
 	/** Initialize all the global shader parameter structs. */
 	static void InitializeAllUniformBufferStructs();
@@ -248,6 +250,8 @@ private:
 	/** Name of the static slot to use for the uniform buffer (or null). */
 	const TCHAR* const StaticSlotName;
 
+	FHashedName ShaderVariableHashedName;
+
 	/** Size of the entire struct in bytes. */
 	const uint32 Size;
 
diff --git a/Engine/Source/Runtime/RenderCore/Public/ShaderParameterStruct.h b/Engine/Source/Runtime/RenderCore/Public/ShaderParameterStruct.h
--- a/Engine/Source/Runtime/RenderCore/Public/ShaderParameterStruct.h
+++ b/Engine/Source/Runtime/RenderCore/Public/ShaderParameterStruct.h
@@ -69,21 +69,21 @@ void BindForLegacyShaderParameters(FShader* Shader, const FShaderParameterMap& P
 
 
 /** Raise fatal error when a required shader parameter has not been set. */
-extern RENDERCORE_API void EmitNullShaderParameterFatalError(const FShader* Shader, const FShaderParametersMetadata* ParametersMetadata, uint16 MemberOffset);
+extern RENDERCORE_API void EmitNullShaderParameterFatalError(const TShaderRef<FShader>& Shader, const FShaderParametersMetadata* ParametersMetadata, uint16 MemberOffset);
 
 
 /** Validates that all resource parameters of a shader are set. */
 #if DO_CHECK
-extern RENDERCORE_API void ValidateShaderParameters(const FShader* Shader, const FShaderParametersMetadata* ParametersMetadata, const void* Parameters);
+extern RENDERCORE_API void ValidateShaderParameters(const TShaderRef<FShader>& Shader, const FShaderParametersMetadata* ParametersMetadata, const void* Parameters);
 
 #else // !DO_CHECK
-FORCEINLINE void ValidateShaderParameters(const FShader* Shader, const FShaderParametersMetadata* ParametersMetadata, const void* Parameters)
+FORCEINLINE void ValidateShaderParameters(const TShaderRef<FShader>& Shader, const FShaderParametersMetadata* ParametersMetadata, const void* Parameters)
 { }
 
 #endif // !DO_CHECK
 
 template<typename TShaderClass>
-FORCEINLINE void ValidateShaderParameters(const TShaderClass* Shader, const typename TShaderClass::FParameters& Parameters)
+FORCEINLINE void ValidateShaderParameters(const TShaderRef<TShaderClass>& Shader, const typename TShaderClass::FParameters& Parameters)
 {
 	const typename TShaderClass::FParameters* ParameterPtr = &Parameters;
 	return ValidateShaderParameters(Shader, TShaderClass::FParameters::FTypeInfo::GetStructMetadata(), ParameterPtr);
@@ -92,7 +92,7 @@ FORCEINLINE void ValidateShaderParameters(const TShaderClass* Shader, const type
 
 /** Set compute shader UAVs. */
 template<typename TRHICmdList, typename TShaderClass, typename TShaderRHI>
-inline void SetShaderUAVs(TRHICmdList& RHICmdList, const TShaderClass* Shader, TShaderRHI* ShadeRHI, const typename TShaderClass::FParameters& Parameters)
+inline void SetShaderUAVs(TRHICmdList& RHICmdList, const TShaderRef<TShaderClass>& Shader, TShaderRHI* ShadeRHI, const typename TShaderClass::FParameters& Parameters)
 {
 	checkf(
 		Shader->Bindings.UAVs.Num() == 0 && Shader->Bindings.GraphUAVs.Num() == 0,
@@ -100,13 +100,13 @@ inline void SetShaderUAVs(TRHICmdList& RHICmdList, const TShaderClass* Shader, T
 }
 
 template<typename TRHICmdList, typename TShaderClass>
-inline void SetShaderUAVs(TRHICmdList& RHICmdList, const TShaderClass* Shader, FRHIPixelShader* ShadeRHI, const typename TShaderClass::FParameters& Parameters)
+inline void SetShaderUAVs(TRHICmdList& RHICmdList, const TShaderRef<TShaderClass>& Shader, FRHIPixelShader* ShadeRHI, const typename TShaderClass::FParameters& Parameters)
 {
 	// Pixelshader UAVs are bound together with rendertargets using BeginRenderPass
 }
 
 template<typename TRHICmdList, typename TShaderClass>
-inline void SetShaderUAVs(TRHICmdList& RHICmdList, const TShaderClass* Shader, FRHIComputeShader* ShadeRHI, const typename TShaderClass::FParameters& Parameters)
+inline void SetShaderUAVs(TRHICmdList& RHICmdList, const TShaderRef<TShaderClass>& Shader, FRHIComputeShader* ShadeRHI, const typename TShaderClass::FParameters& Parameters)
 {
 	const FShaderParameterBindings& Bindings = Shader->Bindings;
 
@@ -134,7 +134,7 @@ inline void SetShaderUAVs(TRHICmdList& RHICmdList, const TShaderClass* Shader, F
 
 /** Unset compute shader UAVs. */
 template<typename TRHICmdList, typename TShaderClass>
-inline void UnsetShaderUAVs(TRHICmdList& RHICmdList, const TShaderClass* Shader, FRHIComputeShader* ShadeRHI)
+inline void UnsetShaderUAVs(TRHICmdList& RHICmdList, const TShaderRef<TShaderClass>& Shader, FRHIComputeShader* ShadeRHI)
 {
 	// TODO(RDG): Once all shader sets their parameter through this, can refactor RHI so all UAVs of a shader get unset through a single RHI function call.
 	const FShaderParameterBindings& Bindings = Shader->Bindings;
@@ -155,7 +155,7 @@ inline void UnsetShaderUAVs(TRHICmdList& RHICmdList, const TShaderClass* Shader,
 
 /** Set shader's parameters from its parameters struct. */
 template<typename TRHICmdList, typename TShaderClass, typename TShaderRHI>
-inline void SetShaderParameters(TRHICmdList& RHICmdList, const TShaderClass* Shader, TShaderRHI* ShadeRHI, const typename TShaderClass::FParameters& Parameters)
+inline void SetShaderParameters(TRHICmdList& RHICmdList, const TShaderRef<TShaderClass>& Shader, TShaderRHI* ShadeRHI, const typename TShaderClass::FParameters& Parameters)
 {
 	ValidateShaderParameters(Shader, Parameters);
 
@@ -231,7 +231,7 @@ inline void SetShaderParameters(TRHICmdList& RHICmdList, const TShaderClass* Sha
 
 /** Set shader's parameters from its parameters struct. */
 template<typename TShaderClass>
-void SetShaderParameters(FRayTracingShaderBindingsWriter& RTBindingsWriter, const TShaderClass* Shader, const typename TShaderClass::FParameters& Parameters)
+void SetShaderParameters(FRayTracingShaderBindingsWriter& RTBindingsWriter, const TShaderRef<TShaderClass>& Shader, const typename TShaderClass::FParameters& Parameters)
 {
 	ValidateShaderParameters(Shader, Parameters);
 
diff --git a/Engine/Source/Runtime/RenderCore/Public/ShaderParameters.h b/Engine/Source/Runtime/RenderCore/Public/ShaderParameters.h
--- a/Engine/Source/Runtime/RenderCore/Public/ShaderParameters.h
+++ b/Engine/Source/Runtime/RenderCore/Public/ShaderParameters.h
@@ -7,6 +7,7 @@
 #pragma once
 
 #include "CoreMinimal.h"
+#include "Serialization/MemoryLayout.h"
 #include "RHI.h"
 
 class FShaderParameterMap;
@@ -27,27 +28,21 @@ enum EShaderParameterFlags
 /** A shader parameter's register binding. e.g. float1/2/3/4, can be an array, UAV */
 class FShaderParameter
 {
+	DECLARE_EXPORTED_TYPE_LAYOUT(FShaderParameter, RENDERCORE_API, NonVirtual);
 public:
 	FShaderParameter()
 	:	BufferIndex(0)
 	,	BaseIndex(0)
 	,	NumBytes(0)
-#if UE_BUILD_DEBUG
-	,	bInitialized(false)
-#endif
 	{}
 
 	RENDERCORE_API void Bind(const FShaderParameterMap& ParameterMap,const TCHAR* ParameterName, EShaderParameterFlags Flags = SPF_Optional);
 	friend RENDERCORE_API FArchive& operator<<(FArchive& Ar,FShaderParameter& P);
 	bool IsBound() const { return NumBytes > 0; }
-
+	
 	inline bool IsInitialized() const 
 	{ 
-#if UE_BUILD_DEBUG
-		return bInitialized; 
-#else 
 		return true;
-#endif
 	}
 
 	uint32 GetBufferIndex() const { return BufferIndex; }
@@ -55,26 +50,20 @@ public:
 	uint32 GetNumBytes() const { return NumBytes; }
 
 private:
-	uint16 BufferIndex;
-	uint16 BaseIndex;
+	LAYOUT_FIELD(uint16, BufferIndex);
+	LAYOUT_FIELD(uint16, BaseIndex);
 	// 0 if the parameter wasn't bound
-	uint16 NumBytes;
-
-#if UE_BUILD_DEBUG
-	bool bInitialized;
-#endif
+	LAYOUT_FIELD(uint16, NumBytes);
 };
 
 /** A shader resource binding (textures or samplerstates). */
 class FShaderResourceParameter
 {
+	DECLARE_EXPORTED_TYPE_LAYOUT(FShaderResourceParameter, RENDERCORE_API, NonVirtual);
 public:
 	FShaderResourceParameter()
 	:	BaseIndex(0)
 	,	NumResources(0) 
-#if UE_BUILD_DEBUG
-	,	bInitialized(false)
-#endif
 	{}
 
 	RENDERCORE_API void Bind(const FShaderParameterMap& ParameterMap,const TCHAR* ParameterName,EShaderParameterFlags Flags = SPF_Optional);
@@ -83,27 +72,21 @@ public:
 
 	inline bool IsInitialized() const 
 	{ 
-#if UE_BUILD_DEBUG
-		return bInitialized; 
-#else 
 		return true;
-#endif
 	}
 
 	uint32 GetBaseIndex() const { return BaseIndex; }
 	uint32 GetNumResources() const { return NumResources; }
-private:
-	uint16 BaseIndex;
-	uint16 NumResources;
 
-#if UE_BUILD_DEBUG
-	bool bInitialized;
-#endif
+private:
+	LAYOUT_FIELD(uint16, BaseIndex);
+	LAYOUT_FIELD(uint16, NumResources);
 };
 
 /** A class that binds either a UAV or SRV of a resource. */
 class FRWShaderParameter
 {
+	DECLARE_EXPORTED_TYPE_LAYOUT(FRWShaderParameter, RENDERCORE_API, NonVirtual);
 public:
 
 	void Bind(const FShaderParameterMap& ParameterMap,const TCHAR* BaseName)
@@ -151,9 +134,8 @@ public:
 	inline void UnsetUAV(TRHICmdList& RHICmdList, FRHIComputeShader* ComputeShader) const;
 
 private:
-
-	FShaderResourceParameter SRVParameter;
-	FShaderResourceParameter UAVParameter;
+	LAYOUT_FIELD(FShaderResourceParameter, SRVParameter);
+	LAYOUT_FIELD(FShaderResourceParameter, UAVParameter);
 };
 
 /** Creates a shader code declaration of this struct for the given shader platform. */
@@ -161,13 +143,10 @@ extern RENDERCORE_API void CreateUniformBufferShaderDeclaration(const TCHAR* Nam
 
 class FShaderUniformBufferParameter
 {
+	DECLARE_EXPORTED_TYPE_LAYOUT(FShaderUniformBufferParameter, RENDERCORE_API, NonVirtual);
 public:
 	FShaderUniformBufferParameter()
-	:	BaseIndex(0)
-	,	bIsBound(false) 
-#if UE_BUILD_DEBUG
-	,	bInitialized(false)
-#endif
+	:	BaseIndex(0xffff)
 	{}
 
 	static RENDERCORE_API void ModifyCompilationEnvironment(const TCHAR* ParameterName,const FShaderParametersMetadata& Struct,EShaderPlatform Platform,FShaderCompilerEnvironment& OutEnvironment);
@@ -180,44 +159,21 @@ public:
 		return Ar;
 	}
 
-	bool IsBound() const { return bIsBound; }
+	bool IsBound() const { return BaseIndex != 0xffff; }
 
 	void Serialize(FArchive& Ar)
 	{
-#if UE_BUILD_DEBUG
-		if (Ar.IsLoading())
-		{
-			bInitialized = true;
-		}
-#endif
-		Ar << BaseIndex << bIsBound;
+		Ar << BaseIndex;
 	}
 
 	inline bool IsInitialized() const 
 	{ 
-#if UE_BUILD_DEBUG
-		return bInitialized; 
-#else 
 		return true;
-#endif
-	}
-
-	inline void SetInitialized() 
-	{ 
-#if UE_BUILD_DEBUG
-		bInitialized = true; 
-#endif
 	}
-
-	uint32 GetBaseIndex() const { return BaseIndex; }
+	uint32 GetBaseIndex() const { ensure(IsBound()); return BaseIndex; }
 
 private:
-	uint16 BaseIndex;
-	bool bIsBound;
-
-#if UE_BUILD_DEBUG
-	bool bInitialized;
-#endif
+	LAYOUT_FIELD(uint16, BaseIndex);
 };
 
 /** A shader uniform buffer binding with a specific structure. */
diff --git a/Engine/Source/Runtime/RenderCore/Public/ShaderPipelineCache.h b/Engine/Source/Runtime/RenderCore/Public/ShaderPipelineCache.h
--- a/Engine/Source/Runtime/RenderCore/Public/ShaderPipelineCache.h
+++ b/Engine/Source/Runtime/RenderCore/Public/ShaderPipelineCache.h
@@ -72,7 +72,6 @@ class RENDERCORE_API FShaderPipelineCache : public FTickableObjectRenderThread
 	{
 		FPipelineCacheFileFormatPSO PSO;
 		FShaderPipelineCacheArchive* ReadRequests;
-		TSet<FSHAHash> ShaderCodeReads;
 	};
 
 public:
diff --git a/Engine/Source/Runtime/RenderCore/Public/VertexFactory.h b/Engine/Source/Runtime/RenderCore/Public/VertexFactory.h
--- a/Engine/Source/Runtime/RenderCore/Public/VertexFactory.h
+++ b/Engine/Source/Runtime/RenderCore/Public/VertexFactory.h
@@ -18,6 +18,7 @@
 class FMaterial;
 class FMeshDrawSingleShaderBindings;
 class FPrimitiveSceneProxy;
+struct FVertexFactoryShaderPermutationParameters;
 
 #if PLATFORM_SUPPORTS_PRAGMA_PACK
 	#pragma pack (push,4)
@@ -155,16 +156,15 @@ struct FVertexStreamComponent
  */
 class RENDERCORE_API FVertexFactoryShaderParameters
 {
+	DECLARE_TYPE_LAYOUT(FVertexFactoryShaderParameters, NonVirtual);
 public:
-	virtual ~FVertexFactoryShaderParameters() {}
-	virtual void Bind(const class FShaderParameterMap& ParameterMap) = 0;
-	virtual void Serialize(FArchive& Ar) = 0;
+	void Bind(const class FShaderParameterMap& ParameterMap) {}
 
 	/** 
 	 * Gets the vertex factory's shader bindings and vertex streams.
 	 * View can be null when caching mesh draw commands (only for supported vertex factories)
 	 */
-	virtual void GetElementShaderBindings(
+	void GetElementShaderBindings(
 		const class FSceneInterface* Scene,
 		const class FSceneView* View,
 		const class FMeshMaterialShader* Shader,
@@ -173,11 +173,120 @@ public:
 		const class FVertexFactory* VertexFactory,
 		const struct FMeshBatchElement& BatchElement,
 		class FMeshDrawSingleShaderBindings& ShaderBindings,
-		FVertexInputStreamArray& VertexStreams) const = 0;
+		FVertexInputStreamArray& VertexStreams) const {}
+
+	// Only for debugging type mismatch
+	LAYOUT_FIELD(uint32, Size);
+};
 
-	virtual uint32 GetSize() const { return sizeof(*this); }
+template<EShaderFrequency ShaderFrequency, typename VertexFactoryType>
+struct TVertexFactoryParameterTraits
+{
+	static const FTypeLayoutDesc* GetLayout() { return nullptr; }
+	static FVertexFactoryShaderParameters* Create(const class FShaderParameterMap& ParameterMap) { return nullptr; }
+	
+	static void GetElementShaderBindings(
+		const FVertexFactoryShaderParameters* Parameters,
+		const class FSceneInterface* Scene,
+		const class FSceneView* View,
+		const class FMeshMaterialShader* Shader,
+		const EVertexInputStreamType InputStreamType,
+		ERHIFeatureLevel::Type FeatureLevel,
+		const class FVertexFactory* VertexFactory,
+		const struct FMeshBatchElement& BatchElement,
+		class FMeshDrawSingleShaderBindings& ShaderBindings,
+		FVertexInputStreamArray& VertexStreams) {}
 };
 
+#define IMPLEMENT_VERTEX_FACTORY_PARAMETER_TYPE(FactoryClass, ShaderFrequency, ParameterClass) \
+	template<> struct TVertexFactoryParameterTraits<ShaderFrequency, FactoryClass> \
+	{ \
+		static const FTypeLayoutDesc* GetLayout() { return &StaticGetTypeLayoutDesc<ParameterClass>(); } \
+		static FVertexFactoryShaderParameters* Create(const class FShaderParameterMap& ParameterMap) { ParameterClass* Result = new ParameterClass(); Result->Size = sizeof(ParameterClass); Result->Bind(ParameterMap); return Result; } \
+		static void GetElementShaderBindings( \
+			const FVertexFactoryShaderParameters* Parameters, \
+			const class FSceneInterface* Scene, \
+			const class FSceneView* View, \
+			const class FMeshMaterialShader* Shader, \
+			const EVertexInputStreamType InputStreamType, \
+			ERHIFeatureLevel::Type FeatureLevel, \
+			const class FVertexFactory* VertexFactory, \
+			const struct FMeshBatchElement& BatchElement, \
+			class FMeshDrawSingleShaderBindings& ShaderBindings, \
+			FVertexInputStreamArray& VertexStreams) \
+		{ \
+			check(Parameters->Size == sizeof(ParameterClass)); \
+			static_cast<const ParameterClass*>(Parameters)->GetElementShaderBindings(Scene, View, Shader, InputStreamType, FeatureLevel, VertexFactory, BatchElement, ShaderBindings, VertexStreams); \
+		} \
+	}
+
+template<typename VertexFactoryType>
+static const FTypeLayoutDesc* GetVertexFactoryParametersLayout(EShaderFrequency ShaderFrequency)
+{
+	switch (ShaderFrequency)
+	{
+	case SF_Vertex: return TVertexFactoryParameterTraits<SF_Vertex, VertexFactoryType>::GetLayout();
+	case SF_Hull: return TVertexFactoryParameterTraits<SF_Hull, VertexFactoryType>::GetLayout();
+	case SF_Domain: return TVertexFactoryParameterTraits<SF_Domain, VertexFactoryType>::GetLayout();
+	case SF_Pixel: return TVertexFactoryParameterTraits<SF_Pixel, VertexFactoryType>::GetLayout();
+	case SF_Geometry: return TVertexFactoryParameterTraits<SF_Geometry, VertexFactoryType>::GetLayout();
+	case SF_Compute: return TVertexFactoryParameterTraits<SF_Compute, VertexFactoryType>::GetLayout();
+	case SF_RayGen: return TVertexFactoryParameterTraits<SF_RayGen, VertexFactoryType>::GetLayout();
+	case SF_RayMiss: return TVertexFactoryParameterTraits<SF_RayMiss, VertexFactoryType>::GetLayout();
+	case SF_RayHitGroup: return TVertexFactoryParameterTraits<SF_RayHitGroup, VertexFactoryType>::GetLayout();
+	case SF_RayCallable: return TVertexFactoryParameterTraits<SF_RayCallable, VertexFactoryType>::GetLayout();
+	default: checkNoEntry(); return nullptr;
+	}
+}
+
+template<typename VertexFactoryType>
+static FVertexFactoryShaderParameters* ConstructVertexFactoryParameters(EShaderFrequency ShaderFrequency, const class FShaderParameterMap& ParameterMap)
+{
+	switch (ShaderFrequency)
+	{
+	case SF_Vertex: return TVertexFactoryParameterTraits<SF_Vertex, VertexFactoryType>::Create(ParameterMap);
+	case SF_Hull: return TVertexFactoryParameterTraits<SF_Hull, VertexFactoryType>::Create(ParameterMap);
+	case SF_Domain: return TVertexFactoryParameterTraits<SF_Domain, VertexFactoryType>::Create(ParameterMap);
+	case SF_Pixel: return TVertexFactoryParameterTraits<SF_Pixel, VertexFactoryType>::Create(ParameterMap);
+	case SF_Geometry: return TVertexFactoryParameterTraits<SF_Geometry, VertexFactoryType>::Create(ParameterMap);
+	case SF_Compute: return TVertexFactoryParameterTraits<SF_Compute, VertexFactoryType>::Create(ParameterMap);
+	case SF_RayGen: return TVertexFactoryParameterTraits<SF_RayGen, VertexFactoryType>::Create(ParameterMap);
+	case SF_RayMiss: return TVertexFactoryParameterTraits<SF_RayMiss, VertexFactoryType>::Create(ParameterMap);
+	case SF_RayHitGroup: return TVertexFactoryParameterTraits<SF_RayHitGroup, VertexFactoryType>::Create(ParameterMap);
+	case SF_RayCallable: return TVertexFactoryParameterTraits<SF_RayCallable, VertexFactoryType>::Create(ParameterMap);
+	default: checkNoEntry(); return nullptr;
+	}
+}
+
+
+template<typename VertexFactoryType>
+static void GetVertexFactoryParametersElementShaderBindings(EShaderFrequency ShaderFrequency, const FVertexFactoryShaderParameters* Parameters,
+	const class FSceneInterface* Scene,
+	const class FSceneView* View,
+	const class FMeshMaterialShader* Shader,
+	const EVertexInputStreamType InputStreamType,
+	ERHIFeatureLevel::Type FeatureLevel,
+	const class FVertexFactory* VertexFactory,
+	const struct FMeshBatchElement& BatchElement,
+	class FMeshDrawSingleShaderBindings& ShaderBindings,
+	FVertexInputStreamArray& VertexStreams)
+{
+	switch (ShaderFrequency)
+	{
+	case SF_Vertex: TVertexFactoryParameterTraits<SF_Vertex, VertexFactoryType>::GetElementShaderBindings(Parameters, Scene, View, Shader, InputStreamType, FeatureLevel, VertexFactory, BatchElement, ShaderBindings, VertexStreams); break;
+	case SF_Hull: TVertexFactoryParameterTraits<SF_Hull, VertexFactoryType>::GetElementShaderBindings(Parameters, Scene, View, Shader, InputStreamType, FeatureLevel, VertexFactory, BatchElement, ShaderBindings, VertexStreams); break;
+	case SF_Domain: TVertexFactoryParameterTraits<SF_Domain, VertexFactoryType>::GetElementShaderBindings(Parameters, Scene, View, Shader, InputStreamType, FeatureLevel, VertexFactory, BatchElement, ShaderBindings, VertexStreams); break;
+	case SF_Pixel: TVertexFactoryParameterTraits<SF_Pixel, VertexFactoryType>::GetElementShaderBindings(Parameters, Scene, View, Shader, InputStreamType, FeatureLevel, VertexFactory, BatchElement, ShaderBindings, VertexStreams); break;
+	case SF_Geometry: TVertexFactoryParameterTraits<SF_Geometry, VertexFactoryType>::GetElementShaderBindings(Parameters, Scene, View, Shader, InputStreamType, FeatureLevel, VertexFactory, BatchElement, ShaderBindings, VertexStreams); break;
+	case SF_Compute: TVertexFactoryParameterTraits<SF_Compute, VertexFactoryType>::GetElementShaderBindings(Parameters, Scene, View, Shader, InputStreamType, FeatureLevel, VertexFactory, BatchElement, ShaderBindings, VertexStreams); break;
+	case SF_RayGen: TVertexFactoryParameterTraits<SF_RayGen, VertexFactoryType>::GetElementShaderBindings(Parameters, Scene, View, Shader, InputStreamType, FeatureLevel, VertexFactory, BatchElement, ShaderBindings, VertexStreams); break;
+	case SF_RayMiss: TVertexFactoryParameterTraits<SF_RayMiss, VertexFactoryType>::GetElementShaderBindings(Parameters, Scene, View, Shader, InputStreamType, FeatureLevel, VertexFactory, BatchElement, ShaderBindings, VertexStreams); break;
+	case SF_RayHitGroup: TVertexFactoryParameterTraits<SF_RayHitGroup, VertexFactoryType>::GetElementShaderBindings(Parameters, Scene, View, Shader, InputStreamType, FeatureLevel, VertexFactory, BatchElement, ShaderBindings, VertexStreams); break;
+	case SF_RayCallable: TVertexFactoryParameterTraits<SF_RayCallable, VertexFactoryType>::GetElementShaderBindings(Parameters, Scene, View, Shader, InputStreamType, FeatureLevel, VertexFactory, BatchElement, ShaderBindings, VertexStreams); break;
+	default: checkNoEntry(); break;
+	}
+}
+
 /**
  * An object used to represent the type of a vertex factory.
  */
@@ -185,21 +294,38 @@ class FVertexFactoryType
 {
 public:
 
-	typedef FVertexFactoryShaderParameters* (*ConstructParametersType)(EShaderFrequency ShaderFrequency);
-	typedef bool (*ShouldCacheType)(EShaderPlatform, const class FMaterial*, const class FShaderType*);
-	typedef void (*ModifyCompilationEnvironmentType)(const FVertexFactoryType*, EShaderPlatform, const class FMaterial*, FShaderCompilerEnvironment&);
+	typedef FVertexFactoryShaderParameters* (*ConstructParametersType)(EShaderFrequency ShaderFrequency, const class FShaderParameterMap& ParameterMap);
+	typedef const FTypeLayoutDesc* (*GetParameterTypeLayoutType)(EShaderFrequency ShaderFrequency);
+	typedef void (*GetParameterTypeElementShaderBindingsType)(EShaderFrequency ShaderFrequency,
+		const FVertexFactoryShaderParameters* Parameters,
+		const class FSceneInterface* Scene,
+		const class FSceneView* View,
+		const class FMeshMaterialShader* Shader,
+		const EVertexInputStreamType InputStreamType,
+		ERHIFeatureLevel::Type FeatureLevel,
+		const class FVertexFactory* VertexFactory,
+		const struct FMeshBatchElement& BatchElement,
+		class FMeshDrawSingleShaderBindings& ShaderBindings,
+		FVertexInputStreamArray& VertexStreams);
+
+	typedef bool (*ShouldCacheType)(const FVertexFactoryShaderPermutationParameters&);
+	typedef void (*ModifyCompilationEnvironmentType)(const FVertexFactoryShaderPermutationParameters&, FShaderCompilerEnvironment&);
 	typedef void (*ValidateCompiledResultType)(const FVertexFactoryType*, EShaderPlatform, const FShaderParameterMap& ParameterMap, TArray<FString>& OutErrors);
 	typedef bool (*SupportsTessellationShadersType)();
 
+	static int32 GetNumVertexFactoryTypes() { return NumVertexFactories; }
+
 	/**
 	 * @return The global shader factory list.
 	 */
 	static RENDERCORE_API TLinkedList<FVertexFactoryType*>*& GetTypeList();
 
+	static RENDERCORE_API const TArray<FVertexFactoryType*>& GetSortedMaterialTypes();
+
 	/**
 	 * Finds a FVertexFactoryType by name.
 	 */
-	static RENDERCORE_API FVertexFactoryType* GetVFByName(const FString& VFName);
+	static RENDERCORE_API FVertexFactoryType* GetVFByName(const FHashedName& VFName);
 
 	/** Initialize FVertexFactoryType static members, this must be called before any VF types are created. */
 	static void Initialize(const TMap<FString, TArray<const TCHAR*> >& ShaderFileToUniformBufferVariables);
@@ -223,6 +349,8 @@ public:
 		bool bInSupportsCachingMeshDrawCommands,
 		bool bInSupportsPrimitiveIdStream,
 		ConstructParametersType InConstructParameters,
+		GetParameterTypeLayoutType InGetParameterTypeLayout,
+		GetParameterTypeElementShaderBindingsType InGetParameterTypeElementShaderBindings,
 		ShouldCacheType InShouldCache,
 		ModifyCompilationEnvironmentType InModifyCompilationEnvironment,
 		ValidateCompiledResultType InValidateCompiledResult,
@@ -234,8 +362,23 @@ public:
 	// Accessors.
 	const TCHAR* GetName() const { return Name; }
 	FName GetFName() const { return TypeName; }
+	const FHashedName& GetHashedName() const { return HashedName; }
 	const TCHAR* GetShaderFilename() const { return ShaderFilename; }
-	FVertexFactoryShaderParameters* CreateShaderParameters(EShaderFrequency ShaderFrequency) const { return (*ConstructParameters)(ShaderFrequency); }
+
+	FVertexFactoryShaderParameters* CreateShaderParameters(EShaderFrequency ShaderFrequency, const class FShaderParameterMap& ParameterMap) const { return (*ConstructParameters)(ShaderFrequency, ParameterMap); }
+	const FTypeLayoutDesc* GetShaderParameterLayout(EShaderFrequency ShaderFrequency) const { return (*GetParameterTypeLayout)(ShaderFrequency); }
+	void GetShaderParameterElementShaderBindings(EShaderFrequency ShaderFrequency,
+		const FVertexFactoryShaderParameters* Parameters,
+		const class FSceneInterface* Scene,
+		const class FSceneView* View,
+		const class FMeshMaterialShader* Shader,
+		const EVertexInputStreamType InputStreamType,
+		ERHIFeatureLevel::Type FeatureLevel,
+		const class FVertexFactory* VertexFactory,
+		const struct FMeshBatchElement& BatchElement,
+		class FMeshDrawSingleShaderBindings& ShaderBindings,
+		FVertexInputStreamArray& VertexStreams) const { (*GetParameterTypeElementShaderBindings)(ShaderFrequency, Parameters, Scene, View, Shader, InputStreamType, FeatureLevel, VertexFactory, BatchElement, ShaderBindings, VertexStreams); }
+
 	bool IsUsedWithMaterials() const { return bUsedWithMaterials; }
 	bool SupportsStaticLighting() const { return bSupportsStaticLighting; }
 	bool SupportsDynamicLighting() const { return bSupportsDynamicLighting; }
@@ -243,19 +386,11 @@ public:
 	bool SupportsPositionOnly() const { return bSupportsPositionOnly; }
 	bool SupportsCachingMeshDrawCommands() const { return bSupportsCachingMeshDrawCommands; }
 	bool SupportsPrimitiveIdStream() const { return bSupportsPrimitiveIdStream; }
-	/** Returns an int32 specific to this vertex factory type. */
-	inline int32 GetId() const { return HashIndex; }
-	static int32 GetNumVertexFactoryTypes() { return NextHashIndex; }
-
-	const FSerializationHistory* GetSerializationHistory(EShaderFrequency Frequency) const
-	{
-		return &SerializationHistory[Frequency];
-	}
 
 	// Hash function.
-	friend uint32 GetTypeHash(FVertexFactoryType* Type)
+	friend uint32 GetTypeHash(const FVertexFactoryType* Type)
 	{ 
-		return Type ? Type->HashIndex : 0;
+		return Type ? GetTypeHash(Type->HashedName) : 0u;
 	}
 
 	/** Calculates a Hash based on this vertex factory type's source code and includes */
@@ -264,23 +399,23 @@ public:
 	/**
 	 * Should we cache the material's shadertype on this platform with this vertex factory? 
 	 */
-	bool ShouldCache(EShaderPlatform Platform, const class FMaterial* Material, const class FShaderType* ShaderType)
+	bool ShouldCache(const FVertexFactoryShaderPermutationParameters& Parameters) const
 	{
-		return (*ShouldCacheRef)(Platform, Material, ShaderType); 
+		return (*ShouldCacheRef)(Parameters);
 	}
 
 	/**
 	* Calls the function ptr for the shader type on the given environment
 	* @param Environment - shader compile environment to modify
 	*/
-	void ModifyCompilationEnvironment(EShaderPlatform Platform, const FMaterial* Material, FShaderCompilerEnvironment& OutEnvironment){
+	void ModifyCompilationEnvironment(const FVertexFactoryShaderPermutationParameters& Parameters, FShaderCompilerEnvironment& OutEnvironment){
 		// Set up the mapping from VertexFactory.usf to the vertex factory type's source code.
 		FString VertexFactoryIncludeString = FString::Printf( TEXT("#include \"%s\""), GetShaderFilename() );
 		OutEnvironment.IncludeVirtualPathToContentsMap.Add(TEXT("/Engine/Generated/VertexFactory.ush"), VertexFactoryIncludeString);
 
 		OutEnvironment.SetDefine(TEXT("HAS_PRIMITIVE_UNIFORM_BUFFER"), 1);
 
-		(*ModifyCompilationEnvironmentRef)(this, Platform, Material, OutEnvironment);
+		(*ModifyCompilationEnvironmentRef)(Parameters, OutEnvironment);
 	}
 
 	void ValidateCompiledResult(EShaderPlatform Platform, const FShaderParameterMap& ParameterMap, TArray<FString>& OutErrors)
@@ -312,15 +447,15 @@ public:
 	}
 
 private:
-	static RENDERCORE_API uint32 NextHashIndex;
+	static RENDERCORE_API uint32 NumVertexFactories;
 
 	/** Tracks whether serialization history for all shader types has been initialized. */
 	static bool bInitializedSerializationHistory;
 
-	uint32 HashIndex;
 	const TCHAR* Name;
 	const TCHAR* ShaderFilename;
 	FName TypeName;
+	FHashedName HashedName;
 	uint32 bUsedWithMaterials : 1;
 	uint32 bSupportsStaticLighting : 1;
 	uint32 bSupportsDynamicLighting : 1;
@@ -329,6 +464,8 @@ private:
 	uint32 bSupportsCachingMeshDrawCommands : 1;
 	uint32 bSupportsPrimitiveIdStream : 1;
 	ConstructParametersType ConstructParameters;
+	GetParameterTypeLayoutType GetParameterTypeLayout;
+	GetParameterTypeElementShaderBindingsType GetParameterTypeElementShaderBindings;
 	ShouldCacheType ShouldCacheRef;
 	ModifyCompilationEnvironmentType ModifyCompilationEnvironmentRef;
 	ValidateCompiledResultType ValidateCompiledResultRef;
@@ -345,12 +482,6 @@ private:
 
 	/** Tracks what platforms ReferencedUniformBufferStructsCache has had declarations cached for. */
 	bool bCachedUniformBufferStructDeclarations;
-
-	/** 
-	 * Stores a history of serialization sizes for this vertex factory's shader parameter class. 
-	 * This is used to invalidate shaders when serialization changes.
-	 */
-	FSerializationHistory SerializationHistory[SF_NumFrequencies];
 };
 
 /**
@@ -362,7 +493,7 @@ extern RENDERCORE_API FArchive& operator<<(FArchive& Ar,FVertexFactoryType*& Typ
  * Find the vertex factory type with the given name.
  * @return NULL if no vertex factory type matched, otherwise a vertex factory type with a matching name.
  */
-extern RENDERCORE_API FVertexFactoryType* FindVertexFactoryType(FName TypeName);
+extern RENDERCORE_API FVertexFactoryType* FindVertexFactoryType(const FHashedName& TypeName);
 
 /**
  * A macro for declaring a new vertex factory type, for use in the vertex factory class's definition body.
@@ -372,13 +503,21 @@ extern RENDERCORE_API FVertexFactoryType* FindVertexFactoryType(FName TypeName);
 	static FVertexFactoryType StaticType; \
 	virtual FVertexFactoryType* GetType() const override;
 
+#define IMPLEMENT_VERTEX_FACTORY_VTABLE(FactoryClass) \
+	&ConstructVertexFactoryParameters<FactoryClass>, \
+	&GetVertexFactoryParametersLayout<FactoryClass>, \
+	&GetVertexFactoryParametersElementShaderBindings<FactoryClass>, \
+	FactoryClass::ShouldCompilePermutation, \
+	FactoryClass::ModifyCompilationEnvironment, \
+	FactoryClass::ValidateCompiledResult, \
+	FactoryClass::SupportsTessellationShaders
+
 /**
  * A macro for implementing the static vertex factory type object, and specifying parameters used by the type.
  * @param bUsedWithMaterials - True if the vertex factory will be rendered in combination with a material.
  * @param bSupportsStaticLighting - True if the vertex factory will be rendered with static lighting.
  */
 #define IMPLEMENT_VERTEX_FACTORY_TYPE(FactoryClass,ShaderFilename,bUsedWithMaterials,bSupportsStaticLighting,bSupportsDynamicLighting,bPrecisePrevWorldPos,bSupportsPositionOnly) \
-	FVertexFactoryShaderParameters* Construct##FactoryClass##ShaderParameters(EShaderFrequency ShaderFrequency) { return FactoryClass::ConstructShaderParameters(ShaderFrequency); } \
 	FVertexFactoryType FactoryClass::StaticType( \
 		TEXT(#FactoryClass), \
 		TEXT(ShaderFilename), \
@@ -389,17 +528,12 @@ extern RENDERCORE_API FVertexFactoryType* FindVertexFactoryType(FName TypeName);
 		bSupportsPositionOnly, \
 		false, \
 		false, \
-		Construct##FactoryClass##ShaderParameters, \
-		FactoryClass::ShouldCompilePermutation, \
-		FactoryClass::ModifyCompilationEnvironment, \
-		FactoryClass::ValidateCompiledResult, \
-		FactoryClass::SupportsTessellationShaders \
+		IMPLEMENT_VERTEX_FACTORY_VTABLE(FactoryClass) \
 		); \
 		FVertexFactoryType* FactoryClass::GetType() const { return &StaticType; }
 
 // @todo - need more extensible type properties - shouldn't have to change all IMPLEMENT_VERTEX_FACTORY_TYPE's when you add one new parameter
 #define IMPLEMENT_VERTEX_FACTORY_TYPE_EX(FactoryClass,ShaderFilename,bUsedWithMaterials,bSupportsStaticLighting,bSupportsDynamicLighting,bPrecisePrevWorldPos,bSupportsPositionOnly,bSupportsCachingMeshDrawCommands,bSupportsPrimitiveIdStream) \
-	FVertexFactoryShaderParameters* Construct##FactoryClass##ShaderParameters(EShaderFrequency ShaderFrequency) { return FactoryClass::ConstructShaderParameters(ShaderFrequency); } \
 	FVertexFactoryType FactoryClass::StaticType( \
 		TEXT(#FactoryClass), \
 		TEXT(ShaderFilename), \
@@ -410,11 +544,7 @@ extern RENDERCORE_API FVertexFactoryType* FindVertexFactoryType(FName TypeName);
 		bSupportsPositionOnly, \
 		bSupportsCachingMeshDrawCommands, \
 		bSupportsPrimitiveIdStream, \
-		Construct##FactoryClass##ShaderParameters, \
-		FactoryClass::ShouldCompilePermutation, \
-		FactoryClass::ModifyCompilationEnvironment, \
-		FactoryClass::ValidateCompiledResult, \
-		FactoryClass::SupportsTessellationShaders \
+		IMPLEMENT_VERTEX_FACTORY_VTABLE(FactoryClass) \
 		); \
 		FVertexFactoryType* FactoryClass::GetType() const { return &StaticType; }
 
@@ -422,36 +552,22 @@ extern RENDERCORE_API FVertexFactoryType* FindVertexFactoryType(FName TypeName);
 class FVertexFactoryTypeDependency
 {
 public:
+	FVertexFactoryTypeDependency() {}
 
-	FVertexFactoryTypeDependency() :
-		VertexFactoryType(NULL)
-	{}
-
-	FVertexFactoryType* VertexFactoryType;
+	FHashedName VertexFactoryTypeName;
 
-#if KEEP_SHADER_SOURCE_HASHES
 	/** Used to detect changes to the vertex factory source files. */
 	FSHAHash VFSourceHash;
-#endif
 
 	friend FArchive& operator<<(FArchive& Ar,class FVertexFactoryTypeDependency& Ref)
 	{
-#if KEEP_SHADER_SOURCE_HASHES
-		FSHAHash& Hash = Ref.VFSourceHash;
-#else
-		FSHAHash Hash;
-#endif
-		Ar << Ref.VertexFactoryType << FShaderResource::FilterShaderSourceHashForSerialization(Ar, Hash);
+		Ar << Ref.VertexFactoryTypeName << Ref.VFSourceHash;
 		return Ar;
 	}
 
 	bool operator==(const FVertexFactoryTypeDependency& Reference) const
 	{
-#if KEEP_SHADER_SOURCE_HASHES
-		return VertexFactoryType == Reference.VertexFactoryType && VFSourceHash == Reference.VFSourceHash;
-#else
-		return VertexFactoryType == Reference.VertexFactoryType;
-#endif
+		return VertexFactoryTypeName == Reference.VertexFactoryTypeName && VFSourceHash == Reference.VFSourceHash;
 	}
 
 	bool operator!=(const FVertexFactoryTypeDependency& Reference) const
@@ -496,7 +612,7 @@ public:
 	/**
 	* Can be overridden by FVertexFactory subclasses to modify their compile environment just before compilation occurs.
 	*/
-	static void ModifyCompilationEnvironment(const FVertexFactoryType* Type, EShaderPlatform Platform, const FMaterial* Material, FShaderCompilerEnvironment& OutEnvironment ) {}
+	static void ModifyCompilationEnvironment(const FVertexFactoryShaderPermutationParameters& Parameters, FShaderCompilerEnvironment& OutEnvironment ) {}
 
 	/**
 	* Can be overridden by FVertexFactory subclasses to fail a compile based on compilation output.
@@ -546,7 +662,7 @@ public:
 	virtual uint64 GetStaticBatchElementVisibility(const class FSceneView& View, const struct FMeshBatch* Batch, const void* InViewCustomData = nullptr) const { return 1; }
 
 	bool NeedsDeclaration() const { return bNeedsDeclaration; }
-	bool SupportsManualVertexFetch(ERHIFeatureLevel::Type InFeatureLevel) const 
+	inline bool SupportsManualVertexFetch(const FStaticFeatureLevel InFeatureLevel) const
 	{ 
 		check(InFeatureLevel != ERHIFeatureLevel::Num);
 		return bSupportsManualVertexFetch && (InFeatureLevel > ERHIFeatureLevel::ES3_1) && RHISupportsManualVertexFetch(GMaxRHIShaderPlatform);
@@ -633,74 +749,6 @@ private:
 	FVertexDeclarationRHIRef PositionAndNormalDeclaration;
 };
 
-/**
- * An encapsulation of the vertex factory parameters for a shader.
- * Handles serialization and binding of the vertex factory parameters for the shader's vertex factory type.
- */
-class RENDERCORE_API FVertexFactoryParameterRef
-{
-public:
-	FVertexFactoryParameterRef(FVertexFactoryType* InVertexFactoryType,const FShaderParameterMap& ParameterMap, EShaderFrequency InShaderFrequency, EShaderPlatform InShaderPlatform);
-
-	FVertexFactoryParameterRef():
-		Parameters(NULL),
-		VertexFactoryType(NULL)
-	{}
-
-	~FVertexFactoryParameterRef()
-	{
-		delete Parameters;
-	}
-
-	void GetElementShaderBindings(
-		const class FSceneInterface* Scene,
-		const FSceneView* View,
-		const class FMeshMaterialShader* Shader,
-		const EVertexInputStreamType InputStreamType,
-		ERHIFeatureLevel::Type FeatureLevel,
-		const FVertexFactory* VertexFactory,
-		const struct FMeshBatchElement& BatchElement,
-		class FMeshDrawSingleShaderBindings& ShaderBindings,
-		FVertexInputStreamArray& VertexStreams
-		) const
-	{
-		if (Parameters)
-		{
-			checkSlow(VertexFactory->GetType() == VertexFactoryType);
-			checkSlow(View || VertexFactoryType->SupportsCachingMeshDrawCommands());
-			Parameters->GetElementShaderBindings(Scene, View, Shader, InputStreamType, FeatureLevel, VertexFactory, BatchElement, ShaderBindings, VertexStreams);
-		}
-	}
-
-	const FVertexFactoryType* GetVertexFactoryType() const { return VertexFactoryType; }
-
-#if KEEP_SHADER_SOURCE_HASHES
-	/** Returns the hash of the vertex factory shader file that this shader was compiled with. */
-	const FSHAHash& GetHash() const;
-#endif
-
-	/** Returns the shader platform that this shader was compiled with. */
-	EShaderPlatform GetShaderPlatform() const;
-
-	friend RENDERCORE_API bool operator<<(FArchive& Ar,FVertexFactoryParameterRef& Ref);
-
-	uint32 GetAllocatedSize() const
-	{
-		return Parameters ? Parameters->GetSize() : 0;
-	}
-
-private:
-	FVertexFactoryShaderParameters* Parameters;
-	FVertexFactoryType* VertexFactoryType;
-	EShaderFrequency ShaderFrequency;
-	EShaderPlatform ShaderPlatform;
-
-#if KEEP_SHADER_SOURCE_HASHES
-	// Hash of the vertex factory's source file at shader compile time, used by the automatic versioning system to detect changes
-	FSHAHash VFHash;
-#endif
-};
-
 /**
 * Default PrimitiveId vertex buffer.  Contains a single index of 0.
 * This is used when the VF is used for rendering outside normal mesh passes, where there is no valid scene.
diff --git a/Engine/Source/Runtime/Renderer/Private/AmbientCubemapParameters.h b/Engine/Source/Runtime/Renderer/Private/AmbientCubemapParameters.h
--- a/Engine/Source/Runtime/Renderer/Private/AmbientCubemapParameters.h
+++ b/Engine/Source/Runtime/Renderer/Private/AmbientCubemapParameters.h
@@ -15,6 +15,7 @@ class FShaderParameterMap;
 // DEPRECATED: use FAmbientCubemapParameters instead.
 class FCubemapShaderParameters
 {
+	DECLARE_TYPE_LAYOUT(FCubemapShaderParameters, NonVirtual);
 public:
 
 	void Bind(const FShaderParameterMap& ParameterMap);
@@ -25,13 +26,15 @@ public:
 	friend FArchive& operator<<(FArchive& Ar, FCubemapShaderParameters& P);
 
 private:
-	FShaderParameter AmbientCubemapColor;
-	FShaderParameter AmbientCubemapMipAdjust;
-	FShaderResourceParameter AmbientCubemap;
-	FShaderResourceParameter AmbientCubemapSampler;
-
 	template<typename TRHIShader>
 	void SetParametersTemplate(FRHICommandList& RHICmdList, TRHIShader* ShaderRHI, const FFinalPostProcessSettings::FCubemapEntry& Entry) const;
+
+	
+		LAYOUT_FIELD(FShaderParameter, AmbientCubemapColor)
+		LAYOUT_FIELD(FShaderParameter, AmbientCubemapMipAdjust)
+		LAYOUT_FIELD(FShaderResourceParameter, AmbientCubemap)
+		LAYOUT_FIELD(FShaderResourceParameter, AmbientCubemapSampler)
+	
 };
 
 /** Shader parameters needed for deferred passes sampling the ambient cube map. */
