Fixed shader DDC key to use build target platform instead of editor platform to get virtual texture support value.
Rearranged DDC key composition since we need to flush out incorrectly keyed shaders in the cache.

diff --git a/Engine/Source/Runtime/RenderCore/Private/Shader.cpp b/Engine/Source/Runtime/RenderCore/Private/Shader.cpp
--- a/Engine/Source/Runtime/RenderCore/Private/Shader.cpp
+++ b/Engine/Source/Runtime/RenderCore/Private/Shader.cpp
@@ -1662,18 +1662,13 @@ void ShaderMapAppendKeyString(EShaderPlatform Platform, FString& KeyString)
 		{
 			VTTextures = (CVarMobileVirtualTexture->GetValueOnAnyThread() != 0);
 		}
-		
-		static const auto CVarVTFactor = IConsoleManager::Get().FindConsoleVariable(TEXT("r.vt.FeedbackFactor")); check(CVarVTFactor);
-		const int32 VTFeedbackFactor = CVarVTFactor->GetInt(); 
 
+		const bool VTSupported = TargetPlatform->SupportsFeature(ETargetPlatformFeatures::VirtualTextureStreaming);
 
-		ITargetPlatformManagerModule* TPM = GetTargetPlatformManager();
-		check(TPM);
-		auto RunningTargetPlatform = TPM->GetRunningTargetPlatform();
-		check(RunningTargetPlatform);
-		const bool VTSupported = RunningTargetPlatform->SupportsFeature(ETargetPlatformFeatures::VirtualTextureStreaming);
+		static const auto CVarVTFactor = IConsoleManager::Get().FindConsoleVariable(TEXT("r.vt.FeedbackFactor")); check(CVarVTFactor);
+		const int32 VTFeedbackFactor = CVarVTFactor->GetInt(); 
 
-		auto tt = FString::Printf(TEXT("_VT-%d-%d-%d-%d"), VTLightmaps, VTTextures, VTFeedbackFactor, VTSupported);
+		auto tt = FString::Printf(TEXT("_VT-%d-%d-%d-%d"), VTLightmaps, VTTextures, VTSupported, VTFeedbackFactor);
  		KeyString += tt;
 	}
 
