fixes for the PSO cache to make it so we can prioritize and compile only those PSOs up front we want based on user category

diff --git a/Engine/Source/Runtime/RenderCore/Private/ShaderPipelineCache.cpp b/Engine/Source/Runtime/RenderCore/Private/ShaderPipelineCache.cpp
--- a/Engine/Source/Runtime/RenderCore/Private/ShaderPipelineCache.cpp
+++ b/Engine/Source/Runtime/RenderCore/Private/ShaderPipelineCache.cpp
@@ -138,6 +138,12 @@ static TAutoConsoleVariable<int32> CVarPSOFileCacheGameFileMaskEnabled(
 																TEXT("Set non zero to use GameFileMask during PSO precompile - recording should always save out the usage masks to make that data availble when needed."),
 																ECVF_ReadOnly | ECVF_RenderThreadSafe);
 
+static TAutoConsoleVariable<int32> CVarPSOFileCachePreOptimizeEnabled(
+																TEXT("r.ShaderPipelineCache.PreOptimizeEnabled"),
+																(int32)0,
+																TEXT("Set non zero to PreOptimize PSOs - this allows some PSOs to be compiled in the foreground before going in to game"),
+																ECVF_ReadOnly | ECVF_RenderThreadSafe);
+
 
 static TAutoConsoleVariable<int32> CVarPSOFileCacheMinBindCount(
                                                                 TEXT("r.ShaderPipelineCache.MinBindCount"),
@@ -332,65 +338,80 @@ bool FShaderPipelineCache::SetGameUsageMaskWithComparison(uint64 InMask, FPSOMas
 {
 	bool bMaskChanged = false;
 	
-	if (ShaderPipelineCache != nullptr && CVarPSOFileCacheGameFileMaskEnabled.GetValueOnAnyThread())
+	if (ShaderPipelineCache != nullptr && CVarPSOFileCacheGameFileMaskEnabled.GetValueOnAnyThread() && !ShaderPipelineCache->bPreOptimizing)
 	{
 		FScopeLock Lock(&ShaderPipelineCache->Mutex);
 		
-		uint64 OldMask = FPipelineFileCache::SetGameUsageMaskWithComparison(InMask, InComparisonFnPtr);
-		bMaskChanged = OldMask != InMask;
+		if (ShaderPipelineCache->bOpened)
+		{
+			uint64 OldMask = FPipelineFileCache::SetGameUsageMaskWithComparison(InMask, InComparisonFnPtr);
+			bMaskChanged = OldMask != InMask;
 		
-		ShaderPipelineCache->bReady = true;
+			ShaderPipelineCache->bReady = true;
 		
-		if(bMaskChanged && ShaderPipelineCache->bOpened)
-		{
-			// Mask has changed and we have an open file refetch PSO's for this Mask - leave the FPipelineFileCache file open - no need to close - just pull out the relevant PSOs.
-			// If this PSO compile run has completed for this Mask in which case don't refetch + compile for that mask
+			if(bMaskChanged)
+			{
+				// Mask has changed and we have an open file refetch PSO's for this Mask - leave the FPipelineFileCache file open - no need to close - just pull out the relevant PSOs.
+				// If this PSO compile run has completed for this Mask in which case don't refetch + compile for that mask
 			
-			// Don't clear already compiled PSOHash list - this is not a full reset
-			ShaderPipelineCache->Flush(false);
+				// Don't clear already compiled PSOHash list - this is not a full reset
+				ShaderPipelineCache->Flush(false);
 			
-			if(!ShaderPipelineCache->CompletedMasks.Contains(InMask))
-			{
-				int32 Order = (int32)FPipelineFileCache::PSOOrder::Default;
-				
-				if(!GConfig->GetInt(FShaderPipelineCacheConstants::SectionHeading, FShaderPipelineCacheConstants::SortOrderKey, Order, *GGameUserSettingsIni))
+				if(!ShaderPipelineCache->CompletedMasks.Contains(InMask))
 				{
-					GConfig->GetInt(FShaderPipelineCacheConstants::SectionHeading, FShaderPipelineCacheConstants::SortOrderKey, Order, *GGameIni);
-				}
+					int32 Order = (int32)FPipelineFileCache::PSOOrder::Default;
 				
-				TArray<FPipelineCachePSOHeader> LocalPreFetchedTasks;
-				FPipelineFileCache::GetOrderedPSOHashes(LocalPreFetchedTasks, (FPipelineFileCache::PSOOrder)Order, (int64)CVarPSOFileCacheMinBindCount.GetValueOnAnyThread(), ShaderPipelineCache->CompiledHashes);
-				// Iterate over all the tasks we haven't yet begun to read data for - these are the 'waiting' tasks
-				int64 Count = 0;
-				for (auto const& Task : LocalPreFetchedTasks)
-				{
-					bool bHasShaders = true;
-					for (FSHAHash const& Hash : Task.Shaders)
+					if(!GConfig->GetInt(FShaderPipelineCacheConstants::SectionHeading, FShaderPipelineCacheConstants::SortOrderKey, Order, *GGameUserSettingsIni))
 					{
-						bHasShaders &= FShaderCodeLibrary::ContainsShaderCode(Hash);
+						GConfig->GetInt(FShaderPipelineCacheConstants::SectionHeading, FShaderPipelineCacheConstants::SortOrderKey, Order, *GGameIni);
 					}
-					if (bHasShaders)
+				
+					TArray<FPipelineCachePSOHeader> LocalPreFetchedTasks;
+					FPipelineFileCache::GetOrderedPSOHashes(LocalPreFetchedTasks, (FPipelineFileCache::PSOOrder)Order, (int64)CVarPSOFileCacheMinBindCount.GetValueOnAnyThread(), ShaderPipelineCache->CompiledHashes);
+					// Iterate over all the tasks we haven't yet begun to read data for - these are the 'waiting' tasks
+					int64 Count = 0;
+					for (auto const& Task : LocalPreFetchedTasks)
 					{
-						Count++;
+						bool bHasShaders = true;
+						for (FSHAHash const& Hash : Task.Shaders)
+						{
+							bHasShaders &= FShaderCodeLibrary::ContainsShaderCode(Hash);
+						}
+						if (bHasShaders)
+						{
+							Count++;
+						}
 					}
-				}
 				
-				FPlatformAtomics::InterlockedAdd(&ShaderPipelineCache->TotalWaitingTasks, Count);
+					FPlatformAtomics::InterlockedAdd(&ShaderPipelineCache->TotalWaitingTasks, Count);
 				
-				if (OnCachedOpened.IsBound())
-				{
-					OnCachedOpened.Broadcast(ShaderPipelineCache->FileName, ShaderPipelineCache->CurrentPlatform, LocalPreFetchedTasks.Num(), ShaderPipelineCache->CacheFileGuid, ShaderPipelineCache->ShaderCachePrecompileContext);
-				}
+					if (OnCachedOpened.IsBound())
+					{
+						OnCachedOpened.Broadcast(ShaderPipelineCache->FileName, ShaderPipelineCache->CurrentPlatform, LocalPreFetchedTasks.Num(), ShaderPipelineCache->CacheFileGuid, ShaderPipelineCache->ShaderCachePrecompileContext);
+					}
 				
-				ShaderPipelineCache->PreFetchedTasks = LocalPreFetchedTasks;
+					ShaderPipelineCache->PreFetchedTasks = LocalPreFetchedTasks;
 				
-				UE_LOG(LogRHI, Display, TEXT("New ShaderPipelineCache GameUsageMask [%llu=>%llu], Enqueued %d of %d tasks for precompile."), OldMask, InMask, Count, ShaderPipelineCache->PreFetchedTasks.Num());
+					UE_LOG(LogRHI, Display, TEXT("New ShaderPipelineCache GameUsageMask [%llu=>%llu], Enqueued %d of %d tasks for precompile."), OldMask, InMask, Count, ShaderPipelineCache->PreFetchedTasks.Num());
+				}
+				else
+				{
+					UE_LOG(LogRHI, Display, TEXT("New ShaderPipelineCache GameUsageMask [%llu=>%llu], Target mask already precompiled."), OldMask, InMask);
+				}
 			}
 			else
 			{
-				UE_LOG(LogRHI, Display, TEXT("New ShaderPipelineCache GameUsageMask [%llu=>%llu], Target mask already precompiled."), OldMask, InMask);
+				UE_LOG(LogRHI, Display, TEXT("ShaderPipelineCache::SetGameUsageMaskWithComparison failed to set a new mask because the game mask was not different"));
 			}
 		}
+		else
+		{
+			UE_LOG(LogRHI, Display, TEXT("ShaderPipelineCache::SetGameUsageMaskWithComparison failed to set a new mask because the cache was not open"));
+		}
+	}
+	else
+	{
+		UE_LOG(LogRHI, Display, TEXT("ShaderPipelineCache::SetGameUsageMaskWithComparison failed to set a new mask because the cache was not open or game mask is not enabled"));
 	}
 	
 	return bMaskChanged;
@@ -1075,6 +1096,7 @@ FShaderPipelineCache::FShaderPipelineCache(EShaderPlatform Platform)
 , bPaused(false)
 , bOpened(false)
 , bReady(false)
+, bPreOptimizing(false)
 , PausedCount(0)
 , TotalActiveTasks(0)
 , TotalWaitingTasks(0)
@@ -1115,7 +1137,7 @@ FShaderPipelineCache::FShaderPipelineCache(EShaderPlatform Platform)
 	
 	FCoreDelegates::ApplicationWillDeactivateDelegate.AddStatic(&PipelineStateCacheOnAppDeactivate);
 	
-	bReady = !CVarPSOFileCacheGameFileMaskEnabled.GetValueOnAnyThread();
+	bReady = !CVarPSOFileCacheGameFileMaskEnabled.GetValueOnAnyThread() || CVarPSOFileCachePreOptimizeEnabled.GetValueOnAnyThread();
 }
 
 FShaderPipelineCache::~FShaderPipelineCache()
@@ -1196,6 +1218,7 @@ void FShaderPipelineCache::Tick( float DeltaTime )
 		
         	FPlatformAtomics::InterlockedExchange(&TotalCompleteTasks, 0);
         	FPlatformAtomics::InterlockedExchange(&TotalPrecompileTime, 0);
+			bPreOptimizing = false;
 		}
     }
 	
@@ -1376,11 +1399,12 @@ TStatId FShaderPipelineCache::GetStatId() const
 // Not sure where the define is for this but most seem to be low, medium, high, epic, cinema, auto, except material quality but that's less anyway
 static const int32 MaxQualityCount = 6;
 static const int32 MaxPlaylistCount = 3;
+static const int32 MaxUserCount = 16;
 
 static bool PreCompileMaskComparison(uint64 ReferenceGameMask, uint64 PSOMask)
 {
 	uint64 UsageMask = (ReferenceGameMask & PSOMask);
-	return (UsageMask & (7 << (MaxQualityCount*2+MaxPlaylistCount))) && (UsageMask & (63 << MaxQualityCount*2)) && (UsageMask & (63 << MaxQualityCount)) && (UsageMask & 63);
+	return (UsageMask & (7l << (MaxQualityCount*3+MaxPlaylistCount))) && (UsageMask & (7 << (MaxQualityCount*3))) && (UsageMask & (63 << MaxQualityCount*2)) && (UsageMask & (63 << MaxQualityCount)) && (UsageMask & 63);
 }
 
 bool FShaderPipelineCache::Open(FString const& Name, EShaderPlatform Platform)
@@ -1439,14 +1463,11 @@ bool FShaderPipelineCache::Open(FString const& Name, EShaderPlatform Platform)
 			PreFetchedTasks = LocalPreFetchedTasks;
 			TotalPrecompileTasks = PreFetchedTasks.Num();
 			
+			bPreOptimizing = TotalPrecompileTasks > 0;
 			UE_LOG(LogRHI, Display, TEXT("Opened pipeline cache and enqueued %d of %d tasks for precompile with BatchSize %d and BatchTime %f."), Count, PreFetchedTasks.Num(), BatchSize, BatchTime);
 		}
 		else
 		{
-			if (CVarPSOFileCacheGameFileMaskEnabled.GetValueOnAnyThread() == 0)
-			{
-				CVarPSOFileCacheGameFileMaskEnabled->Set(1);
-			}
 			UE_LOG(LogRHI, Display, TEXT("Opened pipeline cache - precompile deferred on UsageMask."));
 		}
 	}
diff --git a/Engine/Source/Runtime/RenderCore/Public/ShaderPipelineCache.h b/Engine/Source/Runtime/RenderCore/Public/ShaderPipelineCache.h
--- a/Engine/Source/Runtime/RenderCore/Public/ShaderPipelineCache.h
+++ b/Engine/Source/Runtime/RenderCore/Public/ShaderPipelineCache.h
@@ -238,6 +238,7 @@ private:
 	bool bPaused;
 	bool bOpened;
 	bool bReady;
+	bool bPreOptimizing;
     int32 PausedCount;
 	FShaderCachePrecompileContext ShaderCachePrecompileContext;
 	
