- Fixed copyright notices

diff --git a/Engine/Source/Runtime/Core/Private/Containers/MemoryImageAllocator.cpp b/Engine/Source/Runtime/Core/Private/Containers/MemoryImageAllocator.cpp
--- a/Engine/Source/Runtime/Core/Private/Containers/MemoryImageAllocator.cpp
+++ b/Engine/Source/Runtime/Core/Private/Containers/MemoryImageAllocator.cpp
@@ -1,3 +1,5 @@
+// Copyright Epic Games, Inc. All Rights Reserved.
+
 #include "Containers/ContainerAllocationPolicies.h"
 #include "Containers/ResourceArray.h"
 #include "Serialization/MemoryImage.h"
diff --git a/Engine/Source/Runtime/Core/Private/Misc/DelayedAutoRegister.cpp b/Engine/Source/Runtime/Core/Private/Misc/DelayedAutoRegister.cpp
--- a/Engine/Source/Runtime/Core/Private/Misc/DelayedAutoRegister.cpp
+++ b/Engine/Source/Runtime/Core/Private/Misc/DelayedAutoRegister.cpp
@@ -1,4 +1,4 @@
-// Copyright 1998-2019 Epic Games, Inc. All Rights Reserved.
+// Copyright Epic Games, Inc. All Rights Reserved.
 
 #include "Misc/DelayedAutoRegister.h"
 #include "Delegates/Delegate.h"
diff --git a/Engine/Source/Runtime/Core/Private/Misc/PreLoadFile.cpp b/Engine/Source/Runtime/Core/Private/Misc/PreLoadFile.cpp
--- a/Engine/Source/Runtime/Core/Private/Misc/PreLoadFile.cpp
+++ b/Engine/Source/Runtime/Core/Private/Misc/PreLoadFile.cpp
@@ -1,4 +1,4 @@
-// Copyright 1998-2019 Epic Games, Inc. All Rights Reserved.
+// Copyright Epic Games, Inc. All Rights Reserved.
 
 #include "Misc/PreLoadFile.h"
 #include "HAL/Event.h"
diff --git a/Engine/Source/Runtime/Core/Private/Serialization/MemoryImage.cpp b/Engine/Source/Runtime/Core/Private/Serialization/MemoryImage.cpp
--- a/Engine/Source/Runtime/Core/Private/Serialization/MemoryImage.cpp
+++ b/Engine/Source/Runtime/Core/Private/Serialization/MemoryImage.cpp
@@ -1,4 +1,4 @@
-// Copyright 1998-2019 Epic Games, Inc. All Rights Reserved.
+// Copyright Epic Games, Inc. All Rights Reserved.
 
 #include "Serialization/MemoryImage.h"
 #include "Containers/UnrealString.h"
diff --git a/Engine/Source/Runtime/Core/Public/Misc/DelayedAutoRegister.h b/Engine/Source/Runtime/Core/Public/Misc/DelayedAutoRegister.h
--- a/Engine/Source/Runtime/Core/Public/Misc/DelayedAutoRegister.h
+++ b/Engine/Source/Runtime/Core/Public/Misc/DelayedAutoRegister.h
@@ -1,4 +1,4 @@
-// Copyright 1998-2019 Epic Games, Inc. All Rights Reserved.
+// Copyright Epic Games, Inc. All Rights Reserved.
 
 #pragma once
 
diff --git a/Engine/Source/Runtime/Core/Public/Misc/PreLoadFile.h b/Engine/Source/Runtime/Core/Public/Misc/PreLoadFile.h
--- a/Engine/Source/Runtime/Core/Public/Misc/PreLoadFile.h
+++ b/Engine/Source/Runtime/Core/Public/Misc/PreLoadFile.h
@@ -1,4 +1,4 @@
-// Copyright 1998-2019 Epic Games, Inc. All Rights Reserved.
+// Copyright Epic Games, Inc. All Rights Reserved.
 
 #pragma once
 
diff --git a/Engine/Source/Runtime/Core/Public/Serialization/MemoryImage.h b/Engine/Source/Runtime/Core/Public/Serialization/MemoryImage.h
--- a/Engine/Source/Runtime/Core/Public/Serialization/MemoryImage.h
+++ b/Engine/Source/Runtime/Core/Public/Serialization/MemoryImage.h
@@ -1,3 +1,5 @@
+// Copyright Epic Games, Inc. All Rights Reserved.
+
 #pragma once
 
 #include "CoreTypes.h"
diff --git a/Engine/Source/Runtime/Core/Public/Serialization/MemoryImageWriter.h b/Engine/Source/Runtime/Core/Public/Serialization/MemoryImageWriter.h
--- a/Engine/Source/Runtime/Core/Public/Serialization/MemoryImageWriter.h
+++ b/Engine/Source/Runtime/Core/Public/Serialization/MemoryImageWriter.h
@@ -1,3 +1,5 @@
+// Copyright Epic Games, Inc. All Rights Reserved.
+
 #pragma once
 
 #include "CoreTypes.h"
diff --git a/Engine/Source/Runtime/Core/Public/Serialization/MemoryLayout.h b/Engine/Source/Runtime/Core/Public/Serialization/MemoryLayout.h
--- a/Engine/Source/Runtime/Core/Public/Serialization/MemoryLayout.h
+++ b/Engine/Source/Runtime/Core/Public/Serialization/MemoryLayout.h
@@ -1,4 +1,4 @@
-// Copyright 1998-2019 Epic Games, Inc. All Rights Reserved.
+// Copyright Epic Games, Inc. All Rights Reserved.
 
 #pragma once
 
diff --git a/Engine/Source/Runtime/Core/Public/Templates/IsPolymorphic.h b/Engine/Source/Runtime/Core/Public/Templates/IsPolymorphic.h
--- a/Engine/Source/Runtime/Core/Public/Templates/IsPolymorphic.h
+++ b/Engine/Source/Runtime/Core/Public/Templates/IsPolymorphic.h
@@ -1,4 +1,4 @@
-// Copyright 1998-2019 Epic Games, Inc. All Rights Reserved.
+// Copyright Epic Games, Inc. All Rights Reserved.
 
 #pragma once
 
diff --git a/Engine/Source/Runtime/Engine/Private/Materials/MaterialCachedData.cpp b/Engine/Source/Runtime/Engine/Private/Materials/MaterialCachedData.cpp
--- a/Engine/Source/Runtime/Engine/Private/Materials/MaterialCachedData.cpp
+++ b/Engine/Source/Runtime/Engine/Private/Materials/MaterialCachedData.cpp
@@ -1,3 +1,5 @@
+// Copyright Epic Games, Inc. All Rights Reserved.
+
 #include "MaterialCachedData.h"
 #include "Materials/MaterialExpressionBreakMaterialAttributes.h"
 #include "Materials/MaterialExpressionShadingModel.h"
diff --git a/Engine/Source/Runtime/Engine/Private/Materials/MaterialShaderMapLayout.h b/Engine/Source/Runtime/Engine/Private/Materials/MaterialShaderMapLayout.h
--- a/Engine/Source/Runtime/Engine/Private/Materials/MaterialShaderMapLayout.h
+++ b/Engine/Source/Runtime/Engine/Private/Materials/MaterialShaderMapLayout.h
@@ -1,4 +1,4 @@
-// Copyright 1998-2019 Epic Games, Inc. All Rights Reserved.
+// Copyright Epic Games, Inc. All Rights Reserved.
 
 #pragma once
 
diff --git a/Engine/Source/Runtime/Engine/Public/MaterialCachedData.h b/Engine/Source/Runtime/Engine/Public/MaterialCachedData.h
--- a/Engine/Source/Runtime/Engine/Public/MaterialCachedData.h
+++ b/Engine/Source/Runtime/Engine/Public/MaterialCachedData.h
@@ -1,3 +1,5 @@
+// Copyright Epic Games, Inc. All Rights Reserved.
+
 #pragma once
 
 #include "CoreMinimal.h"
diff --git a/Engine/Source/Runtime/RenderCore/Private/ShaderCodeArchive.cpp b/Engine/Source/Runtime/RenderCore/Private/ShaderCodeArchive.cpp
--- a/Engine/Source/Runtime/RenderCore/Private/ShaderCodeArchive.cpp
+++ b/Engine/Source/Runtime/RenderCore/Private/ShaderCodeArchive.cpp
@@ -1,3 +1,5 @@
+// Copyright Epic Games, Inc. All Rights Reserved.
+
 #include "ShaderCodeArchive.h"
 #include "ShaderCodeLibrary.h"
 #include "Shader.h"
diff --git a/Engine/Source/Runtime/RenderCore/Private/ShaderMap.cpp b/Engine/Source/Runtime/RenderCore/Private/ShaderMap.cpp
--- a/Engine/Source/Runtime/RenderCore/Private/ShaderMap.cpp
+++ b/Engine/Source/Runtime/RenderCore/Private/ShaderMap.cpp
@@ -1,4 +1,4 @@
-// Copyright 1998-2019 Epic Games, Inc. All Rights Reserved.
+// Copyright Epic Games, Inc. All Rights Reserved.
 
 /*=============================================================================
 	Shader.cpp: Shader implementation.
diff --git a/Engine/Source/Runtime/RenderCore/Private/ShaderResource.cpp b/Engine/Source/Runtime/RenderCore/Private/ShaderResource.cpp
--- a/Engine/Source/Runtime/RenderCore/Private/ShaderResource.cpp
+++ b/Engine/Source/Runtime/RenderCore/Private/ShaderResource.cpp
@@ -1,4 +1,4 @@
-// Copyright 1998-2019 Epic Games, Inc. All Rights Reserved.
+// Copyright Epic Games, Inc. All Rights Reserved.
 
 /*=============================================================================
 	ShaderResource.cpp: ShaderResource implementation.
diff --git a/Engine/Source/Runtime/RenderCore/Public/ShaderCodeArchive.h b/Engine/Source/Runtime/RenderCore/Public/ShaderCodeArchive.h
--- a/Engine/Source/Runtime/RenderCore/Public/ShaderCodeArchive.h
+++ b/Engine/Source/Runtime/RenderCore/Public/ShaderCodeArchive.h
@@ -1,3 +1,5 @@
+// Copyright Epic Games, Inc. All Rights Reserved.
+
 #pragma once
 
 #include "CoreMinimal.h"
