Fix potential null dereference
#rnx

diff --git a/Engine/Source/Runtime/RenderCore/Private/Shader.cpp b/Engine/Source/Runtime/RenderCore/Private/Shader.cpp
--- a/Engine/Source/Runtime/RenderCore/Private/Shader.cpp
+++ b/Engine/Source/Runtime/RenderCore/Private/Shader.cpp
@@ -1663,7 +1663,7 @@ void ShaderMapAppendKeyString(EShaderPlatform Platform, FString& KeyString)
 			VTTextures = (CVarMobileVirtualTexture->GetValueOnAnyThread() != 0);
 		}
 
-		const bool VTSupported = TargetPlatform->SupportsFeature(ETargetPlatformFeatures::VirtualTextureStreaming);
+		const bool VTSupported = TargetPlatform != nullptr && TargetPlatform->SupportsFeature(ETargetPlatformFeatures::VirtualTextureStreaming);
 
 		static const auto CVarVTFactor = IConsoleManager::Get().FindConsoleVariable(TEXT("r.vt.FeedbackFactor")); check(CVarVTFactor);
 		const int32 VTFeedbackFactor = CVarVTFactor->GetInt(); 
