Convert some checks() into checkf(), to hopefully gather more information about failing check

diff --git a/Engine/Source/Runtime/OpenGLDrv/Private/OpenGLShaders.cpp b/Engine/Source/Runtime/OpenGLDrv/Private/OpenGLShaders.cpp
--- a/Engine/Source/Runtime/OpenGLDrv/Private/OpenGLShaders.cpp
+++ b/Engine/Source/Runtime/OpenGLDrv/Private/OpenGLShaders.cpp
@@ -882,6 +882,9 @@ ShaderType* CompileOpenGLShader(TArrayView<const uint8> InShaderCode, const FSHA
 		}
 	}
 
+	checkf(Shader->StaticSlots.Num() == Shader->Bindings.ShaderResourceTable.ResourceTableLayoutHashes.Num(), TEXT("StaticSlots %d, Bindings %d"),
+		Shader->StaticSlots.Num(), Shader->Bindings.ShaderResourceTable.ResourceTableLayoutHashes.Num());
+
 	if (FOpenGL::SupportsSeparateShaderObjects())
 	{
 		FSHAHash Hash;
diff --git a/Engine/Source/Runtime/RenderCore/Private/ShaderMap.cpp b/Engine/Source/Runtime/RenderCore/Private/ShaderMap.cpp
--- a/Engine/Source/Runtime/RenderCore/Private/ShaderMap.cpp
+++ b/Engine/Source/Runtime/RenderCore/Private/ShaderMap.cpp
@@ -217,13 +217,13 @@ void FShaderMapBase::Serialize(FArchive& Ar, bool bInlineShaderResources, bool b
 				{
 					FSHAHash CheckLayoutHash;
 					const uint32 CheckLayoutSize = Freeze::HashLayout(*DependencyType, LayoutParams, CheckLayoutHash);
-					checkf(CheckLayoutSize == SavedLayoutSize, TEXT("Mismatch size for type %s"), DependencyType->Name);
+					checkf(CheckLayoutSize == SavedLayoutSize, TEXT("Mismatch size for type %s, compiled size is %d, loaded size is %d"), DependencyType->Name, CheckLayoutSize, SavedLayoutSize);
 					checkf(CheckLayoutHash == SavedLayoutHash, TEXT("Mismatch hash for type %s"), DependencyType->Name);
 				}
 #endif // CHECK_SHADERMAP_DEPENDENCIES
 			}
 		}
-
+		
 		bool bShareCode = false;
 		Ar << bShareCode;
 		if (bShareCode)
diff --git a/Engine/Source/Runtime/RenderCore/Public/ShaderCore.h b/Engine/Source/Runtime/RenderCore/Public/ShaderCore.h
--- a/Engine/Source/Runtime/RenderCore/Public/ShaderCore.h
+++ b/Engine/Source/Runtime/RenderCore/Public/ShaderCore.h
@@ -1266,7 +1266,8 @@ void ApplyGlobalUniformBuffers(
 	const TArray<uint32>& LayoutHashes,
 	const TArray<FRHIUniformBuffer*>& UniformBuffers)
 {
-	check(LayoutHashes.Num() == Slots.Num());
+	checkf(LayoutHashes.Num() == Slots.Num(), TEXT("Shader %s, LayoutHashes %d, Slots %d"),
+		*Shader->ShaderName, LayoutHashes.Num(), Slots.Num());
 
 	for (int32 BufferIndex = 0; BufferIndex < Slots.Num(); ++BufferIndex)
 	{
