Fix event name pointer collision problem in GPU profiler

diff --git a/Engine/Source/Runtime/RenderCore/Private/ProfilingDebugging/RealtimeGPUProfiler.cpp b/Engine/Source/Runtime/RenderCore/Private/ProfilingDebugging/RealtimeGPUProfiler.cpp
--- a/Engine/Source/Runtime/RenderCore/Private/ProfilingDebugging/RealtimeGPUProfiler.cpp
+++ b/Engine/Source/Runtime/RenderCore/Private/ProfilingDebugging/RealtimeGPUProfiler.cpp
@@ -419,12 +419,9 @@ public:
 			const bool bTracingStatsEnabled = !!CVarGPUTracingStatsEnabled.GetValueOnRenderThread();
 			if (bTracingStatsEnabled)
 			{
-				ANSICHAR EventName[NAME_SIZE];
-				Event.GetName().GetPlainANSIString(EventName);
-
 				const uint32 GPUIndex = 0;
 				FTracingProfiler::Get()->AddGPUEvent(
-					EventName,
+					Event.GetName(),
 					Event.GetStartResultMicroseconds(),
 					Event.GetEndResultMicroseconds(),
 					GPUIndex,
diff --git a/Engine/Source/Runtime/RenderCore/Private/ProfilingDebugging/TracingProfiler.cpp b/Engine/Source/Runtime/RenderCore/Private/ProfilingDebugging/TracingProfiler.cpp
--- a/Engine/Source/Runtime/RenderCore/Private/ProfilingDebugging/TracingProfiler.cpp
+++ b/Engine/Source/Runtime/RenderCore/Private/ProfilingDebugging/TracingProfiler.cpp
@@ -12,6 +12,7 @@
 #include "Misc/Paths.h"
 #include "RHI.h"
 #include "RenderingThread.h"
+#include "UObject/NameTypes.h"
 
 static TAutoConsoleVariable<int32> GTracingProfileBufferSize(
 	TEXT("TracingProfiler.BufferSize"),
@@ -112,7 +113,7 @@ void FTracingProfiler::EndFrame()
 
 	GameThreadFrameEndCycle = FPlatformTime::Cycles64();
 
-	AddCPUEvent("GameThread",
+	AddCPUEvent(NAME_GameThread,
 		GameThreadFrameBeginCycle,
 		GameThreadFrameEndCycle,
 		GGameThreadId,
@@ -162,7 +163,7 @@ void FTracingProfiler::EndFrameRT()
 		return;
 	}
 
-	AddCPUEvent("RenderThread",
+	AddCPUEvent(NAME_RenderThread,
 		RenderThreadFrameBeginCycle,
 		RenderThreadFrameEndCycle,
 		GRenderThreadId,
@@ -344,10 +345,13 @@ void FTracingProfiler::WriteCaptureToFile()
 			ThreadOrGpuId = Event.GPU.GPUIndex;
 		}
 
+		ANSICHAR EventName[NAME_SIZE];
+		Event.Name.GetPlainANSIString(EventName);
+
 		FCStringAnsi::Snprintf(StringBuffer, StringBufferSize,
 			R"({"pid":%d, "tid":%d, "ph": "X", "name": "%s", "ts": %llu, "dur": %llu, "args":{"frame":%d}},)"
 			"\n",
-			Pid, ThreadOrGpuId, Event.Name, TimeBeginMicroseconds,
+			Pid, ThreadOrGpuId, EventName, TimeBeginMicroseconds,
 			TimeEndMicroseconds - TimeBeginMicroseconds,
 			Event.FrameNumber);
 
diff --git a/Engine/Source/Runtime/RenderCore/Public/ProfilingDebugging/TracingProfiler.h b/Engine/Source/Runtime/RenderCore/Public/ProfilingDebugging/TracingProfiler.h
--- a/Engine/Source/Runtime/RenderCore/Public/ProfilingDebugging/TracingProfiler.h
+++ b/Engine/Source/Runtime/RenderCore/Public/ProfilingDebugging/TracingProfiler.h
@@ -41,7 +41,7 @@ public:
 
 	struct FEvent
 	{
-		const char* Name;
+		FName Name;
 		uint32 FrameNumber;
 		uint32 SessionId;
 		EEventType Type;
@@ -67,7 +67,7 @@ public:
 
 	RENDERCORE_API void Init();
 
-	uint32 AddCPUEvent(const char* Name,
+	uint32 AddCPUEvent(FName Name,
 		uint64 TimestampBeginCycles,
 		uint64 TimestempEndCycles,
 		uint32 ThreadId,
@@ -84,7 +84,7 @@ public:
 		return AddEvent(Event);
 	}
 
-	uint32 AddGPUEvent(const char* Name,
+	uint32 AddGPUEvent(FName Name,
 		uint64 TimestempBeginMicroseconds,
 		uint64 TimestampEndMicroseconds,
 		uint64 GPUIndex,
