Fixed RDG issue where registered external buffer was not having its RHI resource assigned.

diff --git a/Engine/Source/Runtime/RenderCore/Public/RenderGraphBuilder.inl b/Engine/Source/Runtime/RenderCore/Public/RenderGraphBuilder.inl
--- a/Engine/Source/Runtime/RenderCore/Public/RenderGraphBuilder.inl
+++ b/Engine/Source/Runtime/RenderCore/Public/RenderGraphBuilder.inl
@@ -52,6 +52,18 @@ inline FRDGBufferRef FRDGBuilder::RegisterExternalBuffer(
 
 	FRDGBuffer* OutBuffer = AllocateForRHILifeTime<FRDGBuffer>(Name, ExternalPooledBuffer->Desc, Flags);
 	OutBuffer->PooledBuffer = ExternalPooledBuffer;
+	switch (Buffer->Desc.UnderlyingType)
+	{
+	case FRDGBufferDesc::EUnderlyingType::VertexBuffer:
+		Buffer->ResourceRHI = AllocatedBuffer->VertexBuffer;
+		break;
+	case FRDGBufferDesc::EUnderlyingType::IndexBuffer:
+		Buffer->ResourceRHI = AllocatedBuffer->IndexBuffer;
+		break;
+	case FRDGBufferDesc::EUnderlyingType::StructuredBuffer:
+		Buffer->ResourceRHI = AllocatedBuffer->StructuredBuffer;
+		break;
+	}
 	AllocatedBuffers.Add(OutBuffer, ExternalPooledBuffer);
 
 	IF_RDG_ENABLE_DEBUG(Validation.ValidateCreateExternalBuffer(OutBuffer));
