Fixed build break.

diff --git a/Engine/Source/Runtime/RenderCore/Public/RenderGraphBuilder.inl b/Engine/Source/Runtime/RenderCore/Public/RenderGraphBuilder.inl
--- a/Engine/Source/Runtime/RenderCore/Public/RenderGraphBuilder.inl
+++ b/Engine/Source/Runtime/RenderCore/Public/RenderGraphBuilder.inl
@@ -52,16 +52,16 @@ inline FRDGBufferRef FRDGBuilder::RegisterExternalBuffer(
 
 	FRDGBuffer* OutBuffer = AllocateForRHILifeTime<FRDGBuffer>(Name, ExternalPooledBuffer->Desc, Flags);
 	OutBuffer->PooledBuffer = ExternalPooledBuffer;
-	switch (Buffer->Desc.UnderlyingType)
+	switch (OutBuffer->Desc.UnderlyingType)
 	{
 	case FRDGBufferDesc::EUnderlyingType::VertexBuffer:
-		Buffer->ResourceRHI = AllocatedBuffer->VertexBuffer;
+		OutBuffer->ResourceRHI = ExternalPooledBuffer->VertexBuffer;
 		break;
 	case FRDGBufferDesc::EUnderlyingType::IndexBuffer:
-		Buffer->ResourceRHI = AllocatedBuffer->IndexBuffer;
+		OutBuffer->ResourceRHI = ExternalPooledBuffer->IndexBuffer;
 		break;
 	case FRDGBufferDesc::EUnderlyingType::StructuredBuffer:
-		Buffer->ResourceRHI = AllocatedBuffer->StructuredBuffer;
+		OutBuffer->ResourceRHI = ExternalPooledBuffer->StructuredBuffer;
 		break;
 	}
 	AllocatedBuffers.Add(OutBuffer, ExternalPooledBuffer);
