Disable validation check on RHIResources when compiling for platform shader compilers that currently output the byte-size of the resource.

Workaround for:
Follow on jira: UE-87339

#rnx

diff --git a/Engine/Source/Runtime/RenderCore/Private/ShaderParameterStruct.cpp b/Engine/Source/Runtime/RenderCore/Private/ShaderParameterStruct.cpp
--- a/Engine/Source/Runtime/RenderCore/Private/ShaderParameterStruct.cpp
+++ b/Engine/Source/Runtime/RenderCore/Private/ShaderParameterStruct.cpp
@@ -201,12 +201,18 @@ struct FShaderParameterStructBindingContext
 
 					if (BoundSize != 1)
 					{
-						UE_LOG(LogShaders, Fatal, 
-							TEXT("Error with shader %s's (Permutation Id %d) parameter %s is %i bytes, cpp name = %s.")
-							TEXT("The shader compiler should give precisely which elements of an array did not get compiled out, ")
-							TEXT("for optimal automatic render graph pass dependency with ClearUnusedGraphResources()."),
-							Shader->GetType()->GetName(), Shader->GetPermutationId(),
-							*ElementShaderBindingName, BoundSize, *CppName);
+						// Switch shader compiler does not yet support this validation on RHIResouces, see UE-86533
+						const EShaderPlatform& ShaderPlatform = Shader->GetShaderPlatform();
+						const bool bIsSwitchShader = ShaderPlatform == SP_SWITCH || ShaderPlatform == SP_SWITCH_FORWARD;
+						if (bIsSwitchShader && !bIsRHIResource)
+						{
+							UE_LOG(LogShaders, Fatal,
+								TEXT("Error with shader %s's (Permutation Id %d) parameter %s is %i bytes, cpp name = %s.")
+								TEXT("The shader compiler should give precisely which elements of an array did not get compiled out, ")
+								TEXT("for optimal automatic render graph pass dependency with ClearUnusedGraphResources()."),
+								Shader->GetType()->GetName(), Shader->GetPermutationId(),
+								*ElementShaderBindingName, BoundSize, *CppName);
+						}
 					}
 
 					if (BaseType == UBMT_TEXTURE)
