Add GetShaderName() accessor to FRHIShader, fix compile error with certain build config

diff --git a/Engine/Source/Runtime/RHI/Public/RHIResources.h b/Engine/Source/Runtime/RHI/Public/RHIResources.h
--- a/Engine/Source/Runtime/RHI/Public/RHIResources.h
+++ b/Engine/Source/Runtime/RHI/Public/RHIResources.h
@@ -199,6 +199,9 @@ public:
 #if !(UE_BUILD_SHIPPING || UE_BUILD_TEST)
 	// for debugging only e.g. MaterialName:ShaderFile.usf or ShaderFile.usf/EntryFunc
 	FString ShaderName;
+	FORCEINLINE const TCHAR* GetShaderName() const { return *ShaderName; }
+#else
+	FORCEINLINE const TCHAR* GetShaderName() const { TEXT(""); }
 #endif
 
 	explicit FRHIShader(EShaderFrequency InFrequency)
diff --git a/Engine/Source/Runtime/RenderCore/Public/ShaderCore.h b/Engine/Source/Runtime/RenderCore/Public/ShaderCore.h
--- a/Engine/Source/Runtime/RenderCore/Public/ShaderCore.h
+++ b/Engine/Source/Runtime/RenderCore/Public/ShaderCore.h
@@ -1271,7 +1271,7 @@ void ApplyGlobalUniformBuffers(
 	const TArray<FRHIUniformBuffer*>& UniformBuffers)
 {
 	checkf(LayoutHashes.Num() == Slots.Num(), TEXT("Shader %s, LayoutHashes %d, Slots %d"),
-		*Shader->ShaderName, LayoutHashes.Num(), Slots.Num());
+		Shader->GetShaderName(), LayoutHashes.Num(), Slots.Num());
 
 	for (int32 BufferIndex = 0; BufferIndex < Slots.Num(); ++BufferIndex)
 	{
