Make shader library determinism optional per platform.

#rnx

diff --git a/Engine/Source/Editor/UnrealEd/Classes/Settings/ProjectPackagingSettings.h b/Engine/Source/Editor/UnrealEd/Classes/Settings/ProjectPackagingSettings.h
--- a/Engine/Source/Editor/UnrealEd/Classes/Settings/ProjectPackagingSettings.h
+++ b/Engine/Source/Editor/UnrealEd/Classes/Settings/ProjectPackagingSettings.h
@@ -259,6 +259,15 @@ public:
 	bool bShareMaterialShaderCode;
 
 	/** 
+	 * With this option off, the shader code will be stored in the library essentially in a random order,
+	 * squarely the same in which the assets were loaded by the cooker. Enabling this will sort the shaders
+	 * by their hash, which makes the shader library more similar between the builds which can help patching, but
+	 * can adversely affect loading times.
+	 */
+	UPROPERTY(config, EditAnywhere, Category = Packaging, meta = (EditCondition = "bShareMaterialShaderCode"))
+	bool bDeterministicShaderCodeOrder;
+
+	/**
 	 * By default shader shader code gets saved into individual platform agnostic files,
 	 * enabling this option will use the platform-specific library format if and only if one is available
 	 * This will reduce overall package size but might increase loading time
diff --git a/Engine/Source/Editor/UnrealEd/Private/Commandlets/ShaderCodeLibraryToolsCommandlet.cpp b/Engine/Source/Editor/UnrealEd/Private/Commandlets/ShaderCodeLibraryToolsCommandlet.cpp
--- a/Engine/Source/Editor/UnrealEd/Private/Commandlets/ShaderCodeLibraryToolsCommandlet.cpp
+++ b/Engine/Source/Editor/UnrealEd/Private/Commandlets/ShaderCodeLibraryToolsCommandlet.cpp
@@ -39,11 +39,14 @@ int32 UShaderCodeLibraryToolsCommandlet::StaticMain(const FString& Params)
 		{
 			bNativeFormat |= bArchive;
 		}
+
+		bool bNeedsDeterministicOrder = false;
+		GConfig->GetBool(TEXT("/Script/UnrealEd.ProjectPackagingSettings"), TEXT("bDeterministicShaderCodeOrder"), bNeedsDeterministicOrder, GGameIni);
 		
 		TArray<FString> OldMetaDataDirs;
 		OldMetaDataDirs.Add(Left);
 		
-		return FShaderCodeLibrary::CreatePatchLibrary(OldMetaDataDirs, Right, Output, bNativeFormat) ? 0 : 1;
+		return FShaderCodeLibrary::CreatePatchLibrary(OldMetaDataDirs, Right, Output, bNativeFormat, bNeedsDeterministicOrder) ? 0 : 1;
 	}
 	
 	UE_LOG(LogShaderCodeLibraryTools, Warning, TEXT("Usage: <Path-To-Old-MetaData> <Path-To-New-MetaData> <Output-Path> [-PreferNativeArchives]\n"));
diff --git a/Engine/Source/Editor/UnrealEd/Private/CookOnTheFlyServer.cpp b/Engine/Source/Editor/UnrealEd/Private/CookOnTheFlyServer.cpp
--- a/Engine/Source/Editor/UnrealEd/Private/CookOnTheFlyServer.cpp
+++ b/Engine/Source/Editor/UnrealEd/Private/CookOnTheFlyServer.cpp
@@ -6609,13 +6609,22 @@ void UCookOnTheFlyServer::InitShaderCodeLibrary(void)
 			FConfigFile PlatformIniFile;
 			FConfigCacheIni::LoadLocalIniFile(PlatformIniFile, TEXT("Engine"), true, *TargetPlatform->IniPlatformName());
 			PlatformIniFile.GetBool(TEXT("DevOptions.Shaders"), TEXT("NeedsShaderStableKeys"), bNeedShaderStableKeys);
-            
+
+			bool bNeedsDeterministicOrder = PackagingSettings->bDeterministicShaderCodeOrder;
+			FConfigFile PlatformGameIniFile;
+			FConfigCacheIni::LoadLocalIniFile(PlatformGameIniFile, TEXT("Game"), true, *TargetPlatform->IniPlatformName());
+			PlatformGameIniFile.GetBool(TEXT("/Script/UnrealEd.ProjectPackagingSettings"), TEXT("bDeterministicShaderCodeOrder"), bNeedsDeterministicOrder);
+
             TArray<FName> ShaderFormats;
             TargetPlatform->GetAllTargetedShaderFormats(ShaderFormats);
-			TArray<TPair<FName, bool>> ShaderFormatsWithStableKeys;
+			TArray<FShaderCodeLibrary::FShaderFormatDescriptor> ShaderFormatsWithStableKeys;
 			for (FName& Format : ShaderFormats)
 			{
-				ShaderFormatsWithStableKeys.Push(MakeTuple(Format, bNeedShaderStableKeys));
+				FShaderCodeLibrary::FShaderFormatDescriptor NewDesc;
+				NewDesc.ShaderFormat = Format;
+				NewDesc.bNeedsStableKeys = bNeedShaderStableKeys;
+				NewDesc.bNeedsDeterministicOrder = bNeedsDeterministicOrder;
+				ShaderFormatsWithStableKeys.Push(NewDesc);
 			}
 
             if (ShaderFormats.Num() > 0)
diff --git a/Engine/Source/Runtime/RenderCore/Private/ShaderCodeLibrary.cpp b/Engine/Source/Runtime/RenderCore/Private/ShaderCodeLibrary.cpp
--- a/Engine/Source/Runtime/RenderCore/Private/ShaderCodeLibrary.cpp
+++ b/Engine/Source/Runtime/RenderCore/Private/ShaderCodeLibrary.cpp
@@ -800,9 +800,10 @@ struct FShaderCodeStats
 
 struct FEditorShaderCodeArchive
 {
-	FEditorShaderCodeArchive(FName InFormat)
+	FEditorShaderCodeArchive(FName InFormat, bool bInNeedsDeterministicOrder)
 		: FormatName(InFormat)
 		, Format(nullptr)
+		, bNeedsDeterministicOrder(bInNeedsDeterministicOrder)
 	{
 		Format = GetTargetPlatformManagerRef().FindShaderFormat(InFormat);
 		check(Format);
@@ -1090,8 +1091,6 @@ struct FEditorShaderCodeArchive
 
 		bool bSuccess = IFileManager::Get().MakeDirectory(*OutputDir, true);
 
-		EShaderPlatform Platform = ShaderFormatToLegacyShaderPlatform(FormatName);
-
 		// Shader library
 		if (bSuccess && SerializedShaders.GetNumShaderMaps() > 0)
 		{
@@ -1154,6 +1153,11 @@ struct FEditorShaderCodeArchive
 
 			*FileWriter << GShaderPipelineArchiveVersion;
 
+			if (bNeedsDeterministicOrder)
+			{
+				Pipelines.Sort([](const FShaderCodeLibraryPipeline& A, const FShaderCodeLibraryPipeline& B){ return A.Hash <= B.Hash; });
+			}
+
 			*FileWriter << Pipelines;
 
 			FileWriter->Close();
@@ -1248,12 +1252,12 @@ struct FEditorShaderCodeArchive
 		}
 	}
 	
-	static bool CreatePatchLibrary(FName FormatName, FString const& LibraryName, TArray<FString> const& OldMetaDataDirs, FString const& NewMetaDataDir, FString const& OutDir, bool bNativeFormat)
+	static bool CreatePatchLibrary(FName FormatName, FString const& LibraryName, TArray<FString> const& OldMetaDataDirs, FString const& NewMetaDataDir, FString const& OutDir, bool bNativeFormat, bool bNeedsDeterministicOrder)
 	{
 		TArray<FEditorShaderCodeArchive*> OldLibraries;
 		for (FString const& OldMetaDataDir : OldMetaDataDirs)
 		{
-			FEditorShaderCodeArchive* OldLibrary = new FEditorShaderCodeArchive(FormatName);
+			FEditorShaderCodeArchive* OldLibrary = new FEditorShaderCodeArchive(FormatName, bNeedsDeterministicOrder);
 			OldLibrary->OpenLibrary(LibraryName);
 			if (OldLibrary->LoadExistingShaderCodeLibrary(OldMetaDataDir))
 			{
@@ -1261,12 +1265,12 @@ struct FEditorShaderCodeArchive
 			}
 		}
 
-		FEditorShaderCodeArchive NewLibrary(FormatName);
+		FEditorShaderCodeArchive NewLibrary(FormatName, bNeedsDeterministicOrder);
 		NewLibrary.OpenLibrary(LibraryName);
 		bool bOK = NewLibrary.LoadExistingShaderCodeLibrary(NewMetaDataDir);
 		if (bOK)
 		{
-			FEditorShaderCodeArchive OutLibrary(FormatName);
+			FEditorShaderCodeArchive OutLibrary(FormatName, bNeedsDeterministicOrder);
 			OutLibrary.OpenLibrary(LibraryName);
 			OutLibrary.MakePatchLibrary(OldLibraries, NewLibrary);
 			bOK = OutLibrary.SerializedShaders.GetNumShaderMaps() > 0;
@@ -1305,6 +1309,7 @@ private:
 	FSerializedShaderArchive SerializedShaders;
 	TArray<TArray<uint8>> ShaderCode;
 	const IShaderFormat* Format;
+	bool bNeedsDeterministicOrder;
 };
 
 struct FEditorShaderStableInfo
@@ -1718,27 +1723,27 @@ public:
 		}
 	}
 
-	void CookShaderFormats(TArray<TTuple<FName,bool>> const& ShaderFormats)
+	void CookShaderFormats(TArray<FShaderCodeLibrary::FShaderFormatDescriptor> const& ShaderFormats)
 	{
-		for (auto Pair : ShaderFormats)
+		for (const FShaderCodeLibrary::FShaderFormatDescriptor& Descriptor : ShaderFormats)
 		{
-			FName const& Format = Pair.Get<0>();
+			FName const& Format = Descriptor.ShaderFormat;
 
 			EShaderPlatform Platform = ShaderFormatToLegacyShaderPlatform(Format);
 			FName PossiblyAdjustedFormat = LegacyShaderPlatformToShaderFormat(Platform);	// Vulkan and GL switch between name variants depending on CVars (e.g. see r.Vulkan.UseRealUBs)
 			FEditorShaderCodeArchive* CodeArchive = EditorShaderCodeArchive[Platform];
 			if (!CodeArchive)
 			{
-				CodeArchive = new FEditorShaderCodeArchive(PossiblyAdjustedFormat);
+				CodeArchive = new FEditorShaderCodeArchive(PossiblyAdjustedFormat, Descriptor.bNeedsDeterministicOrder);
 				EditorShaderCodeArchive[Platform] = CodeArchive;
 				EditorArchivePipelines[Platform] = !bNativeFormat;
 			}
 			check(CodeArchive);
 		}
-		for (auto Pair : ShaderFormats)
+		for (const FShaderCodeLibrary::FShaderFormatDescriptor& Descriptor : ShaderFormats)
 		{
-			FName const& Format = Pair.Get<0>();
-			bool bUseStableKeys = Pair.Get<1>();
+			FName const& Format = Descriptor.ShaderFormat;
+			bool bUseStableKeys = Descriptor.bNeedsStableKeys;
 
 			EShaderPlatform Platform = ShaderFormatToLegacyShaderPlatform(Format);
 			FName PossiblyAdjustedFormat = LegacyShaderPlatformToShaderFormat(Platform);	// Vulkan and GL switch between name variants depending on CVars (e.g. see r.Vulkan.UseRealUBs)
@@ -2091,7 +2096,7 @@ void FShaderCodeLibrary::CleanDirectories(TArray<FName> const& ShaderFormats)
 	}
 }
 
-void FShaderCodeLibrary::CookShaderFormats(TArray<TTuple<FName,bool>> const& ShaderFormats)
+void FShaderCodeLibrary::CookShaderFormats(TArray<FShaderFormatDescriptor> const& ShaderFormats)
 {
 	if (FShaderCodeLibraryImpl::Impl)
 	{
@@ -2172,7 +2177,7 @@ void FShaderCodeLibrary::DumpShaderCodeStats()
 	}
 }
 
-bool FShaderCodeLibrary::CreatePatchLibrary(TArray<FString> const& OldMetaDataDirs, FString const& NewMetaDataDir, FString const& OutDir, bool bNativeFormat)
+bool FShaderCodeLibrary::CreatePatchLibrary(TArray<FString> const& OldMetaDataDirs, FString const& NewMetaDataDir, FString const& OutDir, bool bNativeFormat, bool bNeedsDeterministicOrder)
 {
 	TMap<FName, TSet<FString>> FormatLibraryMap;
 	TArray<FString> LibraryFiles;
@@ -2198,7 +2203,7 @@ bool FShaderCodeLibrary::CreatePatchLibrary(TArray<FString> const& OldMetaDataDi
 	{
 		for (auto const& Library : Entry.Value)
 		{
-			bOK |= FEditorShaderCodeArchive::CreatePatchLibrary(Entry.Key, Library, OldMetaDataDirs, NewMetaDataDir, OutDir, bNativeFormat);
+			bOK |= FEditorShaderCodeArchive::CreatePatchLibrary(Entry.Key, Library, OldMetaDataDirs, NewMetaDataDir, OutDir, bNativeFormat, bNeedsDeterministicOrder);
 		}
 	}
 	return bOK;
diff --git a/Engine/Source/Runtime/RenderCore/Public/ShaderCodeLibrary.h b/Engine/Source/Runtime/RenderCore/Public/ShaderCodeLibrary.h
--- a/Engine/Source/Runtime/RenderCore/Public/ShaderCodeLibrary.h
+++ b/Engine/Source/Runtime/RenderCore/Public/ShaderCodeLibrary.h
@@ -181,9 +181,15 @@ struct RENDERCORE_API FShaderCodeLibrary
 	// Clean the cook directories
 	static void CleanDirectories(TArray<FName> const& ShaderFormats);
     
-    // Specify the shader formats to cook and which ones needs stable keys. Provide an array of tuples
-	// with names and whether the format needs stable keys.
-    static void CookShaderFormats(TArray<TTuple<FName,bool>> const& ShaderFormats);
+	struct FShaderFormatDescriptor
+	{
+		FName ShaderFormat;
+		bool bNeedsStableKeys;
+		bool bNeedsDeterministicOrder;
+	};
+
+	// Specify the shader formats to cook and which ones needs stable keys. Provide an array of FShaderFormatDescriptors
+    static void CookShaderFormats(TArray<FShaderFormatDescriptor> const& ShaderFormats);
 	
 	// At cook time, add shader code to collection
 	static bool AddShaderCode(EShaderPlatform ShaderPlatform, const FShaderMapResourceCode* Code);
@@ -209,7 +215,7 @@ struct RENDERCORE_API FShaderCodeLibrary
 	static void DumpShaderCodeStats();
 	
 	// Create a smaller 'patch' library that only contains data from 'NewMetaDataDir' not contained in any of 'OldMetaDataDirs'
-	static bool CreatePatchLibrary(TArray<FString> const& OldMetaDataDirs, FString const& NewMetaDataDir, FString const& OutDir, bool bNativeFormat);
+	static bool CreatePatchLibrary(TArray<FString> const& OldMetaDataDirs, FString const& NewMetaDataDir, FString const& OutDir, bool bNativeFormat, bool bNeedsDeterministicOrder);
 #endif
 	
 	// Safely assign the hash to a shader object
