fix for PSO Cache not properly setting the mask when no cache is present

diff --git a/Engine/Source/Runtime/RenderCore/Private/ShaderPipelineCache.cpp b/Engine/Source/Runtime/RenderCore/Private/ShaderPipelineCache.cpp
--- a/Engine/Source/Runtime/RenderCore/Private/ShaderPipelineCache.cpp
+++ b/Engine/Source/Runtime/RenderCore/Private/ShaderPipelineCache.cpp
@@ -336,8 +336,8 @@ int32 FShaderPipelineCache::GetGameVersionForPSOFileCache()
 
 bool FShaderPipelineCache::SetGameUsageMaskWithComparison(uint64 InMask, FPSOMaskComparisonFn InComparisonFnPtr)
 {
-	bool bMaskChanged = false;
-	
+	static bool bMaskChanged = false;
+
 	if (ShaderPipelineCache != nullptr && CVarPSOFileCacheGameFileMaskEnabled.GetValueOnAnyThread() && !ShaderPipelineCache->bPreOptimizing)
 	{
 		FScopeLock Lock(&ShaderPipelineCache->Mutex);
@@ -345,7 +345,7 @@ bool FShaderPipelineCache::SetGameUsageMaskWithComparison(uint64 InMask, FPSOMas
 		if (ShaderPipelineCache->bOpened)
 		{
 			uint64 OldMask = FPipelineFileCache::SetGameUsageMaskWithComparison(InMask, InComparisonFnPtr);
-			bMaskChanged = OldMask != InMask;
+			bMaskChanged |= OldMask != InMask;
 		
 			ShaderPipelineCache->bReady = true;
 		
@@ -392,7 +392,11 @@ bool FShaderPipelineCache::SetGameUsageMaskWithComparison(uint64 InMask, FPSOMas
 				
 					ShaderPipelineCache->PreFetchedTasks = LocalPreFetchedTasks;
 				
+					bMaskChanged = false;
+					
 					UE_LOG(LogRHI, Display, TEXT("New ShaderPipelineCache GameUsageMask [%llu=>%llu], Enqueued %d of %d tasks for precompile."), OldMask, InMask, Count, ShaderPipelineCache->PreFetchedTasks.Num());
+					
+					return OldMask != InMask;
 				}
 				else
 				{
@@ -406,7 +410,12 @@ bool FShaderPipelineCache::SetGameUsageMaskWithComparison(uint64 InMask, FPSOMas
 		}
 		else
 		{
-			UE_LOG(LogRHI, Display, TEXT("ShaderPipelineCache::SetGameUsageMaskWithComparison failed to set a new mask because the cache was not open"));
+			// NOTE: if this is called and then the cache is opened, but this function is never called again, the PSOs will not get precompiled. Should probably update the open call itself to see if there was a new mask set and call the code above
+			uint64 OldMask = FPipelineFileCache::SetGameUsageMaskWithComparison(InMask, InComparisonFnPtr);
+			bMaskChanged |= OldMask != InMask;
+
+			UE_LOG(LogRHI, Display, TEXT("ShaderPipelineCache::SetGameUsageMaskWithComparison set a new mask but did not attempt to setup any tasks because the cache was not open"));
+			return OldMask != InMask;
 		}
 	}
 	else
