revert out the prioritized PSO caching, seeing some very strange errors which we can't seem to resolve in time for 12.00

diff --git a/Engine/Source/Runtime/RenderCore/Private/ShaderPipelineCache.cpp b/Engine/Source/Runtime/RenderCore/Private/ShaderPipelineCache.cpp
--- a/Engine/Source/Runtime/RenderCore/Private/ShaderPipelineCache.cpp
+++ b/Engine/Source/Runtime/RenderCore/Private/ShaderPipelineCache.cpp
@@ -138,12 +138,6 @@ static TAutoConsoleVariable<int32> CVarPSOFileCacheGameFileMaskEnabled(
 																TEXT("Set non zero to use GameFileMask during PSO precompile - recording should always save out the usage masks to make that data availble when needed."),
 																ECVF_ReadOnly | ECVF_RenderThreadSafe);
 
-static TAutoConsoleVariable<int32> CVarPSOFileCachePreOptimizeEnabled(
-																TEXT("r.ShaderPipelineCache.PreOptimizeEnabled"),
-																(int32)0,
-																TEXT("Set non zero to PreOptimize PSOs - this allows some PSOs to be compiled in the foreground before going in to game"),
-																ECVF_ReadOnly | ECVF_RenderThreadSafe);
-
 
 static TAutoConsoleVariable<int32> CVarPSOFileCacheMinBindCount(
                                                                 TEXT("r.ShaderPipelineCache.MinBindCount"),
@@ -336,91 +330,67 @@ int32 FShaderPipelineCache::GetGameVersionForPSOFileCache()
 
 bool FShaderPipelineCache::SetGameUsageMaskWithComparison(uint64 InMask, FPSOMaskComparisonFn InComparisonFnPtr)
 {
-	static bool bMaskChanged = false;
-
-	if (ShaderPipelineCache != nullptr && CVarPSOFileCacheGameFileMaskEnabled.GetValueOnAnyThread() && !ShaderPipelineCache->bPreOptimizing)
+	bool bMaskChanged = false;
+	
+	if (ShaderPipelineCache != nullptr && CVarPSOFileCacheGameFileMaskEnabled.GetValueOnAnyThread())
 	{
 		FScopeLock Lock(&ShaderPipelineCache->Mutex);
 		
-		if (ShaderPipelineCache->bOpened)
-		{
-			uint64 OldMask = FPipelineFileCache::SetGameUsageMaskWithComparison(InMask, InComparisonFnPtr);
-			bMaskChanged |= OldMask != InMask;
+		uint64 OldMask = FPipelineFileCache::SetGameUsageMaskWithComparison(InMask, InComparisonFnPtr);
+		bMaskChanged = OldMask != InMask;
 		
-			ShaderPipelineCache->bReady = true;
+		ShaderPipelineCache->bReady = true;
 		
-			if(bMaskChanged)
-			{
-				// Mask has changed and we have an open file refetch PSO's for this Mask - leave the FPipelineFileCache file open - no need to close - just pull out the relevant PSOs.
-				// If this PSO compile run has completed for this Mask in which case don't refetch + compile for that mask
+		if(bMaskChanged && ShaderPipelineCache->bOpened)
+		{
+			// Mask has changed and we have an open file refetch PSO's for this Mask - leave the FPipelineFileCache file open - no need to close - just pull out the relevant PSOs.
+			// If this PSO compile run has completed for this Mask in which case don't refetch + compile for that mask
 			
-				// Don't clear already compiled PSOHash list - this is not a full reset
-				ShaderPipelineCache->Flush(false);
+			// Don't clear already compiled PSOHash list - this is not a full reset
+			ShaderPipelineCache->Flush(false);
 			
-				if(!ShaderPipelineCache->CompletedMasks.Contains(InMask))
-				{
-					int32 Order = (int32)FPipelineFileCache::PSOOrder::Default;
+			if(!ShaderPipelineCache->CompletedMasks.Contains(InMask))
+			{
+				int32 Order = (int32)FPipelineFileCache::PSOOrder::Default;
 				
-					if(!GConfig->GetInt(FShaderPipelineCacheConstants::SectionHeading, FShaderPipelineCacheConstants::SortOrderKey, Order, *GGameUserSettingsIni))
-					{
-						GConfig->GetInt(FShaderPipelineCacheConstants::SectionHeading, FShaderPipelineCacheConstants::SortOrderKey, Order, *GGameIni);
-					}
+				if(!GConfig->GetInt(FShaderPipelineCacheConstants::SectionHeading, FShaderPipelineCacheConstants::SortOrderKey, Order, *GGameUserSettingsIni))
+				{
+					GConfig->GetInt(FShaderPipelineCacheConstants::SectionHeading, FShaderPipelineCacheConstants::SortOrderKey, Order, *GGameIni);
+				}
 				
-					TArray<FPipelineCachePSOHeader> LocalPreFetchedTasks;
-					FPipelineFileCache::GetOrderedPSOHashes(LocalPreFetchedTasks, (FPipelineFileCache::PSOOrder)Order, (int64)CVarPSOFileCacheMinBindCount.GetValueOnAnyThread(), ShaderPipelineCache->CompiledHashes);
-					// Iterate over all the tasks we haven't yet begun to read data for - these are the 'waiting' tasks
-					int64 Count = 0;
-					for (auto const& Task : LocalPreFetchedTasks)
+				TArray<FPipelineCachePSOHeader> LocalPreFetchedTasks;
+				FPipelineFileCache::GetOrderedPSOHashes(LocalPreFetchedTasks, (FPipelineFileCache::PSOOrder)Order, (int64)CVarPSOFileCacheMinBindCount.GetValueOnAnyThread(), ShaderPipelineCache->CompiledHashes);
+				// Iterate over all the tasks we haven't yet begun to read data for - these are the 'waiting' tasks
+				int64 Count = 0;
+				for (auto const& Task : LocalPreFetchedTasks)
+				{
+					bool bHasShaders = true;
+					for (FSHAHash const& Hash : Task.Shaders)
 					{
-						bool bHasShaders = true;
-						for (FSHAHash const& Hash : Task.Shaders)
-						{
-							bHasShaders &= FShaderCodeLibrary::ContainsShaderCode(Hash);
-						}
-						if (bHasShaders)
-						{
-							Count++;
-						}
+						bHasShaders &= FShaderCodeLibrary::ContainsShaderCode(Hash);
 					}
-				
-					FPlatformAtomics::InterlockedAdd(&ShaderPipelineCache->TotalWaitingTasks, Count);
-				
-					if (OnCachedOpened.IsBound())
+					if (bHasShaders)
 					{
-						OnCachedOpened.Broadcast(ShaderPipelineCache->FileName, ShaderPipelineCache->CurrentPlatform, LocalPreFetchedTasks.Num(), ShaderPipelineCache->CacheFileGuid, ShaderPipelineCache->ShaderCachePrecompileContext);
+						Count++;
 					}
+				}
 				
-					ShaderPipelineCache->PreFetchedTasks = LocalPreFetchedTasks;
+				FPlatformAtomics::InterlockedAdd(&ShaderPipelineCache->TotalWaitingTasks, Count);
 				
-					bMaskChanged = false;
-					
-					UE_LOG(LogRHI, Display, TEXT("New ShaderPipelineCache GameUsageMask [%llu=>%llu], Enqueued %d of %d tasks for precompile."), OldMask, InMask, Count, ShaderPipelineCache->PreFetchedTasks.Num());
-					
-					return OldMask != InMask;
-				}
-				else
+				if (OnCachedOpened.IsBound())
 				{
-					UE_LOG(LogRHI, Display, TEXT("New ShaderPipelineCache GameUsageMask [%llu=>%llu], Target mask already precompiled."), OldMask, InMask);
+					OnCachedOpened.Broadcast(ShaderPipelineCache->FileName, ShaderPipelineCache->CurrentPlatform, LocalPreFetchedTasks.Num(), ShaderPipelineCache->CacheFileGuid, ShaderPipelineCache->ShaderCachePrecompileContext);
 				}
+				
+				ShaderPipelineCache->PreFetchedTasks = LocalPreFetchedTasks;
+				
+				UE_LOG(LogRHI, Display, TEXT("New ShaderPipelineCache GameUsageMask [%llu=>%llu], Enqueued %d of %d tasks for precompile."), OldMask, InMask, Count, ShaderPipelineCache->PreFetchedTasks.Num());
 			}
 			else
 			{
-				UE_LOG(LogRHI, Display, TEXT("ShaderPipelineCache::SetGameUsageMaskWithComparison failed to set a new mask because the game mask was not different"));
+				UE_LOG(LogRHI, Display, TEXT("New ShaderPipelineCache GameUsageMask [%llu=>%llu], Target mask already precompiled."), OldMask, InMask);
 			}
 		}
-		else
-		{
-			// NOTE: if this is called and then the cache is opened, but this function is never called again, the PSOs will not get precompiled. Should probably update the open call itself to see if there was a new mask set and call the code above
-			uint64 OldMask = FPipelineFileCache::SetGameUsageMaskWithComparison(InMask, InComparisonFnPtr);
-			bMaskChanged |= OldMask != InMask;
-
-			UE_LOG(LogRHI, Display, TEXT("ShaderPipelineCache::SetGameUsageMaskWithComparison set a new mask but did not attempt to setup any tasks because the cache was not open"));
-			return OldMask != InMask;
-		}
-	}
-	else
-	{
-		UE_LOG(LogRHI, Display, TEXT("ShaderPipelineCache::SetGameUsageMaskWithComparison failed to set a new mask because the cache was not open or game mask is not enabled"));
 	}
 	
 	return bMaskChanged;
@@ -1105,7 +1075,6 @@ FShaderPipelineCache::FShaderPipelineCache(EShaderPlatform Platform)
 , bPaused(false)
 , bOpened(false)
 , bReady(false)
-, bPreOptimizing(false)
 , PausedCount(0)
 , TotalActiveTasks(0)
 , TotalWaitingTasks(0)
@@ -1146,7 +1115,7 @@ FShaderPipelineCache::FShaderPipelineCache(EShaderPlatform Platform)
 	
 	FCoreDelegates::ApplicationWillDeactivateDelegate.AddStatic(&PipelineStateCacheOnAppDeactivate);
 	
-	bReady = !CVarPSOFileCacheGameFileMaskEnabled.GetValueOnAnyThread() || CVarPSOFileCachePreOptimizeEnabled.GetValueOnAnyThread();
+	bReady = !CVarPSOFileCacheGameFileMaskEnabled.GetValueOnAnyThread();
 }
 
 FShaderPipelineCache::~FShaderPipelineCache()
@@ -1227,7 +1196,6 @@ void FShaderPipelineCache::Tick( float DeltaTime )
 		
         	FPlatformAtomics::InterlockedExchange(&TotalCompleteTasks, 0);
         	FPlatformAtomics::InterlockedExchange(&TotalPrecompileTime, 0);
-			bPreOptimizing = false;
 		}
     }
 	
@@ -1408,12 +1376,11 @@ TStatId FShaderPipelineCache::GetStatId() const
 // Not sure where the define is for this but most seem to be low, medium, high, epic, cinema, auto, except material quality but that's less anyway
 static const int32 MaxQualityCount = 6;
 static const int32 MaxPlaylistCount = 3;
-static const int32 MaxUserCount = 16;
 
 static bool PreCompileMaskComparison(uint64 ReferenceGameMask, uint64 PSOMask)
 {
 	uint64 UsageMask = (ReferenceGameMask & PSOMask);
-	return (UsageMask & (7l << (MaxQualityCount*3+MaxPlaylistCount))) && (UsageMask & (7 << (MaxQualityCount*3))) && (UsageMask & (63 << MaxQualityCount*2)) && (UsageMask & (63 << MaxQualityCount)) && (UsageMask & 63);
+	return (UsageMask & (7 << (MaxQualityCount*2+MaxPlaylistCount))) && (UsageMask & (63 << MaxQualityCount*2)) && (UsageMask & (63 << MaxQualityCount)) && (UsageMask & 63);
 }
 
 bool FShaderPipelineCache::Open(FString const& Name, EShaderPlatform Platform)
@@ -1472,11 +1439,14 @@ bool FShaderPipelineCache::Open(FString const& Name, EShaderPlatform Platform)
 			PreFetchedTasks = LocalPreFetchedTasks;
 			TotalPrecompileTasks = PreFetchedTasks.Num();
 			
-			bPreOptimizing = TotalPrecompileTasks > 0;
 			UE_LOG(LogRHI, Display, TEXT("Opened pipeline cache and enqueued %d of %d tasks for precompile with BatchSize %d and BatchTime %f."), Count, PreFetchedTasks.Num(), BatchSize, BatchTime);
 		}
 		else
 		{
+			if (CVarPSOFileCacheGameFileMaskEnabled.GetValueOnAnyThread() == 0)
+			{
+				CVarPSOFileCacheGameFileMaskEnabled->Set(1);
+			}
 			UE_LOG(LogRHI, Display, TEXT("Opened pipeline cache - precompile deferred on UsageMask."));
 		}
 	}
diff --git a/Engine/Source/Runtime/RenderCore/Public/ShaderPipelineCache.h b/Engine/Source/Runtime/RenderCore/Public/ShaderPipelineCache.h
--- a/Engine/Source/Runtime/RenderCore/Public/ShaderPipelineCache.h
+++ b/Engine/Source/Runtime/RenderCore/Public/ShaderPipelineCache.h
@@ -238,7 +238,6 @@ private:
 	bool bPaused;
 	bool bOpened;
 	bool bReady;
-	bool bPreOptimizing;
     int32 PausedCount;
 	FShaderCachePrecompileContext ShaderCachePrecompileContext;
 	
