fix for invalid PSO masks when the GameMask is disabled by default and we attempt to gather PSOs

diff --git a/Engine/Source/Runtime/RenderCore/Private/ShaderPipelineCache.cpp b/Engine/Source/Runtime/RenderCore/Private/ShaderPipelineCache.cpp
--- a/Engine/Source/Runtime/RenderCore/Private/ShaderPipelineCache.cpp
+++ b/Engine/Source/Runtime/RenderCore/Private/ShaderPipelineCache.cpp
@@ -1439,6 +1439,10 @@ bool FShaderPipelineCache::Open(FString const& Name, EShaderPlatform Platform)
 			PreFetchedTasks = LocalPreFetchedTasks;
 			TotalPrecompileTasks = PreFetchedTasks.Num();
 			
+			if (CVarPSOFileCacheGameFileMaskEnabled.GetValueOnAnyThread() == 0 && FParse::Param(FCommandLine::Get(), TEXT("logPSO"))
+			{
+				CVarPSOFileCacheGameFileMaskEnabled->Set(1);
+			}
 			UE_LOG(LogRHI, Display, TEXT("Opened pipeline cache and enqueued %d of %d tasks for precompile with BatchSize %d and BatchTime %f."), Count, PreFetchedTasks.Num(), BatchSize, BatchTime);
 		}
 		else
