fix for compile failure

diff --git a/Engine/Source/Runtime/RenderCore/Private/ShaderPipelineCache.cpp b/Engine/Source/Runtime/RenderCore/Private/ShaderPipelineCache.cpp
--- a/Engine/Source/Runtime/RenderCore/Private/ShaderPipelineCache.cpp
+++ b/Engine/Source/Runtime/RenderCore/Private/ShaderPipelineCache.cpp
@@ -1439,7 +1439,7 @@ bool FShaderPipelineCache::Open(FString const& Name, EShaderPlatform Platform)
 			PreFetchedTasks = LocalPreFetchedTasks;
 			TotalPrecompileTasks = PreFetchedTasks.Num();
 			
-			if (CVarPSOFileCacheGameFileMaskEnabled.GetValueOnAnyThread() == 0 && FParse::Param(FCommandLine::Get(), TEXT("logPSO"))
+			if (CVarPSOFileCacheGameFileMaskEnabled.GetValueOnAnyThread() == 0 && FParse::Param(FCommandLine::Get(), TEXT("logPSO")))
 			{
 				CVarPSOFileCacheGameFileMaskEnabled->Set(1);
 			}
