Fixed Buffer flag usage

#rnx

diff --git a/Engine/Source/Runtime/RenderCore/Private/DynamicBufferAllocator.cpp b/Engine/Source/Runtime/RenderCore/Private/DynamicBufferAllocator.cpp
--- a/Engine/Source/Runtime/RenderCore/Private/DynamicBufferAllocator.cpp
+++ b/Engine/Source/Runtime/RenderCore/Private/DynamicBufferAllocator.cpp
@@ -129,7 +129,7 @@ FGlobalDynamicReadBuffer::FAllocation FGlobalDynamicReadBuffer::AllocateFloat(ui
 			const uint32 NewBufferSize = FMath::Max(AlignedNum, (uint32)GMinReadBufferRenderingBufferSize);
 			Buffer = new FDynamicAllocReadBuffer();
 			FloatBufferPool->Buffers.Add(Buffer);
-			Buffer->Initialize(sizeof(float), NewBufferSize, PF_R32_FLOAT, BUF_Volatile);
+			Buffer->Initialize(sizeof(float), NewBufferSize, PF_R32_FLOAT, BUF_Dynamic);
 		}
 
 		// Lock the buffer if needed.
@@ -189,7 +189,7 @@ FGlobalDynamicReadBuffer::FAllocation FGlobalDynamicReadBuffer::AllocateInt32(ui
 			const uint32 NewBufferSize = FMath::Max(AlignedNum, (uint32)GMinReadBufferRenderingBufferSize);
 			Buffer = new FDynamicAllocReadBuffer();
 			Int32BufferPool->Buffers.Add(Buffer);
-			Buffer->Initialize(sizeof(int32), NewBufferSize, PF_R32_SINT, BUF_Volatile);
+			Buffer->Initialize(sizeof(int32), NewBufferSize, PF_R32_SINT, BUF_Dynamic);
 		}
 
 		// Lock the buffer if needed.
@@ -220,48 +220,48 @@ bool FGlobalDynamicReadBuffer::IsRenderAlarmLoggingEnabled() const
 
 void FGlobalDynamicReadBuffer::Commit()
 {
-	for (int32 BufferIndex = 0, NumBuffers = FloatBufferPool->Buffers.Num(); BufferIndex < NumBuffers; ++BufferIndex)
-	{
-		FDynamicAllocReadBuffer& Buffer = FloatBufferPool->Buffers[BufferIndex];
-		if (Buffer.MappedBuffer != NULL)
-		{
-			Buffer.Unlock();
-		}
-		else if (GGlobalBufferNumFramesUnusedThresold && !Buffer.AllocatedByteCount)
+		for (int32 BufferIndex = 0, NumBuffers = FloatBufferPool->Buffers.Num(); BufferIndex < NumBuffers; ++BufferIndex)
 		{
-			++Buffer.NumFramesUnused;
-			if (Buffer.NumFramesUnused >= GGlobalBufferNumFramesUnusedThresold )
+			FDynamicAllocReadBuffer& Buffer = FloatBufferPool->Buffers[BufferIndex];
+			if (Buffer.MappedBuffer != NULL)
 			{
-				// Remove the buffer, assumes they are unordered.
-				Buffer.Release();
-				FloatBufferPool->Buffers.RemoveAtSwap(BufferIndex);
-				--BufferIndex;
-				--NumBuffers;
+				Buffer.Unlock();
+			}
+			else if (GGlobalBufferNumFramesUnusedThresold && !Buffer.AllocatedByteCount)
+			{
+				++Buffer.NumFramesUnused;
+				if (Buffer.NumFramesUnused >= GGlobalBufferNumFramesUnusedThresold)
+				{
+					// Remove the buffer, assumes they are unordered.
+					Buffer.Release();
+					FloatBufferPool->Buffers.RemoveAtSwap(BufferIndex);
+					--BufferIndex;
+					--NumBuffers;
+				}
 			}
 		}
-	}
-	FloatBufferPool->CurrentBuffer = NULL;
+		FloatBufferPool->CurrentBuffer = NULL;
 
-	for (int32 BufferIndex = 0, NumBuffers = Int32BufferPool->Buffers.Num(); BufferIndex < NumBuffers; ++BufferIndex)
-	{
-		FDynamicAllocReadBuffer& Buffer = Int32BufferPool->Buffers[BufferIndex];
-		if (Buffer.MappedBuffer != NULL)
-		{
-			Buffer.Unlock();
-		}
-		else if (GGlobalBufferNumFramesUnusedThresold  && !Buffer.AllocatedByteCount)
+		for (int32 BufferIndex = 0, NumBuffers = Int32BufferPool->Buffers.Num(); BufferIndex < NumBuffers; ++BufferIndex)
 		{
-			++Buffer.NumFramesUnused;
-			if (Buffer.NumFramesUnused >= GGlobalBufferNumFramesUnusedThresold )
+			FDynamicAllocReadBuffer& Buffer = Int32BufferPool->Buffers[BufferIndex];
+			if (Buffer.MappedBuffer != NULL)
 			{
-				// Remove the buffer, assumes they are unordered.
-				Buffer.Release();
-				Int32BufferPool->Buffers.RemoveAtSwap(BufferIndex);
-				--BufferIndex;
-				--NumBuffers;
+				Buffer.Unlock();
+			}
+			else if (GGlobalBufferNumFramesUnusedThresold && !Buffer.AllocatedByteCount)
+			{
+				++Buffer.NumFramesUnused;
+				if (Buffer.NumFramesUnused >= GGlobalBufferNumFramesUnusedThresold)
+				{
+					// Remove the buffer, assumes they are unordered.
+					Buffer.Release();
+					Int32BufferPool->Buffers.RemoveAtSwap(BufferIndex);
+					--BufferIndex;
+					--NumBuffers;
+				}
 			}
 		}
-	}
-	Int32BufferPool->CurrentBuffer = NULL;
+		Int32BufferPool->CurrentBuffer = NULL;
 	TotalAllocatedSinceLastCommit = 0;
 }
