Fix another merge issue where previous change was dropped.

diff --git a/Engine/Source/Runtime/RenderCore/Private/ShaderResource.cpp b/Engine/Source/Runtime/RenderCore/Private/ShaderResource.cpp
--- a/Engine/Source/Runtime/RenderCore/Private/ShaderResource.cpp
+++ b/Engine/Source/Runtime/RenderCore/Private/ShaderResource.cpp
@@ -180,10 +180,18 @@ int32 FShaderMapResourceBuilder::FindOrAddCode(EShaderFrequency InFrequency, con
 	int32 Index = FindCode(InHash, Key);
 	if(Index == INDEX_NONE)
 	{
+		bool bAllowShaderCompression = true;
+
+#if !(UE_BUILD_SHIPPING || UE_BUILD_TEST)
+		static const IConsoleVariable* CVarSkipShaderCompression = IConsoleManager::Get().FindConsoleVariable(TEXT("r.Shaders.SkipCompression"));
+		bAllowShaderCompression = CVarSkipShaderCompression ? CVarSkipShaderCompression->GetInt() == 0 : true;
+#endif
+
 		int32 CompressedSize = InCode.Num();
 		TArray<uint8> CompressedCode;
 		CompressedCode.AddUninitialized(CompressedSize);
-		if(FCompression::CompressMemory(ShaderCompressionFormat, CompressedCode.GetData(), CompressedSize, InCode.GetData(), InCode.Num()))
+
+		if (bAllowShaderCompression && FCompression::CompressMemory(ShaderCompressionFormat, CompressedCode.GetData(), CompressedSize, InCode.GetData(), InCode.Num()))
 		{
 			CompressedCode.SetNum(CompressedSize, false);
 		}
