Fix unity build break.

diff --git a/Engine/Source/Runtime/RenderCore/Private/ShaderResource.cpp b/Engine/Source/Runtime/RenderCore/Private/ShaderResource.cpp
--- a/Engine/Source/Runtime/RenderCore/Private/ShaderResource.cpp
+++ b/Engine/Source/Runtime/RenderCore/Private/ShaderResource.cpp
@@ -183,8 +183,8 @@ int32 FShaderMapResourceBuilder::FindOrAddCode(EShaderFrequency InFrequency, con
 		bool bAllowShaderCompression = true;
 
 #if !(UE_BUILD_SHIPPING || UE_BUILD_TEST)
-		static const IConsoleVariable* CVarSkipShaderCompression = IConsoleManager::Get().FindConsoleVariable(TEXT("r.Shaders.SkipCompression"));
-		bAllowShaderCompression = CVarSkipShaderCompression ? CVarSkipShaderCompression->GetInt() == 0 : true;
+		static const IConsoleVariable* CVarSkipCompression = IConsoleManager::Get().FindConsoleVariable(TEXT("r.Shaders.SkipCompression"));
+		bAllowShaderCompression = CVarSkipCompression ? CVarSkipCompression->GetInt() == 0 : true;
 #endif
 
 		int32 CompressedSize = InCode.Num();
