UE-89171: Clamp GPU profiler events start and end timestamps

diff --git a/Engine/Source/Runtime/RenderCore/Private/ProfilingDebugging/RealtimeGPUProfiler.cpp b/Engine/Source/Runtime/RenderCore/Private/ProfilingDebugging/RealtimeGPUProfiler.cpp
--- a/Engine/Source/Runtime/RenderCore/Private/ProfilingDebugging/RealtimeGPUProfiler.cpp
+++ b/Engine/Source/Runtime/RenderCore/Private/ProfilingDebugging/RealtimeGPUProfiler.cpp
@@ -113,10 +113,10 @@ public:
 
 public:
 	FRealtimeGPUProfilerEvent(FRHIRenderQueryPool& RenderQueryPool)
-		: StartQuery(RenderQueryPool.AllocateQuery())
-		, EndQuery(RenderQueryPool.AllocateQuery())
-		, StartResultMicroseconds(InvalidQueryResult)
+		: StartResultMicroseconds(InvalidQueryResult)
 		, EndResultMicroseconds(InvalidQueryResult)
+		, StartQuery(RenderQueryPool.AllocateQuery())
+		, EndQuery(RenderQueryPool.AllocateQuery())
 #if REALTIME_GPU_PROFILER_EVENT_TRACK_FRAME_NUMBER
 		, FrameNumber(-1)
 #endif
@@ -235,6 +235,8 @@ public:
 		return 2u;
 	}
 
+	uint64 StartResultMicroseconds;
+	uint64 EndResultMicroseconds;
 private:
 	FRHIPooledRenderQuery StartQuery;
 	FRHIPooledRenderQuery EndQuery;
@@ -242,8 +244,6 @@ private:
 	FName Name;
 	STAT(FName StatName;)
 
-	uint64 StartResultMicroseconds;
-	uint64 EndResultMicroseconds;
 
 #if REALTIME_GPU_PROFILER_EVENT_TRACK_FRAME_NUMBER
 	uint32 FrameNumber;
@@ -293,6 +293,22 @@ void TraverseEventTree(
 		FGpuProfilerTrace::EndEvent(GpuProfilerEvents[Root].GetEndResultMicroseconds() + GPUTimestampOffset);
 	}
 }
+
+void SanitizeEventTree(
+	TArray<FRealtimeGPUProfilerEvent, TInlineAllocator<100u>>& GpuProfilerEvents,
+	const TArray<TArray<int32>>& GpuProfilerEventChildrenIndices,
+	int32 Root,
+	uint64& MonotonicTime)
+{
+	GpuProfilerEvents[Root].StartResultMicroseconds = FMath::Max(GpuProfilerEvents[Root].StartResultMicroseconds, MonotonicTime);
+	GpuProfilerEvents[Root].EndResultMicroseconds = FMath::Max(GpuProfilerEvents[Root].EndResultMicroseconds, GpuProfilerEvents[Root].GetStartResultMicroseconds());
+	for (int32 Subroot : GpuProfilerEventChildrenIndices[Root])
+	{
+		SanitizeEventTree(GpuProfilerEvents, GpuProfilerEventChildrenIndices, Subroot, MonotonicTime);
+	}
+	GpuProfilerEvents[Root].EndResultMicroseconds = FMath::Max(GpuProfilerEvents[Root].EndResultMicroseconds, MonotonicTime);
+	MonotonicTime = FMath::Max(MonotonicTime, GpuProfilerEvents[Root].GetEndResultMicroseconds());
+}
 #endif
 
 /*-----------------------------------------------------------------------------
@@ -494,6 +510,9 @@ public:
 		FGPUTimingCalibrationTimestamp CalibrationTimestamp = FGPUTiming::GetCalibrationTimestamp();
 		const uint64 GPUTimeOffset = CalibrationTimestamp.CPUMicroseconds - CalibrationTimestamp.GPUMicroseconds;
 
+		uint64 MonotonicTime = 0;
+		SanitizeEventTree(GpuProfilerEvents, GpuProfilerEventChildrenIndices, 0, MonotonicTime);
+
 		FGpuProfilerTrace::BeginFrame();
 		TraverseEventTree(GpuProfilerEvents, GpuProfilerEventChildrenIndices, 0, GPUTimeOffset);
 		FGpuProfilerTrace::EndFrame();
