Add missing part of UnrealInsights GPU track improvements (from Dev-Rendering)

* Take Dev-Rendering if any conflict happens

diff --git a/Engine/Source/Runtime/D3D12RHI/Private/D3D12CommandContext.h b/Engine/Source/Runtime/D3D12RHI/Private/D3D12CommandContext.h
--- a/Engine/Source/Runtime/D3D12RHI/Private/D3D12CommandContext.h
+++ b/Engine/Source/Runtime/D3D12RHI/Private/D3D12CommandContext.h
@@ -259,6 +259,7 @@ public:
 	virtual void RHIWriteGPUFence(FRHIGPUFence* Fence) final override;
 	virtual void RHIBeginRenderQuery(FRHIRenderQuery* RenderQuery) final override;
 	virtual void RHIEndRenderQuery(FRHIRenderQuery* RenderQuery) final override;
+	virtual void RHICalibrateTimers(FRHITimestampCalibrationQuery* CalibrationQuery) final override;
 	void RHIBeginOcclusionQueryBatch(uint32 NumQueriesInBatch);
 	void RHIEndOcclusionQueryBatch();
 	virtual void RHIBeginScene() final override;
@@ -562,6 +563,10 @@ public:
 	{
 		ContextRedirect(RHIEndRenderQuery(RenderQuery));
 	}
+	FORCEINLINE virtual void RHICalibrateTimers(FRHITimestampCalibrationQuery* CalibrationQuery) final override
+	{
+		ContextRedirect(RHICalibrateTimers(CalibrationQuery));
+	}
 	FORCEINLINE virtual void RHIBeginScene() final override
 	{
 		ContextRedirect(RHIBeginScene());
diff --git a/Engine/Source/Runtime/D3D12RHI/Private/D3D12Commands.cpp b/Engine/Source/Runtime/D3D12RHI/Private/D3D12Commands.cpp
--- a/Engine/Source/Runtime/D3D12RHI/Private/D3D12Commands.cpp
+++ b/Engine/Source/Runtime/D3D12RHI/Private/D3D12Commands.cpp
@@ -1131,6 +1131,13 @@ void FD3D12CommandContext::RHIEndRenderQuery(FRHIRenderQuery* QueryRHI)
 #endif
 }
 
+void FD3D12CommandContext::RHICalibrateTimers(FRHITimestampCalibrationQuery* CalibrationQuery)
+{
+	FGPUTimingCalibrationTimestamp Timestamp = GetParentDevice()->GetCommandListManager().GetCalibrationTimestamp();
+	CalibrationQuery->CPUMicroseconds = Timestamp.CPUMicroseconds;
+	CalibrationQuery->GPUMicroseconds = Timestamp.GPUMicroseconds;
+}
+
 // Primitive drawing.
 
 void FD3D12CommandContext::CommitNonComputeShaderConstants()
diff --git a/Engine/Source/Runtime/D3D12RHI/Private/D3D12DirectCommandListManager.cpp b/Engine/Source/Runtime/D3D12RHI/Private/D3D12DirectCommandListManager.cpp
--- a/Engine/Source/Runtime/D3D12RHI/Private/D3D12DirectCommandListManager.cpp
+++ b/Engine/Source/Runtime/D3D12RHI/Private/D3D12DirectCommandListManager.cpp
@@ -385,6 +385,8 @@ void FD3D12CommandListManager::Create(const TCHAR* Name, uint32 NumCommandLists,
 
 FGPUTimingCalibrationTimestamp FD3D12CommandListManager::GetCalibrationTimestamp()
 {
+	TRACE_CPUPROFILER_EVENT_SCOPE(D3D12GetCalibrationTimestamp);
+
 	check(CommandListType == D3D12_COMMAND_LIST_TYPE_DIRECT || CommandListType == D3D12_COMMAND_LIST_TYPE_COMPUTE);
 
 	uint64 GPUTimestampFrequency;
diff --git a/Engine/Source/Runtime/RHI/Private/GpuProfilerTrace.cpp b/Engine/Source/Runtime/RHI/Private/GpuProfilerTrace.cpp
--- a/Engine/Source/Runtime/RHI/Private/GpuProfilerTrace.cpp
+++ b/Engine/Source/Runtime/RHI/Private/GpuProfilerTrace.cpp
@@ -38,7 +38,7 @@ UE_TRACE_EVENT_BEGIN(GpuProfiler, EventSpec, Important)
 UE_TRACE_EVENT_END()
 
 UE_TRACE_EVENT_BEGIN(GpuProfiler, Frame)
-	UE_TRACE_EVENT_FIELD(int64, CalibrationBias)
+	UE_TRACE_EVENT_FIELD(uint64, CalibrationBias)
 	UE_TRACE_EVENT_FIELD(uint64, TimestampBase)
 	UE_TRACE_EVENT_FIELD(uint32, RenderingFrameNumber)
 	UE_TRACE_EVENT_FIELD(uint8[], Data)
@@ -46,7 +46,7 @@ UE_TRACE_EVENT_END()
 
 } // namespace GpuProfilerTrace
 
-void FGpuProfilerTrace::BeginFrame()
+void FGpuProfilerTrace::BeginFrame(FGPUTimingCalibrationTimestamp& Calibration)
 {
 	using namespace GpuProfilerTrace;
 
@@ -55,14 +55,8 @@ void FGpuProfilerTrace::BeginFrame()
 		return;
 	}
 
-	GCurrentFrame.Calibration = FGPUTiming::GetCalibrationTimestamp();
-	if (GCurrentFrame.Calibration.CPUMicroseconds == 0)
-	{
-		uint64 Cycles = FPlatformTime::Cycles64();
-		Cycles = uint64(FPlatformTime::ToSeconds64(Cycles) * 1000 * 1000);
-		GCurrentFrame.Calibration.CPUMicroseconds = Cycles;
-	}
-
+	GCurrentFrame.Calibration = Calibration;
+	ensure(GCurrentFrame.Calibration.CPUMicroseconds > 0 && GCurrentFrame.Calibration.GPUMicroseconds > 0);
 	GCurrentFrame.TimestampBase = 0;
 	GCurrentFrame.EventBufferSize = 0;
 	GCurrentFrame.bActive = true;
@@ -148,6 +142,7 @@ void FGpuProfilerTrace::EndFrame()
 
 	if (GCurrentFrame.EventBufferSize)
 	{
+		// This subtraction is intended to be performed on uint64 to leverage the wrap around behavior defined by the standard
 		uint64 Bias = GCurrentFrame.Calibration.CPUMicroseconds - GCurrentFrame.Calibration.GPUMicroseconds;
 		UE_TRACE_LOG(GpuProfiler, Frame, GpuChannel)
 			<< Frame.CalibrationBias(Bias)
diff --git a/Engine/Source/Runtime/RHI/Public/GpuProfilerTrace.h b/Engine/Source/Runtime/RHI/Public/GpuProfilerTrace.h
--- a/Engine/Source/Runtime/RHI/Public/GpuProfilerTrace.h
+++ b/Engine/Source/Runtime/RHI/Public/GpuProfilerTrace.h
@@ -19,7 +19,7 @@ class FName;
 
 struct FGpuProfilerTrace
 {
-	RHI_API static void BeginFrame();
+	RHI_API static void BeginFrame(struct FGPUTimingCalibrationTimestamp& Calibration);
 	RHI_API static void SpecifyEventByName(const FName& Name);
 	RHI_API static void BeginEventByName(const FName& Name, uint32 FrameNumber, uint64 TimestampMicroseconds);
 	RHI_API static void EndEvent(uint64 TimestampMicroseconds);
diff --git a/Engine/Source/Runtime/RHI/Public/RHICommandList.h b/Engine/Source/Runtime/RHI/Public/RHICommandList.h
--- a/Engine/Source/Runtime/RHI/Public/RHICommandList.h
+++ b/Engine/Source/Runtime/RHI/Public/RHICommandList.h
@@ -1845,6 +1845,17 @@ FRHICOMMAND_MACRO(FRHICommandEndRenderQuery)
 	RHI_API void Execute(FRHICommandListBase& CmdList);
 };
 
+FRHICOMMAND_MACRO(FRHICommandCalibrateTimers)
+{
+	FRHITimestampCalibrationQuery* CalibrationQuery;
+
+	FORCEINLINE_DEBUGGABLE FRHICommandCalibrateTimers(FRHITimestampCalibrationQuery * CalibrationQuery)
+		: CalibrationQuery(CalibrationQuery)
+	{
+	}
+	RHI_API void Execute(FRHICommandListBase & CmdList);
+};
+
 struct FRHICommandSubmitCommandsHintString
 {
 	static const TCHAR* TStr() { return TEXT("FRHICommandSubmitCommandsHint"); }
@@ -3308,6 +3319,15 @@ public:
 		}
 		ALLOC_COMMAND(FRHICommandEndRenderQuery)(RenderQuery);
 	}
+	FORCEINLINE_DEBUGGABLE void CalibrateTimers(FRHITimestampCalibrationQuery* CalibrationQuery)
+	{
+		if (Bypass())
+		{
+			GetContext().RHICalibrateTimers(CalibrationQuery);
+			return;
+		}
+		ALLOC_COMMAND(FRHICommandCalibrateTimers)(CalibrationQuery);
+	}
 
 	FORCEINLINE_DEBUGGABLE void PollOcclusionQueries()
 	{
diff --git a/Engine/Source/Runtime/RHI/Public/RHICommandListCommandExecutes.inl b/Engine/Source/Runtime/RHI/Public/RHICommandListCommandExecutes.inl
--- a/Engine/Source/Runtime/RHI/Public/RHICommandListCommandExecutes.inl
+++ b/Engine/Source/Runtime/RHI/Public/RHICommandListCommandExecutes.inl
@@ -498,6 +498,12 @@ void FRHICommandEndRenderQuery::Execute(FRHICommandListBase& CmdList)
 	INTERNAL_DECORATOR(RHIEndRenderQuery)(RenderQuery);
 }
 
+void FRHICommandCalibrateTimers::Execute(FRHICommandListBase& CmdList)
+{
+	RHISTAT(CalibrateTimers);
+	INTERNAL_DECORATOR(RHICalibrateTimers)(CalibrationQuery);
+}
+
 void FRHICommandSubmitCommandsHint::Execute(FRHICommandListBase& CmdList)
 {
 	RHISTAT(SubmitCommandsHint);
diff --git a/Engine/Source/Runtime/RHI/Public/RHIContext.h b/Engine/Source/Runtime/RHI/Public/RHIContext.h
--- a/Engine/Source/Runtime/RHI/Public/RHIContext.h
+++ b/Engine/Source/Runtime/RHI/Public/RHIContext.h
@@ -411,6 +411,11 @@ public:
 
 	virtual void RHIEndRenderQuery(FRHIRenderQuery* RenderQuery) = 0;
 
+	virtual void RHICalibrateTimers(FRHITimestampCalibrationQuery* CalibrationQuery)
+	{
+		/* empty default implementation */
+	}
+
 	// Used for OpenGL to check and see if any occlusion queries can be read back on the RHI thread. If they aren't ready when we need them, then we end up stalling.
 	virtual void RHIPollOcclusionQueries()
 	{
diff --git a/Engine/Source/Runtime/RHI/Public/RHIResources.h b/Engine/Source/Runtime/RHI/Public/RHIResources.h
--- a/Engine/Source/Runtime/RHI/Public/RHIResources.h
+++ b/Engine/Source/Runtime/RHI/Public/RHIResources.h
@@ -968,6 +968,12 @@ public:
 // Misc
 //
 
+class RHI_API FRHITimestampCalibrationQuery : public FRHIResource
+{
+public:
+	uint64 GPUMicroseconds = 0;
+	uint64 CPUMicroseconds = 0;
+};
 
 /*
 * Generic GPU fence class.
@@ -1209,6 +1215,7 @@ typedef TRefCountPtr<FRHITextureCube> FTextureCubeRHIRef;
 typedef TRefCountPtr<FRHITextureReference> FTextureReferenceRHIRef;
 typedef TRefCountPtr<FRHIRenderQuery> FRenderQueryRHIRef;
 typedef TRefCountPtr<FRHIRenderQueryPool> FRenderQueryPoolRHIRef;
+typedef TRefCountPtr<FRHITimestampCalibrationQuery> FTimestampCalibrationQueryRHIRef;
 typedef TRefCountPtr<FRHIGPUFence>	FGPUFenceRHIRef;
 typedef TRefCountPtr<FRHIViewport> FViewportRHIRef;
 typedef TRefCountPtr<FRHIUnorderedAccessView> FUnorderedAccessViewRHIRef;
diff --git a/Engine/Source/Runtime/RHI/Public/RHIValidationContext.h b/Engine/Source/Runtime/RHI/Public/RHIValidationContext.h
--- a/Engine/Source/Runtime/RHI/Public/RHIValidationContext.h
+++ b/Engine/Source/Runtime/RHI/Public/RHIValidationContext.h
@@ -302,6 +302,11 @@ public:
 		RHIContext->RHIEndRenderQuery(RenderQuery);
 	}
 
+	virtual void RHICalibrateTimers(FRHITimestampCalibrationQuery* CalibrationQuery) override final
+	{
+		RHIContext->RHICalibrateTimers(CalibrationQuery);
+	}
+
 	virtual void RHISubmitCommandsHint() override final
 	{
 		ensureMsgf(!State.bInsideBeginRenderPass, TEXT("Submitting inside a RenderPass is not efficient!"));
diff --git a/Engine/Source/Runtime/RenderCore/Private/ProfilingDebugging/RealtimeGPUProfiler.cpp b/Engine/Source/Runtime/RenderCore/Private/ProfilingDebugging/RealtimeGPUProfiler.cpp
--- a/Engine/Source/Runtime/RenderCore/Private/ProfilingDebugging/RealtimeGPUProfiler.cpp
+++ b/Engine/Source/Runtime/RenderCore/Private/ProfilingDebugging/RealtimeGPUProfiler.cpp
@@ -237,6 +237,7 @@ public:
 
 	uint64 StartResultMicroseconds;
 	uint64 EndResultMicroseconds;
+
 private:
 	FRHIPooledRenderQuery StartQuery;
 	FRHIPooledRenderQuery EndQuery;
@@ -244,7 +245,6 @@ private:
 	FName Name;
 	STAT(FName StatName;)
 
-
 #if REALTIME_GPU_PROFILER_EVENT_TRACK_FRAME_NUMBER
 	uint32 FrameNumber;
 #endif
@@ -258,8 +258,7 @@ private:
 void TraverseEventTree(
 	const TArray<FRealtimeGPUProfilerEvent, TInlineAllocator<100u>>& GpuProfilerEvents,
 	const TArray<TArray<int32>>& GpuProfilerEventChildrenIndices,
-	int32 Root,
-	const uint64 GPUTimestampOffset)
+	int32 Root)
 {
 	uint64 lastStartTime = 0;
 	uint64 lastEndTime = 0;
@@ -270,7 +269,7 @@ void TraverseEventTree(
 	if (Root != 0)
 	{
 		FGpuProfilerTrace::SpecifyEventByName(GpuProfilerEvents[Root].GetName());
-		FGpuProfilerTrace::BeginEventByName(GpuProfilerEvents[Root].GetName(), GpuProfilerEvents[Root].GetFrameNumber(), GpuProfilerEvents[Root].GetStartResultMicroseconds() + GPUTimestampOffset);
+		FGpuProfilerTrace::BeginEventByName(GpuProfilerEvents[Root].GetName(), GpuProfilerEvents[Root].GetFrameNumber(), GpuProfilerEvents[Root].GetStartResultMicroseconds());
 	}
 
 	for (int32 Subroot : GpuProfilerEventChildrenIndices[Root])
@@ -284,13 +283,13 @@ void TraverseEventTree(
 			check(lastStartTime >= GpuProfilerEvents[Root].GetStartResultMicroseconds());
 			check(lastEndTime <= GpuProfilerEvents[Root].GetEndResultMicroseconds());
 		}
-		TraverseEventTree(GpuProfilerEvents, GpuProfilerEventChildrenIndices, Subroot, GPUTimestampOffset);
+		TraverseEventTree(GpuProfilerEvents, GpuProfilerEventChildrenIndices, Subroot);
 	}
 
 	if (Root != 0)
 	{
 		FGpuProfilerTrace::SpecifyEventByName(GpuProfilerEvents[Root].GetName());
-		FGpuProfilerTrace::EndEvent(GpuProfilerEvents[Root].GetEndResultMicroseconds() + GPUTimestampOffset);
+		FGpuProfilerTrace::EndEvent(GpuProfilerEvents[Root].GetEndResultMicroseconds());
 	}
 }
 
@@ -302,11 +301,14 @@ void SanitizeEventTree(
 {
 	GpuProfilerEvents[Root].StartResultMicroseconds = FMath::Max(GpuProfilerEvents[Root].StartResultMicroseconds, MonotonicTime);
 	GpuProfilerEvents[Root].EndResultMicroseconds = FMath::Max(GpuProfilerEvents[Root].EndResultMicroseconds, GpuProfilerEvents[Root].GetStartResultMicroseconds());
+
 	for (int32 Subroot : GpuProfilerEventChildrenIndices[Root])
 	{
 		SanitizeEventTree(GpuProfilerEvents, GpuProfilerEventChildrenIndices, Subroot, MonotonicTime);
 	}
+
 	GpuProfilerEvents[Root].EndResultMicroseconds = FMath::Max(GpuProfilerEvents[Root].EndResultMicroseconds, MonotonicTime);
+
 	MonotonicTime = FMath::Max(MonotonicTime, GpuProfilerEvents[Root].GetEndResultMicroseconds());
 }
 #endif
@@ -507,20 +509,43 @@ public:
 			GpuProfilerEventChildrenIndices[ParentIdx].Add(EventIdx);
 		}
 
-		FGPUTimingCalibrationTimestamp CalibrationTimestamp = FGPUTiming::GetCalibrationTimestamp();
-		const uint64 GPUTimeOffset = CalibrationTimestamp.CPUMicroseconds - CalibrationTimestamp.GPUMicroseconds;
+		FGPUTimingCalibrationTimestamp Timestamp { 0, 0 };
+
+		if (TimestampCalibrationQuery.IsValid())
+		{
+			Timestamp.GPUMicroseconds = TimestampCalibrationQuery->GPUMicroseconds;
+			Timestamp.CPUMicroseconds = TimestampCalibrationQuery->CPUMicroseconds;
+		}
+		
+		if (Timestamp.GPUMicroseconds == 0 || Timestamp.CPUMicroseconds == 0) // Unimplemented platforms, or invalid on the first frame
+		{
+			if (GpuProfilerEvents.Num() > 1)
+			{
+				// Align CPU and GPU frames
+				Timestamp.GPUMicroseconds = GpuProfilerEvents[1].GetStartResultMicroseconds();
+				Timestamp.CPUMicroseconds = FPlatformTime::ToSeconds64(CPUFrameStartTimestamp) * 1000 * 1000;
+			}
+			else
+			{
+				// Fallback to legacy
+				Timestamp = FGPUTiming::GetCalibrationTimestamp();
+			}
+		}
 
 		uint64 MonotonicTime = 0;
 		SanitizeEventTree(GpuProfilerEvents, GpuProfilerEventChildrenIndices, 0, MonotonicTime);
 
-		FGpuProfilerTrace::BeginFrame();
-		TraverseEventTree(GpuProfilerEvents, GpuProfilerEventChildrenIndices, 0, GPUTimeOffset);
+		FGpuProfilerTrace::BeginFrame(Timestamp);
+		TraverseEventTree(GpuProfilerEvents, GpuProfilerEventChildrenIndices, 0);
 		FGpuProfilerTrace::EndFrame();
 #endif
 
 		return true;
 	}
 
+	uint64 CPUFrameStartTimestamp;
+	FTimestampCalibrationQueryRHIRef TimestampCalibrationQuery;
+
 private:
 	struct FGPUEventTimeAggregate
 	{
@@ -708,6 +733,10 @@ void FRealtimeGPUProfiler::BeginFrame(FRHICommandListImmediate& RHICmdList)
 {
 	check(bInBeginEndBlock == false);
 	bInBeginEndBlock = true;
+	
+	Frames[WriteBufferIndex]->TimestampCalibrationQuery = new FRHITimestampCalibrationQuery();
+	RHICmdList.CalibrateTimers(Frames[WriteBufferIndex]->TimestampCalibrationQuery);
+	Frames[WriteBufferIndex]->CPUFrameStartTimestamp = FPlatformTime::Cycles64();
 }
 
 bool AreGPUStatsEnabled()
