put back in the prioritized PSO caching system

diff --git a/Engine/Source/Runtime/RenderCore/Private/ShaderPipelineCache.cpp b/Engine/Source/Runtime/RenderCore/Private/ShaderPipelineCache.cpp
--- a/Engine/Source/Runtime/RenderCore/Private/ShaderPipelineCache.cpp
+++ b/Engine/Source/Runtime/RenderCore/Private/ShaderPipelineCache.cpp
@@ -138,6 +138,12 @@ static TAutoConsoleVariable<int32> CVarPSOFileCacheGameFileMaskEnabled(
 																TEXT("Set non zero to use GameFileMask during PSO precompile - recording should always save out the usage masks to make that data availble when needed."),
 																ECVF_ReadOnly | ECVF_RenderThreadSafe);
 
+static TAutoConsoleVariable<int32> CVarPSOFileCachePreOptimizeEnabled(
+																TEXT("r.ShaderPipelineCache.PreOptimizeEnabled"),
+																(int32)0,
+																TEXT("Set non zero to PreOptimize PSOs - this allows some PSOs to be compiled in the foreground before going in to game"),
+																ECVF_ReadOnly | ECVF_RenderThreadSafe);
+
 
 static TAutoConsoleVariable<int32> CVarPSOFileCacheMinBindCount(
                                                                 TEXT("r.ShaderPipelineCache.MinBindCount"),
@@ -330,67 +336,91 @@ int32 FShaderPipelineCache::GetGameVersionForPSOFileCache()
 
 bool FShaderPipelineCache::SetGameUsageMaskWithComparison(uint64 InMask, FPSOMaskComparisonFn InComparisonFnPtr)
 {
-	bool bMaskChanged = false;
-	
-	if (ShaderPipelineCache != nullptr && CVarPSOFileCacheGameFileMaskEnabled.GetValueOnAnyThread())
+	static bool bMaskChanged = false;
+
+	if (ShaderPipelineCache != nullptr && CVarPSOFileCacheGameFileMaskEnabled.GetValueOnAnyThread() && !ShaderPipelineCache->bPreOptimizing)
 	{
 		FScopeLock Lock(&ShaderPipelineCache->Mutex);
 		
-		uint64 OldMask = FPipelineFileCache::SetGameUsageMaskWithComparison(InMask, InComparisonFnPtr);
-		bMaskChanged = OldMask != InMask;
+		if (ShaderPipelineCache->bOpened)
+		{
+			uint64 OldMask = FPipelineFileCache::SetGameUsageMaskWithComparison(InMask, InComparisonFnPtr);
+			bMaskChanged |= OldMask != InMask;
 		
-		ShaderPipelineCache->bReady = true;
+			ShaderPipelineCache->bReady = true;
 		
-		if(bMaskChanged && ShaderPipelineCache->bOpened)
-		{
-			// Mask has changed and we have an open file refetch PSO's for this Mask - leave the FPipelineFileCache file open - no need to close - just pull out the relevant PSOs.
-			// If this PSO compile run has completed for this Mask in which case don't refetch + compile for that mask
+			if(bMaskChanged)
+			{
+				// Mask has changed and we have an open file refetch PSO's for this Mask - leave the FPipelineFileCache file open - no need to close - just pull out the relevant PSOs.
+				// If this PSO compile run has completed for this Mask in which case don't refetch + compile for that mask
 			
-			// Don't clear already compiled PSOHash list - this is not a full reset
-			ShaderPipelineCache->Flush(false);
+				// Don't clear already compiled PSOHash list - this is not a full reset
+				ShaderPipelineCache->Flush(false);
 			
-			if(!ShaderPipelineCache->CompletedMasks.Contains(InMask))
-			{
-				int32 Order = (int32)FPipelineFileCache::PSOOrder::Default;
-				
-				if(!GConfig->GetInt(FShaderPipelineCacheConstants::SectionHeading, FShaderPipelineCacheConstants::SortOrderKey, Order, *GGameUserSettingsIni))
+				if(!ShaderPipelineCache->CompletedMasks.Contains(InMask))
 				{
-					GConfig->GetInt(FShaderPipelineCacheConstants::SectionHeading, FShaderPipelineCacheConstants::SortOrderKey, Order, *GGameIni);
-				}
+					int32 Order = (int32)FPipelineFileCache::PSOOrder::Default;
 				
-				TArray<FPipelineCachePSOHeader> LocalPreFetchedTasks;
-				FPipelineFileCache::GetOrderedPSOHashes(LocalPreFetchedTasks, (FPipelineFileCache::PSOOrder)Order, (int64)CVarPSOFileCacheMinBindCount.GetValueOnAnyThread(), ShaderPipelineCache->CompiledHashes);
-				// Iterate over all the tasks we haven't yet begun to read data for - these are the 'waiting' tasks
-				int64 Count = 0;
-				for (auto const& Task : LocalPreFetchedTasks)
-				{
-					bool bHasShaders = true;
-					for (FSHAHash const& Hash : Task.Shaders)
+					if(!GConfig->GetInt(FShaderPipelineCacheConstants::SectionHeading, FShaderPipelineCacheConstants::SortOrderKey, Order, *GGameUserSettingsIni))
 					{
-						bHasShaders &= FShaderCodeLibrary::ContainsShaderCode(Hash);
+						GConfig->GetInt(FShaderPipelineCacheConstants::SectionHeading, FShaderPipelineCacheConstants::SortOrderKey, Order, *GGameIni);
 					}
-					if (bHasShaders)
+				
+					TArray<FPipelineCachePSOHeader> LocalPreFetchedTasks;
+					FPipelineFileCache::GetOrderedPSOHashes(LocalPreFetchedTasks, (FPipelineFileCache::PSOOrder)Order, (int64)CVarPSOFileCacheMinBindCount.GetValueOnAnyThread(), ShaderPipelineCache->CompiledHashes);
+					// Iterate over all the tasks we haven't yet begun to read data for - these are the 'waiting' tasks
+					int64 Count = 0;
+					for (auto const& Task : LocalPreFetchedTasks)
 					{
-						Count++;
+						bool bHasShaders = true;
+						for (FSHAHash const& Hash : Task.Shaders)
+						{
+							bHasShaders &= FShaderCodeLibrary::ContainsShaderCode(Hash);
+						}
+						if (bHasShaders)
+						{
+							Count++;
+						}
 					}
-				}
 				
-				FPlatformAtomics::InterlockedAdd(&ShaderPipelineCache->TotalWaitingTasks, Count);
+					FPlatformAtomics::InterlockedAdd(&ShaderPipelineCache->TotalWaitingTasks, Count);
 				
-				if (OnCachedOpened.IsBound())
-				{
-					OnCachedOpened.Broadcast(ShaderPipelineCache->FileName, ShaderPipelineCache->CurrentPlatform, LocalPreFetchedTasks.Num(), ShaderPipelineCache->CacheFileGuid, ShaderPipelineCache->ShaderCachePrecompileContext);
-				}
+					if (OnCachedOpened.IsBound())
+					{
+						OnCachedOpened.Broadcast(ShaderPipelineCache->FileName, ShaderPipelineCache->CurrentPlatform, LocalPreFetchedTasks.Num(), ShaderPipelineCache->CacheFileGuid, ShaderPipelineCache->ShaderCachePrecompileContext);
+					}
 				
-				ShaderPipelineCache->PreFetchedTasks = LocalPreFetchedTasks;
+					ShaderPipelineCache->PreFetchedTasks = LocalPreFetchedTasks;
 				
-				UE_LOG(LogRHI, Display, TEXT("New ShaderPipelineCache GameUsageMask [%llu=>%llu], Enqueued %d of %d tasks for precompile."), OldMask, InMask, Count, ShaderPipelineCache->PreFetchedTasks.Num());
+					bMaskChanged = false;
+					
+					UE_LOG(LogRHI, Display, TEXT("New ShaderPipelineCache GameUsageMask [%llu=>%llu], Enqueued %d of %d tasks for precompile."), OldMask, InMask, Count, ShaderPipelineCache->PreFetchedTasks.Num());
+					
+					return OldMask != InMask;
+				}
+				else
+				{
+					UE_LOG(LogRHI, Display, TEXT("New ShaderPipelineCache GameUsageMask [%llu=>%llu], Target mask already precompiled."), OldMask, InMask);
+				}
 			}
 			else
 			{
-				UE_LOG(LogRHI, Display, TEXT("New ShaderPipelineCache GameUsageMask [%llu=>%llu], Target mask already precompiled."), OldMask, InMask);
+				UE_LOG(LogRHI, Display, TEXT("ShaderPipelineCache::SetGameUsageMaskWithComparison failed to set a new mask because the game mask was not different"));
 			}
 		}
+		else
+		{
+			// NOTE: if this is called and then the cache is opened, but this function is never called again, the PSOs will not get precompiled. Should probably update the open call itself to see if there was a new mask set and call the code above
+			uint64 OldMask = FPipelineFileCache::SetGameUsageMaskWithComparison(InMask, InComparisonFnPtr);
+			bMaskChanged |= OldMask != InMask;
+
+			UE_LOG(LogRHI, Display, TEXT("ShaderPipelineCache::SetGameUsageMaskWithComparison set a new mask but did not attempt to setup any tasks because the cache was not open"));
+			return OldMask != InMask;
+		}
+	}
+	else
+	{
+		UE_LOG(LogRHI, Display, TEXT("ShaderPipelineCache::SetGameUsageMaskWithComparison failed to set a new mask because the cache was not open or game mask is not enabled"));
 	}
 	
 	return bMaskChanged;
@@ -636,7 +666,7 @@ bool FShaderPipelineCache::Precompile(FRHICommandListImmediate& RHICmdList, ESha
 			
 			auto DepthState = GetOrCreateDepthStencilState(PSO.GraphicsDesc.DepthStencilState);
 			GraphicsInitializer.DepthStencilState = DepthState;
-			
+
 			for (uint32 i = 0; i < MaxSimultaneousRenderTargets; ++i)
 			{
 				GraphicsInitializer.RenderTargetFormats[i] = PSO.GraphicsDesc.RenderTargetFormats[i];
@@ -1075,6 +1105,7 @@ FShaderPipelineCache::FShaderPipelineCache(EShaderPlatform Platform)
 , bPaused(false)
 , bOpened(false)
 , bReady(false)
+, bPreOptimizing(false)
 , PausedCount(0)
 , TotalActiveTasks(0)
 , TotalWaitingTasks(0)
@@ -1115,7 +1146,7 @@ FShaderPipelineCache::FShaderPipelineCache(EShaderPlatform Platform)
 	
 	FCoreDelegates::ApplicationWillDeactivateDelegate.AddStatic(&PipelineStateCacheOnAppDeactivate);
 	
-	bReady = !CVarPSOFileCacheGameFileMaskEnabled.GetValueOnAnyThread();
+	bReady = !CVarPSOFileCacheGameFileMaskEnabled.GetValueOnAnyThread() || CVarPSOFileCachePreOptimizeEnabled.GetValueOnAnyThread();
 }
 
 FShaderPipelineCache::~FShaderPipelineCache()
@@ -1196,6 +1227,7 @@ void FShaderPipelineCache::Tick( float DeltaTime )
 		
         	FPlatformAtomics::InterlockedExchange(&TotalCompleteTasks, 0);
         	FPlatformAtomics::InterlockedExchange(&TotalPrecompileTime, 0);
+			bPreOptimizing = false;
 		}
     }
 	
@@ -1376,11 +1408,12 @@ TStatId FShaderPipelineCache::GetStatId() const
 // Not sure where the define is for this but most seem to be low, medium, high, epic, cinema, auto, except material quality but that's less anyway
 static const int32 MaxQualityCount = 6;
 static const int32 MaxPlaylistCount = 3;
+static const int32 MaxUserCount = 16;
 
 static bool PreCompileMaskComparison(uint64 ReferenceGameMask, uint64 PSOMask)
 {
 	uint64 UsageMask = (ReferenceGameMask & PSOMask);
-	return (UsageMask & (7 << (MaxQualityCount*2+MaxPlaylistCount))) && (UsageMask & (63 << MaxQualityCount*2)) && (UsageMask & (63 << MaxQualityCount)) && (UsageMask & 63);
+	return (UsageMask & (7l << (MaxQualityCount*3+MaxPlaylistCount))) && (UsageMask & (7 << (MaxQualityCount*3))) && (UsageMask & (63 << MaxQualityCount*2)) && (UsageMask & (63 << MaxQualityCount)) && (UsageMask & 63);
 }
 
 bool FShaderPipelineCache::Open(FString const& Name, EShaderPlatform Platform)
@@ -1439,18 +1472,11 @@ bool FShaderPipelineCache::Open(FString const& Name, EShaderPlatform Platform)
 			PreFetchedTasks = LocalPreFetchedTasks;
 			TotalPrecompileTasks = PreFetchedTasks.Num();
 			
-			if (CVarPSOFileCacheGameFileMaskEnabled.GetValueOnAnyThread() == 0 && FParse::Param(FCommandLine::Get(), TEXT("logPSO")))
-			{
-				CVarPSOFileCacheGameFileMaskEnabled->Set(1);
-			}
+			bPreOptimizing = TotalPrecompileTasks > 0;
 			UE_LOG(LogRHI, Display, TEXT("Opened pipeline cache and enqueued %d of %d tasks for precompile with BatchSize %d and BatchTime %f."), Count, PreFetchedTasks.Num(), BatchSize, BatchTime);
 		}
 		else
 		{
-			if (CVarPSOFileCacheGameFileMaskEnabled.GetValueOnAnyThread() == 0)
-			{
-				CVarPSOFileCacheGameFileMaskEnabled->Set(1);
-			}
 			UE_LOG(LogRHI, Display, TEXT("Opened pipeline cache - precompile deferred on UsageMask."));
 		}
 	}
@@ -1514,7 +1540,6 @@ void FShaderPipelineCache::Close(bool bShuttingDown)
 
 	FPipelineFileCache::ClosePipelineFileCache();
 
-
 	//
 	// Clean up cached RHI resources
 	//
@@ -1529,7 +1554,7 @@ void FShaderPipelineCache::Close(bool bShuttingDown)
 		Pair.Value->Release();
 	}
 	RasterizerStateCache.Empty();
-
+	
 	for (auto Pair : DepthStencilStateCache)
 	{
 		Pair.Value->Release();
@@ -1579,7 +1604,6 @@ void FShaderPipelineCache::OnShaderLibraryStateChanged(ELibraryState State, ESha
 	UE_LOG(LogRHI, Display, TEXT("Opened pipeline cache after state change and enqueued %d of %d tasks for precompile."), Count, OrderedCompileTasks.Num());
 }
 
-
 FRHIBlendState* FShaderPipelineCache::GetOrCreateBlendState(const FBlendStateInitializerRHI& Initializer)
 {
 	FRHIBlendState** Found = BlendStateCache.Find(Initializer);
@@ -1587,9 +1611,9 @@ FRHIBlendState* FShaderPipelineCache::GetOrCreateBlendState(const FBlendStateIni
 	{
 		return *Found;
 	}
-
+	
 	FBlendStateRHIRef NewState = RHICreateBlendState(Initializer);
-
+	
 	// Add an extra reference so we don't have TRefCountPtr in the maps
 	NewState->AddRef();
 	BlendStateCache.Add(Initializer, NewState);
@@ -1603,9 +1627,9 @@ FRHIRasterizerState* FShaderPipelineCache::GetOrCreateRasterizerState(const FRas
 	{
 		return *Found;
 	}
-
+	
 	FRasterizerStateRHIRef NewState = RHICreateRasterizerState(Initializer);
-
+	
 	// Add an extra reference so we don't have TRefCountPtr in the maps
 	NewState->AddRef();
 	RasterizerStateCache.Add(Initializer, NewState);
@@ -1619,9 +1643,9 @@ FRHIDepthStencilState* FShaderPipelineCache::GetOrCreateDepthStencilState(const
 	{
 		return *Found;
 	}
-
+	
 	FDepthStencilStateRHIRef NewState = RHICreateDepthStencilState(Initializer);
-
+	
 	// Add an extra reference so we don't have TRefCountPtr in the maps
 	NewState->AddRef();
 	DepthStencilStateCache.Add(Initializer, NewState);
diff --git a/Engine/Source/Runtime/RenderCore/Public/ShaderPipelineCache.h b/Engine/Source/Runtime/RenderCore/Public/ShaderPipelineCache.h
--- a/Engine/Source/Runtime/RenderCore/Public/ShaderPipelineCache.h
+++ b/Engine/Source/Runtime/RenderCore/Public/ShaderPipelineCache.h
@@ -226,7 +226,7 @@ private:
 	FRHIBlendState* GetOrCreateBlendState(const FBlendStateInitializerRHI& Initializer);
 	FRHIRasterizerState* GetOrCreateRasterizerState(const FRasterizerStateInitializerRHI& Initializer);
 	FRHIDepthStencilState* GetOrCreateDepthStencilState(const FDepthStencilStateInitializerRHI& Initializer);
-
+	
 private:
 	static FShaderPipelineCache* ShaderPipelineCache;
 	TArray<CompileJob> ReadTasks;
@@ -242,6 +242,7 @@ private:
 	bool bPaused;
 	bool bOpened;
 	bool bReady;
+	bool bPreOptimizing;
     int32 PausedCount;
 	FShaderCachePrecompileContext ShaderCachePrecompileContext;
 	
