Fix for missing GPU Stats enabled tests in RealtimeGPUProfiler

diff --git a/Engine/Source/Runtime/RenderCore/Private/ProfilingDebugging/RealtimeGPUProfiler.cpp b/Engine/Source/Runtime/RenderCore/Private/ProfilingDebugging/RealtimeGPUProfiler.cpp
--- a/Engine/Source/Runtime/RenderCore/Private/ProfilingDebugging/RealtimeGPUProfiler.cpp
+++ b/Engine/Source/Runtime/RenderCore/Private/ProfilingDebugging/RealtimeGPUProfiler.cpp
@@ -731,6 +731,11 @@ void FRealtimeGPUProfiler::Cleanup()
 
 void FRealtimeGPUProfiler::BeginFrame(FRHICommandListImmediate& RHICmdList)
 {
+	if (!AreGPUStatsEnabled())
+	{
+		return;
+	}
+
 	check(bInBeginEndBlock == false);
 	bInBeginEndBlock = true;
 	
@@ -768,10 +773,15 @@ bool AreGPUStatsEnabled()
 
 void FRealtimeGPUProfiler::EndFrame(FRHICommandListImmediate& RHICmdList)
 {
+	if (!AreGPUStatsEnabled())
+	{
+		return;
+	}
+
 	// This is called at the end of the renderthread frame. Note that the RHI thread may still be processing commands for the frame at this point, however
 	// The read buffer index is always 3 frames beind the write buffer index in order to prevent us reading from the frame the RHI thread is still processing. 
 	// This should also ensure the GPU is done with the queries before we try to read them
-	check(!GSupportsTimestampRenderQueries || Frames.Num() > 0);
+	check(Frames.Num() > 0);
 	check(IsInRenderingThread());
 	check(bInBeginEndBlock == true);
 	bInBeginEndBlock = false;
