Don't attempt to add shader virtual paths if shader compilation was disabled on the command line. Avoids some crashes in editor builds that don't have shader source.

diff --git a/Engine/Source/Runtime/RenderCore/Private/ShaderCore.cpp b/Engine/Source/Runtime/RenderCore/Private/ShaderCore.cpp
--- a/Engine/Source/Runtime/RenderCore/Private/ShaderCore.cpp
+++ b/Engine/Source/Runtime/RenderCore/Private/ShaderCore.cpp
@@ -1268,7 +1268,9 @@ void AddShaderSourceDirectoryMapping(const FString& VirtualShaderDirectory, cons
 {
 	check(IsInGameThread());
 
-	if (FPlatformProperties::RequiresCookedData())
+	static const bool bNoShaderCompile = FParse::Param(FCommandLine::Get(), TEXT("NoShaderCompile"));
+
+	if (FPlatformProperties::RequiresCookedData() || bNoShaderCompile)
 	{
 		return;
 	}
