Allow missing shader resources for other shader platforms (UE-89441).

diff --git a/Engine/Source/Runtime/RenderCore/Private/ShaderMap.cpp b/Engine/Source/Runtime/RenderCore/Private/ShaderMap.cpp
--- a/Engine/Source/Runtime/RenderCore/Private/ShaderMap.cpp
+++ b/Engine/Source/Runtime/RenderCore/Private/ShaderMap.cpp
@@ -231,6 +231,10 @@ void FShaderMapBase::Serialize(FArchive& Ar, bool bInlineShaderResources, bool b
 			FSHAHash ResourceHash;
 			Ar << ResourceHash;
 			Resource = FShaderCodeLibrary::LoadResource(ResourceHash, &Ar);
+			if (UNLIKELY(Resource == nullptr && GMaxRHIShaderPlatform == GetShaderPlatform()))
+			{
+				UE_LOG(LogShaders, Fatal, TEXT("Missing shader resource for hash '%s' in the shader library"), *ResourceHash.ToString());
+			}
 		}
 		else
 		{
@@ -238,12 +242,15 @@ void FShaderMapBase::Serialize(FArchive& Ar, bool bInlineShaderResources, bool b
 			Code->Serialize(Ar);
 			Resource = new FShaderMapResource_InlineCode(GetShaderPlatform(), Code);
 		}
-		
-		BeginInitResource(Resource);
 
-		INC_DWORD_STAT_BY(STAT_Shaders_ShaderResourceMemory, Resource->GetSizeBytes());
-		INC_DWORD_STAT_BY(STAT_Shaders_ShaderMemory, FrozenContentSize);
-		INC_DWORD_STAT_BY(STAT_Shaders_NumShadersLoaded, NumFrozenShaders);
+		if (LIKELY(Resource))
+		{
+			BeginInitResource(Resource);
+
+			INC_DWORD_STAT_BY(STAT_Shaders_ShaderResourceMemory, Resource->GetSizeBytes());
+			INC_DWORD_STAT_BY(STAT_Shaders_ShaderMemory, FrozenContentSize);
+			INC_DWORD_STAT_BY(STAT_Shaders_NumShadersLoaded, NumFrozenShaders);
+		}
 	}
 }
 
