Vulkan: prevent crash on exit due to SCL (UE-89478).

diff --git a/Engine/Source/Runtime/RHI/Public/RHIResources.h b/Engine/Source/Runtime/RHI/Public/RHIResources.h
--- a/Engine/Source/Runtime/RHI/Public/RHIResources.h
+++ b/Engine/Source/Runtime/RHI/Public/RHIResources.h
@@ -2279,6 +2279,7 @@ public:
 	virtual FGraphEventRef PreloadShaderMap(int32 ShaderMapIndex) { return FGraphEventRef(); }
 
 	virtual TRefCountPtr<FRHIShader> CreateShader(int32 ShaderIndex) { return nullptr; }
+	virtual void Teardown() {};
 
 protected:
 	EShaderPlatform Platform;
diff --git a/Engine/Source/Runtime/RenderCore/Private/ShaderCodeArchive.cpp b/Engine/Source/Runtime/RenderCore/Private/ShaderCodeArchive.cpp
--- a/Engine/Source/Runtime/RenderCore/Private/ShaderCodeArchive.cpp
+++ b/Engine/Source/Runtime/RenderCore/Private/ShaderCodeArchive.cpp
@@ -215,9 +215,15 @@ FShaderCodeArchive::FShaderCodeArchive(EShaderPlatform InPlatform, const FString
 FShaderCodeArchive::~FShaderCodeArchive()
 {
 	DEC_DWORD_STAT_BY(STAT_Shaders_ShaderResourceMemory, GetSizeBytes());
+	Teardown();
+}
+
+void FShaderCodeArchive::Teardown()
+{
 	if (FileCacheHandle)
 	{
 		delete FileCacheHandle;
+		FileCacheHandle = nullptr;
 	}
 }
 
diff --git a/Engine/Source/Runtime/RenderCore/Private/ShaderCodeLibrary.cpp b/Engine/Source/Runtime/RenderCore/Private/ShaderCodeLibrary.cpp
--- a/Engine/Source/Runtime/RenderCore/Private/ShaderCodeLibrary.cpp
+++ b/Engine/Source/Runtime/RenderCore/Private/ShaderCodeLibrary.cpp
@@ -594,6 +594,16 @@ public:
 
 	~FShaderLibraryInstance()
 	{
+		// release RHI on all of the resources
+		for (FShaderMapResource_SharedCode* Resource : Resources)
+		{
+			if (Resource)
+			{
+				Resource->ReleaseRHI();
+			}
+		}
+		
+		Library->Teardown();
 		DEC_DWORD_STAT_BY(STAT_Shaders_ShaderResourceMemory, GetSizeBytes());
 	}
 
diff --git a/Engine/Source/Runtime/RenderCore/Public/ShaderCodeArchive.h b/Engine/Source/Runtime/RenderCore/Public/ShaderCodeArchive.h
--- a/Engine/Source/Runtime/RenderCore/Public/ShaderCodeArchive.h
+++ b/Engine/Source/Runtime/RenderCore/Public/ShaderCodeArchive.h
@@ -151,6 +151,7 @@ public:
 	virtual FGraphEventRef PreloadShaderMap(int32 ShaderMapIndex) override;
 
 	virtual TRefCountPtr<FRHIShader> CreateShader(int32 Index) override;
+	virtual void Teardown() override;
 
 protected:
 	FShaderCodeArchive(EShaderPlatform InPlatform, const FString& InLibraryDir, const FString& InLibraryName);
