Allow FRenderResource to be initialized even on exit request.

#rnx

diff --git a/Engine/Source/Runtime/RenderCore/Private/RenderResource.cpp b/Engine/Source/Runtime/RenderCore/Private/RenderResource.cpp
--- a/Engine/Source/Runtime/RenderCore/Private/RenderResource.cpp
+++ b/Engine/Source/Runtime/RenderCore/Private/RenderResource.cpp
@@ -78,8 +78,6 @@ void FRenderResource::InitResource()
 	check(IsInRenderingThread());
 	if (ListIndex == INDEX_NONE)
 	{
-		check(!IsEngineExitRequested());
-
 		TArray<FRenderResource*>& ResourceList = GetResourceList();
 		TArray<int32>& FreeIndicesList = GetFreeIndicesList();
 
