Share sampler state objects between textures

diff --git a/Engine/Source/Runtime/Apple/MetalRHI/Private/MetalState.cpp b/Engine/Source/Runtime/Apple/MetalRHI/Private/MetalState.cpp
--- a/Engine/Source/Runtime/Apple/MetalRHI/Private/MetalState.cpp
+++ b/Engine/Source/Runtime/Apple/MetalRHI/Private/MetalState.cpp
@@ -166,36 +166,6 @@ static EColorWriteMask TranslateWriteMask(MTLColorWriteMask WriteMask)
 	return (EColorWriteMask)Result;
 }
 
-static bool operator==(const FSamplerStateInitializerRHI& Left, const FSamplerStateInitializerRHI& Right)
-{
-	bool bSame = Left.Filter == Right.Filter
-	&& Left.MaxAnisotropy == Right.MaxAnisotropy
-	&& Left.AddressU == Right.AddressU
-	&& Left.AddressV == Right.AddressV
-	&& Left.AddressW == Right.AddressW
-	&& Left.MinMipLevel == Right.MinMipLevel
-	&& Left.MaxMipLevel == Right.MaxMipLevel
-	&& Left.SamplerComparisonFunction == Right.SamplerComparisonFunction
-	&& Left.BorderColor == Right.BorderColor;
-	
-	return bSame;
-}
-
-static uint32 GetTypeHash(const FSamplerStateInitializerRHI& Initializer)
-{
-	uint32 Hash = FCrc::MemCrc32(&Initializer.Filter, sizeof(Initializer.Filter));
-	Hash = FCrc::MemCrc32(&Initializer.MaxAnisotropy, sizeof(Initializer.MaxAnisotropy), Hash);
-	Hash = FCrc::MemCrc32(&Initializer.AddressU, sizeof(Initializer.AddressU), Hash);
-	Hash = FCrc::MemCrc32(&Initializer.AddressV, sizeof(Initializer.AddressV), Hash);
-	Hash = FCrc::MemCrc32(&Initializer.AddressW, sizeof(Initializer.AddressW), Hash);
-	Hash = FCrc::MemCrc32(&Initializer.MinMipLevel, sizeof(Initializer.MinMipLevel), Hash);
-	Hash = FCrc::MemCrc32(&Initializer.MaxMipLevel, sizeof(Initializer.MaxMipLevel), Hash);
-	Hash = FCrc::MemCrc32(&Initializer.SamplerComparisonFunction, sizeof(Initializer.SamplerComparisonFunction), Hash);
-	Hash = FCrc::MemCrc32(&Initializer.BorderColor, sizeof(Initializer.BorderColor), Hash);
-	
-	return Hash;
-}
-
 template <typename InitializerType, typename StateType>
 class FMetalStateObjectCache
 {
diff --git a/Engine/Source/Runtime/Engine/Private/Components/ReflectionCaptureComponent.cpp b/Engine/Source/Runtime/Engine/Private/Components/ReflectionCaptureComponent.cpp
--- a/Engine/Source/Runtime/Engine/Private/Components/ReflectionCaptureComponent.cpp
+++ b/Engine/Source/Runtime/Engine/Private/Components/ReflectionCaptureComponent.cpp
@@ -680,7 +680,7 @@ public:
 			AM_Clamp,
 			AM_Clamp
 		);
-		SamplerStateRHI = RHICreateSamplerState(SamplerStateInitializer);
+		SamplerStateRHI = GetOrCreateSamplerState(SamplerStateInitializer);
 
 		INC_MEMORY_STAT_BY(STAT_ReflectionCaptureTextureMemory,CalcTextureSize(Size,Size,Format,NumMips) * 6);
 	}
diff --git a/Engine/Source/Runtime/Engine/Private/Components/SkyLightComponent.cpp b/Engine/Source/Runtime/Engine/Private/Components/SkyLightComponent.cpp
--- a/Engine/Source/Runtime/Engine/Private/Components/SkyLightComponent.cpp
+++ b/Engine/Source/Runtime/Engine/Private/Components/SkyLightComponent.cpp
@@ -100,7 +100,7 @@ void FSkyTextureCubeResource::InitRHI()
 			AM_Clamp,
 			AM_Clamp
 		);
-		SamplerStateRHI = RHICreateSamplerState(SamplerStateInitializer);
+		SamplerStateRHI = GetOrCreateSamplerState(SamplerStateInitializer);
 	}
 }
 
diff --git a/Engine/Source/Runtime/Engine/Private/Slate/SlateTextures.cpp b/Engine/Source/Runtime/Engine/Private/Slate/SlateTextures.cpp
--- a/Engine/Source/Runtime/Engine/Private/Slate/SlateTextures.cpp
+++ b/Engine/Source/Runtime/Engine/Private/Slate/SlateTextures.cpp
@@ -339,7 +339,7 @@ void FSlateTextureRenderTarget2DResource::InitDynamicRHI()
 		AddressY == TA_Wrap ? AM_Wrap : (AddressY == TA_Clamp ? AM_Clamp : AM_Mirror),
 		AM_Wrap
 	);
-	SamplerStateRHI = RHICreateSamplerState( SamplerStateInitializer );
+	SamplerStateRHI = GetOrCreateSamplerState( SamplerStateInitializer );
 }
 
 void FSlateTextureRenderTarget2DResource::ReleaseDynamicRHI()
diff --git a/Engine/Source/Runtime/Engine/Private/Texture2D.cpp b/Engine/Source/Runtime/Engine/Private/Texture2D.cpp
--- a/Engine/Source/Runtime/Engine/Private/Texture2D.cpp
+++ b/Engine/Source/Runtime/Engine/Private/Texture2D.cpp
@@ -1819,7 +1819,7 @@ void FTexture2DResource::CreateSamplerStates(float MipMapBias)
 	  AM_Wrap,
 	  MipMapBias
 	);
-	SamplerStateRHI = RHICreateSamplerState(SamplerStateInitializer);
+	SamplerStateRHI = GetOrCreateSamplerState(SamplerStateInitializer);
 
 	// Create a custom sampler state for using this texture in a deferred pass, where ddx / ddy are discontinuous
 	FSamplerStateInitializerRHI DeferredPassSamplerStateInitializer
@@ -1837,7 +1837,7 @@ void FTexture2DResource::CreateSamplerStates(float MipMapBias)
 	  2
 	);
 
-	DeferredPassSamplerStateRHI = RHICreateSamplerState(DeferredPassSamplerStateInitializer);
+	DeferredPassSamplerStateRHI = GetOrCreateSamplerState(DeferredPassSamplerStateInitializer);
 }
 
 // Recreate the sampler states (used when updating mip map lod bias offset)
@@ -1971,7 +1971,7 @@ void FVirtualTexture2DResource::InitRHI()
 		// This doesn't really matter when sampling the cache texture (as it only has a level 0, so whatever the bias that is sampled) but it does when we sample miptail texture
 		0 // VT currently ignores global mip bias ensure the miptail works the same -> UTexture2D::GetGlobalMipMapLODBias()
 	);
-	SamplerStateRHI = RHICreateSamplerState(SamplerStateInitializer);
+	SamplerStateRHI = GetOrCreateSamplerState(SamplerStateInitializer);
 
 	const int32 MaxLevel = VTData->GetNumMips() - FirstMipToUse - 1;
 	check(MaxLevel >= 0);
@@ -2454,7 +2454,7 @@ void FTexture2DArrayResource::InitRHI()
 		SamplerZAddress
 	);
 
-	SamplerStateRHI = RHICreateSamplerState(SamplerStateInitializer);
+	SamplerStateRHI = GetOrCreateSamplerState(SamplerStateInitializer);
 }
 
 FIncomingTextureArrayDataEntry::FIncomingTextureArrayDataEntry(UTexture2D* InTexture)
diff --git a/Engine/Source/Runtime/Engine/Private/Texture2DDynamic.cpp b/Engine/Source/Runtime/Engine/Private/Texture2DDynamic.cpp
--- a/Engine/Source/Runtime/Engine/Private/Texture2DDynamic.cpp
+++ b/Engine/Source/Runtime/Engine/Private/Texture2DDynamic.cpp
@@ -44,7 +44,7 @@ void FTexture2DDynamicResource::InitRHI()
 		SamplerAddressMode,
 		SamplerAddressMode
 	);
-	SamplerStateRHI = RHICreateSamplerState( SamplerStateInitializer );
+	SamplerStateRHI = GetOrCreateSamplerState( SamplerStateInitializer );
 
 	uint32 Flags = 0;
 	if ( Owner->bIsResolveTarget )
diff --git a/Engine/Source/Runtime/Engine/Private/TextureCube.cpp b/Engine/Source/Runtime/Engine/Private/TextureCube.cpp
--- a/Engine/Source/Runtime/Engine/Private/TextureCube.cpp
+++ b/Engine/Source/Runtime/Engine/Private/TextureCube.cpp
@@ -239,7 +239,7 @@ public:
 			AM_Clamp,
 			AM_Clamp
 		);
-		SamplerStateRHI = RHICreateSamplerState(SamplerStateInitializer);
+		SamplerStateRHI = GetOrCreateSamplerState(SamplerStateInitializer);
 
 		// Set the greyscale format flag appropriately.
 		EPixelFormat PixelFormat = Owner->GetPixelFormat();
diff --git a/Engine/Source/Runtime/Engine/Private/TextureRenderTarget2D.cpp b/Engine/Source/Runtime/Engine/Private/TextureRenderTarget2D.cpp
--- a/Engine/Source/Runtime/Engine/Private/TextureRenderTarget2D.cpp
+++ b/Engine/Source/Runtime/Engine/Private/TextureRenderTarget2D.cpp
@@ -516,7 +516,7 @@ void FTextureRenderTarget2DResource::InitDynamicRHI()
 		Owner->AddressY == TA_Wrap ? AM_Wrap : (Owner->AddressY == TA_Clamp ? AM_Clamp : AM_Mirror),
 		AM_Wrap
 	);
-	SamplerStateRHI = RHICreateSamplerState( SamplerStateInitializer );
+	SamplerStateRHI = GetOrCreateSamplerState( SamplerStateInitializer );
 }
 
 /**
diff --git a/Engine/Source/Runtime/Engine/Private/TextureRenderTargetCube.cpp b/Engine/Source/Runtime/Engine/Private/TextureRenderTargetCube.cpp
--- a/Engine/Source/Runtime/Engine/Private/TextureRenderTargetCube.cpp
+++ b/Engine/Source/Runtime/Engine/Private/TextureRenderTargetCube.cpp
@@ -266,7 +266,7 @@ void FTextureRenderTargetCubeResource::InitDynamicRHI()
 		AM_Clamp,
 		AM_Clamp
 	);
-	SamplerStateRHI = RHICreateSamplerState(SamplerStateInitializer);
+	SamplerStateRHI = GetOrCreateSamplerState(SamplerStateInitializer);
 }
 
 /**
diff --git a/Engine/Source/Runtime/Engine/Private/VolumeTexture.cpp b/Engine/Source/Runtime/Engine/Private/VolumeTexture.cpp
--- a/Engine/Source/Runtime/Engine/Private/VolumeTexture.cpp
+++ b/Engine/Source/Runtime/Engine/Private/VolumeTexture.cpp
@@ -466,7 +466,7 @@ public:
 			AM_Wrap,
 			AM_Wrap
 		);
-		SamplerStateRHI = RHICreateSamplerState(SamplerStateInitializer);
+		SamplerStateRHI = GetOrCreateSamplerState(SamplerStateInitializer);
 	}
 
 	virtual void ReleaseRHI() override
diff --git a/Engine/Source/Runtime/Engine/Public/PlanarReflectionSceneProxy.h b/Engine/Source/Runtime/Engine/Public/PlanarReflectionSceneProxy.h
--- a/Engine/Source/Runtime/Engine/Public/PlanarReflectionSceneProxy.h
+++ b/Engine/Source/Runtime/Engine/Public/PlanarReflectionSceneProxy.h
@@ -41,7 +41,7 @@ public:
 			AM_Clamp,
 			AM_Clamp
 		);
-		SamplerStateRHI = RHICreateSamplerState(SamplerStateInitializer);
+		SamplerStateRHI = GetOrCreateSamplerState(SamplerStateInitializer);
 
 		FTexture2DRHIRef Texture2DRHI;
 		FRHIResourceCreateInfo CreateInfo = { FClearValueBinding(FLinearColor::Black) };
diff --git a/Engine/Source/Runtime/MediaAssets/Private/Misc/MediaTextureResource.cpp b/Engine/Source/Runtime/MediaAssets/Private/Misc/MediaTextureResource.cpp
--- a/Engine/Source/Runtime/MediaAssets/Private/Misc/MediaTextureResource.cpp
+++ b/Engine/Source/Runtime/MediaAssets/Private/Misc/MediaTextureResource.cpp
@@ -429,7 +429,7 @@ void FMediaTextureResource::SetupSampler()
 			AM_Wrap
 		);
 
-		SamplerStateRHI = RHICreateSamplerState(SamplerStateInitializer);
+		SamplerStateRHI = GetOrCreateSamplerState(SamplerStateInitializer);
 	}
 }
 
diff --git a/Engine/Source/Runtime/OpenGLDrv/Private/OpenGLState.cpp b/Engine/Source/Runtime/OpenGLDrv/Private/OpenGLState.cpp
--- a/Engine/Source/Runtime/OpenGLDrv/Private/OpenGLState.cpp
+++ b/Engine/Source/Runtime/OpenGLDrv/Private/OpenGLState.cpp
@@ -11,46 +11,6 @@
 
 GLint GMaxOpenGLTextureFilterAnisotropic = 1;
 
-// Similar to sizeof(FSamplerStateInitializerRHI), but without any padding added by the compiler
-static uint32 SizeOfSamplerStateInitializer = 0;
-
-static FArchive& operator<<(FArchive& Ar, FSamplerStateInitializerRHI& SamplerStateInitializer)
-{
-	Ar << SamplerStateInitializer.Filter;
-	Ar << SamplerStateInitializer.AddressU;
-	Ar << SamplerStateInitializer.AddressV;
-	Ar << SamplerStateInitializer.AddressW;
-	Ar << SamplerStateInitializer.MipBias;
-	Ar << SamplerStateInitializer.MinMipLevel;
-	Ar << SamplerStateInitializer.MaxMipLevel;
-	Ar << SamplerStateInitializer.MaxAnisotropy;
-	Ar << SamplerStateInitializer.BorderColor;
-	Ar << SamplerStateInitializer.SamplerComparisonFunction;
-	return Ar;
-}
-
-static FORCEINLINE void CalculateSizeOfSamplerStateInitializer()
-{
-	if (SizeOfSamplerStateInitializer == 0)
-	{
-		TArray<uint8> Data;
-		FMemoryWriter Writer(Data);
-		FSamplerStateInitializerRHI State;
-		Writer << State;
-		SizeOfSamplerStateInitializer = Data.Num();
-	}
-}
-
-static bool operator==(const FSamplerStateInitializerRHI& A, const FSamplerStateInitializerRHI& B)
-{
-	return FMemory::Memcmp(&A, &B, SizeOfSamplerStateInitializer) == 0;
-}
-
-static FORCEINLINE uint32 GetTypeHash(const FSamplerStateInitializerRHI& SamplerState)
-{
-	return FCrc::MemCrc_DEPRECATED(&SamplerState, SizeOfSamplerStateInitializer);
-}
-
 // Hash of sampler states, used for caching sampler states and texture objects
 static TMap<FSamplerStateInitializerRHI, FOpenGLSamplerState*> GSamplerStateCache;
 
@@ -271,9 +231,6 @@ FSamplerStateRHIRef FOpenGLDynamicRHI::RHICreateSamplerState(const FSamplerState
 		return *Found;
 	}
 
-	// Create a new one
-	CalculateSizeOfSamplerStateInitializer();
-
 	FOpenGLSamplerState* SamplerState = new FOpenGLSamplerState;
 
 	SamplerState->Data.WrapS = TranslateAddressMode( Initializer.AddressU );
diff --git a/Engine/Source/Runtime/RHI/Private/RHI.cpp b/Engine/Source/Runtime/RHI/Private/RHI.cpp
--- a/Engine/Source/Runtime/RHI/Private/RHI.cpp
+++ b/Engine/Source/Runtime/RHI/Private/RHI.cpp
@@ -169,6 +169,37 @@ void FVertexElement::FromString(const FStringView& InSrc)
 	check(Parts.GetData() + PartCount == PartIt);
 }
 
+uint32 GetTypeHash(const FSamplerStateInitializerRHI& Initializer)
+{
+	uint32 Hash = GetTypeHash(Initializer.Filter);
+	Hash = HashCombine(Hash, GetTypeHash(Initializer.AddressU));
+	Hash = HashCombine(Hash, GetTypeHash(Initializer.AddressV));
+	Hash = HashCombine(Hash, GetTypeHash(Initializer.AddressW));
+	Hash = HashCombine(Hash, GetTypeHash(Initializer.MipBias));
+	Hash = HashCombine(Hash, GetTypeHash(Initializer.MinMipLevel));
+	Hash = HashCombine(Hash, GetTypeHash(Initializer.MaxMipLevel));
+	Hash = HashCombine(Hash, GetTypeHash(Initializer.MaxAnisotropy));
+	Hash = HashCombine(Hash, GetTypeHash(Initializer.BorderColor));
+	Hash = HashCombine(Hash, GetTypeHash(Initializer.SamplerComparisonFunction));
+	return Hash;
+}
+
+bool operator== (const FSamplerStateInitializerRHI& A, const FSamplerStateInitializerRHI& B)
+{
+	bool bSame = 
+		A.Filter == B.Filter &&
+		A.AddressU == B.AddressU &&
+		A.AddressV == B.AddressV &&
+		A.AddressW == B.AddressW &&
+		A.MipBias == B.MipBias &&
+		A.MinMipLevel == B.MinMipLevel &&
+		A.MaxMipLevel == B.MaxMipLevel &&
+		A.MaxAnisotropy == B.MaxAnisotropy &&
+		A.BorderColor == B.BorderColor &&
+		A.SamplerComparisonFunction == B.SamplerComparisonFunction;
+	return bSame;
+}
+
 uint32 GetTypeHash(const FRasterizerStateInitializerRHI& Initializer)
 {
 	uint32 Hash = GetTypeHash(Initializer.FillMode);
diff --git a/Engine/Source/Runtime/RHI/Public/RHI.h b/Engine/Source/Runtime/RHI/Public/RHI.h
--- a/Engine/Source/Runtime/RHI/Public/RHI.h
+++ b/Engine/Source/Runtime/RHI/Public/RHI.h
@@ -848,6 +848,9 @@ struct FSamplerStateInitializerRHI
 	int32 MaxAnisotropy;
 	uint32 BorderColor;
 	TEnumAsByte<ESamplerCompareFunction> SamplerComparisonFunction;
+
+	RHI_API friend uint32 GetTypeHash(const FSamplerStateInitializerRHI& Initializer);
+	RHI_API friend bool operator== (const FSamplerStateInitializerRHI& A, const FSamplerStateInitializerRHI& B);
 };
 
 struct FRasterizerStateInitializerRHI
diff --git a/Engine/Source/Runtime/RenderCore/Private/RenderResource.cpp b/Engine/Source/Runtime/RenderCore/Private/RenderResource.cpp
--- a/Engine/Source/Runtime/RenderCore/Private/RenderResource.cpp
+++ b/Engine/Source/Runtime/RenderCore/Private/RenderResource.cpp
@@ -796,3 +796,36 @@ void FMipBiasFade::SetNewMipCount( float ActualMipCount, float TargetMipCount, d
 		}
 	}
 }
+
+class FTextureSamplerStateCache : public FRenderResource
+{
+public:
+	TMap<FSamplerStateInitializerRHI, FRHISamplerState*> Samplers;
+
+	virtual void ReleaseRHI() override
+	{
+		for (auto Pair : Samplers)
+		{
+			Pair.Value->Release();
+		}
+		Samplers.Empty();
+	}
+};
+
+TGlobalResource<FTextureSamplerStateCache> GTextureSamplerStateCache;
+
+FRHISamplerState* FTexture::GetOrCreateSamplerState(const FSamplerStateInitializerRHI& Initializer)
+{
+	FRHISamplerState** Found = GTextureSamplerStateCache.Samplers.Find(Initializer);
+	if (Found)
+	{
+		return *Found;
+	}
+	
+	FSamplerStateRHIRef NewState = RHICreateSamplerState(Initializer);
+	
+	// Add an extra reference so we don't have TRefCountPtr in the maps
+	NewState->AddRef();
+	GTextureSamplerStateCache.Samplers.Add(Initializer, NewState);
+	return NewState;
+}
\ No newline at end of file
diff --git a/Engine/Source/Runtime/RenderCore/Private/RenderUtils.cpp b/Engine/Source/Runtime/RenderCore/Private/RenderUtils.cpp
--- a/Engine/Source/Runtime/RenderCore/Private/RenderUtils.cpp
+++ b/Engine/Source/Runtime/RenderCore/Private/RenderUtils.cpp
@@ -221,7 +221,7 @@ public:
 
 		// Create the sampler state RHI resource.
 		FSamplerStateInitializerRHI SamplerStateInitializer(SF_Point,AM_Wrap,AM_Wrap,AM_Wrap);
-		SamplerStateRHI = RHICreateSamplerState(SamplerStateInitializer);
+		SamplerStateRHI = GetOrCreateSamplerState(SamplerStateInitializer);
 
 		// Create a view of the texture
 		ShaderResourceViewRHI = RHICreateShaderResourceView(TextureRHI, 0u);
@@ -365,7 +365,7 @@ public:
 
 		// Create the sampler state.
 		FSamplerStateInitializerRHI SamplerStateInitializer(SF_Point, AM_Wrap, AM_Wrap, AM_Wrap);
-		SamplerStateRHI = RHICreateSamplerState(SamplerStateInitializer);
+		SamplerStateRHI = GetOrCreateSamplerState(SamplerStateInitializer);
 	}
 
 	/**
@@ -409,7 +409,7 @@ public:
 
 			// Create the sampler state RHI resource.
 			FSamplerStateInitializerRHI SamplerStateInitializer(SF_Point,AM_Wrap,AM_Wrap,AM_Wrap);
-			SamplerStateRHI = RHICreateSamplerState(SamplerStateInitializer);
+			SamplerStateRHI = GetOrCreateSamplerState(SamplerStateInitializer);
 		}
 	}
 
@@ -473,7 +473,7 @@ public:
 
 		// Create the sampler state RHI resource.
 		FSamplerStateInitializerRHI SamplerStateInitializer(SF_Point,AM_Wrap,AM_Wrap,AM_Wrap);
-		SamplerStateRHI = RHICreateSamplerState(SamplerStateInitializer);
+		SamplerStateRHI = GetOrCreateSamplerState(SamplerStateInitializer);
 	}
 
 	/** Returns the width of the texture in pixels. */
@@ -557,7 +557,7 @@ public:
 
 		// Create the sampler state RHI resource.
 		FSamplerStateInitializerRHI SamplerStateInitializer(SF_Point, AM_Wrap, AM_Wrap, AM_Wrap);
-		SamplerStateRHI = RHICreateSamplerState(SamplerStateInitializer);
+		SamplerStateRHI = GetOrCreateSamplerState(SamplerStateInitializer);
 	}
 
 	/** Returns the width of the texture in pixels. */
@@ -626,7 +626,7 @@ public:
 
 			// Create the sampler state RHI resource.
 			FSamplerStateInitializerRHI SamplerStateInitializer(SF_Point,AM_Wrap,AM_Wrap,AM_Wrap);
-			SamplerStateRHI = RHICreateSamplerState(SamplerStateInitializer);
+			SamplerStateRHI = GetOrCreateSamplerState(SamplerStateInitializer);
 		}
 	}
 
@@ -667,7 +667,7 @@ public:
 
 		// Create the sampler state RHI resource.
 		FSamplerStateInitializerRHI SamplerStateInitializer(SF_Point, AM_Wrap, AM_Wrap, AM_Wrap);
-		SamplerStateRHI = RHICreateSamplerState(SamplerStateInitializer);
+		SamplerStateRHI = GetOrCreateSamplerState(SamplerStateInitializer);
 
 		// Create a view of the texture
 		ShaderResourceViewRHI = RHICreateShaderResourceView(TextureRHI, 0u);
diff --git a/Engine/Source/Runtime/RenderCore/Public/RenderResource.h b/Engine/Source/Runtime/RenderCore/Public/RenderResource.h
--- a/Engine/Source/Runtime/RenderCore/Public/RenderResource.h
+++ b/Engine/Source/Runtime/RenderCore/Public/RenderResource.h
@@ -420,6 +420,9 @@ public:
 		DeferredPassSamplerStateRHI.SafeRelease();
 	}
 	virtual FString GetFriendlyName() const override { return TEXT("FTexture"); }
+
+protected:
+	RENDERCORE_API static FRHISamplerState* GetOrCreateSamplerState(const FSamplerStateInitializerRHI& Initializer);
 };
 
 /** A textures resource that includes an SRV. */
diff --git a/Engine/Source/Runtime/SlateRHIRenderer/Private/SlateRHIFontTexture.cpp b/Engine/Source/Runtime/SlateRHIRenderer/Private/SlateRHIFontTexture.cpp
--- a/Engine/Source/Runtime/SlateRHIRenderer/Private/SlateRHIFontTexture.cpp
+++ b/Engine/Source/Runtime/SlateRHIRenderer/Private/SlateRHIFontTexture.cpp
@@ -42,7 +42,7 @@ void FSlateFontTextureRHIResource::InitDynamicRHI()
 		  0,
 		  0
 		);
-		SamplerStateRHI = RHICreateSamplerState(SamplerStateInitializer);
+		SamplerStateRHI = GetOrCreateSamplerState(SamplerStateInitializer);
 
 		// Create a custom sampler state for using this texture in a deferred pass, where ddx / ddy are discontinuous
 		FSamplerStateInitializerRHI DeferredPassSamplerStateInitializer
@@ -56,7 +56,7 @@ void FSlateFontTextureRHIResource::InitDynamicRHI()
 		  0,
 		  0
 		);
-		DeferredPassSamplerStateRHI = RHICreateSamplerState(DeferredPassSamplerStateInitializer);
+		DeferredPassSamplerStateRHI = GetOrCreateSamplerState(DeferredPassSamplerStateInitializer);
 
 		INC_MEMORY_STAT_BY(STAT_SlateTextureGPUMemory, Width*Height*GPixelFormats[PixelFormat].BlockBytes);
 	}
