Fix generation of shader debug data

diff --git a/Engine/Shaders/Public/ShaderVersion.ush b/Engine/Shaders/Public/ShaderVersion.ush
--- a/Engine/Shaders/Public/ShaderVersion.ush
+++ b/Engine/Shaders/Public/ShaderVersion.ush
@@ -3,4 +3,4 @@
 // in Platform.ush (which should be included in any shader) it allows to invalidate the shader DDC.
 // 
 // If you are merging streams and there is a conflict with this GUID you should make a new GUID rather than taking one or the other.
-// GUID = 979F941B400ECB2A9CDA3CB0FBA2FC9E
\ No newline at end of file
+// GUID = BD84E7E6-F382-413D-B218-71A7172EE73D
\ No newline at end of file
diff --git a/Engine/Source/Runtime/RenderCore/Private/ShaderMap.cpp b/Engine/Source/Runtime/RenderCore/Private/ShaderMap.cpp
--- a/Engine/Source/Runtime/RenderCore/Private/ShaderMap.cpp
+++ b/Engine/Source/Runtime/RenderCore/Private/ShaderMap.cpp
@@ -167,6 +167,11 @@ void FShaderMapBase::Serialize(FArchive& Ar, bool bInlineShaderResources, bool b
 #endif // WITH_EDITOR
 		Ar << bShareCode;
 #if WITH_EDITOR
+		if (Ar.IsCooking())
+		{
+			Code->NotifyShadersCooked(Ar.CookingTarget());
+		}
+
 		if (bShareCode)
 		{
 			FSHAHash ResourceHash = Code->ResourceHash;
@@ -176,7 +181,7 @@ void FShaderMapBase::Serialize(FArchive& Ar, bool bInlineShaderResources, bool b
 		else
 #endif // WITH_EDITOR
 		{
-			Code->Serialize(Ar);
+			Code->Serialize(Ar, bLoadedByCookedMaterial);
 		}
 	}
 	else
@@ -239,7 +244,7 @@ void FShaderMapBase::Serialize(FArchive& Ar, bool bInlineShaderResources, bool b
 		else
 		{
 			Code = new FShaderMapResourceCode();
-			Code->Serialize(Ar);
+			Code->Serialize(Ar, bLoadedByCookedMaterial);
 			Resource = new FShaderMapResource_InlineCode(GetShaderPlatform(), Code);
 		}
 
diff --git a/Engine/Source/Runtime/RenderCore/Private/ShaderResource.cpp b/Engine/Source/Runtime/RenderCore/Private/ShaderResource.cpp
--- a/Engine/Source/Runtime/RenderCore/Private/ShaderResource.cpp
+++ b/Engine/Source/Runtime/RenderCore/Private/ShaderResource.cpp
@@ -142,16 +142,51 @@ void FShaderMapResourceCode::Finalize()
 	ApplyResourceStats(*this);
 }
 
-void FShaderMapResourceCode::Serialize(FArchive& Ar)
+void FShaderMapResourceCode::Serialize(FArchive& Ar, bool bLoadedByCookedMaterial)
 {
 	Ar << ResourceHash;
 	Ar << ShaderHashes;
 	Ar << ShaderEntries;
 	Ar << ShaderCode;
 	check(ShaderEntries.Num() == ShaderHashes.Num());
+#if WITH_EDITORONLY_DATA
+	const bool bSerializePlatformData = !bLoadedByCookedMaterial && (!Ar.IsCooking() || Ar.CookingTarget()->HasEditorOnlyData());
+	if (bSerializePlatformData)
+	{
+		Ar << PlatformDebugEntries;
+		Ar << PlatformDebugData;
+	}
+#endif // WITH_EDITORONLY_DATA
 	ApplyResourceStats(*this);
 }
 
+#if WITH_EDITORONLY_DATA
+void FShaderMapResourceCode::NotifyShadersCooked(const ITargetPlatform* TargetPlatform)
+{
+#if WITH_ENGINE
+	// Notify the platform shader format that this particular shader is being used in the cook.
+	// We discard this data in cooked builds unless Ar.CookingTarget()->HasEditorOnlyData() is true.
+	check(TargetPlatform);
+	if (PlatformDebugEntries.Num())
+	{
+		TArray<FName> ShaderFormatNames;
+		TargetPlatform->GetAllTargetedShaderFormats(ShaderFormatNames);
+		for (FName FormatName : ShaderFormatNames)
+		{
+			const IShaderFormat* ShaderFormat = GetTargetPlatformManagerRef().FindShaderFormat(FormatName);
+			if (ShaderFormat)
+			{
+				for (const FPlatformDebugEntry& Entry : PlatformDebugEntries)
+				{
+					ShaderFormat->NotifyShaderCooked(MakeArrayView(PlatformDebugData.GetData() + Entry.Offset, Entry.Size), FormatName);
+				}
+			}
+		}
+	}
+#endif // WITH_ENGINE
+}
+#endif // WITH_EDITORONLY_DATA
+
 FShaderMapResourceBuilder::FShaderMapResourceBuilder(FShaderMapResourceCode* InCode) : ShaderHashTable(1024, 256), Code(InCode)
 {
 	check(InCode->ShaderEntries.Num() == InCode->ShaderHashes.Num());
@@ -216,6 +251,19 @@ int32 FShaderMapResourceBuilder::FindOrAddCode(EShaderFrequency InFrequency, con
 	return Index;
 }
 
+#if WITH_EDITORONLY_DATA
+void FShaderMapResourceBuilder::AddPlatformDebugData(TConstArrayView<uint8> InPlatformDebugData)
+{
+	if (InPlatformDebugData.Num() > 0)
+	{
+		FShaderMapResourceCode::FPlatformDebugEntry& Entry = Code->PlatformDebugEntries.AddDefaulted_GetRef();
+		Entry.Offset = Code->PlatformDebugData.Num();
+		Entry.Size = InPlatformDebugData.Num();
+		Code->PlatformDebugData.Append(InPlatformDebugData.GetData(), InPlatformDebugData.Num());
+	}
+}
+#endif // WITH_EDITORONLY_DATA
+
 FShaderMapResource::FShaderMapResource(EShaderPlatform InPlatform, int32 NumShaders)
 	: Platform(InPlatform)
 	, NumRefs(0)
diff --git a/Engine/Source/Runtime/RenderCore/Public/Shader.h b/Engine/Source/Runtime/RenderCore/Public/Shader.h
--- a/Engine/Source/Runtime/RenderCore/Public/Shader.h
+++ b/Engine/Source/Runtime/RenderCore/Public/Shader.h
@@ -300,11 +300,27 @@ public:
 		}
 	};
 
+#if WITH_EDITORONLY_DATA
+	struct FPlatformDebugEntry
+	{
+		uint32 Offset;
+		uint32 Size;
+
+		friend FArchive& operator<<(FArchive& Ar, FPlatformDebugEntry& Entry)
+		{
+			return Ar << Entry.Offset << Entry.Size;
+		}
+	};
+#endif // WITH_EDITORONLY_DATA
+
 	RENDERCORE_API ~FShaderMapResourceCode();
 
 	RENDERCORE_API void Finalize();
 
-	RENDERCORE_API void Serialize(FArchive& Ar);
+	RENDERCORE_API void Serialize(FArchive& Ar, bool bLoadedByCookedMaterial);
+#if WITH_EDITORONLY_DATA
+	RENDERCORE_API void NotifyShadersCooked(const ITargetPlatform* TargetPlatform);
+#endif // WITH_EDITORONLY_DATA
 
 	inline const uint32 GetSizeBytes() const
 	{
@@ -315,6 +331,10 @@ public:
 	TArray<FSHAHash> ShaderHashes;
 	TArray<FShaderEntry> ShaderEntries;
 	TArray<uint8> ShaderCode;
+#if WITH_EDITORONLY_DATA
+	TArray<FPlatformDebugEntry> PlatformDebugEntries;
+	TArray<uint8> PlatformDebugData;
+#endif // WITH_EDITORONLY_DATA
 };
 	
 class RENDERCORE_API FShaderMapResource_InlineCode : public FShaderMapResource
@@ -339,12 +359,17 @@ public:
 
 	inline int32 FindOrAddCode(const FShaderCompilerOutput& Output)
 	{
+#if WITH_EDITORONLY_DATA
+		AddPlatformDebugData(Output.PlatformDebugData);
+#endif
 		return FindOrAddCode(Output.Target.GetFrequency(), Output.OutputHash, Output.ShaderCode.GetReadAccess());
 	}
 
 	RENDERCORE_API int32 FindCode(const FSHAHash& InHash, uint32 InKey) const;
 	RENDERCORE_API int32 FindOrAddCode(EShaderFrequency InFrequency, const FSHAHash& InHash, const TConstArrayView<uint8>& InCode);
-
+#if WITH_EDITORONLY_DATA
+	RENDERCORE_API void AddPlatformDebugData(TConstArrayView<uint8> InPlatformDebugData);
+#endif // WITH_EDITORONLY_DATA
 private:
 	FHashTable ShaderHashTable;
 	TRefCountPtr<FShaderMapResourceCode> Code;
