Fix inconsistent operator new/FMemory:Free pair when deleting frozen shaders
- Also fix leaked memory when deleting frozen shaders

diff --git a/Engine/Source/Runtime/Core/Private/Serialization/MemoryImage.cpp b/Engine/Source/Runtime/Core/Private/Serialization/MemoryImage.cpp
--- a/Engine/Source/Runtime/Core/Private/Serialization/MemoryImage.cpp
+++ b/Engine/Source/Runtime/Core/Private/Serialization/MemoryImage.cpp
@@ -251,7 +251,7 @@ void InternalDeleteObjectFromLayout(void* Object, const FTypeLayoutDesc& TypeDes
 	}
 	if (!bIsFrozen)
 	{
-		FMemory::Free(Object);
+		::operator delete(Object);
 	}
 }
 
diff --git a/Engine/Source/Runtime/RenderCore/Private/ShaderMap.cpp b/Engine/Source/Runtime/RenderCore/Private/ShaderMap.cpp
--- a/Engine/Source/Runtime/RenderCore/Private/ShaderMap.cpp
+++ b/Engine/Source/Runtime/RenderCore/Private/ShaderMap.cpp
@@ -267,6 +267,10 @@ void FShaderMapBase::DestroyContent()
 		DEC_DWORD_STAT_BY(STAT_Shaders_NumShadersLoaded, NumFrozenShaders);
 
 		InternalDeleteObjectFromLayout(Content, ContentTypeLayout, FrozenContentSize > 0u);
+		if (FrozenContentSize > 0u)
+		{
+			FMemory::Free(Content);
+		}
 
 		FrozenContentSize = 0u;
 		NumFrozenShaders = 0u;
