Fix a bug that we don't recompile global shader after turn on/off MobileHDR

diff --git a/Engine/Source/Developer/ShaderFormatOpenGL/Private/ShaderFormatOpenGL.cpp b/Engine/Source/Developer/ShaderFormatOpenGL/Private/ShaderFormatOpenGL.cpp
--- a/Engine/Source/Developer/ShaderFormatOpenGL/Private/ShaderFormatOpenGL.cpp
+++ b/Engine/Source/Developer/ShaderFormatOpenGL/Private/ShaderFormatOpenGL.cpp
@@ -18,7 +18,7 @@ class FShaderFormatGLSL : public IShaderFormat
 	enum
 	{
 		/** Version for shader format, this becomes part of the DDC key. */
-		UE_SHADER_GLSL_VER = 90,
+		UE_SHADER_GLSL_VER = 91,
 	};
 
 	void CheckFormat(FName Format) const
diff --git a/Engine/Source/Runtime/RenderCore/Private/Shader.cpp b/Engine/Source/Runtime/RenderCore/Private/Shader.cpp
--- a/Engine/Source/Runtime/RenderCore/Private/Shader.cpp
+++ b/Engine/Source/Runtime/RenderCore/Private/Shader.cpp
@@ -1471,6 +1471,11 @@ void ShaderMapAppendKeyString(EShaderPlatform Platform, FString& KeyString)
 				}
 			}
 		}
+
+		{
+			static IConsoleVariable* CVar = IConsoleManager::Get().FindConsoleVariable(TEXT("r.MobileHDR"));
+			KeyString += (CVar && CVar->GetInt() != 0) ? TEXT("_MobileHDR") : TEXT("");
+		}
 	}
 
 	const FName ShaderFormatName = LegacyShaderPlatformToShaderFormat(Platform);
