Add an assert to a texture sampler cache to ensure its used only from RT

diff --git a/Engine/Source/Runtime/RenderCore/Private/RenderResource.cpp b/Engine/Source/Runtime/RenderCore/Private/RenderResource.cpp
--- a/Engine/Source/Runtime/RenderCore/Private/RenderResource.cpp
+++ b/Engine/Source/Runtime/RenderCore/Private/RenderResource.cpp
@@ -816,6 +816,10 @@ TGlobalResource<FTextureSamplerStateCache> GTextureSamplerStateCache;
 
 FRHISamplerState* FTexture::GetOrCreateSamplerState(const FSamplerStateInitializerRHI& Initializer)
 {
+	// This sampler cache is supposed to be used only from RT
+	// Add a lock here if it's used from multiple threads
+	check(IsInRenderingThread());
+	
 	FRHISamplerState** Found = GTextureSamplerStateCache.Samplers.Find(Initializer);
 	if (Found)
 	{
