Added stat scope for EnqueueUniqueRenderCommand.

#rnx

diff --git a/Engine/Source/Runtime/RenderCore/Public/RenderingThread.h b/Engine/Source/Runtime/RenderCore/Public/RenderingThread.h
--- a/Engine/Source/Runtime/RenderCore/Public/RenderingThread.h
+++ b/Engine/Source/Runtime/RenderCore/Public/RenderingThread.h
@@ -244,6 +244,7 @@ private:
 template<typename TSTR, typename LAMBDA>
 FORCEINLINE_DEBUGGABLE void EnqueueUniqueRenderCommand(LAMBDA&& Lambda)
 {
+	QUICK_SCOPE_CYCLE_COUNTER(STAT_EnqueueUniqueRenderCommand);
 	typedef TEnqueueUniqueRenderCommandType<TSTR, LAMBDA> EURCType;
 
 #if 0 // UE_SERVER && UE_BUILD_DEBUG
