Add extra logging to try to get more information about crash

diff --git a/Engine/Source/Runtime/Core/Private/FileCache/FileCache.cpp b/Engine/Source/Runtime/Core/Private/FileCache/FileCache.cpp
--- a/Engine/Source/Runtime/Core/Private/FileCache/FileCache.cpp
+++ b/Engine/Source/Runtime/Core/Private/FileCache/FileCache.cpp
@@ -745,7 +745,8 @@ FGraphEventRef FFileCacheHandle::PreloadData(const FFileCachePreloadEntry* Prelo
 		const CacheLineID StartLine = GetBlock<CacheLineID>(Entry.Offset);
 		const CacheLineID EndLine = GetBlock<CacheLineID>(Entry.Offset + Entry.Size - 1);
 
-		check(Entry.Offset > PrevOffset);
+		checkf(Entry.Offset > PrevOffset, TEXT("Preload entries must be sorted by Offset [%lld, %lld), %lld"),
+			Entry.Offset, Entry.Offset + Entry.Size, PrevOffset);
 		PrevOffset = Entry.Offset;
 
 		CurrentLine = CacheLineID(FMath::Max(CurrentLine.Get(), StartLine.Get()));
diff --git a/Engine/Source/Runtime/RenderCore/Private/ShaderCodeArchive.cpp b/Engine/Source/Runtime/RenderCore/Private/ShaderCodeArchive.cpp
--- a/Engine/Source/Runtime/RenderCore/Private/ShaderCodeArchive.cpp
+++ b/Engine/Source/Runtime/RenderCore/Private/ShaderCodeArchive.cpp
@@ -121,7 +121,7 @@ void FSerializedShaderArchive::Finalize()
 	{
 		check(ShaderMapEntry.NumShaders > 0u);
 		TArray<FFileCachePreloadEntry> SortedPreloadEntries;
-		SortedPreloadEntries.Empty(ShaderMapEntry.NumShaders);
+		SortedPreloadEntries.Empty(ShaderMapEntry.NumShaders + 1);
 		for (uint32 i = 0; i < ShaderMapEntry.NumShaders; ++i)
 		{
 			const int32 ShaderIndex = ShaderIndices[ShaderMapEntry.ShaderIndicesOffset + i];
@@ -137,6 +137,8 @@ void FSerializedShaderArchive::Finalize()
 		{
 			const FFileCachePreloadEntry& PreloadEntry = SortedPreloadEntries[PreloadIndex];
 			const int64 Gap = PreloadEntry.Offset - CurrentPreloadEntry.Offset - CurrentPreloadEntry.Size;
+			checkf(Gap >= 0, TEXT("Overlapping preload entries, [%lld-%lld), [%lld-%lld)"),
+				CurrentPreloadEntry.Offset, CurrentPreloadEntry.Offset + CurrentPreloadEntry.Size, PreloadEntry.Offset, PreloadEntry.Offset + PreloadEntry.Size);
 			if (Gap > 1024)
 			{
 				++ShaderMapEntry.NumPreloadEntries;
