Added DEFINE_GPU_DRAWCALL_STAT to get drawcall count per specific category
Use DEFINE_GPU_DRAWCALL_STAT macro instead of DEFINE_GPU_STAT to also get drawcall count in CSV for a specific pass.
SCOPED_GPU_STAT will then output CSV time for GPU pass and number of drawcalls of this GPU pass in CSV Drawcall category

diff --git a/Engine/Source/Runtime/RHI/Private/RHI.cpp b/Engine/Source/Runtime/RHI/Private/RHI.cpp
--- a/Engine/Source/Runtime/RHI/Private/RHI.cpp
+++ b/Engine/Source/Runtime/RHI/Private/RHI.cpp
@@ -18,6 +18,7 @@ IMPLEMENT_MODULE(FDefaultModuleImpl, RHI);
 /** RHI Logging. */
 DEFINE_LOG_CATEGORY(LogRHI);
 CSV_DEFINE_CATEGORY(RHI, true);
+CSV_DEFINE_CATEGORY(DrawCall, true);
 
 // Define counter stats.
 DEFINE_STAT(STAT_RHIDrawPrimitiveCalls);
@@ -132,6 +133,8 @@ TLockFreePointerListUnordered<FRHIResource, PLATFORM_CACHE_LINE_SIZE> FRHIResour
 FRHIResource* FRHIResource::CurrentlyDeleting = nullptr;
 TArray<FRHIResource::ResourcesToDelete> FRHIResource::DeferredDeletionQueue;
 uint32 FRHIResource::CurrentFrame = 0;
+RHI_API FDrawCallCategoryName* FDrawCallCategoryName::Array[FDrawCallCategoryName::MAX_DRAWCALL_CATEGORY];
+RHI_API int32 FDrawCallCategoryName::NumCategory = 0;
 
 FString FVertexElement::ToString() const
 {
@@ -731,8 +734,10 @@ RHI_API EShaderPlatform GShaderPlatformForFeatureLevel[ERHIFeatureLevel::Num] =
 
 // simple stats about draw calls. GNum is the previous frame and 
 // GCurrent is the current frame.
+// GCurrentNumDrawCallsRHIPtr points to the drawcall counter to increment
 RHI_API int32 GCurrentNumDrawCallsRHI = 0;
 RHI_API int32 GNumDrawCallsRHI = 0;
+RHI_API int32* GCurrentNumDrawCallsRHIPtr = &GCurrentNumDrawCallsRHI;
 RHI_API int32 GCurrentNumPrimitivesDrawnRHI = 0;
 RHI_API int32 GNumPrimitivesDrawnRHI = 0;
 
@@ -740,6 +745,17 @@ RHI_API int32 GNumPrimitivesDrawnRHI = 0;
 void RHIPrivateBeginFrame()
 {
 	GNumDrawCallsRHI = GCurrentNumDrawCallsRHI;
+	
+#if CSV_PROFILER
+	for (int32 Index=0; Index<FDrawCallCategoryName::NumCategory; ++Index)
+	{
+		FDrawCallCategoryName* CategoryName = FDrawCallCategoryName::Array[Index];
+		GNumDrawCallsRHI += CategoryName->Counter;
+		FCsvProfiler::RecordCustomStat(CategoryName->Name, CSV_CATEGORY_INDEX(DrawCall), CategoryName->Counter, ECsvCustomStatOp::Set);
+		CategoryName->Counter = 0;
+	}
+#endif
+
 	GNumPrimitivesDrawnRHI = GCurrentNumPrimitivesDrawnRHI;
 	CSV_CUSTOM_STAT(RHI, DrawCalls, GNumDrawCallsRHI, ECsvCustomStatOp::Set);
 	CSV_CUSTOM_STAT(RHI, PrimitivesDrawn, GNumPrimitivesDrawnRHI, ECsvCustomStatOp::Set);
diff --git a/Engine/Source/Runtime/RHI/Public/RHI.h b/Engine/Source/Runtime/RHI/Public/RHI.h
--- a/Engine/Source/Runtime/RHI/Public/RHI.h
+++ b/Engine/Source/Runtime/RHI/Public/RHI.h
@@ -497,6 +497,7 @@ extern RHI_API int32 GNumPrimitivesDrawnRHI;
 
 /** Num draw calls and primitives this frame (only accurate on RenderThread) */
 extern RHI_API int32 GCurrentNumDrawCallsRHI;
+extern RHI_API int32* GCurrentNumDrawCallsRHIPtr;
 extern RHI_API int32 GCurrentNumPrimitivesDrawnRHI;
 
 
@@ -1634,6 +1635,32 @@ struct FTextureMemoryStats
 	}
 };
 
+struct RHI_API FDrawCallCategoryName
+{
+	FDrawCallCategoryName() 
+	{ 
+	}
+
+	FDrawCallCategoryName(FName InName)
+		: Name(InName)
+		, Counter(0)
+	{
+		check(NumCategory < MAX_DRAWCALL_CATEGORY);
+		if (NumCategory < MAX_DRAWCALL_CATEGORY)
+		{
+			Array[NumCategory] = this;
+			NumCategory++;
+		}
+	}
+
+	FName Name;
+	int32 Counter = -1;
+
+	static constexpr int32 MAX_DRAWCALL_CATEGORY = 256;
+	static FDrawCallCategoryName* Array[MAX_DRAWCALL_CATEGORY];
+	static int32 NumCategory;
+};
+
 // RHI counter stats.
 DECLARE_DWORD_COUNTER_STAT_EXTERN(TEXT("DrawPrimitive calls"),STAT_RHIDrawPrimitiveCalls,STATGROUP_RHI,RHI_API);
 DECLARE_DWORD_COUNTER_STAT_EXTERN(TEXT("Triangles drawn"),STAT_RHITriangles,STATGROUP_RHI,RHI_API);
@@ -1642,7 +1669,7 @@ DECLARE_DWORD_COUNTER_STAT_EXTERN(TEXT("Lines drawn"),STAT_RHILines,STATGROUP_RH
 #if STATS
 	#define RHI_DRAW_CALL_INC() \
 		INC_DWORD_STAT(STAT_RHIDrawPrimitiveCalls); \
-		FPlatformAtomics::InterlockedIncrement(&GCurrentNumDrawCallsRHI);
+		FPlatformAtomics::InterlockedIncrement(GCurrentNumDrawCallsRHIPtr);
 
 	#define RHI_DRAW_CALL_STATS(PrimitiveType,NumPrimitives) \
 		RHI_DRAW_CALL_INC(); \
@@ -1651,11 +1678,11 @@ DECLARE_DWORD_COUNTER_STAT_EXTERN(TEXT("Lines drawn"),STAT_RHILines,STATGROUP_RH
 		FPlatformAtomics::InterlockedAdd(&GCurrentNumPrimitivesDrawnRHI, NumPrimitives);
 #else
 	#define RHI_DRAW_CALL_INC() \
-		FPlatformAtomics::InterlockedIncrement(&GCurrentNumDrawCallsRHI);
+		FPlatformAtomics::InterlockedIncrement(GCurrentNumDrawCallsRHIPtr);
 
 	#define RHI_DRAW_CALL_STATS(PrimitiveType,NumPrimitives) \
 		FPlatformAtomics::InterlockedAdd(&GCurrentNumPrimitivesDrawnRHI, NumPrimitives); \
-		FPlatformAtomics::InterlockedIncrement(&GCurrentNumDrawCallsRHI);
+		FPlatformAtomics::InterlockedIncrement(GCurrentNumDrawCallsRHIPtr);
 #endif
 
 // RHI memory stats.
diff --git a/Engine/Source/Runtime/RenderCore/Private/ProfilingDebugging/RealtimeGPUProfiler.cpp b/Engine/Source/Runtime/RenderCore/Private/ProfilingDebugging/RealtimeGPUProfiler.cpp
--- a/Engine/Source/Runtime/RenderCore/Private/ProfilingDebugging/RealtimeGPUProfiler.cpp
+++ b/Engine/Source/Runtime/RenderCore/Private/ProfilingDebugging/RealtimeGPUProfiler.cpp
@@ -837,7 +837,7 @@ void FRealtimeGPUProfiler::PopEvent(FRHICommandListImmediate& RHICmdList)
 /*-----------------------------------------------------------------------------
 FScopedGPUStatEvent
 -----------------------------------------------------------------------------*/
-void FScopedGPUStatEvent::Begin(FRHICommandList& InRHICmdList, const FName& Name, const FName& StatName)
+void FScopedGPUStatEvent::Begin(FRHICommandList& InRHICmdList, const FName& Name, const FName& StatName, int32* InNumDrawCallsPtr)
 {
 	check(IsInRenderingThread());
 	if (!AreGPUStatsEnabled())
@@ -850,6 +850,20 @@ void FScopedGPUStatEvent::Begin(FRHICommandList& InRHICmdList, const FName& Name
 	{
 		RHICmdList = (FRHICommandListImmediate*)&InRHICmdList;
 		FRealtimeGPUProfiler::Get()->PushEvent(*RHICmdList, Name, StatName);
+
+		check(InNumDrawCallsPtr != nullptr);
+		if ((*InNumDrawCallsPtr)!=-1)
+		{
+			NumDrawCallsPtr = InNumDrawCallsPtr;
+
+			RHICmdList->EnqueueLambda
+			(
+				[InNumDrawCallsPtr](FRHICommandListImmediate& InRHICmdList)
+				{
+					GCurrentNumDrawCallsRHIPtr = InNumDrawCallsPtr;
+				}
+			);
+		}
 	}
 }
 
@@ -863,6 +877,17 @@ void FScopedGPUStatEvent::End()
 	if (RHICmdList != nullptr)
 	{
 		FRealtimeGPUProfiler::Get()->PopEvent(*RHICmdList);
+		
+		if (NumDrawCallsPtr != nullptr)
+		{
+			RHICmdList->EnqueueLambda
+			(
+				[](FRHICommandListImmediate& InRHICmdList)
+				{
+					GCurrentNumDrawCallsRHIPtr = &GCurrentNumDrawCallsRHI;
+				}
+			);
+		}
 	}
 }
 #endif // HAS_GPU_STATS
diff --git a/Engine/Source/Runtime/RenderCore/Public/ProfilingDebugging/RealtimeGPUProfiler.h b/Engine/Source/Runtime/RenderCore/Public/ProfilingDebugging/RealtimeGPUProfiler.h
--- a/Engine/Source/Runtime/RenderCore/Public/ProfilingDebugging/RealtimeGPUProfiler.h
+++ b/Engine/Source/Runtime/RenderCore/Public/ProfilingDebugging/RealtimeGPUProfiler.h
@@ -195,24 +195,30 @@ class FScopedGPUStatEvent;
 #if HAS_GPU_STATS
  CSV_DECLARE_CATEGORY_MODULE_EXTERN(RENDERCORE_API,GPU);
  // The DECLARE_GPU_STAT macros both declare and define a stat (for use in a single CPP)
- #define DECLARE_GPU_STAT(StatName) DECLARE_FLOAT_COUNTER_STAT(TEXT(#StatName), Stat_GPU_##StatName, STATGROUP_GPU); CSV_DEFINE_STAT(GPU,StatName);
- #define DECLARE_GPU_STAT_NAMED(StatName, NameString) DECLARE_FLOAT_COUNTER_STAT(NameString, Stat_GPU_##StatName, STATGROUP_GPU); CSV_DEFINE_STAT(GPU,StatName);
+ #define DECLARE_GPU_STAT(StatName) DECLARE_FLOAT_COUNTER_STAT(TEXT(#StatName), Stat_GPU_##StatName, STATGROUP_GPU); CSV_DEFINE_STAT(GPU,StatName); FDrawCallCategoryName DrawcallCountCategory_##StatName;
+ #define DECLARE_GPU_DRAWCALL_STAT(StatName) DECLARE_FLOAT_COUNTER_STAT(TEXT(#StatName), Stat_GPU_##StatName, STATGROUP_GPU); CSV_DEFINE_STAT(GPU,StatName); FDrawCallCategoryName DrawcallCountCategory_##StatName((TCHAR*)TEXT(#StatName));
+ #define DECLARE_GPU_STAT_NAMED(StatName, NameString) DECLARE_FLOAT_COUNTER_STAT(NameString, Stat_GPU_##StatName, STATGROUP_GPU); CSV_DEFINE_STAT(GPU,StatName); FDrawCallCategoryName DrawcallCountCategory_##StatName;
+ #define DECLARE_GPU_DRAWCALL_STAT_NAMED(StatName, NameString) DECLARE_FLOAT_COUNTER_STAT(NameString, Stat_GPU_##StatName, STATGROUP_GPU); CSV_DEFINE_STAT(GPU,StatName); FDrawCallCategoryName DrawcallCountCategory_##StatName((TCHAR*)TEXT(#StatName));
 
  // Extern GPU stats are needed where a stat is used in multiple CPPs. Use the DECLARE_GPU_STAT_NAMED_EXTERN in the header and DEFINE_GPU_STAT in the CPPs
- #define DECLARE_GPU_STAT_NAMED_EXTERN(StatName, NameString) DECLARE_FLOAT_COUNTER_STAT_EXTERN(NameString, Stat_GPU_##StatName, STATGROUP_GPU, ); CSV_DECLARE_STAT_EXTERN(GPU,StatName);
- #define DEFINE_GPU_STAT(StatName) DEFINE_STAT(Stat_GPU_##StatName); CSV_DEFINE_STAT(GPU,StatName);
+ #define DECLARE_GPU_STAT_NAMED_EXTERN(StatName, NameString) DECLARE_FLOAT_COUNTER_STAT_EXTERN(NameString, Stat_GPU_##StatName, STATGROUP_GPU, ); CSV_DECLARE_STAT_EXTERN(GPU,StatName); extern FDrawCallCategoryName DrawcallCountCategory_##StatName;
+ #define DEFINE_GPU_STAT(StatName) DEFINE_STAT(Stat_GPU_##StatName); CSV_DEFINE_STAT(GPU,StatName); FDrawCallCategoryName DrawcallCountCategory_##StatName;
+ #define DEFINE_GPU_DRAWCALL_STAT(StatName) DEFINE_STAT(Stat_GPU_##StatName); CSV_DEFINE_STAT(GPU,StatName); FDrawCallCategoryName DrawcallCountCategory_##StatName((TCHAR*)TEXT(#StatName));
 #if STATS
-  #define SCOPED_GPU_STAT(RHICmdList, StatName) FScopedGPUStatEvent PREPROCESSOR_JOIN(GPUStatEvent_##StatName,__LINE__); PREPROCESSOR_JOIN(GPUStatEvent_##StatName,__LINE__).Begin(RHICmdList, CSV_STAT_FNAME(StatName), GET_STATID( Stat_GPU_##StatName ).GetName() );
+  #define SCOPED_GPU_STAT(RHICmdList, StatName) FScopedGPUStatEvent PREPROCESSOR_JOIN(GPUStatEvent_##StatName,__LINE__); PREPROCESSOR_JOIN(GPUStatEvent_##StatName,__LINE__).Begin(RHICmdList, CSV_STAT_FNAME(StatName), GET_STATID( Stat_GPU_##StatName ).GetName(), &DrawcallCountCategory_##StatName.Counter);
  #else
-  #define SCOPED_GPU_STAT(RHICmdList, StatName) FScopedGPUStatEvent PREPROCESSOR_JOIN(GPUStatEvent_##StatName,__LINE__); PREPROCESSOR_JOIN(GPUStatEvent_##StatName,__LINE__).Begin(RHICmdList, CSV_STAT_FNAME(StatName), FName() );
- #endif
+  #define SCOPED_GPU_STAT(RHICmdList, StatName) FScopedGPUStatEvent PREPROCESSOR_JOIN(GPUStatEvent_##StatName,__LINE__); PREPROCESSOR_JOIN(GPUStatEvent_##StatName,__LINE__).Begin(RHICmdList, CSV_STAT_FNAME(StatName), FName(), &DrawcallCountCategory_##StatName.Counter );
+#endif
  #define GPU_STATS_BEGINFRAME(RHICmdList) FRealtimeGPUProfiler::Get()->BeginFrame(RHICmdList);
  #define GPU_STATS_ENDFRAME(RHICmdList) FRealtimeGPUProfiler::Get()->EndFrame(RHICmdList);
 #else
  #define DECLARE_GPU_STAT(StatName)
+ #define DECLARE_GPU_DRAWCALL_STAT(StatName)
  #define DECLARE_GPU_STAT_NAMED(StatName, NameString)
+ #define DECLARE_GPU_DRAWCALL_STAT_NAMED(StatName, NameString)
  #define DECLARE_GPU_STAT_NAMED_EXTERN(StatName, NameString)
  #define DEFINE_GPU_STAT(StatName)
+ #define DEFINE_GPU_DRAWCALL_STAT(StatName)
  #define SCOPED_GPU_STAT(RHICmdList, StatName) 
  #define GPU_STATS_BEGINFRAME(RHICmdList) 
  #define GPU_STATS_ENDFRAME(RHICmdList) 
@@ -281,11 +287,14 @@ class FScopedGPUStatEvent
 	/** The stat event used to record timings */
 	FRealtimeGPUProfilerEvent* RealtimeGPUProfilerEvent;
 
+	int32* NumDrawCallsPtr;
+
 public:
 	/** Default constructor, initializing all member variables. */
 	FORCEINLINE FScopedGPUStatEvent()
 		: RHICmdList(nullptr)
 		, RealtimeGPUProfilerEvent(nullptr)
+		, NumDrawCallsPtr(nullptr)
 	{}
 
 	/**
@@ -302,7 +311,7 @@ public:
 	/**
 	* Start/Stop functions for timer stats
 	*/
-	RENDERCORE_API void Begin(FRHICommandList& InRHICmdList, const FName& Name, const FName& StatName );
+	RENDERCORE_API void Begin(FRHICommandList& InRHICmdList, const FName& Name, const FName& StatName, int32* InNumDrawCallsPtr);
 	RENDERCORE_API void End();
 };
 #endif // HAS_GPU_STATS
diff --git a/Engine/Source/Runtime/Renderer/Private/BasePassRendering.cpp b/Engine/Source/Runtime/Renderer/Private/BasePassRendering.cpp
--- a/Engine/Source/Runtime/Renderer/Private/BasePassRendering.cpp
+++ b/Engine/Source/Runtime/Renderer/Private/BasePassRendering.cpp
@@ -127,7 +127,7 @@ IMPLEMENT_BASEPASS_LIGHTMAPPED_SHADER_TYPE( TUniformLightMapPolicy<LMP_DISTANCE_
 
 IMPLEMENT_MATERIAL_SHADER_TYPE(, F128BitRTBasePassPS, TEXT("/Engine/Private/BasePassPixelShader.usf"), TEXT("MainPS"), SF_Pixel);
 
-DECLARE_GPU_STAT(Basepass);
+DECLARE_GPU_DRAWCALL_STAT(Basepass);
 
 void SetTranslucentRenderState(FMeshPassProcessorRenderState& DrawRenderState, const FMaterial& Material, const EShaderPlatform Platform, ETranslucencyPass::Type InTranslucencyPassType)
 {
diff --git a/Engine/Source/Runtime/Renderer/Private/DeferredShadingRenderer.cpp b/Engine/Source/Runtime/Renderer/Private/DeferredShadingRenderer.cpp
--- a/Engine/Source/Runtime/Renderer/Private/DeferredShadingRenderer.cpp
+++ b/Engine/Source/Runtime/Renderer/Private/DeferredShadingRenderer.cpp
@@ -226,9 +226,9 @@ DECLARE_GPU_STAT(PostRenderOpsFX);
 DECLARE_GPU_STAT(GPUSceneUpdate);
 DECLARE_GPU_STAT(HZB);
 DECLARE_GPU_STAT_NAMED(Unaccounted, TEXT("[unaccounted]"));
-DECLARE_GPU_STAT(WaterRendering);
+DECLARE_GPU_DRAWCALL_STAT(WaterRendering);
 DECLARE_GPU_STAT(HairRendering);
-DECLARE_GPU_STAT(VirtualTextureUpdate);
+DECLARE_GPU_DRAWCALL_STAT(VirtualTextureUpdate);
 DECLARE_GPU_STAT(UploadDynamicBuffers);
 DECLARE_GPU_STAT(PostOpaqueExtensions);
 
diff --git a/Engine/Source/Runtime/Renderer/Private/DepthRendering.cpp b/Engine/Source/Runtime/Renderer/Private/DepthRendering.cpp
--- a/Engine/Source/Runtime/Renderer/Private/DepthRendering.cpp
+++ b/Engine/Source/Runtime/Renderer/Private/DepthRendering.cpp
@@ -71,7 +71,7 @@ const TCHAR* GetDepthDrawingModeString(EDepthDrawingMode Mode)
 	return TEXT("");
 }
 
-DECLARE_GPU_STAT(Prepass);
+DECLARE_GPU_DRAWCALL_STAT(Prepass);
 
 IMPLEMENT_MATERIAL_SHADER_TYPE(template<>,TDepthOnlyVS<true>,TEXT("/Engine/Private/PositionOnlyDepthVertexShader.usf"),TEXT("Main"),SF_Vertex);
 IMPLEMENT_MATERIAL_SHADER_TYPE(template<>,TDepthOnlyVS<false>,TEXT("/Engine/Private/DepthOnlyVertexShader.usf"),TEXT("Main"),SF_Vertex);
diff --git a/Engine/Source/Runtime/Renderer/Private/ShadowDepthRendering.cpp b/Engine/Source/Runtime/Renderer/Private/ShadowDepthRendering.cpp
--- a/Engine/Source/Runtime/Renderer/Private/ShadowDepthRendering.cpp
+++ b/Engine/Source/Runtime/Renderer/Private/ShadowDepthRendering.cpp
@@ -40,7 +40,7 @@
 #include "VisualizeTexture.h"
 #include "GPUScene.h"
 
-DECLARE_GPU_STAT_NAMED(ShadowDepths, TEXT("Shadow Depths"));
+DECLARE_GPU_DRAWCALL_STAT_NAMED(ShadowDepths, TEXT("Shadow Depths"));
 
 IMPLEMENT_GLOBAL_SHADER_PARAMETER_STRUCT(FShadowDepthPassUniformParameters, "ShadowDepthPass");
 IMPLEMENT_GLOBAL_SHADER_PARAMETER_STRUCT(FMobileShadowDepthPassUniformParameters, "MobileShadowDepthPass");
diff --git a/Engine/Source/Runtime/Renderer/Private/ShadowRendering.cpp b/Engine/Source/Runtime/Renderer/Private/ShadowRendering.cpp
--- a/Engine/Source/Runtime/Renderer/Private/ShadowRendering.cpp
+++ b/Engine/Source/Runtime/Renderer/Private/ShadowRendering.cpp
@@ -167,7 +167,7 @@ static TAutoConsoleVariable<float> CVarShadowMaxSlopeScaleDepthBias(
 	TEXT("Higher values give better self-shadowing, but increase self-shadowing artifacts"),
 	ECVF_RenderThreadSafe);
 
-DEFINE_GPU_STAT(ShadowProjection);
+DEFINE_GPU_DRAWCALL_STAT(ShadowProjection);
 
 // 0:off, 1:low, 2:med, 3:high, 4:very high, 5:max
 uint32 GetShadowQuality()
diff --git a/Engine/Source/Runtime/Renderer/Private/TranslucentRendering.cpp b/Engine/Source/Runtime/Renderer/Private/TranslucentRendering.cpp
--- a/Engine/Source/Runtime/Renderer/Private/TranslucentRendering.cpp
+++ b/Engine/Source/Runtime/Renderer/Private/TranslucentRendering.cpp
@@ -21,7 +21,7 @@ DECLARE_CYCLE_STAT(TEXT("TranslucencyTimestampQuery Wait"), STAT_TranslucencyTim
 
 DECLARE_FLOAT_COUNTER_STAT(TEXT("Translucency GPU Time (MS)"), STAT_TranslucencyGPU, STATGROUP_SceneRendering);
 
-DECLARE_GPU_STAT(Translucency);
+DECLARE_GPU_DRAWCALL_STAT(Translucency);
 
 static TAutoConsoleVariable<float> CVarSeparateTranslucencyScreenPercentage(
 	TEXT("r.SeparateTranslucencyScreenPercentage"),
diff --git a/Engine/Source/Runtime/SlateRHIRenderer/Private/SlateRHIRenderer.cpp b/Engine/Source/Runtime/SlateRHIRenderer/Private/SlateRHIRenderer.cpp
--- a/Engine/Source/Runtime/SlateRHIRenderer/Private/SlateRHIRenderer.cpp
+++ b/Engine/Source/Runtime/SlateRHIRenderer/Private/SlateRHIRenderer.cpp
@@ -37,7 +37,7 @@ DECLARE_CYCLE_STAT(TEXT("Slate RT: Rendering"), STAT_SlateRenderingRTTime, STATG
 
 DECLARE_CYCLE_STAT(TEXT("Slate RT: Draw Batches"), STAT_SlateRTDrawBatches, STATGROUP_Slate);
 
-DECLARE_GPU_STAT_NAMED(SlateUI, TEXT("Slate UI"));
+DECLARE_GPU_DRAWCALL_STAT_NAMED(SlateUI, TEXT("Slate UI"));
 
 // Defines the maximum size that a slate viewport will create
 #define MIN_VIEWPORT_SIZE 8
