Fixed CIS

#rnx

diff --git a/Engine/Source/Runtime/RenderCore/Private/ProfilingDebugging/RealtimeGPUProfiler.cpp b/Engine/Source/Runtime/RenderCore/Private/ProfilingDebugging/RealtimeGPUProfiler.cpp
--- a/Engine/Source/Runtime/RenderCore/Private/ProfilingDebugging/RealtimeGPUProfiler.cpp
+++ b/Engine/Source/Runtime/RenderCore/Private/ProfilingDebugging/RealtimeGPUProfiler.cpp
@@ -858,7 +858,7 @@ void FScopedGPUStatEvent::Begin(FRHICommandList& InRHICmdList, const FName& Name
 
 			RHICmdList->EnqueueLambda
 			(
-				[InNumDrawCallsPtr](FRHICommandListImmediate& InRHICmdList)
+				[InNumDrawCallsPtr](FRHICommandListImmediate& CmdList_)
 				{
 					GCurrentNumDrawCallsRHIPtr = InNumDrawCallsPtr;
 				}
