Fix missing UAV binding on pixel shader when using render graph resources.

diff --git a/Engine/Source/Runtime/RenderCore/Public/ShaderParameterStruct.h b/Engine/Source/Runtime/RenderCore/Public/ShaderParameterStruct.h
--- a/Engine/Source/Runtime/RenderCore/Public/ShaderParameterStruct.h
+++ b/Engine/Source/Runtime/RenderCore/Public/ShaderParameterStruct.h
@@ -102,7 +102,27 @@ inline void SetShaderUAVs(TRHICmdList& RHICmdList, const TShaderRef<TShaderClass
 template<typename TRHICmdList, typename TShaderClass>
 inline void SetShaderUAVs(TRHICmdList& RHICmdList, const TShaderRef<TShaderClass>& Shader, FRHIPixelShader* ShadeRHI, const typename TShaderClass::FParameters& Parameters)
 {
-	// Pixelshader UAVs are bound together with rendertargets using BeginRenderPass
+	const FShaderParameterBindings& Bindings = Shader->Bindings;
+
+	const typename TShaderClass::FParameters* ParametersPtr = &Parameters;
+	const uint8* Base = reinterpret_cast<const uint8*>(ParametersPtr);
+
+	// UAVs
+	for (const FShaderParameterBindings::FResourceParameter& ParameterBinding : Bindings.UAVs)
+	{
+		FRHIUnorderedAccessView* ShaderParameterRef = *(FRHIUnorderedAccessView**)(Base + ParameterBinding.ByteOffset);
+		RHICmdList.SetUAVParameter(ShadeRHI, ParameterBinding.BaseIndex, ShaderParameterRef);
+	}
+
+	// Graph UAVs
+	for (const FShaderParameterBindings::FResourceParameter& ParameterBinding : Bindings.GraphUAVs)
+	{
+		auto GraphUAV = *reinterpret_cast<FRDGUnorderedAccessView* const*>(Base + ParameterBinding.ByteOffset);	
+
+		checkSlow(GraphUAV);
+		GraphUAV->MarkResourceAsUsed();
+		RHICmdList.SetUAVParameter(ShadeRHI, ParameterBinding.BaseIndex, GraphUAV->GetRHI());
+	}
 }
 
 template<typename TRHICmdList, typename TShaderClass>
