Reduce FShader object size in cooked games from 264 bytes to 160 bytes (saves ~7MB in the lobby)

diff --git a/Engine/Shaders/Public/ShaderVersion.ush b/Engine/Shaders/Public/ShaderVersion.ush
--- a/Engine/Shaders/Public/ShaderVersion.ush
+++ b/Engine/Shaders/Public/ShaderVersion.ush
@@ -3,4 +3,4 @@
 // in Platform.ush (which should be included in any shader) it allows to invalidate the shader DDC.
 // 
 // If you are merging streams and there is a conflict with this GUID you should make a new GUID rather than taking one or the other.
-// GUID = BD84E7E6-F382-413D-B218-71A7172EE73D
\ No newline at end of file
+// GUID = 2C9048786E8942F699D8C862E803EB60
diff --git a/Engine/Source/Runtime/RenderCore/Private/RenderGraphUtils.cpp b/Engine/Source/Runtime/RenderCore/Private/RenderGraphUtils.cpp
--- a/Engine/Source/Runtime/RenderCore/Private/RenderGraphUtils.cpp
+++ b/Engine/Source/Runtime/RenderCore/Private/RenderGraphUtils.cpp
@@ -22,57 +22,47 @@ void ClearUnusedGraphResourcesImpl(
 	void* InoutParameters,
 	std::initializer_list<FRDGResourceRef> ExcludeList)
 {
-	int32 GraphTextureId = 0;
-	int32 GraphSRVId = 0;
-	int32 GraphUAVId = 0;
-
+	int32 UsedResourceIdx = 0;
 	uint8* Base = reinterpret_cast<uint8*>(InoutParameters);
 
 	for (int32 ResourceIndex = 0, Num = ParametersMetadata->GetLayout().Resources.Num(); ResourceIndex < Num; ResourceIndex++)
 	{
 		EUniformBufferBaseType Type = ParametersMetadata->GetLayout().Resources[ResourceIndex].MemberType;
 		uint16 ByteOffset = ParametersMetadata->GetLayout().Resources[ResourceIndex].MemberOffset;
+		bool bRDGResource = Type == UBMT_RDG_TEXTURE ||
+							Type == UBMT_RDG_TEXTURE_SRV || 
+							Type == UBMT_RDG_BUFFER_SRV ||
+							Type == UBMT_RDG_TEXTURE_UAV || 
+							Type == UBMT_RDG_BUFFER_UAV;
 
-		if (Type == UBMT_RDG_TEXTURE)
-		{
-			if (GraphTextureId < ShaderBindings.GraphTextures.Num() && ByteOffset == ShaderBindings.GraphTextures[GraphTextureId].ByteOffset)
-			{
-				GraphTextureId++;
-				continue;
-			}
-		}
-		else if (Type == UBMT_RDG_TEXTURE_SRV || Type == UBMT_RDG_BUFFER_SRV)
+		if (UsedResourceIdx < ShaderBindings.ResourceParameters.Num())
 		{
-			if (GraphSRVId < ShaderBindings.GraphSRVs.Num() && ByteOffset == ShaderBindings.GraphSRVs[GraphSRVId].ByteOffset)
+			const FShaderParameterBindings::FResourceParameter& ResourceParameter = ShaderBindings.ResourceParameters[UsedResourceIdx];
+			if (ByteOffset == ResourceParameter.ByteOffset)
 			{
-				GraphSRVId++;
+				// check next one
+				UsedResourceIdx++;
 				continue;
 			}
-		}
-		else if (Type == UBMT_RDG_TEXTURE_UAV || Type == UBMT_RDG_BUFFER_UAV)
-		{
-			if (GraphUAVId < ShaderBindings.GraphUAVs.Num() && ByteOffset == ShaderBindings.GraphUAVs[GraphUAVId].ByteOffset)
-			{
-				GraphUAVId++;
-				continue;
-			}
-		}
-		else
-		{
-			continue;
+
+			// make sure we dont jump over used resources
+			check(ByteOffset <= ResourceParameter.ByteOffset);
 		}
 
-		for (FRDGResourceRef ExcludeResource : ExcludeList)
+		if (bRDGResource)
 		{
-			auto Resource = *reinterpret_cast<const FRDGResource* const*>(Base + ByteOffset);
-			if (Resource == ExcludeResource)
+			for( FRDGResourceRef ExcludeResource : ExcludeList )
 			{
-				continue;
+				auto Resource = *reinterpret_cast<const FRDGResource* const*>(Base + ByteOffset);
+				if( Resource == ExcludeResource )
+				{
+					continue;
+				}
 			}
-		}
 
-		void** ResourcePointerAddress = reinterpret_cast<void**>(Base + ByteOffset);
-		*ResourcePointerAddress = nullptr;
+			void** ResourcePointerAddress = reinterpret_cast<void**>(Base + ByteOffset);
+			*ResourcePointerAddress = nullptr;
+		}
 	}
 }
 
@@ -82,12 +72,8 @@ void ClearUnusedGraphResourcesImpl(
 	void* InoutParameters,
 	std::initializer_list<FRDGResourceRef> ExcludeList)
 {
-	TArray<int32, TInlineAllocator<SF_NumFrequencies>> GraphTextureIds;
-	TArray<int32, TInlineAllocator<SF_NumFrequencies>> GraphSRVIds;
-	TArray<int32, TInlineAllocator<SF_NumFrequencies>> GraphUAVIds;
-	GraphTextureIds.SetNumZeroed(ShaderBindingsList.Num());
-	GraphSRVIds.SetNumZeroed(ShaderBindingsList.Num());
-	GraphUAVIds.SetNumZeroed(ShaderBindingsList.Num());
+	TArray<int32, TInlineAllocator<SF_NumFrequencies>> GraphResourceIds;
+	GraphResourceIds.SetNumZeroed(ShaderBindingsList.Num());
 
 	uint8* Base = reinterpret_cast<uint8*>(InoutParameters);
 
@@ -95,75 +81,37 @@ void ClearUnusedGraphResourcesImpl(
 	{
 		EUniformBufferBaseType Type = ParametersMetadata->GetLayout().Resources[ResourceIndex].MemberType;
 		uint16 ByteOffset = ParametersMetadata->GetLayout().Resources[ResourceIndex].MemberOffset;
-		bool bResourceIsUsed = false;
+		bool bRDGResource = Type == UBMT_RDG_TEXTURE ||
+							Type == UBMT_RDG_TEXTURE_SRV || 
+							Type == UBMT_RDG_BUFFER_SRV ||
+							Type == UBMT_RDG_TEXTURE_UAV || 
+							Type == UBMT_RDG_BUFFER_UAV;
 
-		if (Type == UBMT_RDG_TEXTURE)
-		{
-			for (int32 Index = 0; Index < ShaderBindingsList.Num(); ++Index)
-			{
-				const FShaderParameterBindings& ShaderBindings = *ShaderBindingsList[Index];
-				int32& GraphTextureId = GraphTextureIds[Index];
-
-				if (GraphTextureId < ShaderBindings.GraphTextures.Num() && ByteOffset == ShaderBindings.GraphTextures[GraphTextureId].ByteOffset)
-				{
-					GraphTextureId++;
-					bResourceIsUsed = true;
-					break;
-				}
-			}
-		}
-		else if (Type == UBMT_RDG_TEXTURE_SRV || Type == UBMT_RDG_BUFFER_SRV)
+		bool bResourceIsUsed = false;
+		for (int32 Index = 0; Index < ShaderBindingsList.Num(); ++Index)
 		{
-			for (int32 Index = 0; Index < ShaderBindingsList.Num(); ++Index)
+			const FShaderParameterBindings& ShaderBindings = *ShaderBindingsList[Index];
+			int32& GraphResourceId = GraphResourceIds[Index];
+			for (; GraphResourceId < ShaderBindings.ResourceParameters.Num() && ByteOffset < ShaderBindings.ResourceParameters[GraphResourceId].ByteOffset; ++GraphResourceId)
 			{
-				const FShaderParameterBindings& ShaderBindings = *ShaderBindingsList[Index];
-				int32& GraphSRVId = GraphSRVIds[Index];
-
-				if (GraphSRVId < ShaderBindings.GraphSRVs.Num() && ByteOffset == ShaderBindings.GraphSRVs[GraphSRVId].ByteOffset)
-				{
-					GraphSRVId++;
-					bResourceIsUsed = true;
-					break;
-				}
 			}
+			bResourceIsUsed|= GraphResourceId < ShaderBindings.ResourceParameters.Num() && ByteOffset == ShaderBindings.ResourceParameters[GraphResourceId].ByteOffset; // check the resouce type as well?
 		}
-		else if (Type == UBMT_RDG_TEXTURE_UAV || Type == UBMT_RDG_BUFFER_UAV)
+		
+		if (bRDGResource && !bResourceIsUsed)
 		{
-			for (int32 Index = 0; Index < ShaderBindingsList.Num(); ++Index)
+			for (FRDGResourceRef ExcludeResource : ExcludeList)
 			{
-				const FShaderParameterBindings& ShaderBindings = *ShaderBindingsList[Index];
-				int32& GraphUAVId = GraphUAVIds[Index];
-
-				if (GraphUAVId < ShaderBindings.GraphUAVs.Num() && ByteOffset == ShaderBindings.GraphUAVs[GraphUAVId].ByteOffset)
+				auto Resource = *reinterpret_cast<const FRDGResource* const*>(Base + ByteOffset);
+				if (Resource == ExcludeResource)
 				{
-					GraphUAVId++;
-					bResourceIsUsed = true;
-					break;
+					continue;
 				}
 			}
-		}
-		else
-		{
-			// Not a resource we care about.
-			continue;
-		}
 
-		if (bResourceIsUsed)
-		{
-			continue;
+			void** ResourcePointerAddress = reinterpret_cast<void**>(Base + ByteOffset);
+			*ResourcePointerAddress = nullptr;
 		}
-
-		for (FRDGResourceRef ExcludeResource : ExcludeList)
-		{
-			auto Resource = *reinterpret_cast<const FRDGResource* const*>(Base + ByteOffset);
-			if (Resource == ExcludeResource)
-			{
-				continue;
-			}
-		}
-
-		void** ResourcePointerAddress = reinterpret_cast<void**>(Base + ByteOffset);
-		*ResourcePointerAddress = nullptr;
 	}
 }
 
diff --git a/Engine/Source/Runtime/RenderCore/Private/ShaderParameterStruct.cpp b/Engine/Source/Runtime/RenderCore/Private/ShaderParameterStruct.cpp
--- a/Engine/Source/Runtime/RenderCore/Private/ShaderParameterStruct.cpp
+++ b/Engine/Source/Runtime/RenderCore/Private/ShaderParameterStruct.cpp
@@ -195,9 +195,12 @@ struct FShaderParameterStructBindingContext
 				}
 				else if (bIsRHIResource || bIsRDGResource)
 				{
+					checkf(BaseIndex < 256, TEXT("BaseIndex does not fit into uint8. Change FResourceParameter::BaseIndex type to uint16"));
+										
 					FShaderParameterBindings::FResourceParameter Parameter;
-					Parameter.BaseIndex = BaseIndex;
+					Parameter.BaseIndex = (uint8)BaseIndex;
 					Parameter.ByteOffset = ByteOffset + ArrayElementId * SHADER_PARAMETER_POINTER_ALIGNMENT;
+					Parameter.BaseType = BaseType;
 
 					if (BoundSize != 1)
 					{
@@ -215,20 +218,7 @@ struct FShaderParameterStructBindingContext
 						}
 					}
 
-					if (BaseType == UBMT_TEXTURE)
-						Bindings->Textures.Add(Parameter);
-					else if (BaseType == UBMT_SRV)
-						Bindings->SRVs.Add(Parameter);
-					else if (BaseType == UBMT_UAV)
-						Bindings->UAVs.Add(Parameter);
-					else if (BaseType == UBMT_SAMPLER)
-						Bindings->Samplers.Add(Parameter);
-					else if (BaseType == UBMT_RDG_TEXTURE)
-						Bindings->GraphTextures.Add(Parameter);
-					else if (BaseType == UBMT_RDG_TEXTURE_SRV || BaseType == UBMT_RDG_BUFFER_SRV)
-						Bindings->GraphSRVs.Add(Parameter);
-					else // if (BaseType == UBMT_RDG_TEXTURE_UAV || BaseType == UBMT_RDG_BUFFER_UAV)
-						Bindings->GraphUAVs.Add(Parameter);
+					Bindings->ResourceParameters.Add(Parameter);
 				}
 				else
 				{
@@ -450,74 +440,57 @@ void ValidateShaderParameters(const TShaderRef<FShader>& Shader, const FShaderPa
 	const TCHAR* ShaderClassName = Shader.GetType()->GetName();
 	const TCHAR* ShaderParemeterStructName = ParametersMetadata->GetStructTypeName();
 
-	// Textures
-	for (const FShaderParameterBindings::FResourceParameter& ParameterBinding : Bindings.Textures)
-	{
-		FRHITexture* ShaderParameterRef = *(FRHITexture**)(Base + ParameterBinding.ByteOffset);
-		if (!ShaderParameterRef)
-		{
-			EmitNullShaderParameterFatalError(Shader, ParametersMetadata, ParameterBinding.ByteOffset);
-		}
-	}
-
-	// SRVs
-	for (const FShaderParameterBindings::FResourceParameter& ParameterBinding : Bindings.SRVs)
-	{
-		FRHIShaderResourceView* ShaderParameterRef = *(FRHIShaderResourceView**)(Base + ParameterBinding.ByteOffset);
-		if (!ShaderParameterRef)
-		{
-			EmitNullShaderParameterFatalError(Shader, ParametersMetadata, ParameterBinding.ByteOffset);
-		}
-	}
-
-	// Samplers
-	for (const FShaderParameterBindings::FResourceParameter& ParameterBinding : Bindings.Samplers)
-	{
-		FRHISamplerState* ShaderParameterRef = *(FRHISamplerState**)(Base + ParameterBinding.ByteOffset);
-		if (!ShaderParameterRef)
-		{
-			EmitNullShaderParameterFatalError(Shader, ParametersMetadata, ParameterBinding.ByteOffset);
-		}
-	}
-
-	// Graph Textures
-	for (const FShaderParameterBindings::FResourceParameter& ParameterBinding : Bindings.GraphTextures)
-	{
-		auto GraphTexture = *reinterpret_cast<const FRDGTexture* const*>(Base + ParameterBinding.ByteOffset);
-		if (!GraphTexture)
-		{
-			EmitNullShaderParameterFatalError(Shader, ParametersMetadata, ParameterBinding.ByteOffset);
-		}
-		else if ((GraphTexture->Desc.TargetableFlags & TexCreate_ShaderResource) == 0)
-		{
-			FString MemberName = ParametersMetadata->GetFullMemberCodeName(ParameterBinding.ByteOffset);
-
-			UE_LOG(LogShaders, Error,
-				TEXT("Attempting to set shader %s parameter %s::%s with the RDG texture %s which was not created with TexCreate_ShaderResource"),
-				ShaderClassName, ShaderParemeterStructName, *MemberName, GraphTexture->Name);
-		}
-	}
-
-	// Graph SRVs
-	for (const FShaderParameterBindings::FResourceParameter& ParameterBinding : Bindings.GraphSRVs)
+	for (const FShaderParameterBindings::FResourceParameter& ParameterBinding : Bindings.ResourceParameters)
 	{
-		auto GraphSRV = *reinterpret_cast<const FRDGTextureSRV* const*>(Base + ParameterBinding.ByteOffset);
-		if (!GraphSRV)
+		EUniformBufferBaseType BaseType = (EUniformBufferBaseType)ParameterBinding.BaseType;
+		switch (BaseType)
 		{
-			EmitNullShaderParameterFatalError(Shader, ParametersMetadata, ParameterBinding.ByteOffset);
-		}
-	}
+			case UBMT_TEXTURE:
+			case UBMT_SRV:
+			case UBMT_UAV:
+			case UBMT_SAMPLER:
+			{
+				FRHIResource* ShaderParameterRef = *(FRHIResource**)(Base + ParameterBinding.ByteOffset);
+				if (!ShaderParameterRef)
+				{
+					EmitNullShaderParameterFatalError(Shader, ParametersMetadata, ParameterBinding.ByteOffset);
+				}
+			}
+			break;
+			case UBMT_RDG_TEXTURE:
+			{
+				const FRDGTexture* GraphTexture = *reinterpret_cast<const FRDGTexture* const*>(Base + ParameterBinding.ByteOffset);
+				if (!GraphTexture)
+				{
+					EmitNullShaderParameterFatalError(Shader, ParametersMetadata, ParameterBinding.ByteOffset);
+				}
+				else if ((GraphTexture->Desc.TargetableFlags & TexCreate_ShaderResource) == 0)
+				{
+					FString MemberName = ParametersMetadata->GetFullMemberCodeName(ParameterBinding.ByteOffset);
 
-	// Graph UAVs for compute shaders	
-	for (const FShaderParameterBindings::FResourceParameter& ParameterBinding : Bindings.GraphUAVs)
-	{
-		auto GraphUAV = *reinterpret_cast<const FRDGTextureUAV* const*>(Base + ParameterBinding.ByteOffset);
-		if (!GraphUAV)
-		{
-			EmitNullShaderParameterFatalError(Shader, ParametersMetadata, ParameterBinding.ByteOffset);
+					UE_LOG(LogShaders, Error,
+						TEXT("Attempting to set shader %s parameter %s::%s with the RDG texture %s which was not created with TexCreate_ShaderResource"),
+						ShaderClassName, ShaderParemeterStructName, *MemberName, GraphTexture->Name);
+				}
+			}
+			break;
+			case UBMT_RDG_TEXTURE_SRV:
+			case UBMT_RDG_TEXTURE_UAV:
+			case UBMT_RDG_BUFFER_SRV:
+			case UBMT_RDG_BUFFER_UAV:
+			{
+				const FRDGResource* GraphResource = *reinterpret_cast<const FRDGResource* const*>(Base + ParameterBinding.ByteOffset);
+				if (!GraphResource)
+				{
+					EmitNullShaderParameterFatalError(Shader, ParametersMetadata, ParameterBinding.ByteOffset);
+				}
+			}
+			break;
+			default:
+				break;
 		}
 	}
-
+	
 	// Reference structures
 	for (const FShaderParameterBindings::FParameterStructReference& ParameterBinding : Bindings.ParameterReferences)
 	{
diff --git a/Engine/Source/Runtime/RenderCore/Public/Shader.h b/Engine/Source/Runtime/RenderCore/Public/Shader.h
--- a/Engine/Source/Runtime/RenderCore/Public/Shader.h
+++ b/Engine/Source/Runtime/RenderCore/Public/Shader.h
@@ -563,10 +563,10 @@ public:
 
 	struct FResourceParameter
 	{
-		DECLARE_INLINE_TYPE_LAYOUT(FResourceParameter, NonVirtual);
-
-		LAYOUT_FIELD(uint16, BaseIndex);
+		DECLARE_INLINE_TYPE_LAYOUT(FResourceParameter, NonVirtual);		
 		LAYOUT_FIELD(uint16, ByteOffset);
+		LAYOUT_FIELD(uint8, BaseIndex);
+		LAYOUT_FIELD(EUniformBufferBaseType, BaseType);
 
 	};
 
@@ -587,13 +587,7 @@ public:
 	void BindForRootShaderParameters(const FShader* Shader, const FShaderParameterMap& ParameterMaps);
 
 	LAYOUT_FIELD(TMemoryImageArray<FParameter>, Parameters);
-	LAYOUT_FIELD(TMemoryImageArray<FResourceParameter>, Textures);
-	LAYOUT_FIELD(TMemoryImageArray<FResourceParameter>, SRVs);
-	LAYOUT_FIELD(TMemoryImageArray<FResourceParameter>, UAVs);
-	LAYOUT_FIELD(TMemoryImageArray<FResourceParameter>, Samplers);
-	LAYOUT_FIELD(TMemoryImageArray<FResourceParameter>, GraphTextures);
-	LAYOUT_FIELD(TMemoryImageArray<FResourceParameter>, GraphSRVs);
-	LAYOUT_FIELD(TMemoryImageArray<FResourceParameter>, GraphUAVs);
+	LAYOUT_FIELD(TMemoryImageArray<FResourceParameter>, ResourceParameters);
 	LAYOUT_FIELD(TMemoryImageArray<FParameterStructReference>, ParameterReferences);
 
 	// Hash of the shader parameter structure when doing the binding.
diff --git a/Engine/Source/Runtime/RenderCore/Public/ShaderParameterStruct.h b/Engine/Source/Runtime/RenderCore/Public/ShaderParameterStruct.h
--- a/Engine/Source/Runtime/RenderCore/Public/ShaderParameterStruct.h
+++ b/Engine/Source/Runtime/RenderCore/Public/ShaderParameterStruct.h
@@ -89,35 +89,24 @@ FORCEINLINE void ValidateShaderParameters(const TShaderRef<TShaderClass>& Shader
 	return ValidateShaderParameters(Shader, TShaderClass::FParameters::FTypeInfo::GetStructMetadata(), ParameterPtr);
 }
 
-
 /** Set compute shader UAVs. */
 template<typename TRHICmdList, typename TShaderClass, typename TShaderRHI>
-inline void SetShaderUAVs(TRHICmdList& RHICmdList, const TShaderRef<TShaderClass>& Shader, TShaderRHI* ShadeRHI, const typename TShaderClass::FParameters& Parameters)
+inline void SetShaderUAV(TRHICmdList& RHICmdList, const TShaderRef<TShaderClass>& Shader, TShaderRHI* ShadeRHI, const uint8* Base, const FShaderParameterBindings::FResourceParameter& ParameterBinding)
 {
-	checkf(
-		Shader->Bindings.UAVs.Num() == 0 && Shader->Bindings.GraphUAVs.Num() == 0,
-		TEXT("TShaderRHI Can't have compute shader to be set. UAVs are not supported on vertex, tessellation and geometry shaders."));
+	checkf(false, TEXT("TShaderRHI Can't have compute shader to be set. UAVs are not supported on vertex, tessellation and geometry shaders."));
 }
 
 template<typename TRHICmdList, typename TShaderClass>
-inline void SetShaderUAVs(TRHICmdList& RHICmdList, const TShaderRef<TShaderClass>& Shader, FRHIPixelShader* ShadeRHI, const typename TShaderClass::FParameters& Parameters)
+inline void SetShaderUAV(TRHICmdList& RHICmdList, const TShaderRef<TShaderClass>& Shader, FRHIPixelShader* ShadeRHI, const uint8* Base, const FShaderParameterBindings::FResourceParameter& ParameterBinding)
 {
-	const FShaderParameterBindings& Bindings = Shader->Bindings;
-
-	const typename TShaderClass::FParameters* ParametersPtr = &Parameters;
-	const uint8* Base = reinterpret_cast<const uint8*>(ParametersPtr);
-
-	// UAVs
-	for (const FShaderParameterBindings::FResourceParameter& ParameterBinding : Bindings.UAVs)
+	if (ParameterBinding.BaseType == UBMT_UAV)
 	{
 		FRHIUnorderedAccessView* ShaderParameterRef = *(FRHIUnorderedAccessView**)(Base + ParameterBinding.ByteOffset);
 		RHICmdList.SetUAVParameter(ShadeRHI, ParameterBinding.BaseIndex, ShaderParameterRef);
 	}
-
-	// Graph UAVs
-	for (const FShaderParameterBindings::FResourceParameter& ParameterBinding : Bindings.GraphUAVs)
+	else if (ParameterBinding.BaseType == UBMT_RDG_TEXTURE_UAV || ParameterBinding.BaseType == UBMT_RDG_BUFFER_UAV)
 	{
-		auto GraphUAV = *reinterpret_cast<FRDGUnorderedAccessView* const*>(Base + ParameterBinding.ByteOffset);	
+		auto GraphUAV = *reinterpret_cast<FRDGUnorderedAccessView* const*>(Base + ParameterBinding.ByteOffset);
 
 		checkSlow(GraphUAV);
 		GraphUAV->MarkResourceAsUsed();
@@ -126,22 +115,14 @@ inline void SetShaderUAVs(TRHICmdList& RHICmdList, const TShaderRef<TShaderClass
 }
 
 template<typename TRHICmdList, typename TShaderClass>
-inline void SetShaderUAVs(TRHICmdList& RHICmdList, const TShaderRef<TShaderClass>& Shader, FRHIComputeShader* ShadeRHI, const typename TShaderClass::FParameters& Parameters)
+inline void SetShaderUAV(TRHICmdList& RHICmdList, const TShaderRef<TShaderClass>& Shader, FRHIComputeShader* ShadeRHI, const uint8* Base, const FShaderParameterBindings::FResourceParameter& ParameterBinding)
 {
-	const FShaderParameterBindings& Bindings = Shader->Bindings;
-
-	const typename TShaderClass::FParameters* ParametersPtr = &Parameters;
-	const uint8* Base = reinterpret_cast<const uint8*>(ParametersPtr);
-
-	// UAVs
-	for (const FShaderParameterBindings::FResourceParameter& ParameterBinding : Bindings.UAVs)
+	if (ParameterBinding.BaseType == UBMT_UAV)
 	{
 		FRHIUnorderedAccessView* ShaderParameterRef = *(FRHIUnorderedAccessView**)(Base + ParameterBinding.ByteOffset);
 		RHICmdList.SetUAVParameter(ShadeRHI, ParameterBinding.BaseIndex, ShaderParameterRef);
 	}
-
-	// Graph UAVs
-	for (const FShaderParameterBindings::FResourceParameter& ParameterBinding : Bindings.GraphUAVs)
+	else if (ParameterBinding.BaseType == UBMT_RDG_TEXTURE_UAV || ParameterBinding.BaseType == UBMT_RDG_BUFFER_UAV)
 	{
 		auto GraphUAV = *reinterpret_cast<FRDGUnorderedAccessView* const*>(Base + ParameterBinding.ByteOffset);
 
@@ -161,14 +142,14 @@ inline void UnsetShaderUAVs(TRHICmdList& RHICmdList, const TShaderRef<TShaderCla
 
 	checkf(Bindings.RootParameterBufferIndex == FShaderParameterBindings::kInvalidBufferIndex, TEXT("Can't use UnsetShaderUAVs() for root parameter buffer index."));
 
-	for (const FShaderParameterBindings::FResourceParameter& ParameterBinding : Bindings.UAVs)
+	for (const FShaderParameterBindings::FResourceParameter& ParameterBinding : Bindings.ResourceParameters)
 	{
-		RHICmdList.SetUAVParameter(ShadeRHI, ParameterBinding.BaseIndex, nullptr);
-	}
-
-	for (const FShaderParameterBindings::FResourceParameter& ParameterBinding : Bindings.GraphUAVs)
-	{
-		RHICmdList.SetUAVParameter(ShadeRHI, ParameterBinding.BaseIndex, nullptr);
+		if (ParameterBinding.BaseType == UBMT_UAV ||
+			ParameterBinding.BaseType == UBMT_RDG_TEXTURE_UAV ||
+			ParameterBinding.BaseType == UBMT_RDG_BUFFER_UAV)
+		{
+			RHICmdList.SetUAVParameter(ShadeRHI, ParameterBinding.BaseIndex, nullptr);
+		}
 	}
 }
 
@@ -194,42 +175,61 @@ inline void SetShaderParameters(TRHICmdList& RHICmdList, const TShaderRef<TShade
 		RHICmdList.SetShaderParameter(ShadeRHI, ParameterBinding.BufferIndex, ParameterBinding.BaseIndex, ParameterBinding.ByteSize, DataPtr);
 	}
 
-	// Textures
-	for (const FShaderParameterBindings::FResourceParameter& ParameterBinding : Bindings.Textures)
-	{
-		auto ShaderParameterRef = *(FRHITexture**)(Base + ParameterBinding.ByteOffset);
-		RHICmdList.SetShaderTexture(ShadeRHI, ParameterBinding.BaseIndex, ShaderParameterRef);
-	}
+	TArray<FShaderParameterBindings::FResourceParameter, TInlineAllocator<16>> GraphSRVs;
 
-	// SRVs
-	for (const FShaderParameterBindings::FResourceParameter& ParameterBinding : Bindings.SRVs)
+	for (const FShaderParameterBindings::FResourceParameter& ParameterBinding : Bindings.ResourceParameters)
 	{
-		FRHIShaderResourceView* ShaderParameterRef = *(FRHIShaderResourceView**)(Base + ParameterBinding.ByteOffset);
-		RHICmdList.SetShaderResourceViewParameter(ShadeRHI, ParameterBinding.BaseIndex, ShaderParameterRef);
+		EUniformBufferBaseType BaseType = (EUniformBufferBaseType)ParameterBinding.BaseType;
+		switch (BaseType)
+		{
+			case UBMT_TEXTURE:
+			{
+				FRHITexture* TexRef = *(FRHITexture**)(Base + ParameterBinding.ByteOffset);
+				RHICmdList.SetShaderTexture(ShadeRHI, ParameterBinding.BaseIndex, TexRef);
+			}
+			break;
+			case UBMT_SRV:
+			{
+				FRHIShaderResourceView* SRVRef = *(FRHIShaderResourceView**)(Base + ParameterBinding.ByteOffset);
+				RHICmdList.SetShaderResourceViewParameter(ShadeRHI, ParameterBinding.BaseIndex, SRVRef);
+			}
+			break;
+			case UBMT_SAMPLER:
+			{
+				FRHISamplerState* SamplerRef = *(FRHISamplerState**)(Base + ParameterBinding.ByteOffset);
+				RHICmdList.SetShaderSampler(ShadeRHI, ParameterBinding.BaseIndex, SamplerRef);
+			}
+			break;
+			case UBMT_RDG_TEXTURE:
+			{
+				auto GraphTexture = *reinterpret_cast<FRDGTexture* const*>(Base + ParameterBinding.ByteOffset);
+				checkSlow(GraphTexture);
+				GraphTexture->MarkResourceAsUsed();
+				RHICmdList.SetShaderTexture(ShadeRHI, ParameterBinding.BaseIndex, GraphTexture->GetRHI());
+			}
+			break;
+			case UBMT_RDG_TEXTURE_SRV:
+			case UBMT_RDG_BUFFER_SRV:
+			{
+				//HACKHACK: defer SRVs binding after UAVs 
+				GraphSRVs.Add(ParameterBinding);
+			}
+			break;
+			case UBMT_UAV:
+			case UBMT_RDG_TEXTURE_UAV:
+			case UBMT_RDG_BUFFER_UAV:
+			{
+				SetShaderUAV(RHICmdList, Shader, ShadeRHI, Base, ParameterBinding);
+			}
+			break;
+			default:
+				checkf(false, TEXT("Unhandled resource type?"));
+				break;
+		}
 	}
 
-	// Samplers
-	for (const FShaderParameterBindings::FResourceParameter& ParameterBinding : Bindings.Samplers)
-	{
-		FRHISamplerState* ShaderParameterRef = *(FRHISamplerState**)(Base + ParameterBinding.ByteOffset);
-		RHICmdList.SetShaderSampler(ShadeRHI, ParameterBinding.BaseIndex, ShaderParameterRef);
-	}
-
-	// Graph Textures
-	for (const FShaderParameterBindings::FResourceParameter& ParameterBinding : Bindings.GraphTextures)
-	{
-		auto GraphTexture = *reinterpret_cast<FRDGTexture* const*>(Base + ParameterBinding.ByteOffset);
-
-		checkSlow(GraphTexture);
-		GraphTexture->MarkResourceAsUsed();
-		RHICmdList.SetShaderTexture(ShadeRHI, ParameterBinding.BaseIndex, GraphTexture->GetRHI());
-	}
-	
-	// UAVs for compute shaders
-	SetShaderUAVs(RHICmdList, Shader, ShadeRHI, Parameters);	//HACKHACK: Bind UAVs before SRVs as a workaround for D3D11 RHI unbinding SRVs when binding a UAV on the same resource even when the views don't overlap.
-
-	// Graph SRVs
-	for (const FShaderParameterBindings::FResourceParameter& ParameterBinding : Bindings.GraphSRVs)
+	//HACKHACK: Bind SRVs after UAVs as a workaround for D3D11 RHI unbinding SRVs when binding a UAV on the same resource even when the views don't overlap.
+	for (const FShaderParameterBindings::FResourceParameter& ParameterBinding : GraphSRVs)
 	{
 		auto GraphSRV = *reinterpret_cast<FRDGShaderResourceView* const*>(Base + ParameterBinding.ByteOffset);
 
@@ -237,7 +237,7 @@ inline void SetShaderParameters(TRHICmdList& RHICmdList, const TShaderRef<TShade
 		GraphSRV->MarkResourceAsUsed();
 		RHICmdList.SetShaderResourceViewParameter(ShadeRHI, ParameterBinding.BaseIndex, GraphSRV->GetRHI());
 	}
-
+	
 	// Reference structures
 	for (const FShaderParameterBindings::FParameterStructReference& ParameterBinding : Bindings.ParameterReferences)
 	{
@@ -262,63 +262,69 @@ void SetShaderParameters(FRayTracingShaderBindingsWriter& RTBindingsWriter, cons
 	const typename TShaderClass::FParameters* ParametersPtr = &Parameters;
 	const uint8* Base = reinterpret_cast<const uint8*>(ParametersPtr);
 
-	// Textures
-	for (const FShaderParameterBindings::FResourceParameter& ParameterBinding : Bindings.Textures)
-	{
-		auto ShaderParameterRef = *(FRHITexture**)(Base + ParameterBinding.ByteOffset);
-		RTBindingsWriter.SetTexture(ParameterBinding.BaseIndex, ShaderParameterRef);
-	}
-
-	// SRVs
-	for (const FShaderParameterBindings::FResourceParameter& ParameterBinding : Bindings.SRVs)
-	{
-		FRHIShaderResourceView* ShaderParameterRef = *(FRHIShaderResourceView**)(Base + ParameterBinding.ByteOffset);
-		RTBindingsWriter.SetSRV(ParameterBinding.BaseIndex, ShaderParameterRef);
-	}
-
-	// UAVs
-	for (const FShaderParameterBindings::FResourceParameter& ParameterBinding : Bindings.UAVs)
-	{
-		FRHIUnorderedAccessView* ShaderParameterRef = *(FRHIUnorderedAccessView**)(Base + ParameterBinding.ByteOffset);
-		RTBindingsWriter.SetUAV(ParameterBinding.BaseIndex, ShaderParameterRef);
-	}
-
-	// Samplers
-	for (const FShaderParameterBindings::FResourceParameter& ParameterBinding : Bindings.Samplers)
-	{
-		FRHISamplerState* ShaderParameterRef = *(FRHISamplerState**)(Base + ParameterBinding.ByteOffset);
-		RTBindingsWriter.SetSampler(ParameterBinding.BaseIndex, ShaderParameterRef);
-	}
-
-	// Graph Textures
-	for (const FShaderParameterBindings::FResourceParameter& ParameterBinding : Bindings.GraphTextures)
+	for (const FShaderParameterBindings::FResourceParameter& ParameterBinding : Bindings.ResourceParameters)
 	{
-		auto GraphTexture = *reinterpret_cast<FRDGTexture* const*>(Base + ParameterBinding.ByteOffset);
-
-		checkSlow(GraphTexture);
-		GraphTexture->MarkResourceAsUsed();
-		RTBindingsWriter.SetTexture(ParameterBinding.BaseIndex, GraphTexture->GetRHI());
+		EUniformBufferBaseType BaseType = (EUniformBufferBaseType)ParameterBinding.BaseType;
+		switch (BaseType)
+		{
+			case UBMT_TEXTURE:
+			{
+				auto ShaderParameterRef = *(FRHITexture**)(Base + ParameterBinding.ByteOffset);
+				RTBindingsWriter.SetTexture(ParameterBinding.BaseIndex, ShaderParameterRef);
+			}
+			break;
+			case UBMT_SRV:
+			{
+				FRHIShaderResourceView* ShaderParameterRef = *(FRHIShaderResourceView**)(Base + ParameterBinding.ByteOffset);
+				RTBindingsWriter.SetSRV(ParameterBinding.BaseIndex, ShaderParameterRef);
+			}
+			break;
+			case UBMT_UAV:
+			{
+				FRHIUnorderedAccessView* ShaderParameterRef = *(FRHIUnorderedAccessView**)(Base + ParameterBinding.ByteOffset);
+				RTBindingsWriter.SetUAV(ParameterBinding.BaseIndex, ShaderParameterRef);
+			}
+			break;
+			case UBMT_SAMPLER:
+			{
+				FRHISamplerState* ShaderParameterRef = *(FRHISamplerState**)(Base + ParameterBinding.ByteOffset);
+				RTBindingsWriter.SetSampler(ParameterBinding.BaseIndex, ShaderParameterRef);
+			}
+			break;
+			case UBMT_RDG_TEXTURE:
+			{
+				auto GraphTexture = *reinterpret_cast<FRDGTexture* const*>(Base + ParameterBinding.ByteOffset);
+				checkSlow(GraphTexture);
+				GraphTexture->MarkResourceAsUsed();
+				RTBindingsWriter.SetTexture(ParameterBinding.BaseIndex, GraphTexture->GetRHI());
+			}
+			break;
+			case UBMT_RDG_TEXTURE_SRV:
+			case UBMT_RDG_BUFFER_SRV:
+			{
+				auto GraphSRV = *reinterpret_cast<FRDGShaderResourceView* const*>(Base + ParameterBinding.ByteOffset);
+
+				checkSlow(GraphSRV);
+				GraphSRV->MarkResourceAsUsed();
+				RTBindingsWriter.SetSRV(ParameterBinding.BaseIndex, GraphSRV->GetRHI());
+			}
+			break;
+			case UBMT_RDG_TEXTURE_UAV:
+			case UBMT_RDG_BUFFER_UAV:
+			{
+				auto UAV = *reinterpret_cast<FRDGUnorderedAccessView* const*>(Base + ParameterBinding.ByteOffset);
+
+				checkSlow(UAV);
+				UAV->MarkResourceAsUsed();
+				RTBindingsWriter.SetUAV(ParameterBinding.BaseIndex, UAV->GetRHI());
+			}
+			break;
+			default:
+				checkf(false, TEXT("Unhandled resource type?"));
+				break;
+		}
 	}
 
-	// Graph SRVs
-	for (const FShaderParameterBindings::FResourceParameter& ParameterBinding : Bindings.GraphSRVs)
-	{
-		auto GraphSRV = *reinterpret_cast<FRDGShaderResourceView* const*>(Base + ParameterBinding.ByteOffset);
-
-		checkSlow(GraphSRV);
-		GraphSRV->MarkResourceAsUsed();
-		RTBindingsWriter.SetSRV(ParameterBinding.BaseIndex, GraphSRV->GetRHI());
-	}
-
-	// Render graph UAVs
-	for (const FShaderParameterBindings::FResourceParameter& ParameterBinding : Bindings.GraphUAVs)
-	{
-		auto UAV = *reinterpret_cast<FRDGUnorderedAccessView* const*>(Base + ParameterBinding.ByteOffset);
-
-		checkSlow(UAV);
-		UAV->MarkResourceAsUsed();
-		RTBindingsWriter.SetUAV(ParameterBinding.BaseIndex, UAV->GetRHI());
-	}
 
 	// Referenced uniform buffers
 	for (const FShaderParameterBindings::FParameterStructReference& ParameterBinding : Bindings.ParameterReferences)
