Fix for debug draw pass using MSAA depth instead of resolved depth.

diff --git a/Engine/Source/Runtime/RenderCore/Private/RenderGraphBuilder.cpp b/Engine/Source/Runtime/RenderCore/Private/RenderGraphBuilder.cpp
--- a/Engine/Source/Runtime/RenderCore/Private/RenderGraphBuilder.cpp
+++ b/Engine/Source/Runtime/RenderCore/Private/RenderGraphBuilder.cpp
@@ -971,6 +971,11 @@ void FRDGBuilder::PrepareResourcesForExecute(const FRDGPass* Pass, struct FRHIRe
 					FRHITexture* TargetableTexture = Texture->PooledRenderTarget->GetRenderTargetItem().TargetableTexture;
 					FRHITexture* ShaderResourceTexture = Texture->PooledRenderTarget->GetRenderTargetItem().ShaderResourceTexture;
 
+					if (RenderTarget.GetMsaaPlane() == ERenderTargetMsaaPlane::Resolved)
+					{
+						TargetableTexture = ShaderResourceTexture;
+					}
+
 					// TODO(RDG): The load store action could actually be optimised by render graph for tile hardware when there is multiple
 					// consecutive rasterizer passes that have RDG resource as render target, a bit like resource transitions.
 					ERenderTargetStoreAction StoreAction = ERenderTargetStoreAction::EStore;
@@ -1010,7 +1015,9 @@ void FRDGBuilder::PrepareResourcesForExecute(const FRDGPass* Pass, struct FRHIRe
 				ERenderTargetStoreAction DepthStoreAction = ExclusiveDepthStencil.IsDepthWrite() ? ERenderTargetStoreAction::EStore : ERenderTargetStoreAction::ENoAction;
 				ERenderTargetStoreAction StencilStoreAction = ExclusiveDepthStencil.IsStencilWrite() ? ERenderTargetStoreAction::EStore : ERenderTargetStoreAction::ENoAction;
 
-				OutDepthStencil.DepthStencilTarget = Texture->PooledRenderTarget->GetRenderTargetItem().TargetableTexture;
+				const FSceneRenderTargetItem& RenderTargetItem = Texture->PooledRenderTarget->GetRenderTargetItem();
+
+				OutDepthStencil.DepthStencilTarget = DepthStencil.GetMsaaPlane() == ERenderTargetMsaaPlane::Unresolved ? RenderTargetItem.TargetableTexture : RenderTargetItem.ShaderResourceTexture;
 				OutDepthStencil.ResolveTarget = nullptr;
 				OutDepthStencil.Action = MakeDepthStencilTargetActions(
 					MakeRenderTargetActions(DepthStencil.GetDepthLoadAction(), DepthStoreAction),
diff --git a/Engine/Source/Runtime/RenderCore/Public/ShaderParameterMacros.h b/Engine/Source/Runtime/RenderCore/Public/ShaderParameterMacros.h
--- a/Engine/Source/Runtime/RenderCore/Public/ShaderParameterMacros.h
+++ b/Engine/Source/Runtime/RenderCore/Public/ShaderParameterMacros.h
@@ -137,6 +137,11 @@ private:
 	friend class TUniformBuffer;
 };
 
+enum class ERenderTargetMsaaPlane : uint8
+{
+	Unresolved,
+	Resolved
+};
 
 /** Render graph information about how to bind a render target. */
 struct alignas(SHADER_PARAMETER_STRUCT_ALIGNMENT) FRenderTargetBinding
@@ -150,9 +155,10 @@ struct alignas(SHADER_PARAMETER_STRUCT_ALIGNMENT) FRenderTargetBinding
 	 *
 	 * Notes: Load and store action are on purpose without default values, to force the user to not forget one of these.
 	 */
-	FRenderTargetBinding(FRDGTexture* InTexture, ERenderTargetLoadAction InLoadAction, uint8 InMipIndex = 0, int32 InArraySlice = -1)
+	FRenderTargetBinding(FRDGTexture* InTexture, ERenderTargetLoadAction InLoadAction, uint8 InMipIndex = 0, int32 InArraySlice = -1, ERenderTargetMsaaPlane InMsaaPlane = ERenderTargetMsaaPlane::Unresolved)
 		: Texture(InTexture)
 		, LoadAction(InLoadAction)
+		, MsaaPlane(InMsaaPlane)
 		, MipIndex(InMipIndex)
 		, ArraySlice(InArraySlice)
 	{
@@ -175,6 +181,10 @@ struct alignas(SHADER_PARAMETER_STRUCT_ALIGNMENT) FRenderTargetBinding
 	{
 		return ArraySlice;
 	}
+	FORCEINLINE ERenderTargetMsaaPlane GetMsaaPlane() const
+	{
+		return MsaaPlane;
+	}
 
 private:
 	/** All parameters required to bind a render target deferred. This are purposefully private to
@@ -182,6 +192,7 @@ private:
 	 */
 	TAlignedShaderParameterPtr<FRDGTexture*> Texture;
 	ERenderTargetLoadAction		LoadAction		= ERenderTargetLoadAction::ENoAction;
+	ERenderTargetMsaaPlane		MsaaPlane		= ERenderTargetMsaaPlane::Unresolved;
 	uint8						MipIndex		= 0;
 	int32						ArraySlice		= -1;
 
@@ -204,10 +215,12 @@ struct alignas(SHADER_PARAMETER_STRUCT_ALIGNMENT) FDepthStencilBinding
 		FRDGTexture* InTexture,
 		ERenderTargetLoadAction InDepthLoadAction,
 		ERenderTargetLoadAction InStencilLoadAction,
-		FExclusiveDepthStencil InDepthStencilAccess)
+		FExclusiveDepthStencil InDepthStencilAccess,
+		ERenderTargetMsaaPlane InMsaaPlane = ERenderTargetMsaaPlane::Unresolved)
 		: Texture(InTexture)
 		, DepthLoadAction(InDepthLoadAction)
 		, StencilLoadAction(InStencilLoadAction)
+		, MsaaPlane(InMsaaPlane)
 		, DepthStencilAccess(InDepthStencilAccess)
 	{
 		check(Validate());
@@ -216,9 +229,11 @@ struct alignas(SHADER_PARAMETER_STRUCT_ALIGNMENT) FDepthStencilBinding
 	FORCEINLINE FDepthStencilBinding(
 		FRDGTexture* InTexture,
 		ERenderTargetLoadAction InDepthLoadAction,
-		FExclusiveDepthStencil InDepthStencilAccess)
+		FExclusiveDepthStencil InDepthStencilAccess,
+		ERenderTargetMsaaPlane InMsaaPlane = ERenderTargetMsaaPlane::Unresolved)
 		: Texture(InTexture)
 		, DepthLoadAction(InDepthLoadAction)
+		, MsaaPlane(InMsaaPlane)
 		, DepthStencilAccess(InDepthStencilAccess)
 	{
 		check(Validate());
@@ -240,6 +255,10 @@ struct alignas(SHADER_PARAMETER_STRUCT_ALIGNMENT) FDepthStencilBinding
 	{
 		return DepthStencilAccess;
 	}
+	FORCEINLINE ERenderTargetMsaaPlane GetMsaaPlane() const
+	{
+		return MsaaPlane;
+	}
 
 private:
 	/** 
@@ -249,6 +268,7 @@ private:
 	TAlignedShaderParameterPtr<FRDGTexture*> Texture = nullptr;
 	ERenderTargetLoadAction		DepthLoadAction		= ERenderTargetLoadAction::ENoAction;
 	ERenderTargetLoadAction		StencilLoadAction	= ERenderTargetLoadAction::ENoAction;
+	ERenderTargetMsaaPlane		MsaaPlane			= ERenderTargetMsaaPlane::Unresolved;
 	FExclusiveDepthStencil		DepthStencilAccess	= FExclusiveDepthStencil::DepthNop_StencilNop;
 
 	RENDERCORE_API bool Validate() const;
diff --git a/Engine/Source/Runtime/Renderer/Private/GpuDebugRendering.cpp b/Engine/Source/Runtime/Renderer/Private/GpuDebugRendering.cpp
--- a/Engine/Source/Runtime/Renderer/Private/GpuDebugRendering.cpp
+++ b/Engine/Source/Runtime/Renderer/Private/GpuDebugRendering.cpp
@@ -274,7 +274,7 @@ namespace ShaderDrawDebug
 
 			ShaderDrawVSPSParameters* PassParameters = GraphBuilder.AllocParameters<ShaderDrawVSPSParameters>();
 			PassParameters->ShaderDrawPSParameters.RenderTargets[0] = FRenderTargetBinding(OutputTexture, ERenderTargetLoadAction::ELoad);
-			PassParameters->ShaderDrawPSParameters.RenderTargets.DepthStencil = FDepthStencilBinding(DepthTexture, ERenderTargetLoadAction::ELoad, FExclusiveDepthStencil::DepthRead_StencilNop);
+			PassParameters->ShaderDrawPSParameters.RenderTargets.DepthStencil = FDepthStencilBinding(DepthTexture, ERenderTargetLoadAction::ELoad, FExclusiveDepthStencil::DepthRead_StencilNop, ERenderTargetMsaaPlane::Resolved);
 			PassParameters->ShaderDrawPSParameters.ColorScale = bIsBehindDepth ? 0.4f : 1.0f;	// When debug primitive are behind the depth buffer, make them look darker.
 			PassParameters->ShaderDrawVSParameters.View = View.ViewUniformBuffer;
 			if (bUseRdgInput)
