Clamp GpuProfilerEvent times to valid values plus remove unused variables

SanitizeEventTree was not clamping start times to be >= to previous root start times, so
we were hitting asserts in TraverseEventTree() for:

    lastStartTime >= GpuProfilerEvents[Root].GetStartResultMicroseconds()

This CL should sanitize all start/end time values

diff --git a/Engine/Source/Runtime/RenderCore/Private/ProfilingDebugging/RealtimeGPUProfiler.cpp b/Engine/Source/Runtime/RenderCore/Private/ProfilingDebugging/RealtimeGPUProfiler.cpp
--- a/Engine/Source/Runtime/RenderCore/Private/ProfilingDebugging/RealtimeGPUProfiler.cpp
+++ b/Engine/Source/Runtime/RenderCore/Private/ProfilingDebugging/RealtimeGPUProfiler.cpp
@@ -263,9 +263,6 @@ void TraverseEventTree(
 	uint64 lastStartTime = 0;
 	uint64 lastEndTime = 0;
 
-	uint64 RootStartTime = 0;
-	uint64 RootEndTime = 0;
-
 	if (Root != 0)
 	{
 		FGpuProfilerTrace::SpecifyEventByName(GpuProfilerEvents[Root].GetName());
@@ -292,25 +289,6 @@ void TraverseEventTree(
 		FGpuProfilerTrace::EndEvent(GpuProfilerEvents[Root].GetEndResultMicroseconds());
 	}
 }
-
-void SanitizeEventTree(
-	TArray<FRealtimeGPUProfilerEvent, TInlineAllocator<100u>>& GpuProfilerEvents,
-	const TArray<TArray<int32>>& GpuProfilerEventChildrenIndices,
-	int32 Root,
-	uint64& MonotonicTime)
-{
-	GpuProfilerEvents[Root].StartResultMicroseconds = FMath::Max(GpuProfilerEvents[Root].StartResultMicroseconds, MonotonicTime);
-	GpuProfilerEvents[Root].EndResultMicroseconds = FMath::Max(GpuProfilerEvents[Root].EndResultMicroseconds, GpuProfilerEvents[Root].GetStartResultMicroseconds());
-
-	for (int32 Subroot : GpuProfilerEventChildrenIndices[Root])
-	{
-		SanitizeEventTree(GpuProfilerEvents, GpuProfilerEventChildrenIndices, Subroot, MonotonicTime);
-	}
-
-	GpuProfilerEvents[Root].EndResultMicroseconds = FMath::Max(GpuProfilerEvents[Root].EndResultMicroseconds, MonotonicTime);
-
-	MonotonicTime = FMath::Max(MonotonicTime, GpuProfilerEvents[Root].GetEndResultMicroseconds());
-}
 #endif
 
 /*-----------------------------------------------------------------------------
@@ -532,8 +510,36 @@ public:
 			}
 		}
 
-		uint64 MonotonicTime = 0;
-		SanitizeEventTree(GpuProfilerEvents, GpuProfilerEventChildrenIndices, 0, MonotonicTime);
+		// Sanitize event start/end times
+		uint64 lastEndTime = 0;
+		for (int32 EventIdx = 1; EventIdx < GpuProfilerEventParentIndices.Num(); ++EventIdx)
+		{
+			const int32 ParentIdx = GpuProfilerEventParentIndices[EventIdx];
+			FRealtimeGPUProfilerEvent& Event = GpuProfilerEvents[EventIdx];
+
+			if (ParentIdx == 0)
+			{
+				// Start time must be >= last end time
+				Event.StartResultMicroseconds = FMath::Max(Event.StartResultMicroseconds, lastEndTime);
+				// End time must be >= start time
+				Event.EndResultMicroseconds = FMath::Max(Event.StartResultMicroseconds, Event.EndResultMicroseconds);
+
+				// No parent, so set last end time
+				lastEndTime = Event.EndResultMicroseconds;
+			}
+			else
+			{
+				FRealtimeGPUProfilerEvent& EventParent = GpuProfilerEvents[ParentIdx];
+
+				// Clamp start/end times to be inside parent start/end times
+				Event.StartResultMicroseconds = FMath::Clamp(Event.StartResultMicroseconds,
+						EventParent.StartResultMicroseconds,
+						EventParent.EndResultMicroseconds);
+				Event.EndResultMicroseconds = FMath::Clamp(Event.EndResultMicroseconds,
+						Event.StartResultMicroseconds,
+						EventParent.EndResultMicroseconds);
+			}
+		}
 
 		FGpuProfilerTrace::BeginFrame(Timestamp);
 		TraverseEventTree(GpuProfilerEvents, GpuProfilerEventChildrenIndices, 0);
diff --git a/Engine/Source/Runtime/RenderCore/Public/ProfilingDebugging/RealtimeGPUProfiler.h b/Engine/Source/Runtime/RenderCore/Public/ProfilingDebugging/RealtimeGPUProfiler.h
--- a/Engine/Source/Runtime/RenderCore/Public/ProfilingDebugging/RealtimeGPUProfiler.h
+++ b/Engine/Source/Runtime/RenderCore/Public/ProfilingDebugging/RealtimeGPUProfiler.h
@@ -284,16 +284,12 @@ class FScopedGPUStatEvent
 	/** Cmdlist to push onto. */
 	FRHICommandListImmediate* RHICmdList;
 
-	/** The stat event used to record timings */
-	FRealtimeGPUProfilerEvent* RealtimeGPUProfilerEvent;
-
 	int32* NumDrawCallsPtr;
 
 public:
 	/** Default constructor, initializing all member variables. */
 	FORCEINLINE FScopedGPUStatEvent()
 		: RHICmdList(nullptr)
-		, RealtimeGPUProfilerEvent(nullptr)
 		, NumDrawCallsPtr(nullptr)
 	{}
 
