Don't crash when loading cooked shader maps (where frozen content most likely doesn't match runtime layout)

diff --git a/Engine/Source/Runtime/Engine/Private/Materials/MaterialShader.cpp b/Engine/Source/Runtime/Engine/Private/Materials/MaterialShader.cpp
--- a/Engine/Source/Runtime/Engine/Private/Materials/MaterialShader.cpp
+++ b/Engine/Source/Runtime/Engine/Private/Materials/MaterialShader.cpp
@@ -2409,14 +2409,14 @@ void FMaterialShaderMap::FlushShadersByVertexFactoryType(const FVertexFactoryTyp
 	GetMutableContent()->RemoveMeshShaderMap(VertexFactoryType);
 }
 
-void FMaterialShaderMap::Serialize(FArchive& Ar, bool bInlineShaderResources, bool bLoadedByCookedMaterial)
+bool FMaterialShaderMap::Serialize(FArchive& Ar, bool bInlineShaderResources, bool bLoadedByCookedMaterial)
 {
 	SCOPED_LOADTIMER(FMaterialShaderMap_Serialize);
 	// Note: This is saved to the DDC, not into packages (except when cooked)
 	// Backwards compatibility therefore will not work based on the version of Ar
 	// Instead, just bump MATERIALSHADERMAP_DERIVEDDATA_VER
 	ShaderMapId.Serialize(Ar, bLoadedByCookedMaterial);
-	Super::Serialize(Ar, bInlineShaderResources, bLoadedByCookedMaterial);
+	return Super::Serialize(Ar, bInlineShaderResources, bLoadedByCookedMaterial);
 }
 
 /*void FMaterialShaderMap::RegisterSerializedShaders(bool bLoadedByCookedMaterial)
diff --git a/Engine/Source/Runtime/Engine/Private/Materials/MaterialShared.cpp b/Engine/Source/Runtime/Engine/Private/Materials/MaterialShared.cpp
--- a/Engine/Source/Runtime/Engine/Private/Materials/MaterialShared.cpp
+++ b/Engine/Source/Runtime/Engine/Private/Materials/MaterialShared.cpp
@@ -921,8 +921,10 @@ void FMaterial::SerializeInlineShaderMap(FArchive& Ar)
 			if (bValid)
 			{
 				TRefCountPtr<FMaterialShaderMap> LoadedShaderMap = new FMaterialShaderMap();
-				LoadedShaderMap->Serialize(Ar, true, bCooked && Ar.IsLoading());
-				GameThreadShaderMap = LoadedShaderMap;
+				if (LoadedShaderMap->Serialize(Ar, true, bCooked && Ar.IsLoading()))
+				{
+					GameThreadShaderMap = MoveTemp(LoadedShaderMap);
+				}
 			}
 		}
 	}
diff --git a/Engine/Source/Runtime/Engine/Public/MaterialShared.h b/Engine/Source/Runtime/Engine/Public/MaterialShared.h
--- a/Engine/Source/Runtime/Engine/Public/MaterialShared.h
+++ b/Engine/Source/Runtime/Engine/Public/MaterialShared.h
@@ -1172,7 +1172,7 @@ public:
 	static const FMaterialShaderMap* GetShaderMapBeingCompiled(const FMaterial* Material);
 
 	/** Serializes the shader map. */
-	void Serialize(FArchive& Ar, bool bInlineShaderResources=true, bool bLoadedByCookedMaterial=false);
+	bool Serialize(FArchive& Ar, bool bInlineShaderResources=true, bool bLoadedByCookedMaterial=false);
 
 #if WITH_EDITOR
 	/** Saves this shader map to the derived data cache. */
diff --git a/Engine/Source/Runtime/RenderCore/Private/ShaderMap.cpp b/Engine/Source/Runtime/RenderCore/Private/ShaderMap.cpp
--- a/Engine/Source/Runtime/RenderCore/Private/ShaderMap.cpp
+++ b/Engine/Source/Runtime/RenderCore/Private/ShaderMap.cpp
@@ -117,11 +117,12 @@ void FShaderMapBase::UnfreezeContent()
 	}
 }
 
-#define CHECK_SHADERMAP_DEPENDENCIES !(UE_BUILD_SHIPPING || UE_BUILD_TEST)
+#define CHECK_SHADERMAP_DEPENDENCIES (WITH_EDITOR || !(UE_BUILD_SHIPPING || UE_BUILD_TEST))
 
-void FShaderMapBase::Serialize(FArchive& Ar, bool bInlineShaderResources, bool bLoadedByCookedMaterial)
+bool FShaderMapBase::Serialize(FArchive& Ar, bool bInlineShaderResources, bool bLoadedByCookedMaterial)
 {
 	LLM_SCOPE(ELLMTag::Shaders);
+	bool bContentValid = true;
 	if (Ar.IsSaving())
 	{
 		check(Content);
@@ -197,7 +198,6 @@ void FShaderMapBase::Serialize(FArchive& Ar, bool bInlineShaderResources, bool b
 		void* ContentMemory = FMemory::Malloc(FrozenContentSize);
 		Ar.Serialize(ContentMemory, FrozenContentSize);
 		Content = static_cast<FShaderMapContent*>(ContentMemory);
-		NumFrozenShaders = Content->GetNumShaders();
 		FMemoryImageResult::ApplyPatchesFromArchive(Content, Ar);
 		PointerTable->LoadFromArchive(Ar, Content, bInlineShaderResources, bLoadedByCookedMaterial);
 
@@ -226,8 +226,16 @@ void FShaderMapBase::Serialize(FArchive& Ar, bool bInlineShaderResources, bool b
 				{
 					FSHAHash CheckLayoutHash;
 					const uint32 CheckLayoutSize = Freeze::HashLayout(*DependencyType, LayoutParams, CheckLayoutHash);
-					checkf(CheckLayoutSize == SavedLayoutSize, TEXT("Mismatch size for type %s, compiled size is %d, loaded size is %d"), DependencyType->Name, CheckLayoutSize, SavedLayoutSize);
-					checkf(CheckLayoutHash == SavedLayoutHash, TEXT("Mismatch hash for type %s"), DependencyType->Name);
+					if (CheckLayoutSize != SavedLayoutSize)
+					{
+						UE_LOG(LogShaders, Error, TEXT("Mismatch size for type %s, compiled size is %d, loaded size is %d"), DependencyType->Name, CheckLayoutSize, SavedLayoutSize);
+						bContentValid = false;
+					}
+					else if (CheckLayoutHash != SavedLayoutHash)
+					{
+						UE_LOG(LogShaders, Error, TEXT("Mismatch hash for type %s"), DependencyType->Name);
+						bContentValid = false;
+					}
 				}
 #endif // CHECK_SHADERMAP_DEPENDENCIES
 			}
@@ -240,9 +248,10 @@ void FShaderMapBase::Serialize(FArchive& Ar, bool bInlineShaderResources, bool b
 			FSHAHash ResourceHash;
 			Ar << ResourceHash;
 			Resource = FShaderCodeLibrary::LoadResource(ResourceHash, &Ar);
-			if (UNLIKELY(Resource == nullptr && GMaxRHIShaderPlatform == GetShaderPlatform()))
+			if (!Resource)
 			{
-				UE_LOG(LogShaders, Fatal, TEXT("Missing shader resource for hash '%s' in the shader library"), *ResourceHash.ToString());
+				UE_LOG(LogShaders, Error, TEXT("Missing shader resource for hash '%s' in the shader library"), *ResourceHash.ToString());
+				bContentValid = false;
 			}
 		}
 		else
@@ -252,15 +261,28 @@ void FShaderMapBase::Serialize(FArchive& Ar, bool bInlineShaderResources, bool b
 			Resource = new FShaderMapResource_InlineCode(GetShaderPlatform(), Code);
 		}
 
-		if (LIKELY(Resource))
+		if (bContentValid)
 		{
+			check(Resource);
+			NumFrozenShaders = Content->GetNumShaders();
+
 			BeginInitResource(Resource);
 
 			INC_DWORD_STAT_BY(STAT_Shaders_ShaderResourceMemory, Resource->GetSizeBytes());
 			INC_DWORD_STAT_BY(STAT_Shaders_ShaderMemory, FrozenContentSize);
 			INC_DWORD_STAT_BY(STAT_Shaders_NumShadersLoaded, NumFrozenShaders);
 		}
+		else
+		{
+			Resource.SafeRelease();
+
+			// Don't call destructors here, this is basically unknown/invalid memory at this point
+			FMemory::Free(Content);
+			Content = nullptr;
+		}
 	}
+
+	return bContentValid;
 }
 
 void FShaderMapBase::DestroyContent()
diff --git a/Engine/Source/Runtime/RenderCore/Public/Shader.h b/Engine/Source/Runtime/RenderCore/Public/Shader.h
--- a/Engine/Source/Runtime/RenderCore/Public/Shader.h
+++ b/Engine/Source/Runtime/RenderCore/Public/Shader.h
@@ -1967,7 +1967,7 @@ public:
 	void AssignContent(FShaderMapContent* InContent);
 	void FinalizeContent();
 	void UnfreezeContent();
-	void Serialize(FArchive& Ar, bool bInlineShaderResources, bool bLoadedByCookedMaterial);
+	bool Serialize(FArchive& Ar, bool bInlineShaderResources, bool bLoadedByCookedMaterial);
 
 #if WITH_EDITOR
 	inline void GetOutdatedTypes(TArray<const FShaderType*>& OutdatedShaderTypes, TArray<const FShaderPipelineType*>& OutdatedShaderPipelineTypes, TArray<const FVertexFactoryType*>& OutdatedFactoryTypes) const
