Don't crash the editor if 'cooked content in editor' is enabled along with shared shader code library

diff --git a/Engine/Source/Runtime/RenderCore/Private/ShaderCodeLibrary.cpp b/Engine/Source/Runtime/RenderCore/Private/ShaderCodeLibrary.cpp
--- a/Engine/Source/Runtime/RenderCore/Private/ShaderCodeLibrary.cpp
+++ b/Engine/Source/Runtime/RenderCore/Private/ShaderCodeLibrary.cpp
@@ -1932,6 +1932,7 @@ void FShaderCodeLibrary::InitForRuntime(EShaderPlatform ShaderPlatform)
 		}
 		else
 		{
+			Shutdown();
 #if !WITH_EDITOR
 			if (FPlatformProperties::SupportsWindowedMode())
 			{
@@ -1946,9 +1947,8 @@ void FShaderCodeLibrary::InitForRuntime(EShaderPlatform ShaderPlatform)
 			{
 				UE_LOG(LogShaderLibrary, Fatal, TEXT("Failed to initialize ShaderCodeLibrary required by the project because part of the Global shader library is missing from %s."), *FPaths::ProjectContentDir());
 			}
-#endif
-			Shutdown();
 			FPlatformMisc::RequestExit(true);
+#endif // !WITH_EDITOR	
 		}
 	}
 }
