Ensure start time for events is greater than last end time for all nodes

#ryi Yujiang.Wang, Arciel.Rekman

diff --git a/Engine/Source/Runtime/RenderCore/Private/ProfilingDebugging/RealtimeGPUProfiler.cpp b/Engine/Source/Runtime/RenderCore/Private/ProfilingDebugging/RealtimeGPUProfiler.cpp
--- a/Engine/Source/Runtime/RenderCore/Private/ProfilingDebugging/RealtimeGPUProfiler.cpp
+++ b/Engine/Source/Runtime/RenderCore/Private/ProfilingDebugging/RealtimeGPUProfiler.cpp
@@ -511,23 +511,19 @@ public:
 		}
 
 		// Sanitize event start/end times
-		uint64 lastEndTime = 0;
+		TArray<uint64> lastEndTimes;
+		lastEndTimes.AddDefaulted(GpuProfilerEvents.Num());
 		for (int32 EventIdx = 1; EventIdx < GpuProfilerEventParentIndices.Num(); ++EventIdx)
 		{
 			const int32 ParentIdx = GpuProfilerEventParentIndices[EventIdx];
 			FRealtimeGPUProfilerEvent& Event = GpuProfilerEvents[EventIdx];
 
-			if (ParentIdx == 0)
-			{
-				// Start time must be >= last end time
-				Event.StartResultMicroseconds = FMath::Max(Event.StartResultMicroseconds, lastEndTime);
-				// End time must be >= start time
-				Event.EndResultMicroseconds = FMath::Max(Event.StartResultMicroseconds, Event.EndResultMicroseconds);
+			// Start time must be >= last end time
+			Event.StartResultMicroseconds = FMath::Max(Event.StartResultMicroseconds, lastEndTimes[ParentIdx]);
+			// End time must be >= start time
+			Event.EndResultMicroseconds = FMath::Max(Event.StartResultMicroseconds, Event.EndResultMicroseconds);
 
-				// No parent, so set last end time
-				lastEndTime = Event.EndResultMicroseconds;
-			}
-			else
+			if (ParentIdx != 0)
 			{
 				FRealtimeGPUProfilerEvent& EventParent = GpuProfilerEvents[ParentIdx];
 
@@ -539,6 +535,9 @@ public:
 						Event.StartResultMicroseconds,
 						EventParent.EndResultMicroseconds);
 			}
+
+			// Update last end time for this parent
+			lastEndTimes[ParentIdx] = Event.EndResultMicroseconds;
 		}
 
 		FGpuProfilerTrace::BeginFrame(Timestamp);
