Fixed uninitialized varaible in FShaderPipelineCache that was causing stable cache load failures on some computers

diff --git a/Engine/Source/Runtime/RenderCore/Private/ShaderPipelineCache.cpp b/Engine/Source/Runtime/RenderCore/Private/ShaderPipelineCache.cpp
--- a/Engine/Source/Runtime/RenderCore/Private/ShaderPipelineCache.cpp
+++ b/Engine/Source/Runtime/RenderCore/Private/ShaderPipelineCache.cpp
@@ -1100,6 +1100,7 @@ void FShaderPipelineCache::Flush(bool bClearCompiled /*= true*/)
 
 FShaderPipelineCache::FShaderPipelineCache(EShaderPlatform Platform)
 : FTickableObjectRenderThread(true, false) // (RegisterNow, HighFrequency)
+, CurrentPlatform((EShaderPlatform)-1)
 , BatchSize(0)
 , BatchTime(0.0f)
 , bPaused(false)
