Fixed: //UE4/Main - Run Incremental EngineTest Sequencer Win64Editor - MemStack.h Error: Assertion failed: (GIsCriticalError || !NumMarks)

diff --git a/Engine/Source/Runtime/RenderCore/Private/RenderGraphUtils.cpp b/Engine/Source/Runtime/RenderCore/Private/RenderGraphUtils.cpp
--- a/Engine/Source/Runtime/RenderCore/Private/RenderGraphUtils.cpp
+++ b/Engine/Source/Runtime/RenderCore/Private/RenderGraphUtils.cpp
@@ -92,7 +92,7 @@ void ClearUnusedGraphResourcesImpl(
 		{
 			const FShaderParameterBindings& ShaderBindings = *ShaderBindingsList[Index];
 			int32& GraphResourceId = GraphResourceIds[Index];
-			for (; GraphResourceId < ShaderBindings.ResourceParameters.Num() && ByteOffset < ShaderBindings.ResourceParameters[GraphResourceId].ByteOffset; ++GraphResourceId)
+			for (; GraphResourceId < ShaderBindings.ResourceParameters.Num() && ShaderBindings.ResourceParameters[GraphResourceId].ByteOffset < ByteOffset; ++GraphResourceId)
 			{
 			}
 			bResourceIsUsed|= GraphResourceId < ShaderBindings.ResourceParameters.Num() && ByteOffset == ShaderBindings.ResourceParameters[GraphResourceId].ByteOffset; // check the resouce type as well?
