Fixing (blindly) merge of SCL sorting from a 4.24 codebase.

- Functions did not get carried over.

diff --git a/Engine/Source/Runtime/RenderCore/Private/ShaderCodeLibrary.cpp b/Engine/Source/Runtime/RenderCore/Private/ShaderCodeLibrary.cpp
--- a/Engine/Source/Runtime/RenderCore/Private/ShaderCodeLibrary.cpp
+++ b/Engine/Source/Runtime/RenderCore/Private/ShaderCodeLibrary.cpp
@@ -859,6 +859,65 @@ struct FEditorShaderCodeArchive
 		return SerializedShaders.FindShaderMap(Hash) != INDEX_NONE;
 	}
 
+	void BeginShaderMap(const TArray<FString>& AssociatedAssets, const FName& InShaderMapName, const FPakOrderMap* OrderMap, FShaderCodeStats& CodeStats)
+	{
+		CodeStats.NumShaderMaps++;
+
+		if (bHasActiveShaderMap)
+		{
+			UE_LOG(LogShaderLibrary, Error, TEXT("Begin/EndShaderMap parity mismatch for shadermap '%s', currently active '%s'"), *InShaderMapName.ToString(), *ShaderMapName.ToString());
+			bHasActiveShaderMap = false;
+		}
+
+		ShaderMapName = InShaderMapName;
+		if (ShaderMapName == NAME_None)
+		{
+			UE_LOG(LogShaderLibrary, Warning, TEXT("Shadermap without a type name encountered"));
+			CodeStats.NumUnnamedShaderMaps++;
+		}
+
+		UE_LOG(LogInit, VeryVerbose, TEXT("BeginShaderMap, type name %s: Got %d associated assets"), *ShaderMapName.ToString(), AssociatedAssets.Num());
+
+		// unlikely special case for a global shadermap
+		static FName GlobalShaderMapName(TEXT("GlobalShaderMap"));
+		if (LIKELY(OrderMap && ShaderMapName != GlobalShaderMapName))
+		{
+			if (AssociatedAssets.Num())
+			{
+				uint64 MinAssetOrder = MAX_uint64;
+				for (const FString& AssetPath: AssociatedAssets)
+				{
+					uint64 ThisAssetOrder = OrderMap->GetFileOrder(AssetPath, false);
+					UE_LOG(LogInit, VeryVerbose, TEXT(" '%s', order %llu"), *AssetPath, ThisAssetOrder);
+					if (ThisAssetOrder < MinAssetOrder)
+					{
+						MinAssetOrder = ThisAssetOrder;
+						FirstAssociatedAsset = AssetPath;
+					}
+				}
+				FirstAssociatedAssetOrder = MinAssetOrder;
+			}
+			else
+			{
+				UE_LOG(LogShaderLibrary, Warning, TEXT("Shadermap '%s' has zero associated assets"), *ShaderMapName.ToString());
+				CodeStats.NumShaderMapsWithUnknownAssets++;
+			}
+		}
+		else
+		{
+			FirstAssociatedAssetOrder = 0;
+		}
+		UE_LOG(LogInit, VeryVerbose, TEXT(" Order for shaders from this ShaderMap will be %llu"), FirstAssociatedAssetOrder);
+
+		bHasActiveShaderMap = true;
+	}
+
+	void EndShaderMap()
+	{
+		UE_LOG(LogInit, VeryVerbose, TEXT("EndShaderMap"));
+		bHasActiveShaderMap = false;
+	}
+
 	int32 AddShaderCode(FShaderCodeStats& CodeStats, const FShaderMapResourceCode* Code)
 	{
 		int32 ShaderMapIndex = INDEX_NONE;
