Fix cooking builds with SCL and -nullrhi (UE-91021).

diff --git a/Engine/Source/Runtime/RenderCore/Private/ShaderCodeLibrary.cpp b/Engine/Source/Runtime/RenderCore/Private/ShaderCodeLibrary.cpp
--- a/Engine/Source/Runtime/RenderCore/Private/ShaderCodeLibrary.cpp
+++ b/Engine/Source/Runtime/RenderCore/Private/ShaderCodeLibrary.cpp
@@ -1984,7 +1984,7 @@ void FShaderCodeLibrary::InitForRuntime(EShaderPlatform ShaderPlatform)
 	bool bArchive = false;
 	GConfig->GetBool(TEXT("/Script/UnrealEd.ProjectPackagingSettings"), TEXT("bShareMaterialShaderCode"), bArchive, GGameIni);
 
-	bool bEnable = !FPlatformProperties::IsServerOnly() && FApp::CanEverRender() && bArchive;
+	bool bEnable = !FPlatformProperties::IsServerOnly() && bArchive;
 #if !UE_BUILD_SHIPPING
 	FString FileHostIP;
 	const bool bCookOnTheFly = FParse::Value(FCommandLine::Get(), TEXT("filehostip"), FileHostIP);
