Fixing corrupted video output on Android Vulkan devices when integrated in 3D environments

diff --git a/Engine/Source/Runtime/MediaAssets/Private/Misc/MediaTextureResource.cpp b/Engine/Source/Runtime/MediaAssets/Private/Misc/MediaTextureResource.cpp
--- a/Engine/Source/Runtime/MediaAssets/Private/Misc/MediaTextureResource.cpp
+++ b/Engine/Source/Runtime/MediaAssets/Private/Misc/MediaTextureResource.cpp
@@ -1041,7 +1041,7 @@ void FMediaTextureResource::CreateOutputRenderTarget(const FIntPoint & InDim, EP
 		OutputCreateFlags |= (TexCreate_GenerateMipCapable | TexCreate_UAV);
 
 		// Make sure we only set a number of mips that actually makes sense, given the sample size
-		uint8 MaxMips = FGenericPlatformMath::FloorToInt(FGenericPlatformMath::Log2(FGenericPlatformMath::Max(InDim.X, InDim.Y)));
+		uint8 MaxMips = FGenericPlatformMath::FloorToInt(FGenericPlatformMath::Log2(FGenericPlatformMath::Min(InDim.X, InDim.Y)));
 		InNumMips = FMath::Min(InNumMips, MaxMips);
 	}
 
diff --git a/Engine/Source/Runtime/RenderCore/Private/GenerateMips.cpp b/Engine/Source/Runtime/RenderCore/Private/GenerateMips.cpp
--- a/Engine/Source/Runtime/RenderCore/Private/GenerateMips.cpp
+++ b/Engine/Source/Runtime/RenderCore/Private/GenerateMips.cpp
@@ -134,7 +134,7 @@ void FGenerateMips::SetupRendering(FGenerateMipsStruct *GenMipsStruct, FRHITextu
 	};
 
 	FRHIResourceCreateInfo CreateInfo;
-	GenMipsStruct->VertexBuffer = RHICreateVertexBuffer(sizeof(FMipsElementVertex) * 4, BUF_Volatile, CreateInfo);
+	GenMipsStruct->VertexBuffer = RHICreateVertexBuffer(sizeof(FMipsElementVertex) * 4, BUF_Static, CreateInfo);
 	void* VoidPtr = RHILockVertexBuffer(GenMipsStruct->VertexBuffer, 0, sizeof(FMipsElementVertex) * 4, RLM_WriteOnly);
 
 	FMipsElementVertex* Vertices = (FMipsElementVertex*)VoidPtr;
@@ -156,10 +156,7 @@ void FGenerateMips::SetupRendering(FGenerateMipsStruct *GenMipsStruct, FRHITextu
 	GenMipsStruct->VertexDeclaration = PipelineStateCache::GetOrCreateVertexDeclaration(Elements);
 
 	//Specify the Sampler details based on the input.
-	GenMipsStruct->Sampler.Filter = InParams.Filter;
-	GenMipsStruct->Sampler.AddressU = InParams.AddressU;
-	GenMipsStruct->Sampler.AddressV = InParams.AddressV;
-	GenMipsStruct->Sampler.AddressW = InParams.AddressW;
+	GenMipsStruct->Sampler = FSamplerStateInitializerRHI(InParams.Filter, InParams.AddressU, InParams.AddressV, InParams.AddressW);
 }
 
 
@@ -208,6 +205,7 @@ void FGenerateMips::RenderMips(FRHICommandListImmediate& CommandList, FRHITextur
 		int32 Height = InTexture->GetSizeXYZ().Y >> MipLevel;
 
 		FRHIRenderPassInfo RPInfo(InTexture, ERenderTargetActions::DontLoad_Store, nullptr, MipLevel);
+		RPInfo.bGeneratingMips = true;
 		CommandList.BeginRenderPass(RPInfo, TEXT("GenMipsLevel"));
 		{
 			CommandList.ApplyCachedRenderTargets(GraphicsPSOInit);
@@ -221,6 +219,8 @@ void FGenerateMips::RenderMips(FRHICommandListImmediate& CommandList, FRHITextur
 		}
 		CommandList.EndRenderPass();
 	}
+
+	CommandList.TransitionResource(EResourceTransitionAccess::EReadable, InTexture);
 }
 
 
@@ -267,10 +267,7 @@ TSharedPtr<FGenerateMipsStruct> FGenerateMips::SetupTexture(FRHITexture* InTextu
 	GRenderTargetPool.CreateUntrackedElement(Desc, GenMipsStruct->RenderTarget, RenderTexture);
 
 		//Specify the Sampler details based on the input.
-	GenMipsStruct->Sampler.Filter = InParams.Filter;
-	GenMipsStruct->Sampler.AddressU = InParams.AddressU;
-	GenMipsStruct->Sampler.AddressV = InParams.AddressV;
-	GenMipsStruct->Sampler.AddressW = InParams.AddressW;
+	GenMipsStruct->Sampler = FSamplerStateInitializerRHI(InParams.Filter, InParams.AddressU, InParams.AddressV, InParams.AddressW);
 
 	return GenMipsStruct;
 }
