Reduce allocations in render graph builder.

diff --git a/Engine/Source/Runtime/RenderCore/Private/RenderGraphBarrierBatcher.cpp b/Engine/Source/Runtime/RenderCore/Private/RenderGraphBarrierBatcher.cpp
--- a/Engine/Source/Runtime/RenderCore/Private/RenderGraphBarrierBatcher.cpp
+++ b/Engine/Source/Runtime/RenderCore/Private/RenderGraphBarrierBatcher.cpp
@@ -107,17 +107,17 @@ namespace
 	}
 } //! namespace
 
-FRDGBarrierBatcher::FRDGBarrierBatcher(FRHICommandList& InRHICmdList, const FRDGPass* InPass)
-	: RHICmdList(InRHICmdList)
-	, Pass(InPass)
+void FRDGBarrierBatcher::Begin(const FRDGPass* InPass)
 {
+	Pass = InPass;
+
 	if (Pass)
 	{
 		Pipeline = Pass->IsCompute() ? FRDGResourceState::EPipeline::Compute : FRDGResourceState::EPipeline::Graphics;
 	}
 }
 
-FRDGBarrierBatcher::~FRDGBarrierBatcher()
+void FRDGBarrierBatcher::End(FRHICommandList& RHICmdList)
 {
 #if WITH_MGPU
 	// Wait for the temporal effect before executing the first pass in the graph. This
@@ -133,17 +133,20 @@ FRDGBarrierBatcher::~FRDGBarrierBatcher()
 	{
 		RHICmdList.BeginUpdateMultiFrameResource(RHITexture);
 	}
+	TextureUpdateMultiFrameBegins.Reset();
 
 	for (FRHIUnorderedAccessView* RHIUAV : UAVUpdateMultiFrameBegins)
 	{
 		RHICmdList.BeginUpdateMultiFrameResource(RHIUAV);
 	}
+	UAVUpdateMultiFrameBegins.Reset();
 
 	for (auto& Element : TextureBatchMap)
 	{
 		FTransitionParameters TransitionParameters = Element.Key;
 		FTextureBatch& Batch = Element.Value;
 		RHICmdList.TransitionResources(TransitionParameters.TransitionAccess, TransitionParameters.TransitionPipeline, Batch.GetData(), Batch.Num());
+		Batch.Reset();
 	}
 
 	for (auto& Element : UAVBatchMap)
@@ -151,17 +154,20 @@ FRDGBarrierBatcher::~FRDGBarrierBatcher()
 		FTransitionParameters TransitionParameters = Element.Key;
 		FUAVBatch& Batch = Element.Value;
 		RHICmdList.TransitionResources(TransitionParameters.TransitionAccess, TransitionParameters.TransitionPipeline, Batch.GetData(), Batch.Num());
+		Batch.Reset();
 	}
 
 	for (FRHITexture* RHITexture : TextureUpdateMultiFrameEnds)
 	{
 		RHICmdList.EndUpdateMultiFrameResource(RHITexture);
 	}
+	TextureUpdateMultiFrameEnds.Reset();
 
 	for (FRHIUnorderedAccessView* RHIUAV : UAVUpdateMultiFrameEnds)
 	{
 		RHICmdList.EndUpdateMultiFrameResource(RHIUAV);
 	}
+	UAVUpdateMultiFrameEnds.Reset();
 
 #if WITH_MGPU
 	// Broadcast all multi-frame resources when processing deferred resource queries.
@@ -169,7 +175,11 @@ FRDGBarrierBatcher::~FRDGBarrierBatcher()
 	{
 		RHICmdList.BroadcastTemporalEffect(NameForTemporalEffect, TexturesToCopyForTemporalEffect);
 	}
+	TexturesToCopyForTemporalEffect.Reset();
+	NameForTemporalEffect = NAME_None;
 #endif
+
+	Pass = nullptr;
 }
 
 void FRDGBarrierBatcher::QueueTransitionTexture(FRDGTexture* Texture, FRDGResourceState::EAccess AccessAfter)
diff --git a/Engine/Source/Runtime/RenderCore/Private/RenderGraphBarrierBatcher.h b/Engine/Source/Runtime/RenderCore/Private/RenderGraphBarrierBatcher.h
--- a/Engine/Source/Runtime/RenderCore/Private/RenderGraphBarrierBatcher.h
+++ b/Engine/Source/Runtime/RenderCore/Private/RenderGraphBarrierBatcher.h
@@ -3,16 +3,16 @@
 #pragma once
 
 #include "RenderGraphResources.h"
+#include "Containers/SortedMap.h"
 
 /** Class that queues up batches of resource barriers and then submits them to a command list. */
 class RENDERCORE_API FRDGBarrierBatcher final
 {
 public:
-	/** RAII initialization of batcher. Batcher will flush all queued transitions in the destructor.
-	 *  @param Pass The current pass if performing inter-pass barriers. This can be null (e.g. for post-execution barriers).
-	 */
-	FRDGBarrierBatcher(FRHICommandList& RHICmdList, const FRDGPass* Pass);
-	~FRDGBarrierBatcher();
+	FRDGBarrierBatcher() = default;
+
+	void Begin(const FRDGPass* Pass);
+	void End(FRHICommandList& RHICmdList);
 
 	/** Queues a transition of the texture to the requested access state. */
 	void QueueTransitionTexture(FRDGTexture* Texture, FRDGResourceState::EAccess AccessAfter);
@@ -45,6 +45,11 @@ private:
 			return TransitionAccess == Other.TransitionAccess && TransitionPipeline == Other.TransitionPipeline;
 		}
 
+		bool operator<(const FTransitionParameters& Other) const
+		{
+			return GetHash() < Other.GetHash();
+		}
+
 		uint32 GetHash() const
 		{
 			return uint32(TransitionAccess) | (uint32(TransitionPipeline) << 8);
@@ -54,22 +59,11 @@ private:
 		EResourceTransitionPipeline TransitionPipeline = EResourceTransitionPipeline::EGfxToGfx;
 	};
 
-	template <typename TBatchType>
-	struct TBatchMapKeyFuncs : public TDefaultMapKeyFuncs<FTransitionParameters, TBatchType, /** bAllowDuplicateKeys */ false>
-	{
-		static uint32 GetKeyHash(FTransitionParameters Key)
-		{
-			return Key.GetHash();
-		}
-	};
-
 	using FUAVBatch = TArray<FRHIUnorderedAccessView*, SceneRenderingAllocator>;
-	using FUAVBatchMap = TMap<FTransitionParameters, FUAVBatch, SceneRenderingSetAllocator, TBatchMapKeyFuncs<FUAVBatch>>;
+	using FUAVBatchMap = TSortedMap<FTransitionParameters, FUAVBatch, SceneRenderingAllocator>;
 
 	using FTextureBatch = TArray<FRHITexture*, SceneRenderingAllocator>;
-	using FTextureBatchMap = TMap<FTransitionParameters, FTextureBatch, SceneRenderingSetAllocator, TBatchMapKeyFuncs<FTextureBatch>>;
-
-	FRHICommandList& RHICmdList;
+	using FTextureBatchMap = TSortedMap<FTransitionParameters, FTextureBatch, SceneRenderingAllocator>;
 
 	FTextureBatch TextureUpdateMultiFrameBegins;
 	FTextureBatch TextureUpdateMultiFrameEnds;
diff --git a/Engine/Source/Runtime/RenderCore/Private/RenderGraphBuilder.cpp b/Engine/Source/Runtime/RenderCore/Private/RenderGraphBuilder.cpp
--- a/Engine/Source/Runtime/RenderCore/Private/RenderGraphBuilder.cpp
+++ b/Engine/Source/Runtime/RenderCore/Private/RenderGraphBuilder.cpp
@@ -164,13 +164,22 @@ FRDGBuilder::FRDGBuilder(FRHICommandListImmediate& InRHICmdList)
 	, MemStack(FMemStack::Get())
 	, EventScopeStack(RHICmdList)
 	, StatScopeStack(RHICmdList)
-{}
+	, BarrierBatcher(MakeUnique<FRDGBarrierBatcher>())
+{
+	AllocatedTextures.Reserve(16);
+	AllocatedBuffers.Reserve(16);
+}
+
+FRDGBuilder::~FRDGBuilder() {}
 
 void FRDGBuilder::Execute()
 {
 	CSV_SCOPED_TIMING_STAT_EXCLUSIVE(FRDGBuilder_Execute);
+	SCOPED_NAMED_EVENT(FRDGBuilder_Execute, FColor::Emerald);
 
 	IF_RDG_ENABLE_DEBUG(Validation.ValidateExecuteBegin());
+	AllocatedBuffers.Reserve(BufferCount);
+	AllocatedTextures.Reserve(TextureCount);
 
 	EventScopeStack.BeginExecute();
 	StatScopeStack.BeginExecute();
@@ -737,10 +746,10 @@ void FRDGBuilder::PrepareResourcesForExecute(const FRDGPass* Pass, struct FRHIRe
 
 	const bool bIsCompute = Pass->IsCompute();
 
-	FRDGBarrierBatcher BarrierBatcher(RHICmdList, Pass);
+	BarrierBatcher->Begin(Pass);
 
 #if WITH_MGPU
-	BarrierBatcher.SetNameForTemporalEffect(NameForTemporalEffect);
+	BarrierBatcher->SetNameForTemporalEffect(NameForTemporalEffect);
 #endif
 
 	FRDGPassParameterStruct ParameterStruct = Pass->GetParameters();
@@ -750,6 +759,9 @@ void FRDGBuilder::PrepareResourcesForExecute(const FRDGPass* Pass, struct FRHIRe
 	// List all RDG texture being read and modified by the pass.
 	TSet<FRDGTextureRef, DefaultKeyFuncs<FRDGTextureRef>, SceneRenderingSetAllocator> ReadTextures;
 	TSet<FRDGTextureRef, DefaultKeyFuncs<FRDGTextureRef>, SceneRenderingSetAllocator> ModifiedTextures;
+	ReadTextures.Reserve(ParameterCount);
+	ModifiedTextures.Reserve(ParameterCount);
+
 	for (uint32 ParameterIndex = 0; ParameterIndex < ParameterCount; ++ParameterIndex)
 	{
 		FRDGPassParameter Parameter = ParameterStruct.GetParameter(ParameterIndex);
@@ -832,7 +844,7 @@ void FRDGBuilder::PrepareResourcesForExecute(const FRDGPass* Pass, struct FRHIRe
 				check(Texture->ResourceRHI);
 
 				check(!ModifiedTextures.Contains(Texture));
-				BarrierBatcher.QueueTransitionTexture(Texture, FRDGResourceState::EAccess::Read);
+				BarrierBatcher->QueueTransitionTexture(Texture, FRDGResourceState::EAccess::Read);
 			}
 		}
 		break;
@@ -851,7 +863,7 @@ void FRDGBuilder::PrepareResourcesForExecute(const FRDGPass* Pass, struct FRHIRe
 				}
 				else
 				{
-					BarrierBatcher.QueueTransitionTexture(Texture, FRDGResourceState::EAccess::Read);
+					BarrierBatcher->QueueTransitionTexture(Texture, FRDGResourceState::EAccess::Read);
 				}
 			}
 		}
@@ -868,7 +880,7 @@ void FRDGBuilder::PrepareResourcesForExecute(const FRDGPass* Pass, struct FRHIRe
 
 				bool bGeneratingMips = ReadTextures.Contains(Texture);
 
-				BarrierBatcher.QueueTransitionUAV(UAVRHI, Texture, FRDGResourceState::EAccess::Write, bGeneratingMips);
+				BarrierBatcher->QueueTransitionUAV(UAVRHI, Texture, FRDGResourceState::EAccess::Write, bGeneratingMips);
 			}
 		}
 		break;
@@ -879,7 +891,7 @@ void FRDGBuilder::PrepareResourcesForExecute(const FRDGPass* Pass, struct FRHIRe
 				AllocateRHITextureIfNeeded(Texture);
 
 				check(!ReadTextures.Contains(Texture));
-				BarrierBatcher.QueueTransitionTexture(Texture, FRDGResourceState::EAccess::Write);
+				BarrierBatcher->QueueTransitionTexture(Texture, FRDGResourceState::EAccess::Write);
 			}
 		}
 		break;
@@ -892,7 +904,7 @@ void FRDGBuilder::PrepareResourcesForExecute(const FRDGPass* Pass, struct FRHIRe
 				check(Buffer->PooledBuffer->UAVs.Num() == 1);
 				FRHIUnorderedAccessView* UAVRHI = Buffer->PooledBuffer->UAVs.CreateIterator().Value();
 
-				BarrierBatcher.QueueTransitionUAV(UAVRHI, Buffer, FRDGResourceState::EAccess::Read);
+				BarrierBatcher->QueueTransitionUAV(UAVRHI, Buffer, FRDGResourceState::EAccess::Read);
 			}
 		}
 		break;
@@ -912,7 +924,7 @@ void FRDGBuilder::PrepareResourcesForExecute(const FRDGPass* Pass, struct FRHIRe
 					check(Buffer->PooledBuffer->UAVs.Num() == 1);
 					FRHIUnorderedAccessView* UAVRHI = Buffer->PooledBuffer->UAVs.CreateIterator().Value();
 
-					BarrierBatcher.QueueTransitionUAV(UAVRHI, Buffer, FRDGResourceState::EAccess::Read);
+					BarrierBatcher->QueueTransitionUAV(UAVRHI, Buffer, FRDGResourceState::EAccess::Read);
 				}
 			}
 		}
@@ -927,7 +939,7 @@ void FRDGBuilder::PrepareResourcesForExecute(const FRDGPass* Pass, struct FRHIRe
 
 				FRHIUnorderedAccessView* UAVRHI = UAV->GetRHI();
 
-				BarrierBatcher.QueueTransitionUAV(UAVRHI, Buffer, FRDGResourceState::EAccess::Write);
+				BarrierBatcher->QueueTransitionUAV(UAVRHI, Buffer, FRDGResourceState::EAccess::Write);
 			}
 		}
 		break;
@@ -993,7 +1005,7 @@ void FRDGBuilder::PrepareResourcesForExecute(const FRDGPass* Pass, struct FRHIRe
 					OutRenderTarget.MipIndex = RenderTarget.GetMipIndex();
 					OutRenderTarget.Action = MakeRenderTargetActions(RenderTarget.GetLoadAction(), StoreAction);
 
-					BarrierBatcher.QueueTransitionTexture(Texture, FRDGResourceState::EAccess::Write);
+					BarrierBatcher->QueueTransitionTexture(Texture, FRDGResourceState::EAccess::Write);
 
 					SampleCount |= OutRenderTarget.RenderTarget->GetNumSamples();
 					ValidRenderTargetCount++;
@@ -1024,7 +1036,7 @@ void FRDGBuilder::PrepareResourcesForExecute(const FRDGPass* Pass, struct FRHIRe
 					MakeRenderTargetActions(DepthStencil.GetStencilLoadAction(), StencilStoreAction));
 				OutDepthStencil.ExclusiveDepthStencil = ExclusiveDepthStencil;
 
-				BarrierBatcher.QueueTransitionTexture(Texture,
+				BarrierBatcher->QueueTransitionTexture(Texture,
 					ExclusiveDepthStencil.IsAnyWrite() ?
 					FRDGResourceState::EAccess::Write :
 					FRDGResourceState::EAccess::Read);
@@ -1042,6 +1054,8 @@ void FRDGBuilder::PrepareResourcesForExecute(const FRDGPass* Pass, struct FRHIRe
 	}
 
 	OutRPInfo->bGeneratingMips = Pass->IsGenerateMips();
+
+	BarrierBatcher->End(RHICmdList);
 }
 
 void FRDGBuilder::ReleaseRHITextureIfUnreferenced(FRDGTexture* Texture)
@@ -1167,10 +1181,10 @@ void FRDGBuilder::ReleaseUnreferencedResources(const FRDGPass* Pass)
 
 void FRDGBuilder::ProcessDeferredInternalResourceQueries()
 {
-	FRDGBarrierBatcher BarrierBatcher(RHICmdList, nullptr);
+	BarrierBatcher->Begin(nullptr);
 
 #if WITH_MGPU
-	BarrierBatcher.SetNameForTemporalEffect(NameForTemporalEffect);
+	BarrierBatcher->SetNameForTemporalEffect(NameForTemporalEffect);
 #endif
 
 	for (const auto& Query : DeferredInternalTextureQueries)
@@ -1179,7 +1193,7 @@ void FRDGBuilder::ProcessDeferredInternalResourceQueries()
 
 		if (Query.bTransitionToRead)
 		{
-			BarrierBatcher.QueueTransitionTexture(Query.Texture, FRDGResourceState::EAccess::Read);
+			BarrierBatcher->QueueTransitionTexture(Query.Texture, FRDGResourceState::EAccess::Read);
 		}
 
 		*Query.OutTexturePtr = AllocatedTextures.FindChecked(Query.Texture);
@@ -1195,7 +1209,7 @@ void FRDGBuilder::ProcessDeferredInternalResourceQueries()
 
 		for (TMap<FRDGBufferUAVDesc, FUnorderedAccessViewRHIRef, FDefaultSetAllocator, TMapRDGBufferUAVFuncs<FRDGBufferUAVDesc, FUnorderedAccessViewRHIRef>>::TIterator It(Query.Buffer->PooledBuffer->UAVs); It; ++It)
 		{
-			BarrierBatcher.QueueTransitionUAV(It.Value(), Query.Buffer, Query.DestinationAccess, /* bGeneratingMips = */ false, Query.DestinationPipeline);
+			BarrierBatcher->QueueTransitionUAV(It.Value(), Query.Buffer, Query.DestinationAccess, /* bGeneratingMips = */ false, Query.DestinationPipeline);
 		}
 
 		*Query.OutBufferPtr = AllocatedBuffers.FindChecked(Query.Buffer);
@@ -1206,6 +1220,8 @@ void FRDGBuilder::ProcessDeferredInternalResourceQueries()
 			ReleaseRHIBufferIfUnreferenced(Query.Buffer);
 		}
 	}
+
+	BarrierBatcher->End(RHICmdList);
 }
 
 void FRDGBuilder::DestructPasses()
diff --git a/Engine/Source/Runtime/RenderCore/Public/RenderGraphBuilder.h b/Engine/Source/Runtime/RenderCore/Public/RenderGraphBuilder.h
--- a/Engine/Source/Runtime/RenderCore/Public/RenderGraphBuilder.h
+++ b/Engine/Source/Runtime/RenderCore/Public/RenderGraphBuilder.h
@@ -9,6 +9,8 @@
 #include "RenderGraphValidation.h"
 #include "ShaderParameterMacros.h"
 
+class FRDGBarrierBatcher;
+
 /** Builds the per-frame render graph.
  *  Resources must be created from the builder before they can be bound to Pass ResourceTables.
  *  These resources are descriptors only until the graph is executed, where RHI resources are allocated as needed.
@@ -19,6 +21,7 @@ public:
 	/** A RHI cmd list is required, if using the immediate mode. */
 	FRDGBuilder(FRHICommandListImmediate& InRHICmdList);
 	FRDGBuilder(const FRDGBuilder&) = delete;
+	~FRDGBuilder();
 
 	/** Register a external texture to be tracked by the render graph. */
 	FRDGTextureRef RegisterExternalTexture(
@@ -176,8 +179,10 @@ public:
 	FMemStackBase& MemStack;
 
 private:
+	TUniquePtr<FRDGBarrierBatcher> BarrierBatcher;
+
 	/** Array of all pass created */
-	TArray<FRDGPass*, SceneRenderingAllocator> Passes;
+	TArray<FRDGPass*, TInlineAllocator<8, SceneRenderingAllocator>> Passes;
 
 	/** Keep the references over the pooled render target, since FRDGTexture is allocated on MemStack. */
 	TMap<FRDGTexture*, TRefCountPtr<IPooledRenderTarget>, SceneRenderingSetAllocator> AllocatedTextures;
@@ -189,6 +194,9 @@ private:
 	TMap<const IPooledRenderTarget*, FRDGTexture*, SceneRenderingSetAllocator> ExternalTextures;
 	TMap<const FPooledRDGBuffer*, FRDGBuffer*, SceneRenderingSetAllocator> ExternalBuffers;
 
+	uint32 BufferCount = 0;
+	uint32 TextureCount = 0;
+
 	/** Array of all deferred access to internal textures. */
 	struct FDeferredInternalTextureQuery
 	{
@@ -196,7 +204,7 @@ private:
 		TRefCountPtr<IPooledRenderTarget>* OutTexturePtr;
 		bool bTransitionToRead;
 	};
-	TArray<FDeferredInternalTextureQuery, SceneRenderingAllocator> DeferredInternalTextureQueries;
+	TArray<FDeferredInternalTextureQuery, TInlineAllocator<4, SceneRenderingAllocator>> DeferredInternalTextureQueries;
 
 	/** Array of all deferred access to internal buffers. */
 	struct FDeferredInternalBufferQuery
@@ -206,7 +214,7 @@ private:
 		FRDGResourceState::EAccess DestinationAccess;
 		FRDGResourceState::EPipeline DestinationPipeline;
 	};
-	TArray<FDeferredInternalBufferQuery, SceneRenderingAllocator> DeferredInternalBufferQueries;
+	TArray<FDeferredInternalBufferQuery, TInlineAllocator<4, SceneRenderingAllocator>> DeferredInternalBufferQueries;
 
 	FRDGEventScopeStack EventScopeStack;
 	FRDGStatScopeStack StatScopeStack;
diff --git a/Engine/Source/Runtime/RenderCore/Public/RenderGraphBuilder.inl b/Engine/Source/Runtime/RenderCore/Public/RenderGraphBuilder.inl
--- a/Engine/Source/Runtime/RenderCore/Public/RenderGraphBuilder.inl
+++ b/Engine/Source/Runtime/RenderCore/Public/RenderGraphBuilder.inl
@@ -26,6 +26,8 @@ inline FRDGTextureRef FRDGBuilder::RegisterExternalTexture(
 	OutTexture->ResourceRHI = ExternalPooledTexture->GetRenderTargetItem().ShaderResourceTexture;
 	AllocatedTextures.Add(OutTexture, ExternalPooledTexture);
 
+	TextureCount++;
+
 	IF_RDG_ENABLE_DEBUG(Validation.ValidateCreateExternalTexture(OutTexture));
 
 	ExternalTextures.Add(ExternalPooledTexture.GetReference(), OutTexture);
@@ -66,6 +68,8 @@ inline FRDGBufferRef FRDGBuilder::RegisterExternalBuffer(
 	}
 	AllocatedBuffers.Add(OutBuffer, ExternalPooledBuffer);
 
+	BufferCount++;
+
 	IF_RDG_ENABLE_DEBUG(Validation.ValidateCreateExternalBuffer(OutBuffer));
 
 	ExternalBuffers.Add(ExternalPooledBuffer.GetReference(), OutBuffer);
@@ -104,6 +108,8 @@ inline FRDGTextureRef FRDGBuilder::CreateTexture(
 
 	FRDGTexture* Texture = AllocateForRHILifeTime<FRDGTexture>(Name, TextureDesc, Flags);
 
+	TextureCount++;
+
 	IF_RDG_ENABLE_DEBUG(Validation.ValidateCreateTexture(Texture));
 
 	return Texture;
@@ -130,6 +136,8 @@ inline FRDGBufferRef FRDGBuilder::CreateBuffer(
 
 	FRDGBufferRef Buffer = AllocateForRHILifeTime<FRDGBuffer>(Name, Desc, Flags);
 
+	BufferCount++;
+
 	IF_RDG_ENABLE_DEBUG(Validation.ValidateCreateBuffer(Buffer));
 
 	return Buffer;
