Added shadercomplexity support in cooked content.
This requires "r.ShaderComplexity.CacheShaders=true"
Disabled velocity rendering when rendering debug view shaders.

diff --git a/Engine/Source/Runtime/RenderCore/Public/Shader.h b/Engine/Source/Runtime/RenderCore/Public/Shader.h
--- a/Engine/Source/Runtime/RenderCore/Public/Shader.h
+++ b/Engine/Source/Runtime/RenderCore/Public/Shader.h
@@ -715,14 +715,13 @@ public:
 	inline EShaderFrequency GetFrequency() const { return Target.GetFrequency(); }
 	inline const FShaderTarget GetTarget() const { return Target; }
 	inline bool IsFrozen() const { return Type.IsFrozen(); }
+	inline uint32 GetNumInstructions() const { return NumInstructions; }
 
 #if WITH_EDITORONLY_DATA
-	inline uint32 GetNumInstructions() const { return NumInstructions; }
 	inline uint32 GetNumTextureSamplers() const { return NumTextureSamplers; }
 	inline uint32 GetCodeSize() const { return CodeSize; }
 	inline void SetNumInstructions(uint32 Value) { NumInstructions = Value; }
 #else
-	inline uint32 GetNumInstructions() const { return 0u; }
 	inline uint32 GetNumTextureSamplers() const { return 0u; }
 	inline uint32 GetCodeSize() const { return 0u; }
 #endif
@@ -819,7 +818,7 @@ private:
 	LAYOUT_FIELD(int32, PermutationId);
 
 	/** The number of instructions the shader takes to execute. */
-	LAYOUT_FIELD_EDITORONLY(uint32, NumInstructions);
+	LAYOUT_FIELD(uint32, NumInstructions);
 
 	/** Number of texture samplers the shader uses. */
 	LAYOUT_FIELD_EDITORONLY(uint32, NumTextureSamplers);
diff --git a/Engine/Source/Runtime/Renderer/Private/VelocityRendering.cpp b/Engine/Source/Runtime/Renderer/Private/VelocityRendering.cpp
--- a/Engine/Source/Runtime/Renderer/Private/VelocityRendering.cpp
+++ b/Engine/Source/Runtime/Renderer/Private/VelocityRendering.cpp
@@ -375,7 +375,7 @@ void FDeferredShadingSceneRenderer::RenderVelocitiesInner(FRHICommandListImmedia
 
 bool FDeferredShadingSceneRenderer::ShouldRenderVelocities() const
 {
-	if (!FVelocityRendering::IsSeparateVelocityPassSupported())
+	if (!FVelocityRendering::IsSeparateVelocityPassSupported() || ViewFamily.UseDebugViewPS())
 	{
 		return false;
 	}
