Removed additional allocations from FRDGBuilder. Compile out debug labels for test / shipping in render target pool. Don't set global uniform buffers when none exist.

diff --git a/Engine/Source/Runtime/RenderCore/Private/RenderGraphBuilder.cpp b/Engine/Source/Runtime/RenderCore/Private/RenderGraphBuilder.cpp
--- a/Engine/Source/Runtime/RenderCore/Private/RenderGraphBuilder.cpp
+++ b/Engine/Source/Runtime/RenderCore/Private/RenderGraphBuilder.cpp
@@ -755,10 +755,6 @@ void FRDGBuilder::PrepareResourcesForExecute(const FRDGPass* Pass, struct FRHIRe
 	FRDGPassParameterStruct ParameterStruct = Pass->GetParameters();
 
 	const uint32 ParameterCount = ParameterStruct.GetParameterCount();
-
-	// List all RDG texture being read and modified by the pass.
-	TSet<FRDGTextureRef, DefaultKeyFuncs<FRDGTextureRef>, SceneRenderingSetAllocator> ReadTextures;
-	TSet<FRDGTextureRef, DefaultKeyFuncs<FRDGTextureRef>, SceneRenderingSetAllocator> ModifiedTextures;
 	ReadTextures.Reserve(ParameterCount);
 	ModifiedTextures.Reserve(ParameterCount);
 
@@ -1056,6 +1052,8 @@ void FRDGBuilder::PrepareResourcesForExecute(const FRDGPass* Pass, struct FRHIRe
 	OutRPInfo->bGeneratingMips = Pass->IsGenerateMips();
 
 	BarrierBatcher->End(RHICmdList);
+	ReadTextures.Reset();
+	ModifiedTextures.Reset();
 }
 
 void FRDGBuilder::ReleaseRHITextureIfUnreferenced(FRDGTexture* Texture)
diff --git a/Engine/Source/Runtime/RenderCore/Private/RenderGraphPass.cpp b/Engine/Source/Runtime/RenderCore/Private/RenderGraphPass.cpp
--- a/Engine/Source/Runtime/RenderCore/Private/RenderGraphPass.cpp
+++ b/Engine/Source/Runtime/RenderCore/Private/RenderGraphPass.cpp
@@ -40,14 +40,9 @@ FUniformBufferStaticBindings FRDGPassParameterStruct::GetGlobalUniformBuffers()
 
 			FRHIUniformBuffer* UniformBufferPtr = *reinterpret_cast<FUniformBufferRHIRef*>(Contents + MemberOffset);
 
-			if (UniformBufferPtr)
+			if (UniformBufferPtr && UniformBufferPtr->HasStaticSlot())
 			{
-				const FUniformBufferStaticSlot StaticSlot = UniformBufferPtr->GetLayout().StaticSlot;
-
-				if (IsUniformBufferStaticSlotValid(StaticSlot))
-				{
-					GlobalUniformBuffers.AddUniformBuffer(UniformBufferPtr);
-				}
+				GlobalUniformBuffers.AddUniformBuffer(UniformBufferPtr);
 			}
 		}
 	}
@@ -66,10 +61,19 @@ FRDGPass::FRDGPass(
 
 void FRDGPass::Execute(FRHICommandListImmediate& RHICmdList) const
 {
-	FUniformBufferStaticBindings GlobalUniformBuffers = ParameterStruct.GetGlobalUniformBuffers();
-	SCOPED_UNIFORM_BUFFER_GLOBAL_BINDINGS(RHICmdList, GlobalUniformBuffers);
+	SCOPED_NAMED_EVENT(FRDGPass_Execute, FColor::Emerald);
+
+	const FUniformBufferStaticBindings GlobalUniformBuffers = ParameterStruct.GetGlobalUniformBuffers();
+
+	if (GlobalUniformBuffers.GetUniformBufferCount())
+	{
+		RHICmdList.SetGlobalUniformBuffers(GlobalUniformBuffers);
+	}
 
 	ExecuteImpl(RHICmdList);
 
-	ParameterStruct.ClearUniformBuffers();
+	if (GlobalUniformBuffers.GetUniformBufferCount())
+	{
+		RHICmdList.SetGlobalUniformBuffers({});
+	}
 }
\ No newline at end of file
diff --git a/Engine/Source/Runtime/RenderCore/Private/RenderTargetPool.cpp b/Engine/Source/Runtime/RenderCore/Private/RenderTargetPool.cpp
--- a/Engine/Source/Runtime/RenderCore/Private/RenderTargetPool.cpp
+++ b/Engine/Source/Runtime/RenderCore/Private/RenderTargetPool.cpp
@@ -253,10 +253,12 @@ bool FRenderTargetPool::FindFreeElement(FRHICommandList& RHICmdList, const FPool
 		{
 			// we can reuse the same, but the debug name might have changed
 			Current->Desc.DebugName = InDebugName;
+#if !(UE_BUILD_SHIPPING || UE_BUILD_TEST)
 			if (Current->GetRenderTargetItem().TargetableTexture)
 			{
 				RHIBindDebugLabelName(Current->GetRenderTargetItem().TargetableTexture, InDebugName);
 			}
+#endif
 			check(!Out->IsFree());
 			return true;
 		}
@@ -445,7 +447,9 @@ Done:
 				}
 			}
 
+#if !(UE_BUILD_SHIPPING || UE_BUILD_TEST)
 			RHIBindDebugLabelName(Found->RenderTargetItem.TargetableTexture, InDebugName);
+#endif
 		}
 		else if (!bDeferTextureAllocation)
 		{
@@ -488,7 +492,9 @@ Done:
 				}
 			}
 
+#if !(UE_BUILD_SHIPPING || UE_BUILD_TEST)
 			RHIBindDebugLabelName(Found->RenderTargetItem.ShaderResourceTexture, InDebugName);
+#endif
 		}
 
 		if ((Desc.TargetableFlags & TexCreate_UAV) && !bDeferTextureAllocation)
@@ -549,10 +555,12 @@ Done:
 	// Transient RTs have to be targettable
 	check( ( Desc.Flags & TexCreate_Transient ) == 0 || Found->GetRenderTargetItem().TargetableTexture != nullptr );
 
+#if !(UE_BUILD_SHIPPING || UE_BUILD_TEST)
 	if (Found->GetRenderTargetItem().TargetableTexture)
 	{
 		RHIBindDebugLabelName(Found->GetRenderTargetItem().TargetableTexture, InDebugName);
 	}
+#endif
 
 	return false;
 }
diff --git a/Engine/Source/Runtime/RenderCore/Public/RenderGraphBuilder.h b/Engine/Source/Runtime/RenderCore/Public/RenderGraphBuilder.h
--- a/Engine/Source/Runtime/RenderCore/Public/RenderGraphBuilder.h
+++ b/Engine/Source/Runtime/RenderCore/Public/RenderGraphBuilder.h
@@ -194,6 +194,9 @@ private:
 	TMap<const IPooledRenderTarget*, FRDGTexture*, SceneRenderingSetAllocator> ExternalTextures;
 	TMap<const FPooledRDGBuffer*, FRDGBuffer*, SceneRenderingSetAllocator> ExternalBuffers;
 
+	TSet<FRDGTextureRef, DefaultKeyFuncs<FRDGTextureRef>, SceneRenderingSetAllocator> ReadTextures;
+	TSet<FRDGTextureRef, DefaultKeyFuncs<FRDGTextureRef>, SceneRenderingSetAllocator> ModifiedTextures;
+
 	uint32 BufferCount = 0;
 	uint32 TextureCount = 0;
 
