Fix running -nullrhi with native shader libraries (UE-91432).

diff --git a/Engine/Source/Runtime/RenderCore/Private/ShaderCodeLibrary.cpp b/Engine/Source/Runtime/RenderCore/Private/ShaderCodeLibrary.cpp
--- a/Engine/Source/Runtime/RenderCore/Private/ShaderCodeLibrary.cpp
+++ b/Engine/Source/Runtime/RenderCore/Private/ShaderCodeLibrary.cpp
@@ -1986,7 +1986,8 @@ void FShaderCodeLibrary::InitForRuntime(EShaderPlatform ShaderPlatform)
 	bool bArchive = false;
 	GConfig->GetBool(TEXT("/Script/UnrealEd.ProjectPackagingSettings"), TEXT("bShareMaterialShaderCode"), bArchive, GGameIni);
 
-	bool bEnable = !FPlatformProperties::IsServerOnly() && bArchive;
+	// We cannot enable native shader libraries when running with NullRHI, so for consistency all libraries (both native and non-native) are disabled if FApp::CanEverRender() == false
+	bool bEnable = !FPlatformProperties::IsServerOnly() && FApp::CanEverRender() && bArchive;
 #if !UE_BUILD_SHIPPING
 	FString FileHostIP;
 	const bool bCookOnTheFly = FParse::Value(FCommandLine::Get(), TEXT("filehostip"), FileHostIP);
diff --git a/Engine/Source/Runtime/RenderCore/Private/ShaderMap.cpp b/Engine/Source/Runtime/RenderCore/Private/ShaderMap.cpp
--- a/Engine/Source/Runtime/RenderCore/Private/ShaderMap.cpp
+++ b/Engine/Source/Runtime/RenderCore/Private/ShaderMap.cpp
@@ -250,7 +250,10 @@ bool FShaderMapBase::Serialize(FArchive& Ar, bool bInlineShaderResources, bool b
 			Resource = FShaderCodeLibrary::LoadResource(ResourceHash, &Ar);
 			if (!Resource)
 			{
-				UE_LOG(LogShaders, Error, TEXT("Missing shader resource for hash '%s' in the shader library"), *ResourceHash.ToString());
+				if (FApp::CanEverRender())	// when running -nullrhi, the resource may not be created
+				{
+					UE_LOG(LogShaders, Error, TEXT("Missing shader resource for hash '%s' in the shader library"), *ResourceHash.ToString());
+				}
 				bContentValid = false;
 			}
 		}
