Do not log errors when loading shaders for other RHIs.

- If a game targets both Vulkan and D3D, when running with -vulkan D3D shaders will not be loaded.

diff --git a/Engine/Source/Runtime/RenderCore/Private/ShaderMap.cpp b/Engine/Source/Runtime/RenderCore/Private/ShaderMap.cpp
--- a/Engine/Source/Runtime/RenderCore/Private/ShaderMap.cpp
+++ b/Engine/Source/Runtime/RenderCore/Private/ShaderMap.cpp
@@ -250,9 +250,11 @@ bool FShaderMapBase::Serialize(FArchive& Ar, bool bInlineShaderResources, bool b
 			Resource = FShaderCodeLibrary::LoadResource(ResourceHash, &Ar);
 			if (!Resource)
 			{
-				if (FApp::CanEverRender())	// when running -nullrhi, the resource may not be created
+				// do not warn when running -nullrhi (the resource cannot be created as the shader library will not be uninitialized),
+				// also do not warn for shader platforms other than current (if the game targets more than one RHI)
+				if (FApp::CanEverRender() && GetShaderPlatform() == GMaxRHIShaderPlatform)
 				{
-					UE_LOG(LogShaders, Error, TEXT("Missing shader resource for hash '%s' in the shader library"), *ResourceHash.ToString());
+					UE_LOG(LogShaders, Error, TEXT("Missing shader resource for hash '%s' for shader platform %d in the shader library"), *ResourceHash.ToString(), GetShaderPlatform());
 				}
 				bContentValid = false;
 			}
