Cache hashed filename on FShaderType, avoid semi-expensive TCHAR*->FName->FHashedName conversion when looking up global shaders

diff --git a/Engine/Source/Runtime/RenderCore/Private/GlobalShader.cpp b/Engine/Source/Runtime/RenderCore/Private/GlobalShader.cpp
--- a/Engine/Source/Runtime/RenderCore/Private/GlobalShader.cpp
+++ b/Engine/Source/Runtime/RenderCore/Private/GlobalShader.cpp
@@ -298,8 +298,7 @@ FGlobalShaderMap::~FGlobalShaderMap()
 
 TShaderRef<FShader> FGlobalShaderMap::GetShader(FShaderType* ShaderType, int32 PermutationId) const
 {
-	const FHashedName HashedFilename(ShaderType->GetShaderFilename());
-	FGlobalShaderMapSection* const* Section = SectionMap.Find(HashedFilename);
+	FGlobalShaderMapSection* const* Section = SectionMap.Find(ShaderType->GetHashedShaderFilename());
 	return Section ? (*Section)->GetShader(ShaderType, PermutationId) : TShaderRef<FShader>();
 }
 
@@ -359,7 +358,7 @@ void FGlobalShaderMap::Empty()
 
 FShader* FGlobalShaderMap::FindOrAddShader(const FShaderType* ShaderType, FShader* Shader)
 {
-	const FHashedName HashedFilename(ShaderType->GetShaderFilename());
+	const FHashedName HashedFilename(ShaderType->GetHashedShaderFilename());
 	FGlobalShaderMapSection*& Section = SectionMap.FindOrAdd(HashedFilename);
 	if (!Section)
 	{
@@ -380,8 +379,7 @@ FShaderPipeline* FGlobalShaderMap::FindOrAddShaderPipeline(const FShaderPipeline
 
 void FGlobalShaderMap::RemoveShaderTypePermutaion(const FShaderType* Type, int32 PermutationId)
 {
-	const FHashedName HashedFilename(Type->GetShaderFilename());
-	FGlobalShaderMapSection** Section = SectionMap.Find(HashedFilename);
+	FGlobalShaderMapSection** Section = SectionMap.Find(Type->GetHashedShaderFilename());
 	if (Section)
 	{
 		(*Section)->GetMutableContent()->RemoveShaderTypePermutaion(Type, PermutationId);
@@ -413,7 +411,7 @@ FGlobalShaderMapSection* FGlobalShaderMap::FindSection(const FHashedName& Hashed
 
 FGlobalShaderMapSection* FGlobalShaderMap::FindOrAddSection(const FShaderType* ShaderType)
 {
-	const FHashedName HashedFilename(ShaderType->GetShaderFilename());
+	const FHashedName HashedFilename(ShaderType->GetHashedShaderFilename());
 	FGlobalShaderMapSection* Section = FindSection(HashedFilename);
 	if(!Section)
 	{
diff --git a/Engine/Source/Runtime/RenderCore/Private/Shader.cpp b/Engine/Source/Runtime/RenderCore/Private/Shader.cpp
--- a/Engine/Source/Runtime/RenderCore/Private/Shader.cpp
+++ b/Engine/Source/Runtime/RenderCore/Private/Shader.cpp
@@ -193,6 +193,7 @@ FShaderType::FShaderType(
 	Name(InName),
 	TypeName(InName),
 	HashedName(TypeName),
+	HashedSourceFilename(InSourceFilename),
 	SourceFilename(InSourceFilename),
 	FunctionName(InFunctionName),
 	Frequency(InFrequency),
diff --git a/Engine/Source/Runtime/RenderCore/Public/Shader.h b/Engine/Source/Runtime/RenderCore/Public/Shader.h
--- a/Engine/Source/Runtime/RenderCore/Public/Shader.h
+++ b/Engine/Source/Runtime/RenderCore/Public/Shader.h
@@ -1125,6 +1125,10 @@ public:
 	{ 
 		return SourceFilename; 
 	}
+	inline const FHashedName& GetHashedShaderFilename() const
+	{
+		return HashedSourceFilename;
+	}
 	inline const TCHAR* GetFunctionName() const
 	{
 		return FunctionName;
@@ -1175,6 +1179,7 @@ private:
 	const TCHAR* Name;
 	FName TypeName;
 	FHashedName HashedName;
+	FHashedName HashedSourceFilename;
 	const TCHAR* SourceFilename;
 	const TCHAR* FunctionName;
 	uint32 Frequency;
