Add freezable hash table, use to accelerate lookup of shaders within ShaderMap
- Shaders are searched by TypeName and PermutationId, so move those fields out of FShader and into separate arrays inside ShaderMap, to reduce cache misses while finding shaders

diff --git a/Engine/Plugins/Compositing/OpenColorIO/Source/OpenColorIO/Private/OpenColorIODerivedDataVersion.h b/Engine/Plugins/Compositing/OpenColorIO/Source/OpenColorIO/Private/OpenColorIODerivedDataVersion.h
--- a/Engine/Plugins/Compositing/OpenColorIO/Source/OpenColorIO/Private/OpenColorIODerivedDataVersion.h
+++ b/Engine/Plugins/Compositing/OpenColorIO/Source/OpenColorIO/Private/OpenColorIODerivedDataVersion.h
@@ -10,4 +10,4 @@ OpenColorIOShaderDerivedDataVersion.h: Derived data version for OpenColorIO shad
 
 // In case of merge conflicts with DDC versions, you MUST generate a new GUID and set this new
 // guid as version
-#define OPENCOLORIO_DERIVEDDATA_VER		TEXT("883BB66A2A0E4A5BB85621F97919BAF8")
\ No newline at end of file
+#define OPENCOLORIO_DERIVEDDATA_VER		TEXT("176398465B0E4D26A11F2DE8FA3E6916")
\ No newline at end of file
diff --git a/Engine/Plugins/Compositing/OpenColorIO/Source/OpenColorIO/Private/OpenColorIOShaderMap.cpp b/Engine/Plugins/Compositing/OpenColorIO/Source/OpenColorIO/Private/OpenColorIOShaderMap.cpp
--- a/Engine/Plugins/Compositing/OpenColorIO/Source/OpenColorIO/Private/OpenColorIOShaderMap.cpp
+++ b/Engine/Plugins/Compositing/OpenColorIO/Source/OpenColorIO/Private/OpenColorIOShaderMap.cpp
@@ -490,7 +490,7 @@ FShader* FOpenColorIOShaderMap::ProcessCompilationResultsForSingleJob(FShaderCom
 	FOpenColorIOPixelShader *OpenColorIOShader = static_cast<FOpenColorIOPixelShader*>(Shader);
 	check(Shader);
 	check(!GetContent()->HasShader(OpenColorIOShaderType, /* PermutationId = */ 0));
-	return GetMutableContent()->FindOrAddShader(Shader);
+	return GetMutableContent()->FindOrAddShader(OpenColorIOShaderType->GetHashedName(), 0, Shader);
 }
 
 bool FOpenColorIOShaderMap::ProcessCompilationResults(const TArray<TSharedRef<FShaderCommonCompileJob, ESPMode::ThreadSafe>>& InCompilationResults, int32& InOutJobIndex, float& InOutTimeBudget)
diff --git a/Engine/Plugins/FX/Niagara/Source/NiagaraShader/Private/NiagaraShader.cpp b/Engine/Plugins/FX/Niagara/Source/NiagaraShader/Private/NiagaraShader.cpp
--- a/Engine/Plugins/FX/Niagara/Source/NiagaraShader/Private/NiagaraShader.cpp
+++ b/Engine/Plugins/FX/Niagara/Source/NiagaraShader/Private/NiagaraShader.cpp
@@ -796,7 +796,7 @@ FShader* FNiagaraShaderMap::ProcessCompilationResultsForSingleJob(TSharedRef<FSh
 	// UE-67395 - we had a case where we polluted the DDC with a shader containing no bytecode.
 	check(Shader && Shader->GetCodeSize() > 0);
 	check(!GetContent()->HasShader(NiagaraShaderType, /* PermutationId = */ 0));
-	return GetMutableContent()->FindOrAddShader(Shader);
+	return GetMutableContent()->FindOrAddShader(NiagaraShaderType->GetHashedName(), 0, Shader);
 }
 
 bool FNiagaraShaderMap::ProcessCompilationResults(const TArray<TSharedRef<FShaderCommonCompileJob, ESPMode::ThreadSafe>>& InCompilationResults, int32& InOutJobIndex, float& TimeBudget)
diff --git a/Engine/Plugins/FX/Niagara/Source/NiagaraShader/Private/NiagaraShaderDerivedDataVersion.h b/Engine/Plugins/FX/Niagara/Source/NiagaraShader/Private/NiagaraShaderDerivedDataVersion.h
--- a/Engine/Plugins/FX/Niagara/Source/NiagaraShader/Private/NiagaraShaderDerivedDataVersion.h
+++ b/Engine/Plugins/FX/Niagara/Source/NiagaraShader/Private/NiagaraShaderDerivedDataVersion.h
@@ -11,4 +11,4 @@ NiagaraShaderDerivedDataVersion.h: Shader derived data version for Niagara.
 // In case of merge conflicts with DDC versions, you MUST generate a new GUID and set this new
 // guid as version
 
-#define NIAGARASHADERMAP_DERIVEDDATA_VER		TEXT("906AEEB436374DF1ACE9A4AA6BABCED1")
+#define NIAGARASHADERMAP_DERIVEDDATA_VER		TEXT("40A7B0D0F38847349C93BD4445DAEDA5")
diff --git a/Engine/Shaders/Public/ShaderVersion.ush b/Engine/Shaders/Public/ShaderVersion.ush
--- a/Engine/Shaders/Public/ShaderVersion.ush
+++ b/Engine/Shaders/Public/ShaderVersion.ush
@@ -3,4 +3,4 @@
 // in Platform.ush (which should be included in any shader) it allows to invalidate the shader DDC.
 // 
 // If you are merging streams and there is a conflict with this GUID you should make a new GUID rather than taking one or the other.
-// GUID = 8F43812D-41AA-4930-8F86-2AE388575C56
\ No newline at end of file
+// GUID = B83EB8AE-77C7-404F-BD81-9F1CCC330F59
diff --git a/Engine/Source/Runtime/Core/Public/Containers/HashTable.h b/Engine/Source/Runtime/Core/Public/Containers/HashTable.h
--- a/Engine/Source/Runtime/Core/Public/Containers/HashTable.h
+++ b/Engine/Source/Runtime/Core/Public/Containers/HashTable.h
@@ -3,6 +3,7 @@
 #pragma once
 
 #include "CoreTypes.h"
+#include "Serialization/MemoryLayout.h"
 #include "Misc/AssertionMacros.h"
 #include "HAL/UnrealMemory.h"
 #include "Math/UnrealMathUtility.h"
@@ -362,3 +363,250 @@ inline void FHashTable::Remove( uint16 Key, uint32 Index )
 		}
 	}
 }
+
+template<typename InAllocator>
+class THashTable
+{
+public:
+	using Allocator = InAllocator;
+
+	using ElementAllocatorType = typename TChooseClass<
+		Allocator::NeedsElementType,
+		typename Allocator::template ForElementType<uint32>,
+		typename Allocator::ForAnyElementType
+	>::Result;
+
+	explicit THashTable(uint32 InHashSize = 1024, uint32 InIndexSize = 0);
+	THashTable(const THashTable& Other) = delete;
+	THashTable(THashTable&& Other) { MoveAssign(MoveTemp(Other)); }
+	~THashTable();
+
+	THashTable& operator=(const THashTable& Other) = delete;
+	THashTable& operator=(THashTable&& Other) { return MoveAssign(MoveTemp(Other)); }
+
+	THashTable&		MoveAssign(THashTable&& Other);
+	void			Clear();
+	void			Resize(uint32 NewIndexSize);
+
+	const uint32*	GetNextIndices() const { return (uint32*)NextIndex.GetAllocation(); }
+
+	// Functions used to search
+	uint32			First(uint16 Key) const;
+	uint32			Next(uint32 Index) const;
+	bool			IsValid(uint32 Index) const;
+	bool			Contains(uint16 Key) const;
+
+	void			Add(uint16 Key, uint32 Index);
+	void			Remove(uint16 Key, uint32 Index);
+
+private:
+	FORCEINLINE uint32 HashAt(uint32 Index) const { return ((uint32*)Hash.GetAllocation())[Index]; }
+	FORCEINLINE uint32 NextIndexAt(uint32 Index) const { return ((uint32*)NextIndex.GetAllocation())[Index]; }
+	FORCEINLINE uint32& HashAt(uint32 Index) { return ((uint32*)Hash.GetAllocation())[Index]; }
+	FORCEINLINE uint32& NextIndexAt(uint32 Index) { return ((uint32*)NextIndex.GetAllocation())[Index]; }
+
+	ElementAllocatorType	Hash;
+	ElementAllocatorType	NextIndex;
+	uint32					HashMask;
+	uint32					IndexSize;
+
+	template<bool bFreezeMemoryImage, typename Dummy = void>
+	struct TSupportsFreezeMemoryImageHelper
+	{
+		static void WriteMemoryImage(FMemoryImageWriter& Writer, const THashTable&)
+		{
+			check(false);
+		}
+
+		static void CopyUnfrozen(const FMemoryUnfreezeContent& Context, const THashTable&, void* Dst) { new(Dst) THashTable(); }
+	};
+
+	template<typename Dummy>
+	struct TSupportsFreezeMemoryImageHelper<true, Dummy>
+	{
+		static void WriteMemoryImage(FMemoryImageWriter& Writer, const THashTable& Object)
+		{
+			Object.Hash.WriteMemoryImage(Writer, StaticGetTypeLayoutDesc<uint32>(), Object.HashMask + 1u);
+			Object.NextIndex.WriteMemoryImage(Writer, StaticGetTypeLayoutDesc<uint32>(), Object.IndexSize);
+			Writer.WriteBytes(Object.HashMask);
+			Writer.WriteBytes(Object.IndexSize);
+		}
+		static void CopyUnfrozen(const FMemoryUnfreezeContent& Context, const THashTable& Object, void* Dst)
+		{
+			THashTable* DstTable = new(Dst) THashTable(Object.HashMask + 1u, Object.IndexSize);
+			FMemory::Memcpy(DstTable->Hash.GetAllocation(), Object.Hash.GetAllocation(), (Object.HashMask + 1u) * 4);
+			FMemory::Memcpy(DstTable->NextIndex.GetAllocation(), Object.NextIndex.GetAllocation(), Object.IndexSize * 4);
+		}
+	};
+
+public:
+	void WriteMemoryImage(FMemoryImageWriter& Writer) const
+	{
+		static const bool bSupportsFreezeMemoryImage = TAllocatorTraits<Allocator>::SupportsFreezeMemoryImage;
+		TSupportsFreezeMemoryImageHelper<bSupportsFreezeMemoryImage>::WriteMemoryImage(Writer, *this);
+	}
+
+	void CopyUnfrozen(const FMemoryUnfreezeContent& Context, void* Dst) const
+	{
+		static const bool bSupportsFreezeMemoryImage = TAllocatorTraits<Allocator>::SupportsFreezeMemoryImage;
+		TSupportsFreezeMemoryImageHelper<bSupportsFreezeMemoryImage>::CopyUnfrozen(Context, *this, Dst);
+	}
+};
+
+template<typename InAllocator>
+FORCEINLINE THashTable<InAllocator>::THashTable(uint32 InHashSize, uint32 InIndexSize)
+	: HashMask(InHashSize - 1u)
+	, IndexSize(InIndexSize)
+{
+	check(InHashSize > 0u && InHashSize <= 0x10000);
+	check(FMath::IsPowerOfTwo(InHashSize));
+
+	Hash.ResizeAllocation(0, InHashSize, sizeof(uint32));
+	FMemory::Memset(Hash.GetAllocation(), 0xff, InHashSize * 4);
+
+	if (IndexSize)
+	{
+		NextIndex.ResizeAllocation(0, IndexSize, sizeof(uint32));
+	}
+}
+
+template<typename InAllocator>
+FORCEINLINE THashTable<InAllocator>::~THashTable()
+{
+}
+
+template<typename InAllocator>
+THashTable<InAllocator>& THashTable<InAllocator>::MoveAssign(THashTable&& Other)
+{
+	Hash.MoveToEmpty(Other.Hash);
+	NextIndex.MoveToEmpty(Other.NextIndex);
+	HashMask = Other.HashMask;
+	IndexSize = Other.IndexSize;
+	Other.HashMask = 0u;
+	Other.IndexSize = 0u;
+	return *this;
+}
+
+template<typename InAllocator>
+FORCEINLINE void THashTable<InAllocator>::Clear()
+{
+	if (IndexSize)
+	{
+		FMemory::Memset(Hash.GetAllocation(), 0xff, (HashMask + 1u) * 4);
+	}
+}
+
+// First in hash chain
+template<typename InAllocator>
+FORCEINLINE uint32 THashTable<InAllocator>::First(uint16 Key) const
+{
+	Key &= HashMask;
+	return HashAt(Key);
+}
+
+// Next in hash chain
+template<typename InAllocator>
+FORCEINLINE uint32 THashTable<InAllocator>::Next(uint32 Index) const
+{
+	checkSlow(Index < IndexSize);
+	const uint32 Next = NextIndexAt(Index);
+	checkSlow(Next != Index); // check for corrupt tables
+	return Next;
+}
+
+template<typename InAllocator>
+FORCEINLINE bool THashTable<InAllocator>::IsValid(uint32 Index) const
+{
+	return Index != ~0u;
+}
+
+template<typename InAllocator>
+FORCEINLINE bool THashTable<InAllocator>::Contains(uint16 Key) const
+{
+	return First(Key) != ~0u;
+}
+
+template<typename InAllocator>
+FORCEINLINE void THashTable<InAllocator>::Add(uint16 Key, uint32 Index)
+{
+	if (Index >= IndexSize)
+	{
+		Resize(FMath::Max<uint32>(32u, FMath::RoundUpToPowerOfTwo(Index + 1)));
+	}
+
+	Key &= HashMask;
+	NextIndexAt(Index) = HashAt(Key);
+	HashAt(Key) = Index;
+}
+
+template<typename InAllocator>
+inline void THashTable<InAllocator>::Remove(uint16 Key, uint32 Index)
+{
+	if (Index >= IndexSize)
+	{
+		return;
+	}
+
+	Key &= HashMask;
+	if (HashAt(Key) == Index)
+	{
+		// Head of chain
+		HashAt(Key) = NextIndexAt(Index);
+	}
+	else
+	{
+		for (uint32 i = HashAt(Key); IsValid(i); i = NextIndexAt(i))
+		{
+			if (NextIndexAt(i) == Index)
+			{
+				// Next = Next->Next
+				NextIndexAt(i) = NextIndexAt(Index);
+				break;
+			}
+		}
+	}
+}
+
+template<typename InAllocator>
+void THashTable<InAllocator>::Resize(uint32 NewIndexSize)
+{
+	if (NewIndexSize != IndexSize)
+	{
+		NextIndex.ResizeAllocation(IndexSize, NewIndexSize, sizeof(uint32));
+		IndexSize = NewIndexSize;
+	}
+}
+
+namespace Freeze
+{
+	template<typename InAllocator>
+	void IntrinsicWriteMemoryImage(FMemoryImageWriter& Writer, const THashTable<InAllocator>& Object, const FTypeLayoutDesc&)
+	{
+		Object.WriteMemoryImage(Writer);
+	}
+
+	template<typename InAllocator>
+	void IntrinsicUnfrozenCopy(const FMemoryUnfreezeContent& Context, const THashTable<InAllocator>& Object, void* OutDst)
+	{
+		Object.CopyUnfrozen(Context, OutDst);
+	}
+
+	template<typename InAllocator>
+	uint32 IntrinsicAppendHash(const THashTable<InAllocator>* DummyObject, const FTypeLayoutDesc& TypeDesc, const FPlatformTypeLayoutParameters& LayoutParams, FSHA1& Hasher)
+	{
+		// sizeof(TArray) changes depending on target platform 32bit vs 64bit
+		// For now, calculate the size manually
+		static_assert(sizeof(THashTable<InAllocator>) == sizeof(FMemoryImageUPtrInt) * 2 + sizeof(uint32) * 2, "Unexpected THashTable size");
+		const uint32 SizeFromFields = LayoutParams.GetMemoryImagePointerSize() * 2u + sizeof(uint32) * 2u;
+		return AppendHashForNameAndSize(TypeDesc.Name, SizeFromFields, Hasher);
+	}
+
+	template<typename InAllocator>
+	uint32 IntrinsicGetTargetAlignment(const THashTable<InAllocator>* DummyObject, const FTypeLayoutDesc& TypeDesc, const FPlatformTypeLayoutParameters& LayoutParams)
+	{
+		// Assume alignment of array is drive by pointer
+		return FMath::Min(LayoutParams.GetMemoryImagePointerSize(), LayoutParams.MaxFieldAlignment);
+	}
+}
+
+DECLARE_TEMPLATE_INTRINSIC_TYPE_LAYOUT(template <typename InAllocator>, THashTable<InAllocator>);
diff --git a/Engine/Source/Runtime/Engine/Private/Materials/MaterialShader.cpp b/Engine/Source/Runtime/Engine/Private/Materials/MaterialShader.cpp
--- a/Engine/Source/Runtime/Engine/Private/Materials/MaterialShader.cpp
+++ b/Engine/Source/Runtime/Engine/Private/Materials/MaterialShader.cpp
@@ -1259,35 +1259,13 @@ void FMaterialShaderMap::FinalizeContent()
 	FMaterialShaderMapContent* LocalContent = GetMutableContent();
 
 	FSHA1 Hasher;
-	for (const FShader* Shader : LocalContent->GetShaders())
-	{
-		const FHashedName& TypeName = Shader->GetTypeName();
-		const int32 PermutationId = Shader->GetPermutationId();
-		Hasher.Update((uint8*)&TypeName, sizeof(TypeName));
-		Hasher.Update((uint8*)&PermutationId, sizeof(PermutationId));
-	}
-
-	for (const FShaderPipeline* Pipeline : LocalContent->GetShaderPipelines())
-	{
-		const FHashedName& TypeName = Pipeline->TypeName;
-		Hasher.Update((uint8*)&TypeName, sizeof(TypeName));
-	}
+	LocalContent->Finalize();
+	LocalContent->UpdateHash(Hasher);
 
 	for (FMeshMaterialShaderMap* MeshShaderMap : LocalContent->OrderedMeshShaderMaps)
 	{
-		for (const FShader* Shader : MeshShaderMap->GetShaders())
-		{
-			const FHashedName& TypeName = Shader->GetTypeName();
-			const int32 PermutationId = Shader->GetPermutationId();
-			Hasher.Update((uint8*)&TypeName, sizeof(TypeName));
-			Hasher.Update((uint8*)&PermutationId, sizeof(PermutationId));
-		}
-
-		for (const FShaderPipeline* Pipeline : MeshShaderMap->GetShaderPipelines())
-		{
-			const FHashedName& TypeName = Pipeline->TypeName;
-			Hasher.Update((uint8*)&TypeName, sizeof(TypeName));
-		}
+		MeshShaderMap->Finalize();
+		MeshShaderMap->UpdateHash(Hasher);
 	}
 
 	Hasher.Final();
@@ -1613,7 +1591,7 @@ FShader* FMaterialShaderMap::ProcessCompilationResultsForSingleJob(FShaderCompil
 		if (!ShaderPipeline)
 		{
 			check(!MeshShaderMap->HasShader(MeshMaterialShaderType, CurrentJob.PermutationId));
-			Shader = MeshShaderMap->FindOrAddShader(Shader);
+			Shader = MeshShaderMap->FindOrAddShader(MeshMaterialShaderType->GetHashedName(), CurrentJob.PermutationId, Shader);
 		}
 	}
 	else
@@ -1626,7 +1604,7 @@ FShader* FMaterialShaderMap::ProcessCompilationResultsForSingleJob(FShaderCompil
 		if (!ShaderPipeline)
 		{
 			check(!GetContent()->HasShader(MaterialShaderType, CurrentJob.PermutationId));
-			Shader = GetMutableContent()->FindOrAddShader(Shader);
+			Shader = GetMutableContent()->FindOrAddShader(MaterialShaderType->GetHashedName(), CurrentJob.PermutationId, Shader);
 		}
 	}
 
@@ -1673,33 +1651,31 @@ bool FMaterialShaderMap::ProcessCompilationResults(const TArray<TSharedRef<FShad
 			const FShaderPipelineCompileJob& CurrentJob = *PipelineJob;
 			check(CurrentJob.Id == CompilingId);
 
-			TArray<FShader*> ShaderStages;
 			FVertexFactoryType* VertexFactoryType = CurrentJob.StageJobs[0]->GetSingleShaderJob()->VFType;
+			FShaderPipeline* ShaderPipeline = new FShaderPipeline(CurrentJob.ShaderPipeline);
 			for (int32 Index = 0; Index < CurrentJob.StageJobs.Num(); ++Index)
 			{
 				SingleJob = CurrentJob.StageJobs[Index]->GetSingleShaderJob();
 				FShader* Shader = ProcessCompilationResultsForSingleJob(SingleJob, PipelineJob->ShaderPipeline, MaterialShaderMapHash, ResourceBuilder);
-				ShaderStages.Add(Shader);
+				ShaderPipeline->AddShader(Shader, SingleJob->PermutationId);
 				check(VertexFactoryType == CurrentJob.StageJobs[Index]->GetSingleShaderJob()->VFType);
 			}
+			ShaderPipeline->Validate(CurrentJob.ShaderPipeline);
 
-			FShaderPipeline* ShaderPipeline = new FShaderPipeline(PipelineJob->ShaderPipeline, ShaderStages);
-			if (ShaderPipeline)
+			if (VertexFactoryType)
 			{
-				if (VertexFactoryType)
-				{
-					check(VertexFactoryType->IsUsedWithMaterials());
-					FMeshMaterialShaderMap* MeshShaderMap = GetMeshShaderMap(VertexFactoryType);
-					check(MeshShaderMap);
-					check(!MeshShaderMap->HasShaderPipeline(ShaderPipeline->TypeName));
-					MeshShaderMap->FindOrAddShaderPipeline(ShaderPipeline);
-				}
-				else
-				{
-					check(!GetContent()->HasShaderPipeline(ShaderPipeline->TypeName));
-					GetMutableContent()->FindOrAddShaderPipeline(ShaderPipeline);
-				}
+				check(VertexFactoryType->IsUsedWithMaterials());
+				FMeshMaterialShaderMap* MeshShaderMap = GetMeshShaderMap(VertexFactoryType);
+				check(MeshShaderMap);
+				check(!MeshShaderMap->HasShaderPipeline(ShaderPipeline->TypeName));
+				MeshShaderMap->FindOrAddShaderPipeline(ShaderPipeline);
+			}
+			else
+			{
+				check(!GetContent()->HasShaderPipeline(ShaderPipeline->TypeName));
+				GetMutableContent()->FindOrAddShaderPipeline(ShaderPipeline);
 			}
+
 		}
 
 		InOutJobIndex++;
@@ -1728,7 +1704,8 @@ bool FMaterialShaderMap::ProcessCompilationResults(const TArray<TSharedRef<FShad
 							if (ShaderPipelineType->IsMeshMaterialTypePipeline() && !MeshShaderMap->HasShaderPipeline(ShaderPipelineType))
 							{
 								auto& StageTypes = ShaderPipelineType->GetStages();
-								TArray<FShader*> ShaderStages;
+	
+								FShaderPipeline* ShaderPipeline = new FShaderPipeline(ShaderPipelineType);
 								for (int32 Index = 0; Index < StageTypes.Num(); ++Index)
 								{
 									FMeshMaterialShaderType* ShaderType = ((FMeshMaterialShaderType*)(StageTypes[Index]))->GetMeshMaterialShaderType();
@@ -1736,11 +1713,9 @@ bool FMaterialShaderMap::ProcessCompilationResults(const TArray<TSharedRef<FShad
 #if DO_CHECK
 									UE_CLOG(Shader == nullptr, LogShaders, Fatal, TEXT("Failed to get ShaderType %s Permutation %d from MeshMaterial ShaderPipeline %s"), ShaderType->GetName(), kUniqueShaderPermutationId, ShaderPipelineType->GetName());
 #endif	// DO_CHECK
-									ShaderStages.Add(Shader);
+									ShaderPipeline->AddShader(Shader, kUniqueShaderPermutationId);
 								}
-
-								checkf(StageTypes.Num() == ShaderStages.Num(), TEXT("Internal Error adding MeshMaterial ShaderPipeline %s"), ShaderPipelineType->GetName());
-								FShaderPipeline* ShaderPipeline = new FShaderPipeline(ShaderPipelineType, ShaderStages);
+								ShaderPipeline->Validate(ShaderPipelineType);
 								MeshShaderMap->FindOrAddShaderPipeline(ShaderPipeline);
 							}
 						}
@@ -1757,17 +1732,16 @@ bool FMaterialShaderMap::ProcessCompilationResults(const TArray<TSharedRef<FShad
 					if (ShaderPipelineType->IsMaterialTypePipeline() && !GetContent()->HasShaderPipeline(ShaderPipelineType))
 					{
 						auto& StageTypes = ShaderPipelineType->GetStages();
-						TArray<FShader*> ShaderStages;
+
+						FShaderPipeline* ShaderPipeline = new FShaderPipeline(ShaderPipelineType);
 						for (int32 Index = 0; Index < StageTypes.Num(); ++Index)
 						{
 							FMaterialShaderType* ShaderType = ((FMaterialShaderType*)(StageTypes[Index]))->GetMaterialShaderType();
 							FShader* Shader = GetContent()->GetShader(ShaderType, kUniqueShaderPermutationId);
 							check(Shader);
-							ShaderStages.Add(Shader);
+							ShaderPipeline->AddShader(Shader, kUniqueShaderPermutationId);
 						}
-
-						checkf(StageTypes.Num() == ShaderStages.Num(), TEXT("Internal Error adding Material ShaderPipeline %s"), ShaderPipelineType->GetName());
-						FShaderPipeline* ShaderPipeline = new FShaderPipeline(ShaderPipelineType, ShaderStages);
+						ShaderPipeline->Validate(ShaderPipelineType);
 						GetMutableContent()->FindOrAddShaderPipeline(ShaderPipeline);
 					}
 				}
diff --git a/Engine/Source/Runtime/Engine/Private/ShaderCompiler/ShaderCompiler.cpp b/Engine/Source/Runtime/Engine/Private/ShaderCompiler/ShaderCompiler.cpp
--- a/Engine/Source/Runtime/Engine/Private/ShaderCompiler/ShaderCompiler.cpp
+++ b/Engine/Source/Runtime/Engine/Private/ShaderCompiler/ShaderCompiler.cpp
@@ -5006,7 +5006,7 @@ static inline FShader* ProcessCompiledJob(FShaderCompileJob* SingleJob, const FS
 		EShaderPlatform Platform = (EShaderPlatform)SingleJob->Input.Target.Platform;
 		if (!Pipeline || !Pipeline->ShouldOptimizeUnusedOutputs(Platform))
 		{
-			Shader = GGlobalShaderMap[Platform]->FindOrAddShader(GlobalShaderType, Shader);
+			Shader = GGlobalShaderMap[Platform]->FindOrAddShader(GlobalShaderType, SingleJob->PermutationId, Shader);
 			// Add this shared pipeline to the list
 			if (!Pipeline)
 			{
@@ -5052,21 +5052,19 @@ void ProcessCompiledGlobalShaders(const TArray<TSharedRef<FShaderCommonCompileJo
 		{
 			const auto* PipelineJob = CurrentJob.GetShaderPipelineJob();
 			check(PipelineJob);
-			TArray<FShader*> ShaderStages;
+
+			FShaderPipeline* ShaderPipeline = new FShaderPipeline(PipelineJob->ShaderPipeline);
 			for (int32 Index = 0; Index < PipelineJob->StageJobs.Num(); ++Index)
 			{
 				SingleJob = PipelineJob->StageJobs[Index]->GetSingleShaderJob();
 				FShader* Shader = ProcessCompiledJob(SingleJob, PipelineJob->ShaderPipeline, ShaderPlatformsProcessed, SharedPipelines);
-				ShaderStages.Add(Shader);
+				ShaderPipeline->AddShader(Shader, SingleJob->PermutationId);
 			}
+			ShaderPipeline->Validate(PipelineJob->ShaderPipeline);
 
-			FShaderPipeline* ShaderPipeline = new FShaderPipeline(PipelineJob->ShaderPipeline, ShaderStages);
-			if (ShaderPipeline)
-			{
-				EShaderPlatform Platform = (EShaderPlatform)PipelineJob->StageJobs[0]->GetSingleShaderJob()->Input.Target.Platform;
-				check(ShaderPipeline && !GGlobalShaderMap[Platform]->HasShaderPipeline(PipelineJob->ShaderPipeline));
-				GGlobalShaderMap[Platform]->FindOrAddShaderPipeline(PipelineJob->ShaderPipeline, ShaderPipeline);
-			}
+			EShaderPlatform Platform = (EShaderPlatform)PipelineJob->StageJobs[0]->GetSingleShaderJob()->Input.Target.Platform;
+			check(ShaderPipeline && !GGlobalShaderMap[Platform]->HasShaderPipeline(PipelineJob->ShaderPipeline));
+			GGlobalShaderMap[Platform]->FindOrAddShaderPipeline(PipelineJob->ShaderPipeline, ShaderPipeline);
 		}
 	}
 
@@ -5082,7 +5080,8 @@ void ProcessCompiledGlobalShaders(const TArray<TSharedRef<FShaderCommonCompileJo
 				if (!GlobalShaderMap->HasShaderPipeline(ShaderPipelineType))
 				{
 					auto& StageTypes = ShaderPipelineType->GetStages();
-					TArray<FShader*> ShaderStages;
+
+					FShaderPipeline* ShaderPipeline = new FShaderPipeline(ShaderPipelineType);
 					for (int32 Index = 0; Index < StageTypes.Num(); ++Index)
 					{
 						FGlobalShaderType* GlobalShaderType = ((FShaderType*)(StageTypes[Index]))->GetGlobalShaderType();
@@ -5090,16 +5089,14 @@ void ProcessCompiledGlobalShaders(const TArray<TSharedRef<FShaderCommonCompileJo
 						{
 							TShaderRef<FShader> Shader = GlobalShaderMap->GetShader(GlobalShaderType, kUniqueShaderPermutationId);
 							check(Shader.IsValid());
-							ShaderStages.Add(Shader.GetShader());
+							ShaderPipeline->AddShader(Shader.GetShader(), kUniqueShaderPermutationId);
 						}
 						else
 						{
 							break;
 						}
 					}
-
-					checkf(StageTypes.Num() == ShaderStages.Num(), TEXT("Internal Error adding Global ShaderPipeline %s"), ShaderPipelineType->GetName());
-					FShaderPipeline* ShaderPipeline = new FShaderPipeline(ShaderPipelineType, ShaderStages);
+					ShaderPipeline->Validate(ShaderPipelineType);
 					GlobalShaderMap->FindOrAddShaderPipeline(ShaderPipelineType, ShaderPipeline);
 				}
 			}
diff --git a/Engine/Source/Runtime/RHI/Public/RHIDefinitions.h b/Engine/Source/Runtime/RHI/Public/RHIDefinitions.h
--- a/Engine/Source/Runtime/RHI/Public/RHIDefinitions.h
+++ b/Engine/Source/Runtime/RHI/Public/RHIDefinitions.h
@@ -29,7 +29,10 @@ enum EShaderFrequency : uint8
 	SF_Geometry			= 4,
 	SF_Compute			= 5,
 
-	// Number of standard SM5-style graphics pipeline shader frequencies
+	// Number of standard SM5-style shader frequencies for graphics pipeline (excluding compute)
+	SF_NumGraphicsFrequencies = 5,
+
+	// Number of standard SM5-style shader frequencies (including compute)
 	SF_NumStandardFrequencies = 6,
 
 	SF_RayGen			= 6,
diff --git a/Engine/Source/Runtime/RenderCore/Private/GlobalShader.cpp b/Engine/Source/Runtime/RenderCore/Private/GlobalShader.cpp
--- a/Engine/Source/Runtime/RenderCore/Private/GlobalShader.cpp
+++ b/Engine/Source/Runtime/RenderCore/Private/GlobalShader.cpp
@@ -356,7 +356,7 @@ void FGlobalShaderMap::Empty()
 	}
 }
 
-FShader* FGlobalShaderMap::FindOrAddShader(const FShaderType* ShaderType, FShader* Shader)
+FShader* FGlobalShaderMap::FindOrAddShader(const FShaderType* ShaderType, int32 PermutationId, FShader* Shader)
 {
 	const FHashedName HashedFilename(ShaderType->GetHashedShaderFilename());
 	FGlobalShaderMapSection*& Section = SectionMap.FindOrAdd(HashedFilename);
@@ -364,7 +364,7 @@ FShader* FGlobalShaderMap::FindOrAddShader(const FShaderType* ShaderType, FShade
 	{
 		Section = new FGlobalShaderMapSection(Platform, HashedFilename);
 	}
-	return Section->GetMutableContent()->FindOrAddShader(Shader);
+	return Section->GetMutableContent()->FindOrAddShader(ShaderType->GetHashedName(), PermutationId, Shader);
 }
 
 FShaderPipeline* FGlobalShaderMap::FindOrAddShaderPipeline(const FShaderPipelineType* ShaderPipelineType, FShaderPipeline* ShaderPipeline)
@@ -382,7 +382,7 @@ void FGlobalShaderMap::RemoveShaderTypePermutaion(const FShaderType* Type, int32
 	FGlobalShaderMapSection** Section = SectionMap.Find(Type->GetHashedShaderFilename());
 	if (Section)
 	{
-		(*Section)->GetMutableContent()->RemoveShaderTypePermutaion(Type, PermutationId);
+		(*Section)->GetMutableContent()->RemoveShaderTypePermutaion(Type->GetHashedName(), PermutationId);
 	}
 }
 
diff --git a/Engine/Source/Runtime/RenderCore/Private/RenderGraphUtils.cpp b/Engine/Source/Runtime/RenderCore/Private/RenderGraphUtils.cpp
--- a/Engine/Source/Runtime/RenderCore/Private/RenderGraphUtils.cpp
+++ b/Engine/Source/Runtime/RenderCore/Private/RenderGraphUtils.cpp
@@ -447,7 +447,7 @@ public:
 	FClearUAVUIntCS(const ShaderMetaType::CompiledShaderInitializerType& Initializer)
 		: FGlobalShader(Initializer)
 	{
-		Bindings.BindForLegacyShaderParameters(this, Initializer.ParameterMap, *FParameters::FTypeInfo::GetStructMetadata());
+		Bindings.BindForLegacyShaderParameters(this, Initializer.PermutationId, Initializer.ParameterMap, *FParameters::FTypeInfo::GetStructMetadata());
 	}
 
 	FClearUAVUIntCS()
@@ -498,7 +498,7 @@ public:
 	FClearUAVFloatCS(const ShaderMetaType::CompiledShaderInitializerType& Initializer)
 		: FGlobalShader(Initializer)
 	{
-		Bindings.BindForLegacyShaderParameters(this, Initializer.ParameterMap, *FParameters::FTypeInfo::GetStructMetadata());
+		Bindings.BindForLegacyShaderParameters(this, Initializer.PermutationId, Initializer.ParameterMap, *FParameters::FTypeInfo::GetStructMetadata());
 	}
 
 	FClearUAVFloatCS()
diff --git a/Engine/Source/Runtime/RenderCore/Private/Shader.cpp b/Engine/Source/Runtime/RenderCore/Private/Shader.cpp
--- a/Engine/Source/Runtime/RenderCore/Private/Shader.cpp
+++ b/Engine/Source/Runtime/RenderCore/Private/Shader.cpp
@@ -459,7 +459,6 @@ FShader::FShader()
 	// set to undefined (currently shared with SF_Vertex)
 	: Target((EShaderFrequency)0, GShaderPlatformForFeatureLevel[GMaxRHIFeatureLevel])
 	, ResourceIndex(INDEX_NONE)
-	, PermutationId(0)
 #if WITH_EDITORONLY_DATA
 	, NumInstructions(0u)
 	, NumTextureSamplers(0u)
@@ -474,10 +473,8 @@ FShader::FShader()
 FShader::FShader(const CompiledShaderInitializerType& Initializer)
 	: Type(Initializer.Type)
 	, VFType(Initializer.VertexFactoryType)
-	, TypeName(Initializer.Type->GetFName())
 	, Target(Initializer.Target)
 	, ResourceIndex(Initializer.ResourceIndex)
-	, PermutationId(Initializer.PermutationId)
 #if WITH_EDITORONLY_DATA
 	, NumInstructions(Initializer.NumInstructions)
 	, NumTextureSamplers(Initializer.NumTextureSamplers)
@@ -668,14 +665,13 @@ void FShader::DumpDebugInfo(const FShaderMapPointerTable& InPtrTable)
 	UE_LOG(LogConsoleResponse, Display, TEXT("               :Target %s"), *LegacyShaderPlatformToShaderFormat(GetShaderPlatform()).ToString());
 	UE_LOG(LogConsoleResponse, Display, TEXT("               :VFType %s"), VertexFactoryType ? VertexFactoryType->GetName() : TEXT("null"));
 	UE_LOG(LogConsoleResponse, Display, TEXT("               :Type %s"), GetType(InPtrTable)->GetName());
-	UE_LOG(LogConsoleResponse, Display, TEXT("               :PermutationId %d"), PermutationId);
 	UE_LOG(LogConsoleResponse, Display, TEXT("               :SourceHash %s"), *GetHash().ToString());
 	UE_LOG(LogConsoleResponse, Display, TEXT("               :VFSourceHash %s"), *GetVertexFactoryHash().ToString());
 	UE_LOG(LogConsoleResponse, Display, TEXT("               :OutputHash %s"), *GetOutputHash().ToString());
 }
 
 #if WITH_EDITOR
-void FShader::SaveShaderStableKeys(const FShaderMapPointerTable& InPtrTable, EShaderPlatform TargetShaderPlatform, const FStableShaderKeyAndValue& InSaveKeyVal)
+void FShader::SaveShaderStableKeys(const FShaderMapPointerTable& InPtrTable, EShaderPlatform TargetShaderPlatform, int32 PermutationId, const FStableShaderKeyAndValue& InSaveKeyVal)
 {
 	if ((TargetShaderPlatform == EShaderPlatform::SP_NumPlatforms || GetShaderPlatform() == TargetShaderPlatform) 
 		&& FShaderCodeLibrary::NeedsShaderStableKeys(TargetShaderPlatform))
@@ -912,78 +908,21 @@ const FSHAHash& FShaderPipelineType::GetSourceHash(EShaderPlatform ShaderPlatfor
 	return GetShaderFilesHash(Filenames, ShaderPlatform);
 }
 
-
-FShaderPipeline::FShaderPipeline(
-	const FShaderPipelineType* InPipelineType,
-	FShader* InVertexShader,
-	FShader* InHullShader,
-	FShader* InDomainShader,
-	FShader* InGeometryShader,
-	FShader* InPixelShader) :
-	TypeName(InPipelineType->GetHashedName()),
-	VertexShader(InVertexShader),
-	HullShader(InHullShader),
-	DomainShader(InDomainShader),
-	GeometryShader(InGeometryShader),
-	PixelShader(InPixelShader)
+void FShaderPipeline::AddShader(FShader* Shader, int32 PermutationId)
 {
-	check(InPipelineType);
-	Validate(InPipelineType);
-}
-
-FShaderPipeline::FShaderPipeline(const FShaderPipelineType* InPipelineType, const TArray<FShader*>& InStages) :
-	TypeName(InPipelineType->GetHashedName()),
-	VertexShader(nullptr),
-	HullShader(nullptr),
-	DomainShader(nullptr),
-	GeometryShader(nullptr),
-	PixelShader(nullptr)
-{
-	check(InPipelineType);
-	for (FShader* Shader : InStages)
-	{
-		if (Shader)
-		{
-			switch (Shader->GetTypeUnfrozen()->GetFrequency())
-			{
-			case SF_Vertex:
-				check(!VertexShader);
-				VertexShader = Shader;
-				break;
-			case SF_Pixel:
-				check(!PixelShader);
-				PixelShader = Shader;
-				break;
-			case SF_Hull:
-				check(!HullShader);
-				HullShader = Shader;
-				break;
-			case SF_Domain:
-				check(!DomainShader);
-				DomainShader = Shader;
-				break;
-			case SF_Geometry:
-				check(!GeometryShader);
-				GeometryShader = Shader;
-				break;
-			default:
-				checkf(0, TEXT("Invalid stage %u found!"), (uint32)Shader->GetTypeUnfrozen()->GetFrequency());
-				break;
-			}
-		}
-	}
-
-	Validate(InPipelineType);
+	const EShaderFrequency Frequency = Shader->GetFrequency();
+	check(Shaders[Frequency].IsNull());
+	Shaders[Frequency] = Shader;
+	PermutationIds[Frequency] = PermutationId;
 }
 
 FShaderPipeline::~FShaderPipeline()
 {
 	// Manually set references to nullptr, helps debugging
-	VertexShader = nullptr;
-	HullShader = nullptr;
-	DomainShader = nullptr;
-	GeometryShader = nullptr;
-	PixelShader = nullptr;
+	for (uint32 i = 0u; i < SF_NumGraphicsFrequencies; ++i)
+	{
+		Shaders[i] = nullptr;
+	}
 }
 
 void FShaderPipeline::Validate(const FShaderPipelineType* InPipelineType) const
@@ -991,27 +930,9 @@ void FShaderPipeline::Validate(const FShaderPipelineType* InPipelineType) const
 	check(InPipelineType->GetHashedName() == TypeName);
 	for (const FShaderType* Stage : InPipelineType->GetStages())
 	{
-		switch (Stage->GetFrequency())
-		{
-		case SF_Vertex:
-			check(VertexShader && VertexShader->GetTypeUnfrozen() == Stage);
-			break;
-		case SF_Pixel:
-			check(PixelShader && PixelShader->GetTypeUnfrozen() == Stage);
-			break;
-		case SF_Hull:
-			check(HullShader && HullShader->GetTypeUnfrozen() == Stage);
-			break;
-		case SF_Domain:
-			check(DomainShader && DomainShader->GetTypeUnfrozen() == Stage);
-			break;
-		case SF_Geometry:
-			check(GeometryShader && GeometryShader->GetTypeUnfrozen() == Stage);
-			break;
-		default:
-			// Can never happen :)
-			break;
-		}
+		const FShader* Shader = GetShader(Stage->GetFrequency());
+		check(Shader);
+		check(Shader->GetTypeUnfrozen() == Stage);
 	}
 }
 
@@ -1030,25 +951,15 @@ void FShaderPipeline::SaveShaderStableKeys(const FShaderMapPointerTable& InPtrTa
 	if (bCanHaveUniqueShaders)
 	{
 		FStableShaderKeyAndValue SaveKeyVal(InSaveKeyVal);
-
 		SaveKeyVal.SetPipelineHash(this); // could use PipelineType->GetSourceHash(), but each pipeline instance even of the same type can have unique shaders
 
-		// Using GetShaders() would be more future-proof but would result in more dynamic allocation churn.
-		// Instead, mirroring the logic here
-		VertexShader->SaveShaderStableKeys(InPtrTable, TargetShaderPlatform, SaveKeyVal);
-
-		if (PixelShader)
-		{
-			PixelShader->SaveShaderStableKeys(InPtrTable, TargetShaderPlatform, SaveKeyVal);
-		}
-		if (GeometryShader)
+		for (uint32 Frequency = 0u; Frequency < SF_NumGraphicsFrequencies; ++Frequency)
 		{
-			GeometryShader->SaveShaderStableKeys(InPtrTable, TargetShaderPlatform, SaveKeyVal);
-		}
-		if (HullShader)
-		{
-			HullShader->SaveShaderStableKeys(InPtrTable, TargetShaderPlatform, SaveKeyVal);
-			DomainShader->SaveShaderStableKeys(InPtrTable, TargetShaderPlatform, SaveKeyVal);
+			FShader* Shader = Shaders[Frequency];
+			if (Shader)
+			{
+				Shader->SaveShaderStableKeys(InPtrTable, TargetShaderPlatform, PermutationIds[Frequency], SaveKeyVal);
+			}
 		}
 	}
 }
diff --git a/Engine/Source/Runtime/RenderCore/Private/ShaderCodeLibrary.cpp b/Engine/Source/Runtime/RenderCore/Private/ShaderCodeLibrary.cpp
--- a/Engine/Source/Runtime/RenderCore/Private/ShaderCodeLibrary.cpp
+++ b/Engine/Source/Runtime/RenderCore/Private/ShaderCodeLibrary.cpp
@@ -503,37 +503,23 @@ void FStableShaderKeyAndValue::SetPipelineHash(const FShaderPipeline* Pipeline)
 void FShaderCodeLibraryPipeline::Initialize(const FShaderPipeline* Pipeline)
 {
 	check(Pipeline != nullptr);
-
-	if (Pipeline->VertexShader.IsValid())
-	{
-		VertexShader = Pipeline->VertexShader->GetOutputHash();
-	}
-	if (Pipeline->GeometryShader.IsValid())
+	for (uint32 Frequency = 0u; Frequency < SF_NumGraphicsFrequencies; ++Frequency)
 	{
-		GeometryShader = Pipeline->GeometryShader->GetOutputHash();
-	}
-	if (Pipeline->HullShader.IsValid())
-	{
-		HullShader = Pipeline->HullShader->GetOutputHash();
-	}
-	if (Pipeline->DomainShader.IsValid())
-	{
-		DomainShader = Pipeline->DomainShader->GetOutputHash();
-	}
-	if (Pipeline->PixelShader.IsValid())
-	{
-		PixelShader = Pipeline->PixelShader->GetOutputHash();
+		if (Pipeline->Shaders[Frequency].IsValid())
+		{
+			Shaders[Frequency] = Pipeline->Shaders[Frequency]->GetOutputHash();
+		}
 	}
 }
 
 void FShaderCodeLibraryPipeline::GetPipelineHash(FSHAHash& Output)
 {
 	FSHA1 Hasher;
-	Hasher.Update(&VertexShader.Hash[0], sizeof(VertexShader.Hash));
-	Hasher.Update(&PixelShader.Hash[0], sizeof(PixelShader.Hash));
-	Hasher.Update(&GeometryShader.Hash[0], sizeof(GeometryShader.Hash));
-	Hasher.Update(&HullShader.Hash[0], sizeof(HullShader.Hash));
-	Hasher.Update(&DomainShader.Hash[0], sizeof(DomainShader.Hash));
+	Hasher.Update(&Shaders[SF_Vertex].Hash[0], sizeof(FSHAHash));
+	Hasher.Update(&Shaders[SF_Pixel].Hash[0], sizeof(FSHAHash));
+	Hasher.Update(&Shaders[SF_Geometry].Hash[0], sizeof(FSHAHash));
+	Hasher.Update(&Shaders[SF_Hull].Hash[0], sizeof(FSHAHash));
+	Hasher.Update(&Shaders[SF_Domain].Hash[0], sizeof(FSHAHash));
 
 	Hasher.Final();
 	Hasher.GetHash(&Output.Hash[0]);
diff --git a/Engine/Source/Runtime/RenderCore/Private/ShaderMap.cpp b/Engine/Source/Runtime/RenderCore/Private/ShaderMap.cpp
--- a/Engine/Source/Runtime/RenderCore/Private/ShaderMap.cpp
+++ b/Engine/Source/Runtime/RenderCore/Private/ShaderMap.cpp
@@ -309,36 +309,66 @@ void FShaderMapBase::DestroyContent()
 	}
 }
 
+static uint16 MakeShaderHash(const FHashedName& TypeName, int32 PermutationId)
+{
+	return (uint16)CityHash128to64({ TypeName.GetHash(), (uint64)PermutationId });
+}
+
 FShader* FShaderMapContent::GetShader(const FHashedName& TypeName, int32 PermutationId) const
 {
 	TRACE_CPUPROFILER_EVENT_SCOPE(FShaderMapContent::GetShader);
-	const int32 Index = Algo::BinarySearchBy(Shaders, FShaderMapKey(TypeName, PermutationId), FProjectShaderToKey());
-	return (Index != INDEX_NONE) ? Shaders[Index].Get() : nullptr;
+	const uint16 Hash = MakeShaderHash(TypeName, PermutationId);
+	const FHashedName* RESTRICT LocalShaderTypes = ShaderTypes.GetData();
+	const int32* RESTRICT LocalShaderPermutations = ShaderPermutations.GetData();
+	const uint32* RESTRICT LocalNextHashIndices = ShaderHash.GetNextIndices();
+	const uint32 NumShaders = Shaders.Num();
+
+	for (uint32 Index = ShaderHash.First(Hash); ShaderHash.IsValid(Index); Index = LocalNextHashIndices[Index])
+	{
+		checkSlow(Index < NumShaders);
+		if (LocalShaderTypes[Index] == TypeName && LocalShaderPermutations[Index] == PermutationId)
+		{
+			return Shaders[Index].GetChecked();
+		}
+	}
+
+	return nullptr;
 }
 
-void FShaderMapContent::AddShader(FShader* Shader)
+void FShaderMapContent::AddShader(const FHashedName& TypeName, int32 PermutationId, FShader* Shader)
 {
 	check(!Shader->IsFrozen());
-	checkSlow(!HasShader(Shader->GetTypeName(), Shader->GetPermutationId()));
-	const int32 Index = Algo::LowerBoundBy(Shaders, FShaderMapKey(*Shader), FProjectShaderToKey());
-	Shaders.Insert(Shader, Index);
+	checkSlow(!HasShader(TypeName, PermutationId));
+
+	const uint16 Hash = MakeShaderHash(TypeName, PermutationId);
+	const int32 Index = Shaders.Add(Shader);
+	ShaderTypes.Add(TypeName);
+	ShaderPermutations.Add(PermutationId);
+	check(ShaderTypes.Num() == Shaders.Num());
+	check(ShaderPermutations.Num() == Shaders.Num());
+	ShaderHash.Add(Hash, Index);
 }
 
-FShader* FShaderMapContent::FindOrAddShader(FShader* Shader)
+FShader* FShaderMapContent::FindOrAddShader(const FHashedName& TypeName, int32 PermutationId, FShader* Shader)
 {
 	check(!Shader->IsFrozen());
-	const int32 Index = Algo::LowerBoundBy(Shaders, FShaderMapKey(*Shader), FProjectShaderToKey());
-	if (Index < Shaders.Num())
+
+	const uint16 Hash = MakeShaderHash(TypeName, PermutationId);
+	for (uint32 Index = ShaderHash.First(Hash); ShaderHash.IsValid(Index); Index = ShaderHash.Next(Index))
 	{
-		FShader* PrevShader = Shaders[Index];
-		if (PrevShader->GetTypeName() == Shader->GetTypeName() && PrevShader->GetPermutationId() == Shader->GetPermutationId())
+		if (ShaderTypes[Index] == TypeName && ShaderPermutations[Index] == PermutationId)
 		{
 			DeleteObjectFromLayout(Shader);
-			return PrevShader;
+			return Shaders[Index].GetChecked();
 		}
 	}
 
-	Shaders.Insert(Shader, Index);
+	const int32 Index = Shaders.Add(Shader);
+	ShaderHash.Add(Hash, Index);
+	ShaderTypes.Add(TypeName);
+	ShaderPermutations.Add(PermutationId);
+	check(ShaderTypes.Num() == Shaders.Num());
+	check(ShaderPermutations.Num() == Shaders.Num());
 	return Shader;
 }
 
@@ -370,12 +400,32 @@ FShaderPipeline* FShaderMapContent::FindOrAddShaderPipeline(FShaderPipeline* Pip
  * Removes the shader of the given type from the shader map
  * @param Type Shader type to remove the entry for
  */
-void FShaderMapContent::RemoveShaderTypePermutaion(const FShaderType* Type, int32 PermutationId)
+void FShaderMapContent::RemoveShaderTypePermutaion(const FHashedName& TypeName, int32 PermutationId)
 {
-	const int32 Index = Algo::BinarySearchBy(Shaders, FShaderMapKey(Type->GetHashedName(), PermutationId), FProjectShaderToKey());
-	if (Index != INDEX_NONE)
+	const uint16 Hash = MakeShaderHash(TypeName, PermutationId);
+
+	for (uint32 Index = ShaderHash.First(Hash); ShaderHash.IsValid(Index); Index = ShaderHash.Next(Index))
 	{
-		Shaders.RemoveAt(Index, 1, false);
+		FShader* Shader = Shaders[Index].GetChecked();
+		if (ShaderTypes[Index] == TypeName && ShaderPermutations[Index] == PermutationId)
+		{
+			DeleteObjectFromLayout(Shader);
+			Shaders.RemoveAtSwap(Index, 1, false);
+			ShaderTypes.RemoveAtSwap(Index, 1, false);
+			ShaderPermutations.RemoveAtSwap(Index, 1, false);
+			check(ShaderTypes.Num() == Shaders.Num());
+			check(ShaderPermutations.Num() == Shaders.Num());
+
+			ShaderHash.Remove(Hash, Index);
+			const int32 SwapIndex = Shaders.Num();
+			if (SwapIndex > 0)
+			{
+				const uint16 SwapHash = MakeShaderHash(ShaderTypes[Index], ShaderPermutations[Index]);
+				ShaderHash.Remove(SwapHash, SwapIndex);
+				ShaderHash.Add(SwapHash, Index);
+			}
+			break;
+		}
 	}
 }
 
@@ -392,14 +442,15 @@ void FShaderMapContent::RemoveShaderPipelineType(const FShaderPipelineType* Shad
 
 void FShaderMapContent::GetShaderList(const FShaderMapBase& InShaderMap, const FSHAHash& InMaterialShaderMapHash, TMap<FShaderId, TShaderRef<FShader>>& OutShaders) const
 {
-	for (FShader* Shader : Shaders)
+	for (int32 ShaderIndex = 0; ShaderIndex < Shaders.Num(); ++ShaderIndex)
 	{
+		FShader* Shader = Shaders[ShaderIndex].GetChecked();
 		const FShaderId ShaderId(
 			Shader->GetType(InShaderMap.GetPointerTable()),
 			InMaterialShaderMapHash,
 			FHashedName(),
 			Shader->GetVertexFactoryType(InShaderMap.GetPointerTable()),
-			Shader->GetPermutationId(),
+			ShaderPermutations[ShaderIndex],
 			GetShaderPlatform());
 
 		OutShaders.Add(ShaderId, TShaderRef<FShader>(Shader, InShaderMap));
@@ -407,33 +458,37 @@ void FShaderMapContent::GetShaderList(const FShaderMapBase& InShaderMap, const F
 
 	for (const FShaderPipeline* ShaderPipeline : ShaderPipelines)
 	{
-		for (const TShaderRef<FShader>& Shader : ShaderPipeline->GetShaders(InShaderMap))
+		for (uint32 Frequency = 0u; Frequency < SF_NumGraphicsFrequencies; ++Frequency)
 		{
-			const FShaderId ShaderId(
-				Shader.GetType(),
-				InMaterialShaderMapHash,
-				ShaderPipeline->TypeName,
-				Shader.GetVertexFactoryType(),
-				Shader->GetPermutationId(),
-				GetShaderPlatform());
-
-			OutShaders.Add(ShaderId, Shader);
+			FShader* Shader = ShaderPipeline->Shaders[Frequency].Get();
+			if (Shader)
+			{
+				const FShaderId ShaderId(
+					Shader->GetType(InShaderMap.GetPointerTable()),
+					InMaterialShaderMapHash,
+					ShaderPipeline->TypeName,
+					Shader->GetVertexFactoryType(InShaderMap.GetPointerTable()),
+					ShaderPipeline->PermutationIds[Frequency],
+					GetShaderPlatform());
+				OutShaders.Add(ShaderId, TShaderRef<FShader>(Shader, InShaderMap));
+			}
 		}
 	}
 }
 
 void FShaderMapContent::GetShaderList(const FShaderMapBase& InShaderMap, TMap<FHashedName, TShaderRef<FShader>>& OutShaders) const
 {
-	for (FShader* Shader : Shaders)
+	for (int32 ShaderIndex = 0; ShaderIndex < Shaders.Num(); ++ShaderIndex)
 	{
-		OutShaders.Add(Shader->GetTypeName(), TShaderRef<FShader>(Shader, InShaderMap));
+		FShader* Shader = Shaders[ShaderIndex].GetChecked();
+		OutShaders.Add(ShaderTypes[ShaderIndex], TShaderRef<FShader>(Shader, InShaderMap));
 	}
 
 	for (const FShaderPipeline* ShaderPipeline : ShaderPipelines)
 	{
 		for (const TShaderRef<FShader>& Shader : ShaderPipeline->GetShaders(InShaderMap))
 		{
-			OutShaders.Add(Shader->GetTypeName(), Shader);
+			OutShaders.Add(Shader.GetType()->GetHashedName(), Shader);
 		}
 	}
 }
@@ -515,9 +570,10 @@ void FShaderMapContent::GetOutdatedTypes(const FShaderMapBase& InShaderMap, TArr
 
 void FShaderMapContent::SaveShaderStableKeys(const FShaderMapBase& InShaderMap, EShaderPlatform TargetShaderPlatform, const struct FStableShaderKeyAndValue& SaveKeyVal)
 {
-	for (FShader* Shader : Shaders)
+	for (int32 ShaderIndex = 0; ShaderIndex < Shaders.Num(); ++ShaderIndex)
 	{
-		Shader->SaveShaderStableKeys(InShaderMap.GetPointerTable(), TargetShaderPlatform, SaveKeyVal);
+		const int32 PermutationId = ShaderPermutations[ShaderIndex];
+		Shaders[ShaderIndex]->SaveShaderStableKeys(InShaderMap.GetPointerTable(), TargetShaderPlatform, PermutationId, SaveKeyVal);
 	}
 
 	for (const FShaderPipeline* Pipeline : ShaderPipelines)
@@ -578,6 +634,75 @@ uint32 FShaderMapContent::GetMaxNumInstructionsForShader(const FShaderMapBase& I
 	return MaxNumInstructions;
 }
 
+struct FSortedShaderEntry
+{
+	FHashedName TypeName;
+	int32 PermutationId;
+	int32 Index;
+
+	friend bool operator<(const FSortedShaderEntry& Lhs, const FSortedShaderEntry& Rhs)
+	{
+		if (Lhs.TypeName != Rhs.TypeName)
+		{
+			return Lhs.TypeName < Rhs.TypeName;
+		}
+		return Lhs.PermutationId < Rhs.PermutationId;
+	}
+};
+
+void FShaderMapContent::Finalize()
+{
+	// Sort the shaders by type/permutation, so they are consistently ordered
+	TArray<FSortedShaderEntry> SortedEntries;
+	SortedEntries.Empty(Shaders.Num());
+	for (int32 ShaderIndex = 0; ShaderIndex < Shaders.Num(); ++ShaderIndex)
+	{
+		FSortedShaderEntry& Entry = SortedEntries.AddDefaulted_GetRef();
+		Entry.TypeName = ShaderTypes[ShaderIndex];
+		Entry.PermutationId = ShaderPermutations[ShaderIndex];
+		Entry.Index = ShaderIndex;
+	}
+	SortedEntries.Sort();
+
+	// Choose a good hash size based on the number of shaders we have
+	const uint32 HashSize = FMath::RoundUpToPowerOfTwo(FMath::Max((Shaders.Num() * 3) / 2, 1));
+	FMemoryImageHashTable NewShaderHash(HashSize, Shaders.Num());
+	TMemoryImageArray<TMemoryImagePtr<FShader>> NewShaders;
+	NewShaders.Empty(Shaders.Num());
+	ShaderTypes.Empty(Shaders.Num());
+	ShaderPermutations.Empty(Shaders.Num());
+
+	for (int32 SortedIndex = 0; SortedIndex < SortedEntries.Num(); ++SortedIndex)
+	{
+		const FSortedShaderEntry& SortedEntry = SortedEntries[SortedIndex];
+		const uint16 Key = MakeShaderHash(SortedEntry.TypeName, SortedEntry.PermutationId);
+		NewShaders.Add(Shaders[SortedEntry.Index]);
+		ShaderTypes.Add(SortedEntry.TypeName);
+		ShaderPermutations.Add(SortedEntry.PermutationId);
+		NewShaderHash.Add(Key, SortedIndex);
+	}
+
+	Shaders = MoveTemp(NewShaders);
+	ShaderHash = MoveTemp(NewShaderHash);
+}
+
+void FShaderMapContent::UpdateHash(FSHA1& Hasher) const
+{
+	for (int32 ShaderIndex = 0; ShaderIndex < Shaders.Num(); ++ShaderIndex)
+	{
+		const FHashedName& TypeName = ShaderTypes[ShaderIndex];
+		const int32 PermutationId = ShaderPermutations[ShaderIndex];
+		Hasher.Update((uint8*)&TypeName, sizeof(TypeName));
+		Hasher.Update((uint8*)&PermutationId, sizeof(PermutationId));
+	}
+
+	for (const FShaderPipeline* Pipeline : GetShaderPipelines())
+	{
+		const FHashedName& TypeName = Pipeline->TypeName;
+		Hasher.Update((uint8*)&TypeName, sizeof(TypeName));
+	}
+}
+
 void FShaderMapContent::Empty()
 {
 	EmptyShaderPipelines();
@@ -586,6 +711,9 @@ void FShaderMapContent::Empty()
 		Shaders[i].SafeDelete();
 	}
 	Shaders.Empty();
+	ShaderTypes.Empty();
+	ShaderPermutations.Empty();
+	ShaderHash.Clear();
 }
 
 void FShaderMapContent::EmptyShaderPipelines()
diff --git a/Engine/Source/Runtime/RenderCore/Private/ShaderParameterStruct.cpp b/Engine/Source/Runtime/RenderCore/Private/ShaderParameterStruct.cpp
--- a/Engine/Source/Runtime/RenderCore/Private/ShaderParameterStruct.cpp
+++ b/Engine/Source/Runtime/RenderCore/Private/ShaderParameterStruct.cpp
@@ -25,6 +25,9 @@ struct FShaderParameterStructBindingContext
 	// C++ name of the render target binding slot.
 	FString RenderTargetBindingSlotCppName;
 
+	// Shader PermutationId
+	int32 PermutationId;
+
 	// Whether this is for legacy shader parameter settings, or root shader parameter structures/
 	bool bUseRootShaderParameters;
 
@@ -178,7 +181,7 @@ struct FShaderParameterStructBindingContext
 					if (uint32(BoundSize) > ByteSize)
 					{
 						UE_LOG(LogShaders, Fatal, TEXT("The size required to bind shader %s's (Permutation Id %d) struct %s parameter %s is %i bytes, smaller than %s's %i bytes."),
-							Shader->GetTypeUnfrozen()->GetName(), Shader->GetPermutationId(), StructMetaData.GetStructTypeName(),
+							Shader->GetTypeUnfrozen()->GetName(), PermutationId, StructMetaData.GetStructTypeName(),
 							*ElementShaderBindingName, BoundSize, *CppName, ByteSize);
 					}
 
@@ -213,7 +216,7 @@ struct FShaderParameterStructBindingContext
 								TEXT("Error with shader %s's (Permutation Id %d) parameter %s is %i bytes, cpp name = %s.")
 								TEXT("The shader compiler should give precisely which elements of an array did not get compiled out, ")
 								TEXT("for optimal automatic render graph pass dependency with ClearUnusedGraphResources()."),
-								Shader->GetTypeUnfrozen()->GetName(), Shader->GetPermutationId(),
+								Shader->GetTypeUnfrozen()->GetName(), PermutationId,
 								*ElementShaderBindingName, BoundSize, *CppName);
 						}
 					}
@@ -229,7 +232,7 @@ struct FShaderParameterStructBindingContext
 	} // void Bind()
 }; // struct FShaderParameterStructBindingContext
 
-void FShaderParameterBindings::BindForLegacyShaderParameters(const FShader* Shader, const FShaderParameterMap& ParametersMap, const FShaderParametersMetadata& StructMetaData, bool bShouldBindEverything)
+void FShaderParameterBindings::BindForLegacyShaderParameters(const FShader* Shader, int32 PermutationId, const FShaderParameterMap& ParametersMap, const FShaderParametersMetadata& StructMetaData, bool bShouldBindEverything)
 {
 	const FShaderType* Type = Shader->GetTypeUnfrozen();
 	checkf(StructMetaData.GetSize() < (1 << (sizeof(uint16) * 8)), TEXT("Shader parameter structure can only have a size < 65536 bytes."));
@@ -251,6 +254,7 @@ void FShaderParameterBindings::BindForLegacyShaderParameters(const FShader* Shad
 
 	FShaderParameterStructBindingContext BindingContext;
 	BindingContext.Shader = Shader;
+	BindingContext.PermutationId = PermutationId;
 	BindingContext.Bindings = this;
 	BindingContext.ParametersMap = &ParametersMap;
 	BindingContext.bUseRootShaderParameters = false;
@@ -281,7 +285,7 @@ void FShaderParameterBindings::BindForLegacyShaderParameters(const FShader* Shad
 	}
 }
 
-void FShaderParameterBindings::BindForRootShaderParameters(const FShader* Shader, const FShaderParameterMap& ParametersMap)
+void FShaderParameterBindings::BindForRootShaderParameters(const FShader* Shader, int32 PermutationId, const FShaderParameterMap& ParametersMap)
 {
 	const FShaderType* Type = Shader->GetTypeUnfrozen();
 	check(this == &Shader->Bindings);
@@ -304,6 +308,7 @@ void FShaderParameterBindings::BindForRootShaderParameters(const FShader* Shader
 
 	FShaderParameterStructBindingContext BindingContext;
 	BindingContext.Shader = Shader;
+	BindingContext.PermutationId = PermutationId;
 	BindingContext.Bindings = this;
 	BindingContext.ParametersMap = &ParametersMap;
 	BindingContext.bUseRootShaderParameters = true;
diff --git a/Engine/Source/Runtime/RenderCore/Public/GlobalShader.h b/Engine/Source/Runtime/RenderCore/Public/GlobalShader.h
--- a/Engine/Source/Runtime/RenderCore/Public/GlobalShader.h
+++ b/Engine/Source/Runtime/RenderCore/Public/GlobalShader.h
@@ -204,7 +204,7 @@ public:
 
 	void Empty();
 
-	FShader* FindOrAddShader(const FShaderType* ShaderType, FShader* Shader);
+	FShader* FindOrAddShader(const FShaderType* ShaderType, int32 PermutationId, FShader* Shader);
 	FShaderPipeline* FindOrAddShaderPipeline(const FShaderPipelineType* ShaderPipelineType, FShaderPipeline* ShaderPipeline);
 
 	void RemoveShaderTypePermutaion(const FShaderType* Type, int32 PermutationId);
diff --git a/Engine/Source/Runtime/RenderCore/Public/Shader.h b/Engine/Source/Runtime/RenderCore/Public/Shader.h
--- a/Engine/Source/Runtime/RenderCore/Public/Shader.h
+++ b/Engine/Source/Runtime/RenderCore/Public/Shader.h
@@ -583,8 +583,8 @@ public:
 public:
 	static constexpr uint16 kInvalidBufferIndex = 0xFFFF;
 
-	void BindForLegacyShaderParameters(const FShader* Shader, const FShaderParameterMap& ParameterMaps, const FShaderParametersMetadata& StructMetaData, bool bShouldBindEverything = false);
-	void BindForRootShaderParameters(const FShader* Shader, const FShaderParameterMap& ParameterMaps);
+	void BindForLegacyShaderParameters(const FShader* Shader, int32 PermutationId, const FShaderParameterMap& ParameterMaps, const FShaderParametersMetadata& StructMetaData, bool bShouldBindEverything = false);
+	void BindForRootShaderParameters(const FShader* Shader, int32 PermutationId, const FShaderParameterMap& ParameterMaps);
 
 	LAYOUT_FIELD(TMemoryImageArray<FParameter>, Parameters);
 	LAYOUT_FIELD(TMemoryImageArray<FResourceParameter>, ResourceParameters);
@@ -702,9 +702,7 @@ public:
 	inline FVertexFactoryType* GetVertexFactoryType(const FShaderMapPointerTable& InPointerTable) const { return VFType.Get(InPointerTable.VFTypes); }
 	inline FVertexFactoryType* GetVertexFactoryType(const FPointerTableBase* InPointerTable) const { return VFType.Get(InPointerTable); }
 	inline FShaderType* GetTypeUnfrozen() const { return Type.GetUnfrozen(); }
-	inline const FHashedName& GetTypeName() const { return TypeName; }
 	inline int32 GetResourceIndex() const { return ResourceIndex; }
-	inline int32 GetPermutationId() const { return PermutationId; }
 	inline EShaderPlatform GetShaderPlatform() const { return Target.GetPlatform(); }
 	inline EShaderFrequency GetFrequency() const { return Target.GetFrequency(); }
 	inline const FShaderTarget GetTarget() const { return Target; }
@@ -762,7 +760,7 @@ public:
 	void DumpDebugInfo(const FShaderMapPointerTable& InPtrTable);
 
 #if WITH_EDITOR
-	void SaveShaderStableKeys(const FShaderMapPointerTable& InPtrTable, EShaderPlatform TargetShaderPlatform, const struct FStableShaderKeyAndValue& SaveKeyVal);
+	void SaveShaderStableKeys(const FShaderMapPointerTable& InPtrTable, EShaderPlatform TargetShaderPlatform, int32 PermutationId, const struct FStableShaderKeyAndValue& SaveKeyVal);
 #endif // WITH_EDITOR
 
 	/** Returns the meta data for the root shader parameter struct. */
@@ -799,18 +797,12 @@ private:
 
 	LAYOUT_FIELD(TIndexedPtr<FVertexFactoryType>, VFType);
 
-	/** Shader Type metadata for this shader. */
-	LAYOUT_FIELD(FHashedName, TypeName);
-
 	/** Target platform and frequency. */
 	LAYOUT_FIELD(FShaderTarget, Target);
 
 	/** Index of this shader within the FShaderMapResource */
 	LAYOUT_FIELD(int32, ResourceIndex);
 
-	/** Unique permutation identifier of the shader in the shader type. */
-	LAYOUT_FIELD(int32, PermutationId);
-
 	/** The number of instructions the shader takes to execute. */
 	LAYOUT_FIELD(uint32, NumInstructions);
 
@@ -1606,26 +1598,21 @@ class RENDERCORE_API FShaderPipeline
 {
 	DECLARE_TYPE_LAYOUT(FShaderPipeline, NonVirtual);
 public:
-	FShaderPipeline(
-		const FShaderPipelineType* InPipelineType,
-		FShader* InVertexShader,
-		FShader* InHullShader,
-		FShader* InDomainShader,
-		FShader* InGeometryShader,
-		FShader* InPixelShader);
-
-	FShaderPipeline(const FShaderPipelineType* InPipelineType, const TArray<FShader*>& InStages);
-
+	explicit FShaderPipeline(const FShaderPipelineType* InType) : TypeName(InType->GetHashedName()) {}
 	~FShaderPipeline();
 
+	void AddShader(FShader* Shader, int32 PermutationId);
+
 	inline uint32 GetNumShaders() const
 	{
 		uint32 NumShaders = 0u;
-		if (VertexShader.IsValid()) ++NumShaders;
-		if (HullShader.IsValid()) ++NumShaders;
-		if (DomainShader.IsValid()) ++NumShaders;
-		if (GeometryShader.IsValid()) ++NumShaders;
-		if (PixelShader.IsValid()) ++NumShaders;
+		for (uint32 i = 0u; i < SF_NumGraphicsFrequencies; ++i)
+		{
+			if (Shaders[i].IsValid())
+			{
+				++NumShaders;
+			}
+		}
 		return NumShaders;
 	}
 
@@ -1633,84 +1620,42 @@ public:
 	template<typename ShaderType>
 	ShaderType* GetShader(const FShaderMapPointerTable& InPtrTable)
 	{
-		if (PixelShader && PixelShader->GetType(InPtrTable) == &ShaderType::StaticType)
-		{
-			return (ShaderType*)PixelShader.Get();
-		}
-		else if (VertexShader && VertexShader->GetType(InPtrTable) == &ShaderType::StaticType)
-		{
-			return (ShaderType*)VertexShader.Get();
-		}
-		else if (GeometryShader && GeometryShader->GetType(InPtrTable) == &ShaderType::StaticType)
+		const FShaderType& Type = ShaderType::StaticType;
+		const EShaderFrequency Frequency = Type.GetFrequency();
+		if (Frequency < SF_NumGraphicsFrequencies && Shaders[Frequency].IsValid())
 		{
-			return (ShaderType*)GeometryShader.Get();
-		}
-		else if (HullShader)
-		{
-			if (HullShader->GetType(InPtrTable) == &ShaderType::StaticType)
-			{
-				return (ShaderType*)HullShader.Get();
-			}
-			else if (DomainShader && DomainShader->GetType(InPtrTable) == &ShaderType::StaticType)
+			FShader* Shader = Shaders[Frequency].GetChecked();
+			if (Shader->GetType(InPtrTable) == &Type)
 			{
-				return (ShaderType*)DomainShader.Get();
+				return static_cast<ShaderType*>(Shader);
 			}
 		}
-
 		return nullptr;
 	}
 
 	FShader* GetShader(EShaderFrequency Frequency)
 	{
-		switch (Frequency)
-		{
-		case SF_Vertex: return VertexShader;
-		case SF_Domain: return DomainShader;
-		case SF_Hull: return HullShader;
-		case SF_Geometry: return GeometryShader;
-		case SF_Pixel: return PixelShader;
-		default: check(0);
-		}
-
-		return nullptr;
+		check(Frequency < SF_NumGraphicsFrequencies);
+		return Shaders[Frequency];
 	}
 
 	const FShader* GetShader(EShaderFrequency Frequency) const
 	{
-		switch (Frequency)
-		{
-		case SF_Vertex: return VertexShader;
-		case SF_Domain: return DomainShader;
-		case SF_Hull: return HullShader;
-		case SF_Geometry: return GeometryShader;
-		case SF_Pixel: return PixelShader;
-		default: check(0);
-		}
-
-		return nullptr;
+		check(Frequency < SF_NumGraphicsFrequencies);
+		return Shaders[Frequency];
 	}
 
 	inline TArray<TShaderRef<FShader>> GetShaders(const FShaderMapBase& InShaderMap) const
 	{
-		TArray<TShaderRef<FShader>> Shaders;
-
-		if (PixelShader)
-		{
-			Shaders.Add(TShaderRef<FShader>(PixelShader, InShaderMap));
-		}
-		if (GeometryShader)
+		TArray<TShaderRef<FShader>> Result;
+		for (uint32 i = 0u; i < SF_NumGraphicsFrequencies; ++i)
 		{
-			Shaders.Add(TShaderRef<FShader>(GeometryShader, InShaderMap));
-		}
-		if (HullShader)
-		{
-			Shaders.Add(TShaderRef<FShader>(DomainShader, InShaderMap));
-			Shaders.Add(TShaderRef<FShader>(HullShader, InShaderMap));
+			if (Shaders[i].IsValid())
+			{
+				Result.Add(TShaderRef<FShader>(Shaders[i].GetChecked(), InShaderMap));
+			}
 		}
-
-		Shaders.Add(TShaderRef<FShader>(VertexShader, InShaderMap));
-
-		return Shaders;
+		return Result;
 	}
 
 	void Validate(const FShaderPipelineType* InPipelineType) const;
@@ -1728,11 +1673,8 @@ public:
 #endif // WITH_EDITOR
 
 	LAYOUT_FIELD(FHashedName, TypeName);
-	LAYOUT_FIELD(TMemoryImagePtr<FShader>, VertexShader);
-	LAYOUT_FIELD(TMemoryImagePtr<FShader>, HullShader);
-	LAYOUT_FIELD(TMemoryImagePtr<FShader>, DomainShader);
-	LAYOUT_FIELD(TMemoryImagePtr<FShader>, GeometryShader);
-	LAYOUT_FIELD(TMemoryImagePtr<FShader>, PixelShader);
+	LAYOUT_ARRAY(TMemoryImagePtr<FShader>, Shaders, SF_NumGraphicsFrequencies);
+	LAYOUT_ARRAY(int32, PermutationIds, SF_NumGraphicsFrequencies);
 };
 
 inline bool operator<(const FShaderPipeline& Lhs, const FShaderPipeline& Rhs)
@@ -1776,43 +1718,18 @@ private:
 	const FShaderMapBase* ShaderMap;
 };
 
-struct FShaderMapKey
-{
-	FShaderMapKey() {}
-	FShaderMapKey(const FHashedName& InTypeName, int32 InPermutationId) : TypeName(InTypeName), PermutationId(InPermutationId) {}
-	explicit FShaderMapKey(const FShader& InShader) : FShaderMapKey(InShader.GetTypeName(), InShader.GetPermutationId()) {}
-
-	friend inline bool operator<(const FShaderMapKey& Lhs, const FShaderMapKey& Rhs)
-	{
-		if (Lhs.TypeName != Rhs.TypeName)
-		{
-			return Lhs.TypeName < Rhs.TypeName;
-		}
-		return Lhs.PermutationId < Rhs.PermutationId;
-	}
-
-	FHashedName TypeName;
-	int32 PermutationId;
-};
-
 /** A collection of shaders of different types */
 class RENDERCORE_API FShaderMapContent
 {
 	DECLARE_TYPE_LAYOUT(FShaderMapContent, NonVirtual);
 public:
-	struct FProjectShaderToKey
-	{
-		inline FShaderMapKey operator()(const FShader* InShader) { return FShaderMapKey(*InShader); }
-	};
-
 	struct FProjectShaderPipelineToKey
 	{
 		inline FHashedName operator()(const FShaderPipeline* InShaderPipeline) { return InShaderPipeline->TypeName; }
 	};
-
 	/** Default constructor. */
 	explicit FShaderMapContent(EShaderPlatform InPlatform)
-		: Platform(InPlatform)
+		: ShaderHash(128u), Platform(InPlatform)
 	{}
 
 	/** Destructor ensures pipelines cleared up. */
@@ -1857,7 +1774,7 @@ public:
 		return Shader != nullptr;
 	}
 
-	bool HasShader(FShaderType* Type, int32 PermutationId) const
+	bool HasShader(const FShaderType* Type, int32 PermutationId) const
 	{
 		return HasShader(Type->GetHashedName(), PermutationId);
 	}
@@ -1872,9 +1789,9 @@ public:
 		return ShaderPipelines;
 	}
 
-	void AddShader(FShader* Shader);
+	void AddShader(const FHashedName& TypeName, int32 PermutationId, FShader* Shader);
 
-	FShader* FindOrAddShader(FShader* Shader);
+	FShader* FindOrAddShader(const FHashedName& TypeName, int32 PermutationId, FShader* Shader);
 
 	void AddShaderPipeline(FShaderPipeline* Pipeline);
 
@@ -1884,7 +1801,12 @@ public:
 	 * Removes the shader of the given type from the shader map
 	 * @param Type Shader type to remove the entry for 
 	 */
-	void RemoveShaderTypePermutaion(const FShaderType* Type, int32 PermutationId);
+	void RemoveShaderTypePermutaion(const FHashedName& TypeName, int32 PermutationId);
+
+	inline void RemoveShaderTypePermutaion(const FShaderType* Type, int32 PermutationId)
+	{
+		RemoveShaderTypePermutaion(Type->GetHashedName(), PermutationId);
+	}
 
 	void RemoveShaderPipelineType(const FShaderPipelineType* ShaderPipelineType);
 
@@ -1939,9 +1861,18 @@ public:
 
 	uint32 GetMaxNumInstructionsForShader(const FShaderMapBase& InShaderMap, FShaderType* ShaderType) const;
 
+	void Finalize();
+
+	void UpdateHash(FSHA1& Hasher) const;
+
 protected:
 	void EmptyShaderPipelines();
 
+	using FMemoryImageHashTable = THashTable<FMemoryImageAllocator>;
+
+	LAYOUT_FIELD(FMemoryImageHashTable, ShaderHash);
+	LAYOUT_FIELD(TMemoryImageArray<FHashedName>, ShaderTypes);
+	LAYOUT_FIELD(TMemoryImageArray<int32>, ShaderPermutations);
 	LAYOUT_FIELD(TMemoryImageArray<TMemoryImagePtr<FShader>>, Shaders);
 	LAYOUT_FIELD(TMemoryImageArray<TMemoryImagePtr<FShaderPipeline>>, ShaderPipelines);
 	/** The platform this shader map was compiled with */
@@ -2021,6 +1952,12 @@ public:
 	inline const ContentType* GetContent() const { return static_cast<const ContentType*>(FShaderMapBase::GetContent()); }
 	inline ContentType* GetMutableContent() { return static_cast<ContentType*>(FShaderMapBase::GetMutableContent()); }
 
+	void FinalizeContent()
+	{
+		this->GetMutableContent()->Finalize();
+		FShaderMapBase::FinalizeContent();
+	}
+
 protected:
 	TShaderMap() : FShaderMapBase(StaticGetTypeLayoutDesc<ContentType>()) {}
 	virtual FShaderMapPointerTable* CreatePointerTable() const final override { return new PointerTableType(); }
diff --git a/Engine/Source/Runtime/RenderCore/Public/ShaderCodeLibrary.h b/Engine/Source/Runtime/RenderCore/Public/ShaderCodeLibrary.h
--- a/Engine/Source/Runtime/RenderCore/Public/ShaderCodeLibrary.h
+++ b/Engine/Source/Runtime/RenderCore/Public/ShaderCodeLibrary.h
@@ -19,11 +19,7 @@ class FShaderMapResourceCode;
 
 struct RENDERCORE_API FShaderCodeLibraryPipeline
 {
-	FSHAHash VertexShader;
-	FSHAHash PixelShader;
-	FSHAHash GeometryShader;
-	FSHAHash HullShader;
-	FSHAHash DomainShader;
+	FSHAHash Shaders[SF_NumGraphicsFrequencies];
 	mutable uint32 Hash;
 	
 	/** Fills the hashes from the pipeline stage shaders */
@@ -33,18 +29,24 @@ struct RENDERCORE_API FShaderCodeLibraryPipeline
 	
 	friend bool operator ==(const FShaderCodeLibraryPipeline& A,const FShaderCodeLibraryPipeline& B)
 	{
-		return A.VertexShader == B.VertexShader && A.PixelShader == B.PixelShader && A.GeometryShader == B.GeometryShader && A.HullShader == B.HullShader && A.DomainShader == B.DomainShader;
+		for (uint32 Frequency = 0u; Frequency < SF_NumGraphicsFrequencies; ++Frequency)
+		{
+			if (A.Shaders[Frequency] != B.Shaders[Frequency])
+			{
+				return false;
+			}
+		}
+		return true;
 	}
 	
 	friend uint32 GetTypeHash(const FShaderCodeLibraryPipeline &Key)
 	{
 		if(!Key.Hash)
 		{
-			Key.Hash = FCrc::MemCrc32(Key.VertexShader.Hash, sizeof(Key.VertexShader.Hash));
-			Key.Hash = FCrc::MemCrc32(Key.PixelShader.Hash, sizeof(Key.PixelShader.Hash), Key.Hash);
-			Key.Hash = FCrc::MemCrc32(Key.GeometryShader.Hash, sizeof(Key.GeometryShader.Hash), Key.Hash);
-			Key.Hash = FCrc::MemCrc32(Key.HullShader.Hash, sizeof(Key.HullShader.Hash), Key.Hash);
-			Key.Hash = FCrc::MemCrc32(Key.DomainShader.Hash, sizeof(Key.DomainShader.Hash), Key.Hash);
+			for (uint32 Frequency = 0u; Frequency < SF_NumGraphicsFrequencies; ++Frequency)
+			{
+				Key.Hash = FCrc::MemCrc32(Key.Shaders[Frequency].Hash, sizeof(FSHAHash), Key.Hash);
+			}
 		}
 		return Key.Hash;
 	}
@@ -54,7 +56,11 @@ struct RENDERCORE_API FShaderCodeLibraryPipeline
 	
 	friend FArchive& operator<<( FArchive& Ar, FShaderCodeLibraryPipeline& Info )
 	{
-		return Ar << Info.VertexShader << Info.PixelShader << Info.GeometryShader << Info.HullShader << Info.DomainShader << Info.Hash;
+		for (uint32 Frequency = 0u; Frequency < SF_NumGraphicsFrequencies; ++Frequency)
+		{
+			Ar << Info.Shaders[Frequency];
+		}
+		return Ar << Info.Hash;
 	}
 };
 
diff --git a/Engine/Source/Runtime/RenderCore/Public/ShaderParameterStruct.h b/Engine/Source/Runtime/RenderCore/Public/ShaderParameterStruct.h
--- a/Engine/Source/Runtime/RenderCore/Public/ShaderParameterStruct.h
+++ b/Engine/Source/Runtime/RenderCore/Public/ShaderParameterStruct.h
@@ -12,9 +12,9 @@
 #include "RenderGraphResources.h"
 
 template <typename FParameterStruct>
-void BindForLegacyShaderParameters(FShader* Shader, const FShaderParameterMap& ParameterMap, bool bShouldBindEverything = false)
+void BindForLegacyShaderParameters(FShader* Shader, int32 PermutationId, const FShaderParameterMap& ParameterMap, bool bShouldBindEverything = false)
 {
-	Shader->Bindings.BindForLegacyShaderParameters(Shader, ParameterMap, *FParameterStruct::FTypeInfo::GetStructMetadata(), bShouldBindEverything);
+	Shader->Bindings.BindForLegacyShaderParameters(Shader, PermutationId, ParameterMap, *FParameterStruct::FTypeInfo::GetStructMetadata(), bShouldBindEverything);
 }
 
 /** Tag a shader class to use the structured shader parameters API.
@@ -37,7 +37,7 @@ void BindForLegacyShaderParameters(FShader* Shader, const FShaderParameterMap& P
 	ShaderClass(const ShaderMetaType::CompiledShaderInitializerType& Initializer) \
 		: ShaderParentClass(Initializer) \
 	{ \
-		BindForLegacyShaderParameters<FParameters>(this, Initializer.ParameterMap, bShouldBindEverything); \
+		BindForLegacyShaderParameters<FParameters>(this, Initializer.PermutationId, Initializer.ParameterMap, bShouldBindEverything); \
 	} \
 	\
 	ShaderClass() \
@@ -61,7 +61,7 @@ void BindForLegacyShaderParameters(FShader* Shader, const FShaderParameterMap& P
 	ShaderClass(const ShaderMetaType::CompiledShaderInitializerType& Initializer) \
 		: ShaderParentClass(Initializer) \
 	{ \
-		this->Bindings.BindForRootShaderParameters(this, Initializer.ParameterMap); \
+		this->Bindings.BindForRootShaderParameters(this, Initializer.PermutationId, Initializer.ParameterMap); \
 	} \
 	\
 	ShaderClass() \
diff --git a/Engine/Source/Runtime/Renderer/Private/PostProcess/PostProcessMobile.cpp b/Engine/Source/Runtime/Renderer/Private/PostProcess/PostProcessMobile.cpp
--- a/Engine/Source/Runtime/Renderer/Private/PostProcess/PostProcessMobile.cpp
+++ b/Engine/Source/Runtime/Renderer/Private/PostProcess/PostProcessMobile.cpp
@@ -2736,7 +2736,7 @@ public:
 	FClearUAVUIntCS_ES2(const ShaderMetaType::CompiledShaderInitializerType& Initializer)
 		: FGlobalShader(Initializer)
 	{
-		Bindings.BindForLegacyShaderParameters(this, Initializer.ParameterMap, *FParameters::FTypeInfo::GetStructMetadata());
+		Bindings.BindForLegacyShaderParameters(this, Initializer.PermutationId, Initializer.ParameterMap, *FParameters::FTypeInfo::GetStructMetadata());
 	}
 
 	static bool ShouldCompilePermutation(const FGlobalShaderPermutationParameters& Parameters)
@@ -2830,7 +2830,7 @@ public:
 	FPostProcessAverageLuminanceCS_ES2(const ShaderMetaType::CompiledShaderInitializerType& Initializer)
 		: FGlobalShader(Initializer)
 	{
-		Bindings.BindForLegacyShaderParameters(this, Initializer.ParameterMap, *FParameters::FTypeInfo::GetStructMetadata());
+		Bindings.BindForLegacyShaderParameters(this, Initializer.PermutationId, Initializer.ParameterMap, *FParameters::FTypeInfo::GetStructMetadata());
 	}
 
 
@@ -2952,7 +2952,7 @@ public:
 	FBasicEyeAdaptationCS_ES2(const ShaderMetaType::CompiledShaderInitializerType& Initializer)
 		: FGlobalShader(Initializer)
 	{
-		Bindings.BindForLegacyShaderParameters(this, Initializer.ParameterMap, *FParameters::FTypeInfo::GetStructMetadata());
+		Bindings.BindForLegacyShaderParameters(this, Initializer.PermutationId, Initializer.ParameterMap, *FParameters::FTypeInfo::GetStructMetadata());
 	}
 
 public:
@@ -3117,7 +3117,7 @@ public:
 	FPostProcessHistogramCS_ES2(const ShaderMetaType::CompiledShaderInitializerType& Initializer)
 		: FGlobalShader(Initializer)
 	{
-		Bindings.BindForLegacyShaderParameters(this, Initializer.ParameterMap, *FParameters::FTypeInfo::GetStructMetadata());
+		Bindings.BindForLegacyShaderParameters(this, Initializer.PermutationId, Initializer.ParameterMap, *FParameters::FTypeInfo::GetStructMetadata());
 	}
 
 	template <typename TRHICmdList>
@@ -3256,7 +3256,7 @@ public:
 	FHistogramEyeAdaptationCS_ES2(const CompiledShaderInitializerType& Initializer)
 		: FGlobalShader(Initializer)
 	{
-		Bindings.BindForLegacyShaderParameters(this, Initializer.ParameterMap, *FParameters::FTypeInfo::GetStructMetadata());
+		Bindings.BindForLegacyShaderParameters(this, Initializer.PermutationId, Initializer.ParameterMap, *FParameters::FTypeInfo::GetStructMetadata());
 	}
 
 	void Set(const FRenderingCompositePassContext& Context, const TShaderRef<FHistogramEyeAdaptationCS_ES2>& Shader, FRHIShaderResourceView* LastEyeAdaptation, FRHIShaderResourceView* HistogramBuffer, FRHIUnorderedAccessView* TextureUnorderedAccessViewRHI = nullptr)
diff --git a/Engine/Source/Runtime/Renderer/Private/ShaderBaseClasses.cpp b/Engine/Source/Runtime/Renderer/Private/ShaderBaseClasses.cpp
--- a/Engine/Source/Runtime/Renderer/Private/ShaderBaseClasses.cpp
+++ b/Engine/Source/Runtime/Renderer/Private/ShaderBaseClasses.cpp
@@ -113,9 +113,10 @@ FRHIUniformBuffer* FMaterialShader::GetParameterCollectionBuffer(const FGuid& Id
 void FMaterialShader::VerifyExpressionAndShaderMaps(const FMaterialRenderProxy* MaterialRenderProxy, const FMaterial& Material, const FUniformExpressionCache* UniformExpressionCache) const
 {
 	// Validate that the shader is being used for a material that matches the uniform expression set the shader was compiled for.
-	const FUniformExpressionSet& MaterialUniformExpressionSet = Material.GetRenderingThreadShaderMap()->GetUniformExpressionSet();
+	FMaterialShaderMap* ShaderMap = Material.GetRenderingThreadShaderMap();
+	const FUniformExpressionSet& MaterialUniformExpressionSet = ShaderMap->GetUniformExpressionSet();
 	bool bUniformExpressionSetMismatch = !DebugUniformExpressionSet.Matches(MaterialUniformExpressionSet)
-		|| UniformExpressionCache->CachedUniformExpressionShaderMap != Material.GetRenderingThreadShaderMap();
+		|| UniformExpressionCache->CachedUniformExpressionShaderMap != ShaderMap;
 	if (!bUniformExpressionSetMismatch)
 	{
 		auto DumpUB = [](const FRHIUniformBufferLayout& Layout)
@@ -164,10 +165,10 @@ void FMaterialShader::VerifyExpressionAndShaderMaps(const FMaterialRenderProxy*
 	}
 	if (bUniformExpressionSetMismatch)
 	{
-		const FShaderType* ShaderType = FindShaderTypeByName(GetTypeName());
+		const FShaderType* ShaderType = GetType(ShaderMap->GetPointerTable());
 		const FString ProxyName = *MaterialRenderProxy->GetFriendlyName();
 		const FString MaterialName = *Material.GetFriendlyName();
-		const TCHAR* ShaderMapDesc = Material.GetRenderingThreadShaderMap()->GetDebugDescription();
+		const TCHAR* ShaderMapDesc = ShaderMap->GetDebugDescription();
 		UE_LOG(
 			LogShaders,
 			Fatal,
@@ -196,7 +197,7 @@ void FMaterialShader::VerifyExpressionAndShaderMaps(const FMaterialRenderProxy*
 			MaterialUniformExpressionSet.UniformTextureParameters[(uint32)EMaterialTextureParameterType::Array2D].Num(),
 			MaterialUniformExpressionSet.UniformTextureParameters[(uint32)EMaterialTextureParameterType::Volume].Num(),
 			MaterialUniformExpressionSet.UniformTextureParameters[(uint32)EMaterialTextureParameterType::Virtual].Num(),
-			Material.GetRenderingThreadShaderMap()
+			ShaderMap
 		);
 	}
 }
@@ -224,10 +225,11 @@ void FMaterialShader::SetParametersInner(
 		// UE-46061 - Workaround for a rare crash with an outdated cached shader map
 		if (UniformExpressionCache->CachedUniformExpressionShaderMap != Material.GetRenderingThreadShaderMap())
 		{
+			FMaterialShaderMap* ShaderMap = Material.GetRenderingThreadShaderMap();
 			UMaterialInterface* MtlInterface = Material.GetMaterialInterface();
 			UMaterialInterface* ProxyInterface = MaterialRenderProxy->GetMaterialInterface();
 
-			const FShaderType* ShaderType = FindShaderTypeByName(GetTypeName());
+			const FShaderType* ShaderType = GetType(ShaderMap->GetPointerTable());
 			ensureMsgf(false,
 				TEXT("%s shader uniform expression set mismatched shader map for material %s/%s, forcing expression cache evaluation.\n")
 				TEXT("Material:  %s\n")
diff --git a/Engine/Source/Runtime/Renderer/Private/VT/RuntimeVirtualTextureRender.cpp b/Engine/Source/Runtime/Renderer/Private/VT/RuntimeVirtualTextureRender.cpp
--- a/Engine/Source/Runtime/Renderer/Private/VT/RuntimeVirtualTextureRender.cpp
+++ b/Engine/Source/Runtime/Renderer/Private/VT/RuntimeVirtualTextureRender.cpp
@@ -513,7 +513,7 @@ namespace RuntimeVirtualTexture
 		FShader_VirtualTextureCompress(const ShaderMetaType::CompiledShaderInitializerType& Initializer)
 			: FGlobalShader(Initializer)
 		{
-			Bindings.BindForLegacyShaderParameters(this, Initializer.ParameterMap, *FParameters::FTypeInfo::GetStructMetadata());
+			Bindings.BindForLegacyShaderParameters(this, Initializer.PermutationId, Initializer.ParameterMap, *FParameters::FTypeInfo::GetStructMetadata());
 		}
 	};
 
@@ -601,7 +601,7 @@ namespace RuntimeVirtualTexture
 		FShader_VirtualTextureCopy(const ShaderMetaType::CompiledShaderInitializerType& Initializer)
 			: FGlobalShader(Initializer)
 		{
-			Bindings.BindForLegacyShaderParameters(this, Initializer.ParameterMap, *FParameters::FTypeInfo::GetStructMetadata());
+			Bindings.BindForLegacyShaderParameters(this, Initializer.PermutationId, Initializer.ParameterMap, *FParameters::FTypeInfo::GetStructMetadata());
 		}
 	};
 
diff --git a/Engine/Source/Runtime/Renderer/Private/VolumetricFog.cpp b/Engine/Source/Runtime/Renderer/Private/VolumetricFog.cpp
--- a/Engine/Source/Runtime/Renderer/Private/VolumetricFog.cpp
+++ b/Engine/Source/Runtime/Renderer/Private/VolumetricFog.cpp
@@ -163,7 +163,7 @@ public:
 	FVolumetricFogMaterialSetupCS(const ShaderMetaType::CompiledShaderInitializerType& Initializer)
 		: FGlobalShader(Initializer)
 	{
-		Bindings.BindForLegacyShaderParameters(this, Initializer.ParameterMap, *FParameters::FTypeInfo::GetStructMetadata());
+		Bindings.BindForLegacyShaderParameters(this, Initializer.PermutationId, Initializer.ParameterMap, *FParameters::FTypeInfo::GetStructMetadata());
 
 		VolumetricFogParameters.Bind(Initializer.ParameterMap);
 	}
@@ -581,7 +581,7 @@ public:
 	TVolumetricFogLightScatteringCS(const ShaderMetaType::CompiledShaderInitializerType& Initializer)
 		: FGlobalShader(Initializer)
 	{
-		Bindings.BindForLegacyShaderParameters(this, Initializer.ParameterMap, *FParameters::FTypeInfo::GetStructMetadata());
+		Bindings.BindForLegacyShaderParameters(this, Initializer.PermutationId, Initializer.ParameterMap, *FParameters::FTypeInfo::GetStructMetadata());
 
 		LocalShadowedLightScattering.Bind(Initializer.ParameterMap, TEXT("LocalShadowedLightScattering"));
 		LightScatteringHistory.Bind(Initializer.ParameterMap, TEXT("LightScatteringHistory"));
@@ -740,7 +740,7 @@ class FVolumetricFogFinalIntegrationCS : public FGlobalShader
 	FVolumetricFogFinalIntegrationCS(const ShaderMetaType::CompiledShaderInitializerType& Initializer)
 		: FGlobalShader(Initializer)
 	{
-		Bindings.BindForLegacyShaderParameters(this, Initializer.ParameterMap, *FParameters::FTypeInfo::GetStructMetadata());
+		Bindings.BindForLegacyShaderParameters(this, Initializer.PermutationId, Initializer.ParameterMap, *FParameters::FTypeInfo::GetStructMetadata());
 
 		VolumetricFogParameters.Bind(Initializer.ParameterMap);
 	}
