Force 64bit alignment for some uint64s used with atomics, fixes crash on 32bit windows

diff --git a/Engine/Source/Runtime/RenderCore/Public/ShaderPipelineCache.h b/Engine/Source/Runtime/RenderCore/Public/ShaderPipelineCache.h
--- a/Engine/Source/Runtime/RenderCore/Public/ShaderPipelineCache.h
+++ b/Engine/Source/Runtime/RenderCore/Public/ShaderPipelineCache.h
@@ -246,10 +246,10 @@ private:
     int32 PausedCount;
 	FShaderCachePrecompileContext ShaderCachePrecompileContext;
 	
-    volatile int64 TotalActiveTasks;
-    volatile int64 TotalWaitingTasks;
-    volatile int64 TotalCompleteTasks;
-    volatile int64 TotalPrecompileTime;
+	MS_ALIGN(8) volatile int64 TotalActiveTasks GCC_ALIGN(8);
+	MS_ALIGN(8) volatile int64 TotalWaitingTasks GCC_ALIGN(8);
+	MS_ALIGN(8) volatile int64 TotalCompleteTasks GCC_ALIGN(8);
+	MS_ALIGN(8) volatile int64 TotalPrecompileTime GCC_ALIGN(8);
 	double PrecompileStartTime;
 
 	FCriticalSection Mutex;
