Improve behavior when global shades can't be loaded due to memory layout differences
(Trigger a fatal log mentioning missing global shaders, rather than simply crash on nullptr access)

diff --git a/Engine/Source/Runtime/RenderCore/Private/GlobalShader.cpp b/Engine/Source/Runtime/RenderCore/Private/GlobalShader.cpp
--- a/Engine/Source/Runtime/RenderCore/Private/GlobalShader.cpp
+++ b/Engine/Source/Runtime/RenderCore/Private/GlobalShader.cpp
@@ -264,13 +264,17 @@ FGlobalShaderMap* GetGlobalShaderMap(EShaderPlatform Platform)
 FGlobalShaderMapSection* FGlobalShaderMapSection::CreateFromArchive(FArchive& Ar)
 {
 	FGlobalShaderMapSection* Section = new FGlobalShaderMapSection();
-	Section->Serialize(Ar);
-	return Section;
+	if (Section->Serialize(Ar))
+	{
+		return Section;
+	}
+	delete Section;
+	return nullptr;
 }
 
-void FGlobalShaderMapSection::Serialize(FArchive& Ar)
+bool FGlobalShaderMapSection::Serialize(FArchive& Ar)
 {
-	Super::Serialize(Ar, true, false);
+	return Super::Serialize(Ar, true, false);
 }
 
 TShaderRef<FShader> FGlobalShaderMapSection::GetShader(FShaderType* ShaderType, int32 PermutationId) const
@@ -397,6 +401,7 @@ void FGlobalShaderMap::RemoveShaderPipelineType(const FShaderPipelineType* Shade
 
 void FGlobalShaderMap::AddSection(FGlobalShaderMapSection* InSection)
 {
+	check(InSection);
 	const FGlobalShaderMapContent* Content = InSection->GetContent();
 	const FHashedName& HashedFilename = Content->HashedSourceFilename;
 
@@ -440,6 +445,9 @@ void FGlobalShaderMap::LoadFromGlobalArchive(FArchive& Ar)
 	for (int32 i = 0; i < NumSections; ++i)
 	{
 		FGlobalShaderMapSection* Section = FGlobalShaderMapSection::CreateFromArchive(Ar);
-		AddSection(Section);
+		if (Section)
+		{
+			AddSection(Section);
+		}
 	}
 }
diff --git a/Engine/Source/Runtime/RenderCore/Public/GlobalShader.h b/Engine/Source/Runtime/RenderCore/Public/GlobalShader.h
--- a/Engine/Source/Runtime/RenderCore/Public/GlobalShader.h
+++ b/Engine/Source/Runtime/RenderCore/Public/GlobalShader.h
@@ -132,7 +132,7 @@ public:
 
 	~FGlobalShaderMapSection() { ReleaseResourceBuilder(); }
 
-	void Serialize(FArchive& Ar);
+	bool Serialize(FArchive& Ar);
 
 	FShaderMapResourceBuilder& GetResourceBuilder()
 	{
