Added StartRenderCommandFenceBundler/StopRenderCommandFenceBundler to FPendingCleanupObjects to avoid massive render command enqueue hitching.

#rnx

diff --git a/Engine/Source/Runtime/RenderCore/Private/RenderingThread.cpp b/Engine/Source/Runtime/RenderCore/Private/RenderingThread.cpp
--- a/Engine/Source/Runtime/RenderCore/Private/RenderingThread.cpp
+++ b/Engine/Source/Runtime/RenderCore/Private/RenderingThread.cpp
@@ -1311,10 +1311,12 @@ FPendingCleanupObjects::~FPendingCleanupObjects()
 {
 	QUICK_SCOPE_CYCLE_COUNTER(STAT_FPendingCleanupObjects_Destruct);
 
+	StartRenderCommandFenceBundler();
 	for (int32 ObjectIndex = 0; ObjectIndex < CleanupArray.Num(); ObjectIndex++)
 	{
 		delete CleanupArray[ObjectIndex];
 	}
+	StopRenderCommandFenceBundler();
 }
 
 void BeginCleanup(FDeferredCleanupInterface* CleanupObject)
@@ -1340,10 +1342,12 @@ FPendingCleanupObjects::~FPendingCleanupObjects()
 {
 	QUICK_SCOPE_CYCLE_COUNTER(STAT_FPendingCleanupObjects_Destruct);
 
+	StartRenderCommandFenceBundler();
 	for (int32 ObjectIndex = 0; ObjectIndex < CleanupArray.Num(); ObjectIndex++)
 	{
 		delete CleanupArray[ObjectIndex];
 	}
+	StopRenderCommandFenceBundler();
 }
 
 void BeginCleanup(FDeferredCleanupInterface* CleanupObject)
