Fix RemoveShaderTypePermutaion(), avoid crash when recompiling shaders

diff --git a/Engine/Source/Runtime/RenderCore/Private/ShaderMap.cpp b/Engine/Source/Runtime/RenderCore/Private/ShaderMap.cpp
--- a/Engine/Source/Runtime/RenderCore/Private/ShaderMap.cpp
+++ b/Engine/Source/Runtime/RenderCore/Private/ShaderMap.cpp
@@ -410,16 +410,21 @@ void FShaderMapContent::RemoveShaderTypePermutaion(const FHashedName& TypeName,
 		if (ShaderTypes[Index] == TypeName && ShaderPermutations[Index] == PermutationId)
 		{
 			DeleteObjectFromLayout(Shader);
+
+			// Replace the shader we're removing with the last shader in the list
 			Shaders.RemoveAtSwap(Index, 1, false);
 			ShaderTypes.RemoveAtSwap(Index, 1, false);
 			ShaderPermutations.RemoveAtSwap(Index, 1, false);
 			check(ShaderTypes.Num() == Shaders.Num());
 			check(ShaderPermutations.Num() == Shaders.Num());
-
 			ShaderHash.Remove(Hash, Index);
+
+			// SwapIndex is the old index of the shader at the end of the list, that's now been moved to replace the current shader
 			const int32 SwapIndex = Shaders.Num();
-			if (SwapIndex > 0)
+			if (Index != SwapIndex)
 			{
+				// We need to update the hash table to reflect shader previously at SwapIndex being moved to Index
+				// Here we construct the hash from values at Index, since type/permutation have already been moved
 				const uint16 SwapHash = MakeShaderHash(ShaderTypes[Index], ShaderPermutations[Index]);
 				ShaderHash.Remove(SwapHash, SwapIndex);
 				ShaderHash.Add(SwapHash, Index);
