Moved definition of virtual texture log

diff --git a/Engine/Source/Runtime/Engine/Classes/VT/VirtualTexture.h b/Engine/Source/Runtime/Engine/Classes/VT/VirtualTexture.h
--- a/Engine/Source/Runtime/Engine/Classes/VT/VirtualTexture.h
+++ b/Engine/Source/Runtime/Engine/Classes/VT/VirtualTexture.h
@@ -11,8 +11,6 @@
 #include "VirtualTexturing.h"
 #include "VirtualTexture.generated.h"
 
-DECLARE_LOG_CATEGORY_EXTERN(LogVirtualTexturingModule, Log, All);
-
 UCLASS(ClassGroup = Rendering)
 class ENGINE_API UVirtualTexture : public UObject
 {
diff --git a/Engine/Source/Runtime/Engine/Private/VT/VirtualTexture.cpp b/Engine/Source/Runtime/Engine/Private/VT/VirtualTexture.cpp
--- a/Engine/Source/Runtime/Engine/Private/VT/VirtualTexture.cpp
+++ b/Engine/Source/Runtime/Engine/Private/VT/VirtualTexture.cpp
@@ -9,8 +9,6 @@
 #include "FileCache/FileCache.h"
 #include "EngineModule.h"
 
-DEFINE_LOG_CATEGORY(LogVirtualTexturingModule);
-
 static FAutoConsoleCommand GVTFlushAndEvictFileCacheCommand(
 	TEXT("r.VT.FlushAndEvictFileCache"),
 	TEXT("Flush both the virtual texture physcial page cache and disk file cachje"),
diff --git a/Engine/Source/Runtime/Engine/Private/VT/VirtualTextureDataBuilder.cpp b/Engine/Source/Runtime/Engine/Private/VT/VirtualTextureDataBuilder.cpp
--- a/Engine/Source/Runtime/Engine/Private/VT/VirtualTextureDataBuilder.cpp
+++ b/Engine/Source/Runtime/Engine/Private/VT/VirtualTextureDataBuilder.cpp
@@ -704,7 +704,7 @@ void FVirtualTextureDataBuilder::BuildTiles(const TArray<FVTSourceTileEntry>& Ti
 		GeneratedData.TilePayload.Empty();
 		GeneratedData.CodecPayload.Empty();
 		GeneratedData.Codec = EVirtualTextureCodec::Max;
-		UE_LOG(LogVirtualTexturingModule, Fatal, TEXT("Failed build tile"));
+		UE_LOG(LogVirtualTexturing, Fatal, TEXT("Failed build tile"));
 	}
 }
 
diff --git a/Engine/Source/Runtime/RenderCore/Private/VirtualTexturing.cpp b/Engine/Source/Runtime/RenderCore/Private/VirtualTexturing.cpp
new file mode 100644
--- /dev/null
+++ b/Engine/Source/Runtime/RenderCore/Private/VirtualTexturing.cpp
@@ -0,0 +1,5 @@
+// Copyright Epic Games, Inc. All Rights Reserved.
+
+#include "VirtualTexturing.h"
+
+DEFINE_LOG_CATEGORY(LogVirtualTexturing);
diff --git a/Engine/Source/Runtime/RenderCore/Public/VirtualTexturing.h b/Engine/Source/Runtime/RenderCore/Public/VirtualTexturing.h
--- a/Engine/Source/Runtime/RenderCore/Public/VirtualTexturing.h
+++ b/Engine/Source/Runtime/RenderCore/Public/VirtualTexturing.h
@@ -383,5 +383,7 @@ inline uint64 GetTypeHash(const FVirtualTextureLocalTile& T) { return T.PackedVa
 inline bool operator==(const FVirtualTextureLocalTile& Lhs, const FVirtualTextureLocalTile& Rhs) { return Lhs.PackedValue == Rhs.PackedValue; }
 inline bool operator!=(const FVirtualTextureLocalTile& Lhs, const FVirtualTextureLocalTile& Rhs) { return Lhs.PackedValue != Rhs.PackedValue; }
 
+RENDERCORE_API DECLARE_LOG_CATEGORY_EXTERN(LogVirtualTexturing, Log, All);
+
 DECLARE_STATS_GROUP(TEXT("Virtual Texturing"), STATGROUP_VirtualTexturing, STATCAT_Advanced);
 DECLARE_STATS_GROUP(TEXT("Virtual Texture Memory"), STATGROUP_VirtualTextureMemory, STATCAT_Advanced);
diff --git a/Engine/Source/Runtime/Renderer/Private/VT/VirtualTexturing.cpp b/Engine/Source/Runtime/Renderer/Private/VT/VirtualTexturing.cpp
deleted file mode 100644
