FNEM - Optimizations arround TMemoryImageArray where usually the offset would be reapplied when using the index operator. Using an ArrayView avoids this.

diff --git a/Engine/Plugins/FX/Niagara/Source/Niagara/Private/NiagaraDataInterfaceSkeletalMesh.cpp b/Engine/Plugins/FX/Niagara/Source/Niagara/Private/NiagaraDataInterfaceSkeletalMesh.cpp
--- a/Engine/Plugins/FX/Niagara/Source/Niagara/Private/NiagaraDataInterfaceSkeletalMesh.cpp
+++ b/Engine/Plugins/FX/Niagara/Source/Niagara/Private/NiagaraDataInterfaceSkeletalMesh.cpp
@@ -595,8 +595,8 @@ void FSkeletalMeshGpuSpawnStaticBuffers::InitRHI()
 	if (bUseGpuUniformlyDistributedSampling)
 	{
 		const FSkeletalMeshAreaWeightedTriangleSampler& triangleSampler = SkeletalMeshSamplingLODBuiltData->AreaWeightedTriangleSampler;
-		const TArray<float, FMemoryImageAllocator>& Prob = triangleSampler.GetProb();
-		const TArray<int32, FMemoryImageAllocator>& Alias = triangleSampler.GetAlias();
+		TArrayView<const float> Prob = triangleSampler.GetProb();
+		TArrayView<const int32> Alias = triangleSampler.GetAlias();
 		check(TriangleCount == triangleSampler.GetNumEntries());
 
 		FRHIResourceCreateInfo CreateInfo;
diff --git a/Engine/Plugins/FX/Niagara/Source/Niagara/Private/NiagaraDataInterfaceStaticMesh.cpp b/Engine/Plugins/FX/Niagara/Source/Niagara/Private/NiagaraDataInterfaceStaticMesh.cpp
--- a/Engine/Plugins/FX/Niagara/Source/Niagara/Private/NiagaraDataInterfaceStaticMesh.cpp
+++ b/Engine/Plugins/FX/Niagara/Source/Niagara/Private/NiagaraDataInterfaceStaticMesh.cpp
@@ -104,8 +104,8 @@ void FStaticMeshGpuSpawnBuffer::Initialise(const FStaticMeshLODResources* Res, c
 	SectionRenderData = Res;
 
 	const uint32 ValidSectionCount = ValidSection.Num();
-	const TArray<float, FMemoryImageAllocator>& Prob = SectionSamplerParam.GetProb();
-	const TArray<int32, FMemoryImageAllocator>& Alias = SectionSamplerParam.GetAlias();
+	TArrayView<const float> Prob = SectionSamplerParam.GetProb();
+	TArrayView<const int32> Alias = SectionSamplerParam.GetAlias();
 	check(ValidSectionCount == Prob.Num());
 	// Build data that will be uploaded to GPU later from the render thread.
 	// The array contains data used to select regions for uniform particle spawning on them, as well as section triangle ranges.
diff --git a/Engine/Source/Runtime/Engine/Private/Materials/Material.cpp b/Engine/Source/Runtime/Engine/Private/Materials/Material.cpp
--- a/Engine/Source/Runtime/Engine/Private/Materials/Material.cpp
+++ b/Engine/Source/Runtime/Engine/Private/Materials/Material.cpp
@@ -1071,12 +1071,12 @@ void UMaterial::GetUsedTexturesAndIndices(TArray<UTexture*>& OutTextures, TArray
 
 		if (CurrentResource)
 		{
-			const TArray<FMaterialTextureParameterInfo, FMemoryImageAllocator>* ExpressionsByType[NumMaterialTextureParameterTypes];
+			TArrayView<const FMaterialTextureParameterInfo> ExpressionsByType[NumMaterialTextureParameterTypes];
 			uint32 NumTextures = 0u;
 			for (uint32 TypeIndex = 0u; TypeIndex < NumMaterialTextureParameterTypes; ++TypeIndex)
 			{
-				ExpressionsByType[TypeIndex] = &CurrentResource->GetUniformTextureExpressions((EMaterialTextureParameterType)TypeIndex);
-				NumTextures += ExpressionsByType[TypeIndex]->Num();
+				ExpressionsByType[TypeIndex] = CurrentResource->GetUniformTextureExpressions((EMaterialTextureParameterType)TypeIndex);
+				NumTextures += ExpressionsByType[TypeIndex].Num();
 			}
 
 			// Try to prevent resizing since this would be expensive.
@@ -1085,7 +1085,7 @@ void UMaterial::GetUsedTexturesAndIndices(TArray<UTexture*>& OutTextures, TArray
 			for (int32 TypeIndex = 0; TypeIndex < UE_ARRAY_COUNT(ExpressionsByType); TypeIndex++)
 			{
 				// Iterate over each of the material's texture expressions.
-				for (const FMaterialTextureParameterInfo& Parameter : *ExpressionsByType[TypeIndex])
+				for (const FMaterialTextureParameterInfo& Parameter : ExpressionsByType[TypeIndex])
 				{
 					UTexture* Texture = NULL;
 					Parameter.GetGameThreadTextureValue(this, *CurrentResource, Texture);
diff --git a/Engine/Source/Runtime/Engine/Private/Materials/MaterialInstance.cpp b/Engine/Source/Runtime/Engine/Private/Materials/MaterialInstance.cpp
--- a/Engine/Source/Runtime/Engine/Private/Materials/MaterialInstance.cpp
+++ b/Engine/Source/Runtime/Engine/Private/Materials/MaterialInstance.cpp
@@ -1666,7 +1666,7 @@ void UMaterialInstanceDynamic::CopyScalarAndVectorParameters(const UMaterialInte
 
 		// scalar
 		{
-			const TArray<FMaterialScalarParameterInfo, FMemoryImageAllocator>& Array = MaterialResource->GetUniformScalarParameterExpressions();
+			TArrayView<const FMaterialScalarParameterInfo> Array = MaterialResource->GetUniformScalarParameterExpressions();
 
 			for (int32 i = 0, Count = Array.Num(); i < Count; ++i)
 			{
@@ -1688,7 +1688,7 @@ void UMaterialInstanceDynamic::CopyScalarAndVectorParameters(const UMaterialInte
 
 		// vector
 		{
-			const TArray<FMaterialVectorParameterInfo, FMemoryImageAllocator>& Array = MaterialResource->GetUniformVectorParameterExpressions();
+			TArrayView<const FMaterialVectorParameterInfo> Array = MaterialResource->GetUniformVectorParameterExpressions();
 
 			for (int32 i = 0, Count = Array.Num(); i < Count; ++i)
 			{
@@ -4904,7 +4904,7 @@ void UMaterialInstance::CopyMaterialUniformParametersInternal(UMaterialInterface
 				if (MaterialResource)
 				{
 					// Scalars
-					const TArray<FMaterialScalarParameterInfo, FMemoryImageAllocator>& ScalarExpressions = MaterialResource->GetUniformScalarParameterExpressions();
+					TArrayView<const FMaterialScalarParameterInfo> ScalarExpressions = MaterialResource->GetUniformScalarParameterExpressions();
 					for (const FMaterialScalarParameterInfo& Parameter : ScalarExpressions)
 					{
 						FScalarParameterValue* ParameterValue = new(ScalarParameterValues) FScalarParameterValue;
@@ -4913,7 +4913,7 @@ void UMaterialInstance::CopyMaterialUniformParametersInternal(UMaterialInterface
 					}
 
 					// Vectors
-					const TArray<FMaterialVectorParameterInfo, FMemoryImageAllocator>& VectorExpressions = MaterialResource->GetUniformVectorParameterExpressions();
+					TArrayView<const FMaterialVectorParameterInfo> VectorExpressions = MaterialResource->GetUniformVectorParameterExpressions();
 					for (const FMaterialVectorParameterInfo& Parameter : VectorExpressions)
 					{
 						FVectorParameterValue* ParameterValue = new(VectorParameterValues) FVectorParameterValue;
diff --git a/Engine/Source/Runtime/Engine/Private/Materials/MaterialShared.cpp b/Engine/Source/Runtime/Engine/Private/Materials/MaterialShared.cpp
--- a/Engine/Source/Runtime/Engine/Private/Materials/MaterialShared.cpp
+++ b/Engine/Source/Runtime/Engine/Private/Materials/MaterialShared.cpp
@@ -666,17 +666,17 @@ const FUniformExpressionSet& FMaterial::GetUniformExpressions() const
 	return EmptyExpressions;
 }
 
-const TArray<FMaterialTextureParameterInfo, FMemoryImageAllocator>& FMaterial::GetUniformTextureExpressions(EMaterialTextureParameterType Type) const
+TArrayView<const FMaterialTextureParameterInfo> FMaterial::GetUniformTextureExpressions(EMaterialTextureParameterType Type) const
 {
 	return GetUniformExpressions().UniformTextureParameters[(uint32)Type];
 }
 
-const TArray<FMaterialVectorParameterInfo, FMemoryImageAllocator>& FMaterial::GetUniformVectorParameterExpressions() const
+TArrayView<const FMaterialVectorParameterInfo> FMaterial::GetUniformVectorParameterExpressions() const
 { 
 	return GetUniformExpressions().UniformVectorParameters;
 }
 
-const TArray<FMaterialScalarParameterInfo, FMemoryImageAllocator>& FMaterial::GetUniformScalarParameterExpressions() const
+TArrayView<const FMaterialScalarParameterInfo> FMaterial::GetUniformScalarParameterExpressions() const
 { 
 	return GetUniformExpressions().UniformScalarParameters;
 }
diff --git a/Engine/Source/Runtime/Engine/Private/StaticMesh.cpp b/Engine/Source/Runtime/Engine/Private/StaticMesh.cpp
--- a/Engine/Source/Runtime/Engine/Private/StaticMesh.cpp
+++ b/Engine/Source/Runtime/Engine/Private/StaticMesh.cpp
@@ -204,8 +204,8 @@ void FStaticMeshSectionAreaWeightedTriangleSamplerBuffer::InitRHI()
 		for (uint32 i = 0; i < AllSectionCount; ++i)
 		{
 			FStaticMeshSectionAreaWeightedTriangleSampler& sampler = (*Samplers)[i];
-			const TArray<float, FMemoryImageAllocator>& ProbTris = sampler.GetProb();
-			const TArray<int32, FMemoryImageAllocator>& AliasTris = sampler.GetAlias();
+			TArrayView<const float> ProbTris = sampler.GetProb();
+			TArrayView<const int32> AliasTris = sampler.GetAlias();
 			const uint32 NumTriangle = sampler.GetNumEntries();
 
 			for (uint32 t = 0; t < NumTriangle; ++t)
diff --git a/Engine/Source/Runtime/Engine/Public/MaterialShared.h b/Engine/Source/Runtime/Engine/Public/MaterialShared.h
--- a/Engine/Source/Runtime/Engine/Public/MaterialShared.h
+++ b/Engine/Source/Runtime/Engine/Public/MaterialShared.h
@@ -1671,15 +1671,15 @@ public:
 
 	// Accessors.
 	ENGINE_API const FUniformExpressionSet& GetUniformExpressions() const;
-	ENGINE_API const TArray<FMaterialTextureParameterInfo, FMemoryImageAllocator>& GetUniformTextureExpressions(EMaterialTextureParameterType Type) const;
-	ENGINE_API const TArray<FMaterialVectorParameterInfo, FMemoryImageAllocator>& GetUniformVectorParameterExpressions() const;
-	ENGINE_API const TArray<FMaterialScalarParameterInfo, FMemoryImageAllocator>& GetUniformScalarParameterExpressions() const;
-
-	inline const TArray<FMaterialTextureParameterInfo, FMemoryImageAllocator>& GetUniform2DTextureExpressions() const { return GetUniformTextureExpressions(EMaterialTextureParameterType::Standard2D); }
-	inline const TArray<FMaterialTextureParameterInfo, FMemoryImageAllocator>& GetUniformCubeTextureExpressions() const { return GetUniformTextureExpressions(EMaterialTextureParameterType::Cube); }
-	inline const TArray<FMaterialTextureParameterInfo, FMemoryImageAllocator>& GetUniform2DArrayTextureExpressions() const { return GetUniformTextureExpressions(EMaterialTextureParameterType::Array2D); }
-	inline const TArray<FMaterialTextureParameterInfo, FMemoryImageAllocator>& GetUniformVolumeTextureExpressions() const { return GetUniformTextureExpressions(EMaterialTextureParameterType::Volume); }
-	inline const TArray<FMaterialTextureParameterInfo, FMemoryImageAllocator>& GetUniformVirtualTextureExpressions() const { return GetUniformTextureExpressions(EMaterialTextureParameterType::Virtual); }
+	ENGINE_API TArrayView<const FMaterialTextureParameterInfo> GetUniformTextureExpressions(EMaterialTextureParameterType Type) const;
+	ENGINE_API TArrayView<const FMaterialVectorParameterInfo> GetUniformVectorParameterExpressions() const;
+	ENGINE_API TArrayView<const FMaterialScalarParameterInfo> GetUniformScalarParameterExpressions() const;
+
+	inline TArrayView<const FMaterialTextureParameterInfo> GetUniform2DTextureExpressions() const { return GetUniformTextureExpressions(EMaterialTextureParameterType::Standard2D); }
+	inline TArrayView<const FMaterialTextureParameterInfo> GetUniformCubeTextureExpressions() const { return GetUniformTextureExpressions(EMaterialTextureParameterType::Cube); }
+	inline TArrayView<const FMaterialTextureParameterInfo> GetUniform2DArrayTextureExpressions() const { return GetUniformTextureExpressions(EMaterialTextureParameterType::Array2D); }
+	inline TArrayView<const FMaterialTextureParameterInfo> GetUniformVolumeTextureExpressions() const { return GetUniformTextureExpressions(EMaterialTextureParameterType::Volume); }
+	inline TArrayView<const FMaterialTextureParameterInfo> GetUniformVirtualTextureExpressions() const { return GetUniformTextureExpressions(EMaterialTextureParameterType::Virtual); }
 
 #if WITH_EDITOR
 	const TArray<FString>& GetCompileErrors() const { return CompileErrors; }
diff --git a/Engine/Source/Runtime/Engine/Public/WeightedRandomSampler.h b/Engine/Source/Runtime/Engine/Public/WeightedRandomSampler.h
--- a/Engine/Source/Runtime/Engine/Public/WeightedRandomSampler.h
+++ b/Engine/Source/Runtime/Engine/Public/WeightedRandomSampler.h
@@ -50,8 +50,8 @@ public:
 
 	int32 GetNumEntries() const { return Prob.Num(); }
 
-	const TArray<float, FMemoryImageAllocator>& GetProb() const { return Prob; }
-	const TArray<int32, FMemoryImageAllocator>& GetAlias() const { return Alias; }
+	TArrayView<const float> GetProb() const { return Prob; }
+	TArrayView<const int32> GetAlias() const { return Alias; }
 
 protected:
 
diff --git a/Engine/Source/Runtime/RenderCore/Public/RendererInterface.h b/Engine/Source/Runtime/RenderCore/Public/RendererInterface.h
--- a/Engine/Source/Runtime/RenderCore/Public/RendererInterface.h
+++ b/Engine/Source/Runtime/RenderCore/Public/RendererInterface.h
@@ -56,20 +56,20 @@ public:
 
 	/** Default constructor, use one of the factory functions below to make a valid description */
 	FPooledRenderTargetDesc()
-		: ClearValue(FClearValueBinding())
+		: DebugName(TEXT("UnknownTexture"))
+		, ClearValue(FClearValueBinding())
 		, Extent(0, 0)
 		, Depth(0)
 		, ArraySize(1)
-		, bIsArray(false)
-		, bIsCubemap(false)
-		, NumMips(0)
-		, NumSamples(1)
-		, Format(PF_Unknown)
 		, Flags(TexCreate_None)
 		, TargetableFlags(TexCreate_None)
+		, Format(PF_Unknown)
+		, NumMips(0)
+		, NumSamples(1)
+		, bIsArray(false)
+		, bIsCubemap(false)
 		, bForceSeparateTargetAndShaderResource(false)
 		, bForceSharedTargetAndShaderResource(false)
-		, DebugName(TEXT("UnknownTexture"))
 		, AutoWritable(true)
 		, bCreateRenderTargetWriteMask(false)
 		, bCreateRenderTargetFmask(false)
@@ -383,42 +383,40 @@ public:
 		TargetableFlags &= (~TexCreate_UAV);
 	}
 
+	/** only set a pointer to memory that never gets released */
+	const TCHAR* DebugName;
 	/** Value allowed for fast clears for this target. */
 	FClearValueBinding ClearValue;
-
 	/** In pixels, (0,0) if not set, (x,0) for cube maps, todo: make 3d int vector for volume textures */
 	FIntPoint Extent;
 	/** 0, unless it's texture array or volume texture */
 	uint32 Depth;
-	
 	/** >1 if a texture array should be used (not supported on DX9) */
 	uint32 ArraySize;
-	/** true if an array texture. Note that ArraySize still can be 1 */
-	bool bIsArray;
-	/** true if a cubemap texture */
-	bool bIsCubemap;
-	/** Number of mips */
-	uint16 NumMips;
-	/** Number of MSAA samples, default: 1  */
-	uint16 NumSamples;
-	/** Texture format e.g. PF_B8G8R8A8 */
-	EPixelFormat Format;
 	/** The flags that must be set on both the shader-resource and the targetable texture. bit mask combined from elements of ETextureCreateFlags e.g. TexCreate_UAV */
 	uint32 Flags;
 	/** The flags that must be set on the targetable texture. bit mask combined from elements of ETextureCreateFlags e.g. TexCreate_UAV */
 	uint32 TargetableFlags;
+	/** Texture format e.g. PF_B8G8R8A8 */
+	EPixelFormat Format;
+	/** Number of mips */
+	uint16 NumMips;
+	/** Number of MSAA samples, default: 1  */
+	uint16 NumSamples;
+	/** true if an array texture. Note that ArraySize still can be 1 */
+	uint8 bIsArray : 1;
+	/** true if a cubemap texture */
+	uint8 bIsCubemap : 1;
 	/** Whether the shader-resource and targetable texture must be separate textures. */
-	bool bForceSeparateTargetAndShaderResource;
+	uint8 bForceSeparateTargetAndShaderResource : 1;
 	/** Whether the shader-resource and targetable texture must be the same resource. */
-	bool bForceSharedTargetAndShaderResource;
-	/** only set a pointer to memory that never gets released */
-	const TCHAR* DebugName;
+	uint8 bForceSharedTargetAndShaderResource : 1;
 	/** automatically set to writable via barrier during */
-	bool AutoWritable;
+	uint8 AutoWritable : 1;
 	/** create render target write mask (supported only on specific platforms) */
-	bool bCreateRenderTargetWriteMask;
+	uint8 bCreateRenderTargetWriteMask : 1;
 	/** create render target fmask (supported only on specific platforms) */
-	bool bCreateRenderTargetFmask;
+	uint8 bCreateRenderTargetFmask : 1;
 };
 
 
diff --git a/Engine/Source/Runtime/RenderCore/Public/Shader.h b/Engine/Source/Runtime/RenderCore/Public/Shader.h
--- a/Engine/Source/Runtime/RenderCore/Public/Shader.h
+++ b/Engine/Source/Runtime/RenderCore/Public/Shader.h
@@ -731,9 +731,10 @@ public:
 	{
 		const FHashedName SearchName = SearchStruct->GetShaderVariableHashedName();
 		int32 FoundIndex = INDEX_NONE;
-		for (int32 StructIndex = 0, Count = UniformBufferParameterStructs.Num(); StructIndex < Count; StructIndex++)
+		TArrayView<const FHashedName> UniformBufferParameterStructsView(UniformBufferParameterStructs);
+		for (int32 StructIndex = 0, Count = UniformBufferParameterStructsView.Num(); StructIndex < Count; StructIndex++)
 		{
-			if (UniformBufferParameterStructs[StructIndex] == SearchName)
+			if (UniformBufferParameterStructsView[StructIndex] == SearchName)
 			{
 				FoundIndex = StructIndex;
 				break;
@@ -1779,12 +1780,12 @@ public:
 		return HasShader(Type->GetHashedName(), PermutationId);
 	}
 
-	inline const TMemoryImageArray<TMemoryImagePtr<FShader>>& GetShaders() const
+	inline TArrayView<const TMemoryImagePtr<FShader>> GetShaders() const
 	{
 		return Shaders;
 	}
 
-	inline const TMemoryImageArray<TMemoryImagePtr<FShaderPipeline>>& GetShaderPipelines() const
+	inline TArrayView<const TMemoryImagePtr<FShaderPipeline>> GetShaderPipelines() const
 	{
 		return ShaderPipelines;
 	}
diff --git a/Engine/Source/Runtime/Renderer/Public/MeshDrawShaderBindings.h b/Engine/Source/Runtime/Renderer/Public/MeshDrawShaderBindings.h
--- a/Engine/Source/Runtime/Renderer/Public/MeshDrawShaderBindings.h
+++ b/Engine/Source/Runtime/Renderer/Public/MeshDrawShaderBindings.h
@@ -39,9 +39,9 @@ public:
 	inline uint32 GetLooseDataSizeBytes() const
 	{
 		uint32 DataSize = 0;
-		for (int32 LooseBufferIndex = 0; LooseBufferIndex < ParameterMapInfo.LooseParameterBuffers.Num(); LooseBufferIndex++)
+		for (const FShaderLooseParameterBufferInfo& Info : ParameterMapInfo.LooseParameterBuffers)
 		{
-			DataSize += ParameterMapInfo.LooseParameterBuffers[LooseBufferIndex].Size;
+			DataSize += Info.Size;
 		}
 		return DataSize;
 	}
@@ -188,15 +188,17 @@ public:
 			bool bFoundParameter = false;
 			uint8* LooseDataOffset = GetLooseDataStart();
 
-			for (int32 LooseBufferIndex = 0; LooseBufferIndex < ParameterMapInfo.LooseParameterBuffers.Num(); LooseBufferIndex++)
+			TArrayView<const FShaderLooseParameterBufferInfo> LooseParameterBuffers(ParameterMapInfo.LooseParameterBuffers);
+			for (int32 LooseBufferIndex = 0; LooseBufferIndex < LooseParameterBuffers.Num(); LooseBufferIndex++)
 			{
-				const FShaderLooseParameterBufferInfo& LooseParameterBuffer = ParameterMapInfo.LooseParameterBuffers[LooseBufferIndex];
+				const FShaderLooseParameterBufferInfo& LooseParameterBuffer = LooseParameterBuffers[LooseBufferIndex];
 
 				if (LooseParameterBuffer.BaseIndex == Parameter.GetBufferIndex())
 				{
-					for (int32 LooseParameterIndex = 0; LooseParameterIndex < LooseParameterBuffer.Parameters.Num(); LooseParameterIndex++)
+					TArrayView<const FShaderParameterInfo> Parameters(LooseParameterBuffer.Parameters);
+					for (int32 LooseParameterIndex = 0; LooseParameterIndex < Parameters.Num(); LooseParameterIndex++)
 					{
-						FShaderParameterInfo LooseParameter = LooseParameterBuffer.Parameters[LooseParameterIndex];
+						FShaderParameterInfo LooseParameter = Parameters[LooseParameterIndex];
 
 						if (Parameter.GetBaseIndex() == LooseParameter.BaseIndex)
 						{
@@ -258,27 +260,30 @@ private:
 	template<typename ElementType>
 	inline int FindSortedArrayBaseIndex(const TConstArrayView<ElementType>& Array, uint32 BaseIndex)
 	{
-		int S0 = 0;
-		int S1 = Array.Num() - 1;
-		while (S0 < S1)
+		int Index = 0; 
+		int Size = Array.Num();
+
+		//start with binary search for larger lists
+		while (Size > 8)
 		{
-			int MED = (S0 + S1) / 2;
-			if (Array[MED].BaseIndex >= BaseIndex)
-			{
-				S1 = MED;
-			}
-			else
-			{
-				S0 = (S0 == MED) ? S0 + 1 : MED;
-			}
+			const int LeftoverSize = Size % 2;
+			Size = Size / 2;
+
+			const int CheckIndex = Index + Size;
+			const int IndexIfLess = CheckIndex + LeftoverSize;
+
+			Index = Array[CheckIndex].BaseIndex < BaseIndex ? IndexIfLess : Index;
 		}
 
-		if (S0 == S1)
+		//small size array optimization
+		const int ArrayEnd = Index + Size;
+		while (Index < ArrayEnd)
 		{
-			if (Array[S0].BaseIndex == BaseIndex)
+			if (Array[Index].BaseIndex == BaseIndex)
 			{
-				return S0;
+				return Index;
 			}
+			Index++;
 		}
 
 		return -1;
