Allow FShader::Finalize to be called multiple times, ensure index doesn't change

diff --git a/Engine/Source/Runtime/RenderCore/Private/Shader.cpp b/Engine/Source/Runtime/RenderCore/Private/Shader.cpp
--- a/Engine/Source/Runtime/RenderCore/Private/Shader.cpp
+++ b/Engine/Source/Runtime/RenderCore/Private/Shader.cpp
@@ -547,9 +547,12 @@ FShader::~FShader()
 
 void FShader::Finalize(const FShaderMapResourceCode* Code)
 {
-	check(ResourceIndex == INDEX_NONE);
-	ResourceIndex = Code->FindShaderIndex(GetOutputHash());
-	check(ResourceIndex != INDEX_NONE);
+	// Finalize may be called multiple times, as a given shader may be in shader list, as well as pipeline
+	const FSHAHash& Hash = GetOutputHash();
+	const int32 NewResourceIndex = Code->FindShaderIndex(Hash);
+	checkf(NewResourceIndex != INDEX_NONE, TEXT("Missing shader code %s"), *Hash.ToString());
+	checkf(ResourceIndex == INDEX_NONE || ResourceIndex == NewResourceIndex, TEXT("Incoming index %d, existing index %d for shader %s"), NewResourceIndex, ResourceIndex, *Hash.ToString());
+	ResourceIndex = NewResourceIndex;
 }
 
 void FShader::BuildParameterMapInfo(const TMap<FString, FParameterAllocation>& ParameterMap)
