Remove an old broken check() that can potentially fail on certain platforms

diff --git a/Engine/Source/Runtime/RenderCore/Public/VertexFactory.h b/Engine/Source/Runtime/RenderCore/Public/VertexFactory.h
--- a/Engine/Source/Runtime/RenderCore/Public/VertexFactory.h
+++ b/Engine/Source/Runtime/RenderCore/Public/VertexFactory.h
@@ -175,8 +175,9 @@ public:
 		class FMeshDrawSingleShaderBindings& ShaderBindings,
 		FVertexInputStreamArray& VertexStreams) const {}
 
-	// Only for debugging type mismatch
-	LAYOUT_FIELD(uint32, Size);
+private:
+	// Should remove, but first need to fix some logic surrounding empty base classes
+	LAYOUT_FIELD_INITIALIZED(uint32, Size_DEPRECATED, 0u);
 };
 
 template<EShaderFrequency ShaderFrequency, typename VertexFactoryType>
@@ -202,7 +203,7 @@ struct TVertexFactoryParameterTraits
 	template<> struct TVertexFactoryParameterTraits<ShaderFrequency, FactoryClass> \
 	{ \
 		static const FTypeLayoutDesc* GetLayout() { return &StaticGetTypeLayoutDesc<ParameterClass>(); } \
-		static FVertexFactoryShaderParameters* Create(const class FShaderParameterMap& ParameterMap) { ParameterClass* Result = new ParameterClass(); Result->Size = sizeof(ParameterClass); Result->Bind(ParameterMap); return Result; } \
+		static FVertexFactoryShaderParameters* Create(const class FShaderParameterMap& ParameterMap) { ParameterClass* Result = new ParameterClass(); Result->Bind(ParameterMap); return Result; } \
 		static void GetElementShaderBindings( \
 			const FVertexFactoryShaderParameters* Parameters, \
 			const class FSceneInterface* Scene, \
@@ -215,7 +216,6 @@ struct TVertexFactoryParameterTraits
 			class FMeshDrawSingleShaderBindings& ShaderBindings, \
 			FVertexInputStreamArray& VertexStreams) \
 		{ \
-			check(Parameters->Size == sizeof(ParameterClass)); \
 			static_cast<const ParameterClass*>(Parameters)->GetElementShaderBindings(Scene, View, Shader, InputStreamType, FeatureLevel, VertexFactory, BatchElement, ShaderBindings, VertexStreams); \
 		} \
 	}
