Fixing race conditions in RT polling cvar toggles.

Previously, having two commands required a specific ordering to turn RT polling on and off. There was also a race condition where RT is set to poll and wakeups triggered off while RT is already sleeping waiting for a wakeup. Nothing was initially waking it up to start polling. These are all fixed.

#rnx

diff --git a/Engine/Source/Runtime/Core/Private/Async/TaskGraph.cpp b/Engine/Source/Runtime/Core/Private/Async/TaskGraph.cpp
--- a/Engine/Source/Runtime/Core/Private/Async/TaskGraph.cpp
+++ b/Engine/Source/Runtime/Core/Private/Async/TaskGraph.cpp
@@ -61,19 +61,9 @@ namespace ENamedThreads
 	CORE_API int32 bHasHighPriorityThreads = CREATE_HIPRI_TASK_THREADS;
 }
 
-static bool GDoRenderThreadWakeupTrigger = true;
-static FAutoConsoleVariableRef CVarDoRenderThreadWakeupTrigger(
-	TEXT("TaskGraph.DoRenderThreadWakeupTrigger"),
-	GDoRenderThreadWakeupTrigger,
-	TEXT("If true, task graph tasks sent to the render thread trigger a wakeup. See TaskGraph.RenderThreadPollPeriodMs.")
-);
-
-static int32 GRenderThreadPollPeriodMs = (int32)MAX_uint32;
-static FAutoConsoleVariableRef CVarRenderThreadPollPeriodMs(
-	TEXT("TaskGraph.RenderThreadPollPeriodMs"),
-	GRenderThreadPollPeriodMs,
-	TEXT("Render thread polling period in milliseconds.")
-);
+// RenderingThread.cpp sets these values if needed
+TAtomic<bool> GDoRenderThreadWakeupTrigger(true);	// Accessed by any thread, modified by GT.
+int32 GRenderThreadPollPeriodMs = -1;				// Accessed/Modified by RT only.
 
 static int32 GIgnoreThreadToDoGatherOn = 0;
 static FAutoConsoleVariableRef CVarIgnoreThreadToDoGatherOn(
diff --git a/Engine/Source/Runtime/RenderCore/Private/RenderingThread.cpp b/Engine/Source/Runtime/RenderCore/Private/RenderingThread.cpp
--- a/Engine/Source/Runtime/RenderCore/Private/RenderingThread.cpp
+++ b/Engine/Source/Runtime/RenderCore/Private/RenderingThread.cpp
@@ -56,6 +56,35 @@ volatile bool GIsRenderingThreadHealthy = true;
  */
 float GRenderingThreadMaxIdleTickFrequency = 40.f;
 
+/**
+ * RT Task Graph polling.
+ */
+
+extern TAtomic<bool> GDoRenderThreadWakeupTrigger;
+extern int32 GRenderThreadPollPeriodMs;
+
+static void OnRenderThreadPollPeriodMsChanged(IConsoleVariable* Var)
+{
+	const int32 DesiredRTPollPeriod = Var->GetInt();
+
+	GDoRenderThreadWakeupTrigger = true;
+	ENQUEUE_RENDER_COMMAND(WakeupCommand)([DesiredRTPollPeriod](FRHICommandListImmediate&)
+	{
+		GRenderThreadPollPeriodMs = DesiredRTPollPeriod;
+		if (DesiredRTPollPeriod != -1)
+		{
+			GDoRenderThreadWakeupTrigger = false;
+		}
+	});
+}
+
+static FAutoConsoleVariable CVarRenderThreadPollPeriodMs(
+	TEXT("TaskGraph.RenderThreadPollPeriodMs"),
+	-1,
+	TEXT("Render thread polling period in milliseconds. If value is -1, task graph tasks explicitly wake up RT, otherwise RT polls for tasks."),
+	FConsoleVariableDelegate::CreateStatic(&OnRenderThreadPollPeriodMsChanged)
+);
+
 /** Function to stall the rendering thread **/
 static void SuspendRendering()
 {
