handle case where LibraryInstance->Library is null.

diff --git a/Engine/Source/Runtime/RenderCore/Private/ShaderCodeLibrary.cpp b/Engine/Source/Runtime/RenderCore/Private/ShaderCodeLibrary.cpp
--- a/Engine/Source/Runtime/RenderCore/Private/ShaderCodeLibrary.cpp
+++ b/Engine/Source/Runtime/RenderCore/Private/ShaderCodeLibrary.cpp
@@ -774,7 +774,10 @@ void FShaderMapResource_SharedCode::ReleaseRHI()
 		}
 	}
 
-	LibraryInstance->Library->ReleasePreloadedShaderMap(ShaderMapIndex);
+	if(LibraryInstance->Library.IsValid())
+	{
+		LibraryInstance->Library->ReleasePreloadedShaderMap(ShaderMapIndex);
+	}
 
 	FShaderMapResource::ReleaseRHI();
 }
