Removed FShaderParameterMap::operator== call from MeshDrawShaderBindings to remove the implicit FShader* indirection from the cached MDC hash map.

diff --git a/Engine/Source/Runtime/RenderCore/Private/Shader.cpp b/Engine/Source/Runtime/RenderCore/Private/Shader.cpp
--- a/Engine/Source/Runtime/RenderCore/Private/Shader.cpp
+++ b/Engine/Source/Runtime/RenderCore/Private/Shader.cpp
@@ -640,6 +640,7 @@ void FShader::BuildParameterMapInfo(const TMap<FString, FParameterAllocation>& P
 			}
 		}
 	}
+
 	for (FShaderLooseParameterBufferInfo& Info : ParameterMapInfo.LooseParameterBuffers)
 	{
 		Info.Parameters.Sort();
@@ -648,6 +649,32 @@ void FShader::BuildParameterMapInfo(const TMap<FString, FParameterAllocation>& P
 	ParameterMapInfo.UniformBuffers.Sort();
 	ParameterMapInfo.TextureSamplers.Sort();
 	ParameterMapInfo.SRVs.Sort();
+
+	uint64 Hash = 0;
+
+	{
+		const auto CityHashValue = [&](auto Value)
+		{
+			CityHash64WithSeed((const char*)&Value, sizeof(Value), Hash);
+		};
+
+		const auto CityHashArray = [&](const TMemoryImageArray<FShaderParameterInfo>& Array)
+		{
+			CityHash64WithSeed((const char*)Array.GetData(), Array.Num() * sizeof(FShaderParameterInfo), Hash);
+		};
+
+		for (FShaderLooseParameterBufferInfo& Info : ParameterMapInfo.LooseParameterBuffers)
+		{
+			CityHashValue(Info.BaseIndex);
+			CityHashValue(Info.Size);
+			CityHashArray(Info.Parameters);
+		}
+		CityHashArray(ParameterMapInfo.UniformBuffers);
+		CityHashArray(ParameterMapInfo.TextureSamplers);
+		CityHashArray(ParameterMapInfo.SRVs);
+	}
+
+	ParameterMapInfo.Hash = Hash;
 }
 
 const FSHAHash& FShader::GetOutputHash() const
diff --git a/Engine/Source/Runtime/RenderCore/Public/Shader.h b/Engine/Source/Runtime/RenderCore/Public/Shader.h
--- a/Engine/Source/Runtime/RenderCore/Public/Shader.h
+++ b/Engine/Source/Runtime/RenderCore/Public/Shader.h
@@ -177,6 +177,7 @@ public:
 	LAYOUT_FIELD(TMemoryImageArray<FShaderParameterInfo>, TextureSamplers);
 	LAYOUT_FIELD(TMemoryImageArray<FShaderParameterInfo>, SRVs);
 	LAYOUT_FIELD(TMemoryImageArray<FShaderLooseParameterBufferInfo>, LooseParameterBuffers);
+	LAYOUT_FIELD(uint64, Hash);
 
 	friend FArchive& operator<<(FArchive& Ar,FShaderParameterMapInfo& Info)
 	{
@@ -184,15 +185,13 @@ public:
 		Ar << Info.TextureSamplers;
 		Ar << Info.SRVs;
 		Ar << Info.LooseParameterBuffers;
+		Ar << Info.Hash;
 		return Ar;
 	}
 
 	inline bool operator==(const FShaderParameterMapInfo& Rhs) const
 	{
-		return UniformBuffers == Rhs.UniformBuffers
-			&& TextureSamplers == Rhs.TextureSamplers
-			&& SRVs == Rhs.SRVs
-			&& LooseParameterBuffers == Rhs.LooseParameterBuffers;
+		return Hash == Rhs.Hash;
 	}
 };
 
diff --git a/Engine/Source/Runtime/Renderer/Public/MeshDrawShaderBindings.h b/Engine/Source/Runtime/Renderer/Public/MeshDrawShaderBindings.h
--- a/Engine/Source/Runtime/Renderer/Public/MeshDrawShaderBindings.h
+++ b/Engine/Source/Runtime/Renderer/Public/MeshDrawShaderBindings.h
@@ -15,11 +15,13 @@ class FMeshDrawShaderBindingsLayout
 {
 protected:
 	const FShaderParameterMapInfo& ParameterMapInfo;
+	const uint64 ParameterMapHash;
 
 public:
 
-	FMeshDrawShaderBindingsLayout(const TShaderRef<FShader>& Shader) :
-		ParameterMapInfo(Shader->ParameterMapInfo)
+	FMeshDrawShaderBindingsLayout(const TShaderRef<FShader>& Shader)
+		: ParameterMapInfo(Shader->ParameterMapInfo)
+		, ParameterMapHash(ParameterMapInfo.Hash)
 	{
 		check(Shader.IsValid());
 	}
@@ -33,7 +35,7 @@ public:
 
 	bool operator==(const FMeshDrawShaderBindingsLayout& Rhs) const
 	{
-		return ParameterMapInfo == Rhs.ParameterMapInfo;
+		return ParameterMapHash == Rhs.ParameterMapHash;
 	}
 
 	inline uint32 GetLooseDataSizeBytes() const
