Reverting CL 12692247. We're thinking this is causing GameThreadWaitForTask hangs on all platforms.

#rnx

#localization none
#tests Waiting for PF to test on switch.

diff --git a/Engine/Source/Runtime/RenderCore/Private/RenderingThread.cpp b/Engine/Source/Runtime/RenderCore/Private/RenderingThread.cpp
--- a/Engine/Source/Runtime/RenderCore/Private/RenderingThread.cpp
+++ b/Engine/Source/Runtime/RenderCore/Private/RenderingThread.cpp
@@ -1319,6 +1319,13 @@ FRHICommandListImmediate& GetImmediateCommandList_ForRenderCommand()
 	return FRHICommandListExecutor::GetImmediateCommandList();
 }
 
+static bool bEnablePendingCleanupObjectsCommandBatching = false;
+static FAutoConsoleVariableRef CVarEnablePendingCleanupObjectsCommandBatching(
+	TEXT("g.bEnablePendingCleanupObjectsCommandBatching"),
+	bEnablePendingCleanupObjectsCommandBatching,
+	TEXT("Enable batching PendingCleanupObjects destruction.")
+);
+
 #if WITH_EDITOR || IS_PROGRAM
 
 // mainly concerned about the cooker here, but anyway, the editor can run without a frame for a very long time (hours) and we do not have enough lock free links. 
@@ -1340,12 +1347,19 @@ FPendingCleanupObjects::~FPendingCleanupObjects()
 {
 	QUICK_SCOPE_CYCLE_COUNTER(STAT_FPendingCleanupObjects_Destruct);
 
-	StartRenderCommandFenceBundler();
+	const bool bBatchingEnabled = bEnablePendingCleanupObjectsCommandBatching;
+	if (bBatchingEnabled)
+	{
+		StartRenderCommandFenceBundler();
+	}
 	for (int32 ObjectIndex = 0; ObjectIndex < CleanupArray.Num(); ObjectIndex++)
 	{
 		delete CleanupArray[ObjectIndex];
 	}
-	StopRenderCommandFenceBundler();
+	if (bBatchingEnabled)
+	{
+		StopRenderCommandFenceBundler();
+	}
 }
 
 void BeginCleanup(FDeferredCleanupInterface* CleanupObject)
@@ -1371,12 +1385,19 @@ FPendingCleanupObjects::~FPendingCleanupObjects()
 {
 	QUICK_SCOPE_CYCLE_COUNTER(STAT_FPendingCleanupObjects_Destruct);
 
-	StartRenderCommandFenceBundler();
+	const bool bBatchingEnabled = bEnablePendingCleanupObjectsCommandBatching;
+	if (bBatchingEnabled)
+	{
+		StartRenderCommandFenceBundler();
+	}
 	for (int32 ObjectIndex = 0; ObjectIndex < CleanupArray.Num(); ObjectIndex++)
 	{
 		delete CleanupArray[ObjectIndex];
 	}
-	StopRenderCommandFenceBundler();
+	if (bBatchingEnabled)
+	{
+		StopRenderCommandFenceBundler();
+	}
 }
 
 void BeginCleanup(FDeferredCleanupInterface* CleanupObject)
