Fixed edge case in FShaderParameterMapInfo that would have caused collisions.

diff --git a/Engine/Source/Runtime/RenderCore/Private/Shader.cpp b/Engine/Source/Runtime/RenderCore/Private/Shader.cpp
--- a/Engine/Source/Runtime/RenderCore/Private/Shader.cpp
+++ b/Engine/Source/Runtime/RenderCore/Private/Shader.cpp
@@ -660,6 +660,7 @@ void FShader::BuildParameterMapInfo(const TMap<FString, FParameterAllocation>& P
 
 		const auto CityHashArray = [&](const TMemoryImageArray<FShaderParameterInfo>& Array)
 		{
+			CityHashValue(Array.Num());
 			CityHash64WithSeed((const char*)Array.GetData(), Array.Num() * sizeof(FShaderParameterInfo), Hash);
 		};
 
