Integrated fixes from Dev-TaskGraph stream

- ParallelRhiThread was also relying on our old broken ParallelRT implemention.
- Switched order because during shootergame startup is in renderingthread is divergent.

diff --git a/Engine/Source/Runtime/Core/Private/HAL/ThreadingBase.cpp b/Engine/Source/Runtime/Core/Private/HAL/ThreadingBase.cpp
--- a/Engine/Source/Runtime/Core/Private/HAL/ThreadingBase.cpp
+++ b/Engine/Source/Runtime/Core/Private/HAL/ThreadingBase.cpp
@@ -198,6 +198,7 @@ CORE_API bool IsInParallelRenderingThread()
 	{
 		newValue = FTaskTagScope::IsCurrentTag(ETaskTag::EParallelRenderingThread)
 			|| FTaskTagScope::IsCurrentTag(ETaskTag::ERenderingThread)
+			|| FTaskTagScope::IsCurrentTag(ETaskTag::EParallelRhiThread) //TODO lots of RHI functions rely on our broken IsInParallelRenderingThread;
 			|| FTaskTagScope::IsCurrentTag(ETaskTag::ERhiThread); //TODO lots of RHI functions rely on our broken IsInParallelRenderingThread;
 	}
 
diff --git a/Engine/Source/Runtime/RenderCore/Public/UniformBuffer.h b/Engine/Source/Runtime/RenderCore/Public/UniformBuffer.h
--- a/Engine/Source/Runtime/RenderCore/Public/UniformBuffer.h
+++ b/Engine/Source/Runtime/RenderCore/Public/UniformBuffer.h
@@ -83,7 +83,7 @@ public:
 	// Accessors.
 	FRHIUniformBuffer* GetUniformBufferRHI() const
 	{ 
-		checkSlow(IsInRenderingThread() || IsInParallelRenderingThread());
+		checkSlow(IsInParallelRenderingThread() || IsInRenderingThread());
 		checkf(UniformBufferRHI.GetReference(), TEXT("Attempted to access UniformBufferRHI on a TUniformBuffer that was never filled in with anything")); 
 		check(UniformBufferRHI.GetReference()); // you are trying to use a UB that was never filled with anything
 		return UniformBufferRHI; 
