SS - Fix PVS:'ChildStruct' pointer was utilized before it was verified against nullptr

diff --git a/Engine/Source/Runtime/RenderCore/Private/ShaderParameterMetadata.cpp b/Engine/Source/Runtime/RenderCore/Private/ShaderParameterMetadata.cpp
--- a/Engine/Source/Runtime/RenderCore/Private/ShaderParameterMetadata.cpp
+++ b/Engine/Source/Runtime/RenderCore/Private/ShaderParameterMetadata.cpp
@@ -609,7 +609,7 @@ void FShaderParametersMetadata::InitializeLayout()
 		{
 			Layout.bHasNonGraphOutputs = true;
 		}
-		else if (BaseType == UBMT_REFERENCED_STRUCT)
+		else if (ChildStruct && BaseType == UBMT_REFERENCED_STRUCT)
 		{
 			for (const FShaderParametersMetadata::FMember& Member : ChildStruct->GetMembers())
 			{
