Nanite config tweaks on some platforms, and fixed an ini typo

diff --git a/Engine/Source/Runtime/RenderCore/Private/ShaderCodeLibrary.cpp b/Engine/Source/Runtime/RenderCore/Private/ShaderCodeLibrary.cpp
--- a/Engine/Source/Runtime/RenderCore/Private/ShaderCodeLibrary.cpp
+++ b/Engine/Source/Runtime/RenderCore/Private/ShaderCodeLibrary.cpp
@@ -56,10 +56,10 @@ static FString ShaderExtension = TEXT(".ushaderbytecode");
 static FString StableExtension = TEXT(".scl.csv");
 static FString PipelineExtension = TEXT(".ushaderpipelines");
 
-int32 GShaderCodeLibrarySeperateLoadingCache = 0;
-static FAutoConsoleVariableRef CVarShaderCodeLibrarySeperateLoadingCache(
-	TEXT("r.ShaderCodeLibrary.SeperateLoadingCache"),
-	GShaderCodeLibrarySeperateLoadingCache,
+int32 GShaderCodeLibrarySeparateLoadingCache = 0;
+static FAutoConsoleVariableRef CVarShaderCodeLibrarySeparateLoadingCache(
+	TEXT("r.ShaderCodeLibrary.SeparateLoadingCache"),
+	GShaderCodeLibrarySeparateLoadingCache,
 	TEXT("if > 0, each shader code library has it's own loading cache."),
 	ECVF_Default
 );
@@ -551,8 +551,8 @@ public:
 					Library = FShaderCodeArchive::Create(InShaderPlatform, *Ar, DestFilePath, ShaderCodeDir, InLibraryName);
 					if (Library)
 					{
-						bool ShaderCodeLibrarySeperateLoadingCacheCommandLineOverride = FParse::Param(FCommandLine::Get(), TEXT("ShaderCodeLibrarySeperateLoadingCache"));;
-						if (GShaderCodeLibrarySeperateLoadingCache || ShaderCodeLibrarySeperateLoadingCacheCommandLineOverride)
+						bool ShaderCodeLibrarySeparateLoadingCacheCommandLineOverride = FParse::Param(FCommandLine::Get(), TEXT("ShaderCodeLibrarySeparateLoadingCache"));;
+						if (GShaderCodeLibrarySeparateLoadingCache || ShaderCodeLibrarySeparateLoadingCacheCommandLineOverride)
 						{
 							TArray<TArray<FString>> FilesToMakeUnique;
 							FilesToMakeUnique.AddDefaulted(1);
