Fix clang compiler errors (UE-98026).

- The -Wconstant-logical-operand warning shouldn't be disabled as it helps finding if (Flags && SOME_FLAG) bugs.

diff --git a/Engine/Source/Runtime/RenderCore/Private/RenderGraphBuilder.cpp b/Engine/Source/Runtime/RenderCore/Private/RenderGraphBuilder.cpp
--- a/Engine/Source/Runtime/RenderCore/Private/RenderGraphBuilder.cpp
+++ b/Engine/Source/Runtime/RenderCore/Private/RenderGraphBuilder.cpp
@@ -333,7 +333,7 @@ ERDGPassFlags FRDGBuilder::OverridePassFlags(const TCHAR* PassName, ERDGPassFlag
 		PassFlags |= ERDGPassFlags::AsyncCompute;
 	}
 
-	if (EnumHasAnyFlags(PassFlags, ERDGPassFlags::AsyncCompute) && (GRDGAsyncCompute == RDG_ASYNC_COMPUTE_DISABLED || GRDGImmediateMode || !bAsyncComputeSupported))
+	if (EnumHasAnyFlags(PassFlags, ERDGPassFlags::AsyncCompute) && (GRDGAsyncCompute == RDG_ASYNC_COMPUTE_DISABLED || GRDGImmediateMode != 0 || !bAsyncComputeSupported))
 	{
 		PassFlags &= ~ERDGPassFlags::AsyncCompute;
 		PassFlags |= ERDGPassFlags::Compute;
@@ -1018,7 +1018,7 @@ void FRDGBuilder::Clear()
 	Passes.DestructAndClear();
 	ExternalTextures.Empty();
 	ExternalBuffers.Empty();
-	if (GRDGImmediateMode || GRDGExtendResourceLifetimes)
+	if (GRDGImmediateMode != 0 || GRDGExtendResourceLifetimes != 0)
 	{
 		while (AllocatedTextures.Num())
 		{
@@ -1095,7 +1095,7 @@ FRDGPassRef FRDGBuilder::AddPass(FRDGPass* Pass)
 	}
 	IF_RDG_SCOPES(Pass->Scopes = ScopeStacks.GetCurrentScopes(Pipeline));
 
-	if (GRDGImmediateMode || GRDGExtendResourceLifetimes)
+	if (GRDGImmediateMode != 0 || GRDGExtendResourceLifetimes != 0)
 	{
 		EnumerateTextureParameters(PassParameters, [&](auto Resource)
 		{
@@ -1605,7 +1605,7 @@ void FRDGBuilder::AddTransitionInternal(
 	// the transition is split to occur after the previous pass executes.
 
 	FRDGBarrierBatchBegin* BarriersToBegin = nullptr;
-	if (PrevPass == NextPass || GRDGImmediateMode)
+	if (PrevPass == NextPass || GRDGImmediateMode != 0)
 	{
 		BarriersToBegin = &NextPass->PrologueBarriersToBegin;
 	}
@@ -1630,7 +1630,7 @@ void FRDGBuilder::BeginResourceRHI(FRDGPass* Pass, FRDGTextureRef Texture)
 
 	IF_RDG_ENABLE_DEBUG(ConditionalDebugBreak(RDG_BREAKPOINT_RESOURCE_LIFETIME, BuilderName.GetTCHAR(), Pass->GetName(), Texture->Name));
 
-	check(Texture->ReferenceCount > 0 || GRDGImmediateMode || GRDGExtendResourceLifetimes);
+	check(Texture->ReferenceCount > 0 || GRDGImmediateMode != 0 || GRDGExtendResourceLifetimes != 0);
 
 	Texture->Init(GRenderTargetPool.AllocateElementForRDG(RHICmdList, Texture->Desc, Texture->Name));
 
@@ -1751,7 +1751,7 @@ void FRDGBuilder::BeginResourceRHI(FRDGPass* Pass, FRDGBufferRef Buffer)
 
 	IF_RDG_ENABLE_DEBUG(ConditionalDebugBreak(RDG_BREAKPOINT_RESOURCE_LIFETIME, BuilderName.GetTCHAR(), Pass->GetName(), Buffer->Name));
 
-	check(Buffer->ReferenceCount > 0 || GRDGImmediateMode || GRDGExtendResourceLifetimes);
+	check(Buffer->ReferenceCount > 0 || GRDGImmediateMode != 0 || GRDGExtendResourceLifetimes != 0);
 
 	TRefCountPtr<FPooledRDGBuffer> PooledBuffer;
 	GRenderGraphResourcePool.FindFreeBuffer(RHICmdList, Buffer->Desc, PooledBuffer, Buffer->Name);
