Fix debug shaders not being compiled for editor

diff --git a/Engine/Source/Runtime/RenderCore/Private/ShaderCore.cpp b/Engine/Source/Runtime/RenderCore/Private/ShaderCore.cpp
--- a/Engine/Source/Runtime/RenderCore/Private/ShaderCore.cpp
+++ b/Engine/Source/Runtime/RenderCore/Private/ShaderCore.cpp
@@ -264,10 +264,19 @@ bool AllowDebugViewmodes(EShaderPlatform Platform)
 	const int32 ForceDebugViewValue = CVarForceDebugViewModes.GetValueOnAnyThread();
 	bool bForceEnable = ForceDebugViewValue == 1;
 	bool bForceDisable = ForceDebugViewValue == 2;
-	TStringBuilder<64> PlatformName;
-	ShaderPlatformToPlatformName(Platform).ToString(PlatformName);
-	ITargetPlatform* TargetPlatform = GetTargetPlatformManager()->FindTargetPlatform(PlatformName);
-	return (!bForceDisable) && (bForceEnable || !TargetPlatform || !TargetPlatform->RequiresCookedData());
+
+	#if 0
+		// We can't distinguish between editor and non-editor targets solely from EShaderPlatform.
+		// RequiresCookedData() was always returning true for Windows in UE4, and false in UE5, for Windows
+		TStringBuilder<64> PlatformName;
+		ShaderPlatformToPlatformName(Platform).ToString(PlatformName);
+		ITargetPlatform* TargetPlatform = GetTargetPlatformManager()->FindTargetPlatform(PlatformName);
+		return (!bForceDisable) && (bForceEnable || !TargetPlatform || !TargetPlatform->RequiresCookedData());
+	#else
+		// Always include debug shaders for Windows targets until we have a way to distinguish
+		return (!bForceDisable) && (bForceEnable || IsPCPlatform(Platform));
+	#endif
+
 #else
 	return AllowDebugViewmodes();
 #endif
