Slowed down VisualizeTexture blinking on negative values

diff --git a/Engine/Source/Runtime/RenderCore/Private/VisualizeTexture.cpp b/Engine/Source/Runtime/RenderCore/Private/VisualizeTexture.cpp
--- a/Engine/Source/Runtime/RenderCore/Private/VisualizeTexture.cpp
+++ b/Engine/Source/Runtime/RenderCore/Private/VisualizeTexture.cpp
@@ -221,9 +221,9 @@ void FVisualizeTexture::CreateContentCapturePass(FRDGBuilder& GraphBuilder, cons
 
 		{
 			// alternates between 0 and 1 with a short pause
-			const float FracTimeScale = 2.0f;
+			const float FracTimeScale = 1.0f / 4.0f;
 			float FracTime = FApp::GetCurrentTime() * FracTimeScale - floor(FApp::GetCurrentTime() * FracTimeScale);
-			float BlinkState = (FracTime > 0.5f) ? 1.0f : 0.0f;
+			float BlinkState = (FracTime < 1.0f / 16.0f) ? 1.0f : 0.0f;
 
 			FVector4 VisualizeParamValue[3];
 
