Fixed RDG to respect the MSAA plane.

diff --git a/Engine/Source/Runtime/RenderCore/Private/RenderGraphBuilder.cpp b/Engine/Source/Runtime/RenderCore/Private/RenderGraphBuilder.cpp
--- a/Engine/Source/Runtime/RenderCore/Private/RenderGraphBuilder.cpp
+++ b/Engine/Source/Runtime/RenderCore/Private/RenderGraphBuilder.cpp
@@ -1253,8 +1253,10 @@ FRHIRenderPassInfo FRDGBuilder::GetRenderPassInfo(const FRDGPass* Pass) const
 		const ERenderTargetStoreAction DepthStoreAction = ExclusiveDepthStencil.IsDepthWrite() ? ERenderTargetStoreAction::EStore : ERenderTargetStoreAction::ENoAction;
 		const ERenderTargetStoreAction StencilStoreAction = ExclusiveDepthStencil.IsStencilWrite() ? ERenderTargetStoreAction::EStore : ERenderTargetStoreAction::ENoAction;
 
+		auto& RenderTargetItem = Texture->PooledTexture->GetRenderTargetItem();
+
 		auto& OutDepthStencil = RenderPassInfo.DepthStencilRenderTarget;
-		OutDepthStencil.DepthStencilTarget = Texture->PooledTexture->GetRenderTargetItem().TargetableTexture;
+		OutDepthStencil.DepthStencilTarget = DepthStencil.GetMsaaPlane() == ERenderTargetMsaaPlane::Unresolved ? RenderTargetItem.TargetableTexture : RenderTargetItem.ShaderResourceTexture;
 		OutDepthStencil.ResolveTarget = nullptr;
 		OutDepthStencil.Action = MakeDepthStencilTargetActions(
 			MakeRenderTargetActions(DepthStencil.GetDepthLoadAction(), DepthStoreAction),
