Use const-ref to iterate over TMap entries to fix compiler error with clang 10 (Merged from CL 14359172).

#rnx

diff --git a/Engine/Source/Runtime/RenderCore/Private/RenderGraphValidation.cpp b/Engine/Source/Runtime/RenderCore/Private/RenderGraphValidation.cpp
--- a/Engine/Source/Runtime/RenderCore/Private/RenderGraphValidation.cpp
+++ b/Engine/Source/Runtime/RenderCore/Private/RenderGraphValidation.cpp
@@ -716,7 +716,7 @@ void FRDGBarrierValidation::ValidateBarrierBatchBegin(const FRDGBarrierBatchBegi
 		}
 	};
 
-	for (const auto Pair : ResourceMap.Textures)
+	for (const auto& Pair : ResourceMap.Textures)
 	{
 		FRDGTextureRef Texture = Pair.Key;
 		
