Changed unbound parameters assert to ensure, since you can fix it with recompile shaders

diff --git a/Engine/Source/Runtime/RenderCore/Private/ShaderParameterStruct.cpp b/Engine/Source/Runtime/RenderCore/Private/ShaderParameterStruct.cpp
--- a/Engine/Source/Runtime/RenderCore/Private/ShaderParameterStruct.cpp
+++ b/Engine/Source/Runtime/RenderCore/Private/ShaderParameterStruct.cpp
@@ -282,7 +282,7 @@ void FShaderParameterBindings::BindForLegacyShaderParameters(const FShader* Shad
 			}
 		}
 
-		UE_LOG(LogShaders, Fatal, TEXT("%s"), *ErrorString);
+		ensureMsgf(false, TEXT("%s"), *ErrorString);
 	}
 }
 
