Fix for ConvertToExternalX causing passes to be culled in RDG.

 - This function calls FRDGBuilder::PreallocateX, which sets bExternal to true.
 - Culling relies on bUntrackedOutputs to mark passes with external resources as root nodes in the culling traversal.
 - bUntrackedOutputs was being assigned during AddPass instead of during Compile.

diff --git a/Engine/Source/Runtime/RenderCore/Private/RenderGraphBuilder.cpp b/Engine/Source/Runtime/RenderCore/Private/RenderGraphBuilder.cpp
--- a/Engine/Source/Runtime/RenderCore/Private/RenderGraphBuilder.cpp
+++ b/Engine/Source/Runtime/RenderCore/Private/RenderGraphBuilder.cpp
@@ -628,6 +628,8 @@ void FRDGBuilder::Compile()
 		{
 			FRDGPass* Pass = Passes[PassHandle];
 
+			bool bUntrackedOutputs = Pass->GetParameters().HasExternalOutputs();
+
 			for (auto& TexturePair : Pass->TextureStates)
 			{
 				FRDGTextureRef Texture = TexturePair.Key;
@@ -647,12 +649,15 @@ void FRDGBuilder::Compile()
 				{
 					AddCullingDependency(LastProducers[Index], PassHandle, PassState[bWholePassState ? 0 : Index].Access);
 				}
+
+				bUntrackedOutputs |= Texture->bExternal;
 			}
 
 			for (auto& BufferPair : Pass->BufferStates)
 			{
 				FRDGBufferRef Buffer = BufferPair.Key;
 				AddCullingDependency(Buffer->LastProducer, PassHandle, BufferPair.Value.State.Access);
+				bUntrackedOutputs |= Buffer->bExternal;
 			}
 
 			const ERDGPassFlags PassFlags = Pass->GetFlags();
@@ -663,7 +668,7 @@ void FRDGBuilder::Compile()
 			PassesOnRaster[PassHandle] = bRaster;
 			PassesOnAsyncCompute[PassHandle] = bAsyncCompute;
 			PassesToNeverCull[PassHandle] = bNeverCull;
-			PassesWithUntrackedOutputs[PassHandle] = Pass->bUntrackedOutputs;
+			PassesWithUntrackedOutputs[PassHandle] = bUntrackedOutputs;
 			AsyncComputePassCount += bAsyncCompute ? 1 : 0;
 			RasterPassCount += bRaster ? 1 : 0;
 		}
@@ -1321,7 +1326,6 @@ void FRDGBuilder::SetupPass(FRDGPass* Pass)
 	const ERDGPassFlags PassFlags = Pass->GetFlags();
 	const ERHIPipeline PassPipeline = Pass->GetPipeline();
 
-	bool bPassUntrackedOutputs = PassParameters.HasExternalOutputs();
 	bool bPassUAVAccess = false;
 
 	Pass->TextureStates.Reserve(PassParameters.GetTextureParameterCount() + (PassParameters.HasRenderTargets() ? (MaxSimultaneousRenderTargets + 1) : 0));
@@ -1362,7 +1366,6 @@ void FRDGBuilder::SetupPass(FRDGPass* Pass)
 		}
 
 		bPassUAVAccess |= EnumHasAnyFlags(Access, ERHIAccess::UAVMask);
-		bPassUntrackedOutputs |= Texture->bExternal;
 	});
 
 	Pass->BufferStates.Reserve(PassParameters.GetBufferParameterCount());
@@ -1379,11 +1382,9 @@ void FRDGBuilder::SetupPass(FRDGPass* Pass)
 		PassState.State.Pipeline = PassPipeline;
 
 		bPassUAVAccess |= EnumHasAnyFlags(Access, ERHIAccess::UAVMask);
-		bPassUntrackedOutputs |= Buffer->bExternal;
 	});
 
 	Pass->bUAVAccess = bPassUAVAccess;
-	Pass->bUntrackedOutputs = bPassUntrackedOutputs;
 
 	const bool bEmptyParameters = !Pass->TextureStates.Num() && !Pass->BufferStates.Num();
 	PassesWithEmptyParameters.Add(bEmptyParameters);
diff --git a/Engine/Source/Runtime/RenderCore/Public/RenderGraphPass.h b/Engine/Source/Runtime/RenderCore/Public/RenderGraphPass.h
--- a/Engine/Source/Runtime/RenderCore/Public/RenderGraphPass.h
+++ b/Engine/Source/Runtime/RenderCore/Public/RenderGraphPass.h
@@ -298,9 +298,6 @@ private:
 			/** Whether the pass writes to a UAV. */
 			uint32 bUAVAccess : 1;
 
-			/** Whether the pass has untracked outputs. */
-			uint32 bUntrackedOutputs : 1;
-
 			/** Whether this pass allocated a texture through the pool. */
 			IF_RDG_ENABLE_DEBUG(uint32 bFirstTextureAllocated : 1);
 		};
