Check statement should be using "frozen" API during validation.

diff --git a/Engine/Source/Runtime/RenderCore/Private/ShaderParameterStruct.cpp b/Engine/Source/Runtime/RenderCore/Private/ShaderParameterStruct.cpp
--- a/Engine/Source/Runtime/RenderCore/Private/ShaderParameterStruct.cpp
+++ b/Engine/Source/Runtime/RenderCore/Private/ShaderParameterStruct.cpp
@@ -442,7 +442,7 @@ void ValidateShaderParameters(const TShaderRef<FShader>& Shader, const FShaderPa
 	checkf(
 		Bindings.StructureLayoutHash == ParametersMetadata->GetLayoutHash(),
 		TEXT("Shader %s's parameter structure has changed without recompilation of the shader"),
-		Shader->GetTypeUnfrozen()->GetName());
+		Shader->GetType()->GetName());
 
 	const uint8* Base = reinterpret_cast<const uint8*>(Parameters);
 
