Added a check to catch null RDG buffer clear.

#rnx

diff --git a/Engine/Source/Runtime/RenderCore/Private/RenderGraphUtils.cpp b/Engine/Source/Runtime/RenderCore/Private/RenderGraphUtils.cpp
--- a/Engine/Source/Runtime/RenderCore/Private/RenderGraphUtils.cpp
+++ b/Engine/Source/Runtime/RenderCore/Private/RenderGraphUtils.cpp
@@ -299,6 +299,8 @@ END_SHADER_PARAMETER_STRUCT()
 
 void AddClearUAVPass(FRDGBuilder& GraphBuilder, FRDGBufferUAVRef BufferUAV, uint32 Value)
 {
+	check(BufferUAV);
+
 	FClearBufferUAVParameters* Parameters = GraphBuilder.AllocParameters<FClearBufferUAVParameters>();
 	Parameters->BufferUAV = BufferUAV;
 
