Merged CL to fix UAV crash: Refactored RDG pooled buffer UAV/SRV creation into public helper functions to aid with hair conversion.

#rnx

diff --git a/Engine/Source/Runtime/RenderCore/Private/RenderGraphBuilder.cpp b/Engine/Source/Runtime/RenderCore/Private/RenderGraphBuilder.cpp
--- a/Engine/Source/Runtime/RenderCore/Private/RenderGraphBuilder.cpp
+++ b/Engine/Source/Runtime/RenderCore/Private/RenderGraphBuilder.cpp
@@ -2207,7 +2207,6 @@ void FRDGBuilder::BeginResourceRHI(FRDGPassHandle PassHandle, FRDGBufferSRVRef S
 	}
 
 	FRDGBufferRef Buffer = SRV->Desc.Buffer;
-
 	check(Buffer->PooledBuffer);
 
 	if (!Buffer->FirstPass.IsValid())
@@ -2215,29 +2214,7 @@ void FRDGBuilder::BeginResourceRHI(FRDGPassHandle PassHandle, FRDGBufferSRVRef S
 		Buffer->FirstPass = PassHandle;
 	}
 
-	if (Buffer->PooledBuffer->SRVs.Contains(SRV->Desc))
-	{
-		SRV->ResourceRHI = Buffer->PooledBuffer->SRVs[SRV->Desc];
-		return;
-	}
-
-	FShaderResourceViewRHIRef RHIShaderResourceView;
-
-	if (Buffer->Desc.UnderlyingType == FRDGBufferDesc::EUnderlyingType::VertexBuffer)
-	{
-		RHIShaderResourceView = RHICreateShaderResourceView(Buffer->PooledBuffer->VertexBuffer, SRV->Desc.BytesPerElement, SRV->Desc.Format);
-	}
-	else if (Buffer->Desc.UnderlyingType == FRDGBufferDesc::EUnderlyingType::StructuredBuffer)
-	{
-		RHIShaderResourceView = RHICreateShaderResourceView(Buffer->PooledBuffer->StructuredBuffer);
-	}
-	else
-	{
-		check(0);
-	}
-
-	SRV->ResourceRHI = RHIShaderResourceView;
-	Buffer->PooledBuffer->SRVs.Add(SRV->Desc, RHIShaderResourceView);
+	SRV->ResourceRHI = Buffer->PooledBuffer->GetOrCreateSRV(SRV->Desc);
 }
 
 void FRDGBuilder::BeginResourceRHI(FRDGPassHandle PassHandle, FRDGBufferUAV* UAV)
@@ -2251,33 +2228,7 @@ void FRDGBuilder::BeginResourceRHI(FRDGPassHandle PassHandle, FRDGBufferUAV* UAV
 
 	FRDGBufferRef Buffer = UAV->Desc.Buffer;
 	BeginResourceRHI(PassHandle, Buffer);
-
-	if (Buffer->PooledBuffer->UAVs.Contains(UAV->Desc))
-	{
-		UAV->ResourceRHI = Buffer->PooledBuffer->UAVs[UAV->Desc];
-		return;
-	}
-
-	// Hack to make sure only one UAVs is around.
-	Buffer->PooledBuffer->UAVs.Empty();
-
-	FUnorderedAccessViewRHIRef RHIUnorderedAccessView;
-
-	if (Buffer->Desc.UnderlyingType == FRDGBufferDesc::EUnderlyingType::VertexBuffer)
-	{
-		RHIUnorderedAccessView = RHICreateUnorderedAccessView(Buffer->PooledBuffer->VertexBuffer, UAV->Desc.Format);
-	}
-	else if (Buffer->Desc.UnderlyingType == FRDGBufferDesc::EUnderlyingType::StructuredBuffer)
-	{
-		RHIUnorderedAccessView = RHICreateUnorderedAccessView(Buffer->PooledBuffer->StructuredBuffer, UAV->Desc.bSupportsAtomicCounter, UAV->Desc.bSupportsAppendBuffer);
-	}
-	else
-	{
-		check(0);
-	}
-
-	UAV->ResourceRHI = RHIUnorderedAccessView;
-	Buffer->PooledBuffer->UAVs.Add(UAV->Desc, RHIUnorderedAccessView);
+	UAV->ResourceRHI = Buffer->PooledBuffer->GetOrCreateUAV(UAV->Desc);
 }
 
 void FRDGBuilder::EndResourceRHI(FRDGPassHandle PassHandle, FRDGTextureRef Texture, uint32 ReferenceCount)
diff --git a/Engine/Source/Runtime/RenderCore/Private/RenderGraphResources.cpp b/Engine/Source/Runtime/RenderCore/Private/RenderGraphResources.cpp
--- a/Engine/Source/Runtime/RenderCore/Private/RenderGraphResources.cpp
+++ b/Engine/Source/Runtime/RenderCore/Private/RenderGraphResources.cpp
@@ -299,4 +299,58 @@ FRDGTextureRef FRDGTexture::GetPassthrough(const TRefCountPtr<IPooledRenderTarge
 		return Texture;
 	}
 	return nullptr;
+}
+
+FRHIShaderResourceView* FRDGPooledBuffer::GetOrCreateSRV(FRDGBufferSRVDesc SRVDesc)
+{
+	if (const auto* FoundPtr = SRVs.Find(SRVDesc))
+	{
+		return FoundPtr->GetReference();
+	}
+
+	FShaderResourceViewRHIRef RHIShaderResourceView;
+
+	if (Desc.UnderlyingType == FRDGBufferDesc::EUnderlyingType::VertexBuffer)
+	{
+		RHIShaderResourceView = RHICreateShaderResourceView(VertexBuffer, SRVDesc.BytesPerElement, SRVDesc.Format);
+	}
+	else if (Desc.UnderlyingType == FRDGBufferDesc::EUnderlyingType::StructuredBuffer)
+	{
+		RHIShaderResourceView = RHICreateShaderResourceView(StructuredBuffer);
+	}
+	else
+	{
+		checkNoEntry();
+	}
+
+	FRHIShaderResourceView* View = RHIShaderResourceView.GetReference();
+	SRVs.Emplace(SRVDesc, MoveTemp(RHIShaderResourceView));
+	return View;
+}
+
+FRHIUnorderedAccessView* FRDGPooledBuffer::GetOrCreateUAV(FRDGBufferUAVDesc UAVDesc)
+{
+	if (const auto* FoundPtr = UAVs.Find(UAVDesc))
+	{
+		return FoundPtr->GetReference();
+	}
+
+	FUnorderedAccessViewRHIRef RHIUnorderedAccessView;
+
+	if (Desc.UnderlyingType == FRDGBufferDesc::EUnderlyingType::VertexBuffer)
+	{
+		RHIUnorderedAccessView = RHICreateUnorderedAccessView(VertexBuffer, UAVDesc.Format);
+	}
+	else if (Desc.UnderlyingType == FRDGBufferDesc::EUnderlyingType::StructuredBuffer)
+	{
+		RHIUnorderedAccessView = RHICreateUnorderedAccessView(StructuredBuffer, UAVDesc.bSupportsAtomicCounter, UAVDesc.bSupportsAppendBuffer);
+	}
+	else
+	{
+		checkNoEntry();
+	}
+
+	FRHIUnorderedAccessView* View = RHIUnorderedAccessView.GetReference();
+	UAVs.Emplace(UAVDesc, MoveTemp(RHIUnorderedAccessView));
+	return View;
 }
\ No newline at end of file
diff --git a/Engine/Source/Runtime/RenderCore/Public/RenderGraphResources.h b/Engine/Source/Runtime/RenderCore/Public/RenderGraphResources.h
--- a/Engine/Source/Runtime/RenderCore/Public/RenderGraphResources.h
+++ b/Engine/Source/Runtime/RenderCore/Public/RenderGraphResources.h
@@ -1114,6 +1114,12 @@ class FRDGPooledBuffer
 public:
 	const FRDGBufferDesc Desc;
 
+	/** Finds a UAV matching the descriptor in the cache or creates a new one and updates the cache. */
+	FRHIUnorderedAccessView* GetOrCreateUAV(FRDGBufferUAVDesc UAVDesc);
+
+	/** Finds a SRV matching the descriptor in the cache or creates a new one and updates the cache. */
+	FRHIShaderResourceView* GetOrCreateSRV(FRDGBufferSRVDesc SRVDesc);
+
 	FRHIVertexBuffer* GetVertexBufferRHI() const
 	{
 		return VertexBuffer;
