Frosty fixes.

diff --git a/Engine/Source/Runtime/Engine/Private/ShaderCompiler/ShaderGenerationUtil.cpp b/Engine/Source/Runtime/Engine/Private/ShaderCompiler/ShaderGenerationUtil.cpp
--- a/Engine/Source/Runtime/Engine/Private/ShaderCompiler/ShaderGenerationUtil.cpp
+++ b/Engine/Source/Runtime/Engine/Private/ShaderCompiler/ShaderGenerationUtil.cpp
@@ -644,6 +644,9 @@ static int32 GetBufferNumBits(EGBufferType Format, int32 Channel)
 	case GBT_Float_16_16:
 		Ret = 16;
 		break;
+	case GBT_Float_16_16_16_16:
+		Ret = 16;
+		break;
 	default:
 		check(0);
 		break;
@@ -673,6 +676,9 @@ static int32 GetTargetNumChannels(EGBufferType Type)
 	case GBT_Float_16_16:
 		Ret = 2;
 		break;
+	case GBT_Float_16_16_16_16:
+		Ret = 4;
+		break;
 	default:
 		check(0);
 		break;
diff --git a/Engine/Source/Runtime/RenderCore/Private/GBufferInfo.cpp b/Engine/Source/Runtime/RenderCore/Private/GBufferInfo.cpp
--- a/Engine/Source/Runtime/RenderCore/Private/GBufferInfo.cpp
+++ b/Engine/Source/Runtime/RenderCore/Private/GBufferInfo.cpp
@@ -256,7 +256,7 @@ FGBufferInfo RENDERCORE_API FetchLegacyGBufferInfo(const FGBufferParams& Params)
 		TargetGBufferD = 5;
 
 		// note the false for use extra flags for velocity, not quite sure of all the ramifications, but this keeps it consistent with previous usage
-		Info.Targets[4].Init(GBT_Float_16_16,   TEXT("Velocity"), false,  true,  true, false);
+		Info.Targets[4].Init(Params.bUsesVelocityDepth ? GBT_Float_16_16_16_16 : GBT_Float_16_16,   TEXT("Velocity"), false,  true,  true, false);
 		Info.Targets[5].Init(GBT_Unorm_8_8_8_8, TEXT("GBufferD"), false,  true,  true,  true);
 
 		if (Params.bHasPrecShadowFactor)
diff --git a/Engine/Source/Runtime/Renderer/Private/PostProcess/SceneRenderTargets.cpp b/Engine/Source/Runtime/Renderer/Private/PostProcess/SceneRenderTargets.cpp
--- a/Engine/Source/Runtime/Renderer/Private/PostProcess/SceneRenderTargets.cpp
+++ b/Engine/Source/Runtime/Renderer/Private/PostProcess/SceneRenderTargets.cpp
@@ -1380,7 +1380,7 @@ void FSceneRenderTargets::AllocGBufferTargets(FRHICommandList& RHICmdList, EText
 		const int32 IndexC = FindGBufferTargetByName(GBufferInfo, "GBufferC");
 		const int32 IndexD = FindGBufferTargetByName(GBufferInfo, "GBufferD");
 		const int32 IndexE = FindGBufferTargetByName(GBufferInfo, "GBufferE");
-		const int32 IndexVelocity = FindGBufferTargetByName(GBufferInfo, "GBufferVelocity");
+		const int32 IndexVelocity = FindGBufferTargetByName(GBufferInfo, "Velocity");
 
 		if (IndexA >= 0)
 		{
