Fixing build CIS warning.

diff --git a/Engine/Source/Runtime/RenderCore/Private/GBufferInfo.cpp b/Engine/Source/Runtime/RenderCore/Private/GBufferInfo.cpp
--- a/Engine/Source/Runtime/RenderCore/Private/GBufferInfo.cpp
+++ b/Engine/Source/Runtime/RenderCore/Private/GBufferInfo.cpp
@@ -94,14 +94,8 @@ TArray < EGBufferSlot > FetchGBufferSlots(bool bHasVelocity, bool bHasTangent, b
 	NeededSlots.Push(GBS_SelectiveOutputMask);
 	NeededSlots.Push(GBS_BaseColor);
 
-	if (bHasPrecShadowFactor)
-	{
-		NeededSlots.Push(GBS_GenericAO);
-	}
-	else
-	{
-		NeededSlots.Push(GBS_GenericAO);
-	}
+	// this is needed for all combinations, will have to split it later
+	NeededSlots.Push(GBS_GenericAO);
 
 	if (bHasVelocity)
 	{
