Various fixes needed for next round of RDG conversions.

#rnx

diff --git a/Engine/Source/Runtime/RHI/Public/RHIContext.h b/Engine/Source/Runtime/RHI/Public/RHIContext.h
--- a/Engine/Source/Runtime/RHI/Public/RHIContext.h
+++ b/Engine/Source/Runtime/RHI/Public/RHIContext.h
@@ -188,7 +188,7 @@ public:
 
 	virtual void RHISetGlobalUniformBuffers(const FUniformBufferStaticBindings& InUniformBuffers)
 	{
-		checkNoEntry();
+		/** empty default implementation. */
 	}
 
 	virtual void RHIPushEvent(const TCHAR* Name, FColor Color) = 0;
diff --git a/Engine/Source/Runtime/RenderCore/Private/RenderGraphValidation.cpp b/Engine/Source/Runtime/RenderCore/Private/RenderGraphValidation.cpp
--- a/Engine/Source/Runtime/RenderCore/Private/RenderGraphValidation.cpp
+++ b/Engine/Source/Runtime/RenderCore/Private/RenderGraphValidation.cpp
@@ -60,18 +60,29 @@ const ERHIAccess AccessMaskRaster  = ERHIAccess::ResolveSrc | ERHIAccess::Resolv
 bool GRDGInExecutePassScope = false;
 } //! namespace
 
+FRDGUserValidation::FRDGUserValidation()
+	: ExpectedNumMarks(FMemStack::Get().GetNumMarks())
+{}
+
 FRDGUserValidation::~FRDGUserValidation()
 {
 	checkf(bHasExecuted, TEXT("Render graph execution is required to ensure consistency with immediate mode."));
 }
 
+void FRDGUserValidation::MemStackGuard()
+{
+	checkf(ExpectedNumMarks == FMemStack::Get().GetNumMarks(), TEXT("A MemStack mark was added during the FRDGBuilder lifetime. This is not allowed as it will free memory still used by the builder."));
+}
+
 void FRDGUserValidation::ExecuteGuard(const TCHAR* Operation, const TCHAR* ResourceName)
 {
 	checkf(!bHasExecuted, TEXT("Render graph operation '%s' with resource '%s' must be performed prior to graph execution."), Operation, ResourceName);
+	MemStackGuard();
 }
 
 void FRDGUserValidation::ValidateCreateTexture(FRDGTextureRef Texture)
 {
+	MemStackGuard();
 	check(Texture);
 	if (GRDGDebug)
 	{
@@ -81,6 +92,7 @@ void FRDGUserValidation::ValidateCreateTexture(FRDGTextureRef Texture)
 
 void FRDGUserValidation::ValidateCreateBuffer(FRDGBufferRef Buffer)
 {
+	MemStackGuard();
 	check(Buffer);
 	if (GRDGDebug)
 	{
@@ -102,6 +114,7 @@ void FRDGUserValidation::ValidateCreateExternalBuffer(FRDGBufferRef Buffer)
 
 void FRDGUserValidation::ValidateExtractResource(FRDGParentResourceRef Resource)
 {
+	MemStackGuard();
 	check(Resource);
 
 	checkf(Resource->ParentDebugData.bHasBeenProduced,
@@ -117,6 +130,7 @@ void FRDGUserValidation::ValidateExtractResource(FRDGParentResourceRef Resource)
 
 void FRDGUserValidation::RemoveUnusedWarning(FRDGParentResourceRef Resource)
 {
+	MemStackGuard();
 	check(Resource);
 
 	// Removes 'produced but not used' warning.
@@ -140,6 +154,7 @@ bool FRDGUserValidation::TryMarkForClobber(FRDGParentResourceRef Resource) const
 
 void FRDGUserValidation::ValidateAddPass(const FRDGPass* Pass, bool bSkipPassAccessMarking)
 {
+	MemStackGuard();
 	const FRenderTargetBindingSlots* RenderTargetBindingSlots = nullptr;
 
 	// Pass flags are validated as early as possible by the builder in AddPass.
@@ -462,11 +477,14 @@ void FRDGUserValidation::ValidateAddPass(const FRDGPass* Pass, bool bSkipPassAcc
 
 void FRDGUserValidation::ValidateExecuteBegin()
 {
+	MemStackGuard();
 	checkf(!bHasExecuted, TEXT("Render graph execution should only happen once to ensure consistency with immediate mode."));
 }
 
 void FRDGUserValidation::ValidateExecuteEnd()
 {
+	MemStackGuard();
+
 	bHasExecuted = true;
 
 	if (GRDGDebug)
diff --git a/Engine/Source/Runtime/RenderCore/Private/ShaderParameters.cpp b/Engine/Source/Runtime/RenderCore/Private/ShaderParameters.cpp
--- a/Engine/Source/Runtime/RenderCore/Private/ShaderParameters.cpp
+++ b/Engine/Source/Runtime/RenderCore/Private/ShaderParameters.cpp
@@ -242,7 +242,6 @@ static void CreateHLSLUniformBufferStructMembersDeclaration(
 
 		if (IsShaderParameterTypeForUniformBufferLayout(Member.GetBaseType()))
 		{
-			check(Member.GetBaseType() != UBMT_RDG_TEXTURE_SRV && Member.GetBaseType() != UBMT_RDG_TEXTURE_UAV);
 			checkf(Member.GetBaseType() != UBMT_RDG_TEXTURE_ACCESS && Member.GetBaseType() != UBMT_RDG_BUFFER_ACCESS, TEXT("Copy destination usage is not supported in uniform buffers."));
 			if (Member.GetBaseType() == UBMT_SRV)
 			{
diff --git a/Engine/Source/Runtime/RenderCore/Public/RenderGraphUtils.h b/Engine/Source/Runtime/RenderCore/Public/RenderGraphUtils.h
--- a/Engine/Source/Runtime/RenderCore/Public/RenderGraphUtils.h
+++ b/Engine/Source/Runtime/RenderCore/Public/RenderGraphUtils.h
@@ -608,6 +608,14 @@ FORCEINLINE void AddSetCurrentStatPass(FRDGBuilder& GraphBuilder, TStatId StatId
 	});
 }
 
+FORCEINLINE void AddDispatchToRHIThreadPass(FRDGBuilder& GraphBuilder)
+{
+	AddPass(GraphBuilder, [](FRHICommandListImmediate& RHICmdList)
+	{
+		RHICmdList.ImmediateFlush(EImmediateFlushType::DispatchToRHIThread);
+	});
+}
+
 FORCEINLINE void AddBeginUAVOverlapPass(FRDGBuilder& GraphBuilder)
 {
 	AddPass(GraphBuilder, [](FRHICommandListImmediate& RHICmdList)
@@ -676,6 +684,18 @@ void AddReadbackTexturePass(FRDGBuilder& GraphBuilder, FRDGEventName&& Name, FRD
 	GraphBuilder.AddPass(MoveTemp(Name), PassParameters, ERDGPassFlags::Readback, MoveTemp(ExecuteLambda));
 }
 
+BEGIN_SHADER_PARAMETER_STRUCT(FReadbackBufferParameters, )
+	RDG_BUFFER_ACCESS(Buffer, ERHIAccess::CopySrc)
+END_SHADER_PARAMETER_STRUCT()
+
+template <typename ExecuteLambdaType>
+void AddReadbackBufferPass(FRDGBuilder& GraphBuilder, FRDGEventName&& Name, FRDGBufferRef Buffer, ExecuteLambdaType&& ExecuteLambda)
+{
+	auto* PassParameters = GraphBuilder.AllocParameters<FReadbackBufferParameters>();
+	PassParameters->Buffer = Buffer;
+	GraphBuilder.AddPass(MoveTemp(Name), PassParameters, ERDGPassFlags::Readback, MoveTemp(ExecuteLambda));
+}
+
 // Forces the graph to make the resource immediately available as if it were registered. Returns the pooled resource.
 
 inline void ConvertToExternalBuffer(FRDGBuilder& GraphBuilder, FRDGBufferRef Buffer, TRefCountPtr<FRDGPooledBuffer>& OutPooledBuffer)
diff --git a/Engine/Source/Runtime/RenderCore/Public/RenderGraphValidation.h b/Engine/Source/Runtime/RenderCore/Public/RenderGraphValidation.h
--- a/Engine/Source/Runtime/RenderCore/Public/RenderGraphValidation.h
+++ b/Engine/Source/Runtime/RenderCore/Public/RenderGraphValidation.h
@@ -23,7 +23,7 @@
 class RENDERCORE_API FRDGUserValidation final
 {
 public:
-	FRDGUserValidation() = default;
+	FRDGUserValidation();
 	FRDGUserValidation(const FRDGUserValidation&) = delete;
 	~FRDGUserValidation();
 
@@ -72,6 +72,10 @@ private:
 
 	/** Whether the Execute() has already been called. */
 	bool bHasExecuted = false;
+
+	void MemStackGuard();
+
+	int32 ExpectedNumMarks;
 };
 
 /** This class validates and logs barriers submitted by the graph. */
