Added missing validation case for RDG debug logging.

#rnx

diff --git a/Engine/Source/Runtime/RenderCore/Private/RenderGraphValidation.cpp b/Engine/Source/Runtime/RenderCore/Private/RenderGraphValidation.cpp
--- a/Engine/Source/Runtime/RenderCore/Private/RenderGraphValidation.cpp
+++ b/Engine/Source/Runtime/RenderCore/Private/RenderGraphValidation.cpp
@@ -569,10 +569,27 @@ void FRDGUserValidation::ValidateExecutePassBegin(const FRDGPass* Pass)
 				}
 			break;
 			case UBMT_RDG_TEXTURE_ACCESS:
+			{
+				const FRDGTextureAccess TextureAccess = Parameter.GetAsTextureAccess();
+				if (FRDGTextureRef Texture = TextureAccess.GetTexture())
+				{
+					const ERHIAccess Access = TextureAccess.GetAccess();
+					if (EnumHasAnyFlags(Access, ERHIAccess::UAVMask))
+					{
+						Texture->TextureDebugData.bHasNeededUAV = true;
+					}
+					if (EnumHasAnyFlags(Access, ERHIAccess::RTV | ERHIAccess::DSVRead | ERHIAccess::DSVWrite))
+					{
+						Texture->TextureDebugData.bHasBeenBoundAsRenderTarget = true;
+					}
+					Texture->MarkResourceAsUsed();
+				}
+			}
+			break;
 			case UBMT_RDG_BUFFER_ACCESS:
-				if (FRDGResourceRef Resource = Parameter.GetAsResource())
+				if (FRDGBufferRef Buffer = Parameter.GetAsBuffer())
 				{
-					Resource->MarkResourceAsUsed();
+					Buffer->MarkResourceAsUsed();
 				}
 			break;
 			case UBMT_RENDER_TARGET_BINDING_SLOTS:
