Workaround D3D12RHI's dynamic SRV tracking bug for FScatterUploadBuffer by recreating SRVs upon buffer locks. This fixes GPUScene lightmap corruption issues.

diff --git a/Engine/Source/Runtime/RenderCore/Private/UnifiedBuffer.cpp b/Engine/Source/Runtime/RenderCore/Private/UnifiedBuffer.cpp
--- a/Engine/Source/Runtime/RenderCore/Private/UnifiedBuffer.cpp
+++ b/Engine/Source/Runtime/RenderCore/Private/UnifiedBuffer.cpp
@@ -473,6 +473,11 @@ void FScatterUploadBuffer::Init( uint32 NumElements, uint32 InNumBytesPerElement
 
 	ScatterData = (uint32*)RHILockStructuredBuffer( ScatterBuffer.Buffer, 0, ScatterBytes, RLM_WriteOnly );
 	UploadData = (uint8*)RHILockStructuredBuffer( UploadBuffer.Buffer, 0, UploadBytes, RLM_WriteOnly );
+
+	// Recreate SRVs as D3D12RHI's dynamic SRV tracking is bugged and existing SRVs get invalidated if the buffer is locked
+	// TODO: fix dynamic SRV tracking in D3D12
+	ScatterBuffer.SRV = RHICreateShaderResourceView(ScatterBuffer.Buffer);
+	UploadBuffer.SRV = RHICreateShaderResourceView(UploadBuffer.Buffer);
 }
 
 template<typename ResourceType>
