Fixed render target pool event viewer to be stable.

diff --git a/Engine/Source/Runtime/RenderCore/Private/RenderTargetPool.cpp b/Engine/Source/Runtime/RenderCore/Private/RenderTargetPool.cpp
--- a/Engine/Source/Runtime/RenderCore/Private/RenderTargetPool.cpp
+++ b/Engine/Source/Runtime/RenderCore/Private/RenderTargetPool.cpp
@@ -822,6 +822,7 @@ FRenderTargetPool::SMemoryStats FRenderTargetPool::ComputeView()
 	{
 		struct FRTPColumn
 		{
+			FString Name;
 			// index into the column, -1 if this is no valid column
 			uint32 PoolEntryId;
 			// for sorting
@@ -838,17 +839,22 @@ FRenderTargetPool::SMemoryStats FRenderTargetPool::ComputeView()
 
 			// constructor
 			FRTPColumn(const FRenderTargetPoolEvent& Event)
-				: PoolEntryId(Event.GetPoolEntryId())
+				: Name(Event.GetDesc().DebugName)
+				, PoolEntryId(Event.GetPoolEntryId())
 				, bVRam((Event.GetDesc().Flags& TexCreate_FastVRAM) != 0)
 			{
 				SizeInBytes = Event.GetSizeInBytes();
 			}
 
 			// sort criteria
-			bool operator <(const FRTPColumn& rhs) const
+			bool operator < (const FRTPColumn& RHS) const
 			{
-				// we want the large ones first
-				return SizeInBytes > rhs.SizeInBytes;
+				if (SizeInBytes == RHS.SizeInBytes)
+				{
+					return Name < RHS.Name;
+				}
+
+				return SizeInBytes > RHS.SizeInBytes;
 			}
 		};
 
@@ -872,7 +878,7 @@ FRenderTargetPool::SMemoryStats FRenderTargetPool::ComputeView()
 			}
 		}
 
-		Colums.Sort();
+		Colums.StableSort();
 
 		{
 			uint32 ColumnX = 0;
