Added assert when calling GetShader() on a shader with more than 1 permutation - catches the case where you requested permutations but forgot to select the right one.

diff --git a/Engine/Source/Runtime/RenderCore/Public/GlobalShader.h b/Engine/Source/Runtime/RenderCore/Public/GlobalShader.h
--- a/Engine/Source/Runtime/RenderCore/Public/GlobalShader.h
+++ b/Engine/Source/Runtime/RenderCore/Public/GlobalShader.h
@@ -156,7 +156,14 @@ public:
 	FShaderPipelineRef GetShaderPipeline(const FShaderPipelineType* PipelineType) const;
 
 	template<typename ShaderType>
-	TShaderRef<ShaderType> GetShader(int32 PermutationId = 0) const
+	TShaderRef<ShaderType> GetShader() const
+	{
+		ensureMsgf(ShaderType::StaticType.GetPermutationCount() == 1, TEXT("Failed to provide PermutationId for shader type %s with %u permutations"), ShaderType::StaticType.GetName(), ShaderType::StaticType.GetPermutationCount());
+		return GetShader<ShaderType>(0);
+	}
+
+	template<typename ShaderType>
+	TShaderRef<ShaderType> GetShader(int32 PermutationId) const
 	{
 		TShaderRef<FShader> Shader = GetShader(&ShaderType::StaticType, PermutationId);
 		checkf(Shader.IsValid(), TEXT("Failed to find shader type %s in Platform %s"), ShaderType::StaticType.GetName(), *LegacyShaderPlatformToShaderFormat(Platform).ToString());
