Missed files part of prior check in.

diff --git a/Engine/Plugins/Compositing/OpenColorIO/Source/OpenColorIO/Public/OpenColorIOShaderType.h b/Engine/Plugins/Compositing/OpenColorIO/Source/OpenColorIO/Public/OpenColorIOShaderType.h
--- a/Engine/Plugins/Compositing/OpenColorIO/Source/OpenColorIO/Public/OpenColorIOShaderType.h
+++ b/Engine/Plugins/Compositing/OpenColorIO/Source/OpenColorIO/Public/OpenColorIOShaderType.h
@@ -56,7 +56,7 @@ public:
 			const FSHAHash& InOCIOShaderMapHash,
 			const FString& InDebugDescription
 			)
-		: FShaderCompiledShaderInitializerType(InType,InPermutationId,CompilerOutput, InOCIOShaderMapHash,nullptr,nullptr)
+		: FShaderCompiledShaderInitializerType(InType,nullptr,InPermutationId,CompilerOutput, InOCIOShaderMapHash,nullptr,nullptr)
 		, DebugDescription(InDebugDescription)
 		{}
 	};
diff --git a/Engine/Plugins/FX/Niagara/Source/NiagaraShader/Public/NiagaraShaderType.h b/Engine/Plugins/FX/Niagara/Source/NiagaraShader/Public/NiagaraShaderType.h
--- a/Engine/Plugins/FX/Niagara/Source/NiagaraShader/Public/NiagaraShaderType.h
+++ b/Engine/Plugins/FX/Niagara/Source/NiagaraShader/Public/NiagaraShaderType.h
@@ -66,7 +66,7 @@ public:
 			const FString& InDebugDescription,
 			const TArray< FNiagaraDataInterfaceGPUParamInfo > &InDIParamInfo
 			)
-		: FShaderCompiledShaderInitializerType(InType,InPermutationId,CompilerOutput, InNiagaraShaderMapHash,nullptr,nullptr)
+		: FShaderCompiledShaderInitializerType(InType,nullptr,InPermutationId,CompilerOutput, InNiagaraShaderMapHash,nullptr,nullptr)
 		, DebugDescription(InDebugDescription)
 		, DIParamInfo(InDIParamInfo)
 		{}
diff --git a/Engine/Source/Runtime/Engine/Private/ShaderCompiler/ShaderCompiler.cpp b/Engine/Source/Runtime/Engine/Private/ShaderCompiler/ShaderCompiler.cpp
--- a/Engine/Source/Runtime/Engine/Private/ShaderCompiler/ShaderCompiler.cpp
+++ b/Engine/Source/Runtime/Engine/Private/ShaderCompiler/ShaderCompiler.cpp
@@ -151,6 +151,7 @@ void FShaderCompileJobCollection::InternalAddJob(FShaderCommonCompileJob* InJob)
 static FShaderCommonCompileJob* CloneJob_Single(const FShaderCompileJob* SrcJob)
 {
 	FShaderCompileJob* Job = new FShaderCompileJob(SrcJob->Hash, SrcJob->Id, SrcJob->Priority, SrcJob->Key);
+	Job->ShaderParameters = SrcJob->ShaderParameters;
 	Job->PendingShaderMap = SrcJob->PendingShaderMap;
 	Job->Input = SrcJob->Input;
 	return Job;
@@ -4800,7 +4801,7 @@ FShader* FGlobalShaderTypeCompiler::FinishCompileShader(const FGlobalShaderType*
 			HashState.GetHash(&GlobalShaderMapHash.Hash[0]);
 		}
 
-		Shader = ShaderType->ConstructCompiled(FGlobalShaderType::CompiledShaderInitializerType(ShaderType, CurrentJob.Key.PermutationId, CurrentJob.Output, GlobalShaderMapHash, ShaderPipelineType, nullptr));
+		Shader = ShaderType->ConstructCompiled(FGlobalShaderType::CompiledShaderInitializerType(ShaderType, nullptr, CurrentJob.Key.PermutationId, CurrentJob.Output, GlobalShaderMapHash, ShaderPipelineType, nullptr));
 		CurrentJob.Output.ParameterMap.VerifyBindingsAreComplete(ShaderType->GetName(), CurrentJob.Output.Target, CurrentJob.Key.VFType);
 	}
 
diff --git a/Engine/Source/Runtime/Engine/Public/MaterialShaderType.h b/Engine/Source/Runtime/Engine/Public/MaterialShaderType.h
--- a/Engine/Source/Runtime/Engine/Public/MaterialShaderType.h
+++ b/Engine/Source/Runtime/Engine/Public/MaterialShaderType.h
@@ -75,7 +75,15 @@ public:
 			const FVertexFactoryType* InVertexFactoryType,
 			const FString& InDebugDescription
 			)
-		: FGlobalShaderType::CompiledShaderInitializerType(InType,InPermutationId,CompilerOutput,InMaterialShaderMapHash,InShaderPipeline,InVertexFactoryType)
+		: FGlobalShaderType::CompiledShaderInitializerType(
+			InType, 
+			nullptr, 
+			InPermutationId,
+			CompilerOutput,
+			InMaterialShaderMapHash,
+			InShaderPipeline,
+			InVertexFactoryType
+			)
 		, UniformExpressionSet(InUniformExpressionSet)
 		, DebugDescription(InDebugDescription)
 		{}
diff --git a/Engine/Source/Runtime/Engine/Public/ShaderCompiler.h b/Engine/Source/Runtime/Engine/Public/ShaderCompiler.h
--- a/Engine/Source/Runtime/Engine/Public/ShaderCompiler.h
+++ b/Engine/Source/Runtime/Engine/Public/ShaderCompiler.h
@@ -208,6 +208,13 @@ public:
 	static const EShaderCompileJobType Type = EShaderCompileJobType::Single;
 
 	FShaderCompileJobKey Key;
+
+	/** 
+	 * Additional parameters that can be supplied to the compile job such 
+	 * that it is available from the compilation begins to when the FShader is created.
+	 */
+	TSharedPtr<const FShaderType::FParameters, ESPMode::ThreadSafe> ShaderParameters;
+
 	/** Input for the shader compile */
 	FShaderCompilerInput Input;
 	FShaderCompilerOutput Output;
diff --git a/Engine/Source/Runtime/RenderCore/Private/Shader.cpp b/Engine/Source/Runtime/RenderCore/Private/Shader.cpp
--- a/Engine/Source/Runtime/RenderCore/Private/Shader.cpp
+++ b/Engine/Source/Runtime/RenderCore/Private/Shader.cpp
@@ -487,24 +487,26 @@ bool FShaderMapPointerTable::LoadFromArchive(FArchive& Ar, const FPlatformTypeLa
 
 FShaderCompiledShaderInitializerType::FShaderCompiledShaderInitializerType(
 	const FShaderType* InType,
+	const FShaderType::FParameters* InParameters,
 	int32 InPermutationId,
 	const FShaderCompilerOutput& CompilerOutput,
 	const FSHAHash& InMaterialShaderMapHash,
 	const FShaderPipelineType* InShaderPipeline,
 	const FVertexFactoryType* InVertexFactoryType
-) :
-	Type(InType),
-	Target(CompilerOutput.Target),
-	Code(CompilerOutput.ShaderCode.GetReadAccess()),
-	ParameterMap(CompilerOutput.ParameterMap),
-	OutputHash(CompilerOutput.OutputHash),
-	MaterialShaderMapHash(InMaterialShaderMapHash),
-	ShaderPipeline(InShaderPipeline),
-	VertexFactoryType(InVertexFactoryType),
-	NumInstructions(CompilerOutput.NumInstructions),
-	NumTextureSamplers(CompilerOutput.NumTextureSamplers),
-	CodeSize(CompilerOutput.ShaderCode.GetShaderCodeSize()),
-	PermutationId(InPermutationId)
+	) 
+	: Type(InType)
+	, Parameters(InParameters)
+	, Target(CompilerOutput.Target)
+	, Code(CompilerOutput.ShaderCode.GetReadAccess())
+	, ParameterMap(CompilerOutput.ParameterMap)
+	, OutputHash(CompilerOutput.OutputHash)
+	, MaterialShaderMapHash(InMaterialShaderMapHash)
+	, ShaderPipeline(InShaderPipeline)
+	, VertexFactoryType(InVertexFactoryType)
+	, NumInstructions(CompilerOutput.NumInstructions)
+	, NumTextureSamplers(CompilerOutput.NumTextureSamplers)
+	, CodeSize(CompilerOutput.ShaderCode.GetShaderCodeSize())
+	, PermutationId(InPermutationId)
 {
 }
 
diff --git a/Engine/Source/Runtime/RenderCore/Public/Shader.h b/Engine/Source/Runtime/RenderCore/Public/Shader.h
--- a/Engine/Source/Runtime/RenderCore/Public/Shader.h
+++ b/Engine/Source/Runtime/RenderCore/Public/Shader.h
@@ -42,6 +42,10 @@
 class FShaderParametersMetadata;
 class FShaderMapPointerTable;
 struct FShaderCompilerOutput;
+struct FShaderCompiledShaderInitializerType;
 
+UE_DEPRECATED(4.26, "FShadereCompiledShaderInitializerType is deprecated. Use FShaderCompiledShaderInitializerType.")
+typedef FShaderCompiledShaderInitializerType FShadereCompiledShaderInitializerType;
+
 /** Define a shader permutation uniquely according to its type, and permutation id.*/
 template<typename MetaShaderType>
@@ -660,34 +664,6 @@
 	{}
 };
 
-struct FShaderCompiledShaderInitializerType
-{
-	const FShaderType* Type;
-	FShaderTarget Target;
-	const TArray<uint8>& Code;
-	const FShaderParameterMap& ParameterMap;
-	const FSHAHash& OutputHash;
-	FSHAHash MaterialShaderMapHash;
-	const FShaderPipelineType* ShaderPipeline;
-	const FVertexFactoryType* VertexFactoryType;
-	uint32 NumInstructions;
-	uint32 NumTextureSamplers;
-	uint32 CodeSize;
-	int32 PermutationId;
-
-	RENDERCORE_API FShaderCompiledShaderInitializerType(
-		const FShaderType* InType,
-		int32 InPermutationId,
-		const FShaderCompilerOutput& CompilerOutput,
-		const FSHAHash& InMaterialShaderMapHash,
-		const FShaderPipelineType* InShaderPipeline,
-		const FVertexFactoryType* InVertexFactoryType
-	);
-};
-
-UE_DEPRECATED(4.26, "FShadereCompiledShaderInitializerType is deprecated. Use FShaderCompiledShaderInitializerType.")
-typedef FShaderCompiledShaderInitializerType FShadereCompiledShaderInitializerType;
-
 namespace Freeze
 {
 	RENDERCORE_API void IntrinsicToString(const TIndexedPtr<FShaderType>& Object, const FTypeLayoutDesc& TypeDesc, const FPlatformTypeLayoutParameters& LayoutParams, FMemoryToStringContext& OutContext);
@@ -1034,6 +1010,15 @@
 		NumShaderTypes,
 	};
 
+	/**
+	 * Derived FShaderTypes should derive from this class to pass params to FShader constructor
+	 */
+	class FParameters
+	{
+	public:
+		virtual ~FParameters(){}
+	};
+
 	typedef class FShader* (*ConstructSerializedType)();
 	typedef FShader* (*ConstructCompiledType)(const FShader::CompiledShaderInitializerType& Initializer);
 	typedef bool (*ShouldCompilePermutationType)(const FShaderPermutationParameters&);
@@ -1296,6 +1281,33 @@
 
 };
 
+struct FShaderCompiledShaderInitializerType
+{
+	const FShaderType* Type;
+	const FShaderType::FParameters* Parameters;
+	FShaderTarget Target;
+	const TArray<uint8>& Code;
+	const FShaderParameterMap& ParameterMap;
+	const FSHAHash& OutputHash;
+	FSHAHash MaterialShaderMapHash;
+	const FShaderPipelineType* ShaderPipeline;
+	const FVertexFactoryType* VertexFactoryType;
+	uint32 NumInstructions;
+	uint32 NumTextureSamplers;
+	uint32 CodeSize;
+	int32 PermutationId;
+
+	RENDERCORE_API FShaderCompiledShaderInitializerType(
+		const FShaderType* InType,
+		const FShaderType::FParameters* InParameters,
+		int32 InPermutationId,
+		const FShaderCompilerOutput& CompilerOutput,
+		const FSHAHash& InMaterialShaderMapHash,
+		const FShaderPipelineType* InShaderPipeline,
+		const FVertexFactoryType* InVertexFactoryType
+	);
+};
+
 #define SHADER_DECLARE_VTABLE(ShaderClass) \
 	static FShader* ConstructSerializedInstance() { return new ShaderClass(); } \
 	static FShader* ConstructCompiledInstance(const typename FShader::CompiledShaderInitializerType& Initializer) \
