clang CIS fix

diff --git a/Engine/Source/Runtime/RenderCore/Public/ShaderParameterMetadataBuilder.h b/Engine/Source/Runtime/RenderCore/Public/ShaderParameterMetadataBuilder.h
--- a/Engine/Source/Runtime/RenderCore/Public/ShaderParameterMetadataBuilder.h
+++ b/Engine/Source/Runtime/RenderCore/Public/ShaderParameterMetadataBuilder.h
@@ -35,7 +35,7 @@ public:
 			TParamTypeInfo::GetStructMetadata()
 			);
 
-		NextMemberOffset += sizeof(TParamTypeInfo::TAlignedType);
+		NextMemberOffset += sizeof(typename TParamTypeInfo::TAlignedType);
 	}
 
 	void AddRDGBufferSRV(
