Changed group dimension checks in ValidateGroupCount to ensures so that dispatches exceeding these limits hopefully don't take down the editor (but still prompt engineers that a technique is doing something bad and needs to be broken up into batches).

diff --git a/Engine/Source/Runtime/RenderCore/Public/RenderGraphUtils.h b/Engine/Source/Runtime/RenderCore/Public/RenderGraphUtils.h
--- a/Engine/Source/Runtime/RenderCore/Public/RenderGraphUtils.h
+++ b/Engine/Source/Runtime/RenderCore/Public/RenderGraphUtils.h
@@ -421,9 +421,9 @@ struct RENDERCORE_API FComputeShaderUtils
 
 	static inline void ValidateGroupCount(const FIntVector& GroupCount)
 	{
-		check(GroupCount.X <= GRHIMaxDispatchThreadGroupsPerDimension.X);
-		check(GroupCount.Y <= GRHIMaxDispatchThreadGroupsPerDimension.Y);
-		check(GroupCount.Z <= GRHIMaxDispatchThreadGroupsPerDimension.Z);
+		ensure(GroupCount.X <= GRHIMaxDispatchThreadGroupsPerDimension.X);
+		ensure(GroupCount.Y <= GRHIMaxDispatchThreadGroupsPerDimension.Y);
+		ensure(GroupCount.Z <= GRHIMaxDispatchThreadGroupsPerDimension.Z);
 	}
 
 	static inline void ValidateIndirectArgsBuffer(uint32 IndirectArgsBufferSize, uint32 IndirectArgOffset)
