Fixed RDG buffer creation with initial data API to sanitize inputs.

diff --git a/Engine/Source/Runtime/RenderCore/Private/RenderGraphUtils.cpp b/Engine/Source/Runtime/RenderCore/Private/RenderGraphUtils.cpp
--- a/Engine/Source/Runtime/RenderCore/Private/RenderGraphUtils.cpp
+++ b/Engine/Source/Runtime/RenderCore/Private/RenderGraphUtils.cpp
@@ -766,24 +766,27 @@ FRDGBufferRef CreateStructuredBuffer(
 	uint64 InitialDataSize,
 	ERDGInitialDataFlags InitialDataFlags)
 {
-	const void* SourcePtr = GetInitialData(GraphBuilder, InitialData, InitialDataSize, InitialDataFlags);
-
 	FRDGBufferRef Buffer = GraphBuilder.CreateBuffer(FRDGBufferDesc::CreateStructuredDesc(BytesPerElement, NumElements), Name);
 
-	FCopyBufferParameters* PassParameters = GraphBuilder.AllocParameters<FCopyBufferParameters>();
-	PassParameters->Buffer = Buffer;
-
-	GraphBuilder.AddPass(
-		RDG_EVENT_NAME("StructuredBufferUpload(%s)", Buffer->Name),
-		PassParameters,
-		ERDGPassFlags::Copy,
-		[Buffer, SourcePtr, InitialDataSize](FRHICommandListImmediate& RHICmdList)
+	if (InitialData && InitialDataSize)
 	{
-		FRHIStructuredBuffer* StructuredBuffer = Buffer->GetRHIStructuredBuffer();
-		void* DestPtr = RHICmdList.LockStructuredBuffer(StructuredBuffer, 0, InitialDataSize, RLM_WriteOnly);
-		FMemory::Memcpy(DestPtr, SourcePtr, InitialDataSize);
-		RHICmdList.UnlockStructuredBuffer(StructuredBuffer);
-	});
+		const void* SourcePtr = GetInitialData(GraphBuilder, InitialData, InitialDataSize, InitialDataFlags);
+
+		FCopyBufferParameters* PassParameters = GraphBuilder.AllocParameters<FCopyBufferParameters>();
+		PassParameters->Buffer = Buffer;
+
+		GraphBuilder.AddPass(
+			RDG_EVENT_NAME("StructuredBufferUpload(%s)", Buffer->Name),
+			PassParameters,
+			ERDGPassFlags::Copy,
+			[Buffer, SourcePtr, InitialDataSize](FRHICommandListImmediate& RHICmdList)
+		{
+			FRHIStructuredBuffer* StructuredBuffer = Buffer->GetRHIStructuredBuffer();
+			void* DestPtr = RHICmdList.LockStructuredBuffer(StructuredBuffer, 0, InitialDataSize, RLM_WriteOnly);
+			FMemory::Memcpy(DestPtr, SourcePtr, InitialDataSize);
+			RHICmdList.UnlockStructuredBuffer(StructuredBuffer);
+		});
+	}
 
 	return Buffer;
 }
@@ -799,24 +802,27 @@ FRDGBufferRef CreateVertexBuffer(
 	checkf(Name, TEXT("Buffer must have a name."));
 	checkf(Desc.UnderlyingType == FRDGBufferDesc::EUnderlyingType::VertexBuffer, TEXT("CreateVertexBuffer called with an FRDGBufferDesc underlying type that is not 'VertexBuffer'. Buffer: %s"), Name);
 
-	const void* SourcePtr = GetInitialData(GraphBuilder, InitialData, InitialDataSize, InitialDataFlags);
-
 	FRDGBufferRef Buffer = GraphBuilder.CreateBuffer(Desc, Name);
 
-	FCopyBufferParameters* PassParameters = GraphBuilder.AllocParameters<FCopyBufferParameters>();
-	PassParameters->Buffer = Buffer;
-
-	GraphBuilder.AddPass(
-		RDG_EVENT_NAME("VertexBufferUpload(%s)", Buffer->Name),
-		PassParameters,
-		ERDGPassFlags::Copy,
-		[Buffer, SourcePtr, InitialDataSize](FRHICommandListImmediate& RHICmdList)
+	if (InitialData && InitialDataSize)
 	{
-		FRHIVertexBuffer* VertexBuffer = Buffer->GetRHIVertexBuffer();
-		void* DestPtr = RHICmdList.LockVertexBuffer(VertexBuffer, 0, InitialDataSize, RLM_WriteOnly);
-		FMemory::Memcpy(DestPtr, SourcePtr, InitialDataSize);
-		RHICmdList.UnlockVertexBuffer(VertexBuffer);
-	});
+		const void* SourcePtr = GetInitialData(GraphBuilder, InitialData, InitialDataSize, InitialDataFlags);
+
+		FCopyBufferParameters* PassParameters = GraphBuilder.AllocParameters<FCopyBufferParameters>();
+		PassParameters->Buffer = Buffer;
+
+		GraphBuilder.AddPass(
+			RDG_EVENT_NAME("VertexBufferUpload(%s)", Buffer->Name),
+			PassParameters,
+			ERDGPassFlags::Copy,
+			[Buffer, SourcePtr, InitialDataSize](FRHICommandListImmediate& RHICmdList)
+		{
+			FRHIVertexBuffer* VertexBuffer = Buffer->GetRHIVertexBuffer();
+			void* DestPtr = RHICmdList.LockVertexBuffer(VertexBuffer, 0, InitialDataSize, RLM_WriteOnly);
+			FMemory::Memcpy(DestPtr, SourcePtr, InitialDataSize);
+			RHICmdList.UnlockVertexBuffer(VertexBuffer);
+		});
+	}
 
 	return Buffer;
 }
