Fix for RDG pass flags validation and ERHIAccess::IndirectArgs.

diff --git a/Engine/Source/Runtime/RenderCore/Private/RenderGraphValidation.cpp b/Engine/Source/Runtime/RenderCore/Private/RenderGraphValidation.cpp
--- a/Engine/Source/Runtime/RenderCore/Private/RenderGraphValidation.cpp
+++ b/Engine/Source/Runtime/RenderCore/Private/RenderGraphValidation.cpp
@@ -51,8 +51,9 @@ void EnumerateSubresources(const FRHITransitionInfo& Transition, uint32 NumMips,
 }
 
 const ERHIAccess AccessMaskCopy    = ERHIAccess::CopySrc | ERHIAccess::CopyDest | ERHIAccess::CPURead;
-const ERHIAccess AccessMaskCompute = ERHIAccess::SRVCompute | ERHIAccess::UAVCompute | ERHIAccess::IndirectArgs;
-const ERHIAccess AccessMaskRaster  = ERHIAccess::ResolveSrc | ERHIAccess::ResolveDst | ERHIAccess::DSVRead | ERHIAccess::DSVWrite | ERHIAccess::RTV | ERHIAccess::SRVGraphics | ERHIAccess::UAVGraphics | ERHIAccess::Present | ERHIAccess::IndirectArgs | ERHIAccess::VertexOrIndexBuffer;
+const ERHIAccess AccessMaskCompute = ERHIAccess::SRVCompute | ERHIAccess::UAVCompute;
+const ERHIAccess AccessMaskRaster  = ERHIAccess::ResolveSrc | ERHIAccess::ResolveDst | ERHIAccess::DSVRead | ERHIAccess::DSVWrite | ERHIAccess::RTV | ERHIAccess::SRVGraphics | ERHIAccess::UAVGraphics | ERHIAccess::Present | ERHIAccess::VertexOrIndexBuffer;
+const ERHIAccess AccessMaskComputeOrRaster = ERHIAccess::IndirectArgs;
 
 /** Validates that we are only executing a single render graph instance in the callstack. Used to catch if a
  *  user creates a second FRDGBuilder instance inside of a pass that is executing.
@@ -298,6 +299,7 @@ void FRDGUserValidation::ValidateAddPass(const FRDGPass* Pass, bool bSkipPassAcc
 				checkf(bIsCopy       || !EnumHasAnyFlags(Access, AccessMaskCopy),    TEXT("Pass '%s' uses texture '%s' with access '%s' containing states which require the 'ERDGPass::Copy' flag."), Pass->GetName(), Texture->Name, *GetRHIAccessName(Access));
 				checkf(bIsAnyCompute || !EnumHasAnyFlags(Access, AccessMaskCompute), TEXT("Pass '%s' uses texture '%s' with access '%s' containing states which require the 'ERDGPass::Compute' or 'ERDGPassFlags::AsyncCompute' flag."), Pass->GetName(), Texture->Name, *GetRHIAccessName(Access));
 				checkf(bIsRaster     || !EnumHasAnyFlags(Access, AccessMaskRaster),  TEXT("Pass '%s' uses texture '%s' with access '%s' containing states which require the 'ERDGPass::Raster' flag."), Pass->GetName(), Texture->Name, *GetRHIAccessName(Access));
+				checkf(bIsAnyCompute || bIsRaster || !EnumHasAnyFlags(Access, AccessMaskComputeOrRaster), TEXT("Pass '%s' uses texture '%s' with access '%s' containing states which require the 'ERDGPassFlags::Compute' or 'ERDGPassFlags::AsyncCompute' or 'ERDGPass::Raster' flag."), Pass->GetName(), Texture->Name, *GetRHIAccessName(Access));
 
 				if (bIsWriteAccess)
 				{
@@ -318,6 +320,7 @@ void FRDGUserValidation::ValidateAddPass(const FRDGPass* Pass, bool bSkipPassAcc
 				checkf(bIsCopy       || !EnumHasAnyFlags(Access, AccessMaskCopy),    TEXT("Pass '%s' uses buffer '%s' with access '%s' containing states which require the 'ERDGPass::Copy' flag."), Pass->GetName(), Buffer->Name, *GetRHIAccessName(Access));
 				checkf(bIsAnyCompute || !EnumHasAnyFlags(Access, AccessMaskCompute), TEXT("Pass '%s' uses buffer '%s' with access '%s' containing states which require the 'ERDGPass::Compute' or 'ERDGPassFlags::AsyncCompute' flag."), Pass->GetName(), Buffer->Name, *GetRHIAccessName(Access));
 				checkf(bIsRaster     || !EnumHasAnyFlags(Access, AccessMaskRaster),  TEXT("Pass '%s' uses buffer '%s' with access '%s' containing states which require the 'ERDGPass::Raster' flag."), Pass->GetName(), Buffer->Name, *GetRHIAccessName(Access));
+				checkf(bIsAnyCompute || bIsRaster || !EnumHasAnyFlags(Access, AccessMaskComputeOrRaster), TEXT("Pass '%s' uses buffer '%s' with access '%s' containing states which require the 'ERDGPassFlags::Compute' or 'ERDGPassFlags::AsyncCompute' or 'ERDGPass::Raster' flag."), Pass->GetName(), Buffer->Name, *GetRHIAccessName(Access));
 
 				if (IsWritableAccess(BufferAccess.GetAccess()))
 				{
