Removing OnChip shader support.

diff --git a/Engine/Source/Runtime/RenderCore/Public/ShaderCompilerCore.h b/Engine/Source/Runtime/RenderCore/Public/ShaderCompilerCore.h
--- a/Engine/Source/Runtime/RenderCore/Public/ShaderCompilerCore.h
+++ b/Engine/Source/Runtime/RenderCore/Public/ShaderCompilerCore.h
@@ -41,8 +41,6 @@ enum ECompilerFlags
 	CFLAG_StandardOptimization,
 	// Always optimize, even when CFLAG_Debug is set. Required for some complex shaders and features.
 	CFLAG_ForceOptimization,
-	// Shader should use on chip memory instead of main memory ring buffer memory.
-	CFLAG_OnChip,
 	// Shader should insert debug/name info at the risk of generating non-deterministic libraries
 	CFLAG_KeepDebugInfo,
 	CFLAG_NoFastMath,
