Moved RDG uniform buffer creation to CollectPassResources.

#rnx

diff --git a/Engine/Source/Runtime/RenderCore/Private/RenderGraphBuilder.cpp b/Engine/Source/Runtime/RenderCore/Private/RenderGraphBuilder.cpp
--- a/Engine/Source/Runtime/RenderCore/Private/RenderGraphBuilder.cpp
+++ b/Engine/Source/Runtime/RenderCore/Private/RenderGraphBuilder.cpp
@@ -1490,12 +1490,6 @@ void FRDGBuilder::ExecutePassPrologue(FRHIComputeCommandList& RHICmdListPass, FR
 		Pass->PrologueBarriersToEnd->Submit(RHICmdListPass);
 	}
 
-	// Uniform buffers are initialized during first-use execution, since the access checks will allow calling GetRHI on RDG resources.
-	Pass->GetParameters().EnumerateUniformBuffers([&](FRDGUniformBufferRef UniformBuffer)
-	{
-		BeginResourceRHI(UniformBuffer);
-	});
-
 	if (Pass->GetPipeline() == ERHIPipeline::AsyncCompute)
 	{
 		RHICmdListPass.SetAsyncComputeBudget(Pass->AsyncComputeBudget);
@@ -1633,6 +1627,16 @@ void FRDGBuilder::CollectPassResources(FRDGPassHandle PassHandle)
 		{
 			BeginResourceRHI(PassHandle, Resource);
 		});
+
+		// Uniform buffer creation requires access to RHI resources through the validator.
+		IF_RDG_ENABLE_DEBUG(FRDGUserValidation::SetAllowRHIAccess(PassToBegin, true));
+		{
+			PassParameters.EnumerateUniformBuffers([&](FRDGUniformBufferRef UniformBuffer)
+			{
+				BeginResourceRHI(UniformBuffer);
+			});
+		}
+		IF_RDG_ENABLE_DEBUG(FRDGUserValidation::SetAllowRHIAccess(PassToBegin, false));
 	}
 
 	for (FRDGPass* PassToEnd : Pass->ResourcesToEnd)
