Fix shader binding type validation for external textures (UE-102423).

diff --git a/Engine/Source/Runtime/RenderCore/Private/ShaderParameterMetadata.cpp b/Engine/Source/Runtime/RenderCore/Private/ShaderParameterMetadata.cpp
--- a/Engine/Source/Runtime/RenderCore/Private/ShaderParameterMetadata.cpp
+++ b/Engine/Source/Runtime/RenderCore/Private/ShaderParameterMetadata.cpp
@@ -178,6 +178,10 @@ EShaderCodeResourceBindingType ParseShaderResourceBindingType(const TCHAR* Shade
 	{
 		BindingType = bIsRWResource ? EShaderCodeResourceBindingType::RWTexture2D : EShaderCodeResourceBindingType::Texture2D;
 	}
+	else if (ShaderTypeLength == 15 && FCString::Strncmp(ComparedShaderType, TEXT("TextureExternal"), ShaderTypeLength) == 0)
+	{
+		BindingType = bIsRWResource ? EShaderCodeResourceBindingType::RWTexture2D : EShaderCodeResourceBindingType::Texture2D;
+	}
 	else if (ShaderTypeLength == 14 && FCString::Strncmp(ComparedShaderType, TEXT("Texture2DArray"), ShaderTypeLength) == 0)
 	{
 		BindingType = bIsRWResource ? EShaderCodeResourceBindingType::RWTexture2DArray : EShaderCodeResourceBindingType::Texture2DArray;
