Workaround to remove transitions within a render pass for DrawTileMesh.

#rnx

diff --git a/Engine/Source/Runtime/RenderCore/Private/RenderGraphBuilder.cpp b/Engine/Source/Runtime/RenderCore/Private/RenderGraphBuilder.cpp
--- a/Engine/Source/Runtime/RenderCore/Private/RenderGraphBuilder.cpp
+++ b/Engine/Source/Runtime/RenderCore/Private/RenderGraphBuilder.cpp
@@ -313,6 +313,8 @@ inline EResourceTransitionFlags GetTextureViewTransitionFlags(FRDGViewRef Resour
 	return EResourceTransitionFlags::None;
 }
 
+extern void TestBlackboard();
+
 void FRDGBuilder::TickPoolElements()
 {
 	GRenderGraphResourcePool.TickPoolElements();
@@ -349,6 +351,8 @@ void FRDGBuilder::TickPoolElements()
 	GRDGStatTransitionBatchCount = 0;
 	GRDGStatMemoryWatermark = 0;
 #endif
+
+	TestBlackboard();
 }
 
 ERDGPassFlags FRDGBuilder::OverridePassFlags(const TCHAR* PassName, ERDGPassFlags PassFlags, bool bAsyncComputeSupported)
@@ -380,12 +384,13 @@ ERDGPassFlags FRDGBuilder::OverridePassFlags(const TCHAR* PassName, ERDGPassFlag
 
 const char* const FRDGBuilder::kDefaultUnaccountedCSVStat = "RDG_Pass";
 
-FRDGBuilder::FRDGBuilder(FRHICommandListImmediate& InRHICmdList, FRDGEventName InName, const char* UnaccountedCSVStat)
+FRDGBuilder::FRDGBuilder(FRHICommandListImmediate& InRHICmdList, FRDGEventName InName, ERDGBuilderFlags InFlags)
 	: RHICmdList(InRHICmdList)
 	, RHICmdListAsyncCompute(FRHICommandListExecutor::GetImmediateAsyncComputeCommandList())
 	, BuilderName(InName)
+	, BuilderFlags(InFlags)
 #if RDG_CPU_SCOPES
-	, CPUScopeStacks(RHICmdList, UnaccountedCSVStat)
+	, CPUScopeStacks(RHICmdList, kDefaultUnaccountedCSVStat)
 #endif
 #if RDG_GPU_SCOPES
 	, GPUScopeStacks(RHICmdList, RHICmdListAsyncCompute)
@@ -1690,7 +1695,7 @@ void FRDGBuilder::CollectPassBarriers(FRDGPassHandle PassHandle, FRDGPassHandle&
 
 void FRDGBuilder::AddEpilogueTransition(FRDGTextureRef Texture, FRDGPassHandle LastUntrackedPassHandle)
 {
-	if (!Texture->bLastOwner || Texture->bCulled)
+	if (!Texture->bLastOwner || Texture->bCulled || EnumHasAnyFlags(BuilderFlags, ERDGBuilderFlags::SkipBarriers))
 	{
 		return;
 	}
@@ -1912,6 +1917,11 @@ void FRDGBuilder::AddTransitionInternal(
 	FRDGPassHandle LastUntrackedPassHandle,
 	const FRHITransitionInfo& TransitionInfo)
 {
+	if (EnumHasAnyFlags(BuilderFlags, ERDGBuilderFlags::SkipBarriers))
+	{
+		return;
+	}
+
 	check(StateAfter.FirstPass.IsValid());
 
 	if (StateBefore.LastPass.IsNull())
diff --git a/Engine/Source/Runtime/RenderCore/Public/RenderGraphBuilder.h b/Engine/Source/Runtime/RenderCore/Public/RenderGraphBuilder.h
--- a/Engine/Source/Runtime/RenderCore/Public/RenderGraphBuilder.h
+++ b/Engine/Source/Runtime/RenderCore/Public/RenderGraphBuilder.h
@@ -12,6 +12,14 @@
 
 class FRDGLogFile;
 
+enum ERDGBuilderFlags
+{
+	None = 0,
+
+	/** The builder will not perform any resource barriers. This is only allowed when all resources are registered. */
+	SkipBarriers = 1 << 0,
+};
+
 /** Use the render graph builder to build up a graph of passes and then call Execute() to process them. Resource barriers
  *  and lifetimes are derived from _RDG_ parameters in the pass parameter struct provided to each AddPass call. The resulting
  *  graph is compiled, culled, and executed in Execute(). The builder should be created on the stack and executed prior to
@@ -20,7 +28,7 @@ class FRDGLogFile;
 class RENDERCORE_API FRDGBuilder
 {
 public:
-	FRDGBuilder(FRHICommandListImmediate& InRHICmdList, FRDGEventName InName = {}, const char* UnaccountedCSVStat = kDefaultUnaccountedCSVStat);
+	FRDGBuilder(FRHICommandListImmediate& InRHICmdList, FRDGEventName InName = {}, ERDGBuilderFlags Flags = ERDGBuilderFlags::None);
 	FRDGBuilder(const FRDGBuilder&) = delete;
 
 	/** Finds an RDG texture associated with the external texture, or returns null if none is found. */
@@ -221,6 +229,7 @@ private:
 	FRDGAllocator Allocator;
 
 	const FRDGEventName BuilderName;
+	const ERDGBuilderFlags BuilderFlags;
 
 	ERDGPassFlags OverridePassFlags(const TCHAR* PassName, ERDGPassFlags Flags, bool bAsyncComputeSupported);
 
diff --git a/Engine/Source/Runtime/RenderCore/Public/RenderGraphBuilder.inl b/Engine/Source/Runtime/RenderCore/Public/RenderGraphBuilder.inl
--- a/Engine/Source/Runtime/RenderCore/Public/RenderGraphBuilder.inl
+++ b/Engine/Source/Runtime/RenderCore/Public/RenderGraphBuilder.inl
@@ -43,6 +43,8 @@ inline FRDGTextureRef FRDGBuilder::CreateTexture(
 		// D3D11 doesn't allow creating a UAV on MSAA texture.
 		const bool bIsUAVForMSAATexture = bIsMSAA && bCanHaveUAV;
 		checkf(!bIsUAVForMSAATexture, TEXT("TexCreate_UAV is not allowed on MSAA texture %s."), Name);
+
+		checkf(!EnumHasAnyFlags(BuilderFlags, ERDGBuilderFlags::SkipBarriers), TEXT("Cannot create texture '%s' because RDG builder has 'SkipBarriers' flag."), Name);
 	}
 #endif
 
@@ -69,6 +71,8 @@ inline FRDGBufferRef FRDGBuilder::CreateBuffer(
 		{
 			checkf(Desc.BytesPerElement == 4, TEXT("Creating buffer '%s' as a structured buffer that is also byte addressable, BytesPerElement must be 4! Instead it is %d"), Name, Desc.BytesPerElement);
 		}
+
+		checkf(!EnumHasAnyFlags(BuilderFlags, ERDGBuilderFlags::SkipBarriers), TEXT("Cannot create buffer '%s' because RDG builder has 'SkipBarriers' flag."), Name);
 	}
 #endif
 
diff --git a/Engine/Source/Runtime/Renderer/Private/Renderer.cpp b/Engine/Source/Runtime/Renderer/Private/Renderer.cpp
--- a/Engine/Source/Runtime/Renderer/Private/Renderer.cpp
+++ b/Engine/Source/Runtime/Renderer/Private/Renderer.cpp
@@ -199,7 +199,7 @@ void FRendererModule::DrawTileMesh(FRHICommandListImmediate& RHICmdList, FMeshPa
 
 		GSystemTextures.InitializeTextures(RHICmdList, FeatureLevel);
 		FMemMark Mark(FMemStack::Get());
-		FRDGBuilder GraphBuilder(RHICmdList);
+		FRDGBuilder GraphBuilder(RHICmdList, {}, ERDGBuilderFlags::SkipBarriers);
 
 		const ERDGPassFlags PassFlags = ERDGPassFlags::Raster | ERDGPassFlags::SkipRenderPass | ERDGPassFlags::NeverCull;
 
