Fixed build break.

diff --git a/Engine/Source/Runtime/RenderCore/Public/RenderGraphBuilder.h b/Engine/Source/Runtime/RenderCore/Public/RenderGraphBuilder.h
--- a/Engine/Source/Runtime/RenderCore/Public/RenderGraphBuilder.h
+++ b/Engine/Source/Runtime/RenderCore/Public/RenderGraphBuilder.h
@@ -12,7 +12,7 @@
 
 class FRDGLogFile;
 
-enum ERDGBuilderFlags
+enum class ERDGBuilderFlags
 {
 	None = 0,
 
