Fixed build break.

#rnx

diff --git a/Engine/Source/Runtime/RenderCore/Private/RenderGraphBuilder.cpp b/Engine/Source/Runtime/RenderCore/Private/RenderGraphBuilder.cpp
--- a/Engine/Source/Runtime/RenderCore/Private/RenderGraphBuilder.cpp
+++ b/Engine/Source/Runtime/RenderCore/Private/RenderGraphBuilder.cpp
@@ -313,8 +313,6 @@ inline EResourceTransitionFlags GetTextureViewTransitionFlags(FRDGViewRef Resour
 	return EResourceTransitionFlags::None;
 }
 
-extern void TestBlackboard();
-
 void FRDGBuilder::TickPoolElements()
 {
 	GRenderGraphResourcePool.TickPoolElements();
@@ -351,8 +349,6 @@ void FRDGBuilder::TickPoolElements()
 	GRDGStatTransitionBatchCount = 0;
 	GRDGStatMemoryWatermark = 0;
 #endif
-
-	TestBlackboard();
 }
 
 ERDGPassFlags FRDGBuilder::OverridePassFlags(const TCHAR* PassName, ERDGPassFlags PassFlags, bool bAsyncComputeSupported)
