Fix transition errors for GPUScene Instance buffer (not transitioned to SRV if size changed)
- Also changed ResizeResourceIfNeeded to return true also if it was previously zero.

diff --git a/Engine/Source/Runtime/RenderCore/Private/UnifiedBuffer.cpp b/Engine/Source/Runtime/RenderCore/Private/UnifiedBuffer.cpp
--- a/Engine/Source/Runtime/RenderCore/Private/UnifiedBuffer.cpp
+++ b/Engine/Source/Runtime/RenderCore/Private/UnifiedBuffer.cpp
@@ -332,6 +332,7 @@ RENDERCORE_API bool ResizeResourceIfNeeded<FTextureRWBuffer2D>(FRHICommandList&
 	if (Texture.NumBytes == 0)
 	{
 		Texture.Initialize(BytesPerElement, Float4sPerLine, NumLines, PF_A32B32G32R32F, TexCreate_RenderTargetable | TexCreate_UAV);
+		return true;
 	}
 	else if ((NumLines * Float4sPerLine * BytesPerElement) != Texture.NumBytes)
 	{
@@ -353,6 +354,7 @@ RENDERCORE_API bool ResizeResourceIfNeeded<FRWBufferStructured>(FRHICommandList&
 	if (Buffer.NumBytes == 0)
 	{
 		Buffer.Initialize(16, NumBytes / 16, 0, DebugName);
+		return true;
 	}
 	else if (NumBytes != Buffer.NumBytes)
 	{
@@ -379,6 +381,7 @@ RENDERCORE_API bool ResizeResourceIfNeeded<FRWByteAddressBuffer>(FRHICommandList
 	if (Buffer.NumBytes == 0)
 	{
 		Buffer.Initialize(NumBytes, 0, DebugName);
+		return true;
 	}
 	else if (NumBytes != Buffer.NumBytes)
 	{
@@ -409,6 +412,7 @@ RENDERCORE_API bool ResizeResourceSOAIfNeeded<FRWBufferStructured>(FRHICommandLi
 	if (Buffer.NumBytes == 0)
 	{
 		Buffer.Initialize(16, NumBytes / 16, 0, DebugName);
+		return true;
 	}
 	else if (NumBytes != Buffer.NumBytes)
 	{
diff --git a/Engine/Source/Runtime/Renderer/Private/GPUScene.cpp b/Engine/Source/Runtime/Renderer/Private/GPUScene.cpp
--- a/Engine/Source/Runtime/Renderer/Private/GPUScene.cpp
+++ b/Engine/Source/Runtime/Renderer/Private/GPUScene.cpp
@@ -676,10 +676,10 @@ void UpdateGPUSceneInternal(FRHICommandListImmediate& RHICmdList, FScene& Scene)
 				Scene.GPUScene.InstanceUploadBuffer.ResourceUploadTo(RHICmdList, Scene.GPUScene.InstanceDataBuffer, false);
 				RHICmdList.Transition(FRHITransitionInfo(Scene.GPUScene.InstanceDataBuffer.UAV, ERHIAccess::UAVCompute, ERHIAccess::SRVMask));
 			}
-			else
+			else if (bResizedInstanceData)
 			{
 				RHICmdList.Transition(FRHITransitionInfo(Scene.GPUScene.InstanceDataBuffer.UAV, ERHIAccess::Unknown, ERHIAccess::SRVMask));
-			}
+			 }
 
 			if (NumLightmapDataUploads > 0)
 			{
