Added Multi-pipeline Transition Support to platforms and RDG. Fixed uniform buffer hack in composition lighting.

diff --git a/Engine/Source/Runtime/D3D12RHI/Private/D3D12CommandContext.cpp b/Engine/Source/Runtime/D3D12RHI/Private/D3D12CommandContext.cpp
--- a/Engine/Source/Runtime/D3D12RHI/Private/D3D12CommandContext.cpp
+++ b/Engine/Source/Runtime/D3D12RHI/Private/D3D12CommandContext.cpp
@@ -896,8 +896,6 @@ IRHICommandContextContainer* FD3D12DynamicRHI::RHIGetCommandContextContainer(int
 
 void FD3D12DynamicRHI::RHICreateTransition(FRHITransition* Transition, ERHIPipeline SrcPipelines, ERHIPipeline DstPipelines, ERHICreateTransitionFlags CreateFlags, TArrayView<const FRHITransitionInfo> Infos)
 {
-	checkf(FMath::IsPowerOfTwo(uint32(SrcPipelines)) && FMath::IsPowerOfTwo(uint32(DstPipelines)), TEXT("Support for multi-pipe resources is not yet implemented."));
-
 	// Construct the data in-place on the transition instance
 	FD3D12TransitionData* Data = new (Transition->GetPrivateData<FD3D12TransitionData>()) FD3D12TransitionData;
 
diff --git a/Engine/Source/Runtime/D3D12RHI/Private/D3D12Commands.cpp b/Engine/Source/Runtime/D3D12RHI/Private/D3D12Commands.cpp
--- a/Engine/Source/Runtime/D3D12RHI/Private/D3D12Commands.cpp
+++ b/Engine/Source/Runtime/D3D12RHI/Private/D3D12Commands.cpp
@@ -296,12 +296,59 @@ void ProcessResource(FD3D12CommandContext& Context, const FRHITransitionInfo& In
 	}
 }
 
+static D3D12_RESOURCE_STATES GetEndTransitionState(const FRHITransitionInfo& Info, FD3D12Resource* Resource, ERHIPipeline Pipelines, D3D12_RESOURCE_STATES SkipFastClearEliminateState)
+{
+	D3D12_RESOURCE_STATES State = D3D12_RESOURCE_STATE_COMMON;
+
+	if (EnumHasAnyFlags(Info.Flags, EResourceTransitionFlags::MaintainCompression) && IsReadOnlyAccess(Info.AccessAfter))
+	{
+		State |= SkipFastClearEliminateState;
+	}
+
+	const bool bIsAsyncComputeContext = EnumHasAnyFlags(Pipelines, ERHIPipeline::AsyncCompute);
+
+	if (Info.AccessAfter == ERHIAccess::ResolveSrc)
+	{
+		State |= D3D12_RESOURCE_STATE_RESOLVE_SOURCE;
+	}
+	else if (Info.AccessAfter == ERHIAccess::ResolveDst)
+	{
+		State |= D3D12_RESOURCE_STATE_RESOLVE_DEST;
+	}
+	else if (EnumHasAnyFlags(Info.AccessAfter, ERHIAccess::WritableMask))
+	{
+		if (EnumHasAnyFlags(Info.AccessAfter, ERHIAccess::UAVMask) || bIsAsyncComputeContext)
+		{
+			State |= D3D12_RESOURCE_STATE_UNORDERED_ACCESS;
+		}
+		else
+		{
+			State |= Resource->GetWritableState();
+		}
+	}
+	else if (EnumHasAnyFlags(Info.AccessAfter, ERHIAccess::ReadableMask))
+	{
+		if (bIsAsyncComputeContext)
+		{
+			State |= D3D12_RESOURCE_STATE_NON_PIXEL_SHADER_RESOURCE;
+		}
+		else
+		{
+			State |= Resource->GetReadableState();
+		}
+	}
+
+	return State;
+}
+
 void FD3D12CommandContext::RHIBeginTransitionsWithoutFencing(TArrayView<const FRHITransition*> Transitions)
 {
 	static IConsoleVariable* CVarShowTransitions = IConsoleManager::Get().FindConsoleVariable(TEXT("r.ProfileGPU.ShowTransitions"));
 	const bool bShowTransitionEvents = CVarShowTransitions->GetInt() != 0;
 	SCOPED_RHI_CONDITIONAL_DRAW_EVENTF(*this, RHIBeginTransitions, bShowTransitionEvents, TEXT("RHIBeginTransitions"));
 
+	bool bUAVBarrier = false;
+
 	for (const FRHITransition* Transition : Transitions)
 	{
 		const FD3D12TransitionData* Data = Transition->GetPrivateData<FD3D12TransitionData>();
@@ -312,8 +359,15 @@ void FD3D12CommandContext::RHIBeginTransitionsWithoutFencing(TArrayView<const FR
 			continue;
 		}
 
+		const bool bDstAllPipelines = EnumHasAllFlags(Data->DstPipelines, ERHIPipeline::All);
+
 		for (const FRHITransitionInfo& Info : Data->Infos)
 		{
+			if (bDstAllPipelines)
+			{
+				bUAVBarrier |= EnumHasAnyFlags(Info.AccessAfter, ERHIAccess::UAVMask);
+			}
+
 			if (!Info.Resource)
 			{
 				continue;
@@ -328,13 +382,21 @@ void FD3D12CommandContext::RHIBeginTransitionsWithoutFencing(TArrayView<const FR
 
 				D3D12_RESOURCE_STATES State = D3D12_RESOURCE_STATE_COMMON;
 
-				if (EnumHasAnyFlags(Info.AccessAfter, ERHIAccess::WritableMask))
+				// When forking from one pipeline to all, we perform the end transitions immediately.
+				if (bDstAllPipelines)
 				{
-					State |= D3D12_RESOURCE_STATE_UNORDERED_ACCESS;
+					State = GetEndTransitionState(Info, Resource, Data->DstPipelines, SkipFastClearEliminateState);
 				}
-				else if (EnumHasAnyFlags(Info.AccessAfter, ERHIAccess::ReadableMask))
+				else
 				{
-					State |= D3D12_RESOURCE_STATE_NON_PIXEL_SHADER_RESOURCE;
+					if (EnumHasAnyFlags(Info.AccessAfter, ERHIAccess::WritableMask))
+					{
+						State |= D3D12_RESOURCE_STATE_UNORDERED_ACCESS;
+					}
+					else if (EnumHasAnyFlags(Info.AccessAfter, ERHIAccess::ReadableMask))
+					{
+						State |= D3D12_RESOURCE_STATE_NON_PIXEL_SHADER_RESOURCE;
+					}
 				}
 
 				if (State == D3D12_RESOURCE_STATE_COMMON)
@@ -356,6 +418,11 @@ void FD3D12CommandContext::RHIBeginTransitionsWithoutFencing(TArrayView<const FR
 			});
 		}
 	}
+
+	if (bUAVBarrier)
+	{
+		StateCache.FlushComputeShaderCache(true);
+	}
 }
 
 void FD3D12CommandContext::RHIBeginTransitions(TArrayView<const FRHITransition*> Transitions)
@@ -381,15 +448,14 @@ void FD3D12CommandContext::RHIEndTransitions(TArrayView<const FRHITransition*> T
 		const FD3D12TransitionData* Data = Transition->GetPrivateData<FD3D12TransitionData>();
 
 		const bool bSamePipeline = !Data->bCrossPipeline;
+		const bool bDstAllPipelines = EnumHasAllFlags(Data->DstPipelines, ERHIPipeline::All);
 
 		for (const FRHITransitionInfo& Info : Data->Infos)
 		{
-			const bool bUAVAccess = EnumHasAnyFlags(Info.AccessAfter, ERHIAccess::UAVMask);
-
 			// Sometimes we could still have barriers with resources, invalid but can still happen
-			if (bSamePipeline)
+			if (bSamePipeline && !bDstAllPipelines)
 			{
-				bUAVBarrier |= bUAVAccess;
+				bUAVBarrier |= EnumHasAnyFlags(Info.AccessAfter, ERHIAccess::UAVMask);
 			}
 
 			if (!Info.Resource)
@@ -405,40 +471,11 @@ void FD3D12CommandContext::RHIEndTransitions(TArrayView<const FRHITransition*> T
 				}
 
 				D3D12_RESOURCE_STATES State = D3D12_RESOURCE_STATE_COMMON;
-				if (EnumHasAnyFlags(Info.Flags, EResourceTransitionFlags::MaintainCompression) && IsReadOnlyAccess(Info.AccessAfter))
-				{
-					State |= SkipFastClearEliminateState;
-				}
 
-				if (Info.AccessAfter == ERHIAccess::ResolveSrc)
+				// When forking from one pipeline to all, we perform the end transitions in begin.
+				if (!bDstAllPipelines)
 				{
-					State |= D3D12_RESOURCE_STATE_RESOLVE_SOURCE;
-				}
-				else if (Info.AccessAfter == ERHIAccess::ResolveDst)
-				{
-					State |= D3D12_RESOURCE_STATE_RESOLVE_DEST;
-				}
-				else if (EnumHasAnyFlags(Info.AccessAfter, ERHIAccess::WritableMask))
-				{
-					if (bUAVAccess || bIsAsyncComputeContext)
-					{
-						State |= D3D12_RESOURCE_STATE_UNORDERED_ACCESS;
-					}
-					else
-					{
-						State |= Resource->GetWritableState();
-					}
-				}
-				else if (EnumHasAnyFlags(Info.AccessAfter, ERHIAccess::ReadableMask))
-				{
-					if (bIsAsyncComputeContext)
-					{
-						State |= D3D12_RESOURCE_STATE_NON_PIXEL_SHADER_RESOURCE;
-					}
-					else
-					{
-						State |= Resource->GetReadableState();
-					}
+					State = GetEndTransitionState(Info, Resource, Data->DstPipelines, SkipFastClearEliminateState);
 				}
 
 				if (State == D3D12_RESOURCE_STATE_COMMON)
diff --git a/Engine/Source/Runtime/D3D12RHI/Private/D3D12RHI.cpp b/Engine/Source/Runtime/D3D12RHI/Private/D3D12RHI.cpp
--- a/Engine/Source/Runtime/D3D12RHI/Private/D3D12RHI.cpp
+++ b/Engine/Source/Runtime/D3D12RHI/Private/D3D12RHI.cpp
@@ -80,6 +80,7 @@ FD3D12DynamicRHI::FD3D12DynamicRHI(const TArray<TSharedPtr<FD3D12Adapter>>& Chos
 #endif // PLATFORM_WINDOWS
 
 	GRHISupportsMultithreading = true;
+	GRHIMultiPipelineAccessMask = ERHIAccess::ReadOnlyMask;
 
 	GPoolSizeVRAMPercentage = 0;
 	GTexturePoolSize = 0;
diff --git a/Engine/Source/Runtime/RenderCore/Private/RenderGraphBuilder.cpp b/Engine/Source/Runtime/RenderCore/Private/RenderGraphBuilder.cpp
--- a/Engine/Source/Runtime/RenderCore/Private/RenderGraphBuilder.cpp
+++ b/Engine/Source/Runtime/RenderCore/Private/RenderGraphBuilder.cpp
@@ -331,6 +331,7 @@ void FRDGBuilder::TickPoolElements()
 
 #if STATS
 	SET_DWORD_STAT(STAT_RDG_PassCount, GRDGStatPassCount);
+	SET_DWORD_STAT(STAT_RDG_PassWithParameterCount, GRDGStatPassWithParameterCount);
 	SET_DWORD_STAT(STAT_RDG_PassCullCount, GRDGStatPassCullCount);
 	SET_DWORD_STAT(STAT_RDG_RenderPassMergeCount, GRDGStatRenderPassMergeCount);
 	SET_DWORD_STAT(STAT_RDG_PassDependencyCount, GRDGStatPassDependencyCount);
@@ -340,6 +341,7 @@ void FRDGBuilder::TickPoolElements()
 	SET_DWORD_STAT(STAT_RDG_TransitionBatchCount, GRDGStatTransitionBatchCount);
 	SET_MEMORY_STAT(STAT_RDG_MemoryWatermark, int64(GRDGStatMemoryWatermark));
 	GRDGStatPassCount = 0;
+	GRDGStatPassWithParameterCount = 0;
 	GRDGStatPassCullCount = 0;
 	GRDGStatRenderPassMergeCount = 0;
 	GRDGStatPassDependencyCount = 0;
@@ -761,6 +763,11 @@ void FRDGBuilder::Compile()
 				continue;
 			}
 
+			const bool bAsyncComputePass = PassesOnAsyncCompute[PassHandle];
+
+			FRDGPass* Pass = Passes[PassHandle];
+			const ERHIPipeline PassPipeline = Pass->GetPipeline();
+
 			const auto MergeSubresourceStates = [&](ERDGParentResourceType ResourceType, FRDGSubresourceState*& PassMergeState, FRDGSubresourceState*& ResourceMergeState, const FRDGSubresourceState& PassState)
 			{
 				if (PassState.Access == ERHIAccess::Unknown)
@@ -771,29 +778,39 @@ void FRDGBuilder::Compile()
 				if (!ResourceMergeState || !FRDGSubresourceState::IsMergeAllowed(ResourceType, *ResourceMergeState, PassState))
 				{
 					// Cross-pipeline, non-mergable state changes require a new pass dependency for fencing purposes.
-					if (ResourceMergeState && ResourceMergeState->Pipeline != PassState.Pipeline)
+					if (ResourceMergeState)
 					{
-						AddPassDependency(ResourceMergeState->LastPass, PassHandle);
+						for (ERHIPipeline Pipeline : GetRHIPipelines())
+						{
+							if (Pipeline != PassPipeline && ResourceMergeState->LastPass[Pipeline].IsValid())
+							{
+								// Add a dependency from the other pipe to this pass to join back.
+								AddPassDependency(ResourceMergeState->LastPass[Pipeline], PassHandle);
+							}
+						}
 					}
 
 					// Allocate a new pending merge state and assign it to the pass state.
 					ResourceMergeState = AllocSubresource(PassState);
-					ResourceMergeState->SetPass(PassHandle);
 				}
 				else
 				{
 					// Merge the pass state into the merged state.
 					ResourceMergeState->Access |= PassState.Access;
-					ResourceMergeState->LastPass = PassHandle;
+
+					FRDGPassHandle& FirstPassHandle = ResourceMergeState->FirstPass[PassPipeline];
+
+					if (FirstPassHandle.IsNull())
+					{
+						FirstPassHandle = PassHandle;
+					}
+
+					ResourceMergeState->LastPass[PassPipeline] = PassHandle;
 				}
 
 				PassMergeState = ResourceMergeState;
 			};
 
-			const bool bAsyncComputePass = PassesOnAsyncCompute[PassHandle];
-
-			FRDGPass* Pass = Passes[PassHandle];
-
 			for (auto& TexturePair : Pass->TextureStates)
 			{
 				FRDGTextureRef Texture = TexturePair.Key;
@@ -1297,10 +1314,14 @@ void FRDGBuilder::Execute()
 	IF_RDG_CPU_SCOPES(CPUScopeStacks.EndExecute());
 
 	IF_RDG_ENABLE_DEBUG(UserValidation.ValidateExecuteEnd());
-	IF_RDG_ENABLE_DEBUG(BarrierValidation.ValidateExecuteEnd());
 
 #if STATS
 	GRDGStatPassCount += Passes.Num();
+	GRDGStatPassWithParameterCount = GRDGStatPassCount;
+	for (FRDGPassHandle PassHandle = Passes.Begin(); PassHandle != Passes.End(); ++PassHandle)
+	{
+		GRDGStatPassWithParameterCount -= PassesWithEmptyParameters[PassHandle] ? 1 : 0;
+	}
 	GRDGStatBufferCount += Buffers.Num();
 	GRDGStatTextureCount += Textures.Num();
 	GRDGStatMemoryWatermark = FMath::Max(GRDGStatMemoryWatermark, Allocator.GetByteCount());
@@ -1361,7 +1382,7 @@ void FRDGBuilder::SetupPass(FRDGPass* Pass)
 			State.Access = MakeValidAccess(State.Access | Access);
 			State.Flags |= TransitionFlags;
 			State.NoUAVBarrierFilter.AddHandle(NoUAVBarrierHandle);
-			State.Pipeline = PassPipeline;
+			State.SetPass(PassPipeline, PassHandle);
 		};
 
 		if (bWholePassState)
@@ -1387,7 +1408,7 @@ void FRDGBuilder::SetupPass(FRDGPass* Pass)
 		PassState.ReferenceCount++;
 		PassState.State.Access = MakeValidAccess(PassState.State.Access | Access);
 		PassState.State.NoUAVBarrierFilter.AddHandle(NoUAVBarrierHandle);
-		PassState.State.Pipeline = PassPipeline;
+		PassState.State.SetPass(PassPipeline, PassHandle);
 
 		bPassUAVAccess |= EnumHasAnyFlags(Access, ERHIAccess::UAVMask);
 	});
@@ -1450,7 +1471,6 @@ void FRDGBuilder::SetupPassInternal(FRDGPass* Pass, FRDGPassHandle PassHandle, E
 				if (PassState.State[Index].Access != ERHIAccess::Unknown)
 				{
 					PassState.MergeState[Index] = &PassState.State[Index];
-					PassState.MergeState[Index]->SetPass(PassHandle);
 				}
 			}
 		}
@@ -1459,7 +1479,6 @@ void FRDGBuilder::SetupPassInternal(FRDGPass* Pass, FRDGPassHandle PassHandle, E
 		{
 			auto& PassState = BufferPair.Value;
 			PassState.MergeState = &PassState.State;
-			PassState.MergeState->SetPass(PassHandle);
 		}
 
 		check(!EnumHasAnyFlags(PassPipeline, ERHIPipeline::AsyncCompute));
@@ -1479,25 +1498,26 @@ void FRDGBuilder::ExecutePassPrologue(FRHIComputeCommandList& RHICmdListPass, FR
 
 	IF_RDG_ENABLE_DEBUG(UserValidation.ValidateExecutePassBegin(Pass));
 
+	const ERDGPassFlags PassFlags = Pass->GetFlags();
+	const ERHIPipeline PassPipeline = Pass->GetPipeline();
+
 	if (Pass->PrologueBarriersToBegin)
 	{
 		IF_RDG_ENABLE_DEBUG(BarrierValidation.ValidateBarrierBatchBegin(Pass, *Pass->PrologueBarriersToBegin));
-		Pass->PrologueBarriersToBegin->Submit(RHICmdListPass);
+		Pass->PrologueBarriersToBegin->Submit(RHICmdListPass, PassPipeline);
 	}
 
 	if (Pass->PrologueBarriersToEnd)
 	{
 		IF_RDG_ENABLE_DEBUG(BarrierValidation.ValidateBarrierBatchEnd(Pass, *Pass->PrologueBarriersToEnd));
-		Pass->PrologueBarriersToEnd->Submit(RHICmdListPass);
+		Pass->PrologueBarriersToEnd->Submit(RHICmdListPass, PassPipeline);
 	}
 
-	if (Pass->GetPipeline() == ERHIPipeline::AsyncCompute)
+	if (PassPipeline == ERHIPipeline::AsyncCompute)
 	{
 		RHICmdListPass.SetAsyncComputeBudget(Pass->AsyncComputeBudget);
 	}
 
-	const ERDGPassFlags PassFlags = Pass->GetFlags();
-
 	if (EnumHasAnyFlags(PassFlags, ERDGPassFlags::Raster))
 	{
 		if (!EnumHasAnyFlags(PassFlags, ERDGPassFlags::SkipRenderPass) && !Pass->SkipRenderPassBegin())
@@ -1513,6 +1533,8 @@ void FRDGBuilder::ExecutePassEpilogue(FRHIComputeCommandList& RHICmdListPass, FR
 	CSV_SCOPED_TIMING_STAT_EXCLUSIVE_CONDITIONAL(RDGBuilder_ExecutePassEpilogue, GRDGVerboseCSVStats != 0);
 
 	const ERDGPassFlags PassFlags = Pass->GetFlags();
+	const ERHIPipeline PassPipeline = Pass->GetPipeline();
+	const FRDGParameterStruct PassParameters = Pass->GetParameters();
 
 	if (EnumHasAnyFlags(PassFlags, ERDGPassFlags::Raster) && !EnumHasAnyFlags(PassFlags, ERDGPassFlags::SkipRenderPass) && !Pass->SkipRenderPassEnd())
 	{
@@ -1529,18 +1551,38 @@ void FRDGBuilder::ExecutePassEpilogue(FRHIComputeCommandList& RHICmdListPass, FR
 		RHIAcquireTransientResource(Texture);
 	}
 
-	const FRDGParameterStruct PassParameters = Pass->GetParameters();
+	FRDGTransitionQueue Transitions;
 
 	if (Pass->EpilogueBarriersToBeginForGraphics)
 	{
 		IF_RDG_ENABLE_DEBUG(BarrierValidation.ValidateBarrierBatchBegin(Pass, *Pass->EpilogueBarriersToBeginForGraphics));
-		Pass->EpilogueBarriersToBeginForGraphics->Submit(RHICmdListPass);
+		Pass->EpilogueBarriersToBeginForGraphics->Submit(RHICmdListPass, PassPipeline, Transitions);
 	}
 
 	if (Pass->EpilogueBarriersToBeginForAsyncCompute)
 	{
 		IF_RDG_ENABLE_DEBUG(BarrierValidation.ValidateBarrierBatchBegin(Pass, *Pass->EpilogueBarriersToBeginForAsyncCompute));
-		Pass->EpilogueBarriersToBeginForAsyncCompute->Submit(RHICmdListPass);
+		Pass->EpilogueBarriersToBeginForAsyncCompute->Submit(RHICmdListPass, PassPipeline, Transitions);
+	}
+
+	if (Pass->EpilogueBarriersToBeginForAll)
+	{
+		IF_RDG_ENABLE_DEBUG(BarrierValidation.ValidateBarrierBatchBegin(Pass, *Pass->EpilogueBarriersToBeginForAll));
+		Pass->EpilogueBarriersToBeginForAll->Submit(RHICmdListPass, PassPipeline, Transitions);
+	}
+
+	for (FRDGBarrierBatchBegin* BarriersToBegin : Pass->SharedEpilogueBarriersToBegin)
+	{
+		IF_RDG_ENABLE_DEBUG(BarrierValidation.ValidateBarrierBatchBegin(Pass, *BarriersToBegin));
+		BarriersToBegin->Submit(RHICmdListPass, PassPipeline, Transitions);
+	}
+
+	Transitions.Begin(RHICmdListPass);
+
+	if (Pass->EpilogueBarriersToEnd)
+	{
+		IF_RDG_ENABLE_DEBUG(BarrierValidation.ValidateBarrierBatchEnd(Pass, *Pass->EpilogueBarriersToEnd));
+		Pass->EpilogueBarriersToEnd->Submit(RHICmdListPass, PassPipeline);
 	}
 
 	IF_RDG_ENABLE_DEBUG(UserValidation.ValidateExecutePassEnd(Pass));
@@ -1660,8 +1702,7 @@ void FRDGBuilder::CollectPassBarriers(FRDGPassHandle PassHandle, FRDGPassHandle&
 
 	IF_RDG_ENABLE_DEBUG(ConditionalDebugBreak(RDG_BREAKPOINT_PASS_COMPILE, BuilderName.GetTCHAR(), Pass->GetName()));
 
-	const ERDGPassFlags PassFlags = Pass->GetFlags();
-	if (EnumHasAnyFlags(PassFlags, ERDGPassFlags::UntrackedAccess))
+	if (EnumHasAnyFlags(Pass->GetFlags(), ERDGPassFlags::UntrackedAccess))
 	{
 		LastUntrackedPassHandle = PassHandle;
 	}
@@ -1671,9 +1712,6 @@ void FRDGBuilder::CollectPassBarriers(FRDGPassHandle PassHandle, FRDGPassHandle&
 		return;
 	}
 
-	const FRDGParameterStruct PassParameters = Pass->GetParameters();
-	const ERHIPipeline PassPipeline = Pass->GetPipeline();
-
 	for (const auto& TexturePair : Pass->TextureStates)
 	{
 		FRDGTextureRef Texture = TexturePair.Key;
@@ -1703,7 +1741,7 @@ void FRDGBuilder::AddEpilogueTransition(FRDGTextureRef Texture, FRDGPassHandle L
 	// A known final state means extraction from the graph (or an external texture).
 	if (Texture->AccessFinal != ERHIAccess::Unknown)
 	{
-		ScratchSubresourceState.SetPass(EpiloguePassHandle);
+		ScratchSubresourceState.SetPass(ERHIPipeline::Graphics, EpiloguePassHandle);
 		ScratchSubresourceState.Access = Texture->AccessFinal;
 		InitAsWholeResource(ScratchTextureState, &ScratchSubresourceState);
 	}
@@ -1719,10 +1757,9 @@ void FRDGBuilder::AddEpilogueTransition(FRDGTextureRef Texture, FRDGPassHandle L
 			FRDGSubresourceState SubresourceState = TextureState[Index];
 
 			// Transition async compute back to the graphics pipe.
-			if (SubresourceState.Pipeline == ERHIPipeline::AsyncCompute)
+			if (SubresourceState.IsUsedBy(ERHIPipeline::AsyncCompute))
 			{
-				SubresourceState.SetPass(EpiloguePassHandle);
-				SubresourceState.Pipeline = ERHIPipeline::Graphics;
+				SubresourceState.SetPass(ERHIPipeline::Graphics, EpiloguePassHandle);
 
 				ScratchTextureState[Index] = AllocSubresource(SubresourceState);
 			}
@@ -1758,7 +1795,7 @@ void FRDGBuilder::AddEpilogueTransition(FRDGBufferRef Buffer, FRDGPassHandle Las
 	{
 		const FRDGSubresourceState State = Buffer->GetState();
 
-		if (State.Pipeline == ERHIPipeline::AsyncCompute)
+		if (State.IsUsedBy(ERHIPipeline::AsyncCompute))
 		{
 			AccessFinal = State.Access;
 		}
@@ -1767,7 +1804,7 @@ void FRDGBuilder::AddEpilogueTransition(FRDGBufferRef Buffer, FRDGPassHandle Las
 	if (AccessFinal != ERHIAccess::Unknown)
 	{
 		FRDGSubresourceState StateFinal;
-		StateFinal.SetPass(EpiloguePassHandle);
+		StateFinal.SetPass(ERHIPipeline::Graphics, EpiloguePassHandle);
 		StateFinal.Access = AccessFinal;
 		AddTransition(EpiloguePassHandle, Buffer, StateFinal, LastUntrackedPassHandle);
 	}
@@ -1785,7 +1822,6 @@ void FRDGBuilder::AddTransition(FRDGPassHandle PassHandle, FRDGTextureRef Textur
 		FRDGTextureSubresource* Subresource)
 	{
 		check(SubresourceStateAfter.Access != ERHIAccess::Unknown);
-		check(SubresourceStateAfter.FirstPass.IsValid() && SubresourceStateAfter.LastPass.IsValid());
 
 		if (FRDGSubresourceState::IsTransitionRequired(SubresourceStateBefore, SubresourceStateAfter))
 		{
@@ -1819,7 +1855,14 @@ void FRDGBuilder::AddTransition(FRDGPassHandle PassHandle, FRDGTextureRef Textur
 	const auto MergeSubresourceState = [&] (FRDGSubresourceState& SubresourceStateBefore, const FRDGSubresourceState& SubresourceStateAfter)
 	{
 		SubresourceStateBefore = SubresourceStateAfter;
-		SubresourceStateBefore.FirstPass = PassHandle;
+
+		for (ERHIPipeline Pipeline : GetRHIPipelines())
+		{
+			if (SubresourceStateBefore.IsUsedBy(Pipeline))
+			{
+				SubresourceStateBefore.FirstPass[Pipeline] = PassHandle;
+			}
+		}
 	};
 
 	if (IsWholeResource(StateBefore))
@@ -1831,8 +1874,7 @@ void FRDGBuilder::AddTransition(FRDGPassHandle PassHandle, FRDGTextureRef Textur
 			{
 				FRDGSubresourceState& SubresourceStateBefore = GetWholeResource(StateBefore);
 				AddSubresourceTransition(SubresourceStateBefore, *SubresourceStateAfter, nullptr);
-				SubresourceStateBefore = *SubresourceStateAfter;
-				SubresourceStateBefore.FirstPass = PassHandle;
+				MergeSubresourceState(SubresourceStateBefore, *SubresourceStateAfter);
 			}
 		}
 		// 1 -> N
@@ -1886,7 +1928,6 @@ void FRDGBuilder::AddTransition(FRDGPassHandle PassHandle, FRDGTextureRef Textur
 void FRDGBuilder::AddTransition(FRDGPassHandle PassHandle, FRDGBufferRef Buffer, FRDGSubresourceState StateAfter, FRDGPassHandle LastUntrackedPassHandle)
 {
 	check(StateAfter.Access != ERHIAccess::Unknown);
-	check(StateAfter.FirstPass.IsValid() && StateAfter.LastPass.IsValid());
 
 	FRDGSubresourceState& StateBefore = Buffer->GetState();
 
@@ -1918,42 +1959,65 @@ void FRDGBuilder::AddTransitionInternal(
 		return;
 	}
 
-	check(StateAfter.FirstPass.IsValid());
+	const ERHIPipeline Graphics = ERHIPipeline::Graphics;
+	const ERHIPipeline AsyncCompute = ERHIPipeline::AsyncCompute;
+
+#if RDG_ENABLE_DEBUG
+	StateBefore.Validate();
+	StateAfter.Validate();
+#endif
 
-	if (StateBefore.LastPass.IsNull())
+	if (GRDGImmediateMode != 0)
 	{
-		check(StateBefore.Pipeline == ERHIPipeline::Graphics);
-		StateBefore.LastPass = GetProloguePassHandle();
+		FRDGPass* NextPass = Passes[StateAfter.FirstPass[Graphics]];
+		FRDGBarrierBatchBegin& BarriersToBegin = NextPass->GetPrologueBarriersToBegin(Allocator);
+		BarriersToBegin.AddTransition(Resource, TransitionInfo);
+		NextPass->GetPrologueBarriersToEnd(Allocator).AddDependency(&BarriersToBegin);
+		return;
 	}
 
-	check(StateBefore.LastPass < StateAfter.FirstPass);
+	ERHIPipeline PipelinesBefore = StateBefore.GetPipelines();
+	ERHIPipeline PipelinesAfter = StateAfter.GetPipelines();
+
+	// This may be the first use of the resource in the graph, so we assign the prologue as the previous pass.
+	if (PipelinesBefore == ERHIPipeline::None)
+	{
+		StateBefore.SetPass(Graphics, GetProloguePassHandle());
+		PipelinesBefore = Graphics;
+	}
+
+	check(PipelinesBefore != ERHIPipeline::None && PipelinesAfter != ERHIPipeline::None);
+	checkf(StateBefore.GetLastPass() < StateAfter.GetFirstPass(), TEXT("Submitted a state for '%s' that begins before our previous state has ended."), Resource->Name);
+
+	FRDGPassHandlesByPipeline& PassesBefore = StateBefore.LastPass;
+	FRDGPassHandlesByPipeline& PassesAfter  = StateAfter.FirstPass;
 
 	// Before states may come from previous aliases of the texture.
-	if (Resource->bTransient && StateBefore.LastPass < Resource->AcquirePass)
+	if (Resource->bTransient && StateBefore.GetLastPass() < Resource->AcquirePass)
 	{
+		checkf(PipelinesAfter == ERHIPipeline::Graphics, TEXT("The first use of transient resource %s is not exclusively on the graphics pipe."), Resource->Name);
+
 		// If this got left in an async compute state, we need to end the transition in the graphics
-		// fork pass, since that's the lifetime of the previous owner was extended to that point. The
-		// cost should be better absorbed by the fence as well.
-		if (StateBefore.Pipeline == ERHIPipeline::AsyncCompute)
+		// fork pass, since the lifetime of the previous owner was extended to that point. The cost
+		// should be better absorbed by the fence as well.
+		if (EnumHasAnyFlags(PipelinesBefore, AsyncCompute))
 		{
-			StateAfter.FirstPass = Passes[StateBefore.LastPass]->GraphicsJoinPass;
+			PassesAfter[Graphics] = Passes[PassesBefore[AsyncCompute]]->GraphicsJoinPass;
 		}
 		// Otherwise, can push the start of the transition forward until our alias is acquired.
 		else
 		{
-			StateBefore.LastPass = Resource->AcquirePass;
+			PassesBefore[Graphics] = Resource->AcquirePass;
 		}
 	}
 
-	check(StateBefore.LastPass < StateAfter.FirstPass);
-
 	// Avoids splitting across a pass that is doing untracked RHI access.
-	if (StateBefore.LastPass < LastUntrackedPassHandle)
+	if (StateBefore.GetLastPass() < LastUntrackedPassHandle)
 	{
 		// Transitions exclusively on the graphics queue can be pushed forward.
-		if (!EnumHasAnyFlags(StateBefore.Pipeline | StateAfter.Pipeline, ERHIPipeline::AsyncCompute))
+		if (!EnumHasAnyFlags(PipelinesBefore | PipelinesAfter, AsyncCompute))
 		{
-			StateBefore.LastPass = LastUntrackedPassHandle;
+			PassesBefore[Graphics] = LastUntrackedPassHandle;
 		}
 		// Async compute transitions have to be split. Emit a warning to the user to avoid touching the resource.
 		else
@@ -1966,28 +2030,106 @@ void FRDGBuilder::AddTransitionInternal(
 		}
 	}
 
-	FRDGPass* PrevPass = Passes[StateBefore.LastPass];
-	FRDGPass* NextPass = Passes[StateAfter.FirstPass];
+	const auto GetEpilogueBarrierPassHandle = [&](FRDGPassHandle Handle)
+	{
+		return Passes[Handle]->EpilogueBarrierPass;
+	};
 
-	// When in immediate mode or if the previous and next passes are the same, the transition
-	// occurs immediately before the pass executes in the prologue (no splitting). Otherwise,
-	// the transition is split to occur after the previous pass executes.
+	const auto GetPrologueBarrierPassHandle = [&](FRDGPassHandle Handle)
+	{
+		return Passes[Handle]->PrologueBarrierPass;
+	};
 
-	FRDGPass* PrologueBarrierPass = Passes[NextPass->PrologueBarrierPass];
+	const auto GetEpilogueBarrierPass = [&](FRDGPassHandle Handle)
+	{
+		return Passes[GetEpilogueBarrierPassHandle(Handle)];
+	};
 
-	FRDGBarrierBatchBegin* BarriersToBegin = nullptr;
-	if (PrevPass == NextPass || GRDGImmediateMode != 0)
+	const auto GetPrologueBarrierPass = [&](FRDGPassHandle Handle)
 	{
-		BarriersToBegin = &PrologueBarrierPass->GetPrologueBarriersToBegin(Allocator);
+		return Passes[GetPrologueBarrierPassHandle(Handle)];
+	};
+
+	const uint32 PipelinesBeforeCount = FMath::CountBits(uint64(PipelinesBefore));
+	const uint32 PipelinesAfterCount  = FMath::CountBits(uint64(PipelinesAfter));
+
+	// 1-to-1 or 1-to-N pipe transition.
+	if (PipelinesBeforeCount == 1)
+	{
+		const FRDGPassHandle BeginPassHandle = StateBefore.GetLastPass();
+		const FRDGPassHandle FirstEndPassHandle = StateAfter.GetFirstPass();
+
+		FRDGPass* BeginPass = nullptr;
+		FRDGBarrierBatchBegin* BarriersToBegin = nullptr;
+
+		// We can only split the transition if the before / after passes are not the same. Aliasing or untracked access affect this.
+		if (BeginPassHandle != FirstEndPassHandle)
+		{
+			BeginPass = GetEpilogueBarrierPass(BeginPassHandle);
+			BarriersToBegin = &BeginPass->GetEpilogueBarriersToBeginFor(Allocator, PipelinesAfter);
+		}
+		else
+		{
+			checkf(PipelinesAfter == ERHIPipeline::Graphics, TEXT("Attempted to queue a non-graphics pipe transition for %s from the same begin / end pass. Pipelines: %s"), Resource->Name, *GetRHIPipelineName(PipelinesAfter));
+
+			BeginPass = GetPrologueBarrierPass(BeginPassHandle);
+			BarriersToBegin = &BeginPass->GetPrologueBarriersToBegin(Allocator);
+		}
+
+		BarriersToBegin->AddTransition(Resource, TransitionInfo);
+
+		for (ERHIPipeline Pipeline : GetRHIPipelines())
+		{
+			// If doing a 1-to-N transition and this is the same pipe as the begin, we end it immediately afterwards.
+			// This is because we can't guarantee that the other pipeline won't join back before the end.
+			if (PipelinesBefore == Pipeline && PipelinesAfterCount > 1)
+			{
+				BeginPass->GetEpilogueBarriersToEnd(Allocator).AddDependency(BarriersToBegin);
+			}
+			else if (EnumHasAnyFlags(PipelinesAfter, Pipeline))
+			{
+				FRDGPass* EndPass = GetPrologueBarrierPass(PassesAfter[Pipeline]);
+				EndPass->GetPrologueBarriersToEnd(Allocator).AddDependency(BarriersToBegin);
+			}
+		}
 	}
+	// N-to-1 or N-to-N pipe transition.
 	else
 	{
-		FRDGPass* EpilogueBarrierPass = Passes[PrevPass->EpilogueBarrierPass];
-		BarriersToBegin = &EpilogueBarrierPass->GetEpilogueBarriersToBeginFor(Allocator, StateAfter.Pipeline);
-	}
-	BarriersToBegin->AddTransition(Resource, TransitionInfo);
+		checkf(StateBefore.GetLastPass() != StateAfter.GetFirstPass(),
+			TEXT("Attempted to queue a transition for resource '%s' from '%s' to '%s', but previous and next passes are the same on one pipe."),
+			Resource->Name, *GetRHIPipelineName(PipelinesBefore), *GetRHIPipelineName(PipelinesAfter));
+
+		FRDGBarrierBatchBeginId Id;
+		Id.PipelinesAfter = PipelinesAfter;
+		for (ERHIPipeline Pipeline : GetRHIPipelines())
+		{
+			Id.Passes[Pipeline] = GetEpilogueBarrierPassHandle(PassesBefore[Pipeline]);
+		}
+
+		FRDGBarrierBatchBegin*& BarriersToBegin = BarrierBatchMap.FindOrAdd(Id);
 
-	PrologueBarrierPass->GetPrologueBarriersToEnd(Allocator).AddDependency(BarriersToBegin);
+		if (!BarriersToBegin)
+		{
+			BarriersToBegin = Allocator.AllocNoDestruct<FRDGBarrierBatchBegin>(PipelinesBefore, PipelinesAfter, GetEpilogueBarriersToBeginDebugName(PipelinesAfter), Id.Passes);
+
+			for (FRDGPassHandle PassHandle : Id.Passes)
+			{
+				Passes[PassHandle]->SharedEpilogueBarriersToBegin.Add(BarriersToBegin);
+			}
+		}
+
+		BarriersToBegin->AddTransition(Resource, TransitionInfo);
+
+		for (ERHIPipeline Pipeline : GetRHIPipelines())
+		{
+			if (EnumHasAnyFlags(PipelinesAfter, Pipeline))
+			{
+				FRDGPass* EndPass = GetPrologueBarrierPass(PassesAfter[Pipeline]);
+				EndPass->GetPrologueBarriersToEnd(Allocator).AddDependency(BarriersToBegin);
+			}
+		}
+	}
 }
 
 void FRDGBuilder::BeginResourceRHI(FRDGUniformBuffer* UniformBuffer)
diff --git a/Engine/Source/Runtime/RenderCore/Private/RenderGraphPass.cpp b/Engine/Source/Runtime/RenderCore/Private/RenderGraphPass.cpp
--- a/Engine/Source/Runtime/RenderCore/Private/RenderGraphPass.cpp
+++ b/Engine/Source/Runtime/RenderCore/Private/RenderGraphPass.cpp
@@ -4,6 +4,14 @@
 #include "RenderGraphPrivate.h"
 #include "RenderGraphUtils.h"
 
+static FRDGPassHandlesByPipeline GetPassesByPipeline(const FRDGPass* Pass)
+{
+	check(Pass);
+	FRDGPassHandlesByPipeline Passes;
+	Passes[Pass->GetPipeline()] = Pass->GetHandle();
+	return Passes;
+}
+
 FUniformBufferStaticBindings FRDGParameterStruct::GetGlobalUniformBuffers() const
 {
 	FUniformBufferStaticBindings GlobalUniformBuffers;
@@ -97,177 +105,192 @@ FRHIRenderPassInfo FRDGParameterStruct::GetRenderPassInfo() const
 	return RenderPassInfo;
 }
 
-FRDGBarrierBatch::FRDGBarrierBatch(const FRDGPass* InPass, const TCHAR* InName)
-	: Pipeline(InPass->GetPipeline())
-#if RDG_ENABLE_DEBUG
-	, Pass(InPass)
-	, Name(InName)
-#endif
-{}
+uint32 GetTypeHash(FRDGBarrierBatchBeginId Id)
+{
+	static_assert(sizeof(Id.Passes) == 4, "Hash expects the Passes array to be 4 bytes (2 uint16's).");
+	uint32 Hash = *(const uint32*)Id.Passes.GetData();
+	return (Hash << GetRHIPipelineCount()) | uint32(Id.PipelinesAfter);
+}
 
-void FRDGBarrierBatch::SetSubmitted()
+FRDGTransitionQueue::FRDGTransitionQueue(uint32 ReservedCount)
 {
-	check(!bSubmitted);
-	bSubmitted = true;
+	Queue.Reserve(ReservedCount);
 }
 
-FString FRDGBarrierBatch::GetName() const
+void FRDGTransitionQueue::Insert(const FRHITransition* Transition, ERHICreateTransitionFlags TransitionFlags)
 {
-#if RDG_ENABLE_DEBUG
-	return FString::Printf(TEXT("[%s(%s)]: %s"), Pass->GetName(), *GetRHIPipelineName(Pass->GetPipeline()), Name);
-#else
-	return {};
-#endif
+	if (!EnumHasAnyFlags(TransitionFlags, ERHICreateTransitionFlags::NoFence))
+	{
+		QueueWithFences.Add(Transition);
+	}
+	else
+	{
+		Queue.Add(Transition);
+	}
+}
+
+void FRDGTransitionQueue::Begin(FRHIComputeCommandList& RHICmdList)
+{
+	if (Queue.Num() || QueueWithFences.Num())
+	{
+		// Fence signals happen last.
+		Queue.Append(QueueWithFences);
+		RHICmdList.BeginTransitions(Queue);
+		Queue.Empty();
+		QueueWithFences.Empty();
+	}
 }
 
-FRDGBarrierBatchBegin::FRDGBarrierBatchBegin(const FRDGPass* InPass, const TCHAR* InName, TOptional<ERHIPipeline> InOverridePipelineToEnd)
-	: FRDGBarrierBatch(InPass, InName)
-	, OverridePipelineToEnd(InOverridePipelineToEnd)
+void FRDGTransitionQueue::End(FRHIComputeCommandList& RHICmdList)
+{
+	if (Queue.Num() || QueueWithFences.Num())
+	{
+		// Fence waits happen first.
+		Queue.Insert(QueueWithFences, 0);
+		RHICmdList.EndTransitions(Queue);
+		Queue.Empty();
+		QueueWithFences.Empty();
+	}
+}
+
+FRDGBarrierBatchBegin::FRDGBarrierBatchBegin(ERHIPipeline InPipelinesToBegin, ERHIPipeline InPipelinesToEnd, const TCHAR* InDebugName, const FRDGPass* InDebugPass)
+	: FRDGBarrierBatchBegin(InPipelinesToBegin, InPipelinesToEnd, InDebugName, GetPassesByPipeline(InDebugPass))
 {}
 
-FRDGBarrierBatchBegin::~FRDGBarrierBatchBegin()
+FRDGBarrierBatchBegin::FRDGBarrierBatchBegin(ERHIPipeline InPipelinesToBegin, ERHIPipeline InPipelinesToEnd, const TCHAR* InDebugName, FRDGPassHandlesByPipeline InDebugPasses)
+	: PipelinesToBegin(InPipelinesToBegin)
+	, PipelinesToEnd(InPipelinesToEnd)
+#if RDG_ENABLE_DEBUG
+	, DebugPasses(InDebugPasses)
+	, DebugName(InDebugName)
+	, DebugPipelinesToBegin(InPipelinesToBegin)
+	, DebugPipelinesToEnd(InPipelinesToEnd)
+#endif
 {
 #if RDG_ENABLE_DEBUG
-	checkf(!Resources.Num() && !Transitions.Num(), TEXT("Begin barrier batch has unsubmitted transitions."));
-	checkf(!Transition, TEXT("Begin barrier batch %s is currently active and was never ended."), *GetName());
+	for (ERHIPipeline Pipeline : GetRHIPipelines())
+	{
+		// We should have provided corresponding debug passes to match the pipeline flags.
+		check(InDebugPasses[Pipeline].IsValid() == EnumHasAllFlags(InPipelinesToBegin, Pipeline));
+	}
 #endif
 }
 
 void FRDGBarrierBatchBegin::AddTransition(FRDGParentResourceRef Resource, const FRHITransitionInfo& Info)
 {
-	check(Resource && Info.Resource);
-
-	checkf(!IsSubmitted(), TEXT("Attempting to add transition for resource '%s' into begin batch '%s', when the begin batch has already been submitted."), Resource->Name, *GetName());
 	Transitions.Add(Info);
+	bTransitionNeeded = true;
 
 #if STATS
 	GRDGStatTransitionCount++;
 #endif
 
 #if RDG_ENABLE_DEBUG
-	Resources.Add(Resource);
+	DebugResources.Add(Resource);
 #endif
 }
 
-void FRDGBarrierBatchBegin::Submit(FRHIComputeCommandList& RHICmdList)
+void FRDGBarrierBatchBegin::Submit(FRHIComputeCommandList& RHICmdList, ERHIPipeline Pipeline, FRDGTransitionQueue& TransitionsToBegin)
 {
-	SetSubmitted();
-
-	check(!Transition);
-	if (Transitions.Num() || bUseCrossPipelineFence)
+	// Submit may be called once for each pipe. The first to submit creates the transition.
+	if (!Transition && bTransitionNeeded)
 	{
-		const ERHIPipeline PassPipeline = GetPipeline();
-
-		ERHICreateTransitionFlags Flags = ERHICreateTransitionFlags::NoFence;
-
-		if (bUseCrossPipelineFence)
-		{
-			Flags = ERHICreateTransitionFlags::None;
-		}
-
-		const ERHIPipeline DstPipeline = OverridePipelineToEnd ? OverridePipelineToEnd.GetValue() : PassPipeline;
-		Transition = RHICreateTransition(PassPipeline, DstPipeline, Flags, Transitions);
-
-		RHICmdList.BeginTransitions(MakeArrayView(&Transition, 1));
+		Transition = RHICreateTransition(PipelinesToBegin, PipelinesToEnd, TransitionFlags, Transitions);
+	}
 
-		Transitions.Empty();
-#if RDG_ENABLE_DEBUG
-		Resources.Empty();
-#endif
+	if (Transition)
+	{
+		check(EnumHasAnyFlags(PipelinesToBegin, Pipeline));
+		EnumRemoveFlags(PipelinesToBegin, Pipeline);
+		TransitionsToBegin.Insert(Transition, TransitionFlags);
+	}
 
 #if STATS
-		GRDGStatTransitionBatchCount++;
+	GRDGStatTransitionBatchCount++;
 #endif
-	}
-}
-
-FRDGBarrierBatchEnd::~FRDGBarrierBatchEnd()
-{
-	checkf(!Dependencies.Num(), TEXT("End barrier batch has unsubmitted dependencies."));
 }
 
-void FRDGBarrierBatchEnd::ReserveMemory(uint32 ExpectedDependencyCount)
+void FRDGBarrierBatchBegin::Submit(FRHIComputeCommandList& RHICmdList, ERHIPipeline Pipeline)
 {
-	Dependencies.Reserve(ExpectedDependencyCount);
+	FRDGTransitionQueue TransitionsToBegin;
+	Submit(RHICmdList, Pipeline, TransitionsToBegin);
+	TransitionsToBegin.Begin(RHICmdList);
 }
 
 void FRDGBarrierBatchEnd::AddDependency(FRDGBarrierBatchBegin* BeginBatch)
 {
-	check(BeginBatch);
-	checkf(!IsSubmitted(), TEXT("Attempting to add a dependency on begin batch '%s' into end batch '%s', when the end batch has already been submitted."), *BeginBatch->GetName(), *GetName());
 	Dependencies.AddUnique(BeginBatch);
 }
 
-void FRDGBarrierBatchEnd::Submit(FRHIComputeCommandList& RHICmdList)
+void FRDGBarrierBatchEnd::Submit(FRHIComputeCommandList& RHICmdList, ERHIPipeline Pipeline)
 {
-	SetSubmitted();
+	FRDGTransitionQueue Transitions(Dependencies.Num());
 
-	TArray<const FRHITransition*, SceneRenderingAllocator> Transitions;
-	Transitions.Reserve(Dependencies.Num());
-
-	// Process dependencies with cross-pipeline fences first.
 	for (FRDGBarrierBatchBegin* Dependent : Dependencies)
 	{
-		check(Dependent->IsSubmitted());
-
-		if (Dependent->Transition && Dependent->bUseCrossPipelineFence)
+		if (EnumHasAnyFlags(Dependent->PipelinesToEnd, Pipeline))
 		{
-			Transitions.Add(Dependent->Transition);
-			Dependent->Transition = nullptr;
+			check(Dependent->Transition);
+			EnumRemoveFlags(Dependent->PipelinesToEnd, Pipeline);
+			Transitions.Insert(Dependent->Transition, Dependent->TransitionFlags);
 		}
 	}
 
-	for (FRDGBarrierBatchBegin* Dependent : Dependencies)
+	Transitions.End(RHICmdList);
+}
+
+FRDGBarrierBatchBegin& FRDGPass::GetPrologueBarriersToBegin(FRDGAllocator& Allocator)
+{
+	if (!PrologueBarriersToBegin)
 	{
-		if (Dependent->Transition)
-		{
-			Transitions.Add(Dependent->Transition);
-			Dependent->Transition = nullptr;
-		}
+		PrologueBarriersToBegin = Allocator.AllocNoDestruct<FRDGBarrierBatchBegin>(Pipeline, Pipeline, TEXT("Prologue"), this);
 	}
+	return *PrologueBarriersToBegin;
+}
 
-	Dependencies.Empty();
-
-	if (Transitions.Num())
+FRDGBarrierBatchBegin& FRDGPass::GetEpilogueBarriersToBeginForGraphics(FRDGAllocator& Allocator)
+{
+	if (!EpilogueBarriersToBeginForGraphics)
 	{
-		RHICmdList.EndTransitions(Transitions);
+		EpilogueBarriersToBeginForGraphics = Allocator.AllocNoDestruct<FRDGBarrierBatchBegin>(Pipeline, ERHIPipeline::Graphics, GetEpilogueBarriersToBeginDebugName(ERHIPipeline::Graphics), this);
 	}
+	return *EpilogueBarriersToBeginForGraphics;
 }
 
-FRDGBarrierBatchBegin& FRDGPass::GetPrologueBarriersToBegin(FRDGAllocator& Allocator)
+FRDGBarrierBatchBegin& FRDGPass::GetEpilogueBarriersToBeginForAsyncCompute(FRDGAllocator& Allocator)
 {
-	if (!PrologueBarriersToBegin)
+	if (!EpilogueBarriersToBeginForAsyncCompute)
 	{
-		PrologueBarriersToBegin = Allocator.AllocNoDestruct<FRDGBarrierBatchBegin>(this, TEXT("Prologue"));
+		EpilogueBarriersToBeginForAsyncCompute = Allocator.AllocNoDestruct<FRDGBarrierBatchBegin>(Pipeline, ERHIPipeline::AsyncCompute, GetEpilogueBarriersToBeginDebugName(ERHIPipeline::AsyncCompute), this);
 	}
-	return *PrologueBarriersToBegin;
+	return *EpilogueBarriersToBeginForAsyncCompute;
 }
 
-FRDGBarrierBatchEnd& FRDGPass::GetPrologueBarriersToEnd(FRDGAllocator& Allocator)
+FRDGBarrierBatchBegin& FRDGPass::GetEpilogueBarriersToBeginForAll(FRDGAllocator& Allocator)
 {
-	if (!PrologueBarriersToEnd)
+	if (!EpilogueBarriersToBeginForAll)
 	{
-		PrologueBarriersToEnd = Allocator.AllocNoDestruct<FRDGBarrierBatchEnd>(this, TEXT("Prologue"));
+		EpilogueBarriersToBeginForAll = Allocator.AllocNoDestruct<FRDGBarrierBatchBegin>(Pipeline, ERHIPipeline::All, GetEpilogueBarriersToBeginDebugName(ERHIPipeline::AsyncCompute), this);
 	}
-	return *PrologueBarriersToEnd;
+	return *EpilogueBarriersToBeginForAll;
 }
 
-FRDGBarrierBatchBegin& FRDGPass::GetEpilogueBarriersToBeginForGraphics(FRDGAllocator& Allocator)
+FRDGBarrierBatchEnd& FRDGPass::GetPrologueBarriersToEnd(FRDGAllocator& Allocator)
 {
-	if (!EpilogueBarriersToBeginForGraphics)
+	if (!PrologueBarriersToEnd)
 	{
-		EpilogueBarriersToBeginForGraphics = Allocator.AllocNoDestruct<FRDGBarrierBatchBegin>(this, TEXT("Epilogue (For Graphics)"), ERHIPipeline::Graphics);
+		PrologueBarriersToEnd = Allocator.AllocNoDestruct<FRDGBarrierBatchEnd>();
 	}
-	return *EpilogueBarriersToBeginForGraphics;
+	return *PrologueBarriersToEnd;
 }
 
-FRDGBarrierBatchBegin& FRDGPass::GetEpilogueBarriersToBeginForAsyncCompute(FRDGAllocator& Allocator)
+FRDGBarrierBatchEnd& FRDGPass::GetEpilogueBarriersToEnd(FRDGAllocator& Allocator)
 {
-	if (!EpilogueBarriersToBeginForAsyncCompute)
+	if (!EpilogueBarriersToEnd)
 	{
-		EpilogueBarriersToBeginForAsyncCompute = Allocator.AllocNoDestruct<FRDGBarrierBatchBegin>(this, TEXT("Epilogue (For AsyncCompute)"), ERHIPipeline::AsyncCompute);
+		EpilogueBarriersToEnd = Allocator.AllocNoDestruct<FRDGBarrierBatchEnd>();
 	}
-	return *EpilogueBarriersToBeginForAsyncCompute;
+	return *EpilogueBarriersToEnd;
 }
 
 FRDGPass::FRDGPass(
diff --git a/Engine/Source/Runtime/RenderCore/Private/RenderGraphPrivate.cpp b/Engine/Source/Runtime/RenderCore/Private/RenderGraphPrivate.cpp
--- a/Engine/Source/Runtime/RenderCore/Private/RenderGraphPrivate.cpp
+++ b/Engine/Source/Runtime/RenderCore/Private/RenderGraphPrivate.cpp
@@ -288,6 +288,7 @@ FAutoConsoleVariableRef CVarRDGVerboseCSVStats(
 
 #if STATS
 int32 GRDGStatPassCount = 0;
+int32 GRDGStatPassWithParameterCount = 0;
 int32 GRDGStatPassCullCount = 0;
 int32 GRDGStatPassDependencyCount = 0;
 int32 GRDGStatRenderPassMergeCount = 0;
@@ -298,6 +299,7 @@ int32 GRDGStatTransitionBatchCount = 0;
 int32 GRDGStatMemoryWatermark = 0;
 
 DEFINE_STAT(STAT_RDG_PassCount);
+DEFINE_STAT(STAT_RDG_PassWithParameterCount);
 DEFINE_STAT(STAT_RDG_PassCullCount);
 DEFINE_STAT(STAT_RDG_RenderPassMergeCount);
 DEFINE_STAT(STAT_RDG_PassDependencyCount);
diff --git a/Engine/Source/Runtime/RenderCore/Private/RenderGraphPrivate.h b/Engine/Source/Runtime/RenderCore/Private/RenderGraphPrivate.h
--- a/Engine/Source/Runtime/RenderCore/Private/RenderGraphPrivate.h
+++ b/Engine/Source/Runtime/RenderCore/Private/RenderGraphPrivate.h
@@ -96,6 +96,7 @@ const int32 GRDGVerboseCSVStats = 0;
 #endif
 
 #if STATS
+extern int32 GRDGStatPassWithParameterCount;
 extern int32 GRDGStatPassCount;
 extern int32 GRDGStatPassCullCount;
 extern int32 GRDGStatRenderPassMergeCount;
@@ -107,6 +108,7 @@ extern int32 GRDGStatTransitionBatchCount;
 extern int32 GRDGStatMemoryWatermark;
 
 DECLARE_DWORD_COUNTER_STAT_EXTERN(TEXT("Passes"), STAT_RDG_PassCount, STATGROUP_RDG, RENDERCORE_API);
+DECLARE_DWORD_COUNTER_STAT_EXTERN(TEXT("Passes With Parameters"), STAT_RDG_PassWithParameterCount, STATGROUP_RDG, RENDERCORE_API);
 DECLARE_DWORD_COUNTER_STAT_EXTERN(TEXT("Passes Culled"), STAT_RDG_PassCullCount, STATGROUP_RDG, RENDERCORE_API);
 DECLARE_DWORD_COUNTER_STAT_EXTERN(TEXT("Render Passes Merged"), STAT_RDG_RenderPassMergeCount, STATGROUP_RDG, RENDERCORE_API);
 DECLARE_DWORD_COUNTER_STAT_EXTERN(TEXT("Pass Dependencies"), STAT_RDG_PassDependencyCount, STATGROUP_RDG, RENDERCORE_API);
@@ -121,4 +123,20 @@ DECLARE_CYCLE_STAT_EXTERN(TEXT("Collect Barriers"), STAT_RDG_CollectBarriersTime
 DECLARE_CYCLE_STAT_EXTERN(TEXT("Clear"), STAT_RDG_ClearTime, STATGROUP_RDG, RENDERCORE_API);
 
 DECLARE_MEMORY_STAT_EXTERN(TEXT("Builder Watermark"), STAT_RDG_MemoryWatermark, STATGROUP_RDG, RENDERCORE_API);
-#endif
\ No newline at end of file
+#endif
+
+inline const TCHAR* GetEpilogueBarriersToBeginDebugName(ERHIPipeline Pipelines)
+{
+#if RDG_ENABLE_DEBUG
+	switch (Pipelines)
+	{
+	case ERHIPipeline::Graphics:
+		return TEXT("Epilogue (For Graphics)");
+	case ERHIPipeline::AsyncCompute:
+		return TEXT("Epilogue (For AsyncCompute)");
+	case ERHIPipeline::All:
+		return TEXT("Epilogue (For All)");
+	}
+#endif
+	return TEXT("");
+}
\ No newline at end of file
diff --git a/Engine/Source/Runtime/RenderCore/Private/RenderGraphResources.cpp b/Engine/Source/Runtime/RenderCore/Private/RenderGraphResources.cpp
--- a/Engine/Source/Runtime/RenderCore/Private/RenderGraphResources.cpp
+++ b/Engine/Source/Runtime/RenderCore/Private/RenderGraphResources.cpp
@@ -42,7 +42,7 @@ bool FRDGSubresourceState::IsTransitionRequired(const FRDGSubresourceState& Prev
 {
 	check(Next.Access != ERHIAccess::Unknown);
 
-	if (Previous.Access != Next.Access || Previous.Pipeline != Next.Pipeline)
+	if (Previous.Access != Next.Access || Previous.GetPipelines() != Next.GetPipelines())
 	{
 		return true;
 	}
@@ -100,8 +100,14 @@ bool FRDGSubresourceState::IsMergeAllowed(ERDGParentResourceType ResourceType, c
 		return false;
 	}
 
+	// If our platform doesn't support multi-pipeline transitions, then we can't merge the state.
+	if (EnumHasAnyFlags(AccessUnion, ~GRHIMultiPipelineAccessMask) && Previous.GetPipelines() != Next.GetPipelines())
+	{
+		return false;
+	}
+
 	// We are not allowed to cross pipelines or change flags in a merge.
-	if (Previous.Pipeline != Next.Pipeline || Previous.Flags != Next.Flags)
+	if (Previous.Flags != Next.Flags)
 	{
 		return false;
 	}
diff --git a/Engine/Source/Runtime/RenderCore/Private/RenderGraphValidation.cpp b/Engine/Source/Runtime/RenderCore/Private/RenderGraphValidation.cpp
--- a/Engine/Source/Runtime/RenderCore/Private/RenderGraphValidation.cpp
+++ b/Engine/Source/Runtime/RenderCore/Private/RenderGraphValidation.cpp
@@ -721,30 +721,31 @@ FRDGBarrierValidation::FRDGBarrierValidation(const FRDGPassRegistry* InPasses, c
 
 void FRDGBarrierValidation::ValidateBarrierBatchBegin(const FRDGPass* Pass, const FRDGBarrierBatchBegin& Batch)
 {
-	checkf(!Batch.IsTransitionValid() && !BatchMap.Contains(&Batch), TEXT("Begin barrier batch '%s' has already been submitted."), *Batch.GetName());
-
 	if (!GRDGTransitionLog)
 	{
 		return;
 	}
 
-	FResourceMap& ResourceMap = BatchMap.Emplace(&Batch);
+	FResourceMap* ResourceMap = BatchMap.Find(&Batch);
 
-	for (int32 Index = 0; Index < Batch.Transitions.Num(); ++Index)
+	if (!ResourceMap)
 	{
-		FRDGParentResourceRef Resource = Batch.Resources[Index];
-		const FRHITransitionInfo& Transition = Batch.Transitions[Index];
+		ResourceMap = &BatchMap.Emplace(&Batch);
 
-		if (Resource->Type == ERDGParentResourceType::Texture)
-		{
-			ResourceMap.Textures.FindOrAdd(static_cast<FRDGTextureRef>(Resource)).Add(Transition);
-		}
-		else
+		for (int32 Index = 0; Index < Batch.Transitions.Num(); ++Index)
 		{
-			check(Resource->Type == ERDGParentResourceType::Buffer);
-			FRDGBufferRef Buffer = static_cast<FRDGBufferRef>(Resource);
-			checkf(!ResourceMap.Buffers.Contains(Buffer), TEXT("Buffer %s was added multiple times to batch %s."), Buffer->Name, *Batch.GetName());
-			ResourceMap.Buffers.Emplace(Buffer, Transition);
+			FRDGParentResourceRef Resource = Batch.DebugResources[Index];
+			const FRHITransitionInfo& Transition = Batch.Transitions[Index];
+
+			if (Resource->Type == ERDGParentResourceType::Texture)
+			{
+				ResourceMap->Textures.FindOrAdd(static_cast<FRDGTextureRef>(Resource)).Add(Transition);
+			}
+			else
+			{
+				check(Resource->Type == ERDGParentResourceType::Buffer);
+				ResourceMap->Buffers.Emplace(static_cast<FRDGBufferRef>(Resource), Transition);
+			}
 		}
 	}
 
@@ -763,11 +764,11 @@ void FRDGBarrierValidation::ValidateBarrierBatchBegin(const FRDGPass* Pass, cons
 		if (!bFoundFirst)
 		{
 			bFoundFirst = true;
-			UE_CLOG(bAllowedForPass, LogRDG, Display, TEXT("%s (Begin):"), *Batch.GetName());
+			UE_CLOG(bAllowedForPass, LogRDG, Display, TEXT("[%s(Index: %d, Pipeline: %s): %s] (Begin):"), Pass->GetName(), Pass->GetHandle().GetIndex(), *GetRHIPipelineName(Pass->GetPipeline()), Batch.DebugName);
 		}
 	};
 
-	for (const auto& Pair : ResourceMap.Textures)
+	for (const auto& Pair : ResourceMap->Textures)
 	{
 		FRDGTextureRef Texture = Pair.Key;
 		
@@ -791,17 +792,19 @@ void FRDGBarrierValidation::ValidateBarrierBatchBegin(const FRDGPass* Pass, cons
 			{
 				const int32 SubresourceIndex = SubresourceLayout.GetSubresourceIndex(Subresource);
 
-				UE_CLOG(bAllowedForResource, LogRDG, Display, TEXT("\t\tMip(%d), Array(%d), Slice(%d): %s -> %s"),
+				UE_CLOG(bAllowedForResource, LogRDG, Display, TEXT("\t\tMip(%d), Array(%d), Slice(%d): [%s, %s] -> [%s, %s]"),
 					Subresource.MipIndex, Subresource.ArraySlice, Subresource.PlaneSlice,
 					*GetRHIAccessName(Transition.AccessBefore),
-					*GetRHIAccessName(Transition.AccessAfter));
+					*GetRHIPipelineName(Batch.DebugPipelinesToBegin),
+					*GetRHIAccessName(Transition.AccessAfter),
+					*GetRHIPipelineName(Batch.DebugPipelinesToEnd));
 			});
 		}
 	}
 
 	if (bAllowedForPass)
 	{
-		for (auto Pair : ResourceMap.Buffers)
+		for (const auto& Pair : ResourceMap->Buffers)
 		{
 			FRDGBufferRef Buffer = Pair.Key;
 			const FRHITransitionInfo& Transition = Pair.Value;
@@ -813,11 +816,13 @@ void FRDGBarrierValidation::ValidateBarrierBatchBegin(const FRDGPass* Pass, cons
 
 			LogHeader();
 
-			UE_LOG(LogRDG, Display, TEXT("\t(%p) %s: %s -> %s"),
+			UE_LOG(LogRDG, Display, TEXT("\t(%p) %s: [%s, %s] -> [%s, %s]"),
 				Buffer,
 				Buffer->Name,
 				*GetRHIAccessName(Transition.AccessBefore),
-				*GetRHIAccessName(Transition.AccessAfter));
+				*GetRHIPipelineName(Batch.DebugPipelinesToBegin),
+				*GetRHIAccessName(Transition.AccessAfter),
+				*GetRHIPipelineName(Batch.DebugPipelinesToEnd));
 		}
 	}
 }
@@ -833,8 +838,7 @@ void FRDGBarrierValidation::ValidateBarrierBatchEnd(const FRDGPass* Pass, const
 
 	for (const FRDGBarrierBatchBegin* Dependent : Batch.Dependencies)
 	{
-		// Transitions can be queued into multiple end batches. The first to get flushed nulls out the transition.
-		if (!Dependent->IsTransitionValid())
+		if (Dependent->PipelinesToEnd == ERHIPipeline::None)
 		{
 			continue;
 		}
@@ -857,32 +861,29 @@ void FRDGBarrierValidation::ValidateBarrierBatchEnd(const FRDGPass* Pass, const
 		{
 			if (!bFoundFirstBatch)
 			{
-				UE_LOG(LogRDG, Display, TEXT("%s (End):"), *Batch.GetName());
+				UE_LOG(LogRDG, Display, TEXT("[%s(Index: %d, Pipeline: %s) %s] (End):"), Pass->GetName(), Pass->GetHandle().GetIndex(), Dependent->DebugName, *GetRHIPipelineName(Pass->GetPipeline()));
 				bFoundFirstBatch = true;
 			}
 		}
 
 		for (FRDGTextureRef Texture : Textures)
 		{
-			UE_LOG(LogRDG, Display, TEXT("\t(%p) %s"), Texture, Texture->Name);
+			if (IsDebugAllowedForResource(Texture->Name))
+			{
+				UE_LOG(LogRDG, Display, TEXT("\t(%p) %s"), Texture, Texture->Name);
+			}
 		}
 
 		for (FRDGBufferRef Buffer : Buffers)
 		{
-			UE_LOG(LogRDG, Display, TEXT("\t(%p) %s"), Buffer, Buffer->Name);
+			if (IsDebugAllowedForResource(Buffer->Name))
+			{
+				UE_LOG(LogRDG, Display, TEXT("\t(%p) %s"), Buffer, Buffer->Name);
+			}
 		}
 	}
 }
 
-void FRDGBarrierValidation::ValidateExecuteEnd()
-{
-	for (auto Pair : BatchMap)
-	{
-		checkf(!Pair.Key->IsTransitionValid(), TEXT("A batch was begun but never ended."));
-	}
-	BatchMap.Empty();
-}
-
 namespace
 {
 	const TCHAR* RasterColorName = TEXT("#ff7070");
@@ -892,6 +893,7 @@ namespace
 	const TCHAR* TextureColorAttributes = TEXT("color=\"#5800a1\", fontcolor=\"#5800a1\"");
 	const TCHAR* BufferColorAttributes = TEXT("color=\"#007309\", fontcolor=\"#007309\"");
 	const TCHAR* AliasColorAttributes = TEXT("color=\"#00ff00\", fontcolor=\"#00ff00\"");
+	const TCHAR* AllPipelinesColorName = TEXT("#f170ff");
 
 	const TCHAR* GetPassColorName(ERDGPassFlags Flags)
 	{
@@ -916,8 +918,22 @@ namespace
 
 	FString GetSubresourceStateLabel(FRDGSubresourceState State)
 	{
-		check(State.Pipeline == ERHIPipeline::Graphics || State.Pipeline == ERHIPipeline::AsyncCompute);
-		const TCHAR* FontColor = State.Pipeline == ERHIPipeline::AsyncCompute ? AsyncComputeColorName : RasterColorName;
+		const ERHIPipeline Pipelines = State.GetPipelines();
+		const TCHAR* FontColor = nullptr;
+		switch (Pipelines)
+		{
+		default:
+			checkNoEntry();
+		case ERHIPipeline::Graphics:
+			FontColor = RasterColorName;
+			break;
+		case ERHIPipeline::AsyncCompute:
+			FontColor = AsyncComputeColorName;
+			break;
+		case ERHIPipeline::All:
+			FontColor = AllPipelinesColorName;
+			break;
+		}
 		return FString::Printf(TEXT("<font color=\"%s\">%s</font>"), FontColor, *GetRHIAccessName(State.Access));
 	}
 }
@@ -1105,7 +1121,7 @@ void FRDGLogFile::End()
 
 			const FRDGPass* Pass = Passes->Get(PassHandle);
 
-			for (FRDGPassHandle ProducerHandle : Pass->GetProducers())
+			for (const FRDGPassHandle ProducerHandle : Pass->GetProducers())
 			{
 				if (ProducerHandle != ProloguePassHandle)
 				{
@@ -1300,13 +1316,13 @@ void FRDGLogFile::AddAliasEdge(const FRDGBufferRef BufferBefore, FRDGPassHandle
 
 void FRDGLogFile::AddTransitionEdge(FRDGPassHandle PassHandle, FRDGSubresourceState StateBefore, FRDGSubresourceState StateAfter, const FRDGTextureRef Texture)
 {
-	if (GRDGDumpGraph == RDG_DUMP_GRAPH_RESOURCES && bOpen && IncludeTransitionEdgeInGraph(StateBefore.FirstPass, PassHandle))
+	if (GRDGDumpGraph == RDG_DUMP_GRAPH_RESOURCES && bOpen && IncludeTransitionEdgeInGraph(StateBefore.GetFirstPass(), PassHandle))
 	{
 		if (FRDGSubresourceState::IsTransitionRequired(StateBefore, StateAfter))
 		{
 			AddLine(FString::Printf(TEXT("\"%s\" -> \"%s\" [%s, label=<%s: <b>%s -&gt; %s</b>>]"),
-				*GetProducerName(StateBefore.LastPass),
-				*GetConsumerName(StateAfter.FirstPass),
+				*GetProducerName(StateBefore.GetLastPass()),
+				*GetConsumerName(StateAfter.GetFirstPass()),
 				TextureColorAttributes,
 				Texture->Name,
 				*GetSubresourceStateLabel(StateBefore),
@@ -1315,7 +1331,7 @@ void FRDGLogFile::AddTransitionEdge(FRDGPassHandle PassHandle, FRDGSubresourceSt
 		else
 		{
 			AddLine(FString::Printf(TEXT("\"%s\" -> \"%s\" [%s, label=<%s: <b>%s</b>>]"),
-				*GetProducerName(StateBefore.FirstPass),
+				*GetProducerName(StateBefore.GetFirstPass()),
 				*GetConsumerName(PassHandle),
 				TextureColorAttributes,
 				Texture->Name,
@@ -1326,13 +1342,13 @@ void FRDGLogFile::AddTransitionEdge(FRDGPassHandle PassHandle, FRDGSubresourceSt
 
 void FRDGLogFile::AddTransitionEdge(FRDGPassHandle PassHandle, FRDGSubresourceState StateBefore, FRDGSubresourceState StateAfter, const FRDGTextureRef Texture, FRDGTextureSubresource Subresource)
 {
-	if (GRDGDumpGraph == RDG_DUMP_GRAPH_RESOURCES && bOpen && IncludeTransitionEdgeInGraph(StateBefore.FirstPass, PassHandle))
+	if (GRDGDumpGraph == RDG_DUMP_GRAPH_RESOURCES && bOpen && IncludeTransitionEdgeInGraph(StateBefore.GetFirstPass(), PassHandle))
 	{
 		if (FRDGSubresourceState::IsTransitionRequired(StateBefore, StateAfter))
 		{
 			AddLine(FString::Printf(TEXT("\"%s\" -> \"%s\" [%s, label=<%s[%d][%d][%d]: <b>%s -&gt; %s</b>>]"),
-				*GetProducerName(StateBefore.LastPass),
-				*GetConsumerName(StateAfter.FirstPass),
+				*GetProducerName(StateBefore.GetLastPass()),
+				*GetConsumerName(StateAfter.GetFirstPass()),
 				TextureColorAttributes,
 				Texture->Name,
 				Subresource.MipIndex, Subresource.ArraySlice, Subresource.PlaneSlice,
@@ -1342,7 +1358,7 @@ void FRDGLogFile::AddTransitionEdge(FRDGPassHandle PassHandle, FRDGSubresourceSt
 		else
 		{
 			AddLine(FString::Printf(TEXT("\"%s\" -> \"%s\" [%s, label=<%s[%d][%d][%d]: <b>%s</b>>]"),
-				*GetProducerName(StateBefore.FirstPass),
+				*GetProducerName(StateBefore.GetFirstPass()),
 				*GetConsumerName(PassHandle),
 				TextureColorAttributes,
 				Texture->Name,
@@ -1354,13 +1370,13 @@ void FRDGLogFile::AddTransitionEdge(FRDGPassHandle PassHandle, FRDGSubresourceSt
 
 void FRDGLogFile::AddTransitionEdge(FRDGPassHandle PassHandle, FRDGSubresourceState StateBefore, FRDGSubresourceState StateAfter, const FRDGBufferRef Buffer)
 {
-	if (GRDGDumpGraph == RDG_DUMP_GRAPH_RESOURCES && bOpen && IncludeTransitionEdgeInGraph(StateBefore.FirstPass, PassHandle))
+	if (GRDGDumpGraph == RDG_DUMP_GRAPH_RESOURCES && bOpen && IncludeTransitionEdgeInGraph(StateBefore.GetFirstPass(), PassHandle))
 	{
 		if (FRDGSubresourceState::IsTransitionRequired(StateBefore, StateAfter))
 		{
 			AddLine(FString::Printf(TEXT("\"%s\" -> \"%s\" [%s, label=<%s: <b>%s -&gt; %s</b>>]"),
-				*GetProducerName(StateBefore.LastPass),
-				*GetConsumerName(StateAfter.FirstPass),
+				*GetProducerName(StateBefore.GetLastPass()),
+				*GetConsumerName(StateAfter.GetFirstPass()),
 				BufferColorAttributes,
 				Buffer->Name,
 				*GetSubresourceStateLabel(StateBefore),
@@ -1369,7 +1385,7 @@ void FRDGLogFile::AddTransitionEdge(FRDGPassHandle PassHandle, FRDGSubresourceSt
 		else
 		{
 			AddLine(FString::Printf(TEXT("\"%s\" -> \"%s\" [%s, label=<%s: <b>%s</b>>]"),
-				*GetProducerName(StateBefore.FirstPass),
+				*GetProducerName(StateBefore.GetFirstPass()),
 				*GetConsumerName(PassHandle),
 				BufferColorAttributes,
 				Buffer->Name,
@@ -1378,4 +1394,4 @@ void FRDGLogFile::AddTransitionEdge(FRDGPassHandle PassHandle, FRDGSubresourceSt
 	}
 }
 
-#endif
+#endif
\ No newline at end of file
diff --git a/Engine/Source/Runtime/RenderCore/Public/RenderGraphBuilder.h b/Engine/Source/Runtime/RenderCore/Public/RenderGraphBuilder.h
--- a/Engine/Source/Runtime/RenderCore/Public/RenderGraphBuilder.h
+++ b/Engine/Source/Runtime/RenderCore/Public/RenderGraphBuilder.h
@@ -258,8 +258,11 @@ private:
 	FRDGPassBitArray PassesWithEmptyParameters;
 
 	/** Tracks external resources to their registered render graph counterparts for de-duplication. */
-	TSortedMap<FRHITexture*, FRDGTexture*, TInlineAllocator<4, SceneRenderingAllocator>> ExternalTextures;
-	TSortedMap<const FRDGPooledBuffer*, FRDGBuffer*, TInlineAllocator<4, SceneRenderingAllocator>> ExternalBuffers;
+	TSortedMap<FRHITexture*, FRDGTexture*, SceneRenderingAllocator> ExternalTextures;
+	TSortedMap<const FRDGPooledBuffer*, FRDGBuffer*, SceneRenderingAllocator> ExternalBuffers;
+
+	/** Map of barrier batches begun from more than one pipe. */
+	TMap<FRDGBarrierBatchBeginId, FRDGBarrierBatchBegin*, SceneRenderingSetAllocator> BarrierBatchMap;
 
 	/** The epilogue and prologue passes are sentinels that are used to simplify graph logic around barriers
 	 *  and traversal. The prologue pass is used exclusively for barriers before the graph executes, while the
diff --git a/Engine/Source/Runtime/RenderCore/Public/RenderGraphDefinitions.h b/Engine/Source/Runtime/RenderCore/Public/RenderGraphDefinitions.h
--- a/Engine/Source/Runtime/RenderCore/Public/RenderGraphDefinitions.h
+++ b/Engine/Source/Runtime/RenderCore/Public/RenderGraphDefinitions.h
@@ -426,6 +426,8 @@ class TRDGHandleBitArray : public TBitArray<SceneRenderingBitArrayAllocator>
 	using Base = TBitArray<SceneRenderingBitArrayAllocator>;
 public:
 	TRDGHandleBitArray() = default;
+	TRDGHandleBitArray(const TRDGHandleBitArray&) = default;
+	TRDGHandleBitArray& operator=(const TRDGHandleBitArray&) = default;
 
 	explicit TRDGHandleBitArray(bool bValue, int32 InNumBits)
 		: Base(bValue, InNumBits)
@@ -565,4 +567,6 @@ template <typename TUniformStruct> class TRDGUniformBuffer;
 template <typename TUniformStruct> using TRDGUniformBufferRef = TRDGUniformBuffer<TUniformStruct>*;
 
 template <typename InElementType, typename InAllocatorType = FDefaultAllocator>
-using TRDGTextureSubresourceArray = TArray<InElementType, TInlineAllocator<1, InAllocatorType>>;
\ No newline at end of file
+using TRDGTextureSubresourceArray = TArray<InElementType, TInlineAllocator<1, InAllocatorType>>;
+
+using FRDGPassHandlesByPipeline = TRHIPipelineArray<FRDGPassHandle>;
\ No newline at end of file
diff --git a/Engine/Source/Runtime/RenderCore/Public/RenderGraphPass.h b/Engine/Source/Runtime/RenderCore/Public/RenderGraphPass.h
--- a/Engine/Source/Runtime/RenderCore/Public/RenderGraphPass.h
+++ b/Engine/Source/Runtime/RenderCore/Public/RenderGraphPass.h
@@ -6,99 +6,89 @@
 #include "RenderGraphEvent.h"
 #include "Containers/SortedMap.h"
 
-class RENDERCORE_API FRDGBarrierBatch
+class FRDGTransitionQueue
 {
 public:
-	FRDGBarrierBatch(const FRDGBarrierBatch&) = delete;
+	FRDGTransitionQueue() = default;
+	FRDGTransitionQueue(uint32 ReservedCount);
 
-	bool IsSubmitted() const
-	{
-		return bSubmitted;
-	}
+	void Insert(const FRHITransition* Transition, ERHICreateTransitionFlags TransitionFlags);
+	void Begin(FRHIComputeCommandList& RHICmdList);
+	void End(FRHIComputeCommandList& RHICmdList);
 
-	FString GetName() const;
-
-protected:
-	FRDGBarrierBatch(const FRDGPass* InPass, const TCHAR* InName);
+private:
+	TArray<const FRHITransition*, TInlineAllocator<1, SceneRenderingAllocator>> Queue;
+	TArray<const FRHITransition*, TInlineAllocator<1, SceneRenderingAllocator>> QueueWithFences;
+};
 
-	void SetSubmitted();
+struct FRDGBarrierBatchBeginId
+{
+	FRDGBarrierBatchBeginId() = default;
 
-	ERHIPipeline GetPipeline() const
+	bool operator==(FRDGBarrierBatchBeginId Other) const
 	{
-		return Pipeline;
+		return Passes == Other.Passes && PipelinesAfter == Other.PipelinesAfter;
 	}
 
-private:
-	bool bSubmitted = false;
-	ERHIPipeline Pipeline;
+	bool operator!=(FRDGBarrierBatchBeginId Other) const
+	{
+		return !(*this == Other);
+	}
 
-#if RDG_ENABLE_DEBUG
-	const FRDGPass* Pass;
-	const TCHAR* Name;
-#endif
+	FRDGPassHandlesByPipeline Passes;
+	ERHIPipeline PipelinesAfter = ERHIPipeline::None;
 };
 
-class RENDERCORE_API FRDGBarrierBatchBegin final : public FRDGBarrierBatch
+RENDERCORE_API uint32 GetTypeHash(FRDGBarrierBatchBeginId Id);
+
+class RENDERCORE_API FRDGBarrierBatchBegin
 {
 public:
-	FRDGBarrierBatchBegin(const FRDGPass* InPass, const TCHAR* InName, TOptional<ERHIPipeline> InOverridePipelineForEnd = {});
-	~FRDGBarrierBatchBegin();
+	FRDGBarrierBatchBegin(ERHIPipeline PipelinesToBegin, ERHIPipeline PipelinesToEnd, const TCHAR* DebugName, const FRDGPass* DebugPass);
+	FRDGBarrierBatchBegin(ERHIPipeline PipelinesToBegin, ERHIPipeline PipelinesToEnd, const TCHAR* DebugName, FRDGPassHandlesByPipeline DebugPasses);
 
-	/** Adds a resource transition into the batch. */
 	void AddTransition(FRDGParentResourceRef Resource, const FRHITransitionInfo& Info);
 
-	const FRHITransition* GetTransition() const
-	{
-		return Transition;
-	}
-
-	bool IsTransitionValid() const
-	{
-		return Transition != nullptr;
-	}
-
 	void SetUseCrossPipelineFence()
 	{
-		check(!bUseCrossPipelineFence);
-		bUseCrossPipelineFence = true;
+		TransitionFlags = ERHICreateTransitionFlags::None;
+		bTransitionNeeded = true;
 	}
 
-	void Submit(FRHIComputeCommandList& RHICmdList);
+	void Submit(FRHIComputeCommandList& RHICmdList, ERHIPipeline Pipeline);
+	void Submit(FRHIComputeCommandList& RHICmdList, ERHIPipeline Pipeline, FRDGTransitionQueue& TransitionsToBegin);
 
 private:
-	TOptional<ERHIPipeline> OverridePipelineToEnd;
-	bool bUseCrossPipelineFence = false;
-
-	/** The transition to store after submission. It is assigned back to null by the end batch. */
 	const FRHITransition* Transition = nullptr;
-
-	/** An array of asynchronous resource transitions to perform. */
 	TArray<FRHITransitionInfo, TInlineAllocator<1, SceneRenderingAllocator>> Transitions;
+	ERHICreateTransitionFlags TransitionFlags = ERHICreateTransitionFlags::NoFence;
+	bool bTransitionNeeded = false;
+
+	/** These pipelines masks are set at creation time and reset with each submission. */
+	ERHIPipeline PipelinesToBegin;
+	ERHIPipeline PipelinesToEnd;
 
 #if RDG_ENABLE_DEBUG
-	/** An array of RDG resources matching with the Transitions array. For debugging only. */
-	TArray<FRDGParentResource*, SceneRenderingAllocator> Resources;
+	FRDGPassHandlesByPipeline DebugPasses;
+	TArray<FRDGParentResource*, SceneRenderingAllocator> DebugResources;
+	const TCHAR* DebugName;
+	ERHIPipeline DebugPipelinesToBegin;
+	ERHIPipeline DebugPipelinesToEnd;
 #endif
 
 	friend class FRDGBarrierBatchEnd;
 	friend class FRDGBarrierValidation;
 };
 
-class RENDERCORE_API FRDGBarrierBatchEnd final : public FRDGBarrierBatch
+class RENDERCORE_API FRDGBarrierBatchEnd
 {
 public:
-	FRDGBarrierBatchEnd(const FRDGPass* InPass, const TCHAR* InName)
-		: FRDGBarrierBatch(InPass, InName)
-	{}
-
-	~FRDGBarrierBatchEnd();
-
-	void ReserveMemory(uint32 ExpectedDependencyCount);
+	FRDGBarrierBatchEnd() = default;
 
 	/** Inserts a dependency on a begin batch. A begin batch can be inserted into more than one end batch. */
 	void AddDependency(FRDGBarrierBatchBegin* BeginBatch);
 
-	void Submit(FRHIComputeCommandList& RHICmdList);
+	void Submit(FRHIComputeCommandList& RHICmdList, ERHIPipeline Pipeline);
 
 private:
 	TArray<FRDGBarrierBatchBegin*, TInlineAllocator<1, SceneRenderingAllocator>> Dependencies;
@@ -238,9 +228,9 @@ public:
 
 private:
 	FRDGBarrierBatchBegin& GetPrologueBarriersToBegin(FRDGAllocator& Allocator);
-	FRDGBarrierBatchEnd& GetPrologueBarriersToEnd(FRDGAllocator& Allocator);
 	FRDGBarrierBatchBegin& GetEpilogueBarriersToBeginForGraphics(FRDGAllocator& Allocator);
 	FRDGBarrierBatchBegin& GetEpilogueBarriersToBeginForAsyncCompute(FRDGAllocator& Allocator);
+	FRDGBarrierBatchBegin& GetEpilogueBarriersToBeginForAll(FRDGAllocator& Allocator);
 
 	FRDGBarrierBatchBegin& GetEpilogueBarriersToBeginFor(FRDGAllocator& Allocator, ERHIPipeline PipelineForEnd)
 	{
@@ -255,9 +245,15 @@ private:
 
 		case ERHIPipeline::AsyncCompute:
 			return GetEpilogueBarriersToBeginForAsyncCompute(Allocator);
+
+		case ERHIPipeline::All:
+			return GetEpilogueBarriersToBeginForAll(Allocator);
 		}
 	}
 
+	FRDGBarrierBatchEnd& GetPrologueBarriersToEnd(FRDGAllocator& Allocator);
+	FRDGBarrierBatchEnd& GetEpilogueBarriersToEnd(FRDGAllocator& Allocator);
+
 	//////////////////////////////////////////////////////////////////////////
 	//! User Methods to Override
 
@@ -355,13 +351,14 @@ private:
 	 */
 	TArray<FRHITexture*, SceneRenderingAllocator> TexturesToDiscard;
 
-	/** Barriers to begin / end prior to executing a pass. */
+	/** Split-barrier batches at various points of execution of the pass. */
 	FRDGBarrierBatchBegin* PrologueBarriersToBegin = nullptr;
 	FRDGBarrierBatchEnd* PrologueBarriersToEnd = nullptr;
-
-	/** Barriers to begin after executing a pass. */
 	FRDGBarrierBatchBegin* EpilogueBarriersToBeginForGraphics = nullptr;
 	FRDGBarrierBatchBegin* EpilogueBarriersToBeginForAsyncCompute = nullptr;
+	FRDGBarrierBatchBegin* EpilogueBarriersToBeginForAll = nullptr;
+	TArray<FRDGBarrierBatchBegin*> SharedEpilogueBarriersToBegin;
+	FRDGBarrierBatchEnd* EpilogueBarriersToEnd = nullptr;
 
 	EAsyncComputeBudget AsyncComputeBudget = EAsyncComputeBudget::EAll_4;
 
diff --git a/Engine/Source/Runtime/RenderCore/Public/RenderGraphResources.h b/Engine/Source/Runtime/RenderCore/Public/RenderGraphResources.h
--- a/Engine/Source/Runtime/RenderCore/Public/RenderGraphResources.h
+++ b/Engine/Source/Runtime/RenderCore/Public/RenderGraphResources.h
@@ -19,28 +19,38 @@ struct FRDGSubresourceState
 
 	FRDGSubresourceState() = default;
 
-	inline void SetPass(FRDGPassHandle PassHandle)
-	{
-		FirstPass = LastPass = PassHandle;
-	}
+	/** Initializes the first and last pass and the pipeline. Clears any other pass state. */
+	void SetPass(ERHIPipeline Pipeline, FRDGPassHandle PassHandle);
 
 	/** Finalizes the state at the end of the transition chain; keeps access intact. */
 	void Finalize();
 
+	/** Validates that the state is in a correct configuration for use. */
+	void Validate();
+
+	/** Returns whether the state is used by the pipeline. */
+	bool IsUsedBy(ERHIPipeline Pipeline) const;
+
+	/** Returns the last pass across either pipe. */
+	FRDGPassHandle GetLastPass() const;
+
+	/** Returns the first pass across either pipe. */
+	FRDGPassHandle GetFirstPass() const;
+
+	/** Returns the pipeline mask this state is used on. */
+	ERHIPipeline GetPipelines() const;
+
 	/** The last used access on the pass. */
 	ERHIAccess Access = ERHIAccess::Unknown;
 
 	/** The last used transition flags on the pass. */
 	EResourceTransitionFlags Flags = EResourceTransitionFlags::None;
 
-	/** The last used pass hardware pipeline. */
-	ERHIPipeline Pipeline = ERHIPipeline::Graphics;
-
 	/** The first pass in this state. */
-	FRDGPassHandle FirstPass;
+	FRDGPassHandlesByPipeline FirstPass;
 
 	/** The last pass in this state. */
-	FRDGPassHandle LastPass;
+	FRDGPassHandlesByPipeline LastPass;
 
 	/** The last no-UAV barrier to be used by this subresource. */
 	FRDGViewUniqueFilter NoUAVBarrierFilter;
diff --git a/Engine/Source/Runtime/RenderCore/Public/RenderGraphResources.inl b/Engine/Source/Runtime/RenderCore/Public/RenderGraphResources.inl
--- a/Engine/Source/Runtime/RenderCore/Public/RenderGraphResources.inl
+++ b/Engine/Source/Runtime/RenderCore/Public/RenderGraphResources.inl
@@ -4,13 +4,81 @@
 
 inline void FRDGSubresourceState::Finalize()
 {
-	ensureMsgf(Pipeline == ERHIPipeline::Graphics, TEXT("Resource should be on the graphics pipeline!"));
+	ensureMsgf(GetPipelines() == ERHIPipeline::Graphics, TEXT("Resource should be on the graphics pipeline!"));
 
 	const ERHIAccess LocalAccess = Access;
 	*this = {};
 	Access = LocalAccess;
 }
 
+inline void FRDGSubresourceState::SetPass(ERHIPipeline Pipeline, FRDGPassHandle PassHandle)
+{
+	FirstPass = {};
+	LastPass = {};
+	FirstPass[Pipeline] = PassHandle;
+	LastPass[Pipeline] = PassHandle;
+}
+
+inline void FRDGSubresourceState::Validate()
+{
+#if RDG_ENABLE_DEBUG
+	for (ERHIPipeline Pipeline : GetRHIPipelines())
+	{
+		checkf(FirstPass[Pipeline].IsValid() == LastPass[Pipeline].IsValid(), TEXT("Subresource state has unset first or last pass on '%s."), *GetRHIPipelineName(Pipeline));
+	}
+#endif
+}
+
+inline bool FRDGSubresourceState::IsUsedBy(ERHIPipeline Pipeline) const
+{
+	check(FirstPass[Pipeline].IsValid() == LastPass[Pipeline].IsValid());
+	return FirstPass[Pipeline].IsValid();
+}
+
+inline FRDGPassHandle FRDGSubresourceState::GetLastPass() const
+{
+	const FRDGPassHandle GraphicsPass = LastPass[ERHIPipeline::Graphics];
+	const FRDGPassHandle AsyncComputePass = LastPass[ERHIPipeline::AsyncCompute];
+
+	if (GraphicsPass.IsNull())
+	{
+		return AsyncComputePass;
+	}
+	
+	if (AsyncComputePass.IsNull())
+	{
+		return GraphicsPass;
+	}
+
+	return FMath::Max(GraphicsPass, AsyncComputePass);
+}
+
+inline FRDGPassHandle FRDGSubresourceState::GetFirstPass() const
+{
+	const FRDGPassHandle GraphicsPass = FirstPass[ERHIPipeline::Graphics];
+	const FRDGPassHandle AsyncComputePass = FirstPass[ERHIPipeline::AsyncCompute];
+
+	if (GraphicsPass.IsNull())
+	{
+		return AsyncComputePass;
+	}
+
+	if (AsyncComputePass.IsNull())
+	{
+		return GraphicsPass;
+	}
+
+	return FMath::Min(GraphicsPass, AsyncComputePass);
+}
+
+FORCEINLINE ERHIPipeline FRDGSubresourceState::GetPipelines() const
+{
+	ERHIPipeline Pipelines = ERHIPipeline::None;
+	Pipelines |= FirstPass[ERHIPipeline::Graphics].IsValid()     ? ERHIPipeline::Graphics     : ERHIPipeline::None;
+	Pipelines |= FirstPass[ERHIPipeline::AsyncCompute].IsValid() ? ERHIPipeline::AsyncCompute : ERHIPipeline::None;
+	return Pipelines;
+}
+
 inline FPooledRenderTargetDesc Translate(const FRDGTextureDesc& InDesc)
 {
 	check(InDesc.IsValid());
diff --git a/Engine/Source/Runtime/RenderCore/Public/RenderGraphValidation.h b/Engine/Source/Runtime/RenderCore/Public/RenderGraphValidation.h
--- a/Engine/Source/Runtime/RenderCore/Public/RenderGraphValidation.h
+++ b/Engine/Source/Runtime/RenderCore/Public/RenderGraphValidation.h
@@ -91,9 +91,6 @@ public:
 	/** Validates an end barrier batch just prior to submission to the command list. */
 	void ValidateBarrierBatchEnd(const FRDGPass* Pass, const FRDGBarrierBatchEnd& Batch);
 
-	/** Validates that all barrier batches were flushed at execution end. */
-	void ValidateExecuteEnd();
-
 private:
 	struct FResourceMap
 	{
diff --git a/Engine/Source/Runtime/Renderer/Private/CompositionLighting/CompositionLighting.cpp b/Engine/Source/Runtime/Renderer/Private/CompositionLighting/CompositionLighting.cpp
--- a/Engine/Source/Runtime/Renderer/Private/CompositionLighting/CompositionLighting.cpp
+++ b/Engine/Source/Runtime/Renderer/Private/CompositionLighting/CompositionLighting.cpp
@@ -129,14 +129,12 @@ static FSSAOCommonParameters GetSSAOCommonParameters(
 	FRDGBuilder& GraphBuilder,
 	const FViewInfo& View,
 	TRDGUniformBufferRef<FSceneTextureUniformParameters> SceneTexturesUniformBuffer,
-	TUniformBufferRef<FSceneTextureUniformParameters> SceneTexturesUniformBufferRHI,
 	uint32 Levels)
 {
 	const FSceneTextureParameters SceneTextureParameters = GetSceneTextureParameters(GraphBuilder, SceneTexturesUniformBuffer);
 
 	FSSAOCommonParameters CommonParameters;
 	CommonParameters.SceneTexturesUniformBuffer = SceneTexturesUniformBuffer;
-	CommonParameters.SceneTexturesUniformBufferRHI = SceneTexturesUniformBufferRHI;
 	CommonParameters.SceneTexturesViewport = FScreenPassTextureViewport(SceneTextureParameters.SceneDepthTexture, View.ViewRect);
 
 	CommonParameters.HZBInput = FScreenPassTexture(View.HZB);
@@ -159,7 +157,6 @@ FGTAOCommonParameters GetGTAOCommonParameters(
 	FRDGBuilder& GraphBuilder,
 	const FViewInfo& View,
 	TRDGUniformBufferRef<FSceneTextureUniformParameters> SceneTexturesUniformBuffer,
-	TUniformBufferRef<FSceneTextureUniformParameters> SceneTexturesUniformBufferRHI,
 	EGTAOType GTAOType
 	)
 {
@@ -167,7 +164,6 @@ FGTAOCommonParameters GetGTAOCommonParameters(
 
 	FGTAOCommonParameters CommonParameters;
 	CommonParameters.SceneTexturesUniformBuffer = SceneTexturesUniformBuffer;
-	CommonParameters.SceneTexturesUniformBufferRHI = SceneTexturesUniformBufferRHI;
 	CommonParameters.SceneTexturesViewport = FScreenPassTextureViewport(SceneTextureParameters.SceneDepthTexture, View.ViewRect);
 
 	CommonParameters.HZBInput = FScreenPassTexture(View.HZB);
@@ -525,8 +521,7 @@ void FCompositionLighting::ProcessBeforeBasePass(
 
 		if (bNeedSSAO)
 		{
-			TUniformBufferRef<FSceneTextureUniformParameters> SceneTexturesUniformBufferRHI = CreateSceneTextureUniformBuffer(GraphBuilder.RHICmdList, View.FeatureLevel, ESceneTextureSetupMode::SceneDepth);
-			FSSAOCommonParameters Parameters = GetSSAOCommonParameters(GraphBuilder, View, SceneTexturesUniformBuffer, SceneTexturesUniformBufferRHI, SSAOLevels);
+			FSSAOCommonParameters Parameters = GetSSAOCommonParameters(GraphBuilder, View, SceneTexturesUniformBuffer, SSAOLevels);
 			FScreenPassRenderTarget FinalTarget = FScreenPassRenderTarget(GraphBuilder.RegisterExternalTexture(SceneContext.ScreenSpaceAO, TEXT("AmbientOcclusionDirect")), View.ViewRect, ERenderTargetLoadAction::ENoAction);
 
 			AddPostProcessingAmbientOcclusion(
@@ -601,12 +596,10 @@ void FCompositionLighting::ProcessAfterBasePass(
 			{
 				const EGTAOType GTAOType = FSSAOHelper::GetGTAOPassType(View, SSAOLevels);
 
-				TUniformBufferRef<FSceneTextureUniformParameters> SceneTexturesUniformBufferRHI = CreateSceneTextureUniformBuffer(GraphBuilder.RHICmdList, View.FeatureLevel);
-
 				// If doing the Split GTAO method then we need to do the second part here.
 				if (GTAOType == EGTAOType::EAsyncHorizonSearch || GTAOType == EGTAOType::EAsyncCombinedSpatial)
 				{
-					FGTAOCommonParameters Parameters = GetGTAOCommonParameters(GraphBuilder, View, SceneTexturesUniformBuffer, SceneTexturesUniformBufferRHI, GTAOType);
+					FGTAOCommonParameters Parameters = GetGTAOCommonParameters(GraphBuilder, View, SceneTexturesUniformBuffer, GTAOType);
 
 					FScreenPassTexture GTAOHorizons(GraphBuilder.RegisterExternalTexture(SceneContext.ScreenSpaceGTAOHorizons, TEXT("GTAOHorizons")), Parameters.DownsampledViewRect);
 					AmbientOcclusion = AddPostProcessingGTAOPostAsync(GraphBuilder, View, Parameters, GTAOHorizons, FinalTarget);
@@ -619,12 +612,12 @@ void FCompositionLighting::ProcessAfterBasePass(
 				{
 					if (GTAOType == EGTAOType::ENonAsync)
 					{
-						FGTAOCommonParameters Parameters = GetGTAOCommonParameters(GraphBuilder, View, SceneTexturesUniformBuffer, SceneTexturesUniformBufferRHI, GTAOType);
+						FGTAOCommonParameters Parameters = GetGTAOCommonParameters(GraphBuilder, View, SceneTexturesUniformBuffer, GTAOType);
 						AmbientOcclusion = AddPostProcessingGTAOAllPasses(GraphBuilder, View, Parameters, FinalTarget);
 					}
 					else
 					{
-						FSSAOCommonParameters Parameters = GetSSAOCommonParameters(GraphBuilder, View, SceneTexturesUniformBuffer, SceneTexturesUniformBufferRHI, SSAOLevels);
+						FSSAOCommonParameters Parameters = GetSSAOCommonParameters(GraphBuilder, View, SceneTexturesUniformBuffer, SSAOLevels);
 						AmbientOcclusion = AddPostProcessingAmbientOcclusion(GraphBuilder, View, Parameters, FinalTarget);
 					}
 
@@ -672,8 +665,7 @@ void FCompositionLighting::ProcessAsyncSSAO(
 		{
 			RDG_GPU_MASK_SCOPE(GraphBuilder, View.GPUMask);
 
-			TUniformBufferRef<FSceneTextureUniformParameters> SceneTexturesUniformBufferRHI = CreateSceneTextureUniformBuffer(GraphBuilder.RHICmdList, View.FeatureLevel, ESceneTextureSetupMode::SceneDepth);
-			FGTAOCommonParameters CommonParameters = GetGTAOCommonParameters(GraphBuilder, View, SceneTexturesUniformBuffer, SceneTexturesUniformBufferRHI, GTAOType);
+			FGTAOCommonParameters CommonParameters = GetGTAOCommonParameters(GraphBuilder, View, SceneTexturesUniformBuffer, GTAOType);
 
 			FRDGTextureRef GTAOHorizonsTexture = nullptr;
 
diff --git a/Engine/Source/Runtime/Renderer/Private/CompositionLighting/PostProcessAmbientOcclusion.cpp b/Engine/Source/Runtime/Renderer/Private/CompositionLighting/PostProcessAmbientOcclusion.cpp
--- a/Engine/Source/Runtime/Renderer/Private/CompositionLighting/PostProcessAmbientOcclusion.cpp
+++ b/Engine/Source/Runtime/Renderer/Private/CompositionLighting/PostProcessAmbientOcclusion.cpp
@@ -668,7 +668,7 @@ END_SHADER_PARAMETER_STRUCT();
 
 BEGIN_SHADER_PARAMETER_STRUCT(FAmbientOcclusionParameters, )
 	SHADER_PARAMETER_STRUCT_REF(FViewUniformShaderParameters, View)
-	SHADER_PARAMETER_STRUCT_REF(FSceneTextureUniformParameters, SceneTextures)
+	SHADER_PARAMETER_RDG_UNIFORM_BUFFER(FSceneTextureUniformParameters, SceneTextures)
 
 	SHADER_PARAMETER_STRUCT_INCLUDE(FHZBParameters, HZBParameters)
 	SHADER_PARAMETER_STRUCT_INCLUDE(FSSAOShaderParameters, SSAOParameters)
@@ -808,7 +808,7 @@ FScreenPassTexture AddAmbientOcclusionPass(
 
 	FAmbientOcclusionParameters SharedParameters;
 	SharedParameters.View = View.ViewUniformBuffer;
-	SharedParameters.SceneTextures = CommonParameters.SceneTexturesUniformBufferRHI;
+	SharedParameters.SceneTextures = CommonParameters.SceneTexturesUniformBuffer;
 	SharedParameters.HZBParameters = GetHZBParameters(View, HZBInput, CommonParameters.SceneTexturesViewport.Extent, EAOTechnique::SSAO);
 	SharedParameters.SSAOParameters = GetSSAOShaderParameters(View, InputViewport, OutputViewport, CommonParameters.SceneTexturesViewport, EAOTechnique::SSAO);
 
@@ -970,7 +970,7 @@ public:
 
 	BEGIN_SHADER_PARAMETER_STRUCT(FParameters, )
 		SHADER_PARAMETER_STRUCT_REF(FViewUniformShaderParameters, View)
-		SHADER_PARAMETER_STRUCT_REF(FSceneTextureUniformParameters, SceneTextures)
+		SHADER_PARAMETER_RDG_UNIFORM_BUFFER(FSceneTextureUniformParameters, SceneTextures)
 		SHADER_PARAMETER_STRUCT_INCLUDE(FHZBParameters, HZBParameters)
 
 		SHADER_PARAMETER_STRUCT_INCLUDE(FSSAOShaderParameters, SSAOParameters)
@@ -1014,7 +1014,7 @@ FGTAOHorizonSearchOutputs AddGTAOHorizonSearchIntegratePass(
 
 	FGTAOHorizonSearchAndIntegrateCS::FParameters* PassParameters = GraphBuilder.AllocParameters<FGTAOHorizonSearchAndIntegrateCS::FParameters>();
 	PassParameters->View = View.ViewUniformBuffer;
-	PassParameters->SceneTextures = CommonParameters.SceneTexturesUniformBufferRHI;
+	PassParameters->SceneTextures = CommonParameters.SceneTexturesUniformBuffer;
 	PassParameters->HZBParameters = GetHZBParameters(View, HZBInput, SceneViewport.Extent, EAOTechnique::GTAO);
 	PassParameters->SSAOParameters = GetSSAOShaderParameters(View, SceneViewport, OutputViewport, CommonParameters.SceneTexturesViewport, EAOTechnique::GTAO);
 	PassParameters->GTAOParameters = GetGTAOShaderParameters(View, OutputViewport.Extent);
@@ -1153,7 +1153,7 @@ public:
 
 	BEGIN_SHADER_PARAMETER_STRUCT(FParameters, )
 		SHADER_PARAMETER_STRUCT_REF(FViewUniformShaderParameters, View)
-		SHADER_PARAMETER_STRUCT_REF(FSceneTextureUniformParameters, SceneTextures)
+		SHADER_PARAMETER_RDG_UNIFORM_BUFFER(FSceneTextureUniformParameters, SceneTextures)
 		SHADER_PARAMETER_STRUCT_INCLUDE(FHZBParameters, HZBParameters)
 
 		SHADER_PARAMETER_STRUCT_INCLUDE(FSSAOShaderParameters, SSAOParameters)
@@ -1181,7 +1181,7 @@ FGTAOHorizonSearchOutputs AddGTAOHorizonSearchPass(
 	FGTAOHorizonSearchCS::FParameters* PassParameters = GraphBuilder.AllocParameters<FGTAOHorizonSearchCS::FParameters>();
 
 	PassParameters->View = View.ViewUniformBuffer;
-	PassParameters->SceneTextures = CommonParameters.SceneTexturesUniformBufferRHI;
+	PassParameters->SceneTextures = CommonParameters.SceneTexturesUniformBuffer;
 	PassParameters->HZBParameters = GetHZBParameters(View, HZBInput, SceneViewport.Extent, EAOTechnique::GTAO);
 	PassParameters->SSAOParameters = GetSSAOShaderParameters(View, SceneViewport, OutputViewport, CommonParameters.SceneTexturesViewport, EAOTechnique::GTAO);
 	PassParameters->GTAOParameters = GetGTAOShaderParameters(View, OutputViewport.Extent);
@@ -1229,7 +1229,7 @@ public:
 
 	BEGIN_SHADER_PARAMETER_STRUCT(FParameters, )
 		SHADER_PARAMETER_STRUCT_REF(FViewUniformShaderParameters, View)
-		SHADER_PARAMETER_STRUCT_REF(FSceneTextureUniformParameters, SceneTextures)
+		SHADER_PARAMETER_RDG_UNIFORM_BUFFER(FSceneTextureUniformParameters, SceneTextures)
 
 		SHADER_PARAMETER_STRUCT_INCLUDE(FSSAOShaderParameters, SSAOParameters)
 		SHADER_PARAMETER_STRUCT_INCLUDE(FGTAOShaderParameters, GTAOParameters)
@@ -1303,7 +1303,7 @@ FGTAOTemporalOutputs AddGTAOTemporalPass(
 	FGTAOTemporalFilterCS::FParameters* PassParameters = GraphBuilder.AllocParameters<FGTAOTemporalFilterCS::FParameters>();
 
 	PassParameters->View = View.ViewUniformBuffer;
-	PassParameters->SceneTextures = CommonParameters.SceneTexturesUniformBufferRHI;
+	PassParameters->SceneTextures = CommonParameters.SceneTexturesUniformBuffer;
 	PassParameters->SSAOParameters = GetSSAOShaderParameters(View, InputViewport, OutputViewport, CommonParameters.SceneTexturesViewport, EAOTechnique::GTAO);
 	PassParameters->GTAOParameters = GetGTAOShaderParameters(View, OutputViewport.Extent);
 
@@ -1366,7 +1366,7 @@ public:
 
 	BEGIN_SHADER_PARAMETER_STRUCT(FParameters, )
 		SHADER_PARAMETER_STRUCT_REF(FViewUniformShaderParameters, View)
-		SHADER_PARAMETER_STRUCT_REF(FSceneTextureUniformParameters, SceneTextures)
+		SHADER_PARAMETER_RDG_UNIFORM_BUFFER(FSceneTextureUniformParameters, SceneTextures)
 
 		SHADER_PARAMETER_STRUCT_INCLUDE(FSSAOShaderParameters, SSAOParameters)
 
@@ -1413,7 +1413,7 @@ FScreenPassTexture AddGTAOSpatialFilter(
 
 	FGTAOSpatialFilterCS::FParameters* PassParameters = GraphBuilder.AllocParameters<FGTAOSpatialFilterCS::FParameters>();
 	PassParameters->View = View.ViewUniformBuffer;
-	PassParameters->SceneTextures = CommonParameters.SceneTexturesUniformBufferRHI;
+	PassParameters->SceneTextures = CommonParameters.SceneTexturesUniformBuffer;
 	PassParameters->SSAOParameters = GetSSAOShaderParameters(View, InputViewport, OutputViewport, CommonParameters.SceneTexturesViewport, EAOTechnique::GTAO);
 
 	PassParameters->GTAOSpatialFilterExtents = OutputViewport.Rect.Size();
diff --git a/Engine/Source/Runtime/Renderer/Private/CompositionLighting/PostProcessAmbientOcclusion.h b/Engine/Source/Runtime/Renderer/Private/CompositionLighting/PostProcessAmbientOcclusion.h
--- a/Engine/Source/Runtime/Renderer/Private/CompositionLighting/PostProcessAmbientOcclusion.h
+++ b/Engine/Source/Runtime/Renderer/Private/CompositionLighting/PostProcessAmbientOcclusion.h
@@ -152,7 +152,6 @@ FScreenPassTexture AddAmbientOcclusionFinalPass(
 
 struct FGTAOCommonParameters
 {
-	TUniformBufferRef<FSceneTextureUniformParameters> SceneTexturesUniformBufferRHI;
 	TRDGUniformBufferRef<FSceneTextureUniformParameters> SceneTexturesUniformBuffer = nullptr;
 	FScreenPassTextureViewport SceneTexturesViewport;
 	
