Fix for RDG ensure at startup.

#rnx

diff --git a/Engine/Source/Runtime/RenderCore/Public/RenderGraphResources.inl b/Engine/Source/Runtime/RenderCore/Public/RenderGraphResources.inl
--- a/Engine/Source/Runtime/RenderCore/Public/RenderGraphResources.inl
+++ b/Engine/Source/Runtime/RenderCore/Public/RenderGraphResources.inl
@@ -4,7 +4,7 @@
 
 inline void FRDGSubresourceState::Finalize()
 {
-	ensureMsgf(GetPipelines() == ERHIPipeline::Graphics, TEXT("Resource should be on the graphics pipeline!"));
+	ensureMsgf(!EnumHasAnyFlags(GetPipelines(), ERHIPipeline::AsyncCompute), TEXT("Resource should not be on the async compute pipeline!"));
 
 	const ERHIAccess LocalAccess = Access;
 	*this = {};
