Emit value of VirtualShaderDirectory in assertion.

#rnx

diff --git a/Engine/Source/Runtime/RenderCore/Private/ShaderCore.cpp b/Engine/Source/Runtime/RenderCore/Private/ShaderCore.cpp
--- a/Engine/Source/Runtime/RenderCore/Private/ShaderCore.cpp
+++ b/Engine/Source/Runtime/RenderCore/Private/ShaderCore.cpp
@@ -1350,9 +1350,13 @@ void AddShaderSourceDirectoryMapping(const FString& VirtualShaderDirectory, cons
 	}
 
 	// Do sanity checks of the virtual shader directory to map.
-	check(VirtualShaderDirectory.StartsWith(TEXT("/")));
-	check(!VirtualShaderDirectory.EndsWith(TEXT("/")));
-	check(!VirtualShaderDirectory.Contains(FString(TEXT("."))));
+	checkf(
+		VirtualShaderDirectory.StartsWith(TEXT("/")) &&
+		!VirtualShaderDirectory.EndsWith(TEXT("/")) &&
+		!VirtualShaderDirectory.Contains(FString(TEXT("."))),
+		TEXT("VirtualShaderDirectory = \"%s\""),
+		*VirtualShaderDirectory
+	);
 
 	// Detect collisions with any other mappings.
 	check(!GShaderSourceDirectoryMappings.Contains(VirtualShaderDirectory));
