RDG refactors in preparation for scene render targets refactor.
 - Added simple blackboard container.
 - Added support for checking if a resource has been produced within the graph.
 - Moved all RHI validation into user validation cpp file.
 - Fixed bug in graph producer compilation logic when switching to subresources.

diff --git a/Engine/Source/Runtime/RenderCore/Private/RenderGraphBlackboard.cpp b/Engine/Source/Runtime/RenderCore/Private/RenderGraphBlackboard.cpp
new file mode 100644
--- /dev/null
+++ b/Engine/Source/Runtime/RenderCore/Private/RenderGraphBlackboard.cpp
@@ -0,0 +1,28 @@
+// Copyright Epic Games, Inc. All Rights Reserved.
+
+#include "RenderGraphBlackboard.h"
+
+FString FRDGBlackboard::GetTypeName(const TCHAR* ClassName, const TCHAR* FileName, uint32 LineNumber)
+{
+	return FString::Printf(TEXT("%s %s %d"), ClassName, FileName, LineNumber);
+}
+
+uint32 FRDGBlackboard::AllocateIndex(FString&& TypeName)
+{
+	check(IsInRenderingThread());
+	static TMap<FString, uint32> StructMap;
+	static uint32 NextIndex = 0;
+
+	uint32 Result;
+	if (const uint32* FoundIndex = StructMap.Find(TypeName))
+	{
+		Result = *FoundIndex;
+	}
+	else
+	{
+		StructMap.Emplace(MoveTemp(TypeName), NextIndex);
+		Result = NextIndex;
+		NextIndex++;
+	}
+	return Result;
+}
\ No newline at end of file
diff --git a/Engine/Source/Runtime/RenderCore/Private/RenderGraphBuilder.cpp b/Engine/Source/Runtime/RenderCore/Private/RenderGraphBuilder.cpp
--- a/Engine/Source/Runtime/RenderCore/Private/RenderGraphBuilder.cpp
+++ b/Engine/Source/Runtime/RenderCore/Private/RenderGraphBuilder.cpp
@@ -384,6 +384,7 @@ const char* const FRDGBuilder::kDefaultUnaccountedCSVStat = "RDG_Pass";
 
 FRDGBuilder::FRDGBuilder(FRHICommandListImmediate& InRHICmdList, FRDGEventName InName, ERDGBuilderFlags InFlags)
 	: RHICmdList(InRHICmdList)
+	, Blackboard(Allocator)
 	, RHICmdListAsyncCompute(FRHICommandListExecutor::GetImmediateAsyncComputeCommandList())
 	, BuilderName(InName)
 	, BuilderFlags(InFlags)
@@ -394,6 +395,7 @@ FRDGBuilder::FRDGBuilder(FRHICommandListImmediate& InRHICmdList, FRDGEventName I
 	, GPUScopeStacks(RHICmdList, RHICmdListAsyncCompute)
 #endif
 #if RDG_ENABLE_DEBUG
+	, UserValidation(Allocator, BuilderFlags)
 	, BarrierValidation(&Passes, BuilderName)
 #endif
 {
@@ -446,13 +448,7 @@ FRDGTextureRef FRDGBuilder::RegisterExternalTexture(
 	ERenderTargetTexture RenderTargetTexture,
 	ERDGTextureFlags Flags)
 {
-#if RDG_ENABLE_DEBUG
-	checkf(Name, TEXT("Attempted to register external texture with NULL name."));
-	checkf(ExternalPooledTexture.IsValid(), TEXT("Attempted to register NULL external texture."));
-	checkf(ExternalPooledTexture->IsCompatibleWithRDG(), TEXT("Pooled render target %s is not a compatible type for RDG."), Name);
-	UserValidation.ExecuteGuard(TEXT("RegisterExternalTexture"), Name);
-#endif
-
+	IF_RDG_ENABLE_DEBUG(UserValidation.ValidateRegisterExternalTexture(ExternalPooledTexture, Name, RenderTargetTexture, Flags));
 	FRHITexture* ExternalTextureRHI = ExternalPooledTexture->GetRenderTargetItem().GetRHI(RenderTargetTexture);
 	IF_RDG_ENABLE_DEBUG(checkf(ExternalTextureRHI, TEXT("Attempted to register texture %s, but its RHI texture is null."), Name));
 
@@ -485,6 +481,7 @@ FRDGTextureRef FRDGBuilder::RegisterExternalTexture(
 		TEXT("Externally registered texture '%s' has known RDG state. This means the graph did not sanitize it correctly, or ")
 		TEXT("an IPooledRenderTarget reference was improperly held within a pass."), Texture->Name);
 
+	if (!EnumHasAnyFlags(BuilderFlags, ERDGBuilderFlags::SkipBarriers))
 	{
 		FRDGSubresourceState SubresourceState;
 		SubresourceState.Access = AccessInitial;
@@ -493,7 +490,7 @@ FRDGTextureRef FRDGBuilder::RegisterExternalTexture(
 
 	ExternalTextures.Add(Texture->GetRHIUnchecked(), Texture);
 
-	IF_RDG_ENABLE_DEBUG(UserValidation.ValidateCreateExternalTexture(Texture));
+	IF_RDG_ENABLE_DEBUG(UserValidation.ValidateRegisterExternalTexture(Texture));
 	return Texture;
 }
 
@@ -516,11 +513,7 @@ FRDGBufferRef FRDGBuilder::RegisterExternalBuffer(
 	const TCHAR* Name,
 	ERDGBufferFlags Flags)
 {
-#if RDG_ENABLE_DEBUG
-	checkf(Name, TEXT("Attempted to register external buffer with NULL name."));
-	checkf(ExternalPooledBuffer.IsValid(), TEXT("Attempted to register NULL external buffer."));
-	UserValidation.ExecuteGuard(TEXT("RegisterExternalBuffer"), Name);
-#endif
+	IF_RDG_ENABLE_DEBUG(UserValidation.ValidateRegisterExternalBuffer(ExternalPooledBuffer, Name, Flags));
 
 	if (FRDGBufferRef* FoundBufferPtr = ExternalBuffers.Find(ExternalPooledBuffer.GetReference()))
 	{
@@ -550,11 +543,15 @@ FRDGBufferRef FRDGBuilder::RegisterExternalBuffer(
 	checkf(BufferState.Access == ERHIAccess::Unknown,
 		TEXT("Externally registered buffer '%s' has known RDG state. This means the graph did not sanitize it correctly, or ")
 		TEXT("an FRDGPooledBuffer reference was improperly held within a pass."), Buffer->Name);
-	BufferState.Access = AccessInitial;
+
+	if (!EnumHasAnyFlags(BuilderFlags, ERDGBuilderFlags::SkipBarriers))
+	{
+		BufferState.Access = AccessInitial;
+	}
 
 	ExternalBuffers.Add(ExternalPooledBuffer, Buffer);
 
-	IF_RDG_ENABLE_DEBUG(UserValidation.ValidateCreateExternalBuffer(Buffer));
+	IF_RDG_ENABLE_DEBUG(UserValidation.ValidateRegisterExternalBuffer(Buffer));
 	return Buffer;
 }
 
@@ -652,7 +649,7 @@ void FRDGBuilder::Compile()
 				// The producer array needs to be at least as large as the pass state array.
 				if (bWholeProducers && !bWholePassState)
 				{
-					InitAsSubresources(LastProducers, Texture->Layout);
+					InitAsSubresources(LastProducers, Texture->Layout, GetWholeResource(LastProducers));
 				}
 
 				for (uint32 Index = 0, Count = LastProducers.Num(); Index < Count; ++Index)
@@ -1343,6 +1340,7 @@ void FRDGBuilder::Clear()
 	Textures.Clear();
 	Buffers.Clear();
 	UniformBuffers.Clear();
+	Blackboard.Clear();
 	Allocator.ReleaseAll();
 }
 
@@ -1395,6 +1393,7 @@ void FRDGBuilder::SetupPass(FRDGPass* Pass)
 		}
 
 		bPassUAVAccess |= EnumHasAnyFlags(Access, ERHIAccess::UAVMask);
+		Texture->bProduced |= IsWritableAccess(Access);
 	});
 
 	Pass->BufferStates.Reserve(PassParameters.GetBufferParameterCount());
@@ -1411,6 +1410,7 @@ void FRDGBuilder::SetupPass(FRDGPass* Pass)
 		PassState.State.SetPass(PassPipeline, PassHandle);
 
 		bPassUAVAccess |= EnumHasAnyFlags(Access, ERHIAccess::UAVMask);
+		Buffer->bProduced |= IsWritableAccess(Access);
 	});
 
 	Pass->bUAVAccess = bPassUAVAccess;
@@ -1707,7 +1707,7 @@ void FRDGBuilder::CollectPassBarriers(FRDGPassHandle PassHandle, FRDGPassHandle&
 		LastUntrackedPassHandle = PassHandle;
 	}
 
-	if (PassesWithEmptyParameters[PassHandle])
+	if (PassesWithEmptyParameters[PassHandle] || EnumHasAnyFlags(BuilderFlags, ERDGBuilderFlags::SkipBarriers))
 	{
 		return;
 	}
@@ -1729,7 +1729,7 @@ void FRDGBuilder::CollectPassBarriers(FRDGPassHandle PassHandle, FRDGPassHandle&
 
 void FRDGBuilder::AddEpilogueTransition(FRDGTextureRef Texture, FRDGPassHandle LastUntrackedPassHandle)
 {
-	if (!Texture->bLastOwner || Texture->bCulled || EnumHasAnyFlags(BuilderFlags, ERDGBuilderFlags::SkipBarriers))
+	if (!Texture->bLastOwner || Texture->bCulled)
 	{
 		return;
 	}
@@ -1954,10 +1954,7 @@ void FRDGBuilder::AddTransitionInternal(
 	FRDGPassHandle LastUntrackedPassHandle,
 	const FRHITransitionInfo& TransitionInfo)
 {
-	if (EnumHasAnyFlags(BuilderFlags, ERDGBuilderFlags::SkipBarriers))
-	{
-		return;
-	}
+	check(!EnumHasAnyFlags(BuilderFlags, ERDGBuilderFlags::SkipBarriers));
 
 	const ERHIPipeline Graphics = ERHIPipeline::Graphics;
 	const ERHIPipeline AsyncCompute = ERHIPipeline::AsyncCompute;
diff --git a/Engine/Source/Runtime/RenderCore/Private/RenderGraphResources.cpp b/Engine/Source/Runtime/RenderCore/Private/RenderGraphResources.cpp
--- a/Engine/Source/Runtime/RenderCore/Private/RenderGraphResources.cpp
+++ b/Engine/Source/Runtime/RenderCore/Private/RenderGraphResources.cpp
@@ -4,33 +4,12 @@
 #include "RenderGraphPass.h"
 #include "RenderGraphPrivate.h"
 
-#if RDG_ENABLE_DEBUG
-void FRDGResource::MarkResourceAsUsed()
-{
-	ValidateRHIAccess();
-	DebugData.bIsActuallyUsedByPass = true;
-}
-
-void FRDGUniformBuffer::MarkResourceAsUsed()
-{
-	FRDGResource::MarkResourceAsUsed();
-
-	// Individual resources can't be culled from a uniform buffer, so we have to mark them all as used.
-	ParameterStruct.Enumerate([](FRDGParameter Parameter)
-	{
-		if (FRDGResourceRef Resource = Parameter.GetAsResource())
-		{
-			Resource->MarkResourceAsUsed();
-		}
-	});
-}
-#endif
-
 FRDGParentResource::FRDGParentResource(const TCHAR* InName, const ERDGParentResourceType InType)
 	: FRDGResource(InName)
 	, Type(InType)
 	, bExternal(0)
 	, bExtracted(0)
+	, bProduced(0)
 	, bTransient(0)
 	, bLastOwner(1)
 	// Culling logic runs only when immediate mode is off.
@@ -38,6 +17,11 @@ FRDGParentResource::FRDGParentResource(const TCHAR* InName, const ERDGParentReso
 	, bUsedByAsyncComputePass(0)
 {}
 
+void FRDGParentResource::SetPassthroughRHI(FRHIResource* InResourceRHI)
+{
+	ResourceRHI = InResourceRHI;
+}
+
 bool FRDGSubresourceState::IsTransitionRequired(const FRDGSubresourceState& Previous, const FRDGSubresourceState& Next)
 {
 	check(Next.Access != ERHIAccess::Unknown);
diff --git a/Engine/Source/Runtime/RenderCore/Private/RenderGraphValidation.cpp b/Engine/Source/Runtime/RenderCore/Private/RenderGraphValidation.cpp
--- a/Engine/Source/Runtime/RenderCore/Private/RenderGraphValidation.cpp
+++ b/Engine/Source/Runtime/RenderCore/Private/RenderGraphValidation.cpp
@@ -61,8 +61,114 @@ const ERHIAccess AccessMaskComputeOrRaster = ERHIAccess::IndirectArgs;
 bool GRDGInExecutePassScope = false;
 } //! namespace
 
-FRDGUserValidation::FRDGUserValidation()
-	: ExpectedNumMarks(FMemStack::Get().GetNumMarks())
+struct FRDGResourceDebugData
+{
+	/** Boolean to track at runtime whether a resource is actually used by the lambda of a pass or not, to detect unnecessary resource dependencies on passes. */
+	bool bIsActuallyUsedByPass = false;
+
+	/** Boolean to track at pass execution whether the underlying RHI resource is allowed to be accessed. */
+	bool bAllowRHIAccess = false;
+};
+
+void FRDGResource::MarkResourceAsUsed()
+{
+	ValidateRHIAccess();
+
+	if (DebugData)
+	{
+		DebugData->bIsActuallyUsedByPass = true;
+	}
+}
+
+void FRDGResource::ValidateRHIAccess() const
+{
+	// Passthrough resources will not have debug data, since they are not tied to a graph instance.
+	if (DebugData)
+	{
+		checkf(DebugData->bAllowRHIAccess,
+			TEXT("Accessing the RHI resource of %s at this time is not allowed. If you hit this check in pass, ")
+			TEXT("that is due to this resource not being referenced in the parameters of your pass."),
+			Name);
+	}
+}
+
+FRDGResourceDebugData& FRDGResource::GetDebugData() const
+{
+	check(DebugData);
+	return *DebugData;
+}
+
+struct FRDGParentResourceDebugData
+{
+	/** Pointer towards the pass that is the first to produce it, for even more convenient error message. */
+	const FRDGPass* FirstProducer = nullptr;
+
+	/** Count the number of times it has been used by a pass (without culling). */
+	uint32 PassAccessCount = 0;
+
+	/** Tracks whether a resource has valid contents. This version is true by default for external resources which have prior unknown contents. */
+	bool bHasBeenProduced = false;
+
+	/** Tracks whether this resource was clobbered by the builder prior to use. */
+	bool bHasBeenClobbered = false;
+};
+
+FRDGParentResourceDebugData& FRDGParentResource::GetParentDebugData() const
+{
+	check(ParentDebugData);
+	return *ParentDebugData;
+}
+
+bool FRDGResource::IsPassthrough() const
+{
+	return DebugData == nullptr;
+}
+
+struct FRDGTextureDebugData
+{
+	/** Tracks whether a UAV has ever been allocated to catch when TexCreate_UAV was unneeded. */
+	bool bHasNeededUAV = false;
+
+	/** Tracks whether has ever been bound as a render target to catch when TexCreate_RenderTargetable was unneeded. */
+	bool bHasBeenBoundAsRenderTarget = false;
+};
+
+FRDGTextureDebugData& FRDGTexture::GetTextureDebugData() const
+{
+	check(TextureDebugData);
+	return *TextureDebugData;
+}
+
+struct FRDGBufferDebugData
+{
+	/** Tracks state changes in order of execution. */
+	TArray<TPair<FRDGPassHandle, FRDGSubresourceState>, SceneRenderingAllocator> States;
+};
+
+FRDGBufferDebugData& FRDGBuffer::GetBufferDebugData() const
+{
+	check(BufferDebugData);
+	return *BufferDebugData;
+}
+
+void FRDGUniformBuffer::MarkResourceAsUsed()
+{
+	FRDGResource::MarkResourceAsUsed();
+
+	// Individual resources can't be culled from a uniform buffer, so we have to mark them all as used.
+	ParameterStruct.Enumerate([](FRDGParameter Parameter)
+	{
+		if (FRDGResourceRef Resource = Parameter.GetAsResource())
+		{
+			Resource->MarkResourceAsUsed();
+		}
+	});
+}
+
+FRDGUserValidation::FRDGUserValidation(FRDGAllocator& InAllocator, ERDGBuilderFlags InBuilderFlags)
+	: Allocator(InAllocator)
+	, BuilderFlags(InBuilderFlags)
+	, ExpectedNumMarks(FMemStack::Get().GetNumMarks())
 {}
 
 FRDGUserValidation::~FRDGUserValidation()
@@ -81,10 +187,23 @@ void FRDGUserValidation::ExecuteGuard(const TCHAR* Operation, const TCHAR* Resou
 	MemStackGuard();
 }
 
-void FRDGUserValidation::ValidateCreateTexture(FRDGTextureRef Texture)
+void FRDGUserValidation::ValidateCreateResource(FRDGResourceRef Resource)
 {
 	MemStackGuard();
-	check(Texture);
+	check(Resource);
+	Resource->DebugData = Allocator.AllocObject<FRDGResourceDebugData>();
+}
+
+void FRDGUserValidation::ValidateCreateParentResource(FRDGParentResourceRef Resource)
+{
+	ValidateCreateResource(Resource);
+	Resource->ParentDebugData = Allocator.AllocObject<FRDGParentResourceDebugData>();
+}
+
+void FRDGUserValidation::ValidateCreateTexture(FRDGTextureRef Texture)
+{
+	ValidateCreateParentResource(Texture);
+	Texture->TextureDebugData = Allocator.AllocObject<FRDGTextureDebugData>();
 	if (GRDGDebug)
 	{
 		TrackedTextures.Add(Texture);
@@ -93,24 +212,197 @@ void FRDGUserValidation::ValidateCreateTexture(FRDGTextureRef Texture)
 
 void FRDGUserValidation::ValidateCreateBuffer(FRDGBufferRef Buffer)
 {
-	MemStackGuard();
-	check(Buffer);
+	ValidateCreateParentResource(Buffer);
+	Buffer->BufferDebugData = Allocator.AllocObject<FRDGBufferDebugData>();
 	if (GRDGDebug)
 	{
 		TrackedBuffers.Add(Buffer);
 	}
 }
 
-void FRDGUserValidation::ValidateCreateExternalTexture(FRDGTextureRef Texture)
+void FRDGUserValidation::ValidateCreateSRV(FRDGTextureSRVRef SRV)
+{
+	ValidateCreateResource(SRV);
+}
+
+void FRDGUserValidation::ValidateCreateSRV(FRDGBufferSRVRef SRV)
+{
+	ValidateCreateResource(SRV);
+}
+
+void FRDGUserValidation::ValidateCreateUAV(FRDGTextureUAVRef UAV)
+{
+	ValidateCreateResource(UAV);
+}
+
+void FRDGUserValidation::ValidateCreateUAV(FRDGBufferUAVRef UAV)
+{
+	ValidateCreateResource(UAV);
+}
+
+void FRDGUserValidation::ValidateCreateUniformBuffer(FRDGUniformBufferRef UniformBuffer)
+{
+	ValidateCreateResource(UniformBuffer);
+}
+
+void FRDGUserValidation::ValidateRegisterExternalTexture(
+	const TRefCountPtr<IPooledRenderTarget>& ExternalPooledTexture,
+	const TCHAR* Name,
+	ERenderTargetTexture RenderTargetTexture,
+	ERDGTextureFlags Flags)
+{
+	checkf(Name, TEXT("Attempted to register external texture with NULL name."));
+	checkf(ExternalPooledTexture.IsValid(), TEXT("Attempted to register NULL external texture."));
+	checkf(ExternalPooledTexture->IsCompatibleWithRDG(), TEXT("Pooled render target %s is not a compatible type for RDG."), Name);
+	ExecuteGuard(TEXT("RegisterExternalTexture"), Name);
+}
+
+void FRDGUserValidation::ValidateRegisterExternalBuffer(const TRefCountPtr<FRDGPooledBuffer>& ExternalPooledBuffer, const TCHAR* Name, ERDGBufferFlags Flags)
+{
+	checkf(Name, TEXT("Attempted to register external buffer with NULL name."));
+	checkf(ExternalPooledBuffer.IsValid(), TEXT("Attempted to register NULL external buffer."));
+	ExecuteGuard(TEXT("RegisterExternalBuffer"), Name);
+}
+
+void FRDGUserValidation::ValidateRegisterExternalTexture(FRDGTextureRef Texture)
 {
 	ValidateCreateTexture(Texture);
-	Texture->ParentDebugData.bHasBeenProduced = true;
+	Texture->GetParentDebugData().bHasBeenProduced = true;
 }
 
-void FRDGUserValidation::ValidateCreateExternalBuffer(FRDGBufferRef Buffer)
+void FRDGUserValidation::ValidateRegisterExternalBuffer(FRDGBufferRef Buffer)
 {
 	ValidateCreateBuffer(Buffer);
-	Buffer->ParentDebugData.bHasBeenProduced = true;
+	Buffer->GetParentDebugData().bHasBeenProduced = true;
+}
+
+void FRDGUserValidation::ValidateCreateTexture(const FRDGTextureDesc& Desc, const TCHAR* Name, ERDGTextureFlags Flags)
+{
+	checkf(Name, TEXT("Creating a texture requires a valid debug name."));
+	ExecuteGuard(TEXT("CreateTexture"), Name);
+
+	// Validate the pixel format.
+	checkf(Desc.Format != PF_Unknown, TEXT("Illegal to create texture %s with an invalid pixel format."), Name);
+	checkf(Desc.Format < PF_MAX, TEXT("Illegal to create texture %s with invalid FPooledRenderTargetDesc::Format."), Name);
+	checkf(GPixelFormats[Desc.Format].Supported,
+		TEXT("Failed to create texture %s with pixel format %s because it is not supported."), Name, GPixelFormats[Desc.Format].Name);
+	checkf(Desc.IsValid(), TEXT("Texture %s was created with an invalid descriptor."), Name);
+
+	const bool bCanHaveUAV = (Desc.Flags & TexCreate_UAV) > 0;
+	const bool bIsMSAA = Desc.NumSamples > 1;
+
+	// D3D11 doesn't allow creating a UAV on MSAA texture.
+	const bool bIsUAVForMSAATexture = bIsMSAA && bCanHaveUAV;
+	checkf(!bIsUAVForMSAATexture, TEXT("TexCreate_UAV is not allowed on MSAA texture %s."), Name);
+
+	checkf(!EnumHasAnyFlags(BuilderFlags, ERDGBuilderFlags::SkipBarriers), TEXT("Cannot create texture '%s' because RDG builder has 'SkipBarriers' flag."), Name);
+}
+
+void FRDGUserValidation::ValidateCreateBuffer(const FRDGBufferDesc& Desc, const TCHAR* Name, ERDGBufferFlags Flags)
+{
+	checkf(Name, TEXT("Creating a buffer requires a valid debug name."));
+	ExecuteGuard(TEXT("CreateBuffer"), Name);
+
+	checkf(Desc.GetTotalNumBytes() > 0, TEXT("Creating buffer '%s' is zero bytes in size."), Name);
+
+	const bool bIsByteAddress = (Desc.Usage & BUF_ByteAddressBuffer) == BUF_ByteAddressBuffer;
+
+	if (bIsByteAddress && Desc.UnderlyingType == FRDGBufferDesc::EUnderlyingType::StructuredBuffer)
+	{
+		checkf(Desc.BytesPerElement == 4, TEXT("Creating buffer '%s' as a structured buffer that is also byte addressable, BytesPerElement must be 4! Instead it is %d"), Name, Desc.BytesPerElement);
+	}
+
+	checkf(!EnumHasAnyFlags(BuilderFlags, ERDGBuilderFlags::SkipBarriers), TEXT("Cannot create buffer '%s' because RDG builder has 'SkipBarriers' flag."), Name);
+}
+
+void FRDGUserValidation::ValidateCreateSRV(const FRDGTextureSRVDesc& Desc)
+{
+	FRDGTextureRef Texture = Desc.Texture;
+	checkf(Texture, TEXT("Texture SRV created with a null texture."));
+	checkf(!Texture->IsPassthrough(), TEXT("Texture SRV created with passthrough texture '%s'."), Texture->Name);
+	ExecuteGuard(TEXT("CreateSRV"), Texture->Name);
+	checkf(Texture->Desc.Flags & TexCreate_ShaderResource, TEXT("Attempted to create SRV from texture %s which was not created with TexCreate_ShaderResource"), Desc.Texture->Name);
+
+	// Validate the pixel format if overridden by the SRV's descriptor.
+	if (Desc.Format == PF_X24_G8)
+	{
+		// PF_X24_G8 is a bit of mess in the RHI, used to read the stencil, but have varying BlockBytes.
+		checkf(Texture->Desc.Format == PF_DepthStencil, TEXT("PF_X24_G8 is only to read stencil from a PF_DepthStencil texture"));
+	}
+	else if (Desc.Format != PF_Unknown)
+	{
+		checkf(Desc.Format < PF_MAX, TEXT("Illegal to create SRV for texture %s with invalid FPooledRenderTargetDesc::Format."), Texture->Name);
+		checkf(GPixelFormats[Desc.Format].Supported, TEXT("Failed to create SRV for texture %s with pixel format %s because it is not supported."),
+			Texture->Name, GPixelFormats[Desc.Format].Name);
+
+		EPixelFormat ResourcePixelFormat = Texture->Desc.Format;
+
+		checkf(
+			GPixelFormats[Desc.Format].BlockBytes == GPixelFormats[ResourcePixelFormat].BlockBytes &&
+			GPixelFormats[Desc.Format].BlockSizeX == GPixelFormats[ResourcePixelFormat].BlockSizeX &&
+			GPixelFormats[Desc.Format].BlockSizeY == GPixelFormats[ResourcePixelFormat].BlockSizeY &&
+			GPixelFormats[Desc.Format].BlockSizeZ == GPixelFormats[ResourcePixelFormat].BlockSizeZ,
+			TEXT("Failed to create SRV for texture %s with pixel format %s because it does not match the byte size of the texture's pixel format %s."),
+			Texture->Name, GPixelFormats[Desc.Format].Name, GPixelFormats[ResourcePixelFormat].Name);
+	}
+
+	checkf((Desc.MipLevel + Desc.NumMipLevels) <= Texture->Desc.NumMips, TEXT("Failed to create SRV at mips %d-%d: the texture %s has only %d mip levels."),
+		Desc.MipLevel, (Desc.MipLevel + Desc.NumMipLevels), Texture->Name, Texture->Desc.NumMips);
+
+	checkf(Desc.MetaData != ERDGTextureMetaDataAccess::FMask || GRHISupportsExplicitFMask,
+		TEXT("Failed to create FMask SRV for texture %s because the current RHI doesn't support it. Be sure to gate the call with GRHISupportsExplicitFMask."),
+		Texture->Name);
+
+	checkf(Desc.MetaData != ERDGTextureMetaDataAccess::HTile || GRHISupportsExplicitHTile,
+		TEXT("Failed to create HTile SRV for texture %s because the current RHI doesn't support it. Be sure to gate the call with GRHISupportsExplicitHTile."),
+		Texture->Name);
+}
+
+void FRDGUserValidation::ValidateCreateSRV(const FRDGBufferSRVDesc& Desc)
+{
+	FRDGBufferRef Buffer = Desc.Buffer;
+	checkf(Buffer, TEXT("Buffer SRV created with a null buffer."));
+	checkf(!Buffer->IsPassthrough(), TEXT("Buffer SRV created with passthrough buffer '%s'."), Buffer->Name);
+	ExecuteGuard(TEXT("CreateSRV"), Buffer->Name);
+}
+
+void FRDGUserValidation::ValidateCreateUAV(const FRDGTextureUAVDesc& Desc)
+{
+	FRDGTextureRef Texture = Desc.Texture;
+
+	checkf(Texture, TEXT("Texture UAV created with a null texture."));
+	checkf(!Texture->IsPassthrough(), TEXT("Texture UAV created with passthrough texture '%s'."), Texture->Name);
+	ExecuteGuard(TEXT("CreateUAV"), Texture->Name);
+	checkf(Texture->Desc.Flags & TexCreate_UAV, TEXT("Attempted to create UAV from texture %s which was not created with TexCreate_UAV"), Texture->Name);
+	checkf(Desc.MipLevel < Texture->Desc.NumMips, TEXT("Failed to create UAV at mip %d: the texture %s has only %d mip levels."), Desc.MipLevel, Texture->Name, Texture->Desc.NumMips);
+}
+
+void FRDGUserValidation::ValidateCreateUAV(const FRDGBufferUAVDesc& Desc)
+{
+	FRDGBufferRef Buffer = Desc.Buffer;
+	checkf(Buffer, TEXT("Buffer UAV created with a null buffer."));
+	checkf(!Buffer->IsPassthrough(), TEXT("Buffer UAV created with passthrough buffer '%s'."), Buffer->Name);
+	ExecuteGuard(TEXT("CreateUAV"), Buffer->Name);
+}
+
+void FRDGUserValidation::ValidateCreateUniformBuffer(const void* ParameterStruct, const FShaderParametersMetadata* Metadata)
+{
+	check(Metadata);
+	const TCHAR* Name = Metadata->GetShaderVariableName();
+	checkf(ParameterStruct, TEXT("Uniform buffer '%s' created with null parameters."), Name);
+	ExecuteGuard(TEXT("CreateUniformBuffer"), Name);
+}
+
+void FRDGUserValidation::ValidateExtractTexture(FRDGTextureRef Texture, TRefCountPtr<IPooledRenderTarget>* OutTexturePtr)
+{
+	ValidateExtractResource(Texture);
+	checkf(OutTexturePtr, TEXT("Texture %s was extracted, but the output texture pointer is null."), Texture->Name);
+}
+
+void FRDGUserValidation::ValidateExtractBuffer(FRDGBufferRef Buffer, TRefCountPtr<FRDGPooledBuffer>* OutBufferPtr)
+{
+	ValidateExtractResource(Buffer);
+	checkf(OutBufferPtr, TEXT("Texture %s was extracted, but the output texture pointer is null."), Buffer->Name);
 }
 
 void FRDGUserValidation::ValidateExtractResource(FRDGParentResourceRef Resource)
@@ -118,7 +410,7 @@ void FRDGUserValidation::ValidateExtractResource(FRDGParentResourceRef Resource)
 	MemStackGuard();
 	check(Resource);
 
-	checkf(Resource->ParentDebugData.bHasBeenProduced,
+	checkf(Resource->bProduced || Resource->bExternal,
 		TEXT("Unable to queue the extraction of the resource %s because it has not been produced by any pass."),
 		Resource->Name);
 
@@ -126,36 +418,92 @@ void FRDGUserValidation::ValidateExtractResource(FRDGParentResourceRef Resource)
 	 *  emitting a 'produced but never used' warning. We don't have the history of registered
 	 *  resources to be able to emit a proper warning.
 	 */
-	Resource->ParentDebugData.PassAccessCount++;
+	Resource->GetParentDebugData().PassAccessCount++;
 }
 
 void FRDGUserValidation::RemoveUnusedWarning(FRDGParentResourceRef Resource)
 {
-	MemStackGuard();
 	check(Resource);
+	ExecuteGuard(TEXT("RemoveUnusedResourceWarning"), Resource->Name);
 
 	// Removes 'produced but not used' warning.
-	Resource->ParentDebugData.PassAccessCount++;
+	Resource->GetParentDebugData().PassAccessCount++;
 
 	// Removes 'not used' warning.
-	Resource->DebugData.bIsActuallyUsedByPass = true;
+	Resource->GetDebugData().bIsActuallyUsedByPass = true;
 }
 
 bool FRDGUserValidation::TryMarkForClobber(FRDGParentResourceRef Resource) const
 {
-	const bool bClobber = !Resource->ParentDebugData.bHasBeenClobbered && !Resource->bExternal && IsDebugAllowedForResource(Resource->Name);
+	check(Resource);
+	FRDGParentResourceDebugData& DebugData = Resource->GetParentDebugData();
+
+	const bool bClobber = !DebugData.bHasBeenClobbered && !Resource->bExternal && IsDebugAllowedForResource(Resource->Name);
 
 	if (bClobber)
 	{
-		Resource->ParentDebugData.bHasBeenClobbered = true;
+		DebugData.bHasBeenClobbered = true;
 	}
 
 	return bClobber;
 }
 
+void FRDGUserValidation::ValidateGetPooledTexture(FRDGTextureRef Texture) const
+{
+	check(Texture);
+	checkf(Texture->bExternal, TEXT("GetPooledTexture called on texture %s, but it is not external. Call PreallocateTexture or register as an external texture instead."), Texture->Name);
+}
+
+void FRDGUserValidation::ValidateGetPooledBuffer(FRDGBufferRef Buffer) const
+{
+	check(Buffer);
+	checkf(Buffer->bExternal, TEXT("GetPooledBuffer called on buffer %s, but it is not external. Call PreallocateBuffer or register as an external buffer instead."), Buffer->Name);
+}
+
+void FRDGUserValidation::ValidateSetTextureAccessFinal(FRDGTextureRef Texture, ERHIAccess AccessFinal)
+{
+	check(Texture);
+	check(AccessFinal != ERHIAccess::Unknown && IsValidAccess(AccessFinal));
+	checkf(Texture->bExternal || Texture->bExtracted, TEXT("Cannot set final access on nont-external texture '%s' unless it is first extracted."), Texture->Name);
+}
+
+void FRDGUserValidation::ValidateSetBufferAccessFinal(FRDGBufferRef Buffer, ERHIAccess AccessFinal)
+{
+	check(Buffer);
+	check(AccessFinal != ERHIAccess::Unknown && IsValidAccess(AccessFinal));
+	checkf(Buffer->bExternal || Buffer->bExtracted, TEXT("Cannot set final access on nont-external buffer '%s' unless it is first extracted."), Buffer->Name);
+}
+
+void FRDGUserValidation::ValidateAddPass(const FRDGEventName& Name, ERDGPassFlags Flags)
+{
+	ExecuteGuard(TEXT("AddPass"), Name.GetTCHAR());
+
+	checkf(!EnumHasAnyFlags(Flags, ERDGPassFlags::Copy | ERDGPassFlags::Compute | ERDGPassFlags::AsyncCompute | ERDGPassFlags::Raster),
+		TEXT("Pass %s may not specify any of the (Copy, Compute, AsyncCompute, Raster) flags, because it has no parameters. Use None instead."), Name.GetTCHAR());
+}
+
+void FRDGUserValidation::ValidateAddPass(const void* ParameterStruct, const FShaderParametersMetadata* Metadata, const FRDGEventName& Name, ERDGPassFlags Flags)
+{
+	checkf(ParameterStruct, TEXT("Pass '%s' created with null parameters."), Name.GetTCHAR());
+	ExecuteGuard(TEXT("AddPass"), Name.GetTCHAR());
+
+	checkf(EnumHasAnyFlags(Flags, ERDGPassFlags::CommandMask),
+		TEXT("Pass %s must specify at least one of the following flags: (Copy, Compute, AsyncCompute, Raster)"), Name.GetTCHAR());
+
+	checkf(!EnumHasAllFlags(Flags, ERDGPassFlags::Compute | ERDGPassFlags::AsyncCompute),
+		TEXT("Pass %s specified both Compute and AsyncCompute. They are mutually exclusive."), Name.GetTCHAR());
+
+	checkf(!EnumHasAllFlags(Flags, ERDGPassFlags::Raster | ERDGPassFlags::AsyncCompute),
+		TEXT("Pass %s specified both Raster and AsyncCompute. They are mutually exclusive."), Name.GetTCHAR());
+
+	checkf(!EnumHasAllFlags(Flags, ERDGPassFlags::SkipRenderPass) || EnumHasAllFlags(Flags, ERDGPassFlags::Raster),
+		TEXT("Pass %s specified SkipRenderPass without Raster. Only raster passes support this flag."));
+}
+
 void FRDGUserValidation::ValidateAddPass(const FRDGPass* Pass, bool bSkipPassAccessMarking)
 {
 	MemStackGuard();
+
 	const FRenderTargetBindingSlots* RenderTargetBindingSlots = nullptr;
 
 	// Pass flags are validated as early as possible by the builder in AddPass.
@@ -173,10 +521,9 @@ void FRDGUserValidation::ValidateAddPass(const FRDGPass* Pass, bool bSkipPassAcc
 	{
 		if (!bSkipPassAccessMarking)
 		{
-			auto& Debug = Resource->ParentDebugData;
-			if (!Debug.bHasBeenProduced)
+			auto& Debug = Resource->GetParentDebugData();
+			if (!Debug.FirstProducer)
 			{
-				Debug.bHasBeenProduced = true;
 				Debug.FirstProducer = Pass;
 			}
 			Debug.PassAccessCount++;
@@ -185,19 +532,17 @@ void FRDGUserValidation::ValidateAddPass(const FRDGPass* Pass, bool bSkipPassAcc
 
 	const auto MarkAsConsumed = [&] (FRDGParentResourceRef Resource)
 	{
-		auto& Debug = Resource->ParentDebugData;
-
-		ensureMsgf(Debug.bHasBeenProduced,
+		ensureMsgf(Resource->bProduced || Resource->bExternal,
 			TEXT("Pass %s has a read dependency on %s, but it was never written to."),
 			PassName, Resource->Name);
 
 		if (!bSkipPassAccessMarking)
 		{
-			Debug.PassAccessCount++;
+			Resource->GetParentDebugData().PassAccessCount++;
 		}
 	};
 
-	const auto CheckNotPassthrough = [&](FRDGResourceRef Resource)
+	const auto CheckNotPassthrough = [&](FRDGParentResourceRef Resource)
 	{
 		checkf(!Resource->IsPassthrough(), TEXT("Resource '%s' was created as a passthrough resource but is attached to pass '%s'."), Resource->Name, Pass->GetName());
 	};
@@ -222,7 +567,6 @@ void FRDGUserValidation::ValidateAddPass(const FRDGPass* Pass, bool bSkipPassAcc
 		{
 			if (FRDGViewRef View = Parameter.GetAsView())
 			{
-				CheckNotPassthrough(View);
 				CheckNotPassthrough(View->GetParent());
 			}
 		}
@@ -421,17 +765,6 @@ void FRDGUserValidation::ValidateAddPass(const FRDGPass* Pass, bool bSkipPassAcc
 						PassName, Texture->Name);
 
 					CheckNotPassthrough(Texture);
-					const bool bIsLoadAction = RenderTarget.GetLoadAction() == ERenderTargetLoadAction::ELoad;
-
-					/** Validate that load action is correct. We can only load contents if a pass previously produced something. */
-					{
-						const bool bIsLoadActionInvalid = bIsLoadAction && !Texture->ParentDebugData.bHasBeenProduced;
-						checkf(
-							!bIsLoadActionInvalid,
-							TEXT("Pass '%s' attempted to bind texture '%s' as a render target with the 'Load' action specified, but the texture has not been produced yet. The render target must use either 'Clear' or 'NoAction' action instead."),
-							PassName,
-							Texture->Name);
-					}
 
 					/** Validate that any previously produced texture contents are loaded. This occurs if the user failed to specify a load action
 					 *  on a texture that was produced by a previous pass, effectively losing that data. This can also happen if the user 're-uses'
@@ -440,8 +773,10 @@ void FRDGUserValidation::ValidateAddPass(const FRDGPass* Pass, bool bSkipPassAcc
 					 *  which are not actually managed by the render target pool and likely represent the frame buffer.
 					 */
 					{
+						const bool bIsLoadAction = RenderTarget.GetLoadAction() == ERenderTargetLoadAction::ELoad;
+
 						// Ignore external textures which are always marked as produced. We don't need to enforce this warning on them.
-						const bool bHasBeenProduced = Texture->ParentDebugData.bHasBeenProduced && !Texture->bExternal;
+						const bool bHasBeenProduced = Texture->bProduced && !Texture->bExternal;
 
 						// We only validate single-mip textures since we don't track production at the subresource level.
 						const bool bFailedToLoadProducedContent = !bIsLoadAction && bHasBeenProduced && Texture->Desc.NumMips == 1;
@@ -496,7 +831,7 @@ void FRDGUserValidation::ValidateExecuteEnd()
 		{
 			check(Resource->ReferenceCount == 0);
 
-			const auto& ParentDebugData = Resource->ParentDebugData;
+			const auto& ParentDebugData = Resource->GetParentDebugData();
 			const bool bProducedButNeverUsed = ParentDebugData.PassAccessCount == 1 && ParentDebugData.FirstProducer;
 
 			if (bProducedButNeverUsed)
@@ -513,20 +848,23 @@ void FRDGUserValidation::ValidateExecuteEnd()
 		{
 			ValidateResourceAtExecuteEnd(Texture);
 
-			const bool bHasBeenProducedByGraph = !Texture->bExternal && Texture->ParentDebugData.PassAccessCount > 0;
+			const auto& ParentDebugData = Texture->GetParentDebugData();
+			const auto& TextureDebugData = Texture->GetTextureDebugData();
+
+			const bool bHasBeenProducedByGraph = !Texture->bExternal && ParentDebugData.PassAccessCount > 0;
 
-			if (bHasBeenProducedByGraph && !Texture->TextureDebugData.bHasNeededUAV && EnumHasAnyFlags(Texture->Desc.Flags, TexCreate_UAV))
+			if (bHasBeenProducedByGraph && !TextureDebugData.bHasNeededUAV && EnumHasAnyFlags(Texture->Desc.Flags, TexCreate_UAV))
 			{
 				EmitRDGWarningf(
 					TEXT("Resource %s first produced by the pass %s had the TexCreate_UAV flag, but no UAV has been used."),
-					Texture->Name, Texture->ParentDebugData.FirstProducer->GetName());
+					Texture->Name, ParentDebugData.FirstProducer->GetName());
 			}
 
-			if (bHasBeenProducedByGraph && !Texture->TextureDebugData.bHasBeenBoundAsRenderTarget && EnumHasAnyFlags(Texture->Desc.Flags, TexCreate_RenderTargetable))
+			if (bHasBeenProducedByGraph && !TextureDebugData.bHasBeenBoundAsRenderTarget && EnumHasAnyFlags(Texture->Desc.Flags, TexCreate_RenderTargetable))
 			{
 				EmitRDGWarningf(
 					TEXT("Resource %s first produced by the pass %s had the TexCreate_RenderTargetable flag, but has never been bound as a render target of a pass."),
-					Texture->Name, Texture->ParentDebugData.FirstProducer->GetName());
+					Texture->Name, ParentDebugData.FirstProducer->GetName());
 			}
 		}
 
@@ -568,7 +906,7 @@ void FRDGUserValidation::ValidateExecutePassBegin(const FRDGPass* Pass)
 				if (FRDGTextureUAVRef UAV = Parameter.GetAsTextureUAV())
 				{
 					FRDGTextureRef Texture = UAV->Desc.Texture;
-					Texture->TextureDebugData.bHasNeededUAV = true;
+					Texture->GetTextureDebugData().bHasNeededUAV = true;
 				}
 			break;
 			case UBMT_RDG_TEXTURE_ACCESS:
@@ -579,11 +917,11 @@ void FRDGUserValidation::ValidateExecutePassBegin(const FRDGPass* Pass)
 					const ERHIAccess Access = TextureAccess.GetAccess();
 					if (EnumHasAnyFlags(Access, ERHIAccess::UAVMask))
 					{
-						Texture->TextureDebugData.bHasNeededUAV = true;
+						Texture->GetTextureDebugData().bHasNeededUAV = true;
 					}
 					if (EnumHasAnyFlags(Access, ERHIAccess::RTV | ERHIAccess::DSVRead | ERHIAccess::DSVWrite))
 					{
-						Texture->TextureDebugData.bHasBeenBoundAsRenderTarget = true;
+						Texture->GetTextureDebugData().bHasBeenBoundAsRenderTarget = true;
 					}
 					Texture->MarkResourceAsUsed();
 				}
@@ -602,13 +940,13 @@ void FRDGUserValidation::ValidateExecutePassBegin(const FRDGPass* Pass)
 				RenderTargets.Enumerate([&](FRenderTargetBinding RenderTarget)
 				{
 					FRDGTextureRef Texture = RenderTarget.GetTexture();
-					Texture->TextureDebugData.bHasBeenBoundAsRenderTarget = true;
+					Texture->GetTextureDebugData().bHasBeenBoundAsRenderTarget = true;
 					Texture->MarkResourceAsUsed();
 				});
 
 				if (FRDGTextureRef Texture = RenderTargets.DepthStencil.GetTexture())
 				{
-					Texture->TextureDebugData.bHasBeenBoundAsRenderTarget = true;
+					Texture->GetTextureDebugData().bHasBeenBoundAsRenderTarget = true;
 					Texture->MarkResourceAsUsed();
 				}
 			}
@@ -636,7 +974,7 @@ void FRDGUserValidation::ValidateExecutePassEnd(const FRDGPass* Pass)
 				if (FRDGResourceRef Resource = Parameter.GetAsResource())
 				{
 					TrackedResourceCount++;
-					UsedResourceCount += Resource->DebugData.bIsActuallyUsedByPass ? 1 : 0;
+					UsedResourceCount += Resource->GetDebugData().bIsActuallyUsedByPass ? 1 : 0;
 				}
 			}
 		});
@@ -653,7 +991,7 @@ void FRDGUserValidation::ValidateExecutePassEnd(const FRDGPass* Pass)
 				{
 					if (const FRDGResourceRef Resource = Parameter.GetAsResource())
 					{
-						if (!Resource->DebugData.bIsActuallyUsedByPass)
+						if (!Resource->GetDebugData().bIsActuallyUsedByPass)
 						{
 							WarningMessage += FString::Printf(TEXT("\n    %s"), Resource->Name);
 						}
@@ -671,7 +1009,7 @@ void FRDGUserValidation::ValidateExecutePassEnd(const FRDGPass* Pass)
 		{
 			if (FRDGResourceRef Resource = Parameter.GetAsResource())
 			{
-				Resource->DebugData.bIsActuallyUsedByPass = false;
+				Resource->GetDebugData().bIsActuallyUsedByPass = false;
 			}
 		}
 	});
@@ -687,7 +1025,7 @@ void FRDGUserValidation::SetAllowRHIAccess(const FRDGPass* Pass, bool bAllowAcce
 		{
 			if (FRDGResourceRef Resource = Parameter.GetAsResource())
 			{
-				Resource->DebugData.bAllowRHIAccess = bAllowAccess;
+				Resource->GetDebugData().bAllowRHIAccess = bAllowAccess;
 			}
 		}
 		else if (Parameter.IsRenderTargetBindingSlots())
@@ -696,17 +1034,17 @@ void FRDGUserValidation::SetAllowRHIAccess(const FRDGPass* Pass, bool bAllowAcce
 
 			RenderTargets.Enumerate([&](FRenderTargetBinding RenderTarget)
 			{
-				RenderTarget.GetTexture()->DebugData.bAllowRHIAccess = bAllowAccess;
+				RenderTarget.GetTexture()->GetDebugData().bAllowRHIAccess = bAllowAccess;
 
 				if (FRDGTextureRef ResolveTexture = RenderTarget.GetResolveTexture())
 				{
-					ResolveTexture->DebugData.bAllowRHIAccess = bAllowAccess;
+					ResolveTexture->GetDebugData().bAllowRHIAccess = bAllowAccess;
 				}
 			});
 
 			if (FRDGTextureRef Texture = RenderTargets.DepthStencil.GetTexture())
 			{
-				Texture->DebugData.bAllowRHIAccess = bAllowAccess;
+				Texture->GetDebugData().bAllowRHIAccess = bAllowAccess;
 			}
 		}
 	});
diff --git a/Engine/Source/Runtime/RenderCore/Public/RenderGraphBlackboard.h b/Engine/Source/Runtime/RenderCore/Public/RenderGraphBlackboard.h
new file mode 100644
--- /dev/null
+++ b/Engine/Source/Runtime/RenderCore/Public/RenderGraphBlackboard.h
@@ -0,0 +1,151 @@
+// Copyright Epic Games, Inc. All Rights Reserved.
+
+#pragma once
+
+#include "RenderGraphDefinitions.h"
+#include "RendererInterface.h"
+#include "Misc/GeneratedTypeName.h"
+
+class FRDGBlackboard;
+
+/** Declares a struct for use by the RDG blackboard. */
+#define RDG_REGISTER_BLACKBOARD_STRUCT(StructType)										\
+	template <>																			\
+	inline FString FRDGBlackboard::GetTypeName<StructType>()							\
+	{																					\
+		return GetTypeName(TEXT(#StructType), TEXT(__FILE__), __LINE__);				\
+	}
+
+/** The blackboard is a map of struct instances with a lifetime tied to a render graph allocator. It is designed
+ *  to solve cases where explicit marshaling of immutable data is undesirable. Structures are created once and the mutable
+ *  reference is returned. Only the immutable version is accessible from the blackboard. This constraint on mutability
+ *  is to discourage relying entirely on the blackboard. The mutable version should be marshaled around if needed.
+ *  Good candidates for the blackboard would be data that is created once and immutably fetched across the entire
+ *  renderer pipeline, where marshaling would create more maintenance burden than benefit. More constrained data
+ *  structures should be marshaled through function calls instead.
+ *
+ *  Example of Usage:
+ *
+ *	class FMyStruct
+ *	{
+ *	public:
+ *		FRDGTextureRef TextureA = nullptr;
+ *		FRDGTextureRef TextureB = nullptr;
+ *		FRDGTextureRef TextureC = nullptr;
+ *	};
+ *
+ *	RDG_REGISTER_BLACKBOARD_STRUCT(FMyStruct);
+ *
+ *	static void InitStruct(FRDGBlackboard& GraphBlackboard)
+ *	{
+ *		auto& MyStruct = GraphBlackboard.Create<FMyStruct>();
+ *
+ *		//...
+ *	}
+ *
+ *	static void UseStruct(const FRDGBlackboard& GraphBlackboard)
+ *	{
+ *		const auto& MyStruct = GraphBlackboard.GetChecked<FMyStruct>();
+ *
+ *		//...
+ *	}
+ */
+class FRDGBlackboard
+{
+public:
+	FRDGBlackboard(FRDGAllocator& InAllocator)
+		: Allocator(InAllocator)
+	{}
+
+	/** Copy / Move not allowed. */
+	FRDGBlackboard(FRDGBlackboard&&) = delete;
+	FRDGBlackboard(const FRDGBlackboard&) = delete;
+
+	/** Creates a new instance of a struct. Asserts if one already existed. */
+	template <typename StructType, typename... ArgsType>
+	StructType& Create(ArgsType&&... Args)
+	{
+		using HelperStructType = TStruct<StructType>;
+
+		const int32 StructIndex = GetStructIndex<StructType>();
+		if (StructIndex >= Blackboard.Num())
+		{
+			Blackboard.SetNum(StructIndex + 1);
+		}
+
+		FStruct*& Result = Blackboard[StructIndex];
+		checkf(!Result, TEXT("RDGBlackboard duplicate Create called on struct '%s'. Only one Create call per struct is allowed."), GetGeneratedTypeName<StructType>());
+		Result = Allocator.AllocObject<HelperStructType>(Forward<ArgsType&&>(Args)...);
+		check(Result);
+		return static_cast<HelperStructType*>(Result)->Struct;
+	}
+
+	/** Gets an immutable instance of the struct. Returns null if not present in the blackboard. */
+	template <typename StructType>
+	const StructType* Get() const
+	{
+		using HelperStructType = TStruct<StructType>;
+
+		const int32 StructIndex = GetStructIndex<StructType>();
+		if (StructIndex < Blackboard.Num())
+		{
+			return &static_cast<const HelperStructType*>(Blackboard[StructIndex])->Struct;
+		}
+		return nullptr;
+	}
+
+	/** Gets an immutable instance of the struct. Asserts if not present in the blackboard. */
+	template <typename StructType>
+	const StructType& GetChecked() const
+	{
+		const StructType* Struct = Get<StructType>();
+		checkf(Struct, TEXT("RDGBlackboard Get failed to find instance of struct '%s' in the blackboard."), GetGeneratedTypeName<StructType>());
+		return *Struct;
+	}
+
+	void Clear()
+	{
+		Blackboard.Empty();
+	}
+
+private:
+	struct FStruct
+	{
+		virtual ~FStruct() = default;
+	};
+
+	template <typename StructType>
+	struct TStruct final : public FStruct
+	{
+		template <typename... TArgs>
+		FORCEINLINE TStruct(TArgs&&... Args)
+			: Struct(Forward<TArgs&&>(Args)...)
+		{}
+
+		StructType Struct;
+	};
+
+	template <typename StructType>
+	static FString GetTypeName()
+	{
+		// Forces the compiler to only evaluate the assert on a concrete type.
+		static_assert(sizeof(StructType) == 0, "Struct has not been registered with the RDG blackboard. Use RDG_REGISTER_BLACKBOARD_STRUCT to do this.");
+	}
+
+	static RENDERCORE_API FString GetTypeName(const TCHAR* ClassName, const TCHAR* FileName, uint32 LineNumber);
+	static RENDERCORE_API uint32 AllocateIndex(FString&& TypeName);
+
+	template <typename StructType>
+	static uint32 GetStructIndex()
+	{
+		static uint32 Index = UINT_MAX;
+		if (Index == UINT_MAX)
+		{
+			Index = AllocateIndex(GetTypeName<StructType>());
+		}
+		return Index;
+	}
+
+	FRDGAllocator& Allocator;
+	TArray<FStruct*, SceneRenderingAllocator> Blackboard;
+};
\ No newline at end of file
diff --git a/Engine/Source/Runtime/RenderCore/Public/RenderGraphBuilder.h b/Engine/Source/Runtime/RenderCore/Public/RenderGraphBuilder.h
--- a/Engine/Source/Runtime/RenderCore/Public/RenderGraphBuilder.h
+++ b/Engine/Source/Runtime/RenderCore/Public/RenderGraphBuilder.h
@@ -7,19 +7,12 @@
 #include "RenderGraphResources.h"
 #include "RenderGraphPass.h"
 #include "RenderGraphValidation.h"
+#include "RenderGraphBlackboard.h"
 #include "ShaderParameterMacros.h"
 #include "ProfilingDebugging/CsvProfiler.h"
 
 class FRDGLogFile;
 
-enum class ERDGBuilderFlags
-{
-	None = 0,
-
-	/** The builder will not perform any resource barriers. This is only allowed when all resources are registered. */
-	SkipBarriers = 1 << 0,
-};
-
 /** Use the render graph builder to build up a graph of passes and then call Execute() to process them. Resource barriers
  *  and lifetimes are derived from _RDG_ parameters in the pass parameter struct provided to each AddPass call. The resulting
  *  graph is compiled, culled, and executed in Execute(). The builder should be created on the stack and executed prior to
@@ -102,6 +95,16 @@ public:
 		return CreateUAV(FRDGBufferUAVDesc(Buffer, Format), Flags);
 	}
 
+	/** Creates a graph tracked uniform buffer which can be attached to passes. These uniform buffers require some care
+	 *  because they will bulk transition all resources. The graph will only transition resources which are not also
+	 *  bound for write access by the pass.
+	 */
+	template <typename ParameterStructType>
+	TRDGUniformBufferRef<ParameterStructType> CreateUniformBuffer(ParameterStructType* ParameterStruct);
+
+	//////////////////////////////////////////////////////////////////////////
+	// Allocation Methods
+
 	/** Allocates raw memory using an allocator tied to the lifetime of the graph. */
 	void* Alloc(uint32 SizeInBytes, uint32 AlignInBytes);
 
@@ -117,12 +120,7 @@ public:
 	template <typename ParameterStructType>
 	ParameterStructType* AllocParameters();
 
-	/** Creates a graph tracked uniform buffer which can be attached to passes. These uniform buffers require some care
-	 *  because they will bulk transition all resources. The graph will only transition resources which are not also
-	 *  bound for write access by the pass.
-	 */
-	template <typename ParameterStructType>
-	TRDGUniformBufferRef<ParameterStructType> CreateUniformBuffer(ParameterStructType* ParameterStruct);
+	//////////////////////////////////////////////////////////////////////////
 
 	/** Adds a lambda pass to the graph with an accompanied pass parameter struct.
 	 *
@@ -220,6 +218,9 @@ public:
 	/** The RHI command list used for the render graph. */
 	FRHICommandListImmediate& RHICmdList;
 
+	/** The blackboard used to hold common data tied to the graph lifetime. */
+	FRDGBlackboard Blackboard;
+
 private:
 	static const ERHIAccess kDefaultAccessInitial = ERHIAccess::Unknown;
 	static const ERHIAccess kDefaultAccessFinal = ERHIAccess::SRVMask;
diff --git a/Engine/Source/Runtime/RenderCore/Public/RenderGraphBuilder.inl b/Engine/Source/Runtime/RenderCore/Public/RenderGraphBuilder.inl
--- a/Engine/Source/Runtime/RenderCore/Public/RenderGraphBuilder.inl
+++ b/Engine/Source/Runtime/RenderCore/Public/RenderGraphBuilder.inl
@@ -25,29 +25,7 @@ inline FRDGTextureRef FRDGBuilder::CreateTexture(
 	const TCHAR* Name,
 	ERDGTextureFlags Flags)
 {
-#if RDG_ENABLE_DEBUG
-	{
-		checkf(Name, TEXT("Creating a texture requires a valid debug name."));
-		UserValidation.ExecuteGuard(TEXT("CreateTexture"), Name);
-
-		// Validate the pixel format.
-		checkf(Desc.Format != PF_Unknown, TEXT("Illegal to create texture %s with an invalid pixel format."), Name);
-		checkf(Desc.Format < PF_MAX, TEXT("Illegal to create texture %s with invalid FPooledRenderTargetDesc::Format."), Name);
-		checkf(GPixelFormats[Desc.Format].Supported,
-			TEXT("Failed to create texture %s with pixel format %s because it is not supported."), Name, GPixelFormats[Desc.Format].Name);
-		checkf(Desc.IsValid(), TEXT("Texture %s was created with an invalid descriptor."), Name);
-
-		const bool bCanHaveUAV = (Desc.Flags & TexCreate_UAV) > 0;
-		const bool bIsMSAA = Desc.NumSamples > 1;
-
-		// D3D11 doesn't allow creating a UAV on MSAA texture.
-		const bool bIsUAVForMSAATexture = bIsMSAA && bCanHaveUAV;
-		checkf(!bIsUAVForMSAATexture, TEXT("TexCreate_UAV is not allowed on MSAA texture %s."), Name);
-
-		checkf(!EnumHasAnyFlags(BuilderFlags, ERDGBuilderFlags::SkipBarriers), TEXT("Cannot create texture '%s' because RDG builder has 'SkipBarriers' flag."), Name);
-	}
-#endif
-
+	IF_RDG_ENABLE_DEBUG(UserValidation.ValidateCreateTexture(Desc, Name, Flags));
 	FRDGTextureRef Texture = Textures.Allocate(Allocator, Name, Desc, Flags, ERenderTargetTexture::ShaderResource);
 	IF_RDG_ENABLE_DEBUG(UserValidation.ValidateCreateTexture(Texture));
 	return Texture;
@@ -58,24 +36,7 @@ inline FRDGBufferRef FRDGBuilder::CreateBuffer(
 	const TCHAR* Name,
 	ERDGBufferFlags Flags)
 {
-#if RDG_ENABLE_DEBUG
-	{
-		checkf(Name, TEXT("Creating a buffer requires a valid debug name."));
-		UserValidation.ExecuteGuard(TEXT("CreateBuffer"), Name);
-
-		checkf(Desc.GetTotalNumBytes() > 0, TEXT("Creating buffer '%s' is zero bytes in size."), Name);
-
-		const bool bIsByteAddress = (Desc.Usage & BUF_ByteAddressBuffer) == BUF_ByteAddressBuffer;
-
-		if (bIsByteAddress && Desc.UnderlyingType == FRDGBufferDesc::EUnderlyingType::StructuredBuffer)
-		{
-			checkf(Desc.BytesPerElement == 4, TEXT("Creating buffer '%s' as a structured buffer that is also byte addressable, BytesPerElement must be 4! Instead it is %d"), Name, Desc.BytesPerElement);
-		}
-
-		checkf(!EnumHasAnyFlags(BuilderFlags, ERDGBuilderFlags::SkipBarriers), TEXT("Cannot create buffer '%s' because RDG builder has 'SkipBarriers' flag."), Name);
-	}
-#endif
-
+	IF_RDG_ENABLE_DEBUG(UserValidation.ValidateCreateBuffer(Desc, Name, Flags));
 	FRDGBufferRef Buffer = Buffers.Allocate(Allocator, Name, Desc, Flags);
 	IF_RDG_ENABLE_DEBUG(UserValidation.ValidateCreateBuffer(Buffer));
 	return Buffer;
@@ -83,97 +44,34 @@ inline FRDGBufferRef FRDGBuilder::CreateBuffer(
 
 inline FRDGTextureSRVRef FRDGBuilder::CreateSRV(const FRDGTextureSRVDesc& Desc)
 {
-	FRDGTextureRef Texture = Desc.Texture;
-
-#if RDG_ENABLE_DEBUG
-	{
-		checkf(Texture, TEXT("Texture SRV created with a null texture."));
-		checkf(!Texture->IsPassthrough(), TEXT("Texture SRV created with passthrough texture '%s'."), Texture->Name);
-		UserValidation.ExecuteGuard(TEXT("CreateSRV"), Texture->Name);
-		checkf(Texture->Desc.Flags & TexCreate_ShaderResource, TEXT("Attempted to create SRV from texture %s which was not created with TexCreate_ShaderResource"), Desc.Texture->Name);
-		
-		// Validate the pixel format if overridden by the SRV's descriptor.
-		if (Desc.Format == PF_X24_G8)
-		{
-			// PF_X24_G8 is a bit of mess in the RHI, used to read the stencil, but have varying BlockBytes.
-			checkf(Texture->Desc.Format == PF_DepthStencil, TEXT("PF_X24_G8 is only to read stencil from a PF_DepthStencil texture"));
-		}
-		else if (Desc.Format != PF_Unknown)
-		{
-			checkf(Desc.Format < PF_MAX, TEXT("Illegal to create SRV for texture %s with invalid FPooledRenderTargetDesc::Format."), Texture->Name);
-			checkf(GPixelFormats[Desc.Format].Supported, TEXT("Failed to create SRV for texture %s with pixel format %s because it is not supported."),
-				Texture->Name, GPixelFormats[Desc.Format].Name);
-
-			EPixelFormat ResourcePixelFormat = Texture->Desc.Format;
-
-			checkf(
-				GPixelFormats[Desc.Format].BlockBytes == GPixelFormats[ResourcePixelFormat].BlockBytes &&
-				GPixelFormats[Desc.Format].BlockSizeX == GPixelFormats[ResourcePixelFormat].BlockSizeX &&
-				GPixelFormats[Desc.Format].BlockSizeY == GPixelFormats[ResourcePixelFormat].BlockSizeY &&
-				GPixelFormats[Desc.Format].BlockSizeZ == GPixelFormats[ResourcePixelFormat].BlockSizeZ,
-				TEXT("Failed to create SRV for texture %s with pixel format %s because it does not match the byte size of the texture's pixel format %s."),
-				Texture->Name, GPixelFormats[Desc.Format].Name, GPixelFormats[ResourcePixelFormat].Name);
-		}
-
-		checkf((Desc.MipLevel + Desc.NumMipLevels) <= Texture->Desc.NumMips, TEXT("Failed to create SRV at mips %d-%d: the texture %s has only %d mip levels."),
-			Desc.MipLevel, (Desc.MipLevel + Desc.NumMipLevels), Texture->Name, Texture->Desc.NumMips);
-
-		checkf(Desc.MetaData != ERDGTextureMetaDataAccess::FMask || GRHISupportsExplicitFMask,
-			TEXT("Failed to create FMask SRV for texture %s because the current RHI doesn't support it. Be sure to gate the call with GRHISupportsExplicitFMask."),
-			Texture->Name);
-
-		checkf(Desc.MetaData != ERDGTextureMetaDataAccess::HTile || GRHISupportsExplicitHTile,
-			TEXT("Failed to create HTile SRV for texture %s because the current RHI doesn't support it. Be sure to gate the call with GRHISupportsExplicitHTile."),
-			Texture->Name);
-	}
-#endif
-
-	return Views.Allocate<FRDGTextureSRV>(Allocator, Texture->Name, Desc);
+	IF_RDG_ENABLE_DEBUG(UserValidation.ValidateCreateSRV(Desc));
+	FRDGTextureSRVRef SRV = Views.Allocate<FRDGTextureSRV>(Allocator, Desc.Texture->Name, Desc);
+	IF_RDG_ENABLE_DEBUG(UserValidation.ValidateCreateSRV(SRV));
+	return SRV;
 }
 
 inline FRDGBufferSRVRef FRDGBuilder::CreateSRV(const FRDGBufferSRVDesc& Desc)
 {
-	FRDGBufferRef Buffer = Desc.Buffer;
-
-#if RDG_ENABLE_DEBUG
-	{
-		checkf(Buffer, TEXT("Buffer SRV created with a null buffer."));
-		UserValidation.ExecuteGuard(TEXT("CreateSRV"), Buffer->Name);
-	}
-#endif
-
-	return Views.Allocate<FRDGBufferSRV>(Allocator, Buffer->Name, Desc);
+	IF_RDG_ENABLE_DEBUG(UserValidation.ValidateCreateSRV(Desc));
+	FRDGBufferSRVRef SRV = Views.Allocate<FRDGBufferSRV>(Allocator, Desc.Buffer->Name, Desc);
+	IF_RDG_ENABLE_DEBUG(UserValidation.ValidateCreateSRV(SRV));
+	return SRV;
 }
 
 inline FRDGTextureUAVRef FRDGBuilder::CreateUAV(const FRDGTextureUAVDesc& Desc, ERDGUnorderedAccessViewFlags InFlags)
 {
-	FRDGTextureRef Texture = Desc.Texture;
-
-#if RDG_ENABLE_DEBUG
-	{
-		checkf(Texture, TEXT("Texture UAV created with a null texture."));
-		checkf(!Texture->IsPassthrough(), TEXT("Texture UAV created with passthrough texture '%s'."), Texture->Name);
-		UserValidation.ExecuteGuard(TEXT("CreateUAV"), Texture->Name);
-		checkf(Texture->Desc.Flags & TexCreate_UAV, TEXT("Attempted to create UAV from texture %s which was not created with TexCreate_UAV"), Texture->Name);
-		checkf(Desc.MipLevel < Texture->Desc.NumMips, TEXT("Failed to create UAV at mip %d: the texture %s has only %d mip levels."), Desc.MipLevel, Texture->Name, Texture->Desc.NumMips);
-	}
-#endif
-
-	return Views.Allocate<FRDGTextureUAV>(Allocator, Texture->Name, Desc, InFlags);
+	IF_RDG_ENABLE_DEBUG(UserValidation.ValidateCreateUAV(Desc));
+	FRDGTextureUAVRef UAV = Views.Allocate<FRDGTextureUAV>(Allocator, Desc.Texture->Name, Desc, InFlags);
+	IF_RDG_ENABLE_DEBUG(UserValidation.ValidateCreateUAV(UAV));
+	return UAV;
 }
 
 inline FRDGBufferUAVRef FRDGBuilder::CreateUAV(const FRDGBufferUAVDesc& Desc, ERDGUnorderedAccessViewFlags InFlags)
 {
-	FRDGBufferRef Buffer = Desc.Buffer;
-
-#if RDG_ENABLE_DEBUG
-	{
-		checkf(Buffer, TEXT("Buffer UAV created with a null buffer."));
-		UserValidation.ExecuteGuard(TEXT("CreateUAV"), Buffer->Name);
-	}
-#endif
-
-	return Views.Allocate<FRDGBufferUAV>(Allocator, Buffer->Name, Desc, InFlags);
+	IF_RDG_ENABLE_DEBUG(UserValidation.ValidateCreateUAV(Desc));
+	FRDGBufferUAVRef UAV = Views.Allocate<FRDGBufferUAV>(Allocator, Desc.Buffer->Name, Desc, InFlags);
+	IF_RDG_ENABLE_DEBUG(UserValidation.ValidateCreateUAV(UAV));
+	return UAV;
 }
 
 FORCEINLINE void* FRDGBuilder::Alloc(uint32 SizeInBytes, uint32 AlignInBytes)
@@ -209,16 +107,10 @@ FORCEINLINE FRDGSubresourceState* FRDGBuilder::AllocSubresource(const FRDGSubres
 template <typename ParameterStructType>
 TRDGUniformBufferRef<ParameterStructType> FRDGBuilder::CreateUniformBuffer(ParameterStructType* ParameterStruct)
 {
-	const TCHAR* Name = ParameterStructType::StaticStructMetadata.GetShaderVariableName();
-
-#if RDG_ENABLE_DEBUG
-	{
-		checkf(ParameterStruct, TEXT("Uniform buffer '%s' created with null parameters."), Name);
-		UserValidation.ExecuteGuard(TEXT("CreateUniformBuffer"), Name);
-	}
-#endif
-
-	return UniformBuffers.Allocate<TRDGUniformBuffer<ParameterStructType>>(Allocator, ParameterStruct, Name);
+	IF_RDG_ENABLE_DEBUG(UserValidation.ValidateCreateUniformBuffer(ParameterStruct, &ParameterStructType::StaticStructMetadata));
+	auto* UniformBuffer = UniformBuffers.Allocate<TRDGUniformBuffer<ParameterStructType>>(Allocator, ParameterStruct, ParameterStructType::StaticStructMetadata.GetShaderVariableName());
+	IF_RDG_ENABLE_DEBUG(UserValidation.ValidateCreateUniformBuffer(UniformBuffer));
+	return UniformBuffer;
 }
 
 template <typename ExecuteLambdaType>
@@ -229,14 +121,7 @@ FRDGPassRef FRDGBuilder::AddPass(
 {
 	using LambdaPassType = TRDGEmptyLambdaPass<ExecuteLambdaType>;
 
-#if RDG_ENABLE_DEBUG
-	{
-		UserValidation.ExecuteGuard(TEXT("AddPass"), Name.GetTCHAR());
-
-		checkf(!EnumHasAnyFlags(Flags, ERDGPassFlags::Copy | ERDGPassFlags::Compute | ERDGPassFlags::AsyncCompute | ERDGPassFlags::Raster),
-			TEXT("Pass %s may not specify any of the (Copy, Compute, AsyncCompute, Raster) flags, because it has no parameters. Use None instead."), Name.GetTCHAR());
-	}
-#endif
+	IF_RDG_ENABLE_DEBUG(UserValidation.ValidateAddPass(Name, Flags));
 
 	Flags |= ERDGPassFlags::NeverCull;
 
@@ -255,24 +140,7 @@ FRDGPassRef FRDGBuilder::AddPass(
 {
 	using LambdaPassType = TRDGLambdaPass<ParameterStructType, ExecuteLambdaType>;
 
-#if RDG_ENABLE_DEBUG
-	{
-		checkf(ParameterStruct, TEXT("Pass '%s' created with null parameters."), Name.GetTCHAR());
-		UserValidation.ExecuteGuard(TEXT("AddPass"), Name.GetTCHAR());
-
-		checkf(EnumHasAnyFlags(Flags, ERDGPassFlags::CommandMask),
-			TEXT("Pass %s must specify at least one of the following flags: (Copy, Compute, AsyncCompute, Raster)"), Name.GetTCHAR());
-
-		checkf(!EnumHasAllFlags(Flags, ERDGPassFlags::Compute | ERDGPassFlags::AsyncCompute),
-			TEXT("Pass %s specified both Compute and AsyncCompute. They are mutually exclusive."), Name.GetTCHAR());
-
-		checkf(!EnumHasAllFlags(Flags, ERDGPassFlags::Raster | ERDGPassFlags::AsyncCompute),
-			TEXT("Pass %s specified both Raster and AsyncCompute. They are mutually exclusive."), Name.GetTCHAR());
-
-		checkf(!EnumHasAllFlags(Flags, ERDGPassFlags::SkipRenderPass) || EnumHasAllFlags(Flags, ERDGPassFlags::Raster),
-			TEXT("Pass %s specified SkipRenderPass without Raster. Only raster passes support this flag."));
-	}
-#endif
+	IF_RDG_ENABLE_DEBUG(UserValidation.ValidateAddPass(ParameterStruct, ParametersMetadata, Name, Flags));
 
 	FRDGPass* Pass = Allocator.AllocObject<LambdaPassType>(
 		MoveTemp(Name),
@@ -313,10 +181,7 @@ inline void FRDGBuilder::QueueTextureExtraction(FRDGTextureRef Texture, TRefCoun
 
 inline void FRDGBuilder::QueueTextureExtraction(FRDGTextureRef Texture, TRefCountPtr<IPooledRenderTarget>* OutTexturePtr)
 {
-#if RDG_ENABLE_DEBUG
-	check(OutTexturePtr);
-	UserValidation.ValidateExtractResource(Texture);
-#endif
+	IF_RDG_ENABLE_DEBUG(UserValidation.ValidateExtractTexture(Texture, OutTexturePtr));
 
 	Texture->bExtracted = true;
 	Texture->bCulled = false;
@@ -330,10 +195,7 @@ inline void FRDGBuilder::QueueTextureExtraction(FRDGTextureRef Texture, TRefCoun
 
 inline void FRDGBuilder::QueueBufferExtraction(FRDGBufferRef Buffer, TRefCountPtr<FRDGPooledBuffer>* OutBufferPtr)
 {
-#if RDG_ENABLE_DEBUG
-	check(OutBufferPtr);
-	UserValidation.ValidateExtractResource(Buffer);
-#endif
+	IF_RDG_ENABLE_DEBUG(UserValidation.ValidateExtractBuffer(Buffer, OutBufferPtr));
 
 	Buffer->bExtracted = true;
 	Buffer->bCulled = false;
@@ -353,62 +215,36 @@ inline void FRDGBuilder::QueueBufferExtraction(FRDGBufferRef Buffer, TRefCountPt
 
 inline const TRefCountPtr<IPooledRenderTarget>& FRDGBuilder::GetPooledTexture(FRDGTextureRef Texture) const
 {
-#if RDG_ENABLE_DEBUG
-	check(Texture);
-	checkf(Texture->bExternal, TEXT("GetPooledTexture called on texture %s, but it is not external. Call PreallocateTexture or register as an external texture instead."), Texture->Name);
-#endif
-
+	IF_RDG_ENABLE_DEBUG(UserValidation.ValidateGetPooledTexture(Texture));
 	return Texture->Allocation;
 }
 
 inline const TRefCountPtr<FRDGPooledBuffer>& FRDGBuilder::GetPooledBuffer(FRDGBufferRef Buffer) const
 {
-#if RDG_ENABLE_DEBUG
-	check(Buffer);
-	checkf(Buffer->bExternal, TEXT("GetPooledBuffer called on buffer %s, but it is not external. Call PreallocateBuffer or register as an external buffer instead."), Buffer->Name);
-#endif
-
+	IF_RDG_ENABLE_DEBUG(UserValidation.ValidateGetPooledBuffer(Buffer));
 	return Buffer->Allocation;
 }
 
 inline void FRDGBuilder::SetTextureAccessFinal(FRDGTextureRef Texture, ERHIAccess AccessFinal)
 {
-#if RDG_ENABLE_DEBUG
-	check(Texture);
-	check(AccessFinal != ERHIAccess::Unknown && IsValidAccess(AccessFinal));
-	checkf(Texture->bExternal || Texture->bExtracted, TEXT("Cannot set final access on nont-external texture '%s' unless it is first extracted."), Texture->Name);
-#endif
-
+	IF_RDG_ENABLE_DEBUG(UserValidation.ValidateSetTextureAccessFinal(Texture, AccessFinal));
 	Texture->AccessFinal = AccessFinal;
 }
 
 inline void FRDGBuilder::SetBufferAccessFinal(FRDGBufferRef Buffer, ERHIAccess AccessFinal)
 {
-#if RDG_ENABLE_DEBUG
-	check(Buffer);
-	check(AccessFinal != ERHIAccess::Unknown && IsValidAccess(AccessFinal));
-	checkf(Buffer->bExternal || Buffer->bExtracted, TEXT("Cannot set final access on nont-external buffer '%s' unless it is first extracted."), Buffer->Name);
-#endif
-
+	IF_RDG_ENABLE_DEBUG(UserValidation.ValidateSetBufferAccessFinal(Buffer, AccessFinal));
 	Buffer->AccessFinal = AccessFinal;
 }
 
 inline void FRDGBuilder::RemoveUnusedTextureWarning(FRDGTextureRef Texture)
 {
-#if RDG_ENABLE_DEBUG
-	check(Texture);
-	UserValidation.ExecuteGuard(TEXT("RemoveUnusedTextureWarning"), Texture->Name);
-	UserValidation.RemoveUnusedWarning(Texture);
-#endif
+	IF_RDG_ENABLE_DEBUG(UserValidation.RemoveUnusedWarning(Texture));
 }
 
 inline void FRDGBuilder::RemoveUnusedBufferWarning(FRDGBufferRef Buffer)
 {
-#if RDG_ENABLE_DEBUG
-	check(Buffer);
-	UserValidation.ExecuteGuard(TEXT("RemoveUnusedBufferWarning"), Buffer->Name);
-	UserValidation.RemoveUnusedWarning(Buffer);
-#endif
+	IF_RDG_ENABLE_DEBUG(UserValidation.RemoveUnusedWarning(Buffer));
 }
 
 inline void FRDGBuilder::BeginEventScope(FRDGEventName&& ScopeName)
diff --git a/Engine/Source/Runtime/RenderCore/Public/RenderGraphDefinitions.h b/Engine/Source/Runtime/RenderCore/Public/RenderGraphDefinitions.h
--- a/Engine/Source/Runtime/RenderCore/Public/RenderGraphDefinitions.h
+++ b/Engine/Source/Runtime/RenderCore/Public/RenderGraphDefinitions.h
@@ -54,6 +54,16 @@
 
 /** ENUMS */
 
+enum class ERDGBuilderFlags
+{
+	None = 0,
+
+	/** The builder will not perform any resource barriers. This is only allowed when all resources are registered. This exists
+	 *  as a workaround for the RDG engine port, and will be deprecated immediately once the scaffolding code is removed.
+	 */
+	SkipBarriers = 1 << 0,
+};
+
 /** Flags to annotate a pass with when calling AddPass. */
 enum class ERDGPassFlags : uint8
 {
@@ -490,6 +500,8 @@ const TRDGHandle<ObjectType, IndexType> TRDGHandle<ObjectType, IndexType>::Null;
 
 struct FRDGTextureDesc;
 
+class FRDGBlackboard;
+
 class FRDGPassFlagsScopeGuard;
 class FRDGAsyncComputeBudgetScopeGuard;
 class FRDGEventScopeGuard;
diff --git a/Engine/Source/Runtime/RenderCore/Public/RenderGraphResources.h b/Engine/Source/Runtime/RenderCore/Public/RenderGraphResources.h
--- a/Engine/Source/Runtime/RenderCore/Public/RenderGraphResources.h
+++ b/Engine/Source/Runtime/RenderCore/Public/RenderGraphResources.h
@@ -83,7 +83,7 @@ public:
 
 	FRHIResource* GetRHI() const
 	{
-		ValidateRHIAccess();
+		IF_RDG_ENABLE_DEBUG(ValidateRHIAccess());
 		return ResourceRHI;
 	}
 
@@ -94,36 +94,6 @@ protected:
 		: Name(InName)
 	{}
 
-	/** Assigns this resource as a simple passthrough container for an RHI resource. */
-	void SetPassthroughRHI(FRHIResource* InResourceRHI)
-	{
-		ResourceRHI = InResourceRHI;
-#if RDG_ENABLE_DEBUG
-		DebugData.bAllowRHIAccess = true;
-		DebugData.bPassthrough = true;
-#endif
-	}
-
-	bool IsPassthrough() const
-	{
-#if RDG_ENABLE_DEBUG
-		return DebugData.bPassthrough;
-#else
-		return false;
-#endif
-	}
-
-	/** Verify that the RHI resource can be accessed at a pass execution. */
-	void ValidateRHIAccess() const
-	{
-#if RDG_ENABLE_DEBUG
-		checkf(DebugData.bAllowRHIAccess,
-			TEXT("Accessing the RHI resource of %s at this time is not allowed. If you hit this check in pass, ")
-			TEXT("that is due to this resource not being referenced in the parameters of your pass."),
-			Name);
-#endif
-	}
-
 	FRHIResource* GetRHIUnchecked() const
 	{
 		return ResourceRHI;
@@ -131,23 +101,15 @@ protected:
 
 	FRHIResource* ResourceRHI = nullptr;
 
-private:
 #if RDG_ENABLE_DEBUG
-	class FDebugData
-	{
-	private:
-		/** Boolean to track at runtime whether a resource is actually used by the lambda of a pass or not, to detect unnecessary resource dependencies on passes. */
-		bool bIsActuallyUsedByPass = false;
-
-		/** Boolean to track at pass execution whether the underlying RHI resource is allowed to be accessed. */
-		bool bAllowRHIAccess = false;
-
-		/** If true, the resource is not attached to any builder and exists as a dummy container for staging code to RDG. */
-		bool bPassthrough = false;
+	void ValidateRHIAccess() const;
+#endif
 
-		friend FRDGResource;
-		friend FRDGUserValidation;
-	} DebugData;
+private:
+#if RDG_ENABLE_DEBUG
+	struct FRDGResourceDebugData* DebugData = nullptr;
+	FRDGResourceDebugData& GetDebugData() const;
+	bool IsPassthrough() const;
 #endif
 
 	friend FRDGBuilder;
@@ -170,6 +132,8 @@ public:
 
 #if RDG_ENABLE_DEBUG
 	RENDERCORE_API void MarkResourceAsUsed() override;
+#else
+	inline         void MarkResourceAsUsed() {}
 #endif
 
 	//////////////////////////////////////////////////////////////////////////
@@ -241,11 +205,20 @@ public:
 	/** The type of this resource; useful for casting between types. */
 	const ERDGParentResourceType Type;
 
+	/** Whether this resource is externally registered with the graph (i.e. the user holds a reference to the underlying resource outside the graph). */
 	bool IsExternal() const
 	{
 		return bExternal;
 	}
 
+	/** Whether a prior pass added to the graph produced contents for this resource. External resources are not considered produced
+	 *  until used for a write operation. This is a union of all subresources, so any subresource write will set this to true.
+	 */
+	bool HasBeenProduced() const
+	{
+		return bProduced;
+	}
+
 protected:
 	FRDGParentResource(const TCHAR* InName, ERDGParentResourceType InType);
 
@@ -255,6 +228,9 @@ protected:
 	/** Whether this is an extracted resource. */
 	uint8 bExtracted : 1;
 
+	/** Whether any sub-resource has been used for write by a pass. */
+	uint8 bProduced : 1;
+
 	/** Whether this resource needs acquire / discard. */
 	uint8 bTransient : 1;
 
@@ -267,6 +243,9 @@ protected:
 	/** If true, the resource has been used on an async compute pass and may have async compute states. */
 	uint8 bUsedByAsyncComputePass : 1;
 
+	/** Assigns this resource as a simple passthrough container for an RHI resource. */
+	void SetPassthroughRHI(FRHIResource* InResourceRHI);
+
 private:
 	/** Number of references in passes and deferred queries. */
 	uint16 ReferenceCount = 0;
@@ -280,23 +259,8 @@ private:
 	FRDGPassHandle LastPass;
 
 #if RDG_ENABLE_DEBUG
-	class FParentDebugData
-	{
-	private:
-		/** Pointer towards the pass that is the first to produce it, for even more convenient error message. */
-		const FRDGPass* FirstProducer = nullptr;
-
-		/** Count the number of times it has been used by a pass (without culling). */
-		uint32 PassAccessCount = 0;
-
-		/** Tracks at wiring time if a resource has ever been produced by a pass, to error out early if accessing a resource that has not been produced. */
-		bool bHasBeenProduced = false;
-
-		/** Tracks whether this resource was clobbered by the builder prior to use. */
-		bool bHasBeenClobbered = false;
-
-		friend FRDGUserValidation;
-	} ParentDebugData;
+	struct FRDGParentResourceDebugData* ParentDebugData = nullptr;
+	FRDGParentResourceDebugData& GetParentDebugData() const;
 #endif
 
 	friend FRDGBuilder;
@@ -667,7 +631,7 @@ public:
 	/** Returns the allocated pooled render target. */
 	IPooledRenderTarget* GetPooledRenderTarget() const
 	{
-		ValidateRHIAccess();
+		IF_RDG_ENABLE_DEBUG(ValidateRHIAccess());
 		check(PooledRenderTarget);
 		return PooledRenderTarget;
 	}
@@ -760,18 +724,8 @@ private:
 	TRDGTextureSubresourceArray<FRDGPassHandle> LastProducers;
 
 #if RDG_ENABLE_DEBUG
-	class FTextureDebugData
-	{
-	private:
-		/** Tracks whether a UAV has ever been allocated to catch when TexCreate_UAV was unneeded. */
-		bool bHasNeededUAV = false;
-
-		/** Tracks whether has ever been bound as a render target to catch when TexCreate_RenderTargetable was unneeded. */
-		bool bHasBeenBoundAsRenderTarget = false;
-
-		friend FRDGUserValidation;
-		friend FRDGBarrierValidation;
-	} TextureDebugData;
+	struct FRDGTextureDebugData* TextureDebugData = nullptr;
+	FRDGTextureDebugData& GetTextureDebugData() const;
 #endif
 
 	friend FRDGBuilder;
@@ -1323,18 +1277,13 @@ private:
 	FRDGSubresourceState* MergeState = nullptr;
 
 #if RDG_ENABLE_DEBUG
-	class FBufferDebugData
-	{
-	private:
-		/** Tracks state changes in order of execution. */
-		TArray<TPair<FRDGPassHandle, FRDGSubresourceState>, SceneRenderingAllocator> States;
-
-		friend FRDGBarrierValidation;
-	} BufferDebugData;
+	struct FRDGBufferDebugData* BufferDebugData = nullptr;
+	FRDGBufferDebugData& GetBufferDebugData() const;
 #endif
 
 	friend FRDGBuilder;
 	friend FRDGBarrierValidation;
+	friend FRDGUserValidation;
 	friend FRDGBufferRegistry;
 	friend FRDGAllocator;
 };
diff --git a/Engine/Source/Runtime/RenderCore/Public/RenderGraphUtils.h b/Engine/Source/Runtime/RenderCore/Public/RenderGraphUtils.h
--- a/Engine/Source/Runtime/RenderCore/Public/RenderGraphUtils.h
+++ b/Engine/Source/Runtime/RenderCore/Public/RenderGraphUtils.h
@@ -11,6 +11,11 @@
 
 class FGlobalShaderMap;
 
+inline bool HasBeenProduced(FRDGParentResourceRef Resource)
+{
+	return Resource && Resource->HasBeenProduced();
+}
+
 /** Fetches the RHI texture from an RDG texture or null if the RDG texture is null. */
 inline FRHITexture* TryGetRHI(FRDGTextureRef Texture)
 {
diff --git a/Engine/Source/Runtime/RenderCore/Public/RenderGraphValidation.h b/Engine/Source/Runtime/RenderCore/Public/RenderGraphValidation.h
--- a/Engine/Source/Runtime/RenderCore/Public/RenderGraphValidation.h
+++ b/Engine/Source/Runtime/RenderCore/Public/RenderGraphValidation.h
@@ -23,29 +23,53 @@
 class RENDERCORE_API FRDGUserValidation final
 {
 public:
-	FRDGUserValidation();
+	FRDGUserValidation(FRDGAllocator& Allocator, ERDGBuilderFlags BuilderFlags);
 	FRDGUserValidation(const FRDGUserValidation&) = delete;
 	~FRDGUserValidation();
 
-	/** Validates that the graph has not executed yet. */
-	void ExecuteGuard(const TCHAR* Operation, const TCHAR* ResourceName);
+	/** Tracks and validates inputs into resource creation functions. */
+	void ValidateCreateTexture(const FRDGTextureDesc& Desc, const TCHAR* Name, ERDGTextureFlags Flags);
+	void ValidateCreateBuffer(const FRDGBufferDesc& Desc, const TCHAR* Name, ERDGBufferFlags Flags);
+	void ValidateCreateSRV(const FRDGTextureSRVDesc& Desc);
+	void ValidateCreateSRV(const FRDGBufferSRVDesc& Desc);
+	void ValidateCreateUAV(const FRDGTextureUAVDesc& Desc);
+	void ValidateCreateUAV(const FRDGBufferUAVDesc& Desc);
+	void ValidateCreateUniformBuffer(const void* ParameterStruct, const FShaderParametersMetadata* Metadata);
 
-	/** Tracks and validates the creation of a new resource in the graph. */
+	/** Tracks and validates the creation of a new externally created / registered resource instances. */
 	void ValidateCreateTexture(FRDGTextureRef Texture);
 	void ValidateCreateBuffer(FRDGBufferRef Buffer);
-
-	/** Tracks and validates the creation of a new externally registered resource. */
-	void ValidateCreateExternalTexture(FRDGTextureRef Texture);
-	void ValidateCreateExternalBuffer(FRDGBufferRef Buffer);
+	void ValidateCreateSRV(FRDGTextureSRVRef SRV);
+	void ValidateCreateSRV(FRDGBufferSRVRef SRV);
+	void ValidateCreateUAV(FRDGTextureUAVRef UAV);
+	void ValidateCreateUAV(FRDGBufferUAVRef UAV);
+	void ValidateCreateUniformBuffer(FRDGUniformBufferRef UniformBuffer);
+
+	void ValidateRegisterExternalTexture(
+		const TRefCountPtr<IPooledRenderTarget>& ExternalPooledTexture,
+		const TCHAR* Name,
+		ERenderTargetTexture RenderTargetTexture,
+		ERDGTextureFlags Flags);
+
+	void ValidateRegisterExternalBuffer(
+		const TRefCountPtr<FRDGPooledBuffer>& ExternalPooledBuffer,
+		const TCHAR* Name,
+		ERDGBufferFlags Flags);
+
+	void ValidateRegisterExternalTexture(FRDGTextureRef Texture);
+	void ValidateRegisterExternalBuffer(FRDGBufferRef Buffer);
 
 	/** Validates a resource extraction operation. */
-	void ValidateExtractResource(FRDGParentResourceRef Resource);
+	void ValidateExtractTexture(FRDGTextureRef Texture, TRefCountPtr<IPooledRenderTarget>* OutTexturePtr);
+	void ValidateExtractBuffer(FRDGBufferRef Buffer, TRefCountPtr<FRDGPooledBuffer>* OutBufferPtr);
 
 	/** Tracks and validates the addition of a new pass to the graph.
 	 *  @param bSkipPassAccessMarking Skips marking the pass as a producer or incrementing the pass access. Useful when
 	 *      the builder needs to inject a pass for debugging while preserving error messages and warnings for the original
 	 *      graph structure.
 	 */
+	void ValidateAddPass(const void* ParameterStruct, const FShaderParametersMetadata* Metadata, const FRDGEventName& Name, ERDGPassFlags Flags);
+	void ValidateAddPass(const FRDGEventName& Name, ERDGPassFlags Flags);
 	void ValidateAddPass(const FRDGPass* Pass, bool bSkipPassAccessMarking);
 
 	/** Validate pass state before and after execution. */
@@ -62,10 +86,20 @@ public:
 	/** Attempts to mark a resource for clobbering. If already marked, returns false.  */
 	bool TryMarkForClobber(FRDGParentResourceRef Resource) const;
 
+	void ValidateGetPooledTexture(FRDGTextureRef Texture) const;
+	void ValidateGetPooledBuffer(FRDGBufferRef Buffer) const;
+
+	void ValidateSetTextureAccessFinal(FRDGTextureRef Texture, ERHIAccess AccessFinal);
+	void ValidateSetBufferAccessFinal(FRDGBufferRef Buffer, ERHIAccess AccessFinal);
+
 	/** Traverses all resources in the pass and marks whether they are externally accessible by user pass implementations. */
 	static void SetAllowRHIAccess(const FRDGPass* Pass, bool bAllowAccess);
 
 private:
+	void ValidateCreateParentResource(FRDGParentResourceRef Resource);
+	void ValidateCreateResource(FRDGResourceRef Resource);
+	void ValidateExtractResource(FRDGParentResourceRef Resource);
+
 	/** List of tracked resources for validation prior to shutdown. */
 	TArray<FRDGTextureRef, SceneRenderingAllocator> TrackedTextures;
 	TArray<FRDGBufferRef, SceneRenderingAllocator> TrackedBuffers;
@@ -74,7 +108,10 @@ private:
 	bool bHasExecuted = false;
 
 	void MemStackGuard();
+	void ExecuteGuard(const TCHAR* Operation, const TCHAR* ResourceName);
 
+	FRDGAllocator& Allocator;
+	ERDGBuilderFlags BuilderFlags;
 	int32 ExpectedNumMarks;
 };
 
