Fix warning V670: The uninitialized class member 'Allocator' is used to initialize the 'Blackboard' member. Remember that members are initialized in the order of their declarations inside a class.

diff --git a/Engine/Source/Runtime/RenderCore/Public/RenderGraphBuilder.h b/Engine/Source/Runtime/RenderCore/Public/RenderGraphBuilder.h
--- a/Engine/Source/Runtime/RenderCore/Public/RenderGraphBuilder.h
+++ b/Engine/Source/Runtime/RenderCore/Public/RenderGraphBuilder.h
@@ -218,9 +218,6 @@ public:
 	/** The RHI command list used for the render graph. */
 	FRHICommandListImmediate& RHICmdList;
 
-	/** The blackboard used to hold common data tied to the graph lifetime. */
-	FRDGBlackboard Blackboard;
-
 private:
 	static const ERHIAccess kDefaultAccessInitial = ERHIAccess::Unknown;
 	static const ERHIAccess kDefaultAccessFinal = ERHIAccess::SRVMask;
@@ -229,6 +226,12 @@ private:
 	FRHIAsyncComputeCommandListImmediate& RHICmdListAsyncCompute;
 	FRDGAllocator Allocator;
 
+public:
+	/** The blackboard used to hold common data tied to the graph lifetime. */
+	FRDGBlackboard Blackboard;
+
+private
+
 	const FRDGEventName BuilderName;
 	const ERDGBuilderFlags BuilderFlags;
 
