Fix typo in RenderGraphBuilder.h

diff --git a/Engine/Source/Runtime/RenderCore/Public/RenderGraphBuilder.h b/Engine/Source/Runtime/RenderCore/Public/RenderGraphBuilder.h
--- a/Engine/Source/Runtime/RenderCore/Public/RenderGraphBuilder.h
+++ b/Engine/Source/Runtime/RenderCore/Public/RenderGraphBuilder.h
@@ -230,7 +230,7 @@ public:
 	/** The blackboard used to hold common data tied to the graph lifetime. */
 	FRDGBlackboard Blackboard;
 
-private
+private:
 
 	const FRDGEventName BuilderName;
 	const ERDGBuilderFlags BuilderFlags;
