Allow ray tracing unless specifically disabled by the project.

diff --git a/Engine/Source/Runtime/RenderCore/Private/RenderUtils.cpp b/Engine/Source/Runtime/RenderCore/Private/RenderUtils.cpp
--- a/Engine/Source/Runtime/RenderCore/Private/RenderUtils.cpp
+++ b/Engine/Source/Runtime/RenderCore/Private/RenderUtils.cpp
@@ -1258,9 +1258,9 @@ RENDERCORE_API void RenderUtilsInit()
 				}
 				else
 				{
-					GUseRayTracing = false;
+					GUseRayTracing = true;
 
-					UE_LOG(LogRendererCore, Log, TEXT("Ray tracing is disabled for the game. Reason: r.RayTracing=1, but user setting r.RayTracing.EnableInGame is not present (default false)."));
+					UE_LOG(LogRendererCore, Log, TEXT("Ray tracing is enabled for the game. Reason: r.RayTracing=1, and r.RayTracing.EnableInGame is not present (default true)."));
 				}
 			}
 
