Reworking RDG allocator to include container memory. This removes the dependency on the thread local memstack, and MemMark can no longer release RDG memory. All RDG memory is now released after execution.

diff --git a/Engine/Source/Runtime/RenderCore/Private/RenderGraphAllocator.cpp b/Engine/Source/Runtime/RenderCore/Private/RenderGraphAllocator.cpp
new file mode 100644
--- /dev/null
+++ b/Engine/Source/Runtime/RenderCore/Private/RenderGraphAllocator.cpp
@@ -0,0 +1,24 @@
+// Copyright Epic Games, Inc. All Rights Reserved.
+
+#include "RenderGraphAllocator.h"
+
+FRDGAllocator& FRDGAllocator::Get()
+{
+	static FRDGAllocator Instance;
+	return Instance;
+}
+
+FRDGAllocator::~FRDGAllocator()
+{
+	check(MemStack.IsEmpty());
+}
+
+void FRDGAllocator::ReleaseAll()
+{
+	for (int32 Index = TrackedAllocs.Num() - 1; Index >= 0; --Index)
+	{
+		TrackedAllocs[Index]->~FTrackedAlloc();
+	}
+	TrackedAllocs.Reset();
+	MemStack.Flush();
+}
\ No newline at end of file
diff --git a/Engine/Source/Runtime/RenderCore/Private/RenderGraphBuilder.cpp b/Engine/Source/Runtime/RenderCore/Private/RenderGraphBuilder.cpp
--- a/Engine/Source/Runtime/RenderCore/Private/RenderGraphBuilder.cpp
+++ b/Engine/Source/Runtime/RenderCore/Private/RenderGraphBuilder.cpp
@@ -384,8 +384,9 @@ const char* const FRDGBuilder::kDefaultUnaccountedCSVStat = "RDG_Pass";
 
 FRDGBuilder::FRDGBuilder(FRHICommandListImmediate& InRHICmdList, FRDGEventName InName, ERDGBuilderFlags InFlags)
 	: RHICmdList(InRHICmdList)
-	, RHICmdListAsyncCompute(FRHICommandListExecutor::GetImmediateAsyncComputeCommandList())
+	, Allocator(FRDGAllocator::Get())
 	, Blackboard(Allocator)
+	, RHICmdListAsyncCompute(FRHICommandListExecutor::GetImmediateAsyncComputeCommandList())
 	, BuilderName(InName)
 	, BuilderFlags(InFlags)
 #if RDG_CPU_SCOPES
@@ -702,7 +703,7 @@ void FRDGBuilder::Compile()
 	if (GRDGCullPasses)
 	{
 		SCOPED_NAMED_EVENT(FRDGBuilder_Compile_Cull_Passes, FColor::Emerald);
-		TArray<FRDGPassHandle, TInlineAllocator<32, SceneRenderingAllocator>> PassStack;
+		TArray<FRDGPassHandle, TInlineAllocator<32, FRDGArrayAllocator>> PassStack;
 		PassesToCull.Init(true, Passes.Num());
 
 	#if STATS
@@ -1037,7 +1038,7 @@ void FRDGBuilder::Compile()
 	{
 		SCOPED_NAMED_EVENT(FRDGBuilder_Compile_RenderPassMerge, FColor::Emerald);
 
-		TArray<FRDGPassHandle, SceneRenderingAllocator> PassesToMerge;
+		TArray<FRDGPassHandle, FRDGArrayAllocator> PassesToMerge;
 		FRDGPass* PrevPass = nullptr;
 		const FRenderTargetBindingSlots* PrevRenderTargets = nullptr;
 
diff --git a/Engine/Source/Runtime/RenderCore/Private/RenderGraphValidation.cpp b/Engine/Source/Runtime/RenderCore/Private/RenderGraphValidation.cpp
--- a/Engine/Source/Runtime/RenderCore/Private/RenderGraphValidation.cpp
+++ b/Engine/Source/Runtime/RenderCore/Private/RenderGraphValidation.cpp
@@ -55,10 +55,8 @@ const ERHIAccess AccessMaskCompute = ERHIAccess::SRVCompute | ERHIAccess::UAVCom
 const ERHIAccess AccessMaskRaster  = ERHIAccess::ResolveSrc | ERHIAccess::ResolveDst | ERHIAccess::DSVRead | ERHIAccess::DSVWrite | ERHIAccess::RTV | ERHIAccess::SRVGraphics | ERHIAccess::UAVGraphics | ERHIAccess::Present | ERHIAccess::VertexOrIndexBuffer;
 const ERHIAccess AccessMaskComputeOrRaster = ERHIAccess::IndirectArgs;
 
-/** Validates that we are only executing a single render graph instance in the callstack. Used to catch if a
- *  user creates a second FRDGBuilder instance inside of a pass that is executing.
- */
-bool GRDGInExecutePassScope = false;
+/** Validates that only one builder instance exists at any time. This is currently a requirement for state tracking and allocation lifetimes. */
+bool GRDGBuilderActive = false;
 } //! namespace
 
 struct FRDGResourceDebugData
@@ -142,7 +140,7 @@ FRDGTextureDebugData& FRDGTexture::GetTextureDebugData() const
 struct FRDGBufferDebugData
 {
 	/** Tracks state changes in order of execution. */
-	TArray<TPair<FRDGPassHandle, FRDGSubresourceState>, SceneRenderingAllocator> States;
+	TArray<TPair<FRDGPassHandle, FRDGSubresourceState>, FRDGArrayAllocator> States;
 };
 
 FRDGBufferDebugData& FRDGBuffer::GetBufferDebugData() const
@@ -169,11 +167,15 @@ FRDGUserValidation::FRDGUserValidation(FRDGAllocator& InAllocator, ERDGBuilderFl
 	: Allocator(InAllocator)
 	, BuilderFlags(InBuilderFlags)
 	, ExpectedNumMarks(FMemStack::Get().GetNumMarks())
-{}
+{
+	checkf(!GRDGBuilderActive, TEXT("Another FRDGBuilder already exists on the stack. Only one builder can be created at a time. This builder instance should be merged into the parent one."));
+	GRDGBuilderActive = true;
+}
 
 FRDGUserValidation::~FRDGUserValidation()
 {
 	checkf(bHasExecuted, TEXT("Render graph execution is required to ensure consistency with immediate mode."));
+	GRDGBuilderActive = false;
 }
 
 void FRDGUserValidation::MemStackGuard()
@@ -191,19 +193,19 @@ void FRDGUserValidation::ValidateCreateResource(FRDGResourceRef Resource)
 {
 	MemStackGuard();
 	check(Resource);
-	Resource->DebugData = Allocator.AllocObject<FRDGResourceDebugData>();
+	Resource->DebugData = Allocator.Alloc<FRDGResourceDebugData>();
 }
 
 void FRDGUserValidation::ValidateCreateParentResource(FRDGParentResourceRef Resource)
 {
 	ValidateCreateResource(Resource);
-	Resource->ParentDebugData = Allocator.AllocObject<FRDGParentResourceDebugData>();
+	Resource->ParentDebugData = Allocator.Alloc<FRDGParentResourceDebugData>();
 }
 
 void FRDGUserValidation::ValidateCreateTexture(FRDGTextureRef Texture)
 {
 	ValidateCreateParentResource(Texture);
-	Texture->TextureDebugData = Allocator.AllocObject<FRDGTextureDebugData>();
+	Texture->TextureDebugData = Allocator.Alloc<FRDGTextureDebugData>();
 	if (GRDGDebug)
 	{
 		TrackedTextures.Add(Texture);
@@ -213,7 +215,7 @@ void FRDGUserValidation::ValidateCreateTexture(FRDGTextureRef Texture)
 void FRDGUserValidation::ValidateCreateBuffer(FRDGBufferRef Buffer)
 {
 	ValidateCreateParentResource(Buffer);
-	Buffer->BufferDebugData = Allocator.AllocObject<FRDGBufferDebugData>();
+	Buffer->BufferDebugData = Allocator.Alloc<FRDGBufferDebugData>();
 	if (GRDGDebug)
 	{
 		TrackedBuffers.Add(Buffer);
@@ -897,9 +899,6 @@ void FRDGUserValidation::ValidateExecuteEnd()
 void FRDGUserValidation::ValidateExecutePassBegin(const FRDGPass* Pass)
 {
 	check(Pass);
-	checkf(!GRDGInExecutePassScope, TEXT("Render graph is being executed recursively. This usually means a separate FRDGBuilder instance was created inside of an executing pass."));
-
-	GRDGInExecutePassScope = true;
 
 	SetAllowRHIAccess(Pass, true);
 
@@ -1029,8 +1028,6 @@ void FRDGUserValidation::ValidateExecutePassEnd(const FRDGPass* Pass)
 			}
 		}
 	});
-
-	GRDGInExecutePassScope = false;
 }
 
 void FRDGUserValidation::SetAllowRHIAccess(const FRDGPass* Pass, bool bAllowAccess)
diff --git a/Engine/Source/Runtime/RenderCore/Public/RenderGraphAllocator.h b/Engine/Source/Runtime/RenderCore/Public/RenderGraphAllocator.h
new file mode 100644
--- /dev/null
+++ b/Engine/Source/Runtime/RenderCore/Public/RenderGraphAllocator.h
@@ -0,0 +1,172 @@
+// Copyright Epic Games, Inc. All Rights Reserved.
+
+#pragma once
+
+#include "CoreMinimal.h"
+#include "Misc/MemStack.h"
+
+/** Private allocator used by RDG to track its internal memory. All memory is released after RDG builder execution. */
+class RENDERCORE_API FRDGAllocator final
+{
+public:
+	~FRDGAllocator();
+
+	/** Allocates raw memory. */
+	FORCEINLINE void* Alloc(uint32 SizeInBytes, uint32 AlignInBytes)
+	{
+		return MemStack.Alloc(SizeInBytes, AlignInBytes);
+	}
+
+	/** Allocates an uninitialized type without destructor tracking. */
+	template <typename PODType>
+	FORCEINLINE PODType* AllocUninitialized()
+	{
+		return reinterpret_cast<PODType*>(Alloc(sizeof(PODType), alignof(PODType)));
+	}
+
+	/** Allocates and constructs an object and tracks it for destruction. */
+	template <typename T, typename... TArgs>
+	FORCEINLINE T* Alloc(TArgs&&... Args)
+	{
+		TTrackedAlloc<T>* TrackedAlloc = new(MemStack) TTrackedAlloc<T>(Forward<TArgs&&>(Args)...);
+		check(TrackedAlloc);
+		TrackedAllocs.Add(TrackedAlloc);
+		return &TrackedAlloc->Alloc;
+	}
+
+	/** Allocates a C++ object with no destructor tracking (dangerous!). */
+	template <typename T, typename... TArgs>
+	FORCEINLINE T* AllocNoDestruct(TArgs&&... Args)
+	{
+		return new (MemStack) T(Forward<TArgs&&>(Args)...);
+	}
+
+	FORCEINLINE int32 GetByteCount() const
+	{
+		return MemStack.GetByteCount();
+	}
+
+private:
+	static FRDGAllocator& Get();
+	FRDGAllocator() = default;
+	void ReleaseAll();
+
+	struct FTrackedAlloc
+	{
+		virtual ~FTrackedAlloc() = default;
+	};
+
+	template <typename T>
+	struct TTrackedAlloc final : FTrackedAlloc
+	{
+		template <typename... TArgs>
+		FORCEINLINE TTrackedAlloc(TArgs&&... Args)
+			: Alloc(Forward<TArgs&&>(Args)...)
+		{}
+
+		T Alloc;
+	};
+
+	FMemStackBase MemStack{0};
+	TArray<FTrackedAlloc*> TrackedAllocs;
+
+	template <uint32>
+	friend class TRDGArrayAllocator;
+	friend class FRDGBuilder;
+};
+
+/** A container allocator that allocates from a global RDG allocator instance. */
+template<uint32 Alignment = DEFAULT_ALIGNMENT>
+class TRDGArrayAllocator
+{
+public:
+	using SizeType = int32;
+
+	enum { NeedsElementType = true };
+	enum { RequireRangeCheck = true };
+
+	template<typename ElementType>
+	class ForElementType
+	{
+	public:
+		ForElementType() = default;
+
+		FORCEINLINE void MoveToEmpty(ForElementType& Other)
+		{
+			checkSlow(this != &Other);
+
+			Data = Other.Data;
+			Other.Data = nullptr;
+		}
+
+		FORCEINLINE ElementType* GetAllocation() const
+		{
+			return Data;
+		}
+
+		void ResizeAllocation(SizeType PreviousNumElements, SizeType NumElements, SIZE_T NumBytesPerElement)
+		{
+			void* OldData = Data;
+			if (NumElements)
+			{
+				// Allocate memory from the allocator.
+				const int32 AllocSize = (int32)(NumElements * NumBytesPerElement);
+				const int32 AllocAlignment = FMath::Max(Alignment, (uint32)alignof(ElementType));
+				Data = (ElementType*)FRDGAllocator::Get().Alloc(AllocSize, FMath::Max(AllocSize >= 16 ? (int32)16 : (int32)8, AllocAlignment));
+
+				// If the container previously held elements, copy them into the new allocation.
+				if (OldData && PreviousNumElements)
+				{
+					const SizeType NumCopiedElements = FMath::Min(NumElements, PreviousNumElements);
+					FMemory::Memcpy(Data, OldData, NumCopiedElements * NumBytesPerElement);
+				}
+			}
+		}
+		FORCEINLINE SizeType CalculateSlackReserve(SizeType NumElements, SIZE_T NumBytesPerElement) const
+		{
+			return DefaultCalculateSlackReserve(NumElements, NumBytesPerElement, false, Alignment);
+		}
+
+		FORCEINLINE SizeType CalculateSlackShrink(SizeType NumElements, SizeType NumAllocatedElements, SIZE_T NumBytesPerElement) const
+		{
+			return DefaultCalculateSlackShrink(NumElements, NumAllocatedElements, NumBytesPerElement, false, Alignment);
+		}
+
+		FORCEINLINE SizeType CalculateSlackGrow(SizeType NumElements, SizeType NumAllocatedElements, SIZE_T NumBytesPerElement) const
+		{
+			return DefaultCalculateSlackGrow(NumElements, NumAllocatedElements, NumBytesPerElement, false, Alignment);
+		}
+
+		FORCEINLINE SIZE_T GetAllocatedSize(SizeType NumAllocatedElements, SIZE_T NumBytesPerElement) const
+		{
+			return NumAllocatedElements * NumBytesPerElement;
+		}
+
+		bool HasAllocation() const
+		{
+			return !!Data;
+		}
+
+		SizeType GetInitialCapacity() const
+		{
+			return 0;
+		}
+
+	private:
+		ElementType* Data = nullptr;
+	};
+
+	typedef ForElementType<FScriptContainerElement> ForAnyElementType;
+};
+
+template <uint32 Alignment>
+struct TAllocatorTraits<TRDGArrayAllocator<Alignment>> : TAllocatorTraitsBase<TRDGArrayAllocator<Alignment>>
+{
+	enum { SupportsMove = true };
+	enum { IsZeroConstruct = true };
+};
+
+using FRDGArrayAllocator = TRDGArrayAllocator<>;
+using FRDGBitArrayAllocator = TInlineAllocator<4, FRDGArrayAllocator>;
+using FRDGSparseArrayAllocator = TSparseArrayAllocator<FRDGArrayAllocator, FRDGBitArrayAllocator>;
+using FRDGSetAllocator = TSetAllocator<FRDGSparseArrayAllocator, TInlineAllocator<1, FRDGBitArrayAllocator>>;
\ No newline at end of file
diff --git a/Engine/Source/Runtime/RenderCore/Public/RenderGraphBlackboard.h b/Engine/Source/Runtime/RenderCore/Public/RenderGraphBlackboard.h
--- a/Engine/Source/Runtime/RenderCore/Public/RenderGraphBlackboard.h
+++ b/Engine/Source/Runtime/RenderCore/Public/RenderGraphBlackboard.h
@@ -3,11 +3,8 @@
 #pragma once
 
 #include "RenderGraphDefinitions.h"
-#include "RendererInterface.h"
 #include "Misc/GeneratedTypeName.h"
 
-class FRDGBlackboard;
-
 /** Declares a struct for use by the RDG blackboard. */
 #define RDG_REGISTER_BLACKBOARD_STRUCT(StructType)										\
 	template <>																			\
@@ -53,14 +50,6 @@ class FRDGBlackboard;
 class FRDGBlackboard
 {
 public:
-	FRDGBlackboard(FRDGAllocator& InAllocator)
-		: Allocator(InAllocator)
-	{}
-
-	/** Copy / Move not allowed. */
-	FRDGBlackboard(FRDGBlackboard&&) = delete;
-	FRDGBlackboard(const FRDGBlackboard&) = delete;
-
 	/** Creates a new instance of a struct. Asserts if one already existed. */
 	template <typename StructType, typename... ArgsType>
 	StructType& Create(ArgsType&&... Args)
@@ -75,7 +64,7 @@ public:
 
 		FStruct*& Result = Blackboard[StructIndex];
 		checkf(!Result, TEXT("RDGBlackboard duplicate Create called on struct '%s'. Only one Create call per struct is allowed."), GetGeneratedTypeName<StructType>());
-		Result = Allocator.AllocObject<HelperStructType>(Forward<ArgsType&&>(Args)...);
+		Result = Allocator.Alloc<HelperStructType>(Forward<ArgsType&&>(Args)...);
 		check(Result);
 		return static_cast<HelperStructType*>(Result)->Struct;
 	}
@@ -103,12 +92,16 @@ public:
 		return *Struct;
 	}
 
+private:
+	FRDGBlackboard(FRDGAllocator& InAllocator)
+		: Allocator(InAllocator)
+	{}
+
 	void Clear()
 	{
 		Blackboard.Empty();
 	}
 
-private:
 	struct FStruct
 	{
 		virtual ~FStruct() = default;
@@ -147,5 +140,7 @@ private:
 	}
 
 	FRDGAllocator& Allocator;
-	TArray<FStruct*, SceneRenderingAllocator> Blackboard;
+	TArray<FStruct*, FRDGArrayAllocator> Blackboard;
+
+	friend class FRDGBuilder;
 };
\ No newline at end of file
diff --git a/Engine/Source/Runtime/RenderCore/Public/RenderGraphBuilder.h b/Engine/Source/Runtime/RenderCore/Public/RenderGraphBuilder.h
--- a/Engine/Source/Runtime/RenderCore/Public/RenderGraphBuilder.h
+++ b/Engine/Source/Runtime/RenderCore/Public/RenderGraphBuilder.h
@@ -218,19 +218,18 @@ public:
 	/** The RHI command list used for the render graph. */
 	FRHICommandListImmediate& RHICmdList;
 
+	/** The RDG allocator instance tied to the graph lifetime. All memory is released after Execute(). */
+	FRDGAllocator& Allocator;
+
+	/** The blackboard used to hold common data tied to the graph lifetime. */
+	FRDGBlackboard Blackboard;
+
 private:
 	static const ERHIAccess kDefaultAccessInitial = ERHIAccess::Unknown;
 	static const ERHIAccess kDefaultAccessFinal = ERHIAccess::SRVMask;
 	static const char* const kDefaultUnaccountedCSVStat;
 
 	FRHIAsyncComputeCommandListImmediate& RHICmdListAsyncCompute;
-	FRDGAllocator Allocator;
-
-public:
-	/** The blackboard used to hold common data tied to the graph lifetime. */
-	FRDGBlackboard Blackboard;
-
-private:
 
 	const FRDGEventName BuilderName;
 	const ERDGBuilderFlags BuilderFlags;
@@ -262,11 +261,11 @@ private:
 	FRDGPassBitArray PassesWithEmptyParameters;
 
 	/** Tracks external resources to their registered render graph counterparts for de-duplication. */
-	TSortedMap<FRHITexture*, FRDGTexture*, SceneRenderingAllocator> ExternalTextures;
-	TSortedMap<const FRDGPooledBuffer*, FRDGBuffer*, SceneRenderingAllocator> ExternalBuffers;
+	TSortedMap<FRHITexture*, FRDGTexture*, FRDGArrayAllocator> ExternalTextures;
+	TSortedMap<const FRDGPooledBuffer*, FRDGBuffer*, FRDGArrayAllocator> ExternalBuffers;
 
 	/** Map of barrier batches begun from more than one pipe. */
-	TMap<FRDGBarrierBatchBeginId, FRDGBarrierBatchBegin*, SceneRenderingSetAllocator> BarrierBatchMap;
+	TMap<FRDGBarrierBatchBeginId, FRDGBarrierBatchBegin*, FRDGSetAllocator> BarrierBatchMap;
 
 	/** The epilogue and prologue passes are sentinels that are used to simplify graph logic around barriers
 	 *  and traversal. The prologue pass is used exclusively for barriers before the graph executes, while the
@@ -278,8 +277,8 @@ private:
 	FRDGPass* EpiloguePass = nullptr;
 
 	/** Array of all requested resource extractions from the graph. */
-	TArray<TPair<FRDGTextureRef, TRefCountPtr<IPooledRenderTarget>*>, TInlineAllocator<4, SceneRenderingAllocator>> ExtractedTextures;
-	TArray<TPair<FRDGBufferRef, TRefCountPtr<FRDGPooledBuffer>*>, TInlineAllocator<4, SceneRenderingAllocator>> ExtractedBuffers;
+	TArray<TPair<FRDGTextureRef, TRefCountPtr<IPooledRenderTarget>*>, FRDGArrayAllocator> ExtractedTextures;
+	TArray<TPair<FRDGBufferRef, TRefCountPtr<FRDGPooledBuffer>*>, FRDGArrayAllocator> ExtractedBuffers;
 
 	/** Texture state used for intermediate operations. Held here to avoid re-allocating. */
 	FRDGTextureTransientSubresourceStateIndirect ScratchTextureState;
diff --git a/Engine/Source/Runtime/RenderCore/Public/RenderGraphBuilder.inl b/Engine/Source/Runtime/RenderCore/Public/RenderGraphBuilder.inl
--- a/Engine/Source/Runtime/RenderCore/Public/RenderGraphBuilder.inl
+++ b/Engine/Source/Runtime/RenderCore/Public/RenderGraphBuilder.inl
@@ -82,26 +82,24 @@ FORCEINLINE void* FRDGBuilder::Alloc(uint32 SizeInBytes, uint32 AlignInBytes)
 template <typename PODType>
 FORCEINLINE PODType* FRDGBuilder::AllocPOD()
 {
-	return Allocator.AllocPOD<PODType>();
+	return Allocator.AllocUninitialized<PODType>();
 }
 
 template <typename ObjectType, typename... TArgs>
 FORCEINLINE ObjectType* FRDGBuilder::AllocObject(TArgs&&... Args)
 {
-	return Allocator.AllocObject<ObjectType>(Forward<TArgs&&>(Args)...);
+	return Allocator.Alloc<ObjectType>(Forward<TArgs&&>(Args)...);
 }
 
 template <typename ParameterStructType>
 FORCEINLINE ParameterStructType* FRDGBuilder::AllocParameters()
 {
-	return Allocator.AllocObject<ParameterStructType>();
+	return Allocator.Alloc<ParameterStructType>();
 }
 
 FORCEINLINE FRDGSubresourceState* FRDGBuilder::AllocSubresource(const FRDGSubresourceState& Other)
 {
-	FRDGSubresourceState* State = Allocator.AllocPOD<FRDGSubresourceState>();
-	*State = Other;
-	return State;
+	return Allocator.AllocNoDestruct<FRDGSubresourceState>(Other);
 }
 
 template <typename ParameterStructType>
@@ -142,7 +140,7 @@ FRDGPassRef FRDGBuilder::AddPass(
 
 	IF_RDG_ENABLE_DEBUG(UserValidation.ValidateAddPass(ParameterStruct, ParametersMetadata, Name, Flags));
 
-	FRDGPass* Pass = Allocator.AllocObject<LambdaPassType>(
+	FRDGPass* Pass = Allocator.Alloc<LambdaPassType>(
 		MoveTemp(Name),
 		ParametersMetadata,
 		ParameterStruct,
diff --git a/Engine/Source/Runtime/RenderCore/Public/RenderGraphDefinitions.h b/Engine/Source/Runtime/RenderCore/Public/RenderGraphDefinitions.h
--- a/Engine/Source/Runtime/RenderCore/Public/RenderGraphDefinitions.h
+++ b/Engine/Source/Runtime/RenderCore/Public/RenderGraphDefinitions.h
@@ -2,7 +2,7 @@
 
 #pragma once
 
-#include "CoreMinimal.h"
+#include "RenderGraphAllocator.h"
 #include "ProfilingDebugging/RealtimeGPUProfiler.h"
 
 /** DEFINES */
@@ -220,91 +220,6 @@ inline int32 GetResourceTransitionPlaneForMetadataAccess(ERDGTextureMetaDataAcce
 	}
 }
 
-/** Simple C++ object allocator which tracks and destructs objects allocated using the MemStack allocator. */
-class FRDGAllocator final
-{
-public:
-	FRDGAllocator()
-		: MemStack(0)
-	{
-		TrackedAllocs.Reserve(16);
-	}
-
-	~FRDGAllocator()
-	{
-		checkSlow(MemStack.IsEmpty());
-	}
-
-	/** Allocates raw memory. */
-	FORCEINLINE void* Alloc(uint32 SizeInBytes, uint32 AlignInBytes)
-	{
-		return MemStack.Alloc(SizeInBytes, AlignInBytes);
-	}
-
-	/** Allocates POD memory without destructor tracking. */
-	template <typename PODType>
-	FORCEINLINE PODType* AllocPOD()
-	{
-		return reinterpret_cast<PODType*>(Alloc(sizeof(PODType), alignof(PODType)));
-	}
-
-	/** Allocates a C++ object with destructor tracking. */
-	template <typename ObjectType, typename... TArgs>
-	FORCEINLINE ObjectType* AllocObject(TArgs&&... Args)
-	{
-		TTrackedAlloc<ObjectType>* TrackedAlloc = new(MemStack) TTrackedAlloc<ObjectType>(Forward<TArgs&&>(Args)...);
-		check(TrackedAlloc);
-		TrackedAllocs.Add(TrackedAlloc);
-		return TrackedAlloc->Get();
-	}
-
-	/** Allocates a C++ object with no destructor tracking (dangerous!). */
-	template <typename ObjectType, typename... TArgs>
-	FORCEINLINE ObjectType* AllocNoDestruct(TArgs&&... Args)
-	{
-		return new (MemStack) ObjectType(Forward<TArgs&&>(Args)...);
-	}
-
-	/** Releases all allocations. */
-	void ReleaseAll()
-	{
-		for (int32 Index = TrackedAllocs.Num() - 1; Index >= 0; --Index)
-		{
-			TrackedAllocs[Index]->~FTrackedAlloc();
-		}
-		TrackedAllocs.Empty();
-		MemStack.Flush();
-	}
-
-	FORCEINLINE int32 GetByteCount() const
-	{
-		return MemStack.GetByteCount();
-	}
-
-private:
-	class FTrackedAlloc
-	{
-	public:
-		virtual ~FTrackedAlloc() = default;
-	};
-
-	template <typename ObjectType>
-	class TTrackedAlloc : public FTrackedAlloc
-	{
-	public:
-		template <typename... TArgs>
-		FORCEINLINE TTrackedAlloc(TArgs&&... Args) : Object(Forward<TArgs&&>(Args)...) {}
-
-		FORCEINLINE ObjectType* Get() { return &Object; }
-
-	private:
-		ObjectType Object;
-	};
-
-	FMemStackBase MemStack;
-	TArray<FTrackedAlloc*, TMemStackAllocator<>> TrackedAllocs;
-};
-
 /** HANDLE UTILITIES */
 
 /** Handle helper class for internal tracking of RDG types. */
@@ -379,7 +294,7 @@ public:
 	DerivedType* Allocate(FRDGAllocator& Allocator, TArgs&&... Args)
 	{
 		static_assert(TIsDerivedFrom<DerivedType, ObjectType>::Value, "You must specify a type that derives from ObjectType");
-		DerivedType* Object = Allocator.AllocObject<DerivedType>(Forward<TArgs>(Args)...);
+		DerivedType* Object = Allocator.Alloc<DerivedType>(Forward<TArgs>(Args)...);
 		Insert(Object);
 		return Object;
 	}
@@ -430,14 +345,14 @@ public:
 	}
 
 private:
-	TArray<ObjectType*, TMemStackAllocator<>> Array;
+	TArray<ObjectType*, FRDGArrayAllocator> Array;
 };
 
 /** Specialization of bit array with compile-time type checking for handles and a pre-configured allocator. */
 template <typename HandleType>
-class TRDGHandleBitArray : public TBitArray<TInlineAllocator<4, TMemStackAllocator<>>>
+class TRDGHandleBitArray : public TBitArray<FRDGBitArrayAllocator>
 {
-	using Base = TBitArray<TInlineAllocator<4, TMemStackAllocator<>>>;
+	using Base = TBitArray<FRDGBitArrayAllocator>;
 public:
 	TRDGHandleBitArray() = default;
 	TRDGHandleBitArray(const TRDGHandleBitArray&) = default;
@@ -552,7 +467,7 @@ class FRDGPass;
 using FRDGPassRef = const FRDGPass*;
 using FRDGPassHandle = TRDGHandle<FRDGPass, uint16>;
 using FRDGPassRegistry = TRDGHandleRegistry<FRDGPassHandle>;
-using FRDGPassHandleArray = TArray<FRDGPassHandle, TInlineAllocator<4, TMemStackAllocator<>>>;
+using FRDGPassHandleArray = TArray<FRDGPassHandle, TInlineAllocator<4, FRDGArrayAllocator>>;
 using FRDGPassBitArray = TRDGHandleBitArray<FRDGPassHandle>;
 
 class FRDGUniformBuffer;
diff --git a/Engine/Source/Runtime/RenderCore/Public/RenderGraphEvent.h b/Engine/Source/Runtime/RenderCore/Public/RenderGraphEvent.h
--- a/Engine/Source/Runtime/RenderCore/Public/RenderGraphEvent.h
+++ b/Engine/Source/Runtime/RenderCore/Public/RenderGraphEvent.h
@@ -99,7 +99,7 @@ private:
 	const TScopeType* CurrentScope = nullptr;
 
 	/** Tracks scopes allocated through MemStack for destruction. */
-	TArray<TScopeType*, SceneRenderingAllocator> Scopes;
+	TArray<TScopeType*, FRDGArrayAllocator> Scopes;
 
 	/** Stacks of scopes pushed to the RHI command list during execution. */
 	TStaticArray<const TScopeType*, kScopeStackDepthMax> ScopeStack;
diff --git a/Engine/Source/Runtime/RenderCore/Public/RenderGraphPass.h b/Engine/Source/Runtime/RenderCore/Public/RenderGraphPass.h
--- a/Engine/Source/Runtime/RenderCore/Public/RenderGraphPass.h
+++ b/Engine/Source/Runtime/RenderCore/Public/RenderGraphPass.h
@@ -17,8 +17,8 @@ public:
 	void End(FRHIComputeCommandList& RHICmdList);
 
 private:
-	TArray<const FRHITransition*, TInlineAllocator<1, SceneRenderingAllocator>> Queue;
-	TArray<const FRHITransition*, TInlineAllocator<1, SceneRenderingAllocator>> QueueWithFences;
+	TArray<const FRHITransition*, TInlineAllocator<1, FRDGArrayAllocator>> Queue;
+	TArray<const FRHITransition*, TInlineAllocator<1, FRDGArrayAllocator>> QueueWithFences;
 };
 
 struct FRDGBarrierBatchBeginId
@@ -60,7 +60,7 @@ public:
 
 private:
 	const FRHITransition* Transition = nullptr;
-	TArray<FRHITransitionInfo, TInlineAllocator<1, SceneRenderingAllocator>> Transitions;
+	TArray<FRHITransitionInfo, TInlineAllocator<1, FRDGArrayAllocator>> Transitions;
 	ERHICreateTransitionFlags TransitionFlags = ERHICreateTransitionFlags::NoFence;
 	bool bTransitionNeeded = false;
 
@@ -70,7 +70,7 @@ private:
 
 #if RDG_ENABLE_DEBUG
 	FRDGPassHandlesByPipeline DebugPasses;
-	TArray<FRDGParentResource*, SceneRenderingAllocator> DebugResources;
+	TArray<FRDGParentResource*, FRDGArrayAllocator> DebugResources;
 	const TCHAR* DebugName;
 	ERHIPipeline DebugPipelinesToBegin;
 	ERHIPipeline DebugPipelinesToEnd;
@@ -91,7 +91,7 @@ public:
 	void Submit(FRHIComputeCommandList& RHICmdList, ERHIPipeline Pipeline);
 
 private:
-	TArray<FRDGBarrierBatchBegin*, TInlineAllocator<1, SceneRenderingAllocator>> Dependencies;
+	TArray<FRDGBarrierBatchBegin*, TInlineAllocator<1, FRDGArrayAllocator>> Dependencies;
 
 	friend class FRDGBarrierValidation;
 };
@@ -335,21 +335,21 @@ private:
 	};
 
 	/** Maps textures / buffers to information on how they are used in the pass. */
-	TSortedMap<FRDGTexture*, FTextureState, SceneRenderingAllocator> TextureStates;
-	TSortedMap<FRDGBuffer*, FBufferState, SceneRenderingAllocator> BufferStates;
+	TSortedMap<FRDGTexture*, FTextureState, FRDGArrayAllocator> TextureStates;
+	TSortedMap<FRDGBuffer*, FBufferState, FRDGArrayAllocator> BufferStates;
 
 	/** Lists of pass parameters scheduled for begin during execution of this pass. */
-	TArray<FRDGPass*, TInlineAllocator<1, SceneRenderingAllocator>> ResourcesToBegin;
-	TArray<FRDGPass*, TInlineAllocator<1, SceneRenderingAllocator>> ResourcesToEnd;
+	TArray<FRDGPass*, TInlineAllocator<1, FRDGArrayAllocator>> ResourcesToBegin;
+	TArray<FRDGPass*, TInlineAllocator<1, FRDGArrayAllocator>> ResourcesToEnd;
 
 	/** List of textures to acquire *after* the pass completes, *before* discards. Acquires apply to all allocated textures. */
-	TArray<FRHITexture*, SceneRenderingAllocator> TexturesToAcquire;
+	TArray<FRHITexture*, FRDGArrayAllocator> TexturesToAcquire;
 
 	/** List of textures to discard *after* the pass completes, *after* acquires. Discards only apply to textures marked as
 	 *  transient and the last alias of the texture uses the automatic discard behavior (in order to support cleaner hand-off
 	 *  to the user or back to the pool.
 	 */
-	TArray<FRHITexture*, SceneRenderingAllocator> TexturesToDiscard;
+	TArray<FRHITexture*, FRDGArrayAllocator> TexturesToDiscard;
 
 	/** Split-barrier batches at various points of execution of the pass. */
 	FRDGBarrierBatchBegin* PrologueBarriersToBegin = nullptr;
diff --git a/Engine/Source/Runtime/RenderCore/Public/RenderGraphResources.h b/Engine/Source/Runtime/RenderCore/Public/RenderGraphResources.h
--- a/Engine/Source/Runtime/RenderCore/Public/RenderGraphResources.h
+++ b/Engine/Source/Runtime/RenderCore/Public/RenderGraphResources.h
@@ -58,8 +58,8 @@ struct FRDGSubresourceState
 };
 
 using FRDGTextureSubresourceState = TRDGTextureSubresourceArray<FRDGSubresourceState, FDefaultAllocator>;
-using FRDGTextureTransientSubresourceState = TRDGTextureSubresourceArray<FRDGSubresourceState, SceneRenderingAllocator>;
-using FRDGTextureTransientSubresourceStateIndirect = TRDGTextureSubresourceArray<FRDGSubresourceState*, SceneRenderingAllocator>;
+using FRDGTextureTransientSubresourceState = TRDGTextureSubresourceArray<FRDGSubresourceState, FRDGArrayAllocator>;
+using FRDGTextureTransientSubresourceStateIndirect = TRDGTextureSubresourceArray<FRDGSubresourceState*, FRDGArrayAllocator>;
 
 /** Generic graph resource. */
 class RENDERCORE_API FRDGResource
diff --git a/Engine/Source/Runtime/RenderCore/Public/RenderGraphValidation.h b/Engine/Source/Runtime/RenderCore/Public/RenderGraphValidation.h
--- a/Engine/Source/Runtime/RenderCore/Public/RenderGraphValidation.h
+++ b/Engine/Source/Runtime/RenderCore/Public/RenderGraphValidation.h
@@ -101,8 +101,8 @@ private:
 	void ValidateExtractResource(FRDGParentResourceRef Resource);
 
 	/** List of tracked resources for validation prior to shutdown. */
-	TArray<FRDGTextureRef, SceneRenderingAllocator> TrackedTextures;
-	TArray<FRDGBufferRef, SceneRenderingAllocator> TrackedBuffers;
+	TArray<FRDGTextureRef, FRDGArrayAllocator> TrackedTextures;
+	TArray<FRDGBufferRef, FRDGArrayAllocator> TrackedBuffers;
 
 	/** Whether the Execute() has already been called. */
 	bool bHasExecuted = false;
