Fix for RDG crash due to early release of allocations.

diff --git a/Engine/Source/Runtime/RenderCore/Private/RenderGraphBuilder.cpp b/Engine/Source/Runtime/RenderCore/Private/RenderGraphBuilder.cpp
--- a/Engine/Source/Runtime/RenderCore/Private/RenderGraphBuilder.cpp
+++ b/Engine/Source/Runtime/RenderCore/Private/RenderGraphBuilder.cpp
@@ -384,7 +384,6 @@ const char* const FRDGBuilder::kDefaultUnaccountedCSVStat = "RDG_Pass";
 
 FRDGBuilder::FRDGBuilder(FRHICommandListImmediate& InRHICmdList, FRDGEventName InName, ERDGBuilderFlags InFlags)
 	: RHICmdList(InRHICmdList)
-	, Allocator(FRDGAllocator::Get())
 	, Blackboard(Allocator)
 	, RHICmdListAsyncCompute(FRHICommandListExecutor::GetImmediateAsyncComputeCommandList())
 	, BuilderName(InName)
@@ -1335,7 +1334,6 @@ void FRDGBuilder::Clear()
 	Buffers.Clear();
 	UniformBuffers.Clear();
 	Blackboard.Clear();
-	Allocator.ReleaseAll();
 }
 
 void FRDGBuilder::SetupPass(FRDGPass* Pass)
diff --git a/Engine/Source/Runtime/RenderCore/Public/RenderGraphAllocator.h b/Engine/Source/Runtime/RenderCore/Public/RenderGraphAllocator.h
--- a/Engine/Source/Runtime/RenderCore/Public/RenderGraphAllocator.h
+++ b/Engine/Source/Runtime/RenderCore/Public/RenderGraphAllocator.h
@@ -72,6 +72,25 @@ private:
 
 	template <uint32>
 	friend class TRDGArrayAllocator;
+	friend class FRDGAllocatorScope;
+};
+
+/** Base class for RDG builder which scopes the allocations and releases them in the destructor. */
+class FRDGAllocatorScope
+{
+protected:
+	FRDGAllocator& Allocator;
+
+private:
+	FRDGAllocatorScope()
+		: Allocator(FRDGAllocator::Get())
+	{}
+
+	~FRDGAllocatorScope()
+	{
+		Allocator.ReleaseAll();
+	}
+
 	friend class FRDGBuilder;
 };
 
diff --git a/Engine/Source/Runtime/RenderCore/Public/RenderGraphBuilder.h b/Engine/Source/Runtime/RenderCore/Public/RenderGraphBuilder.h
--- a/Engine/Source/Runtime/RenderCore/Public/RenderGraphBuilder.h
+++ b/Engine/Source/Runtime/RenderCore/Public/RenderGraphBuilder.h
@@ -19,6 +19,7 @@ class FRDGLogFile;
  *  destruction.
  */
 class RENDERCORE_API FRDGBuilder
+	: FRDGAllocatorScope
 {
 public:
 	FRDGBuilder(FRHICommandListImmediate& InRHICmdList, FRDGEventName InName = {}, ERDGBuilderFlags Flags = ERDGBuilderFlags::None);
@@ -218,9 +219,6 @@ public:
 	/** The RHI command list used for the render graph. */
 	FRHICommandListImmediate& RHICmdList;
 
-	/** The RDG allocator instance tied to the graph lifetime. All memory is released after Execute(). */
-	FRDGAllocator& Allocator;
-
 	/** The blackboard used to hold common data tied to the graph lifetime. */
 	FRDGBlackboard Blackboard;
 
