UE5 - Proper fix for r.Shaders.KeepDebugInfo passing in the FRHIShader::ShaderName correctly
Initial fixes and tested on D3D11, D3D12 and Vulkan
Changed 'n' to FShaderCodeName::Key
Also added checks if trying to add optional shader data out of order

diff --git a/Engine/Shaders/Public/ShaderVersion.ush b/Engine/Shaders/Public/ShaderVersion.ush
--- a/Engine/Shaders/Public/ShaderVersion.ush
+++ b/Engine/Shaders/Public/ShaderVersion.ush
@@ -3,4 +3,4 @@
 // in Platform.ush (which should be included in any shader) it allows to invalidate the shader DDC.
 // 
 // If you are merging streams and there is a conflict with this GUID you should make a new GUID rather than taking one or the other.
-// GUID = 66906A4D-A1F5-4A5D-9AAB-70FFF8DB1805
+// GUID = 5952C866-E7D6-4116-ABCC-48ACA608D216
diff --git a/Engine/Source/Developer/VulkanShaderFormat/Private/VulkanShaderCompiler.cpp b/Engine/Source/Developer/VulkanShaderFormat/Private/VulkanShaderCompiler.cpp
--- a/Engine/Source/Developer/VulkanShaderFormat/Private/VulkanShaderCompiler.cpp
+++ b/Engine/Source/Developer/VulkanShaderFormat/Private/VulkanShaderCompiler.cpp
@@ -2905,6 +2905,12 @@ static bool CompileWithShaderConductor(
 	// Write final output shader
 	Output.Target = Input.Target;
 	Output.ShaderCode.GetWriteAccess().Append(reinterpret_cast<const uint8*>(Spirv.Data.GetData()), Spirv.Data.Num() * sizeof(uint32));
+
+	if (Input.Environment.CompilerFlags.Contains(CFLAG_KeepDebugInfo))
+	{
+		Output.ShaderCode.AddOptionalData(FShaderCodeName::Key, TCHAR_TO_UTF8(*Input.GenerateShaderName()));
+	}
+
 	Output.bSucceeded = true;
 
 	if (bDebugDump)
diff --git a/Engine/Source/Developer/Windows/ShaderFormatD3D/Private/D3DShaderCompiler.inl b/Engine/Source/Developer/Windows/ShaderFormatD3D/Private/D3DShaderCompiler.inl
--- a/Engine/Source/Developer/Windows/ShaderFormatD3D/Private/D3DShaderCompiler.inl
+++ b/Engine/Source/Developer/Windows/ShaderFormatD3D/Private/D3DShaderCompiler.inl
@@ -526,10 +526,10 @@ inline void GenerateFinalOutput(TRefCountPtr<TBlob>& CompressedData,
 		}
 	}
 
-	// Store data we can pickup later with ShaderCode.FindOptionalData('n'), could be removed for shipping
-	// Daniel L: This GenerateShaderName does not generate a deterministic output among shaders as the shader code can be shared. 
-	//			uncommenting this will cause the project to have non deterministic materials and will hurt patch sizes
-	//Output.ShaderCode.AddOptionalData('n', TCHAR_TO_UTF8(*Input.GenerateShaderName()));
+	if (Input.Environment.CompilerFlags.Contains(CFLAG_KeepDebugInfo))
+	{
+		Output.ShaderCode.AddOptionalData(FShaderCodeName::Key, TCHAR_TO_UTF8(*Input.GenerateShaderName()));
+	}
 
 	// Set the number of instructions.
 	Output.NumInstructions = NumInstructions;
diff --git a/Engine/Source/Runtime/D3D12RHI/Private/D3D12Shaders.cpp b/Engine/Source/Runtime/D3D12RHI/Private/D3D12Shaders.cpp
--- a/Engine/Source/Runtime/D3D12RHI/Private/D3D12Shaders.cpp
+++ b/Engine/Source/Runtime/D3D12RHI/Private/D3D12Shaders.cpp
@@ -16,7 +16,7 @@ static inline bool ReadShaderOptionalData(FShaderCodeReader& InShaderCode, TShad
 	}
 	OutShader.ResourceCounts = *PackedResourceCounts;
 #if !(UE_BUILD_SHIPPING || UE_BUILD_TEST)
-	OutShader.ShaderName = InShaderCode.FindOptionalData('n');
+	OutShader.ShaderName = InShaderCode.FindOptionalData(FShaderCodeName::Key);
 
 	int32 UniformBufferTableSize = 0;
 	const auto* UniformBufferData = InShaderCode.FindOptionalDataAndSize('u', UniformBufferTableSize);
diff --git a/Engine/Source/Runtime/RenderCore/Public/ShaderCore.h b/Engine/Source/Runtime/RenderCore/Public/ShaderCore.h
--- a/Engine/Source/Runtime/RenderCore/Public/ShaderCore.h
+++ b/Engine/Source/Runtime/RenderCore/Public/ShaderCore.h
@@ -537,6 +537,14 @@ struct FShaderCodeFeatures
 	bool bUsesWaveOps;
 };
 
+// if this changes you need to make sure all shaders get invalidated
+struct FShaderCodeName
+{
+	static const uint8 Key = 'n';
+
+	// We store the straight ANSICHAR zero-terminated string
+};
+
 // if this changes you need to make sure all shaders get invalidated
 struct FShaderCodeVendorExtension
 {
@@ -737,9 +745,11 @@ public:
 		}
 	}
 
-	// for write access
+	// Write access for regular microcode: Optional Data must be added AFTER regular microcode and BEFORE Finalize
 	TArray<uint8>& GetWriteAccess()
 	{
+		checkf(OptionalDataSize != -1, TEXT("Tried to add ShaderCode after adding being finalized!"));
+		checkf(OptionalDataSize == 0, TEXT("Tried to add ShaderCode after adding Optional data!"));
 		return ShaderCodeWithOptionalData;
 	}
 
@@ -752,16 +762,6 @@ public:
 		return Wrapper.GetShaderCodeSize();
 	}
 
-	// inefficient, will/should be replaced by GetShaderCodeToRead()
-	UE_DEPRECATED(4.26, "Please switch to GetShaderCodeToRead()")
-	void GetShaderCodeLegacy(TArray<uint8>& Out) const
-	{
-		Out.Empty();
-
-		Out.AddUninitialized(GetShaderCodeSize());
-		FMemory::Memcpy(Out.GetData(), GetReadAccess().GetData(), ShaderCodeWithOptionalData.Num());
-	}
-
 	// for read access, can have additional data attached to the end
 	const TArray<uint8>& GetReadAccess() const
 	{
diff --git a/Engine/Source/Runtime/VulkanRHI/Private/VulkanShaders.cpp b/Engine/Source/Runtime/VulkanRHI/Private/VulkanShaders.cpp
--- a/Engine/Source/Runtime/VulkanRHI/Private/VulkanShaders.cpp
+++ b/Engine/Source/Runtime/VulkanRHI/Private/VulkanShaders.cpp
@@ -53,7 +53,13 @@ ShaderType* FVulkanShaderFactory::CreateShader(TArrayView<const uint8> Code, FVu
 			
 		FRWScopeLock ScopedLock(Lock, SLT_Write);
 		ShaderMap[ShaderType::StaticFrequency].Add(ShaderKey, RetShader);
+
+#if !(UE_BUILD_SHIPPING || UE_BUILD_TEST)
+		FShaderCodeReader ShaderCode(Code);
+		RetShader->ShaderName = ShaderCode.FindOptionalData(FShaderCodeName::Key);
+#endif
 	}
+
 	return RetShader;
 }
 
diff --git a/Engine/Source/Runtime/Windows/D3D11RHI/Private/D3D11Shaders.cpp b/Engine/Source/Runtime/Windows/D3D11RHI/Private/D3D11Shaders.cpp
--- a/Engine/Source/Runtime/Windows/D3D11RHI/Private/D3D11Shaders.cpp
+++ b/Engine/Source/Runtime/Windows/D3D11RHI/Private/D3D11Shaders.cpp
@@ -67,7 +67,7 @@ static inline void ReadShaderOptionalData(FShaderCodeReader& InShaderCode, TShad
 
 	OutShader.bShaderNeedsGlobalConstantBuffer = PackedResourceCounts->bGlobalUniformBufferUsed;
 #if !(UE_BUILD_SHIPPING || UE_BUILD_TEST)
-	OutShader.ShaderName = InShaderCode.FindOptionalData('n');
+	OutShader.ShaderName = InShaderCode.FindOptionalData(FShaderCodeName::Key);
 	int32 UniformBufferTableSize = 0;
 	auto* UniformBufferData = InShaderCode.FindOptionalDataAndSize('u', UniformBufferTableSize);
 	if (UniformBufferData && UniformBufferTableSize > 0)
