Fixed erroneous assert in RDG debug mode.

diff --git a/Engine/Source/Runtime/RenderCore/Private/RenderGraphValidation.cpp b/Engine/Source/Runtime/RenderCore/Private/RenderGraphValidation.cpp
--- a/Engine/Source/Runtime/RenderCore/Private/RenderGraphValidation.cpp
+++ b/Engine/Source/Runtime/RenderCore/Private/RenderGraphValidation.cpp
@@ -104,9 +104,6 @@ struct FRDGParentResourceDebugData
 	/** Count the number of times it has been used by a pass (without culling). */
 	uint32 PassAccessCount = 0;
 
-	/** Tracks whether a resource has valid contents. This version is true by default for external resources which have prior unknown contents. */
-	bool bHasBeenProduced = false;
-
 	/** Tracks whether this resource was clobbered by the builder prior to use. */
 	bool bHasBeenClobbered = false;
 };
@@ -269,13 +266,11 @@ void FRDGUserValidation::ValidateRegisterExternalBuffer(const TRefCountPtr<FRDGP
 void FRDGUserValidation::ValidateRegisterExternalTexture(FRDGTextureRef Texture)
 {
 	ValidateCreateTexture(Texture);
-	Texture->GetParentDebugData().bHasBeenProduced = true;
 }
 
 void FRDGUserValidation::ValidateRegisterExternalBuffer(FRDGBufferRef Buffer)
 {
 	ValidateCreateBuffer(Buffer);
-	Buffer->GetParentDebugData().bHasBeenProduced = true;
 }
 
 void FRDGUserValidation::ValidateCreateTexture(const FRDGTextureDesc& Desc, const TCHAR* Name, ERDGTextureFlags Flags)
@@ -854,7 +849,7 @@ void FRDGUserValidation::ValidateExecuteEnd()
 
 			if (bProducedButNeverUsed)
 			{
-				check(ParentDebugData.bHasBeenProduced);
+				check(Resource->bProduced || Resource->bExternal || Resource->bExtracted);
 
 				EmitRDGWarningf(
 					TEXT("Resource %s has been produced by the pass %s, but never used by another pass."),
