Implemented support for multi-pipe UAV access. Implemented command list stats on RDG builder to avoid helper passes. Added {Begin, End}UAVOverlap to RDG UAVs.

diff --git a/Engine/Source/Runtime/D3D12RHI/Private/D3D12RHI.cpp b/Engine/Source/Runtime/D3D12RHI/Private/D3D12RHI.cpp
--- a/Engine/Source/Runtime/D3D12RHI/Private/D3D12RHI.cpp
+++ b/Engine/Source/Runtime/D3D12RHI/Private/D3D12RHI.cpp
@@ -80,7 +80,7 @@ FD3D12DynamicRHI::FD3D12DynamicRHI(const TArray<TSharedPtr<FD3D12Adapter>>& Chos
 #endif // PLATFORM_WINDOWS
 
 	GRHISupportsMultithreading = true;
-	GRHIMultiPipelineAccessMask = ERHIAccess::ReadOnlyMask;
+	GRHIMultiPipelineMergeableAccessMask = GRHIMergeableAccessMask;
 
 	GPoolSizeVRAMPercentage = 0;
 	GTexturePoolSize = 0;
diff --git a/Engine/Source/Runtime/RHI/Private/RHI.cpp b/Engine/Source/Runtime/RHI/Private/RHI.cpp
--- a/Engine/Source/Runtime/RHI/Private/RHI.cpp
+++ b/Engine/Source/Runtime/RHI/Private/RHI.cpp
@@ -903,8 +903,11 @@ RHI_API int32 GNumPrimitivesDrawnRHI = 0;
 RHI_API uint64 GRHITransitionPrivateData_SizeInBytes = 0;
 RHI_API uint64 GRHITransitionPrivateData_AlignInBytes = 0;
 
-ERHIAccess GRHITextureReadAccessMask = ERHIAccess::ReadOnlyMask;
-ERHIAccess GRHIMultiPipelineAccessMask = ERHIAccess::Unknown;
+// By default, read only states and UAV states are allowed to participate in state merging.
+ERHIAccess GRHIMergeableAccessMask = ERHIAccess::ReadOnlyMask | ERHIAccess::UAVMask;
+
+// By default, no multi-pipe states are allowed. Each platform must opt-in.
+ERHIAccess GRHIMultiPipelineMergeableAccessMask = ERHIAccess::Unknown;
 
 /** Called once per frame only from within an RHI. */
 void RHIPrivateBeginFrame()
diff --git a/Engine/Source/Runtime/RHI/Public/RHI.h b/Engine/Source/Runtime/RHI/Public/RHI.h
--- a/Engine/Source/Runtime/RHI/Public/RHI.h
+++ b/Engine/Source/Runtime/RHI/Public/RHI.h
@@ -843,11 +843,11 @@ enum class ERHIAccess
 };
 ENUM_CLASS_FLAGS(ERHIAccess)
 
-/** Mask of read states that can be used together for textures. */
-extern RHI_API ERHIAccess GRHITextureReadAccessMask;
+/** Mask of states which are allowed to be considered for state merging. */
+extern RHI_API ERHIAccess GRHIMergeableAccessMask;
 
-/** Mask of states which are allowed to be used on multiple pipelines at the same time. */
-extern RHI_API ERHIAccess GRHIMultiPipelineAccessMask;
+/** Mask of states which are allowed to be considered for multi-pipeline state merging. This should be a subset of GRHIMergeableAccessMask. */
+extern RHI_API ERHIAccess GRHIMultiPipelineMergeableAccessMask;
 
 /** to customize the RHIReadSurfaceData() output */
 class FReadSurfaceDataFlags
diff --git a/Engine/Source/Runtime/RenderCore/Private/RenderGraphBuilder.cpp b/Engine/Source/Runtime/RenderCore/Private/RenderGraphBuilder.cpp
--- a/Engine/Source/Runtime/RenderCore/Private/RenderGraphBuilder.cpp
+++ b/Engine/Source/Runtime/RenderCore/Private/RenderGraphBuilder.cpp
@@ -607,26 +607,38 @@ void FRDGBuilder::Compile()
 	{
 		SCOPED_NAMED_EVENT(FRDGBuilder_Compile_Culling_Dependencies, FColor::Emerald);
 
-		const auto AddCullingDependency = [&](FRDGPassHandle& ProducerHandle, FRDGPassHandle PassHandle, ERHIAccess Access)
+		const auto AddCullingDependency = [&](FRDGProducerStatesByPipeline& LastProducers, const FRDGProducerState& NextState, ERHIPipeline NextPipeline)
 		{
-			if (Access != ERHIAccess::Unknown)
+			if (NextState.Access == ERHIAccess::Unknown)
 			{
-				if (ProducerHandle.IsValid())
+				return;
+			}
+
+			for (ERHIPipeline LastPipeline : GetRHIPipelines())
+			{
+				FRDGProducerState& LastProducer = LastProducers[LastPipeline];
+
+				if (LastProducer.Access == ERHIAccess::Unknown)
 				{
-					AddPassDependency(ProducerHandle, PassHandle);
+					continue;
 				}
 
-				// If the access is writable, we store the new producer.
-				if (IsWritableAccess(Access))
+				if (FRDGProducerState::IsDependencyRequired(LastProducer, LastPipeline, NextState, NextPipeline))
 				{
-					ProducerHandle = PassHandle;
+					AddPassDependency(LastProducer.PassHandle, NextState.PassHandle);
 				}
 			}
+
+			if (IsWritableAccess(NextState.Access))
+			{
+				LastProducers[NextPipeline] = NextState;
+			}
 		};
 
 		for (FRDGPassHandle PassHandle = Passes.Begin(); PassHandle != Passes.End(); ++PassHandle)
 		{
 			FRDGPass* Pass = Passes[PassHandle];
+			const ERHIPipeline PassPipeline = Pass->GetPipeline();
 
 			bool bUntrackedOutputs = Pass->GetParameters().HasExternalOutputs();
 
@@ -647,7 +659,14 @@ void FRDGBuilder::Compile()
 
 				for (uint32 Index = 0, Count = LastProducers.Num(); Index < Count; ++Index)
 				{
-					AddCullingDependency(LastProducers[Index], PassHandle, PassState[bWholePassState ? 0 : Index].Access);
+					const auto& SubresourceState = PassState[bWholePassState ? 0 : Index];
+
+					FRDGProducerState ProducerState;
+					ProducerState.Access = SubresourceState.Access;
+					ProducerState.PassHandle = PassHandle;
+					ProducerState.NoUAVBarrierHandle = SubresourceState.NoUAVBarrierFilter.GetUniqueHandle();
+
+					AddCullingDependency(LastProducers[Index], ProducerState, PassPipeline);
 				}
 
 				bUntrackedOutputs |= Texture->bExternal;
@@ -656,7 +675,14 @@ void FRDGBuilder::Compile()
 			for (auto& BufferPair : Pass->BufferStates)
 			{
 				FRDGBufferRef Buffer = BufferPair.Key;
-				AddCullingDependency(Buffer->LastProducer, PassHandle, BufferPair.Value.State.Access);
+				const auto& SubresourceState = BufferPair.Value.State;
+
+				FRDGProducerState ProducerState;
+				ProducerState.Access = SubresourceState.Access;
+				ProducerState.PassHandle = PassHandle;
+				ProducerState.NoUAVBarrierHandle = SubresourceState.NoUAVBarrierFilter.GetUniqueHandle();
+
+				AddCullingDependency(Buffer->LastProducer, ProducerState, PassPipeline);
 				bUntrackedOutputs |= Buffer->bExternal;
 			}
 
@@ -680,9 +706,13 @@ void FRDGBuilder::Compile()
 		for (const auto& Query : ExtractedTextures)
 		{
 			FRDGTextureRef Texture = Query.Key;
-			for (FRDGPassHandle& ProducerHandle : Texture->LastProducers)
+			for (auto& LastProducer : Texture->LastProducers)
 			{
-				AddCullingDependency(ProducerHandle, EpiloguePassHandle, Texture->AccessFinal);
+				FRDGProducerState StateFinal;
+				StateFinal.Access = Texture->AccessFinal;
+				StateFinal.PassHandle = EpiloguePassHandle;
+
+				AddCullingDependency(LastProducer, StateFinal, ERHIPipeline::Graphics);
 			}
 			Texture->ReferenceCount++;
 		}
@@ -690,7 +720,12 @@ void FRDGBuilder::Compile()
 		for (const auto& Query : ExtractedBuffers)
 		{
 			FRDGBufferRef Buffer = Query.Key;
-			AddCullingDependency(Buffer->LastProducer, EpiloguePassHandle, Buffer->AccessFinal);
+
+			FRDGProducerState StateFinal;
+			StateFinal.Access = Buffer->AccessFinal;
+			StateFinal.PassHandle = EpiloguePassHandle;
+
+			AddCullingDependency(Buffer->LastProducer, StateFinal, ERHIPipeline::Graphics);
 			Buffer->ReferenceCount++;
 		}
 	}
@@ -1440,6 +1475,10 @@ void FRDGBuilder::SetupPassInternal(FRDGPass* Pass, FRDGPassHandle PassHandle, E
 	Pass->GPUMask = RHICmdList.GetGPUMask();
 #endif
 
+#if STATS
+	Pass->CommandListStat = CommandListStat;
+#endif
+
 	IF_RDG_CPU_SCOPES(Pass->CPUScopes = CPUScopeStacks.GetCurrentScopes());
 	IF_RDG_GPU_SCOPES(Pass->GPUScopes = GPUScopeStacks.GetCurrentScopes(PassPipeline));
 
@@ -1489,6 +1528,7 @@ void FRDGBuilder::ExecutePassPrologue(FRHIComputeCommandList& RHICmdListPass, FR
 	CSV_SCOPED_TIMING_STAT_EXCLUSIVE_CONDITIONAL(RDGBuilder_ExecutePassPrologue, GRDGVerboseCSVStats != 0);
 
 	IF_RDG_ENABLE_DEBUG(UserValidation.ValidateExecutePassBegin(Pass));
+	IF_RDG_CMDLIST_STATS(RHICmdList.SetCurrentStat(Pass->CommandListStat));
 
 	const ERDGPassFlags PassFlags = Pass->GetFlags();
 	const ERHIPipeline PassPipeline = Pass->GetPipeline();
@@ -1517,6 +1557,8 @@ void FRDGBuilder::ExecutePassPrologue(FRHIComputeCommandList& RHICmdListPass, FR
 			static_cast<FRHICommandList&>(RHICmdListPass).BeginRenderPass(Pass->GetParameters().GetRenderPassInfo(), Pass->GetName());
 		}
 	}
+
+	BeginUAVOverlap(Pass, RHICmdListPass);
 }
 
 void FRDGBuilder::ExecutePassEpilogue(FRHIComputeCommandList& RHICmdListPass, FRDGPass* Pass)
@@ -1524,6 +1566,8 @@ void FRDGBuilder::ExecutePassEpilogue(FRHIComputeCommandList& RHICmdListPass, FR
 	QUICK_SCOPE_CYCLE_COUNTER(STAT_FRDGBuilder_ExecutePassEpilogue);
 	CSV_SCOPED_TIMING_STAT_EXCLUSIVE_CONDITIONAL(RDGBuilder_ExecutePassEpilogue, GRDGVerboseCSVStats != 0);
 
+	EndUAVOverlap(Pass, RHICmdListPass);
+
 	const ERDGPassFlags PassFlags = Pass->GetFlags();
 	const ERHIPipeline PassPipeline = Pass->GetPipeline();
 	const FRDGParameterStruct PassParameters = Pass->GetParameters();
@@ -1927,7 +1971,7 @@ void FRDGBuilder::AddTransition(FRDGPassHandle PassHandle, FRDGBufferRef Buffer,
 	{
 		FRHITransitionInfo Info;
 		Info.Resource = Buffer->GetRHIUnchecked();
-		Info.Type = FRDGBufferDesc::GetTransitionResourceType(Buffer->Desc.UnderlyingType);
+		Info.Type = FRHITransitionInfo::EType::Buffer;
 		Info.Flags = StateAfter.Flags;
 		Info.AccessBefore = StateBefore.Access;
 		Info.AccessAfter = StateAfter.Access;
diff --git a/Engine/Source/Runtime/RenderCore/Private/RenderGraphPrivate.h b/Engine/Source/Runtime/RenderCore/Private/RenderGraphPrivate.h
--- a/Engine/Source/Runtime/RenderCore/Private/RenderGraphPrivate.h
+++ b/Engine/Source/Runtime/RenderCore/Private/RenderGraphPrivate.h
@@ -139,4 +139,43 @@ inline const TCHAR* GetEpilogueBarriersToBeginDebugName(ERHIPipeline Pipelines)
 	}
 #endif
 	return TEXT("");
+}
+
+template <typename FunctionType>
+inline void EnumeratePassUAVs(const FRDGPass* Pass, FunctionType Function)
+{
+	Pass->GetParameters().Enumerate([&](FRDGParameter Parameter)
+	{
+		if (Parameter.IsUAV())
+		{
+			if (FRDGUnorderedAccessViewRef UAV = Parameter.GetAsUAV())
+			{
+				Function(UAV->GetRHI());
+			}
+		}
+	});
+}
+
+inline void BeginUAVOverlap(const FRDGPass* Pass, FRHIComputeCommandList& RHICmdList)
+{
+#if ENABLE_RHI_VALIDATION
+	TArray<FRHIUnorderedAccessView*, TInlineAllocator<8, FRDGArrayAllocator>> UAVs;
+	EnumeratePassUAVs(Pass, [&](FRHIUnorderedAccessView* UAV)
+	{
+		UAVs.Add(UAV);
+	});
+	RHICmdList.BeginUAVOverlap(UAVs);
+#endif
+}
+
+inline void EndUAVOverlap(const FRDGPass* Pass, FRHIComputeCommandList& RHICmdList)
+{
+#if ENABLE_RHI_VALIDATION
+	TArray<FRHIUnorderedAccessView*, TInlineAllocator<8, FRDGArrayAllocator>> UAVs;
+	EnumeratePassUAVs(Pass, [&](FRHIUnorderedAccessView* UAV)
+	{
+		UAVs.Add(UAV);
+	});
+	RHICmdList.EndUAVOverlap(UAVs);
+#endif
 }
\ No newline at end of file
diff --git a/Engine/Source/Runtime/RenderCore/Private/RenderGraphResourcePool.cpp b/Engine/Source/Runtime/RenderCore/Private/RenderGraphResourcePool.cpp
--- a/Engine/Source/Runtime/RenderCore/Private/RenderGraphResourcePool.cpp
+++ b/Engine/Source/Runtime/RenderCore/Private/RenderGraphResourcePool.cpp
@@ -66,11 +66,11 @@ TRefCountPtr<FRDGPooledBuffer> FRenderGraphResourcePool::FindFreeBufferInternal(
 
 		if (Desc.UnderlyingType == FRDGBufferDesc::EUnderlyingType::VertexBuffer)
 		{
-			PooledBuffer->VertexBuffer = RHICreateVertexBuffer(NumBytes, Desc.Usage, CreateInfo);
+			PooledBuffer->Buffer = RHICreateVertexBuffer(NumBytes, Desc.Usage, CreateInfo);
 		}
 		else if (Desc.UnderlyingType == FRDGBufferDesc::EUnderlyingType::StructuredBuffer)
 		{
-			PooledBuffer->StructuredBuffer = RHICreateStructuredBuffer(Desc.BytesPerElement, NumBytes, Desc.Usage, CreateInfo);
+			PooledBuffer->Buffer = RHICreateStructuredBuffer(Desc.BytesPerElement, NumBytes, Desc.Usage, CreateInfo);
 		}
 		else
 		{
diff --git a/Engine/Source/Runtime/RenderCore/Private/RenderGraphResources.cpp b/Engine/Source/Runtime/RenderCore/Private/RenderGraphResources.cpp
--- a/Engine/Source/Runtime/RenderCore/Private/RenderGraphResources.cpp
+++ b/Engine/Source/Runtime/RenderCore/Private/RenderGraphResources.cpp
@@ -4,6 +4,22 @@
 #include "RenderGraphPass.h"
 #include "RenderGraphPrivate.h"
 
+inline bool NeedsUAVBarrier(FRDGViewHandle PreviousHandle, FRDGViewHandle NextHandle)
+{
+	// Barrier if previous / next don't have a matching valid skip-barrier UAV handle.
+	if (GRDGOverlapUAVs != 0 && NextHandle.IsValid() && PreviousHandle == NextHandle)
+	{
+		return false;
+	}
+
+	return true;
+}
+
+inline bool NeedsUAVBarrier(const FRDGSubresourceState& Previous, const FRDGSubresourceState& Next)
+{
+	return NeedsUAVBarrier(Previous.NoUAVBarrierFilter.GetUniqueHandle(), Next.NoUAVBarrierFilter.GetUniqueHandle());
+}
+
 FRDGParentResource::FRDGParentResource(const TCHAR* InName, const ERDGParentResourceType InType)
 	: FRDGResource(InName)
 	, Type(InType)
@@ -22,38 +38,48 @@ void FRDGParentResource::SetPassthroughRHI(FRHIResource* InResourceRHI)
 	ResourceRHI = InResourceRHI;
 }
 
-bool FRDGSubresourceState::IsTransitionRequired(const FRDGSubresourceState& Previous, const FRDGSubresourceState& Next)
+bool FRDGProducerState::IsDependencyRequired(FRDGProducerState LastProducer, ERHIPipeline LastPipeline, FRDGProducerState NextState, ERHIPipeline NextPipeline)
 {
-	check(Next.Access != ERHIAccess::Unknown);
-
-	if (Previous.Access != Next.Access || Previous.GetPipelines() != Next.GetPipelines())
+	/** This function determines whether a producer-consumer relationship exists in the graph, which is used for culling and
+	 *  async-compute fence derivation. Producers are tracked per-pipeline, so it's safe to elide a cross-pipeline producer
+	 *  for the purposes of overlapping producers, as long as a dependency exists on the same pipeline. Eliding both will
+	 *  split the producer / consumer graph into two and break culling. The only current use case this is allowing multiple
+	 *  pipes to write UAVs.
+	 *
+	 *  Producer / consumer dependencies take place independent of resource state merging / transitions, so the logic must
+	 *  be carefully aligned so that cross-pipe dependencies align with transitions.
+	 */
+
+	// The first needs to be known producers.
+	check(IsWritableAccess(LastProducer.Access));
+
+	// A dependency is always applied on the same pipe to ensure that connectivity is preserved for culling purposes.
+	if (LastPipeline == NextPipeline)
 	{
 		return true;
 	}
-	else if (EnumHasAnyFlags(Previous.Access, ERHIAccess::UAVMask) && EnumHasAnyFlags(Next.Access, ERHIAccess::UAVMask))
-	{
-		if (!GRDGOverlapUAVs)
-		{
-			return true;
-		}
 
-		const FRDGViewHandle PreviousUniqueHandle = Previous.NoUAVBarrierFilter.GetUniqueHandle();
-		const FRDGViewHandle NextUniqueHandle = Next.NoUAVBarrierFilter.GetUniqueHandle();
+	// Only certain platforms allow multi-pipe UAV access.
+	const ERHIAccess MultiPipelineUAVMask = ERHIAccess::UAVMask & GRHIMultiPipelineMergeableAccessMask;
 
-		// Previous / Next have the same non-null no-barrier UAV.
-		const bool bHasNoBarrierUAV = PreviousUniqueHandle == NextUniqueHandle && PreviousUniqueHandle.IsValid();
-
-		// We require a UAV barrier unless we have a valid no-barrier UAV being used.
-		return !bHasNoBarrierUAV;
-	}
-	else
+	// Skip the dependency if the states are used as UAV on different pipes and a UAV barrier can be skipped. This elides the async fence.
+	if (EnumHasAnyFlags(NextState.Access, MultiPipelineUAVMask) && !NeedsUAVBarrier(LastProducer.NoUAVBarrierHandle, NextState.NoUAVBarrierHandle))
 	{
 		return false;
 	}
+
+	// Everything else requires a dependency.
+	return true;
 }
 
 bool FRDGSubresourceState::IsMergeAllowed(ERDGParentResourceType ResourceType, const FRDGSubresourceState& Previous, const FRDGSubresourceState& Next)
 {
+	/** State merging occurs during compilation and before resource transitions are collected. It serves to remove the bulk
+	 *  of unnecessary transitions by looking ahead in the resource usage chain. A resource transition cannot occur within
+	 *  a merged state, so a merge is not allowed to proceed if a barrier might be required. Merging is also where multi-pipe
+	 *  transitions are determined, if supported by the platform.
+	 */
+
 	const ERHIAccess AccessUnion = Previous.Access | Next.Access;
 
 	// If we have the same access between the two states, we don't need to check for invalid access combinations.
@@ -71,26 +97,32 @@ bool FRDGSubresourceState::IsMergeAllowed(ERDGParentResourceType ResourceType, c
 			return false;
 		}
 
-		// Textures only allow certain read states to merge.
-		if (ResourceType == ERDGParentResourceType::Texture && EnumHasAnyFlags(AccessUnion, ~GRHITextureReadAccessMask))
+		// UAVs will filter through the above checks because they are both read and write. UAV can only merge it itself.
+		if (EnumHasAnyFlags(AccessUnion, ERHIAccess::UAVMask) && EnumHasAnyFlags(AccessUnion, ~ERHIAccess::UAVMask))
+		{
+			return false;
+		}
+
+		// Filter out platform-specific unsupported mergeable states.
+		if (EnumHasAnyFlags(AccessUnion, ~GRHIMergeableAccessMask))
 		{
 			return false;
 		}
 	}
 
-	// For merging purposes we are conservative and assume a UAV barrier.
-	if (EnumHasAnyFlags(AccessUnion, ERHIAccess::UAVMask))
+	// Not allowed if the resource is being used as a UAV and needs a barrier.
+	if (EnumHasAnyFlags(Next.Access, ERHIAccess::UAVMask) && NeedsUAVBarrier(Previous, Next))
 	{
 		return false;
 	}
 
-	// If our platform doesn't support multi-pipeline transitions, then we can't merge the state.
-	if (EnumHasAnyFlags(AccessUnion, ~GRHIMultiPipelineAccessMask) && Previous.GetPipelines() != Next.GetPipelines())
+	// Filter out unsupported platform-specific multi-pipeline merged accesses.
+	if (EnumHasAnyFlags(AccessUnion, ~GRHIMultiPipelineMergeableAccessMask) && Previous.GetPipelines() != Next.GetPipelines())
 	{
 		return false;
 	}
 
-	// We are not allowed to cross pipelines or change flags in a merge.
+	// Not allowed to merge differing flags.
 	if (Previous.Flags != Next.Flags)
 	{
 		return false;
@@ -99,6 +131,25 @@ bool FRDGSubresourceState::IsMergeAllowed(ERDGParentResourceType ResourceType, c
 	return true;
 }
 
+bool FRDGSubresourceState::IsTransitionRequired(const FRDGSubresourceState& Previous, const FRDGSubresourceState& Next)
+{
+	// This function only needs to filter out identical states and handle UAV barriers.
+	check(Next.Access != ERHIAccess::Unknown);
+
+	if (Previous.Access != Next.Access || Previous.GetPipelines() != Next.GetPipelines() || Previous.Flags != Next.Flags)
+	{
+		return true;
+	}
+
+	// UAV is a special case as a barrier may still be required even if the states match.
+	if (EnumHasAnyFlags(Next.Access, ERHIAccess::UAVMask) && NeedsUAVBarrier(Previous, Next))
+	{
+		return true;
+	}
+
+	return false;
+}
+
 bool FRDGTextureDesc::IsValid() const
 {
 	if (Extent.X <= 0 || Extent.Y <= 0 || Depth == 0 || ArraySize == 0 || NumMips == 0 || NumSamples < 1 || NumSamples > 8)
@@ -242,18 +293,7 @@ void FRDGBuffer::SetRHI(FRDGPooledBuffer* InPooledBuffer, FRDGBufferRef& OutPrev
 	PooledBuffer = InPooledBuffer;
 	Allocation = InPooledBuffer;
 	State = &PooledBuffer->State;
-	switch (Desc.UnderlyingType)
-	{
-	case FRDGBufferDesc::EUnderlyingType::VertexBuffer:
-		ResourceRHI = InPooledBuffer->VertexBuffer;
-		break;
-	case FRDGBufferDesc::EUnderlyingType::IndexBuffer:
-		ResourceRHI = InPooledBuffer->IndexBuffer;
-		break;
-	case FRDGBufferDesc::EUnderlyingType::StructuredBuffer:
-		ResourceRHI = InPooledBuffer->StructuredBuffer;
-		break;
-	}
+	ResourceRHI = InPooledBuffer->Buffer;
 	check(ResourceRHI);
 }
 
@@ -302,11 +342,11 @@ FRHIShaderResourceView* FRDGPooledBuffer::GetOrCreateSRV(FRDGBufferSRVDesc SRVDe
 
 	if (Desc.UnderlyingType == FRDGBufferDesc::EUnderlyingType::VertexBuffer)
 	{
-		RHIShaderResourceView = RHICreateShaderResourceView(VertexBuffer, SRVDesc.BytesPerElement, SRVDesc.Format);
+		RHIShaderResourceView = RHICreateShaderResourceView(Buffer, SRVDesc.BytesPerElement, SRVDesc.Format);
 	}
 	else if (Desc.UnderlyingType == FRDGBufferDesc::EUnderlyingType::StructuredBuffer)
 	{
-		RHIShaderResourceView = RHICreateShaderResourceView(StructuredBuffer);
+		RHIShaderResourceView = RHICreateShaderResourceView(Buffer);
 	}
 	else
 	{
@@ -329,11 +369,11 @@ FRHIUnorderedAccessView* FRDGPooledBuffer::GetOrCreateUAV(FRDGBufferUAVDesc UAVD
 
 	if (Desc.UnderlyingType == FRDGBufferDesc::EUnderlyingType::VertexBuffer)
 	{
-		RHIUnorderedAccessView = RHICreateUnorderedAccessView(VertexBuffer, UAVDesc.Format);
+		RHIUnorderedAccessView = RHICreateUnorderedAccessView(Buffer, UAVDesc.Format);
 	}
 	else if (Desc.UnderlyingType == FRDGBufferDesc::EUnderlyingType::StructuredBuffer)
 	{
-		RHIUnorderedAccessView = RHICreateUnorderedAccessView(StructuredBuffer, UAVDesc.bSupportsAtomicCounter, UAVDesc.bSupportsAppendBuffer);
+		RHIUnorderedAccessView = RHICreateUnorderedAccessView(Buffer, UAVDesc.bSupportsAtomicCounter, UAVDesc.bSupportsAppendBuffer);
 	}
 	else
 	{
diff --git a/Engine/Source/Runtime/RenderCore/Public/RenderGraphBuilder.h b/Engine/Source/Runtime/RenderCore/Public/RenderGraphBuilder.h
--- a/Engine/Source/Runtime/RenderCore/Public/RenderGraphBuilder.h
+++ b/Engine/Source/Runtime/RenderCore/Public/RenderGraphBuilder.h
@@ -162,6 +162,9 @@ public:
 	}
 #endif
 
+	/** Sets the current command list stat for all subsequent passes. */
+	void SetCommandListStat(TStatId StatId);
+
 	/** Queues a pooled render target extraction to happen at the end of graph execution. For graph-created textures, this extends
 	 *  the lifetime of the GPU resource until execution, at which point the pointer is filled. If specified, the texture is transitioned
 	 *  to the AccessFinal state, or kDefaultAccessFinal otherwise.
@@ -304,6 +307,8 @@ private:
 	bool bWaitedForTemporalEffect = false;
 #endif
 
+	IF_RDG_CMDLIST_STATS(TStatId CommandListStat);
+
 	void Compile();
 	void Clear();
 
diff --git a/Engine/Source/Runtime/RenderCore/Public/RenderGraphBuilder.inl b/Engine/Source/Runtime/RenderCore/Public/RenderGraphBuilder.inl
--- a/Engine/Source/Runtime/RenderCore/Public/RenderGraphBuilder.inl
+++ b/Engine/Source/Runtime/RenderCore/Public/RenderGraphBuilder.inl
@@ -211,6 +211,14 @@ inline void FRDGBuilder::QueueBufferExtraction(FRDGBufferRef Buffer, TRefCountPt
 	SetBufferAccessFinal(Buffer, AccessFinal);
 }
 
+inline void FRDGBuilder::SetCommandListStat(TStatId StatId)
+{
+#if RDG_CMDLIST_STATS
+	CommandListStat = StatId;
+	RHICmdList.SetCurrentStat(StatId);
+#endif
+}
+
 inline const TRefCountPtr<IPooledRenderTarget>& FRDGBuilder::GetPooledTexture(FRDGTextureRef Texture) const
 {
 	IF_RDG_ENABLE_DEBUG(UserValidation.ValidateGetPooledTexture(Texture));
diff --git a/Engine/Source/Runtime/RenderCore/Public/RenderGraphDefinitions.h b/Engine/Source/Runtime/RenderCore/Public/RenderGraphDefinitions.h
--- a/Engine/Source/Runtime/RenderCore/Public/RenderGraphDefinitions.h
+++ b/Engine/Source/Runtime/RenderCore/Public/RenderGraphDefinitions.h
@@ -52,6 +52,14 @@
 	#define IF_RDG_CPU_SCOPES(Op)
 #endif
 
+#define RDG_CMDLIST_STATS (STATS || ENABLE_STATNAMEDEVENTS)
+
+#if RDG_CMDLIST_STATS
+	#define IF_RDG_CMDLIST_STATS(Op) Op
+#else
+	#define IF_RDG_CMDLIST_STATS(Op)
+#endif
+
 /** ENUMS */
 
 enum class ERDGBuilderFlags
diff --git a/Engine/Source/Runtime/RenderCore/Public/RenderGraphPass.h b/Engine/Source/Runtime/RenderCore/Public/RenderGraphPass.h
--- a/Engine/Source/Runtime/RenderCore/Public/RenderGraphPass.h
+++ b/Engine/Source/Runtime/RenderCore/Public/RenderGraphPass.h
@@ -366,6 +366,8 @@ private:
 	FRHIGPUMask GPUMask;
 #endif
 
+	IF_RDG_CMDLIST_STATS(TStatId CommandListStat);
+
 	IF_RDG_CPU_SCOPES(FRDGCPUScopes CPUScopes);
 	IF_RDG_GPU_SCOPES(FRDGGPUScopes GPUScopes);
 
diff --git a/Engine/Source/Runtime/RenderCore/Public/RenderGraphResources.h b/Engine/Source/Runtime/RenderCore/Public/RenderGraphResources.h
--- a/Engine/Source/Runtime/RenderCore/Public/RenderGraphResources.h
+++ b/Engine/Source/Runtime/RenderCore/Public/RenderGraphResources.h
@@ -9,6 +9,21 @@
 struct FPooledRenderTarget;
 class FRenderTargetPool;
 
+/** Used for tracking pass producer / consumer edges in the graph for culling and pipe fencing. */
+struct FRDGProducerState
+{
+	/** Returns whether the next state is dependent on the last producer in the producer graph. */
+	static bool IsDependencyRequired(FRDGProducerState LastProducer, ERHIPipeline LastPipeline, FRDGProducerState NextState, ERHIPipeline NextPipeline);
+
+	FRDGProducerState() = default;
+
+	ERHIAccess Access = ERHIAccess::Unknown;
+	FRDGPassHandle PassHandle;
+	FRDGViewHandle NoUAVBarrierHandle;
+};
+
+using FRDGProducerStatesByPipeline = TRHIPipelineArray<FRDGProducerState>;
+
 /** Used for tracking the state of an individual subresource during execution. */
 struct FRDGSubresourceState
 {
@@ -721,7 +736,7 @@ private:
 	FRDGTextureTransientSubresourceStateIndirect MergeState;
 
 	/** Tracks pass producers for each subresource as the graph is built. */
-	TRDGTextureSubresourceArray<FRDGPassHandle> LastProducers;
+	TRDGTextureSubresourceArray<FRDGProducerStatesByPipeline> LastProducers;
 
 #if RDG_ENABLE_DEBUG
 	struct FRDGTextureDebugData* TextureDebugData = nullptr;
@@ -910,26 +925,12 @@ private:
 /** Descriptor for render graph tracked Buffer. */
 struct FRDGBufferDesc
 {
-	// Type of buffers to the RHI
 	enum class EUnderlyingType
 	{
 		VertexBuffer,
-		IndexBuffer, // not implemented yet.
-		StructuredBuffer,
+		StructuredBuffer
 	};
 
-	static inline FRHITransitionInfo::EType GetTransitionResourceType(EUnderlyingType Type)
-	{
-		switch (Type)
-		{
-		default: checkNoEntry();
-		case EUnderlyingType::VertexBuffer:
-		case EUnderlyingType::IndexBuffer:
-		case EUnderlyingType::StructuredBuffer:
-			return FRHITransitionInfo::EType::Buffer;
-		}
-	}
-
 	/** Create the descriptor for an indirect RHI call.
 	 *
 	 * Note, IndirectParameterStruct should be one of the:
@@ -1028,7 +1029,7 @@ struct FRDGBufferDesc
 	/** Bitfields describing the uses of that buffer. */
 	EBufferUsageFlags Usage = BUF_None;
 
-	/** The underlying RHI type to use. A bit of a work around because RHI still have 3 different objects. */
+	/** The underlying RHI type to use. */
 	EUnderlyingType UnderlyingType = EUnderlyingType::VertexBuffer;
 };
 
@@ -1086,19 +1087,22 @@ public:
 	/** Finds a SRV matching the descriptor in the cache or creates a new one and updates the cache. */
 	FRHIShaderResourceView* GetOrCreateSRV(FRDGBufferSRVDesc SRVDesc);
 
-	FRHIVertexBuffer* GetVertexBufferRHI() const
+	/** Returns the RHI buffer. */
+	FRHIBuffer* GetRHI() const
 	{
-		return VertexBuffer;
+		return Buffer;
 	}
 
-	FRHIIndexBuffer* GetIndexBufferRHI() const
+	UE_DEPRECATED(5.0, "Buffers types have been consolidated; use GetRHI() instead.")
+	FRHIBuffer* GetVertexBufferRHI() const
 	{
-		return IndexBuffer;
+		return Buffer;
 	}
 
-	FRHIStructuredBuffer* GetStructuredBufferRHI() const
+	UE_DEPRECATED(5.0, "Buffers types have been consolidated; use GetRHI() instead.")
+	FRHIBuffer* GetStructuredBufferRHI() const
 	{
-		return StructuredBuffer;
+		return Buffer;
 	}
 
 	uint32 GetRefCount() const
@@ -1166,9 +1170,7 @@ private:
 		: Desc(InDesc)
 	{}
 
-	FVertexBufferRHIRef VertexBuffer;
-	FIndexBufferRHIRef IndexBuffer;
-	FStructuredBufferRHIRef StructuredBuffer;
+	TRefCountPtr<FRHIBuffer> Buffer;
 	TMap<FRDGBufferUAVDesc, FUnorderedAccessViewRHIRef, FDefaultSetAllocator, TUAVFuncs<FRDGBufferUAVDesc, FUnorderedAccessViewRHIRef>> UAVs;
 	TMap<FRDGBufferSRVDesc, FShaderResourceViewRHIRef, FDefaultSetAllocator, TSRVFuncs<FRDGBufferSRVDesc, FShaderResourceViewRHIRef>> SRVs;
 
@@ -1197,9 +1199,6 @@ private:
 	friend FRDGBuffer;
 };
 
-UE_DEPRECATED(4.26, "FRDGPooledBuffer has been renamed to FRDGPooledBuffer.")
-typedef FRDGPooledBuffer FPooledRDGBuffer;
-
 /** A render graph tracked buffer. */
 class RENDERCORE_API FRDGBuffer final
 	: public FRDGParentResource
@@ -1211,26 +1210,31 @@ public:
 	//////////////////////////////////////////////////////////////////////////
 	//! The following methods may only be called during pass execution.
 
+	/** Returns the underlying RHI buffer resource */
+	FRHIBuffer* GetRHI() const
+	{
+		return static_cast<FRHIBuffer*>(FRDGParentResource::GetRHI());
+	}
+
 	/** Returns the buffer to use for indirect RHI calls. */
-	FRHIVertexBuffer* GetIndirectRHICallBuffer() const
+	FORCEINLINE FRHIBuffer* GetIndirectRHICallBuffer() const
 	{
-		checkf(Desc.UnderlyingType == FRDGBufferDesc::EUnderlyingType::VertexBuffer, TEXT("Buffer %s is not an underlying vertex buffer."), Name);
 		checkf(Desc.Usage & BUF_DrawIndirect, TEXT("Buffer %s was not flagged for indirect draw usage."), Name);
-		return static_cast<FRHIVertexBuffer*>(GetRHI());
+		return GetRHI();
 	}
 
 	/** Returns the buffer to use for RHI calls, eg RHILockVertexBuffer. */
-	FRHIVertexBuffer* GetRHIVertexBuffer() const
+	UE_DEPRECATED(5.0, "Buffers types have been consolidated; use GetRHI() instead.")
+	FORCEINLINE FRHIBuffer* GetRHIVertexBuffer() const
 	{
-		checkf(Desc.UnderlyingType == FRDGBufferDesc::EUnderlyingType::VertexBuffer, TEXT("Buffer %s is not an underlying vertex buffer."), Name);
-		return static_cast<FRHIVertexBuffer*>(GetRHI());
+		return GetRHI();
 	}
 
 	/** Returns the buffer to use for structured buffer calls. */
-	FRHIStructuredBuffer* GetRHIStructuredBuffer() const
+	UE_DEPRECATED(5.0, "Buffers types have been consolidated; use GetRHI() instead.")
+	FORCEINLINE FRHIBuffer* GetRHIStructuredBuffer() const
 	{
-		checkf(Desc.UnderlyingType == FRDGBufferDesc::EUnderlyingType::StructuredBuffer, TEXT("Buffer %s is not an underlying structured buffer."), Name);
-		return static_cast<FRHIStructuredBuffer*>(GetRHI());
+		return GetRHI();
 	}
 
 	//////////////////////////////////////////////////////////////////////////
@@ -1259,7 +1263,7 @@ private:
 	FRDGBufferHandle Handle;
 
 	/** Tracks the last pass that produced this resource as the graph is built. */
-	FRDGPassHandle LastProducer;
+	FRDGProducerStatesByPipeline LastProducer;
 
 	/** The next buffer to own the PooledBuffer allocation during execution. */
 	FRDGBufferHandle NextOwner;
diff --git a/Engine/Source/Runtime/RenderCore/Public/RenderGraphUtils.h b/Engine/Source/Runtime/RenderCore/Public/RenderGraphUtils.h
--- a/Engine/Source/Runtime/RenderCore/Public/RenderGraphUtils.h
+++ b/Engine/Source/Runtime/RenderCore/Public/RenderGraphUtils.h
@@ -657,12 +657,10 @@ FORCEINLINE void AddUntrackedAccessPassIfDebug(FRDGBuilder& GraphBuilder, Execut
 	AddUntrackedAccessPassIfDebug(GraphBuilder, {}, MoveTemp(ExecuteLambda));
 }
 
+UE_DEPRECATED(5.0, "AddSetCurrentStatPass is deprecated. Use GraphBuilder.SetCommandListStat instead.")
 FORCEINLINE void AddSetCurrentStatPass(FRDGBuilder& GraphBuilder, TStatId StatId)
 {
-	AddPassIfDebug(GraphBuilder, [StatId](FRHICommandListImmediate& RHICmdList)
-	{
-		RHICmdList.SetCurrentStat(StatId);
-	});
+	GraphBuilder.SetCommandListStat(StatId);
 }
 
 FORCEINLINE void AddDispatchToRHIThreadPass(FRDGBuilder& GraphBuilder)
diff --git a/Engine/Source/Runtime/VulkanRHI/Private/VulkanRHI.cpp b/Engine/Source/Runtime/VulkanRHI/Private/VulkanRHI.cpp
--- a/Engine/Source/Runtime/VulkanRHI/Private/VulkanRHI.cpp
+++ b/Engine/Source/Runtime/VulkanRHI/Private/VulkanRHI.cpp
@@ -206,7 +206,7 @@ FVulkanDynamicRHI::FVulkanDynamicRHI()
 	GConfig->GetInt(TEXT("TextureStreaming"), TEXT("PoolSizeVRAMPercentage"), GPoolSizeVRAMPercentage, GEngineIni);
 
 	// Copy source requires its own image layout.
-	EnumRemoveFlags(GRHITextureReadAccessMask, ERHIAccess::CopySrc);
+	EnumRemoveFlags(GRHIMergeableAccessMask, ERHIAccess::CopySrc);
 }
 
 void FVulkanDynamicRHI::Init()
