Fix for non-unity builds.

diff --git a/Engine/Source/Runtime/RenderCore/Private/RenderGraphBuilder.cpp b/Engine/Source/Runtime/RenderCore/Private/RenderGraphBuilder.cpp
--- a/Engine/Source/Runtime/RenderCore/Private/RenderGraphBuilder.cpp
+++ b/Engine/Source/Runtime/RenderCore/Private/RenderGraphBuilder.cpp
@@ -7,6 +7,45 @@
 #include "VisualizeTexture.h"
 #include "ProfilingDebugging/CsvProfiler.h"
 
+template <typename FunctionType>
+inline void EnumeratePassUAVs(const FRDGPass* Pass, FunctionType Function)
+{
+	Pass->GetParameters().Enumerate([&](FRDGParameter Parameter)
+	{
+		if (Parameter.IsUAV())
+		{
+			if (FRDGUnorderedAccessViewRef UAV = Parameter.GetAsUAV())
+			{
+				Function(UAV->GetRHI());
+			}
+		}
+	});
+}
+
+inline void BeginUAVOverlap(const FRDGPass* Pass, FRHIComputeCommandList& RHICmdList)
+{
+#if ENABLE_RHI_VALIDATION
+	TArray<FRHIUnorderedAccessView*, TInlineAllocator<8, FRDGArrayAllocator>> UAVs;
+	EnumeratePassUAVs(Pass, [&](FRHIUnorderedAccessView* UAV)
+	{
+		UAVs.Add(UAV);
+	});
+	RHICmdList.BeginUAVOverlap(UAVs);
+#endif
+}
+
+inline void EndUAVOverlap(const FRDGPass* Pass, FRHIComputeCommandList& RHICmdList)
+{
+#if ENABLE_RHI_VALIDATION
+	TArray<FRHIUnorderedAccessView*, TInlineAllocator<8, FRDGArrayAllocator>> UAVs;
+	EnumeratePassUAVs(Pass, [&](FRHIUnorderedAccessView* UAV)
+	{
+		UAVs.Add(UAV);
+	});
+	RHICmdList.EndUAVOverlap(UAVs);
+#endif
+}
+
 inline ERHIAccess MakeValidAccess(ERHIAccess Access)
 {
 	// If we find any write states in the access mask, remove all read-only states. This mainly exists
diff --git a/Engine/Source/Runtime/RenderCore/Private/RenderGraphPrivate.h b/Engine/Source/Runtime/RenderCore/Private/RenderGraphPrivate.h
--- a/Engine/Source/Runtime/RenderCore/Private/RenderGraphPrivate.h
+++ b/Engine/Source/Runtime/RenderCore/Private/RenderGraphPrivate.h
@@ -139,43 +139,4 @@ inline const TCHAR* GetEpilogueBarriersToBeginDebugName(ERHIPipeline Pipelines)
 	}
 #endif
 	return TEXT("");
-}
-
-template <typename FunctionType>
-inline void EnumeratePassUAVs(const FRDGPass* Pass, FunctionType Function)
-{
-	Pass->GetParameters().Enumerate([&](FRDGParameter Parameter)
-	{
-		if (Parameter.IsUAV())
-		{
-			if (FRDGUnorderedAccessViewRef UAV = Parameter.GetAsUAV())
-			{
-				Function(UAV->GetRHI());
-			}
-		}
-	});
-}
-
-inline void BeginUAVOverlap(const FRDGPass* Pass, FRHIComputeCommandList& RHICmdList)
-{
-#if ENABLE_RHI_VALIDATION
-	TArray<FRHIUnorderedAccessView*, TInlineAllocator<8, FRDGArrayAllocator>> UAVs;
-	EnumeratePassUAVs(Pass, [&](FRHIUnorderedAccessView* UAV)
-	{
-		UAVs.Add(UAV);
-	});
-	RHICmdList.BeginUAVOverlap(UAVs);
-#endif
-}
-
-inline void EndUAVOverlap(const FRDGPass* Pass, FRHIComputeCommandList& RHICmdList)
-{
-#if ENABLE_RHI_VALIDATION
-	TArray<FRHIUnorderedAccessView*, TInlineAllocator<8, FRDGArrayAllocator>> UAVs;
-	EnumeratePassUAVs(Pass, [&](FRHIUnorderedAccessView* UAV)
-	{
-		UAVs.Add(UAV);
-	});
-	RHICmdList.EndUAVOverlap(UAVs);
-#endif
 }
\ No newline at end of file
