Removed deprecated RDG texture desc create functions.

diff --git a/Engine/Source/Runtime/RenderCore/Public/RenderGraphResources.h b/Engine/Source/Runtime/RenderCore/Public/RenderGraphResources.h
--- a/Engine/Source/Runtime/RenderCore/Public/RenderGraphResources.h
+++ b/Engine/Source/Runtime/RenderCore/Public/RenderGraphResources.h
@@ -316,50 +316,6 @@ private:
 /** Descriptor used to create a render graph texture. */
 struct RENDERCORE_API FRDGTextureDesc
 {
-	UE_DEPRECATED(4.26, "FRDGTextureDesc has been refactored. Use Create2D instead.")
-	static FRDGTextureDesc Create2DDesc(
-		FIntPoint InExtent,
-		EPixelFormat InFormat,
-		const FClearValueBinding& InClearValue,
-		ETextureCreateFlags InFlags,
-		ETextureCreateFlags InTargetableFlags,
-		bool bInForceSeparateTargetAndShaderResource,
-		uint16 InNumMips = 1);
-
-	UE_DEPRECATED(4.26, "FRDGTextureDesc has been refactored. Use Create3D instead.")
-	static FRDGTextureDesc CreateVolumeDesc(
-		uint32 InSizeX,
-		uint32 InSizeY,
-		uint32 InSizeZ,
-		EPixelFormat InFormat,
-		const FClearValueBinding& InClearValue,
-		ETextureCreateFlags InFlags,
-		ETextureCreateFlags InTargetableFlags,
-		bool bInForceSeparateTargetAndShaderResource,
-		uint16 InNumMips = 1);
-
-	UE_DEPRECATED(4.26, "FRDGTextureDesc has been refactored. Use CreateCube instead.")
-	static FRDGTextureDesc CreateCubemapDesc(
-		uint32 InExtent,
-		EPixelFormat InFormat,
-		const FClearValueBinding& InClearValue,
-		ETextureCreateFlags InFlags,
-		ETextureCreateFlags InTargetableFlags,
-		bool bInForceSeparateTargetAndShaderResource,
-		uint32 InArraySize = 1,
-		uint16 InNumMips = 1);
-
-	UE_DEPRECATED(4.26, "FRDGTextureDesc has been refactored. Use CreateCubeArray instead.")
-	static FRDGTextureDesc CreateCubemapArrayDesc(
-		uint32 InExtent,
-		EPixelFormat InFormat,
-		const FClearValueBinding& InClearValue,
-		ETextureCreateFlags InFlags,
-		ETextureCreateFlags InTargetableFlags,
-		bool bInForceSeparateTargetAndShaderResource,
-		uint32 InArraySize,
-		uint16 InNumMips = 1);
-
 	static FRDGTextureDesc Create2D(
 		FIntPoint InExtent,
 		EPixelFormat InFormat,
diff --git a/Engine/Source/Runtime/RenderCore/Public/RenderGraphResources.inl b/Engine/Source/Runtime/RenderCore/Public/RenderGraphResources.inl
--- a/Engine/Source/Runtime/RenderCore/Public/RenderGraphResources.inl
+++ b/Engine/Source/Runtime/RenderCore/Public/RenderGraphResources.inl
@@ -172,58 +172,6 @@ FRDGTextureDesc Translate(const FPooledRenderTargetDesc& InDesc, ERenderTargetTe
 	return OutDesc;
 }
 
-inline FRDGTextureDesc FRDGTextureDesc::Create2DDesc(
-	FIntPoint InExtent,
-	EPixelFormat InFormat,
-	const FClearValueBinding& InClearValue,
-	ETextureCreateFlags InFlags,
-	ETextureCreateFlags InTargetableFlags,
-	bool bInForceSeparateTargetAndShaderResource,
-	uint16 InNumMips)
-{
-	return Translate(FPooledRenderTargetDesc::Create2DDesc(InExtent, InFormat, InClearValue, InFlags, InTargetableFlags, bInForceSeparateTargetAndShaderResource, InNumMips));
-}
-
-inline FRDGTextureDesc FRDGTextureDesc::CreateVolumeDesc(
-	uint32 InSizeX,
-	uint32 InSizeY,
-	uint32 InSizeZ,
-	EPixelFormat InFormat,
-	const FClearValueBinding& InClearValue,
-	ETextureCreateFlags InFlags,
-	ETextureCreateFlags InTargetableFlags,
-	bool bInForceSeparateTargetAndShaderResource,
-	uint16 InNumMips)
-{
-	return Translate(FPooledRenderTargetDesc::CreateVolumeDesc(InSizeX, InSizeY, InSizeZ, InFormat, InClearValue, InFlags, InTargetableFlags, bInForceSeparateTargetAndShaderResource, InNumMips));
-}
-
-inline FRDGTextureDesc FRDGTextureDesc::CreateCubemapDesc(
-	uint32 InExtent,
-	EPixelFormat InFormat,
-	const FClearValueBinding& InClearValue,
-	ETextureCreateFlags InFlags,
-	ETextureCreateFlags InTargetableFlags,
-	bool bInForceSeparateTargetAndShaderResource,
-	uint32 InArraySize,
-	uint16 InNumMips)
-{
-	return Translate(FPooledRenderTargetDesc::CreateCubemapDesc(InExtent, InFormat, InClearValue, InFlags, InTargetableFlags, bInForceSeparateTargetAndShaderResource, InArraySize, InNumMips));
-}
-
-inline FRDGTextureDesc FRDGTextureDesc::CreateCubemapArrayDesc(
-	uint32 InExtent,
-	EPixelFormat InFormat,
-	const FClearValueBinding& InClearValue,
-	ETextureCreateFlags InFlags,
-	ETextureCreateFlags InTargetableFlags,
-	bool bInForceSeparateTargetAndShaderResource,
-	uint32 InArraySize,
-	uint16 InNumMips)
-{
-	return Translate(FPooledRenderTargetDesc::CreateCubemapArrayDesc(InExtent, InFormat, InClearValue, InFlags, InTargetableFlags, bInForceSeparateTargetAndShaderResource, InArraySize, InNumMips));
-}
-
 inline FRDGTextureSubresourceRange FRDGTextureSRV::GetSubresourceRange() const
 {
 	FRDGTextureSubresourceRange Range = GetParent()->GetSubresourceRange();
