- Breaking a circular references that was on Mac but not Windows (editor compiles and runs on Windows, Mac still compiling)

diff --git a/Engine/Source/Runtime/DeveloperSettings/DeveloperSettings.Build.cs b/Engine/Source/Runtime/DeveloperSettings/DeveloperSettings.Build.cs
--- a/Engine/Source/Runtime/DeveloperSettings/DeveloperSettings.Build.cs
+++ b/Engine/Source/Runtime/DeveloperSettings/DeveloperSettings.Build.cs
@@ -11,8 +11,6 @@ public class DeveloperSettings : ModuleRules
 				"Core",
 				"CoreUObject",
                 "SlateCore",
-				"RenderCore",
-				"RHI"
 			}
 		);
 	}
diff --git a/Engine/Source/Runtime/RenderCore/RenderCore.Build.cs b/Engine/Source/Runtime/RenderCore/RenderCore.Build.cs
--- a/Engine/Source/Runtime/RenderCore/RenderCore.Build.cs
+++ b/Engine/Source/Runtime/RenderCore/RenderCore.Build.cs
@@ -7,21 +7,17 @@ public class RenderCore : ModuleRules
 	public RenderCore(ReadOnlyTargetRules Target) : base(Target)
     {
         PublicDependencyModuleNames.AddRange(new string[] { "RHI" });
-		
-        if (Target.bBuildEditor == true)
+
+		PrivateIncludePathModuleNames.AddRange(new string[] { "TargetPlatform" });
+
+		if (Target.bBuildEditor == true)
         {
-            PrivateDependencyModuleNames.Add("TargetPlatform");
+			DynamicallyLoadedModuleNames.Add("TargetPlatform");
 			// PakFileUtitilities due to file open order usage by the shader library
 			PrivateDependencyModuleNames.Add("PakFileUtilities");
 			// JSON is used for the asset info in the shader library
 			PrivateDependencyModuleNames.Add("Json");
 		}
-		else
-        {
-
-            PrivateIncludePathModuleNames.AddRange(new string[] { "TargetPlatform" });
-
-        }
 
         PrivateDependencyModuleNames.AddRange(new string[] { "Core", "Projects", "RHI", "ApplicationCore", "TraceLog" });
 
