Rename YUV to RGB matrices and fix errors

diff --git a/Engine/Plugins/Media/AjaMedia/Source/AjaMedia/Private/Player/AjaMediaTextureSample.h b/Engine/Plugins/Media/AjaMedia/Source/AjaMedia/Private/Player/AjaMediaTextureSample.h
--- a/Engine/Plugins/Media/AjaMedia/Source/AjaMedia/Private/Player/AjaMediaTextureSample.h
+++ b/Engine/Plugins/Media/AjaMedia/Source/AjaMedia/Private/Player/AjaMediaTextureSample.h
@@ -72,7 +72,7 @@ public:
 	 */
 	virtual const FMatrix& GetYUVToRGBMatrix() const override
 	{
-		return MediaShaders::YuvToRgbRec709Full;
+		return MediaShaders::YuvToRgbRec709Scaled;
 	}
 
 };
diff --git a/Engine/Plugins/Media/AvfMedia/Source/AvfMedia/Private/Player/AvfMediaVideoSampler.cpp b/Engine/Plugins/Media/AvfMedia/Source/AvfMedia/Private/Player/AvfMediaVideoSampler.cpp
--- a/Engine/Plugins/Media/AvfMedia/Source/AvfMedia/Private/Player/AvfMediaVideoSampler.cpp
+++ b/Engine/Plugins/Media/AvfMedia/Source/AvfMedia/Private/Player/AvfMediaVideoSampler.cpp
@@ -173,11 +173,11 @@ void FAvfMediaVideoSampler::SetOutput(AVPlayerItemVideoOutput* InOutput, float I
 	
 	if(bFullRange)
 	{
-		ColorTransform = &MediaShaders::YuvToRgbRec709Full;
+		ColorTransform = &MediaShaders::YuvToRgbRec709Scaled;
 	}
 	else
 	{
-		ColorTransform = &MediaShaders::YuvToRgbRec709;
+		ColorTransform = &MediaShaders::YuvToRgbRec709Unscaled;
 	}
 }
 
diff --git a/Engine/Plugins/Media/BlackmagicMedia/Source/BlackmagicMedia/Private/Player/BlackmagicMediaPlayer.h b/Engine/Plugins/Media/BlackmagicMedia/Source/BlackmagicMedia/Private/Player/BlackmagicMediaPlayer.h
--- a/Engine/Plugins/Media/BlackmagicMedia/Source/BlackmagicMedia/Private/Player/BlackmagicMediaPlayer.h
+++ b/Engine/Plugins/Media/BlackmagicMedia/Source/BlackmagicMedia/Private/Player/BlackmagicMediaPlayer.h
@@ -21,7 +21,7 @@ namespace BlackmagicMediaPlayerHelpers
 
 class FBlackmagicMediaTextureSample : public FMediaIOCoreTextureSampleBase
 {
-	virtual const FMatrix& GetYUVToRGBMatrix() const override { return MediaShaders::YuvToRgbRec709Full; }
+	virtual const FMatrix& GetYUVToRGBMatrix() const override { return MediaShaders::YuvToRgbRec709Scaled; }
 };
 
 class FBlackmagicMediaAudioSamplePool : public TMediaObjectPool<FMediaIOCoreAudioSampleBase> { };
diff --git a/Engine/Plugins/Media/MediaIOFramework/Source/MediaIOCore/Private/MediaCapture.cpp b/Engine/Plugins/Media/MediaIOFramework/Source/MediaIOCore/Private/MediaCapture.cpp
--- a/Engine/Plugins/Media/MediaIOFramework/Source/MediaIOCore/Private/MediaCapture.cpp
+++ b/Engine/Plugins/Media/MediaIOFramework/Source/MediaIOCore/Private/MediaCapture.cpp
@@ -853,7 +853,7 @@ void UMediaCapture::Capture_RenderThread(FRHICommandListImmediate& RHICmdList,
 					TShaderMapRef<FRGB8toUYVY8ConvertPS> ConvertShader(ShaderMap);
 					GraphicsPSOInit.BoundShaderState.PixelShaderRHI = ConvertShader.GetPixelShader();
 					SetGraphicsPipelineState(RHICmdList, GraphicsPSOInit);
-					ConvertShader->SetParameters(RHICmdList, SourceTexture, MediaShaders::RgbToYuvRec709Full, MediaShaders::YUVOffset8bits, bDoLinearToSRGB);
+					ConvertShader->SetParameters(RHICmdList, SourceTexture, MediaShaders::RgbToYuvRec709Scaled, MediaShaders::YUVOffset8bits, bDoLinearToSRGB);
 				}
 				break;
 				case EMediaCaptureConversionOperation::RGB10_TO_YUVv210_10BIT:
@@ -861,7 +861,7 @@ void UMediaCapture::Capture_RenderThread(FRHICommandListImmediate& RHICmdList,
 					TShaderMapRef<FRGB10toYUVv210ConvertPS> ConvertShader(ShaderMap);
 					GraphicsPSOInit.BoundShaderState.PixelShaderRHI = ConvertShader.GetPixelShader();
 					SetGraphicsPipelineState(RHICmdList, GraphicsPSOInit);
-					ConvertShader->SetParameters(RHICmdList, SourceTexture, MediaShaders::RgbToYuvRec709Full, MediaShaders::YUVOffset10bits, bDoLinearToSRGB);
+					ConvertShader->SetParameters(RHICmdList, SourceTexture, MediaShaders::RgbToYuvRec709Scaled, MediaShaders::YUVOffset10bits, bDoLinearToSRGB);
 				}
 				break;
 				case EMediaCaptureConversionOperation::INVERT_ALPHA:
diff --git a/Engine/Plugins/Media/WebMMedia/Source/WebMMedia/Private/Player/WebMVideoDecoder.cpp b/Engine/Plugins/Media/WebMMedia/Source/WebMMedia/Private/Player/WebMVideoDecoder.cpp
--- a/Engine/Plugins/Media/WebMMedia/Source/WebMMedia/Private/Player/WebMVideoDecoder.cpp
+++ b/Engine/Plugins/Media/WebMMedia/Source/WebMMedia/Private/Player/WebMVideoDecoder.cpp
@@ -275,7 +275,7 @@ void FWebMVideoDecoder::ConvertYUVToRGBAndSubmit(const FConvertParams& Params)
 			}
 
 			SetGraphicsPipelineState(CommandList, GraphicsPSOInit);
-			PixelShader->SetParameters(CommandList, DecodedY->GetTexture2D(), DecodedU->GetTexture2D(), DecodedV->GetTexture2D(), FIntPoint(Image->d_w, Image->d_h), MediaShaders::YuvToSrgbDefault, MediaShaders::YUVOffset8bits, true);
+			PixelShader->SetParameters(CommandList, DecodedY->GetTexture2D(), DecodedU->GetTexture2D(), DecodedV->GetTexture2D(), FIntPoint(Image->d_w, Image->d_h), MediaShaders::YuvToRgbRec709Scaled, MediaShaders::YUVOffset8bits, true);
 
 			// draw full-size quad
 			CommandList.SetViewport(0, 0, 0.0f, Image->d_w, Image->d_h, 1.0f);
diff --git a/Engine/Plugins/Media/WmfMedia/Source/WmfMedia/Public/WmfMediaHardwareVideoDecodingShaders.h b/Engine/Plugins/Media/WmfMedia/Source/WmfMedia/Public/WmfMediaHardwareVideoDecodingShaders.h
--- a/Engine/Plugins/Media/WmfMedia/Source/WmfMedia/Public/WmfMediaHardwareVideoDecodingShaders.h
+++ b/Engine/Plugins/Media/WmfMedia/Source/WmfMedia/Public/WmfMediaHardwareVideoDecodingShaders.h
@@ -58,7 +58,7 @@ public:
 		SetSamplerParameter(RHICmdList, ShaderRHI, PointClampedSamplerY, TStaticSamplerState<SF_Point, AM_Clamp, AM_Clamp, AM_Clamp>::GetRHI());
 		SetSamplerParameter(RHICmdList, ShaderRHI, BilinearClampedSamplerUV, InFilterUV ? TStaticSamplerState<SF_Bilinear, AM_Clamp, AM_Clamp, AM_Clamp>::GetRHI() : TStaticSamplerState<SF_Point, AM_Clamp, AM_Clamp, AM_Clamp>::GetRHI());
 
-		SetShaderValue(RHICmdList, ShaderRHI, ColorTransform, MediaShaders::CombineColorTransformAndOffset(MediaShaders::YuvToSrgbDefault, MediaShaders::YUVOffset8bits));
+		SetShaderValue(RHICmdList, ShaderRHI, ColorTransform, MediaShaders::CombineColorTransformAndOffset(MediaShaders::YuvToRgbRec709Scaled, MediaShaders::YUVOffset8bits));
 		SetShaderValue(RHICmdList, ShaderRHI, SrgbToLinear, InIsOutputSrgb ? 1 : 0); // Explicitly specify integer value, as using boolean falls over on XboxOne.
 	}
 
@@ -75,7 +75,7 @@ public:
 		SetSamplerParameter(RHICmdList, ShaderRHI, PointClampedSamplerY, TStaticSamplerState<SF_Point, AM_Clamp, AM_Clamp, AM_Clamp>::GetRHI());
 		SetSamplerParameter(RHICmdList, ShaderRHI, BilinearClampedSamplerUV, TStaticSamplerState<SF_Bilinear, AM_Clamp, AM_Clamp, AM_Clamp>::GetRHI());
 
-		SetShaderValue(RHICmdList, ShaderRHI, ColorTransform, MediaShaders::CombineColorTransformAndOffset(MediaShaders::YuvToSrgbDefault, MediaShaders::YUVOffset8bits));
+		SetShaderValue(RHICmdList, ShaderRHI, ColorTransform, MediaShaders::CombineColorTransformAndOffset(MediaShaders::YuvToRgbRec709Scaled, MediaShaders::YUVOffset8bits));
 		SetShaderValue(RHICmdList, ShaderRHI, SrgbToLinear, InIsOutputSrgb ? 1 : 0); // Explicitly specify integer value, as using boolean falls over on XboxOne.
 	}
 
diff --git a/Engine/Plugins/Runtime/WindowsMoviePlayer/Source/WindowsMoviePlayer/Private/WindowsMovieStreamer.cpp b/Engine/Plugins/Runtime/WindowsMoviePlayer/Source/WindowsMoviePlayer/Private/WindowsMovieStreamer.cpp
--- a/Engine/Plugins/Runtime/WindowsMoviePlayer/Source/WindowsMoviePlayer/Private/WindowsMovieStreamer.cpp
+++ b/Engine/Plugins/Runtime/WindowsMoviePlayer/Source/WindowsMoviePlayer/Private/WindowsMovieStreamer.cpp
@@ -156,7 +156,7 @@ void FMediaFoundationMovieStreamer::ConvertSample()
 			TShaderMapRef<FYUY2ConvertPS> ConvertShader(ShaderMap);
 			GraphicsPSOInit.BoundShaderState.PixelShaderRHI = ConvertShader.GetPixelShader();
 			SetGraphicsPipelineState(CommandList, GraphicsPSOInit);
-			ConvertShader->SetParameters(CommandList, InputTarget, OutputDim, MediaShaders::YuvToSrgbDefault, MediaShaders::YUVOffset8bits, bSampleIsOutputSrgb);
+			ConvertShader->SetParameters(CommandList, InputTarget, OutputDim, MediaShaders::YuvToRgbRec709Scaled, MediaShaders::YUVOffset8bits, bSampleIsOutputSrgb);
 		}
 		break;
 
diff --git a/Engine/Source/Runtime/Media/Public/IMediaTextureSample.h b/Engine/Source/Runtime/Media/Public/IMediaTextureSample.h
--- a/Engine/Source/Runtime/Media/Public/IMediaTextureSample.h
+++ b/Engine/Source/Runtime/Media/Public/IMediaTextureSample.h
@@ -271,16 +271,16 @@ public:
 	/**
 	 * Get the YUV to RGB conversion matrix.
 	 *
-	 * Equivalent to MediaShaders::YuvToSrgbDefault Matrix.
+	 * Equivalent to MediaShaders::YuvToRgbRec709Scaled Matrix. NOTE: previously in UE4 this was YuvToRgbRec601Scaled
 	 *
 	 * @return Conversion Matrix
 	 */
 	virtual const FMatrix& GetYUVToRGBMatrix() const
 	{
 		static const FMatrix DefaultMatrix(
-			FPlane(1.164383f, 0.000000f, 1.596027f, 0.000000f),
-			FPlane(1.164383f, -0.391762f, -0.812968f, 0.000000f),
-			FPlane(1.164383f, 2.017232f, 0.000000f, 0.000000f),
+			FPlane(1.16438356164f, 0.000000000000f, 1.792652263418f, 0.000000f),
+			FPlane(1.16438356164f, -0.213237021569f, -0.533004040142f, 0.000000f),
+			FPlane(1.16438356164f, 2.112419281991f, 0.000000000000f, 0.000000f),
 			FPlane(0.000000f, 0.000000f, 0.000000f, 0.000000f)
 		);
 
diff --git a/Engine/Source/Runtime/RenderCore/Private/MediaShaders.cpp b/Engine/Source/Runtime/RenderCore/Private/MediaShaders.cpp
--- a/Engine/Source/Runtime/RenderCore/Private/MediaShaders.cpp
+++ b/Engine/Source/Runtime/RenderCore/Private/MediaShaders.cpp
@@ -13,52 +13,67 @@ TGlobalResource<FMediaVertexDeclaration> GMediaVertexDeclaration;
 
 namespace MediaShaders
 {
-	const FMatrix YuvToSrgbJpeg = FMatrix(
-		FPlane(1.000000f,  0.000000f,  1.402000f, 0.000000f),
-		FPlane(1.000000f, -0.344140f, -0.714140f, 0.000000f),
-		FPlane(1.000000f,  1.772000f,  0.000000f, 0.000000f),
+	/**
+	 * YUV to RGB matrices for Rec601, Rec709, and Rec2020. Each are provided with a Scaled and Unscaled version.
+	 *
+	 * The scaling is according to SMPTE EG36, resulting in (255/219.0, 255/224.0, 255/224.0) being multiplied by the Unscaled matrix.
+	 */
+
+	const FMatrix YuvToRgbRec601Unscaled = FMatrix(
+		FPlane(1.000000f,  0.00000000000f,  1.4020000000f, 0.000000f),
+		FPlane(1.000000f, -0.34414362862f, -0.7141362862f, 0.000000f),
+		FPlane(1.000000f,  1.77200000000f,  0.0000000000f, 0.000000f),
 		FPlane(0.000000f,  0.000000f,  0.000000f, 0.000000f)
 	);
 
-	const FMatrix YuvToSrgbDefault = FMatrix(
-		FPlane(1.164383f,  0.000000f,  1.596027f, 0.000000f),
-		FPlane(1.164383f, -0.391762f, -0.812968f, 0.000000f),
-		FPlane(1.164383f,  2.017232f,  0.000000f, 0.000000f),
+	const FMatrix YuvToRgbRec601Scaled = FMatrix(
+		FPlane(1.16438356164f,  0.000000000000f,  1.596026785714f, 0.000000f),
+		FPlane(1.16438356164f, -0.391762290095f, -0.812967647238f, 0.000000f),
+		FPlane(1.16438356164f,  2.017232142857f,  0.000000000000f, 0.000000f),
 		FPlane(0.000000f,  0.000000f,  0.000000f, 0.000000f)
 	);
 
-	const FMatrix YuvToSrgbPs4 = FMatrix(
-		FPlane(1.164400f,  0.000000f,  1.792700f, 0.000000f),
-		FPlane(1.164400f, -0.213300f, -0.532900f, 0.000000f),
-		FPlane(1.164400f,  2.112400f,  0.000000f, 0.000000f),
-		FPlane(0.000000f,  0.000000f,  0.000000f, 0.000000f)
+	const FMatrix YuvToRgbRec709Unscaled = FMatrix(
+		FPlane(1.000000f,  0.000000000000f,  1.574721988260f, 0.000000f),
+		FPlane(1.000000f, -0.187314089535f, -0.468207470556f, 0.000000f),
+		FPlane(1.000000f,  1.855615369279f,  0.000000000000f, 0.000000f),
+		FPlane(0.000000f, 0.000000f, 0.000000f, 0.000000f)
 	);
 
-	const FMatrix YuvToSrgbRec601 = FMatrix(
-		FPlane(1.000000f, 0.000000f, 1.139830f, 0.000000f),
-		FPlane(1.000000f, -0.394650f, -0.580600f, 0.000000f),
-		FPlane(1.000000f, 2.032110, 0.000000, 0.000000f),
+	const FMatrix YuvToRgbRec709Scaled = FMatrix(
+		FPlane(1.16438356164f,  0.000000000000f,  1.792652263418f, 0.000000f),
+		FPlane(1.16438356164f, -0.213237021569f, -0.533004040142f, 0.000000f),
+		FPlane(1.16438356164f,  2.112419281991f,  0.000000000000f, 0.000000f),
 		FPlane(0.000000f, 0.000000f, 0.000000f, 0.000000f)
 	);
 
-	const FMatrix YuvToRgbRec709 = FMatrix(
-		FPlane(1.000000f, 0.000000f, 1.280330f, 0.000000f),
-		FPlane(1.000000f, -0.214820f, -0.380590f, 0.000000f),
-		FPlane(1.000000f, 2.127980f, 0.000000f, 0.000000f),
+	const FMatrix YuvToRgbRec2020Unscaled = FMatrix(
+		FPlane(1.000000f,  0.000000000000f,  1.474599575977f, 0.000000f),
+		FPlane(1.000000f, -0.164558057720f, -0.571355048803f, 0.000000f),
+		FPlane(1.000000f,  1.881396567060f,  0.000000000000f, 0.000000f),
 		FPlane(0.000000f, 0.000000f, 0.000000f, 0.000000f)
 	);
 
-	const FMatrix YuvToRgbRec709Full = FMatrix(
-		FPlane(1.164400f, 0.000000f, 1.792700f, 0.000000f),
-		FPlane(1.164400f, -0.213300f, -0.532900f, 0.000000f),
-		FPlane(1.164400f, 2.112400f, 0.000000f, 0.000000f),
+	const FMatrix YuvToRgbRec2020Scaled = FMatrix(
+		FPlane(1.16438356164f,  0.000000000000f,  1.678673624439f, 0.000000f),
+		FPlane(1.16438356164f, -0.187331717494f, -0.650426506450f, 0.000000f),
+		FPlane(1.16438356164f,  2.141768413395f,  0.000000000000f, 0.000000f),
 		FPlane(0.000000f, 0.000000f, 0.000000f, 0.000000f)
 	);
 
-	const FMatrix RgbToYuvRec709Full = FMatrix(
-		FPlane(0.182581f, 0.614210f, 0.062020f, 0.000000f),
-		FPlane(-0.100642f, -0.338566f, 0.439208f, 0.000000f),
-		FPlane(0.439227f, -0.398944f, -0.040283f, 0.000000f),
+	// YuvToSrgbPs4 is legacy. It is apparently a YuvToRgbRec709Scaled that was rounded to four digits after radix. 
+	const FMatrix YuvToSrgbPs4 = FMatrix(
+		FPlane(1.164400f,  0.000000f,  1.792700f, 0.000000f),
+		FPlane(1.164400f, -0.213300f, -0.532900f, 0.000000f),
+		FPlane(1.164400f,  2.112400f,  0.000000f, 0.000000f),
+		FPlane(0.000000f,  0.000000f,  0.000000f, 0.000000f)
+	);
+
+	// Inverse of YuvToRgbRec709Scaled
+	const FMatrix RgbToYuvRec709Scaled = FMatrix(
+		FPlane( 0.18261938151317932966f,  0.61420368882407283539f,  0.062000459074512487279f, 0.000000f),
+		FPlane(-0.10066136381366230790f, -0.33855432246084737891f,  0.439215686274509797830f, 0.000000f),
+		FPlane( 0.43921568627450979783f, -0.39894445418228852152f, -0.040271232092221290189f, 0.000000f),
 		FPlane(0.000000f, 0.000000f, 0.000000f, 0.000000f)
 	);
 
diff --git a/Engine/Source/Runtime/RenderCore/Public/MediaShaders.h b/Engine/Source/Runtime/RenderCore/Public/MediaShaders.h
--- a/Engine/Source/Runtime/RenderCore/Public/MediaShaders.h
+++ b/Engine/Source/Runtime/RenderCore/Public/MediaShaders.h
@@ -12,26 +12,29 @@
 
 namespace MediaShaders
 {
-	/** Color transform from YUV to sRGB (using values from MSDN). */
-	RENDERCORE_API extern const FMatrix YuvToSrgbDefault;
+	/** Color transform from YUV to Rec601 without range scaling. */
+	RENDERCORE_API extern const FMatrix YuvToRgbRec601Unscaled;
 
-	/** Color transform from YUV to sRGB (in JPEG color space). */
-	RENDERCORE_API extern const FMatrix YuvToSrgbJpeg;
+	/** Color transform from YUV Video Range to Rec601 Full Range. */
+	RENDERCORE_API extern const FMatrix YuvToRgbRec601Scaled;
 
-	/** Color transform from YUV to sRGB (using values from PS4 AvPlayer codec). */
-	RENDERCORE_API extern const FMatrix YuvToSrgbPs4;
+	/** Color transform from YUV to Rec709 without range scaling. */
+	RENDERCORE_API extern const FMatrix YuvToRgbRec709Unscaled;
+
+	/** Color transform from YUV Video Range to Rec709 Full Range. */
+	RENDERCORE_API extern const FMatrix YuvToRgbRec709Scaled;
 
-	/** Color transform from YUV to sRGB (in Rec. 601 color space). */
-	RENDERCORE_API extern const FMatrix YuvToSrgbRec601;
+	/** Color transform from YUV to Rec2020 without range scaling. */
+	RENDERCORE_API extern const FMatrix YuvToRgbRec2020Unscaled;
 
-	/** Color transform from YUV to sRGB (in Rec. 709 color space). */
-	RENDERCORE_API extern const FMatrix YuvToRgbRec709;
+	/** Color transform from YUV Video Range to Rec2020 Full Range. */
+	RENDERCORE_API extern const FMatrix YuvToRgbRec2020Scaled;
 
-	/** Color transform from YUV to RGB (in Rec. 709 color space, RGB full range) */
-	RENDERCORE_API extern const FMatrix YuvToRgbRec709Full;
+	/** Color transform from YUV to sRGB (using rounded values from PS4 AvPlayer codec). */
+	RENDERCORE_API extern const FMatrix YuvToSrgbPs4;
 
-	/** Color transform from RGB to YUV (in Rec. 709 color space, RGB full range) */
-	RENDERCORE_API extern const FMatrix RgbToYuvRec709Full;
+	/** Color transform from RGB to YUV (in Rec. 709 color space, including inversion of range scaling) */
+	RENDERCORE_API extern const FMatrix RgbToYuvRec709Scaled;
 
 	/** YUV Offset for 8 bit conversion (Computed as 16/255, 128/255, 128/255) */
 	RENDERCORE_API extern const FVector YUVOffset8bits;
