Fixing thin translucent rendering.

diff --git a/Engine/Source/Runtime/RenderCore/Private/ShaderMaterialDerivedHelpers.cpp b/Engine/Source/Runtime/RenderCore/Private/ShaderMaterialDerivedHelpers.cpp
--- a/Engine/Source/Runtime/RenderCore/Private/ShaderMaterialDerivedHelpers.cpp
+++ b/Engine/Source/Runtime/RenderCore/Private/ShaderMaterialDerivedHelpers.cpp
@@ -105,8 +105,7 @@ FShaderMaterialDerivedDefines RENDERCORE_API CalculateDerivedMaterialParameters(
 	Dst.STRATA_ENABLED = Mat.PROJECT_STRATA | Mat.MATERIAL_IS_STRATA;
 	Dst.SHADER_STRATA_TRANSLUCENT_ENABLED = Dst.STRATA_ENABLED && Dst.MATERIALBLENDING_ANY_TRANSLUCENT;
 	// As of today, all translucent Strata material forces dual source color blending. STRATA_TODO: optimize out dual color blending when not needed or requested
-	Dst.MATERIAL_WORKS_WITH_DUAL_SOURCE_COLOR_BLENDING = (Dst.MATERIAL_SHADINGMODEL_THIN_TRANSLUCENT || Dst.SHADER_STRATA_TRANSLUCENT_ENABLED);
-
+	Dst.MATERIAL_WORKS_WITH_DUAL_SOURCE_COLOR_BLENDING = (Mat.MATERIAL_SHADINGMODEL_THIN_TRANSLUCENT || Dst.SHADER_STRATA_TRANSLUCENT_ENABLED);
 
 	// There are 4 different ways of setting MRTs depending on which .usf file is the base shader, which sets defines which includes a different include file
 	// which may or may not override the current defines. Here is my best attempt at figuring out the logic.
diff --git a/Engine/Source/Runtime/RenderCore/Public/ShaderMaterial.h b/Engine/Source/Runtime/RenderCore/Public/ShaderMaterial.h
--- a/Engine/Source/Runtime/RenderCore/Public/ShaderMaterial.h
+++ b/Engine/Source/Runtime/RenderCore/Public/ShaderMaterial.h
@@ -228,8 +228,6 @@ struct FShaderMaterialDerivedDefines
 	bool WRITES_PRECSHADOWFACTOR_ZERO;
 	bool WRITES_PRECSHADOWFACTOR_TO_GBUFFER;
 
-	bool MATERIAL_SHADINGMODEL_THIN_TRANSLUCENT;
-
 	// If a primitive has static lighting, we assume it is not moving. If it is, it will be rerendered in an extra renderpass.
 	bool SUPPORTS_WRITING_VELOCITY_TO_BASE_PASS;
 	bool WRITES_VELOCITY_TO_GBUFFER;
