Cleaning up derived defines.

diff --git a/Engine/Source/Runtime/RenderCore/Private/ShaderMaterialDerivedHelpers.cpp b/Engine/Source/Runtime/RenderCore/Private/ShaderMaterialDerivedHelpers.cpp
--- a/Engine/Source/Runtime/RenderCore/Private/ShaderMaterialDerivedHelpers.cpp
+++ b/Engine/Source/Runtime/RenderCore/Private/ShaderMaterialDerivedHelpers.cpp
@@ -86,7 +86,6 @@ FShaderMaterialDerivedDefines RENDERCORE_API CalculateDerivedMaterialParameters(
 
 	//Materials also have to opt in to these features.
 	Dst.USE_EDITOR_COMPOSITING = (Dst.USE_EDITOR_SHADERS && Mat.EDITOR_PRIMITIVE_MATERIAL);
-	Dst.MATERIALBLENDING_ANY_TRANSLUCENT = (Mat.MATERIALBLENDING_TRANSLUCENT || Mat.MATERIALBLENDING_ADDITIVE || Mat.MATERIALBLENDING_MODULATE);
 
 	Dst.IS_MESHPARTICLE_FACTORY = (Lightmap.PARTICLE_MESH_FACTORY || Lightmap.NIAGARA_MESH_FACTORY);
 
@@ -103,7 +102,7 @@ FShaderMaterialDerivedDefines RENDERCORE_API CalculateDerivedMaterialParameters(
 
 	// matching the logic in BasePassPixelShader.usf
 	Dst.STRATA_ENABLED = Mat.PROJECT_STRATA | Mat.MATERIAL_IS_STRATA;
-	Dst.SHADER_STRATA_TRANSLUCENT_ENABLED = Dst.STRATA_ENABLED && Dst.MATERIALBLENDING_ANY_TRANSLUCENT;
+	Dst.SHADER_STRATA_TRANSLUCENT_ENABLED = Dst.STRATA_ENABLED && Mat.MATERIALBLENDING_ANY_TRANSLUCENT;
 	// As of today, all translucent Strata material forces dual source color blending. STRATA_TODO: optimize out dual color blending when not needed or requested
 	Dst.MATERIAL_WORKS_WITH_DUAL_SOURCE_COLOR_BLENDING = (Mat.MATERIAL_SHADINGMODEL_THIN_TRANSLUCENT || Dst.SHADER_STRATA_TRANSLUCENT_ENABLED);
 
diff --git a/Engine/Source/Runtime/RenderCore/Public/ShaderMaterial.h b/Engine/Source/Runtime/RenderCore/Public/ShaderMaterial.h
--- a/Engine/Source/Runtime/RenderCore/Public/ShaderMaterial.h
+++ b/Engine/Source/Runtime/RenderCore/Public/ShaderMaterial.h
@@ -267,7 +267,6 @@ struct FShaderMaterialDerivedDefines
 
 	bool USE_EDITOR_SHADERS;
 	bool USE_EDITOR_COMPOSITING;
-	bool MATERIALBLENDING_ANY_TRANSLUCENT;
 	bool IS_MESHPARTICLE_FACTORY;
 
 	bool PLATFORM_SUPPORTS_EDITOR_SHADERS;
