Refactor of uniform buffer binding to allow for runtime selection of static or per-shader binding. This is a necessary precursor to moving the View uniform buffer to use static bindings on mesh passes; it is used in so many places that both binding models are required to preserve compatibility. Added validation to catch when a uniform buffer is bound using both methods at the same time. Added CVar to RHI transition validation to control whether to break in the debugger on a transition failure.

diff --git a/Engine/Plugins/FX/Niagara/Source/NiagaraShader/Private/NiagaraShader.cpp b/Engine/Plugins/FX/Niagara/Source/NiagaraShader/Private/NiagaraShader.cpp
--- a/Engine/Plugins/FX/Niagara/Source/NiagaraShader/Private/NiagaraShader.cpp
+++ b/Engine/Plugins/FX/Niagara/Source/NiagaraShader/Private/NiagaraShader.cpp
@@ -462,7 +462,7 @@ void FNiagaraShaderType::AddReferencedUniformBufferIncludes(FShaderCompilerEnvir
 		{
 			if (It.Key() == StructIt->GetShaderVariableName())
 			{
-				StructIt->AddResourceTableEntries(OutEnvironment.ResourceTableMap, OutEnvironment.ResourceTableLayoutHashes, OutEnvironment.ResourceTableLayoutSlots);
+				StructIt->AddResourceTableEntries(OutEnvironment.ResourceTableMap, OutEnvironment.UniformBufferMap);
 			}
 		}
 	}
diff --git a/Engine/Source/Developer/Apple/MetalShaderFormat/Private/MetalDerivedData.cpp b/Engine/Source/Developer/Apple/MetalShaderFormat/Private/MetalDerivedData.cpp
--- a/Engine/Source/Developer/Apple/MetalShaderFormat/Private/MetalDerivedData.cpp
+++ b/Engine/Source/Developer/Apple/MetalShaderFormat/Private/MetalDerivedData.cpp
@@ -508,7 +508,7 @@ bool FMetalShaderOutputCooker::Build(TArray<uint8>& OutData)
 			TMap<FString, FMetalResourceTableEntry> ResourceTable;
 			if (IABTier >= 1)
 			{
-				for (auto Pair : Input.Environment.ResourceTableLayoutHashes)
+				for (auto Pair : Input.Environment.UniformBufferMap)
 				{
 					TableNames.Add(*Pair.Key);
 				}
diff --git a/Engine/Source/Developer/Apple/MetalShaderFormat/Private/MetalShaderCompiler.cpp b/Engine/Source/Developer/Apple/MetalShaderFormat/Private/MetalShaderCompiler.cpp
--- a/Engine/Source/Developer/Apple/MetalShaderFormat/Private/MetalShaderCompiler.cpp
+++ b/Engine/Source/Developer/Apple/MetalShaderFormat/Private/MetalShaderCompiler.cpp
@@ -767,8 +767,8 @@ void BuildMetalShaderOutput(
 	{
 		// Build the generic SRT for this shader.
 		FShaderCompilerResourceTable GenericSRT;
-		BuildResourceTableMapping(ShaderInput.Environment.ResourceTableMap, ShaderInput.Environment.ResourceTableLayoutHashes, UsedUniformBufferSlots, ShaderOutput.ParameterMap, GenericSRT);
-		CullGlobalUniformBuffers(ShaderInput.Environment.ResourceTableLayoutSlots, ShaderOutput.ParameterMap);
+		BuildResourceTableMapping(ShaderInput.Environment.ResourceTableMap, ShaderInput.Environment.UniformBufferMap, UsedUniformBufferSlots, ShaderOutput.ParameterMap, GenericSRT);
+		CullGlobalUniformBuffers(ShaderInput.Environment.UniformBufferMap, ShaderOutput.ParameterMap);
 
 		// Copy over the bits indicating which resource tables are active.
 		Header.Bindings.ShaderResourceTable.ResourceTableBits = GenericSRT.ResourceTableBits;
diff --git a/Engine/Source/Developer/ShaderCompilerCommon/Private/ShaderCompilerCommon.cpp b/Engine/Source/Developer/ShaderCompilerCommon/Private/ShaderCompilerCommon.cpp
--- a/Engine/Source/Developer/ShaderCompilerCommon/Private/ShaderCompilerCommon.cpp
+++ b/Engine/Source/Developer/ShaderCompilerCommon/Private/ShaderCompilerCommon.cpp
@@ -81,8 +81,8 @@ void BuildResourceTableTokenStream(const TArray<uint32>& InResourceMap, int32 Ma
 
 
 bool BuildResourceTableMapping(
-	const TMap<FString,FResourceTableEntry>& ResourceTableMap,
-	const TMap<FString, uint32>& ResourceTableLayoutHashes,
+	const TMap<FString, FResourceTableEntry>& ResourceTableMap,
+	const TMap<FString, FUniformBufferEntry>& UniformBufferMap,
 	TBitArray<>& UsedUniformBufferSlots,
 	FShaderParameterMap& ParameterMap,
 	FShaderCompilerResourceTable& OutSRT)
@@ -153,22 +153,26 @@ bool BuildResourceTableMapping(
 				return false;
 			}
 		}
+	}
 
-		// We have to do this separately from the resource table member check above. We want to include the hash even
-		// if the uniform buffer does not have any actual members used, because it will still be in the parameter map
-		// and certain platforms (like DX12) will pessimise and require them.
-		{
-			uint16 UniformBufferIndex = INDEX_NONE;
-			uint16 UBBaseIndex, UBSize;
+	// Emit hashes for all uniform buffers in the parameter map. We need to include the ones without resources as well
+	// (i.e. just constants), since the global uniform buffer bindings rely on valid hashes.
+	for (const auto& KeyValue : ParameterMap.GetParameterMap())
+	{
+		const FString& UniformBufferName = KeyValue.Key;
+		const FParameterAllocation& UniformBufferParameter = KeyValue.Value;
 
-			if (ParameterMap.FindParameterAllocation(*Entry.UniformBufferName, UniformBufferIndex, UBBaseIndex, UBSize))
+		if (UniformBufferParameter.Type == EShaderParameterType::UniformBuffer)
+		{
+			if (OutSRT.ResourceTableLayoutHashes.Num() <= UniformBufferParameter.BufferIndex)
 			{
-				while (OutSRT.ResourceTableLayoutHashes.Num() <= UniformBufferIndex)
-				{
-					OutSRT.ResourceTableLayoutHashes.Add(0);
-				}
+				OutSRT.ResourceTableLayoutHashes.SetNumZeroed(UniformBufferParameter.BufferIndex + 1);
+			}
 
-				OutSRT.ResourceTableLayoutHashes[UniformBufferIndex] = ResourceTableLayoutHashes.FindChecked(Entry.UniformBufferName);
+			// Data-driven uniform buffers will not have registered this information.
+			if (const FUniformBufferEntry* UniformBufferEntry = UniformBufferMap.Find(UniformBufferName))
+			{
+				OutSRT.ResourceTableLayoutHashes[UniformBufferParameter.BufferIndex] = UniformBufferEntry->LayoutHash;
 			}
 		}
 	}
@@ -177,15 +181,21 @@ bool BuildResourceTableMapping(
 	return true;
 }
 
-void CullGlobalUniformBuffers(const TMap<FString, FString>& ResourceTableLayoutSlots, FShaderParameterMap& ParameterMap)
+void CullGlobalUniformBuffers(const TMap<FString, FUniformBufferEntry>& UniformBufferMap, FShaderParameterMap& ParameterMap)
 {
 	TArray<FString> ParameterNames;
 	ParameterMap.GetAllParameterNames(ParameterNames);
 
 	for (const FString& Name : ParameterNames)
 	{
-		if (ResourceTableLayoutSlots.Contains(*Name))
+		if (const FUniformBufferEntry* UniformBufferEntry = UniformBufferMap.Find(*Name))
 		{
+			// A uniform buffer that is bound per-shader keeps its allocation in the map.
+			if (EnumHasAnyFlags(UniformBufferEntry->BindingFlags, EUniformBufferBindingFlags::Shader))
+			{
+				continue;
+			}
+
 			ParameterMap.RemoveParameterAllocation(*Name);
 		}
 	}
@@ -1094,7 +1104,7 @@ void FShaderParameterParser::ExtractFileAndLine(int32 PragamLineoffset, int32 Li
 void RemoveUniformBuffersFromSource(const FShaderCompilerEnvironment& Environment, FString& PreprocessedShaderSource)
 {
 	TMap<FString, TArray<FUniformBufferMemberInfo>> UniformBufferNameToMembers;
-	UniformBufferNameToMembers.Reserve(Environment.ResourceTableLayoutHashes.Num());
+	UniformBufferNameToMembers.Reserve(Environment.UniformBufferMap.Num());
 
 	// Build a mapping from uniform buffer name to its members
 	{
@@ -1557,9 +1567,9 @@ namespace CrossCompiler
 	FString CreateResourceTableFromEnvironment(const FShaderCompilerEnvironment& Environment)
 	{
 		FString Line = TEXT("\n#if 0 /*BEGIN_RESOURCE_TABLES*/\n");
-		for (auto Pair : Environment.ResourceTableLayoutHashes)
+		for (auto Pair : Environment.UniformBufferMap)
 		{
-			Line += FString::Printf(TEXT("%s, %d\n"), *Pair.Key, Pair.Value);
+			Line += FString::Printf(TEXT("%s, %d\n"), *Pair.Key, Pair.Value.LayoutHash);
 		}
 		Line += TEXT("NULL, 0\n");
 		for (auto Pair : Environment.ResourceTableMap)
@@ -1621,7 +1631,9 @@ namespace CrossCompiler
 			{
 				break;
 			}
-			OutEnvironment.ResourceTableLayoutHashes.FindOrAdd(UB) = (uint32)Hash;
+
+			FUniformBufferEntry& UniformBufferEntry = OutEnvironment.UniformBufferMap.FindOrAdd(UB);
+			UniformBufferEntry.LayoutHash = (uint32)Hash;
 		}
 
 		while (Ptr < PtrEnd)
diff --git a/Engine/Source/Developer/ShaderCompilerCommon/Public/ShaderCompilerCommon.h b/Engine/Source/Developer/ShaderCompilerCommon/Public/ShaderCompilerCommon.h
--- a/Engine/Source/Developer/ShaderCompilerCommon/Public/ShaderCompilerCommon.h
+++ b/Engine/Source/Developer/ShaderCompilerCommon/Public/ShaderCompilerCommon.h
@@ -17,14 +17,14 @@
  */
 extern SHADERCOMPILERCOMMON_API bool BuildResourceTableMapping(
 		const TMap<FString,FResourceTableEntry>& ResourceTableMap,
-		const TMap<FString,uint32>& ResourceTableLayoutHashes,
+		const TMap<FString,FUniformBufferEntry>& UniformBufferMap,
 		TBitArray<>& UsedUniformBufferSlots,
 		FShaderParameterMap& ParameterMap,
 		FShaderCompilerResourceTable& OutSRT
 	);
 
 /** Culls global uniform buffer entries from the parameter map. */
-extern SHADERCOMPILERCOMMON_API void CullGlobalUniformBuffers(const TMap<FString, FString>& ResourceTableLayoutSlots, FShaderParameterMap& ParameterMap);
+extern SHADERCOMPILERCOMMON_API void CullGlobalUniformBuffers(const TMap<FString, FUniformBufferEntry>& UniformBufferMap, FShaderParameterMap& ParameterMap);
 
 /**
  * Builds a token stream out of the resource map. The resource map is one
diff --git a/Engine/Source/Developer/ShaderFormatOpenGL/Private/OpenGLShaderCompiler.cpp b/Engine/Source/Developer/ShaderFormatOpenGL/Private/OpenGLShaderCompiler.cpp
--- a/Engine/Source/Developer/ShaderFormatOpenGL/Private/OpenGLShaderCompiler.cpp
+++ b/Engine/Source/Developer/ShaderFormatOpenGL/Private/OpenGLShaderCompiler.cpp
@@ -938,8 +938,8 @@ void FOpenGLFrontend::BuildShaderOutput(
 	{
 		// Build the generic SRT for this shader.
 		FShaderCompilerResourceTable GenericSRT;
-		BuildResourceTableMapping(ShaderInput.Environment.ResourceTableMap, ShaderInput.Environment.ResourceTableLayoutHashes, UsedUniformBufferSlots, ShaderOutput.ParameterMap, GenericSRT);
-		CullGlobalUniformBuffers(ShaderInput.Environment.ResourceTableLayoutSlots, ShaderOutput.ParameterMap);
+		BuildResourceTableMapping(ShaderInput.Environment.ResourceTableMap, ShaderInput.Environment.UniformBufferMap, UsedUniformBufferSlots, ShaderOutput.ParameterMap, GenericSRT);
+		CullGlobalUniformBuffers(ShaderInput.Environment.UniformBufferMap, ShaderOutput.ParameterMap);
 
 		// Copy over the bits indicating which resource tables are active.
 		Header.Bindings.ShaderResourceTable.ResourceTableBits = GenericSRT.ResourceTableBits;
diff --git a/Engine/Source/Developer/VulkanShaderFormat/Private/VulkanShaderCompiler.cpp b/Engine/Source/Developer/VulkanShaderFormat/Private/VulkanShaderCompiler.cpp
--- a/Engine/Source/Developer/VulkanShaderFormat/Private/VulkanShaderCompiler.cpp
+++ b/Engine/Source/Developer/VulkanShaderFormat/Private/VulkanShaderCompiler.cpp
@@ -578,8 +578,8 @@ static void AddUniformBuffer(FOLDVulkanCodeHeader& OLDHeader,
 		HeaderUBIndex = BindingIndex;
 	}
 	FVulkanShaderHeader::FUniformBufferInfo& UBInfo = OutHeader.UniformBuffers[HeaderUBIndex];
-	const uint32* LayoutHash = ShaderInput.Environment.ResourceTableLayoutHashes.Find(UBName);
-	UBInfo.LayoutHash = LayoutHash ? *LayoutHash : 0;
+	const FUniformBufferEntry* UniformBufferEntry = ShaderInput.Environment.UniformBufferMap.Find(UBName);
+	UBInfo.LayoutHash = UniformBufferEntry ? UniformBufferEntry->LayoutHash : 0;
 #if VULKAN_ENABLE_BINDING_DEBUG_NAMES
 	UBInfo.DebugName = UBName;
 #endif
@@ -625,7 +625,7 @@ static void AddUniformBuffer(FOLDVulkanCodeHeader& OLDHeader,
 	}
 
 	// Currently we don't support mismatched uniform buffer layouts/cbuffers with resources!
-	check(LayoutHash || UBInfo.ResourceEntries.Num() == 0);
+	check(UniformBufferEntry || UBInfo.ResourceEntries.Num() == 0);
 
 	InOutParameterMap.RemoveParameterAllocation(*UBName);
 	InOutParameterMap.AddParameterAllocation(*UBName, HeaderUBIndex, (uint16)FVulkanShaderHeader::UniformBuffer, 1, EShaderParameterType::UniformBuffer);
@@ -1451,7 +1451,7 @@ static void BuildShaderOutput(
 	{
 		// Build the generic SRT for this shader.
 		FShaderCompilerResourceTable GenericSRT;
-		if (!BuildResourceTableMapping(ShaderInput.Environment.ResourceTableMap, ShaderInput.Environment.ResourceTableLayoutHashes, UsedUniformBufferSlots, ShaderOutput.ParameterMap, /*MaxBoundResourceTable, */GenericSRT))
+		if (!BuildResourceTableMapping(ShaderInput.Environment.ResourceTableMap, ShaderInput.Environment.UniformBufferMap, UsedUniformBufferSlots, ShaderOutput.ParameterMap, /*MaxBoundResourceTable, */GenericSRT))
 		{
 			ShaderOutput.Errors.Add(TEXT("Internal error on BuildResourceTableMapping."));
 			return;
@@ -1529,7 +1529,7 @@ static void BuildShaderOutput(
 		}
 	}
 
-	CullGlobalUniformBuffers(ShaderInput.Environment.ResourceTableLayoutSlots, ShaderOutput.ParameterMap);
+	CullGlobalUniformBuffers(ShaderInput.Environment.UniformBufferMap, ShaderOutput.ParameterMap);
 }
 
 
diff --git a/Engine/Source/Developer/Windows/ShaderFormatD3D/Private/D3DShaderCompiler.inl b/Engine/Source/Developer/Windows/ShaderFormatD3D/Private/D3DShaderCompiler.inl
--- a/Engine/Source/Developer/Windows/ShaderFormatD3D/Private/D3DShaderCompiler.inl
+++ b/Engine/Source/Developer/Windows/ShaderFormatD3D/Private/D3DShaderCompiler.inl
@@ -452,12 +452,12 @@ inline void GenerateFinalOutput(TRefCountPtr<TBlob>& CompressedData,
 	{
 		// Build the generic SRT for this shader.
 		FShaderCompilerResourceTable GenericSRT;
-		BuildResourceTableMapping(Input.Environment.ResourceTableMap, Input.Environment.ResourceTableLayoutHashes, UsedUniformBufferSlots, Output.ParameterMap, GenericSRT);
+		BuildResourceTableMapping(Input.Environment.ResourceTableMap, Input.Environment.UniformBufferMap, UsedUniformBufferSlots, Output.ParameterMap, GenericSRT);
 
 		// Ray generation shaders rely on a different binding model that aren't compatible with global uniform buffers.
 		if (Input.Target.Frequency != SF_RayGen)
 		{
-			CullGlobalUniformBuffers(Input.Environment.ResourceTableLayoutSlots, Output.ParameterMap);
+			CullGlobalUniformBuffers(Input.Environment.UniformBufferMap, Output.ParameterMap);
 		}
 
 		if (UniformBufferNames.Num() < GenericSRT.ResourceTableLayoutHashes.Num())
@@ -469,9 +469,16 @@ inline void GenerateFinalOutput(TRefCountPtr<TBlob>& CompressedData,
 		{
 			if (GenericSRT.ResourceTableLayoutHashes[Index] != 0 && UniformBufferNames[Index].Len() == 0)
 			{
-				auto* Name = Input.Environment.ResourceTableLayoutHashes.FindKey(GenericSRT.ResourceTableLayoutHashes[Index]);
-				check(Name);
-				UniformBufferNames[Index] = *Name;
+				for (const auto& KeyValue : Input.Environment.UniformBufferMap)
+				{
+					const FUniformBufferEntry& UniformBufferEntry = KeyValue.Value;
+
+					if (UniformBufferEntry.LayoutHash == GenericSRT.ResourceTableLayoutHashes[Index])
+					{
+						UniformBufferNames[Index] = KeyValue.Key;
+						break;
+					}
+				}
 			}
 		}
 
diff --git a/Engine/Source/Runtime/Engine/Private/Materials/MaterialUniformExpressions.cpp b/Engine/Source/Runtime/Engine/Private/Materials/MaterialUniformExpressions.cpp
--- a/Engine/Source/Runtime/Engine/Private/Materials/MaterialUniformExpressions.cpp
+++ b/Engine/Source/Runtime/Engine/Private/Materials/MaterialUniformExpressions.cpp
@@ -518,6 +518,7 @@ FShaderParametersMetadata* FUniformExpressionSet::CreateBufferStruct()
 	const uint32 StructSize = Align(NextMemberOffset, SHADER_PARAMETER_STRUCT_ALIGNMENT);
 	FShaderParametersMetadata* UniformBufferStruct = new FShaderParametersMetadata(
 		FShaderParametersMetadata::EUseCase::DataDrivenUniformBuffer,
+		EUniformBufferBindingFlags::Shader,
 		TEXT("Material"),
 		TEXT("MaterialUniforms"),
 		TEXT("Material"),
diff --git a/Engine/Source/Runtime/Engine/Private/Materials/ParameterCollection.cpp b/Engine/Source/Runtime/Engine/Private/Materials/ParameterCollection.cpp
--- a/Engine/Source/Runtime/Engine/Private/Materials/ParameterCollection.cpp
+++ b/Engine/Source/Runtime/Engine/Private/Materials/ParameterCollection.cpp
@@ -420,6 +420,7 @@ void UMaterialParameterCollection::CreateBufferStruct()
 	// (and the hlsl cbuffers are named MaterialCollection0, etc, so the names don't match the layout)
 	UniformBufferStruct = MakeUnique<FShaderParametersMetadata>(
 		FShaderParametersMetadata::EUseCase::DataDrivenUniformBuffer,
+		EUniformBufferBindingFlags::Shader,
 		TEXT("MaterialCollection"),
 		TEXT("MaterialCollection"),
 		TEXT("MaterialCollection"),
diff --git a/Engine/Source/Runtime/RHI/Private/RHIValidation.cpp b/Engine/Source/Runtime/RHI/Private/RHIValidation.cpp
--- a/Engine/Source/Runtime/RHI/Private/RHIValidation.cpp
+++ b/Engine/Source/Runtime/RHI/Private/RHIValidation.cpp
@@ -11,6 +11,18 @@
 
 #if ENABLE_RHI_VALIDATION
 
+namespace RHIValidation
+{
+	int32 GBreakOnTransitionError = 1;
+	FAutoConsoleVariableRef CVarBreakOnTransitionError(
+		TEXT("r.RHIValidation.DebugBreak.Transitions"),
+		GBreakOnTransitionError,
+		TEXT("Controls whether the debugger should break when a validation error is encountered.\n")
+		TEXT(" 0: disabled;\n")
+		TEXT(" 1: break in the debugger if a validation error is encountered."),
+		ECVF_RenderThreadSafe);
+}
+
 bool GRHIValidationEnabled = false;
 
 // When set to 1, callstack for each uniform buffer allocation will be tracked 
@@ -252,7 +264,7 @@ void FValidationRHI::ReportValidationFailure(const TCHAR* InMessage)
 
 	UE_LOG(LogRHI, Error, TEXT("%s"), InMessage);
 
-	if (FPlatformMisc::IsDebuggerPresent())
+	if (FPlatformMisc::IsDebuggerPresent() && RHIValidation::GBreakOnTransitionError)
 	{
 		// Print the message again using the debug output function, because UE_LOG doesn't always reach
 		// the VS output window before the breakpoint is triggered, despite the log flush call below.
@@ -274,9 +286,9 @@ void FValidationComputeContext::FState::Reset()
 	ComputePassName.Reset();
 	bComputeShaderSet = false;
 	TrackerInstance.ResetAllUAVState();
+	GlobalUniformBuffers.Reset();
 }
 
-
 FValidationContext::FValidationContext()
 	: RHIContext(nullptr)
 {
@@ -303,10 +315,42 @@ void FValidationContext::FState::Reset()
 	ComputePassName.Reset();
 	bComputeShaderSet = false;
 	TrackerInstance.ResetAllUAVState();
+	GlobalUniformBuffers.Reset();
 }
 
 namespace RHIValidation
 {
+	void FGlobalUniformBuffers::Reset()
+	{
+		Bindings.Reset();
+		check(!bInSetPipelineStateCall);
+	}
+
+	void FGlobalUniformBuffers::ValidateSetShaderUniformBuffer(FRHIUniformBuffer* UniformBuffer)
+	{
+		check(UniformBuffer);
+		UniformBuffer->ValidateLifeTime();
+
+		// Skip validating global uniform buffers that are set internally by the RHI as part of the pipeline state.
+		if (bInSetPipelineStateCall)
+		{
+			return;
+		}
+
+		const FRHIUniformBufferLayout& Layout = UniformBuffer->GetLayout();
+
+		checkf(EnumHasAnyFlags(Layout.BindingFlags, EUniformBufferBindingFlags::Shader), TEXT("Uniform buffer '%s' does not have the 'Shader' binding flag."), *Layout.GetDebugName());
+
+		if (Layout.StaticSlot < Bindings.Num())
+		{
+			check(Layout.BindingFlags == EUniformBufferBindingFlags::StaticAndShader);
+
+			ensureMsgf(Bindings[Layout.StaticSlot] == nullptr,
+				TEXT("Uniform buffer '%s' was bound statically and is now being bound on a specific RHI shader. Only one binding model should be used at a time."),
+				*Layout.GetDebugName());
+		}
+	}
+
 	ERHIAccess DecayResourceAccess(ERHIAccess AccessMask, ERHIAccess RequiredAccess, bool bAllowUAVOverlap)
 	{
 		using T = __underlying_type(ERHIAccess);
diff --git a/Engine/Source/Runtime/RHI/Public/RHIContext.h b/Engine/Source/Runtime/RHI/Public/RHIContext.h
--- a/Engine/Source/Runtime/RHI/Public/RHIContext.h
+++ b/Engine/Source/Runtime/RHI/Public/RHIContext.h
@@ -61,6 +61,7 @@ public:
 
 		Slots.Add(Slot);
 		UniformBuffers.Add(UniformBuffer);
+		SlotCount = FMath::Max(SlotCount, Slot + 1);
 	}
 
 	inline void TryAddUniformBuffer(FRHIUniformBuffer* UniformBuffer)
@@ -86,10 +87,27 @@ public:
 		return Slots[Index];
 	}
 
+	int32 GetSlotCount() const
+	{
+		return SlotCount;
+	}
+
+	void Bind(TArray<FRHIUniformBuffer*>& Bindings) const
+	{
+		Bindings.Reset();
+		Bindings.SetNumZeroed(SlotCount);
+
+		for (int32 Index = 0; Index < UniformBuffers.Num(); ++Index)
+		{
+			Bindings[Slots[Index]] = UniformBuffers[Index];
+		}
+	}
+
 private:
 	static const uint32 InlineUniformBufferCount = 8;
 	TArray<FUniformBufferStaticSlot, TInlineAllocator<InlineUniformBufferCount>> Slots;
 	TArray<FRHIUniformBuffer*, TInlineAllocator<InlineUniformBufferCount>> UniformBuffers;
+	int32 SlotCount = 0;
 };
 
 /** Context that is capable of doing Compute work.  Can be async or compute on the gfx pipe. */
diff --git a/Engine/Source/Runtime/RHI/Public/RHIDefinitions.h b/Engine/Source/Runtime/RHI/Public/RHIDefinitions.h
--- a/Engine/Source/Runtime/RHI/Public/RHIDefinitions.h
+++ b/Engine/Source/Runtime/RHI/Public/RHIDefinitions.h
@@ -994,6 +994,25 @@ enum EUniformBufferBaseType : uint8
 static_assert(EUniformBufferBaseType_Num <= (1 << EUniformBufferBaseType_NumBits), "EUniformBufferBaseType_Num will not fit on EUniformBufferBaseType_NumBits");
 DECLARE_INTRINSIC_TYPE_LAYOUT(EUniformBufferBaseType);
 
+/** The list of flags declaring which binding models are allowed for a uniform buffer layout. */
+enum class EUniformBufferBindingFlags : uint8
+{
+	/** If set, the uniform buffer can be bound as an RHI shader parameter on an RHI shader (i.e. RHISetShaderUniformBuffer). */
+	Shader = 1 << 0,
+
+	/** If set, the uniform buffer can be bound globally through a static slot (i.e. RHISetGlobalUniformBuffers). */
+	Static = 1 << 1,
+
+	/** If set, the uniform buffer can be bound globally or per-shader, depending on the use case. Only one binding model should be
+	 *  used at a time, and RHI validation will emit an error if both are used for a particular uniform buffer at the same time. This
+	 *  is designed for difficult cases where a fixed single binding model would produce an unnecessary maintenance burden. Using this
+	 *  disables some RHI validation errors for global bindings, so use with care.
+	 */
+	StaticAndShader = Static | Shader
+};
+ENUM_CLASS_FLAGS(EUniformBufferBindingFlags);
+DECLARE_INTRINSIC_TYPE_LAYOUT(EUniformBufferBindingFlags);
+
 /** Numerical type used to store the static slot indices. */
 using FUniformBufferStaticSlot = uint8;
 
diff --git a/Engine/Source/Runtime/RHI/Public/RHIResources.h b/Engine/Source/Runtime/RHI/Public/RHIResources.h
--- a/Engine/Source/Runtime/RHI/Public/RHIResources.h
+++ b/Engine/Source/Runtime/RHI/Public/RHIResources.h
@@ -603,7 +603,8 @@ public:
 
 	void ComputeHash()
 	{
-		uint32 TmpHash = ConstantBufferSize << 16 | static_cast<uint32>(StaticSlot);
+		// Static slot is not stable. Just track whether we have one at all.
+		uint32 TmpHash = ConstantBufferSize << 16 | static_cast<uint32>(BindingFlags) << 8 | static_cast<uint32>(StaticSlot != MAX_UNIFORM_BUFFER_STATIC_SLOTS);
 
 		for (int32 ResourceIndex = 0; ResourceIndex < Resources.Num(); ResourceIndex++)
 		{
@@ -649,6 +650,7 @@ public:
 	{
 		ConstantBufferSize = Source.ConstantBufferSize;
 		StaticSlot = Source.StaticSlot;
+		BindingFlags = Source.BindingFlags;
 		Resources = Source.Resources;
 		Name = Source.Name;
 		Hash = Source.Hash;
@@ -680,6 +682,7 @@ public:
 		Ar << Ref.StaticSlot;
 		Ar << Ref.RenderTargetsOffset;
 		Ar << Ref.bHasNonGraphOutputs;
+		Ar << Ref.BindingFlags;
 		Ar << Ref.Resources;
 		Ar << Ref.GraphResources;
 		Ar << Ref.GraphTextures;
@@ -703,6 +706,9 @@ public:
 	/** Whether this layout may contain non-render-graph outputs (e.g. RHI UAVs). */
 	LAYOUT_FIELD_INITIALIZED(bool, bHasNonGraphOutputs, false);
 
+	/** The binding flags describing how this resource can be bound to the RHI. */
+	LAYOUT_FIELD_INITIALIZED(EUniformBufferBindingFlags, BindingFlags, EUniformBufferBindingFlags::Shader);
+
 	/** The list of all resource inlined into the shader parameter structure. */
 	LAYOUT_FIELD(TMemoryImageArray<FResourceParameter>, Resources);
 
@@ -741,6 +747,7 @@ inline bool operator==(const FRHIUniformBufferLayout& A, const FRHIUniformBuffer
 {
 	return A.ConstantBufferSize == B.ConstantBufferSize
 		&& A.StaticSlot == B.StaticSlot
+		&& A.BindingFlags == B.BindingFlags
 		&& A.Resources == B.Resources;
 }
 
diff --git a/Engine/Source/Runtime/RHI/Public/RHIValidationCommon.h b/Engine/Source/Runtime/RHI/Public/RHIValidationCommon.h
--- a/Engine/Source/Runtime/RHI/Public/RHIValidationCommon.h
+++ b/Engine/Source/Runtime/RHI/Public/RHIValidationCommon.h
@@ -21,6 +21,15 @@ static inline bool IsStencilFormat(EPixelFormat Format);
 
 namespace RHIValidation
 {
+	struct FGlobalUniformBuffers
+	{
+		TArray<FRHIUniformBuffer*> Bindings;
+		bool bInSetPipelineStateCall{};
+
+		void Reset();
+		void ValidateSetShaderUniformBuffer(FRHIUniformBuffer* UniformBuffer);
+	};
+
 	static bool IsValidCopyFormat(EPixelFormat SourceFormat, EPixelFormat DestFormat)
 	{
 		if (SourceFormat == DestFormat)
diff --git a/Engine/Source/Runtime/RHI/Public/RHIValidationContext.h b/Engine/Source/Runtime/RHI/Public/RHIValidationContext.h
--- a/Engine/Source/Runtime/RHI/Public/RHIValidationContext.h
+++ b/Engine/Source/Runtime/RHI/Public/RHIValidationContext.h
@@ -32,7 +32,9 @@ public:
 		// Reset the compute UAV tracker since the renderer must re-bind all resources after changing a shader.
 		Tracker->ResetUAVState(RHIValidation::EUAVMode::Compute);
 
+		State.GlobalUniformBuffers.bInSetPipelineStateCall = true;
 		RHIContext->RHISetComputeShader(Shader);
+		State.GlobalUniformBuffers.bInSetPipelineStateCall = false;
 	}
 
 	virtual void RHIDispatchComputeShader(uint32 ThreadGroupCountX, uint32 ThreadGroupCountY, uint32 ThreadGroupCountZ) override final
@@ -152,7 +154,7 @@ public:
 	virtual void RHISetShaderUniformBuffer(FRHIComputeShader* Shader, uint32 BufferIndex, FRHIUniformBuffer* Buffer) override final
 	{
 		checkf(State.bComputeShaderSet, TEXT("A Compute shader has to be set to set resources into a shader!"));
-		Buffer->ValidateLifeTime();
+		State.GlobalUniformBuffers.ValidateSetShaderUniformBuffer(Buffer);
 		RHIContext->RHISetShaderUniformBuffer(Shader, BufferIndex, Buffer);
 	}
 
@@ -164,6 +166,7 @@ public:
 
 	virtual void RHISetGlobalUniformBuffers(const FUniformBufferStaticBindings& InUniformBuffers) override final
 	{
+		InUniformBuffers.Bind(State.GlobalUniformBuffers.Bindings);
 		RHIContext->RHISetGlobalUniformBuffers(InUniformBuffers);
 	}
 
@@ -200,14 +203,11 @@ public:
 protected:
 	struct FState
 	{
-		RHIValidation::FTracker TrackerInstance;
+		RHIValidation::FTracker TrackerInstance{ ERHIPipeline::AsyncCompute };
+		RHIValidation::FGlobalUniformBuffers GlobalUniformBuffers;
 
 		FString ComputePassName;
-		bool bComputeShaderSet;
-
-		FState()
-			: TrackerInstance(ERHIPipeline::AsyncCompute)
-		{}
+		bool bComputeShaderSet{};
 
 		void Reset();
 	};
@@ -537,7 +537,9 @@ public:
 		// Setting a new PSO unbinds all previous bound resources
 		Tracker->ResetUAVState(RHIValidation::EUAVMode::Graphics);
 
+		State.GlobalUniformBuffers.bInSetPipelineStateCall = true;
 		RHIContext->RHISetGraphicsPipelineState(GraphicsState, bApplyAdditionalState);
+		State.GlobalUniformBuffers.bInSetPipelineStateCall = false;
 	}
 
 	/** Set the shader resource view of a surface.  This is used for binding TextureMS parameter types that need a multi sampled view. */
@@ -646,14 +648,14 @@ public:
 	virtual void RHISetShaderUniformBuffer(FRHIGraphicsShader* Shader, uint32 BufferIndex, FRHIUniformBuffer* Buffer) override final
 	{
 		checkf(State.bGfxPSOSet, TEXT("A Graphics PSO has to be set to set resources into a shader!"));
-		Buffer->ValidateLifeTime();
+		State.GlobalUniformBuffers.ValidateSetShaderUniformBuffer(Buffer);
 		RHIContext->RHISetShaderUniformBuffer(Shader, BufferIndex, Buffer);
 	}
 
 	virtual void RHISetShaderUniformBuffer(FRHIComputeShader* Shader, uint32 BufferIndex, FRHIUniformBuffer* Buffer) override final
 	{
 		checkf(State.bComputeShaderSet, TEXT("A Compute shader has to be set to set resources into a shader!"));
-		Buffer->ValidateLifeTime();
+		State.GlobalUniformBuffers.ValidateSetShaderUniformBuffer(Buffer);
 		RHIContext->RHISetShaderUniformBuffer(Shader, BufferIndex, Buffer);
 	}
 
@@ -671,6 +673,7 @@ public:
 
 	virtual void RHISetGlobalUniformBuffers(const FUniformBufferStaticBindings& InUniformBuffers) override final
 	{
+		InUniformBuffers.Bind(State.GlobalUniformBuffers.Bindings);
 		RHIContext->RHISetGlobalUniformBuffers(InUniformBuffers);
 	}
 
@@ -965,23 +968,19 @@ public:
 		PlatformContext->Tracker = &State.TrackerInstance;
 	}
 
-
 protected:
 	struct FState
 	{
-		RHIValidation::FTracker TrackerInstance;
+		RHIValidation::FTracker TrackerInstance{ ERHIPipeline::Graphics };
+		RHIValidation::FGlobalUniformBuffers GlobalUniformBuffers;
 
-		bool bInsideBeginRenderPass;
 		FRHIRenderPassInfo RenderPassInfo;
 		FString RenderPassName;
 		FString PreviousRenderPassName;
 		FString ComputePassName;
-		bool bGfxPSOSet;
-		bool bComputeShaderSet;
-
-		FState()
-			: TrackerInstance(ERHIPipeline::Graphics)
-		{}
+		bool bGfxPSOSet{};
+		bool bComputeShaderSet{};
+		bool bInsideBeginRenderPass{};
 
 		void Reset();
 	};
diff --git a/Engine/Source/Runtime/RenderCore/Private/RenderGraphBuilder.cpp b/Engine/Source/Runtime/RenderCore/Private/RenderGraphBuilder.cpp
--- a/Engine/Source/Runtime/RenderCore/Private/RenderGraphBuilder.cpp
+++ b/Engine/Source/Runtime/RenderCore/Private/RenderGraphBuilder.cpp
@@ -1749,9 +1749,9 @@ void FRDGBuilder::CollectPassResources(FRDGPassHandle PassHandle)
 		// Uniform buffer creation requires access to RHI resources through the validator.
 		IF_RDG_ENABLE_DEBUG(FRDGUserValidation::SetAllowRHIAccess(PassToBegin, true));
 		{
-			PassParameters.EnumerateUniformBuffers([&](FRDGUniformBufferRef UniformBuffer)
+			PassParameters.EnumerateUniformBuffers([&](FRDGUniformBufferBinding UniformBuffer)
 			{
-				BeginResourceRHI(UniformBuffer);
+				BeginResourceRHI(UniformBuffer.GetUniformBuffer());
 			});
 		}
 		IF_RDG_ENABLE_DEBUG(FRDGUserValidation::SetAllowRHIAccess(PassToBegin, false));
diff --git a/Engine/Source/Runtime/RenderCore/Private/RenderGraphPass.cpp b/Engine/Source/Runtime/RenderCore/Private/RenderGraphPass.cpp
--- a/Engine/Source/Runtime/RenderCore/Private/RenderGraphPass.cpp
+++ b/Engine/Source/Runtime/RenderCore/Private/RenderGraphPass.cpp
@@ -19,17 +19,17 @@ FUniformBufferStaticBindings FRDGParameterStruct::GetGlobalUniformBuffers() cons
 	for (uint32 Index = 0, Count = Layout->UniformBuffers.Num(); Index < Count; ++Index)
 	{
 		const uint32 MemberOffset = Layout->UniformBuffers[Index].MemberOffset;
-		FRHIUniformBuffer* UniformBufferPtr = *reinterpret_cast<FUniformBufferRHIRef*>(const_cast<uint8*>(Contents + MemberOffset));
+		const FUniformBufferBinding& UniformBuffer = *reinterpret_cast<const FUniformBufferBinding*>(const_cast<uint8*>(Contents + MemberOffset));
 
-		if (UniformBufferPtr && UniformBufferPtr->IsGlobal())
+		if (UniformBuffer && UniformBuffer.IsStatic())
 		{
-			GlobalUniformBuffers.AddUniformBuffer(UniformBufferPtr);
+			GlobalUniformBuffers.AddUniformBuffer(UniformBuffer.GetUniformBuffer());
 		}
 	}
 
-	EnumerateUniformBuffers([&](FRDGUniformBufferRef UniformBuffer)
+	EnumerateUniformBuffers([&](FRDGUniformBufferBinding UniformBuffer)
 	{
-		if (UniformBuffer->IsGlobal())
+		if (UniformBuffer.IsStatic())
 		{
 			GlobalUniformBuffers.AddUniformBuffer(UniformBuffer->GetRHI());
 		}
diff --git a/Engine/Source/Runtime/RenderCore/Private/RenderGraphUtils.cpp b/Engine/Source/Runtime/RenderCore/Private/RenderGraphUtils.cpp
--- a/Engine/Source/Runtime/RenderCore/Private/RenderGraphUtils.cpp
+++ b/Engine/Source/Runtime/RenderCore/Private/RenderGraphUtils.cpp
@@ -52,8 +52,8 @@ void ClearUnusedGraphResourcesImpl(
 				continue;
 			}
 
-			FRDGUniformBufferRef UniformBuffer = *reinterpret_cast<FRDGUniformBufferRef*>(Base + ByteOffset);
-			if (!UniformBuffer || UniformBuffer->IsGlobal())
+			const FRDGUniformBufferBinding& UniformBuffer = *reinterpret_cast<FRDGUniformBufferBinding*>(Base + ByteOffset);
+			if (!UniformBuffer || UniformBuffer.IsStatic())
 			{
 				continue;
 			}
@@ -126,8 +126,8 @@ void ClearUnusedGraphResourcesImpl(
 				bResourceIsUsed |= GraphUniformBufferId < GraphUniformBuffers.Num() && ByteOffset == GraphUniformBuffers[GraphUniformBufferId].ByteOffset;
 			}
 
-			FRDGUniformBufferRef UniformBuffer = *reinterpret_cast<FRDGUniformBufferRef*>(Base + ByteOffset);
-			if (!UniformBuffer || UniformBuffer->IsGlobal())
+			const FRDGUniformBufferBinding& UniformBuffer = *reinterpret_cast<const FRDGUniformBufferBinding*>(Base + ByteOffset);
+			if (!UniformBuffer || UniformBuffer.IsStatic())
 			{
 				continue;
 			}
diff --git a/Engine/Source/Runtime/RenderCore/Private/RenderGraphValidation.cpp b/Engine/Source/Runtime/RenderCore/Private/RenderGraphValidation.cpp
--- a/Engine/Source/Runtime/RenderCore/Private/RenderGraphValidation.cpp
+++ b/Engine/Source/Runtime/RenderCore/Private/RenderGraphValidation.cpp
@@ -899,10 +899,10 @@ void FRDGUserValidation::ValidateExecutePassBegin(const FRDGPass* Pass)
 
 	if (GRDGDebug)
 	{
-		Pass->GetParameters().EnumerateUniformBuffers([&](FRDGUniformBufferRef UniformBuffer)
+		Pass->GetParameters().EnumerateUniformBuffers([&](FRDGUniformBufferBinding UniformBuffer)
 		{
 			// Global uniform buffers are always marked as used, because FShader traversal doesn't know about them.
-			if (UniformBuffer->IsGlobal())
+			if (UniformBuffer.IsStatic())
 			{
 				UniformBuffer->MarkResourceAsUsed();
 			}
diff --git a/Engine/Source/Runtime/RenderCore/Private/ShaderCore.cpp b/Engine/Source/Runtime/RenderCore/Private/ShaderCore.cpp
--- a/Engine/Source/Runtime/RenderCore/Private/ShaderCore.cpp
+++ b/Engine/Source/Runtime/RenderCore/Private/ShaderCore.cpp
@@ -381,11 +381,22 @@ void ValidateStaticUniformBuffer(FRHIUniformBuffer* UniformBuffer, FUniformBuffe
 			TEXT("Shader is requesting a uniform buffer at slot %s with hash '%u', but a reverse lookup of the hash can't find it. The shader cache may be out of date."),
 			*SlotRegistry.GetDebugDescription(Slot), ExpectedHash);
 
-		UE_LOG(LogShaders, Fatal,
-			TEXT("Shader requested a global uniform buffer of type '%s' at static slot '%s', but it was null. The uniform buffer should ")
-			TEXT("be bound using RHICmdList.SetGlobalUniformBuffers() or passed into an RDG pass using SHADER_PARAMETER_STRUCT_REF() or ")
-			TEXT("SHADER_PARAMETER_RDG_UNIFORM_BUFFER()."),
+		const EUniformBufferBindingFlags BindingFlags = ExpectedStructMetadata->GetBindingFlags();
+
+		checkf(EnumHasAnyFlags(BindingFlags, EUniformBufferBindingFlags::Static),
+			TEXT("Shader requested a global uniform buffer of type '%s' at static slot '%s', but it is not registered with the Global binding flag. The shader cache may be out of date."),
 			ExpectedStructMetadata->GetShaderVariableName(), *SlotRegistry.GetDebugDescription(Slot));
+
+		// Structs can be bound both globally or per-shader, effectively leaving it up to the user to choose which to bind.
+		// But that also means we can't validate existence at the global level.
+		if (!EnumHasAnyFlags(BindingFlags, EUniformBufferBindingFlags::Shader))
+		{
+			UE_LOG(LogShaders, Fatal,
+				TEXT("Shader requested a global uniform buffer of type '%s' at static slot '%s', but it was null. The uniform buffer should ")
+				TEXT("be bound using RHICmdList.SetGlobalUniformBuffers() or passed into an RDG pass using SHADER_PARAMETER_STRUCT_REF() or ")
+				TEXT("SHADER_PARAMETER_RDG_UNIFORM_BUFFER()."),
+				ExpectedStructMetadata->GetShaderVariableName(), *SlotRegistry.GetDebugDescription(Slot));
+		}
 	}
 	else
 	{
diff --git a/Engine/Source/Runtime/RenderCore/Private/ShaderParameterMetadata.cpp b/Engine/Source/Runtime/RenderCore/Private/ShaderParameterMetadata.cpp
--- a/Engine/Source/Runtime/RenderCore/Private/ShaderParameterMetadata.cpp
+++ b/Engine/Source/Runtime/RenderCore/Private/ShaderParameterMetadata.cpp
@@ -274,6 +274,7 @@ public:
 
 FShaderParametersMetadata::FShaderParametersMetadata(
 	EUseCase InUseCase,
+	EUniformBufferBindingFlags InBindingFlags,
 	const TCHAR* InLayoutName,
 	const TCHAR* InStructTypeName,
 	const TCHAR* InShaderVariableName,
@@ -290,11 +291,14 @@ FShaderParametersMetadata::FShaderParametersMetadata(
 	, FileLine(InFileLine)
 	, Size(InSize)
 	, UseCase(InUseCase)
+	, BindingFlags(InBindingFlags)
 	, Layout(InLayoutName)
 	, Members(InMembers)
 	, GlobalListLink(this)
 	, bLayoutInitialized(false)
 {
+	checkf(UseCase == EUseCase::UniformBuffer || !EnumHasAnyFlags(BindingFlags, EUniformBufferBindingFlags::Static), TEXT("Only uniform buffers can utilize the global binding flag."));
+
 	check(StructTypeName);
 	if (UseCase == EUseCase::ShaderParameterStruct)
 	{
@@ -383,6 +387,9 @@ void FShaderParametersMetadata::InitializeLayout()
 
 	if (StaticSlotName)
 	{
+		checkf(EnumHasAnyFlags(BindingFlags, EUniformBufferBindingFlags::Static), TEXT("Uniform buffer of type '%s' and shader name '%s' attempted to reference static slot '%s', but the binding model does not contain 'Global'."),
+			StructTypeName, ShaderVariableName, StaticSlotName);
+
 		checkf(UseCase == EUseCase::UniformBuffer,
 			TEXT("Attempted to assign static slot %s to uniform buffer %s. Static slots are only supported for compile-time uniform buffers."),
 			ShaderVariableName, StaticSlotName);
@@ -394,6 +401,12 @@ void FShaderParametersMetadata::InitializeLayout()
 			StructTypeName, ShaderVariableName, StaticSlotName);
 
 		Layout.StaticSlot = StaticSlot;
+		Layout.BindingFlags = BindingFlags;
+	}
+	else
+	{
+		checkf(!EnumHasAnyFlags(BindingFlags, EUniformBufferBindingFlags::Static), TEXT("Uniform buffer of type '%s' and shader name '%s' has no binding slot specified, but the binding model specifies 'Global'."),
+			StructTypeName, ShaderVariableName, StaticSlotName);
 	}
 
 	TArray<FUniformBufferMemberAndOffset> MemberStack;
@@ -749,17 +762,17 @@ void FShaderParametersMetadata::GetNestedStructs(TArray<const FShaderParametersM
 	}
 }
 
-void FShaderParametersMetadata::AddResourceTableEntries(TMap<FString, FResourceTableEntry>& ResourceTableMap, TMap<FString, uint32>& ResourceTableLayoutHashes, TMap<FString, FString>& UniformBufferStaticSlots) const
+void FShaderParametersMetadata::AddResourceTableEntries(TMap<FString, FResourceTableEntry>& ResourceTableMap, TMap<FString, FUniformBufferEntry>& UniformBufferMap) const
 {
 	uint16 ResourceIndex = 0;
 	FString Prefix = FString::Printf(TEXT("%s_"), ShaderVariableName);
 	AddResourceTableEntriesRecursive(ShaderVariableName, *Prefix, ResourceIndex, ResourceTableMap);
-	ResourceTableLayoutHashes.Add(ShaderVariableName, Layout.GetHash());
-
-	if (StaticSlotName)
-	{
-		UniformBufferStaticSlots.Add(ShaderVariableName, StaticSlotName);
-	}
+	
+	FUniformBufferEntry UniformBufferEntry;
+	UniformBufferEntry.StaticSlotName = StaticSlotName;
+	UniformBufferEntry.LayoutHash = Layout.GetHash();
+	UniformBufferEntry.BindingFlags = BindingFlags;
+	UniformBufferMap.Add(ShaderVariableName, UniformBufferEntry);
 }
 
 void FShaderParametersMetadata::AddResourceTableEntriesRecursive(const TCHAR* UniformBufferName, const TCHAR* Prefix, uint16& ResourceIndex, TMap<FString, FResourceTableEntry>& ResourceTableMap) const
@@ -953,6 +966,7 @@ FShaderParametersMetadata* FShaderParametersMetadataBuilder::Build(
 
 	FShaderParametersMetadata* ShaderParameterMetadata = new FShaderParametersMetadata(
 		UseCase,
+		EUniformBufferBindingFlags::Shader,
 		ShaderParameterName,
 		ShaderParameterName,
 		nullptr,
diff --git a/Engine/Source/Runtime/RenderCore/Private/ShaderParameterStruct.cpp b/Engine/Source/Runtime/RenderCore/Private/ShaderParameterStruct.cpp
--- a/Engine/Source/Runtime/RenderCore/Private/ShaderParameterStruct.cpp
+++ b/Engine/Source/Runtime/RenderCore/Private/ShaderParameterStruct.cpp
@@ -495,8 +495,9 @@ void ValidateShaderParameters(const TShaderRef<FShader>& Shader, const FShaderPa
 	// Graph Uniform Buffers
 	for (const FShaderParameterBindings::FParameterStructReference& ParameterBinding : Bindings.GraphUniformBuffers)
 	{
-		auto GraphUniformBuffer = *reinterpret_cast<const FRDGUniformBuffer* const*>(Base + ParameterBinding.ByteOffset);
-		if (!GraphUniformBuffer)
+		const FRDGUniformBufferBinding& UniformBufferBinding = *reinterpret_cast<const FRDGUniformBufferBinding*>(Base + ParameterBinding.ByteOffset);
+
+		if (!UniformBufferBinding)
 		{
 			EmitNullShaderParameterFatalError(Shader, ParametersMetadata, ParameterBinding.ByteOffset);
 		}
@@ -505,9 +506,9 @@ void ValidateShaderParameters(const TShaderRef<FShader>& Shader, const FShaderPa
 	// Reference structures
 	for (const FShaderParameterBindings::FParameterStructReference& ParameterBinding : Bindings.ParameterReferences)
 	{
-		const TRefCountPtr<FRHIUniformBuffer>& ShaderParameterRef = *reinterpret_cast<const TRefCountPtr<FRHIUniformBuffer>*>(Base + ParameterBinding.ByteOffset);
+		const FUniformBufferBinding& UniformBufferBinding = *reinterpret_cast<const FUniformBufferBinding*>(Base + ParameterBinding.ByteOffset);
 
-		if (!ShaderParameterRef.IsValid())
+		if (!UniformBufferBinding)
 		{
 			EmitNullShaderParameterFatalError(Shader, ParametersMetadata, ParameterBinding.ByteOffset);
 		}
diff --git a/Engine/Source/Runtime/RenderCore/Private/ShaderParameters.cpp b/Engine/Source/Runtime/RenderCore/Private/ShaderParameters.cpp
--- a/Engine/Source/Runtime/RenderCore/Private/ShaderParameters.cpp
+++ b/Engine/Source/Runtime/RenderCore/Private/ShaderParameters.cpp
@@ -79,7 +79,7 @@ void FShaderUniformBufferParameter::ModifyCompilationEnvironment(const TCHAR* Pa
 	FString Include = FString::Printf(TEXT("#include \"/Engine/Generated/UniformBuffers/%s.ush\"") LINE_TERMINATOR, ParameterName);
 
 	GeneratedUniformBuffersInclude.Append(Include);
-	Struct.AddResourceTableEntries(OutEnvironment.ResourceTableMap, OutEnvironment.ResourceTableLayoutHashes, OutEnvironment.ResourceTableLayoutSlots);
+	Struct.AddResourceTableEntries(OutEnvironment.ResourceTableMap, OutEnvironment.UniformBufferMap);
 }
 
 void FShaderUniformBufferParameter::Bind(const FShaderParameterMap& ParameterMap,const TCHAR* ParameterName,EShaderParameterFlags Flags)
@@ -376,7 +376,7 @@ void FShaderType::AddReferencedUniformBufferIncludes(FShaderCompilerEnvironment&
 		{
 			if (It.Key() == StructIt->GetShaderVariableName())
 			{
-				StructIt->AddResourceTableEntries(OutEnvironment.ResourceTableMap, OutEnvironment.ResourceTableLayoutHashes, OutEnvironment.ResourceTableLayoutSlots);
+				StructIt->AddResourceTableEntries(OutEnvironment.ResourceTableMap, OutEnvironment.UniformBufferMap);
 			}
 		}
 	}
@@ -499,7 +499,7 @@ void FVertexFactoryType::AddReferencedUniformBufferIncludes(FShaderCompilerEnvir
 		{
 			if (It.Key() == StructIt->GetShaderVariableName())
 			{
-				StructIt->AddResourceTableEntries(OutEnvironment.ResourceTableMap, OutEnvironment.ResourceTableLayoutHashes, OutEnvironment.ResourceTableLayoutSlots);
+				StructIt->AddResourceTableEntries(OutEnvironment.ResourceTableMap, OutEnvironment.UniformBufferMap);
 			}
 		}
 	}
diff --git a/Engine/Source/Runtime/RenderCore/Public/RenderGraphParameter.h b/Engine/Source/Runtime/RenderCore/Public/RenderGraphParameter.h
--- a/Engine/Source/Runtime/RenderCore/Public/RenderGraphParameter.h
+++ b/Engine/Source/Runtime/RenderCore/Public/RenderGraphParameter.h
@@ -81,10 +81,10 @@ public:
 		return *GetAs<FRDGResourceRef>();
 	}
 
-	FRDGUniformBufferRef GetAsUniformBuffer() const
+	FRDGUniformBufferBinding GetAsUniformBuffer() const
 	{
 		check(IsUniformBuffer());
-		return *GetAs<FRDGUniformBufferRef>();
+		return *GetAs<FRDGUniformBufferBinding>();
 	}
 
 	FRDGParentResourceRef GetAsParentResource() const
@@ -188,7 +188,7 @@ class RENDERCORE_API FRDGParameterStruct
 {
 public:
 	template <typename ParameterStructType>
-	explicit FRDGParameterStruct(ParameterStructType* Parameters)
+	explicit FRDGParameterStruct(const ParameterStructType* Parameters)
 		: FRDGParameterStruct(Parameters, &ParameterStructType::FTypeInfo::GetStructMetadata()->GetLayout())
 	{}
 
@@ -238,7 +238,7 @@ public:
 	template <typename FunctionType>
 	void EnumerateBuffers(FunctionType Function) const;
 
-	/** Enumerates all non-null uniform buffers. Expected function signature: void(FRDGUniformBuffer*). */
+	/** Enumerates all non-null uniform buffers. Expected function signature: void(FRDGUniformBufferBinding). */
 	template <typename FunctionType>
 	void EnumerateUniformBuffers(FunctionType Function) const;
 
diff --git a/Engine/Source/Runtime/RenderCore/Public/RenderGraphParameters.inl b/Engine/Source/Runtime/RenderCore/Public/RenderGraphParameters.inl
--- a/Engine/Source/Runtime/RenderCore/Public/RenderGraphParameters.inl
+++ b/Engine/Source/Runtime/RenderCore/Public/RenderGraphParameters.inl
@@ -5,7 +5,7 @@
 template <typename FunctionType>
 void FRDGParameterStruct::Enumerate(FunctionType Function) const
 {
-	EnumerateUniformBuffers([&](FRDGUniformBuffer* UniformBuffer)
+	EnumerateUniformBuffers([&](FRDGUniformBufferBinding UniformBuffer)
 	{
 		UniformBuffer->GetParameters().Enumerate(Function);
 	});
@@ -19,7 +19,7 @@ void FRDGParameterStruct::Enumerate(FunctionType Function) const
 template <typename FunctionType>
 void FRDGParameterStruct::EnumerateTextures(FunctionType Function) const
 {
-	EnumerateUniformBuffers([&](FRDGUniformBuffer* UniformBuffer)
+	EnumerateUniformBuffers([&](FRDGUniformBufferBinding UniformBuffer)
 	{
 		UniformBuffer->GetParameters().EnumerateTextures(Function);
 	});
@@ -33,7 +33,7 @@ void FRDGParameterStruct::EnumerateTextures(FunctionType Function) const
 template <typename FunctionType>
 void FRDGParameterStruct::EnumerateBuffers(FunctionType Function) const
 {
-	EnumerateUniformBuffers([&](FRDGUniformBuffer* UniformBuffer)
+	EnumerateUniformBuffers([&](FRDGUniformBufferBinding UniformBuffer)
 	{
 		UniformBuffer->GetParameters().EnumerateBuffers(Function);
 	});
@@ -51,7 +51,7 @@ void FRDGParameterStruct::EnumerateUniformBuffers(FunctionType Function) const
 	{
 		const FRDGParameter Parameter = GetParameterInternal(Layout->GraphUniformBuffers, Index);
 
-		if (FRDGUniformBufferRef UniformBuffer = Parameter.GetAsUniformBuffer())
+		if (FRDGUniformBufferBinding UniformBuffer = Parameter.GetAsUniformBuffer())
 		{
 			Function(UniformBuffer);
 		}
diff --git a/Engine/Source/Runtime/RenderCore/Public/RenderGraphResources.h b/Engine/Source/Runtime/RenderCore/Public/RenderGraphResources.h
--- a/Engine/Source/Runtime/RenderCore/Public/RenderGraphResources.h
+++ b/Engine/Source/Runtime/RenderCore/Public/RenderGraphResources.h
@@ -135,11 +135,6 @@ class FRDGUniformBuffer
 	: public FRDGResource
 {
 public:
-	FORCEINLINE bool IsGlobal() const
-	{
-		return bGlobal;
-	}
-
 	FORCEINLINE const FRDGParameterStruct& GetParameters() const
 	{
 		return ParameterStruct;
@@ -163,10 +158,9 @@ public:
 
 protected:
 	template <typename TParameterStruct>
-	explicit FRDGUniformBuffer(TParameterStruct* InParameters, const TCHAR* InName)
+	explicit FRDGUniformBuffer(const TParameterStruct* InParameters, const TCHAR* InName)
 		: FRDGResource(InName)
 		, ParameterStruct(InParameters)
-		, bGlobal(ParameterStruct.HasStaticSlot())
 	{}
 
 private:
@@ -175,9 +169,6 @@ private:
 	TRefCountPtr<FRHIUniformBuffer> UniformBufferRHI;
 	FRDGUniformBufferHandle Handle;
 
-	/** Whether this uniform buffer is bound globally or locally to a shader. */
-	uint8 bGlobal : 1;
-
 	friend FRDGBuilder;
 	friend FRDGUniformBufferRegistry;
 	friend FRDGAllocator;
@@ -199,14 +190,17 @@ public:
 
 	FORCEINLINE const ParameterStructType* operator->() const
 	{
-		return GetParameters().GetContents();
+		return Parameters;
 	}
 
 private:
-	explicit TRDGUniformBuffer(ParameterStructType* InParameters, const TCHAR* InName)
+	explicit TRDGUniformBuffer(const ParameterStructType* InParameters, const TCHAR* InName)
 		: FRDGUniformBuffer(InParameters, InName)
+		, Parameters(InParameters)
 	{}
 
+	const ParameterStructType* Parameters;
+
 	friend FRDGBuilder;
 	friend FRDGUniformBufferRegistry;
 	friend FRDGAllocator;
diff --git a/Engine/Source/Runtime/RenderCore/Public/ShaderCore.h b/Engine/Source/Runtime/RenderCore/Public/ShaderCore.h
--- a/Engine/Source/Runtime/RenderCore/Public/ShaderCore.h
+++ b/Engine/Source/Runtime/RenderCore/Public/ShaderCore.h
@@ -394,6 +394,14 @@ inline FArchive& operator<<(FArchive& Ar, FResourceTableEntry& Entry)
 	return Ar;
 }
 
+inline FArchive& operator<<(FArchive& Ar, FUniformBufferEntry& Entry)
+{
+	Ar << Entry.StaticSlotName;
+	Ar << Entry.LayoutHash;
+	Ar << Entry.BindingFlags;
+	return Ar;
+}
+
 using FThreadSafeSharedStringPtr = TSharedPtr<FString, ESPMode::ThreadSafe>;
 
 /** The environment used to compile a shader. */
@@ -407,9 +415,8 @@ struct FShaderCompilerEnvironment
 
 	TArray<uint32> CompilerFlags;
 	TMap<uint32,uint8> RenderTargetOutputFormatsMap;
-	TMap<FString,FResourceTableEntry> ResourceTableMap;
-	TMap<FString,uint32> ResourceTableLayoutHashes;
-	TMap<FString, FString> ResourceTableLayoutSlots;
+	TMap<FString, FResourceTableEntry> ResourceTableMap;
+	TMap<FString, FUniformBufferEntry> UniformBufferMap;
 	TMap<FString, FString> RemoteServerData;
 	TMap<FString, FString> ShaderFormatCVars;
 
@@ -466,8 +473,7 @@ struct FShaderCompilerEnvironment
 		Ar << Environment.CompilerFlags;
 		Ar << Environment.RenderTargetOutputFormatsMap;
 		Ar << Environment.ResourceTableMap;
-		Ar << Environment.ResourceTableLayoutHashes;
-		Ar << Environment.ResourceTableLayoutSlots;
+		Ar << Environment.UniformBufferMap;
 		Ar << Environment.RemoteServerData;
 		Ar << Environment.ShaderFormatCVars;
 
@@ -496,8 +502,7 @@ struct FShaderCompilerEnvironment
 
 		CompilerFlags.Append(Other.CompilerFlags);
 		ResourceTableMap.Append(Other.ResourceTableMap);
-		ResourceTableLayoutHashes.Append(Other.ResourceTableLayoutHashes);
-		ResourceTableLayoutSlots.Append(Other.ResourceTableLayoutSlots);
+		UniformBufferMap.Append(Other.UniformBufferMap);
 		Definitions.Merge(Other.Definitions);
 		RenderTargetOutputFormatsMap.Append(Other.RenderTargetOutputFormatsMap);
 		RemoteServerData.Append(Other.RemoteServerData);
diff --git a/Engine/Source/Runtime/RenderCore/Public/ShaderParameterMacros.h b/Engine/Source/Runtime/RenderCore/Public/ShaderParameterMacros.h
--- a/Engine/Source/Runtime/RenderCore/Public/ShaderParameterMacros.h
+++ b/Engine/Source/Runtime/RenderCore/Public/ShaderParameterMacros.h
@@ -139,6 +139,187 @@ private:
 	friend class TRDGUniformBuffer<TBufferStruct>;
 };
 
+class alignas(SHADER_PARAMETER_STRUCT_ALIGNMENT) FUniformBufferBinding
+{
+public:
+	FUniformBufferBinding() = default;
+
+	FORCEINLINE FUniformBufferRHIRef GetUniformBuffer() const
+	{
+		return UniformBuffer;
+	}
+
+	FORCEINLINE EUniformBufferBindingFlags GetBindingFlags() const
+	{
+		return BindingFlags;
+	}
+
+	FORCEINLINE bool IsStatic() const
+	{
+		return EnumHasAnyFlags(BindingFlags, EUniformBufferBindingFlags::Static);
+	}
+
+	FORCEINLINE bool IsShader() const
+	{
+		return EnumHasAnyFlags(BindingFlags, EUniformBufferBindingFlags::Shader);
+	}
+
+	FORCEINLINE operator bool() const
+	{
+		return UniformBuffer.IsValid();
+	}
+
+protected:
+	FORCEINLINE FUniformBufferBinding(const FUniformBufferRHIRef& InUniformBuffer, EUniformBufferBindingFlags InBindingFlags)
+		: UniformBuffer(InUniformBuffer)
+		, BindingFlags(InBindingFlags)
+	{}
+
+private:
+	FUniformBufferRHIRef UniformBuffer;
+	EUniformBufferBindingFlags BindingFlags = EUniformBufferBindingFlags::Shader;
+};
+
+template<typename TBufferStruct>
+class alignas(SHADER_PARAMETER_STRUCT_ALIGNMENT) TUniformBufferBinding
+	: public FUniformBufferBinding
+{
+public:
+	TUniformBufferBinding() = default;
+
+	TUniformBufferBinding(const TUniformBufferRef<TBufferStruct>& InUniformBuffer)
+		: FUniformBufferBinding(InUniformBuffer, TBufferStruct::FTypeInfo::GetStructMetadata()->GetPreferredBindingFlag())
+	{}
+
+	TUniformBufferBinding(const TUniformBufferRef<TBufferStruct>& InUniformBuffer, EUniformBufferBindingFlags InBindingFlags)
+		: FUniformBufferBinding(InUniformBuffer, InBindingFlags)
+	{
+#if DO_CHECK
+		const auto* StructMetadata = TBufferStruct::FTypeInfo::GetStructMetadata();
+		checkf(
+			EnumHasAllFlags(StructMetadata->GetBindingFlags(), GetBindingFlags()),
+			TEXT("Uniform buffer binding flags don't match those supported by the uniform buffer layout '%s."),
+			StructMetadata->GetStructTypeName());
+#endif
+	}
+
+	FORCEINLINE const TUniformBufferRef<TBufferStruct>& GetUniformBuffer() const
+	{
+		return static_cast<const TUniformBufferRef<TBufferStruct>&>(FUniformBufferBinding::GetUniformBuffer());
+	}
+
+	FORCEINLINE TRDGUniformBuffer<TBufferStruct>* operator->() const
+	{
+		return GetUniformBuffer();
+	}
+};
+
+template <typename TBufferStruct>
+inline TUniformBufferBinding<TBufferStruct> GetStaticBinding(const TUniformBufferRef<TBufferStruct>& InUniformBuffer)
+{
+	return TUniformBufferBinding<TBufferStruct>(InUniformBuffer, EUniformBufferBindingFlags::Static);
+}
+
+template <typename TBufferStruct>
+inline TUniformBufferBinding<TBufferStruct> GetShaderBinding(const TUniformBufferRef<TBufferStruct>& InUniformBuffer)
+{
+	return TUniformBufferBinding<TBufferStruct>(InUniformBuffer, EUniformBufferBindingFlags::Shader);
+}
+
+class alignas(SHADER_PARAMETER_STRUCT_ALIGNMENT) FRDGUniformBufferBinding
+{
+public:
+	FRDGUniformBufferBinding() = default;
+
+	FORCEINLINE FRDGUniformBuffer* GetUniformBuffer() const
+	{
+		return UniformBuffer;
+	}
+
+	FORCEINLINE EUniformBufferBindingFlags GetBindingFlags() const
+	{
+		return BindingFlags;
+	}
+
+	FORCEINLINE bool IsStatic() const
+	{
+		return EnumHasAnyFlags(BindingFlags, EUniformBufferBindingFlags::Static);
+	}
+
+	FORCEINLINE bool IsShader() const
+	{
+		return EnumHasAnyFlags(BindingFlags, EUniformBufferBindingFlags::Shader);
+	}
+
+	FORCEINLINE operator bool() const
+	{
+		return UniformBuffer != nullptr;
+	}
+
+	FORCEINLINE FRDGUniformBuffer* operator->() const
+	{
+		return UniformBuffer;
+	}
+
+protected:
+	FRDGUniformBufferBinding(FRDGUniformBuffer* InUniformBuffer, EUniformBufferBindingFlags InBindingFlags)
+		: UniformBuffer(InUniformBuffer)
+		, BindingFlags(InBindingFlags)
+	{}
+
+private:
+	FRDGUniformBuffer* UniformBuffer = nullptr;
+	EUniformBufferBindingFlags BindingFlags = EUniformBufferBindingFlags::Shader;
+};
+
+template <typename TBufferStruct>
+class alignas(SHADER_PARAMETER_STRUCT_ALIGNMENT) TRDGUniformBufferBinding
+	: public FRDGUniformBufferBinding
+{
+public:
+	TRDGUniformBufferBinding() = default;
+
+	TRDGUniformBufferBinding(TRDGUniformBuffer<TBufferStruct>* InUniformBuffer)
+		: FRDGUniformBufferBinding(InUniformBuffer, TBufferStruct::FTypeInfo::GetStructMetadata()->GetPreferredBindingFlag())
+	{}
+
+	TRDGUniformBufferBinding(TRDGUniformBuffer<TBufferStruct>* InUniformBuffer, EUniformBufferBindingFlags InBindingFlags)
+		: FRDGUniformBufferBinding(InUniformBuffer, InBindingFlags)
+	{
+#if DO_CHECK
+		const auto* StructMetadata = TBufferStruct::FTypeInfo::GetStructMetadata();
+		checkf(
+			EnumHasAllFlags(StructMetadata->GetBindingFlags(), GetBindingFlags()),
+			TEXT("RDG uniform buffer binding flags don't match those supported by the uniform buffer layout '%s."),
+			StructMetadata->GetStructTypeName());
+#endif
+	}
+
+	FORCEINLINE TRDGUniformBuffer<TBufferStruct>* GetUniformBuffer() const
+	{
+		return static_cast<TRDGUniformBuffer<TBufferStruct>*>(FRDGUniformBufferBinding::GetUniformBuffer());
+	}
+
+	FORCEINLINE TRDGUniformBuffer<TBufferStruct>* operator->() const
+	{
+		return GetUniformBuffer();
+	}
+};
+
+/** Constructs a uniform buffer binding that will utilize the shader binding model. */
+template <typename TBufferStruct>
+inline TRDGUniformBufferBinding<TBufferStruct> GetShaderBinding(TRDGUniformBuffer<TBufferStruct>* InUniformBuffer)
+{
+	return TRDGUniformBufferBinding<TBufferStruct>(InUniformBuffer, EUniformBufferBindingFlags::Shader);
+}
+
+/** Constructs a uniform buffer binding that will utilize the static binding model. */
+template <typename TBufferStruct>
+inline TRDGUniformBufferBinding<TBufferStruct> GetStaticBinding(TRDGUniformBuffer<TBufferStruct>* InUniformBuffer)
+{
+	return TRDGUniformBufferBinding<TBufferStruct>(InUniformBuffer, EUniformBufferBindingFlags::Static);
+}
+
 class alignas(SHADER_PARAMETER_STRUCT_ALIGNMENT) FRDGBufferAccess
 {
 public:
@@ -514,7 +695,11 @@ struct alignas(SHADER_PARAMETER_STRUCT_ALIGNMENT) FRenderTargetBindingSlots
 			return false;
 		}
 
-		if (ResolveRect != Other.ResolveRect)
+		if (ResolveRect != Other.ResolveRect ||
+			(NumOcclusionQueries != Other.NumOcclusionQueries && Other.NumOcclusionQueries != 0) ||
+			SubpassHint != Other.SubpassHint ||
+			MultiViewCount != Other.MultiViewCount ||
+			FoveationTexture != Other.FoveationTexture)
 		{
 			return false;
 		}
@@ -888,15 +1073,29 @@ struct TShaderParameterTypeInfo<TUniformBufferRef<UniformBufferStructType>>
 };
 
 template<class UniformBufferStructType>
-struct TShaderParameterTypeInfo<TRDGUniformBuffer<UniformBufferStructType>*>
+struct TShaderParameterTypeInfo<TUniformBufferBinding<UniformBufferStructType>>
 {
 	static constexpr int32 NumRows = 1;
 	static constexpr int32 NumColumns = 1;
 	static constexpr int32 NumElements = 0;
-	static constexpr int32 Alignment = SHADER_PARAMETER_POINTER_ALIGNMENT;
+	static constexpr int32 Alignment = SHADER_PARAMETER_STRUCT_ALIGNMENT;
+	static constexpr bool bIsStoredInConstantBuffer = false;
+
+	using TAlignedType = TAlignedShaderParameterPtr<TUniformBufferBinding<UniformBufferStructType>>;
+
+	static const FShaderParametersMetadata* GetStructMetadata() { return &UniformBufferStructType::StaticStructMetadata; }
+};
+
+template<class UniformBufferStructType>
+struct TShaderParameterTypeInfo<TRDGUniformBufferBinding<UniformBufferStructType>>
+{
+	static constexpr int32 NumRows = 1;
+	static constexpr int32 NumColumns = 1;
+	static constexpr int32 NumElements = 0;
+	static constexpr int32 Alignment = SHADER_PARAMETER_STRUCT_ALIGNMENT;
 	static constexpr bool bIsStoredInConstantBuffer = false;
 
-	using TAlignedType = TAlignedShaderParameterPtr<TRDGUniformBuffer<UniformBufferStructType>*>;
+	using TAlignedType = TRDGUniformBufferBinding<UniformBufferStructType>;
 
 	static const FShaderParametersMetadata* GetStructMetadata() { return &UniformBufferStructType::StaticStructMetadata; }
 };
@@ -938,6 +1137,7 @@ struct TShaderParameterStructTypeInfo<StructType[InNumElements]>
 #define INTERNAL_SHADER_PARAMETER_GET_STRUCT_METADATA(StructTypeName) \
 	static FShaderParametersMetadata StaticStructMetadata(\
 		FShaderParametersMetadata::EUseCase::ShaderParameterStruct, \
+		EUniformBufferBindingFlags::Shader, \
 		TEXT(#StructTypeName), \
 		TEXT(#StructTypeName), \
 		nullptr, \
@@ -1073,6 +1273,7 @@ extern RENDERCORE_API FShaderParametersMetadata* FindUniformBufferStructByShader
 #define IMPLEMENT_UNIFORM_BUFFER_STRUCT(StructTypeName,ShaderVariableName) \
 	FShaderParametersMetadata StructTypeName::StaticStructMetadata( \
 	FShaderParametersMetadata::EUseCase::UniformBuffer, \
+	EUniformBufferBindingFlags::Shader, \
 	TEXT(#StructTypeName), \
 	TEXT(#StructTypeName), \
 	TEXT(ShaderVariableName), \
@@ -1085,6 +1286,7 @@ extern RENDERCORE_API FShaderParametersMetadata* FindUniformBufferStructByShader
 #define IMPLEMENT_UNIFORM_BUFFER_ALIAS_STRUCT(StructTypeName, UniformBufferAlias) \
 	static const FShaderParametersMetadata UniformBufferAlias( \
 	FShaderParametersMetadata::EUseCase::UniformBuffer, \
+	EUniformBufferBindingFlags::Shader, \
 	TEXT(#StructTypeName), \
 	TEXT(#StructTypeName), \
 	TEXT(#UniformBufferAlias), \
@@ -1094,33 +1296,47 @@ extern RENDERCORE_API FShaderParametersMetadata* FindUniformBufferStructByShader
 	sizeof(StructTypeName), \
 	StructTypeName::zzGetMembers())
 
-/** Implements a uniform buffer tied to a static binding slot. The third parameter is the name of the slot.
- *  Multiple uniform buffers can be associated to a slot; only one uniform buffer can be bound to a slot
- *  at one time.
- *
- * Example:
- *	BEGIN_UNIFORM_BUFFER_STRUCT(FMyParameterStruct, RENDERER_API)
- *	END_UNIFORM_BUFFER_STRUCT()
- *
- *	// C++ file
- *
- *	// Define uniform buffer slot.
- *	IMPLEMENT_STATIC_UNIFORM_BUFFER_SLOT(MySlot)
- *
- *	// Associate uniform buffer with slot.
- *	IMPLEMENT_STATIC_UNIFORM_BUFFER_STRUCT(FMyParameterStruct, "MyShaderBindingName", MySlot);
+#define IMPLEMENT_STATIC_UNIFORM_BUFFER_STRUCT_EX(StructTypeName,ShaderVariableName,StaticSlotName,BindingFlagsEnum) \
+	static_assert(EnumHasAnyFlags(EUniformBufferBindingFlags::BindingFlagsEnum, EUniformBufferBindingFlags::Static), "Shader bindings must include 'Static'."); \
+	FShaderParametersMetadata StructTypeName::StaticStructMetadata( \
+		FShaderParametersMetadata::EUseCase::UniformBuffer, \
+		EUniformBufferBindingFlags::BindingFlagsEnum, \
+		TEXT(#StructTypeName), \
+		TEXT(#StructTypeName), \
+		TEXT(ShaderVariableName), \
+		TEXT(#StaticSlotName), \
+		StructTypeName::FTypeInfo::FileName, \
+		StructTypeName::FTypeInfo::FileLine, \
+		sizeof(StructTypeName), \
+		StructTypeName::zzGetMembers())
+
+/** Implements the same contract as IMPLEMENT_STATIC_UNIFORM_BUFFER_STRUCT, with the addition of the 'StaticAndShader' binding
+ *  flag. This means that the uniform buffer may be bound statically or through SetShaderParameters. The uniform buffer is NOT
+ *  removed from the parameter map, which means it is possible to end up with both global and shader bindings. The RHI validation
+ *  layer will catch and report these cases, as they are both inefficient and would produce confusing behavior if the two bindings
+ *  differ. Only one binding model should be used at a time.
  */
+#define IMPLEMENT_STATIC_AND_SHADER_UNIFORM_BUFFER_STRUCT(StructTypeName,ShaderVariableName,StaticSlotName) \
+	IMPLEMENT_STATIC_UNIFORM_BUFFER_STRUCT_EX(StructTypeName, ShaderVariableName, StaticSlotName, StaticAndShader)
+
+ /** Implements a uniform buffer tied to a static binding slot. The third parameter is the name of the slot.
+  *  Multiple uniform buffers can be associated to a slot; only one uniform buffer can be bound to a slot
+  *  at one time.
+  *
+  * Example:
+  *	BEGIN_UNIFORM_BUFFER_STRUCT(FMyParameterStruct, RENDERER_API)
+  *	END_UNIFORM_BUFFER_STRUCT()
+  *
+  *	// C++ file
+  *
+  *	// Define uniform buffer slot.
+  *	IMPLEMENT_STATIC_UNIFORM_BUFFER_SLOT(MySlot)
+  *
+  *	// Associate uniform buffer with slot.
+  *	IMPLEMENT_STATIC_UNIFORM_BUFFER_STRUCT(FMyParameterStruct, "MyShaderBindingName", MySlot);
+  */
 #define IMPLEMENT_STATIC_UNIFORM_BUFFER_STRUCT(StructTypeName,ShaderVariableName,StaticSlotName) \
-	FShaderParametersMetadata StructTypeName::StaticStructMetadata( \
-	FShaderParametersMetadata::EUseCase::UniformBuffer, \
-	TEXT(#StructTypeName), \
-	TEXT(#StructTypeName), \
-	TEXT(ShaderVariableName), \
-	TEXT(#StaticSlotName), \
-	StructTypeName::FTypeInfo::FileName, \
-	StructTypeName::FTypeInfo::FileLine, \
-	sizeof(StructTypeName), \
-	StructTypeName::zzGetMembers())
+	IMPLEMENT_STATIC_UNIFORM_BUFFER_STRUCT_EX(StructTypeName, ShaderVariableName, StaticSlotName, Static)
 
 /** Implements a uniform buffer static binding slot. */
 #define IMPLEMENT_STATIC_UNIFORM_BUFFER_SLOT(SlotName) \
@@ -1289,7 +1505,7 @@ extern RENDERCORE_API FShaderParametersMetadata* FindUniformBufferStructByShader
   *	SHADER_PARAMETER_RDG_UNIFORM_BUFFER(FMyStruct, MemberName)
   */
 #define SHADER_PARAMETER_RDG_UNIFORM_BUFFER(StructType, MemberName) \
-	INTERNAL_SHADER_PARAMETER_EXPLICIT(UBMT_RDG_UNIFORM_BUFFER, TShaderParameterTypeInfo<TRDGUniformBuffer<StructType>*>, TRDGUniformBuffer<StructType>*,MemberName,, = nullptr,EShaderPrecisionModifier::Float,TEXT(#StructType),false)
+	INTERNAL_SHADER_PARAMETER_EXPLICIT(UBMT_RDG_UNIFORM_BUFFER, TShaderParameterTypeInfo<TRDGUniformBufferBinding<StructType>>, TRDGUniformBufferBinding<StructType>,MemberName,,,EShaderPrecisionModifier::Float,TEXT(#StructType),false)
 
 /** Nests a shader parameter structure into another one, in C++ and shader code.
  *
@@ -1375,7 +1591,7 @@ extern RENDERCORE_API FShaderParametersMetadata* FindUniformBufferStructByShader
  *		SHADER_PARAMETER_STRUCT_REF(FMyNestedStruct, MyStruct)
  */
 #define SHADER_PARAMETER_STRUCT_REF(StructType,MemberName) \
-	INTERNAL_SHADER_PARAMETER_EXPLICIT(UBMT_REFERENCED_STRUCT, TShaderParameterTypeInfo<TUniformBufferRef<StructType>>, TUniformBufferRef<StructType>,MemberName,,,EShaderPrecisionModifier::Float,TEXT(#StructType),false)
+	INTERNAL_SHADER_PARAMETER_EXPLICIT(UBMT_REFERENCED_STRUCT, TShaderParameterTypeInfo<TUniformBufferBinding<StructType>>, TUniformBufferBinding<StructType>,MemberName,,,EShaderPrecisionModifier::Float,TEXT(#StructType),false)
 
 /** Informs the RDG pass to transition the buffer into the requested state.
  *
diff --git a/Engine/Source/Runtime/RenderCore/Public/ShaderParameterMetadata.h b/Engine/Source/Runtime/RenderCore/Public/ShaderParameterMetadata.h
--- a/Engine/Source/Runtime/RenderCore/Public/ShaderParameterMetadata.h
+++ b/Engine/Source/Runtime/RenderCore/Public/ShaderParameterMetadata.h
@@ -34,9 +34,20 @@ struct FResourceTableEntry
 	/** The name of the uniform buffer in which this resource exists. */
 	FString UniformBufferName;
 	/** The type of the resource (EUniformBufferBaseType). */
-	uint16 Type;
+	uint16 Type{};
 	/** The index of the resource in the table. */
-	uint16 ResourceIndex;
+	uint16 ResourceIndex{};
+};
+
+/** Minimal information about each uniform buffer entry fed to the shader compiler. */
+struct FUniformBufferEntry
+{
+	/** The name of the uniform buffer static slot (if global). */
+	FString StaticSlotName;
+	/** Hash of the resource table layout. */
+	uint32 LayoutHash{};
+	/** The binding flags used by this resource table. */
+	EUniformBufferBindingFlags BindingFlags{ EUniformBufferBindingFlags::Shader };
 };
 
 /** Parse the shader resource binding from the binding type used in shader code. */
@@ -206,6 +217,7 @@ public:
 	/** Initialization constructor. */
 	FShaderParametersMetadata(
 		EUseCase UseCase,
+		EUniformBufferBindingFlags InBindingFlags,
 		const TCHAR* InLayoutName,
 		const TCHAR* InStructTypeName,
 		const TCHAR* InShaderVariableName,
@@ -219,7 +231,7 @@ public:
 
 	void GetNestedStructs(TArray<const FShaderParametersMetadata*>& OutNestedStructs) const;
 
-	void AddResourceTableEntries(TMap<FString, FResourceTableEntry>& ResourceTableMap, TMap<FString, uint32>& ResourceTableLayoutHashes, TMap<FString, FString>& ResourceTableLayoutSlots) const;
+	void AddResourceTableEntries(TMap<FString, FResourceTableEntry>& ResourceTableMap, TMap<FString, FUniformBufferEntry>& UniformBufferMap) const;
 
 	const TCHAR* GetStructTypeName() const { return StructTypeName; }
 	const TCHAR* GetShaderVariableName() const { return ShaderVariableName; }
@@ -228,6 +240,16 @@ public:
 
 	bool HasStaticSlot() const { return StaticSlotName != nullptr; }
 
+	EUniformBufferBindingFlags GetBindingFlags() const { return BindingFlags; }
+
+	EUniformBufferBindingFlags GetPreferredBindingFlag() const
+	{
+		// Decay to static when both binding flags are specified.
+		return BindingFlags != EUniformBufferBindingFlags::StaticAndShader
+			? BindingFlags
+			: EUniformBufferBindingFlags::Static;
+	}
+
 	/** Returns the C++ file name where the parameter structure is declared. */
 	const ANSICHAR* GetFileName() const { return FileName; }
 
@@ -292,6 +314,9 @@ private:
 	/** The use case of this shader parameter struct. */
 	const EUseCase UseCase;
 
+	/** The binding model used by this parameter struct. */
+	const EUniformBufferBindingFlags BindingFlags;
+
 	/** Layout of all the resources in the shader parameter struct. */
 	FRHIUniformBufferLayout Layout;
 	
@@ -311,4 +336,4 @@ private:
 	void InitializeLayout();
 
 	void AddResourceTableEntriesRecursive(const TCHAR* UniformBufferName, const TCHAR* Prefix, uint16& ResourceIndex, TMap<FString, FResourceTableEntry>& ResourceTableMap) const;
-};
+};
\ No newline at end of file
diff --git a/Engine/Source/Runtime/RenderCore/Public/ShaderParameterStruct.h b/Engine/Source/Runtime/RenderCore/Public/ShaderParameterStruct.h
--- a/Engine/Source/Runtime/RenderCore/Public/ShaderParameterStruct.h
+++ b/Engine/Source/Runtime/RenderCore/Public/ShaderParameterStruct.h
@@ -272,18 +272,24 @@ inline void SetShaderParameters(
 	// Graph Uniform Buffers
 	for (const FShaderParameterBindings::FParameterStructReference& ParameterBinding : Bindings.GraphUniformBuffers)
 	{
-		auto GraphUniformBuffer = *reinterpret_cast<FRDGUniformBuffer* const*>(Base + ParameterBinding.ByteOffset);
+		const FRDGUniformBufferBinding& UniformBufferBinding = *reinterpret_cast<const FRDGUniformBufferBinding*>(Base + ParameterBinding.ByteOffset);
 
-		checkSlow(GraphUniformBuffer);
-		GraphUniformBuffer->MarkResourceAsUsed();
-		RHICmdList.SetShaderUniformBuffer(ShadeRHI, ParameterBinding.BufferIndex, GraphUniformBuffer->GetRHI());
+		if (UniformBufferBinding.IsShader())
+		{
+			UniformBufferBinding->MarkResourceAsUsed();
+			RHICmdList.SetShaderUniformBuffer(ShadeRHI, ParameterBinding.BufferIndex, UniformBufferBinding->GetRHI());
+		}
 	}
 
 	// Reference structures
 	for (const FShaderParameterBindings::FParameterStructReference& ParameterBinding : Bindings.ParameterReferences)
 	{
-		const TRefCountPtr<FRHIUniformBuffer>& ShaderParameterRef = *reinterpret_cast<const TRefCountPtr<FRHIUniformBuffer>*>(Base + ParameterBinding.ByteOffset);
-		RHICmdList.SetShaderUniformBuffer(ShadeRHI, ParameterBinding.BufferIndex, ShaderParameterRef);
+		const FUniformBufferBinding& UniformBufferBinding = *reinterpret_cast<const FUniformBufferBinding*>(Base + ParameterBinding.ByteOffset);
+
+		if (UniformBufferBinding.IsShader())
+		{
+			RHICmdList.SetShaderUniformBuffer(ShadeRHI, ParameterBinding.BufferIndex, UniformBufferBinding.GetUniformBuffer());
+		}
 	}
 }
 
@@ -375,18 +381,18 @@ void SetShaderParameters(FRayTracingShaderBindingsWriter& RTBindingsWriter, cons
 	// Graph Uniform Buffers
 	for (const FShaderParameterBindings::FParameterStructReference& ParameterBinding : Bindings.GraphUniformBuffers)
 	{
-		auto GraphUniformBuffer = *reinterpret_cast<FRDGUniformBuffer* const*>(Base + ParameterBinding.ByteOffset);
+		const FRDGUniformBufferBinding& UniformBufferBinding = *reinterpret_cast<const FRDGUniformBufferBinding*>(Base + ParameterBinding.ByteOffset);
 
-		checkSlow(GraphUniformBuffer);
-		GraphUniformBuffer->MarkResourceAsUsed();
-		RTBindingsWriter.SetUniformBuffer(ParameterBinding.BufferIndex, GraphUniformBuffer->GetRHI());
+		checkSlow(UniformBufferBinding);
+		UniformBufferBinding->MarkResourceAsUsed();
+		RTBindingsWriter.SetUniformBuffer(ParameterBinding.BufferIndex, UniformBufferBinding->GetRHI());
 	}
 
 	// Referenced uniform buffers
 	for (const FShaderParameterBindings::FParameterStructReference& ParameterBinding : Bindings.ParameterReferences)
 	{
-		const TRefCountPtr<FRHIUniformBuffer>& ShaderParameterRef = *reinterpret_cast<const TRefCountPtr<FRHIUniformBuffer>*>(Base + ParameterBinding.ByteOffset);
-		RTBindingsWriter.SetUniformBuffer(ParameterBinding.BufferIndex, ShaderParameterRef);
+		const FUniformBufferBinding& UniformBufferBinding = *reinterpret_cast<const FUniformBufferBinding*>(Base + ParameterBinding.ByteOffset);
+		RTBindingsWriter.SetUniformBuffer(ParameterBinding.BufferIndex, UniformBufferBinding.GetUniformBuffer());
 	}
 
 	// Root uniform buffer.
