Fixed build break.

diff --git a/Engine/Source/Runtime/RenderCore/Public/ShaderParameterMacros.h b/Engine/Source/Runtime/RenderCore/Public/ShaderParameterMacros.h
--- a/Engine/Source/Runtime/RenderCore/Public/ShaderParameterMacros.h
+++ b/Engine/Source/Runtime/RenderCore/Public/ShaderParameterMacros.h
@@ -1081,7 +1081,7 @@ struct TShaderParameterTypeInfo<TUniformBufferBinding<UniformBufferStructType>>
 	static constexpr int32 Alignment = SHADER_PARAMETER_STRUCT_ALIGNMENT;
 	static constexpr bool bIsStoredInConstantBuffer = false;
 
-	using TAlignedType = TAlignedShaderParameterPtr<TUniformBufferBinding<UniformBufferStructType>>;
+	using TAlignedType = TUniformBufferBinding<UniformBufferStructType>;
 
 	static const FShaderParametersMetadata* GetStructMetadata() { return &UniformBufferStructType::StaticStructMetadata; }
 };
