Improved RDG state merging to better handle subresources. Subresource tracking is used for the rest of the graph if subresource states are encountered. Improved error reporting for edge case where transitions are inserted out of order. Added explicit merge check for depth read / write to make sure that it doesn't merge with anything else.

diff --git a/Engine/Source/Runtime/RenderCore/Private/RenderGraphBuilder.cpp b/Engine/Source/Runtime/RenderCore/Private/RenderGraphBuilder.cpp
--- a/Engine/Source/Runtime/RenderCore/Private/RenderGraphBuilder.cpp
+++ b/Engine/Source/Runtime/RenderCore/Private/RenderGraphBuilder.cpp
@@ -687,14 +687,11 @@ void FRDGBuilder::Compile()
 				auto& LastProducers = Texture->LastProducers;
 				auto& PassState = TexturePair.Value.State;
 
+			#if DO_CHECK
 				const bool bWholePassState = IsWholeResource(PassState);
 				const bool bWholeProducers = IsWholeResource(LastProducers);
-
-				// The producer array needs to be at least as large as the pass state array.
-				if (bWholeProducers && !bWholePassState)
-				{
-					InitAsSubresources(LastProducers, Texture->Layout, GetWholeResource(LastProducers));
-				}
+				checkf(!bWholeProducers || bWholePassState == bWholeProducers, TEXT("The producer array needs to be at least as large as the pass state array."));
+			#endif
 
 				for (uint32 Index = 0, Count = LastProducers.Num(); Index < Count; ++Index)
 				{
@@ -883,17 +880,11 @@ void FRDGBuilder::Compile()
 				Texture->ReferenceCount += PassState.ReferenceCount;
 				Texture->bUsedByAsyncComputePass |= bAsyncComputePass;
 
-				const bool bWholePassState = IsWholeResource(PassState.State);
-				const bool bWholeMergeState = IsWholeResource(Texture->MergeState);
-
 				// For simplicity, the merge / pass state dimensionality should match.
-				if (bWholeMergeState && !bWholePassState)
-				{
-					InitAsSubresources(Texture->MergeState, Texture->Layout);
-				}
-				else if (!bWholeMergeState && bWholePassState)
+				if (Texture->MergeState.Num() != PassState.State.Num())
 				{
-					InitAsWholeResource(Texture->MergeState);
+					check(IsSubresources(Texture->MergeState));
+					InitAsSubresources(PassState.State, Texture->Layout, GetWholeResource(PassState.State));
 				}
 
 				const uint32 SubresourceCount = PassState.State.Num();
@@ -1166,7 +1157,7 @@ void FRDGBuilder::Compile()
 
 		for (FRDGPassHandle PassHandle = Passes.Begin(); PassHandle != Passes.End(); ++PassHandle)
 		{
-			if (PassesToCull[PassHandle])
+			if (PassesToCull[PassHandle] || PassesWithEmptyParameters[PassHandle])
 			{
 				continue;
 			}
@@ -1437,8 +1428,16 @@ void FRDGBuilder::SetupPass(FRDGPass* Pass)
 		// Convert the pass state to subresource dimensionality if we've found a subresource range.
 		if (!bWholeTextureRange && bWholePassState)
 		{
-			InitAsSubresources(PassState.State, Texture->Layout);
+			InitAsSubresources(PassState.State, Texture->Layout, GetWholeResource(PassState.State));
 			bWholePassState = false;
+
+			// Also convert the texture's internal arrays to subresources. We'll process everything as subresources from now on.
+			if (IsWholeResource(Texture->MergeState))
+			{
+				check(IsWholeResource(Texture->LastProducers));
+				InitAsSubresources(Texture->MergeState, Texture->Layout);
+				InitAsSubresources(Texture->LastProducers, Texture->Layout);
+			}
 		}
 
 		const auto AddSubresourceAccess = [&](FRDGSubresourceState& State)
diff --git a/Engine/Source/Runtime/RenderCore/Private/RenderGraphPass.cpp b/Engine/Source/Runtime/RenderCore/Private/RenderGraphPass.cpp
--- a/Engine/Source/Runtime/RenderCore/Private/RenderGraphPass.cpp
+++ b/Engine/Source/Runtime/RenderCore/Private/RenderGraphPass.cpp
@@ -222,6 +222,19 @@ void FRDGBarrierBatchBegin::Submit(FRHIComputeCommandList& RHICmdList, ERHIPipel
 
 void FRDGBarrierBatchEnd::AddDependency(FRDGBarrierBatchBegin* BeginBatch)
 {
+#if RDG_ENABLE_DEBUG
+	check(BeginBatch);
+
+	for (ERHIPipeline Pipeline : GetRHIPipelines())
+	{
+		const FRDGPassHandle BeginPassHandle = BeginBatch->DebugPasses[Pipeline];
+		if (BeginPassHandle.IsValid())
+		{
+			checkf(BeginPassHandle <= PassHandle, TEXT("A transition end batch for pass %d is dependent on begin batch for pass %d."), PassHandle.GetIndex(), BeginPassHandle.GetIndex());
+		}
+	}
+#endif
+
 	Dependencies.AddUnique(BeginBatch);
 }
 
@@ -282,7 +295,7 @@ FRDGBarrierBatchEnd& FRDGPass::GetPrologueBarriersToEnd(FRDGAllocator& Allocator
 {
 	if (!PrologueBarriersToEnd)
 	{
-		PrologueBarriersToEnd = Allocator.AllocNoDestruct<FRDGBarrierBatchEnd>();
+		PrologueBarriersToEnd = Allocator.AllocNoDestruct<FRDGBarrierBatchEnd>(Handle);
 	}
 	return *PrologueBarriersToEnd;
 }
@@ -291,7 +304,7 @@ FRDGBarrierBatchEnd& FRDGPass::GetEpilogueBarriersToEnd(FRDGAllocator& Allocator
 {
 	if (!EpilogueBarriersToEnd)
 	{
-		EpilogueBarriersToEnd = Allocator.AllocNoDestruct<FRDGBarrierBatchEnd>();
+		EpilogueBarriersToEnd = Allocator.AllocNoDestruct<FRDGBarrierBatchEnd>(Handle);
 	}
 	return *EpilogueBarriersToEnd;
 }
diff --git a/Engine/Source/Runtime/RenderCore/Private/RenderGraphResources.cpp b/Engine/Source/Runtime/RenderCore/Private/RenderGraphResources.cpp
--- a/Engine/Source/Runtime/RenderCore/Private/RenderGraphResources.cpp
+++ b/Engine/Source/Runtime/RenderCore/Private/RenderGraphResources.cpp
@@ -81,6 +81,7 @@ bool FRDGSubresourceState::IsMergeAllowed(ERDGParentResourceType ResourceType, c
 	 */
 
 	const ERHIAccess AccessUnion = Previous.Access | Next.Access;
+	const ERHIAccess DSVMask = ERHIAccess::DSVRead | ERHIAccess::DSVWrite;
 
 	// If we have the same access between the two states, we don't need to check for invalid access combinations.
 	if (Previous.Access != Next.Access)
@@ -103,6 +104,12 @@ bool FRDGSubresourceState::IsMergeAllowed(ERDGParentResourceType ResourceType, c
 			return false;
 		}
 
+		// Depth Read / Write should never merge with anything other than itself.
+		if (EnumHasAllFlags(AccessUnion, DSVMask) && EnumHasAnyFlags(AccessUnion, ~DSVMask))
+		{
+			return false;
+		}
+
 		// Filter out platform-specific unsupported mergeable states.
 		if (EnumHasAnyFlags(AccessUnion, ~GRHIMergeableAccessMask))
 		{
diff --git a/Engine/Source/Runtime/RenderCore/Public/RenderGraphBuilder.h b/Engine/Source/Runtime/RenderCore/Public/RenderGraphBuilder.h
--- a/Engine/Source/Runtime/RenderCore/Public/RenderGraphBuilder.h
+++ b/Engine/Source/Runtime/RenderCore/Public/RenderGraphBuilder.h
@@ -101,7 +101,7 @@ public:
 	 *  bound for write access by the pass.
 	 */
 	template <typename ParameterStructType>
-	TRDGUniformBufferRef<ParameterStructType> CreateUniformBuffer(ParameterStructType* ParameterStruct);
+	TRDGUniformBufferRef<ParameterStructType> CreateUniformBuffer(const ParameterStructType* ParameterStruct);
 
 	//////////////////////////////////////////////////////////////////////////
 	// Allocation Methods
diff --git a/Engine/Source/Runtime/RenderCore/Public/RenderGraphBuilder.inl b/Engine/Source/Runtime/RenderCore/Public/RenderGraphBuilder.inl
--- a/Engine/Source/Runtime/RenderCore/Public/RenderGraphBuilder.inl
+++ b/Engine/Source/Runtime/RenderCore/Public/RenderGraphBuilder.inl
@@ -103,7 +103,7 @@ FORCEINLINE FRDGSubresourceState* FRDGBuilder::AllocSubresource(const FRDGSubres
 }
 
 template <typename ParameterStructType>
-TRDGUniformBufferRef<ParameterStructType> FRDGBuilder::CreateUniformBuffer(ParameterStructType* ParameterStruct)
+TRDGUniformBufferRef<ParameterStructType> FRDGBuilder::CreateUniformBuffer(const ParameterStructType* ParameterStruct)
 {
 	IF_RDG_ENABLE_DEBUG(UserValidation.ValidateCreateUniformBuffer(ParameterStruct, &ParameterStructType::StaticStructMetadata));
 	auto* UniformBuffer = UniformBuffers.Allocate<TRDGUniformBuffer<ParameterStructType>>(Allocator, ParameterStruct, ParameterStructType::StaticStructMetadata.GetShaderVariableName());
diff --git a/Engine/Source/Runtime/RenderCore/Public/RenderGraphPass.h b/Engine/Source/Runtime/RenderCore/Public/RenderGraphPass.h
--- a/Engine/Source/Runtime/RenderCore/Public/RenderGraphPass.h
+++ b/Engine/Source/Runtime/RenderCore/Public/RenderGraphPass.h
@@ -83,7 +83,11 @@ private:
 class RENDERCORE_API FRDGBarrierBatchEnd
 {
 public:
-	FRDGBarrierBatchEnd() = default;
+	FRDGBarrierBatchEnd(FRDGPassHandle InPassHandle)
+#if RDG_ENABLE_DEBUG
+		: PassHandle(InPassHandle)
+#endif
+	{}
 
 	/** Inserts a dependency on a begin batch. A begin batch can be inserted into more than one end batch. */
 	void AddDependency(FRDGBarrierBatchBegin* BeginBatch);
@@ -93,6 +97,10 @@ public:
 private:
 	TArray<FRDGBarrierBatchBegin*, TInlineAllocator<1, FRDGArrayAllocator>> Dependencies;
 
+#if RDG_ENABLE_DEBUG
+	FRDGPassHandle PassHandle;
+#endif
+
 	friend class FRDGBarrierValidation;
 };
 
@@ -412,6 +420,9 @@ public:
 		ExecuteLambdaType&& InExecuteLambda)
 		: FRDGPass(MoveTemp(InName), FRDGParameterStruct(InParameterStruct, &InParameterMetadata->GetLayout()), InPassFlags)
 		, ExecuteLambda(MoveTemp(InExecuteLambda))
+#if RDG_ENABLE_DEBUG
+		, DebugParameterStruct(InParameterStruct)
+#endif
 	{
 		checkf(kSupportsAsyncCompute || !EnumHasAnyFlags(InPassFlags, ERDGPassFlags::AsyncCompute),
 			TEXT("Pass %s is set to use 'AsyncCompute', but the pass lambda's first argument is not FRHIComputeCommandList&."), GetName());
@@ -425,6 +436,8 @@ private:
 	}
 
 	ExecuteLambdaType ExecuteLambda;
+
+	IF_RDG_ENABLE_DEBUG(const ParameterStructType* DebugParameterStruct);
 };
 
 template <typename ExecuteLambdaType>
