Fixed build break.

diff --git a/Engine/Source/Runtime/RenderCore/Private/RenderGraphBuilder.cpp b/Engine/Source/Runtime/RenderCore/Private/RenderGraphBuilder.cpp
--- a/Engine/Source/Runtime/RenderCore/Private/RenderGraphBuilder.cpp
+++ b/Engine/Source/Runtime/RenderCore/Private/RenderGraphBuilder.cpp
@@ -687,11 +687,9 @@ void FRDGBuilder::Compile()
 				auto& LastProducers = Texture->LastProducers;
 				auto& PassState = TexturePair.Value.State;
 
-			#if DO_CHECK
 				const bool bWholePassState = IsWholeResource(PassState);
 				const bool bWholeProducers = IsWholeResource(LastProducers);
 				checkf(!bWholeProducers || bWholePassState == bWholeProducers, TEXT("The producer array needs to be at least as large as the pass state array."));
-			#endif
 
 				for (uint32 Index = 0, Count = LastProducers.Num(); Index < Count; ++Index)
 				{
