Fixed crash in UseNanite() if the r.Nanite cvar is missing.

diff --git a/Engine/Source/Runtime/RenderCore/Public/RenderUtils.h b/Engine/Source/Runtime/RenderCore/Public/RenderUtils.h
--- a/Engine/Source/Runtime/RenderCore/Public/RenderUtils.h
+++ b/Engine/Source/Runtime/RenderCore/Public/RenderUtils.h
@@ -605,6 +605,6 @@ inline bool UseNanite(EShaderPlatform ShaderPlatform)
 	static const auto EnableNaniteCVar = IConsoleManager::Get().FindConsoleVariable(TEXT("r.Nanite"));
 	bool bNaniteSupported = DoesPlatformSupportNanite(ShaderPlatform);
 	bool bForwardShadingEnabled = IsForwardShadingEnabled(ShaderPlatform);
-	return bNaniteSupported /*&& GRHISupportsAtomicUInt64*/ && !bForwardShadingEnabled && EnableNaniteCVar->GetInt() > 0;
+	return bNaniteSupported /*&& GRHISupportsAtomicUInt64*/ && !bForwardShadingEnabled && (EnableNaniteCVar != nullptr && EnableNaniteCVar->GetInt() > 0);
 }
 
