Refactored all RHI's to include RHICore. Moved ApplyStaticUniformBuffer helpers from ShaderCore into RHICore.

diff --git a/Engine/Plugins/Lumin/MagicLeap/Source/MagicLeapHelperOpenGL/MagicLeapHelperOpenGL.Build.cs b/Engine/Plugins/Lumin/MagicLeap/Source/MagicLeapHelperOpenGL/MagicLeapHelperOpenGL.Build.cs
--- a/Engine/Plugins/Lumin/MagicLeap/Source/MagicLeapHelperOpenGL/MagicLeapHelperOpenGL.Build.cs
+++ b/Engine/Plugins/Lumin/MagicLeap/Source/MagicLeapHelperOpenGL/MagicLeapHelperOpenGL.Build.cs
@@ -17,6 +17,7 @@ namespace UnrealBuildTool.Rules
 					"Engine",
 					"RenderCore",
 					"RHI",
+					"RHICore",
 					"MLSDK"
 				});
 
diff --git a/Engine/Plugins/Lumin/MagicLeap/Source/MagicLeapHelperVulkan/MagicLeapHelperVulkan.Build.cs b/Engine/Plugins/Lumin/MagicLeap/Source/MagicLeapHelperVulkan/MagicLeapHelperVulkan.Build.cs
--- a/Engine/Plugins/Lumin/MagicLeap/Source/MagicLeapHelperVulkan/MagicLeapHelperVulkan.Build.cs
+++ b/Engine/Plugins/Lumin/MagicLeap/Source/MagicLeapHelperVulkan/MagicLeapHelperVulkan.Build.cs
@@ -17,6 +17,7 @@ namespace UnrealBuildTool.Rules
 					"Engine",
 					"MLSDK",
 					"RHI",
+					"RHICore",
 					"RenderCore",
 					"HeadMountedDisplay"
 				});
diff --git a/Engine/Plugins/Media/ImgMedia/Source/ImgMedia/ImgMedia.Build.cs b/Engine/Plugins/Media/ImgMedia/Source/ImgMedia/ImgMedia.Build.cs
--- a/Engine/Plugins/Media/ImgMedia/Source/ImgMedia/ImgMedia.Build.cs
+++ b/Engine/Plugins/Media/ImgMedia/Source/ImgMedia/ImgMedia.Build.cs
@@ -24,6 +24,7 @@ namespace UnrealBuildTool.Rules
 					"RenderCore",
 					"Renderer",
 					"RHI",
+					"RHICore",
 				});
 
 			PrivateIncludePathModuleNames.AddRange(
diff --git a/Engine/Plugins/Runtime/GoogleVR/GoogleVRHMD/Source/GoogleVRHMD/GoogleVRHMD.Build.cs b/Engine/Plugins/Runtime/GoogleVR/GoogleVRHMD/Source/GoogleVRHMD/GoogleVRHMD.Build.cs
--- a/Engine/Plugins/Runtime/GoogleVR/GoogleVRHMD/Source/GoogleVRHMD/GoogleVRHMD.Build.cs
+++ b/Engine/Plugins/Runtime/GoogleVR/GoogleVRHMD/Source/GoogleVRHMD/GoogleVRHMD.Build.cs
@@ -24,6 +24,7 @@ namespace UnrealBuildTool.Rules
 					"CoreUObject",
 					"Engine",
 					"RHI",
+					"RHICore",
 					"RenderCore",
 					"Renderer",
 					"HeadMountedDisplay",
diff --git a/Engine/Plugins/Runtime/TextureShare/Source/TextureShareD3D12/TextureShareD3D12.Build.cs b/Engine/Plugins/Runtime/TextureShare/Source/TextureShareD3D12/TextureShareD3D12.Build.cs
--- a/Engine/Plugins/Runtime/TextureShare/Source/TextureShareD3D12/TextureShareD3D12.Build.cs
+++ b/Engine/Plugins/Runtime/TextureShare/Source/TextureShareD3D12/TextureShareD3D12.Build.cs
@@ -14,6 +14,7 @@ public class TextureShareD3D12 : ModuleRules
 							"Engine",
 							"RenderCore",
 							"RHI",
+							"RHICore",
 							"D3D12RHI",
 						});
 
diff --git a/Engine/Plugins/Runtime/nDisplay/Source/DisplayCluster/DisplayCluster.Build.cs b/Engine/Plugins/Runtime/nDisplay/Source/DisplayCluster/DisplayCluster.Build.cs
--- a/Engine/Plugins/Runtime/nDisplay/Source/DisplayCluster/DisplayCluster.Build.cs
+++ b/Engine/Plugins/Runtime/nDisplay/Source/DisplayCluster/DisplayCluster.Build.cs
@@ -27,6 +27,7 @@ public class DisplayCluster : ModuleRules
 				"Networking",
 				"RenderCore",
 				"RHI",
+				"RHICore",
 				"Slate",
 				"SlateCore",
 				"Sockets",
diff --git a/Engine/Plugins/Runtime/nDisplay/Source/DisplayClusterRendering/DisplayClusterRendering.Build.cs b/Engine/Plugins/Runtime/nDisplay/Source/DisplayClusterRendering/DisplayClusterRendering.Build.cs
--- a/Engine/Plugins/Runtime/nDisplay/Source/DisplayClusterRendering/DisplayClusterRendering.Build.cs
+++ b/Engine/Plugins/Runtime/nDisplay/Source/DisplayClusterRendering/DisplayClusterRendering.Build.cs
@@ -35,6 +35,7 @@ public class DisplayClusterRendering : ModuleRules
 				"Renderer",
 				"RenderCore",
 				"RHI",
+				"RHICore",
 				"Slate",
 				"SlateCore"
 			}
diff --git a/Engine/Source/Runtime/Apple/MetalRHI/MetalRHI.Build.cs b/Engine/Source/Runtime/Apple/MetalRHI/MetalRHI.Build.cs
--- a/Engine/Source/Runtime/Apple/MetalRHI/MetalRHI.Build.cs
+++ b/Engine/Source/Runtime/Apple/MetalRHI/MetalRHI.Build.cs
@@ -13,6 +13,7 @@ public class MetalRHI : ModuleRules
 				"ApplicationCore",
 				"Engine",
 				"RHI",
+				"RHICore",
 				"RenderCore",
 			}
 			);
diff --git a/Engine/Source/Runtime/Apple/MetalRHI/Private/MetalCommands.cpp b/Engine/Source/Runtime/Apple/MetalRHI/Private/MetalCommands.cpp
--- a/Engine/Source/Runtime/Apple/MetalRHI/Private/MetalCommands.cpp
+++ b/Engine/Source/Runtime/Apple/MetalRHI/Private/MetalCommands.cpp
@@ -134,7 +134,7 @@ void FMetalRHICommandContext::RHISetComputeShader(FRHIComputeShader* ComputeShad
 	// sets this compute shader pipeline as the current (this resets all state, so we need to set all resources after calling this)
 	Context->GetCurrentState().SetComputeShader(ComputeShader);
 
-	ApplyGlobalUniformBuffers(ComputeShader);
+	ApplyStaticUniformBuffers(ComputeShader);
 	}
 }
 
@@ -147,7 +147,7 @@ void FMetalRHICommandContext::RHISetComputePipelineState(FRHIComputePipelineStat
 	// sets this compute shader pipeline as the current (this resets all state, so we need to set all resources after calling this)
 	Context->GetCurrentState().SetComputeShader(ComputePipeline->GetComputeShader());
 
-	ApplyGlobalUniformBuffers(ComputePipeline->GetComputeShader());
+	ApplyStaticUniformBuffers(ComputePipeline->GetComputeShader());
 	}
 }
 
@@ -267,12 +267,12 @@ void FMetalRHICommandContext::RHISetGraphicsPipelineState(FRHIGraphicsPipelineSt
 
 		if (bApplyAdditionalState)
 		{
-			ApplyGlobalUniformBuffers(PipelineState->VertexShader.GetReference());
+			ApplyStaticUniformBuffers(PipelineState->VertexShader.GetReference());
 #if PLATFORM_SUPPORTS_TESSELLATION_SHADERS
-			ApplyGlobalUniformBuffers(PipelineState->HullShader.GetReference());
-			ApplyGlobalUniformBuffers(PipelineState->DomainShader.GetReference());
+			ApplyStaticUniformBuffers(PipelineState->HullShader.GetReference());
+			ApplyStaticUniformBuffers(PipelineState->DomainShader.GetReference());
 #endif
-			ApplyGlobalUniformBuffers(PipelineState->PixelShader.GetReference());
+			ApplyStaticUniformBuffers(PipelineState->PixelShader.GetReference());
 		}
 	}
 }
diff --git a/Engine/Source/Runtime/Apple/MetalRHI/Public/MetalRHIContext.h b/Engine/Source/Runtime/Apple/MetalRHI/Public/MetalRHIContext.h
--- a/Engine/Source/Runtime/Apple/MetalRHI/Public/MetalRHIContext.h
+++ b/Engine/Source/Runtime/Apple/MetalRHI/Public/MetalRHIContext.h
@@ -7,6 +7,7 @@
 #include "MetalState.h"
 #include "MetalResources.h"
 #include "MetalViewport.h"
+#include "RHICoreShader.h"
 
 #define UE_METAL_RHI_SUPPORT_CLEAR_UAV_WITH_BLIT_ENCODER	1
 
@@ -236,11 +237,11 @@ protected:
 	uint32 PendingNumPrimitives;
 
 	template <typename TRHIShader>
-	void ApplyGlobalUniformBuffers(TRHIShader* Shader)
+	void ApplyStaticUniformBuffers(TRHIShader* Shader)
 	{
 		if (Shader)
 		{
-			::ApplyGlobalUniformBuffers(this, Shader, Shader->StaticSlots, Shader->Bindings.ShaderResourceTable.ResourceTableLayoutHashes, GlobalUniformBuffers);
+			UE::RHICore::ApplyStaticUniformBuffers(this, Shader, Shader->StaticSlots, Shader->Bindings.ShaderResourceTable.ResourceTableLayoutHashes, GlobalUniformBuffers);
 		}
 	}
 
diff --git a/Engine/Source/Runtime/D3D12RHI/D3D12RHI.Build.cs b/Engine/Source/Runtime/D3D12RHI/D3D12RHI.Build.cs
--- a/Engine/Source/Runtime/D3D12RHI/D3D12RHI.Build.cs
+++ b/Engine/Source/Runtime/D3D12RHI/D3D12RHI.Build.cs
@@ -19,6 +19,7 @@ public class D3D12RHI : ModuleRules
 				"CoreUObject",
 				"Engine",
 				"RHI",
+				"RHICore",
 				"RenderCore",
 				}
 			);
diff --git a/Engine/Source/Runtime/D3D12RHI/Private/D3D12CommandContext.h b/Engine/Source/Runtime/D3D12RHI/Private/D3D12CommandContext.h
--- a/Engine/Source/Runtime/D3D12RHI/Private/D3D12CommandContext.h
+++ b/Engine/Source/Runtime/D3D12RHI/Private/D3D12CommandContext.h
@@ -26,6 +26,8 @@ THIRD_PARTY_INCLUDES_START
 #include "Windows/HideWindowsPlatformTypes.h"
 THIRD_PARTY_INCLUDES_END
 
+#include "RHICoreShader.h"
+
 struct FRayTracingShaderBindings;
 
 // Base class used to define commands that are not device specific, or that broadcast to all devices.
@@ -446,11 +448,11 @@ private:
 	static void ClearUAV(TRHICommandList_RecursiveHazardous<FD3D12CommandContext>& RHICmdList, FD3D12UnorderedAccessView* UAV, const void* ClearValues, bool bFloat);
 
 	template <typename TRHIShader>
-	void ApplyGlobalUniformBuffers(TRHIShader* Shader)
+	void ApplyStaticUniformBuffers(TRHIShader* Shader)
 	{
 		if (Shader)
 		{
-			::ApplyGlobalUniformBuffers(this, Shader, Shader->StaticSlots, Shader->ShaderResourceTable.ResourceTableLayoutHashes, GlobalUniformBuffers);
+			UE::RHICore::ApplyStaticUniformBuffers(this, Shader, Shader->StaticSlots, Shader->ShaderResourceTable.ResourceTableLayoutHashes, GlobalUniformBuffers);
 		}
 	}
 
diff --git a/Engine/Source/Runtime/D3D12RHI/Private/D3D12Commands.cpp b/Engine/Source/Runtime/D3D12RHI/Private/D3D12Commands.cpp
--- a/Engine/Source/Runtime/D3D12RHI/Private/D3D12Commands.cpp
+++ b/Engine/Source/Runtime/D3D12RHI/Private/D3D12Commands.cpp
@@ -565,11 +565,11 @@ void FD3D12CommandContext::RHISetGraphicsPipelineState(FRHIGraphicsPipelineState
 
 	if (bApplyAdditionalState)
 	{
-		ApplyGlobalUniformBuffers(GraphicsPipelineState->GetVertexShader());
-		ApplyGlobalUniformBuffers(GraphicsPipelineState->GetHullShader());
-		ApplyGlobalUniformBuffers(GraphicsPipelineState->GetDomainShader());
-		ApplyGlobalUniformBuffers(GraphicsPipelineState->GetGeometryShader());
-		ApplyGlobalUniformBuffers(GraphicsPipelineState->GetPixelShader());
+		ApplyStaticUniformBuffers(GraphicsPipelineState->GetVertexShader());
+		ApplyStaticUniformBuffers(GraphicsPipelineState->GetHullShader());
+		ApplyStaticUniformBuffers(GraphicsPipelineState->GetDomainShader());
+		ApplyStaticUniformBuffers(GraphicsPipelineState->GetGeometryShader());
+		ApplyStaticUniformBuffers(GraphicsPipelineState->GetPixelShader());
 	}
 }
 
@@ -591,7 +591,7 @@ void FD3D12CommandContext::RHISetComputePipelineState(FRHIComputePipelineState*
 
 	StateCache.SetComputePipelineState(ComputePipelineState);
 
-	ApplyGlobalUniformBuffers(ComputePipelineState->ComputeShader.GetReference());
+	ApplyStaticUniformBuffers(ComputePipelineState->ComputeShader.GetReference());
 }
 
 void FD3D12CommandContext::RHISetShaderTexture(FRHIGraphicsShader* ShaderRHI, uint32 TextureIndex, FRHITexture* NewTextureRHI)
diff --git a/Engine/Source/Runtime/OpenGLDrv/OpenGLDrv.Build.cs b/Engine/Source/Runtime/OpenGLDrv/OpenGLDrv.Build.cs
--- a/Engine/Source/Runtime/OpenGLDrv/OpenGLDrv.Build.cs
+++ b/Engine/Source/Runtime/OpenGLDrv/OpenGLDrv.Build.cs
@@ -16,6 +16,7 @@ public class OpenGLDrv : ModuleRules
 				"ApplicationCore",
 				"Engine",
 				"RHI",
+				"RHICore",
 				"RenderCore",
 				"PreLoadScreen"
 			}
diff --git a/Engine/Source/Runtime/OpenGLDrv/Private/OpenGLCommands.cpp b/Engine/Source/Runtime/OpenGLDrv/Private/OpenGLCommands.cpp
--- a/Engine/Source/Runtime/OpenGLDrv/Private/OpenGLCommands.cpp
+++ b/Engine/Source/Runtime/OpenGLDrv/Private/OpenGLCommands.cpp
@@ -3194,7 +3194,7 @@ void FOpenGLDynamicRHI::RHISetComputeShader(FRHIComputeShader* ComputeShaderRHI)
 		UE_LOG(LogRHI,Fatal,TEXT("Platform doesn't support SM5 for OpenGL but set feature level to SM5"));
 	}
 
-	ApplyGlobalUniformBuffers(ComputeShaderRHI, ResourceCast(ComputeShaderRHI));
+	ApplyStaticUniformBuffers(ComputeShaderRHI, ResourceCast(ComputeShaderRHI));
 }
 
 void FOpenGLDynamicRHI::RHIDispatchComputeShader(uint32 ThreadGroupCountX, uint32 ThreadGroupCountY, uint32 ThreadGroupCountZ)
diff --git a/Engine/Source/Runtime/OpenGLDrv/Public/OpenGLDrv.h b/Engine/Source/Runtime/OpenGLDrv/Public/OpenGLDrv.h
--- a/Engine/Source/Runtime/OpenGLDrv/Public/OpenGLDrv.h
+++ b/Engine/Source/Runtime/OpenGLDrv/Public/OpenGLDrv.h
@@ -15,6 +15,7 @@
 #include "GPUProfiler.h"
 #include "RenderResource.h"
 #include "Templates/EnableIf.h"
+#include "RHICoreShader.h"
 
 // @todo platplug: Replace all of these includes with a call to COMPILED_PLATFORM_HEADER(OpenGLDrvPrivate.h)
 //TODO: Move these to OpenGLDrvPrivate.h
@@ -980,11 +981,11 @@ public:
 
 		if (bApplyAdditionalState)
 		{
-			ApplyGlobalUniformBuffers(PsoInit.BoundShaderState.VertexShaderRHI, ResourceCast(PsoInit.BoundShaderState.VertexShaderRHI));
-			ApplyGlobalUniformBuffers(PsoInit.BoundShaderState.HullShaderRHI, ResourceCast(PsoInit.BoundShaderState.HullShaderRHI));
-			ApplyGlobalUniformBuffers(PsoInit.BoundShaderState.DomainShaderRHI, ResourceCast(PsoInit.BoundShaderState.DomainShaderRHI));
-			ApplyGlobalUniformBuffers(PsoInit.BoundShaderState.GeometryShaderRHI, ResourceCast(PsoInit.BoundShaderState.GeometryShaderRHI));
-			ApplyGlobalUniformBuffers(PsoInit.BoundShaderState.PixelShaderRHI, ResourceCast(PsoInit.BoundShaderState.PixelShaderRHI));
+			ApplyStaticUniformBuffers(PsoInit.BoundShaderState.VertexShaderRHI, ResourceCast(PsoInit.BoundShaderState.VertexShaderRHI));
+			ApplyStaticUniformBuffers(PsoInit.BoundShaderState.HullShaderRHI, ResourceCast(PsoInit.BoundShaderState.HullShaderRHI));
+			ApplyStaticUniformBuffers(PsoInit.BoundShaderState.DomainShaderRHI, ResourceCast(PsoInit.BoundShaderState.DomainShaderRHI));
+			ApplyStaticUniformBuffers(PsoInit.BoundShaderState.GeometryShaderRHI, ResourceCast(PsoInit.BoundShaderState.GeometryShaderRHI));
+			ApplyStaticUniformBuffers(PsoInit.BoundShaderState.PixelShaderRHI, ResourceCast(PsoInit.BoundShaderState.PixelShaderRHI));
 		}
 
 		// Store the PSO's primitive (after since IRHICommandContext::RHISetGraphicsPipelineState sets the BSS)
@@ -1074,12 +1075,12 @@ private:
 	int32 BeginSceneContextType;
 
 	template <typename TRHIShader, typename TRHIProxyShader>
-	void ApplyGlobalUniformBuffers(TRHIShader* Shader, TRHIProxyShader* ProxyShader)
+	void ApplyStaticUniformBuffers(TRHIShader* Shader, TRHIProxyShader* ProxyShader)
 	{
 		if (ProxyShader)
 		{
 			check(Shader);
-			::ApplyGlobalUniformBuffers(this, Shader, ProxyShader->StaticSlots, ProxyShader->Bindings.ShaderResourceTable.ResourceTableLayoutHashes, GlobalUniformBuffers);
+			UE::RHICore::ApplyStaticUniformBuffers(this, Shader, ProxyShader->StaticSlots, ProxyShader->Bindings.ShaderResourceTable.ResourceTableLayoutHashes, GlobalUniformBuffers);
 		}
 	}
 
diff --git a/Engine/Source/Runtime/RHICore/Public/RHICoreShader.h b/Engine/Source/Runtime/RHICore/Public/RHICoreShader.h
--- a/Engine/Source/Runtime/RHICore/Public/RHICoreShader.h
+++ b/Engine/Source/Runtime/RHICore/Public/RHICoreShader.h
@@ -31,7 +31,7 @@ void ApplyStaticUniformBuffers(
 		if (IsUniformBufferStaticSlotValid(Slot))
 		{
 			FRHIUniformBuffer* Buffer = UniformBuffers[Slot];
-			RHICore::ValidateStaticUniformBuffer(Buffer, Slot, LayoutHashes[BufferIndex]);
+			ValidateStaticUniformBuffer(Buffer, Slot, LayoutHashes[BufferIndex]);
 
 			if (Buffer)
 			{
diff --git a/Engine/Source/Runtime/RenderCore/Private/ShaderCore.cpp b/Engine/Source/Runtime/RenderCore/Private/ShaderCore.cpp
--- a/Engine/Source/Runtime/RenderCore/Private/ShaderCore.cpp
+++ b/Engine/Source/Runtime/RenderCore/Private/ShaderCore.cpp
@@ -26,8 +26,6 @@
 #include "Misc/StringBuilder.h"
 #endif
 
-#define VALIDATE_GLOBAL_UNIFORM_BUFFERS (!UE_BUILD_SHIPPING && !UE_BUILD_TEST)
-
 static TAutoConsoleVariable<int32> CVarShaderDevelopmentMode(
 	TEXT("r.ShaderDevelopmentMode"),
 	0,
@@ -367,59 +365,6 @@ bool ShouldExportShaderDebugInfo(EShaderPlatform Platform)
 	return (GExportDebugInfoPlatforms & (uint64(1) << Platform)) != 0;
 }
 
-void ValidateStaticUniformBuffer(FRHIUniformBuffer* UniformBuffer, FUniformBufferStaticSlot Slot, uint32 ExpectedHash)
-{
-#if VALIDATE_GLOBAL_UNIFORM_BUFFERS
-	FUniformBufferStaticSlotRegistry& SlotRegistry = FUniformBufferStaticSlotRegistry::Get();
-
-	if (!UniformBuffer)
-	{
-		const FShaderParametersMetadata* ExpectedStructMetadata = FindUniformBufferStructByLayoutHash(ExpectedHash);
-
-		checkf(
-			ExpectedStructMetadata,
-			TEXT("Shader is requesting a uniform buffer at slot %s with hash '%u', but a reverse lookup of the hash can't find it. The shader cache may be out of date."),
-			*SlotRegistry.GetDebugDescription(Slot), ExpectedHash);
-
-		const EUniformBufferBindingFlags BindingFlags = ExpectedStructMetadata->GetBindingFlags();
-
-		checkf(EnumHasAnyFlags(BindingFlags, EUniformBufferBindingFlags::Static),
-			TEXT("Shader requested a global uniform buffer of type '%s' at static slot '%s', but it is not registered with the Global binding flag. The shader cache may be out of date."),
-			ExpectedStructMetadata->GetShaderVariableName(), *SlotRegistry.GetDebugDescription(Slot));
-
-		// Structs can be bound both globally or per-shader, effectively leaving it up to the user to choose which to bind.
-		// But that also means we can't validate existence at the global level.
-		if (!EnumHasAnyFlags(BindingFlags, EUniformBufferBindingFlags::Shader))
-		{
-			UE_LOG(LogShaders, Fatal,
-				TEXT("Shader requested a global uniform buffer of type '%s' at static slot '%s', but it was null. The uniform buffer should ")
-				TEXT("be bound using RHICmdList.SetGlobalUniformBuffers() or passed into an RDG pass using SHADER_PARAMETER_STRUCT_REF() or ")
-				TEXT("SHADER_PARAMETER_RDG_UNIFORM_BUFFER()."),
-				ExpectedStructMetadata->GetShaderVariableName(), *SlotRegistry.GetDebugDescription(Slot));
-		}
-	}
-	else
-	{
-		const FRHIUniformBufferLayout& Layout = UniformBuffer->GetLayout();
-
-		if (Layout.GetHash() != ExpectedHash)
-		{
-			const FShaderParametersMetadata* ExpectedStructMetadata = FindUniformBufferStructByLayoutHash(ExpectedHash);
-
-			checkf(
-				ExpectedStructMetadata,
-				TEXT("Shader is requesting uniform buffer '%s' at slot %s with hash '%u', but a reverse lookup of the hash can't find it. The shader cache may be out of date."),
-				*Layout.GetDebugName(), *SlotRegistry.GetDebugDescription(Slot), ExpectedHash);
-
-			checkf(
-				false,
-				TEXT("Shader attempted to bind uniform buffer '%s' at slot %s with hash '%u', but the shader expected '%s' with hash '%u'."),
-				*Layout.GetDebugName(), *SlotRegistry.GetDebugDescription(Slot), ExpectedHash, ExpectedStructMetadata->GetShaderVariableName(), Layout.GetHash());
-		}
-	}
-#endif
-}
-
 bool FShaderParameterMap::FindParameterAllocation(const TCHAR* ParameterName,uint16& OutBufferIndex,uint16& OutBaseIndex,uint16& OutSize) const
 {
 	const FParameterAllocation* Allocation = ParameterMap.Find(ParameterName);
diff --git a/Engine/Source/Runtime/RenderCore/Public/ShaderCore.h b/Engine/Source/Runtime/RenderCore/Public/ShaderCore.h
--- a/Engine/Source/Runtime/RenderCore/Public/ShaderCore.h
+++ b/Engine/Source/Runtime/RenderCore/Public/ShaderCore.h
@@ -909,35 +909,4 @@ extern RENDERCORE_API void ResetAllShaderSourceDirectoryMappings();
 extern RENDERCORE_API void AddShaderSourceDirectoryMapping(const FString& VirtualShaderDirectory, const FString& RealShaderDirectory);
 
 extern RENDERCORE_API void AddShaderSourceFileEntry(TArray<FString>& OutVirtualFilePaths, FString VirtualFilePath, EShaderPlatform ShaderPlatform);
-extern RENDERCORE_API void GetAllVirtualShaderSourcePaths(TArray<FString>& OutVirtualFilePaths, EShaderPlatform ShaderPlatform);
-
-/** Validates that the uniform buffer at the requested static slot. */
-extern RENDERCORE_API void ValidateStaticUniformBuffer(FRHIUniformBuffer* UniformBuffer, FUniformBufferStaticSlot Slot, uint32 ExpectedHash);
-
-template <typename TRHIContext, typename TRHIShader>
-void ApplyGlobalUniformBuffers(
-	TRHIContext* CommandContext,
-	TRHIShader* Shader,
-	const TArray<FUniformBufferStaticSlot>& Slots,
-	const TArray<uint32>& LayoutHashes,
-	const TArray<FRHIUniformBuffer*>& UniformBuffers)
-{
-	checkf(LayoutHashes.Num() == Slots.Num(), TEXT("Shader %s, LayoutHashes %d, Slots %d"),
-		Shader->GetShaderName(), LayoutHashes.Num(), Slots.Num());
-
-	for (int32 BufferIndex = 0; BufferIndex < Slots.Num(); ++BufferIndex)
-	{
-		const FUniformBufferStaticSlot Slot = Slots[BufferIndex];
-
-		if (IsUniformBufferStaticSlotValid(Slot))
-		{
-			FRHIUniformBuffer* Buffer = UniformBuffers[Slot];
-			ValidateStaticUniformBuffer(Buffer, Slot, LayoutHashes[BufferIndex]);
-
-			if (Buffer)
-			{
-				CommandContext->RHISetShaderUniformBuffer(Shader, BufferIndex, Buffer);
-			}
-		}
-	}
-}
+extern RENDERCORE_API void GetAllVirtualShaderSourcePaths(TArray<FString>& OutVirtualFilePaths, EShaderPlatform ShaderPlatform);
\ No newline at end of file
diff --git a/Engine/Source/Runtime/VulkanRHI/Private/VulkanCommands.cpp b/Engine/Source/Runtime/VulkanRHI/Private/VulkanCommands.cpp
--- a/Engine/Source/Runtime/VulkanRHI/Private/VulkanCommands.cpp
+++ b/Engine/Source/Runtime/VulkanRHI/Private/VulkanCommands.cpp
@@ -222,7 +222,7 @@ void FVulkanCommandListContext::RHISetComputePipelineState(FRHIComputePipelineSt
 	FVulkanComputePipeline* ComputePipeline = ResourceCast(ComputePipelineState);
 	PendingComputeState->SetComputePipeline(ComputePipeline);
 
-	ApplyGlobalUniformBuffers(const_cast<FVulkanComputeShader*>(ComputePipeline->GetShader()));
+	ApplyStaticUniformBuffers(const_cast<FVulkanComputeShader*>(ComputePipeline->GetShader()));
 }
 
 void FVulkanCommandListContext::RHIDispatchComputeShader(uint32 ThreadGroupCountX, uint32 ThreadGroupCountY, uint32 ThreadGroupCountZ)
diff --git a/Engine/Source/Runtime/VulkanRHI/Private/VulkanContext.h b/Engine/Source/Runtime/VulkanRHI/Private/VulkanContext.h
--- a/Engine/Source/Runtime/VulkanRHI/Private/VulkanContext.h
+++ b/Engine/Source/Runtime/VulkanRHI/Private/VulkanContext.h
@@ -9,6 +9,7 @@
 #include "VulkanResources.h"
 #include "VulkanGPUProfiler.h"
 #include "VulkanBarriers.h"
+#include "RHICoreShader.h"
 
 class FVulkanDevice;
 class FVulkanCommandBufferManager;
@@ -247,7 +248,7 @@ private:
 	FVulkanGPUTiming* FrameTiming;
 
 	template <typename TRHIShader>
-	void ApplyGlobalUniformBuffers(TRHIShader* Shader)
+	void ApplyStaticUniformBuffers(TRHIShader* Shader)
 	{
 		if (Shader)
 		{
@@ -261,7 +262,7 @@ private:
 				if (IsUniformBufferStaticSlotValid(Slot))
 				{
 					FRHIUniformBuffer* Buffer = GlobalUniformBuffers[Slot];
-					ValidateStaticUniformBuffer(Buffer, Slot, UBInfos[BufferIndex].LayoutHash);
+					UE::RHICore::ValidateStaticUniformBuffer(Buffer, Slot, UBInfos[BufferIndex].LayoutHash);
 
 					if (Buffer)
 					{
diff --git a/Engine/Source/Runtime/VulkanRHI/Private/VulkanPipelineState.cpp b/Engine/Source/Runtime/VulkanRHI/Private/VulkanPipelineState.cpp
--- a/Engine/Source/Runtime/VulkanRHI/Private/VulkanPipelineState.cpp
+++ b/Engine/Source/Runtime/VulkanRHI/Private/VulkanPipelineState.cpp
@@ -371,15 +371,15 @@ void FVulkanCommandListContext::RHISetGraphicsPipelineState(FRHIGraphicsPipeline
 
 	if (bApplyAdditionalState)
 	{
-		ApplyGlobalUniformBuffers(static_cast<FVulkanVertexShader*>(Pipeline->VulkanShaders[ShaderStage::Vertex]));
+		ApplyStaticUniformBuffers(static_cast<FVulkanVertexShader*>(Pipeline->VulkanShaders[ShaderStage::Vertex]));
 #if PLATFORM_SUPPORTS_TESSELLATION_SHADERS
-		ApplyGlobalUniformBuffers(static_cast<FVulkanHullShader*>(Pipeline->VulkanShaders[ShaderStage::Hull]));
-		ApplyGlobalUniformBuffers(static_cast<FVulkanDomainShader*>(Pipeline->VulkanShaders[ShaderStage::Domain]));
+		ApplyStaticUniformBuffers(static_cast<FVulkanHullShader*>(Pipeline->VulkanShaders[ShaderStage::Hull]));
+		ApplyStaticUniformBuffers(static_cast<FVulkanDomainShader*>(Pipeline->VulkanShaders[ShaderStage::Domain]));
 #endif
 #if PLATFORM_SUPPORTS_GEOMETRY_SHADERS
-		ApplyGlobalUniformBuffers(static_cast<FVulkanGeometryShader*>(Pipeline->VulkanShaders[ShaderStage::Geometry]));
+		ApplyStaticUniformBuffers(static_cast<FVulkanGeometryShader*>(Pipeline->VulkanShaders[ShaderStage::Geometry]));
 #endif
-		ApplyGlobalUniformBuffers(static_cast<FVulkanPixelShader*>(Pipeline->VulkanShaders[ShaderStage::Pixel]));
+		ApplyStaticUniformBuffers(static_cast<FVulkanPixelShader*>(Pipeline->VulkanShaders[ShaderStage::Pixel]));
 	}
 }
 
diff --git a/Engine/Source/Runtime/VulkanRHI/VulkanRHI.Build.cs b/Engine/Source/Runtime/VulkanRHI/VulkanRHI.Build.cs
--- a/Engine/Source/Runtime/VulkanRHI/VulkanRHI.Build.cs
+++ b/Engine/Source/Runtime/VulkanRHI/VulkanRHI.Build.cs
@@ -36,7 +36,8 @@ public class VulkanRHI : ModuleRules
 				"CoreUObject",
 				"ApplicationCore",
 				"Engine", 
-				"RHI", 
+				"RHI",
+				"RHICore",
 				"RenderCore", 
 				"HeadMountedDisplay",
                 "PreLoadScreen",
