Fixed RDG 'Clear' stat timings.

#rnx

diff --git a/Engine/Source/Runtime/RenderCore/Private/RenderGraphAllocator.cpp b/Engine/Source/Runtime/RenderCore/Private/RenderGraphAllocator.cpp
--- a/Engine/Source/Runtime/RenderCore/Private/RenderGraphAllocator.cpp
+++ b/Engine/Source/Runtime/RenderCore/Private/RenderGraphAllocator.cpp
@@ -1,6 +1,7 @@
 // Copyright Epic Games, Inc. All Rights Reserved.
 
 #include "RenderGraphAllocator.h"
+#include "RenderGraphPrivate.h"
 
 FRDGAllocator& FRDGAllocator::Get()
 {
@@ -15,6 +16,8 @@ FRDGAllocator::~FRDGAllocator()
 
 void FRDGAllocator::ReleaseAll()
 {
+	SCOPE_CYCLE_COUNTER(STAT_RDG_ClearTime);
+	CSV_SCOPED_TIMING_STAT_EXCLUSIVE_CONDITIONAL(RDGAllocator_Clear, GRDGVerboseCSVStats != 0);
 	for (int32 Index = TrackedAllocs.Num() - 1; Index >= 0; --Index)
 	{
 		TrackedAllocs[Index]->~FTrackedAlloc();
diff --git a/Engine/Source/Runtime/RenderCore/Private/RenderGraphBuilder.cpp b/Engine/Source/Runtime/RenderCore/Private/RenderGraphBuilder.cpp
--- a/Engine/Source/Runtime/RenderCore/Private/RenderGraphBuilder.cpp
+++ b/Engine/Source/Runtime/RenderCore/Private/RenderGraphBuilder.cpp
@@ -1385,8 +1385,6 @@ void FRDGBuilder::Execute()
 
 void FRDGBuilder::Clear()
 {
-	SCOPE_CYCLE_COUNTER(STAT_RDG_ClearTime);
-	CSV_SCOPED_TIMING_STAT_EXCLUSIVE_CONDITIONAL(RDGBuilder_Clear, GRDGVerboseCSVStats != 0);
 	ExternalTextures.Empty();
 	ExternalBuffers.Empty();
 	ExtractedTextures.Empty();
