Fix typo preventing the creation of Texture arrays via RDG

Fix incorrect assert preventing the use of Texture 2D Arrays with UAVs

diff --git a/Engine/Source/Runtime/RHI/Public/RHIUtilities.h b/Engine/Source/Runtime/RHI/Public/RHIUtilities.h
--- a/Engine/Source/Runtime/RHI/Public/RHIUtilities.h
+++ b/Engine/Source/Runtime/RHI/Public/RHIUtilities.h
@@ -729,7 +729,7 @@ inline void RHICreateTargetableShaderResource2DArray(
 	check(!(bForceSeparateTargetAndShaderResource && bForceSharedTargetAndShaderResource));
 
 	// Ensure that the targetable texture is either render or depth-stencil targetable.
-	check(TargetableTextureFlags & (TexCreate_RenderTargetable | TexCreate_DepthStencilTargetable));
+	check(TargetableTextureFlags & (TexCreate_RenderTargetable | TexCreate_DepthStencilTargetable | TexCreate_UAV));
 
 	if (NumSamples > 1 && !bForceSharedTargetAndShaderResource)
 	{
diff --git a/Engine/Source/Runtime/RenderCore/Public/RenderGraphResources.h b/Engine/Source/Runtime/RenderCore/Public/RenderGraphResources.h
--- a/Engine/Source/Runtime/RenderCore/Public/RenderGraphResources.h
+++ b/Engine/Source/Runtime/RenderCore/Public/RenderGraphResources.h
@@ -330,7 +330,7 @@ struct RENDERCORE_API FRDGTextureDesc
 		uint8 InNumMips = 1,
 		uint8 InNumSamples = 1)
 	{
-		return FRDGTextureDesc(InClearValue, ETextureDimension::Texture2D, InFlags, InFormat, InExtent, 1, InArraySize, InNumMips, InNumSamples);
+		return FRDGTextureDesc(InClearValue, ETextureDimension::Texture2DArray, InFlags, InFormat, InExtent, 1, InArraySize, InNumMips, InNumSamples);
 	}
 
 	static FRDGTextureDesc Create3D(
@@ -364,7 +364,7 @@ struct RENDERCORE_API FRDGTextureDesc
 		uint8 InNumMips = 1,
 		uint8 InNumSamples = 1)
 	{
-		return FRDGTextureDesc(InClearValue, ETextureDimension::TextureCube, InFlags, InFormat, FIntPoint(InSizeInPixels, InSizeInPixels), 1, InArraySize, InNumMips, InNumSamples);
+		return FRDGTextureDesc(InClearValue, ETextureDimension::TextureCubeArray, InFlags, InFormat, FIntPoint(InSizeInPixels, InSizeInPixels), 1, InArraySize, InNumMips, InNumSamples);
 	}
 
 	FRDGTextureDesc() = default;
