Miscellaneous render graph changes in preparation for the scene textures refactor.
 - Added GetLoadActionIfProduced helper function.
 - Made GetPooledFreeBuffer attempt to reuse the existing allocation.
 - Disable load action validation when SkipRenderPass is specified.

#rnx

diff --git a/Engine/Source/Runtime/RenderCore/Private/RenderGraphUtils.cpp b/Engine/Source/Runtime/RenderCore/Private/RenderGraphUtils.cpp
--- a/Engine/Source/Runtime/RenderCore/Private/RenderGraphUtils.cpp
+++ b/Engine/Source/Runtime/RenderCore/Private/RenderGraphUtils.cpp
@@ -916,11 +916,19 @@ FRDGWaitForTasksScope::~FRDGWaitForTasksScope()
 	}
 }
 
-void GetPooledFreeBuffer(
+bool GetPooledFreeBuffer(
 	FRHICommandList& RHICmdList,
 	const FRDGBufferDesc& Desc,
 	TRefCountPtr<FRDGPooledBuffer>& Out,
 	const TCHAR* InDebugName)
 {
+	if (Out && Out->Desc == Desc)
+	{
+		// Kept current allocation.
+		return false;
+	}
+
+	// New allocation.
 	Out = GRenderGraphResourcePool.FindFreeBuffer(RHICmdList, Desc, InDebugName);
+	return true;
 }
\ No newline at end of file
diff --git a/Engine/Source/Runtime/RenderCore/Private/RenderGraphValidation.cpp b/Engine/Source/Runtime/RenderCore/Private/RenderGraphValidation.cpp
--- a/Engine/Source/Runtime/RenderCore/Private/RenderGraphValidation.cpp
+++ b/Engine/Source/Runtime/RenderCore/Private/RenderGraphValidation.cpp
@@ -511,7 +511,6 @@ void FRDGUserValidation::ValidateAddPass(const FRDGPass* Pass, bool bSkipPassAcc
 	const ERDGPassFlags PassFlags = Pass->GetFlags();
 	const FRDGParameterStruct PassParameters = Pass->GetParameters();
 
-
 	const TCHAR* PassName = Pass->GetName();
 	const bool bIsRaster = EnumHasAnyFlags(PassFlags, ERDGPassFlags::Raster);
 	const bool bIsCopy = EnumHasAnyFlags(PassFlags, ERDGPassFlags::Copy);
@@ -785,6 +784,7 @@ void FRDGUserValidation::ValidateAddPass(const FRDGPass* Pass, bool bSkipPassAcc
 					 *  target pool. Instead, the user should create a new texture instance. An exception to this rule are untracked render targets,
 					 *  which are not actually managed by the render target pool and likely represent the frame buffer.
 					 */
+					if (!bSkipRenderPass)
 					{
 						const bool bIsLoadAction = RenderTarget.GetLoadAction() == ERenderTargetLoadAction::ELoad;
 
diff --git a/Engine/Source/Runtime/RenderCore/Public/RenderGraphUtils.h b/Engine/Source/Runtime/RenderCore/Public/RenderGraphUtils.h
--- a/Engine/Source/Runtime/RenderCore/Public/RenderGraphUtils.h
+++ b/Engine/Source/Runtime/RenderCore/Public/RenderGraphUtils.h
@@ -26,6 +26,11 @@ inline FRDGBufferRef GetIfProduced(FRDGBufferRef Buffer, FRDGBufferRef FallbackB
 	return HasBeenProduced(Buffer) ? Buffer : FallbackBuffer;
 }
 
+inline ERenderTargetLoadAction GetLoadActionIfProduced(FRDGTextureRef Texture, ERenderTargetLoadAction ActionIfNotProduced)
+{
+	return HasBeenProduced(Texture) ? ERenderTargetLoadAction::ELoad : ActionIfNotProduced;
+}
+
 /** Fetches the RHI texture from an RDG texture or null if the RDG texture is null. */
 inline FRHITexture* TryGetRHI(FRDGTextureRef Texture)
 {
@@ -816,7 +821,8 @@ private:
 #define RDG_WAIT_FOR_TASKS_CONDITIONAL(GraphBuilder, bCondition) FRDGWaitForTasksScope PREPROCESSOR_JOIN(RDGWaitForTasksScope, __LINE__){ GraphBuilder, bCondition }
 #define RDG_WAIT_FOR_TASKS(GraphBuilder) RDG_WAIT_FOR_TASKS_CONDITIONAL(GraphBuilder, true)
 
-RENDERCORE_API void GetPooledFreeBuffer(
+// Allocates an RDG pooled buffer instance. Attempts to reuse allocation if Out has a value. Returns true a new instance was allocated, or false if the existing allocation was reused.
+RENDERCORE_API bool GetPooledFreeBuffer(
 	FRHICommandList& RHICmdList,
 	const FRDGBufferDesc& Desc,
 	TRefCountPtr<FRDGPooledBuffer>& Out,
