Added FGBufferBinding to GBuffer info to store pertinent GBuffer binding information in a much smaller footprint.

diff --git a/Engine/Source/Runtime/RenderCore/Private/GBufferInfo.cpp b/Engine/Source/Runtime/RenderCore/Private/GBufferInfo.cpp
--- a/Engine/Source/Runtime/RenderCore/Private/GBufferInfo.cpp
+++ b/Engine/Source/Runtime/RenderCore/Private/GBufferInfo.cpp
@@ -117,7 +117,7 @@ TArray < EGBufferSlot > FetchGBufferSlots(bool bHasVelocity, bool bHasTangent, b
 
 
 // finds the target by searching for the name. returns -1 if not found.
-int32 RENDERCORE_API FindGBufferTargetByName(const FGBufferInfo& GBufferInfo, const FString& Name)
+int32 FindGBufferTargetByName(const FGBufferInfo& GBufferInfo, const FString& Name)
 {
 	for (int32 I = 0; I < GBufferInfo.NumTargets; I++)
 	{
@@ -130,6 +130,48 @@ int32 RENDERCORE_API FindGBufferTargetByName(const FGBufferInfo& GBufferInfo, co
 	return -1;
 }
 
+FGBufferBinding FindGBufferBindingByName(const FGBufferInfo& GBufferInfo, const FString& Name)
+{
+	const int32 Index = FindGBufferTargetByName(GBufferInfo, Name);
+
+	FGBufferBinding Binding;
+
+	if (Index >= 0)
+	{
+		const FGBufferTarget& Target = GBufferInfo.Targets[Index];
+
+		EPixelFormat PixelFormat = PF_Unknown;
+
+		switch (Target.TargetType)
+		{
+		case GBT_Unorm_8_8_8_8:
+			PixelFormat = PF_B8G8R8A8;
+			break;
+		case GBT_Unorm_11_11_10:
+			PixelFormat = PF_FloatR11G11B10;
+			break;
+		case GBT_Unorm_10_10_10_2:
+			PixelFormat = PF_A2B10G10R10;
+			break;
+		case GBT_Float_16_16:
+			PixelFormat = PF_G16R16;
+			break;
+		case GBT_Float_16_16_16_16:
+			PixelFormat = PF_A16B16G16R16;
+			break;
+		case GBT_Invalid:
+		default:
+			check(0);
+			break;
+		}
+
+		Binding.Index = Index;
+		Binding.Format = PixelFormat;
+		Binding.Flags = TexCreate_ShaderResource | TexCreate_RenderTargetable | (Target.bIsSrgb ? TexCreate_SRGB : TexCreate_None);
+	}
+
+	return Binding;
+}
 
 /*
  * 4.25 Logic:
diff --git a/Engine/Source/Runtime/RenderCore/Public/GBufferInfo.h b/Engine/Source/Runtime/RenderCore/Public/GBufferInfo.h
--- a/Engine/Source/Runtime/RenderCore/Public/GBufferInfo.h
+++ b/Engine/Source/Runtime/RenderCore/Public/GBufferInfo.h
@@ -3,6 +3,8 @@
 #pragma once
 
 #include "CoreMinimal.h"
+#include "PixelFormat.h"
+#include "RHIDefinitions.h"
 
 enum EGBufferSlot
 {
@@ -236,7 +238,6 @@ struct FGBufferTarget
 		bIsUsingExtraFlags = bInIsUsingExtraFlags;
 	}
 
-
 	EGBufferType TargetType;
 	FString TargetName;
 	bool bIsSrgb;
@@ -245,6 +246,13 @@ struct FGBufferTarget
 	bool bIsUsingExtraFlags;
 };
 
+struct FGBufferBinding
+{
+	int32 Index = -1;
+	EPixelFormat Format = PF_Unknown;
+	ETextureCreateFlags Flags = TexCreate_None;
+};
+
 struct FGBufferInfo
 {
 	static const int MaxTargets = 8;
@@ -258,16 +266,33 @@ struct FGBufferInfo
 
 struct FGBufferParams
 {
-	bool bHasVelocity;
-	bool bHasTangent;
-	bool bHasPrecShadowFactor;
-	int LegacyFormatIndex;
-	bool bUsesVelocityDepth;
+	int32 LegacyFormatIndex = 0;
+	bool bHasVelocity = false;
+	bool bHasTangent = false;
+	bool bHasPrecShadowFactor = false;
+	bool bUsesVelocityDepth = false;
+
+	bool operator == (const FGBufferParams& RHS) const
+	{
+		return
+			LegacyFormatIndex == RHS.LegacyFormatIndex &&
+			bHasVelocity == RHS.bHasVelocity &&
+			bHasTangent == RHS.bHasTangent &&
+			bHasPrecShadowFactor == RHS.bHasPrecShadowFactor &&
+			bUsesVelocityDepth == RHS.bUsesVelocityDepth;
+	}
+
+	bool operator != (const FGBufferParams& RHS) const
+	{
+		return !(*this == RHS);
+	}
 };
 
 
 int32 RENDERCORE_API FindGBufferTargetByName(const FGBufferInfo& GBufferInfo, const FString& Name);
 
+FGBufferBinding RENDERCORE_API FindGBufferBindingByName(const FGBufferInfo& GBufferInfo, const FString& Name);
+
 //FGBufferInfo RENDERCORE_API FetchLegacyGBufferInfo(bool bHasVelocity, bool bHasTangent, bool bHasPrecShadowFactor, int LegacyFormatIndex, bool bUsesVelocityDepth);
 //FGBufferInfo RENDERCORE_API FetchFullGBufferInfo(bool bHasVelocity, bool bHasTangent, bool bHasPrecShadowFactor, int LegacyFormatIndex, bool bUsesVelocityDepth);
 FGBufferInfo RENDERCORE_API FetchLegacyGBufferInfo(const FGBufferParams& Params);
