Changing some shader compilers to use common flags for exporting debug information:
- r.Shaders.KeepDebugInfo: if we should compile with debugging information included and export that data to files for platforms that support it.
- r.Shaders.PrepareExportedDebugInfo: precompile debug info for platforms that export it to files.
- r.Shaders.AllowUniqueDebugInfo: used to disallow debug info sharing, will increase debug export sizes.

diff --git a/Engine/Config/BaseEngine.ini b/Engine/Config/BaseEngine.ini
--- a/Engine/Config/BaseEngine.ini
+++ b/Engine/Config/BaseEngine.ini
@@ -2301,6 +2301,9 @@ DefaultASTCQualityBySize=3
 +VersionedIntRValues=r.SelectiveBasePassOutputs
 +VersionedIntRValues=r.DBuffer
 +VersionedIntRValues=r.Shaders.KeepDebugInfo
++VersionedIntRValues=r.Shaders.ExportDebugInfo
++VersionedIntRValues=r.Shaders.AllowUniqueDebugInfo
++VersionedIntRValues=r.Shaders.PrepareExportedDebugInfo
 +VersionedIntRValues=r.Shaders.Optimize
 +VersionedIntRValues=r.CompileShadersForDevelopment
 +VersionedIntRValues=r.MobileHDR
diff --git a/Engine/Source/Runtime/Engine/Private/ShaderCompiler/ShaderCompiler.cpp b/Engine/Source/Runtime/Engine/Private/ShaderCompiler/ShaderCompiler.cpp
--- a/Engine/Source/Runtime/Engine/Private/ShaderCompiler/ShaderCompiler.cpp
+++ b/Engine/Source/Runtime/Engine/Private/ShaderCompiler/ShaderCompiler.cpp
@@ -492,6 +492,14 @@ static TAutoConsoleVariable<int32> CVarKeepShaderDebugData(
 	TEXT("For some platforms this cvar can be overriden in the Engine.ini, under the [ShaderCompiler] section."),
 	ECVF_ReadOnly);
 
+static TAutoConsoleVariable<int32> CVarPrepareExportedDebugInfo(
+	TEXT("r.Shaders.PrepareExportedDebugInfo"),
+	0,
+	TEXT("For platforms that work on exported shader debug info, this should be set to keep shader reflection and debug data from shader bytecode, default is to strip.")
+	TEXT("For these platforms, r.Shaders.KeepDebugInfo has the same behavior but will also export the debug info.")
+	TEXT("This will only be read from the platform Engine.ini, under the [ShaderCompiler] section."),
+	ECVF_ReadOnly);
+
 static TAutoConsoleVariable<int32> CVarExportShaderDebugData(
 	TEXT("r.Shaders.ExportDebugInfo"),
 	0,
@@ -500,6 +508,14 @@ static TAutoConsoleVariable<int32> CVarExportShaderDebugData(
 	TEXT("For some platforms this cvar can be overriden in the Engine.ini, under the [ShaderCompiler] section."),
 	ECVF_ReadOnly);
 
+static TAutoConsoleVariable<int32> CVarAllowUniqueDebugInfo(
+	TEXT("r.Shaders.AllowUniqueDebugInfo"),
+	0,
+	TEXT("Whether to allow (exported) shader debug info to be unique per shader.\n")
+	TEXT("Enabling this can cause a drastic increase in debug data, enable only if absolutely necessary.\n")
+	TEXT("For some platforms this cvar can be overriden in the Engine.ini, under the [ShaderCompiler] section."),
+	ECVF_ReadOnly);
+
 static TAutoConsoleVariable<int32> CVarExportShaderDebugDataMode(
 	TEXT("r.Shaders.ExportDebugInfoMode"),
 	0,
@@ -4149,11 +4165,14 @@ void GlobalBeginCompileShader(
 		}
 	}
 
+	if (ShouldKeepShaderDebugInfo(Target.GetPlatform()))
 	{
-		if (ShouldKeepShaderDebugInfo((EShaderPlatform)Target.Platform))
-		{
-			Input.Environment.CompilerFlags.Add(CFLAG_KeepDebugInfo);
-		}
+		Input.Environment.CompilerFlags.Add(CFLAG_KeepDebugInfo);
+	}
+
+	if (ShouldAllowUniqueDebugInfo(Target.GetPlatform()))
+	{
+		Input.Environment.CompilerFlags.Add(CFLAG_AllowUniqueDebugInfo);
 	}
 
 	if (CVarShaderFastMath.GetValueOnAnyThread() == 0)
diff --git a/Engine/Source/Runtime/RenderCore/Private/Shader.cpp b/Engine/Source/Runtime/RenderCore/Private/Shader.cpp
--- a/Engine/Source/Runtime/RenderCore/Private/Shader.cpp
+++ b/Engine/Source/Runtime/RenderCore/Private/Shader.cpp
@@ -1413,6 +1413,7 @@ void ShaderMapAppendKeyString(EShaderPlatform Platform, FString& KeyString)
 
 	{
 		KeyString += ShouldKeepShaderDebugInfo(Platform) ? TEXT("_NoStrip") : TEXT("");
+		KeyString += ShouldAllowUniqueDebugInfo(Platform) ? TEXT("_FullDbg") : TEXT("");
 	}
 
 	{
diff --git a/Engine/Source/Runtime/RenderCore/Private/ShaderCore.cpp b/Engine/Source/Runtime/RenderCore/Private/ShaderCore.cpp
--- a/Engine/Source/Runtime/RenderCore/Private/ShaderCore.cpp
+++ b/Engine/Source/Runtime/RenderCore/Private/ShaderCore.cpp
@@ -280,89 +280,73 @@ bool AllowDebugViewmodes(EShaderPlatform Platform)
 #endif
 }
 
-#if WITH_EDITOR
-static void GetShaderCompilerPlatformConfigs(const TCHAR* Key, uint64& OutPlatformFlags)
+static bool IsShaderCompilerConfigEnabledForPlatform(const IConsoleVariable* GlobalCvar, const TCHAR* Key, EShaderPlatform ShaderPlatform)
 {
-	for (uint32 ShaderPlatformIndex = 0; ShaderPlatformIndex < SP_NumPlatforms; ++ShaderPlatformIndex)
+	bool bEnabled = false;
+
+	// First check the global cvar
+	if (GlobalCvar && GlobalCvar->GetInt())
+	{
+		bEnabled = true;
+	}
+#if WITH_EDITOR
+	else
 	{
-		EShaderPlatform ShaderPlatform = EShaderPlatform(ShaderPlatformIndex);
-		FName PlatformName = ShaderPlatformToPlatformName(ShaderPlatform);
+		// Then check the per platform settings.
+		const FName PlatformName = ShaderPlatformToPlatformName(ShaderPlatform);
 		if (!PlatformName.IsNone())
 		{
 			FConfigFile EngineSettings;
 			FConfigCacheIni::LoadLocalIniFile(EngineSettings, TEXT("Engine"), true, *PlatformName.ToString());
 
-			bool bEnabled = false;
-			if (EngineSettings.GetBool(TEXT("ShaderCompiler"), Key, bEnabled))
-			{
-				uint64 Mask = (uint64)1 << ShaderPlatformIndex;
-				if (bEnabled)
-				{
-					OutPlatformFlags |= Mask;
-				}
-				else
-				{
-					OutPlatformFlags &= ~Mask;
-				}
-			}
+			EngineSettings.GetBool(TEXT("ShaderCompiler"), Key, bEnabled);
 		}
 	}
-}
 #endif
 
-static uint64 GetKeepShaderDebugInfoPlatforms()
-{
-	uint64 KeepDebugInfoPlatforms = 0;
-
-	// First check the global cvars
-	static IConsoleVariable* CVarKeepDebugInfo = IConsoleManager::Get().FindConsoleVariable(TEXT("r.Shaders.KeepDebugInfo"));
-	if (CVarKeepDebugInfo && CVarKeepDebugInfo->GetInt())
-	{
-		KeepDebugInfoPlatforms = ~(uint64)0;
-	}
-
-#if WITH_EDITOR
-	// Then load the per platform settings.
-	GetShaderCompilerPlatformConfigs(TEXT("r.Shaders.KeepDebugInfo"), KeepDebugInfoPlatforms);
-#endif
-
-	return KeepDebugInfoPlatforms;
+	return bEnabled;
 }
 
-bool ShouldKeepShaderDebugInfo(EShaderPlatform Platform)
+bool ShouldKeepShaderDebugInfo(EShaderPlatform ShaderPlatform)
 {
-	static uint64 KeepShaderDebugInfoPlatforms = GetKeepShaderDebugInfoPlatforms();
-	return (KeepShaderDebugInfoPlatforms & (uint64(1) << Platform)) != 0;
-}
+	static const TCHAR* const KeepDebugInfoName = TEXT("r.Shaders.KeepDebugInfo");
+	static const IConsoleVariable* KeepDebugInfoCvar = IConsoleManager::Get().FindConsoleVariable(KeepDebugInfoName);
 
-static uint64 GetExportShaderDebugInfoPlatforms()
-{
-	uint64 ExportDebugInfoPlatforms = 0;
+	// PrepareExportedDebugInfo is only read from the platform settings, not the global cvar
+	static const TCHAR* const PrepareExportedDebugInfoName = TEXT("r.Shaders.PrepareExportedDebugInfo");
 
-	// First check the global cvars
+	return IsShaderCompilerConfigEnabledForPlatform(KeepDebugInfoCvar, KeepDebugInfoName, ShaderPlatform)
+		|| IsShaderCompilerConfigEnabledForPlatform(nullptr, PrepareExportedDebugInfoName, ShaderPlatform);
+}
 
+bool ShouldExportShaderDebugInfo(EShaderPlatform ShaderPlatform)
+{
 	// r.DumpShaderDebugInfo should also turn on ExportShaderDebugInfo
 	// The difference is that r.DumpShaderDebugInfo will also output engine debug files such as converted hlsl or SCW helper files.
 	// Where as r.Shader.ExportDebugInfo is purely to export the graphics debugging tool's debug info files.
-	static IConsoleVariable* CVarExportDebugInfo = IConsoleManager::Get().FindConsoleVariable(TEXT("r.Shaders.ExportDebugInfo"));
 	static IConsoleVariable* CVarDumpDebugInfo = IConsoleManager::Get().FindConsoleVariable(TEXT("r.DumpShaderDebugInfo"));
-	if ((CVarExportDebugInfo && CVarExportDebugInfo->GetInt()) || (CVarDumpDebugInfo && CVarDumpDebugInfo->GetInt()))
+	if (CVarDumpDebugInfo && CVarDumpDebugInfo->GetInt())
 	{
-		ExportDebugInfoPlatforms = ~(uint64)0;
+		return true;
 	}
 
-#if WITH_EDITOR
-	// Then load the per platform settings.
-	GetShaderCompilerPlatformConfigs(TEXT("r.Shaders.ExportDebugInfo"), ExportDebugInfoPlatforms);
-#endif
+	static const TCHAR* const KeepDebugInfoName = TEXT("r.Shaders.KeepDebugInfo");
+	static const IConsoleVariable* KeepDebugInfoCvar = IConsoleManager::Get().FindConsoleVariable(KeepDebugInfoName);
+	if (IsShaderCompilerConfigEnabledForPlatform(KeepDebugInfoCvar, KeepDebugInfoName, ShaderPlatform))
+	{
+		return true;
+	}
 
-	return ExportDebugInfoPlatforms;
+	static const TCHAR* const ExportDebugInfoName = TEXT("r.Shaders.ExportDebugInfo");
+	static const IConsoleVariable* ExportDebugInfoCvar = IConsoleManager::Get().FindConsoleVariable(ExportDebugInfoName);
+	return IsShaderCompilerConfigEnabledForPlatform(ExportDebugInfoCvar, ExportDebugInfoName, ShaderPlatform);
 }
 
-bool ShouldExportShaderDebugInfo(EShaderPlatform Platform)
+bool ShouldAllowUniqueDebugInfo(EShaderPlatform ShaderPlatform)
 {
-	static uint64 GExportDebugInfoPlatforms = GetExportShaderDebugInfoPlatforms();
-	return (GExportDebugInfoPlatforms & (uint64(1) << Platform)) != 0;
+	static const TCHAR* const VariableName = TEXT("r.Shaders.AllowUniqueDebugInfo");
+	static const IConsoleVariable* GlobalCvar = IConsoleManager::Get().FindConsoleVariable(VariableName);
+	return IsShaderCompilerConfigEnabledForPlatform(GlobalCvar, VariableName, ShaderPlatform);
 }
 
 bool FShaderParameterMap::FindParameterAllocation(const TCHAR* ParameterName,uint16& OutBufferIndex,uint16& OutBaseIndex,uint16& OutSize) const
diff --git a/Engine/Source/Runtime/RenderCore/Public/ShaderCompilerCore.h b/Engine/Source/Runtime/RenderCore/Public/ShaderCompilerCore.h
--- a/Engine/Source/Runtime/RenderCore/Public/ShaderCompilerCore.h
+++ b/Engine/Source/Runtime/RenderCore/Public/ShaderCompilerCore.h
@@ -43,6 +43,8 @@ enum ECompilerFlags
 	CFLAG_ForceOptimization,
 	// Shader should insert debug/name info at the risk of generating non-deterministic libraries
 	CFLAG_KeepDebugInfo,
+	// Allows the (external) debug info to be specific to each shader rather than trying to deduplicate.
+	CFLAG_AllowUniqueDebugInfo,
 	CFLAG_NoFastMath,
 	// Explicitly enforce zero initialization on shader platforms that may omit it.
 	CFLAG_ZeroInitialise,
diff --git a/Engine/Source/Runtime/RenderCore/Public/ShaderCore.h b/Engine/Source/Runtime/RenderCore/Public/ShaderCore.h
--- a/Engine/Source/Runtime/RenderCore/Public/ShaderCore.h
+++ b/Engine/Source/Runtime/RenderCore/Public/ShaderCore.h
@@ -105,6 +105,9 @@ extern RENDERCORE_API bool ShouldKeepShaderDebugInfo(EShaderPlatform Platform);
 /** Returns true if debug information should be exported to separate files for a given platform . */
 extern RENDERCORE_API bool ShouldExportShaderDebugInfo(EShaderPlatform Platform);
 
+/** Returns true if (external) debug information should be specific to each shader rather than be deduplicated. */
+extern RENDERCORE_API bool ShouldAllowUniqueDebugInfo(EShaderPlatform Platform);
+
 struct FShaderTarget
 {
 	// The rest of uint32 holding the bitfields can be left unitialized. Union with a uint32 serves to prevent that to be able to set the whole uint32 value
