Quick fix for excessive load times when launching a project for the first time. I'll need to investigate a better way to cache the necessary per platform configs.

diff --git a/Engine/Source/Runtime/RenderCore/Private/ShaderCore.cpp b/Engine/Source/Runtime/RenderCore/Private/ShaderCore.cpp
--- a/Engine/Source/Runtime/RenderCore/Private/ShaderCore.cpp
+++ b/Engine/Source/Runtime/RenderCore/Private/ShaderCore.cpp
@@ -289,7 +289,8 @@ static bool IsShaderCompilerConfigEnabledForPlatform(const IConsoleVariable* Glo
 	{
 		bEnabled = true;
 	}
-#if WITH_EDITOR
+	// TODO: christopher.waters - find out a better way to load these to prevent excessive ini load times on startup without cached ini files.
+#if WITH_EDITOR && 0
 	else
 	{
 		// Then check the per platform settings.
