Added StructuredBuffers to ClearUAV on Vulkan

diff --git a/Engine/Shaders/Private/ClearReplacementShaders.usf b/Engine/Shaders/Private/ClearReplacementShaders.usf
--- a/Engine/Shaders/Private/ClearReplacementShaders.usf
+++ b/Engine/Shaders/Private/ClearReplacementShaders.usf
@@ -49,6 +49,12 @@
 		#define ADDR(p) (p).xyz
 		#define COMP(a,b) (((a).x < (b).x) && ((a).y < (b).y) && ((a).z < (b).z))
 
+	#elif RESOURCE_TYPE == 4
+	
+		StructuredBuffer<VALUE_TYPE> ClearResource;
+		#define ADDR(p) (p).x
+		#define COMP(a,b) ((a).x < (b).x)
+
 	#endif
 
 	[numthreads(THREADGROUPSIZE_X, THREADGROUPSIZE_Y, THREADGROUPSIZE_Z)]
diff --git a/Engine/Source/Runtime/RenderCore/Private/ClearReplacementShaders.cpp b/Engine/Source/Runtime/RenderCore/Private/ClearReplacementShaders.cpp
--- a/Engine/Source/Runtime/RenderCore/Private/ClearReplacementShaders.cpp
+++ b/Engine/Source/Runtime/RenderCore/Private/ClearReplacementShaders.cpp
@@ -5,34 +5,42 @@
 #include "ShaderParameterUtils.h"
 #include "RendererInterface.h"
 
-IMPLEMENT_SHADER_TYPE4_WITH_TEMPLATE_PREFIX(template<>, RENDERCORE_API, FClearReplacementVS				                , SF_Vertex);
-IMPLEMENT_SHADER_TYPE4_WITH_TEMPLATE_PREFIX(template<>, RENDERCORE_API, FClearReplacementVS_Bounds                      , SF_Vertex);
-IMPLEMENT_SHADER_TYPE4_WITH_TEMPLATE_PREFIX(template<>, RENDERCORE_API, FClearReplacementVS_Depth		                , SF_Vertex);
-IMPLEMENT_SHADER_TYPE4_WITH_TEMPLATE_PREFIX(template<>, RENDERCORE_API, FClearReplacementPS                             , SF_Pixel);
-IMPLEMENT_SHADER_TYPE4_WITH_TEMPLATE_PREFIX(template<>, RENDERCORE_API, FClearReplacementPS_128                         , SF_Pixel);
-IMPLEMENT_SHADER_TYPE4_WITH_TEMPLATE_PREFIX(template<>, RENDERCORE_API, FClearReplacementPS_Zero                        , SF_Pixel);
-IMPLEMENT_SHADER_TYPE4_WITH_TEMPLATE_PREFIX(template<>, RENDERCORE_API, FClearReplacementCS_Buffer_Uint_Bounds          , SF_Compute);
-IMPLEMENT_SHADER_TYPE4_WITH_TEMPLATE_PREFIX(template<>, RENDERCORE_API, FClearReplacementCS_Buffer_Uint4_Bounds         , SF_Compute);
-IMPLEMENT_SHADER_TYPE4_WITH_TEMPLATE_PREFIX(template<>, RENDERCORE_API, FClearReplacementCS_Buffer_Uint4         		, SF_Compute);
-IMPLEMENT_SHADER_TYPE4_WITH_TEMPLATE_PREFIX(template<>, RENDERCORE_API, FClearReplacementCS_Buffer_Float_Bounds         , SF_Compute);
-IMPLEMENT_SHADER_TYPE4_WITH_TEMPLATE_PREFIX(template<>, RENDERCORE_API, FClearReplacementCS_Buffer_Float4_Bounds        , SF_Compute);
-IMPLEMENT_SHADER_TYPE4_WITH_TEMPLATE_PREFIX(template<>, RENDERCORE_API, FClearReplacementCS_Texture3D_Float4_Bounds     , SF_Compute);
-IMPLEMENT_SHADER_TYPE4_WITH_TEMPLATE_PREFIX(template<>, RENDERCORE_API, FClearReplacementCS_Texture2DArray_Float4_Bounds, SF_Compute);
-IMPLEMENT_SHADER_TYPE4_WITH_TEMPLATE_PREFIX(template<>, RENDERCORE_API, FClearReplacementCS_Texture2D_Float4_Bounds     , SF_Compute);
-IMPLEMENT_SHADER_TYPE4_WITH_TEMPLATE_PREFIX(template<>, RENDERCORE_API, FClearReplacementCS_Buffer_Uint_Zero            , SF_Compute);
-IMPLEMENT_SHADER_TYPE4_WITH_TEMPLATE_PREFIX(template<>, RENDERCORE_API, FClearReplacementCS_Texture2DArray_Uint_Zero    , SF_Compute);
-IMPLEMENT_SHADER_TYPE4_WITH_TEMPLATE_PREFIX(template<>, RENDERCORE_API, FClearReplacementCS_Buffer_Uint                 , SF_Compute);
-IMPLEMENT_SHADER_TYPE4_WITH_TEMPLATE_PREFIX(template<>, RENDERCORE_API, FClearReplacementCS_Texture2DArray_Uint         , SF_Compute);
-IMPLEMENT_SHADER_TYPE4_WITH_TEMPLATE_PREFIX(template<>, RENDERCORE_API, FClearReplacementCS_Texture3D_Float4            , SF_Compute);
-IMPLEMENT_SHADER_TYPE4_WITH_TEMPLATE_PREFIX(template<>, RENDERCORE_API, FClearReplacementCS_Texture2D_Float4            , SF_Compute);
-IMPLEMENT_SHADER_TYPE4_WITH_TEMPLATE_PREFIX(template<>, RENDERCORE_API, FClearReplacementCS_Texture2DArray_Float4       , SF_Compute);
-IMPLEMENT_SHADER_TYPE4_WITH_TEMPLATE_PREFIX(template<>, RENDERCORE_API, FClearReplacementCS_Texture3D_Uint4             , SF_Compute);
-IMPLEMENT_SHADER_TYPE4_WITH_TEMPLATE_PREFIX(template<>, RENDERCORE_API, FClearReplacementCS_Texture2D_Uint4             , SF_Compute);
-IMPLEMENT_SHADER_TYPE4_WITH_TEMPLATE_PREFIX(template<>, RENDERCORE_API, FClearReplacementCS_Texture2DArray_Uint4        , SF_Compute);
-IMPLEMENT_SHADER_TYPE4_WITH_TEMPLATE_PREFIX(template<>, RENDERCORE_API, FClearReplacementCS_Texture3D_Uint4_Bounds      , SF_Compute);
-IMPLEMENT_SHADER_TYPE4_WITH_TEMPLATE_PREFIX(template<>, RENDERCORE_API, FClearReplacementCS_Texture2D_Uint4_Bounds      , SF_Compute);
-IMPLEMENT_SHADER_TYPE4_WITH_TEMPLATE_PREFIX(template<>, RENDERCORE_API, FClearReplacementCS_Texture2DArray_Uint4_Bounds , SF_Compute);
-IMPLEMENT_SHADER_TYPE4_WITH_TEMPLATE_PREFIX(template<>, RENDERCORE_API, FClearReplacementCS_Buffer_Sint4_Bounds         , SF_Compute);
-IMPLEMENT_SHADER_TYPE4_WITH_TEMPLATE_PREFIX(template<>, RENDERCORE_API, FClearReplacementCS_Texture3D_Sint4_Bounds      , SF_Compute);
-IMPLEMENT_SHADER_TYPE4_WITH_TEMPLATE_PREFIX(template<>, RENDERCORE_API, FClearReplacementCS_Texture2D_Sint4_Bounds      , SF_Compute);
-IMPLEMENT_SHADER_TYPE4_WITH_TEMPLATE_PREFIX(template<>, RENDERCORE_API, FClearReplacementCS_Texture2DArray_Sint4_Bounds , SF_Compute);
+IMPLEMENT_SHADER_TYPE4_WITH_TEMPLATE_PREFIX(template<>, RENDERCORE_API, FClearReplacementVS									, SF_Vertex);
+IMPLEMENT_SHADER_TYPE4_WITH_TEMPLATE_PREFIX(template<>, RENDERCORE_API, FClearReplacementVS_Bounds							, SF_Vertex);
+IMPLEMENT_SHADER_TYPE4_WITH_TEMPLATE_PREFIX(template<>, RENDERCORE_API, FClearReplacementVS_Depth							, SF_Vertex);
+IMPLEMENT_SHADER_TYPE4_WITH_TEMPLATE_PREFIX(template<>, RENDERCORE_API, FClearReplacementPS									, SF_Pixel);
+IMPLEMENT_SHADER_TYPE4_WITH_TEMPLATE_PREFIX(template<>, RENDERCORE_API, FClearReplacementPS_128								, SF_Pixel);
+IMPLEMENT_SHADER_TYPE4_WITH_TEMPLATE_PREFIX(template<>, RENDERCORE_API, FClearReplacementPS_Zero							, SF_Pixel);
+IMPLEMENT_SHADER_TYPE4_WITH_TEMPLATE_PREFIX(template<>, RENDERCORE_API, FClearReplacementCS_Buffer_Uint_Bounds				, SF_Compute);
+IMPLEMENT_SHADER_TYPE4_WITH_TEMPLATE_PREFIX(template<>, RENDERCORE_API, FClearReplacementCS_Buffer_Uint4_Bounds				, SF_Compute);
+IMPLEMENT_SHADER_TYPE4_WITH_TEMPLATE_PREFIX(template<>, RENDERCORE_API, FClearReplacementCS_Buffer_Uint4					, SF_Compute);
+IMPLEMENT_SHADER_TYPE4_WITH_TEMPLATE_PREFIX(template<>, RENDERCORE_API, FClearReplacementCS_Buffer_Float_Bounds				, SF_Compute);
+IMPLEMENT_SHADER_TYPE4_WITH_TEMPLATE_PREFIX(template<>, RENDERCORE_API, FClearReplacementCS_Buffer_Float4_Bounds			, SF_Compute);
+IMPLEMENT_SHADER_TYPE4_WITH_TEMPLATE_PREFIX(template<>, RENDERCORE_API, FClearReplacementCS_StructuredBuffer_Uint_Bounds	, SF_Compute);
+IMPLEMENT_SHADER_TYPE4_WITH_TEMPLATE_PREFIX(template<>, RENDERCORE_API, FClearReplacementCS_StructuredBuffer_Uint4_Bounds	, SF_Compute);
+IMPLEMENT_SHADER_TYPE4_WITH_TEMPLATE_PREFIX(template<>, RENDERCORE_API, FClearReplacementCS_StructuredBuffer_Uint4			, SF_Compute);
+IMPLEMENT_SHADER_TYPE4_WITH_TEMPLATE_PREFIX(template<>, RENDERCORE_API, FClearReplacementCS_StructuredBuffer_Float_Bounds	, SF_Compute);
+IMPLEMENT_SHADER_TYPE4_WITH_TEMPLATE_PREFIX(template<>, RENDERCORE_API, FClearReplacementCS_StructuredBuffer_Float4_Bounds	, SF_Compute);
+IMPLEMENT_SHADER_TYPE4_WITH_TEMPLATE_PREFIX(template<>, RENDERCORE_API, FClearReplacementCS_Texture3D_Float4_Bounds			, SF_Compute);
+IMPLEMENT_SHADER_TYPE4_WITH_TEMPLATE_PREFIX(template<>, RENDERCORE_API, FClearReplacementCS_Texture2DArray_Float4_Bounds	, SF_Compute);
+IMPLEMENT_SHADER_TYPE4_WITH_TEMPLATE_PREFIX(template<>, RENDERCORE_API, FClearReplacementCS_Texture2D_Float4_Bounds			, SF_Compute);
+IMPLEMENT_SHADER_TYPE4_WITH_TEMPLATE_PREFIX(template<>, RENDERCORE_API, FClearReplacementCS_Buffer_Uint_Zero				, SF_Compute);
+IMPLEMENT_SHADER_TYPE4_WITH_TEMPLATE_PREFIX(template<>, RENDERCORE_API, FClearReplacementCS_StructuredBuffer_Uint_Zero		, SF_Compute);
+IMPLEMENT_SHADER_TYPE4_WITH_TEMPLATE_PREFIX(template<>, RENDERCORE_API, FClearReplacementCS_Texture2DArray_Uint_Zero		, SF_Compute);
+IMPLEMENT_SHADER_TYPE4_WITH_TEMPLATE_PREFIX(template<>, RENDERCORE_API, FClearReplacementCS_Buffer_Uint						, SF_Compute);
+IMPLEMENT_SHADER_TYPE4_WITH_TEMPLATE_PREFIX(template<>, RENDERCORE_API, FClearReplacementCS_StructuredBuffer_Uint			, SF_Compute);
+IMPLEMENT_SHADER_TYPE4_WITH_TEMPLATE_PREFIX(template<>, RENDERCORE_API, FClearReplacementCS_Texture2DArray_Uint				, SF_Compute);
+IMPLEMENT_SHADER_TYPE4_WITH_TEMPLATE_PREFIX(template<>, RENDERCORE_API, FClearReplacementCS_Texture3D_Float4				, SF_Compute);
+IMPLEMENT_SHADER_TYPE4_WITH_TEMPLATE_PREFIX(template<>, RENDERCORE_API, FClearReplacementCS_Texture2D_Float4				, SF_Compute);
+IMPLEMENT_SHADER_TYPE4_WITH_TEMPLATE_PREFIX(template<>, RENDERCORE_API, FClearReplacementCS_Texture2DArray_Float4			, SF_Compute);
+IMPLEMENT_SHADER_TYPE4_WITH_TEMPLATE_PREFIX(template<>, RENDERCORE_API, FClearReplacementCS_Texture3D_Uint4					, SF_Compute);
+IMPLEMENT_SHADER_TYPE4_WITH_TEMPLATE_PREFIX(template<>, RENDERCORE_API, FClearReplacementCS_Texture2D_Uint4					, SF_Compute);
+IMPLEMENT_SHADER_TYPE4_WITH_TEMPLATE_PREFIX(template<>, RENDERCORE_API, FClearReplacementCS_Texture2DArray_Uint4			, SF_Compute);
+IMPLEMENT_SHADER_TYPE4_WITH_TEMPLATE_PREFIX(template<>, RENDERCORE_API, FClearReplacementCS_Texture3D_Uint4_Bounds			, SF_Compute);
+IMPLEMENT_SHADER_TYPE4_WITH_TEMPLATE_PREFIX(template<>, RENDERCORE_API, FClearReplacementCS_Texture2D_Uint4_Bounds			, SF_Compute);
+IMPLEMENT_SHADER_TYPE4_WITH_TEMPLATE_PREFIX(template<>, RENDERCORE_API, FClearReplacementCS_Texture2DArray_Uint4_Bounds		, SF_Compute);
+IMPLEMENT_SHADER_TYPE4_WITH_TEMPLATE_PREFIX(template<>, RENDERCORE_API, FClearReplacementCS_Buffer_Sint4_Bounds				, SF_Compute);
+IMPLEMENT_SHADER_TYPE4_WITH_TEMPLATE_PREFIX(template<>, RENDERCORE_API, FClearReplacementCS_StructuredBuffer_Sint4_Bounds	, SF_Compute);
+IMPLEMENT_SHADER_TYPE4_WITH_TEMPLATE_PREFIX(template<>, RENDERCORE_API, FClearReplacementCS_Texture3D_Sint4_Bounds			, SF_Compute);
+IMPLEMENT_SHADER_TYPE4_WITH_TEMPLATE_PREFIX(template<>, RENDERCORE_API, FClearReplacementCS_Texture2D_Sint4_Bounds			, SF_Compute);
+IMPLEMENT_SHADER_TYPE4_WITH_TEMPLATE_PREFIX(template<>, RENDERCORE_API, FClearReplacementCS_Texture2DArray_Sint4_Bounds		, SF_Compute);
diff --git a/Engine/Source/Runtime/RenderCore/Public/ClearReplacementShaders.h b/Engine/Source/Runtime/RenderCore/Public/ClearReplacementShaders.h
--- a/Engine/Source/Runtime/RenderCore/Public/ClearReplacementShaders.h
+++ b/Engine/Source/Runtime/RenderCore/Public/ClearReplacementShaders.h
@@ -14,7 +14,8 @@ enum class EClearReplacementResourceType
 	Buffer = 0,
 	Texture2D = 1,
 	Texture2DArray = 2,
-	Texture3D = 3
+	Texture3D = 3,
+	StructuredBuffer = 4
 };
 
 enum class EClearReplacementValueType
@@ -109,10 +110,11 @@ private:
 namespace ClearReplacementCS
 {
 	template <EClearReplacementResourceType> struct TThreadGroupSize {};
-	template <> struct TThreadGroupSize<EClearReplacementResourceType::Buffer>         { static constexpr int32 X = 64, Y = 1, Z = 1; };
-	template <> struct TThreadGroupSize<EClearReplacementResourceType::Texture2D>      { static constexpr int32 X =  8, Y = 8, Z = 1; };
-	template <> struct TThreadGroupSize<EClearReplacementResourceType::Texture2DArray> { static constexpr int32 X =  8, Y = 8, Z = 1; };
-	template <> struct TThreadGroupSize<EClearReplacementResourceType::Texture3D>      { static constexpr int32 X =  4, Y = 4, Z = 4; };
+	template <> struct TThreadGroupSize<EClearReplacementResourceType::Buffer>				{ static constexpr int32 X = 64, Y = 1, Z = 1; };
+	template <> struct TThreadGroupSize<EClearReplacementResourceType::Texture2D>			{ static constexpr int32 X =  8, Y = 8, Z = 1; };
+	template <> struct TThreadGroupSize<EClearReplacementResourceType::Texture2DArray>		{ static constexpr int32 X =  8, Y = 8, Z = 1; };
+	template <> struct TThreadGroupSize<EClearReplacementResourceType::Texture3D>			{ static constexpr int32 X =  4, Y = 4, Z = 4; };
+	template <> struct TThreadGroupSize<EClearReplacementResourceType::StructuredBuffer>	{ static constexpr int32 X = 64, Y = 1, Z = 1; };
 }
 
 template <EClearReplacementResourceType ResourceType, typename BaseType>
@@ -255,38 +257,46 @@ typedef TClearReplacementPS<true,  FClearReplacementBase_Float4>             FCl
 typedef TClearReplacementPS<false, FClearReplacementBase_Float4_Zero>        FClearReplacementPS_Zero;
 	
 // Compute shaders for clearing each resource type, with a min/max bounds enabled.
-typedef TClearReplacementCS<EClearReplacementResourceType::Buffer,         FClearReplacementBase_Uint_Bounds>   FClearReplacementCS_Buffer_Uint_Bounds;
-typedef TClearReplacementCS<EClearReplacementResourceType::Buffer,         FClearReplacementBase_Float_Bounds> FClearReplacementCS_Buffer_Float_Bounds;
-typedef TClearReplacementCS<EClearReplacementResourceType::Buffer,         FClearReplacementBase_Float4_Bounds> FClearReplacementCS_Buffer_Float4_Bounds;
-typedef TClearReplacementCS<EClearReplacementResourceType::Texture3D,      FClearReplacementBase_Float4_Bounds> FClearReplacementCS_Texture3D_Float4_Bounds;
-typedef TClearReplacementCS<EClearReplacementResourceType::Texture2D,      FClearReplacementBase_Float4_Bounds> FClearReplacementCS_Texture2D_Float4_Bounds;
-typedef TClearReplacementCS<EClearReplacementResourceType::Texture2DArray, FClearReplacementBase_Float4_Bounds> FClearReplacementCS_Texture2DArray_Float4_Bounds;
+typedef TClearReplacementCS<EClearReplacementResourceType::Buffer,           FClearReplacementBase_Uint_Bounds>   FClearReplacementCS_Buffer_Uint_Bounds;
+typedef TClearReplacementCS<EClearReplacementResourceType::Buffer,           FClearReplacementBase_Float_Bounds>  FClearReplacementCS_Buffer_Float_Bounds;
+typedef TClearReplacementCS<EClearReplacementResourceType::Buffer,           FClearReplacementBase_Float4_Bounds> FClearReplacementCS_Buffer_Float4_Bounds;
+typedef TClearReplacementCS<EClearReplacementResourceType::StructuredBuffer, FClearReplacementBase_Uint_Bounds>   FClearReplacementCS_StructuredBuffer_Uint_Bounds;
+typedef TClearReplacementCS<EClearReplacementResourceType::StructuredBuffer, FClearReplacementBase_Float_Bounds>  FClearReplacementCS_StructuredBuffer_Float_Bounds;
+typedef TClearReplacementCS<EClearReplacementResourceType::StructuredBuffer, FClearReplacementBase_Float4_Bounds> FClearReplacementCS_StructuredBuffer_Float4_Bounds;
+typedef TClearReplacementCS<EClearReplacementResourceType::Texture3D,        FClearReplacementBase_Float4_Bounds> FClearReplacementCS_Texture3D_Float4_Bounds;
+typedef TClearReplacementCS<EClearReplacementResourceType::Texture2D,		 FClearReplacementBase_Float4_Bounds> FClearReplacementCS_Texture2D_Float4_Bounds;
+typedef TClearReplacementCS<EClearReplacementResourceType::Texture2DArray,	 FClearReplacementBase_Float4_Bounds> FClearReplacementCS_Texture2DArray_Float4_Bounds;
 
 // Compute shaders for clearing each resource type. No bounds checks enabled.
-typedef TClearReplacementCS<EClearReplacementResourceType::Buffer,         FClearReplacementBase_Uint_Zero>     FClearReplacementCS_Buffer_Uint_Zero;
-typedef TClearReplacementCS<EClearReplacementResourceType::Texture2DArray, FClearReplacementBase_Uint_Zero>     FClearReplacementCS_Texture2DArray_Uint_Zero;
-typedef TClearReplacementCS<EClearReplacementResourceType::Buffer,         FClearReplacementBase_Uint>          FClearReplacementCS_Buffer_Uint;
-typedef TClearReplacementCS<EClearReplacementResourceType::Texture2DArray, FClearReplacementBase_Uint>          FClearReplacementCS_Texture2DArray_Uint;
+typedef TClearReplacementCS<EClearReplacementResourceType::Buffer,           FClearReplacementBase_Uint_Zero>     FClearReplacementCS_Buffer_Uint_Zero;
+typedef TClearReplacementCS<EClearReplacementResourceType::StructuredBuffer, FClearReplacementBase_Uint_Zero>     FClearReplacementCS_StructuredBuffer_Uint_Zero;
+typedef TClearReplacementCS<EClearReplacementResourceType::Texture2DArray,   FClearReplacementBase_Uint_Zero>     FClearReplacementCS_Texture2DArray_Uint_Zero;
+typedef TClearReplacementCS<EClearReplacementResourceType::Buffer,           FClearReplacementBase_Uint>          FClearReplacementCS_Buffer_Uint;
+typedef TClearReplacementCS<EClearReplacementResourceType::StructuredBuffer, FClearReplacementBase_Uint>          FClearReplacementCS_StructuredBuffer_Uint;
+typedef TClearReplacementCS<EClearReplacementResourceType::Texture2DArray,   FClearReplacementBase_Uint>          FClearReplacementCS_Texture2DArray_Uint;
 	
-typedef TClearReplacementCS<EClearReplacementResourceType::Texture3D,      FClearReplacementBase_Float4>        FClearReplacementCS_Texture3D_Float4;
-typedef TClearReplacementCS<EClearReplacementResourceType::Texture2D,      FClearReplacementBase_Float4>        FClearReplacementCS_Texture2D_Float4;
-typedef TClearReplacementCS<EClearReplacementResourceType::Texture2DArray, FClearReplacementBase_Float4>        FClearReplacementCS_Texture2DArray_Float4;
+typedef TClearReplacementCS<EClearReplacementResourceType::Texture3D,        FClearReplacementBase_Float4>        FClearReplacementCS_Texture3D_Float4;
+typedef TClearReplacementCS<EClearReplacementResourceType::Texture2D,        FClearReplacementBase_Float4>        FClearReplacementCS_Texture2D_Float4;
+typedef TClearReplacementCS<EClearReplacementResourceType::Texture2DArray,   FClearReplacementBase_Float4>        FClearReplacementCS_Texture2DArray_Float4;
 	
 // Used by ClearUAV_T in ClearQuad.cpp
-typedef TClearReplacementCS<EClearReplacementResourceType::Texture3D,      FClearReplacementBase_Uint4>         FClearReplacementCS_Texture3D_Uint4;
-typedef TClearReplacementCS<EClearReplacementResourceType::Texture2D,      FClearReplacementBase_Uint4>         FClearReplacementCS_Texture2D_Uint4;
-typedef TClearReplacementCS<EClearReplacementResourceType::Texture2DArray, FClearReplacementBase_Uint4>         FClearReplacementCS_Texture2DArray_Uint4;
-typedef TClearReplacementCS<EClearReplacementResourceType::Buffer,         FClearReplacementBase_Uint4>         FClearReplacementCS_Buffer_Uint4;
-
-typedef TClearReplacementCS<EClearReplacementResourceType::Buffer,         FClearReplacementBase_Uint4_Bounds>  FClearReplacementCS_Buffer_Uint4_Bounds;
-typedef TClearReplacementCS<EClearReplacementResourceType::Texture3D,      FClearReplacementBase_Uint4_Bounds>  FClearReplacementCS_Texture3D_Uint4_Bounds;
-typedef TClearReplacementCS<EClearReplacementResourceType::Texture2D,      FClearReplacementBase_Uint4_Bounds>  FClearReplacementCS_Texture2D_Uint4_Bounds;
-typedef TClearReplacementCS<EClearReplacementResourceType::Texture2DArray, FClearReplacementBase_Uint4_Bounds>  FClearReplacementCS_Texture2DArray_Uint4_Bounds;
-
-typedef TClearReplacementCS<EClearReplacementResourceType::Buffer,         FClearReplacementBase_Sint4_Bounds>  FClearReplacementCS_Buffer_Sint4_Bounds;
-typedef TClearReplacementCS<EClearReplacementResourceType::Texture3D,      FClearReplacementBase_Sint4_Bounds>  FClearReplacementCS_Texture3D_Sint4_Bounds;
-typedef TClearReplacementCS<EClearReplacementResourceType::Texture2D,      FClearReplacementBase_Sint4_Bounds>  FClearReplacementCS_Texture2D_Sint4_Bounds;
-typedef TClearReplacementCS<EClearReplacementResourceType::Texture2DArray, FClearReplacementBase_Sint4_Bounds>  FClearReplacementCS_Texture2DArray_Sint4_Bounds;
+typedef TClearReplacementCS<EClearReplacementResourceType::Texture3D,        FClearReplacementBase_Uint4>         FClearReplacementCS_Texture3D_Uint4;
+typedef TClearReplacementCS<EClearReplacementResourceType::Texture2D,        FClearReplacementBase_Uint4>         FClearReplacementCS_Texture2D_Uint4;
+typedef TClearReplacementCS<EClearReplacementResourceType::Texture2DArray,   FClearReplacementBase_Uint4>         FClearReplacementCS_Texture2DArray_Uint4;
+typedef TClearReplacementCS<EClearReplacementResourceType::Buffer,			 FClearReplacementBase_Uint4>         FClearReplacementCS_Buffer_Uint4;
+typedef TClearReplacementCS<EClearReplacementResourceType::StructuredBuffer, FClearReplacementBase_Uint4>         FClearReplacementCS_StructuredBuffer_Uint4;
+
+typedef TClearReplacementCS<EClearReplacementResourceType::Buffer,		     FClearReplacementBase_Uint4_Bounds>  FClearReplacementCS_Buffer_Uint4_Bounds;
+typedef TClearReplacementCS<EClearReplacementResourceType::StructuredBuffer, FClearReplacementBase_Uint4_Bounds>  FClearReplacementCS_StructuredBuffer_Uint4_Bounds;
+typedef TClearReplacementCS<EClearReplacementResourceType::Texture3D,        FClearReplacementBase_Uint4_Bounds>  FClearReplacementCS_Texture3D_Uint4_Bounds;
+typedef TClearReplacementCS<EClearReplacementResourceType::Texture2D,        FClearReplacementBase_Uint4_Bounds>  FClearReplacementCS_Texture2D_Uint4_Bounds;
+typedef TClearReplacementCS<EClearReplacementResourceType::Texture2DArray,   FClearReplacementBase_Uint4_Bounds>  FClearReplacementCS_Texture2DArray_Uint4_Bounds;
+
+typedef TClearReplacementCS<EClearReplacementResourceType::Buffer,           FClearReplacementBase_Sint4_Bounds>  FClearReplacementCS_Buffer_Sint4_Bounds;
+typedef TClearReplacementCS<EClearReplacementResourceType::StructuredBuffer, FClearReplacementBase_Sint4_Bounds>  FClearReplacementCS_StructuredBuffer_Sint4_Bounds;
+typedef TClearReplacementCS<EClearReplacementResourceType::Texture3D,        FClearReplacementBase_Sint4_Bounds>  FClearReplacementCS_Texture3D_Sint4_Bounds;
+typedef TClearReplacementCS<EClearReplacementResourceType::Texture2D,        FClearReplacementBase_Sint4_Bounds>  FClearReplacementCS_Texture2D_Sint4_Bounds;
+typedef TClearReplacementCS<EClearReplacementResourceType::Texture2DArray,   FClearReplacementBase_Sint4_Bounds>  FClearReplacementCS_Texture2DArray_Sint4_Bounds;
 
 /**
  * Helper functions for running the clear replacement shader for specific resource types, values types and number of channels.
diff --git a/Engine/Source/Runtime/VulkanRHI/Private/VulkanUAV.cpp b/Engine/Source/Runtime/VulkanRHI/Private/VulkanUAV.cpp
--- a/Engine/Source/Runtime/VulkanRHI/Private/VulkanUAV.cpp
+++ b/Engine/Source/Runtime/VulkanRHI/Private/VulkanUAV.cpp
@@ -534,6 +534,12 @@ void FVulkanCommandListContext::ClearUAV(TRHICommandList_RecursiveHazardous<FVul
 			ensure(0);
 		}
 	}
+	else if (UnorderedAccessView->SourceStructuredBuffer)
+	{
+		ensure(UnorderedAccessView->BufferViewFormat == PF_Unknown);
+		uint32 NumUints = UnorderedAccessView->SourceStructuredBuffer->GetSize() / 4;
+		ClearUAVShader_T<EClearReplacementResourceType::StructuredBuffer, 4, false>(RHICmdList, UnorderedAccessView, NumUints, 1, 1, ClearValue, ValueType);
+	}
 	else
 	{
 		ensure(0);
