Fixed RDG scopes not being pushed due to unsigned / signed mismatch.

diff --git a/Engine/Source/Runtime/RenderCore/Public/RenderGraphEvent.inl b/Engine/Source/Runtime/RenderCore/Public/RenderGraphEvent.inl
--- a/Engine/Source/Runtime/RenderCore/Public/RenderGraphEvent.inl
+++ b/Engine/Source/Runtime/RenderCore/Public/RenderGraphEvent.inl
@@ -7,13 +7,13 @@ TRDGScopeStack<TScopeType>::TRDGScopeStack(
 	FRHIComputeCommandList& InRHICmdList,
 	FPushFunction InPushFunction,
 	FPopFunction InPopFunction,
-	bool bRDGEvents)
+	bool bInRDGEvents)
 	: RHICmdList(InRHICmdList)
 	, MemStack(FMemStack::Get())
 	, PushFunction(InPushFunction)
 	, PopFunction(InPopFunction)
 	, ScopeStack(MakeUniformStaticArray<const TScopeType*, kScopeStackDepthMax>(nullptr))
-	, bRDGEvents(bRDGEvents)
+	, bRDGEvents(bInRDGEvents)
 {}
 
 template <typename TScopeType>
@@ -88,7 +88,7 @@ void TRDGScopeStack<TScopeType>::BeginExecutePass(const TScopeType* ParentScope)
 	}
 
 	// Push new scopes.
-	for (int32 i = TraversedScopeCount - 1; i >= 0 && CommonScopeId < kScopeStackDepthMax; i--)
+	for (int32 i = TraversedScopeCount - 1; i >= 0 && CommonScopeId < static_cast<int32>(kScopeStackDepthMax); i--)
 	{
 		PushFunction(RHICmdList, TraversedScopes[i], bRDGEvents);
 		CommonScopeId++;
