Fixed crash in DrawTileMesh due to SkipBarriers flag.

diff --git a/Engine/Source/Runtime/RenderCore/Private/RenderGraphBuilder.cpp b/Engine/Source/Runtime/RenderCore/Private/RenderGraphBuilder.cpp
--- a/Engine/Source/Runtime/RenderCore/Private/RenderGraphBuilder.cpp
+++ b/Engine/Source/Runtime/RenderCore/Private/RenderGraphBuilder.cpp
@@ -421,12 +421,11 @@ ERDGPassFlags FRDGBuilder::OverridePassFlags(const TCHAR* PassName, ERDGPassFlag
 
 const char* const FRDGBuilder::kDefaultUnaccountedCSVStat = "RDG_Pass";
 
-FRDGBuilder::FRDGBuilder(FRHICommandListImmediate& InRHICmdList, FRDGEventName InName, ERDGBuilderFlags InFlags)
+FRDGBuilder::FRDGBuilder(FRHICommandListImmediate& InRHICmdList, FRDGEventName InName)
 	: RHICmdList(InRHICmdList)
 	, Blackboard(Allocator)
 	, RHICmdListAsyncCompute(FRHICommandListExecutor::GetImmediateAsyncComputeCommandList())
 	, BuilderName(InName)
-	, BuilderFlags(InFlags)
 #if RDG_CPU_SCOPES
 	, CPUScopeStacks(RHICmdList, kDefaultUnaccountedCSVStat)
 #endif
@@ -434,7 +433,7 @@ FRDGBuilder::FRDGBuilder(FRHICommandListImmediate& InRHICmdList, FRDGEventName I
 	, GPUScopeStacks(RHICmdList, RHICmdListAsyncCompute)
 #endif
 #if RDG_ENABLE_DEBUG
-	, UserValidation(Allocator, BuilderFlags)
+	, UserValidation(Allocator)
 	, BarrierValidation(&Passes, BuilderName)
 #endif
 {
@@ -517,7 +516,6 @@ FRDGTextureRef FRDGBuilder::RegisterExternalTexture(
 		TEXT("Externally registered texture '%s' has known RDG state. This means the graph did not sanitize it correctly, or ")
 		TEXT("an IPooledRenderTarget reference was improperly held within a pass."), Texture->Name);
 
-	if (!EnumHasAnyFlags(BuilderFlags, ERDGBuilderFlags::SkipBarriers))
 	{
 		FRDGSubresourceState SubresourceState;
 		SubresourceState.Access = AccessInitial;
@@ -576,10 +574,7 @@ FRDGBufferRef FRDGBuilder::RegisterExternalBuffer(
 		TEXT("Externally registered buffer '%s' has known RDG state. This means the graph did not sanitize it correctly, or ")
 		TEXT("an FRDGPooledBuffer reference was improperly held within a pass."), Buffer->Name);
 
-	if (!EnumHasAnyFlags(BuilderFlags, ERDGBuilderFlags::SkipBarriers))
-	{
-		BufferState.Access = AccessInitial;
-	}
+	BufferState.Access = AccessInitial;
 
 	ExternalBuffers.Add(ExternalPooledBuffer, Buffer);
 
@@ -1777,7 +1772,7 @@ void FRDGBuilder::CollectPassBarriers(FRDGPassHandle PassHandle, FRDGPassHandle&
 		LastUntrackedPassHandle = PassHandle;
 	}
 
-	if (PassesWithEmptyParameters[PassHandle] || EnumHasAnyFlags(BuilderFlags, ERDGBuilderFlags::SkipBarriers))
+	if (PassesWithEmptyParameters[PassHandle])
 	{
 		return;
 	}
@@ -2024,8 +2019,6 @@ void FRDGBuilder::AddTransitionInternal(
 	FRDGPassHandle LastUntrackedPassHandle,
 	const FRHITransitionInfo& TransitionInfo)
 {
-	check(!EnumHasAnyFlags(BuilderFlags, ERDGBuilderFlags::SkipBarriers));
-
 	const ERHIPipeline Graphics = ERHIPipeline::Graphics;
 	const ERHIPipeline AsyncCompute = ERHIPipeline::AsyncCompute;
 
diff --git a/Engine/Source/Runtime/RenderCore/Private/RenderGraphValidation.cpp b/Engine/Source/Runtime/RenderCore/Private/RenderGraphValidation.cpp
--- a/Engine/Source/Runtime/RenderCore/Private/RenderGraphValidation.cpp
+++ b/Engine/Source/Runtime/RenderCore/Private/RenderGraphValidation.cpp
@@ -160,9 +160,8 @@ void FRDGUniformBuffer::MarkResourceAsUsed()
 	});
 }
 
-FRDGUserValidation::FRDGUserValidation(FRDGAllocator& InAllocator, ERDGBuilderFlags InBuilderFlags)
+FRDGUserValidation::FRDGUserValidation(FRDGAllocator& InAllocator)
 	: Allocator(InAllocator)
-	, BuilderFlags(InBuilderFlags)
 	, ExpectedNumMarks(FMemStack::Get().GetNumMarks())
 {
 	checkf(!GRDGBuilderActive, TEXT("Another FRDGBuilder already exists on the stack. Only one builder can be created at a time. This builder instance should be merged into the parent one."));
@@ -293,8 +292,6 @@ void FRDGUserValidation::ValidateCreateTexture(const FRDGTextureDesc& Desc, cons
 	checkf(!bIsUAVForMSAATexture, TEXT("TexCreate_UAV is not allowed on MSAA texture %s."), Name);
 
 	checkf(!EnumHasAnyFlags(Flags, ERDGTextureFlags::ReadOnly), TEXT("Cannot create texture %s with the ReadOnly flag. Only registered textures can use this flag."), Name);
-
-	checkf(!EnumHasAnyFlags(BuilderFlags, ERDGBuilderFlags::SkipBarriers), TEXT("Cannot create texture '%s' because RDG builder has 'SkipBarriers' flag."), Name);
 }
 
 void FRDGUserValidation::ValidateCreateBuffer(const FRDGBufferDesc& Desc, const TCHAR* Name, ERDGBufferFlags Flags)
@@ -312,8 +309,6 @@ void FRDGUserValidation::ValidateCreateBuffer(const FRDGBufferDesc& Desc, const
 	}
 
 	checkf(!EnumHasAnyFlags(Flags, ERDGBufferFlags::ReadOnly), TEXT("Cannot create buffer %s with the ReadOnly flag. Only registered buffers can use this flag."), Name);
-
-	checkf(!EnumHasAnyFlags(BuilderFlags, ERDGBuilderFlags::SkipBarriers), TEXT("Cannot create buffer '%s' because RDG builder has 'SkipBarriers' flag."), Name);
 }
 
 void FRDGUserValidation::ValidateCreateSRV(const FRDGTextureSRVDesc& Desc)
diff --git a/Engine/Source/Runtime/RenderCore/Public/RenderGraphBuilder.h b/Engine/Source/Runtime/RenderCore/Public/RenderGraphBuilder.h
--- a/Engine/Source/Runtime/RenderCore/Public/RenderGraphBuilder.h
+++ b/Engine/Source/Runtime/RenderCore/Public/RenderGraphBuilder.h
@@ -22,7 +22,7 @@ class RENDERCORE_API FRDGBuilder
 	: FRDGAllocatorScope
 {
 public:
-	FRDGBuilder(FRHICommandListImmediate& InRHICmdList, FRDGEventName InName = {}, ERDGBuilderFlags Flags = ERDGBuilderFlags::None);
+	FRDGBuilder(FRHICommandListImmediate& InRHICmdList, FRDGEventName InName = {});
 	FRDGBuilder(const FRDGBuilder&) = delete;
 
 	/** Finds an RDG texture associated with the external texture, or returns null if none is found. */
@@ -230,7 +230,6 @@ private:
 	FRHIAsyncComputeCommandListImmediate& RHICmdListAsyncCompute;
 
 	const FRDGEventName BuilderName;
-	const ERDGBuilderFlags BuilderFlags;
 
 	ERDGPassFlags OverridePassFlags(const TCHAR* PassName, ERDGPassFlags Flags, bool bAsyncComputeSupported);
 
diff --git a/Engine/Source/Runtime/RenderCore/Public/RenderGraphDefinitions.h b/Engine/Source/Runtime/RenderCore/Public/RenderGraphDefinitions.h
--- a/Engine/Source/Runtime/RenderCore/Public/RenderGraphDefinitions.h
+++ b/Engine/Source/Runtime/RenderCore/Public/RenderGraphDefinitions.h
@@ -62,16 +62,6 @@
 
 /** ENUMS */
 
-enum class ERDGBuilderFlags
-{
-	None = 0,
-
-	/** The builder will not perform any resource barriers. This is only allowed when all resources are registered. This exists
-	 *  as a workaround for the RDG engine port, and will be deprecated immediately once the scaffolding code is removed.
-	 */
-	SkipBarriers = 1 << 0,
-};
-
 /** Flags to annotate a pass with when calling AddPass. */
 enum class ERDGPassFlags : uint8
 {
diff --git a/Engine/Source/Runtime/RenderCore/Public/RenderGraphValidation.h b/Engine/Source/Runtime/RenderCore/Public/RenderGraphValidation.h
--- a/Engine/Source/Runtime/RenderCore/Public/RenderGraphValidation.h
+++ b/Engine/Source/Runtime/RenderCore/Public/RenderGraphValidation.h
@@ -23,7 +23,7 @@
 class RENDERCORE_API FRDGUserValidation final
 {
 public:
-	FRDGUserValidation(FRDGAllocator& Allocator, ERDGBuilderFlags BuilderFlags);
+	FRDGUserValidation(FRDGAllocator& Allocator);
 	FRDGUserValidation(const FRDGUserValidation&) = delete;
 	~FRDGUserValidation();
 
@@ -111,7 +111,6 @@ private:
 	void ExecuteGuard(const TCHAR* Operation, const TCHAR* ResourceName);
 
 	FRDGAllocator& Allocator;
-	ERDGBuilderFlags BuilderFlags;
 	int32 ExpectedNumMarks;
 };
 
diff --git a/Engine/Source/Runtime/Renderer/Private/Renderer.cpp b/Engine/Source/Runtime/Renderer/Private/Renderer.cpp
--- a/Engine/Source/Runtime/Renderer/Private/Renderer.cpp
+++ b/Engine/Source/Runtime/Renderer/Private/Renderer.cpp
@@ -183,7 +183,7 @@ void FRendererModule::DrawTileMesh(FRHICommandListImmediate& RHICmdList, FMeshPa
 
 		GSystemTextures.InitializeTextures(RHICmdList, FeatureLevel);
 		FMemMark Mark(FMemStack::Get());
-		FRDGBuilder GraphBuilder(RHICmdList, {}, ERDGBuilderFlags::SkipBarriers);
+		FRDGBuilder GraphBuilder(RHICmdList);
 
 		// Initialize the RDG read-only system textures.
 		FRDGSystemTextures::Create(GraphBuilder);
