Modified RenderGraph CreateStructuredBuffer helper to support 0 element allocations by returning a dummy buffer with 1 element

diff --git a/Engine/Source/Runtime/RenderCore/Public/RenderGraphUtils.h b/Engine/Source/Runtime/RenderCore/Public/RenderGraphUtils.h
--- a/Engine/Source/Runtime/RenderCore/Public/RenderGraphUtils.h
+++ b/Engine/Source/Runtime/RenderCore/Public/RenderGraphUtils.h
@@ -617,6 +617,11 @@ FORCEINLINE FRDGBufferRef CreateStructuredBuffer(
 	const TArray<ElementType, AllocatorType>& InitialData,
 	ERDGInitialDataFlags InitialDataFlags = ERDGInitialDataFlags::None)
 {
+	static const ElementType DummyElement = ElementType();
+	if (InitialData.Num() == 0)
+	{
+		return CreateStructuredBuffer(GraphBuilder, Name, InitialData.GetTypeSize(), 1, &DummyElement, InitialData.GetTypeSize());
+	}
 	return CreateStructuredBuffer(GraphBuilder, Name, InitialData.GetTypeSize(), InitialData.Num(), InitialData.GetData(), InitialData.Num() * InitialData.GetTypeSize(), InitialDataFlags);
 }
 
