Small refactor to dynamic RDG 'AddPass' function to take a void parameter struct pointer, since it's providing its own metadata.

#rnx

diff --git a/Engine/Source/Runtime/RenderCore/Public/RenderGraphBuilder.h b/Engine/Source/Runtime/RenderCore/Public/RenderGraphBuilder.h
--- a/Engine/Source/Runtime/RenderCore/Public/RenderGraphBuilder.h
+++ b/Engine/Source/Runtime/RenderCore/Public/RenderGraphBuilder.h
@@ -140,12 +140,13 @@ public:
 	 *
 	 *  The pass name is used by debugging / profiling tools.
 	 */
-	template <typename ParameterStructType, typename ExecuteLambdaType>
-	FRDGPassRef AddPass(FRDGEventName&& Name, const FShaderParametersMetadata* ParametersMetadata, const ParameterStructType* ParameterStruct, ERDGPassFlags Flags, ExecuteLambdaType&& ExecuteLambda);
-	
 	template <typename ParameterStructType, typename ExecuteLambdaType>
 	FRDGPassRef AddPass(FRDGEventName&& Name, const ParameterStructType* ParameterStruct, ERDGPassFlags Flags, ExecuteLambdaType&& ExecuteLambda);
 
+	/** Adds a lambda pass to the graph with a runtime-generated parameter struct. */
+	template <typename ExecuteLambdaType>
+	FRDGPassRef AddPass(FRDGEventName&& Name, const FShaderParametersMetadata* ParametersMetadata, const void* ParameterStruct, ERDGPassFlags Flags, ExecuteLambdaType&& ExecuteLambda);
+
 	/** Adds a lambda pass to the graph without any parameters. This useful for deferring RHI work onto the graph timeline,
 	 *  or incrementally porting code to use the graph system. NeverCull and SkipRenderPass (if Raster) are implicitly added
 	 *  to Flags. AsyncCompute is not allowed. Accessing RHI resources which are tracked by the graph is only allowed if the
@@ -231,6 +232,14 @@ private:
 
 	const FRDGEventName BuilderName;
 
+	template <typename ParameterStructType, typename ExecuteLambdaType>
+	FRDGPassRef AddPassInternal(
+		FRDGEventName&& Name,
+		const FShaderParametersMetadata* ParametersMetadata,
+		const ParameterStructType* ParameterStruct,
+		ERDGPassFlags Flags,
+		ExecuteLambdaType&& ExecuteLambda);
+
 	ERDGPassFlags OverridePassFlags(const TCHAR* PassName, ERDGPassFlags Flags, bool bAsyncComputeSupported);
 
 	FORCEINLINE FRDGPassHandle GetProloguePassHandle() const
diff --git a/Engine/Source/Runtime/RenderCore/Public/RenderGraphBuilder.inl b/Engine/Source/Runtime/RenderCore/Public/RenderGraphBuilder.inl
--- a/Engine/Source/Runtime/RenderCore/Public/RenderGraphBuilder.inl
+++ b/Engine/Source/Runtime/RenderCore/Public/RenderGraphBuilder.inl
@@ -129,7 +129,7 @@ FRDGPassRef FRDGBuilder::AddPass(
 }
 
 template <typename ParameterStructType, typename ExecuteLambdaType>
-FRDGPassRef FRDGBuilder::AddPass(
+FRDGPassRef FRDGBuilder::AddPassInternal(
 	FRDGEventName&& Name,
 	const FShaderParametersMetadata* ParametersMetadata,
 	const ParameterStructType* ParameterStruct,
@@ -153,6 +153,17 @@ FRDGPassRef FRDGBuilder::AddPass(
 	return Pass;
 }
 
+template <typename ExecuteLambdaType>
+FRDGPassRef FRDGBuilder::AddPass(
+	FRDGEventName&& Name,
+	const FShaderParametersMetadata* ParametersMetadata,
+	const void* ParameterStruct,
+	ERDGPassFlags Flags,
+	ExecuteLambdaType&& ExecuteLambda)
+{
+	return AddPassInternal(Forward<FRDGEventName>(Name), ParametersMetadata, ParameterStruct, Flags, Forward<ExecuteLambdaType>(ExecuteLambda));
+}
+
 template <typename ParameterStructType, typename ExecuteLambdaType>
 FRDGPassRef FRDGBuilder::AddPass(
 	FRDGEventName&& Name,
@@ -160,7 +171,7 @@ FRDGPassRef FRDGBuilder::AddPass(
 	ERDGPassFlags Flags,
 	ExecuteLambdaType&& ExecuteLambda)
 {
-	return AddPass(Forward<FRDGEventName>(Name), ParameterStructType::FTypeInfo::GetStructMetadata(), ParameterStruct, Flags, Forward<ExecuteLambdaType>(ExecuteLambda));
+	return AddPassInternal(Forward<FRDGEventName>(Name), ParameterStructType::FTypeInfo::GetStructMetadata(), ParameterStruct, Flags, Forward<ExecuteLambdaType>(ExecuteLambda));
 }
 
 inline void FRDGBuilder::QueueTextureExtraction(FRDGTextureRef Texture, TRefCountPtr<IPooledRenderTarget>* OutTexturePtr, ERHIAccess AccessFinal)
