Fix out of bounds scope stack index

diff --git a/Engine/Source/Runtime/RenderCore/Public/RenderGraphEvent.inl b/Engine/Source/Runtime/RenderCore/Public/RenderGraphEvent.inl
--- a/Engine/Source/Runtime/RenderCore/Public/RenderGraphEvent.inl
+++ b/Engine/Source/Runtime/RenderCore/Public/RenderGraphEvent.inl
@@ -88,7 +88,7 @@ void TRDGScopeStack<TScopeType>::BeginExecutePass(const TScopeType* ParentScope)
 	}
 
 	// Push new scopes.
-	for (int32 i = TraversedScopeCount - 1; i >= 0 && CommonScopeId < static_cast<int32>(kScopeStackDepthMax); i--)
+	for (int32 i = TraversedScopeCount - 1; i >= 0 && CommonScopeId + 1 < static_cast<int32>(kScopeStackDepthMax); i--)
 	{
 		PushFunction(RHICmdList, TraversedScopes[i], bRDGEvents);
 		CommonScopeId++;
