Converting some ShaderPlatform usages to ShaderFormat.

diff --git a/Engine/Source/Runtime/Engine/Private/ShaderCompiler/ShaderCompiler.cpp b/Engine/Source/Runtime/Engine/Private/ShaderCompiler/ShaderCompiler.cpp
--- a/Engine/Source/Runtime/Engine/Private/ShaderCompiler/ShaderCompiler.cpp
+++ b/Engine/Source/Runtime/Engine/Private/ShaderCompiler/ShaderCompiler.cpp
@@ -3952,16 +3952,18 @@ void GlobalBeginCompileShader(
 	COOK_STAT(FScopedDurationTimer DurationTimer(ShaderCompilerCookStats::GlobalBeginCompileShaderTimeSec));
 
 	EShaderPlatform ShaderPlatform = EShaderPlatform(Target.Platform);
+	const FName ShaderFormatName = LegacyShaderPlatformToShaderFormat(ShaderPlatform);
+
 	FShaderCompileUtilities::GenerateBrdfHeaders(ShaderPlatform);
 
 	Input.Target = Target;
-	Input.ShaderFormat = LegacyShaderPlatformToShaderFormat(ShaderPlatform);
+	Input.ShaderFormat = ShaderFormatName;
 	Input.VirtualSourceFilePath = SourceFilename;
 	Input.EntryPointName = FunctionName;
 	Input.bCompilingForShaderPipeline = false;
 	Input.bIncludeUsedOutputs = false;
 	Input.bGenerateDirectCompileFile = (GDumpShaderDebugInfoSCWCommandLine != 0);
-	Input.DumpDebugInfoRootPath = GShaderCompilingManager->GetAbsoluteShaderDebugInfoDirectory() / Input.ShaderFormat.ToString();
+	Input.DumpDebugInfoRootPath = GShaderCompilingManager->GetAbsoluteShaderDebugInfoDirectory() / ShaderFormatName.ToString();
 	// asset material name or "Global"
 	Input.DebugGroupName = DebugGroupName;
 	Input.DebugDescription = DebugDescription;
@@ -4143,7 +4145,7 @@ void GlobalBeginCompileShader(
 		// Throw a warning if we are silently disabling ISR due to missing platform support.
 		if (bIsInstancedStereoCVar && !bIsInstancedStereo && !GShaderCompilingManager->AreWarningsSuppressed(ShaderPlatform))
 		{
-			UE_LOG(LogShaderCompilers, Log, TEXT("Instanced stereo rendering is not supported for the %s shader platform."), *LegacyShaderPlatformToShaderFormat(ShaderPlatform).ToString());
+			UE_LOG(LogShaderCompilers, Log, TEXT("Instanced stereo rendering is not supported for the %s shader platform."), *ShaderFormatName.ToString());
 			GShaderCompilingManager->SuppressWarnings(ShaderPlatform);
 		}
 
@@ -4177,12 +4179,12 @@ void GlobalBeginCompileShader(
 		}
 	}
 
-	if (ShouldKeepShaderDebugInfo(Target.GetPlatform()))
+	if (ShouldKeepShaderDebugInfo(ShaderFormatName))
 	{
 		Input.Environment.CompilerFlags.Add(CFLAG_KeepDebugInfo);
 	}
 
-	if (ShouldAllowUniqueDebugInfo(Target.GetPlatform()))
+	if (ShouldAllowUniqueDebugInfo(ShaderFormatName))
 	{
 		Input.Environment.CompilerFlags.Add(CFLAG_AllowUniqueDebugInfo);
 	}
@@ -4256,8 +4258,7 @@ void GlobalBeginCompileShader(
 		
 		// Check whether we can compile metal shaders to bytecode - avoids poisoning the DDC
 		static ITargetPlatformManagerModule& TPM = GetTargetPlatformManagerRef();
-		const FName Format = LegacyShaderPlatformToShaderFormat(EShaderPlatform(Target.Platform));
-		const IShaderFormat* Compiler = TPM.FindShaderFormat(Format);
+		const IShaderFormat* Compiler = TPM.FindShaderFormat(ShaderFormatName);
 		static const bool bCanCompileOfflineMetalShaders = Compiler && Compiler->CanCompileBinaryShaders();
 		if (!bCanCompileOfflineMetalShaders)
 		{
@@ -4397,7 +4398,7 @@ void GlobalBeginCompileShader(
 		Input.Environment.SetDefine(TEXT("PROJECT_ALLOW_GLOBAL_CLIP_PLANE"), CVar ? (CVar->GetInt() != 0) : 0);
 	}
 
-	ITargetPlatform* TargetPlatform = GetTargetPlatformManager()->FindTargetPlatformWithSupport(TEXT("ShaderFormat"), LegacyShaderPlatformToShaderFormat((EShaderPlatform)Target.Platform));
+	ITargetPlatform* TargetPlatform = GetTargetPlatformManager()->FindTargetPlatformWithSupport(TEXT("ShaderFormat"), ShaderFormatName);
 	bool bForwardShading = false;
 	{
 		if (TargetPlatform)
@@ -4528,7 +4529,7 @@ void GlobalBeginCompileShader(
 	}
 
 	// Allow the target shader format to modify the shader input before we add it as a job
-	const IShaderFormat* Format = GetTargetPlatformManagerRef().FindShaderFormat(Input.ShaderFormat);
+	const IShaderFormat* Format = GetTargetPlatformManagerRef().FindShaderFormat(ShaderFormatName);
 	Format->ModifyShaderCompilerInput(Input);
 
 	// Allow the GBuffer and other shader defines to cause dependend environment changes, but minimizing the #ifdef magic in the shaders, which
diff --git a/Engine/Source/Runtime/RenderCore/Private/Shader.cpp b/Engine/Source/Runtime/RenderCore/Private/Shader.cpp
--- a/Engine/Source/Runtime/RenderCore/Private/Shader.cpp
+++ b/Engine/Source/Runtime/RenderCore/Private/Shader.cpp
@@ -1322,6 +1322,8 @@ bool IsDxcEnabledForPlatform(EShaderPlatform Platform)
 
 void ShaderMapAppendKeyString(EShaderPlatform Platform, FString& KeyString)
 {
+	const FName ShaderFormatName = LegacyShaderPlatformToShaderFormat(Platform);
+
 	// Globals that should cause all shaders to recompile when changed must be appended to the key here
 	// Key should be kept as short as possible while being somewhat human readable for debugging
 
@@ -1412,8 +1414,8 @@ void ShaderMapAppendKeyString(EShaderPlatform Platform, FString& KeyString)
 	}
 
 	{
-		KeyString += ShouldKeepShaderDebugInfo(Platform) ? TEXT("_NoStrip") : TEXT("");
-		KeyString += ShouldAllowUniqueDebugInfo(Platform) ? TEXT("_FullDbg") : TEXT("");
+		KeyString += ShouldKeepShaderDebugInfo(ShaderFormatName) ? TEXT("_NoStrip") : TEXT("");
+		KeyString += ShouldAllowUniqueDebugInfo(ShaderFormatName) ? TEXT("_FullDbg") : TEXT("");
 	}
 
 	{
@@ -1543,7 +1545,6 @@ void ShaderMapAppendKeyString(EShaderPlatform Platform, FString& KeyString)
 		}
 	}
 
-	const FName ShaderFormatName = LegacyShaderPlatformToShaderFormat(Platform);
 	const IShaderFormat* ShaderFormat = GetTargetPlatformManagerRef().FindShaderFormat(ShaderFormatName);
 	if (ShaderFormat)
 	{
@@ -1626,7 +1627,7 @@ void ShaderMapAppendKeyString(EShaderPlatform Platform, FString& KeyString)
 		}
 	}
 
-	ITargetPlatform* TargetPlatform = GetTargetPlatformManager()->FindTargetPlatformWithSupport(TEXT("ShaderFormat"), LegacyShaderPlatformToShaderFormat(Platform));
+	ITargetPlatform* TargetPlatform = GetTargetPlatformManager()->FindTargetPlatformWithSupport(TEXT("ShaderFormat"), ShaderFormatName);
 
 	{
 		bool bForwardShading = false;
diff --git a/Engine/Source/Runtime/RenderCore/Private/ShaderCore.cpp b/Engine/Source/Runtime/RenderCore/Private/ShaderCore.cpp
--- a/Engine/Source/Runtime/RenderCore/Private/ShaderCore.cpp
+++ b/Engine/Source/Runtime/RenderCore/Private/ShaderCore.cpp
@@ -280,7 +280,7 @@ bool AllowDebugViewmodes(EShaderPlatform Platform)
 #endif
 }
 
-static bool IsShaderCompilerConfigEnabledForPlatform(const IConsoleVariable* GlobalCvar, const TCHAR* Key, EShaderPlatform ShaderPlatform)
+static bool IsShaderCompilerConfigEnabledForPlatform(const IConsoleVariable* GlobalCvar, const TCHAR* Key, FName ShaderFormat)
 {
 	bool bEnabled = false;
 
@@ -294,11 +294,11 @@ static bool IsShaderCompilerConfigEnabledForPlatform(const IConsoleVariable* Glo
 	else
 	{
 		// Then check the per platform settings.
-		const FName PlatformName = ShaderPlatformToPlatformName(ShaderPlatform);
-		if (!PlatformName.IsNone())
+		ITargetPlatform* TargetPlatform = GetTargetPlatformManager()->FindTargetPlatformWithSupport(TEXT("ShaderFormat"), ShaderFormat);
+		if (TargetPlatform)
 		{
 			FConfigFile EngineSettings;
-			FConfigCacheIni::LoadLocalIniFile(EngineSettings, TEXT("Engine"), true, *PlatformName.ToString());
+			FConfigCacheIni::LoadLocalIniFile(EngineSettings, TEXT("Engine"), true, *TargetPlatform->IniPlatformName());
 
 			EngineSettings.GetBool(TEXT("ShaderCompiler"), Key, bEnabled);
 		}
@@ -308,7 +308,7 @@ static bool IsShaderCompilerConfigEnabledForPlatform(const IConsoleVariable* Glo
 	return bEnabled;
 }
 
-bool ShouldKeepShaderDebugInfo(EShaderPlatform ShaderPlatform)
+bool ShouldKeepShaderDebugInfo(FName ShaderFormat)
 {
 	static const TCHAR* const KeepDebugInfoName = TEXT("r.Shaders.KeepDebugInfo");
 	static const IConsoleVariable* KeepDebugInfoCvar = IConsoleManager::Get().FindConsoleVariable(KeepDebugInfoName);
@@ -316,11 +316,16 @@ bool ShouldKeepShaderDebugInfo(EShaderPlatform ShaderPlatform)
 	// PrepareExportedDebugInfo is only read from the platform settings, not the global cvar
 	static const TCHAR* const PrepareExportedDebugInfoName = TEXT("r.Shaders.PrepareExportedDebugInfo");
 
-	return IsShaderCompilerConfigEnabledForPlatform(KeepDebugInfoCvar, KeepDebugInfoName, ShaderPlatform)
-		|| IsShaderCompilerConfigEnabledForPlatform(nullptr, PrepareExportedDebugInfoName, ShaderPlatform);
+	return IsShaderCompilerConfigEnabledForPlatform(KeepDebugInfoCvar, KeepDebugInfoName, ShaderFormat)
+		|| IsShaderCompilerConfigEnabledForPlatform(nullptr, PrepareExportedDebugInfoName, ShaderFormat);
 }
 
-bool ShouldExportShaderDebugInfo(EShaderPlatform ShaderPlatform)
+bool ShouldKeepShaderDebugInfo(EShaderPlatform ShaderPlatform)
+{
+	return ShouldKeepShaderDebugInfo(LegacyShaderPlatformToShaderFormat(ShaderPlatform));
+}
+
+bool ShouldExportShaderDebugInfo(FName ShaderFormat)
 {
 	// r.DumpShaderDebugInfo should also turn on ExportShaderDebugInfo
 	// The difference is that r.DumpShaderDebugInfo will also output engine debug files such as converted hlsl or SCW helper files.
@@ -333,21 +338,26 @@ bool ShouldExportShaderDebugInfo(EShaderPlatform ShaderPlatform)
 
 	static const TCHAR* const KeepDebugInfoName = TEXT("r.Shaders.KeepDebugInfo");
 	static const IConsoleVariable* KeepDebugInfoCvar = IConsoleManager::Get().FindConsoleVariable(KeepDebugInfoName);
-	if (IsShaderCompilerConfigEnabledForPlatform(KeepDebugInfoCvar, KeepDebugInfoName, ShaderPlatform))
+	if (IsShaderCompilerConfigEnabledForPlatform(KeepDebugInfoCvar, KeepDebugInfoName, ShaderFormat))
 	{
 		return true;
 	}
 
 	static const TCHAR* const ExportDebugInfoName = TEXT("r.Shaders.ExportDebugInfo");
 	static const IConsoleVariable* ExportDebugInfoCvar = IConsoleManager::Get().FindConsoleVariable(ExportDebugInfoName);
-	return IsShaderCompilerConfigEnabledForPlatform(ExportDebugInfoCvar, ExportDebugInfoName, ShaderPlatform);
+	return IsShaderCompilerConfigEnabledForPlatform(ExportDebugInfoCvar, ExportDebugInfoName, ShaderFormat);
+}
+
+bool ShouldExportShaderDebugInfo(EShaderPlatform ShaderPlatform)
+{
+	return ShouldExportShaderDebugInfo(LegacyShaderPlatformToShaderFormat(ShaderPlatform));
 }
 
-bool ShouldAllowUniqueDebugInfo(EShaderPlatform ShaderPlatform)
+bool ShouldAllowUniqueDebugInfo(FName ShaderFormat)
 {
 	static const TCHAR* const VariableName = TEXT("r.Shaders.AllowUniqueDebugInfo");
 	static const IConsoleVariable* GlobalCvar = IConsoleManager::Get().FindConsoleVariable(VariableName);
-	return IsShaderCompilerConfigEnabledForPlatform(GlobalCvar, VariableName, ShaderPlatform);
+	return IsShaderCompilerConfigEnabledForPlatform(GlobalCvar, VariableName, ShaderFormat);
 }
 
 bool FShaderParameterMap::FindParameterAllocation(const TCHAR* ParameterName,uint16& OutBufferIndex,uint16& OutBaseIndex,uint16& OutSize) const
diff --git a/Engine/Source/Runtime/RenderCore/Public/ShaderCompilerCore.h b/Engine/Source/Runtime/RenderCore/Public/ShaderCompilerCore.h
--- a/Engine/Source/Runtime/RenderCore/Public/ShaderCompilerCore.h
+++ b/Engine/Source/Runtime/RenderCore/Public/ShaderCompilerCore.h
@@ -26,9 +26,14 @@ const int32 ShaderCompileWorkerPipelineJobHeader = 'P';
 
 /** Returns true if debug information should be kept for a given platform. */
 extern RENDERCORE_API bool ShouldKeepShaderDebugInfo(EShaderPlatform Platform);
+extern RENDERCORE_API bool ShouldKeepShaderDebugInfo(FName ShaderFormat);
 
 /** Returns true if debug information should be exported to separate files for a given platform . */
 extern RENDERCORE_API bool ShouldExportShaderDebugInfo(EShaderPlatform Platform);
+extern RENDERCORE_API bool ShouldExportShaderDebugInfo(FName ShaderFormat);
+
+/** Returns true if (external) debug information should be specific to each shader rather than be deduplicated. */
+extern RENDERCORE_API bool ShouldAllowUniqueDebugInfo(FName ShaderFormat);
 
 enum ECompilerFlags
 {
diff --git a/Engine/Source/Runtime/RenderCore/Public/ShaderCore.h b/Engine/Source/Runtime/RenderCore/Public/ShaderCore.h
--- a/Engine/Source/Runtime/RenderCore/Public/ShaderCore.h
+++ b/Engine/Source/Runtime/RenderCore/Public/ShaderCore.h
@@ -99,15 +99,6 @@ extern RENDERCORE_API bool AllowDebugViewmodes();
 /** Returns true if debug viewmodes are allowed for the given platform. */
 extern RENDERCORE_API bool AllowDebugViewmodes(EShaderPlatform Platform);
 
-/** Returns true if debug information should be kept for a given platform. */
-extern RENDERCORE_API bool ShouldKeepShaderDebugInfo(EShaderPlatform Platform);
-
-/** Returns true if debug information should be exported to separate files for a given platform . */
-extern RENDERCORE_API bool ShouldExportShaderDebugInfo(EShaderPlatform Platform);
-
-/** Returns true if (external) debug information should be specific to each shader rather than be deduplicated. */
-extern RENDERCORE_API bool ShouldAllowUniqueDebugInfo(EShaderPlatform Platform);
-
 struct FShaderTarget
 {
 	// The rest of uint32 holding the bitfields can be left unitialized. Union with a uint32 serves to prevent that to be able to set the whole uint32 value
