UE5 - Dont use VirtualShadowMaps on non-Nanite platforms

diff --git a/Engine/Source/Runtime/RenderCore/Public/RenderUtils.h b/Engine/Source/Runtime/RenderCore/Public/RenderUtils.h
--- a/Engine/Source/Runtime/RenderCore/Public/RenderUtils.h
+++ b/Engine/Source/Runtime/RenderCore/Public/RenderUtils.h
@@ -631,5 +631,5 @@ inline bool UseVirtualShadowMaps(EShaderPlatform ShaderPlatform, const FStaticFe
 	// bool bNaniteSupported = DoesPlatformSupportNanite(ShaderPlatform);
 	bool bForwardShadingEnabled = IsForwardShadingEnabled(ShaderPlatform);
 
-	return UseGPUScene(ShaderPlatform, FeatureLevel) && !bForwardShadingEnabled && EnableVirtualSMCVar->GetInt() != 0;
+	return UseGPUScene(ShaderPlatform, FeatureLevel) && !bForwardShadingEnabled && EnableVirtualSMCVar->GetInt() != 0 && DoesPlatformSupportNanite(ShaderPlatform);
 }
