Adds support for wildcard in VisualizeTexture tool and fix bug where some where not listed

diff --git a/Engine/Source/Runtime/RenderCore/Private/VisualizeTexture.cpp b/Engine/Source/Runtime/RenderCore/Private/VisualizeTexture.cpp
--- a/Engine/Source/Runtime/RenderCore/Private/VisualizeTexture.cpp
+++ b/Engine/Source/Runtime/RenderCore/Private/VisualizeTexture.cpp
@@ -15,179 +15,232 @@
 void FVisualizeTexture::ParseCommands(const TCHAR* Cmd, FOutputDevice &Ar)
 {
 #if SUPPORTS_VISUALIZE_TEXTURE
-	Config = {};
-	uint32 ParameterCount = 0;
-	for (;;)
+	// Find out what command to do based on first parameter.
+	ECommand Command = ECommand::Unknown;
+	FString ResourceName;
+	TOptional<uint32> ResourceVersion;
+	TOptional<FWildcardString> ResourceListWildCard;
 	{
-		FString Parameter = FParse::Token(Cmd, 0);
+		FString FirstParameter = FParse::Token(Cmd, 0);
+
 
-		if (Parameter.IsEmpty())
+		if (FirstParameter.IsEmpty())
+		{
+			// NOP
+		}
+		else if (FChar::IsDigit(**FirstParameter))
+		{
+			Command = ECommand::DisableVisualization;
+		}
+		else if (FirstParameter == TEXT("help"))
+		{
+			Command = ECommand::DisplayHelp;
+		}
+		else if (FirstParameter == TEXT("pool"))
 		{
-			break;
+			Command = ECommand::DisplayPoolResourceList;
 		}
-		else if (ParameterCount == 0)
+		else
 		{
-			if (!FChar::IsDigit(**Parameter))
+			const TCHAR* AfterAt = *FirstParameter;
+
+			while (*AfterAt != 0 && *AfterAt != TCHAR('@'))
+			{
+				++AfterAt;
+			}
+
+			if (*AfterAt == TCHAR('@'))
+			{
+				ResourceName = FirstParameter.Left(AfterAt - *FirstParameter);
+				ResourceVersion = FCString::Atoi(AfterAt + 1);
+			}
+			else
+			{
+				ResourceName = FirstParameter;
+			}
+
+			if (ResourceName.Contains(TEXT("*")))
+			{
+				ResourceListWildCard = FWildcardString(ResourceName);
+				Command = ECommand::DisplayResourceList;
+			}
+			else
 			{
-				const TCHAR* AfterAt = *Parameter;
+				Command = ECommand::VisualizeResource;
+				Visualize(ResourceName, ResourceVersion);
+			}
+		}
+	}
+
+	if (Command == ECommand::Unknown)
+	{
+		FVisualizeTexture::DisplayHelp(Ar);
+		DisplayResourceListToLog(TOptional<FWildcardString>());
+	}
+	else if (Command == ECommand::DisplayHelp)
+	{
+		FVisualizeTexture::DisplayHelp(Ar);
+	}
+	else if (Command == ECommand::DisableVisualization)
+	{
+		Visualize({});
+	}
+	else if (Command == ECommand::VisualizeResource)
+	{
+		Config = {};
+		Visualize(ResourceName, ResourceVersion);
+		for (;;)
+		{
+			FString Parameter = FParse::Token(Cmd, 0);
 
-				while (*AfterAt != 0 && *AfterAt != TCHAR('@'))
+			if (Parameter.IsEmpty())
+			{
+				break;
+			}
+			else if (Parameter == TEXT("uv0"))
+			{
+				Config.InputUVMapping = EInputUVMapping::LeftTop;
+			}
+			else if (Parameter == TEXT("uv1"))
+			{
+				Config.InputUVMapping = EInputUVMapping::Whole;
+			}
+			else if (Parameter == TEXT("uv2"))
+			{
+				Config.InputUVMapping = EInputUVMapping::PixelPerfectCenter;
+			}
+			else if (Parameter == TEXT("pip"))
+			{
+				Config.InputUVMapping = EInputUVMapping::PictureInPicture;
+			}
+			else if (Parameter == TEXT("bmp"))
+			{
+				Config.Flags |= EFlags::SaveBitmap;
+			}
+			else if (Parameter == TEXT("stencil"))
+			{
+				Config.Flags |= EFlags::SaveBitmapAsStencil;
+			}
+			else if (Parameter == TEXT("frac"))
+			{
+				Config.ShaderOp = EShaderOp::Frac;
+			}
+			else if (Parameter == TEXT("sat"))
+			{
+				Config.ShaderOp = EShaderOp::Saturate;
+			}
+			else if (Parameter.Left(3) == TEXT("mip"))
+			{
+				Parameter.RightInline(Parameter.Len() - 3, false);
+				Config.MipIndex = FCString::Atoi(*Parameter);
+			}
+			else if (Parameter.Left(5) == TEXT("index"))
+			{
+				Parameter.RightInline(Parameter.Len() - 5, false);
+				Config.ArrayIndex = FCString::Atoi(*Parameter);
+			}
+			// e.g. RGB*6, A, *22, /2.7, A*7
+			else if (Parameter.Left(3) == TEXT("rgb")
+				|| Parameter.Left(1) == TEXT("a")
+				|| Parameter.Left(1) == TEXT("r")
+				|| Parameter.Left(1) == TEXT("g")
+				|| Parameter.Left(1) == TEXT("b")
+				|| Parameter.Left(1) == TEXT("*")
+				|| Parameter.Left(1) == TEXT("/"))
+			{
+				Config.SingleChannel = -1;
+
+				if (Parameter.Left(3) == TEXT("rgb"))
 				{
-					++AfterAt;
+					Parameter.RightInline(Parameter.Len() - 3, false);
 				}
+				else if (Parameter.Left(1) == TEXT("r")) Config.SingleChannel = 0;
+				else if (Parameter.Left(1) == TEXT("g")) Config.SingleChannel = 1;
+				else if (Parameter.Left(1) == TEXT("b")) Config.SingleChannel = 2;
+				else if (Parameter.Left(1) == TEXT("a")) Config.SingleChannel = 3;
+				if (Config.SingleChannel >= 0)
+				{
+					Parameter.RightInline(Parameter.Len() - 1, false);
+					Config.SingleChannelMul = 1.0f;
+					Config.RGBMul = 0.0f;
+				}
+
+				float Mul = 1.0f;
 
-				if (*AfterAt == TCHAR('@'))
+				// * or /
+				if (Parameter.Left(1) == TEXT("*"))
 				{
-					FString NameWithoutAt = Parameter.Left(AfterAt - *Parameter);
-					Visualize(*NameWithoutAt, FCString::Atoi(AfterAt + 1));
+					Parameter.RightInline(Parameter.Len() - 1, false);
+					Mul = FCString::Atof(*Parameter);
 				}
-				else
+				else if (Parameter.Left(1) == TEXT("/"))
 				{
-					Visualize(*Parameter);
+					Parameter.RightInline(Parameter.Len() - 1, false);
+					Mul = 1.0f / FCString::Atof(*Parameter);
 				}
+				Config.RGBMul *= Mul;
+				Config.SingleChannelMul *= Mul;
+				Config.AMul *= Mul;
 			}
 			else
 			{
-				Visualize({});
+				Ar.Logf(TEXT("Error: parameter \"%s\" not recognized"), *Parameter);
 			}
 		}
-		else if (Parameter == TEXT("fulllist") || Parameter == TEXT("full"))
-		{
-			Config.Flags |= EFlags::FullList;
-		}
-		else if (Parameter == TEXT("byname"))
-		{
-			Config.SortBy = ESortBy::Name;
-		}
-		else if (Parameter == TEXT("bysize"))
-		{
-			Config.SortBy = ESortBy::Size;
-		}
-		else if (Parameter == TEXT("uv0"))
-		{
-			Config.InputUVMapping = EInputUVMapping::LeftTop;
-		}
-		else if (Parameter == TEXT("uv1"))
-		{
-			Config.InputUVMapping = EInputUVMapping::Whole;
-		}
-		else if (Parameter == TEXT("uv2"))
-		{
-			Config.InputUVMapping = EInputUVMapping::PixelPerfectCenter;
-		}
-		else if (Parameter == TEXT("pip"))
-		{
-			Config.InputUVMapping = EInputUVMapping::PictureInPicture;
-		}
-		else if (Parameter == TEXT("bmp"))
-		{
-			Config.Flags |= EFlags::SaveBitmap;
-		}
-		else if (Parameter == TEXT("stencil"))
-		{
-			Config.Flags |= EFlags::SaveBitmapAsStencil;
-		}
-		else if (Parameter == TEXT("frac"))
-		{
-			Config.ShaderOp = EShaderOp::Frac;
-		}
-		else if (Parameter == TEXT("sat"))
-		{
-			Config.ShaderOp = EShaderOp::Saturate;
-		}
-		else if (Parameter.Left(3) == TEXT("mip"))
-		{
-			Parameter.RightInline(Parameter.Len() - 3, false);
-			Config.MipIndex = FCString::Atoi(*Parameter);
-		}
-		else if (Parameter.Left(5) == TEXT("index"))
-		{
-			Parameter.RightInline(Parameter.Len() - 5, false);
-			Config.ArrayIndex = FCString::Atoi(*Parameter);
-		}
-		// e.g. RGB*6, A, *22, /2.7, A*7
-		else if (Parameter.Left(3) == TEXT("rgb")
-			|| Parameter.Left(1) == TEXT("a")
-			|| Parameter.Left(1) == TEXT("r")
-			|| Parameter.Left(1) == TEXT("g")
-			|| Parameter.Left(1) == TEXT("b")
-			|| Parameter.Left(1) == TEXT("*")
-			|| Parameter.Left(1) == TEXT("/"))
+
+	}
+	else if (Command == ECommand::DisplayPoolResourceList)
+	{
+		ESortBy SortBy = ESortBy::Index;
+
+		for (;;)
 		{
-			Config.SingleChannel = -1;
+			FString Parameter = FParse::Token(Cmd, 0);
 
-			if (Parameter.Left(3) == TEXT("rgb"))
+			if (Parameter.IsEmpty())
 			{
-				Parameter.RightInline(Parameter.Len() - 3, false);
+				break;
 			}
-			else if (Parameter.Left(1) == TEXT("r")) Config.SingleChannel = 0;
-			else if (Parameter.Left(1) == TEXT("g")) Config.SingleChannel = 1;
-			else if (Parameter.Left(1) == TEXT("b")) Config.SingleChannel = 2;
-			else if (Parameter.Left(1) == TEXT("a")) Config.SingleChannel = 3;
-			if (Config.SingleChannel >= 0)
+			else if (Parameter == TEXT("byname"))
 			{
-				Parameter.RightInline(Parameter.Len() - 1, false);
-				Config.SingleChannelMul = 1.0f;
-				Config.RGBMul = 0.0f;
+				SortBy = ESortBy::Name;
 			}
-
-			float Mul = 1.0f;
-
-			// * or /
-			if (Parameter.Left(1) == TEXT("*"))
+			else if (Parameter == TEXT("bysize"))
 			{
-				Parameter.RightInline(Parameter.Len() - 1, false);
-				Mul = FCString::Atof(*Parameter);
+				SortBy = ESortBy::Size;
 			}
-			else if (Parameter.Left(1) == TEXT("/"))
+			else
 			{
-				Parameter.RightInline(Parameter.Len() - 1, false);
-				Mul = 1.0f / FCString::Atof(*Parameter);
+				Ar.Logf(TEXT("Error: parameter \"%s\" not recognized"), *Parameter);
 			}
-			Config.RGBMul *= Mul;
-			Config.SingleChannelMul *= Mul;
-			Config.AMul *= Mul;
 		}
-		else
+
+		DisplayPoolResourceListToLog(SortBy);
+	}
+	else if (Command == ECommand::DisplayResourceList)
+	{
+		bool bListAllocated = false;
+		ESortBy SortBy = ESortBy::Index;
+
+		for (;;)
 		{
-			Ar.Logf(TEXT("Error: parameter \"%s\" not recognized"), *Parameter);
+			FString Parameter = FParse::Token(Cmd, 0);
+
+			if (Parameter.IsEmpty())
+			{
+				break;
+			}
+			else
+			{
+				Ar.Logf(TEXT("Error: parameter \"%s\" not recognized"), *Parameter);
+			}
 		}
 
-		++ParameterCount;
+		DisplayResourceListToLog(ResourceListWildCard);
 	}
-
-	if (!ParameterCount)
+	else
 	{
-		Ar.Logf(TEXT("VisualizeTexture/Vis <CheckpointName> [<Mode>] [PIP/UV0/UV1/UV2] [BMP] [FRAC/SAT] [FULL]:"));
-		Ar.Logf(TEXT("Mode (examples):"));
-		Ar.Logf(TEXT("  RGB      = RGB in range 0..1 (default)"));
-		Ar.Logf(TEXT("  *8       = RGB * 8"));
-		Ar.Logf(TEXT("  A        = alpha channel in range 0..1"));
-		Ar.Logf(TEXT("  R        = red channel in range 0..1"));
-		Ar.Logf(TEXT("  G        = green channel in range 0..1"));
-		Ar.Logf(TEXT("  B        = blue channel in range 0..1"));
-		Ar.Logf(TEXT("  A*16     = Alpha * 16"));
-		Ar.Logf(TEXT("  RGB/2    = RGB / 2"));
-		Ar.Logf(TEXT("SubResource:"));
-		Ar.Logf(TEXT("  MIP5     = Mip level 5 (0 is default)"));
-		Ar.Logf(TEXT("  INDEX5   = Array Element 5 (0 is default)"));
-		Ar.Logf(TEXT("InputMapping:"));
-		Ar.Logf(TEXT("  PIP      = like UV1 but as picture in picture with normal rendering  (default)"));
-		Ar.Logf(TEXT("  UV0      = UV in left top"));
-		Ar.Logf(TEXT("  UV1      = full texture"));
-		Ar.Logf(TEXT("  UV2      = pixel perfect centered"));
-		Ar.Logf(TEXT("Flags:"));
-		Ar.Logf(TEXT("  BMP      = save out bitmap to the screenshots folder (not on console, normalized)"));
-		Ar.Logf(TEXT("STENCIL    = Stencil normally displayed in alpha channel of depth.  This option is used for BMP to get a stencil only BMP."));
-		Ar.Logf(TEXT("  FRAC     = use frac() in shader (default)"));
-		Ar.Logf(TEXT("  SAT      = use saturate() in shader"));
-		Ar.Logf(TEXT("  FULLLIST = show full list, otherwise we hide some textures in the printout"));
-		Ar.Logf(TEXT("  BYNAME   = sort list by name"));
-		Ar.Logf(TEXT("  BYSIZE   = show list by size"));
-		Ar.Logf(TEXT("TextureId:"));
-		Ar.Logf(TEXT("  0        = <off>"));
-
-		DebugLog(EDebugLogVerbosity::Extended);
+		unimplemented();
 	}
 #endif
 }
@@ -195,9 +248,8 @@ void FVisualizeTexture::ParseCommands(const TCHAR* Cmd, FOutputDevice &Ar)
 void FVisualizeTexture::DebugLogOnCrash()
 {
 #if SUPPORTS_VISUALIZE_TEXTURE
-	Config.SortBy = ESortBy::Size;
-	Config.Flags |= EFlags::FullList;
-	DebugLog(EDebugLogVerbosity::Default);
+	DisplayPoolResourceListToLog(ESortBy::Size);
+	DisplayResourceListToLog(TOptional<FWildcardString>());
 #endif
 }
 
@@ -230,223 +282,265 @@ TGlobalResource<FVisualizeTexture> GVisualizeTexture;
 
 #if SUPPORTS_VISUALIZE_TEXTURE
 
-void FVisualizeTexture::DebugLog(EDebugLogVerbosity Verbosity)
+void FVisualizeTexture::DisplayHelp(FOutputDevice &Ar)
 {
+	Ar.Logf(TEXT("VisualizeTexture/Vis <RDGResourceNameWildcard>:"));
+	Ar.Logf(TEXT("  Lists all RDG resource names with wildcard filtering."));
+	Ar.Logf(TEXT(""));
+	Ar.Logf(TEXT("VisualizeTexture/Vis <RDGResourceName>[@<Version>] [<Mode>] [PIP/UV0/UV1/UV2] [BMP] [FRAC/SAT] [FULL]:"));
+	Ar.Logf(TEXT("  RDGResourceName = Name of the resource set when creating it with RDG."));
+	Ar.Logf(TEXT("  Version = Integer to specify a specific intermediate version."));
+	Ar.Logf(TEXT("  Mode (examples):"));
+	Ar.Logf(TEXT("    RGB      = RGB in range 0..1 (default)"));
+	Ar.Logf(TEXT("    *8       = RGB * 8"));
+	Ar.Logf(TEXT("    A        = alpha channel in range 0..1"));
+	Ar.Logf(TEXT("    R        = red channel in range 0..1"));
+	Ar.Logf(TEXT("    G        = green channel in range 0..1"));
+	Ar.Logf(TEXT("    B        = blue channel in range 0..1"));
+	Ar.Logf(TEXT("    A*16     = Alpha * 16"));
+	Ar.Logf(TEXT("    RGB/2    = RGB / 2"));
+	Ar.Logf(TEXT("  SubResource:"));
+	Ar.Logf(TEXT("    MIP5     = Mip level 5 (0 is default)"));
+	Ar.Logf(TEXT("    INDEX5   = Array Element 5 (0 is default)"));
+	Ar.Logf(TEXT("  InputMapping:"));
+	Ar.Logf(TEXT("    PIP      = like UV1 but as picture in picture with normal rendering  (default)"));
+	Ar.Logf(TEXT("    UV0      = UV in left top"));
+	Ar.Logf(TEXT("    UV1      = full texture"));
+	Ar.Logf(TEXT("    UV2      = pixel perfect centered"));
+	Ar.Logf(TEXT("  Flags:"));
+	Ar.Logf(TEXT("    BMP      = save out bitmap to the screenshots folder (not on console, normalized)"));
+	Ar.Logf(TEXT("    STENCIL  = Stencil normally displayed in alpha channel of depth.  This option is used for BMP to get a stencil only BMP."));
+	Ar.Logf(TEXT("    FRAC     = use frac() in shader (default)"));
+	Ar.Logf(TEXT("    SAT      = use saturate() in shader"));
+	Ar.Logf(TEXT(""));
+	Ar.Logf(TEXT("VisualizeTexture/Vis 0"));
+	Ar.Logf(TEXT("  Stops visualizing a resource."));
+	Ar.Logf(TEXT(""));
+	Ar.Logf(TEXT("VisualizeTexture/Vis pool [BYNAME/BYSIZE]:"));
+	Ar.Logf(TEXT("  Shows list of all resources in the pool."));
+	Ar.Logf(TEXT("  BYNAME   = sort pool list by name"));
+	Ar.Logf(TEXT("  BYSIZE   = show pool list by size"));
+	Ar.Logf(TEXT(""));
+}
+
+void FVisualizeTexture::DisplayPoolResourceListToLog(FVisualizeTexture::ESortBy SortBy)
+{
+	struct FSortedLines
 	{
-		struct FSortedLines
-		{
-			FString Line;
-			int32   SortIndex = 0;
-			uint32  PoolIndex = 0;
+		FString Line;
+		int32   SortIndex = 0;
+		uint32  PoolIndex = 0;
 
-			bool operator < (const FSortedLines &B) const
+		bool operator < (const FSortedLines& B) const
+		{
+			// first large ones
+			if (SortIndex < B.SortIndex)
 			{
-				// first large ones
-				if (SortIndex < B.SortIndex)
-				{
-					return true;
-				}
-				if (SortIndex > B.SortIndex)
-				{
-					return false;
-				}
-
-				return Line < B.Line;
+				return true;
 			}
-		};
+			if (SortIndex > B.SortIndex)
+			{
+				return false;
+			}
+
+			return Line < B.Line;
+		}
+	};
 
-		TArray<FSortedLines> SortedLines;
+	TArray<FSortedLines> SortedLines;
 
-		for (uint32 Index = 0, Num = GRenderTargetPool.GetElementCount(); Index < Num; ++Index)
+	for (uint32 Index = 0, Num = GRenderTargetPool.GetElementCount(); Index < Num; ++Index)
+	{
+		FPooledRenderTarget* RenderTarget = GRenderTargetPool.GetElementById(Index);
+
+		if (!RenderTarget)
 		{
-			FPooledRenderTarget* RenderTarget = GRenderTargetPool.GetElementById(Index);
+			continue;
+		}
 
-			if (!RenderTarget)
-			{
-				continue;
-			}
+		const FPooledRenderTargetDesc Desc = RenderTarget->GetDesc();
 
-			const bool bFullList = EnumHasAnyFlags(Config.Flags, EFlags::FullList);
+		const uint32 SizeInKB = (RenderTarget->ComputeMemorySize() + 1023) / 1024;
 
-			const FPooledRenderTargetDesc Desc = RenderTarget->GetDesc();
+		FString UnusedStr;
 
-			if (bFullList || (Desc.Flags & TexCreate_HideInVisualizeTexture) == 0)
-			{
-				const uint32 SizeInKB = (RenderTarget->ComputeMemorySize() + 1023) / 1024;
+		if (RenderTarget->GetUnusedForNFrames() > 0)
+		{
+			UnusedStr = FString::Printf(TEXT(" unused(%d)"), RenderTarget->GetUnusedForNFrames());
+		}
 
-				FString UnusedStr;
+		FSortedLines Element;
+		Element.PoolIndex = Index;
+		Element.SortIndex = Index;
 
-				if (RenderTarget->GetUnusedForNFrames() > 0)
-				{
-					if (!bFullList)
-					{
-						continue;
-					}
+		FString InfoString = Desc.GenerateInfoString();
 
-					UnusedStr = FString::Printf(TEXT(" unused(%d)"), RenderTarget->GetUnusedForNFrames());
-				}
+		switch (SortBy)
+		{
+		case ESortBy::Index:
+		{
+			// Constant works well with the average name length
+			const uint32 TotalSpacerSize = 36;
+			const uint32 SpaceCount = FMath::Max<int32>(0, TotalSpacerSize - InfoString.Len());
 
-				FSortedLines Element;
-				Element.PoolIndex = Index;
-				Element.SortIndex = Index;
+			for (uint32 Space = 0; Space < SpaceCount; ++Space)
+			{
+				InfoString.AppendChar((TCHAR)' ');
+			}
 
-				FString InfoString = Desc.GenerateInfoString();
+			// Sort by index
+			Element.Line = FString::Printf(TEXT("%s %s %d KB%s"), *InfoString, Desc.DebugName, SizeInKB, *UnusedStr);
+		}
+		break;
 
-				switch (Config.SortBy)
-				{
-				case ESortBy::Index:
-				{
-					// Constant works well with the average name length
-					const uint32 TotalSpacerSize = 36;
-					const uint32 SpaceCount = FMath::Max<int32>(0, TotalSpacerSize - InfoString.Len());
+		case ESortBy::Name:
+		{
+			Element.Line = FString::Printf(TEXT("%s %s %d KB%s"), Desc.DebugName, *InfoString, SizeInKB, *UnusedStr);
+			Element.SortIndex = 0;
+		}
+		break;
 
-					for (uint32 Space = 0; Space < SpaceCount; ++Space)
-					{
-						InfoString.AppendChar((TCHAR)' ');
-					}
+		case ESortBy::Size:
+		{
+			Element.Line = FString::Printf(TEXT("%d KB %s %s%s"), SizeInKB, *InfoString, Desc.DebugName, *UnusedStr);
+			Element.SortIndex = -(int32)SizeInKB;
+		}
+		break;
 
-					// Sort by index
-					Element.Line = FString::Printf(TEXT("%s %s %d KB%s"), *InfoString, Desc.DebugName, SizeInKB, *UnusedStr);
-				}
-				break;
+		default:
+			checkNoEntry();
+		}
 
-				case ESortBy::Name:
-				{
-					Element.Line = FString::Printf(TEXT("%s %s %d KB%s"), Desc.DebugName, *InfoString, SizeInKB, *UnusedStr);
-					Element.SortIndex = 0;
-				}
-				break;
+		if (Desc.Flags & TexCreate_FastVRAM)
+		{
+			FRHIResourceInfo Info;
 
-				case ESortBy::Size:
-				{
-					Element.Line = FString::Printf(TEXT("%d KB %s %s%s"), SizeInKB, *InfoString, Desc.DebugName, *UnusedStr);
-					Element.SortIndex = -(int32)SizeInKB;
-				}
-				break;
+			FTextureRHIRef Texture = RenderTarget->GetRenderTargetItem().ShaderResourceTexture;
 
-				default:
-					checkNoEntry();
-				}
+			if (!IsValidRef(Texture))
+			{
+				Texture = RenderTarget->GetRenderTargetItem().TargetableTexture;
+			}
 
-				if (Desc.Flags & TexCreate_FastVRAM)
-				{
-					FRHIResourceInfo Info;
-
-					FTextureRHIRef Texture = RenderTarget->GetRenderTargetItem().ShaderResourceTexture;
-
-					if (!IsValidRef(Texture))
-					{
-						Texture = RenderTarget->GetRenderTargetItem().TargetableTexture;
-					}
-
-					if (IsValidRef(Texture))
-					{
-						RHIGetResourceInfo(Texture, Info);
-					}
-
-					if (Info.VRamAllocation.AllocationSize)
-					{
-						// Note we do KB for more readable numbers but this can cause quantization loss
-						Element.Line += FString::Printf(TEXT(" VRamInKB(Start/Size):%d/%d"),
-							Info.VRamAllocation.AllocationStart / 1024,
-							(Info.VRamAllocation.AllocationSize + 1023) / 1024);
-					}
-					else
-					{
-						Element.Line += TEXT(" VRamInKB(Start/Size):<NONE>");
-					}
-				}
+			if (IsValidRef(Texture))
+			{
+				RHIGetResourceInfo(Texture, Info);
+			}
 
-				SortedLines.Add(Element);
+			if (Info.VRamAllocation.AllocationSize)
+			{
+				// Note we do KB for more readable numbers but this can cause quantization loss
+				Element.Line += FString::Printf(TEXT(" VRamInKB(Start/Size):%d/%d"),
+					Info.VRamAllocation.AllocationStart / 1024,
+					(Info.VRamAllocation.AllocationSize + 1023) / 1024);
+			}
+			else
+			{
+				Element.Line += TEXT(" VRamInKB(Start/Size):<NONE>");
 			}
 		}
 
-		SortedLines.Sort();
+		SortedLines.Add(Element);
+	}
 
-		for (int32 Index = 0; Index < SortedLines.Num(); Index++)
-		{
-			const FSortedLines& Entry = SortedLines[Index];
+	SortedLines.Sort();
 
-			UE_LOG(LogConsoleResponse, Log, TEXT("   %3d = %s"), Entry.PoolIndex + 1, *Entry.Line);
-		}
+	for (int32 Index = 0; Index < SortedLines.Num(); Index++)
+	{
+		const FSortedLines& Entry = SortedLines[Index];
+
+		UE_LOG(LogConsoleResponse, Log, TEXT("   %3d = %s"), Entry.PoolIndex + 1, *Entry.Line);
 	}
 
 	UE_LOG(LogConsoleResponse, Log, TEXT(""));
 
-	if (Verbosity == EDebugLogVerbosity::Extended)
-	{
-		UE_LOG(LogConsoleResponse, Log, TEXT("CheckpointName (what was rendered this frame, use <Name>@<Number> to get intermediate versions):"));
+	uint32 WholeCount;
+	uint32 WholePoolInKB;
+	uint32 UsedInKB;
+	GRenderTargetPool.GetStats(WholeCount, WholePoolInKB, UsedInKB);
 
-		TArray<FString> Entries;
-		Entries.Reserve(VersionCountMap.Num());
-		for (auto KV : VersionCountMap)
-		{
-			Entries.Add(KV.Key);
-		}
-		Entries.Sort();
+	UE_LOG(LogConsoleResponse, Log, TEXT("Pool: %d/%d MB (referenced/allocated)"), (UsedInKB + 1023) / 1024, (WholePoolInKB + 1023) / 1024);
 
-		// Magic number works well with the name length we have
-		const uint32 ColumnCount = 5;
-		const uint32 SpaceBetweenColumns = 1;
-		const uint32 ColumnHeight = FMath::DivideAndRoundUp((uint32)Entries.Num(), ColumnCount);
+	UE_LOG(LogConsoleResponse, Log, TEXT(""));
+}
 
-		// Width of the column in characters, init with 0
-		uint32 ColumnWidths[ColumnCount] = {};
+void FVisualizeTexture::DisplayResourceListToLog(const TOptional<FWildcardString>& Wildcard)
+{
+	UE_LOG(LogConsoleResponse, Log, TEXT("RDGResourceName (what was rendered this frame, use <RDGResourceName>@<Version> to get intermediate versions):"));
 
-		for (uint32 Index = 0, Count = Entries.Num(); Index < Count; ++Index)
+	TArray<FString> Entries;
+	Entries.Reserve(VersionCountMap.Num());
+	for (auto KV : VersionCountMap)
+	{
+		if (Wildcard.IsSet())
 		{
-			const FString& Entry = *Entries[Index];
-			const uint32 Column = Index / ColumnHeight;
-			ColumnWidths[Column] = FMath::Max(ColumnWidths[Column], (uint32)Entry.Len());
+			if (Wildcard->IsMatch(KV.Key))
+			{
+				Entries.Add(KV.Key);
+			}
 		}
-
-		// Print them sorted, if possible multiple in a line
+		else
 		{
-			FString Line;
+			Entries.Add(KV.Key);
+		}
+	}
+	Entries.Sort();
 
-			for (uint32 OutputIndex = 0, Count = Entries.Num(); OutputIndex < Count; ++OutputIndex)
-			{
-				// [0, ColumnCount-1]
-				const uint32 Column = OutputIndex % ColumnCount;
-				const uint32 Row = OutputIndex / ColumnCount;
-				const uint32 Index = Row + Column * ColumnHeight;
+	// Magic number works well with the name length we have
+	const int32 MaxColumnCount = 5;
+	const int32 SpaceBetweenColumns = 1;
+	const int32 TargetColumnHeight = 8;
 
-				bool bLineEnd = true;
+	int32 ColumnCount = FMath::Clamp(FMath::DivideAndRoundUp(Entries.Num(), TargetColumnHeight), 1, MaxColumnCount);
+	int32 ColumnHeight = FMath::DivideAndRoundUp(Entries.Num(), ColumnCount);
 
-				if (Index < Count)
-				{
-					bLineEnd = (Column + 1 == ColumnCount);
+	// Width of the column in characters, init with 0
+	TStaticArray<int32, MaxColumnCount> ColumnWidths;
+	for (int32 ColumnId = 0; ColumnId < MaxColumnCount; ColumnId++)
+	{
+		ColumnWidths[ColumnId] = 0;
+	}
 
-					// Order per column for readability.
-					const FString& Entry = *Entries[Index];
+	for (int32 Index = 0, Count = Entries.Num(); Index < Count; ++Index)
+	{
+		const FString& Entry = *Entries[Index];
+		const int32 Column = Index / ColumnHeight;
+		ColumnWidths[Column] = FMath::Max(ColumnWidths[Column], Entry.Len());
+	}
 
-					Line += Entry;
+	// Print them sorted, if possible multiple in a line
+	for (int32 RowId = 0; RowId < ColumnHeight; RowId++)
+	{
+		FString Line;
+		int32 ColumnAlignment = 0;
 
-					const int32 SpaceCount = ColumnWidths[Column] + SpaceBetweenColumns - Entry.Len();
-					checkf(SpaceCount >= 0, TEXT("A previous pass produced bad data"));
+		for (int32 ColumnId = 0; ColumnId < ColumnCount; ColumnId++)
+		{
+			int32 EntryId = ColumnId * ColumnHeight + RowId;
 
-					for (int32 Space = 0; Space < SpaceCount; ++Space)
-					{
-						Line.AppendChar((TCHAR)' ');
-					}
-				}
+			if (EntryId >= Entries.Num())
+			{
+				break;
+			}
 
-				if (bLineEnd)
-				{
-					Line.TrimEndInline();
-					UE_LOG(LogConsoleResponse, Log, TEXT("   %s"), *Line);
-					Line.Empty();
-				}
+			const FString& Entry = *Entries[EntryId];
+
+			const int32 SpaceCount = ColumnAlignment - Line.Len();
+			check(SpaceCount >= 0);
+			for (int32 Space = 0; Space < SpaceCount; ++Space)
+			{
+				Line.AppendChar((TCHAR)' ');
 			}
+
+			Line += Entry;
+
+			ColumnAlignment += SpaceBetweenColumns + ColumnWidths[ColumnId];
 		}
-	}
 
-	uint32 WholeCount;
-	uint32 WholePoolInKB;
-	uint32 UsedInKB;
-	GRenderTargetPool.GetStats(WholeCount, WholePoolInKB, UsedInKB);
+		UE_LOG(LogConsoleResponse, Log, TEXT("   %s"), *Line);
+	}
 
-	UE_LOG(LogConsoleResponse, Log, TEXT("Pool: %d/%d MB (referenced/allocated)"), (UsedInKB + 1023) / 1024, (WholePoolInKB + 1023) / 1024);
+	UE_LOG(LogConsoleResponse, Log, TEXT(""));
 }
 
 static TAutoConsoleVariable<int32> CVarAllowBlinking(
diff --git a/Engine/Source/Runtime/RenderCore/Public/VisualizeTexture.h b/Engine/Source/Runtime/RenderCore/Public/VisualizeTexture.h
--- a/Engine/Source/Runtime/RenderCore/Public/VisualizeTexture.h
+++ b/Engine/Source/Runtime/RenderCore/Public/VisualizeTexture.h
@@ -3,6 +3,7 @@
 #pragma once
 
 #include "RenderGraph.h"
+#include "Misc\WildcardString.h"
 
 class FOutputDevice;
 
@@ -32,14 +33,17 @@ private:
 		None				= 0,
 		SaveBitmap			= 1 << 0,
 		SaveBitmapAsStencil = 1 << 1, // stencil normally displays in the alpha channel of depth buffer visualization. This option is just for BMP writeout to get a stencil only BMP.
-		FullList			= 1 << 2,
 	};
 	FRIEND_ENUM_CLASS_FLAGS(EFlags);
 
-	enum class EDebugLogVerbosity
+	enum class ECommand
 	{
-		Default,
-		Extended
+		Unknown,
+		DisableVisualization,
+		VisualizeResource,
+		DisplayHelp,
+		DisplayPoolResourceList,
+		DisplayResourceList,
 	};
 
 	enum class EInputUVMapping
@@ -57,6 +61,12 @@ private:
 		Shadow
 	};
 
+	enum class EDisplayMode
+	{
+		MultiColomn,
+		Detailed,
+	};
+
 	enum class ESortBy
 	{
 		Index,
@@ -71,7 +81,9 @@ private:
 	};
 
 #if SUPPORTS_VISUALIZE_TEXTURE
-	void DebugLog(EDebugLogVerbosity Verbosity);
+	static void DisplayHelp(FOutputDevice &Ar);
+	void DisplayPoolResourceListToLog(ESortBy SortBy);
+	void DisplayResourceListToLog(const TOptional<FWildcardString>& Wildcard);
 
 	/** Determine whether a texture should be captured for debugging purposes and return the capture id if needed. */
 	TOptional<uint32> ShouldCapture(const TCHAR* DebugName, uint32 MipIndex);
@@ -95,7 +107,6 @@ private:
 		float SingleChannelMul = 0.0f;
 
 		EFlags Flags = EFlags::None;
-		ESortBy SortBy = ESortBy::Index;
 		EInputUVMapping InputUVMapping = EInputUVMapping::PictureInPicture;
 		EShaderOp ShaderOp = EShaderOp::Frac;
 		uint32 MipIndex = 0;
