Bumped render thread heartbeat thread's stack size so that it is big enough to handle failed ensures and checks.

#rnx

diff --git a/Engine/Source/Runtime/RenderCore/Private/RenderingThread.cpp b/Engine/Source/Runtime/RenderCore/Private/RenderingThread.cpp
--- a/Engine/Source/Runtime/RenderCore/Private/RenderingThread.cpp
+++ b/Engine/Source/Runtime/RenderCore/Private/RenderingThread.cpp
@@ -790,7 +790,7 @@ PRAGMA_ENABLE_DEPRECATION_WARNINGS
 	GRenderingThreadRunnableHeartbeat = new FRenderingThreadTickHeartbeat();
 
 	UE::Trace::ThreadGroupBegin(TEXT("Render"));
-	GRenderingThreadHeartbeat = FRunnableThread::Create(GRenderingThreadRunnableHeartbeat, *FString::Printf(TEXT("RTHeartBeat %d"), ThreadCount), 16 * 1024, TPri_AboveNormal, FPlatformAffinity::GetRTHeartBeatMask());
+	GRenderingThreadHeartbeat = FRunnableThread::Create(GRenderingThreadRunnableHeartbeat, *FString::Printf(TEXT("RTHeartBeat %d"), ThreadCount), 80 * 1024, TPri_AboveNormal, FPlatformAffinity::GetRTHeartBeatMask());
 	UE::Trace::ThreadGroupEnd();
 
 	ThreadCount++;
