- Allow creation of RHIRayTracingGeometry when IndexBuffer is nullptr
- Remove RT GPUSkinCache spam to log

diff --git a/Engine/Source/Runtime/Engine/Private/GPUSkinCache.cpp b/Engine/Source/Runtime/Engine/Private/GPUSkinCache.cpp
--- a/Engine/Source/Runtime/Engine/Private/GPUSkinCache.cpp
+++ b/Engine/Source/Runtime/Engine/Private/GPUSkinCache.cpp
@@ -1535,8 +1535,6 @@ void FGPUSkinCache::ProcessRayTracingGeometryToUpdate(
 				Initializer.Segments.Add(Segment);
 			}
 
-			UE_LOG(LogSkinCache, Log, TEXT("Build RT BVH with %d triangles"), TrianglesCount);
-
 			FGPUSkinCache::GetRayTracingSegmentVertexBuffers(*SkinCacheEntry, Initializer.Segments);
 
 			// Flush pending resource barriers before BVH is built for the first time
diff --git a/Engine/Source/Runtime/RenderCore/Private/RenderResource.cpp b/Engine/Source/Runtime/RenderCore/Private/RenderResource.cpp
--- a/Engine/Source/Runtime/RenderCore/Private/RenderResource.cpp
+++ b/Engine/Source/Runtime/RenderCore/Private/RenderResource.cpp
@@ -373,7 +373,7 @@ void FRayTracingGeometry::CreateRayTracingGeometry(ERTAccelerationStructureBuild
 		Initializer.OfflineData = &RawData;
 	}
 
-	bool bAllSegmentsAreValid = true;
+	bool bAllSegmentsAreValid = Initializer.Segments.Num() > 0;
 	for (const FRayTracingGeometrySegment& Segment : Initializer.Segments)
 	{
 		if (!Segment.VertexBuffer)
@@ -383,7 +383,7 @@ void FRayTracingGeometry::CreateRayTracingGeometry(ERTAccelerationStructureBuild
 		}
 	}
 
-	if (Initializer.IndexBuffer && bAllSegmentsAreValid)
+	if (bAllSegmentsAreValid)
 	{
 		RayTracingGeometryRHI = RHICreateRayTracingGeometry(Initializer);
 		if (Initializer.OfflineData == nullptr)
