Pulling in fix by [REDACTED_PERSON] from UE4/Main to fix RHI refs being lost for CharAYUV type texture samples in MFW

diff --git a/Engine/Source/Runtime/MediaAssets/Private/Misc/MediaTextureResource.cpp b/Engine/Source/Runtime/MediaAssets/Private/Misc/MediaTextureResource.cpp
--- a/Engine/Source/Runtime/MediaAssets/Private/Misc/MediaTextureResource.cpp
+++ b/Engine/Source/Runtime/MediaAssets/Private/Misc/MediaTextureResource.cpp
@@ -742,6 +742,9 @@ void FMediaTextureResource::ConvertSample(const TSharedPtr<IMediaTextureSample,
 
 			bool bIsSampleOutputSrgb = Sample->IsOutputSrgb();
 
+			// Temporary SRV variables to hold references for the draw
+			FShaderResourceViewRHIRef TempSRV0, TempSRV1;
+
 			// Use the sample format to choose the conversion path
 			switch (Sample->GetFormat())
 			{
@@ -771,7 +774,10 @@ void FMediaTextureResource::ConvertSample(const TSharedPtr<IMediaTextureSample,
 						TShaderMapRef<FNV12ConvertPS> ConvertShader(ShaderMap);
 						GraphicsPSOInit.BoundShaderState.PixelShaderRHI = ConvertShader.GetPixelShader();
 						SetGraphicsPipelineState(CommandList, GraphicsPSOInit);
-						ConvertShader->SetParameters(CommandList, InputTexture, OutputDim, YUVToRGBMatrix, YUVOffset, bIsSampleOutputSrgb);
+						FIntPoint TexDim = InputTexture->GetSizeXY();
+						TempSRV0 = RHICreateShaderResourceView(InputTexture, 0, 1, PF_G8);
+						TempSRV1 = RHICreateShaderResourceView(InputTexture, 0, 1, PF_R8G8);
+						ConvertShader->SetParameters(CommandList, TexDim, TempSRV0, TempSRV1, OutputDim, YUVToRGBMatrix, YUVOffset, bIsSampleOutputSrgb);
 					}
 					else
 					{
diff --git a/Engine/Source/Runtime/RenderCore/Private/MediaShaders.cpp b/Engine/Source/Runtime/RenderCore/Private/MediaShaders.cpp
--- a/Engine/Source/Runtime/RenderCore/Private/MediaShaders.cpp
+++ b/Engine/Source/Runtime/RenderCore/Private/MediaShaders.cpp
@@ -177,25 +177,6 @@ END_GLOBAL_SHADER_PARAMETER_STRUCT()
 IMPLEMENT_GLOBAL_SHADER_PARAMETER_STRUCT(FNV12ConvertUB, "NV12ConvertUB");
 IMPLEMENT_SHADER_TYPE(, FNV12ConvertPS, TEXT("/Engine/Private/MediaShaders.usf"), TEXT("NV12ConvertPS"), SF_Pixel);
 
-
-void FNV12ConvertPS::SetParameters(FRHICommandList& CommandList, TRefCountPtr<FRHITexture2D> NV12Texture, const FIntPoint& OutputDimensions, const FMatrix& ColorTransform, const FVector& YUVOffset, bool SrgbToLinear)
-{
-	FNV12ConvertUB UB;
-	{
-		UB.ColorTransform = MediaShaders::CombineColorTransformAndOffset(ColorTransform, YUVOffset);
-		UB.OutputWidth = OutputDimensions.X;
-		UB.SamplerB = TStaticSamplerState<SF_Bilinear>::GetRHI();
-		UB.SamplerP = TStaticSamplerState<SF_Point>::GetRHI();
-		UB.SrgbToLinear = SrgbToLinear;
-		UB.SRV_Y = RHICreateShaderResourceView(NV12Texture, 0, 1, PF_G8);
-		UB.SRV_UV = RHICreateShaderResourceView(NV12Texture, 0, 1, PF_R8G8);
-		UB.UVScale = FVector2D((float)OutputDimensions.X / (float)NV12Texture->GetSizeX(), (float)OutputDimensions.Y / (float)NV12Texture->GetSizeY());
-	}
-
-	TUniformBufferRef<FNV12ConvertUB> Data = TUniformBufferRef<FNV12ConvertUB>::CreateUniformBufferImmediate(UB, UniformBuffer_SingleFrame);
-	SetUniformBufferParameter(CommandList, CommandList.GetBoundPixelShader(), GetUniformBufferParameter<FNV12ConvertUB>(), Data);
-}
-
 void FNV12ConvertPS::SetParameters(FRHICommandList& CommandList, const FIntPoint & TexDim, FShaderResourceViewRHIRef SRV_Y, FShaderResourceViewRHIRef SRV_UV, const FIntPoint& OutputDimensions, const FMatrix& ColorTransform, const FVector& YUVOffset, bool SrgbToLinear)
 {
 	FNV12ConvertUB UB;
diff --git a/Engine/Source/Runtime/RenderCore/Public/MediaShaders.h b/Engine/Source/Runtime/RenderCore/Public/MediaShaders.h
--- a/Engine/Source/Runtime/RenderCore/Public/MediaShaders.h
+++ b/Engine/Source/Runtime/RenderCore/Public/MediaShaders.h
@@ -221,7 +221,6 @@ public:
 		: FGlobalShader(Initializer)
 	{ }
 
-	RENDERCORE_API void SetParameters(FRHICommandList& RHICmdList, TRefCountPtr<FRHITexture2D> NV12Texture, const FIntPoint& OutputDimensions, const FMatrix& ColorTransform, const FVector& YUVOffset, bool SrgbToLinear);
 	RENDERCORE_API void SetParameters(FRHICommandList& RHICmdList, const FIntPoint & TexDim, FShaderResourceViewRHIRef SRV_Y, FShaderResourceViewRHIRef SRV_UV, const FIntPoint& OutputDimensions, const FMatrix& ColorTransform, const FVector& YUVOffset, bool SrgbToLinear);
 };
 
