Fixed auto GBuffer packing for Velocity Depth

diff --git a/Engine/Source/Runtime/RenderCore/Private/GBufferInfo.cpp b/Engine/Source/Runtime/RenderCore/Private/GBufferInfo.cpp
--- a/Engine/Source/Runtime/RenderCore/Private/GBufferInfo.cpp
+++ b/Engine/Source/Runtime/RenderCore/Private/GBufferInfo.cpp
@@ -406,9 +406,15 @@ FGBufferInfo RENDERCORE_API FetchLegacyGBufferInfo(const FGBufferParams& Params)
 
 	if (Params.bHasVelocity)
 	{
-		Info.Slots[GBS_Velocity] = FGBufferItem(GBS_Velocity, GBC_Raw_Float_16_16, GBCH_Both);
+		Info.Slots[GBS_Velocity] = FGBufferItem(GBS_Velocity, Params.bUsesVelocityDepth ? GBC_Raw_Float_16_16_16_16 : GBC_Raw_Float_16_16, GBCH_Both);
 		Info.Slots[GBS_Velocity].Packing[0] = FGBufferPacking(TargetVelocity, 0, 0);
 		Info.Slots[GBS_Velocity].Packing[1] = FGBufferPacking(TargetVelocity, 1, 1);
+
+		if (Params.bUsesVelocityDepth)
+		{
+			Info.Slots[GBS_Velocity].Packing[2] = FGBufferPacking(TargetVelocity, 2, 2);
+			Info.Slots[GBS_Velocity].Packing[3] = FGBufferPacking(TargetVelocity, 3, 3);
+		}
 	}
 	if (Params.bHasPrecShadowFactor)
 	{
