Fixed RDG texture to use the RDG allocator instead the default allocator.

#rnx

diff --git a/Engine/Source/Runtime/RenderCore/Public/RenderGraphResources.h b/Engine/Source/Runtime/RenderCore/Public/RenderGraphResources.h
--- a/Engine/Source/Runtime/RenderCore/Public/RenderGraphResources.h
+++ b/Engine/Source/Runtime/RenderCore/Public/RenderGraphResources.h
@@ -703,7 +703,7 @@ private:
 	FRDGTextureTransientSubresourceStateIndirect MergeState;
 
 	/** Tracks pass producers for each subresource as the graph is built. */
-	TRDGTextureSubresourceArray<FRDGProducerStatesByPipeline> LastProducers;
+	TRDGTextureSubresourceArray<FRDGProducerStatesByPipeline, FRDGArrayAllocator> LastProducers;
 
 #if RDG_ENABLE_DEBUG
 	struct FRDGTextureDebugData* TextureDebugData = nullptr;
