Added tag tracking whether a pass uses the immediate command list to RDG insights.

#rnx

diff --git a/Engine/Plugins/RenderGraphInsights/Source/RenderGraphInsights/Private/RenderGraphProvider.cpp b/Engine/Plugins/RenderGraphInsights/Source/RenderGraphInsights/Private/RenderGraphProvider.cpp
--- a/Engine/Plugins/RenderGraphInsights/Source/RenderGraphInsights/Private/RenderGraphProvider.cpp
+++ b/Engine/Plugins/RenderGraphInsights/Source/RenderGraphInsights/Private/RenderGraphProvider.cpp
@@ -118,6 +118,7 @@ FPassPacket::FPassPacket(const UE::Trace::IAnalyzer::FOnEventContext& Context)
 	, bAsyncComputeEnd(Context.EventData.GetValue<bool>("IsAsyncComputeEnd"))
 	, bSkipRenderPassBegin(Context.EventData.GetValue<bool>("SkipRenderPassBegin"))
 	, bSkipRenderPassEnd(Context.EventData.GetValue<bool>("SkipRenderPassEnd"))
+	, bImmediateCommandList(Context.EventData.GetValue<bool>("UsesImmediateCommandList"))
 {}
 
 static const uint64 PageSize = 1024;
diff --git a/Engine/Plugins/RenderGraphInsights/Source/RenderGraphInsights/Private/RenderGraphProvider.h b/Engine/Plugins/RenderGraphInsights/Source/RenderGraphInsights/Private/RenderGraphProvider.h
--- a/Engine/Plugins/RenderGraphInsights/Source/RenderGraphInsights/Private/RenderGraphProvider.h
+++ b/Engine/Plugins/RenderGraphInsights/Source/RenderGraphInsights/Private/RenderGraphProvider.h
@@ -126,6 +126,7 @@ public:
 	bool bAsyncComputeEnd{};
 	bool bSkipRenderPassBegin{};
 	bool bSkipRenderPassEnd{};
+	bool bImmediateCommandList{};
 
 	FPassPacket(const UE::Trace::IAnalyzer::FOnEventContext& Context);
 };
diff --git a/Engine/Plugins/RenderGraphInsights/Source/RenderGraphInsights/Private/RenderGraphTrack.cpp b/Engine/Plugins/RenderGraphInsights/Source/RenderGraphInsights/Private/RenderGraphTrack.cpp
--- a/Engine/Plugins/RenderGraphInsights/Source/RenderGraphInsights/Private/RenderGraphTrack.cpp
+++ b/Engine/Plugins/RenderGraphInsights/Source/RenderGraphInsights/Private/RenderGraphTrack.cpp
@@ -1430,6 +1430,11 @@ void FRenderGraphTrack::InitTooltip(FTooltipDrawState& Tooltip, const ITimingEve
 			{
 				Tooltip.AddTextLine(TEXT("Merged RenderPass"), FLinearColor::Red);
 			}
+
+			if (Pass.bImmediateCommandList)
+			{
+				Tooltip.AddTextLine(TEXT("Immediate Command List"), FLinearColor::Red);
+			}
 		}
 		else if (InTooltipEvent.Is<FVisibleTextureEvent>())
 		{
diff --git a/Engine/Source/Runtime/RenderCore/Private/RenderGraphTrace.cpp b/Engine/Source/Runtime/RenderCore/Private/RenderGraphTrace.cpp
--- a/Engine/Source/Runtime/RenderCore/Private/RenderGraphTrace.cpp
+++ b/Engine/Source/Runtime/RenderCore/Private/RenderGraphTrace.cpp
@@ -35,6 +35,7 @@ UE_TRACE_EVENT_BEGIN(RDGTrace, PassMessage)
 	UE_TRACE_EVENT_FIELD(bool, IsAsyncComputeEnd)
 	UE_TRACE_EVENT_FIELD(bool, SkipRenderPassBegin)
 	UE_TRACE_EVENT_FIELD(bool, SkipRenderPassEnd)
+	UE_TRACE_EVENT_FIELD(bool, UsesImmediateCommandList)
 UE_TRACE_EVENT_END()
 
 UE_TRACE_EVENT_BEGIN(RDGTrace, BufferMessage)
@@ -50,7 +51,7 @@ UE_TRACE_EVENT_BEGIN(RDGTrace, BufferMessage)
 	UE_TRACE_EVENT_FIELD(bool, IsExtracted)
 	UE_TRACE_EVENT_FIELD(bool, IsCulled)
 UE_TRACE_EVENT_END()
-
+ 
 UE_TRACE_EVENT_BEGIN(RDGTrace, TextureMessage)
 	UE_TRACE_EVENT_FIELD(UE::Trace::WideString, Name)
 	UE_TRACE_EVENT_FIELD(uint64, StartCycles)
@@ -107,6 +108,8 @@ void FRDGTrace::OutputGraphEnd(const FRDGBuilder& GraphBuilder)
 		return;
 	}
 
+	TRACE_CPUPROFILER_EVENT_SCOPE(FRDGTrace::OutputGraphEnd);
+
 	const FRDGPassHandle ProloguePassHandle = GraphBuilder.GetProloguePassHandle();
 
 	const auto& Passes = GraphBuilder.Passes;
@@ -141,7 +144,8 @@ void FRDGTrace::OutputGraphEnd(const FRDGBuilder& GraphBuilder)
 			<< PassMessage.IsAsyncComputeBegin(Pass->bAsyncComputeBegin != 0)
 			<< PassMessage.IsAsyncComputeEnd(Pass->bAsyncComputeEnd != 0)
 			<< PassMessage.SkipRenderPassBegin(Pass->bSkipRenderPassBegin != 0)
-			<< PassMessage.SkipRenderPassEnd(Pass->bSkipRenderPassEnd != 0);
+			<< PassMessage.SkipRenderPassEnd(Pass->bSkipRenderPassEnd != 0)
+			<< PassMessage.UsesImmediateCommandList(Pass->bImmediateCommandList != 0);
 	}
 
 #if RDG_EVENTS
diff --git a/Engine/Source/Runtime/RenderCore/Public/RenderGraphPass.h b/Engine/Source/Runtime/RenderCore/Public/RenderGraphPass.h
--- a/Engine/Source/Runtime/RenderCore/Public/RenderGraphPass.h
+++ b/Engine/Source/Runtime/RenderCore/Public/RenderGraphPass.h
@@ -302,6 +302,9 @@ private:
 			/** Whether the pass writes to a UAV. */
 			uint32 bUAVAccess : 1;
 
+			/** Whether the pass uses the immediate command list. */
+			uint32 bImmediateCommandList : 1;
+
 			/** Whether this pass allocated a texture through the pool. */
 			IF_RDG_ENABLE_DEBUG(uint32 bFirstTextureAllocated : 1);
 		};
@@ -436,6 +439,8 @@ public:
 	{
 		checkf(kSupportsAsyncCompute || !EnumHasAnyFlags(InPassFlags, ERDGPassFlags::AsyncCompute),
 			TEXT("Pass %s is set to use 'AsyncCompute', but the pass lambda's first argument is not FRHIComputeCommandList&."), GetName());
+
+		bImmediateCommandList = TIsSame<TRHICommandList, FRHICommandListImmediate>::Value;
 	}
 
 private:
