Removed untracked access from RDG and other conversion-related stuff.

#rnx

diff --git a/Engine/Source/Runtime/RenderCore/Private/RenderGraphPrivate.cpp b/Engine/Source/Runtime/RenderCore/Private/RenderGraphPrivate.cpp
--- a/Engine/Source/Runtime/RenderCore/Private/RenderGraphPrivate.cpp
+++ b/Engine/Source/Runtime/RenderCore/Private/RenderGraphPrivate.cpp
@@ -235,6 +235,8 @@ void EmitRDGWarning(const FString& WarningMessage)
 	}
 }
 
+bool GRDGAllowRHIAccess = false;
+
 #endif
 
 int32 GRDGAsyncCompute = 1;
diff --git a/Engine/Source/Runtime/RenderCore/Private/RenderGraphPrivate.h b/Engine/Source/Runtime/RenderCore/Private/RenderGraphPrivate.h
--- a/Engine/Source/Runtime/RenderCore/Private/RenderGraphPrivate.h
+++ b/Engine/Source/Runtime/RenderCore/Private/RenderGraphPrivate.h
@@ -32,6 +32,25 @@ extern int32 GRDGTransitionLog;
 extern int32 GRDGImmediateMode;
 extern int32 GRDGOverlapUAVs;
 extern int32 GRDGExtendResourceLifetimes;
+extern bool  GRDGAllowRHIAccess;
+
+class FRDGAllowRHIAccessScope
+{
+public:
+	FRDGAllowRHIAccessScope()
+	{
+		check(!GRDGAllowRHIAccess);
+		GRDGAllowRHIAccess = true;
+	}
+
+	~FRDGAllowRHIAccessScope()
+	{
+		check(GRDGAllowRHIAccess);
+		GRDGAllowRHIAccess = false;
+	}
+};
+
+#define RDG_ALLOW_RHI_ACCESS_SCOPE() FRDGAllowRHIAccessScope RDGAllowRHIAccessScopeRAII;
 
 // Colors for texture / buffer clobbering.
 FLinearColor GetClobberColor();
@@ -76,6 +95,8 @@ const int32 GRDGImmediateMode = 0;
 const int32 GRDGOverlapUAVs = 1;
 const int32 GRDGExtendResourceLifetimes = 0;
 
+#define RDG_ALLOW_RHI_ACCESS_SCOPE()
+
 #define EmitRDGWarningf(WarningMessageFormat, ...)
 
 #endif
diff --git a/Engine/Source/Runtime/RenderCore/Public/RenderGraphDefinitions.h b/Engine/Source/Runtime/RenderCore/Public/RenderGraphDefinitions.h
--- a/Engine/Source/Runtime/RenderCore/Public/RenderGraphDefinitions.h
+++ b/Engine/Source/Runtime/RenderCore/Public/RenderGraphDefinitions.h
@@ -95,20 +95,8 @@ enum class ERDGPassFlags : uint8
 	/** Render pass begin / end is skipped and left to the user. Only valid when combined with 'Raster'. Disables render pass merging for the pass. */
 	SkipRenderPass = 1 << 5,
 
-	/** Pass accesses raw RHI resources which may be registered with the graph, but all resources are kept in their current state. This flag prevents
-	 *  the graph from scheduling split barriers across the pass. Any splitting is deferred until after the pass executes. The resource may not change
-	 *  state within the pass execution. Affects barrier performance. May not be combined with Async Compute.
-	 */
-	UntrackedAccess = 1 << 6,
-
 	/** Pass uses copy commands but writes to a staging resource. */
-	Readback = Copy | NeverCull,
-
-	/** Mask of flags denoting the kinds of RHI commands submitted to a pass. */
-	CommandMask = Raster | Compute | AsyncCompute | Copy,
-
-	/** Mask of flags which can used by a pass flag scope. */
-	ScopeMask = NeverCull | UntrackedAccess
+	Readback = Copy | NeverCull
 };
 ENUM_CLASS_FLAGS(ERDGPassFlags);
 
@@ -431,7 +419,6 @@ struct FRDGTextureDesc;
 
 class FRDGBlackboard;
 
-class FRDGPassFlagsScopeGuard;
 class FRDGAsyncComputeBudgetScopeGuard;
 class FRDGEventScopeGuard;
 class FRDGGPUStatScopeGuard;
diff --git a/Engine/Source/Runtime/RenderCore/Public/RenderGraphUtils.h b/Engine/Source/Runtime/RenderCore/Public/RenderGraphUtils.h
--- a/Engine/Source/Runtime/RenderCore/Public/RenderGraphUtils.h
+++ b/Engine/Source/Runtime/RenderCore/Public/RenderGraphUtils.h
@@ -648,18 +648,6 @@ FORCEINLINE void AddPass(FRDGBuilder& GraphBuilder, ExecuteLambdaType&& ExecuteL
 	AddPass(GraphBuilder, {}, MoveTemp(ExecuteLambda));
 }
 
-template <typename ExecuteLambdaType>
-FORCEINLINE void AddUntrackedAccessPass(FRDGBuilder& GraphBuilder, FRDGEventName&& Name, ExecuteLambdaType&& ExecuteLambda)
-{
-	GraphBuilder.AddPass(MoveTemp(Name), ERDGPassFlags::UntrackedAccess, MoveTemp(ExecuteLambda));
-}
-
-template <typename ExecuteLambdaType>
-FORCEINLINE void AddUntrackedAccessPass(FRDGBuilder& GraphBuilder, ExecuteLambdaType&& ExecuteLambda)
-{
-	AddUntrackedAccessPass(GraphBuilder, {}, MoveTemp(ExecuteLambda));
-}
-
 template <typename ExecuteLambdaType>
 FORCEINLINE void AddPassIfDebug(FRDGBuilder& GraphBuilder, FRDGEventName&& Name, ExecuteLambdaType&& ExecuteLambda)
 {
@@ -674,20 +662,6 @@ FORCEINLINE void AddPassIfDebug(FRDGBuilder& GraphBuilder, ExecuteLambdaType&& E
 	AddPassIfDebug(GraphBuilder, {}, MoveTemp(ExecuteLambda));
 }
 
-template <typename ExecuteLambdaType>
-FORCEINLINE void AddUntrackedAccessPassIfDebug(FRDGBuilder& GraphBuilder, FRDGEventName&& Name, ExecuteLambdaType&& ExecuteLambda)
-{
-#if !(UE_BUILD_SHIPPING || UE_BUILD_TEST)
-	AddUntrackedAccessPass(GraphBuilder, MoveTemp(Name), MoveTemp(ExecuteLambda));
-#endif
-}
-
-template <typename ExecuteLambdaType>
-FORCEINLINE void AddUntrackedAccessPassIfDebug(FRDGBuilder& GraphBuilder, ExecuteLambdaType&& ExecuteLambda)
-{
-	AddUntrackedAccessPassIfDebug(GraphBuilder, {}, MoveTemp(ExecuteLambda));
-}
-
 UE_DEPRECATED(5.0, "AddSetCurrentStatPass is deprecated. Use GraphBuilder.SetCommandListStat instead.")
 FORCEINLINE void AddSetCurrentStatPass(FRDGBuilder& GraphBuilder, TStatId StatId)
 {
diff --git a/Engine/Source/Runtime/Renderer/Private/AtmosphereRendering.cpp b/Engine/Source/Runtime/Renderer/Private/AtmosphereRendering.cpp
--- a/Engine/Source/Runtime/Renderer/Private/AtmosphereRendering.cpp
+++ b/Engine/Source/Runtime/Renderer/Private/AtmosphereRendering.cpp
@@ -1662,7 +1662,7 @@ void FAtmosphericFogSceneInfo::PrecomputeTextures(FRDGBuilder& GraphBuilder, con
 	// Atmosphere 
 	if (bPrecomputationStarted && !bPrecomputationFinished)
 	{
-		AddUntrackedAccessPass(GraphBuilder, [this, View, ViewFamily](FRHICommandListImmediate& RHICmdList)
+		AddPass(GraphBuilder, [this, View, ViewFamily](FRHICommandListImmediate& RHICmdList)
 		{
 			PrecomputeAtmosphereData(RHICmdList, View, *ViewFamily);
 
diff --git a/Engine/Source/Runtime/Renderer/Private/DistanceFieldObjectCulling.cpp b/Engine/Source/Runtime/Renderer/Private/DistanceFieldObjectCulling.cpp
--- a/Engine/Source/Runtime/Renderer/Private/DistanceFieldObjectCulling.cpp
+++ b/Engine/Source/Runtime/Renderer/Private/DistanceFieldObjectCulling.cpp
@@ -633,7 +633,7 @@ void BuildTileObjectLists(FRDGBuilder& GraphBuilder, FScene* Scene, TArray<FView
 		GraphBuilder.AddPass(
 			{},
 			PassParameters,
-			ERDGPassFlags::Compute | ERDGPassFlags::UntrackedAccess | ERDGPassFlags::NeverCull,
+			ERDGPassFlags::Compute | ERDGPassFlags::NeverCull,
 			[Scene, &View, DistanceFieldNormal, Parameters](FRHICommandListImmediate& RHICmdList)
 		{
 			const FIntPoint TileListGroupSize = GetTileListGroupSizeForView(View);
diff --git a/Engine/Source/Runtime/Renderer/Private/DistanceFieldScreenGridLighting.cpp b/Engine/Source/Runtime/Renderer/Private/DistanceFieldScreenGridLighting.cpp
--- a/Engine/Source/Runtime/Renderer/Private/DistanceFieldScreenGridLighting.cpp
+++ b/Engine/Source/Runtime/Renderer/Private/DistanceFieldScreenGridLighting.cpp
@@ -525,7 +525,7 @@ void FDeferredShadingSceneRenderer::RenderDistanceFieldAOScreenGrid(
 		GraphBuilder.AddPass(
 			{},
 			PassParameters,
-			ERDGPassFlags::Compute | ERDGPassFlags::UntrackedAccess | ERDGPassFlags::NeverCull,
+			ERDGPassFlags::Compute | ERDGPassFlags::NeverCull,
 			[&View, &ScreenGridResources, Parameters, DistanceFieldNormal, TileListGroupSize, bUseGlobalDistanceField, bUseObjectDistanceField]
 			(FRHICommandListImmediate& RHICmdList)
 		{
diff --git a/Engine/Source/Runtime/Renderer/Private/LightRendering.cpp b/Engine/Source/Runtime/Renderer/Private/LightRendering.cpp
--- a/Engine/Source/Runtime/Renderer/Private/LightRendering.cpp
+++ b/Engine/Source/Runtime/Renderer/Private/LightRendering.cpp
@@ -2543,7 +2543,7 @@ void FDeferredShadingSceneRenderer::RenderLightForHair(
 		GraphBuilder.AddPass(
 			{},
 			PassParameters,
-			ERDGPassFlags::Raster | ERDGPassFlags::UntrackedAccess,
+			ERDGPassFlags::Raster,
 			[&HairVisibilityData, &View, PassParameters, LightSceneInfo, MaxTransmittanceElementCount, HairShadowMaskTexture, LightingChannelsTexture, bIsShadowMaskValid](FRHICommandList& RHICmdList)
 		{
 			RHICmdList.SetViewport(0, 0, 0.0f, HairVisibilityData.SampleLightingViewportResolution.X, HairVisibilityData.SampleLightingViewportResolution.Y, 1.0f);
