Fixed build break.

diff --git a/Engine/Source/Runtime/RenderCore/Private/RenderGraphBuilder.cpp b/Engine/Source/Runtime/RenderCore/Private/RenderGraphBuilder.cpp
--- a/Engine/Source/Runtime/RenderCore/Private/RenderGraphBuilder.cpp
+++ b/Engine/Source/Runtime/RenderCore/Private/RenderGraphBuilder.cpp
@@ -403,7 +403,6 @@ ERDGPassFlags FRDGBuilder::OverridePassFlags(const TCHAR* PassName, ERDGPassFlag
 #endif
 
 	const bool bGlobalForceAsyncCompute = (GRDGAsyncCompute == RDG_ASYNC_COMPUTE_FORCE_ENABLED && !GRDGImmediateMode && bDebugAllowedForPass);
-	bAsyncComputeSupported &= !EnumHasAnyFlags(PassFlags, ERDGPassFlags::UntrackedAccess);
 
 	if (EnumHasAnyFlags(PassFlags, ERDGPassFlags::Compute) && (bGlobalForceAsyncCompute))
 	{
@@ -1226,15 +1225,15 @@ void FRDGBuilder::Execute()
 
 	const FRDGPassHandle ProloguePassHandle = GetProloguePassHandle();
 	const FRDGPassHandle EpiloguePassHandle = GetEpiloguePassHandle();
-	FRDGPassHandle LastUntrackedPassHandle = ProloguePassHandle;
 
 	if (!GRDGImmediateMode)
 	{
 		Compile();
 
-		IF_RDG_ENABLE_DEBUG(LogFile.Begin(BuilderName, &Passes, PassesToCull, GetProloguePassHandle(), GetEpiloguePassHandle()));
+		IF_RDG_ENABLE_DEBUG(LogFile.Begin(BuilderName, &Passes, PassesToCull, ProloguePassHandle, EpiloguePassHandle));
 
 		{
+			TRACE_CPUPROFILER_EVENT_SCOPE(FRDGBuilder::CollectResources);
 			SCOPE_CYCLE_COUNTER(STAT_RDG_CollectResourcesTime);
 			CSV_SCOPED_TIMING_STAT_EXCLUSIVE(RDG_CollectResources);
 
@@ -1248,16 +1247,17 @@ void FRDGBuilder::Execute()
 
 			for (const auto& Query : ExtractedTextures)
 			{
-				EndResourceRHI(EpiloguePassHandle, Query.Key, 1);
+				EndResourceRHI(EpiloguePassHandle, Query.Key);
 			}
 
 			for (const auto& Query : ExtractedBuffers)
 			{
-				EndResourceRHI(EpiloguePassHandle, Query.Key, 1);
+				EndResourceRHI(EpiloguePassHandle, Query.Key);
 			}
 		}
 
 		{
+			TRACE_CPUPROFILER_EVENT_SCOPE(FRDGBuilder::CollectBarriers);
 			SCOPE_CYCLE_COUNTER(STAT_RDG_CollectBarriersTime);
 			CSV_SCOPED_TIMING_STAT_EXCLUSIVE_CONDITIONAL(RDG_CollectBarriers, GRDGVerboseCSVStats != 0);
 
@@ -1265,7 +1265,7 @@ void FRDGBuilder::Execute()
 			{
 				if (!PassesToCull[PassHandle])
 				{
-					CollectPassBarriers(PassHandle, LastUntrackedPassHandle);
+					CollectPassBarriers(PassHandle);
 				}
 			}
 		}
@@ -1292,7 +1292,7 @@ void FRDGBuilder::Execute()
 
 		if (Texture->GetRHIUnchecked())
 		{
-			AddEpilogueTransition(Texture, LastUntrackedPassHandle);
+			AddEpilogueTransition(Texture);
 			Texture->Finalize();
 
 			IF_RDG_ENABLE_DEBUG(LogResource(Texture, Textures));
@@ -1305,7 +1305,7 @@ void FRDGBuilder::Execute()
 
 		if (Buffer->GetRHIUnchecked())
 		{
-			AddEpilogueTransition(Buffer, LastUntrackedPassHandle);
+			AddEpilogueTransition(Buffer);
 			Buffer->Finalize();
 
 			IF_RDG_ENABLE_DEBUG(LogResource(Buffer, Buffers));
@@ -1385,6 +1385,7 @@ void FRDGBuilder::Execute()
 
 void FRDGBuilder::Clear()
 {
+	TRACE_CPUPROFILER_EVENT_SCOPE(FRDGBuilder::Clear);
 	ExternalTextures.Empty();
 	ExternalBuffers.Empty();
 	ExtractedTextures.Empty();
@@ -1551,9 +1552,8 @@ void FRDGBuilder::SetupPassInternal(FRDGPass* Pass, FRDGPassHandle PassHandle, E
 		}
 
 		check(!EnumHasAnyFlags(PassPipeline, ERHIPipeline::AsyncCompute));
-		FRDGPassHandle LastUntrackedPassHandle = GetProloguePassHandle();
 		CollectPassResources(PassHandle);
-		CollectPassBarriers(PassHandle, LastUntrackedPassHandle);
+		CollectPassBarriers(PassHandle);
 		ExecutePass(Pass);
 	}
 
@@ -1728,8 +1728,9 @@ void FRDGBuilder::ExecutePass(FRDGPass* Pass)
 
 void FRDGBuilder::CollectPassResources(FRDGPassHandle PassHandle)
 {
+	RDG_ALLOW_RHI_ACCESS_SCOPE();
+
 	FRDGPass* Pass = Passes[PassHandle];
-	const ERDGPassFlags PassFlags = Pass->GetFlags();
 
 	for (FRDGPass* PassToBegin : Pass->ResourcesToBegin)
 	{
@@ -1746,41 +1747,31 @@ void FRDGBuilder::CollectPassResources(FRDGPassHandle PassHandle)
 		});
 
 		// Uniform buffer creation requires access to RHI resources through the validator.
-		IF_RDG_ENABLE_DEBUG(FRDGUserValidation::SetAllowRHIAccess(PassToBegin, true));
+		PassParameters.EnumerateUniformBuffers([&](FRDGUniformBufferBinding UniformBuffer)
 		{
-			PassParameters.EnumerateUniformBuffers([&](FRDGUniformBufferBinding UniformBuffer)
-			{
-				BeginResourceRHI(UniformBuffer.GetUniformBuffer());
-			});
-		}
-		IF_RDG_ENABLE_DEBUG(FRDGUserValidation::SetAllowRHIAccess(PassToBegin, false));
+			BeginResourceRHI(UniformBuffer.GetUniformBuffer());
+		});
 	}
 
 	for (FRDGPass* PassToEnd : Pass->ResourcesToEnd)
 	{
 		for (const auto& TexturePair : PassToEnd->TextureStates)
 		{
-			EndResourceRHI(PassHandle, TexturePair.Key, TexturePair.Value.ReferenceCount);
+			EndResourceRHI(PassHandle, TexturePair.Key);
 		}
 
 		for (const auto& BufferPair : PassToEnd->BufferStates)
 		{
-			EndResourceRHI(PassHandle, BufferPair.Key, BufferPair.Value.ReferenceCount);
+			EndResourceRHI(PassHandle, BufferPair.Key);
 		}
 	}
 }
 
-void FRDGBuilder::CollectPassBarriers(FRDGPassHandle PassHandle, FRDGPassHandle& LastUntrackedPassHandle)
+void FRDGBuilder::CollectPassBarriers(FRDGPassHandle PassHandle)
 {
 	FRDGPass* Pass = Passes[PassHandle];
 
 	IF_RDG_ENABLE_DEBUG(ConditionalDebugBreak(RDG_BREAKPOINT_PASS_COMPILE, BuilderName.GetTCHAR(), Pass->GetName()));
-
-	if (EnumHasAnyFlags(Pass->GetFlags(), ERDGPassFlags::UntrackedAccess))
-	{
-		LastUntrackedPassHandle = PassHandle;
-	}
-
 	if (PassesWithEmptyParameters[PassHandle])
 	{
 		return;
@@ -1789,7 +1780,7 @@ void FRDGBuilder::CollectPassBarriers(FRDGPassHandle PassHandle, FRDGPassHandle&
 	for (const auto& TexturePair : Pass->TextureStates)
 	{
 		FRDGTextureRef Texture = TexturePair.Key;
-		AddTransition(PassHandle, Texture, TexturePair.Value.MergeState, LastUntrackedPassHandle);
+		AddTransition(PassHandle, Texture, TexturePair.Value.MergeState);
 		Texture->bCulled = false;
 
 		IF_RDG_ENABLE_TRACE(Trace.AddTexturePassDependency(Texture, Pass));
@@ -1798,14 +1789,14 @@ void FRDGBuilder::CollectPassBarriers(FRDGPassHandle PassHandle, FRDGPassHandle&
 	for (const auto& BufferPair : Pass->BufferStates)
 	{
 		FRDGBufferRef Buffer = BufferPair.Key;
-		AddTransition(PassHandle, Buffer, *BufferPair.Value.MergeState, LastUntrackedPassHandle);
+		AddTransition(PassHandle, Buffer, *BufferPair.Value.MergeState);
 		Buffer->bCulled = false;
 
 		IF_RDG_ENABLE_TRACE(Trace.AddBufferPassDependency(Buffer, Pass));
 	}
 }
 
-void FRDGBuilder::AddEpilogueTransition(FRDGTextureRef Texture, FRDGPassHandle LastUntrackedPassHandle)
+void FRDGBuilder::AddEpilogueTransition(FRDGTextureRef Texture)
 {
 	if (!Texture->bLastOwner || Texture->bCulled)
 	{
@@ -1853,11 +1844,11 @@ void FRDGBuilder::AddEpilogueTransition(FRDGTextureRef Texture, FRDGPassHandle L
 		return;
 	}
 
-	AddTransition(EpiloguePassHandle, Texture, ScratchTextureState, LastUntrackedPassHandle);
+	AddTransition(EpiloguePassHandle, Texture, ScratchTextureState);
 	ScratchTextureState.Reset();
 }
 
-void FRDGBuilder::AddEpilogueTransition(FRDGBufferRef Buffer, FRDGPassHandle LastUntrackedPassHandle)
+void FRDGBuilder::AddEpilogueTransition(FRDGBufferRef Buffer)
 {
 	if (!Buffer->bLastOwner || Buffer->bCulled)
 	{
@@ -1884,11 +1875,11 @@ void FRDGBuilder::AddEpilogueTransition(FRDGBufferRef Buffer, FRDGPassHandle Las
 		FRDGSubresourceState StateFinal;
 		StateFinal.SetPass(ERHIPipeline::Graphics, EpiloguePassHandle);
 		StateFinal.Access = AccessFinal;
-		AddTransition(EpiloguePassHandle, Buffer, StateFinal, LastUntrackedPassHandle);
+		AddTransition(EpiloguePassHandle, Buffer, StateFinal);
 	}
 }
 
-void FRDGBuilder::AddTransition(FRDGPassHandle PassHandle, FRDGTextureRef Texture, const FRDGTextureTransientSubresourceStateIndirect& StateAfter, FRDGPassHandle LastUntrackedPassHandle)
+void FRDGBuilder::AddTransition(FRDGPassHandle PassHandle, FRDGTextureRef Texture, const FRDGTextureTransientSubresourceStateIndirect& StateAfter)
 {
 	const FRDGTextureSubresourceRange WholeRange = Texture->GetSubresourceRange();
 	const FRDGTextureSubresourceLayout Layout = Texture->Layout;
@@ -1917,7 +1908,7 @@ void FRDGBuilder::AddTransition(FRDGPassHandle PassHandle, FRDGTextureRef Textur
 				Info.PlaneSlice = Subresource->PlaneSlice;
 			}
 
-			AddTransitionInternal(Texture, SubresourceStateBefore, SubresourceStateAfter, LastUntrackedPassHandle, Info);
+			AddTransitionInternal(Texture, SubresourceStateBefore, SubresourceStateAfter, Info);
 		}
 
 		if (Subresource)
@@ -2003,7 +1994,7 @@ void FRDGBuilder::AddTransition(FRDGPassHandle PassHandle, FRDGTextureRef Textur
 	}
 }
 
-void FRDGBuilder::AddTransition(FRDGPassHandle PassHandle, FRDGBufferRef Buffer, FRDGSubresourceState StateAfter, FRDGPassHandle LastUntrackedPassHandle)
+void FRDGBuilder::AddTransition(FRDGPassHandle PassHandle, FRDGBufferRef Buffer, FRDGSubresourceState StateAfter)
 {
 	check(StateAfter.Access != ERHIAccess::Unknown);
 
@@ -2018,7 +2009,7 @@ void FRDGBuilder::AddTransition(FRDGPassHandle PassHandle, FRDGBufferRef Buffer,
 		Info.AccessBefore = StateBefore.Access;
 		Info.AccessAfter = StateAfter.Access;
 
-		AddTransitionInternal(Buffer, StateBefore, StateAfter, LastUntrackedPassHandle, Info);
+		AddTransitionInternal(Buffer, StateBefore, StateAfter, Info);
 	}
 
 	IF_RDG_ENABLE_DEBUG(LogFile.AddTransitionEdge(PassHandle, StateBefore, StateAfter, Buffer));
@@ -2029,7 +2020,6 @@ void FRDGBuilder::AddTransitionInternal(
 	FRDGParentResource* Resource,
 	FRDGSubresourceState StateBefore,
 	FRDGSubresourceState StateAfter,
-	FRDGPassHandle LastUntrackedPassHandle,
 	const FRHITransitionInfo& TransitionInfo)
 {
 	const ERHIPipeline Graphics = ERHIPipeline::Graphics;
@@ -2084,25 +2074,6 @@ void FRDGBuilder::AddTransitionInternal(
 		}
 	}
 
-	// Avoids splitting across a pass that is doing untracked RHI access.
-	if (StateBefore.GetLastPass() < LastUntrackedPassHandle)
-	{
-		// Transitions exclusively on the graphics queue can be pushed forward.
-		if (!EnumHasAnyFlags(PipelinesBefore | PipelinesAfter, AsyncCompute))
-		{
-			PassesBefore[Graphics] = LastUntrackedPassHandle;
-		}
-		// Async compute transitions have to be split. Emit a warning to the user to avoid touching the resource.
-		else
-		{
-			EmitRDGWarningf(
-				TEXT("Resource '%s' is being split-transitioned across untracked pass '%s'. It was not possible to avoid the ")
-				TEXT("split barrier due to async compute. Accessing this resource in the pass will cause an RHI validation failure. ")
-				TEXT("Otherwise, it's safe to ignore this warning."),
-				Resource->Name, Passes[LastUntrackedPassHandle]->GetName());
-		}
-	}
-
 	const auto GetEpilogueBarrierPassHandle = [&](FRDGPassHandle Handle)
 	{
 		return Passes[Handle]->EpilogueBarrierPass;
@@ -2260,8 +2231,6 @@ void FRDGBuilder::BeginResourceRHI(FRDGPassHandle PassHandle, FRDGTextureRef Tex
 	Texture->SetRHI(PooledRenderTarget, PreviousOwner);
 	Texture->bTransient = PooledRenderTarget->IsTransient();
 
-	check(!IsResourceLifetimeExtended() || !PreviousOwner);
-
 	// Acquires are made in the last pass of the previous alias.
 	FRDGPassHandle AcquirePassHandle = PreviousOwner ? PreviousOwner->LastPass : GetProloguePassHandle();
 
@@ -2414,8 +2383,6 @@ void FRDGBuilder::BeginResourceRHI(FRDGPassHandle PassHandle, FRDGBufferRef Buff
 	Buffer->SetRHI(PooledBuffer, PreviousOwner);
 	Buffer->FirstPass = PassHandle;
 
-	check(!IsResourceLifetimeExtended() || !PreviousOwner);
-
 	// Acquires are made in the last pass of the previous alias.
 	Buffer->AcquirePass = PreviousOwner ? PreviousOwner->LastPass : GetProloguePassHandle();
 
@@ -2456,16 +2423,14 @@ void FRDGBuilder::BeginResourceRHI(FRDGPassHandle PassHandle, FRDGBufferUAV* UAV
 	UAV->ResourceRHI = Buffer->PooledBuffer->GetOrCreateUAV(UAV->Desc);
 }
 
-void FRDGBuilder::EndResourceRHI(FRDGPassHandle PassHandle, FRDGTextureRef Texture, uint32 ReferenceCount)
+void FRDGBuilder::EndResourceRHI(FRDGPassHandle PassHandle, FRDGTextureRef Texture)
 {
 	check(Texture);
 
 	if (!IsResourceLifetimeExtended())
 	{
-		check(Texture->ReferenceCount >= ReferenceCount);
-		Texture->ReferenceCount -= ReferenceCount;
-
-		if (Texture->ReferenceCount == 0)
+		check(Texture->ReferenceCount > 0);
+		if (--Texture->ReferenceCount == 0)
 		{
 			// External textures should never release the reference.
 			if (!Texture->bExternal)
@@ -2477,16 +2442,14 @@ void FRDGBuilder::EndResourceRHI(FRDGPassHandle PassHandle, FRDGTextureRef Textu
 	}
 }
 
-void FRDGBuilder::EndResourceRHI(FRDGPassHandle PassHandle, FRDGBufferRef Buffer, uint32 ReferenceCount)
+void FRDGBuilder::EndResourceRHI(FRDGPassHandle PassHandle, FRDGBufferRef Buffer)
 {
 	check(Buffer);
 
 	if (!IsResourceLifetimeExtended())
 	{
-		check(Buffer->ReferenceCount >= ReferenceCount);
-		Buffer->ReferenceCount -= ReferenceCount;
-
-		if (Buffer->ReferenceCount == 0)
+		check(Buffer->ReferenceCount > 0);
+		if (--Buffer->ReferenceCount == 0)
 		{
 			// External buffers should never release the reference.
 			if (!Buffer->bExternal)
diff --git a/Engine/Source/Runtime/RenderCore/Private/RenderGraphValidation.cpp b/Engine/Source/Runtime/RenderCore/Private/RenderGraphValidation.cpp
--- a/Engine/Source/Runtime/RenderCore/Private/RenderGraphValidation.cpp
+++ b/Engine/Source/Runtime/RenderCore/Private/RenderGraphValidation.cpp
@@ -57,6 +57,7 @@ const ERHIAccess AccessMaskComputeOrRaster = ERHIAccess::IndirectArgs;
 
 /** Validates that only one builder instance exists at any time. This is currently a requirement for state tracking and allocation lifetimes. */
 bool GRDGBuilderActive = false;
+
 } //! namespace
 
 struct FRDGResourceDebugData
@@ -83,7 +84,7 @@ void FRDGResource::ValidateRHIAccess() const
 	// Passthrough resources will not have debug data, since they are not tied to a graph instance.
 	if (DebugData)
 	{
-		checkf(DebugData->bAllowRHIAccess,
+		checkf(DebugData->bAllowRHIAccess || GRDGAllowRHIAccess,
 			TEXT("Accessing the RHI resource of %s at this time is not allowed. If you hit this check in pass, ")
 			TEXT("that is due to this resource not being referenced in the parameters of your pass."),
 			Name);
@@ -483,7 +484,7 @@ void FRDGUserValidation::ValidateAddPass(const void* ParameterStruct, const FSha
 	checkf(ParameterStruct, TEXT("Pass '%s' created with null parameters."), Name.GetTCHAR());
 	ExecuteGuard(TEXT("AddPass"), Name.GetTCHAR());
 
-	checkf(EnumHasAnyFlags(Flags, ERDGPassFlags::CommandMask),
+	checkf(EnumHasAnyFlags(Flags, ERDGPassFlags::Raster | ERDGPassFlags::Compute | ERDGPassFlags::AsyncCompute | ERDGPassFlags::Copy),
 		TEXT("Pass %s must specify at least one of the following flags: (Copy, Compute, AsyncCompute, Raster)"), Name.GetTCHAR());
 
 	checkf(!EnumHasAllFlags(Flags, ERDGPassFlags::Compute | ERDGPassFlags::AsyncCompute),
diff --git a/Engine/Source/Runtime/RenderCore/Public/RenderGraphBuilder.h b/Engine/Source/Runtime/RenderCore/Public/RenderGraphBuilder.h
--- a/Engine/Source/Runtime/RenderCore/Public/RenderGraphBuilder.h
+++ b/Engine/Source/Runtime/RenderCore/Public/RenderGraphBuilder.h
@@ -150,9 +150,8 @@ public:
 
 	/** Adds a lambda pass to the graph without any parameters. This useful for deferring RHI work onto the graph timeline,
 	 *  or incrementally porting code to use the graph system. NeverCull and SkipRenderPass (if Raster) are implicitly added
-	 *  to Flags. AsyncCompute is not allowed. Accessing RHI resources which are tracked by the graph is only allowed if the
-	 *  UntrackedAccess flag is specified. It is never permitted to access a created (i.e. not externally registered) RDG resource
-	 *  in an independent pass, as the RHI lifetime is not guaranteed.
+	 *  to Flags. AsyncCompute is not allowed. It is never permitted to access a created (i.e. not externally registered) RDG
+	 *  resource outside of passes it is registered with, as the RHI lifetime is not guaranteed.
 	 */
 	template <typename ExecuteLambdaType>
 	FRDGPassRef AddPass(FRDGEventName&& Name, ERDGPassFlags Flags, ExecuteLambdaType&& ExecuteLambda);
@@ -328,8 +327,8 @@ private:
 	void BeginResourceRHI(FRDGPassHandle, FRDGBufferSRV* SRV);
 	void BeginResourceRHI(FRDGPassHandle, FRDGBufferUAV* UAV);
 
-	void EndResourceRHI(FRDGPassHandle, FRDGTexture* Texture, uint32 ReferenceCount);
-	void EndResourceRHI(FRDGPassHandle, FRDGBuffer* Buffer, uint32 ReferenceCount);
+	void EndResourceRHI(FRDGPassHandle, FRDGTexture* Texture);
+	void EndResourceRHI(FRDGPassHandle, FRDGBuffer* Buffer);
 
 	void SetupPassInternal(FRDGPass* Pass, FRDGPassHandle PassHandle, ERHIPipeline PassPipeline);
 	void SetupPass(FRDGPass* Pass);
@@ -340,21 +339,20 @@ private:
 	void ExecutePassEpilogue(FRHIComputeCommandList& RHICmdListPass, FRDGPass* Pass);
 
 	void CollectPassResources(FRDGPassHandle PassHandle);
-	void CollectPassBarriers(FRDGPassHandle PassHandle, FRDGPassHandle& LastUntrackedPassHandle);
+	void CollectPassBarriers(FRDGPassHandle PassHandle);
 
 	void AddPassDependency(FRDGPassHandle ProducerHandle, FRDGPassHandle ConsumerHandle);
 
-	void AddEpilogueTransition(FRDGTextureRef Texture, FRDGPassHandle LastUntrackedPassHandle);
-	void AddEpilogueTransition(FRDGBufferRef Buffer, FRDGPassHandle LastUntrackedPassHandle);
+	void AddEpilogueTransition(FRDGTextureRef Texture);
+	void AddEpilogueTransition(FRDGBufferRef Buffer);
 
-	void AddTransition(FRDGPassHandle PassHandle, FRDGTextureRef Texture, const FRDGTextureTransientSubresourceStateIndirect& StateAfter, FRDGPassHandle LastUntrackedPassHandle);
-	void AddTransition(FRDGPassHandle PassHandle, FRDGBufferRef Buffer, FRDGSubresourceState StateAfter, FRDGPassHandle LastUntrackedPassHandle);
+	void AddTransition(FRDGPassHandle PassHandle, FRDGTextureRef Texture, const FRDGTextureTransientSubresourceStateIndirect& StateAfter);
+	void AddTransition(FRDGPassHandle PassHandle, FRDGBufferRef Buffer, FRDGSubresourceState StateAfter);
 
 	void AddTransitionInternal(
 		FRDGParentResource* Resource,
 		FRDGSubresourceState StateBefore,
 		FRDGSubresourceState StateAfter,
-		FRDGPassHandle LastUntrackedPassHandle,
 		const FRHITransitionInfo& TransitionInfo);
 
 	FRHIRenderPassInfo GetRenderPassInfo(const FRDGPass* Pass) const;
