Fixed clang build break.

diff --git a/Engine/Source/Runtime/RenderCore/Public/RenderGraphPass.h b/Engine/Source/Runtime/RenderCore/Public/RenderGraphPass.h
--- a/Engine/Source/Runtime/RenderCore/Public/RenderGraphPass.h
+++ b/Engine/Source/Runtime/RenderCore/Public/RenderGraphPass.h
@@ -234,7 +234,7 @@ public:
 	}
 #endif
 
-private:
+protected:
 	FRDGBarrierBatchBegin& GetPrologueBarriersToBegin(FRDGAllocator& Allocator);
 	FRDGBarrierBatchBegin& GetEpilogueBarriersToBeginForGraphics(FRDGAllocator& Allocator);
 	FRDGBarrierBatchBegin& GetEpilogueBarriersToBeginForAsyncCompute(FRDGAllocator& Allocator);
