Reimplemented optimization to reuse existing shadow mask texture in lighting phase.

#rnx

diff --git a/Engine/Source/Runtime/RenderCore/Private/RenderGraphValidation.cpp b/Engine/Source/Runtime/RenderCore/Private/RenderGraphValidation.cpp
--- a/Engine/Source/Runtime/RenderCore/Private/RenderGraphValidation.cpp
+++ b/Engine/Source/Runtime/RenderCore/Private/RenderGraphValidation.cpp
@@ -774,33 +774,6 @@ void FRDGUserValidation::ValidateAddPass(const FRDGPass* Pass, bool bSkipPassAcc
 
 					CheckNotPassthrough(Texture);
 
-					/** Validate that any previously produced texture contents are loaded. This occurs if the user failed to specify a load action
-					 *  on a texture that was produced by a previous pass, effectively losing that data. This can also happen if the user 're-uses'
-					 *  a texture for some other purpose. The latter is considered bad practice, since it increases memory pressure on the render
-					 *  target pool. Instead, the user should create a new texture instance. An exception to this rule are untracked render targets,
-					 *  which are not actually managed by the render target pool and likely represent the frame buffer.
-					 */
-					if (!bSkipRenderPass)
-					{
-						const bool bIsLoadAction = RenderTarget.GetLoadAction() == ERenderTargetLoadAction::ELoad;
-
-						// Ignore external textures which are always marked as produced. We don't need to enforce this warning on them.
-						const bool bHasBeenProduced = Texture->bProduced && !Texture->bExternal;
-
-						// We only validate single-mip textures since we don't track production at the subresource level.
-						const bool bFailedToLoadProducedContent = !bIsLoadAction && bHasBeenProduced && Texture->Desc.NumMips == 1;
-
-						// Untracked render targets aren't actually managed by the render target pool.
-						const bool bIsUntrackedRenderTarget = Texture->PooledRenderTarget && !Texture->PooledRenderTarget->IsTracked();
-
-						// In multi-gpu, when running with "r.EnableMultiGPUForkAndJoin", it's possible for each GPU to clear the same RT in turn.
-						// When this happens, they are not actually working on the same resource, see for example the implementation of FD3D12MultiNodeGPUObject.
-						ensureMsgf((!bFailedToLoadProducedContent || bIsUntrackedRenderTarget) || (GNumExplicitGPUsForRendering > 1 && RenderTarget.GetLoadAction() == ERenderTargetLoadAction::EClear),
-							TEXT("Pass '%s' attempted to bind texture '%s' as a render target without the 'Load' action specified, despite a prior pass having produced it. It's invalid to completely clobber the contents of a resource. Create a new texture instance instead."),
-							PassName,
-							Texture->Name);
-					}
-
 					/** Mark the pass as a producer for render targets with a store action. */
 					MarkTextureAsProduced(Texture);
 				}
diff --git a/Engine/Source/Runtime/Renderer/Private/LightRendering.cpp b/Engine/Source/Runtime/Renderer/Private/LightRendering.cpp
--- a/Engine/Source/Runtime/Renderer/Private/LightRendering.cpp
+++ b/Engine/Source/Runtime/Renderer/Private/LightRendering.cpp
@@ -1468,6 +1468,9 @@ void FDeferredShadingSceneRenderer::RenderLights(
 
 			const bool bDirectLighting = ViewFamily.EngineShowFlags.DirectLighting;
 
+			FRDGTextureRef SharedScreenShadowMaskTexture = nullptr;
+			FRDGTextureRef SharedScreenShadowMaskSubPixelTexture = nullptr;
+
 			// Draw shadowed and light function lights
 			for (int32 LightIndex = AttenuationLightStart; LightIndex < SortedLights.Num(); LightIndex++)
 			{
@@ -1492,12 +1495,18 @@ void FDeferredShadingSceneRenderer::RenderLights(
 
 				if (bDrawShadows || bDrawLightFunction || bDrawPreviewIndicator)
 				{
-					const FRDGTextureDesc Desc(FRDGTextureDesc::Create2D(SceneTextures.Config.Extent, PF_B8G8R8A8, FClearValueBinding::White, TexCreate_RenderTargetable | TexCreate_ShaderResource | GFastVRamConfig.ScreenSpaceShadowMask));
-					ScreenShadowMaskTexture = GraphBuilder.CreateTexture(Desc, TEXT("ShadowMaskTexture"));
-					if (bUseHairLighting)
+					if (!SharedScreenShadowMaskTexture)
 					{
-						ScreenShadowMaskSubPixelTexture = GraphBuilder.CreateTexture(Desc, TEXT("ShadowMaskSubPixelTexture"));
+						const FRDGTextureDesc SharedScreenShadowMaskTextureDesc(FRDGTextureDesc::Create2D(SceneTextures.Config.Extent, PF_B8G8R8A8, FClearValueBinding::White, TexCreate_RenderTargetable | TexCreate_ShaderResource | GFastVRamConfig.ScreenSpaceShadowMask));
+						SharedScreenShadowMaskTexture = GraphBuilder.CreateTexture(SharedScreenShadowMaskTextureDesc, TEXT("ShadowMaskTexture"));
+
+						if (bUseHairLighting)
+						{
+							SharedScreenShadowMaskSubPixelTexture = GraphBuilder.CreateTexture(SharedScreenShadowMaskTextureDesc, TEXT("ShadowMaskSubPixelTexture"));
+						}
 					}
+					ScreenShadowMaskTexture = SharedScreenShadowMaskTexture;
+					ScreenShadowMaskSubPixelTexture = SharedScreenShadowMaskSubPixelTexture;
 				}
 
 				FString LightNameWithLevel;
