Reimplemented fix for transient resources in RDG.

#rnx

diff --git a/Engine/Source/Runtime/RenderCore/Private/RenderGraphBuilder.cpp b/Engine/Source/Runtime/RenderCore/Private/RenderGraphBuilder.cpp
--- a/Engine/Source/Runtime/RenderCore/Private/RenderGraphBuilder.cpp
+++ b/Engine/Source/Runtime/RenderCore/Private/RenderGraphBuilder.cpp
@@ -463,6 +463,25 @@ void FRDGBuilder::PreallocateTexture(FRDGTextureRef Texture)
 	}
 }
 
+FRDGTextureRef FRDGBuilder::CreateTexture(
+	const FRDGTextureDesc& Desc,
+	const TCHAR* Name,
+	ERDGTextureFlags Flags)
+{
+	FRDGTextureDesc TransientDesc = Desc;
+
+	if (FRenderTargetPool::DoesTargetNeedTransienceOverride(Desc.Flags, ERenderTargetTransience::Transient))
+	{
+		TransientDesc.Flags |= TexCreate_Transient;
+	}
+
+	IF_RDG_ENABLE_DEBUG(UserValidation.ValidateCreateTexture(TransientDesc, Name, Flags));
+	FRDGTextureRef Texture = Textures.Allocate(Allocator, Name, TransientDesc, Flags, ERenderTargetTexture::ShaderResource);
+	IF_RDG_ENABLE_DEBUG(UserValidation.ValidateCreateTexture(Texture));
+	IF_RDG_ENABLE_TRACE(Trace.AddResource(Texture));
+	return Texture;
+}
+
 FRDGTextureRef FRDGBuilder::RegisterExternalTexture(
 	const TRefCountPtr<IPooledRenderTarget>& ExternalPooledTexture,
 	ERenderTargetTexture RenderTargetTexture,
@@ -508,7 +527,7 @@ FRDGTextureRef FRDGBuilder::RegisterExternalTexture(
 	Texture->bExternal = true;
 	Texture->AccessInitial = AccessInitial;
 	Texture->AccessFinal = kDefaultAccessFinal;
-	Texture->AcquirePass = GetProloguePassHandle();
+	Texture->FirstPass = GetProloguePassHandle();
 
 	FRDGTextureSubresourceState& TextureState = Texture->GetState();
 
@@ -568,7 +587,7 @@ FRDGBufferRef FRDGBuilder::RegisterExternalBuffer(
 	Buffer->bExternal = true;
 	Buffer->AccessInitial = AccessInitial;
 	Buffer->AccessFinal = kDefaultAccessFinal;
-	Buffer->AcquirePass = GetProloguePassHandle();
+	Buffer->FirstPass = GetProloguePassHandle();
 
 	FRDGSubresourceState& BufferState = Buffer->GetState();
 	checkf(BufferState.Access == ERHIAccess::Unknown,
@@ -1279,7 +1298,7 @@ void FRDGBuilder::Execute()
 			if (!Resource->bLastOwner)
 			{
 				auto* NextOwner = Registry[Resource->NextOwner];
-				LogFile.AddAliasEdge(Resource, Resource->LastPass, NextOwner, NextOwner->AcquirePass);
+				LogFile.AddAliasEdge(Resource, Resource->LastPass, NextOwner, NextOwner->FirstPass);
 			}
 			LogFile.AddFirstEdge(Resource, Resource->FirstPass);
 		}
@@ -1666,6 +1685,11 @@ void FRDGBuilder::ExecutePass(FRDGPass* Pass)
 {
 	QUICK_SCOPE_CYCLE_COUNTER(STAT_FRDGBuilder_ExecutePass);
 
+	for (FRHITexture* Texture : Pass->TexturesToAcquire)
+	{
+		RHIAcquireTransientResource(Texture);
+	}
+
 	// Note that we must do this before doing anything with RHICmdList for the pass.
 	// For example, if this pass only executes on GPU 1 we want to avoid adding a
 	// 0-duration event for this pass on GPU 0's time line.
@@ -1710,6 +1734,11 @@ void FRDGBuilder::ExecutePass(FRDGPass* Pass)
 
 	ExecutePassEpilogue(RHICmdListPass, Pass);
 
+	for (FRHITexture* Texture : Pass->TexturesToDiscard)
+	{
+		RHIDiscardTransientResource(Texture);
+	}
+
 	if (Pass->bAsyncComputeEnd)
 	{
 		if (Pass->bAsyncComputeEndExecute)
@@ -2056,7 +2085,7 @@ void FRDGBuilder::AddTransitionInternal(
 	FRDGPassHandlesByPipeline& PassesAfter  = StateAfter.FirstPass;
 
 	// Before states may come from previous aliases of the texture.
-	if (Resource->bTransient && StateBefore.GetLastPass() < Resource->AcquirePass)
+	if (Resource->bTransient && StateBefore.GetLastPass() < Resource->FirstPass)
 	{
 		checkf(PipelinesAfter == ERHIPipeline::Graphics, TEXT("The first use of transient resource %s is not exclusively on the graphics pipe."), Resource->Name);
 
@@ -2070,7 +2099,7 @@ void FRDGBuilder::AddTransitionInternal(
 		// Otherwise, can push the start of the transition forward until our alias is acquired.
 		else
 		{
-			PassesBefore[Graphics] = Resource->AcquirePass;
+			PassesBefore[Graphics] = Resource->FirstPass;
 		}
 	}
 
@@ -2227,30 +2256,23 @@ void FRDGBuilder::BeginResourceRHI(FRDGPassHandle PassHandle, FRDGTextureRef Tex
 
 	TRefCountPtr<FPooledRenderTarget> PooledRenderTarget = GRenderTargetPool.FindFreeElementForRDG(RHICmdList, Texture->Desc, Texture->Name);
 
+	const bool bTransient = PooledRenderTarget->IsTransient();
 	FRDGTextureRef PreviousOwner = nullptr;
 	Texture->SetRHI(PooledRenderTarget, PreviousOwner);
-	Texture->bTransient = PooledRenderTarget->IsTransient();
-
-	// Acquires are made in the last pass of the previous alias.
-	FRDGPassHandle AcquirePassHandle = PreviousOwner ? PreviousOwner->LastPass : GetProloguePassHandle();
+	Texture->FirstPass = PassHandle;
 
-	if (Texture->bTransient)
+	if (bTransient)
 	{
-		// We will handle the discard behavior ourselves.
-		PooledRenderTarget->bAutoDiscard = false;
-
-		FRDGPass* AcquirePass = Passes[AcquirePassHandle];
-
-		// Discards for the previous owner occur in the same pass as the acquire.
-		if (PreviousOwner)
+		if (Texture->bExternal)
 		{
-			AcquirePass->TexturesToDiscard.Add(PreviousOwner->GetRHIUnchecked());
+			RHIAcquireTransientResource(Texture->GetRHIUnchecked());
+		}
+		else if (!GRDGImmediateMode)
+		{
+			Texture->bTransient = bTransient;
+			Passes[PassHandle]->TexturesToAcquire.Emplace(Texture->GetRHIUnchecked());
 		}
-
-		AcquirePass->TexturesToAcquire.Add(Texture->GetRHIUnchecked());
 	}
-
-	Texture->AcquirePass = AcquirePassHandle;
 }
 
 void FRDGBuilder::BeginResourceRHI(FRDGPassHandle PassHandle, FRDGTextureSRVRef SRV)
@@ -2266,11 +2288,6 @@ void FRDGBuilder::BeginResourceRHI(FRDGPassHandle PassHandle, FRDGTextureSRVRef
 	FRDGPooledTexture* PooledTexture = Texture->PooledTexture;
 	checkf(PooledTexture, TEXT("Pass parameters contained an SRV of RDG Texture %s, before that texture was referenced as a UAV or RTV, which isn't supported.  Make sure to call ClearUnusedGraphResources to remove unbound parameters which can also cause this."), Texture->Name);
 
-	if (!Texture->FirstPass.IsValid())
-	{
-		Texture->FirstPass = PassHandle;
-	}
-
 	if (SRV->Desc.MetaData == ERDGTextureMetaDataAccess::HTile)
 	{
 		check(GRHISupportsExplicitHTile);
@@ -2382,11 +2399,6 @@ void FRDGBuilder::BeginResourceRHI(FRDGPassHandle PassHandle, FRDGBufferRef Buff
 	FRDGBufferRef PreviousOwner = nullptr;
 	Buffer->SetRHI(PooledBuffer, PreviousOwner);
 	Buffer->FirstPass = PassHandle;
-
-	// Acquires are made in the last pass of the previous alias.
-	Buffer->AcquirePass = PreviousOwner ? PreviousOwner->LastPass : GetProloguePassHandle();
-
-	// Transient buffer are not yet implemented.
 }
 
 void FRDGBuilder::BeginResourceRHI(FRDGPassHandle PassHandle, FRDGBufferSRVRef SRV)
@@ -2433,10 +2445,19 @@ void FRDGBuilder::EndResourceRHI(FRDGPassHandle PassHandle, FRDGTextureRef Textu
 		if (--Texture->ReferenceCount == 0)
 		{
 			// External textures should never release the reference.
-			if (!Texture->bExternal)
+			if (!Texture->bExternal && !Texture->bTransient)
 			{
 				Texture->Allocation = nullptr;
 			}
+
+			// Extracted textures will be discarded on release of the reference.
+			if (Texture->bTransient && !Texture->bExtracted)
+			{
+				Passes[PassHandle]->TexturesToDiscard.Emplace(Texture->GetRHIUnchecked());
+
+				static_cast<FPooledRenderTarget*>(Texture->PooledRenderTarget)->bAutoDiscard = false;
+			}
+
 			Texture->LastPass = PassHandle;
 		}
 	}
diff --git a/Engine/Source/Runtime/RenderCore/Private/RenderGraphResources.cpp b/Engine/Source/Runtime/RenderCore/Private/RenderGraphResources.cpp
--- a/Engine/Source/Runtime/RenderCore/Private/RenderGraphResources.cpp
+++ b/Engine/Source/Runtime/RenderCore/Private/RenderGraphResources.cpp
@@ -274,9 +274,6 @@ void FRDGTexture::Finalize()
 			PooledTexture->Finalize();
 		}
 
-		// Resume automatic discard behavior for transient resources.
-		static_cast<FPooledRenderTarget*>(PooledRenderTarget)->bAutoDiscard = true;
-
 		// Restore the reference to the last owner in the aliasing chain.
 		Allocation = PooledRenderTarget;
 	}
diff --git a/Engine/Source/Runtime/RenderCore/Private/RenderTargetPool.cpp b/Engine/Source/Runtime/RenderCore/Private/RenderTargetPool.cpp
--- a/Engine/Source/Runtime/RenderCore/Private/RenderTargetPool.cpp
+++ b/Engine/Source/Runtime/RenderCore/Private/RenderTargetPool.cpp
@@ -222,7 +222,7 @@ FRenderTargetPool::FRenderTargetPool()
 }
 
 // Logic for determining whether to make a rendertarget transient
-bool FRenderTargetPool::DoesTargetNeedTransienceOverride(const FPooledRenderTargetDesc& InputDesc, ERenderTargetTransience TransienceHint) const
+bool FRenderTargetPool::DoesTargetNeedTransienceOverride(ETextureCreateFlags Flags, ERenderTargetTransience TransienceHint)
 {
 	if (!GSupportsTransientResourceAliasing)
 	{
@@ -231,14 +231,12 @@ bool FRenderTargetPool::DoesTargetNeedTransienceOverride(const FPooledRenderTarg
 	int32 AliasingMode = CVarRtPoolTransientMode.GetValueOnRenderThread();
 
 	// We only override transience if aliasing is supported and enabled, the format is suitable, and the target is not already transient
-	if (AliasingMode > 0 &&
-		(InputDesc.TargetableFlags & (TexCreate_RenderTargetable | TexCreate_DepthStencilTargetable | TexCreate_UAV)) &&
-		((InputDesc.Flags & TexCreate_Transient) == 0))
+	if (AliasingMode > 0 && EnumHasAnyFlags(Flags, TexCreate_RenderTargetable | TexCreate_DepthStencilTargetable | TexCreate_UAV) && !EnumHasAnyFlags(Flags, TexCreate_Transient))
 	{
 		if (AliasingMode == 1)
 		{
 			// Mode 1: Only make FastVRAM rendertargets transient
-			if (InputDesc.Flags & TexCreate_FastVRAM)
+			if (EnumHasAnyFlags(Flags, TexCreate_FastVRAM))
 			{
 				return true;
 			}
@@ -246,7 +244,7 @@ bool FRenderTargetPool::DoesTargetNeedTransienceOverride(const FPooledRenderTarg
 		else if (AliasingMode == 2)
 		{
 			// Mode 2: Make fastvram and ERenderTargetTransience::Transient rendertargets transient
-			if (InputDesc.Flags & TexCreate_FastVRAM || TransienceHint == ERenderTargetTransience::Transient)
+			if (EnumHasAnyFlags(Flags, TexCreate_FastVRAM) || TransienceHint == ERenderTargetTransience::Transient)
 			{
 				return true;
 			}
@@ -645,7 +643,7 @@ bool FRenderTargetPool::FindFreeElement(
 
 	// If we're doing aliasing, we may need to override Transient flags, depending on the input format and mode
 	FPooledRenderTargetDesc ModifiedDesc;
-	bool bMakeTransient = DoesTargetNeedTransienceOverride(InputDesc, TransienceHint);
+	bool bMakeTransient = DoesTargetNeedTransienceOverride(InputDesc.Flags, TransienceHint);
 	if (bMakeTransient)
 	{
 		ModifiedDesc = InputDesc;
diff --git a/Engine/Source/Runtime/RenderCore/Public/RenderGraphBuilder.inl b/Engine/Source/Runtime/RenderCore/Public/RenderGraphBuilder.inl
--- a/Engine/Source/Runtime/RenderCore/Public/RenderGraphBuilder.inl
+++ b/Engine/Source/Runtime/RenderCore/Public/RenderGraphBuilder.inl
@@ -20,18 +20,6 @@ inline FRDGTextureRef FRDGBuilder::FindExternalTexture(IPooledRenderTarget* Exte
 	return nullptr;
 }
 
-inline FRDGTextureRef FRDGBuilder::CreateTexture(
-	const FRDGTextureDesc& Desc,
-	const TCHAR* Name,
-	ERDGTextureFlags Flags)
-{
-	IF_RDG_ENABLE_DEBUG(UserValidation.ValidateCreateTexture(Desc, Name, Flags));
-	FRDGTextureRef Texture = Textures.Allocate(Allocator, Name, Desc, Flags, ERenderTargetTexture::ShaderResource);
-	IF_RDG_ENABLE_DEBUG(UserValidation.ValidateCreateTexture(Texture));
-	IF_RDG_ENABLE_TRACE(Trace.AddResource(Texture));
-	return Texture;
-}
-
 inline FRDGBufferRef FRDGBuilder::CreateBuffer(
 	const FRDGBufferDesc& Desc,
 	const TCHAR* Name,
diff --git a/Engine/Source/Runtime/RenderCore/Public/RenderGraphResources.h b/Engine/Source/Runtime/RenderCore/Public/RenderGraphResources.h
--- a/Engine/Source/Runtime/RenderCore/Public/RenderGraphResources.h
+++ b/Engine/Source/Runtime/RenderCore/Public/RenderGraphResources.h
@@ -269,7 +269,6 @@ private:
 	ERHIAccess AccessInitial = ERHIAccess::Unknown;
 	ERHIAccess AccessFinal = ERHIAccess::Unknown;
 
-	FRDGPassHandle AcquirePass;
 	FRDGPassHandle FirstPass;
 	FRDGPassHandle LastPass;
 
@@ -426,7 +425,7 @@ struct RENDERCORE_API FRDGTextureDesc
 
 		// Remove UAV flag for textures that don't need it (some formats are incompatible).
 		Flags |= TexCreate_RenderTargetable;
-		Flags &= ~(TexCreate_UAV | TexCreate_FastVRAM | TexCreate_Transient | TexCreate_ResolveTargetable | TexCreate_DepthStencilResolveTarget);
+		Flags &= ~(TexCreate_UAV | TexCreate_ResolveTargetable | TexCreate_DepthStencilResolveTarget);
 	}
 
 	bool IsTexture2D() const
diff --git a/Engine/Source/Runtime/RenderCore/Public/RenderTargetPool.h b/Engine/Source/Runtime/RenderCore/Public/RenderTargetPool.h
--- a/Engine/Source/Runtime/RenderCore/Public/RenderTargetPool.h
+++ b/Engine/Source/Runtime/RenderCore/Public/RenderTargetPool.h
@@ -335,7 +335,7 @@ private:
 		bool bDeferTextureAllocation,
 		bool bDoAcquireTransientResource);
 
-	bool DoesTargetNeedTransienceOverride(const FPooledRenderTargetDesc& InputDesc, ERenderTargetTransience TransienceHint) const;
+	static bool DoesTargetNeedTransienceOverride(ETextureCreateFlags Flags, ERenderTargetTransience TransienceHint);
 
 	friend void RenderTargetPoolEvents(const TArray<FString>& Args);
 
