Optimized string usage in ReplaceVirtualFilePathForShaderAutogen from ShaderCore.cpp
    * no more allocs (previously this function made 3 to 6 allocs per call)
    * 5x faster code

diff --git a/Engine/Source/Runtime/RenderCore/Private/ShaderCore.cpp b/Engine/Source/Runtime/RenderCore/Private/ShaderCore.cpp
--- a/Engine/Source/Runtime/RenderCore/Private/ShaderCore.cpp
+++ b/Engine/Source/Runtime/RenderCore/Private/ShaderCore.cpp
@@ -11,6 +11,7 @@
 #include "Misc/ScopeLock.h"
 #include "Stats/StatsMisc.h"
 #include "Misc/ScopedSlowTask.h"
+#include "Misc/StringBuilder.h"
 #include "Shader.h"
 #include "ShaderCompilerCore.h"
 #include "VertexFactory.h"
@@ -716,17 +717,29 @@ bool ReplaceVirtualFilePathForShaderPlatform(FString& InOutVirtualFilePath, ESha
 
 bool ReplaceVirtualFilePathForShaderAutogen(FString& InOutVirtualFilePath, EShaderPlatform ShaderPlatform)
 {
-	// Tweak the autogen path
-	const FString ShaderAutogenStem = TEXT("/Engine/Generated/ShaderAutogen/");
-	const FString PlatformName = LegacyShaderPlatformToShaderFormat(ShaderPlatform).GetPlainNameString();
-	const FString ShaderAutogenPlatformStem = TEXT("/ShaderAutogen/") + PlatformName;
+	const FStringView ShaderAutogenStem = TEXT("/Engine/Generated/ShaderAutogen/"_SV);
 
+	// Tweak the autogen path
 	// for examples, if it starts with "/Engine/Generated/ShaderAutogen/" change it to "ShaderAutogen/PCD3D_SM5/"
-	if (InOutVirtualFilePath.StartsWith(ShaderAutogenStem))
+	if (!FCString::Strnicmp(*InOutVirtualFilePath, ShaderAutogenStem.GetData(), ShaderAutogenStem.Len()))
 	{
-		// Parse the right side, after /ShaderAutogen/
-		FString RelativeShaderName = InOutVirtualFilePath.RightChop(ShaderAutogenStem.Len());
-		FString OutputShaderName = ShaderAutogenPlatformStem + TEXT("/") + RelativeShaderName;
+		TStringBuilder<1024> OutputShaderName;
+
+		// Append the prefix.
+		OutputShaderName.Append(TEXT("/ShaderAutogen/"_SV)); // using FStringView (_SV suffix) in order to avoid a strlen call
+
+		// Append the platform name.
+		const FName PlatformName = LegacyShaderPlatformToShaderFormat(ShaderPlatform);
+		TCHAR PlatformNameString[NAME_SIZE];
+		const uint32 PlatformNameLen = PlatformName.GetPlainNameString(PlatformNameString);
+		OutputShaderName.Append(PlatformNameString, PlatformNameLen);
+
+		OutputShaderName.Append(TEXT('/'));
+
+		// Append the relative name (the substring after "/Engine/Generated/ShaderAutogen/").
+		const TCHAR* RelativeShaderName = *InOutVirtualFilePath + ShaderAutogenStem.Len();
+		OutputShaderName.Append(RelativeShaderName, InOutVirtualFilePath.Len() - ShaderAutogenStem.Len());
+
 		InOutVirtualFilePath = OutputShaderName;
 		return true;
 	}
