UE5EA - Disallow arrays of scalars in uniform buffers

diff --git a/Engine/Source/Runtime/RenderCore/Private/ShaderParameters.cpp b/Engine/Source/Runtime/RenderCore/Private/ShaderParameters.cpp
--- a/Engine/Source/Runtime/RenderCore/Private/ShaderParameters.cpp
+++ b/Engine/Source/Runtime/RenderCore/Private/ShaderParameters.cpp
@@ -141,9 +141,11 @@ static void CreateHLSLUniformBufferStructMembersDeclaration(
 		const FShaderParametersMetadata::FMember& Member = StructMembers[MemberIndex];
 		
 		FString ArrayDim;
+		bool bIsArray = false;
 		if(Member.GetNumElements() > 0)
 		{
 			ArrayDim = FString::Printf(TEXT("[%u]"),Member.GetNumElements());
+			bIsArray = true;
 		}
 
 		if(Member.GetBaseType() == UBMT_NESTED_STRUCT)
@@ -169,11 +171,13 @@ static void CreateHLSLUniformBufferStructMembersDeclaration(
 		{
 			// Generate the base type name.
 			FString BaseTypeName;
+			bool bIsNumeric = false;
 			switch(Member.GetBaseType())
 			{
-			case UBMT_INT32:   BaseTypeName = TEXT("int"); break;
-			case UBMT_UINT32:  BaseTypeName = TEXT("uint"); break;
+			case UBMT_INT32:   bIsNumeric = true; BaseTypeName = TEXT("int"); break;
+			case UBMT_UINT32:  bIsNumeric = true; BaseTypeName = TEXT("uint"); break;
 			case UBMT_FLOAT32: 
+				bIsNumeric = true;
 				if (Member.GetPrecision() == EShaderPrecisionModifier::Float || !SupportShaderPrecisionModifier(Platform))
 				{
 					BaseTypeName = TEXT("float"); 
@@ -190,6 +194,11 @@ static void CreateHLSLUniformBufferStructMembersDeclaration(
 			default:           UE_LOG(LogShaders, Fatal,TEXT("Unrecognized uniform buffer struct member base type."));
 			};
 
+			if (bIsArray && bIsNumeric && Member.GetNumColumns() < 3)
+			{
+				UE_LOG(LogShaders, Fatal, TEXT("Arrays of scalars cause padding/offset issues on different platforms. Please use vector elements. (%s, member %s)"), UniformBufferStruct.GetStructTypeName(), Member.GetName());
+			}
+
 			// Generate the type dimensions for vectors and matrices.
 			FString TypeDim;
 			uint32 HLSLMemberSize = 4;
