Crash when building HLODs from the Build menu (Build -> Build LODs)
Rendering stopped directly with StopRenderingThread() does not change GUseThreadedRendering...

diff --git a/Engine/Source/Runtime/RenderCore/Private/RenderCaptureInterface.cpp b/Engine/Source/Runtime/RenderCore/Private/RenderCaptureInterface.cpp
--- a/Engine/Source/Runtime/RenderCore/Private/RenderCaptureInterface.cpp
+++ b/Engine/Source/Runtime/RenderCore/Private/RenderCaptureInterface.cpp
@@ -59,7 +59,7 @@ namespace RenderCaptureInterface
 		: bCapture(bEnable)
 		, RHICmdList(nullptr)
 	{
-		check(!GUseThreadedRendering || !IsInRenderingThread());
+		check(!GIsThreadedRendering || !IsInRenderingThread());
 
 		if (bCapture)
 		{
@@ -74,7 +74,7 @@ namespace RenderCaptureInterface
 		: bCapture(bEnable)
 		, RHICmdList(RHICommandList)
 	{
-		check(!GUseThreadedRendering || IsInRenderingThread());
+		check(!GIsThreadedRendering || IsInRenderingThread());
 
 		if (bCapture)
 		{
@@ -88,13 +88,13 @@ namespace RenderCaptureInterface
 		{
 			if (RHICmdList != nullptr)
 			{
-				check(!GUseThreadedRendering || IsInRenderingThread());
+				check(!GIsThreadedRendering || IsInRenderingThread());
 				
 				EndCapture(RHICmdList);
 			}
 			else
 			{
-				check(!GUseThreadedRendering || !IsInRenderingThread());
+				check(!GIsThreadedRendering || !IsInRenderingThread());
 
 				ENQUEUE_RENDER_COMMAND(CaptureCommand)([](FRHICommandListImmediate& RHICommandListLocal)
 				{
