Fixed compile errors when static analyzer is enabled in Test configs.

Static analyzer always evaluates expressions in check(), despite DO_CHECK=0 due to implementation of check():

    #define check(expr) { CA_ASSUME(expr); }

If checked expression depends on something that's wrapped in `#if DO_CHECK`, it will produce a compile error. The fix is to add USING_CODE_ANALYSIS condition where applicable.

diff --git a/Engine/Source/Runtime/CoreUObject/Private/Serialization/AsyncLoading.cpp b/Engine/Source/Runtime/CoreUObject/Private/Serialization/AsyncLoading.cpp
--- a/Engine/Source/Runtime/CoreUObject/Private/Serialization/AsyncLoading.cpp
+++ b/Engine/Source/Runtime/CoreUObject/Private/Serialization/AsyncLoading.cpp
@@ -3912,7 +3912,7 @@ void FAsyncLoadingThread::InsertPackage(FAsyncPackage* Package, bool bReinsert,
 	checkSlow(IsInAsyncLoadThread());
 	check(!IsInGameThread() || !IsMultithreaded());
 
-#if DO_CHECK
+#if DO_CHECK || USING_CODE_ANALYSIS
 	FWeakAsyncPackagePtr WeakPtr;
 	if (GEventDrivenLoaderEnabled)
 	{
diff --git a/Engine/Source/Runtime/CoreUObject/Private/Serialization/UnversionedPropertySerialization.cpp b/Engine/Source/Runtime/CoreUObject/Private/Serialization/UnversionedPropertySerialization.cpp
--- a/Engine/Source/Runtime/CoreUObject/Private/Serialization/UnversionedPropertySerialization.cpp
+++ b/Engine/Source/Runtime/CoreUObject/Private/Serialization/UnversionedPropertySerialization.cpp
@@ -569,7 +569,7 @@ public:
 			, ZeroMask(Header.ZeroMask)
 			, FragmentIt(Header.Fragments.GetData())
 			, bDone(!Header.HasValues())
-#if DO_CHECK
+#if DO_CHECK || USING_CODE_ANALYSIS
 			, SchemaEnd(Schema.End)
 #endif
 		{
@@ -622,7 +622,7 @@ public:
 		bool bDone = false;
 		uint32 ZeroMaskIndex = 0;
 		uint32 RemainingFragmentValues = 0;
-#if DO_CHECK
+#if DO_CHECK || USING_CODE_ANALYSIS
 		FUnversionedSchemaIterator SchemaEnd;
 #endif
 
diff --git a/Engine/Source/Runtime/D3D12RHI/Private/D3D12View.h b/Engine/Source/Runtime/D3D12RHI/Private/D3D12View.h
--- a/Engine/Source/Runtime/D3D12RHI/Private/D3D12View.h
+++ b/Engine/Source/Runtime/D3D12RHI/Private/D3D12View.h
@@ -680,7 +680,7 @@ protected:
 	CViewSubresourceSubset ViewSubresourceSubset;
 	TDesc Desc;
 
-#if DO_CHECK
+#if DO_CHECK || USING_CODE_ANALYSIS
 	bool bInitialized;
 #endif
 
@@ -688,7 +688,7 @@ protected:
 		: Descriptor(InParent)
 		, Flags(InFlags)
 		, BaseShaderResource(nullptr)
-#if DO_CHECK
+#if DO_CHECK || USING_CODE_ANALYSIS
 		, bInitialized(false)
 #endif
 	{}
@@ -700,7 +700,7 @@ protected:
 			BaseShaderResource->RemoveRenameListener(this);
 		}
 
-#if DO_CHECK
+#if DO_CHECK || USING_CODE_ANALYSIS
 		bInitialized = false;
 #endif
 	}
@@ -731,7 +731,7 @@ private:
 			Resource->GetDesc().Format,
 			Flags);
 
-#if DO_CHECK
+#if DO_CHECK || USING_CODE_ANALYSIS
 		bInitialized = true;
 #endif
 	}
diff --git a/Engine/Source/Runtime/Engine/Classes/GameFramework/Actor.h b/Engine/Source/Runtime/Engine/Classes/GameFramework/Actor.h
--- a/Engine/Source/Runtime/Engine/Classes/GameFramework/Actor.h
+++ b/Engine/Source/Runtime/Engine/Classes/GameFramework/Actor.h
@@ -3352,7 +3352,7 @@ public:
 	 */
 	void RemoveOwnedComponent(UActorComponent* Component);
 
-#if DO_CHECK
+#if DO_CHECK || USING_CODE_ANALYSIS
 	/** Utility function for validating that a component is correctly in its Owner's OwnedComponents array */
 	bool OwnsComponent(UActorComponent* Component) const;
 #endif
diff --git a/Engine/Source/Runtime/Engine/Private/Actor.cpp b/Engine/Source/Runtime/Engine/Private/Actor.cpp
--- a/Engine/Source/Runtime/Engine/Private/Actor.cpp
+++ b/Engine/Source/Runtime/Engine/Private/Actor.cpp
@@ -2936,7 +2936,7 @@ void AActor::RemoveOwnedComponent(UActorComponent* Component)
 	}
 }
 
-#if DO_CHECK
+#if DO_CHECK || USING_CODE_ANALYSIS
 bool AActor::OwnsComponent(UActorComponent* Component) const
 {
 	return OwnedComponents.Contains(Component);
diff --git a/Engine/Source/Runtime/Engine/Private/SceneView.cpp b/Engine/Source/Runtime/Engine/Private/SceneView.cpp
--- a/Engine/Source/Runtime/Engine/Private/SceneView.cpp
+++ b/Engine/Source/Runtime/Engine/Private/SceneView.cpp
@@ -857,7 +857,7 @@ FSceneView::FSceneView(const FSceneViewInitOptions& InitOptions)
 }
 
 
-#if DO_CHECK
+#if DO_CHECK || USING_CODE_ANALYSIS
 bool FSceneView::VerifyMembersChecks() const
 {
 	if (PrimaryScreenPercentageMethod == EPrimaryScreenPercentageMethod::TemporalUpscale)
diff --git a/Engine/Source/Runtime/Engine/Public/SceneView.h b/Engine/Source/Runtime/Engine/Public/SceneView.h
--- a/Engine/Source/Runtime/Engine/Public/SceneView.h
+++ b/Engine/Source/Runtime/Engine/Public/SceneView.h
@@ -1208,7 +1208,7 @@ public:
 	/** Initialization constructor. */
 	FSceneView(const FSceneViewInitOptions& InitOptions);
 
-#if DO_CHECK
+#if DO_CHECK || USING_CODE_ANALYSIS
 	/** Verifies all the assertions made on members. */
 	bool VerifyMembersChecks() const;
 #endif
@@ -1446,7 +1446,7 @@ struct FSceneViewScreenPercentageConfig
 	static constexpr float kMaxTAAUpsampleResolutionFraction = 2.0f;
 
 
-	#if DO_CHECK
+	#if DO_CHECK || USING_CODE_ANALYSIS
 		static bool IsValidResolutionFraction(float ResolutionFraction)
 		{
 			return ResolutionFraction >= kMinResolutionFraction && ResolutionFraction <= kMaxResolutionFraction;
diff --git a/Engine/Source/Runtime/PropertyPath/Private/PropertyPathHelpers.cpp b/Engine/Source/Runtime/PropertyPath/Private/PropertyPathHelpers.cpp
--- a/Engine/Source/Runtime/PropertyPath/Private/PropertyPathHelpers.cpp
+++ b/Engine/Source/Runtime/PropertyPath/Private/PropertyPathHelpers.cpp
@@ -815,7 +815,7 @@ bool FCachedPropertyPath::Equals(const FString& Other) const
 	return ToString() == Other;
 }
 
-#if DO_CHECK
+#if DO_CHECK || USING_CODE_ANALYSIS
 void* FCachedPropertyPath::GetCachedContainer() const
 {
 	return CachedContainer;
diff --git a/Engine/Source/Runtime/PropertyPath/Public/PropertyPathHelpers.h b/Engine/Source/Runtime/PropertyPath/Public/PropertyPathHelpers.h
--- a/Engine/Source/Runtime/PropertyPath/Public/PropertyPathHelpers.h
+++ b/Engine/Source/Runtime/PropertyPath/Public/PropertyPathHelpers.h
@@ -170,7 +170,7 @@ struct PROPERTYPATH_API FCachedPropertyPath
 	/** Compares this property path to a string */
 	bool Equals(const FString& Other) const;
 
-#if DO_CHECK
+#if DO_CHECK || USING_CODE_ANALYSIS
 	/** Get the cached container for this property path, for checking purposes */
 	void* GetCachedContainer() const;
 
@@ -198,7 +198,7 @@ private:
 	UPROPERTY()
 	mutable TObjectPtr<UFunction> CachedFunction;
 
-#if DO_CHECK
+#if DO_CHECK || USING_CODE_ANALYSIS
 	/** Cached container, used for checking purposes only */
 	mutable void* CachedContainer;
 #endif
diff --git a/Engine/Source/Runtime/RHI/Public/RHI.h b/Engine/Source/Runtime/RHI/Public/RHI.h
--- a/Engine/Source/Runtime/RHI/Public/RHI.h
+++ b/Engine/Source/Runtime/RHI/Public/RHI.h
@@ -1982,7 +1982,7 @@ private:
 	// Private constructor. Memory for transitions is allocated manually with extra space at the tail of the structure for RHI use.
 	FRHITransition(ERHIPipeline SrcPipelines, ERHIPipeline DstPipelines)
 		: State(int8(SrcPipelines) | (int8(DstPipelines) << int32(ERHIPipeline::Num)))
-#if DO_CHECK
+#if DO_CHECK || USING_CODE_ANALYSIS
 		, AllowedSrc(SrcPipelines)
 		, AllowedDst(DstPipelines)
 #endif
@@ -2042,7 +2042,7 @@ private:
 	mutable int8 State;
 	static_assert((int32(ERHIPipeline::Num) * 2) < (sizeof(State) * 8), "Not enough bits to hold pipeline state.");
 
-#if DO_CHECK
+#if DO_CHECK || USING_CODE_ANALYSIS
 	mutable ERHIPipeline AllowedSrc;
 	mutable ERHIPipeline AllowedDst;
 #endif
diff --git a/Engine/Source/Runtime/RHI/Public/RHICommandList.h b/Engine/Source/Runtime/RHI/Public/RHICommandList.h
--- a/Engine/Source/Runtime/RHI/Public/RHICommandList.h
+++ b/Engine/Source/Runtime/RHI/Public/RHICommandList.h
@@ -1654,14 +1654,14 @@ struct FLocalUniformBufferWorkArea
 	void* Contents;
 	const FRHIUniformBufferLayout* Layout;
 	FComputedUniformBuffer* ComputedUniformBuffer;
-#if DO_CHECK // the below variables are used in check(), which can be enabled in Shipping builds (see Build.h)
+#if DO_CHECK || USING_CODE_ANALYSIS // the below variables are used in check(), which can be enabled in Shipping builds (see Build.h)
 	FRHICommandListBase* CheckCmdList;
 	int32 UID;
 #endif
 
 	FLocalUniformBufferWorkArea(FRHICommandListBase* InCheckCmdList, const void* InContents, uint32 ContentsSize, const FRHIUniformBufferLayout* InLayout)
 		: Layout(InLayout)
-#if DO_CHECK
+#if DO_CHECK || USING_CODE_ANALYSIS
 		, CheckCmdList(InCheckCmdList)
 		, UID(InCheckCmdList->GetUID())
 #endif
diff --git a/Engine/Source/Runtime/RenderCore/Private/ProfilingDebugging/RealtimeGPUProfiler.cpp b/Engine/Source/Runtime/RenderCore/Private/ProfilingDebugging/RealtimeGPUProfiler.cpp
--- a/Engine/Source/Runtime/RenderCore/Private/ProfilingDebugging/RealtimeGPUProfiler.cpp
+++ b/Engine/Source/Runtime/RenderCore/Private/ProfilingDebugging/RealtimeGPUProfiler.cpp
@@ -113,7 +113,7 @@ public:
 		, StartQuery(RenderQueryPool.AllocateQuery())
 		, EndQuery(RenderQueryPool.AllocateQuery())
 		, FrameNumber(-1)
-#if DO_CHECK
+#if DO_CHECK || USING_CODE_ANALYSIS
 		, bInsideQuery(false)
 #endif
 	{
@@ -124,7 +124,7 @@ public:
 	{
 		check(IsInRenderingThread());
 		check(!bInsideQuery && StartQuery.IsValid());
-#if DO_CHECK
+#if DO_CHECK || USING_CODE_ANALYSIS
 		bInsideQuery = true;
 #endif
 		GPUMask = RHICmdList.GetGPUMask();
@@ -141,7 +141,7 @@ public:
 	{
 		check(IsInRenderingThread());
 		check(bInsideQuery && EndQuery.IsValid());
-#if DO_CHECK
+#if DO_CHECK || USING_CODE_ANALYSIS
 		bInsideQuery = false;
 #endif
 		SCOPED_GPU_MASK(RHICmdList, GPUMask);
@@ -258,7 +258,7 @@ private:
 
 	uint32 FrameNumber;
 
-#if DO_CHECK
+#if DO_CHECK || USING_CODE_ANALYSIS
 	bool bInsideQuery;
 #endif
 };
diff --git a/Engine/Source/Runtime/Renderer/Private/SceneRendering.cpp b/Engine/Source/Runtime/Renderer/Private/SceneRendering.cpp
--- a/Engine/Source/Runtime/Renderer/Private/SceneRendering.cpp
+++ b/Engine/Source/Runtime/Renderer/Private/SceneRendering.cpp
@@ -978,7 +978,7 @@ FViewInfo::~FViewInfo()
 	TranslucentSelfShadowUniformBufferMap.Empty();
 }
 
-#if DO_CHECK
+#if DO_CHECK || USING_CODE_ANALYSIS
 bool FViewInfo::VerifyMembersChecks() const
 {
 	FSceneView::VerifyMembersChecks();
diff --git a/Engine/Source/Runtime/Renderer/Private/SceneRendering.h b/Engine/Source/Runtime/Renderer/Private/SceneRendering.h
--- a/Engine/Source/Runtime/Renderer/Private/SceneRendering.h
+++ b/Engine/Source/Runtime/Renderer/Private/SceneRendering.h
@@ -1435,7 +1435,7 @@ public:
 	*/
 	RENDERER_API ~FViewInfo();
 
-#if DO_CHECK
+#if DO_CHECK || USING_CODE_ANALYSIS
 	/** Verifies all the assertions made on members. */
 	bool VerifyMembersChecks() const;
 #endif
