Fixed bad merge of 4.26 transient resource fixes in RDG.

#rnx

diff --git a/Engine/Source/Runtime/RenderCore/Private/RenderGraphBuilder.cpp b/Engine/Source/Runtime/RenderCore/Private/RenderGraphBuilder.cpp
--- a/Engine/Source/Runtime/RenderCore/Private/RenderGraphBuilder.cpp
+++ b/Engine/Source/Runtime/RenderCore/Private/RenderGraphBuilder.cpp
@@ -1288,12 +1288,12 @@ void FRDGBuilder::Execute()
 
 			for (const auto& Query : ExtractedTextures)
 			{
-				EndResourceRHI(EpiloguePassHandle, Query.Key);
+				EndResourceRHI(EpiloguePassHandle, Query.Key, 1);
 			}
 
 			for (const auto& Query : ExtractedBuffers)
 			{
-				EndResourceRHI(EpiloguePassHandle, Query.Key);
+				EndResourceRHI(EpiloguePassHandle, Query.Key, 1);
 			}
 		}
 
@@ -1656,16 +1656,6 @@ void FRDGBuilder::ExecutePassEpilogue(FRHIComputeCommandList& RHICmdListPass, FR
 		static_cast<FRHICommandList&>(RHICmdListPass).EndRenderPass();
 	}
 
-	for (FRHITexture* Texture : Pass->TexturesToDiscard)
-	{
-		RHIDiscardTransientResource(Texture);
-	}
-
-	for (FRHITexture* Texture : Pass->TexturesToAcquire)
-	{
-		RHIAcquireTransientResource(Texture);
-	}
-
 	FRDGTransitionQueue Transitions;
 
 	if (Pass->EpilogueBarriersToBeginForGraphics)
@@ -1808,12 +1798,12 @@ void FRDGBuilder::CollectPassResources(FRDGPassHandle PassHandle)
 	{
 		for (const auto& TexturePair : PassToEnd->TextureStates)
 		{
-			EndResourceRHI(PassHandle, TexturePair.Key);
+			EndResourceRHI(PassHandle, TexturePair.Key, TexturePair.Value.ReferenceCount);
 		}
 
 		for (const auto& BufferPair : PassToEnd->BufferStates)
 		{
-			EndResourceRHI(PassHandle, BufferPair.Key);
+			EndResourceRHI(PassHandle, BufferPair.Key, BufferPair.Value.ReferenceCount);
 		}
 	}
 }
@@ -2457,14 +2447,16 @@ void FRDGBuilder::BeginResourceRHI(FRDGPassHandle PassHandle, FRDGBufferUAV* UAV
 	UAV->ResourceRHI = Buffer->PooledBuffer->GetOrCreateUAV(UAV->Desc);
 }
 
-void FRDGBuilder::EndResourceRHI(FRDGPassHandle PassHandle, FRDGTextureRef Texture)
+void FRDGBuilder::EndResourceRHI(FRDGPassHandle PassHandle, FRDGTextureRef Texture, uint32 ReferenceCount)
 {
 	check(Texture);
 
 	if (!IsResourceLifetimeExtended())
 	{
-		check(Texture->ReferenceCount > 0);
-		if (--Texture->ReferenceCount == 0)
+		check(Texture->ReferenceCount >= ReferenceCount);
+		Texture->ReferenceCount -= ReferenceCount;
+
+		if (Texture->ReferenceCount == 0)
 		{
 			// External textures should never release the reference.
 			if (!Texture->bExternal && !Texture->bTransient)
@@ -2485,14 +2477,16 @@ void FRDGBuilder::EndResourceRHI(FRDGPassHandle PassHandle, FRDGTextureRef Textu
 	}
 }
 
-void FRDGBuilder::EndResourceRHI(FRDGPassHandle PassHandle, FRDGBufferRef Buffer)
+void FRDGBuilder::EndResourceRHI(FRDGPassHandle PassHandle, FRDGBufferRef Buffer, uint32 ReferenceCount)
 {
 	check(Buffer);
 
 	if (!IsResourceLifetimeExtended())
 	{
-		check(Buffer->ReferenceCount > 0);
-		if (--Buffer->ReferenceCount == 0)
+		check(Buffer->ReferenceCount >= ReferenceCount);
+		Buffer->ReferenceCount -= ReferenceCount;
+
+		if (Buffer->ReferenceCount == 0)
 		{
 			// External buffers should never release the reference.
 			if (!Buffer->bExternal)
diff --git a/Engine/Source/Runtime/RenderCore/Private/RenderGraphResources.cpp b/Engine/Source/Runtime/RenderCore/Private/RenderGraphResources.cpp
--- a/Engine/Source/Runtime/RenderCore/Private/RenderGraphResources.cpp
+++ b/Engine/Source/Runtime/RenderCore/Private/RenderGraphResources.cpp
@@ -33,6 +33,11 @@ FRDGParentResource::FRDGParentResource(const TCHAR* InName, const ERDGParentReso
 	, bUsedByAsyncComputePass(0)
 {}
 
+FRDGParentResource::~FRDGParentResource()
+{
+	check(GRDGImmediateMode || ReferenceCount == 0);
+}
+
 void FRDGParentResource::SetPassthroughRHI(FRHIResource* InResourceRHI)
 {
 	ResourceRHI = InResourceRHI;
diff --git a/Engine/Source/Runtime/RenderCore/Public/RenderGraphBuilder.h b/Engine/Source/Runtime/RenderCore/Public/RenderGraphBuilder.h
--- a/Engine/Source/Runtime/RenderCore/Public/RenderGraphBuilder.h
+++ b/Engine/Source/Runtime/RenderCore/Public/RenderGraphBuilder.h
@@ -327,8 +327,8 @@ private:
 	void BeginResourceRHI(FRDGPassHandle, FRDGBufferSRV* SRV);
 	void BeginResourceRHI(FRDGPassHandle, FRDGBufferUAV* UAV);
 
-	void EndResourceRHI(FRDGPassHandle, FRDGTexture* Texture);
-	void EndResourceRHI(FRDGPassHandle, FRDGBuffer* Buffer);
+	void EndResourceRHI(FRDGPassHandle, FRDGTexture* Texture, uint32 ReferenceCount);
+	void EndResourceRHI(FRDGPassHandle, FRDGBuffer* Buffer, uint32 ReferenceCount);
 
 	void SetupPassInternal(FRDGPass* Pass, FRDGPassHandle PassHandle, ERHIPipeline PassPipeline);
 	void SetupPass(FRDGPass* Pass);
diff --git a/Engine/Source/Runtime/RenderCore/Public/RenderGraphResources.h b/Engine/Source/Runtime/RenderCore/Public/RenderGraphResources.h
--- a/Engine/Source/Runtime/RenderCore/Public/RenderGraphResources.h
+++ b/Engine/Source/Runtime/RenderCore/Public/RenderGraphResources.h
@@ -236,6 +236,7 @@ public:
 
 protected:
 	FRDGParentResource(const TCHAR* InName, ERDGParentResourceType InType);
+	~FRDGParentResource();
 
 	/** Whether this is an externally registered resource. */
 	uint8 bExternal : 1;
@@ -1221,6 +1222,12 @@ private:
 	/** Assigns the pooled buffer to this buffer; returns the previous buffer to own the allocation. */
 	void SetRHI(FRDGPooledBuffer* InPooledBuffer, FRDGBufferRef& OutPreviousOwner);
 
+	/** Returns RHI buffer without access checks. */
+	FRHIBuffer* GetRHIUnchecked() const
+	{
+		return static_cast<FRHIBuffer*>(FRDGResource::GetRHIUnchecked());
+	}
+
 	/** Finalizes the buffer for execution; no other transitions are allowed after calling this. */
 	void Finalize();
 
