Use a simple array to hold ForceBuildRayTracingGeometries instead of TSet. Rename FRayTracingGeometryManager::ForceBuild -> ForceBuildIfPending and allow duplicates in the geometry list.

diff --git a/Engine/Source/Runtime/RenderCore/Private/RayTracingGeometryManager.cpp b/Engine/Source/Runtime/RenderCore/Private/RayTracingGeometryManager.cpp
--- a/Engine/Source/Runtime/RenderCore/Private/RayTracingGeometryManager.cpp
+++ b/Engine/Source/Runtime/RenderCore/Private/RayTracingGeometryManager.cpp
@@ -92,17 +92,23 @@ void FRayTracingGeometryManager::BoostPriority(BuildRequestIndex InRequestIndex,
 	GeometryBuildRequests[InRequestIndex].BuildPriority += InBoostValue;
 }
 
-void FRayTracingGeometryManager::ForceBuild(FRHIComputeCommandList& InCmdList, const TArrayView<const FRayTracingGeometry*> InGeometries)
+void FRayTracingGeometryManager::ForceBuildIfPending(FRHIComputeCommandList& InCmdList, const TArrayView<const FRayTracingGeometry*> InGeometries)
 {
 	FScopeLock ScopeLock(&RequestCS);
 
 	BuildParams.Empty(FMath::Max(BuildParams.Max(), InGeometries.Num()));
 	for (const FRayTracingGeometry* Geometry : InGeometries)
 	{
-		check(Geometry->HasPendingBuildRequest());
-		SetupBuildParams(GeometryBuildRequests[Geometry->RayTracingBuildRequestIndex], BuildParams);
+		if (Geometry->HasPendingBuildRequest())
+		{
+			SetupBuildParams(GeometryBuildRequests[Geometry->RayTracingBuildRequestIndex], BuildParams);
+		}
+	}
+
+	if (BuildParams.Num())
+	{
+		InCmdList.BuildAccelerationStructures(BuildParams);
 	}
-	InCmdList.BuildAccelerationStructures(BuildParams);
 }
 
 void FRayTracingGeometryManager::ProcessBuildRequests(FRHIComputeCommandList& InCmdList, bool bInBuildAll)
diff --git a/Engine/Source/Runtime/RenderCore/Public/RayTracingGeometryManager.h b/Engine/Source/Runtime/RenderCore/Public/RayTracingGeometryManager.h
--- a/Engine/Source/Runtime/RenderCore/Public/RayTracingGeometryManager.h
+++ b/Engine/Source/Runtime/RenderCore/Public/RayTracingGeometryManager.h
@@ -32,7 +32,7 @@ public:
 
 	void RemoveBuildRequest(BuildRequestIndex InRequestIndex);
 	void BoostPriority(BuildRequestIndex InRequestIndex, float InBoostValue);
-	void ForceBuild(FRHIComputeCommandList& InCmdList, const TArrayView<const FRayTracingGeometry*> InGeometries);
+	void ForceBuildIfPending(FRHIComputeCommandList& InCmdList, const TArrayView<const FRayTracingGeometry*> InGeometries);
 	void ProcessBuildRequests(FRHIComputeCommandList& InCmdList, bool bInBuildAll = false);
 
 private:
diff --git a/Engine/Source/Runtime/Renderer/Private/DeferredShadingRenderer.cpp b/Engine/Source/Runtime/Renderer/Private/DeferredShadingRenderer.cpp
--- a/Engine/Source/Runtime/Renderer/Private/DeferredShadingRenderer.cpp
+++ b/Engine/Source/Runtime/Renderer/Private/DeferredShadingRenderer.cpp
@@ -1149,8 +1149,7 @@ bool FDeferredShadingSceneRenderer::DispatchRayTracingWorldUpdates(FRDGBuilder&
 	if (ReferenceView.ForceBuildRayTracingGeometries.Num() > 0)
 	{
 		// Force update all the collected geometries (use stack allocator?)
-		TArray<const FRayTracingGeometry*> ForceBuildRayTracingGeometries = ReferenceView.ForceBuildRayTracingGeometries.Array();
-		GRayTracingGeometryManager.ForceBuild(GraphBuilder.RHICmdList, MakeArrayView(ForceBuildRayTracingGeometries.GetData(), ForceBuildRayTracingGeometries.Num()));
+		GRayTracingGeometryManager.ForceBuildIfPending(GraphBuilder.RHICmdList, ReferenceView.ForceBuildRayTracingGeometries);
 	}
 	
 	if (ReferenceView.AddRayTracingMeshBatchTaskList.Num() > 0)
diff --git a/Engine/Source/Runtime/Renderer/Private/SceneRendering.h b/Engine/Source/Runtime/Renderer/Private/SceneRendering.h
--- a/Engine/Source/Runtime/Renderer/Private/SceneRendering.h
+++ b/Engine/Source/Runtime/Renderer/Private/SceneRendering.h
@@ -1403,7 +1403,7 @@ public:
 	TArray<FRayTracingGeometryInstance, SceneRenderingAllocator> RayTracingGeometryInstances;
 
 	// Geometries which still have a pending build request but are used this frame and require a force build
-	TSet<const FRayTracingGeometry*> ForceBuildRayTracingGeometries;
+	TArray<const FRayTracingGeometry*> ForceBuildRayTracingGeometries;
 
 #ifdef DO_CHECK
 	// Keep track of all used RT Geometries which are used to validate the vertex buffer data (see FRayTracingGeometry::DynamicGeometrySharedBufferGenerationID)
