Removed dead code from render target pool; removed 'passthrough' RDG hacks used during conversion.

#rnx

diff --git a/Engine/Source/Runtime/RenderCore/Private/RenderGraphResources.cpp b/Engine/Source/Runtime/RenderCore/Private/RenderGraphResources.cpp
--- a/Engine/Source/Runtime/RenderCore/Private/RenderGraphResources.cpp
+++ b/Engine/Source/Runtime/RenderCore/Private/RenderGraphResources.cpp
@@ -328,18 +328,6 @@ void FRDGBuffer::Finalize()
 	}
 }
 
-FRDGTextureRef FRDGTexture::GetPassthrough(const TRefCountPtr<IPooledRenderTarget>& PooledRenderTargetBase)
-{
-	if (PooledRenderTargetBase)
-	{
-		check(PooledRenderTargetBase->IsCompatibleWithRDG());
-		FRDGTextureRef Texture = &static_cast<FPooledRenderTarget&>(*PooledRenderTargetBase).PassthroughShaderResourceTexture;
-		checkf(Texture->GetRHI(), TEXT("The render target pool didn't allocate a passthrough RHI texture for %s"), PooledRenderTargetBase->GetDesc().DebugName);
-		return Texture;
-	}
-	return nullptr;
-}
-
 FRHIShaderResourceView* FRDGPooledBuffer::GetOrCreateSRV(FRDGBufferSRVDesc SRVDesc)
 {
 	if (const auto* FoundPtr = SRVs.Find(SRVDesc))
diff --git a/Engine/Source/Runtime/RenderCore/Private/RenderGraphUtils.cpp b/Engine/Source/Runtime/RenderCore/Private/RenderGraphUtils.cpp
--- a/Engine/Source/Runtime/RenderCore/Private/RenderGraphUtils.cpp
+++ b/Engine/Source/Runtime/RenderCore/Private/RenderGraphUtils.cpp
@@ -875,22 +875,6 @@ void ConvertToUntrackedBuffer(
 		[](FRHICommandList&) {});
 }
 
-FRDGTextureRef RegisterExternalOrPassthroughTexture(
-	FRDGBuilder* GraphBuilder,
-	const TRefCountPtr<IPooledRenderTarget>& PooledRenderTarget,
-	ERDGTextureFlags Flags)
-{
-	check(PooledRenderTarget);
-	if (GraphBuilder)
-	{
-		return GraphBuilder->RegisterExternalTexture(PooledRenderTarget, ERenderTargetTexture::ShaderResource, Flags);
-	}
-	else
-	{
-		return FRDGTexture::GetPassthrough(PooledRenderTarget);
-	}
-}
-
 FRDGWaitForTasksScope::~FRDGWaitForTasksScope()
 {
 	if (bCondition)
diff --git a/Engine/Source/Runtime/RenderCore/Private/RenderTargetPool.cpp b/Engine/Source/Runtime/RenderCore/Private/RenderTargetPool.cpp
--- a/Engine/Source/Runtime/RenderCore/Private/RenderTargetPool.cpp
+++ b/Engine/Source/Runtime/RenderCore/Private/RenderTargetPool.cpp
@@ -258,77 +258,19 @@ bool FRenderTargetPool::DoesTargetNeedTransienceOverride(ETextureCreateFlags Fla
 	return false;
 }
 
-void FRenderTargetPool::TransitionTargetsWritable(FRHICommandListImmediate& RHICmdList)
-{
-	QUICK_SCOPE_CYCLE_COUNTER(STAT_RenderTargetPoolTransition);
-	check(IsInRenderingThread());
-	WaitForTransitionFence();
-
-	// Stack allocate the transition descriptors. These will get memcpy()ed onto the RHI command list if required.
-	FMemMark Mark(FMemStack::Get());
-	TArray<FRHITransitionInfo, TMemStackAllocator<>> TransitionInfos;
-	TransitionInfos.AddDefaulted(PooledRenderTargets.Num());
-	uint32 TransitionInfoCount = 0;
-
-	for (int32 i = 0; i < PooledRenderTargets.Num(); ++i)
-	{
-		FPooledRenderTarget* PooledRT = PooledRenderTargets[i];
-		if (PooledRT && PooledRT->GetDesc().AutoWritable)
-		{
-			FRHITexture* RenderTarget = PooledRT->GetRenderTargetItem().TargetableTexture;
-			if (RenderTarget)
-			{
-				uint32 CreateFlags = RenderTarget->GetFlags();
-				if ((CreateFlags & TexCreate_DepthStencilTargetable) != 0)
-				{
-					TransitionInfos[TransitionInfoCount++] = FRHITransitionInfo(RenderTarget, ERHIAccess::Unknown, ERHIAccess::DSVRead | ERHIAccess::DSVWrite);
-				}
-				else if ((CreateFlags & TexCreate_RenderTargetable) != 0)
-				{
-					TransitionInfos[TransitionInfoCount++] = FRHITransitionInfo(RenderTarget, ERHIAccess::Unknown, ERHIAccess::RTV);
-				}
-			}
-		}
-	}
-
-	if (TransitionInfoCount > 0)
-	{
-		RHICmdList.Transition(MakeArrayView(TransitionInfos.GetData(), TransitionInfoCount));
-		if (IsRunningRHIInSeparateThread())
-		{
-			TransitionFence = RHICmdList.RHIThreadFence(false);
-		}
-	}
-}
-
-void FRenderTargetPool::WaitForTransitionFence()
-{
-	QUICK_SCOPE_CYCLE_COUNTER(STAT_RenderTargetPoolTransitionWait);
-	check(IsInRenderingThread());
-	if (TransitionFence)
-	{
-		check(IsInRenderingThread());
-		FRHICommandListExecutor::WaitOnRHIThreadFence(TransitionFence);
-		TransitionFence = nullptr;
-	}
-	DeferredDeleteArray.Reset();
-}
-
 TRefCountPtr<FPooledRenderTarget> FRenderTargetPool::FindFreeElementForRDG(
 	FRHICommandList& RHICmdList,
 	const FRDGTextureDesc& Desc,
 	const TCHAR* Name)
 {
-	const bool bDeferTextureAllocation = false;
 	const bool bDoAcquireTransientTexture = false;
-	return FindFreeElementInternal(RHICmdList, Translate(Desc), Name, bDeferTextureAllocation, bDoAcquireTransientTexture);
+	return FindFreeElementInternal(RHICmdList, Translate(Desc), Name, bDoAcquireTransientTexture);
 }
 
 TRefCountPtr<FPooledRenderTarget> FRenderTargetPool::FindFreeElementInternal(
 	FRHICommandList& RHICmdList,
 	const FPooledRenderTargetDesc& Desc,
 	const TCHAR* InDebugName,
-	bool bDeferTextureAllocation,
 	bool bDoAcquireTransientTexture)
 {
 	const int32 AliasingMode = CVarRtPoolTransientMode.GetValueOnRenderThread();
@@ -384,7 +326,6 @@ TRefCountPtr<FPooledRenderTarget> FRenderTargetPool::FindFreeElementInternal(
 						continue;
 					}
 
-					check(!Element->IsSnapshot());
 					Found = Element;
 					FoundIndex = Index;
 					bReusingExistingTarget = true;
@@ -410,7 +351,7 @@ Done:
 
 		FRHIResourceCreateInfo CreateInfo(InDebugName, Desc.ClearValue);
 
-		if (Desc.TargetableFlags & (TexCreate_RenderTargetable | TexCreate_DepthStencilTargetable | TexCreate_UAV) && !bDeferTextureAllocation)
+		if (Desc.TargetableFlags & (TexCreate_RenderTargetable | TexCreate_DepthStencilTargetable | TexCreate_UAV))
 		{
 			// Only create resources if we're not asked to defer creation.
 			if (Desc.Is2DTexture())
@@ -512,9 +453,8 @@ Done:
 			RHIBindDebugLabelName(Found->RenderTargetItem.TargetableTexture, InDebugName);
 #endif
 		}
-		else if (!bDeferTextureAllocation)
+		else
 		{
-			// Only create resources if we're not asked to defer creation.
 			if (Desc.Is2DTexture())
 			{
 				// this is useful to get a CPU lockable texture through the same interface
@@ -558,29 +498,20 @@ Done:
 #endif
 		}
 
-		if (!bDeferTextureAllocation)
+		if ((Desc.TargetableFlags & TexCreate_UAV))
 		{
-			if ((Desc.TargetableFlags & TexCreate_UAV))
-			{
-				// The render target desc is invalid if a UAV is requested with an RHI that doesn't support the high-end feature level.
-				check(GMaxRHIFeatureLevel >= ERHIFeatureLevel::SM5 || GMaxRHIFeatureLevel == ERHIFeatureLevel::ES3_1);
-				Found->RenderTargetItem.UAV = RHICreateUnorderedAccessView(Found->RenderTargetItem.TargetableTexture, 0);
-			}
-
-			// Only calculate allocation level if we actually allocated something. If bDeferTextureAllocation is true, the caller should call 
-			// UpdateElementSize once it's set the resources on the created object.
-			AllocationLevelInKB += ComputeSizeInKB(*Found);
-			VerifyAllocationLevel();
-
-			Found->InitPassthroughRDG();
+			// The render target desc is invalid if a UAV is requested with an RHI that doesn't support the high-end feature level.
+			check(GMaxRHIFeatureLevel >= ERHIFeatureLevel::SM5 || GMaxRHIFeatureLevel == ERHIFeatureLevel::ES3_1);
+			Found->RenderTargetItem.UAV = RHICreateUnorderedAccessView(Found->RenderTargetItem.TargetableTexture, 0);
 		}
 
+		AllocationLevelInKB += ComputeSizeInKB(*Found);
+
 		FoundIndex = PooledRenderTargets.Num() - 1;
 		Found->Desc.DebugName = InDebugName;
 	}
 
 	check(Found->IsFree());
-	check(!Found->IsSnapshot());
 
 	Found->Desc.DebugName = InDebugName;
 	Found->UnusedForNFrames = 0;
@@ -621,8 +552,7 @@ bool FRenderTargetPool::FindFreeElement(
 	const FPooledRenderTargetDesc& InputDesc,
 	TRefCountPtr<IPooledRenderTarget>& Out,
 	const TCHAR* InDebugName,
-	ERenderTargetTransience TransienceHint,
-	bool bDeferTextureAllocation)
+	ERenderTargetTransience TransienceHint)
 {
 	check(IsInRenderingThread());
 
@@ -658,8 +588,6 @@ bool FRenderTargetPool::FindFreeElement(
 	{
 		FPooledRenderTarget* Current = (FPooledRenderTarget*)Out.GetReference();
 
-		check(!Current->IsSnapshot());
-
 		const bool bExactMatch = true;
 
 		if (Out->GetDesc().Compare(Desc, bExactMatch))
@@ -689,14 +617,12 @@ bool FRenderTargetPool::FindFreeElement(
 				check(Index >= 0);
 
 				FreeElementAtIndex(Index);
-
-				VerifyAllocationLevel();
 			}
 		}
 	}
 
 	const bool bDoAcquireTransientResource = true;
-	TRefCountPtr<FPooledRenderTarget> Result = FindFreeElementInternal(RHICmdList, Desc, InDebugName, bDeferTextureAllocation, bDoAcquireTransientResource);
+	TRefCountPtr<FPooledRenderTarget> Result = FindFreeElementInternal(RHICmdList, Desc, InDebugName, bDoAcquireTransientResource);
 
 	// Reset RDG state back to an unknown default. The resource is being handed off to a user outside of RDG, so the state is no longer valid.
 	{
@@ -730,25 +656,11 @@ void FRenderTargetPool::CreateUntrackedElement(const FPooledRenderTargetDesc& De
 	FPooledRenderTarget* Found = new FPooledRenderTarget(Desc, NULL);
 
 	Found->RenderTargetItem = Item;
-	Found->InitPassthroughRDG();
-	check(!Found->IsSnapshot());
 
 	// assign to the reference counted variable
 	Out = Found;
 }
 
-IPooledRenderTarget* FRenderTargetPool::MakeSnapshot(const TRefCountPtr<IPooledRenderTarget>& In)
-{
-	check(IsInRenderingThread());
-	FPooledRenderTarget* NewSnapshot = nullptr;
-	if (In.GetReference())
-	{
-		NewSnapshot = new (FMemStack::Get()) FPooledRenderTarget(*static_cast<FPooledRenderTarget*>(In.GetReference()));
-		PooledRenderTargetSnapshots.Add(NewSnapshot);
-	}
-	return NewSnapshot;
-}
-
 void FRenderTargetPool::GetStats(uint32& OutWholeCount, uint32& OutWholePoolInKB, uint32& OutUsedInKB) const
 {
 	OutWholeCount = (uint32)PooledRenderTargets.Num();
@@ -761,7 +673,6 @@ void FRenderTargetPool::GetStats(uint32& OutWholeCount, uint32& OutWholePoolInKB
 
 		if (Element)
 		{
-			check(!Element->IsSnapshot());
 			uint32 SizeInKB = ComputeSizeInKB(*Element);
 
 			OutWholePoolInKB += SizeInKB;
@@ -932,13 +843,6 @@ FRenderTargetPool::SMemoryStats FRenderTargetPool::ComputeView()
 	return MemoryStats;
 }
 
-void FRenderTargetPool::UpdateElementSize(const TRefCountPtr<IPooledRenderTarget>& Element, const uint32 OldElementSize)
-{
-	check(Element.IsValid() && FindIndex(&(*Element)) >= 0);
-	AllocationLevelInKB -= (OldElementSize + 1023) / 1024;
-	AllocationLevelInKB += (Element->ComputeMemorySize() + 1023) / 1024;
-}
-
 void FRenderTargetPool::AddDeallocEvents()
 {
 	check(IsInRenderingThread());
@@ -1015,7 +919,6 @@ void FRenderTargetPool::AddAllocEventsFromCurrentState()
 
 void FRenderTargetPool::TickPoolElements()
 {
-	// gather stats on deferred allocs before calling WaitForTransitionFence
 	uint32 DeferredAllocationLevelInKB = 0;
 	for (FPooledRenderTarget* Element : DeferredDeleteArray)
 	{
@@ -1023,7 +926,7 @@ void FRenderTargetPool::TickPoolElements()
 	}
 
 	check(IsInRenderingThread());
-	WaitForTransitionFence();
+	DeferredDeleteArray.Reset();
 
 	if (bStartEventRecordingNextTick)
 	{
@@ -1047,7 +950,6 @@ void FRenderTargetPool::TickPoolElements()
 
 		if (Element)
 		{
-			check(!Element->IsSnapshot());
 			Element->OnFrameStart();
 			if (Element->UnusedForNFrames > 2)
 			{
@@ -1108,8 +1010,6 @@ void FRenderTargetPool::TickPoolElements()
 			// we assume because of reference counting the resource gets released when not needed any more
 			// we don't use Remove() to not shuffle around the elements for better transparency on RenderTargetPoolEvents
 			FreeElementAtIndex(OldestElementIndex);
-
-			VerifyAllocationLevel();
 		}
 		else
 		{
@@ -1166,7 +1066,6 @@ int32 FRenderTargetPool::FindIndex(IPooledRenderTarget* In) const
 
 			if (Element == In)
 			{
-				check(!Element->IsSnapshot());
 				return i;
 			}
 		}
@@ -1199,13 +1098,10 @@ void FRenderTargetPool::FreeUnusedResource(TRefCountPtr<IPooledRenderTarget>& In
 
 		if (Element->IsFree())
 		{
-			check(!Element->IsSnapshot());
 			AllocationLevelInKB -= ComputeSizeInKB(*Element);
 			// we assume because of reference counting the resource gets released when not needed any more
 			DeferredDeleteArray.Add(PooledRenderTargets[Index]);
 			FreeElementAtIndex(Index);
-
-			VerifyAllocationLevel();
 		}
 	}
 }
@@ -1220,7 +1116,6 @@ void FRenderTargetPool::FreeUnusedResources()
 
 		if (Element && Element->IsFree())
 		{
-			check(!Element->IsSnapshot());
 			AllocationLevelInKB -= ComputeSizeInKB(*Element);
 			// we assume because of reference counting the resource gets released when not needed any more
 			// we don't use Remove() to not shuffle around the elements for better transparency on RenderTargetPoolEvents
@@ -1229,8 +1124,6 @@ void FRenderTargetPool::FreeUnusedResources()
 		}
 	}
 
-	VerifyAllocationLevel();
-
 #if LOG_MAX_RENDER_TARGET_POOL_USAGE
 	MaxUsedRenderTargetInKB = 0;
 #endif
@@ -1253,7 +1146,6 @@ void FRenderTargetPool::DumpMemoryUsage(FOutputDevice& OutputDevice)
 				UnusedAllocationInKB += ElementAllocationInKB;
 			}
 
-			check(!Element->IsSnapshot());
 			OutputDevice.Logf(
 				TEXT("  %6.3fMB %4dx%4d%s%s %2dmip(s) %s (%s) %s %s Unused frames: %d"),
 				ElementAllocationInKB / 1024.0f,
@@ -1284,7 +1176,6 @@ void FRenderTargetPool::DumpMemoryUsage(FOutputDevice& OutputDevice)
 
 		if (Element)
 		{
-			check(!Element->IsSnapshot());
 			OutputDevice.Logf(
 				TEXT("  %6.3fMB %4dx%4d%s%s %2dmip(s) %s (%s) %s %s"),
 				ComputeSizeInKB(*Element) / 1024.0f,
@@ -1326,38 +1217,28 @@ void FPooledRenderTarget::InitRDG()
 
 uint32 FPooledRenderTarget::AddRef() const
 {
-	if (!bSnapshot)
-	{
-		check(IsInRenderingThread());
-		return uint32(++NumRefs);
-	}
-	check(NumRefs == 1);
-	return 1;
+	check(IsInRenderingThread());
+	return uint32(++NumRefs);
 }
 
 uint32 FPooledRenderTarget::Release()
 {
-	if (!bSnapshot)
+	checkf(IsInRenderingThread(), TEXT("Tried to delete on non-render thread, PooledRT %s %s"), Desc.DebugName ? Desc.DebugName : TEXT("<Unnamed>"), *Desc.GenerateInfoString());
+	uint32 Refs = uint32(--NumRefs);
+	if (Refs == 0)
 	{
-		checkf(IsInRenderingThread(), TEXT("Tried to delete on non-render thread, PooledRT %s %s"), Desc.DebugName ? Desc.DebugName : TEXT("<Unnamed>"), *Desc.GenerateInfoString());
-		uint32 Refs = uint32(--NumRefs);
-		if (Refs == 0)
-		{
-			RenderTargetItem.SafeRelease();
-			delete this;
-		}
-		else if (Refs == 1 && RenderTargetPool && IsTransient())
+		RenderTargetItem.SafeRelease();
+		delete this;
+	}
+	else if (Refs == 1 && RenderTargetPool && IsTransient())
+	{
+		if (bAutoDiscard && RenderTargetItem.TargetableTexture)
 		{
-			if (bAutoDiscard && RenderTargetItem.TargetableTexture)
-			{
-				RHIDiscardTransientResource(RenderTargetItem.TargetableTexture);
-			}
-			FrameNumberLastDiscard = GFrameNumberRenderThread;
+			RHIDiscardTransientResource(RenderTargetItem.TargetableTexture);
 		}
-		return Refs;
+		FrameNumberLastDiscard = GFrameNumberRenderThread;
 	}
-	check(NumRefs == 1);
-	return 1;
+	return Refs;
 }
 
 uint32 FPooledRenderTarget::GetRefCount() const
@@ -1380,22 +1261,8 @@ const FPooledRenderTargetDesc& FPooledRenderTarget::GetDesc() const
 void FRenderTargetPool::ReleaseDynamicRHI()
 {
 	check(IsInRenderingThread());
-	WaitForTransitionFence();
-
+	DeferredDeleteArray.Empty();
 	PooledRenderTargets.Empty();
-	if (PooledRenderTargetSnapshots.Num())
-	{
-		DestructSnapshots();
-	}
-}
-
-void FRenderTargetPool::DestructSnapshots()
-{
-	for (auto Snapshot : PooledRenderTargetSnapshots)
-	{
-		Snapshot->~FPooledRenderTarget();
-	}
-	PooledRenderTargetSnapshots.Reset();
 }
 
 // for debugging purpose
@@ -1411,10 +1278,6 @@ FPooledRenderTarget* FRenderTargetPool::GetElementById(uint32 Id) const
 	return PooledRenderTargets[Id];
 }
 
-void FRenderTargetPool::VerifyAllocationLevel() const
-{
-}
-
 void FRenderTargetPool::CompactPool()
 {
 	for (uint32 i = 0, Num = (uint32)PooledRenderTargets.Num(); i < Num; ++i)
@@ -1432,7 +1295,7 @@ void FRenderTargetPool::CompactPool()
 
 bool FPooledRenderTarget::OnFrameStart()
 {
-	check(IsInRenderingThread() && !bSnapshot);
+	check(IsInRenderingThread());
 
 	// If there are any references to the pooled render target other than the pool itself, then it may not be freed.
 	if (!IsFree())
@@ -1456,31 +1319,28 @@ bool FPooledRenderTarget::OnFrameStart()
 uint32 FPooledRenderTarget::ComputeMemorySize() const
 {
 	uint32 Size = 0;
-	if (!bSnapshot)
+	if (Desc.Is2DTexture())
 	{
-		if (Desc.Is2DTexture())
+		Size += RHIComputeMemorySize(RenderTargetItem.TargetableTexture);
+		if (RenderTargetItem.ShaderResourceTexture != RenderTargetItem.TargetableTexture)
 		{
-			Size += RHIComputeMemorySize(RenderTargetItem.TargetableTexture);
-			if (RenderTargetItem.ShaderResourceTexture != RenderTargetItem.TargetableTexture)
-			{
-				Size += RHIComputeMemorySize(RenderTargetItem.ShaderResourceTexture);
-			}
+			Size += RHIComputeMemorySize(RenderTargetItem.ShaderResourceTexture);
 		}
-		else if (Desc.Is3DTexture())
+	}
+	else if (Desc.Is3DTexture())
+	{
+		Size += RHIComputeMemorySize(RenderTargetItem.TargetableTexture);
+		if (RenderTargetItem.ShaderResourceTexture != RenderTargetItem.TargetableTexture)
 		{
-			Size += RHIComputeMemorySize(RenderTargetItem.TargetableTexture);
-			if (RenderTargetItem.ShaderResourceTexture != RenderTargetItem.TargetableTexture)
-			{
-				Size += RHIComputeMemorySize(RenderTargetItem.ShaderResourceTexture);
-			}
+			Size += RHIComputeMemorySize(RenderTargetItem.ShaderResourceTexture);
 		}
-		else
+	}
+	else
+	{
+		Size += RHIComputeMemorySize(RenderTargetItem.TargetableTexture);
+		if (RenderTargetItem.ShaderResourceTexture != RenderTargetItem.TargetableTexture)
 		{
-			Size += RHIComputeMemorySize(RenderTargetItem.TargetableTexture);
-			if (RenderTargetItem.ShaderResourceTexture != RenderTargetItem.TargetableTexture)
-			{
-				Size += RHIComputeMemorySize(RenderTargetItem.ShaderResourceTexture);
-			}
+			Size += RHIComputeMemorySize(RenderTargetItem.ShaderResourceTexture);
 		}
 	}
 	return Size;
@@ -1492,11 +1352,5 @@ bool FPooledRenderTarget::IsFree() const
 	check(RefCount >= 1);
 
 	// If the only reference to the pooled render target is from the pool, then it's unused.
-	return !bSnapshot && RefCount == 1;
-}
-
-void FPooledRenderTarget::InitPassthroughRDG()
-{
-	check(RenderTargetItem.ShaderResourceTexture);
-	PassthroughShaderResourceTexture.SetPassthroughRHI(RenderTargetItem.ShaderResourceTexture);
+	return RefCount == 1;
 }
\ No newline at end of file
diff --git a/Engine/Source/Runtime/RenderCore/Private/RenderUtils.cpp b/Engine/Source/Runtime/RenderCore/Private/RenderUtils.cpp
--- a/Engine/Source/Runtime/RenderCore/Private/RenderUtils.cpp
+++ b/Engine/Source/Runtime/RenderCore/Private/RenderUtils.cpp
@@ -27,12 +27,6 @@ FRDGTexture* FTextureWithRDG::GetRDG(FRDGBuilder& GraphBuilder) const
 	return GraphBuilder.RegisterExternalTexture(RenderTarget);
 }
 
-FRDGTexture* FTextureWithRDG::GetPassthroughRDG() const
-{
-	checkf(RenderTarget, TEXT("InitRDG was not called before use."));
-	return FRDGTexture::GetPassthrough(RenderTarget);
-}
-
 void FTextureWithRDG::ReleaseRHI()
 {
 	RenderTarget = nullptr;
diff --git a/Engine/Source/Runtime/RenderCore/Public/RenderGraphResources.h b/Engine/Source/Runtime/RenderCore/Public/RenderGraphResources.h
--- a/Engine/Source/Runtime/RenderCore/Public/RenderGraphResources.h
+++ b/Engine/Source/Runtime/RenderCore/Public/RenderGraphResources.h
@@ -533,6 +533,7 @@ public:
 
 	uint32 Release() const
 	{
+		check(RefCount > 0);
 		if (--RefCount == 0)
 		{
 			delete this;
@@ -596,9 +597,6 @@ class RENDERCORE_API FRDGTexture final
 	: public FRDGParentResource
 {
 public:
-	/** Creates a passthrough texture suitable for filling RHI uniform buffers with RDG parameters for passes not yet ported to RDG. */
-	static FRDGTextureRef GetPassthrough(const TRefCountPtr<IPooledRenderTarget>& PooledRenderTarget);
-
 	const FRDGTextureDesc Desc;
 	const ERDGTextureFlags Flags;
 
diff --git a/Engine/Source/Runtime/RenderCore/Public/RenderGraphUtils.h b/Engine/Source/Runtime/RenderCore/Public/RenderGraphUtils.h
--- a/Engine/Source/Runtime/RenderCore/Public/RenderGraphUtils.h
+++ b/Engine/Source/Runtime/RenderCore/Public/RenderGraphUtils.h
@@ -798,12 +798,6 @@ inline void ConvertToUntrackedExternalBuffer(
 	ConvertToUntrackedBuffer(GraphBuilder, Buffer, AccessFinal);
 }
 
-// Used to help port code over to RDG. If the graph builder is null, creates a passthrough RDG texture instead. Will be removed once port is complete.
-RENDERCORE_API FRDGTextureRef RegisterExternalOrPassthroughTexture(
-	FRDGBuilder* GraphBuilder,
-	const TRefCountPtr<IPooledRenderTarget>& PooledRenderTarget,
-	ERDGTextureFlags Flags = ERDGTextureFlags::None);
-
 /** Scope used to wait for outstanding tasks when the scope destructor is called. Used for command list recording tasks. */
 class FRDGWaitForTasksScope
 {
diff --git a/Engine/Source/Runtime/RenderCore/Public/RenderTargetPool.h b/Engine/Source/Runtime/RenderCore/Public/RenderTargetPool.h
--- a/Engine/Source/Runtime/RenderCore/Public/RenderTargetPool.h
+++ b/Engine/Source/Runtime/RenderCore/Public/RenderTargetPool.h
@@ -15,40 +15,21 @@
 #define LOG_MAX_RENDER_TARGET_POOL_USAGE 0
 
 /** The reference to a pooled render target, use like this: TRefCountPtr<IPooledRenderTarget> */
-struct RENDERCORE_API FPooledRenderTarget : public IPooledRenderTarget
+struct RENDERCORE_API FPooledRenderTarget final : public IPooledRenderTarget
 {
 	FPooledRenderTarget(const FPooledRenderTargetDesc& InDesc, class FRenderTargetPool* InRenderTargetPool) 
 		: RenderTargetPool(InRenderTargetPool)
 		, Desc(InDesc)
-		, PassthroughShaderResourceTexture(TEXT("Passthrough"), {}, ERDGTextureFlags::None, ERenderTargetTexture::ShaderResource)
 	{}
 
-	/* Constructor that makes a snapshot */
-	FPooledRenderTarget(const FPooledRenderTarget& SnaphotSource)
-		: RenderTargetPool(SnaphotSource.RenderTargetPool)
-		, Desc(SnaphotSource.Desc)
-		, NumRefs(1)
-		, bSnapshot(true)
-		, PassthroughShaderResourceTexture(TEXT("Passthrough"), {}, ERDGTextureFlags::None, ERenderTargetTexture::ShaderResource)
+	~FPooledRenderTarget()
 	{
-		check(IsInRenderingThread());
-		RenderTargetItem = SnaphotSource.RenderTargetItem;
-	}
-
-	virtual ~FPooledRenderTarget()
-	{
-		check(!NumRefs || (bSnapshot && NumRefs == 1));
+		check(!NumRefs);
 		RenderTargetItem.SafeRelease();
 	}
 
-	bool IsSnapshot() const 
-	{ 
-		return bSnapshot;
-	}
-
 	uint32 GetUnusedForNFrames() const 
 	{ 
-		check(!bSnapshot);
 		return UnusedForNFrames; 
 	}
 
@@ -67,13 +48,7 @@ struct RENDERCORE_API FPooledRenderTarget : public IPooledRenderTarget
 		return Texture == ERenderTargetTexture::Targetable ? TargetableTexture : ShaderResourceTexture;
 	}
 
-	FRDGTextureRef GetPassthroughRDG() const
-	{
-		return &PassthroughShaderResourceTexture;
-	}
-
 	void InitRDG();
-	void InitPassthroughRDG();
 
 	// interface IPooledRenderTarget --------------
 
@@ -121,16 +96,12 @@ private:
 	/** Keeps track of the last frame we unmapped physical memory for this resource. We can't map again in the same frame if we did that */
 	uint32 FrameNumberLastDiscard = -1;
 
-	/** Snapshots are sortof fake pooled render targets, they don't own anything and can outlive the things that created them. These are for threaded rendering. */
-	bool bSnapshot = false;
-
 	/** The transient resource discard will happen automatically on free. */
 	bool bAutoDiscard = true;
 
 	/** Pooled textures for use with RDG. */
 	TRefCountPtr<FRDGPooledTexture> TargetableTexture;
 	TRefCountPtr<FRDGPooledTexture> ShaderResourceTexture;
-	mutable FRDGTexture PassthroughShaderResourceTexture;
 
 	/** @return true:release this one, false otherwise */
 	bool OnFrameStart();
@@ -256,9 +227,6 @@ class RENDERCORE_API FRenderTargetPool : public FRenderResource
 public:
 	FRenderTargetPool();
 
-	/** Transitions all targets in the pool to writable. */
-	void TransitionTargetsWritable(FRHICommandListImmediate& RHICmdList);
-
 	/**
 	 * @param DebugName must not be 0, we only store the pointer
 	 * @param Out is not the return argument to avoid double allocation because of wrong reference counting
@@ -270,19 +238,10 @@ public:
 		const FPooledRenderTargetDesc& Desc,
 		TRefCountPtr<IPooledRenderTarget>& Out,
 		const TCHAR* InDebugName,
-		ERenderTargetTransience TransienceHint = ERenderTargetTransience::Transient,
-		bool bDeferTextureAllocation = false);
+		ERenderTargetTransience TransienceHint = ERenderTargetTransience::Transient);
 
 	void CreateUntrackedElement(const FPooledRenderTargetDesc& Desc, TRefCountPtr<IPooledRenderTarget>& Out, const FSceneRenderTargetItem& Item);
 
-	/** Destruct all snapshots, this must be done after all outstanding async tasks are done. It is important because they hold ref counted texture pointers etc **/
-	IPooledRenderTarget* MakeSnapshot(const TRefCountPtr<IPooledRenderTarget>& In);
-
-	/** Destruct all snapshots, this must be done after all outstanding async tasks are done. It is important because they hold ref counted texture pointers etc **/
-	void DestructSnapshots();
-
-	void OnRenderTargetUnreferenced(IPooledRenderTarget* RenderTarget);
-
 	/** Only to get statistics on usage and free elements. Normally only called in renderthread or if FlushRenderingCommands was called() */
 	void GetStats(uint32& OutWholeCount, uint32& OutWholePoolInKB, uint32& OutUsedInKB) const;
 	/**
@@ -323,8 +282,6 @@ public:
 
 	void AddPhaseEvent(const TCHAR* InPhaseName);
 
-	void UpdateElementSize(const TRefCountPtr<IPooledRenderTarget>& Element, const uint32 OldSize);
-
 private:
 	TRefCountPtr<FPooledRenderTarget> FindFreeElementForRDG(FRHICommandList& RHICmdList, const FRDGTextureDesc& Desc, const TCHAR* Name);
 
@@ -332,7 +289,6 @@ private:
 		FRHICommandList& RHICmdList,
 		const FPooledRenderTargetDesc& InputDesc,
 		const TCHAR* InDebugName,
-		bool bDeferTextureAllocation,
 		bool bDoAcquireTransientResource);
 
 	static bool DoesTargetNeedTransienceOverride(ETextureCreateFlags Flags, ERenderTargetTransience TransienceHint);
@@ -346,9 +302,6 @@ private:
 	TArray< TRefCountPtr<FPooledRenderTarget> > PooledRenderTargets;
 	TArray< TRefCountPtr<FPooledRenderTarget> > DeferredDeleteArray;
 
-	/** These are snapshots, have odd life times, live in the scene allocator, and don't contribute to any accounting or other management. */
-	TArray<FPooledRenderTarget*> PooledRenderTargetSnapshots;
-
 	// redundant, can always be computed with GetStats(), to debug "out of memory" situations and used for r.RenderTargetPoolMin
 	uint32 AllocationLevelInKB;
 
@@ -357,14 +310,9 @@ private:
 	// to avoid log spam
 	bool bCurrentlyOverBudget;
 
-	// for debugging purpose
-	void VerifyAllocationLevel() const;
-
 	// could be done on the fly but that makes the RenderTargetPoolEvents harder to read
 	void CompactPool();
 
-	void WaitForTransitionFence();
-
 	// the following is used for Event recording --------------------------------
 
 	struct SMemoryStats
diff --git a/Engine/Source/Runtime/RenderCore/Public/RenderUtils.h b/Engine/Source/Runtime/RenderCore/Public/RenderUtils.h
--- a/Engine/Source/Runtime/RenderCore/Public/RenderUtils.h
+++ b/Engine/Source/Runtime/RenderCore/Public/RenderUtils.h
@@ -27,7 +27,6 @@ public:
 	~FTextureWithRDG() override;
 
 	FRDGTexture* GetRDG(FRDGBuilder& GraphBuilder) const;
-	FRDGTexture* GetPassthroughRDG() const;
 
 	void ReleaseRHI() override;
 
