Fixed bad merge from 4.26 transient fixes.

#rnx

diff --git a/Engine/Source/Runtime/RenderCore/Private/RenderTargetPool.cpp b/Engine/Source/Runtime/RenderCore/Private/RenderTargetPool.cpp
--- a/Engine/Source/Runtime/RenderCore/Private/RenderTargetPool.cpp
+++ b/Engine/Source/Runtime/RenderCore/Private/RenderTargetPool.cpp
@@ -1237,6 +1237,7 @@ uint32 FPooledRenderTarget::Release()
 			RHIDiscardTransientResource(RenderTargetItem.TargetableTexture);
 		}
 		FrameNumberLastDiscard = GFrameNumberRenderThread;
+		bAutoDiscard = true;
 	}
 	return Refs;
 }
@@ -1319,28 +1320,31 @@ bool FPooledRenderTarget::OnFrameStart()
 uint32 FPooledRenderTarget::ComputeMemorySize() const
 {
 	uint32 Size = 0;
-	if (Desc.Is2DTexture())
+	if (!IsTransient())
 	{
-		Size += RHIComputeMemorySize(RenderTargetItem.TargetableTexture);
-		if (RenderTargetItem.ShaderResourceTexture != RenderTargetItem.TargetableTexture)
+		if (Desc.Is2DTexture())
 		{
-			Size += RHIComputeMemorySize(RenderTargetItem.ShaderResourceTexture);
+			Size += RHIComputeMemorySize(RenderTargetItem.TargetableTexture);
+			if (RenderTargetItem.ShaderResourceTexture != RenderTargetItem.TargetableTexture)
+			{
+				Size += RHIComputeMemorySize(RenderTargetItem.ShaderResourceTexture);
+			}
 		}
-	}
-	else if (Desc.Is3DTexture())
-	{
-		Size += RHIComputeMemorySize(RenderTargetItem.TargetableTexture);
-		if (RenderTargetItem.ShaderResourceTexture != RenderTargetItem.TargetableTexture)
+		else if (Desc.Is3DTexture())
 		{
-			Size += RHIComputeMemorySize(RenderTargetItem.ShaderResourceTexture);
+			Size += RHIComputeMemorySize(RenderTargetItem.TargetableTexture);
+			if (RenderTargetItem.ShaderResourceTexture != RenderTargetItem.TargetableTexture)
+			{
+				Size += RHIComputeMemorySize(RenderTargetItem.ShaderResourceTexture);
+			}
 		}
-	}
-	else
-	{
-		Size += RHIComputeMemorySize(RenderTargetItem.TargetableTexture);
-		if (RenderTargetItem.ShaderResourceTexture != RenderTargetItem.TargetableTexture)
+		else
 		{
-			Size += RHIComputeMemorySize(RenderTargetItem.ShaderResourceTexture);
+			Size += RHIComputeMemorySize(RenderTargetItem.TargetableTexture);
+			if (RenderTargetItem.ShaderResourceTexture != RenderTargetItem.TargetableTexture)
+			{
+				Size += RHIComputeMemorySize(RenderTargetItem.ShaderResourceTexture);
+			}
 		}
 	}
 	return Size;
