Fix missing symbol when building development programs

diff --git a/Engine/Source/Runtime/RenderCore/Private/RenderGraphEvent.cpp b/Engine/Source/Runtime/RenderCore/Private/RenderGraphEvent.cpp
--- a/Engine/Source/Runtime/RenderCore/Private/RenderGraphEvent.cpp
+++ b/Engine/Source/Runtime/RenderCore/Private/RenderGraphEvent.cpp
@@ -9,7 +9,11 @@ RENDERCORE_API bool GetEmitRDGEvents()
 {
 	check(IsInRenderingThread());
 #if RDG_EVENTS != RDG_EVENTS_NONE
-	return GetEmitDrawEvents() || GRDGDebug || UE_TRACE_CHANNELEXPR_IS_ENABLED(RDGChannel);
+	bool bRDGChannelEnabled = false;
+#if RDG_ENABLE_TRACE
+	bRDGChannelEnabled = UE_TRACE_CHANNELEXPR_IS_ENABLED(RDGChannel);
+#endif // RDG_ENABLE_TRACE
+	return GetEmitDrawEvents() || GRDGDebug || bRDGChannelEnabled;
 #else
 	return false;
 #endif
