Fix a crash bug of converting EShaderPlatform to FStaticShaderPlatform on some specific platforms.

diff --git a/Engine/Source/Runtime/RenderCore/Private/RenderUtils.cpp b/Engine/Source/Runtime/RenderCore/Private/RenderUtils.cpp
--- a/Engine/Source/Runtime/RenderCore/Private/RenderUtils.cpp
+++ b/Engine/Source/Runtime/RenderCore/Private/RenderUtils.cpp
@@ -1284,7 +1284,7 @@ RENDERCORE_API void RenderUtilsInit()
 					GDBufferPlatformMask &= ~Mask;
 				}
 
-				if (TargetPlatform->UsesBasePassVelocity())
+				if (TargetPlatform->UsesBasePassVelocity() && !IsMobilePlatform(ShaderPlatform))
 				{
 					GBasePassVelocityPlatformMask |= Mask;
 				}
@@ -1340,17 +1340,13 @@ RENDERCORE_API void RenderUtilsInit()
 			}
 		}
 	}
-#endif // WITH_EDITOR
-
-	for (uint32 ShaderPlatformIndex = 0; ShaderPlatformIndex < SP_NumPlatforms; ++ShaderPlatformIndex)
+#else
+	if (IsMobilePlatform(GMaxRHIShaderPlatform))
 	{
-		EShaderPlatform ShaderPlatform = EShaderPlatform(ShaderPlatformIndex);
-		if (IsMobilePlatform(ShaderPlatform))
-		{
-			uint64 Mask = 1ull << ShaderPlatformIndex;
-			GBasePassVelocityPlatformMask &= ~Mask;
-		}
+		uint64 Mask = 1ull << GMaxRHIShaderPlatform;
+		GBasePassVelocityPlatformMask &= ~Mask;
 	}
+#endif // WITH_EDITOR
 
 	// Run-time ray tracing support depends on the following factors:
 	// - Ray tracing must be enabled for the project
