Small refactor of FVisualizeTexture::SetCheckpoint to include RDG builder version. Allows for visualizing pooled render targets not produced by the current RDG builder.

diff --git a/Engine/Source/Runtime/RenderCore/Private/VisualizeTexture.cpp b/Engine/Source/Runtime/RenderCore/Private/VisualizeTexture.cpp
--- a/Engine/Source/Runtime/RenderCore/Private/VisualizeTexture.cpp
+++ b/Engine/Source/Runtime/RenderCore/Private/VisualizeTexture.cpp
@@ -850,7 +850,7 @@ uint32 FVisualizeTexture::GetVersionCount(const TCHAR* InName) const
 	return 0;
 }
 
-void FVisualizeTexture::SetCheckPoint(FRHICommandListImmediate& RHICmdList, IPooledRenderTarget* PooledRenderTarget)
+void FVisualizeTexture::SetCheckPoint(FRDGBuilder& GraphBuilder, IPooledRenderTarget* PooledRenderTarget)
 {
 	check(IsInRenderingThread());
 
@@ -872,9 +872,14 @@ void FVisualizeTexture::SetCheckPoint(FRHICommandListImmediate& RHICmdList, IPoo
 		return;
 	}
 
-	FRDGBuilder GraphBuilder(RHICmdList);
 	FRDGTextureRef TextureToCapture = GraphBuilder.RegisterExternalTexture(PooledRenderTarget, ERenderTargetTexture::Targetable);
 	CreateContentCapturePass(GraphBuilder, TextureToCapture, CaptureId.GetValue());
+}
+
+void FVisualizeTexture::SetCheckPoint(FRHICommandListImmediate& RHICmdList, IPooledRenderTarget* PooledRenderTarget)
+{
+	FRDGBuilder GraphBuilder(RHICmdList);
+	SetCheckPoint(GraphBuilder, PooledRenderTarget);
 	GraphBuilder.Execute();
 }
 
diff --git a/Engine/Source/Runtime/RenderCore/Public/VisualizeTexture.h b/Engine/Source/Runtime/RenderCore/Public/VisualizeTexture.h
--- a/Engine/Source/Runtime/RenderCore/Public/VisualizeTexture.h
+++ b/Engine/Source/Runtime/RenderCore/Public/VisualizeTexture.h
@@ -22,8 +22,10 @@ public:
 
 	/** Creates a new checkpoint (e.g. "SceneDepth@N") for the pooled render target. A null parameter is a no-op. */
 #if SUPPORTS_VISUALIZE_TEXTURE
+	void SetCheckPoint(FRDGBuilder& GraphBuilder, IPooledRenderTarget* PooledRenderTarget);
 	void SetCheckPoint(FRHICommandListImmediate& RHICmdList, IPooledRenderTarget* PooledRenderTarget);
 #else
+	inline void SetCheckPoint(FRDGBuilder& GraphBuilder, IPooledRenderTarget* PooledRenderTarget) {}
 	inline void SetCheckPoint(FRHICommandListImmediate& RHICmdList, IPooledRenderTarget* PooledRenderTarget) {}
 #endif
 
