Adjust the way of excluding mobile platform from GBasePassVelocityPlatformMask

diff --git a/Engine/Source/Runtime/RenderCore/Private/RenderUtils.cpp b/Engine/Source/Runtime/RenderCore/Private/RenderUtils.cpp
--- a/Engine/Source/Runtime/RenderCore/Private/RenderUtils.cpp
+++ b/Engine/Source/Runtime/RenderCore/Private/RenderUtils.cpp
@@ -1343,8 +1343,7 @@ RENDERCORE_API void RenderUtilsInit()
 #else
 	if (IsMobilePlatform(GMaxRHIShaderPlatform))
 	{
-		uint64 Mask = 1ull << GMaxRHIShaderPlatform;
-		GBasePassVelocityPlatformMask &= ~Mask;
+		GBasePassVelocityPlatformMask = 0;
 	}
 #endif // WITH_EDITOR
 
