Fix assert on RHI access for foveation texture.

diff --git a/Engine/Source/Runtime/RenderCore/Private/RenderGraphBuilder.cpp b/Engine/Source/Runtime/RenderCore/Private/RenderGraphBuilder.cpp
--- a/Engine/Source/Runtime/RenderCore/Private/RenderGraphBuilder.cpp
+++ b/Engine/Source/Runtime/RenderCore/Private/RenderGraphBuilder.cpp
@@ -189,6 +189,11 @@ void EnumerateTextureAccess(FRDGParameterStruct PassParameters, ERDGPassFlags Pa
 					AccessFunction(nullptr, Texture, NewAccess, Range);
 				});
 			}
+
+			if (FRDGTextureRef Texture = RenderTargets.FoveationTexture)
+			{
+				AccessFunction(nullptr, Texture, ERHIAccess::SRVGraphics, Texture->GetSubresourceRangeSRV());
+			}
 		}
 		break;
 		}
diff --git a/Engine/Source/Runtime/RenderCore/Private/RenderGraphValidation.cpp b/Engine/Source/Runtime/RenderCore/Private/RenderGraphValidation.cpp
--- a/Engine/Source/Runtime/RenderCore/Private/RenderGraphValidation.cpp
+++ b/Engine/Source/Runtime/RenderCore/Private/RenderGraphValidation.cpp
@@ -707,8 +707,12 @@ void FRDGUserValidation::ValidateAddPass(const FRDGPass* Pass, bool bSkipPassAcc
 	if (RenderTargetBindingSlots)
 	{
 		const auto& RenderTargets = RenderTargetBindingSlots->Output;
-
 		{
+			if (FRDGTextureRef Texture = RenderTargetBindingSlots->FoveationTexture)
+			{
+				MarkAsConsumed(Texture);
+			}
+
 			const auto& DepthStencil = RenderTargetBindingSlots->DepthStencil;
 
 			const auto CheckDepthStencil = [&](FRDGTextureRef Texture)
@@ -928,6 +932,11 @@ void FRDGUserValidation::ValidateExecutePassBegin(const FRDGPass* Pass)
 					Texture->GetTextureDebugData().bHasBeenBoundAsRenderTarget = true;
 					Texture->MarkResourceAsUsed();
 				}
+
+				if (FRDGTextureRef Texture = RenderTargets.FoveationTexture)
+				{
+					Texture->MarkResourceAsUsed();
+				}			
 			}
 			break;
 			}
@@ -1023,6 +1032,11 @@ void FRDGUserValidation::SetAllowRHIAccess(const FRDGPass* Pass, bool bAllowAcce
 			{
 				Texture->GetDebugData().bAllowRHIAccess = bAllowAccess;
 			}
+
+			if (FRDGTexture* Texture = RenderTargets.FoveationTexture)
+			{
+				Texture->GetDebugData().bAllowRHIAccess = bAllowAccess;
+			}
 		}
 	});
 }
