Add missing fix for foveation texture

diff --git a/Engine/Source/Runtime/RenderCore/Private/RenderGraphBuilder.cpp b/Engine/Source/Runtime/RenderCore/Private/RenderGraphBuilder.cpp
--- a/Engine/Source/Runtime/RenderCore/Private/RenderGraphBuilder.cpp
+++ b/Engine/Source/Runtime/RenderCore/Private/RenderGraphBuilder.cpp
@@ -247,6 +247,11 @@ void EnumerateTextureParameters(FRDGParameterStruct PassParameters, TParameterFu
 			{
 				ParameterFunction(Texture);
 			}
+
+			if (FRDGTextureRef Texture = RenderTargets.FoveationTexture)
+			{
+				ParameterFunction(Texture);
+			}
 		}
 		break;
 		}
