Fix CIS break (with rename of FoveationTexture to ShadingRateTexture).

diff --git a/Engine/Source/Runtime/RenderCore/Private/RenderGraphBuilder.cpp b/Engine/Source/Runtime/RenderCore/Private/RenderGraphBuilder.cpp
--- a/Engine/Source/Runtime/RenderCore/Private/RenderGraphBuilder.cpp
+++ b/Engine/Source/Runtime/RenderCore/Private/RenderGraphBuilder.cpp
@@ -190,7 +190,7 @@ void EnumerateTextureAccess(FRDGParameterStruct PassParameters, ERDGPassFlags Pa
 				});
 			}
 
-			if (FRDGTextureRef Texture = RenderTargets.FoveationTexture)
+			if (FRDGTextureRef Texture = RenderTargets.ShadingRateTexture)
 			{
 				AccessFunction(nullptr, Texture, ERHIAccess::SRVGraphics, Texture->GetSubresourceRangeSRV());
 			}
@@ -248,7 +248,7 @@ void EnumerateTextureParameters(FRDGParameterStruct PassParameters, TParameterFu
 				ParameterFunction(Texture);
 			}
 
-			if (FRDGTextureRef Texture = RenderTargets.FoveationTexture)
+			if (FRDGTextureRef Texture = RenderTargets.ShadingRateTexture)
 			{
 				ParameterFunction(Texture);
 			}
diff --git a/Engine/Source/Runtime/RenderCore/Private/RenderGraphValidation.cpp b/Engine/Source/Runtime/RenderCore/Private/RenderGraphValidation.cpp
--- a/Engine/Source/Runtime/RenderCore/Private/RenderGraphValidation.cpp
+++ b/Engine/Source/Runtime/RenderCore/Private/RenderGraphValidation.cpp
@@ -708,7 +708,7 @@ void FRDGUserValidation::ValidateAddPass(const FRDGPass* Pass, bool bSkipPassAcc
 	{
 		const auto& RenderTargets = RenderTargetBindingSlots->Output;
 		{
-			if (FRDGTextureRef Texture = RenderTargetBindingSlots->FoveationTexture)
+			if (FRDGTextureRef Texture = RenderTargetBindingSlots->ShadingRateTexture)
 			{
 				MarkAsConsumed(Texture);
 			}
@@ -933,7 +933,7 @@ void FRDGUserValidation::ValidateExecutePassBegin(const FRDGPass* Pass)
 					Texture->MarkResourceAsUsed();
 				}
 
-				if (FRDGTextureRef Texture = RenderTargets.FoveationTexture)
+				if (FRDGTextureRef Texture = RenderTargets.ShadingRateTexture)
 				{
 					Texture->MarkResourceAsUsed();
 				}			
@@ -1033,7 +1033,7 @@ void FRDGUserValidation::SetAllowRHIAccess(const FRDGPass* Pass, bool bAllowAcce
 				Texture->GetDebugData().bAllowRHIAccess = bAllowAccess;
 			}
 
-			if (FRDGTexture* Texture = RenderTargets.FoveationTexture)
+			if (FRDGTexture* Texture = RenderTargets.ShadingRateTexture)
 			{
 				Texture->GetDebugData().bAllowRHIAccess = bAllowAccess;
 			}
