Fix compiler warning in FShaderBindingInOutMask

#rnx

diff --git a/Engine/Source/Runtime/RenderCore/Public/CrossCompilerCommon.h b/Engine/Source/Runtime/RenderCore/Public/CrossCompilerCommon.h
--- a/Engine/Source/Runtime/RenderCore/Public/CrossCompilerCommon.h
+++ b/Engine/Source/Runtime/RenderCore/Public/CrossCompilerCommon.h
@@ -189,15 +189,17 @@ namespace CrossCompiler
 		return Ar;
 	}
 
+#pragma warning(push)
+#pragma warning(error : 4596)
 	struct FShaderBindingInOutMask
 	{
 		uint32 Bitmask = 0;
 
 		/** Maximum value for a valid index in this bitmask. */
-		static constexpr int32 MaxIndex = (sizeof(Bitmask) * 8u) - 1u;
+		static constexpr int32 MaxIndex = (sizeof(decltype(FShaderBindingInOutMask::Bitmask)) * 8u) - 1u;
 
 		/** Index to mark the binding of a depth-stencil output resource. */
-		static constexpr int32 DepthStencilMaskIndex = FShaderBindingInOutMask::MaxIndex;
+		static constexpr int32 DepthStencilMaskIndex = MaxIndex;
 
 		/** Sets the specified bitfield in this bitmask and validates its index boundary. */
 		FORCEINLINE void EnableField(int32 Index)
@@ -223,6 +225,7 @@ namespace CrossCompiler
 			return !(Lhs == Rhs);
 		}
 	};
+#pragma warning(pop)
 
 	inline FArchive& operator<<(FArchive& Ar, FShaderBindingInOutMask& BindingInOutMask)
 	{
