Pulling in fixes to mobile mip generation (related to MFW) into UE5 ahead of automatic integration (CL 15667874)

diff --git a/Engine/Source/Runtime/MediaAssets/Private/Misc/MediaTextureResource.cpp b/Engine/Source/Runtime/MediaAssets/Private/Misc/MediaTextureResource.cpp
--- a/Engine/Source/Runtime/MediaAssets/Private/Misc/MediaTextureResource.cpp
+++ b/Engine/Source/Runtime/MediaAssets/Private/Misc/MediaTextureResource.cpp
@@ -466,11 +466,7 @@ void FMediaTextureResource::Render(const FRenderParams& Params)
 			{
 				check(OutputTarget);
 
-#if PLATFORM_ANDROID
-				const EGenerateMipsPass GenerateMipsPass = EGenerateMipsPass::Raster;
-#else
 				const EGenerateMipsPass GenerateMipsPass = EGenerateMipsPass::Compute;
-#endif
 
 				CacheRenderTarget(OutputTarget, TEXT("MipGeneration"), MipGenerationCache);
 
diff --git a/Engine/Source/Runtime/OpenGLDrv/Private/OpenGLTexture.cpp b/Engine/Source/Runtime/OpenGLDrv/Private/OpenGLTexture.cpp
--- a/Engine/Source/Runtime/OpenGLDrv/Private/OpenGLTexture.cpp
+++ b/Engine/Source/Runtime/OpenGLDrv/Private/OpenGLTexture.cpp
@@ -2087,20 +2087,20 @@ FShaderResourceViewRHIRef FOpenGLDynamicRHI::RHICreateShaderResourceView(FRHITex
 /** Generates mip maps for the surface. */
 void FOpenGLDynamicRHI::RHIGenerateMips(FRHITexture* SurfaceRHI)
 {
-	VERIFY_GL_SCOPE();
-
-	FOpenGLTextureBase* Texture = GetOpenGLTextureFromRHITexture(SurfaceRHI);
-
-	if ( FOpenGL::SupportsGenerateMipmap())
+	if (FOpenGL::SupportsGenerateMipmap())
 	{
-		GPUProfilingData.RegisterGPUWork(0);
-
-		FOpenGLContextState& ContextState = GetContextStateForCurrentContext();
-		// Setup the texture on a disused unit
-		// need to figure out how to setup mips properly in no views case
-		CachedSetupTextureStage(ContextState, FOpenGL::GetMaxCombinedTextureImageUnits() - 1, Texture->Target, Texture->GetResource(), -1, 1);
-
-		FOpenGL::GenerateMipmap( Texture->Target);
+		RunOnGLRenderContextThread([=]()
+		{
+			VERIFY_GL_SCOPE();
+			GPUProfilingData.RegisterGPUWork(0);
+
+			FOpenGLContextState& ContextState = GetContextStateForCurrentContext();
+			FOpenGLTextureBase* Texture = GetOpenGLTextureFromRHITexture(SurfaceRHI);
+			// Setup the texture on a disused unit
+			// need to figure out how to setup mips properly in no views case
+			CachedSetupTextureStage(ContextState, FOpenGL::GetMaxCombinedTextureImageUnits() - 1, Texture->Target, Texture->GetResource(), -1, Texture->NumMips); 
+			FOpenGL::GenerateMipmap(Texture->Target);
+		});
 	}
 	else
 	{
diff --git a/Engine/Source/Runtime/OpenGLDrv/Public/OpenGLES.h b/Engine/Source/Runtime/OpenGLDrv/Public/OpenGLES.h
--- a/Engine/Source/Runtime/OpenGLDrv/Public/OpenGLES.h
+++ b/Engine/Source/Runtime/OpenGLDrv/Public/OpenGLES.h
@@ -580,6 +580,13 @@ struct FOpenGLES : public FOpenGLBase
 		glClearDepthf(Depth);
 	}
 
+	static FORCEINLINE void GenerateMipmap( GLenum Target )
+	{
+		glGenerateMipmap( Target);
+	}
+	
+	static FORCEINLINE bool SupportsGenerateMipmap() { return true; }
+
 	static FORCEINLINE GLuint GetMajorVersion()
 	{
 		return 3;
diff --git a/Engine/Source/Runtime/RenderCore/Private/GenerateMips.cpp b/Engine/Source/Runtime/RenderCore/Private/GenerateMips.cpp
--- a/Engine/Source/Runtime/RenderCore/Private/GenerateMips.cpp
+++ b/Engine/Source/Runtime/RenderCore/Private/GenerateMips.cpp
@@ -144,15 +144,19 @@ void FGenerateMips::ExecuteCompute(FRDGBuilder& GraphBuilder, FRDGTextureRef Tex
 	const FRDGTextureDesc& TextureDesc = Texture->Desc;
 
 	// Select compute shader variant (normal vs. sRGB etc.)
+	bool bMipsSRGB = EnumHasAnyFlags(TextureDesc.Flags, TexCreate_SRGB);
 #if PLATFORM_ANDROID
-	const bool bIsUsingVulkan = FAndroidMisc::ShouldUseVulkan();
-#else
-	const bool bIsUsingVulkan = false;
+	if (IsVulkanPlatform(GMaxRHIShaderPlatform))
+	{
+		// Vulkan Android seems to skip sRGB->Lin conversion when sampling texture in compute
+		bMipsSRGB = false;
+	}
 #endif
+	const bool bMipsSwizzle = false; 
 
 	FGenerateMipsCS::FPermutationDomain PermutationVector;
-	PermutationVector.Set<FGenerateMipsCS::FGenMipsSRGB>(EnumHasAnyFlags(TextureDesc.Flags, TexCreate_SRGB));
-	PermutationVector.Set<FGenerateMipsCS::FGenMipsSwizzle>(bIsUsingVulkan); // On Vulkan we experience RGB being swizzled around, this little switch circumvents the issue.
+	PermutationVector.Set<FGenerateMipsCS::FGenMipsSRGB>(bMipsSRGB);
+	PermutationVector.Set<FGenerateMipsCS::FGenMipsSwizzle>(bMipsSwizzle);
 	TShaderMapRef<FGenerateMipsCS> ComputeShader(GetGlobalShaderMap(GMaxRHIFeatureLevel), PermutationVector);
 
 	// Loop through each level of the mips that require creation and add a dispatch pass per level.
@@ -188,8 +192,7 @@ void FGenerateMips::Execute(FRDGBuilder& GraphBuilder, FRDGTextureRef Texture, F
 		if (RHIRequiresComputeGenerateMips())
 		{
 			FSamplerStateInitializerRHI SamplerInit(Params.Filter, Params.AddressU, Params.AddressV, Params.AddressW);
-			FSamplerStateRHIRef Sampler = *GraphBuilder.AllocObject<FSamplerStateRHIRef>(RHICreateSamplerState(SamplerInit));
-			Execute(GraphBuilder, Texture, Sampler, Pass);
+			Execute(GraphBuilder, Texture, RHICreateSamplerState(SamplerInit), Pass);
 		}
 		else
 		{
