Fixing issues preventing previewing of lower shader models when launching with SM6.

diff --git a/Engine/Source/Runtime/RenderCore/Private/RenderUtils.cpp b/Engine/Source/Runtime/RenderCore/Private/RenderUtils.cpp
--- a/Engine/Source/Runtime/RenderCore/Private/RenderUtils.cpp
+++ b/Engine/Source/Runtime/RenderCore/Private/RenderUtils.cpp
@@ -1113,7 +1113,7 @@ RENDERCORE_API bool MobileRequiresSceneDepthAux(const FStaticShaderPlatform Plat
 
 RENDERCORE_API bool SupportsTextureCubeArray(ERHIFeatureLevel::Type FeatureLevel)
 {
-	return FeatureLevel == ERHIFeatureLevel::SM5 
+	return FeatureLevel >= ERHIFeatureLevel::SM5 
 		// mobile deferred requries ES3.2 feature set
 		|| IsMobileDeferredShadingEnabled(GMaxRHIShaderPlatform);
 }
diff --git a/Engine/Source/Runtime/RenderCore/Private/ShaderResource.cpp b/Engine/Source/Runtime/RenderCore/Private/ShaderResource.cpp
--- a/Engine/Source/Runtime/RenderCore/Private/ShaderResource.cpp
+++ b/Engine/Source/Runtime/RenderCore/Private/ShaderResource.cpp
@@ -39,11 +39,9 @@ bool FShaderMapResource::ArePlatformsCompatible(EShaderPlatform CurrentPlatform,
 	{
 		bFeatureLevelCompatible = GetMaxSupportedFeatureLevel(CurrentPlatform) >= GetMaxSupportedFeatureLevel(TargetPlatform);
 
-		bool const bIsTargetD3D = TargetPlatform == SP_PCD3D_SM5 ||
-			TargetPlatform == SP_PCD3D_ES3_1;
+		bool const bIsTargetD3D = IsD3DPlatform(TargetPlatform);
 
-		bool const bIsCurrentPlatformD3D = CurrentPlatform == SP_PCD3D_SM5 ||
-			TargetPlatform == SP_PCD3D_ES3_1;
+		bool const bIsCurrentPlatformD3D = IsD3DPlatform(CurrentPlatform);
 
 		// For Metal in Editor we can switch feature-levels, but not in cooked projects when using Metal shader librariss.
 		bool const bIsCurrentMetal = IsMetalPlatform(CurrentPlatform);
