Make shader compression a project setting and move to SCWs (UE-73155).

diff --git a/Engine/Source/Programs/ShaderCompileWorker/Private/ShaderCompileWorker.cpp b/Engine/Source/Programs/ShaderCompileWorker/Private/ShaderCompileWorker.cpp
--- a/Engine/Source/Programs/ShaderCompileWorker/Private/ShaderCompileWorker.cpp
+++ b/Engine/Source/Programs/ShaderCompileWorker/Private/ShaderCompileWorker.cpp
@@ -128,6 +128,14 @@ static void ProcessCompilationJob(const FShaderCompilerInput& Input,FShaderCompi
 	// Compile the shader directly through the platform dll (directly from the shader dir as the working directory)
 	double TimeStart = FPlatformTime::Seconds();
 	Compiler->CompileShader(Input.ShaderFormat, Input, Output, WorkingDirectory);
+	if (Output.bSucceeded)
+	{
+		Output.GenerateOutputHash();
+		if (Input.CompressionFormat != NAME_None)
+		{
+			Output.CompressOutput(Input.CompressionFormat);
+		}
+	}
 	Output.CompileTime = FPlatformTime::Seconds() - TimeStart;
 
 	if (Compiler->UsesHLSLcc(Input))
diff --git a/Engine/Source/Runtime/Engine/Classes/Engine/RendererSettings.h b/Engine/Source/Runtime/Engine/Classes/Engine/RendererSettings.h
--- a/Engine/Source/Runtime/Engine/Classes/Engine/RendererSettings.h
+++ b/Engine/Source/Runtime/Engine/Classes/Engine/RendererSettings.h
@@ -147,6 +147,19 @@ namespace EDefaultBackBufferPixelFormat
 	ENGINE_API EDefaultBackBufferPixelFormat::Type FromInt(int32 InDefaultBackBufferPixelFormat);
 }
 
+/**
+ * Enumerates supported shader compression formats.
+ */
+UENUM()
+namespace EShaderCompressionFormat
+{
+	enum Type
+	{
+		None = 0 UMETA(DisplayName = "Do not compress", ToolTip = "Fastest, but disk and memory footprint will be large"),
+		LZ4 = 1 UMETA(DisplayName = "LZ4", ToolTip = "Compressing using LZ4"),
+	};
+}
+
 /**
  * Rendering settings.
  */
@@ -201,6 +214,12 @@ class ENGINE_API URendererSettings : public UDeveloperSettings
 		ToolTip = "When running in game mode, whether to keep shaders for all quality levels in memory or only those needed for the current quality level.\nUnchecked: Keep all quality levels in memory allowing a runtime quality level change. (default)\nChecked: Discard unused quality levels when loading content for the game, saving some memory."))
 	uint32 bDiscardUnusedQualityLevels : 1;
 
+	UPROPERTY(config, EditAnywhere, Category = Materials, meta = (
+		ConsoleVariable = "r.Shaders.CompressionFormat", DisplayName = "Shader Compression Format",
+		ToolTip = "Select how the shaders are compressed for storage",
+		ConfigRestartRequired = true))
+	TEnumAsByte<EShaderCompressionFormat::Type> ShaderCompressionFormat;	// requires restart as we may have jobs in flight that already used the previous one. Also it's also used for decompression the currently compressed shaders
+
 	UPROPERTY(config, EditAnywhere, Category=Culling, meta=(
 		ConsoleVariable="r.AllowOcclusionQueries",DisplayName="Occlusion Culling",
 		ToolTip="Allows occluded meshes to be culled and not rendered."))
diff --git a/Engine/Source/Runtime/Engine/Private/ShaderCompiler/ShaderCompiler.cpp b/Engine/Source/Runtime/Engine/Private/ShaderCompiler/ShaderCompiler.cpp
--- a/Engine/Source/Runtime/Engine/Private/ShaderCompiler/ShaderCompiler.cpp
+++ b/Engine/Source/Runtime/Engine/Private/ShaderCompiler/ShaderCompiler.cpp
@@ -1281,9 +1281,8 @@ static bool ReadSingleJob(FShaderCompileJob* CurrentJob, FArchive& OutputFile)
 	// Deserialize the shader compilation output.
 	OutputFile << CurrentJob->Output;
 
-	// Generate a hash of the output and cache it
-	// The shader processing this output will use it to search for existing FShaderResources
-	CurrentJob->Output.GenerateOutputHash();
+	// The job should already have a non-zero output hash
+	checkf(CurrentJob->Output.OutputHash != FSHAHash(), TEXT("OutputHash was not set in the shader compile worker!"));
 	CurrentJob->bSucceeded = CurrentJob->Output.bSucceeded;
 
 	if (CurrentJob->bSucceeded && CurrentJob->Input.DumpDebugInfoPath.Len() > 0)
@@ -2233,9 +2232,10 @@ void FShaderCompileUtilities::ExecuteShaderCompileJob(FShaderCommonCompileJob& J
 
 		if (SingleJob->Output.bSucceeded)
 		{
-			// Generate a hash of the output and cache it
-			// The shader processing this output will use it to search for existing FShaderResources
+			// Generate a hash of the output and compress the code
+			// The shader processing this output will use the heash to search for existing FShaderResources
 			SingleJob->Output.GenerateOutputHash();
+			SingleJob->Output.CompressOutput(NAME_LZ4);
 		}
 	}
 	else
@@ -4469,6 +4469,7 @@ void GlobalBeginCompileShader(
 
 	Input.Target = Target;
 	Input.ShaderFormat = ShaderFormatName;
+	Input.CompressionFormat = GetShaderCompressionFormat();
 	Input.VirtualSourceFilePath = SourceFilename;
 	Input.EntryPointName = FunctionName;
 	Input.bCompilingForShaderPipeline = false;
diff --git a/Engine/Source/Runtime/RenderCore/Private/Shader.cpp b/Engine/Source/Runtime/RenderCore/Private/Shader.cpp
--- a/Engine/Source/Runtime/RenderCore/Private/Shader.cpp
+++ b/Engine/Source/Runtime/RenderCore/Private/Shader.cpp
@@ -1838,4 +1838,10 @@ void ShaderMapAppendKeyString(EShaderPlatform Platform, FString& KeyString)
 		const bool bSupportsAnisotropicMaterials = FDataDrivenShaderPlatformInfo::GetSupportsAnisotropicMaterials(Platform);
 		KeyString += FString::Printf(TEXT("_Aniso-%d"), bSupportsAnisotropicMaterials ? 1 : 0);
 	}
+
+	{
+		// add shader compression format
+		KeyString += TEXT("_Compr");
+		KeyString += GetShaderCompressionFormat().ToString();
+	}
 }
diff --git a/Engine/Source/Runtime/RenderCore/Private/ShaderCodeArchive.cpp b/Engine/Source/Runtime/RenderCore/Private/ShaderCodeArchive.cpp
--- a/Engine/Source/Runtime/RenderCore/Private/ShaderCodeArchive.cpp
+++ b/Engine/Source/Runtime/RenderCore/Private/ShaderCodeArchive.cpp
@@ -31,8 +31,6 @@ static FAutoConsoleVariableRef CVarShaderCodeLibraryAsyncLoadingAllowDontCache(
 );
 
 
-static const FName ShaderLibraryCompressionFormat = NAME_LZ4;
-
 int32 FSerializedShaderArchive::FindShaderMapWithKey(const FSHAHash& Hash, uint32 Key) const
 {
 	for (uint32 Index = ShaderMapHashTable.First(Key); ShaderMapHashTable.IsValid(Index); Index = ShaderMapHashTable.Next(Index))
@@ -140,7 +138,7 @@ void FSerializedShaderArchive::DecompressShader(int32 Index, const TArray<TArray
 	}
 	else
 	{
-		bool bSucceed = FCompression::UncompressMemory(ShaderLibraryCompressionFormat, OutDecompressedShader.GetData(), Entry.UncompressedSize, ShaderCode[Index].GetData(), Entry.Size);
+		bool bSucceed = FCompression::UncompressMemory(GetShaderCompressionFormat(), OutDecompressedShader.GetData(), Entry.UncompressedSize, ShaderCode[Index].GetData(), Entry.Size);
 		check(bSucceed);
 	}
 }
@@ -956,7 +954,7 @@ TRefCountPtr<FRHIShader> FShaderCodeArchive::CreateShader(int32 Index)
 	if (ShaderEntry.UncompressedSize != ShaderEntry.Size)
 	{
 		void* UncompressedCode = MemStack.Alloc(ShaderEntry.UncompressedSize, 16);
-		const bool bDecompressResult = FCompression::UncompressMemory(ShaderLibraryCompressionFormat, UncompressedCode, ShaderEntry.UncompressedSize, ShaderCode, ShaderEntry.Size);
+		const bool bDecompressResult = FCompression::UncompressMemory(GetShaderCompressionFormat(), UncompressedCode, ShaderEntry.UncompressedSize, ShaderCode, ShaderEntry.Size);
 		check(bDecompressResult);
 		ShaderCode = (uint8*)UncompressedCode;
 	}
diff --git a/Engine/Source/Runtime/RenderCore/Private/ShaderCore.cpp b/Engine/Source/Runtime/RenderCore/Private/ShaderCore.cpp
--- a/Engine/Source/Runtime/RenderCore/Private/ShaderCore.cpp
+++ b/Engine/Source/Runtime/RenderCore/Private/ShaderCore.cpp
@@ -424,6 +424,14 @@ void FShaderCompilerOutput::GenerateOutputHash()
 	HashState.GetHash(&OutputHash.Hash[0]);
 }
 
+void FShaderCompilerOutput::CompressOutput(FName ShaderCompressionFormat)
+{
+	// make sure the hash has been generated
+	checkf(OutputHash != FSHAHash(), TEXT("Output hash must be generated before compressing the shader code."));
+	checkf(ShaderCompressionFormat != NAME_None, TEXT("Compression format should be valid"));
+	ShaderCode.Compress(ShaderCompressionFormat);
+}
+
 static void ReportVirtualShaderFilePathError(TArray<FShaderCompilerError>* CompileErrors, FString ErrorString)
 {
 	if (CompileErrors)
@@ -1359,3 +1367,85 @@ void AddShaderSourceDirectoryMapping(const FString& VirtualShaderDirectory, cons
 	GShaderSourceDirectoryMappings.Add(VirtualShaderDirectory, RealShaderDirectory);
 }
 
+void FShaderCode::Compress(FName ShaderCompressionFormat)
+{
+	checkf(OptionalDataSize == -1, TEXT("FShaderCode::Compress() was called before calling FShaderCode::FinalizeShaderCode()"));
+
+	TArray<uint8> Compressed;
+	int32 CompressedSize = ShaderCodeWithOptionalData.Num();
+	Compressed.AddUninitialized(CompressedSize);
+
+	// there is code that assumes that if CompressedSize == CodeSize, the shader isn't compressed. Because of that, do not accept equal compressed size (very unlikely anyway)
+	if (FCompression::CompressMemory(ShaderCompressionFormat, Compressed.GetData(), CompressedSize, ShaderCodeWithOptionalData.GetData(), ShaderCodeWithOptionalData.Num()) && CompressedSize < ShaderCodeWithOptionalData.Num())
+	{
+		// cache the ShaderCodeSize since it will no longer possible to get it as the reader will fail to parse the compressed data
+		FShaderCodeReader Wrapper(ShaderCodeWithOptionalData);
+		ShaderCodeSize = Wrapper.GetShaderCodeSize();
+
+		// finalize the compression
+		CompressionFormat = ShaderCompressionFormat;
+		UncompressedSize = ShaderCodeWithOptionalData.Num();
+
+		Compressed.SetNum(CompressedSize);
+		ShaderCodeWithOptionalData = Compressed;
+	}
+}
+
+FArchive& operator<<(FArchive& Ar, FShaderCode& Output)
+{
+	if (Ar.IsLoading())
+	{
+		Output.OptionalDataSize = -1;
+	}
+	else
+	{
+		Output.FinalizeShaderCode();
+	}
+
+	// Note: this serialize is used to pass between UE4 and the shader compile worker, recompile both when modifying
+	Ar << Output.ShaderCodeWithOptionalData;
+	Ar << Output.UncompressedSize;
+	{
+		FString CompressionFormatString(Output.CompressionFormat.ToString());
+		Ar << CompressionFormatString;
+		Output.CompressionFormat = FName(*CompressionFormatString);
+	}
+	Ar << Output.ShaderCodeSize;
+	return Ar;
+}
+
+FArchive& operator<<(FArchive& Ar, FShaderCompilerInput& Input)
+{
+	// Note: this serialize is used to pass between UE4 and the shader compile worker, recompile both when modifying
+	Ar << Input.Target;
+	{
+		FString ShaderFormatString(Input.ShaderFormat.ToString());
+		Ar << ShaderFormatString;
+		Input.ShaderFormat = FName(*ShaderFormatString);
+	}
+	{
+		FString CompressionFormatString(Input.CompressionFormat.ToString());
+		Ar << CompressionFormatString;
+		Input.CompressionFormat = FName(*CompressionFormatString);
+	}
+	Ar << Input.SourceFilePrefix;
+	Ar << Input.VirtualSourceFilePath;
+	Ar << Input.EntryPointName;
+	Ar << Input.bSkipPreprocessedCache;
+	Ar << Input.bCompilingForShaderPipeline;
+	Ar << Input.bGenerateDirectCompileFile;
+	Ar << Input.bIncludeUsedOutputs;
+	Ar << Input.UsedOutputs;
+	Ar << Input.DumpDebugInfoRootPath;
+	Ar << Input.DumpDebugInfoPath;
+	Ar << Input.DebugExtension;
+	Ar << Input.DebugGroupName;
+	Ar << Input.DebugDescription;
+	Ar << Input.Environment;
+	Ar << Input.ExtraSettings;
+	Ar << Input.RootParameterBindings;
+
+	// Note: skipping Input.SharedEnvironment, which is handled by FShaderCompileUtilities::DoWriteTasks in order to maintain sharing
+
+	return Ar;
+}
diff --git a/Engine/Source/Runtime/RenderCore/Private/ShaderResource.cpp b/Engine/Source/Runtime/RenderCore/Private/ShaderResource.cpp
--- a/Engine/Source/Runtime/RenderCore/Private/ShaderResource.cpp
+++ b/Engine/Source/Runtime/RenderCore/Private/ShaderResource.cpp
@@ -29,7 +29,29 @@
 #include "Interfaces/IShaderFormat.h"
 #endif
 
-static const FName ShaderCompressionFormat = NAME_LZ4;
+int32 GShaderCompressionFormatChoice = 1;
+static FAutoConsoleVariableRef CVarShaderCompressionFormatChoice(
+	TEXT("r.Shaders.CompressionFormat"),
+	GShaderCompressionFormatChoice,
+	TEXT("Select the compression methods for the shader code.\n")
+	TEXT(" 0: None (uncompressed)\n")
+	TEXT(" 1: LZ4 (default)"),
+	ECVF_ReadOnly);
+
+FName GetShaderCompressionFormat(const FName& ShaderFormat)
+{
+	// support an older developer-only CVar for compatibility and make it preempt
+#if !(UE_BUILD_SHIPPING || UE_BUILD_TEST)
+	static const IConsoleVariable* CVarSkipCompression = IConsoleManager::Get().FindConsoleVariable(TEXT("r.Shaders.SkipCompression"));
+	static bool bSkipCompression = (CVarSkipCompression && CVarSkipCompression->GetInt() != 0);
+	if (UNLIKELY(bSkipCompression))
+	{
+		return NAME_None;
+	}
+#endif
+
+	return (GShaderCompressionFormatChoice == 1) ? NAME_LZ4 : NAME_None;
+}
 
 bool FShaderMapResource::ArePlatformsCompatible(EShaderPlatform CurrentPlatform, EShaderPlatform TargetPlatform)
 {
@@ -162,10 +184,10 @@ void FShaderMapResourceCode::AddShaderCompilerOutput(const FShaderCompilerOutput
 #if WITH_EDITORONLY_DATA
 	AddPlatformDebugData(Output.PlatformDebugData);
 #endif
-	AddShaderCode(Output.Target.GetFrequency(), Output.OutputHash, Output.ShaderCode.GetReadAccess());
+	AddShaderCode(Output.Target.GetFrequency(), Output.OutputHash, Output.ShaderCode);
 }
 
-void FShaderMapResourceCode::AddShaderCode(EShaderFrequency InFrequency, const FSHAHash& InHash, TConstArrayView<uint8> InCode)
+void FShaderMapResourceCode::AddShaderCode(EShaderFrequency InFrequency, const FSHAHash& InHash, const FShaderCode& InCode)
 {
 	TRACE_CPUPROFILER_EVENT_SCOPE(FShaderMapResourceCode::AddShaderCode);
 
@@ -176,26 +198,44 @@ void FShaderMapResourceCode::AddShaderCode(EShaderFrequency InFrequency, const F
 
 		FShaderEntry& Entry = ShaderEntries.InsertDefaulted_GetRef(Index);
 		Entry.Frequency = InFrequency;
-		Entry.UncompressedSize = InCode.Num();
+		const TArray<uint8>& ShaderCode = InCode.GetReadAccess();
 
-		bool bAllowShaderCompression = true;
-#if !(UE_BUILD_SHIPPING || UE_BUILD_TEST)
-		static const IConsoleVariable* CVarSkipCompression = IConsoleManager::Get().FindConsoleVariable(TEXT("r.Shaders.SkipCompression"));
-		bAllowShaderCompression = CVarSkipCompression ? CVarSkipCompression->GetInt() == 0 : true;
-#endif
-
-		int32 CompressedSize = InCode.Num();
-		Entry.Code.AddUninitialized(CompressedSize);
-
-		if (bAllowShaderCompression && FCompression::CompressMemory(ShaderCompressionFormat, Entry.Code.GetData(), CompressedSize, InCode.GetData(), InCode.Num()))
+		FName ShaderCompressionFormat = GetShaderCompressionFormat();
+		if (ShaderCompressionFormat != NAME_None)
 		{
-			// resize to fit reduced compressed size, but don't reallocate memory
-			Entry.Code.SetNum(CompressedSize, false);
+			Entry.UncompressedSize = InCode.GetUncompressedSize();
+
+			// we trust that SCWs also obeyed by the same CVar, so we expect a compressed shader code at this point
+			// However, if we see an uncompressed shader, it perhaps means that SCW tried to compress it, but the result was worse than uncompressed. 
+			// Because of that we special-case NAME_None here
+			if (ShaderCompressionFormat != InCode.GetCompressionFormat())
+			{
+				if (InCode.GetCompressionFormat() != NAME_None)
+				{
+					UE_LOG(LogShaders, Fatal, TEXT("Shader %s is expected to be compressed with %s, but it is compressed with %s instead."),
+						*InHash.ToString(),
+						*ShaderCompressionFormat.ToString(),
+						*InCode.GetCompressionFormat().ToString()
+						);
+					// unreachable
+					return;
+				}
+				
+				// assume uncompressed due to worse ratio than the compression
+				Entry.UncompressedSize = ShaderCode.Num();
+				UE_LOG(LogShaders, Warning, TEXT("Shader %s is expected to be compressed with %s, but it arrived uncompressed. Assuming compressing made it longer and storing uncompressed."),
+					*InHash.ToString(),
+					*ShaderCompressionFormat.ToString(),
+					*InCode.GetCompressionFormat().ToString()
+				);
+			}
 		}
 		else
 		{
-			FMemory::Memcpy(Entry.Code.GetData(), InCode.GetData(), InCode.Num());
+			Entry.UncompressedSize = ShaderCode.Num();
 		}
+
+		Entry.Code = ShaderCode;
 	}
 }
 
@@ -403,7 +443,7 @@ TRefCountPtr<FRHIShader> FShaderMapResource_InlineCode::CreateRHIShader(int32 Sh
 	if (ShaderEntry.Code.Num() != ShaderEntry.UncompressedSize)
 	{
 		void* UncompressedCode = MemStack.Alloc(ShaderEntry.UncompressedSize, 16);
-		auto bSucceed = FCompression::UncompressMemory(ShaderCompressionFormat, UncompressedCode, ShaderEntry.UncompressedSize, ShaderCode, ShaderEntry.Code.Num());
+		auto bSucceed = FCompression::UncompressMemory(GetShaderCompressionFormat(), UncompressedCode, ShaderEntry.UncompressedSize, ShaderCode, ShaderEntry.Code.Num());
 		check(bSucceed);
 		ShaderCode = (uint8*)UncompressedCode;
 	}
diff --git a/Engine/Source/Runtime/RenderCore/Public/Shader.h b/Engine/Source/Runtime/RenderCore/Public/Shader.h
--- a/Engine/Source/Runtime/RenderCore/Public/Shader.h
+++ b/Engine/Source/Runtime/RenderCore/Public/Shader.h
@@ -359,7 +359,7 @@ public:
 
 	int32 FindShaderIndex(const FSHAHash& InHash) const;
 
-	RENDERCORE_API void AddShaderCode(EShaderFrequency InFrequency, const FSHAHash& InHash, TConstArrayView<uint8> InCode);
+	RENDERCORE_API void AddShaderCode(EShaderFrequency InFrequency, const FSHAHash& InHash, const FShaderCode& InCode);
 #if WITH_EDITORONLY_DATA
 	RENDERCORE_API void AddPlatformDebugData(TConstArrayView<uint8> InPlatformDebugData);
 #endif
diff --git a/Engine/Source/Runtime/RenderCore/Public/ShaderCompilerCore.h b/Engine/Source/Runtime/RenderCore/Public/ShaderCompilerCore.h
--- a/Engine/Source/Runtime/RenderCore/Public/ShaderCompilerCore.h
+++ b/Engine/Source/Runtime/RenderCore/Public/ShaderCompilerCore.h
@@ -145,6 +145,7 @@ struct FShaderCompilerInput
 {
 	FShaderTarget Target;
 	FName ShaderFormat;
+	FName CompressionFormat;
 	FString SourceFilePrefix;
 	FString VirtualSourceFilePath;
 	FString EntryPointName;
@@ -296,52 +297,12 @@ struct FShaderCompilerInput
 		}
 	}
 
-	friend FArchive& operator<<(FArchive& Ar,FShaderCompilerInput& Input)
-	{
-		// Note: this serialize is used to pass between UE4 and the shader compile worker, recompile both when modifying
-		Ar << Input.Target;
-		{
-			FString ShaderFormatString(Input.ShaderFormat.ToString());
-			Ar << ShaderFormatString;
-			Input.ShaderFormat = FName(*ShaderFormatString);
-		}
-		Ar << Input.SourceFilePrefix;
-		Ar << Input.VirtualSourceFilePath;
-		Ar << Input.EntryPointName;
-		Ar << Input.bSkipPreprocessedCache;
-		Ar << Input.bCompilingForShaderPipeline;
-		Ar << Input.bGenerateDirectCompileFile;
-		Ar << Input.bIncludeUsedOutputs;
-		Ar << Input.UsedOutputs;
-		Ar << Input.DumpDebugInfoRootPath;
-		Ar << Input.DumpDebugInfoPath;
-		Ar << Input.DebugExtension;
-		Ar << Input.DebugGroupName;
-		Ar << Input.DebugDescription;
-		Ar << Input.Environment;
-		Ar << Input.ExtraSettings;
-		Ar << Input.RootParameterBindings;
-
-		// Note: skipping Input.SharedEnvironment, which is handled by FShaderCompileUtilities::DoWriteTasks in order to maintain sharing
-
-		return Ar;
-	}
+	friend RENDERCORE_API FArchive& operator<<(FArchive& Ar, FShaderCompilerInput& Input);
 
 	bool IsUsingTessellation() const
 	{
-		switch (Target.GetFrequency())
-		{
-		case SF_Vertex:
-		{
-			const FString* UsingTessellationDefine = Environment.GetDefinitions().Find(TEXT("USING_TESSELLATION"));
-			return (UsingTessellationDefine != nullptr && *UsingTessellationDefine == TEXT("1"));
-		}
-		case SF_Hull:
-		case SF_Domain:
-			return true;
-		default:
-			return false;
-		}
+		// no more tessellation in UE5
+		return false;
 	}
 
 	bool IsRayTracingShader() const
@@ -478,10 +439,13 @@ struct FShaderCompilerOutput
 	/** Generates OutputHash from the compiler output. */
 	RENDERCORE_API void GenerateOutputHash();
 
+	/** Calls GenerateOutputHash() before the compression, replaces FShaderCode with the compressed data (if compression result was smaller). */
+	RENDERCORE_API void CompressOutput(FName ShaderCompressionFormat);
+
 	friend FArchive& operator<<(FArchive& Ar, FShaderCompilerOutput& Output)
 	{
 		// Note: this serialize is used to pass between UE4 and the shader compile worker, recompile both when modifying
-		Ar << Output.ParameterMap << Output.Errors << Output.Target << Output.ShaderCode << Output.NumInstructions << Output.NumTextureSamplers << Output.bSucceeded;
+		Ar << Output.ParameterMap << Output.Errors << Output.Target << Output.ShaderCode << Output.OutputHash << Output.NumInstructions << Output.NumTextureSamplers << Output.bSucceeded;
 		Ar << Output.bFailedRemovingUnused << Output.bSupportsQueryingUsedAttributes << Output.UsedAttributes;
 		Ar << Output.CompileTime;
 		Ar << Output.OptionalFinalShaderSource;
diff --git a/Engine/Source/Runtime/RenderCore/Public/ShaderCore.h b/Engine/Source/Runtime/RenderCore/Public/ShaderCore.h
--- a/Engine/Source/Runtime/RenderCore/Public/ShaderCore.h
+++ b/Engine/Source/Runtime/RenderCore/Public/ShaderCore.h
@@ -99,6 +99,9 @@ extern RENDERCORE_API bool AllowDebugViewmodes();
 /** Returns true if debug viewmodes are allowed for the given platform. */
 extern RENDERCORE_API bool AllowDebugViewmodes(EShaderPlatform Platform);
 
+/** Returns the shader compression format (passing ShaderFormat for future proofing, but as of now the setting is global for all formats). */
+extern RENDERCORE_API FName GetShaderCompressionFormat(const FName& ShaderFormat = NAME_None);
+
 struct FShaderTarget
 {
 	// The rest of uint32 holding the bitfields can be left unitialized. Union with a uint32 serves to prevent that to be able to set the whole uint32 value
@@ -733,10 +736,22 @@ class FShaderCode
 	// access through class methods
 	mutable TArray<uint8> ShaderCodeWithOptionalData;
 
+	/** ShaderCode may be compressed in SCWs on demand. If this value isn't null, the shader code is compressed. */
+	mutable int32 UncompressedSize;
+
+	/** Compression algo */
+	mutable FName CompressionFormat;
+
+	/** We cannot get the code size after the compression, so store it here */
+	mutable int32 ShaderCodeSize;
+
 public:
 
 	FShaderCode()
 	: OptionalDataSize(0)
+	, UncompressedSize(0)
+	, CompressionFormat(NAME_None)
+	, ShaderCodeSize(0)
 	{
 	}
 
@@ -745,30 +760,40 @@ public:
 	{
 		if(OptionalDataSize != -1)
 		{
+			checkf(UncompressedSize == 0, TEXT("FShaderCode::FinalizeShaderCode() was called after compressing the code"));
 			OptionalDataSize += sizeof(OptionalDataSize);
 			ShaderCodeWithOptionalData.Append((const uint8*)&OptionalDataSize, sizeof(OptionalDataSize));
 			OptionalDataSize = -1;
 		}
 	}
 
+	void Compress(FName ShaderCompressionFormat);
+
 	// Write access for regular microcode: Optional Data must be added AFTER regular microcode and BEFORE Finalize
 	TArray<uint8>& GetWriteAccess()
 	{
-		checkf(OptionalDataSize != -1, TEXT("Tried to add ShaderCode after adding being finalized!"));
+		checkf(OptionalDataSize != -1, TEXT("Tried to add ShaderCode after being finalized!"));
 		checkf(OptionalDataSize == 0, TEXT("Tried to add ShaderCode after adding Optional data!"));
 		return ShaderCodeWithOptionalData;
 	}
 
 	int32 GetShaderCodeSize() const
 	{
-		FinalizeShaderCode();
-
-		FShaderCodeReader Wrapper(ShaderCodeWithOptionalData);
+		// use the cached size whenever available
+		if (ShaderCodeSize != 0)
+		{
+			return ShaderCodeSize;
+		}
+		else
+		{
+			FinalizeShaderCode();
 
-		return Wrapper.GetShaderCodeSize();
+			FShaderCodeReader Wrapper(ShaderCodeWithOptionalData);
+			return Wrapper.GetShaderCodeSize();
+		}
 	}
 
-	// for read access, can have additional data attached to the end
+	// for read access, can have additional data attached to the end. Can also be compressed
 	const TArray<uint8>& GetReadAccess() const
 	{
 		FinalizeShaderCode();
@@ -776,6 +801,21 @@ public:
 		return ShaderCodeWithOptionalData;
 	}
 
+	bool IsCompressed() const
+	{
+		return UncompressedSize != 0;
+	}
+
+	FName GetCompressionFormat() const
+	{
+		return CompressionFormat;
+	}
+
+	int32 GetUncompressedSize() const
+	{
+		return UncompressedSize;
+	}
+
 	// for convenience
 	template <class T>
 	void AddOptionalData(const T &In)
@@ -809,21 +849,7 @@ public:
 		AddOptionalData(Key, (uint8*)InString, Size);
 	}
 
-	friend FArchive& operator<<(FArchive& Ar, FShaderCode& Output)
-	{
-		if(Ar.IsLoading())
-		{
-			Output.OptionalDataSize = -1;
-		}
-		else
-		{
-			Output.FinalizeShaderCode();
-		}
-
-		// Note: this serialize is used to pass between UE4 and the shader compile worker, recompile both when modifying
-		Ar << Output.ShaderCodeWithOptionalData;
-		return Ar;
-	}
+	friend RENDERCORE_API FArchive& operator<<(FArchive& Ar, FShaderCode& Output);
 };
 
 /**
