Disabled high-frequency CPU profiler scopes for FShaderMapContent::GetShader by default. Add FMaterial::GetShader only for the slow path branch.

diff --git a/Engine/Source/Runtime/RenderCore/Private/ShaderMap.cpp b/Engine/Source/Runtime/RenderCore/Private/ShaderMap.cpp
--- a/Engine/Source/Runtime/RenderCore/Private/ShaderMap.cpp
+++ b/Engine/Source/Runtime/RenderCore/Private/ShaderMap.cpp
@@ -277,7 +277,8 @@ static uint16 MakeShaderHash(const FHashedName& TypeName, int32 PermutationId)
 
 FShader* FShaderMapContent::GetShader(const FHashedName& TypeName, int32 PermutationId) const
 {
-	TRACE_CPUPROFILER_EVENT_SCOPE(FShaderMapContent::GetShader);
+	// TRACE_CPUPROFILER_EVENT_SCOPE(FShaderMapContent::GetShader); -- this function is called too frequently, so don't add the scope by default
+
 	const uint16 Hash = MakeShaderHash(TypeName, PermutationId);
 	const FHashedName* RESTRICT LocalShaderTypes = ShaderTypes.GetData();
 	const int32* RESTRICT LocalShaderPermutations = ShaderPermutations.GetData();
