Reduces the alignment of FRDG{Texture,Buffer}Access and allow casting to pointer

diff --git a/Engine/Shaders/Public/ShaderVersion.ush b/Engine/Shaders/Public/ShaderVersion.ush
--- a/Engine/Shaders/Public/ShaderVersion.ush
+++ b/Engine/Shaders/Public/ShaderVersion.ush
@@ -1,6 +1,5 @@
 // This file is automatically generated by the console command r.InvalidateCachedShaders
 // Each time the console command is executed it generates a new GUID. As this file is included
 // in Platform.ush (which should be included in any shader) it allows to invalidate the shader DDC.
-//
-// If you are merging streams and there is a conflict with this GUID you should make a new GUID rather than taking one or the other.
-// GUID = A2DAA4AD-E532-470F-A9C2-6BB2951C7368
+// 
+// GUID = F4FB83CC4C6472BBB80D0EA07AE3A02F
\ No newline at end of file
diff --git a/Engine/Source/Runtime/RenderCore/Public/ShaderParameterMacros.h b/Engine/Source/Runtime/RenderCore/Public/ShaderParameterMacros.h
--- a/Engine/Source/Runtime/RenderCore/Public/ShaderParameterMacros.h
+++ b/Engine/Source/Runtime/RenderCore/Public/ShaderParameterMacros.h
@@ -320,7 +320,7 @@ inline TRDGUniformBufferBinding<TBufferStruct> GetStaticBinding(TRDGUniformBuffe
 	return TRDGUniformBufferBinding<TBufferStruct>(InUniformBuffer, EUniformBufferBindingFlags::Static);
 }
 
-class alignas(SHADER_PARAMETER_STRUCT_ALIGNMENT) FRDGBufferAccess
+class alignas(SHADER_PARAMETER_POINTER_ALIGNMENT) FRDGBufferAccess
 {
 public:
 	FRDGBufferAccess() = default;
@@ -345,13 +345,18 @@ public:
 		return Buffer;
 	}
 
+	FORCEINLINE operator FRDGBuffer*() const
+	{
+		return Buffer;
+	}
+
 private:
-	FRDGBuffer* Buffer = nullptr;
+	TAlignedShaderParameterPtr<FRDGBuffer*> Buffer = nullptr;
 	ERHIAccess  Access = ERHIAccess::Unknown;
 };
 
 template <ERHIAccess InAccess>
-class alignas(SHADER_PARAMETER_STRUCT_ALIGNMENT) TRDGBufferAccess
+class alignas(SHADER_PARAMETER_POINTER_ALIGNMENT) TRDGBufferAccess
 	: public FRDGBufferAccess
 {
 public:
@@ -368,7 +373,7 @@ public:
 	{}
 };
 
-class alignas(SHADER_PARAMETER_STRUCT_ALIGNMENT) FRDGTextureAccess
+class alignas(SHADER_PARAMETER_POINTER_ALIGNMENT) FRDGTextureAccess
 {
 public:
 	FRDGTextureAccess() = default;
@@ -391,13 +396,18 @@ public:
 		return Texture;
 	}
 
+	FORCEINLINE operator FRDGTexture*() const
+	{
+		return Texture;
+	}
+
 private:
-	FRDGTexture* Texture = nullptr;
+	TAlignedShaderParameterPtr<FRDGTexture*> Texture = nullptr;
 	ERHIAccess   Access  = ERHIAccess::Unknown;
 };
 
 template <ERHIAccess InAccess>
-class alignas(SHADER_PARAMETER_STRUCT_ALIGNMENT) TRDGTextureAccess
+class alignas(SHADER_PARAMETER_POINTER_ALIGNMENT) TRDGTextureAccess
 	: public FRDGTextureAccess
 {
 public:
@@ -988,14 +998,14 @@ struct TRDGResourceAccessTypeInfo
 	static constexpr int32 NumRows = 1;
 	static constexpr int32 NumColumns = 1;
 	static constexpr int32 NumElements = 0;
-	static constexpr int32 Alignment = SHADER_PARAMETER_STRUCT_ALIGNMENT;
+	static constexpr int32 Alignment = SHADER_PARAMETER_POINTER_ALIGNMENT;
 	static constexpr bool bIsStoredInConstantBuffer = false;
 
 	using TAlignedType = RDGResourceType;
 
 	static const FShaderParametersMetadata* GetStructMetadata() { return nullptr; }
 
-	static_assert(sizeof(TAlignedType) == SHADER_PARAMETER_STRUCT_ALIGNMENT, "Uniform buffer layout must not be platform dependent.");
+	static_assert(sizeof(TAlignedType) == SHADER_PARAMETER_POINTER_ALIGNMENT * 2, "Uniform buffer layout must not be platform dependent.");
 };
 
 template<typename T, size_t InNumElements>
