Integration of RHI transient allocator into RDG.
 - Replaced legacy transient support from RDG and replaced with new API.
 - Reworked acquire / discard operations a bit and added RHI validation to track correctness.
 - Reworked RDG barrier batching to include acquire / discard operations.
 - Hardened render pass merging logic and expanded to support lifetime extension of transient resources.
 - Added transient tag to RDG insights to track which resources are transient.

diff --git a/Engine/Plugins/RenderGraphInsights/Source/RenderGraphInsights/Private/RenderGraphProvider.cpp b/Engine/Plugins/RenderGraphInsights/Source/RenderGraphInsights/Private/RenderGraphProvider.cpp
--- a/Engine/Plugins/RenderGraphInsights/Source/RenderGraphInsights/Private/RenderGraphProvider.cpp
+++ b/Engine/Plugins/RenderGraphInsights/Source/RenderGraphInsights/Private/RenderGraphProvider.cpp
@@ -65,6 +65,7 @@ FResourcePacket::FResourcePacket(const UE::Trace::IAnalyzer::FOnEventContext& Co
 	, bExternal(Context.EventData.GetValue<bool>("IsExternal"))
 	, bExtracted(Context.EventData.GetValue<bool>("IsExtracted"))
 	, bCulled(Context.EventData.GetValue<bool>("IsCulled"))
+	, bTransient(Context.EventData.GetValue<bool>("IsTransient"))
 {
 	if (Passes.Num())
 	{
diff --git a/Engine/Plugins/RenderGraphInsights/Source/RenderGraphInsights/Private/RenderGraphProvider.h b/Engine/Plugins/RenderGraphInsights/Source/RenderGraphInsights/Private/RenderGraphProvider.h
--- a/Engine/Plugins/RenderGraphInsights/Source/RenderGraphInsights/Private/RenderGraphProvider.h
+++ b/Engine/Plugins/RenderGraphInsights/Source/RenderGraphInsights/Private/RenderGraphProvider.h
@@ -82,6 +82,7 @@ public:
 	bool bExternal{};
 	bool bExtracted{};
 	bool bCulled{};
+	bool bTransient{};
 
 	FResourcePacket(const UE::Trace::IAnalyzer::FOnEventContext& Context);
 };
diff --git a/Engine/Plugins/RenderGraphInsights/Source/RenderGraphInsights/Private/RenderGraphTrack.cpp b/Engine/Plugins/RenderGraphInsights/Source/RenderGraphInsights/Private/RenderGraphTrack.cpp
--- a/Engine/Plugins/RenderGraphInsights/Source/RenderGraphInsights/Private/RenderGraphTrack.cpp
+++ b/Engine/Plugins/RenderGraphInsights/Source/RenderGraphInsights/Private/RenderGraphTrack.cpp
@@ -1397,6 +1397,11 @@ void FRenderGraphTrack::InitTooltip(FTooltipDrawState& Tooltip, const ITimingEve
 			{
 				Tooltip.AddTextLine(TEXT("External"), FLinearColor::Red);
 			}
+
+			if (Resource.bTransient)
+			{
+				Tooltip.AddTextLine(TEXT("Transient"), FLinearColor::Red);
+			}
 		};
 
 		if (InTooltipEvent.Is<FVisibleScopeEvent>())
diff --git a/Engine/Source/Runtime/Apple/MetalRHI/Private/MetalDynamicRHI.cpp b/Engine/Source/Runtime/Apple/MetalRHI/Private/MetalDynamicRHI.cpp
--- a/Engine/Source/Runtime/Apple/MetalRHI/Private/MetalDynamicRHI.cpp
+++ b/Engine/Source/Runtime/Apple/MetalRHI/Private/MetalDynamicRHI.cpp
@@ -110,12 +110,12 @@ void FMetalDynamicRHI::RHIUnlockStagingBuffer(FRHIStagingBuffer* StagingBuffer)
 #pragma mark - Metal Dynamic RHI Resource Transition Methods -
 
 
-void FMetalDynamicRHI::RHICreateTransition(FRHITransition* Transition, ERHIPipeline SrcPipelines, ERHIPipeline DstPipelines, ERHICreateTransitionFlags CreateFlags, TArrayView<const FRHITransitionInfo> Infos)
+void FMetalDynamicRHI::RHICreateTransition(FRHITransition* Transition, const FRHITransitionCreateInfo& CreateInfo)
 {
-	checkf(FMath::IsPowerOfTwo(uint32(SrcPipelines)) && FMath::IsPowerOfTwo(uint32(DstPipelines)), TEXT("Support for multi-pipe resources is not yet implemented."));
+	checkf(FMath::IsPowerOfTwo(uint32(CreateInfo.SrcPipelines)) && FMath::IsPowerOfTwo(uint32(CreateInfo.DstPipelines)), TEXT("Support for multi-pipe resources is not yet implemented."));
 
 	// Construct the data in-place on the transition instance
-	new (Transition->GetPrivateData<FMetalTransitionData>()) FMetalTransitionData(SrcPipelines, DstPipelines, CreateFlags, Infos);
+	new (Transition->GetPrivateData<FMetalTransitionData>()) FMetalTransitionData(CreateInfo.SrcPipelines, CreateInfo.DstPipelines, CreateInfo.Flags, CreateInfo.TransitionInfos);
 }
 
 void FMetalDynamicRHI::RHIReleaseTransition(FRHITransition* Transition)
diff --git a/Engine/Source/Runtime/Apple/MetalRHI/Private/MetalDynamicRHI.h b/Engine/Source/Runtime/Apple/MetalRHI/Private/MetalDynamicRHI.h
--- a/Engine/Source/Runtime/Apple/MetalRHI/Private/MetalDynamicRHI.h
+++ b/Engine/Source/Runtime/Apple/MetalRHI/Private/MetalDynamicRHI.h
@@ -162,7 +162,7 @@ public:
 	virtual FShaderResourceViewRHIRef RHICreateShaderResourceView_RenderThread(class FRHICommandListImmediate& RHICmdList, FRHIBuffer* Buffer) final override;
 	virtual FShaderResourceViewRHIRef RHICreateShaderResourceView_RenderThread(class FRHICommandListImmediate& RHICmdList, FRHITexture* Texture2DRHI, const FRHITextureSRVCreateInfo& CreateInfo) final override;
 
-	virtual void RHICreateTransition(FRHITransition* Transition, ERHIPipeline SrcPipelines, ERHIPipeline DstPipelines, ERHICreateTransitionFlags CreateFlags, TArrayView<const FRHITransitionInfo> Infos) final override;
+	virtual void RHICreateTransition(FRHITransition* Transition, const FRHITransitionCreateInfo& CreateInfo) final override;
 	virtual void RHIReleaseTransition(FRHITransition* Transition) final override;
 
 	virtual FGPUFenceRHIRef RHICreateGPUFence(const FName &Name) final override;
diff --git a/Engine/Source/Runtime/Apple/MetalRHI/Private/MetalTransitionData.cpp b/Engine/Source/Runtime/Apple/MetalRHI/Private/MetalTransitionData.cpp
--- a/Engine/Source/Runtime/Apple/MetalRHI/Private/MetalTransitionData.cpp
+++ b/Engine/Source/Runtime/Apple/MetalRHI/Private/MetalTransitionData.cpp
@@ -22,7 +22,7 @@
 
 FMetalTransitionData::FMetalTransitionData(ERHIPipeline                         InSrcPipelines,
 										   ERHIPipeline                         InDstPipelines,
-										   ERHICreateTransitionFlags            InCreateFlags,
+										   ERHITransitionCreateFlags            InCreateFlags,
 										   TArrayView<const FRHITransitionInfo> InInfos)
 {
 	SrcPipelines   = InSrcPipelines;
@@ -35,7 +35,7 @@ FMetalTransitionData::FMetalTransitionData(ERHIPipeline
 
 	// TODO: Determine whether the Metal RHI needs to create a separate, per-transition fence.
 #if UE_METAL_TRANSITION_DATA_USES_EXPLICIT_FENCING
-	if (bCrossPipeline && !EnumHasAnyFlags(CreateFlags, ERHICreateTransitionFlags::NoFence))
+	if (bCrossPipeline && !EnumHasAnyFlags(CreateFlags, ERHITransitionCreateFlags::NoFence))
 	{
 		// Get the current context pointer.
 		FMetalContext* Context = GetMetalDeviceContext().GetCurrentContext();
diff --git a/Engine/Source/Runtime/Apple/MetalRHI/Private/MetalTransitionData.h b/Engine/Source/Runtime/Apple/MetalRHI/Private/MetalTransitionData.h
--- a/Engine/Source/Runtime/Apple/MetalRHI/Private/MetalTransitionData.h
+++ b/Engine/Source/Runtime/Apple/MetalRHI/Private/MetalTransitionData.h
@@ -25,7 +25,7 @@ class FMetalTransitionData
 public:
 	explicit FMetalTransitionData(ERHIPipeline                         InSrcPipelines,
 								  ERHIPipeline                         InDstPipelines,
-								  ERHICreateTransitionFlags            InCreateFlags,
+								  ERHITransitionCreateFlags            InCreateFlags,
 								  TArrayView<const FRHITransitionInfo> InInfos);
 
 	// The default destructor is sufficient.
@@ -45,7 +45,7 @@ public:
 private:
 	ERHIPipeline              SrcPipelines   = ERHIPipeline::Num;
 	ERHIPipeline              DstPipelines   = ERHIPipeline::Num;
-	ERHICreateTransitionFlags CreateFlags    = ERHICreateTransitionFlags::None;
+	ERHITransitionCreateFlags CreateFlags    = ERHITransitionCreateFlags::None;
 	bool                      bCrossPipeline = false;
 	FMetalTransitionInfoArray Infos          = {};
 	TRefCountPtr<FMetalFence> Fence          = nullptr;
diff --git a/Engine/Source/Runtime/D3D12RHI/Private/D3D12CommandContext.cpp b/Engine/Source/Runtime/D3D12RHI/Private/D3D12CommandContext.cpp
--- a/Engine/Source/Runtime/D3D12RHI/Private/D3D12CommandContext.cpp
+++ b/Engine/Source/Runtime/D3D12RHI/Private/D3D12CommandContext.cpp
@@ -930,27 +930,28 @@ IRHICommandContextContainer* FD3D12DynamicRHI::RHIGetCommandContextContainer(int
 
 #endif // D3D12_SUPPORTS_PARALLEL_RHI_EXECUTE
 
-void FD3D12DynamicRHI::RHICreateTransition(FRHITransition* Transition, ERHIPipeline SrcPipelines, ERHIPipeline DstPipelines, ERHICreateTransitionFlags CreateFlags, TArrayView<const FRHITransitionInfo> Infos)
+void FD3D12DynamicRHI::RHICreateTransition(FRHITransition* Transition, const FRHITransitionCreateInfo& CreateInfo)
 {
 	// Construct the data in-place on the transition instance
 	FD3D12TransitionData* Data = new (Transition->GetPrivateData<FD3D12TransitionData>()) FD3D12TransitionData;
 
-	Data->SrcPipelines = SrcPipelines;
-	Data->DstPipelines = DstPipelines;
-	Data->CreateFlags = CreateFlags;
+	Data->SrcPipelines = CreateInfo.SrcPipelines;
+	Data->DstPipelines = CreateInfo.DstPipelines;
+	Data->CreateFlags = CreateInfo.Flags;
 
-	const bool bCrossPipeline = SrcPipelines != DstPipelines;
+	const bool bCrossPipeline = CreateInfo.SrcPipelines != CreateInfo.DstPipelines;
 
-	if (bCrossPipeline && !EnumHasAnyFlags(Data->CreateFlags, ERHICreateTransitionFlags::NoFence))
+	if (bCrossPipeline && !EnumHasAnyFlags(Data->CreateFlags, ERHITransitionCreateFlags::NoFence))
 	{
-		const FName Name = SrcPipelines == ERHIPipeline::Graphics ? L"<Graphics To AsyncCompute>" : L"<AsyncCompute To Graphics>";
+		const FName Name = CreateInfo.SrcPipelines == ERHIPipeline::Graphics ? L"<Graphics To AsyncCompute>" : L"<AsyncCompute To Graphics>";
 
 		Data->Fence = new FD3D12Fence(&GetAdapter(), FRHIGPUMask::All(), Name);
 		Data->Fence->CreateFence();
 	}
 
 	Data->bCrossPipeline = bCrossPipeline;
-	Data->Infos.Append(Infos.GetData(), Infos.Num());
+	Data->TransitionInfos = CreateInfo.TransitionInfos;
+	Data->AliasingInfos = CreateInfo.AliasingInfos;
 }
 
 void FD3D12DynamicRHI::RHIReleaseTransition(FRHITransition* Transition)
diff --git a/Engine/Source/Runtime/D3D12RHI/Private/D3D12CommandContext.h b/Engine/Source/Runtime/D3D12RHI/Private/D3D12CommandContext.h
--- a/Engine/Source/Runtime/D3D12RHI/Private/D3D12CommandContext.h
+++ b/Engine/Source/Runtime/D3D12RHI/Private/D3D12CommandContext.h
@@ -878,9 +878,10 @@ private:
 struct FD3D12TransitionData
 {
 	ERHIPipeline SrcPipelines, DstPipelines;
-	ERHICreateTransitionFlags CreateFlags = ERHICreateTransitionFlags::None;
+	ERHITransitionCreateFlags CreateFlags = ERHITransitionCreateFlags::None;
 
-	TArray<FRHITransitionInfo, TInlineAllocator<4>> Infos;
+	TArray<FRHITransitionInfo, TInlineAllocator<4>> TransitionInfos;
+	TArray<FRHITransientAliasingInfo, TInlineAllocator<4>> AliasingInfos;
 	TRefCountPtr<FD3D12Fence> Fence;
 
 	bool bCrossPipeline = false;
diff --git a/Engine/Source/Runtime/D3D12RHI/Private/D3D12Commands.cpp b/Engine/Source/Runtime/D3D12RHI/Private/D3D12Commands.cpp
--- a/Engine/Source/Runtime/D3D12RHI/Private/D3D12Commands.cpp
+++ b/Engine/Source/Runtime/D3D12RHI/Private/D3D12Commands.cpp
@@ -290,85 +290,196 @@ void ProcessResource(FD3D12CommandContext& Context, const FRHITransitionInfo& In
 	}
 }
 
-void ProcessDiscardState(FD3D12CommandContext& Context, const FRHITransitionInfo& Info)
+static bool ProcessTransitionDuringBegin(const FD3D12TransitionData* Data)
+{
+	// Pipe changes which are not ending with graphics or targeting all pipelines are handle during begin
+	return ((Data->SrcPipelines != Data->DstPipelines && Data->DstPipelines != ERHIPipeline::Graphics) || EnumHasAllFlags(Data->DstPipelines, ERHIPipeline::All));
+}
+
+struct FD3D12DiscardResource
+{
+	FD3D12DiscardResource(FD3D12Resource* InResource, EResourceTransitionFlags InFlags, uint32 InSubresource = D3D12_RESOURCE_BARRIER_ALL_SUBRESOURCES)
+		: Resource(InResource)
+		, Flags(InFlags)
+		, Subresource(InSubresource)
+	{}
+
+	FD3D12Resource* Resource;
+	EResourceTransitionFlags Flags;
+	uint32 Subresource;
+};
+
+using FD3D12DiscardResourceArray = TArray<FD3D12DiscardResource>;
+
+static void HandleResourceDiscardTransitions(
+	FD3D12CommandContext& Context,
+	const FD3D12TransitionData* TransitionData,
+	D3D12_RESOURCE_STATES SkipFastClearEliminateState,
+	FD3D12DiscardResourceArray& ResourcesToDiscard)
 {
-	// No discard before or after state, then done
-	if (!EnumHasAnyFlags(Info.AccessBefore, ERHIAccess::Discard) && !EnumHasAnyFlags(Info.AccessAfter, ERHIAccess::Discard))
+	for (const FRHITransitionInfo& Info : TransitionData->TransitionInfos)
 	{
-		return;
+		if (Info.AccessBefore != ERHIAccess::Discard || Info.Type != FRHITransitionInfo::EType::Texture || !EnumHasAnyFlags(Info.Flags, EResourceTransitionFlags::Clear | EResourceTransitionFlags::Discard))
+		{
+			continue;
+		}
+
+		ProcessResource(Context, Info, [&](const FRHITransitionInfo& Info, FD3D12Resource* Resource)
+		{
+			if (!Resource->RequiresResourceStateTracking())
+			{
+				return;
+			}
+
+			const D3D12_RESOURCE_FLAGS ResourceFlags = Resource->GetDesc().Flags;
+
+			D3D12_RESOURCE_STATES DiscardState;
+
+			if (EnumHasAnyFlags(ResourceFlags, D3D12_RESOURCE_FLAG_ALLOW_RENDER_TARGET))
+			{
+				DiscardState = D3D12_RESOURCE_STATE_RENDER_TARGET;
+			}
+			else if (EnumHasAnyFlags(ResourceFlags, D3D12_RESOURCE_FLAG_ALLOW_DEPTH_STENCIL))
+			{
+				DiscardState = D3D12_RESOURCE_STATE_DEPTH_WRITE;
+			}
+			else if (EnumHasAnyFlags(ResourceFlags, D3D12_RESOURCE_FLAG_ALLOW_UNORDERED_ACCESS))
+			{
+				DiscardState = D3D12_RESOURCE_STATE_UNORDERED_ACCESS;
+			}
+			else
+			{
+				return;
+			}
+
+			if (EnumHasAnyFlags(Info.Flags, EResourceTransitionFlags::MaintainCompression))
+			{
+				DiscardState |= SkipFastClearEliminateState;
+			}
+
+			if (Info.IsWholeResource() || Resource->GetSubresourceCount() == 1)
+			{
+				if (FD3D12DynamicRHI::TransitionResource(Context.CommandListHandle, Resource, D3D12_RESOURCE_STATE_TBD, DiscardState, D3D12_RESOURCE_BARRIER_ALL_SUBRESOURCES, FD3D12DynamicRHI::ETransitionMode::Apply))
+				{
+					ResourcesToDiscard.Emplace(Resource, Info.Flags, D3D12_RESOURCE_BARRIER_ALL_SUBRESOURCES);
+				}
+			}
+			else
+			{
+				EnumerateSubresources(Resource, Info, [&](uint32 Subresource)
+				{
+					if (FD3D12DynamicRHI::TransitionResource(Context.CommandListHandle, Resource, D3D12_RESOURCE_STATE_TBD, DiscardState, Subresource, FD3D12DynamicRHI::ETransitionMode::Apply))
+					{
+						ResourcesToDiscard.Emplace(Resource, Info.Flags, Subresource);
+					}
+				});
+			}
+		});
 	}
+}
 
-	FD3D12BaseShaderResource* BaseShaderResource = nullptr;
-	switch (Info.Type)
-	{
-	case FRHITransitionInfo::EType::Buffer:
+static void HandleDiscardResources(
+	FD3D12CommandContext& Context,
+	D3D12_RESOURCE_STATES SkipFastClearEliminateState,
+	TArrayView<const FRHITransition*> Transitions)
+{
+	FD3D12DiscardResourceArray ResourcesToDiscard;
+
+	for (const FRHITransition* Transition : Transitions)
 	{
-		FD3D12Buffer* Buffer = Context.RetrieveObject<FD3D12Buffer>(Info.Buffer);
-		check(Buffer);
-		BaseShaderResource = Buffer;
-		break;
+		const FD3D12TransitionData* Data = Transition->GetPrivateData<FD3D12TransitionData>();
+
+		HandleResourceDiscardTransitions(Context, Data, SkipFastClearEliminateState, ResourcesToDiscard);
 	}
-	case FRHITransitionInfo::EType::Texture:
+
+	Context.CommandListHandle.FlushResourceBarriers();
+
+	for (const FD3D12DiscardResource& DiscardResource : ResourcesToDiscard)
 	{
-		FD3D12TextureBase* Texture = Context.RetrieveTextureBase(Info.Texture);
-		check(Texture);
-		BaseShaderResource = Texture;
-		break;
-	}
-	default:
-		checkNoEntry();
-		break;
+		if (EnumHasAnyFlags(DiscardResource.Flags, EResourceTransitionFlags::Clear))
+		{
+			// add correct ops for clear
+			check(false);
+		}
+		else if (EnumHasAnyFlags(DiscardResource.Flags, EResourceTransitionFlags::Discard))
+		{
+			if (DiscardResource.Subresource == D3D12_RESOURCE_BARRIER_ALL_SUBRESOURCES)
+			{
+				Context.CommandListHandle->DiscardResource(DiscardResource.Resource->GetResource(), nullptr);
+			}
+			else
+			{
+				D3D12_DISCARD_REGION Region;
+				Region.NumRects = 0;
+				Region.pRects = nullptr;
+				Region.FirstSubresource = DiscardResource.Subresource;
+				Region.NumSubresources = 1;
+
+				Context.CommandListHandle->DiscardResource(DiscardResource.Resource->GetResource(), &Region);
+			}
+		}
 	}
+}
 
-	if (EnumHasAnyFlags(Info.AccessBefore, ERHIAccess::Discard))
+static void HandleTransientAliasing(FD3D12CommandContext& Context, const FD3D12TransitionData* TransitionData)
+{
+	for (const FRHITransientAliasingInfo& Info : TransitionData->AliasingInfos)
 	{
-		TRACE_CPUPROFILER_EVENT_SCOPE(D3D12RHI::AcquireTransient);
-
-		FD3D12Resource* Resource = BaseShaderResource->ResourceLocation.GetResource();
+		FD3D12BaseShaderResource* BaseShaderResource = nullptr;
+		switch (Info.Type)
+		{
+		case FRHITransientAliasingInfo::EType::Buffer:
+		{
+			FD3D12Buffer* Buffer = Context.RetrieveObject<FD3D12Buffer>(Info.Buffer);
+			check(Buffer);
+			BaseShaderResource = Buffer;
+			break;
+		}
+		case FRHITransientAliasingInfo::EType::Texture:
+		{
+			FD3D12TextureBase* Texture = Context.RetrieveTextureBase(Info.Texture);
+			check(Texture);
+			BaseShaderResource = Texture;
+			break;
+		}
+		default:
+			checkNoEntry();
+			break;
+		}
 
-		// Try and see if the resource has overlapping resources
-		if (BaseShaderResource->ResourceLocation.IsTransient() && BaseShaderResource->ResourceLocation.GetAllocatorType() == FD3D12ResourceLocation::AT_Pool)
+		if (Info.Action == FRHITransientAliasingInfo::EAction::Acquire)
 		{
-			// We know it's a transient resource allocator
-			FD3D12TransientResourceAllocator* Allocator = (FD3D12TransientResourceAllocator*) BaseShaderResource->ResourceLocation.GetPoolAllocator();
+			TRACE_CPUPROFILER_EVENT_SCOPE(D3D12RHI::AcquireTransient);
+
+			FD3D12Resource* Resource = BaseShaderResource->ResourceLocation.GetResource();
 
-			// Get the overlapping resources and append aliasing barriers for each of them
-			TArrayView<FD3D12Resource*> OverlappingResources = Allocator->GetOverlappingResources(BaseShaderResource);
-			for (FD3D12Resource* OverlappingResource : OverlappingResources)
+			// Try and see if the resource has overlapping resources
+			if (BaseShaderResource->ResourceLocation.IsTransient() && BaseShaderResource->ResourceLocation.GetAllocatorType() == FD3D12ResourceLocation::AT_Pool)
 			{
-				Context.CommandListHandle.AddAliasingBarrier(OverlappingResource, Resource);
+				// We know it's a transient resource allocator
+				FD3D12TransientResourceAllocator* Allocator = (FD3D12TransientResourceAllocator*)BaseShaderResource->ResourceLocation.GetPoolAllocator();
+
+				// Get the overlapping resources and append aliasing barriers for each of them
+				TArrayView<FD3D12Resource*> OverlappingResources = Allocator->GetOverlappingResources(BaseShaderResource);
+				for (FD3D12Resource* OverlappingResource : OverlappingResources)
+				{
+					Context.CommandListHandle.AddAliasingBarrier(OverlappingResource, Resource);
+				}
 			}
-		}		
-
-		bool bIsAsyncCompute = false;
-		D3D12_RESOURCE_STATES AfterState = GetD3D12ResourceState(Info.AccessAfter, bIsAsyncCompute);
-		FD3D12DynamicRHI::TransitionResource(Context.CommandListHandle, Resource, D3D12_RESOURCE_STATE_TBD, AfterState, D3D12_RESOURCE_BARRIER_ALL_SUBRESOURCES, FD3D12DynamicRHI::ETransitionMode::Apply);
-		Context.CommandListHandle.FlushResourceBarriers();
-		
-		// Add clear or discard operation depending on flag
-		if (EnumHasAnyFlags(Info.Flags, EResourceTransitionFlags::Clear))
-		{
-			// add correct ops for clear
-			check(false);
 		}
 		else
 		{
-			check(Info.Flags & EResourceTransitionFlags::Discard);
-			Context.CommandListHandle->DiscardResource(Resource->GetResource(), nullptr);
-		}
-	}
-	else
-	{
-		TRACE_CPUPROFILER_EVENT_SCOPE(D3D12RHI::DiscardTransient);
+			TRACE_CPUPROFILER_EVENT_SCOPE(D3D12RHI::DiscardTransient);
 
-		// Remove from caches
-		Context.ConditionalClearShaderResource(&BaseShaderResource->ResourceLocation);
+			// Remove from caches
+			Context.ConditionalClearShaderResource(&BaseShaderResource->ResourceLocation);
 
-		// Release the resource location - transient allocator internally will take care of correct caching
-		BaseShaderResource->ResourceLocation.Clear();
+			// Release the resource location - transient allocator internally will take care of correct caching
+			BaseShaderResource->ResourceLocation.Clear();
 
-		// Release the views after clearing resource location
-		BaseShaderResource->ResourceRenamed(&BaseShaderResource->ResourceLocation);
+			// Release the views after clearing resource location
+			BaseShaderResource->ResourceRenamed(&BaseShaderResource->ResourceLocation);
+		}
 	}
 }
 
@@ -376,10 +487,13 @@ static void HandleResourceTransitions(FD3D12CommandContext& Context, const FD3D1
 {
 	const bool bDstAllPipelines = EnumHasAllFlags(TransitionData->DstPipelines, ERHIPipeline::All);
 
-	for (const FRHITransitionInfo& Info : TransitionData->Infos)
+	for (const FRHITransitionInfo& Info : TransitionData->TransitionInfos)
 	{
-		// Process the discard state
-		ProcessDiscardState(Context, Info);
+		// A transition back to discard is a no-op.
+		if (Info.AccessAfter == ERHIAccess::Discard)
+		{
+			continue;
+		}
 
 		const bool bUAVAccess = EnumHasAnyFlags(Info.AccessAfter, ERHIAccess::UAVMask);
 
@@ -406,61 +520,64 @@ static void HandleResourceTransitions(FD3D12CommandContext& Context, const FD3D1
 			else
 			{
 				ProcessResource(Context, Info, [&](const FRHITransitionInfo& Info, FD3D12Resource* Resource)
+				{
+					if (!Resource->RequiresResourceStateTracking())
 					{
-						if (!Resource->RequiresResourceStateTracking())
-						{
-							return;
-						}
+						return;
+					}
+
+					const bool bIsAsyncCompute = EnumHasAnyFlags(TransitionData->DstPipelines, ERHIPipeline::AsyncCompute);
+					const bool bKnownBeforeState = (Info.AccessBefore != ERHIAccess::Unknown && Info.AccessBefore != ERHIAccess::Discard);
 
-						const bool bIsAsyncCompute = EnumHasAnyFlags(TransitionData->DstPipelines, ERHIPipeline::AsyncCompute);
+					// Derive or get the correct before & after state
+					D3D12_RESOURCE_STATES BeforeState = bKnownBeforeState ? GetD3D12ResourceState(Info.AccessBefore, bIsAsyncCompute) : D3D12_RESOURCE_STATE_TBD;
+					D3D12_RESOURCE_STATES AfterState = GetD3D12ResourceState(Info.AccessAfter, bIsAsyncCompute);
 
-						// Derive or get the correct before & after state
-						D3D12_RESOURCE_STATES BeforeState = (Info.AccessBefore != ERHIAccess::Unknown) ? GetD3D12ResourceState(Info.AccessBefore, bIsAsyncCompute) : D3D12_RESOURCE_STATE_TBD;
-						D3D12_RESOURCE_STATES AfterState = GetD3D12ResourceState(Info.AccessAfter, bIsAsyncCompute);
+					// Add the compression flags if needed
+					if (EnumHasAnyFlags(Info.Flags, EResourceTransitionFlags::MaintainCompression))
+					{
+						AfterState |= SkipFastClearEliminateState;
+					}
 
-						// Add the compression flags if needed
-						if (EnumHasAnyFlags(Info.Flags, EResourceTransitionFlags::MaintainCompression))
+					// enqueue the correct transitions
+					if (Info.IsWholeResource() || Resource->GetSubresourceCount() == 1)
+					{
+						if (FD3D12DynamicRHI::TransitionResource(Context.CommandListHandle, Resource, BeforeState, AfterState, D3D12_RESOURCE_BARRIER_ALL_SUBRESOURCES, FD3D12DynamicRHI::ETransitionMode::Apply))
 						{
-							AfterState |= SkipFastClearEliminateState;
+							bUAVBarrier = true;
 						}
-
-						// enqueue the correct transitions
-						if (Info.IsWholeResource() || Resource->GetSubresourceCount() == 1)
+					}
+					else
+					{
+						EnumerateSubresources(Resource, Info, [&](uint32 Subresource)
 						{
-							if (FD3D12DynamicRHI::TransitionResource(Context.CommandListHandle, Resource, BeforeState, AfterState, D3D12_RESOURCE_BARRIER_ALL_SUBRESOURCES, FD3D12DynamicRHI::ETransitionMode::Apply))
+							if (FD3D12DynamicRHI::TransitionResource(Context.CommandListHandle, Resource, BeforeState, AfterState, Subresource, FD3D12DynamicRHI::ETransitionMode::Apply))
 							{
 								bUAVBarrier = true;
 							}
-
-						}
-						else
-						{
-							EnumerateSubresources(Resource, Info, [&](uint32 Subresource)
-								{
-									if (FD3D12DynamicRHI::TransitionResource(Context.CommandListHandle, Resource, BeforeState, AfterState, Subresource, FD3D12DynamicRHI::ETransitionMode::Apply))
-									{
-										bUAVBarrier = true;
-									}
-								});
-						}
-					});
+						});
+					}
+				});
 			}
 		}
 	}
 }
 
-static bool ProcessTransitionDuringBegin(const FD3D12TransitionData* Data)
-{
-	// Pipe changes which are not ending with graphics or targeting all pipelines are handle during begin
-	return ((Data->SrcPipelines != Data->DstPipelines && Data->DstPipelines != ERHIPipeline::Graphics) || EnumHasAllFlags(Data->DstPipelines, ERHIPipeline::All));
-}
-
 void FD3D12CommandContext::RHIBeginTransitionsWithoutFencing(TArrayView<const FRHITransition*> Transitions)
 {
 	static IConsoleVariable* CVarShowTransitions = IConsoleManager::Get().FindConsoleVariable(TEXT("r.ProfileGPU.ShowTransitions"));
 	const bool bShowTransitionEvents = CVarShowTransitions->GetInt() != 0;
 	SCOPED_RHI_CONDITIONAL_DRAW_EVENTF(*this, RHIBeginTransitions, bShowTransitionEvents, TEXT("RHIBeginTransitions"));
 
+	for (const FRHITransition* Transition : Transitions)
+	{
+		const FD3D12TransitionData* Data = Transition->GetPrivateData<FD3D12TransitionData>();
+
+		HandleTransientAliasing(*this, Data);
+	}
+
+	HandleDiscardResources(*this, SkipFastClearEliminateState, Transitions);
+
 	bool bUAVBarrier = false;
 
 	for (const FRHITransition* Transition : Transitions)
diff --git a/Engine/Source/Runtime/D3D12RHI/Private/D3D12RHIPrivate.h b/Engine/Source/Runtime/D3D12RHI/Private/D3D12RHIPrivate.h
--- a/Engine/Source/Runtime/D3D12RHI/Private/D3D12RHIPrivate.h
+++ b/Engine/Source/Runtime/D3D12RHI/Private/D3D12RHIPrivate.h
@@ -281,7 +281,7 @@ public:
 	virtual FBoundShaderStateRHIRef RHICreateBoundShaderState(FRHIVertexDeclaration* VertexDeclaration, FRHIVertexShader* VertexShader, FRHIHullShader* HullShader, FRHIDomainShader* DomainShader, FRHIPixelShader* PixelShader, FRHIGeometryShader* GeometryShader) final override;
 	virtual FGraphicsPipelineStateRHIRef RHICreateGraphicsPipelineState(const FGraphicsPipelineStateInitializer& Initializer) final override;
 	virtual TRefCountPtr<FRHIComputePipelineState> RHICreateComputePipelineState(FRHIComputeShader* ComputeShader) final override;
-	virtual void RHICreateTransition(FRHITransition* Transition, ERHIPipeline SrcPipelines, ERHIPipeline DstPipelines, ERHICreateTransitionFlags CreateFlags, TArrayView<const FRHITransitionInfo> Infos) final override;
+	virtual void RHICreateTransition(FRHITransition* Transition, const FRHITransitionCreateInfo& CreateInfo) final override;
 	virtual void RHIReleaseTransition(FRHITransition* Transition) final override;
 	virtual FUniformBufferRHIRef RHICreateUniformBuffer(const void* Contents, const FRHIUniformBufferLayout& Layout, EUniformBufferUsage Usage, EUniformBufferValidation Validation) final override;
 	virtual void RHIUpdateUniformBuffer(FRHIUniformBuffer* UniformBufferRHI, const void* Contents) final override;
diff --git a/Engine/Source/Runtime/D3D12RHI/Private/D3D12TransientResourceAllocator.cpp b/Engine/Source/Runtime/D3D12RHI/Private/D3D12TransientResourceAllocator.cpp
--- a/Engine/Source/Runtime/D3D12RHI/Private/D3D12TransientResourceAllocator.cpp
+++ b/Engine/Source/Runtime/D3D12RHI/Private/D3D12TransientResourceAllocator.cpp
@@ -484,9 +484,6 @@ FRHITexture* FD3D12TransientResourceAllocator::CreateTexture(const FRHITextureCr
 {
 	TRACE_CPUPROFILER_EVENT_SCOPE(D3D12RHI::CreateTransientTexture);
 
-	// Should be marked transient
-	check((InCreateInfo.Flags & TexCreate_Transient) == TexCreate_Transient);
-
 	// Try and get a pooled texture from the manager
 	FD3D12TransientMemoryPoolManager& PoolManager = GetParentDevice()->GetTransientMemoryPoolManager();
 	FD3D12PooledTextureData TextureData = PoolManager.GetPooledTexture(InCreateInfo, InDebugName);
@@ -551,7 +548,7 @@ FRHITexture* FD3D12TransientResourceAllocator::CreateTexture(const FRHITextureCr
 			TextureData.ResourceDesc = Texture3D->GetResource()->GetDesc();
 			TextureData.ClearValue = Texture3D->GetResource()->GetClearValue();
 			TextureData.CreateInfo = InCreateInfo;
-			
+
 			BaseShaderResource = Texture3D;
 			break;
 		}
@@ -563,7 +560,7 @@ FRHITexture* FD3D12TransientResourceAllocator::CreateTexture(const FRHITextureCr
 		}
 		}
 	}
-	
+
 	SetupAllocatedResource(BaseShaderResource, TextureMemoryStats);
 
 	// keep track of all allocated textures
@@ -577,9 +574,6 @@ FRHIBuffer* FD3D12TransientResourceAllocator::CreateBuffer(const FRHIBufferCreat
 {
 	TRACE_CPUPROFILER_EVENT_SCOPE(D3D12RHI::CreateTransientBuffer);
 
-	// Should be marked transient
-	check((InCreateInfo.Usage & BUF_Transient) == BUF_Transient);
-
 	// Try and get a pooled texture from the manager
 	FD3D12TransientMemoryPoolManager& PoolManager = GetParentDevice()->GetTransientMemoryPoolManager();
 	FD3D12PooledBufferData BufferData = PoolManager.GetPooledBuffer(InCreateInfo, InDebugName);
diff --git a/Engine/Source/Runtime/NullDrv/Public/NullRHI.h b/Engine/Source/Runtime/NullDrv/Public/NullRHI.h
--- a/Engine/Source/Runtime/NullDrv/Public/NullRHI.h
+++ b/Engine/Source/Runtime/NullDrv/Public/NullRHI.h
@@ -93,7 +93,7 @@ public:
 
 	}
 
-	virtual void RHICreateTransition(FRHITransition* Transition, ERHIPipeline SrcPipelines, ERHIPipeline DstPipelines, ERHICreateTransitionFlags CreateFlags, TArrayView<const FRHITransitionInfo> Infos) final override
+	virtual void RHICreateTransition(FRHITransition* Transition, const FRHITransitionCreateInfo& CreateInfo) final override
 	{
 	}
 
diff --git a/Engine/Source/Runtime/RHI/Private/RHIValidation.cpp b/Engine/Source/Runtime/RHI/Private/RHIValidation.cpp
--- a/Engine/Source/Runtime/RHI/Private/RHIValidation.cpp
+++ b/Engine/Source/Runtime/RHI/Private/RHIValidation.cpp
@@ -180,10 +180,13 @@ void FValidationRHI::ValidatePipeline(const FGraphicsPipelineStateInitializer& P
 	}
 }
 
-void FValidationRHI::RHICreateTransition(FRHITransition* Transition, ERHIPipeline SrcPipelines, ERHIPipeline DstPipelines, ERHICreateTransitionFlags CreateFlags, TArrayView<const FRHITransitionInfo> Infos)
+void FValidationRHI::RHICreateTransition(FRHITransition* Transition, const FRHITransitionCreateInfo& CreateInfo)
 {
 	using namespace RHIValidation;
 
+	const ERHIPipeline SrcPipelines = CreateInfo.SrcPipelines;
+	const ERHIPipeline DstPipelines = CreateInfo.DstPipelines;
+
 	struct FFenceEdge
 	{
 		FFence* Fence = nullptr;
@@ -218,20 +221,61 @@ void FValidationRHI::RHICreateTransition(FRHITransition* Transition, ERHIPipelin
 		}
 	}
 
-	TArray<FOperation> BeginOps, EndOps;
-	BeginOps.Reserve(Infos.Num() + Fences.Num());
-	EndOps  .Reserve(Infos.Num() + Fences.Num());
+	TArray<FOperation> SignalOps, WaitOps, AcquireOps, DiscardOps, BeginOps, EndOps;
+	SignalOps .Reserve(Fences.Num());
+	WaitOps   .Reserve(Fences.Num());
+	AcquireOps.Reserve(CreateInfo.AliasingInfos.Num());
+	DiscardOps.Reserve(CreateInfo.AliasingInfos.Num());
+	BeginOps  .Reserve(CreateInfo.TransitionInfos.Num());
+	EndOps    .Reserve(CreateInfo.TransitionInfos.Num());
 
 	for (const FFenceEdge& FenceEdge : Fences)
 	{
-		EndOps.Emplace(FOperation::Wait(FenceEdge.Fence, FenceEdge.DstPipe));
+		WaitOps.Emplace(FOperation::Wait(FenceEdge.Fence, FenceEdge.DstPipe));
 	}
 
+	// Take a backtrace of this transition creation if any of the resources it contains have logging enabled.
 	bool bDoTrace = false;
 
-	for (int32 Index = 0; Index < Infos.Num(); ++Index)
+	for (const FRHITransientAliasingInfo& Info : CreateInfo.AliasingInfos)
+	{
+		if (!Info.Resource)
+		{
+			continue;
+		}
+
+		FResource* Resource = nullptr;
+
+		if (Info.Type == FRHITransientAliasingInfo::EType::Texture)
+		{
+			Resource = Info.Texture->GetTrackerResource();
+		}
+		else
+		{
+			Resource = Info.Buffer;
+		}
+
+		bDoTrace |= (Resource->LoggingMode != RHIValidation::ELoggingMode::None);
+
+		if (Info.IsAcquire())
+		{
+			checkf(Resource->TransientState.bTransient, TEXT("Acquiring resource %s which is not transient. Only transient resources can be acquired."), Resource->GetDebugName());
+			checkf(SrcPipelines == ERHIPipeline::Graphics, TEXT("Acquiring a transient resource (%s) must begin on the graphics pipe."), Resource->GetDebugName());
+
+			AcquireOps.Emplace(FOperation::AcquireTransientResource(Resource, nullptr));
+		}
+		else
+		{
+			checkf(Resource->TransientState.bTransient, TEXT("Discarding resource %s which is not transient. Only transient resources can be discarded."), Resource->GetDebugName());
+			checkf(DstPipelines == ERHIPipeline::Graphics, TEXT("Discarding a transient resource (%s) must end on the graphics pipe."), Resource->GetDebugName());
+
+			DiscardOps.Emplace(FOperation::DiscardTransientResource(Resource, nullptr));
+		}
+	}
+
+	for (int32 Index = 0; Index < CreateInfo.TransitionInfos.Num(); ++Index)
 	{
-		const FRHITransitionInfo& Info = Infos[Index];
+		const FRHITransitionInfo& Info = CreateInfo.TransitionInfos[Index];
 		if (!Info.Resource)
 			continue;
 
@@ -255,7 +299,6 @@ void FValidationRHI::RHICreateTransition(FRHITransition* Transition, ERHIPipelin
 			break;
 		}
 
-		// Take a backtrace of this transition creation if any of the resources it contains have logging enabled.
 		bDoTrace |= (Identity.Resource->LoggingMode != RHIValidation::ELoggingMode::None);
 
 		FState PreviousState = FState(Info.AccessBefore, SrcPipelines);
@@ -268,19 +311,25 @@ void FValidationRHI::RHICreateTransition(FRHITransition* Transition, ERHIPipelin
 	if (bDoTrace)
 	{
 		void* Backtrace = CaptureBacktrace();
+		for (FOperation& Op : AcquireOps) { Op.Data_AcquireTransient.CreateBacktrace = Backtrace; }
+		for (FOperation& Op : DiscardOps) { Op.Data_DiscardTransient.CreateBacktrace = Backtrace; }
 		for (FOperation& Op : BeginOps) { Op.Data_BeginTransition.CreateBacktrace = Backtrace; }
 		for (FOperation& Op : EndOps) { Op.Data_EndTransition.CreateBacktrace = Backtrace; }
 	}
 
 	for (const FFenceEdge& FenceEdge : Fences)
 	{
-		BeginOps.Emplace(FOperation::Signal(FenceEdge.Fence, FenceEdge.SrcPipe));
+		SignalOps.Emplace(FOperation::Signal(FenceEdge.Fence, FenceEdge.SrcPipe));
 	}
 
+	Transition->PendingSignals.Operations = MoveTemp(SignalOps);
+	Transition->PendingWaits.Operations = MoveTemp(WaitOps);
+	Transition->PendingAcquires.Operations = MoveTemp(AcquireOps);
+	Transition->PendingDiscards.Operations = MoveTemp(DiscardOps);
 	Transition->PendingOperationsBegin.Operations = MoveTemp(BeginOps);
 	Transition->PendingOperationsEnd.Operations = MoveTemp(EndOps);
 
-	return RHI->RHICreateTransition(Transition, SrcPipelines, DstPipelines, CreateFlags, Infos);
+	return RHI->RHICreateTransition(Transition, CreateInfo);
 }
 
 void FValidationRHI::ReportValidationFailure(const TCHAR* InMessage)
@@ -449,6 +498,10 @@ namespace RHIValidation
 	TEXT("--------------------------------------------------------------------\n")\
 	TEXT("\n\n")
 
+#define BARRIER_TRACKER_LOG_ENABLE_TRANSITION_BACKTRACE \
+	TEXT("    --- Enable barrier logging for this resource to see a callstack backtrace for the RHIBeginTransitions() call ") \
+	TEXT("which has not been completed. Use -RHIValidationLog=X,Y,Z to enable backtrace logging for individual resources.\n\n")
+
 	static inline FString GetResourceDebugName(FResource const* Resource, FSubresourceIndex const& SubresourceIndex)
 	{
 		const TCHAR* DebugName = Resource->GetDebugName();
@@ -511,7 +564,38 @@ namespace RHIValidation
 		}
 		else
 		{
-			return TEXT("    --- Enable barrier logging for this resource to see a callstack backtrace for the RHIBeginTransitions() call which has not been completed. See RHIValidation.cpp.\n\n");
+			return BARRIER_TRACKER_LOG_ENABLE_TRANSITION_BACKTRACE;
+		}
+	}
+
+	static inline FString GetReasonString_Backtrace(const TCHAR* OperationPrefix, const TCHAR* TracePrefix, void* Trace)
+	{
+		if (Trace)
+		{
+			return FString::Printf(
+				TEXT("    --- Callstack backtrace for %s operation (resolve in the Watch window):\n")
+				TEXT("        %s: (void**)0x%p,32\n"),
+				OperationPrefix,
+				TracePrefix,
+				Trace);
+		}
+		else
+		{
+			return FString(BARRIER_TRACKER_LOG_ENABLE_TRANSITION_BACKTRACE);
+		}
+	}
+
+	static inline FString GetReasonString_DuplicateBackTrace(void* PreviousTrace, void* CurrentTrace)
+	{
+		if (PreviousTrace || CurrentTrace)
+		{
+			return
+				GetReasonString_Backtrace(TEXT("previous"), TEXT("RHICreateTransition"), PreviousTrace) +
+				GetReasonString_Backtrace(TEXT("current"), TEXT("RHICreateTransition"), CurrentTrace);
+		}
+		else
+		{
+			return BARRIER_TRACKER_LOG_ENABLE_TRANSITION_BACKTRACE;
 		}
 	}
 
@@ -540,6 +624,78 @@ namespace RHIValidation
 			*GetReasonString_BeginBacktrace(CreateTrace, BeginTrace));
 	}
 
+	static inline FString GetReasonString_TransitionWithoutAcquire(FResource* Resource)
+	{
+		FString DebugName = GetResourceDebugName(Resource, {});
+		return FString::Printf(
+			BARRIER_TRACKER_LOG_PREFIX_RESNAME
+			TEXT("Attempted a resource transition for transient resource %s without acquiring it. Transient resources must be acquired before any transitions are begun and discarded after all transitions are complete.\n")
+			BARRIER_TRACKER_LOG_SUFFIX,
+			*DebugName,
+			*DebugName);
+	}
+
+	static inline FString GetReasonString_AcquireNonTransient(FResource* Resource)
+	{
+		FString DebugName = GetResourceDebugName(Resource, {});
+		return FString::Printf(
+			BARRIER_TRACKER_LOG_PREFIX_RESNAME
+			TEXT("Attempted to acquire non-transient resource %s. Only transient resources may be acquired with the transient aliasing API.\n")
+			BARRIER_TRACKER_LOG_SUFFIX,
+			*DebugName,
+			*DebugName);
+	}
+
+	static inline FString GetReasonString_DiscardNonTransient(FResource* Resource)
+	{
+		FString DebugName = GetResourceDebugName(Resource, {});
+		return FString::Printf(
+			BARRIER_TRACKER_LOG_PREFIX_RESNAME
+			TEXT("Attempted to discard non-transient resource %s. Only transient resources may be discarded with the transient aliasing API.\n")
+			BARRIER_TRACKER_LOG_SUFFIX,
+			*DebugName,
+			*DebugName);
+	}
+
+	static inline FString GetReasonString_DuplicateAcquireTransient(FResource* Resource, void* PreviousAcquireTrace, void* CurrentAcquireTrace)
+	{
+		FString DebugName = GetResourceDebugName(Resource, {});
+		return FString::Printf(
+			BARRIER_TRACKER_LOG_PREFIX_RESNAME
+			TEXT("Mismatched acquire of transient resource %s. A transient resource may only be acquired once in its lifetime.\n")
+			TEXT("%s")
+			BARRIER_TRACKER_LOG_SUFFIX,
+			*DebugName,
+			*DebugName,
+			*GetReasonString_DuplicateBackTrace(PreviousAcquireTrace, CurrentAcquireTrace));
+	}
+
+	static inline FString GetReasonString_DuplicateDiscardTransient(FResource* Resource, void* PreviousDiscardTrace, void* CurrentDiscardTrace)
+	{
+		FString DebugName = GetResourceDebugName(Resource, {});
+		return FString::Printf(
+			BARRIER_TRACKER_LOG_PREFIX_RESNAME
+			TEXT("Mismatched discard of transient resource %s. A transient resource may only be discarded once in its lifetime.\n")
+			TEXT("%s")
+			BARRIER_TRACKER_LOG_SUFFIX,
+			*DebugName,
+			*DebugName,
+			*GetReasonString_DuplicateBackTrace(PreviousDiscardTrace, CurrentDiscardTrace));
+	}
+
+	static inline FString GetReasonString_DiscardWithoutAcquireTransient(FResource* Resource, void* DiscardTrace)
+	{
+		FString DebugName = GetResourceDebugName(Resource, {});
+		return FString::Printf(
+			BARRIER_TRACKER_LOG_PREFIX_RESNAME
+			TEXT("Attempted to discard transient resource %s, but it was never acquired.\n")
+			TEXT("%s")
+			BARRIER_TRACKER_LOG_SUFFIX,
+			*DebugName,
+			*DebugName,
+			*GetReasonString_Backtrace(TEXT("acquire"), TEXT("RHICreateTransition"), DiscardTrace));
+	}
+
 	static inline FString GetReasonString_DuplicateBeginTransition(
 		FResource* Resource, FSubresourceIndex const& SubresourceIndex,
 		const FState& PendingState,
@@ -691,6 +847,68 @@ namespace RHIValidation
 			*GetRHIPipelineName(RequiredState.Pipelines));
 	}
 
+	static inline void* Log(FResource* Resource, FSubresourceIndex const& SubresourceIndex, void* CreateTrace, const TCHAR* TracePrefix, const TCHAR* Type, const TCHAR* LogStr)
+	{
+		void* Trace = CaptureBacktrace();
+
+		if (CreateTrace)
+		{
+			FPlatformMisc::LowLevelOutputDebugStringf(TEXT("\n%s: Type: %s, %s, CreateTrace: 0x%p, %sTrace: 0x%p\n"),
+				*GetResourceDebugName(Resource, SubresourceIndex),
+				Type,
+				LogStr,
+				CreateTrace,
+				TracePrefix,
+				Trace);
+		}
+		else
+		{
+			FPlatformMisc::LowLevelOutputDebugStringf(TEXT("\n%s: Type: %s, %s, Trace: 0x%p\n"),
+				*GetResourceDebugName(Resource, SubresourceIndex),
+				Type,
+				LogStr,
+				Trace);
+		}
+
+		return Trace;
+	}
+
+	void FTransientState::Acquire(FResource* Resource, void* CreateTrace)
+	{
+		RHI_VALIDATION_CHECK(bTransient, *GetReasonString_AcquireNonTransient(Resource));
+		RHI_VALIDATION_CHECK(Status == EStatus::None, *GetReasonString_DuplicateAcquireTransient(Resource, AcquireBacktrace, CreateTrace));
+		Status = EStatus::Acquired;
+
+		if (!AcquireBacktrace)
+		{
+			AcquireBacktrace = CreateTrace;
+		}
+
+		if (Resource->LoggingMode != ELoggingMode::None)
+		{
+			Log(Resource, {}, CreateTrace, TEXT("Acquire"), TEXT("Acquire"), TEXT("Transient Acquire"));
+		}
+	}
+
+	void FTransientState::Discard(FResource* Resource, void* CreateTrace)
+	{
+		RHI_VALIDATION_CHECK(bTransient, *GetReasonString_DiscardNonTransient(Resource));
+
+		RHI_VALIDATION_CHECK(Status != EStatus::None, *GetReasonString_DiscardWithoutAcquireTransient(Resource, CreateTrace));
+		RHI_VALIDATION_CHECK(Status != EStatus::Discarded, *GetReasonString_DuplicateDiscardTransient(Resource, DiscardBacktrace, CreateTrace));
+		Status = EStatus::Discarded;
+
+		if (!DiscardBacktrace)
+		{
+			DiscardBacktrace = CreateTrace;
+		}
+
+		if (Resource->LoggingMode != ELoggingMode::None)
+		{
+			Log(Resource, {}, CreateTrace, TEXT("Discard"), TEXT("Discard"), TEXT("Transient Discard"));
+		}
+	}
+
 	void FResource::SetDebugName(const TCHAR* Name, const TCHAR* Suffix)
 	{
 		DebugName = Suffix
@@ -717,32 +935,6 @@ namespace RHIValidation
 		}
 	}
 
-	void* FSubresourceState::Log(FResource* Resource, FSubresourceIndex const& SubresourceIndex, void* CreateTrace, const TCHAR* TracePrefix, const TCHAR* Type, const TCHAR* LogStr)
-	{
-		void* Trace = CaptureBacktrace();
-
-		if (CreateTrace)
-		{
-			FPlatformMisc::LowLevelOutputDebugStringf(TEXT("\n%s: Type: %s, %s, CreateTrace: 0x%p, %sTrace: 0x%p\n"),
-				*GetResourceDebugName(Resource, SubresourceIndex),
-				Type,
-				LogStr,
-				CreateTrace,
-				TracePrefix,
-				Trace);
-		}
-		else
-		{
-			FPlatformMisc::LowLevelOutputDebugStringf(TEXT("\n%s: Type: %s, %s, Trace: 0x%p\n"),
-				*GetResourceDebugName(Resource, SubresourceIndex),
-				Type,
-				LogStr,
-				Trace);
-		}
-
-		return Trace;
-	}
-
 	void FSubresourceState::BeginTransition(FResource* Resource, FSubresourceIndex const& SubresourceIndex, const FState& CurrentStateFromRHI, const FState& TargetState, EResourceTransitionFlags NewFlags, ERHIPipeline ExecutingPipeline, void* CreateTrace)
 	{
 		FPipelineState& State = States[ExecutingPipeline];
@@ -762,6 +954,11 @@ namespace RHIValidation
 			));
 		}
 
+		if (Resource->TransientState.bTransient)
+		{
+			RHI_VALIDATION_CHECK(Resource->TransientState.IsAcquired(), *GetReasonString_TransitionWithoutAcquire(Resource));
+		}
+
 		// Check we're not already transitioning
 		RHI_VALIDATION_CHECK(!State.bTransitioning, *GetReasonString_DuplicateBeginTransition(Resource, SubresourceIndex, State.Current, TargetState, State.CreateTransitionBacktrace, BeginTrace));
 
@@ -821,6 +1018,11 @@ namespace RHIValidation
 		State.bTransitioning = false;
 		State.BeginTransitionBacktrace = nullptr;
 
+		if (Resource->TransientState.bTransient)
+		{
+			RHI_VALIDATION_CHECK(Resource->TransientState.IsAcquired(), *GetReasonString_TransitionWithoutAcquire(Resource));
+		}
+
 		// Check that the end matches the begin.
 		RHI_VALIDATION_CHECK(TargetState == State.Current, *GetReasonString_MismatchedEndTransition(Resource, SubresourceIndex, State.Current, TargetState));
 
@@ -972,6 +1174,16 @@ namespace RHIValidation
 			Data_EndTransition.Identity.Resource->ReleaseOpRef();
 			break;
 
+		case EOpType::AcquireTransient:
+			Data_AcquireTransient.Resource->TransientState.Acquire(Data_AcquireTransient.Resource, Data_AcquireTransient.CreateBacktrace);
+			Data_AcquireTransient.Resource->ReleaseOpRef();
+			break;
+
+		case EOpType::DiscardTransient:
+			Data_DiscardTransient.Resource->TransientState.Discard(Data_DiscardTransient.Resource, Data_DiscardTransient.CreateBacktrace);
+			Data_AcquireTransient.Resource->ReleaseOpRef();
+			break;
+
 		case EOpType::Assert:
 			Data_Assert.Identity.Resource->EnumerateSubresources(Data_Assert.Identity.SubresourceRange, [this, Pipeline, &bAllowAllUAVsOverlap](FSubresourceState& State, FSubresourceIndex const& SubresourceIndex)
 			{
@@ -1205,6 +1417,7 @@ FValidationTransientResourceAllocator::~FValidationTransientResourceAllocator()
 FRHITexture* FValidationTransientResourceAllocator::CreateTexture(const FRHITextureCreateInfo& InCreateInfo, const TCHAR* InDebugName)
 {
 	check(!bFrozen);
+	checkf(!EnumHasAnyFlags(InCreateInfo.Flags, TexCreate_Transient), TEXT("Attempted to create transient texture %s using legacy TexCreate_Transient flag."), InDebugName);
 
 	FRHITexture* RHITexture = RHIAllocator->CreateTexture(InCreateInfo, InDebugName);
 
@@ -1224,6 +1437,7 @@ FRHITexture* FValidationTransientResourceAllocator::CreateTexture(const FRHIText
 FRHIBuffer* FValidationTransientResourceAllocator::CreateBuffer(const FRHIBufferCreateInfo& InCreateInfo, const TCHAR* InDebugName)
 {
 	check(!bFrozen);
+	checkf(!EnumHasAnyFlags(InCreateInfo.Usage, BUF_Transient), TEXT("Attempted to create transient buffer %s using legacy BUF_Transient flag."), InDebugName);
 
 	FRHIBuffer* RHIBuffer = RHIAllocator->CreateBuffer(InCreateInfo, InDebugName);
 
@@ -1296,6 +1510,8 @@ void FValidationTransientResourceAllocator::Freeze()
 	}
 
 	bFrozen = true;
+
+	RHIAllocator->Freeze();
 }
 
 #endif	// ENABLE_RHI_VALIDATION
diff --git a/Engine/Source/Runtime/RHI/Public/DynamicRHI.h b/Engine/Source/Runtime/RHI/Public/DynamicRHI.h
--- a/Engine/Source/Runtime/RHI/Public/DynamicRHI.h
+++ b/Engine/Source/Runtime/RHI/Public/DynamicRHI.h
@@ -427,7 +427,7 @@ public:
 		return new FGenericRHIGPUFence(Name);
 	}
 
-	virtual void RHICreateTransition(FRHITransition* Transition, ERHIPipeline SrcPipelines, ERHIPipeline DstPipelines, ERHICreateTransitionFlags CreateFlags, TArrayView<const FRHITransitionInfo> Infos)
+	virtual void RHICreateTransition(FRHITransition* Transition, const FRHITransitionCreateInfo& CreateInfo)
 	{
 	}
 
@@ -1621,14 +1621,20 @@ FORCEINLINE class IRHICommandContextContainer* RHIGetCommandContextContainer(int
 
 RHI_API FRenderQueryPoolRHIRef RHICreateRenderQueryPool(ERenderQueryType QueryType, uint32 NumQueries = UINT32_MAX);
 
-FORCEINLINE const FRHITransition* RHICreateTransition(ERHIPipeline SrcPipelines, ERHIPipeline DstPipelines, ERHICreateTransitionFlags CreateFlags, TArrayView<const FRHITransitionInfo> Infos)
+FORCEINLINE const FRHITransition* RHICreateTransition(const FRHITransitionCreateInfo& CreateInfo)
 {
 	// Placement create the transition on the heap.
-	FRHITransition* Transition = new (FMemory::Malloc(FRHITransition::GetTotalAllocationSize(), FRHITransition::GetAlignment())) FRHITransition(SrcPipelines, DstPipelines);
-	GDynamicRHI->RHICreateTransition(Transition, SrcPipelines, DstPipelines, CreateFlags, Infos);
+	FRHITransition* Transition = new (FMemory::Malloc(FRHITransition::GetTotalAllocationSize(), FRHITransition::GetAlignment())) FRHITransition(CreateInfo.SrcPipelines, CreateInfo.DstPipelines);
+	GDynamicRHI->RHICreateTransition(Transition, CreateInfo);
 	return Transition;
 }
 
+UE_DEPRECATED(5.0, "Use the FRHITransitionCreateInfo version of RHICreateTransition instead.")
+FORCEINLINE const FRHITransition* RHICreateTransition(ERHIPipeline SrcPipelines, ERHIPipeline DstPipelines, ERHITransitionCreateFlags CreateFlags, TArrayView<const FRHITransitionInfo> Infos)
+{
+	return RHICreateTransition(FRHITransitionCreateInfo(SrcPipelines, DstPipelines, CreateFlags, Infos));
+}
+
 FORCEINLINE void RHIReleaseTransition(FRHITransition* Transition)
 {
 	GDynamicRHI->RHIReleaseTransition(Transition);
diff --git a/Engine/Source/Runtime/RHI/Public/RHI.h b/Engine/Source/Runtime/RHI/Public/RHI.h
--- a/Engine/Source/Runtime/RHI/Public/RHI.h
+++ b/Engine/Source/Runtime/RHI/Public/RHI.h
@@ -1805,7 +1805,7 @@ inline constexpr bool IsValidAccess(ERHIAccess Access)
 	return !IsInvalidAccess(Access);
 }
 
-enum class ERHICreateTransitionFlags
+enum class ERHITransitionCreateFlags
 {
 	None = 0,
 
@@ -1816,7 +1816,7 @@ enum class ERHICreateTransitionFlags
 	// so should use a partial flush rather than a fence as this is more optimal.
 	NoSplit = 1 << 1
 };
-ENUM_CLASS_FLAGS(ERHICreateTransitionFlags);
+ENUM_CLASS_FLAGS(ERHITransitionCreateFlags);
 
 enum class EResourceTransitionFlags
 {
@@ -1831,16 +1831,6 @@ enum class EResourceTransitionFlags
 };
 ENUM_CLASS_FLAGS(EResourceTransitionFlags);
 
-enum class ERHICreateTransientAllocationFlags
-{
-	// Discard data on acquire
-	Discard,
-
-	// Clear data on acquire
-	Clear,
-};
-ENUM_CLASS_FLAGS(ERHICreateTransientAllocationFlags);
-
 RHI_API FString GetRHIAccessName(ERHIAccess Access);
 RHI_API FString GetResourceTransitionFlagsName(EResourceTransitionFlags Flags);
 RHI_API FString GetRHIPipelineName(ERHIPipeline Pipeline);
@@ -1985,6 +1975,122 @@ struct FRHITransitionInfo : public FRHISubresourceRange
 	}
 };
 
+struct FRHITransientAliasingInfo
+{
+	union
+	{
+		class FRHIResource* Resource = nullptr;
+		class FRHITexture* Texture;
+		class FRHIBuffer* Buffer;
+	};
+
+	enum class EType : uint8
+	{
+		Texture,
+		Buffer
+	} Type = EType::Texture;
+
+	enum class EAction : uint8
+	{
+		Acquire,
+		Discard
+	} Action = EAction::Acquire;
+
+	FRHITransientAliasingInfo() = default;
+
+	static FRHITransientAliasingInfo Acquire(class FRHITexture* Texture)
+	{
+		FRHITransientAliasingInfo Info;
+		Info.Texture = Texture;
+		Info.Type = EType::Texture;
+		Info.Action = EAction::Acquire;
+		return Info;
+	}
+
+	static FRHITransientAliasingInfo Acquire(class FRHIBuffer* Buffer)
+	{
+		FRHITransientAliasingInfo Info;
+		Info.Buffer = Buffer;
+		Info.Type = EType::Buffer;
+		Info.Action = EAction::Acquire;
+		return Info;
+	}
+
+	static FRHITransientAliasingInfo Discard(class FRHITexture* Texture)
+	{
+		FRHITransientAliasingInfo Info;
+		Info.Texture = Texture;
+		Info.Type = EType::Texture;
+		Info.Action = EAction::Discard;
+		return Info;
+	}
+
+	static FRHITransientAliasingInfo Discard(class FRHIBuffer* Buffer)
+	{
+		FRHITransientAliasingInfo Info;
+		Info.Buffer = Buffer;
+		Info.Type = EType::Buffer;
+		Info.Action = EAction::Discard;
+		return Info;
+	}
+
+	bool IsAcquire() const
+	{
+		return Action == EAction::Acquire;
+	}
+
+	bool IsDiscard() const
+	{
+		return Action == EAction::Discard;
+	}
+
+	bool IsTexture() const
+	{
+		return Type == EType::Texture;
+	}
+
+	bool IsBuffer() const
+	{
+		return Type == EType::Buffer;
+	}
+
+	inline bool operator == (const FRHITransientAliasingInfo& RHS) const
+	{
+		return Resource == RHS.Resource
+			&& Type == RHS.Type
+			&& Action == RHS.Action;
+	}
+
+	inline bool operator != (const FRHITransientAliasingInfo& RHS) const
+	{
+		return !(*this == RHS);
+	}
+};
+
+struct FRHITransitionCreateInfo
+{
+	FRHITransitionCreateInfo() = default;
+
+	FRHITransitionCreateInfo(
+		ERHIPipeline InSrcPipelines,
+		ERHIPipeline InDstPipelines,
+		ERHITransitionCreateFlags InFlags = ERHITransitionCreateFlags::None,
+		TArrayView<const FRHITransitionInfo> InTransitionInfos = {},
+		TArrayView<const FRHITransientAliasingInfo> InAliasingInfos = {})
+		: SrcPipelines(InSrcPipelines)
+		, DstPipelines(InDstPipelines)
+		, Flags(InFlags)
+		, TransitionInfos(InTransitionInfos)
+		, AliasingInfos(InAliasingInfos)
+	{}
+
+	ERHIPipeline SrcPipelines = ERHIPipeline::None;
+	ERHIPipeline DstPipelines = ERHIPipeline::None;
+	ERHITransitionCreateFlags Flags = ERHITransitionCreateFlags::None;
+	TArrayView<const FRHITransitionInfo> TransitionInfos;
+	TArrayView<const FRHITransientAliasingInfo> AliasingInfos;
+};
+
 #include "RHIValidationCommon.h"
 
 // Opaque data structure used to represent a pending resource transition in the RHI.
@@ -2024,7 +2130,7 @@ private:
 	{}
 
 	// Give private access to specific functions/RHI commands that need to allocate or control transitions.
-	friend const FRHITransition* RHICreateTransition(ERHIPipeline, ERHIPipeline, ERHICreateTransitionFlags, TArrayView<const struct FRHITransitionInfo>);
+	friend const FRHITransition* RHICreateTransition(const FRHITransitionCreateInfo&);
 	friend class FRHIComputeCommandList;
 	friend struct FRHICommandBeginTransitions;
 	friend struct FRHICommandEndTransitions;
@@ -2084,6 +2190,10 @@ private:
 	friend class FValidationComputeContext;
 	friend class FValidationContext;
 
+	RHIValidation::FOperationsList PendingSignals;
+	RHIValidation::FOperationsList PendingWaits;
+	RHIValidation::FOperationsList PendingAcquires;
+	RHIValidation::FOperationsList PendingDiscards;
 	RHIValidation::FOperationsList PendingOperationsBegin;
 	RHIValidation::FOperationsList PendingOperationsEnd;
 #endif
diff --git a/Engine/Source/Runtime/RHI/Public/RHICommandList.h b/Engine/Source/Runtime/RHI/Public/RHICommandList.h
--- a/Engine/Source/Runtime/RHI/Public/RHICommandList.h
+++ b/Engine/Source/Runtime/RHI/Public/RHICommandList.h
@@ -2509,7 +2509,7 @@ public:
 			FMemStack& MemStack = FMemStack::Get();
 			FMemMark Mark(MemStack);
 			FRHITransition* Transition = new (MemStack.Alloc(FRHITransition::GetTotalAllocationSize(), FRHITransition::GetAlignment())) FRHITransition(Pipeline, Pipeline);
-			GDynamicRHI->RHICreateTransition(Transition, Pipeline, Pipeline, ERHICreateTransitionFlags::NoSplit, Infos);
+			GDynamicRHI->RHICreateTransition(Transition, FRHITransitionCreateInfo(Pipeline, Pipeline, ERHITransitionCreateFlags::NoSplit, Infos));
 
 			GetComputeContext().RHIBeginTransitions(MakeArrayView((const FRHITransition**)&Transition, 1));
 			GetComputeContext().RHIEndTransitions(MakeArrayView((const FRHITransition**)&Transition, 1));
@@ -2522,7 +2522,7 @@ public:
 		{
 			// Allocate the transition in the command list
 			FRHITransition* Transition = new (Alloc(FRHITransition::GetTotalAllocationSize(), FRHITransition::GetAlignment())) FRHITransition(Pipeline, Pipeline);
-			GDynamicRHI->RHICreateTransition(Transition, Pipeline, Pipeline, ERHICreateTransitionFlags::NoSplit, Infos);
+			GDynamicRHI->RHICreateTransition(Transition, FRHITransitionCreateInfo(Pipeline, Pipeline, ERHITransitionCreateFlags::NoSplit, Infos));
 
 			ALLOC_COMMAND(FRHICommandResourceTransition)(Transition);
 		}
diff --git a/Engine/Source/Runtime/RHI/Public/RHIResources.h b/Engine/Source/Runtime/RHI/Public/RHIResources.h
--- a/Engine/Source/Runtime/RHI/Public/RHIResources.h
+++ b/Engine/Source/Runtime/RHI/Public/RHIResources.h
@@ -2908,7 +2908,8 @@ struct RHI_API FRHITextureCreateInfo
 		uint16 InArraySize = 1,
 		uint8 InNumMips = 1,
 		uint8 InNumSamples = 1)
-		: Dimension(InDimension)
+		: ClearValue(InClearValue)
+		, Dimension(InDimension)
 		, Flags(InFlags)
 		, Format(InFormat)
 		, Extent(InExtent)
@@ -2916,7 +2917,6 @@ struct RHI_API FRHITextureCreateInfo
 		, ArraySize(InArraySize)
 		, NumMips(InNumMips)
 		, NumSamples(InNumSamples)
-		, ClearValue(InClearValue)
 	{}
 
 	bool operator == (const FRHITextureCreateInfo& Other) const
@@ -2972,9 +2972,22 @@ struct RHI_API FRHITextureCreateInfo
 		return FIntVector(Extent.X, Extent.Y, Depth);
 	}
 
+	void Reset()
+	{
+		// Usually we don't want to propagate MSAA samples.
+		NumSamples = 1;
+
+		// Remove UAV flag for textures that don't need it (some formats are incompatible).
+		Flags |= TexCreate_RenderTargetable;
+		Flags &= ~(TexCreate_UAV | TexCreate_ResolveTargetable | TexCreate_DepthStencilResolveTarget);
+	}
+
 	/** Returns whether this descriptor conforms to requirements. */
 	bool IsValid() const;
 
+	/** Clear value to use when fast-clearing the texture. */
+	FClearValueBinding ClearValue;
+
 	/** Texture dimension to use when creating the RHI texture. */
 	ETextureDimension Dimension = ETextureDimension::Texture2D;
 
@@ -2998,9 +3011,6 @@ struct RHI_API FRHITextureCreateInfo
 
 	/** Number of samples in the texture. >1 for MSAA. */
 	uint8 NumSamples = 1;
-
-	/** Clear value to use when fast-clearing the texture. */
-	FClearValueBinding ClearValue;
 };
 
 
diff --git a/Engine/Source/Runtime/RHI/Public/RHIValidation.h b/Engine/Source/Runtime/RHI/Public/RHIValidation.h
--- a/Engine/Source/Runtime/RHI/Public/RHIValidation.h
+++ b/Engine/Source/Runtime/RHI/Public/RHIValidation.h
@@ -203,7 +203,7 @@ public:
 		return RHI->RHICreateGPUFence(Name);
 	}
 
-	virtual void RHICreateTransition(FRHITransition* Transition, ERHIPipeline SrcPipelines, ERHIPipeline DstPipelines, ERHICreateTransitionFlags CreateFlags, TArrayView<const FRHITransitionInfo> Infos);
+	virtual void RHICreateTransition(FRHITransition* Transition, const FRHITransitionCreateInfo& CreateInfo);
 
 	virtual void RHIReleaseTransition(FRHITransition* Transition)
 	{
diff --git a/Engine/Source/Runtime/RHI/Public/RHIValidationCommon.h b/Engine/Source/Runtime/RHI/Public/RHIValidationCommon.h
--- a/Engine/Source/Runtime/RHI/Public/RHIValidationCommon.h
+++ b/Engine/Source/Runtime/RHI/Public/RHIValidationCommon.h
@@ -15,6 +15,7 @@ const bool GRHIValidationEnabled = false;
 #if ENABLE_RHI_VALIDATION
 
 class FRHIUniformBuffer;
+class FValidationRHI;
 
 // Forward declaration of function defined in RHIUtilities.h
 static inline bool IsStencilFormat(EPixelFormat Format);
@@ -175,7 +176,6 @@ namespace RHIValidation
 		void EndTransition     (FResource* Resource, FSubresourceIndex const& SubresourceIndex, const FState& CurrentStateFromRHI, const FState& TargetState, ERHIPipeline Pipeline, void* CreateTrace);
 		void Assert            (FResource* Resource, FSubresourceIndex const& SubresourceIndex, const FState& RequiredState, bool bAllowAllUAVsOverlap);
 		void SpecificUAVOverlap(FResource* Resource, FSubresourceIndex const& SubresourceIndex, ERHIPipeline Pipeline, bool bAllow);
-		void* Log              (FResource* Resource, FSubresourceIndex const& SubresourceIndex, void* CreateTrace, const TCHAR* TracePrefix, const TCHAR* Type, const TCHAR* LogStr);
 	};
 
 	struct FSubresourceRange
@@ -232,6 +232,36 @@ namespace RHIValidation
 		}
 	};
 
+	struct FTransientState
+	{
+		FTransientState() = default;
+
+		enum class EStatus : uint8
+		{
+			None,
+			Acquired,
+			Discarded
+		};
+
+		FTransientState(ERHIAccess InitialAccess)
+			: bTransient(InitialAccess == ERHIAccess::Discard)
+		{}
+
+		void* AcquireBacktrace = nullptr;
+		void* DiscardBacktrace = nullptr;
+
+		bool bTransient = false;
+		EStatus Status = EStatus::None;
+
+		bool IsAcquired() const
+		{
+			return Status == EStatus::Acquired;
+		}
+
+		void Acquire(FResource* Resource, void* CreateTrace);
+		void Discard(FResource* Resource, void* CreateTrace);
+	};
+
 	class FResource
 	{
 		friend FTracker;
@@ -239,11 +269,13 @@ namespace RHIValidation
 		friend FOperation;
 		friend FSubresourceState;
 		friend FSubresourceRange;
+		friend FValidationRHI;
 
 	protected:
 		uint32 NumMips = 0;
 		uint32 NumArraySlices = 0;
 		uint32 NumPlanes = 0;
+		FTransientState TransientState;
 
 	private:
 		FString DebugName;
@@ -282,25 +314,23 @@ namespace RHIValidation
 	protected:
 		inline void InitBarrierTracking(int32 InNumMips, int32 InNumArraySlices, int32 InNumPlanes, ERHIAccess InResourceState, const TCHAR* InDebugName)
 		{
-			if (!IsBarrierTrackingInitialized())
-			{
-				checkSlow(InNumMips > 0 && InNumArraySlices > 0 && InNumPlanes > 0);
-				check(InResourceState != ERHIAccess::Unknown);
+			checkSlow(InNumMips > 0 && InNumArraySlices > 0 && InNumPlanes > 0);
+			check(InResourceState != ERHIAccess::Unknown);
 
-				NumMips = InNumMips;
-				NumArraySlices = InNumArraySlices;
-				NumPlanes = InNumPlanes;
+			NumMips = InNumMips;
+			NumArraySlices = InNumArraySlices;
+			NumPlanes = InNumPlanes;
+			TransientState = FTransientState(InResourceState);
 
-				for (auto& State : WholeResourceState.States)
-				{
-					State.Current.Access = InResourceState;
-					State.Previous = State.Current;
-				}
+			for (auto& State : WholeResourceState.States)
+			{
+				State.Current.Access = InResourceState;
+				State.Previous = State.Current;
+			}
 
-				if (InDebugName != nullptr)
-				{
-					SetDebugName(InDebugName);
-				}
+			if (InDebugName != nullptr)
+			{
+				SetDebugName(InDebugName);
 			}
 		}
 	};
@@ -507,6 +537,8 @@ namespace RHIValidation
 	{
 		BeginTransition,
 		EndTransition,
+		AcquireTransient,
+		DiscardTransient,
 		Assert,
 		Rename,
 		Signal,
@@ -549,6 +581,18 @@ namespace RHIValidation
 				void* CreateBacktrace;
 			} Data_EndTransition;
 
+			struct
+			{
+				FResource* Resource;
+				void* CreateBacktrace;
+			} Data_AcquireTransient;
+
+			struct
+			{
+				FResource* Resource;
+				void* CreateBacktrace;
+			} Data_DiscardTransient;
+
 			struct
 			{
 				FResourceIdentity Identity;
@@ -627,6 +671,28 @@ namespace RHIValidation
 			return MoveTemp(Op);
 		}
 
+		static inline FOperation AcquireTransientResource(FResource* Resource, void* CreateBacktrace)
+		{
+			Resource->AddOpRef();
+
+			FOperation Op;
+			Op.Type = EOpType::AcquireTransient;
+			Op.Data_AcquireTransient.Resource = Resource;
+			Op.Data_AcquireTransient.CreateBacktrace = CreateBacktrace;
+			return MoveTemp(Op);
+		}
+
+		static inline FOperation DiscardTransientResource(FResource* Resource, void* CreateBacktrace)
+		{
+			Resource->AddOpRef();
+
+			FOperation Op;
+			Op.Type = EOpType::DiscardTransient;
+			Op.Data_DiscardTransient.Resource = Resource;
+			Op.Data_DiscardTransient.CreateBacktrace = CreateBacktrace;
+			return MoveTemp(Op);
+		}
+
 		static inline FOperation Assert(FResourceIdentity Identity, FState RequiredState)
 		{
 			Identity.Resource->AddOpRef();
diff --git a/Engine/Source/Runtime/RHI/Public/RHIValidationContext.h b/Engine/Source/Runtime/RHI/Public/RHIValidationContext.h
--- a/Engine/Source/Runtime/RHI/Public/RHIValidationContext.h
+++ b/Engine/Source/Runtime/RHI/Public/RHIValidationContext.h
@@ -58,21 +58,41 @@ public:
 
 	virtual void RHIBeginTransitions(TArrayView<const FRHITransition*> Transitions) override final
 	{
+		for (const FRHITransition* Transition : Transitions)
+		{
+			Tracker->AddOps(Transition->PendingAcquires);
+		}
+
 		for (const FRHITransition* Transition : Transitions)
 		{
 			Tracker->AddOps(Transition->PendingOperationsBegin);
 		}
 
+		for (const FRHITransition* Transition : Transitions)
+		{
+			Tracker->AddOps(Transition->PendingSignals);
+		}
+
 		RHIContext->RHIBeginTransitions(Transitions);
 	}
 
 	virtual void RHIEndTransitions(TArrayView<const FRHITransition*> Transitions) override final
 	{
+		for (const FRHITransition* Transition : Transitions)
+		{
+			Tracker->AddOps(Transition->PendingWaits);
+		}
+
 		for (const FRHITransition* Transition : Transitions)
 		{
 			Tracker->AddOps(Transition->PendingOperationsEnd);
 		}
 
+		for (const FRHITransition* Transition : Transitions)
+		{
+			Tracker->AddOps(Transition->PendingDiscards);
+		}
+
 		RHIContext->RHIEndTransitions(Transitions);
 	}
 
@@ -356,11 +376,21 @@ public:
 	{
 		ensureMsgf(!State.bInsideBeginRenderPass, TEXT("Beginning a transition within a renderpass is not supported!"));
 
+		for (const FRHITransition* Transition : Transitions)
+		{
+			Tracker->AddOps(Transition->PendingAcquires);
+		}
+
 		for (const FRHITransition* Transition : Transitions)
 		{
 			Tracker->AddOps(Transition->PendingOperationsBegin);
 		}
 
+		for (const FRHITransition* Transition : Transitions)
+		{
+			Tracker->AddOps(Transition->PendingSignals);
+		}
+
 		RHIContext->RHIBeginTransitions(Transitions);
 	}
 
@@ -368,11 +398,21 @@ public:
 	{
 		ensureMsgf(!State.bInsideBeginRenderPass, TEXT("Ending a transition within a renderpass is not supported!"));
 
+		for (const FRHITransition* Transition : Transitions)
+		{
+			Tracker->AddOps(Transition->PendingWaits);
+		}
+
 		for (const FRHITransition* Transition : Transitions)
 		{
 			Tracker->AddOps(Transition->PendingOperationsEnd);
 		}
 
+		for (const FRHITransition* Transition : Transitions)
+		{
+			Tracker->AddOps(Transition->PendingDiscards);
+		}
+
 		RHIContext->RHIEndTransitions(Transitions);
 	}
 
diff --git a/Engine/Source/Runtime/RenderCore/Private/RenderGraphBuilder.cpp b/Engine/Source/Runtime/RenderCore/Private/RenderGraphBuilder.cpp
--- a/Engine/Source/Runtime/RenderCore/Private/RenderGraphBuilder.cpp
+++ b/Engine/Source/Runtime/RenderCore/Private/RenderGraphBuilder.cpp
@@ -101,12 +101,23 @@ inline void GetPassAccess(ERDGPassFlags PassFlags, ERHIAccess& SRVAccess, ERHIAc
 	}
 }
 
+enum class ERDGTextureAccessFlags
+{
+	None = 0,
+
+	// Access is within the fixed-function render pass.
+	RenderTarget = 1 << 0
+};
+ENUM_CLASS_FLAGS(ERDGTextureAccessFlags);
+
 /** Enumerates all texture accesses and provides the access and subresource range info. This results in
  *  multiple invocations of the same resource, but with different access / subresource range.
  */
 template <typename TAccessFunction>
 void EnumerateTextureAccess(FRDGParameterStruct PassParameters, ERDGPassFlags PassFlags, TAccessFunction AccessFunction)
 {
+	const ERDGTextureAccessFlags NoneFlags = ERDGTextureAccessFlags::None;
+
 	ERHIAccess SRVAccess, UAVAccess;
 	GetPassAccess(PassFlags, SRVAccess, UAVAccess);
 
@@ -117,31 +128,33 @@ void EnumerateTextureAccess(FRDGParameterStruct PassParameters, ERDGPassFlags Pa
 		case UBMT_RDG_TEXTURE:
 			if (FRDGTextureRef Texture = Parameter.GetAsTexture())
 			{
-				AccessFunction(nullptr, Texture, SRVAccess, Texture->GetSubresourceRangeSRV());
+				AccessFunction(nullptr, Texture, SRVAccess, NoneFlags, Texture->GetSubresourceRangeSRV());
 			}
 		break;
 		case UBMT_RDG_TEXTURE_ACCESS:
 		{
 			if (FRDGTextureAccess TextureAccess = Parameter.GetAsTextureAccess())
 			{
-				AccessFunction(nullptr, TextureAccess.GetTexture(), TextureAccess.GetAccess(), TextureAccess.GetTexture()->GetSubresourceRange());
+				AccessFunction(nullptr, TextureAccess.GetTexture(), TextureAccess.GetAccess(), NoneFlags, TextureAccess.GetTexture()->GetSubresourceRange());
 			}
 		}
 		break;
 		case UBMT_RDG_TEXTURE_SRV:
 			if (FRDGTextureSRVRef SRV = Parameter.GetAsTextureSRV())
 			{
-				AccessFunction(SRV, SRV->GetParent(), SRVAccess, SRV->GetSubresourceRange());
+				AccessFunction(SRV, SRV->GetParent(), SRVAccess, NoneFlags, SRV->GetSubresourceRange());
 			}
 		break;
 		case UBMT_RDG_TEXTURE_UAV:
 			if (FRDGTextureUAVRef UAV = Parameter.GetAsTextureUAV())
 			{
-				AccessFunction(UAV, UAV->GetParent(), UAVAccess, UAV->GetSubresourceRange());
+				AccessFunction(UAV, UAV->GetParent(), UAVAccess, NoneFlags, UAV->GetSubresourceRange());
 			}
 		break;
 		case UBMT_RENDER_TARGET_BINDING_SLOTS:
 		{
+			const ERDGTextureAccessFlags RenderTargetAccess = ERDGTextureAccessFlags::RenderTarget;
+
 			const ERHIAccess RTVAccess = ERHIAccess::RTV;
 
 			const FRenderTargetBindingSlots& RenderTargets = Parameter.GetAsRenderTargetBindingSlots();
@@ -161,13 +174,13 @@ void EnumerateTextureAccess(FRDGParameterStruct PassParameters, ERDGPassFlags Pa
 					Range.NumArraySlices = 1;
 				}
 
-				AccessFunction(nullptr, Texture, RTVAccess, Range);
+				AccessFunction(nullptr, Texture, RTVAccess, RenderTargetAccess, Range);
 
 				if (ResolveTexture && ResolveTexture != Texture)
 				{
 					// Resolve targets must use the RTV|ResolveDst flag combination when the resolve is performed through the render
 					// pass. The ResolveDst flag must be used alone only when the resolve is performed using RHICopyToResolveTarget.
-					AccessFunction(nullptr, ResolveTexture, ERHIAccess::RTV | ERHIAccess::ResolveDst, Range);
+					AccessFunction(nullptr, ResolveTexture, ERHIAccess::RTV | ERHIAccess::ResolveDst, RenderTargetAccess, Range);
 				}
 			});
 
@@ -186,13 +199,13 @@ void EnumerateTextureAccess(FRDGParameterStruct PassParameters, ERDGPassFlags Pa
 						Range.NumPlaneSlices = 1;
 					}
 
-					AccessFunction(nullptr, Texture, NewAccess, Range);
+					AccessFunction(nullptr, Texture, NewAccess, RenderTargetAccess, Range);
 				});
 			}
 
 			if (FRDGTextureRef Texture = RenderTargets.ShadingRateTexture)
 			{
-				AccessFunction(nullptr, Texture, ERHIAccess::SRVGraphics, Texture->GetSubresourceRangeSRV());
+				AccessFunction(nullptr, Texture, ERHIAccess::SRVGraphics, RenderTargetAccess, Texture->GetSubresourceRangeSRV());
 			}
 		}
 		break;
@@ -409,7 +422,10 @@ void FRDGBuilder::TickPoolElements()
 	SET_DWORD_STAT(STAT_RDG_PassDependencyCount, GRDGStatPassDependencyCount);
 	SET_DWORD_STAT(STAT_RDG_TextureCount, GRDGStatTextureCount);
 	SET_DWORD_STAT(STAT_RDG_BufferCount, GRDGStatBufferCount);
+	SET_DWORD_STAT(STAT_RDG_TransientTextureCount, GRDGStatTransientTextureCount);
+	SET_DWORD_STAT(STAT_RDG_TransientBufferCount, GRDGStatTransientBufferCount);
 	SET_DWORD_STAT(STAT_RDG_TransitionCount, GRDGStatTransitionCount);
+	SET_DWORD_STAT(STAT_RDG_AliasingCount, GRDGStatAliasingCount);
 	SET_DWORD_STAT(STAT_RDG_TransitionBatchCount, GRDGStatTransitionBatchCount);
 	SET_MEMORY_STAT(STAT_RDG_MemoryWatermark, int64(GRDGStatMemoryWatermark));
 	GRDGStatPassCount = 0;
@@ -419,7 +435,10 @@ void FRDGBuilder::TickPoolElements()
 	GRDGStatPassDependencyCount = 0;
 	GRDGStatTextureCount = 0;
 	GRDGStatBufferCount = 0;
+	GRDGStatTransientTextureCount = 0;
+	GRDGStatTransientBufferCount = 0;
 	GRDGStatTransitionCount = 0;
+	GRDGStatAliasingCount = 0;
 	GRDGStatTransitionBatchCount = 0;
 	GRDGStatMemoryWatermark = 0;
 #endif
@@ -434,7 +453,7 @@ ERDGPassFlags FRDGBuilder::OverridePassFlags(const TCHAR* PassName, ERDGPassFlag
 		true;
 #endif
 
-	const bool bGlobalForceAsyncCompute = (GRDGAsyncCompute == RDG_ASYNC_COMPUTE_FORCE_ENABLED && !GRDGImmediateMode && bDebugAllowedForPass);
+	const bool bGlobalForceAsyncCompute = (GRDGAsyncCompute == RDG_ASYNC_COMPUTE_FORCE_ENABLED && !IsImmediateMode() && bDebugAllowedForPass);
 
 	if (EnumHasAnyFlags(PassFlags, ERDGPassFlags::Compute) && (bGlobalForceAsyncCompute))
 	{
@@ -442,7 +461,7 @@ ERDGPassFlags FRDGBuilder::OverridePassFlags(const TCHAR* PassName, ERDGPassFlag
 		PassFlags |= ERDGPassFlags::AsyncCompute;
 	}
 
-	if (EnumHasAnyFlags(PassFlags, ERDGPassFlags::AsyncCompute) && (GRDGAsyncCompute == RDG_ASYNC_COMPUTE_DISABLED || GRDGImmediateMode != 0 || !bAsyncComputeSupported))
+	if (EnumHasAnyFlags(PassFlags, ERDGPassFlags::AsyncCompute) && (GRDGAsyncCompute == RDG_ASYNC_COMPUTE_DISABLED || IsImmediateMode() || !bAsyncComputeSupported))
 	{
 		PassFlags &= ~ERDGPassFlags::AsyncCompute;
 		PassFlags |= ERDGPassFlags::Compute;
@@ -451,6 +470,62 @@ ERDGPassFlags FRDGBuilder::OverridePassFlags(const TCHAR* PassName, ERDGPassFlag
 	return PassFlags;
 }
 
+bool FRDGBuilder::IsTransient(FRDGBufferRef Buffer) const
+{
+	if (!IsTransientInternal(Buffer))
+	{
+		return false;
+	}
+
+	// Buffers smaller than a page go through normal pooled allocation.
+	const uint32 TransientSizeMin = 64 * 1024;
+
+	return Buffer->Desc.GetTotalNumBytes() >= TransientSizeMin;
+}
+
+bool FRDGBuilder::IsTransient(FRDGTextureRef Texture) const
+{
+	return IsTransientInternal(Texture);
+}
+
+bool FRDGBuilder::IsTransientInternal(FRDGParentResourceRef Resource) const
+{
+	// Immediate mode can't use the transient allocator because we don't know if the user will extract the resource.
+	if (!GRDGTransientAllocator || IsImmediateMode())
+	{
+		return false;
+	}
+
+	// Transient resources must stay within the graph.
+	if (Resource->bExternal || Resource->bExtracted)
+	{
+		return false;
+	}
+
+	// This resource cannot be extracted or made external later.
+	check(bSetupComplete);
+
+	return true;
+}
+
+FRDGBuilder::FTransientResourceAllocator::~FTransientResourceAllocator()
+{
+	if (Allocator)
+	{
+		RHICmdList.ReleaseTransientResourceAllocator(Allocator);
+	}
+}
+
+IRHITransientResourceAllocator* FRDGBuilder::FTransientResourceAllocator::GetOrCreate()
+{
+	check(!IsImmediateMode());
+	if (!Allocator)
+	{
+		Allocator = RHICreateTransientResourceAllocator();
+	}
+	return Allocator;
+}
+
 const char* const FRDGBuilder::kDefaultUnaccountedCSVStat = "RDG_Pass";
 
 FRDGBuilder::FRDGBuilder(FRHICommandListImmediate& InRHICmdList, FRDGEventName InName)
@@ -479,7 +554,7 @@ void FRDGBuilder::PreallocateBuffer(FRDGBufferRef Buffer)
 	{
 		Buffer->bExternal = 1;
 		Buffer->AccessFinal = kDefaultAccessFinal;
-		BeginResourceRHI(GetProloguePassHandle(), Buffer);
+		BeginResourceRHI(nullptr, GetProloguePassHandle(), Buffer);
 		ExternalBuffers.Add(Buffer->PooledBuffer, Buffer);
 	}
 }
@@ -490,30 +565,11 @@ void FRDGBuilder::PreallocateTexture(FRDGTextureRef Texture)
 	{
 		Texture->bExternal = 1;
 		Texture->AccessFinal = kDefaultAccessFinal;
-		BeginResourceRHI(GetProloguePassHandle(), Texture);
+		BeginResourceRHI(nullptr, GetProloguePassHandle(), Texture);
 		ExternalTextures.Add(Texture->GetRHIUnchecked(), Texture);
 	}
 }
 
-FRDGTextureRef FRDGBuilder::CreateTexture(
-	const FRDGTextureDesc& Desc,
-	const TCHAR* Name,
-	ERDGTextureFlags Flags)
-{
-	FRDGTextureDesc TransientDesc = Desc;
-
-	if (FRenderTargetPool::DoesTargetNeedTransienceOverride(Desc.Flags, ERenderTargetTransience::Transient))
-	{
-		TransientDesc.Flags |= TexCreate_Transient;
-	}
-
-	IF_RDG_ENABLE_DEBUG(UserValidation.ValidateCreateTexture(TransientDesc, Name, Flags));
-	FRDGTextureRef Texture = Textures.Allocate(Allocator, Name, TransientDesc, Flags, ERenderTargetTexture::ShaderResource);
-	IF_RDG_ENABLE_DEBUG(UserValidation.ValidateCreateTexture(Texture));
-	IF_RDG_ENABLE_TRACE(Trace.AddResource(Texture));
-	return Texture;
-}
-
 FRDGTextureRef FRDGBuilder::RegisterExternalTexture(
 	const TRefCountPtr<IPooledRenderTarget>& ExternalPooledTexture,
 	ERenderTargetTexture RenderTargetTexture,
@@ -547,12 +603,7 @@ FRDGTextureRef FRDGBuilder::RegisterExternalTexture(
 	}
 
 	FRDGTextureRef Texture = Textures.Allocate(Allocator, Name, Translate(ExternalPooledTexture->GetDesc(), RenderTargetTexture), Flags, RenderTargetTexture);
-
-	FRDGTextureRef PreviousOwner = nullptr;
-	Texture->SetRHI(static_cast<FPooledRenderTarget*>(ExternalPooledTexture.GetReference()), PreviousOwner);
-	checkf(!PreviousOwner,
-		TEXT("Externally registered texture '%s' has a previous RDG texture owner '%s'. This can happen if two RDG builder instances register the same resource."),
-		Name, PreviousOwner->Name);
+	Texture->SetRHI(static_cast<FPooledRenderTarget*>(ExternalPooledTexture.GetReference()));
 
 	const ERHIAccess AccessInitial = kDefaultAccessInitial;
 
@@ -607,12 +658,7 @@ FRDGBufferRef FRDGBuilder::RegisterExternalBuffer(
 	}
 
 	FRDGBufferRef Buffer = Buffers.Allocate(Allocator, Name, ExternalPooledBuffer->Desc, Flags);
-
-	FRDGBufferRef PreviousOwner = nullptr;
-	Buffer->SetRHI(ExternalPooledBuffer, PreviousOwner);
-	checkf(!PreviousOwner,
-		TEXT("Externally registered buffer '%s' has a previous RDG buffer owner '%s'. This can happen if two RDG builder instances register the same resource."),
-		Name, PreviousOwner->Name);
+	Buffer->SetRHI(ExternalPooledBuffer);
 
 	const ERHIAccess AccessInitial = kDefaultAccessInitial;
 
@@ -1212,8 +1258,15 @@ void FRDGBuilder::Compile()
 			{
 				FRDGPass* NextPass = Passes[PassHandle];
 
-				// A pass where the user controls the render pass can't merge with other passes, and raster UAV passes can't merge due to potential interdependencies.
-				if (EnumHasAnyFlags(NextPass->GetFlags(), ERDGPassFlags::SkipRenderPass) || NextPass->bUAVAccess)
+				// A pass where the user controls the render pass can't merge with other passes
+				if (EnumHasAnyFlags(NextPass->GetFlags(), ERDGPassFlags::SkipRenderPass))
+				{
+					CommitMerge();
+					continue;
+				}
+
+				// A pass which writes to resources outside of the render pass introduces new dependencies which break merging.
+				if (!NextPass->bRenderPassOnlyWrites)
 				{
 					CommitMerge();
 					continue;
@@ -1261,6 +1314,20 @@ void FRDGBuilder::Compile()
 
 		CommitMerge();
 	}
+
+	for (FRDGPassHandle PassHandle = Passes.Begin(); PassHandle != Passes.End(); ++PassHandle)
+	{
+		if (PassesToCull[PassHandle] || PassesWithEmptyParameters[PassHandle] || PassesOnAsyncCompute[PassHandle])
+		{
+			continue;
+		}
+
+		FRDGPass* Pass = Passes[PassHandle];
+		FRDGPass* PrologueBarrierPass = GetPrologueBarrierPass(PassHandle);
+		FRDGPass* EpilogueBarrierPass = GetEpilogueBarrierPass(PassHandle);
+		PrologueBarrierPass->ResourcesToBegin.Add(Pass);
+		EpilogueBarrierPass->ResourcesToEnd.Add(Pass);
+	}
 }
 
 void FRDGBuilder::Execute()
@@ -1271,13 +1338,16 @@ void FRDGBuilder::Execute()
 	// Create the epilogue pass at the end of the graph just prior to compilation.
 	EpiloguePass = Passes.Allocate<FRDGSentinelPass>(Allocator, RDG_EVENT_NAME("Graph Epilogue"));
 	SetupEmptyPass(EpiloguePass);
+	bSetupComplete = true;
 
 	IF_RDG_ENABLE_DEBUG(UserValidation.ValidateExecuteBegin());
 
 	const FRDGPassHandle ProloguePassHandle = GetProloguePassHandle();
 	const FRDGPassHandle EpiloguePassHandle = GetEpiloguePassHandle();
 
-	if (!GRDGImmediateMode)
+	FTransientResourceAllocator TransientResourceAllocator(RHICmdList);
+
+	if (!IsImmediateMode())
 	{
 		Compile();
 
@@ -1292,18 +1362,23 @@ void FRDGBuilder::Execute()
 			{
 				if (!PassesToCull[PassHandle])
 				{
-					CollectPassResources(PassHandle);
+					CollectPassResources(TransientResourceAllocator, PassHandle);
 				}
 			}
 
 			for (const auto& Query : ExtractedTextures)
 			{
-				EndResourceRHI(EpiloguePassHandle, Query.Key, 1);
+				EndResourceRHI(TransientResourceAllocator, EpiloguePassHandle, Query.Key, 1);
 			}
 
 			for (const auto& Query : ExtractedBuffers)
 			{
-				EndResourceRHI(EpiloguePassHandle, Query.Key, 1);
+				EndResourceRHI(TransientResourceAllocator, EpiloguePassHandle, Query.Key, 1);
+			}
+
+			if (TransientResourceAllocator)
+			{
+				TransientResourceAllocator->Freeze();
 			}
 		}
 
@@ -1367,7 +1442,7 @@ void FRDGBuilder::Execute()
 	IF_RDG_GPU_SCOPES(GPUScopeStacks.BeginExecute());
 	IF_RDG_ENABLE_TRACE(Trace.OutputGraphBegin());
 
-	if (!GRDGImmediateMode)
+	if (!IsImmediateMode())
 	{
 		QUICK_SCOPE_CYCLE_COUNTER(STAT_FRDGBuilder_Execute_Passes);
 
@@ -1416,7 +1491,6 @@ void FRDGBuilder::Execute()
 	IF_RDG_ENABLE_TRACE(Trace.OutputGraphEnd(*this));
 	IF_RDG_GPU_SCOPES(GPUScopeStacks.Graphics.EndExecute());
 	IF_RDG_CPU_SCOPES(CPUScopeStacks.EndExecute());
-
 	IF_RDG_ENABLE_DEBUG(UserValidation.ValidateExecuteEnd());
 
 #if STATS
@@ -1458,10 +1532,10 @@ void FRDGBuilder::SetupPass(FRDGPass* Pass)
 	const ERDGPassFlags PassFlags = Pass->GetFlags();
 	const ERHIPipeline PassPipeline = Pass->GetPipeline();
 
-	bool bPassUAVAccess = false;
+	bool bRenderPassOnlyWrites = true;
 
 	Pass->TextureStates.Reserve(PassParameters.GetTextureParameterCount() + (PassParameters.HasRenderTargets() ? (MaxSimultaneousRenderTargets + 1) : 0));
-	EnumerateTextureAccess(PassParameters, PassFlags, [&](FRDGViewRef TextureView, FRDGTextureRef Texture, ERHIAccess Access, FRDGTextureSubresourceRange Range)
+	EnumerateTextureAccess(PassParameters, PassFlags, [&](FRDGViewRef TextureView, FRDGTextureRef Texture, ERHIAccess Access, ERDGTextureAccessFlags AccessFlags, FRDGTextureSubresourceRange Range)
 	{
 		check(Access != ERHIAccess::Unknown);
 		const FRDGViewHandle NoUAVBarrierHandle = GetHandleIfNoUAVBarrier(TextureView);
@@ -1505,8 +1579,9 @@ void FRDGBuilder::SetupPass(FRDGPass* Pass)
 			EnumerateSubresourceRange(PassState.State, Texture->Layout, Range, AddSubresourceAccess);
 		}
 
-		bPassUAVAccess |= EnumHasAnyFlags(Access, ERHIAccess::UAVMask);
-		Texture->bProduced |= IsWritableAccess(Access);
+		const bool bWritableAccess = IsWritableAccess(Access);
+		bRenderPassOnlyWrites &= (!bWritableAccess || EnumHasAnyFlags(AccessFlags, ERDGTextureAccessFlags::RenderTarget));
+		Texture->bProduced |= bWritableAccess;
 	});
 
 	Pass->BufferStates.Reserve(PassParameters.GetBufferParameterCount());
@@ -1522,22 +1597,16 @@ void FRDGBuilder::SetupPass(FRDGPass* Pass)
 		PassState.State.NoUAVBarrierFilter.AddHandle(NoUAVBarrierHandle);
 		PassState.State.SetPass(PassPipeline, PassHandle);
 
-		bPassUAVAccess |= EnumHasAnyFlags(Access, ERHIAccess::UAVMask);
-		Buffer->bProduced |= IsWritableAccess(Access);
+		const bool bWritableAccess = IsWritableAccess(Access);
+		bRenderPassOnlyWrites &= !bWritableAccess;
+		Buffer->bProduced |= bWritableAccess;
 	});
 
-	Pass->bUAVAccess = bPassUAVAccess;
+	Pass->bRenderPassOnlyWrites = bRenderPassOnlyWrites;
 
 	const bool bEmptyParameters = !Pass->TextureStates.Num() && !Pass->BufferStates.Num();
 	PassesWithEmptyParameters.Add(bEmptyParameters);
 
-	// The pass can begin / end its own resources on the graphics pipe; async compute is scheduled during compilation.
-	if (PassPipeline == ERHIPipeline::Graphics && !bEmptyParameters)
-	{
-		Pass->ResourcesToBegin.Add(Pass);
-		Pass->ResourcesToEnd.Add(Pass);
-	}
-
 	SetupPassInternal(Pass, PassHandle, PassPipeline);
 }
 
@@ -1579,8 +1648,11 @@ void FRDGBuilder::SetupPassInternal(FRDGPass* Pass, FRDGPassHandle PassHandle, E
 	}
 #endif
 
-	if (GRDGImmediateMode && Pass != EpiloguePass)
+	if (IsImmediateMode() && Pass != EpiloguePass)
 	{
+		Pass->ResourcesToBegin.Add(Pass);
+		Pass->ResourcesToEnd.Add(Pass);
+
 		// Trivially redirect the merge states to the pass states, since we won't be compiling the graph.
 		for (auto& TexturePair : Pass->TextureStates)
 		{
@@ -1603,7 +1675,14 @@ void FRDGBuilder::SetupPassInternal(FRDGPass* Pass, FRDGPassHandle PassHandle, E
 		}
 
 		check(!EnumHasAnyFlags(PassPipeline, ERHIPipeline::AsyncCompute));
-		CollectPassResources(PassHandle);
+
+		// Transient allocator should not be initialized in immediate mode.
+		{
+			FTransientResourceAllocator TransientResourceAllocator(RHICmdList);
+			CollectPassResources(TransientResourceAllocator, PassHandle);
+			check(!TransientResourceAllocator);
+		}
+
 		CollectPassBarriers(PassHandle);
 		ExecutePass(Pass);
 	}
@@ -1707,11 +1786,6 @@ void FRDGBuilder::ExecutePass(FRDGPass* Pass)
 {
 	QUICK_SCOPE_CYCLE_COUNTER(STAT_FRDGBuilder_ExecutePass);
 
-	for (FRHITexture* Texture : Pass->TexturesToAcquire)
-	{
-		RHIAcquireTransientResource(Texture);
-	}
-
 	// Note that we must do this before doing anything with RHICmdList for the pass.
 	// For example, if this pass only executes on GPU 1 we want to avoid adding a
 	// 0-duration event for this pass on GPU 0's time line.
@@ -1756,11 +1830,6 @@ void FRDGBuilder::ExecutePass(FRDGPass* Pass)
 
 	ExecutePassEpilogue(RHICmdListPass, Pass);
 
-	for (FRHITexture* Texture : Pass->TexturesToDiscard)
-	{
-		RHIDiscardTransientResource(Texture);
-	}
-
 	if (Pass->bAsyncComputeEnd)
 	{
 		if (Pass->bAsyncComputeEndExecute)
@@ -1777,7 +1846,7 @@ void FRDGBuilder::ExecutePass(FRDGPass* Pass)
 	}
 }
 
-void FRDGBuilder::CollectPassResources(FRDGPassHandle PassHandle)
+void FRDGBuilder::CollectPassResources(FTransientResourceAllocator& TransientResourceAllocator, FRDGPassHandle PassHandle)
 {
 	RDG_ALLOW_RHI_ACCESS_SCOPE();
 
@@ -1789,15 +1858,14 @@ void FRDGBuilder::CollectPassResources(FRDGPassHandle PassHandle)
 
 		EnumerateTextureParameters(PassParameters, [&](auto Resource)
 		{
-			BeginResourceRHI(PassHandle, Resource);
+			BeginResourceRHI(&TransientResourceAllocator, PassHandle, Resource);
 		});
 
 		EnumerateBufferParameters(PassParameters, [&](auto Resource)
 		{
-			BeginResourceRHI(PassHandle, Resource);
+			BeginResourceRHI(&TransientResourceAllocator, PassHandle, Resource);
 		});
 
-		// Uniform buffer creation requires access to RHI resources through the validator.
 		PassParameters.EnumerateUniformBuffers([&](FRDGUniformBufferBinding UniformBuffer)
 		{
 			BeginResourceRHI(UniformBuffer.GetUniformBuffer());
@@ -1808,12 +1876,12 @@ void FRDGBuilder::CollectPassResources(FRDGPassHandle PassHandle)
 	{
 		for (const auto& TexturePair : PassToEnd->TextureStates)
 		{
-			EndResourceRHI(PassHandle, TexturePair.Key, TexturePair.Value.ReferenceCount);
+			EndResourceRHI(TransientResourceAllocator, PassHandle, TexturePair.Key, TexturePair.Value.ReferenceCount);
 		}
 
 		for (const auto& BufferPair : PassToEnd->BufferStates)
 		{
-			EndResourceRHI(PassHandle, BufferPair.Key, BufferPair.Value.ReferenceCount);
+			EndResourceRHI(TransientResourceAllocator, PassHandle, BufferPair.Key, BufferPair.Value.ReferenceCount);
 		}
 	}
 }
@@ -1857,9 +1925,19 @@ void FRDGBuilder::AddEpilogueTransition(FRDGTextureRef Texture)
 	const FRDGPassHandle EpiloguePassHandle = GetEpiloguePassHandle();
 
 	FRDGSubresourceState ScratchSubresourceState;
+	EBarrierLocation BarrierLocation = EBarrierLocation::Prologue;
+
+	// Texture is using the RHI transient allocator. Transition it back to Discard in the final pass it is used.
+	if (Texture->bTransient)
+	{
+		ScratchSubresourceState.SetPass(ERHIPipeline::Graphics, Texture->LastPass);
+		ScratchSubresourceState.Access = ERHIAccess::Discard;
+		InitAsWholeResource(ScratchTextureState, &ScratchSubresourceState);
 
+		BarrierLocation = EBarrierLocation::Epilogue;
+	}
 	// A known final state means extraction from the graph (or an external texture).
-	if (Texture->AccessFinal != ERHIAccess::Unknown)
+	else if(Texture->AccessFinal != ERHIAccess::Unknown)
 	{
 		ScratchSubresourceState.SetPass(ERHIPipeline::Graphics, EpiloguePassHandle);
 		ScratchSubresourceState.Access = Texture->AccessFinal;
@@ -1895,7 +1973,7 @@ void FRDGBuilder::AddEpilogueTransition(FRDGTextureRef Texture)
 		return;
 	}
 
-	AddTransition(EpiloguePassHandle, Texture, ScratchTextureState);
+	AddTransition(EpiloguePassHandle, Texture, ScratchTextureState, BarrierLocation);
 	ScratchTextureState.Reset();
 }
 
@@ -1908,29 +1986,39 @@ void FRDGBuilder::AddEpilogueTransition(FRDGBufferRef Buffer)
 
 	const FRDGPassHandle EpiloguePassHandle = GetEpiloguePassHandle();
 
-	ERHIAccess AccessFinal = Buffer->AccessFinal;
-
-	// Transition async compute back to the graphics pipe.
-	if (AccessFinal == ERHIAccess::Unknown)
+	if (Buffer->bTransient)
+	{
+		FRDGSubresourceState StateFinal;
+		StateFinal.SetPass(ERHIPipeline::Graphics, Buffer->LastPass);
+		StateFinal.Access = ERHIAccess::Discard;
+		AddTransition(Buffer->LastPass, Buffer, StateFinal, EBarrierLocation::Epilogue);
+	}
+	else
 	{
-		const FRDGSubresourceState State = Buffer->GetState();
+		ERHIAccess AccessFinal = Buffer->AccessFinal;
 
-		if (State.IsUsedBy(ERHIPipeline::AsyncCompute))
+		// Transition async compute back to the graphics pipe.
+		if (AccessFinal == ERHIAccess::Unknown)
 		{
-			AccessFinal = State.Access;
+			const FRDGSubresourceState State = Buffer->GetState();
+
+			if (State.IsUsedBy(ERHIPipeline::AsyncCompute))
+			{
+				AccessFinal = State.Access;
+			}
 		}
-	}
 
-	if (AccessFinal != ERHIAccess::Unknown)
-	{
-		FRDGSubresourceState StateFinal;
-		StateFinal.SetPass(ERHIPipeline::Graphics, EpiloguePassHandle);
-		StateFinal.Access = AccessFinal;
-		AddTransition(EpiloguePassHandle, Buffer, StateFinal);
+		if (AccessFinal != ERHIAccess::Unknown)
+		{
+			FRDGSubresourceState StateFinal;
+			StateFinal.SetPass(ERHIPipeline::Graphics, EpiloguePassHandle);
+			StateFinal.Access = AccessFinal;
+			AddTransition(EpiloguePassHandle, Buffer, StateFinal);
+		}
 	}
 }
 
-void FRDGBuilder::AddTransition(FRDGPassHandle PassHandle, FRDGTextureRef Texture, const FRDGTextureTransientSubresourceStateIndirect& StateAfter)
+void FRDGBuilder::AddTransition(FRDGPassHandle PassHandle, FRDGTextureRef Texture, const FRDGTextureTransientSubresourceStateIndirect& StateAfter, EBarrierLocation BarrierLocation)
 {
 	const FRDGTextureSubresourceRange WholeRange = Texture->GetSubresourceRange();
 	const FRDGTextureSubresourceLayout Layout = Texture->Layout;
@@ -1959,7 +2047,7 @@ void FRDGBuilder::AddTransition(FRDGPassHandle PassHandle, FRDGTextureRef Textur
 				Info.PlaneSlice = Subresource->PlaneSlice;
 			}
 
-			AddTransitionInternal(Texture, SubresourceStateBefore, SubresourceStateAfter, Info);
+			AddTransition(Texture, SubresourceStateBefore, SubresourceStateAfter, BarrierLocation, Info);
 		}
 
 		if (Subresource)
@@ -2045,7 +2133,7 @@ void FRDGBuilder::AddTransition(FRDGPassHandle PassHandle, FRDGTextureRef Textur
 	}
 }
 
-void FRDGBuilder::AddTransition(FRDGPassHandle PassHandle, FRDGBufferRef Buffer, FRDGSubresourceState StateAfter)
+void FRDGBuilder::AddTransition(FRDGPassHandle PassHandle, FRDGBufferRef Buffer, FRDGSubresourceState StateAfter, EBarrierLocation BarrierLocation)
 {
 	check(StateAfter.Access != ERHIAccess::Unknown);
 
@@ -2060,17 +2148,18 @@ void FRDGBuilder::AddTransition(FRDGPassHandle PassHandle, FRDGBufferRef Buffer,
 		Info.AccessBefore = StateBefore.Access;
 		Info.AccessAfter = StateAfter.Access;
 
-		AddTransitionInternal(Buffer, StateBefore, StateAfter, Info);
+		AddTransition(Buffer, StateBefore, StateAfter, BarrierLocation, Info);
 	}
 
 	IF_RDG_ENABLE_DEBUG(LogFile.AddTransitionEdge(PassHandle, StateBefore, StateAfter, Buffer));
 	StateBefore = StateAfter;
 }
 
-void FRDGBuilder::AddTransitionInternal(
+void FRDGBuilder::AddTransition(
 	FRDGParentResource* Resource,
 	FRDGSubresourceState StateBefore,
 	FRDGSubresourceState StateAfter,
+	EBarrierLocation BarrierLocation,
 	const FRHITransitionInfo& TransitionInfo)
 {
 	const ERHIPipeline Graphics = ERHIPipeline::Graphics;
@@ -2081,12 +2170,13 @@ void FRDGBuilder::AddTransitionInternal(
 	StateAfter.Validate();
 #endif
 
-	if (GRDGImmediateMode != 0)
+	if (IsImmediateMode())
 	{
-		FRDGPass* NextPass = Passes[StateAfter.FirstPass[Graphics]];
-		FRDGBarrierBatchBegin& BarriersToBegin = NextPass->GetPrologueBarriersToBegin(Allocator);
-		BarriersToBegin.AddTransition(Resource, TransitionInfo);
-		NextPass->GetPrologueBarriersToEnd(Allocator).AddDependency(&BarriersToBegin);
+		// Immediate mode simply enqueues the barrier into the 'after' pass. Everything is on the graphics pipe.
+		AddToBarriers(StateAfter.FirstPass[Graphics], BarrierLocation, [&](FRDGBarrierBatchBegin& Barriers)
+		{
+			Barriers.AddTransition(Resource, TransitionInfo);
+		});
 		return;
 	}
 
@@ -2101,52 +2191,14 @@ void FRDGBuilder::AddTransitionInternal(
 	}
 
 	check(PipelinesBefore != ERHIPipeline::None && PipelinesAfter != ERHIPipeline::None);
-	checkf(StateBefore.GetLastPass() < StateAfter.GetFirstPass(), TEXT("Submitted a state for '%s' that begins before our previous state has ended."), Resource->Name);
+	checkf(StateBefore.GetLastPass() <= StateAfter.GetFirstPass(), TEXT("Submitted a state for '%s' that begins before our previous state has ended."), Resource->Name);
+	check(!Resource->bTransient || StateBefore.GetLastPass() >= Resource->FirstPass);
 
 	FRDGPassHandlesByPipeline& PassesBefore = StateBefore.LastPass;
 	FRDGPassHandlesByPipeline& PassesAfter  = StateAfter.FirstPass;
 
-	// Before states may come from previous aliases of the texture.
-	if (Resource->bTransient && StateBefore.GetLastPass() < Resource->FirstPass)
-	{
-		checkf(PipelinesAfter == ERHIPipeline::Graphics, TEXT("The first use of transient resource %s is not exclusively on the graphics pipe."), Resource->Name);
-
-		// If this got left in an async compute state, we need to end the transition in the graphics
-		// fork pass, since the lifetime of the previous owner was extended to that point. The cost
-		// should be better absorbed by the fence as well.
-		if (EnumHasAnyFlags(PipelinesBefore, AsyncCompute))
-		{
-			PassesAfter[Graphics] = Passes[PassesBefore[AsyncCompute]]->GraphicsJoinPass;
-		}
-		// Otherwise, can push the start of the transition forward until our alias is acquired.
-		else
-		{
-			PassesBefore[Graphics] = Resource->FirstPass;
-		}
-	}
-
-	const auto GetEpilogueBarrierPassHandle = [&](FRDGPassHandle Handle)
-	{
-		return Passes[Handle]->EpilogueBarrierPass;
-	};
-
-	const auto GetPrologueBarrierPassHandle = [&](FRDGPassHandle Handle)
-	{
-		return Passes[Handle]->PrologueBarrierPass;
-	};
-
-	const auto GetEpilogueBarrierPass = [&](FRDGPassHandle Handle)
-	{
-		return Passes[GetEpilogueBarrierPassHandle(Handle)];
-	};
-
-	const auto GetPrologueBarrierPass = [&](FRDGPassHandle Handle)
-	{
-		return Passes[GetPrologueBarrierPassHandle(Handle)];
-	};
-
 	const uint32 PipelinesBeforeCount = FMath::CountBits(uint64(PipelinesBefore));
-	const uint32 PipelinesAfterCount  = FMath::CountBits(uint64(PipelinesAfter));
+	const uint32 PipelinesAfterCount = FMath::CountBits(uint64(PipelinesAfter));
 
 	// 1-to-1 or 1-to-N pipe transition.
 	if (PipelinesBeforeCount == 1)
@@ -2157,15 +2209,18 @@ void FRDGBuilder::AddTransitionInternal(
 		FRDGPass* BeginPass = nullptr;
 		FRDGBarrierBatchBegin* BarriersToBegin = nullptr;
 
-		// We can only split the transition if the before / after passes are not the same. Aliasing or untracked access affect this.
-		if (BeginPassHandle != FirstEndPassHandle)
+		// Issue the begin in the epilogue of the begin pass if the barrier is being split across multiple passes or the barrier end is in the epilogue.
+		if (BeginPassHandle < FirstEndPassHandle || BarrierLocation == EBarrierLocation::Epilogue)
 		{
 			BeginPass = GetEpilogueBarrierPass(BeginPassHandle);
 			BarriersToBegin = &BeginPass->GetEpilogueBarriersToBeginFor(Allocator, PipelinesAfter);
 		}
+		// This is an immediate prologue transition in the same pass. Issue the begin in the prologue.
 		else
 		{
-			checkf(PipelinesAfter == ERHIPipeline::Graphics, TEXT("Attempted to queue a non-graphics pipe transition for %s from the same begin / end pass. Pipelines: %s"), Resource->Name, *GetRHIPipelineName(PipelinesAfter));
+			checkf(PipelinesAfter == ERHIPipeline::Graphics,
+				TEXT("Attempted to queue an immediate async pipe transition for %s. Pipelines: %s. Async transitions must be split."),
+				Resource->Name, *GetRHIPipelineName(PipelinesAfter));
 
 			BeginPass = GetPrologueBarrierPass(BeginPassHandle);
 			BarriersToBegin = &BeginPass->GetPrologueBarriersToBegin(Allocator);
@@ -2175,26 +2230,51 @@ void FRDGBuilder::AddTransitionInternal(
 
 		for (ERHIPipeline Pipeline : GetRHIPipelines())
 		{
-			// If doing a 1-to-N transition and this is the same pipe as the begin, we end it immediately afterwards.
-			// This is because we can't guarantee that the other pipeline won't join back before the end.
-			if (PipelinesBefore == Pipeline && PipelinesAfterCount > 1)
-			{
-				BeginPass->GetEpilogueBarriersToEnd(Allocator).AddDependency(BarriersToBegin);
+			/** If doing a 1-to-N transition and this is the same pipe as the begin, we end it immediately afterwards in the epilogue
+			 *  of the begin pass. This is because we can't guarantee that the other pipeline won't join back before the end. This can
+			 *  happen if the forking async compute pass joins back to graphics (via another independent transition) before the current
+			 *  graphics transition is ended.
+			 *
+			 *  Async Compute Pipe:               EndA  BeginB
+			 *                                   /            \
+			 *  Graphics Pipe:            BeginA               EndB   EndA
+			 *
+			 *  A is our 1-to-N transition and B is a future transition of the same resource that we haven't evaluated yet. Instead, the
+			 *  same pipe End is performed in the epilogue of the begin pass, which removes the spit barrier but simplifies the tracking:
+			 *
+			 *  Async Compute Pipe:               EndA  BeginB
+			 *                                   /            \
+			 *  Graphics Pipe:            BeginA  EndA         EndB
+			 */
+			if ((PipelinesBefore == Pipeline && PipelinesAfterCount > 1))
+			{
+				AddToEpilogueBarriersToEnd(BeginPassHandle, *BarriersToBegin);
 			}
 			else if (EnumHasAnyFlags(PipelinesAfter, Pipeline))
 			{
-				FRDGPass* EndPass = GetPrologueBarrierPass(PassesAfter[Pipeline]);
-				EndPass->GetPrologueBarriersToEnd(Allocator).AddDependency(BarriersToBegin);
+				if (BarrierLocation == EBarrierLocation::Prologue)
+				{
+					AddToPrologueBarriersToEnd(PassesAfter[Pipeline], *BarriersToBegin);
+				}
+				else
+				{
+					AddToEpilogueBarriersToEnd(PassesAfter[Pipeline], *BarriersToBegin);
+				}
 			}
 		}
 	}
-	// N-to-1 or N-to-N pipe transition.
+	// N-to-1 transition.
 	else
 	{
-		checkf(StateBefore.GetLastPass() != StateAfter.GetFirstPass(),
+		check(PipelinesAfterCount == 1);
+		checkf(StateBefore.GetLastPass() != StateAfter.GetFirstPass() || BarrierLocation == EBarrierLocation::Epilogue,
 			TEXT("Attempted to queue a transition for resource '%s' from '%s' to '%s', but previous and next passes are the same on one pipe."),
 			Resource->Name, *GetRHIPipelineName(PipelinesBefore), *GetRHIPipelineName(PipelinesAfter));
 
+		/** We have multiple begin passes which are being issued from All pipes to 1 end pass on pipe. So, we create a unique barrier batch
+		 *  instance for the requested configuration and add it to all begin passes. The 'End' side of things is slightly simpler than the
+		 *  other path, because there's only one of them.
+		 */
 		FRDGBarrierBatchBeginId Id;
 		Id.PipelinesAfter = PipelinesAfter;
 		for (ERHIPipeline Pipeline : GetRHIPipelines())
@@ -2220,8 +2300,14 @@ void FRDGBuilder::AddTransitionInternal(
 		{
 			if (EnumHasAnyFlags(PipelinesAfter, Pipeline))
 			{
-				FRDGPass* EndPass = GetPrologueBarrierPass(PassesAfter[Pipeline]);
-				EndPass->GetPrologueBarriersToEnd(Allocator).AddDependency(BarriersToBegin);
+				if (BarrierLocation == EBarrierLocation::Prologue)
+				{
+					AddToPrologueBarriersToEnd(PassesAfter[Pipeline], *BarriersToBegin);
+				}
+				else
+				{
+					AddToEpilogueBarriersToEnd(PassesAfter[Pipeline], *BarriersToBegin);
+				}
 			}
 		}
 	}
@@ -2249,21 +2335,17 @@ void FRDGBuilder::BeginResourceRHI(FRDGUniformBuffer* UniformBuffer)
 	UniformBuffer->ResourceRHI = UniformBuffer->UniformBufferRHI;
 }
 
-void FRDGBuilder::BeginResourceRHI(FRDGPassHandle PassHandle, FRDGTextureRef Texture)
+void FRDGBuilder::BeginResourceRHI(FTransientResourceAllocator* TransientResourceAllocator, FRDGPassHandle PassHandle, FRDGTextureRef Texture)
 {
 	check(Texture);
 
-	if (!Texture->FirstPass.IsValid())
-	{
-		Texture->FirstPass = PassHandle;
-	}
-
 	if (Texture->PooledTexture)
 	{
 		return;
 	}
 
-	check(Texture->ReferenceCount > 0 || Texture->bExternal || IsResourceLifetimeExtended());
+	check(TransientResourceAllocator  || Texture->bExternal);
+	check(Texture->ReferenceCount > 0 || Texture->bExternal || IsImmediateMode());
 
 #if RDG_ENABLE_DEBUG
 	{
@@ -2276,28 +2358,51 @@ void FRDGBuilder::BeginResourceRHI(FRDGPassHandle PassHandle, FRDGTextureRef Tex
 	}
 #endif
 
-	TRefCountPtr<FPooledRenderTarget> PooledRenderTarget = GRenderTargetPool.FindFreeElementForRDG(RHICmdList, Texture->Desc, Texture->Name);
+	bool bTransientAllocated = false;
 
-	const bool bTransient = PooledRenderTarget->IsTransient();
-	FRDGTextureRef PreviousOwner = nullptr;
-	Texture->SetRHI(PooledRenderTarget, PreviousOwner);
-	Texture->FirstPass = PassHandle;
-
-	if (bTransient)
+	if (IsTransient(Texture))
 	{
-		if (Texture->bExternal)
-		{
-			RHIAcquireTransientResource(Texture->GetRHIUnchecked());
-		}
-		else if (!GRDGImmediateMode)
+		check(TransientResourceAllocator);
+
+		if (IRHITransientResourceAllocator* AllocatorRHI = TransientResourceAllocator->GetOrCreate())
 		{
-			Texture->bTransient = bTransient;
-			Passes[PassHandle]->TexturesToAcquire.Emplace(Texture->GetRHIUnchecked());
+			TRACE_CPUPROFILER_EVENT_SCOPE(FRDGBuilder::SetupTransientTexture);
+
+			Texture->SetRHI(FRDGPooledTexture::CreateTransient(Allocator, Texture->Desc, AllocatorRHI->CreateTexture(Texture->Desc, Texture->Name)));
+			Texture->bTransient = true;
+
+			check(GetPrologueBarrierPassHandle(PassHandle) == PassHandle);
+
+			FRDGSubresourceState InitialState;
+			InitialState.SetPass(ERHIPipeline::Graphics, PassHandle);
+			InitialState.Access = ERHIAccess::Discard;
+			InitialState.Flags = EResourceTransitionFlags::Discard;
+			InitAsWholeResource(Texture->GetState(), InitialState);
+
+			AddToPrologueBarriers(PassHandle, [&](FRDGBarrierBatchBegin& Barriers)
+			{
+				Barriers.AddAlias(Texture, FRHITransientAliasingInfo::Acquire(Texture->GetRHIUnchecked()));
+			});
+
+			bTransientAllocated = true;
+
+		#if STATS
+			GRDGStatTransientTextureCount++;
+		#endif
 		}
 	}
+
+	if (!bTransientAllocated)
+	{
+		Texture->SetRHI(GRenderTargetPool.FindFreeElementForRDG(RHICmdList, Texture->Desc, Texture->Name));
+	}
+
+	Texture->FirstPass = PassHandle;
+
+	check(Texture->PooledTexture);
 }
 
-void FRDGBuilder::BeginResourceRHI(FRDGPassHandle PassHandle, FRDGTextureSRVRef SRV)
+void FRDGBuilder::BeginResourceRHI(FTransientResourceAllocator*, FRDGPassHandle PassHandle, FRDGTextureSRVRef SRV)
 {
 	check(SRV);
 
@@ -2315,7 +2420,7 @@ void FRDGBuilder::BeginResourceRHI(FRDGPassHandle PassHandle, FRDGTextureSRVRef
 		check(GRHISupportsExplicitHTile);
 		if (!PooledTexture->HTileSRV)
 		{
-			PooledTexture->HTileSRV = RHICreateShaderResourceViewHTile((FRHITexture2D*)PooledTexture->Texture);
+			PooledTexture->HTileSRV = RHICreateShaderResourceViewHTile((FRHITexture2D*)PooledTexture->Texture.GetReference());
 		}
 		SRV->ResourceRHI = PooledTexture->HTileSRV;
 		check(SRV->ResourceRHI);
@@ -2326,7 +2431,7 @@ void FRDGBuilder::BeginResourceRHI(FRDGPassHandle PassHandle, FRDGTextureSRVRef
 	{
 		if (!PooledTexture->FMaskSRV)
 		{
-			PooledTexture->FMaskSRV = RHICreateShaderResourceViewFMask((FRHITexture2D*)PooledTexture->Texture);
+			PooledTexture->FMaskSRV = RHICreateShaderResourceViewFMask((FRHITexture2D*)PooledTexture->Texture.GetReference());
 		}
 		SRV->ResourceRHI = PooledTexture->FMaskSRV;
 		check(SRV->ResourceRHI);
@@ -2337,7 +2442,7 @@ void FRDGBuilder::BeginResourceRHI(FRDGPassHandle PassHandle, FRDGTextureSRVRef
 	{
 		if (!PooledTexture->CMaskSRV)
 		{
-			PooledTexture->CMaskSRV = RHICreateShaderResourceViewWriteMask((FRHITexture2D*)PooledTexture->Texture);
+			PooledTexture->CMaskSRV = RHICreateShaderResourceViewWriteMask((FRHITexture2D*)PooledTexture->Texture.GetReference());
 		}
 		SRV->ResourceRHI = PooledTexture->CMaskSRV;
 		check(SRV->ResourceRHI);
@@ -2359,7 +2464,7 @@ void FRDGBuilder::BeginResourceRHI(FRDGPassHandle PassHandle, FRDGTextureSRVRef
 	PooledTexture->SRVs.Emplace(SRV->Desc, MoveTemp(RHIShaderResourceView));
 }
 
-void FRDGBuilder::BeginResourceRHI(FRDGPassHandle PassHandle, FRDGTextureUAVRef UAV)
+void FRDGBuilder::BeginResourceRHI(FTransientResourceAllocator* TransientResourceAllocator, FRDGPassHandle PassHandle, FRDGTextureUAVRef UAV)
 {
 	check(UAV);
 
@@ -2368,7 +2473,7 @@ void FRDGBuilder::BeginResourceRHI(FRDGPassHandle PassHandle, FRDGTextureUAVRef
 		return;
 	}
 
-	BeginResourceRHI(PassHandle, UAV->Desc.Texture);
+	BeginResourceRHI(TransientResourceAllocator, PassHandle, UAV->Desc.Texture);
 
 	FRDGTextureRef Texture = UAV->Desc.Texture;
 	FRDGPooledTexture* PooledTexture = Texture->PooledTexture;
@@ -2379,7 +2484,7 @@ void FRDGBuilder::BeginResourceRHI(FRDGPassHandle PassHandle, FRDGTextureUAVRef
 		check(GRHISupportsExplicitHTile);
 		if (!PooledTexture->HTileUAV)
 		{
-			PooledTexture->HTileUAV = RHICreateUnorderedAccessViewHTile((FRHITexture2D*)PooledTexture->Texture);
+			PooledTexture->HTileUAV = RHICreateUnorderedAccessViewHTile((FRHITexture2D*)PooledTexture->Texture.GetReference());
 		}
 		UAV->ResourceRHI = PooledTexture->HTileUAV;
 		check(UAV->ResourceRHI);
@@ -2390,7 +2495,7 @@ void FRDGBuilder::BeginResourceRHI(FRDGPassHandle PassHandle, FRDGTextureUAVRef
 	{
 		if (!PooledTexture->StencilUAV)
 		{
-			PooledTexture->StencilUAV = RHICreateUnorderedAccessViewStencil((FRHITexture2D*)PooledTexture->Texture, 0);
+			PooledTexture->StencilUAV = RHICreateUnorderedAccessViewStencil((FRHITexture2D*)PooledTexture->Texture.GetReference(), 0);
 		}
 		UAV->ResourceRHI = PooledTexture->StencilUAV;
 		check(UAV->ResourceRHI);
@@ -2400,30 +2505,62 @@ void FRDGBuilder::BeginResourceRHI(FRDGPassHandle PassHandle, FRDGTextureUAVRef
 	UAV->ResourceRHI = PooledTexture->MipUAVs[UAV->Desc.MipLevel];
 }
 
-void FRDGBuilder::BeginResourceRHI(FRDGPassHandle PassHandle, FRDGBufferRef Buffer)
+void FRDGBuilder::BeginResourceRHI(FTransientResourceAllocator* TransientResourceAllocator, FRDGPassHandle PassHandle, FRDGBufferRef Buffer)
 {
 	check(Buffer);
 
-	if (!Buffer->FirstPass.IsValid())
-	{
-		Buffer->FirstPass = PassHandle;
-	}
-
 	if (Buffer->PooledBuffer)
 	{
 		return;
 	}
 
-	check(Buffer->ReferenceCount > 0 || Buffer->bExternal || IsResourceLifetimeExtended());
+	check(TransientResourceAllocator || Buffer->bExternal);
+	check(Buffer->ReferenceCount > 0 || Buffer->bExternal || IsImmediateMode());
 
-	TRefCountPtr<FRDGPooledBuffer> PooledBuffer = GRenderGraphResourcePool.FindFreeBufferInternal(RHICmdList, Buffer->Desc, Buffer->Name);
+	bool bTransientAllocated = false;
+
+	// If transient then create the resource on the transient allocator. External or extracted resource can't be transient because of lifetime tracking issues.
+	if (IsTransient(Buffer))
+	{
+		check(TransientResourceAllocator);
+
+		if (IRHITransientResourceAllocator* AllocatorRHI = TransientResourceAllocator->GetOrCreate())
+		{
+			TRACE_CPUPROFILER_EVENT_SCOPE(FRDGBuilder::SetupTransientBuffer);
+
+			Buffer->SetRHI(FRDGPooledBuffer::CreateTransient(Allocator, AllocatorRHI->CreateBuffer(Translate(Buffer->Desc), Buffer->Name), Buffer->Desc));
+			Buffer->bTransient = true;
+
+			check(GetPrologueBarrierPassHandle(PassHandle) == PassHandle);
+
+			FRDGSubresourceState& InitialState = Buffer->GetState();
+			InitialState.SetPass(ERHIPipeline::Graphics, PassHandle);
+			InitialState.Access = ERHIAccess::Discard;
+
+			AddToPrologueBarriers(PassHandle, [&](FRDGBarrierBatchBegin& Barriers)
+			{
+				Barriers.AddAlias(Buffer, FRHITransientAliasingInfo::Acquire(Buffer->GetRHIUnchecked()));
+			});
+
+			bTransientAllocated = true;
+
+		#if STATS
+			GRDGStatTransientBufferCount++;
+		#endif
+		}
+	}
+
+	if (!bTransientAllocated)
+	{
+		Buffer->SetRHI(GRenderGraphResourcePool.FindFreeBufferInternal(RHICmdList, Buffer->Desc, Buffer->Name));
+	}
 
-	FRDGBufferRef PreviousOwner = nullptr;
-	Buffer->SetRHI(PooledBuffer, PreviousOwner);
 	Buffer->FirstPass = PassHandle;
+
+	check(Buffer->PooledBuffer);
 }
 
-void FRDGBuilder::BeginResourceRHI(FRDGPassHandle PassHandle, FRDGBufferSRVRef SRV)
+void FRDGBuilder::BeginResourceRHI(FTransientResourceAllocator*, FRDGPassHandle PassHandle, FRDGBufferSRVRef SRV)
 {
 	check(SRV);
 
@@ -2435,15 +2572,10 @@ void FRDGBuilder::BeginResourceRHI(FRDGPassHandle PassHandle, FRDGBufferSRVRef S
 	FRDGBufferRef Buffer = SRV->Desc.Buffer;
 	checkf(Buffer->PooledBuffer, TEXT("Pass parameters contained an SRV of RDG buffer %s, before that buffer was referenced as a UAV, which isn't supported.  Make sure to call ClearUnusedGraphResources to remove unbound parameters which can also cause this."), Buffer->Name);
 
-	if (!Buffer->FirstPass.IsValid())
-	{
-		Buffer->FirstPass = PassHandle;
-	}
-
 	SRV->ResourceRHI = Buffer->PooledBuffer->GetOrCreateSRV(SRV->Desc);
 }
 
-void FRDGBuilder::BeginResourceRHI(FRDGPassHandle PassHandle, FRDGBufferUAV* UAV)
+void FRDGBuilder::BeginResourceRHI(FTransientResourceAllocator* TransientResourceAllocator, FRDGPassHandle PassHandle, FRDGBufferUAV* UAV)
 {
 	check(UAV);
 
@@ -2453,58 +2585,63 @@ void FRDGBuilder::BeginResourceRHI(FRDGPassHandle PassHandle, FRDGBufferUAV* UAV
 	}
 
 	FRDGBufferRef Buffer = UAV->Desc.Buffer;
-	BeginResourceRHI(PassHandle, Buffer);
+	BeginResourceRHI(TransientResourceAllocator, PassHandle, Buffer);
 	UAV->ResourceRHI = Buffer->PooledBuffer->GetOrCreateUAV(UAV->Desc);
 }
 
-void FRDGBuilder::EndResourceRHI(FRDGPassHandle PassHandle, FRDGTextureRef Texture, uint32 ReferenceCount)
+void FRDGBuilder::EndResourceRHI(FTransientResourceAllocator& TransientResourceAllocator, FRDGPassHandle PassHandle, FRDGTextureRef Texture, uint32 ReferenceCount)
 {
 	check(Texture);
+	check(Texture->ReferenceCount >= ReferenceCount || IsImmediateMode());
+	Texture->ReferenceCount -= ReferenceCount;
 
-	if (!IsResourceLifetimeExtended())
+	if (Texture->ReferenceCount == 0)
 	{
-		check(Texture->ReferenceCount >= ReferenceCount);
-		Texture->ReferenceCount -= ReferenceCount;
+		// External and transient textures never release the reference.
+		if (!Texture->bExternal && !Texture->bTransient)
+		{
+			Texture->Allocation = nullptr;
+		}
 
-		if (Texture->ReferenceCount == 0)
+		if (Texture->bTransient)
 		{
-			// External textures should never release the reference.
-			if (!Texture->bExternal && !Texture->bTransient)
-			{
-				Texture->Allocation = nullptr;
-			}
+			TransientResourceAllocator->DeallocateMemory(Texture->GetRHIUnchecked());
 
-			// Extracted textures will be discarded on release of the reference.
-			if (Texture->bTransient && !Texture->bExtracted)
+			AddToEpilogueBarriers(PassHandle, [&](FRDGBarrierBatchBegin& Barriers)
 			{
-				Passes[PassHandle]->TexturesToDiscard.Emplace(Texture->GetRHIUnchecked());
-
-				static_cast<FPooledRenderTarget*>(Texture->PooledRenderTarget)->bAutoDiscard = false;
-			}
-
-			Texture->LastPass = PassHandle;
+				Barriers.AddAlias(Texture, FRHITransientAliasingInfo::Discard(Texture->GetRHIUnchecked()));
+			});
 		}
+
+		Texture->LastPass = PassHandle;
 	}
 }
 
-void FRDGBuilder::EndResourceRHI(FRDGPassHandle PassHandle, FRDGBufferRef Buffer, uint32 ReferenceCount)
+void FRDGBuilder::EndResourceRHI(FTransientResourceAllocator& TransientResourceAllocator, FRDGPassHandle PassHandle, FRDGBufferRef Buffer, uint32 ReferenceCount)
 {
 	check(Buffer);
+	check(Buffer->ReferenceCount >= ReferenceCount || IsImmediateMode());
+	Buffer->ReferenceCount -= ReferenceCount;
 
-	if (!IsResourceLifetimeExtended())
+	if (Buffer->ReferenceCount == 0)
 	{
-		check(Buffer->ReferenceCount >= ReferenceCount);
-		Buffer->ReferenceCount -= ReferenceCount;
+		// External and transient buffers never release the reference.
+		if (!Buffer->bExternal && !Buffer->bTransient)
+		{
+			Buffer->Allocation = nullptr;
+		}
 
-		if (Buffer->ReferenceCount == 0)
+		if (Buffer->bTransient)
 		{
-			// External buffers should never release the reference.
-			if (!Buffer->bExternal)
+			TransientResourceAllocator->DeallocateMemory(Buffer->GetRHIUnchecked());
+
+			AddToEpilogueBarriers(PassHandle, [&](FRDGBarrierBatchBegin& Barriers)
 			{
-				Buffer->Allocation = nullptr;
-			}
-			Buffer->LastPass = PassHandle;
+				Barriers.AddAlias(Buffer, FRHITransientAliasingInfo::Discard(Buffer->GetRHIUnchecked()));
+			});
 		}
+
+		Buffer->LastPass = PassHandle;
 	}
 }
 
diff --git a/Engine/Source/Runtime/RenderCore/Private/RenderGraphPass.cpp b/Engine/Source/Runtime/RenderCore/Private/RenderGraphPass.cpp
--- a/Engine/Source/Runtime/RenderCore/Private/RenderGraphPass.cpp
+++ b/Engine/Source/Runtime/RenderCore/Private/RenderGraphPass.cpp
@@ -120,40 +120,21 @@ FRDGTransitionQueue::FRDGTransitionQueue(uint32 ReservedCount)
 	Queue.Reserve(ReservedCount);
 }
 
-void FRDGTransitionQueue::Insert(const FRHITransition* Transition, ERHICreateTransitionFlags TransitionFlags)
+void FRDGTransitionQueue::Insert(const FRHITransition* Transition, ERHITransitionCreateFlags TransitionFlags)
 {
-	if (!EnumHasAnyFlags(TransitionFlags, ERHICreateTransitionFlags::NoFence))
-	{
-		QueueWithFences.Add(Transition);
-	}
-	else
-	{
-		Queue.Add(Transition);
-	}
+	Queue.Add(Transition);
 }
 
 void FRDGTransitionQueue::Begin(FRHIComputeCommandList& RHICmdList)
 {
-	if (Queue.Num() || QueueWithFences.Num())
-	{
-		// Fence signals happen last.
-		Queue.Append(QueueWithFences);
-		RHICmdList.BeginTransitions(Queue);
-		Queue.Empty();
-		QueueWithFences.Empty();
-	}
+	RHICmdList.BeginTransitions(Queue);
+	Queue.Empty();
 }
 
 void FRDGTransitionQueue::End(FRHIComputeCommandList& RHICmdList)
 {
-	if (Queue.Num() || QueueWithFences.Num())
-	{
-		// Fence waits happen first.
-		Queue.Insert(QueueWithFences, 0);
-		RHICmdList.EndTransitions(Queue);
-		Queue.Empty();
-		QueueWithFences.Empty();
-	}
+	RHICmdList.EndTransitions(Queue);
+	Queue.Empty();
 }
 
 FRDGBarrierBatchBegin::FRDGBarrierBatchBegin(ERHIPipeline InPipelinesToBegin, ERHIPipeline InPipelinesToEnd, const TCHAR* InDebugName, const FRDGPass* InDebugPass)
@@ -189,7 +170,21 @@ void FRDGBarrierBatchBegin::AddTransition(FRDGParentResourceRef Resource, const
 #endif
 
 #if RDG_ENABLE_DEBUG
-	DebugResources.Add(Resource);
+	DebugTransitionResources.Add(Resource);
+#endif
+}
+
+void FRDGBarrierBatchBegin::AddAlias(FRDGParentResourceRef Resource, const FRHITransientAliasingInfo& Info)
+{
+	Aliases.Add(Info);
+	bTransitionNeeded = true;
+
+#if STATS
+	GRDGStatAliasingCount++;
+#endif
+
+#if RDG_ENABLE_DEBUG
+	DebugAliasingResources.Add(Resource);
 #endif
 }
 
@@ -198,7 +193,7 @@ void FRDGBarrierBatchBegin::Submit(FRHIComputeCommandList& RHICmdList, ERHIPipel
 	// Submit may be called once for each pipe. The first to submit creates the transition.
 	if (!Transition && bTransitionNeeded)
 	{
-		Transition = RHICreateTransition(PipelinesToBegin, PipelinesToEnd, TransitionFlags, Transitions);
+		Transition = RHICreateTransition(FRHITransitionCreateInfo(PipelinesToBegin, PipelinesToEnd, TransitionFlags, Transitions, Aliases));
 	}
 
 	if (Transition)
diff --git a/Engine/Source/Runtime/RenderCore/Private/RenderGraphPrivate.cpp b/Engine/Source/Runtime/RenderCore/Private/RenderGraphPrivate.cpp
--- a/Engine/Source/Runtime/RenderCore/Private/RenderGraphPrivate.cpp
+++ b/Engine/Source/Runtime/RenderCore/Private/RenderGraphPrivate.cpp
@@ -86,12 +86,6 @@ FAutoConsoleVariableRef CVarRDGOverlapUAVs(
 	TEXT("RDG will overlap UAV work when requested; if disabled, UAV barriers are always inserted."),
 	ECVF_RenderThreadSafe);
 
-int32 GRDGExtendResourceLifetimes = 0;
-FAutoConsoleVariableRef CVarRDGExtendResourceLifetimes(
-	TEXT("r.RDG.ExtendResourceLifetimes"), GRDGExtendResourceLifetimes,
-	TEXT("RDG will extend resource lifetimes to the full length of the graph. Increases memory usage."),
-	ECVF_RenderThreadSafe);
-
 int32 GRDGTransitionLog = 0;
 FAutoConsoleVariableRef CVarRDGTransitionLog(
 	TEXT("r.RDG.TransitionLog"), GRDGTransitionLog,
@@ -277,6 +271,12 @@ FAutoConsoleVariableRef CVarRDGMergeRenderPasses(
 	TEXT(" 1:on(default);\n"),
 	ECVF_RenderThreadSafe);
 
+int32 GRDGTransientAllocator = 0;
+FAutoConsoleVariableRef CVarRDGUseTransientAllocator(
+	TEXT("r.RDG.TransientAllocator"), GRDGTransientAllocator,
+	TEXT("RDG will use the RHITransientResourceAllocator to allocate all transient resources."),
+	ECVF_RenderThreadSafe);
+
 #if CSV_PROFILER
 int32 GRDGVerboseCSVStats = 0;
 FAutoConsoleVariableRef CVarRDGVerboseCSVStats(
@@ -296,7 +296,10 @@ int32 GRDGStatPassDependencyCount = 0;
 int32 GRDGStatRenderPassMergeCount = 0;
 int32 GRDGStatTextureCount = 0;
 int32 GRDGStatBufferCount = 0;
+int32 GRDGStatTransientTextureCount = 0;
+int32 GRDGStatTransientBufferCount = 0;
 int32 GRDGStatTransitionCount = 0;
+int32 GRDGStatAliasingCount = 0;
 int32 GRDGStatTransitionBatchCount = 0;
 int32 GRDGStatMemoryWatermark = 0;
 
@@ -307,7 +310,10 @@ DEFINE_STAT(STAT_RDG_RenderPassMergeCount);
 DEFINE_STAT(STAT_RDG_PassDependencyCount);
 DEFINE_STAT(STAT_RDG_TextureCount);
 DEFINE_STAT(STAT_RDG_BufferCount);
+DEFINE_STAT(STAT_RDG_TransientTextureCount);
+DEFINE_STAT(STAT_RDG_TransientBufferCount);
 DEFINE_STAT(STAT_RDG_TransitionCount);
+DEFINE_STAT(STAT_RDG_AliasingCount);
 DEFINE_STAT(STAT_RDG_TransitionBatchCount);
 DEFINE_STAT(STAT_RDG_CompileTime);
 DEFINE_STAT(STAT_RDG_CollectResourcesTime);
@@ -336,7 +342,7 @@ void InitRenderGraph()
 	}
 
 	int32 BreakpointValue = 0;
-	if (FParse::Value(FCommandLine::Get(), TEXT("rdgbreakpoint"), BreakpointValue))
+	if (FParse::Value(FCommandLine::Get(), TEXT("rdgbreakpoint="), BreakpointValue))
 	{
 		GRDGBreakpoint = BreakpointValue;
 	}
@@ -347,54 +353,49 @@ void InitRenderGraph()
 	}
 
 	int32 CullPassesValue = 0;
-	if (FParse::Value(FCommandLine::Get(), TEXT("rdgcullpasses"), CullPassesValue))
+	if (FParse::Value(FCommandLine::Get(), TEXT("rdgcullpasses="), CullPassesValue))
 	{
 		GRDGCullPasses = CullPassesValue;
 	}
 
 	int32 MergeRenderPassesValue = 0;
-	if (FParse::Value(FCommandLine::Get(), TEXT("rdgmergerenderpasses"), MergeRenderPassesValue))
+	if (FParse::Value(FCommandLine::Get(), TEXT("rdgmergerenderpasses="), MergeRenderPassesValue))
 	{
 		GRDGMergeRenderPasses = MergeRenderPassesValue;
 	}
 
 	int32 OverlapUAVsValue = 0;
-	if (FParse::Value(FCommandLine::Get(), TEXT("rdgoverlapuavs"), OverlapUAVsValue))
+	if (FParse::Value(FCommandLine::Get(), TEXT("rdgoverlapuavs="), OverlapUAVsValue))
 	{
 		GRDGOverlapUAVs = OverlapUAVsValue;
 	}
 
-	if (FParse::Param(FCommandLine::Get(), TEXT("rdgextendresourcelifetimes")))
-	{
-		GRDGExtendResourceLifetimes = 1;
-	}
-
 	int32 DumpGraphValue = 0;
-	if (FParse::Value(FCommandLine::Get(), TEXT("rdgdumpgraph"), DumpGraphValue))
+	if (FParse::Value(FCommandLine::Get(), TEXT("rdgdumpgraph="), DumpGraphValue))
 	{
 		CVarDumpGraph->Set(DumpGraphValue);
 	}
 
 	int32 AsyncComputeValue = 0;
-	if (FParse::Value(FCommandLine::Get(), TEXT("rdgasynccompute"), AsyncComputeValue))
+	if (FParse::Value(FCommandLine::Get(), TEXT("rdgasynccompute="), AsyncComputeValue))
 	{
 		CVarRDGAsyncCompute->Set(AsyncComputeValue);
 	}
 
 	FString GraphFilter;
-	if (FParse::Value(FCommandLine::Get(), TEXT("rdgdebuggraphfilter"), GraphFilter))
+	if (FParse::Value(FCommandLine::Get(), TEXT("rdgdebuggraphfilter="), GraphFilter))
 	{
 		CVarRDGDebugGraphFilter->Set(*GraphFilter);
 	}
 
 	FString PassFilter;
-	if (FParse::Value(FCommandLine::Get(), TEXT("rdgdebugpassfilter"), PassFilter))
+	if (FParse::Value(FCommandLine::Get(), TEXT("rdgdebugpassfilter="), PassFilter))
 	{
 		CVarRDGDebugPassFilter->Set(*PassFilter);
 	}
 
 	FString ResourceFilter;
-	if (FParse::Value(FCommandLine::Get(), TEXT("rdgdebugresourcefilter"), ResourceFilter))
+	if (FParse::Value(FCommandLine::Get(), TEXT("rdgdebugresourcefilter="), ResourceFilter))
 	{
 		CVarRDGDebugResourceFilter->Set(*ResourceFilter);
 	}
diff --git a/Engine/Source/Runtime/RenderCore/Private/RenderGraphPrivate.h b/Engine/Source/Runtime/RenderCore/Private/RenderGraphPrivate.h
--- a/Engine/Source/Runtime/RenderCore/Private/RenderGraphPrivate.h
+++ b/Engine/Source/Runtime/RenderCore/Private/RenderGraphPrivate.h
@@ -31,7 +31,6 @@ extern int32 GRDGBreakpoint;
 extern int32 GRDGTransitionLog;
 extern int32 GRDGImmediateMode;
 extern int32 GRDGOverlapUAVs;
-extern int32 GRDGExtendResourceLifetimes;
 extern bool  GRDGAllowRHIAccess;
 
 class FRDGAllowRHIAccessScope
@@ -93,7 +92,6 @@ const int32 GRDGBreakpoint = 0;
 const int32 GRDGTransitionLog = 0;
 const int32 GRDGImmediateMode = 0;
 const int32 GRDGOverlapUAVs = 1;
-const int32 GRDGExtendResourceLifetimes = 0;
 
 #define RDG_ALLOW_RHI_ACCESS_SCOPE()
 
@@ -101,14 +99,10 @@ const int32 GRDGExtendResourceLifetimes = 0;
 
 #endif
 
-inline bool IsResourceLifetimeExtended()
-{
-	return (GRDGImmediateMode != 0) || (GRDGExtendResourceLifetimes != 0);
-}
-
 extern int32 GRDGAsyncCompute;
 extern int32 GRDGCullPasses;
 extern int32 GRDGMergeRenderPasses;
+extern int32 GRDGTransientAllocator;
 
 #if CSV_PROFILER
 extern int32 GRDGVerboseCSVStats;
@@ -124,7 +118,10 @@ extern int32 GRDGStatRenderPassMergeCount;
 extern int32 GRDGStatPassDependencyCount;
 extern int32 GRDGStatTextureCount;
 extern int32 GRDGStatBufferCount;
+extern int32 GRDGStatTransientTextureCount;
+extern int32 GRDGStatTransientBufferCount;
 extern int32 GRDGStatTransitionCount;
+extern int32 GRDGStatAliasingCount;
 extern int32 GRDGStatTransitionBatchCount;
 extern int32 GRDGStatMemoryWatermark;
 
@@ -135,7 +132,10 @@ DECLARE_DWORD_COUNTER_STAT_EXTERN(TEXT("Render Passes Merged"), STAT_RDG_RenderP
 DECLARE_DWORD_COUNTER_STAT_EXTERN(TEXT("Pass Dependencies"), STAT_RDG_PassDependencyCount, STATGROUP_RDG, RENDERCORE_API);
 DECLARE_DWORD_COUNTER_STAT_EXTERN(TEXT("Textures"), STAT_RDG_TextureCount, STATGROUP_RDG, RENDERCORE_API);
 DECLARE_DWORD_COUNTER_STAT_EXTERN(TEXT("Buffers"), STAT_RDG_BufferCount, STATGROUP_RDG, RENDERCORE_API);
+DECLARE_DWORD_COUNTER_STAT_EXTERN(TEXT("Transient Textures"), STAT_RDG_TransientTextureCount, STATGROUP_RDG, RENDERCORE_API);
+DECLARE_DWORD_COUNTER_STAT_EXTERN(TEXT("Transient Buffers"), STAT_RDG_TransientBufferCount, STATGROUP_RDG, RENDERCORE_API);
 DECLARE_DWORD_COUNTER_STAT_EXTERN(TEXT("Resource Transitions"), STAT_RDG_TransitionCount, STATGROUP_RDG, RENDERCORE_API);
+DECLARE_DWORD_COUNTER_STAT_EXTERN(TEXT("Resource Acquires and Discards"), STAT_RDG_AliasingCount, STATGROUP_RDG, RENDERCORE_API);
 DECLARE_DWORD_COUNTER_STAT_EXTERN(TEXT("Resource Transition Batches"), STAT_RDG_TransitionBatchCount, STATGROUP_RDG, RENDERCORE_API);
 
 DECLARE_CYCLE_STAT_EXTERN(TEXT("Compile"), STAT_RDG_CompileTime, STATGROUP_RDG, RENDERCORE_API);
@@ -160,4 +160,9 @@ inline const TCHAR* GetEpilogueBarriersToBeginDebugName(ERHIPipeline Pipelines)
 	}
 #endif
 	return TEXT("");
+}
+
+inline bool IsImmediateMode()
+{
+	return GRDGImmediateMode != 0;
 }
\ No newline at end of file
diff --git a/Engine/Source/Runtime/RenderCore/Private/RenderGraphResourcePool.cpp b/Engine/Source/Runtime/RenderCore/Private/RenderGraphResourcePool.cpp
--- a/Engine/Source/Runtime/RenderCore/Private/RenderGraphResourcePool.cpp
+++ b/Engine/Source/Runtime/RenderCore/Private/RenderGraphResourcePool.cpp
@@ -52,30 +52,31 @@ TRefCountPtr<FRDGPooledBuffer> FRenderGraphResourcePool::FindFreeBufferInternal(
 
 	// Allocate new one
 	{
-		TRefCountPtr<FRDGPooledBuffer> PooledBuffer = new FRDGPooledBuffer(Desc);
-		AllocatedBuffers.Add(PooledBuffer);
-		check(PooledBuffer->GetRefCount() == 2);
-
-		PooledBuffer->Name = InDebugName;
-		PooledBuffer->LastUsedFrame = FrameCounter;
-
 		uint32 NumBytes = Desc.GetTotalNumBytes();
 
 		FRHIResourceCreateInfo CreateInfo(InDebugName);
+		TRefCountPtr<FRHIBuffer> BufferRHI;
 
 		if (Desc.UnderlyingType == FRDGBufferDesc::EUnderlyingType::VertexBuffer)
 		{
-			PooledBuffer->Buffer = RHICreateVertexBuffer(NumBytes, Desc.Usage, CreateInfo);
+			BufferRHI = RHICreateVertexBuffer(NumBytes, Desc.Usage, CreateInfo);
 		}
 		else if (Desc.UnderlyingType == FRDGBufferDesc::EUnderlyingType::StructuredBuffer)
 		{
-			PooledBuffer->Buffer = RHICreateStructuredBuffer(Desc.BytesPerElement, NumBytes, Desc.Usage, CreateInfo);
+			BufferRHI = RHICreateStructuredBuffer(Desc.BytesPerElement, NumBytes, Desc.Usage, CreateInfo);
 		}
 		else
 		{
 			check(0);
 		}
 
+		TRefCountPtr<FRDGPooledBuffer> PooledBuffer = FRDGPooledBuffer::CreateCommitted(MoveTemp(BufferRHI), Desc);
+		AllocatedBuffers.Add(PooledBuffer);
+		check(PooledBuffer->GetRefCount() == 2);
+
+		PooledBuffer->Name = InDebugName;
+		PooledBuffer->LastUsedFrame = FrameCounter;
+
 		return PooledBuffer;
 	}
 }
diff --git a/Engine/Source/Runtime/RenderCore/Private/RenderGraphResources.cpp b/Engine/Source/Runtime/RenderCore/Private/RenderGraphResources.cpp
--- a/Engine/Source/Runtime/RenderCore/Private/RenderGraphResources.cpp
+++ b/Engine/Source/Runtime/RenderCore/Private/RenderGraphResources.cpp
@@ -35,7 +35,7 @@ FRDGParentResource::FRDGParentResource(const TCHAR* InName, const ERDGParentReso
 
 FRDGParentResource::~FRDGParentResource()
 {
-	check(GRDGImmediateMode || ReferenceCount == 0);
+	check(IsImmediateMode() || ReferenceCount == 0);
 }
 
 void FRDGParentResource::SetPassthroughRHI(FRHIResource* InResourceRHI)
@@ -162,48 +162,19 @@ bool FRDGSubresourceState::IsTransitionRequired(const FRDGSubresourceState& Prev
 	return false;
 }
 
-bool FRDGTextureDesc::IsValid() const
-{
-	if (Extent.X <= 0 || Extent.Y <= 0 || Depth == 0 || ArraySize == 0 || NumMips == 0 || NumSamples < 1 || NumSamples > 8)
-	{
-		return false;
-	}
-
-	if (NumSamples > 1 && !(Dimension == ETextureDimension::Texture2D || Dimension == ETextureDimension::Texture2DArray))
-	{
-		return false;
-	}
-
-	if (Dimension == ETextureDimension::Texture3D)
-	{
-		if (ArraySize > 1)
-		{
-			return false;
-		}
-	}
-	else if (Depth > 1)
-	{
-		return false;
-	}
-
-	if (Format == PF_Unknown)
-	{
-		return false;
-	}
-
-	return true;
-}
-
 void FRDGPooledTexture::InitViews(const FUnorderedAccessViewRHIRef& FirstMipUAV)
 {
-	if (EnumHasAnyFlags(Desc.Flags, TexCreate_ShaderResource))
+	const ETextureCreateFlags Flags = Texture->GetFlags();
+	const uint32 NumMips = Texture->GetNumMips();
+
+	if (EnumHasAnyFlags(Flags, TexCreate_ShaderResource))
 	{
-		SRVs.Empty(Desc.NumMips);
+		SRVs.Empty();
 	}
 
-	if (EnumHasAnyFlags(Desc.Flags, TexCreate_UAV))
+	if (EnumHasAnyFlags(Flags, TexCreate_UAV))
 	{
-		MipUAVs.Empty(Desc.NumMips);
+		MipUAVs.Empty(NumMips);
 
 		uint32 MipLevel = 0;
 
@@ -213,13 +184,28 @@ void FRDGPooledTexture::InitViews(const FUnorderedAccessViewRHIRef& FirstMipUAV)
 			MipLevel++;
 		}
 
-		for (; MipLevel < Desc.NumMips; MipLevel++)
+		for (; MipLevel < NumMips; MipLevel++)
 		{
 			MipUAVs.Add(RHICreateUnorderedAccessView(Texture, MipLevel));
 		}
 	}
 }
 
+void FRDGPooledTexture::Finalize()
+{
+	for (FRDGSubresourceState& SubresourceState : State)
+	{
+		SubresourceState.Finalize();
+	}
+	Owner = nullptr;
+}
+
+void FRDGPooledTexture::Reset()
+{
+	InitAsWholeResource(State);
+	Owner = nullptr;
+}
+
 FRDGTextureSubresourceRange FRDGTexture::GetSubresourceRangeSRV() const
 {
 	FRDGTextureSubresourceRange Range = GetSubresourceRange();
@@ -232,32 +218,37 @@ FRDGTextureSubresourceRange FRDGTexture::GetSubresourceRangeSRV() const
 	return Range;
 }
 
-void FRDGTexture::SetRHI(FPooledRenderTarget* InPooledRenderTarget, FRDGTextureRef& OutPreviousOwner)
+void FRDGTexture::SetRHI(FPooledRenderTarget* InPooledRenderTarget)
 {
 	check(InPooledRenderTarget);
+	Allocation = InPooledRenderTarget;
+	PooledRenderTarget = InPooledRenderTarget;
 
 	if (!InPooledRenderTarget->HasRDG())
 	{
 		InPooledRenderTarget->InitRDG();
 	}
-	PooledTexture = InPooledRenderTarget->GetRDG(RenderTargetTexture);
-	check(PooledTexture);
 
+	SetRHI(InPooledRenderTarget->GetRDG(RenderTargetTexture));
+}
+
+void FRDGTexture::SetRHI(FRDGPooledTexture* InPooledTexture)
+{
+	check(InPooledTexture);
+	PooledTexture = InPooledTexture;
 	State = &PooledTexture->State;
 
 	// Return the previous owner and assign this texture as the new one.
-	OutPreviousOwner = PooledTexture->Owner;
+	FRDGTextureRef PreviousOwner = PooledTexture->Owner;
 	PooledTexture->Owner = this;
 
 	// Link the previous alias to this one.
-	if (OutPreviousOwner)
+	if (PreviousOwner)
 	{
-		OutPreviousOwner->NextOwner = Handle;
-		OutPreviousOwner->bLastOwner = false;
+		PreviousOwner->NextOwner = Handle;
+		PreviousOwner->bLastOwner = false;
 	}
 
-	Allocation = InPooledRenderTarget;
-	PooledRenderTarget = InPooledRenderTarget;
 	ResourceRHI = PooledTexture->GetRHI();
 	check(ResourceRHI);
 }
@@ -279,24 +270,27 @@ void FRDGTexture::Finalize()
 			PooledTexture->Finalize();
 		}
 
-		// Restore the reference to the last owner in the aliasing chain.
-		Allocation = PooledRenderTarget;
+		// Restore the reference to the last owner in the aliasing chain; not necessary for the transient resource allocator.
+		if (PooledRenderTarget)
+		{
+			Allocation = PooledRenderTarget;
+		}
 	}
 }
 
-void FRDGBuffer::SetRHI(FRDGPooledBuffer* InPooledBuffer, FRDGBufferRef& OutPreviousOwner)
+void FRDGBuffer::SetRHI(FRDGPooledBuffer* InPooledBuffer)
 {
 	check(InPooledBuffer);
 
 	// Return the previous owner and assign this buffer as the new one.
-	OutPreviousOwner = InPooledBuffer->Owner;
+	FRDGBuffer* PreviousOwner = InPooledBuffer->Owner;
 	InPooledBuffer->Owner = this;
 
 	// Link the previous owner to this one.
-	if (OutPreviousOwner)
+	if (PreviousOwner)
 	{
-		OutPreviousOwner->NextOwner = Handle;
-		OutPreviousOwner->bLastOwner = false;
+		PreviousOwner->NextOwner = Handle;
+		PreviousOwner->bLastOwner = false;
 	}
 
 	PooledBuffer = InPooledBuffer;
@@ -323,7 +317,7 @@ void FRDGBuffer::Finalize()
 			PooledBuffer->Finalize();
 		}
 
-		// Restore the reference to the last owner in the chain and sanitize all graph state.
+		// Restore the reference to the last owner in the aliasing chain.
 		Allocation = PooledBuffer;
 	}
 }
diff --git a/Engine/Source/Runtime/RenderCore/Private/RenderGraphTrace.cpp b/Engine/Source/Runtime/RenderCore/Private/RenderGraphTrace.cpp
--- a/Engine/Source/Runtime/RenderCore/Private/RenderGraphTrace.cpp
+++ b/Engine/Source/Runtime/RenderCore/Private/RenderGraphTrace.cpp
@@ -50,6 +50,7 @@ UE_TRACE_EVENT_BEGIN(RDGTrace, BufferMessage)
 	UE_TRACE_EVENT_FIELD(bool, IsExternal)
 	UE_TRACE_EVENT_FIELD(bool, IsExtracted)
 	UE_TRACE_EVENT_FIELD(bool, IsCulled)
+	UE_TRACE_EVENT_FIELD(bool, IsTransient)
 UE_TRACE_EVENT_END()
  
 UE_TRACE_EVENT_BEGIN(RDGTrace, TextureMessage)
@@ -73,6 +74,7 @@ UE_TRACE_EVENT_BEGIN(RDGTrace, TextureMessage)
 	UE_TRACE_EVENT_FIELD(bool, IsExternal)
 	UE_TRACE_EVENT_FIELD(bool, IsExtracted)
 	UE_TRACE_EVENT_FIELD(bool, IsCulled)
+	UE_TRACE_EVENT_FIELD(bool, IsTransient)
 UE_TRACE_EVENT_END()
 
 UE_TRACE_EVENT_BEGIN(RDGTrace, ScopeMessage)
@@ -88,7 +90,7 @@ static_assert(sizeof(FRDGBufferHandle) == sizeof(uint16), "Expected 16 bit buffe
 
 bool IsTraceEnabled()
 {
-	return UE_TRACE_CHANNELEXPR_IS_ENABLED(RDGChannel) && !GRDGImmediateMode;
+	return UE_TRACE_CHANNELEXPR_IS_ENABLED(RDGChannel) && !IsImmediateMode();
 }
 
 void FRDGTrace::OutputGraphBegin()
@@ -239,7 +241,8 @@ void FRDGTrace::OutputGraphEnd(const FRDGBuilder& GraphBuilder)
 			<< TextureMessage.NumSamples(Texture->Desc.NumSamples)
 			<< TextureMessage.IsExternal(bool(Texture->bExternal))
 			<< TextureMessage.IsExtracted(bool(Texture->bExtracted))
-			<< TextureMessage.IsCulled(bool(Texture->bCulled));
+			<< TextureMessage.IsCulled(bool(Texture->bCulled))
+			<< TextureMessage.IsTransient(bool(Texture->bTransient));
 	}
 
 	for (FRDGBufferHandle Handle = Buffers.Begin(); Handle != Buffers.End(); ++Handle)
@@ -257,7 +260,8 @@ void FRDGTrace::OutputGraphEnd(const FRDGBuilder& GraphBuilder)
 			<< BufferMessage.NumElements(Buffer->Desc.NumElements)
 			<< BufferMessage.IsExternal(bool(Buffer->bExternal))
 			<< BufferMessage.IsExtracted(bool(Buffer->bExtracted))
-			<< BufferMessage.IsCulled(bool(Buffer->bCulled));
+			<< BufferMessage.IsCulled(bool(Buffer->bCulled))
+			<< BufferMessage.IsTransient(bool(Buffer->bTransient));
 	}
 
 	UE_TRACE_LOG(RDGTrace, GraphEndMessage, RDGChannel);
diff --git a/Engine/Source/Runtime/RenderCore/Private/RenderGraphValidation.cpp b/Engine/Source/Runtime/RenderCore/Private/RenderGraphValidation.cpp
--- a/Engine/Source/Runtime/RenderCore/Private/RenderGraphValidation.cpp
+++ b/Engine/Source/Runtime/RenderCore/Private/RenderGraphValidation.cpp
@@ -936,7 +936,7 @@ void FRDGUserValidation::ValidateExecutePassBegin(const FRDGPass* Pass)
 				if (FRDGTextureRef Texture = RenderTargets.ShadingRateTexture)
 				{
 					Texture->MarkResourceAsUsed();
-				}			
+				}
 			}
 			break;
 			}
@@ -1063,7 +1063,7 @@ void FRDGBarrierValidation::ValidateBarrierBatchBegin(const FRDGPass* Pass, cons
 
 		for (int32 Index = 0; Index < Batch.Transitions.Num(); ++Index)
 		{
-			FRDGParentResourceRef Resource = Batch.DebugResources[Index];
+			FRDGParentResourceRef Resource = Batch.DebugTransitionResources[Index];
 			const FRHITransitionInfo& Transition = Batch.Transitions[Index];
 
 			if (Resource->Type == ERDGParentResourceType::Texture)
@@ -1076,12 +1076,14 @@ void FRDGBarrierValidation::ValidateBarrierBatchBegin(const FRDGPass* Pass, cons
 				ResourceMap->Buffers.Emplace(static_cast<FRDGBufferRef>(Resource), Transition);
 			}
 		}
-	}
 
-	const bool bAllowedForPass = IsDebugAllowedForGraph(GraphName) && IsDebugAllowedForPass(Pass->GetName());
+		for (int32 Index = 0; Index < Batch.Aliases.Num(); ++Index)
+		{
+			ResourceMap->Aliases.Emplace(Batch.DebugAliasingResources[Index], Batch.Aliases[Index]);
+		}
+	}
 
-	// Debug mode will report errors regardless of logging filter.
-	if (!bAllowedForPass && !GRDGDebug)
+	if (!IsDebugAllowedForGraph(GraphName) || !IsDebugAllowedForPass(Pass->GetName()))
 	{
 		return;
 	}
@@ -1093,21 +1095,39 @@ void FRDGBarrierValidation::ValidateBarrierBatchBegin(const FRDGPass* Pass, cons
 		if (!bFoundFirst)
 		{
 			bFoundFirst = true;
-			UE_CLOG(bAllowedForPass, LogRDG, Display, TEXT("[%s(Index: %d, Pipeline: %s): %s] (Begin):"), Pass->GetName(), Pass->GetHandle().GetIndex(), *GetRHIPipelineName(Pass->GetPipeline()), Batch.DebugName);
+			UE_LOG(LogRDG, Display, TEXT("[%s(Index: %d, Pipeline: %s): %s] (Begin):"), Pass->GetName(), Pass->GetHandle().GetIndex(), *GetRHIPipelineName(Pass->GetPipeline()), Batch.DebugName);
 		}
 	};
 
+	for (const auto& KeyValue : ResourceMap->Aliases)
+	{
+		const FRHITransientAliasingInfo& Info = KeyValue.Value;
+		if (Info.IsAcquire())
+		{
+			FRDGParentResourceRef Resource = KeyValue.Key;
+
+			if (IsDebugAllowedForResource(Resource->Name))
+			{
+				LogHeader();
+				UE_LOG(LogRDG, Display, TEXT("\tRDG(%p) RHI(%p) %s - Acquire"), Resource, Resource->GetRHIUnchecked(), Resource->Name);
+			}
+		}
+	}
+
 	for (const auto& Pair : ResourceMap->Textures)
 	{
 		FRDGTextureRef Texture = Pair.Key;
-		
-		const bool bAllowedForResource = bAllowedForPass && IsDebugAllowedForResource(Texture->Name);
+
+		if (!IsDebugAllowedForResource(Texture->Name))
+		{
+			continue;
+		}
 
 		const auto& Transitions = Pair.Value;
-		if (bAllowedForResource && Transitions.Num())
+		if (Transitions.Num())
 		{
 			LogHeader();
-			UE_LOG(LogRDG, Display, TEXT("\t(%p) %s:"), Texture, Texture->Name);
+			UE_LOG(LogRDG, Display, TEXT("\tRDG(%p) RHI(%p) %s:"), Texture, Texture->GetRHIUnchecked(), Texture->Name);
 		}
 
 		const FRDGTextureSubresourceLayout SubresourceLayout = Texture->GetSubresourceLayout();
@@ -1121,7 +1141,7 @@ void FRDGBarrierValidation::ValidateBarrierBatchBegin(const FRDGPass* Pass, cons
 			{
 				const int32 SubresourceIndex = SubresourceLayout.GetSubresourceIndex(Subresource);
 
-				UE_CLOG(bAllowedForResource, LogRDG, Display, TEXT("\t\tMip(%d), Array(%d), Slice(%d): [%s, %s] -> [%s, %s]"),
+				UE_LOG(LogRDG, Display, TEXT("\t\tMip(%d), Array(%d), Slice(%d): [%s, %s] -> [%s, %s]"),
 					Subresource.MipIndex, Subresource.ArraySlice, Subresource.PlaneSlice,
 					*GetRHIAccessName(Transition.AccessBefore),
 					*GetRHIPipelineName(Batch.DebugPipelinesToBegin),
@@ -1131,28 +1151,26 @@ void FRDGBarrierValidation::ValidateBarrierBatchBegin(const FRDGPass* Pass, cons
 		}
 	}
 
-	if (bAllowedForPass)
+	for (const auto& Pair : ResourceMap->Buffers)
 	{
-		for (const auto& Pair : ResourceMap->Buffers)
-		{
-			FRDGBufferRef Buffer = Pair.Key;
-			const FRHITransitionInfo& Transition = Pair.Value;
+		FRDGBufferRef Buffer = Pair.Key;
+		const FRHITransitionInfo& Transition = Pair.Value;
 
-			if (!IsDebugAllowedForResource(Buffer->Name))
-			{
-				continue;
-			}
+		if (!IsDebugAllowedForResource(Buffer->Name))
+		{
+			continue;
+		}
 
-			LogHeader();
+		LogHeader();
 
-			UE_LOG(LogRDG, Display, TEXT("\t(%p) %s: [%s, %s] -> [%s, %s]"),
-				Buffer,
-				Buffer->Name,
-				*GetRHIAccessName(Transition.AccessBefore),
-				*GetRHIPipelineName(Batch.DebugPipelinesToBegin),
-				*GetRHIAccessName(Transition.AccessAfter),
-				*GetRHIPipelineName(Batch.DebugPipelinesToEnd));
-		}
+		UE_LOG(LogRDG, Display, TEXT("\tRDG(%p) RHI(%p) %s: [%s, %s] -> [%s, %s]"),
+			Buffer,
+			Buffer->GetRHIUnchecked(),
+			Buffer->Name,
+			*GetRHIAccessName(Transition.AccessBefore),
+			*GetRHIPipelineName(Batch.DebugPipelinesToBegin),
+			*GetRHIAccessName(Transition.AccessAfter),
+			*GetRHIPipelineName(Batch.DebugPipelinesToEnd));
 	}
 }
 
@@ -1163,7 +1181,9 @@ void FRDGBarrierValidation::ValidateBarrierBatchEnd(const FRDGPass* Pass, const
 		return;
 	}
 
-	bool bFoundFirstBatch = false;
+	const bool bAllowedForPass = IsDebugAllowedForGraph(GraphName) && IsDebugAllowedForPass(Pass->GetName());
+
+	bool bFoundFirst = false;
 
 	for (const FRDGBarrierBatchBegin* Dependent : Batch.Dependencies)
 	{
@@ -1186,20 +1206,21 @@ void FRDGBarrierValidation::ValidateBarrierBatchEnd(const FRDGPass* Pass, const
 			ResourceMap.Buffers.GetKeys(Buffers);
 		}
 
-		if (Textures.Num() || Buffers.Num())
+		const auto LogHeader = [&]()
 		{
-			if (!bFoundFirstBatch)
+			if (!bFoundFirst)
 			{
+				bFoundFirst = true;
 				UE_LOG(LogRDG, Display, TEXT("[%s(Index: %d, Pipeline: %s) %s] (End):"), Pass->GetName(), Pass->GetHandle().GetIndex(), Dependent->DebugName, *GetRHIPipelineName(Pass->GetPipeline()));
-				bFoundFirstBatch = true;
 			}
-		}
+		};
 
 		for (FRDGTextureRef Texture : Textures)
 		{
 			if (IsDebugAllowedForResource(Texture->Name))
 			{
-				UE_LOG(LogRDG, Display, TEXT("\t(%p) %s"), Texture, Texture->Name);
+				LogHeader();
+				UE_LOG(LogRDG, Display, TEXT("\tRDG(%p) RHI(%p) %s - End:"), Texture, Texture->GetRHIUnchecked(), Texture->Name);
 			}
 		}
 
@@ -1207,7 +1228,23 @@ void FRDGBarrierValidation::ValidateBarrierBatchEnd(const FRDGPass* Pass, const
 		{
 			if (IsDebugAllowedForResource(Buffer->Name))
 			{
-				UE_LOG(LogRDG, Display, TEXT("\t(%p) %s"), Buffer, Buffer->Name);
+				LogHeader();
+				UE_LOG(LogRDG, Display, TEXT("\tRDG(%p) RHI(%p) %s - End"), Buffer, Buffer->GetRHIUnchecked(), Buffer->Name);
+			}
+		}
+
+		for (const auto& KeyValue : ResourceMap.Aliases)
+		{
+			const FRHITransientAliasingInfo& Info = KeyValue.Value;
+			if (Info.IsDiscard())
+			{
+				FRDGParentResourceRef Resource = KeyValue.Key;
+
+				if (IsDebugAllowedForResource(Resource->Name))
+				{
+					LogHeader();
+					UE_LOG(LogRDG, Display, TEXT("\tRDG(%p) RHI(%p) %s - Discard"), Resource, Resource->GetRHIUnchecked(), Resource->Name);
+				}
 			}
 		}
 	}
@@ -1335,7 +1372,7 @@ void FRDGLogFile::Begin(
 {
 	if (GRDGDumpGraph)
 	{
-		if (GRDGImmediateMode)
+		if (IsImmediateMode())
 		{
 			UE_LOG(LogRDG, Warning, TEXT("Dump graph (%d) requested, but immediate mode is enabled. Skipping."), GRDGDumpGraph);
 			return;
@@ -1596,7 +1633,7 @@ bool FRDGLogFile::IncludeTransitionEdgeInGraph(FRDGPassHandle PassBefore, FRDGPa
 
 void FRDGLogFile::AddFirstEdge(const FRDGTextureRef Texture, FRDGPassHandle FirstPass)
 {
-	if (GRDGDumpGraph == RDG_DUMP_GRAPH_RESOURCES && bOpen && IncludeTransitionEdgeInGraph(FirstPass))
+	if (GRDGDumpGraph == RDG_DUMP_GRAPH_RESOURCES && bOpen && IncludeTransitionEdgeInGraph(FirstPass) && IsDebugAllowedForResource(Texture->Name))
 	{
 		AddLine(FString::Printf(TEXT("\"%s\" -> \"%s\" [%s]"),
 			*GetNodeName(Texture),
@@ -1607,7 +1644,7 @@ void FRDGLogFile::AddFirstEdge(const FRDGTextureRef Texture, FRDGPassHandle Firs
 
 void FRDGLogFile::AddFirstEdge(const FRDGBufferRef Buffer, FRDGPassHandle FirstPass)
 {
-	if (GRDGDumpGraph == RDG_DUMP_GRAPH_RESOURCES && bOpen && IncludeTransitionEdgeInGraph(FirstPass))
+	if (GRDGDumpGraph == RDG_DUMP_GRAPH_RESOURCES && bOpen && IncludeTransitionEdgeInGraph(FirstPass) && IsDebugAllowedForResource(Buffer->Name))
 	{
 		AddLine(FString::Printf(TEXT("\"%s\" -> \"%s\" [%s]"),
 			*GetNodeName(Buffer),
@@ -1619,7 +1656,7 @@ void FRDGLogFile::AddFirstEdge(const FRDGBufferRef Buffer, FRDGPassHandle FirstP
 
 void FRDGLogFile::AddAliasEdge(const FRDGTextureRef TextureBefore, FRDGPassHandle BeforePass, const FRDGTextureRef TextureAfter, FRDGPassHandle AfterPass)
 {
-	if (GRDGDumpGraph == RDG_DUMP_GRAPH_RESOURCES && bOpen && IncludeTransitionEdgeInGraph(BeforePass, AfterPass))
+	if (GRDGDumpGraph == RDG_DUMP_GRAPH_RESOURCES && bOpen && IncludeTransitionEdgeInGraph(BeforePass, AfterPass) && IsDebugAllowedForResource(TextureBefore->Name) && IsDebugAllowedForResource(TextureAfter->Name))
 	{
 		AddLine(FString::Printf(TEXT("\"%s\" -> \"%s\" [%s, label=<Alias: <b>%s -&gt; %s</b>>]"),
 			*GetProducerName(BeforePass),
@@ -1632,7 +1669,7 @@ void FRDGLogFile::AddAliasEdge(const FRDGTextureRef TextureBefore, FRDGPassHandl
 
 void FRDGLogFile::AddAliasEdge(const FRDGBufferRef BufferBefore, FRDGPassHandle BeforePass, const FRDGBufferRef BufferAfter, FRDGPassHandle AfterPass)
 {
-	if (GRDGDumpGraph == RDG_DUMP_GRAPH_RESOURCES && bOpen && IncludeTransitionEdgeInGraph(BeforePass, AfterPass))
+	if (GRDGDumpGraph == RDG_DUMP_GRAPH_RESOURCES && bOpen && IncludeTransitionEdgeInGraph(BeforePass, AfterPass) && IsDebugAllowedForResource(BufferBefore->Name) && IsDebugAllowedForResource(BufferAfter->Name))
 	{
 		AddLine(FString::Printf(TEXT("\"%s\" -> \"%s\" [%s, label=<Alias: <b>%s -&gt; %s</b>>]"),
 			*GetProducerName(BeforePass),
@@ -1645,7 +1682,7 @@ void FRDGLogFile::AddAliasEdge(const FRDGBufferRef BufferBefore, FRDGPassHandle
 
 void FRDGLogFile::AddTransitionEdge(FRDGPassHandle PassHandle, FRDGSubresourceState StateBefore, FRDGSubresourceState StateAfter, const FRDGTextureRef Texture)
 {
-	if (GRDGDumpGraph == RDG_DUMP_GRAPH_RESOURCES && bOpen && IncludeTransitionEdgeInGraph(StateBefore.GetFirstPass(), PassHandle))
+	if (GRDGDumpGraph == RDG_DUMP_GRAPH_RESOURCES && bOpen && IncludeTransitionEdgeInGraph(StateBefore.GetFirstPass(), PassHandle) && IsDebugAllowedForResource(Texture->Name))
 	{
 		if (FRDGSubresourceState::IsTransitionRequired(StateBefore, StateAfter))
 		{
@@ -1671,7 +1708,7 @@ void FRDGLogFile::AddTransitionEdge(FRDGPassHandle PassHandle, FRDGSubresourceSt
 
 void FRDGLogFile::AddTransitionEdge(FRDGPassHandle PassHandle, FRDGSubresourceState StateBefore, FRDGSubresourceState StateAfter, const FRDGTextureRef Texture, FRDGTextureSubresource Subresource)
 {
-	if (GRDGDumpGraph == RDG_DUMP_GRAPH_RESOURCES && bOpen && IncludeTransitionEdgeInGraph(StateBefore.GetFirstPass(), PassHandle))
+	if (GRDGDumpGraph == RDG_DUMP_GRAPH_RESOURCES && bOpen && IncludeTransitionEdgeInGraph(StateBefore.GetFirstPass(), PassHandle) && IsDebugAllowedForResource(Texture->Name))
 	{
 		if (FRDGSubresourceState::IsTransitionRequired(StateBefore, StateAfter))
 		{
@@ -1699,7 +1736,7 @@ void FRDGLogFile::AddTransitionEdge(FRDGPassHandle PassHandle, FRDGSubresourceSt
 
 void FRDGLogFile::AddTransitionEdge(FRDGPassHandle PassHandle, FRDGSubresourceState StateBefore, FRDGSubresourceState StateAfter, const FRDGBufferRef Buffer)
 {
-	if (GRDGDumpGraph == RDG_DUMP_GRAPH_RESOURCES && bOpen && IncludeTransitionEdgeInGraph(StateBefore.GetFirstPass(), PassHandle))
+	if (GRDGDumpGraph == RDG_DUMP_GRAPH_RESOURCES && bOpen && IncludeTransitionEdgeInGraph(StateBefore.GetFirstPass(), PassHandle) && IsDebugAllowedForResource(Buffer->Name))
 	{
 		if (FRDGSubresourceState::IsTransitionRequired(StateBefore, StateAfter))
 		{
diff --git a/Engine/Source/Runtime/RenderCore/Private/RenderTargetPool.cpp b/Engine/Source/Runtime/RenderCore/Private/RenderTargetPool.cpp
--- a/Engine/Source/Runtime/RenderCore/Private/RenderTargetPool.cpp
+++ b/Engine/Source/Runtime/RenderCore/Private/RenderTargetPool.cpp
@@ -263,15 +263,14 @@ TRefCountPtr<FPooledRenderTarget> FRenderTargetPool::FindFreeElementForRDG(
 	const FRDGTextureDesc& Desc,
 	const TCHAR* Name)
 {
-	const bool bDoAcquireTransientTexture = false;
-	return FindFreeElementInternal(RHICmdList, Translate(Desc), Name, bDoAcquireTransientTexture);
+	checkf(!EnumHasAnyFlags(Desc.Flags, TexCreate_Transient), TEXT("RDG does not support the render target pool transient API"));
+	return FindFreeElementInternal(RHICmdList, Translate(Desc), Name);
 }
 
 TRefCountPtr<FPooledRenderTarget> FRenderTargetPool::FindFreeElementInternal(
 	FRHICommandList& RHICmdList,
 	const FPooledRenderTargetDesc& Desc,
-	const TCHAR* InDebugName,
-	bool bDoAcquireTransientTexture)
+	const TCHAR* InDebugName)
 {
 	const int32 AliasingMode = CVarRtPoolTransientMode.GetValueOnRenderThread();
 	FPooledRenderTarget* Found = 0;
@@ -525,13 +524,10 @@ Done:
 
 	check(!Found->IsFree());
 
-	if (bDoAcquireTransientTexture)
+	// Only referenced by the pool, map the physical pages
+	if (Found->IsTransient() && OriginalNumRefs == 1 && Found->GetRenderTargetItem().TargetableTexture != nullptr)
 	{
-		// Only referenced by the pool, map the physical pages
-		if (Found->IsTransient() && OriginalNumRefs == 1 && Found->GetRenderTargetItem().TargetableTexture != nullptr)
-		{
-			RHIAcquireTransientResource(Found->GetRenderTargetItem().TargetableTexture);
-		}
+		RHIAcquireTransientResource(Found->GetRenderTargetItem().TargetableTexture);
 	}
 
 	// Transient RTs have to be targettable
@@ -621,8 +617,7 @@ bool FRenderTargetPool::FindFreeElement(
 		}
 	}
 
-	const bool bDoAcquireTransientResource = true;
-	TRefCountPtr<FPooledRenderTarget> Result = FindFreeElementInternal(RHICmdList, Desc, InDebugName, bDoAcquireTransientResource);
+	TRefCountPtr<FPooledRenderTarget> Result = FindFreeElementInternal(RHICmdList, Desc, InDebugName);
 
 	// Reset RDG state back to an unknown default. The resource is being handed off to a user outside of RDG, so the state is no longer valid.
 	{
@@ -1202,12 +1197,12 @@ void FPooledRenderTarget::InitRDG()
 
 	if (RenderTargetItem.TargetableTexture)
 	{
-		TargetableTexture = new FRDGPooledTexture(RenderTargetItem.TargetableTexture, Translate(Desc, ERenderTargetTexture::Targetable), RenderTargetItem.UAV);
+		TargetableTexture = FRDGPooledTexture::CreateCommitted(RenderTargetItem.TargetableTexture, Translate(Desc, ERenderTargetTexture::Targetable), RenderTargetItem.UAV);
 	}
 
 	if (RenderTargetItem.ShaderResourceTexture != RenderTargetItem.TargetableTexture)
 	{
-		ShaderResourceTexture = new FRDGPooledTexture(RenderTargetItem.ShaderResourceTexture, Translate(Desc, ERenderTargetTexture::ShaderResource), nullptr);
+		ShaderResourceTexture = FRDGPooledTexture::CreateCommitted(RenderTargetItem.ShaderResourceTexture, Translate(Desc, ERenderTargetTexture::ShaderResource), nullptr);
 	}
 	else
 	{
@@ -1232,12 +1227,11 @@ uint32 FPooledRenderTarget::Release()
 	}
 	else if (Refs == 1 && RenderTargetPool && IsTransient())
 	{
-		if (bAutoDiscard && RenderTargetItem.TargetableTexture)
+		if (RenderTargetItem.TargetableTexture)
 		{
 			RHIDiscardTransientResource(RenderTargetItem.TargetableTexture);
 		}
 		FrameNumberLastDiscard = GFrameNumberRenderThread;
-		bAutoDiscard = true;
 	}
 	return Refs;
 }
diff --git a/Engine/Source/Runtime/RenderCore/Public/RenderGraphBuilder.h b/Engine/Source/Runtime/RenderCore/Public/RenderGraphBuilder.h
--- a/Engine/Source/Runtime/RenderCore/Public/RenderGraphBuilder.h
+++ b/Engine/Source/Runtime/RenderCore/Public/RenderGraphBuilder.h
@@ -240,19 +240,109 @@ private:
 		ERDGPassFlags Flags,
 		ExecuteLambdaType&& ExecuteLambda);
 
-	ERDGPassFlags OverridePassFlags(const TCHAR* PassName, ERDGPassFlags Flags, bool bAsyncComputeSupported);
+	static ERDGPassFlags OverridePassFlags(const TCHAR* PassName, ERDGPassFlags Flags, bool bAsyncComputeSupported);
 
+	/** Returns the graph prologue pass handle. */
 	FORCEINLINE FRDGPassHandle GetProloguePassHandle() const
 	{
 		return Passes.Begin();
 	}
 
+	/** Returns the graph epilogue pass handle. */
 	FORCEINLINE FRDGPassHandle GetEpiloguePassHandle() const
 	{
 		checkf(EpiloguePass, TEXT("The handle is not valid until the epilogue has been added to the graph during execution."));
 		return Passes.Last();
 	}
 
+	/** Barrier location controls where the barrier is 'Ended' relative to the pass lambda being executed.
+	 *  Most barrier locations are done in the prologue prior to the executing lambda. But certain cases
+	 *  like an aliasing discard operation need to be done *after* the pass being invoked. Therefore, when
+	 *  adding a transition the user can specify where to place the barrier.
+	 */
+	enum class EBarrierLocation
+	{
+		/** The barrier occurs in the prologue of the pass (before execution). */
+		Prologue,
+
+		/** The barrier occurs in the epilogue of the pass (after execution). */
+		Epilogue
+	};
+
+	/** Prologue and Epilogue barrier passes are used to plan transitions around RHI render pass merging,
+	 *  as it is illegal to issue a barrier during a render pass. If passes [A, B, C] are merged together,
+	 *  'A' becomes 'B's prologue pass and 'C' becomes 'A's epilogue pass. This way, any transitions that
+	 *  need to happen before the merged pass (i.e. in the prologue) are done in A. Any transitions after
+	 *  the render pass merge are done in C.
+	 */
+	FRDGPassHandle GetEpilogueBarrierPassHandle(FRDGPassHandle Handle)
+	{
+		return Passes[Handle]->EpilogueBarrierPass;
+	}
+
+	FRDGPassHandle GetPrologueBarrierPassHandle(FRDGPassHandle Handle)
+	{
+		return Passes[Handle]->PrologueBarrierPass;
+	}
+
+	FRDGPass* GetEpilogueBarrierPass(FRDGPassHandle Handle)
+	{
+		return Passes[GetEpilogueBarrierPassHandle(Handle)];
+	}
+
+	FRDGPass* GetPrologueBarrierPass(FRDGPassHandle Handle)
+	{
+		return Passes[GetPrologueBarrierPassHandle(Handle)];
+	}
+
+	/** Ends the barrier batch in the prologue of the provided pass. */
+	void AddToPrologueBarriersToEnd(FRDGPassHandle Handle, FRDGBarrierBatchBegin& BarriersToBegin)
+	{
+		FRDGPass* Pass = GetPrologueBarrierPass(Handle);
+		Pass->GetPrologueBarriersToEnd(Allocator).AddDependency(&BarriersToBegin);
+	}
+
+	/** Ends the barrier batch in the epilogue of the provided pass. */
+	void AddToEpilogueBarriersToEnd(FRDGPassHandle Handle, FRDGBarrierBatchBegin& BarriersToBegin)
+	{
+		FRDGPass* Pass = GetEpilogueBarrierPass(Handle);
+		Pass->GetEpilogueBarriersToEnd(Allocator).AddDependency(&BarriersToBegin);
+	}
+
+	/** Utility function to add an immediate barrier dependency in the prologue of the provided pass. */
+	template <typename FunctionType>
+	void AddToPrologueBarriers(FRDGPassHandle PassHandle, FunctionType Function)
+	{
+		FRDGPass* Pass = GetPrologueBarrierPass(PassHandle);
+		FRDGBarrierBatchBegin& BarriersToBegin = Pass->GetPrologueBarriersToBegin(Allocator);
+		Function(BarriersToBegin);
+		Pass->GetPrologueBarriersToEnd(Allocator).AddDependency(&BarriersToBegin);
+	}
+
+	/** Utility function to add an immediate barrier dependency in the epilogue of the provided pass. */
+	template <typename FunctionType>
+	void AddToEpilogueBarriers(FRDGPassHandle PassHandle, FunctionType Function)
+	{
+		FRDGPass* Pass = GetEpilogueBarrierPass(PassHandle);
+		FRDGBarrierBatchBegin& BarriersToBegin = Pass->GetEpilogueBarriersToBeginFor(Allocator, Pass->GetPipeline());
+		Function(BarriersToBegin);
+		Pass->GetEpilogueBarriersToEnd(Allocator).AddDependency(&BarriersToBegin);
+	}
+
+	/** Utility function to add an immediate barrier dependency in either the prologue or epilogue of the provided pass, depending on BarrierLocation. */
+	template <typename FunctionType>
+	void AddToBarriers(FRDGPassHandle PassHandle, EBarrierLocation BarrierLocation, FunctionType Function)
+	{
+		if (BarrierLocation == EBarrierLocation::Prologue)
+		{
+			AddToPrologueBarriers(PassHandle, Function);
+		}
+		else
+		{
+			AddToEpilogueBarriers(PassHandle, Function);
+		}
+	}
+
 	/** Registry of graph objects. */
 	FRDGPassRegistry Passes;
 	FRDGTextureRegistry Textures;
@@ -314,21 +404,43 @@ private:
 	bool bWaitedForTemporalEffect = false;
 #endif
 
+	/** Tracks whether all passes / resources have been added to the graph. */
+	bool bSetupComplete = false;
+
 	IF_RDG_CMDLIST_STATS(TStatId CommandListStat);
 
+	class FTransientResourceAllocator
+	{
+	public:
+		FTransientResourceAllocator(FRHICommandListImmediate& InRHICmdList)
+			: RHICmdList(InRHICmdList)
+		{}
+
+		~FTransientResourceAllocator();
+
+		IRHITransientResourceAllocator* GetOrCreate();
+		IRHITransientResourceAllocator* Get() const { return Allocator; }
+		IRHITransientResourceAllocator* operator->() const { check(Allocator);  return Allocator; }
+		operator bool() const { return Allocator != nullptr; }
+
+	private:
+		FRHICommandListImmediate& RHICmdList;
+		IRHITransientResourceAllocator* Allocator = nullptr;
+	};
+
 	void Compile();
 	void Clear();
 
 	void BeginResourceRHI(FRDGUniformBuffer* UniformBuffer);
-	void BeginResourceRHI(FRDGPassHandle, FRDGTexture* Texture);
-	void BeginResourceRHI(FRDGPassHandle, FRDGTextureSRV* SRV);
-	void BeginResourceRHI(FRDGPassHandle, FRDGTextureUAV* UAV);
-	void BeginResourceRHI(FRDGPassHandle, FRDGBuffer* Buffer);
-	void BeginResourceRHI(FRDGPassHandle, FRDGBufferSRV* SRV);
-	void BeginResourceRHI(FRDGPassHandle, FRDGBufferUAV* UAV);
+	void BeginResourceRHI(FTransientResourceAllocator*, FRDGPassHandle, FRDGTexture* Texture);
+	void BeginResourceRHI(FTransientResourceAllocator*, FRDGPassHandle, FRDGTextureSRV* SRV);
+	void BeginResourceRHI(FTransientResourceAllocator*, FRDGPassHandle, FRDGTextureUAV* UAV);
+	void BeginResourceRHI(FTransientResourceAllocator*, FRDGPassHandle, FRDGBuffer* Buffer);
+	void BeginResourceRHI(FTransientResourceAllocator*, FRDGPassHandle, FRDGBufferSRV* SRV);
+	void BeginResourceRHI(FTransientResourceAllocator*, FRDGPassHandle, FRDGBufferUAV* UAV);
 
-	void EndResourceRHI(FRDGPassHandle, FRDGTexture* Texture, uint32 ReferenceCount);
-	void EndResourceRHI(FRDGPassHandle, FRDGBuffer* Buffer, uint32 ReferenceCount);
+	void EndResourceRHI(FTransientResourceAllocator&, FRDGPassHandle, FRDGTexture* Texture, uint32 ReferenceCount);
+	void EndResourceRHI(FTransientResourceAllocator&, FRDGPassHandle, FRDGBuffer* Buffer, uint32 ReferenceCount);
 
 	void SetupPassInternal(FRDGPass* Pass, FRDGPassHandle PassHandle, ERHIPipeline PassPipeline);
 	void SetupPass(FRDGPass* Pass);
@@ -338,7 +450,7 @@ private:
 	void ExecutePassPrologue(FRHIComputeCommandList& RHICmdListPass, FRDGPass* Pass);
 	void ExecutePassEpilogue(FRHIComputeCommandList& RHICmdListPass, FRDGPass* Pass);
 
-	void CollectPassResources(FRDGPassHandle PassHandle);
+	void CollectPassResources(FTransientResourceAllocator&, FRDGPassHandle PassHandle);
 	void CollectPassBarriers(FRDGPassHandle PassHandle);
 
 	void AddPassDependency(FRDGPassHandle ProducerHandle, FRDGPassHandle ConsumerHandle);
@@ -346,15 +458,34 @@ private:
 	void AddEpilogueTransition(FRDGTextureRef Texture);
 	void AddEpilogueTransition(FRDGBufferRef Buffer);
 
-	void AddTransition(FRDGPassHandle PassHandle, FRDGTextureRef Texture, const FRDGTextureTransientSubresourceStateIndirect& StateAfter);
-	void AddTransition(FRDGPassHandle PassHandle, FRDGBufferRef Buffer, FRDGSubresourceState StateAfter);
+	void AddTransition(
+		FRDGPassHandle PassHandle,
+		FRDGTextureRef Texture,
+		const FRDGTextureTransientSubresourceStateIndirect& StateAfter,
+		EBarrierLocation BarrierLocation = EBarrierLocation::Prologue);
+
+	void AddTransition(
+		FRDGPassHandle PassHandle,
+		FRDGBufferRef Buffer,
+		FRDGSubresourceState StateAfter,
+		EBarrierLocation BarrierLocation = EBarrierLocation::Prologue);
 
-	void AddTransitionInternal(
+	void AddTransition(
 		FRDGParentResource* Resource,
 		FRDGSubresourceState StateBefore,
 		FRDGSubresourceState StateAfter,
+		EBarrierLocation BarrierLocation,
 		const FRHITransitionInfo& TransitionInfo);
 
+	bool IsTransient(FRDGTextureRef Texture) const;
+	bool IsTransient(FRDGBufferRef Buffer) const;
+	bool IsTransientInternal(FRDGParentResourceRef Resource) const;
+
+	void AddAcquireResourceRHI(FRDGTexture* Texture, FRDGPassHandle PassHandle);
+	void AddAcquireResourceRHI(FRDGBuffer* Buffer, FRDGPassHandle PassHandle);
+	void AddDiscardResourceRHI(FRDGTexture* Texture, FRDGPassHandle PassHandle);
+	void AddDiscardResourceRHI(FRDGBuffer* Buffer, FRDGPassHandle PassHandle);
+
 	FRHIRenderPassInfo GetRenderPassInfo(const FRDGPass* Pass) const;
 
 	FRDGSubresourceState* AllocSubresource(const FRDGSubresourceState& Other);
diff --git a/Engine/Source/Runtime/RenderCore/Public/RenderGraphBuilder.inl b/Engine/Source/Runtime/RenderCore/Public/RenderGraphBuilder.inl
--- a/Engine/Source/Runtime/RenderCore/Public/RenderGraphBuilder.inl
+++ b/Engine/Source/Runtime/RenderCore/Public/RenderGraphBuilder.inl
@@ -20,13 +20,33 @@ inline FRDGTextureRef FRDGBuilder::FindExternalTexture(IPooledRenderTarget* Exte
 	return nullptr;
 }
 
+inline FRDGTextureRef FRDGBuilder::CreateTexture(
+	const FRDGTextureDesc& Desc,
+	const TCHAR* Name,
+	ERDGTextureFlags Flags)
+{
+	// RDG no longer supports the legacy transient resource API.
+	FRDGTextureDesc OverrideDesc = Desc;
+	EnumRemoveFlags(OverrideDesc.Flags, TexCreate_Transient);
+
+	IF_RDG_ENABLE_DEBUG(UserValidation.ValidateCreateTexture(OverrideDesc, Name, Flags));
+	FRDGTextureRef Texture = Textures.Allocate(Allocator, Name, OverrideDesc, Flags, ERenderTargetTexture::ShaderResource);
+	IF_RDG_ENABLE_DEBUG(UserValidation.ValidateCreateTexture(Texture));
+	IF_RDG_ENABLE_TRACE(Trace.AddResource(Texture));
+	return Texture;
+}
+
 inline FRDGBufferRef FRDGBuilder::CreateBuffer(
 	const FRDGBufferDesc& Desc,
 	const TCHAR* Name,
 	ERDGBufferFlags Flags)
 {
+	// RDG no longer supports the legacy transient resource API.
+	FRDGBufferDesc OverrideDesc = Desc;
+	EnumRemoveFlags(OverrideDesc.Usage, BUF_Transient);
+
 	IF_RDG_ENABLE_DEBUG(UserValidation.ValidateCreateBuffer(Desc, Name, Flags));
-	FRDGBufferRef Buffer = Buffers.Allocate(Allocator, Name, Desc, Flags);
+	FRDGBufferRef Buffer = Buffers.Allocate(Allocator, Name, OverrideDesc, Flags);
 	IF_RDG_ENABLE_DEBUG(UserValidation.ValidateCreateBuffer(Buffer));
 	IF_RDG_ENABLE_TRACE(Trace.AddResource(Buffer));
 	return Buffer;
@@ -130,7 +150,7 @@ FRDGPassRef FRDGBuilder::AddPassInternal(
 
 	IF_RDG_ENABLE_DEBUG(UserValidation.ValidateAddPass(ParameterStruct, ParametersMetadata, Name, Flags));
 
-	FRDGPass* Pass = Allocator.Alloc<LambdaPassType>(
+	FRDGPass* Pass = Allocator.AllocNoDestruct<LambdaPassType>(
 		MoveTemp(Name),
 		ParametersMetadata,
 		ParameterStruct,
diff --git a/Engine/Source/Runtime/RenderCore/Public/RenderGraphDefinitions.h b/Engine/Source/Runtime/RenderCore/Public/RenderGraphDefinitions.h
--- a/Engine/Source/Runtime/RenderCore/Public/RenderGraphDefinitions.h
+++ b/Engine/Source/Runtime/RenderCore/Public/RenderGraphDefinitions.h
@@ -290,13 +290,17 @@ public:
 	DerivedType* Allocate(FRDGAllocator& Allocator, TArgs&&... Args)
 	{
 		static_assert(TIsDerivedFrom<DerivedType, ObjectType>::Value, "You must specify a type that derives from ObjectType");
-		DerivedType* Object = Allocator.Alloc<DerivedType>(Forward<TArgs>(Args)...);
+		DerivedType* Object = Allocator.AllocNoDestruct<DerivedType>(Forward<TArgs>(Args)...);
 		Insert(Object);
 		return Object;
 	}
 
 	void Clear()
 	{
+		for (int32 Index = Array.Num() - 1; Index >= 0; --Index)
+		{
+			Array[Index]->~ObjectType();
+		}
 		Array.Empty();
 	}
 
@@ -415,7 +419,7 @@ const TRDGHandle<ObjectType, IndexType> TRDGHandle<ObjectType, IndexType>::Null;
 
 /** FORWARD DECLARATIONS */
 
-struct FRDGTextureDesc;
+using FRDGTextureDesc = FRHITextureCreateInfo;
 
 class FRDGBlackboard;
 
diff --git a/Engine/Source/Runtime/RenderCore/Public/RenderGraphPass.h b/Engine/Source/Runtime/RenderCore/Public/RenderGraphPass.h
--- a/Engine/Source/Runtime/RenderCore/Public/RenderGraphPass.h
+++ b/Engine/Source/Runtime/RenderCore/Public/RenderGraphPass.h
@@ -12,13 +12,12 @@ public:
 	FRDGTransitionQueue() = default;
 	FRDGTransitionQueue(uint32 ReservedCount);
 
-	void Insert(const FRHITransition* Transition, ERHICreateTransitionFlags TransitionFlags);
+	void Insert(const FRHITransition* Transition, ERHITransitionCreateFlags TransitionFlags);
 	void Begin(FRHIComputeCommandList& RHICmdList);
 	void End(FRHIComputeCommandList& RHICmdList);
 
 private:
 	TArray<const FRHITransition*, TInlineAllocator<1, FRDGArrayAllocator>> Queue;
-	TArray<const FRHITransition*, TInlineAllocator<1, FRDGArrayAllocator>> QueueWithFences;
 };
 
 struct FRDGBarrierBatchBeginId
@@ -49,9 +48,11 @@ public:
 
 	void AddTransition(FRDGParentResourceRef Resource, const FRHITransitionInfo& Info);
 
+	void AddAlias(FRDGParentResourceRef Resource, const FRHITransientAliasingInfo& Info);
+
 	void SetUseCrossPipelineFence()
 	{
-		TransitionFlags = ERHICreateTransitionFlags::None;
+		TransitionFlags = ERHITransitionCreateFlags::None;
 		bTransitionNeeded = true;
 	}
 
@@ -61,7 +62,8 @@ public:
 private:
 	const FRHITransition* Transition = nullptr;
 	TArray<FRHITransitionInfo, TInlineAllocator<1, FRDGArrayAllocator>> Transitions;
-	ERHICreateTransitionFlags TransitionFlags = ERHICreateTransitionFlags::NoFence;
+	TArray<FRHITransientAliasingInfo, FRDGArrayAllocator> Aliases;
+	ERHITransitionCreateFlags TransitionFlags = ERHITransitionCreateFlags::NoFence;
 	bool bTransitionNeeded = false;
 
 	/** These pipelines masks are set at creation time and reset with each submission. */
@@ -70,12 +72,14 @@ private:
 
 #if RDG_ENABLE_DEBUG
 	FRDGPassHandlesByPipeline DebugPasses;
-	TArray<FRDGParentResource*, FRDGArrayAllocator> DebugResources;
+	TArray<FRDGParentResource*, FRDGArrayAllocator> DebugTransitionResources;
+	TArray<FRDGParentResource*, FRDGArrayAllocator> DebugAliasingResources;
 	const TCHAR* DebugName;
 	ERHIPipeline DebugPipelinesToBegin;
 	ERHIPipeline DebugPipelinesToEnd;
 #endif
 
+	friend class FRDGBarrierBatchBegin;
 	friend class FRDGBarrierBatchEnd;
 	friend class FRDGBarrierValidation;
 };
@@ -299,8 +303,8 @@ protected:
 			uint32 bGraphicsFork : 1;
 			uint32 bGraphicsJoin : 1;
 
-			/** Whether the pass writes to a UAV. */
-			uint32 bUAVAccess : 1;
+			/** Whether the pass only writes to resources in its render pass. */
+			uint32 bRenderPassOnlyWrites : 1;
 
 			/** Whether the pass uses the immediate command list. */
 			uint32 bImmediateCommandList : 1;
@@ -353,15 +357,6 @@ protected:
 	TArray<FRDGPass*, TInlineAllocator<1, FRDGArrayAllocator>> ResourcesToBegin;
 	TArray<FRDGPass*, TInlineAllocator<1, FRDGArrayAllocator>> ResourcesToEnd;
 
-	/** List of textures to acquire *after* the pass completes, *before* discards. Acquires apply to all allocated textures. */
-	TArray<FRHITexture*, FRDGArrayAllocator> TexturesToAcquire;
-
-	/** List of textures to discard *after* the pass completes, *after* acquires. Discards only apply to textures marked as
-	 *  transient and the last alias of the texture uses the automatic discard behavior (in order to support cleaner hand-off
-	 *  to the user or back to the pool.
-	 */
-	TArray<FRHITexture*, FRDGArrayAllocator> TexturesToDiscard;
-
 	/** Split-barrier batches at various points of execution of the pass. */
 	FRDGBarrierBatchBegin* PrologueBarriersToBegin = nullptr;
 	FRDGBarrierBatchEnd* PrologueBarriersToEnd = nullptr;
diff --git a/Engine/Source/Runtime/RenderCore/Public/RenderGraphResources.h b/Engine/Source/Runtime/RenderCore/Public/RenderGraphResources.h
--- a/Engine/Source/Runtime/RenderCore/Public/RenderGraphResources.h
+++ b/Engine/Source/Runtime/RenderCore/Public/RenderGraphResources.h
@@ -129,6 +129,7 @@ private:
 
 	friend FRDGBuilder;
 	friend FRDGUserValidation;
+	friend FRDGBarrierValidation;
 };
 
 class FRDGUniformBuffer
@@ -319,234 +320,95 @@ private:
 	friend FRDGAllocator;
 };
 
-/** Descriptor used to create a render graph texture. */
-struct RENDERCORE_API FRDGTextureDesc
-{
-	static FRDGTextureDesc Create2D(
-		FIntPoint InExtent,
-		EPixelFormat InFormat,
-		FClearValueBinding InClearValue,
-		ETextureCreateFlags InFlags,
-		uint8 InNumMips = 1,
-		uint8 InNumSamples = 1)
-	{
-		return FRDGTextureDesc(InClearValue, ETextureDimension::Texture2D, InFlags, InFormat, InExtent, 1, 1, InNumMips, InNumSamples);
-	}
-
-	static FRDGTextureDesc Create2DArray(
-		FIntPoint InExtent,
-		EPixelFormat InFormat,
-		FClearValueBinding InClearValue,
-		ETextureCreateFlags InFlags,
-		uint32 InArraySize,
-		uint8 InNumMips = 1,
-		uint8 InNumSamples = 1)
-	{
-		return FRDGTextureDesc(InClearValue, ETextureDimension::Texture2DArray, InFlags, InFormat, InExtent, 1, InArraySize, InNumMips, InNumSamples);
-	}
-
-	static FRDGTextureDesc Create3D(
-		FIntVector InSize,
-		EPixelFormat InFormat,
-		FClearValueBinding InClearValue,
-		ETextureCreateFlags InFlags,
-		uint8 InNumMips = 1,
-		uint8 InNumSamples = 1)
-	{
-		return FRDGTextureDesc(InClearValue, ETextureDimension::Texture3D, InFlags, InFormat, FIntPoint(InSize.X, InSize.Y), InSize.Z, 1, InNumMips, InNumSamples);
-	}
-
-	static FRDGTextureDesc CreateCube(
-		uint32 InSizeInPixels,
-		EPixelFormat InFormat,
-		FClearValueBinding InClearValue,
-		ETextureCreateFlags InFlags,
-		uint8 InNumMips = 1,
-		uint8 InNumSamples = 1)
-	{
-		return FRDGTextureDesc(InClearValue, ETextureDimension::TextureCube, InFlags, InFormat, FIntPoint(InSizeInPixels, InSizeInPixels), 1, 1, InNumMips, InNumSamples);
-	}
-
-	static FRDGTextureDesc CreateCubeArray(
-		uint32 InSizeInPixels,
-		EPixelFormat InFormat,
-		FClearValueBinding InClearValue,
-		ETextureCreateFlags InFlags,
-		uint32 InArraySize,
-		uint8 InNumMips = 1,
-		uint8 InNumSamples = 1)
-	{
-		return FRDGTextureDesc(InClearValue, ETextureDimension::TextureCubeArray, InFlags, InFormat, FIntPoint(InSizeInPixels, InSizeInPixels), 1, InArraySize, InNumMips, InNumSamples);
-	}
-
-	FRDGTextureDesc() = default;
-	FRDGTextureDesc(
-		FClearValueBinding InClearValue,
-		ETextureDimension InDimension,
-		ETextureCreateFlags InFlags,
-		EPixelFormat InFormat,
-		FIntPoint InExtent,
-		uint16 InDepth = 1,
-		uint16 InArraySize = 1,
-		uint8 InNumMips = 1,
-		uint8 InNumSamples = 1)
-		: ClearValue(InClearValue)
-		, Dimension(InDimension)
-		, Flags(InFlags)
-		, Format(InFormat)
-		, Extent(InExtent)
-		, Depth(InDepth)
-		, ArraySize(InArraySize)
-		, NumMips(InNumMips)
-		, NumSamples(InNumSamples)
-	{}
-
-	bool operator == (const FRDGTextureDesc& Other) const
-	{
-		return ClearValue == Other.ClearValue
-			&& Dimension == Other.Dimension
-			&& Flags == Other.Flags
-			&& Format == Other.Format
-			&& Extent == Other.Extent
-			&& Depth == Other.Depth
-			&& ArraySize == Other.ArraySize
-			&& NumMips == Other.NumMips
-			&& NumSamples == Other.NumSamples;
-	}
-
-	bool operator != (const FRDGTextureDesc& Other) const
-	{
-		return !(*this == Other);
-	}
-
-	void Reset()
-	{
-		// Usually we don't want to propagate MSAA samples.
-		NumSamples = 1;
-
-		// Remove UAV flag for textures that don't need it (some formats are incompatible).
-		Flags |= TexCreate_RenderTargetable;
-		Flags &= ~(TexCreate_UAV | TexCreate_ResolveTargetable | TexCreate_DepthStencilResolveTarget);
-	}
-
-	bool IsTexture2D() const
-	{
-		return Dimension == ETextureDimension::Texture2D || Dimension == ETextureDimension::Texture2DArray;
-	}
-
-	bool IsTexture3D() const
-	{
-		return Dimension == ETextureDimension::Texture3D;
-	}
+/** Translates from a pooled render target descriptor to an RDG texture descriptor. */
+inline FRDGTextureDesc Translate(const FPooledRenderTargetDesc& InDesc, ERenderTargetTexture InTexture = ERenderTargetTexture::Targetable);
 
-	bool IsTextureCube() const
-	{
-		return Dimension == ETextureDimension::TextureCube || Dimension == ETextureDimension::TextureCubeArray;
-	}
+/** Translates from an RDG texture descriptor to a pooled render target descriptor. */
+inline FPooledRenderTargetDesc Translate(const FRDGTextureDesc& InDesc);
 
-	bool IsTextureArray() const
-	{
-		return Dimension == ETextureDimension::Texture2DArray || Dimension == ETextureDimension::TextureCubeArray;
-	}
+class FRDGPooledResource
+{
+public:
+	virtual ~FRDGPooledResource() = default;
 
-	bool IsMipChain() const
+	uint32 GetRefCount() const
 	{
-		return NumMips > 1;
+		return RefCount;
 	}
 
-	bool IsMultisample() const
+	uint32 AddRef() const
 	{
-		return NumSamples > 1;
+		return ++RefCount;
 	}
 
-	FIntVector GetSize() const
+	uint32 Release() const
 	{
-		return FIntVector(Extent.X, Extent.Y, Depth);
+		check(RefCount > 0);
+		uint32 LocalRefCount = --RefCount;
+		if (LocalRefCount == 0)
+		{
+			if (Allocator == EAllocator::Default)
+			{
+				delete this;
+			}
+			else
+			{
+				this->~FRDGPooledResource();
+			}
+		}
+		return LocalRefCount;
 	}
 
-	FRDGTextureSubresourceLayout GetSubresourceLayout() const
+protected:
+	enum class EAllocator
 	{
-		return FRDGTextureSubresourceLayout(NumMips, ArraySize * (IsTextureCube() ? 6 : 1), IsStencilFormat(Format) ? 2 : 1);
-	}
-
-	/** Returns whether this descriptor conforms to requirements. */
-	bool IsValid() const;
-
-	/** Clear value to use when fast-clearing the texture. */
-	FClearValueBinding ClearValue;
+		/** Allocated with the default allocator (new / delete). */
+		Default,
 
-	/** Texture dimension to use when creating the RHI texture. */
-	ETextureDimension Dimension = ETextureDimension::Texture2D;
-
-	/** Texture flags passed on to RHI texture. */
-	ETextureCreateFlags Flags = TexCreate_None;
-
-	/** Pixel format used to create RHI texture. */
-	EPixelFormat Format = PF_Unknown;
-
-	/** Extent of the texture in x and y. */
-	FIntPoint Extent = FIntPoint(1, 1);
-
-	/** Depth of the texture if the dimension is 3D. */
-	uint16 Depth = 1;
+		/** Allocated with the RDG allocator. */
+		RDG
+	};
 
-	/** The number of array elements in the texture. (Keep at 1 if dimension is 3D). */
-	uint16 ArraySize = 1;
+	FRDGPooledResource(EAllocator InAllocator)
+		: Allocator(InAllocator)
+	{}
 
-	/** Number of mips in the texture mip-map chain. */
-	uint8 NumMips = 1;
+	EAllocator Allocator;
 
-	/** Number of samples in the texture. >1 for MSAA. */
-	uint8 NumSamples = 1;
+	mutable uint32 RefCount = 0;
 };
 
-/** Translates from a pooled render target descriptor to an RDG texture descriptor. */
-inline FRDGTextureDesc Translate(const FPooledRenderTargetDesc& InDesc, ERenderTargetTexture InTexture = ERenderTargetTexture::Targetable);
-
-/** Translates from an RDG texture descriptor to a pooled render target descriptor. */
-inline FPooledRenderTargetDesc Translate(const FRDGTextureDesc& InDesc);
-
-class RENDERCORE_API FRDGPooledTexture
+class RENDERCORE_API FRDGPooledTexture final
+	: public FRDGPooledResource
 {
 public:
-	const FRDGTextureDesc Desc;
-
-	FRHITexture* GetRHI() const
-	{
-		return Texture;
-	}
-
-	FRDGTexture* GetOwner() const
+	/** Creates a pooled texture using the default allocator. Ref-counting determines deletion of the memory. Safe to be persistently cached. */
+	static TRefCountPtr<FRDGPooledTexture> CreateCommitted(TRefCountPtr<FRHITexture> Texture, const FRDGTextureSubresourceLayout& Layout, const FUnorderedAccessViewRHIRef& FirstMipUAV = nullptr)
 	{
-		return Owner;
+		return TRefCountPtr<FRDGPooledTexture>(new FRDGPooledTexture(MoveTemp(Texture), Layout, FirstMipUAV, EAllocator::Default));
 	}
 
-	uint32 GetRefCount() const
+	/** Creates a pooled texture using the RDG allocator. All references should be released before the allocator is freed. */
+	static TRefCountPtr<FRDGPooledTexture> CreateTransient(FRDGAllocator& Allocator, const FRDGTextureSubresourceLayout& Layout, TRefCountPtr<FRHITexture> Texture, const FUnorderedAccessViewRHIRef& FirstMipUAV = nullptr)
 	{
-		return RefCount;
+		return TRefCountPtr<FRDGPooledTexture>(Allocator.AllocNoDestruct<FRDGPooledTexture>(MoveTemp(Texture), Layout, FirstMipUAV, EAllocator::RDG));
 	}
 
-	uint32 AddRef() const
+	FRHITexture* GetRHI() const
 	{
-		return ++RefCount;
+		return Texture;
 	}
 
-	uint32 Release() const
+	FRDGTexture* GetOwner() const
 	{
-		check(RefCount > 0);
-		if (--RefCount == 0)
-		{
-			delete this;
-			return 0;
-		}
-		return RefCount;
+		return Owner;
 	}
 
 private:
-	FRDGPooledTexture(FRHITexture* InTexture, const FRDGTextureDesc& InDesc, const FUnorderedAccessViewRHIRef& FirstMipUAV)
-		: Desc(InDesc)
-		, Texture(InTexture)
-		, Layout(InDesc.GetSubresourceLayout())
+	FRDGPooledTexture(TRefCountPtr<FRHITexture> InTexture, const FRDGTextureSubresourceLayout& InLayout, const FUnorderedAccessViewRHIRef& FirstMipUAV, EAllocator InAllocator)
+		: FRDGPooledResource(InAllocator)
+		, Texture(MoveTemp(InTexture))
+		, Layout(InLayout)
 	{
 		InitViews(FirstMipUAV);
 		Reset();
@@ -555,22 +417,13 @@ private:
 	/** Initializes cached views. Safe to call multiple times; each call will recreate. */
 	void InitViews(const FUnorderedAccessViewRHIRef& FirstMipUAV);
 
-	void Finalize()
-	{
-		for (FRDGSubresourceState& SubresourceState : State)
-		{
-			SubresourceState.Finalize();
-		}
-		Owner = nullptr;
-	}
+	/** Prepares the pooled texture state for re-use across RDG builder instances. */
+	void Finalize();
 
-	void Reset()
-	{
-		InitAsWholeResource(State);
-		Owner = nullptr;
-	}
+	/** Resets the pooled texture state back an unknown value. */
+	void Reset();
 
-	FRHITexture* Texture = nullptr;
+	TRefCountPtr<FRHITexture> Texture;
 	FRDGTexture* Owner = nullptr;
 	FRDGTextureSubresourceLayout Layout;
 	FRDGTextureSubresourceState State;
@@ -584,12 +437,10 @@ private:
 	FShaderResourceViewRHIRef  FMaskSRV;
 	FShaderResourceViewRHIRef  CMaskSRV;
 
-	mutable uint32 RefCount = 0;
-
 	friend FRDGTexture;
 	friend FRDGBuilder;
-	friend FPooledRenderTarget;
 	friend FRenderTargetPool;
+	friend FRDGAllocator;
 };
 
 /** Render graph tracked Texture. */
@@ -642,14 +493,15 @@ private:
 		, Desc(InDesc)
 		, Flags(InFlags)
 		, RenderTargetTexture(InRenderTargetTexture)
-		, Layout(InDesc.GetSubresourceLayout())
+		, Layout(InDesc)
 	{
 		InitAsWholeResource(MergeState);
 		InitAsWholeResource(LastProducers);
 	}
 
 	/** Assigns the pooled texture to this texture; returns the previous texture to own the allocation. */
-	void SetRHI(FPooledRenderTarget* PooledRenderTarget, FRDGTextureRef& OutPreviousOwner);
+	void SetRHI(FPooledRenderTarget* PooledRenderTarget);
+	void SetRHI(FRDGPooledTexture* PooledTexture);
 
 	/** Finalizes the texture for execution; no other transitions are allowed after calling this. */
 	void Finalize();
@@ -1042,9 +894,25 @@ struct FRDGBufferUAVDesc
 	bool bSupportsAppendBuffer = false;
 };
 
-class RENDERCORE_API FRDGPooledBuffer
+/** Translates from a RDG buffer descriptor to a RHI buffer creation info */
+inline FRHIBufferCreateInfo Translate(const FRDGBufferDesc& InDesc);
+
+class RENDERCORE_API FRDGPooledBuffer final
+	: public FRDGPooledResource
 {
 public:
+	/** Creates a pooled texture using the default allocator. Ref-counting determines deletion of the memory. Safe to be persistently cached. */
+	static TRefCountPtr<FRDGPooledBuffer> CreateCommitted(TRefCountPtr<FRHIBuffer> Buffer, const FRDGBufferDesc& Desc)
+	{
+		return TRefCountPtr<FRDGPooledBuffer>(new FRDGPooledBuffer(MoveTemp(Buffer), Desc, EAllocator::Default));
+	}
+
+	/** Creates a pooled texture using the RDG allocator. All references should be released before the allocator is freed. */
+	static TRefCountPtr<FRDGPooledBuffer> CreateTransient(FRDGAllocator& Allocator, TRefCountPtr<FRHIBuffer> Buffer, const FRDGBufferDesc& Desc)
+	{
+		return TRefCountPtr<FRDGPooledBuffer>(Allocator.AllocNoDestruct<FRDGPooledBuffer>(MoveTemp(Buffer), Desc, EAllocator::RDG));
+	}
+
 	const FRDGBufferDesc Desc;
 
 	/** Finds a UAV matching the descriptor in the cache or creates a new one and updates the cache. */
@@ -1071,26 +939,6 @@ public:
 		return Buffer;
 	}
 
-	uint32 GetRefCount() const
-	{
-		return RefCount;
-	}
-
-	uint32 AddRef() const
-	{
-		return ++RefCount;
-	}
-
-	uint32 Release() const
-	{
-		const uint32 LocalRefCount = --RefCount;
-		if (LocalRefCount == 0)
-		{
-			delete this;
-		}
-		return LocalRefCount;
-	}
-
 	template<typename KeyType, typename ValueType>
 	struct TSRVFuncs : BaseKeyFuncs<TPair<KeyType, ValueType>, KeyType, /* bInAllowDuplicateKeys = */ false>
 	{
@@ -1132,8 +980,10 @@ public:
 	};
 
 private:
-	FRDGPooledBuffer(const FRDGBufferDesc& InDesc)
-		: Desc(InDesc)
+	FRDGPooledBuffer(TRefCountPtr<FRHIBuffer> InBuffer, const FRDGBufferDesc& InDesc, EAllocator InAllocator)
+		: FRDGPooledResource(InAllocator)
+		, Desc(InDesc)
+		, Buffer(MoveTemp(InBuffer))
 	{}
 
 	TRefCountPtr<FRHIBuffer> Buffer;
@@ -1153,7 +1003,6 @@ private:
 	}
 
 	const TCHAR* Name = nullptr;
-
 	FRDGBufferRef Owner = nullptr;
 	FRDGSubresourceState State;
 
@@ -1163,6 +1012,7 @@ private:
 	friend FRenderGraphResourcePool;
 	friend FRDGBuilder;
 	friend FRDGBuffer;
+	friend FRDGAllocator;
 };
 
 /** A render graph tracked buffer. */
@@ -1218,7 +1068,7 @@ private:
 	{}
 
 	/** Assigns the pooled buffer to this buffer; returns the previous buffer to own the allocation. */
-	void SetRHI(FRDGPooledBuffer* InPooledBuffer, FRDGBufferRef& OutPreviousOwner);
+	void SetRHI(FRDGPooledBuffer* InPooledBuffer);
 
 	/** Returns RHI buffer without access checks. */
 	FRHIBuffer* GetRHIUnchecked() const
diff --git a/Engine/Source/Runtime/RenderCore/Public/RenderGraphResources.inl b/Engine/Source/Runtime/RenderCore/Public/RenderGraphResources.inl
--- a/Engine/Source/Runtime/RenderCore/Public/RenderGraphResources.inl
+++ b/Engine/Source/Runtime/RenderCore/Public/RenderGraphResources.inl
@@ -107,6 +107,28 @@ inline FPooledRenderTargetDesc Translate(const FRDGTextureDesc& InDesc)
 	return OutDesc;
 }
 
+inline FRHIBufferCreateInfo Translate(const FRDGBufferDesc& InDesc)
+{
+	FRHIBufferCreateInfo CreateInfo;
+	CreateInfo.Size = InDesc.GetTotalNumBytes();
+	if (InDesc.UnderlyingType == FRDGBufferDesc::EUnderlyingType::VertexBuffer)
+	{
+		CreateInfo.Stride = 0;
+		CreateInfo.Usage = InDesc.Usage | BUF_VertexBuffer;
+	}
+	else if (InDesc.UnderlyingType == FRDGBufferDesc::EUnderlyingType::StructuredBuffer)
+	{
+		CreateInfo.Stride = InDesc.BytesPerElement;
+		CreateInfo.Usage = InDesc.Usage | BUF_StructuredBuffer;
+	}
+	else
+	{
+		check(0);
+	}
+
+	return CreateInfo;
+}
+
 FRDGTextureDesc Translate(const FPooledRenderTargetDesc& InDesc, ERenderTargetTexture InTexture)
 {
 	check(InDesc.IsValid());
diff --git a/Engine/Source/Runtime/RenderCore/Public/RenderGraphTextureSubresource.h b/Engine/Source/Runtime/RenderCore/Public/RenderGraphTextureSubresource.h
--- a/Engine/Source/Runtime/RenderCore/Public/RenderGraphTextureSubresource.h
+++ b/Engine/Source/Runtime/RenderCore/Public/RenderGraphTextureSubresource.h
@@ -77,6 +77,10 @@ struct FRDGTextureSubresourceLayout
 		, NumArraySlices(InNumArraySlices)
 	{}
 
+	FRDGTextureSubresourceLayout(const FRHITextureCreateInfo& CreateInfo)
+		: FRDGTextureSubresourceLayout(CreateInfo.NumMips, CreateInfo.ArraySize * (CreateInfo.IsTextureCube() ? 6 : 1), IsStencilFormat(CreateInfo.Format) ? 2 : 1)
+	{}
+
 	inline uint32 GetSubresourceCount() const
 	{
 		return NumMips * NumArraySlices * NumPlaneSlices;
diff --git a/Engine/Source/Runtime/RenderCore/Public/RenderGraphValidation.h b/Engine/Source/Runtime/RenderCore/Public/RenderGraphValidation.h
--- a/Engine/Source/Runtime/RenderCore/Public/RenderGraphValidation.h
+++ b/Engine/Source/Runtime/RenderCore/Public/RenderGraphValidation.h
@@ -132,6 +132,7 @@ private:
 	{
 		TMap<FRDGTextureRef, TArray<FRHITransitionInfo>> Textures;
 		TMap<FRDGBufferRef, FRHITransitionInfo> Buffers;
+		TMap<FRDGParentResourceRef, FRHITransientAliasingInfo> Aliases;
 	};
 
 	using FBarrierBatchMap = TMap<const FRDGBarrierBatchBegin*, FResourceMap>;
diff --git a/Engine/Source/Runtime/RenderCore/Public/RenderTargetPool.h b/Engine/Source/Runtime/RenderCore/Public/RenderTargetPool.h
--- a/Engine/Source/Runtime/RenderCore/Public/RenderTargetPool.h
+++ b/Engine/Source/Runtime/RenderCore/Public/RenderTargetPool.h
@@ -96,9 +96,6 @@ private:
 	/** Keeps track of the last frame we unmapped physical memory for this resource. We can't map again in the same frame if we did that */
 	uint32 FrameNumberLastDiscard = -1;
 
-	/** The transient resource discard will happen automatically on free. */
-	bool bAutoDiscard = true;
-
 	/** Pooled textures for use with RDG. */
 	TRefCountPtr<FRDGPooledTexture> TargetableTexture;
 	TRefCountPtr<FRDGPooledTexture> ShaderResourceTexture;
@@ -288,8 +285,7 @@ private:
 	TRefCountPtr<FPooledRenderTarget> FindFreeElementInternal(
 		FRHICommandList& RHICmdList,
 		const FPooledRenderTargetDesc& InputDesc,
-		const TCHAR* InDebugName,
-		bool bDoAcquireTransientResource);
+		const TCHAR* InDebugName);
 
 	static bool DoesTargetNeedTransienceOverride(ETextureCreateFlags Flags, ERenderTargetTransience TransienceHint);
 
diff --git a/Engine/Source/Runtime/Renderer/Private/DeferredShadingRenderer.cpp b/Engine/Source/Runtime/Renderer/Private/DeferredShadingRenderer.cpp
--- a/Engine/Source/Runtime/Renderer/Private/DeferredShadingRenderer.cpp
+++ b/Engine/Source/Runtime/Renderer/Private/DeferredShadingRenderer.cpp
@@ -1230,8 +1230,8 @@ bool FDeferredShadingSceneRenderer::DispatchRayTracingWorldUpdates(FRDGBuilder&
 		AddPass(GraphBuilder, [this](FRHICommandList& RHICmdList)
 		{
 			check(RayTracingDynamicGeometryUpdateEndTransition == nullptr);
-			const FRHITransition* RayTracingDynamicGeometryUpdateBeginTransition = RHICreateTransition(ERHIPipeline::Graphics, ERHIPipeline::AsyncCompute, ERHICreateTransitionFlags::None, {});
-			RayTracingDynamicGeometryUpdateEndTransition = RHICreateTransition(ERHIPipeline::AsyncCompute, ERHIPipeline::Graphics, ERHICreateTransitionFlags::None, {});
+			const FRHITransition* RayTracingDynamicGeometryUpdateBeginTransition = RHICreateTransition(FRHITransitionCreateInfo(ERHIPipeline::Graphics, ERHIPipeline::AsyncCompute));
+			RayTracingDynamicGeometryUpdateEndTransition = RHICreateTransition(FRHITransitionCreateInfo(ERHIPipeline::AsyncCompute, ERHIPipeline::Graphics));
 
 			FRHIAsyncComputeCommandListImmediate& RHIAsyncCmdList = FRHICommandListExecutor::GetImmediateAsyncComputeCommandList();
 
diff --git a/Engine/Source/Runtime/VulkanRHI/Private/VulkanBarriers.cpp b/Engine/Source/Runtime/VulkanRHI/Private/VulkanBarriers.cpp
--- a/Engine/Source/Runtime/VulkanRHI/Private/VulkanBarriers.cpp
+++ b/Engine/Source/Runtime/VulkanRHI/Private/VulkanBarriers.cpp
@@ -493,8 +493,11 @@ static void AddMemoryBarrier(VkMemoryBarrier& MemoryBarrier, VkAccessFlags SrcAc
 	}
 }
 
-void FVulkanDynamicRHI::RHICreateTransition(FRHITransition* Transition, ERHIPipeline SrcPipelines, ERHIPipeline DstPipelines, ERHICreateTransitionFlags CreateFlags, TArrayView<const FRHITransitionInfo> Infos)
+void FVulkanDynamicRHI::RHICreateTransition(FRHITransition* Transition, const FRHITransitionCreateInfo& CreateInfo)
 {
+	const ERHIPipeline SrcPipelines = CreateInfo.SrcPipelines;
+	const ERHIPipeline DstPipelines = CreateInfo.DstPipelines;
+
 	checkf(FMath::IsPowerOfTwo((uint32)SrcPipelines) && FMath::IsPowerOfTwo((uint32)DstPipelines), TEXT("Support for multi-pipe resources is not yet implemented."));
 
 	FVulkanPipelineBarrier* Data = new (Transition->GetPrivateData<FVulkanPipelineBarrier>()) FVulkanPipelineBarrier;
@@ -528,7 +531,7 @@ void FVulkanDynamicRHI::RHICreateTransition(FRHITransition* Transition, ERHIPipe
 
 	// Count the images and buffers to be able to pre-allocate the arrays.
 	int32 NumTextures = 0, NumBuffers = 0;
-	for (const FRHITransitionInfo& Info : Infos)
+	for (const FRHITransitionInfo& Info : CreateInfo.TransitionInfos)
 	{
 		if (!Info.Resource)
 		{
@@ -565,7 +568,7 @@ void FVulkanDynamicRHI::RHICreateTransition(FRHITransition* Transition, ERHIPipe
 
 	TArray<FDepthStencilSubresTransition, TInlineAllocator<4>> DSSubresTransitions;
 
-	for (const FRHITransitionInfo& Info : Infos)
+	for (const FRHITransitionInfo& Info : CreateInfo.TransitionInfos)
 	{
 		if (!Info.Resource)
 		{
diff --git a/Engine/Source/Runtime/VulkanRHI/Public/VulkanDynamicRHI.h b/Engine/Source/Runtime/VulkanRHI/Public/VulkanDynamicRHI.h
--- a/Engine/Source/Runtime/VulkanRHI/Public/VulkanDynamicRHI.h
+++ b/Engine/Source/Runtime/VulkanRHI/Public/VulkanDynamicRHI.h
@@ -73,7 +73,7 @@ public:
 	virtual FGeometryShaderRHIRef RHICreateGeometryShader(TArrayView<const uint8> Code, const FSHAHash& Hash) final override;
 	virtual FComputeShaderRHIRef RHICreateComputeShader(TArrayView<const uint8> Code, const FSHAHash& Hash) final override;
 	virtual FGPUFenceRHIRef RHICreateGPUFence(const FName &Name) final override;
-	virtual void RHICreateTransition(FRHITransition* Transition, ERHIPipeline SrcPipelines, ERHIPipeline DstPipelines, ERHICreateTransitionFlags CreateFlags, TArrayView<const FRHITransitionInfo> Infos) final override;
+	virtual void RHICreateTransition(FRHITransition* Transition, const FRHITransitionCreateInfo& CreateInfo) final override;
 	virtual void RHIReleaseTransition(FRHITransition* Transition) final override;
 	virtual FStagingBufferRHIRef RHICreateStagingBuffer() final override;
 	virtual FBoundShaderStateRHIRef RHICreateBoundShaderState(FRHIVertexDeclaration* VertexDeclaration, FRHIVertexShader* VertexShader, FRHIHullShader* HullShader, FRHIDomainShader* DomainShader, FRHIPixelShader* PixelShader, FRHIGeometryShader* GeometryShader) final override;
diff --git a/Engine/Source/Runtime/Windows/D3D11RHI/Private/D3D11Commands.cpp b/Engine/Source/Runtime/Windows/D3D11RHI/Private/D3D11Commands.cpp
--- a/Engine/Source/Runtime/Windows/D3D11RHI/Private/D3D11Commands.cpp
+++ b/Engine/Source/Runtime/Windows/D3D11RHI/Private/D3D11Commands.cpp
@@ -2069,15 +2069,15 @@ void FD3D11DynamicRHI::DisableUAVOverlap()
 	bUAVOverlapEnabled = false;
 }
 
-void FD3D11DynamicRHI::RHICreateTransition(FRHITransition* Transition, ERHIPipeline SrcPipelines, ERHIPipeline DstPipelines, ERHICreateTransitionFlags CreateFlags, TArrayView<const FRHITransitionInfo> Infos)
+void FD3D11DynamicRHI::RHICreateTransition(FRHITransition* Transition, const FRHITransitionCreateInfo& CreateInfo)
 {
-	checkf(FMath::IsPowerOfTwo(uint32(SrcPipelines)) && FMath::IsPowerOfTwo(uint32(DstPipelines)), TEXT("Support for multi-pipe resources is not yet implemented."));
+	checkf(FMath::IsPowerOfTwo(uint32(CreateInfo.SrcPipelines)) && FMath::IsPowerOfTwo(uint32(CreateInfo.DstPipelines)), TEXT("Support for multi-pipe resources is not yet implemented."));
 
 	FD3D11TransitionData* Data = new (Transition->GetPrivateData<FD3D11TransitionData>()) FD3D11TransitionData;
 	Data->bUAVBarrier = false;
 
 	// If we have any transitions to UAVCompute or UAVGraphics, we need to break up the current overlap group.
-	for (const FRHITransitionInfo& Info : Infos)
+	for (const FRHITransitionInfo& Info : CreateInfo.TransitionInfos)
 	{
 		if (Info.Resource && EnumHasAnyFlags(Info.AccessAfter, ERHIAccess::UAVMask))
 		{
diff --git a/Engine/Source/Runtime/Windows/D3D11RHI/Private/D3D11RHIPrivate.h b/Engine/Source/Runtime/Windows/D3D11RHI/Private/D3D11RHIPrivate.h
--- a/Engine/Source/Runtime/Windows/D3D11RHI/Private/D3D11RHIPrivate.h
+++ b/Engine/Source/Runtime/Windows/D3D11RHI/Private/D3D11RHIPrivate.h
@@ -530,7 +530,7 @@ public:
 	virtual void RHICopyToResolveTarget(FRHITexture* SourceTexture, FRHITexture* DestTexture, const FResolveParams& ResolveParams) final override;
 	virtual void RHICopyTexture(FRHITexture* SourceTexture, FRHITexture* DestTexture, const FRHICopyTextureInfo& CopyInfo) final override;
 	virtual void RHICopyBufferRegion(FRHIBuffer* DestBuffer, uint64 DstOffset, FRHIBuffer* SourceBuffer, uint64 SrcOffset, uint64 NumBytes) final override;
-	virtual void RHICreateTransition(FRHITransition* Transition, ERHIPipeline SrcPipelines, ERHIPipeline DstPipelines, ERHICreateTransitionFlags CreateFlags, TArrayView<const FRHITransitionInfo> Infos) final override;
+	virtual void RHICreateTransition(FRHITransition* Transition, const FRHITransitionCreateInfo& CreateInfo) final override;
 	virtual void RHIReleaseTransition(FRHITransition* Transition) final override;
 	virtual void RHIBeginTransitions(TArrayView<const FRHITransition*> Transitions) override final;
 	virtual void RHIEndTransitions(TArrayView<const FRHITransition*> Transitions) override final;
