Shader compiler: prevent attempts to generate output hash from the compressed code.

diff --git a/Engine/Source/Runtime/Engine/Private/ShaderCompiler/ShaderCompiler.cpp b/Engine/Source/Runtime/Engine/Private/ShaderCompiler/ShaderCompiler.cpp
--- a/Engine/Source/Runtime/Engine/Private/ShaderCompiler/ShaderCompiler.cpp
+++ b/Engine/Source/Runtime/Engine/Private/ShaderCompiler/ShaderCompiler.cpp
@@ -6529,13 +6529,12 @@ void FShaderCompileJob::SerializeOutput(FArchive& Ar)
 	}
 
 	Ar << Output;
+	// output hash is now serialized as part of the output, as the shader code is compressed in SCWs
+	checkf(Output.OutputHash != FSHAHash(), TEXT("Compile job does not have an OutputHash generated."));
 
 	if (Ar.IsLoading())
 	{
 		bFinalized = true;
-
-		// serialize the hash as well? minor optimization
-		Output.GenerateOutputHash();
 		bSucceeded = Output.bSucceeded;
 	}
 	else
diff --git a/Engine/Source/Runtime/RenderCore/Private/ShaderCore.cpp b/Engine/Source/Runtime/RenderCore/Private/ShaderCore.cpp
--- a/Engine/Source/Runtime/RenderCore/Private/ShaderCore.cpp
+++ b/Engine/Source/Runtime/RenderCore/Private/ShaderCore.cpp
@@ -418,6 +418,9 @@ void FShaderCompilerOutput::GenerateOutputHash()
 	// we don't hash the optional attachments as they would prevent sharing (e.g. many materials share the same VS)
 	uint32 ShaderCodeSize = ShaderCode.GetShaderCodeSize();
 
+	// make sure we are not generating the hash on compressed data
+	checkf(!ShaderCode.IsCompressed(), TEXT("Attempting to generate the output hash of a compressed code"));
+
 	HashState.Update(Code.GetData(), ShaderCodeSize * Code.GetTypeSize());
 	ParameterMap.UpdateHash(HashState);
 	HashState.Final();
