Add name of material to AllocateVT as debug aid
- Rework hashing of FAllocatedVTDescription to be more explicit about hashing each individual member, don't hash the debug name

diff --git a/Engine/Source/Runtime/Engine/Private/Materials/MaterialShared.cpp b/Engine/Source/Runtime/Engine/Private/Materials/MaterialShared.cpp
--- a/Engine/Source/Runtime/Engine/Private/Materials/MaterialShared.cpp
+++ b/Engine/Source/Runtime/Engine/Private/Materials/MaterialShared.cpp
@@ -2945,7 +2945,13 @@ IAllocatedVirtualTexture* FMaterialRenderProxy::AllocateVTStack(const FMaterialR
 	const UTexture2D* LayerTextures[VIRTUALTEXTURE_SPACE_MAXLAYERS] = { nullptr };
 	VTStack.GetTextureValues(Context, UniformExpressionSet, LayerTextures);
 
+	const UMaterialInterface* MaterialInterface = GetMaterialInterface();
+
 	FAllocatedVTDescription VTDesc;
+	if (MaterialInterface)
+	{
+		VTDesc.Name = MaterialInterface->GetFName();
+	}
 	VTDesc.Dimensions = 2;
 	VTDesc.NumTextureLayers = NumLayers;
 	bool bFoundValidLayer = false;
diff --git a/Engine/Source/Runtime/RenderCore/Public/VirtualTexturing.h b/Engine/Source/Runtime/RenderCore/Public/VirtualTexturing.h
--- a/Engine/Source/Runtime/RenderCore/Public/VirtualTexturing.h
+++ b/Engine/Source/Runtime/RenderCore/Public/VirtualTexturing.h
@@ -49,6 +49,8 @@ inline bool operator!=(const FVirtualTextureProducerHandle& Lhs, const FVirtualT
  */
 struct FAllocatedVTDescription
 {
+	FName Name;
+
 	uint32 TileSize = 0u;
 	uint32 TileBorderSize = 0u;
 	uint8 Dimensions = 0u;
@@ -109,6 +111,21 @@ inline bool operator!=(const FAllocatedVTDescription& Lhs, const FAllocatedVTDes
 	return !operator==(Lhs, Rhs);
 }
 
+inline uint32 GetTypeHash(const FAllocatedVTDescription& Description)
+{
+	uint32 Hash = GetTypeHash(Description.TileSize);
+	Hash = HashCombine(Hash, GetTypeHash(Description.TileBorderSize));
+	Hash = HashCombine(Hash, GetTypeHash(Description.Dimensions));
+	Hash = HashCombine(Hash, GetTypeHash(Description.NumTextureLayers));
+	Hash = HashCombine(Hash, GetTypeHash(Description.PackedFlags));
+	for (uint32 LayerIndex = 0u; LayerIndex < Description.NumTextureLayers; ++LayerIndex)
+	{
+		Hash = HashCombine(Hash, GetTypeHash(Description.ProducerHandle[LayerIndex].PackedValue));
+		Hash = HashCombine(Hash, GetTypeHash(Description.ProducerLayerIndex[LayerIndex]));
+	}
+	return Hash;
+}
+
 struct FVTProducerDescription
 {
 	FName Name; /** Will be name of UTexture for streaming VTs, mostly here for debugging */
diff --git a/Engine/Source/Runtime/Renderer/Private/VT/VirtualTextureSystem.cpp b/Engine/Source/Runtime/Renderer/Private/VT/VirtualTextureSystem.cpp
--- a/Engine/Source/Runtime/Renderer/Private/VT/VirtualTextureSystem.cpp
+++ b/Engine/Source/Runtime/Renderer/Private/VT/VirtualTextureSystem.cpp
@@ -389,11 +389,6 @@ void FVirtualTextureSystem::ListPhysicalPoolsFromConsole()
 	}
 }
 
-uint32 GetTypeHash(const FAllocatedVTDescription& Description)
-{
-	return FCrc::MemCrc32(&Description, sizeof(Description));
-}
-
 IAllocatedVirtualTexture* FVirtualTextureSystem::AllocateVirtualTexture(const FAllocatedVTDescription& Desc)
 {
 	check(Desc.NumTextureLayers <= VIRTUALTEXTURE_SPACE_MAXLAYERS);
