Do not defer delete when ticking pool elements for unused elements - fixes GPU VRAM spikes

diff --git a/Engine/Source/Runtime/RenderCore/Private/RenderTargetPool.cpp b/Engine/Source/Runtime/RenderCore/Private/RenderTargetPool.cpp
--- a/Engine/Source/Runtime/RenderCore/Private/RenderTargetPool.cpp
+++ b/Engine/Source/Runtime/RenderCore/Private/RenderTargetPool.cpp
@@ -1002,6 +1002,15 @@ void FRenderTargetPool::TickPoolElements()
 		{
 			AllocationLevelInKB -= ComputeSizeInKB(*PooledRenderTargets[OldestElementIndex]);
 
+			// Don't defer delete because we know it hasn't been used for at least 2 frames 
+			// and defer deletion queue might not get immediatly flushed causing VRAM
+			// memory spikes
+			if (PooledRenderTargets[OldestElementIndex])
+			{
+				PooledRenderTargets[OldestElementIndex]->GetTargetableRHI()->DoNoDeferDelete();
+				PooledRenderTargets[OldestElementIndex]->GetShaderResourceRHI()->DoNoDeferDelete();
+			}
+
 			// we assume because of reference counting the resource gets released when not needed any more
 			// we don't use Remove() to not shuffle around the elements for better transparency on RenderTargetPoolEvents
 			FreeElementAtIndex(OldestElementIndex);
