Fixed race with profile gpu and draw events in RDG.

diff --git a/Engine/Source/Runtime/RenderCore/Private/RenderGraphBuilder.cpp b/Engine/Source/Runtime/RenderCore/Private/RenderGraphBuilder.cpp
--- a/Engine/Source/Runtime/RenderCore/Private/RenderGraphBuilder.cpp
+++ b/Engine/Source/Runtime/RenderCore/Private/RenderGraphBuilder.cpp
@@ -442,6 +442,11 @@ void FRDGBuilder::TickPoolElements()
 	GRDGStatTransitionBatchCount = 0;
 	GRDGStatMemoryWatermark = 0;
 #endif
+
+#if RDG_EVENTS != RDG_EVENTS_NONE
+	// This is polled once per frame as the underlying global variable changes on the main thread.
+	GRDGEmitEvents = GetEmitDrawEvents();
+#endif
 }
 
 ERDGPassFlags FRDGBuilder::OverridePassFlags(const TCHAR* PassName, ERDGPassFlags PassFlags, bool bAsyncComputeSupported)
diff --git a/Engine/Source/Runtime/RenderCore/Private/RenderGraphEvent.cpp b/Engine/Source/Runtime/RenderCore/Private/RenderGraphEvent.cpp
--- a/Engine/Source/Runtime/RenderCore/Private/RenderGraphEvent.cpp
+++ b/Engine/Source/Runtime/RenderCore/Private/RenderGraphEvent.cpp
@@ -13,7 +13,7 @@ RENDERCORE_API bool GetEmitRDGEvents()
 #if RDG_ENABLE_TRACE
 	bRDGChannelEnabled = UE_TRACE_CHANNELEXPR_IS_ENABLED(RDGChannel);
 #endif // RDG_ENABLE_TRACE
-	return GetEmitDrawEvents() || GRDGDebug || bRDGChannelEnabled;
+	return GRDGEmitEvents || GRDGDebug || bRDGChannelEnabled;
 #else
 	return false;
 #endif
diff --git a/Engine/Source/Runtime/RenderCore/Private/RenderGraphPrivate.cpp b/Engine/Source/Runtime/RenderCore/Private/RenderGraphPrivate.cpp
--- a/Engine/Source/Runtime/RenderCore/Private/RenderGraphPrivate.cpp
+++ b/Engine/Source/Runtime/RenderCore/Private/RenderGraphPrivate.cpp
@@ -322,6 +322,10 @@ DEFINE_STAT(STAT_RDG_ClearTime);
 DEFINE_STAT(STAT_RDG_MemoryWatermark);
 #endif
 
+#if RDG_EVENTS != RDG_EVENTS_NONE
+int32 GRDGEmitEvents = 0;
+#endif
+
 void InitRenderGraph()
 {
 #if RDG_ENABLE_DEBUG_WITH_ENGINE
diff --git a/Engine/Source/Runtime/RenderCore/Private/RenderGraphPrivate.h b/Engine/Source/Runtime/RenderCore/Private/RenderGraphPrivate.h
--- a/Engine/Source/Runtime/RenderCore/Private/RenderGraphPrivate.h
+++ b/Engine/Source/Runtime/RenderCore/Private/RenderGraphPrivate.h
@@ -146,6 +146,10 @@ DECLARE_CYCLE_STAT_EXTERN(TEXT("Clear"), STAT_RDG_ClearTime, STATGROUP_RDG, REND
 DECLARE_MEMORY_STAT_EXTERN(TEXT("Builder Watermark"), STAT_RDG_MemoryWatermark, STATGROUP_RDG, RENDERCORE_API);
 #endif
 
+#if RDG_EVENTS != RDG_EVENTS_NONE
+extern int32 GRDGEmitEvents;
+#endif
+
 inline const TCHAR* GetEpilogueBarriersToBeginDebugName(ERHIPipeline Pipelines)
 {
 #if RDG_ENABLE_DEBUG
