Fix for crash in render target pool texture trimming.

diff --git a/Engine/Source/Runtime/RenderCore/Private/RenderTargetPool.cpp b/Engine/Source/Runtime/RenderCore/Private/RenderTargetPool.cpp
--- a/Engine/Source/Runtime/RenderCore/Private/RenderTargetPool.cpp
+++ b/Engine/Source/Runtime/RenderCore/Private/RenderTargetPool.cpp
@@ -1008,10 +1008,20 @@ void FRenderTargetPool::TickPoolElements()
 			// Don't defer delete because we know it hasn't been used for at least 2 frames 
 			// and defer deletion queue might not get immediatly flushed causing VRAM
 			// memory spikes
-			if (PooledRenderTargets[OldestElementIndex])
+			if (TRefCountPtr<FPooledRenderTarget>& PooledTarget = PooledRenderTargets[OldestElementIndex])
 			{
-				PooledRenderTargets[OldestElementIndex]->GetTargetableRHI()->DoNoDeferDelete();
-				PooledRenderTargets[OldestElementIndex]->GetShaderResourceRHI()->DoNoDeferDelete();
+				FRHITexture* TargetableRHI = PooledTarget->GetTargetableRHI();
+				FRHITexture* ShaderResourceRHI = PooledTarget->GetShaderResourceRHI();
+
+				if (TargetableRHI)
+				{
+					TargetableRHI->DoNoDeferDelete();
+				}
+
+				if (ShaderResourceRHI && ShaderResourceRHI != TargetableRHI)
+				{
+					ShaderResourceRHI->DoNoDeferDelete();
+				}
 			}
 
 			// we assume because of reference counting the resource gets released when not needed any more
