Fix issue where frame fence is retrieved from RenderThread while it's being updated by RHI Thread while deleting resources from the RenderThread (should be moved to RHIThread via request deletion)

diff --git a/Engine/Source/Runtime/RenderCore/Private/RenderTargetPool.cpp b/Engine/Source/Runtime/RenderCore/Private/RenderTargetPool.cpp
--- a/Engine/Source/Runtime/RenderCore/Private/RenderTargetPool.cpp
+++ b/Engine/Source/Runtime/RenderCore/Private/RenderTargetPool.cpp
@@ -1005,6 +1005,9 @@ void FRenderTargetPool::TickPoolElements()
 		{
 			AllocationLevelInKB -= ComputeSizeInKB(*PooledRenderTargets[OldestElementIndex]);
 
+			// NOTE: Comment out because RHI resource shouldn't be deleted on the RenderThread
+			//       But command to request delete should be added to fix this
+			/*
 			// Don't defer delete because we know it hasn't been used for at least 2 frames 
 			// and defer deletion queue might not get immediatly flushed causing VRAM
 			// memory spikes
@@ -1023,6 +1026,7 @@ void FRenderTargetPool::TickPoolElements()
 					ShaderResourceRHI->DoNoDeferDelete();
 				}
 			}
+			*/
 
 			// we assume because of reference counting the resource gets released when not needed any more
 			// we don't use Remove() to not shuffle around the elements for better transparency on RenderTargetPoolEvents
