Fixed profilegpu missing RDG passes.

diff --git a/Engine/Source/Runtime/RenderCore/Private/RenderGraphBuilder.cpp b/Engine/Source/Runtime/RenderCore/Private/RenderGraphBuilder.cpp
--- a/Engine/Source/Runtime/RenderCore/Private/RenderGraphBuilder.cpp
+++ b/Engine/Source/Runtime/RenderCore/Private/RenderGraphBuilder.cpp
@@ -442,11 +442,6 @@ void FRDGBuilder::TickPoolElements()
 	GRDGStatTransitionBatchCount = 0;
 	GRDGStatMemoryWatermark = 0;
 #endif
-
-#if RDG_EVENTS != RDG_EVENTS_NONE
-	// This is polled once per frame as the underlying global variable changes on the main thread.
-	GRDGEmitEvents = GetEmitDrawEvents();
-#endif
 }
 
 ERDGPassFlags FRDGBuilder::OverridePassFlags(const TCHAR* PassName, ERDGPassFlags PassFlags, bool bAsyncComputeSupported)
@@ -551,6 +546,11 @@ FRDGBuilder::FRDGBuilder(FRHICommandListImmediate& InRHICmdList, FRDGEventName I
 {
 	ProloguePass = Passes.Allocate<FRDGSentinelPass>(Allocator, RDG_EVENT_NAME("Graph Prologue"));
 	SetupEmptyPass(ProloguePass);
+
+#if RDG_EVENTS != RDG_EVENTS_NONE
+	// This is polled once as a workaround for a race condition since the underlying global is not always changed on the render thread.
+	GRDGEmitEvents = GetEmitDrawEvents();
+#endif
 }
 
 void FRDGBuilder::PreallocateBuffer(FRDGBufferRef Buffer)
