Deleted FRayTracingScene, as it's no longer used.

diff --git a/Engine/Source/Runtime/RenderCore/Public/RenderResource.h b/Engine/Source/Runtime/RenderCore/Public/RenderResource.h
--- a/Engine/Source/Runtime/RenderCore/Public/RenderResource.h
+++ b/Engine/Source/Runtime/RenderCore/Public/RenderResource.h
@@ -703,21 +703,6 @@ protected:
 #endif
 };
 
-#if RHI_RAYTRACING
-class RENDERCORE_API FRayTracingScene : public FRenderResource
-{
-public:
-	FRayTracingSceneRHIRef RayTracingSceneRHI = nullptr;
-
-	virtual FString GetFriendlyName() const override { return TEXT("FRayTracingScene"); }
-
-	virtual void ReleaseRHI()
-	{
-		RayTracingSceneRHI.SafeRelease();
-	}
-};
-#endif // RHI_RAYTRACING
-
 /**
  * A system for dynamically allocating GPU memory for vertices.
  */
