Fix for RDG using MemStack allocator.

#rnx

diff --git a/Engine/Source/Runtime/RenderCore/Private/RenderGraphBuilder.cpp b/Engine/Source/Runtime/RenderCore/Private/RenderGraphBuilder.cpp
--- a/Engine/Source/Runtime/RenderCore/Private/RenderGraphBuilder.cpp
+++ b/Engine/Source/Runtime/RenderCore/Private/RenderGraphBuilder.cpp
@@ -534,10 +534,10 @@ FRDGBuilder::FRDGBuilder(FRHICommandListImmediate& InRHICmdList, FRDGEventName I
 	, RHICmdListAsyncCompute(FRHICommandListExecutor::GetImmediateAsyncComputeCommandList())
 	, BuilderName(InName)
 #if RDG_CPU_SCOPES
-	, CPUScopeStacks(RHICmdList, kDefaultUnaccountedCSVStat)
+	, CPUScopeStacks(RHICmdList, Allocator, kDefaultUnaccountedCSVStat)
 #endif
 #if RDG_GPU_SCOPES
-	, GPUScopeStacks(RHICmdList, RHICmdListAsyncCompute)
+	, GPUScopeStacks(RHICmdList, RHICmdListAsyncCompute, Allocator)
 #endif
 #if RDG_ENABLE_DEBUG
 	, UserValidation(Allocator)
diff --git a/Engine/Source/Runtime/RenderCore/Private/RenderGraphEvent.cpp b/Engine/Source/Runtime/RenderCore/Private/RenderGraphEvent.cpp
--- a/Engine/Source/Runtime/RenderCore/Private/RenderGraphEvent.cpp
+++ b/Engine/Source/Runtime/RenderCore/Private/RenderGraphEvent.cpp
@@ -106,8 +106,8 @@ bool FRDGEventScopeStack::IsEnabled()
 #endif
 }
 
-FRDGEventScopeStack::FRDGEventScopeStack(FRHIComputeCommandList& InRHICmdList)
-	: ScopeStack(InRHICmdList, &OnPushEvent, &OnPopEvent)
+FRDGEventScopeStack::FRDGEventScopeStack(FRHIComputeCommandList& InRHICmdList, FRDGAllocator& Allocator)
+	: ScopeStack(InRHICmdList, Allocator, &OnPushEvent, &OnPopEvent)
 {}
 
 void FRDGEventScopeStack::BeginScope(FRDGEventName&& EventName)
@@ -230,8 +230,8 @@ bool FRDGGPUStatScopeStack::IsEnabled()
 #endif
 }
 
-FRDGGPUStatScopeStack::FRDGGPUStatScopeStack(FRHIComputeCommandList& InRHICmdList)
-	: ScopeStack(InRHICmdList, &OnPushGPUStat, &OnPopGPUStat)
+FRDGGPUStatScopeStack::FRDGGPUStatScopeStack(FRHIComputeCommandList& InRHICmdList, FRDGAllocator& Allocator)
+	: ScopeStack(InRHICmdList, Allocator, &OnPushGPUStat, &OnPopGPUStat)
 {}
 
 void FRDGGPUStatScopeStack::BeginScope(const FName& Name, const FName& StatName, int32 (*DrawCallCounter)[MAX_NUM_GPUS])
@@ -363,8 +363,8 @@ static void OnPopCSVStat(FRHIComputeCommandList&, const FRDGCSVStatScope* Scope)
 #endif
 }
 
-FRDGCSVStatScopeStack::FRDGCSVStatScopeStack(FRHIComputeCommandList& InRHICmdList, const char* InUnaccountedStatName)
-	: ScopeStack(InRHICmdList, &OnPushCSVStat, &OnPopCSVStat)
+FRDGCSVStatScopeStack::FRDGCSVStatScopeStack(FRHIComputeCommandList& InRHICmdList, FRDGAllocator& Allocator, const char* InUnaccountedStatName)
+	: ScopeStack(InRHICmdList, Allocator, &OnPushCSVStat, &OnPopCSVStat)
 	, UnaccountedStatName(InUnaccountedStatName)
 {
 	BeginScope(UnaccountedStatName);
diff --git a/Engine/Source/Runtime/RenderCore/Public/RenderGraphEvent.h b/Engine/Source/Runtime/RenderCore/Public/RenderGraphEvent.h
--- a/Engine/Source/Runtime/RenderCore/Public/RenderGraphEvent.h
+++ b/Engine/Source/Runtime/RenderCore/Public/RenderGraphEvent.h
@@ -76,7 +76,7 @@ public:
 	using FPushFunction = void(*)(FRHIComputeCommandList&, const TScopeType*);
 	using FPopFunction = void(*)(FRHIComputeCommandList&, const TScopeType*);
 
-	TRDGScopeStack(FRHIComputeCommandList& InRHICmdList, FPushFunction InPushFunction, FPopFunction InPopFunction);
+	TRDGScopeStack(FRHIComputeCommandList& InRHICmdList, FRDGAllocator& Allocator, FPushFunction InPushFunction, FPopFunction InPopFunction);
 	~TRDGScopeStack();
 
 	//////////////////////////////////////////////////////////////////////////
@@ -109,11 +109,12 @@ public:
 	}
 
 	FRHIComputeCommandList& RHICmdList;
-	FMemStackBase& MemStack;
 
 private:
 	void ClearScopes();
 
+	FRDGAllocator& Allocator;
+
 	FPushFunction PushFunction;
 	FPopFunction PopFunction;
 
@@ -195,7 +196,7 @@ public:
 class RENDERCORE_API FRDGEventScopeStack final
 {
 public:
-	FRDGEventScopeStack(FRHIComputeCommandList& RHICmdList);
+	FRDGEventScopeStack(FRHIComputeCommandList& RHICmdList, FRDGAllocator& Allocator);
 
 	void BeginScope(FRDGEventName&& EventName);
 
@@ -259,7 +260,7 @@ public:
 class RENDERCORE_API FRDGGPUStatScopeStack final
 {
 public:
-	FRDGGPUStatScopeStack(FRHIComputeCommandList& RHICmdList);
+	FRDGGPUStatScopeStack(FRHIComputeCommandList& RHICmdList, FRDGAllocator& Allocator);
 
 	void BeginScope(const FName& Name, const FName& StatName, int32 (*DrawCallCounter)[MAX_NUM_GPUS]);
 
@@ -301,7 +302,7 @@ struct FRDGGPUScopes
 /** The complete set of scope stack implementations. */
 struct FRDGGPUScopeStacks
 {
-	FRDGGPUScopeStacks(FRHIComputeCommandList& RHICmdList);
+	FRDGGPUScopeStacks(FRHIComputeCommandList& RHICmdList, FRDGAllocator& Allocator);
 
 	void BeginExecute();
 
@@ -319,7 +320,7 @@ struct FRDGGPUScopeStacks
 
 struct FRDGGPUScopeStacksByPipeline
 {
-	FRDGGPUScopeStacksByPipeline(FRHICommandListImmediate& RHICmdListGraphics, FRHIComputeCommandList& RHICmdListAsyncCompute);
+	FRDGGPUScopeStacksByPipeline(FRHICommandListImmediate& RHICmdListGraphics, FRHIComputeCommandList& RHICmdListAsyncCompute, FRDGAllocator& Allocator);
 
 	void BeginEventScope(FRDGEventName&& ScopeName);
 
@@ -372,7 +373,7 @@ public:
 class RENDERCORE_API FRDGCSVStatScopeStack final
 {
 public:
-	FRDGCSVStatScopeStack(FRHIComputeCommandList& RHICmdList, const char* UnaccountedStatName);
+	FRDGCSVStatScopeStack(FRHIComputeCommandList& RHICmdList, FRDGAllocator& Allocator, const char* UnaccountedStatName);
 
 	void BeginScope(const char* StatName);
 
@@ -426,7 +427,7 @@ struct FRDGCPUScopes
 
 struct FRDGCPUScopeStacks
 {
-	FRDGCPUScopeStacks(FRHIComputeCommandList& RHICmdList, const char* UnaccountedCSVStat);
+	FRDGCPUScopeStacks(FRHIComputeCommandList& RHICmdList, FRDGAllocator& Allocator, const char* UnaccountedCSVStat);
 
 	void BeginExecute();
 
diff --git a/Engine/Source/Runtime/RenderCore/Public/RenderGraphEvent.inl b/Engine/Source/Runtime/RenderCore/Public/RenderGraphEvent.inl
--- a/Engine/Source/Runtime/RenderCore/Public/RenderGraphEvent.inl
+++ b/Engine/Source/Runtime/RenderCore/Public/RenderGraphEvent.inl
@@ -73,10 +73,11 @@ void TRDGScopeStackHelper<TScopeType>::EndExecute(PopFunctionType PopFunction)
 template <typename TScopeType>
 TRDGScopeStack<TScopeType>::TRDGScopeStack(
 	FRHIComputeCommandList& InRHICmdList,
+	FRDGAllocator& InAllocator,
 	FPushFunction InPushFunction,
 	FPopFunction InPopFunction)
 	: RHICmdList(InRHICmdList)
-	, MemStack(FMemStack::Get())
+	, Allocator(InAllocator)
 	, PushFunction(InPushFunction)
 	, PopFunction(InPopFunction)
 {}
@@ -91,7 +92,7 @@ template <typename TScopeType>
 template <typename... TScopeConstructArgs>
 void TRDGScopeStack<TScopeType>::BeginScope(TScopeConstructArgs... ScopeConstructArgs)
 {
-	auto Scope = new(MemStack) TScopeType(CurrentScope, Forward<TScopeConstructArgs>(ScopeConstructArgs)...);
+	auto Scope = Allocator.AllocNoDestruct<TScopeType>(CurrentScope, Forward<TScopeConstructArgs>(ScopeConstructArgs)...);
 	Scopes.Add(Scope);
 	CurrentScope = Scope;
 }
@@ -214,9 +215,9 @@ inline const TCHAR* FRDGEventName::GetTCHAR() const
 
 #if RDG_GPU_SCOPES
 
-inline FRDGGPUScopeStacks::FRDGGPUScopeStacks(FRHIComputeCommandList& RHICmdList)
-	: Event(RHICmdList)
-	, Stat(RHICmdList)
+inline FRDGGPUScopeStacks::FRDGGPUScopeStacks(FRHIComputeCommandList& RHICmdList, FRDGAllocator& Allocator)
+	: Event(RHICmdList, Allocator)
+	, Stat(RHICmdList, Allocator)
 {}
 
 inline void FRDGGPUScopeStacks::BeginExecute()
@@ -250,9 +251,9 @@ inline FRDGGPUScopes FRDGGPUScopeStacks::GetCurrentScopes() const
 	return Scopes;
 }
 
-inline FRDGGPUScopeStacksByPipeline::FRDGGPUScopeStacksByPipeline(FRHICommandListImmediate& RHICmdListGraphics, FRHIComputeCommandList& RHICmdListAsyncCompute)
-	: Graphics(RHICmdListGraphics)
-	, AsyncCompute(RHICmdListAsyncCompute)
+inline FRDGGPUScopeStacksByPipeline::FRDGGPUScopeStacksByPipeline(FRHICommandListImmediate& RHICmdListGraphics, FRHIComputeCommandList& RHICmdListAsyncCompute, FRDGAllocator& Allocator)
+	: Graphics(RHICmdListGraphics, Allocator)
+	, AsyncCompute(RHICmdListAsyncCompute, Allocator)
 {}
 
 inline void FRDGGPUScopeStacksByPipeline::BeginEventScope(FRDGEventName&& ScopeName)
@@ -333,8 +334,8 @@ inline FRDGGPUScopes FRDGGPUScopeStacksByPipeline::GetCurrentScopes(ERHIPipeline
 
 #if RDG_CPU_SCOPES
 
-inline FRDGCPUScopeStacks::FRDGCPUScopeStacks(FRHIComputeCommandList& RHICmdList, const char* UnaccountedCSVStat)
-	: CSV(RHICmdList, UnaccountedCSVStat)
+inline FRDGCPUScopeStacks::FRDGCPUScopeStacks(FRHIComputeCommandList& RHICmdList, FRDGAllocator& Allocator, const char* UnaccountedCSVStat)
+	: CSV(RHICmdList, Allocator, UnaccountedCSVStat)
 {}
 
 inline void FRDGCPUScopeStacks::BeginExecute()
