Shrink scattered upload buffers to prevent growth to some unbounded high watermark. This prevents pathalogical case of growing beyond the size supported by optimal RHI buffer allocation path.

diff --git a/Engine/Source/Runtime/RenderCore/Private/UnifiedBuffer.cpp b/Engine/Source/Runtime/RenderCore/Private/UnifiedBuffer.cpp
--- a/Engine/Source/Runtime/RenderCore/Private/UnifiedBuffer.cpp
+++ b/Engine/Source/Runtime/RenderCore/Private/UnifiedBuffer.cpp
@@ -459,12 +459,13 @@ void FScatterUploadBuffer::Init( uint32 NumElements, uint32 InNumBytesPerElement
 	const uint32 TypeSize = bInFloat4Buffer ? 16 : 4;
 
 	uint32 ScatterBytes = NumElements * sizeof( uint32 );
+	uint32 ScatterBufferSize = NumScattersAllocated * sizeof(uint32);
 
-	if( ScatterBytes > ScatterBuffer.NumBytes )
+	if( ScatterBytes > ScatterBuffer.NumBytes || ScatterBufferSize < ScatterBuffer.NumBytes / 2)
 	{
 		// Resize Scatter Buffer
 		ScatterBuffer.Release();
-		ScatterBuffer.NumBytes = NumScattersAllocated * sizeof( uint32 );
+		ScatterBuffer.NumBytes = ScatterBufferSize;
 
 		FRHIResourceCreateInfo CreateInfo(DebugName);
 		ScatterBuffer.Buffer = RHICreateStructuredBuffer( sizeof( uint32 ), ScatterBuffer.NumBytes, BUF_ShaderResource | BUF_Volatile | Usage, CreateInfo );
@@ -472,11 +473,13 @@ void FScatterUploadBuffer::Init( uint32 NumElements, uint32 InNumBytesPerElement
 	}
 
 	uint32 UploadBytes = NumElements * NumBytesPerElement;
-	if( UploadBytes > UploadBuffer.NumBytes )
+	uint32 UploadBufferSize = NumScattersAllocated * NumBytesPerElement;
+
+	if( UploadBytes > UploadBuffer.NumBytes || UploadBufferSize < UploadBuffer.NumBytes / 2)
 	{
 		// Resize Upload Buffer
 		UploadBuffer.Release();
-		UploadBuffer.NumBytes = NumScattersAllocated * NumBytesPerElement;
+		UploadBuffer.NumBytes = UploadBufferSize;
 
 		FRHIResourceCreateInfo CreateInfo(DebugName);
 		UploadBuffer.Buffer = RHICreateStructuredBuffer( TypeSize, UploadBuffer.NumBytes, BUF_ShaderResource | BUF_Volatile | Usage, CreateInfo );
