Buffers need UAV flag set to be transient - otherwise it could be upload resource

diff --git a/Engine/Source/Runtime/RenderCore/Private/RenderGraphBuilder.cpp b/Engine/Source/Runtime/RenderCore/Private/RenderGraphBuilder.cpp
--- a/Engine/Source/Runtime/RenderCore/Private/RenderGraphBuilder.cpp
+++ b/Engine/Source/Runtime/RenderCore/Private/RenderGraphBuilder.cpp
@@ -480,7 +480,7 @@ bool FRDGBuilder::IsTransient(FRDGBufferRef Buffer) const
 	// Buffers smaller than a page go through normal pooled allocation.
 	const uint32 TransientSizeMin = 64 * 1024;
 
-	return Buffer->Desc.GetTotalNumBytes() >= TransientSizeMin;
+	return Buffer->Desc.GetTotalNumBytes() >= TransientSizeMin && EnumHasAnyFlags(Buffer->Desc.Usage, BUF_UnorderedAccess);
 }
 
 bool FRDGBuilder::IsTransient(FRDGTextureRef Texture) const
