Groundwork for managing ray tracing resources in RDG.

- Added BVHRead to ERHIAccess::ReadOnlyExclusiveMask BVHRead
- Added FRDGBufferDesc::EUnderlyingType::AccelerationStructure
- Added acceleration structure handling in EnumerateBufferAccess, FindFreeBufferInternal, GetOrCreateSRV, FRDGBufferSRVDesc

diff --git a/Engine/Source/Runtime/RHI/Public/RHI.h b/Engine/Source/Runtime/RHI/Public/RHI.h
--- a/Engine/Source/Runtime/RHI/Public/RHI.h
+++ b/Engine/Source/Runtime/RHI/Public/RHI.h
@@ -893,7 +893,7 @@ enum class ERHIAccess
 	UAVMask = UAVCompute | UAVGraphics,
 
 	// A mask of all bits representing read-only states which cannot be combined with other write states.
-	ReadOnlyExclusiveMask = CPURead | Present | IndirectArgs | VertexOrIndexBuffer | SRVGraphics | SRVCompute | CopySrc | ResolveSrc,
+	ReadOnlyExclusiveMask = CPURead | Present | IndirectArgs | VertexOrIndexBuffer | SRVGraphics | SRVCompute | CopySrc | ResolveSrc | BVHRead,
 
 	// A mask of all bits representing read-only states which may be combined with other write states.
 	ReadOnlyMask = ReadOnlyExclusiveMask | DSVRead,
diff --git a/Engine/Source/Runtime/RenderCore/Private/RenderGraphBuilder.cpp b/Engine/Source/Runtime/RenderCore/Private/RenderGraphBuilder.cpp
--- a/Engine/Source/Runtime/RenderCore/Private/RenderGraphBuilder.cpp
+++ b/Engine/Source/Runtime/RenderCore/Private/RenderGraphBuilder.cpp
@@ -304,7 +304,15 @@ void EnumerateBufferAccess(FRDGParameterStruct PassParameters, ERDGPassFlags Pas
 		case UBMT_RDG_BUFFER_SRV:
 			if (FRDGBufferSRVRef SRV = Parameter.GetAsBufferSRV())
 			{
-				AccessFunction(SRV, SRV->GetParent(), SRVAccess);
+				FRDGBufferRef Buffer = SRV->GetParent();
+				ERHIAccess BufferAccess = SRVAccess;
+
+				if (EnumHasAnyFlags(Buffer->Desc.Usage, BUF_AccelerationStructure))
+				{
+					BufferAccess = ERHIAccess::BVHRead;
+				}
+
+				AccessFunction(SRV, Buffer, BufferAccess);
 			}
 		break;
 		case UBMT_RDG_BUFFER_UAV:
diff --git a/Engine/Source/Runtime/RenderCore/Private/RenderGraphResourcePool.cpp b/Engine/Source/Runtime/RenderCore/Private/RenderGraphResourcePool.cpp
--- a/Engine/Source/Runtime/RenderCore/Private/RenderGraphResourcePool.cpp
+++ b/Engine/Source/Runtime/RenderCore/Private/RenderGraphResourcePool.cpp
@@ -65,6 +65,10 @@ TRefCountPtr<FRDGPooledBuffer> FRenderGraphResourcePool::FindFreeBufferInternal(
 		{
 			BufferRHI = RHICreateStructuredBuffer(Desc.BytesPerElement, NumBytes, Desc.Usage, CreateInfo);
 		}
+		else if (Desc.UnderlyingType == FRDGBufferDesc::EUnderlyingType::AccelerationStructure)
+		{
+			BufferRHI = RHICreateBuffer(NumBytes, Desc.Usage, 0, ERHIAccess::BVHWrite, CreateInfo);
+		}
 		else
 		{
 			check(0);
diff --git a/Engine/Source/Runtime/RenderCore/Private/RenderGraphResources.cpp b/Engine/Source/Runtime/RenderCore/Private/RenderGraphResources.cpp
--- a/Engine/Source/Runtime/RenderCore/Private/RenderGraphResources.cpp
+++ b/Engine/Source/Runtime/RenderCore/Private/RenderGraphResources.cpp
@@ -339,6 +339,10 @@ FRHIShaderResourceView* FRDGPooledBuffer::GetOrCreateSRV(FRDGBufferSRVDesc SRVDe
 	{
 		RHIShaderResourceView = RHICreateShaderResourceView(Buffer);
 	}
+	else if (Desc.UnderlyingType == FRDGBufferDesc::EUnderlyingType::AccelerationStructure)
+	{
+		RHIShaderResourceView = RHICreateShaderResourceView(Buffer);
+	}
 	else
 	{
 		checkNoEntry();
diff --git a/Engine/Source/Runtime/RenderCore/Public/RenderGraphResources.h b/Engine/Source/Runtime/RenderCore/Public/RenderGraphResources.h
--- a/Engine/Source/Runtime/RenderCore/Public/RenderGraphResources.h
+++ b/Engine/Source/Runtime/RenderCore/Public/RenderGraphResources.h
@@ -746,7 +746,8 @@ struct FRDGBufferDesc
 	enum class EUnderlyingType
 	{
 		VertexBuffer,
-		StructuredBuffer
+		StructuredBuffer,
+		AccelerationStructure
 	};
 
 	/** Create the descriptor for an indirect RHI call.
diff --git a/Engine/Source/Runtime/RenderCore/Public/RenderGraphResources.inl b/Engine/Source/Runtime/RenderCore/Public/RenderGraphResources.inl
--- a/Engine/Source/Runtime/RenderCore/Public/RenderGraphResources.inl
+++ b/Engine/Source/Runtime/RenderCore/Public/RenderGraphResources.inl
@@ -243,11 +243,15 @@ inline FRDGTextureSubresourceRange FRDGTextureUAV::GetSubresourceRange() const
 inline FRDGBufferSRVDesc::FRDGBufferSRVDesc(FRDGBufferRef InBuffer)
 	: Buffer(InBuffer)
 {
-	if (Buffer->Desc.Usage & BUF_DrawIndirect)
+	if (EnumHasAnyFlags(Buffer->Desc.Usage, BUF_DrawIndirect))
 	{
 		BytesPerElement = 4;
 		Format = PF_R32_UINT;
 	}
+	else if (EnumHasAnyFlags(Buffer->Desc.Usage, BUF_AccelerationStructure))
+	{
+		// nothing special here
+	}
 	else
 	{
 		checkf(Buffer->Desc.UnderlyingType != FRDGBufferDesc::EUnderlyingType::VertexBuffer, TEXT("VertexBuffer %s requires a type when creating a SRV."), Buffer->Name);
@@ -257,7 +261,7 @@ inline FRDGBufferSRVDesc::FRDGBufferSRVDesc(FRDGBufferRef InBuffer)
 inline FRDGBufferUAVDesc::FRDGBufferUAVDesc(FRDGBufferRef InBuffer)
 	: Buffer(InBuffer)
 {
-	if (Buffer->Desc.Usage & BUF_DrawIndirect)
+	if (EnumHasAnyFlags(Buffer->Desc.Usage, BUF_DrawIndirect))
 	{
 		Format = PF_R32_UINT;
 	}
