Cleaned up a few more references to UE4.

#rnx

diff --git a/Engine/Source/Runtime/RenderCore/Public/ShaderCompilerCore.h b/Engine/Source/Runtime/RenderCore/Public/ShaderCompilerCore.h
--- a/Engine/Source/Runtime/RenderCore/Public/ShaderCompilerCore.h
+++ b/Engine/Source/Runtime/RenderCore/Public/ShaderCompilerCore.h
@@ -135,7 +135,7 @@ struct FExtraShaderCompilerSettings
 
 	friend FArchive& operator<<(FArchive& Ar, FExtraShaderCompilerSettings& StatsSettings)
 	{
-		// Note: this serialize is used to pass between UE4 and the shader compile worker, recompile both when modifying
+		// Note: this serialize is used to pass between UE and the shader compile worker, recompile both when modifying
 		return Ar << StatsSettings.bExtractShaderSource << StatsSettings.OfflineCompilerPath;
 	}
 };
@@ -160,7 +160,7 @@ struct FShaderCompilerInput
 	bool bIncludeUsedOutputs;
 	TArray<FString> UsedOutputs;
 
-	// Dump debug path (up to platform) e.g. "D:/MMittring-Z3941-A/UE4-Orion/OrionGame/Saved/ShaderDebugInfo/PCD3D_SM5"
+	// Dump debug path (up to platform) e.g. "D:/Project/Saved/ShaderDebugInfo/PCD3D_SM5"
 	FString DumpDebugInfoRootPath;
 	// only used if enabled by r.DumpShaderDebugInfo (platform/groupname) e.g. ""
 	FString DumpDebugInfoPath;
@@ -452,7 +452,7 @@ struct FShaderCompilerOutput
 
 	friend FArchive& operator<<(FArchive& Ar, FShaderCompilerOutput& Output)
 	{
-		// Note: this serialize is used to pass between UE4 and the shader compile worker, recompile both when modifying
+		// Note: this serialize is used to pass between UE and the shader compile worker, recompile both when modifying
 		Ar << Output.ParameterMap << Output.Errors << Output.Target << Output.ShaderCode << Output.OutputHash << Output.NumInstructions << Output.NumTextureSamplers << Output.bSucceeded;
 		Ar << Output.bFailedRemovingUnused << Output.bSupportsQueryingUsedAttributes << Output.UsedAttributes;
 		Ar << Output.CompileTime;
