Fix for RDG submitting an empty transition array.

diff --git a/Engine/Source/Runtime/RenderCore/Private/RenderGraphPass.cpp b/Engine/Source/Runtime/RenderCore/Private/RenderGraphPass.cpp
--- a/Engine/Source/Runtime/RenderCore/Private/RenderGraphPass.cpp
+++ b/Engine/Source/Runtime/RenderCore/Private/RenderGraphPass.cpp
@@ -127,14 +127,20 @@ void FRDGTransitionQueue::Insert(const FRHITransition* Transition, ERHITransitio
 
 void FRDGTransitionQueue::Begin(FRHIComputeCommandList& RHICmdList)
 {
-	RHICmdList.BeginTransitions(Queue);
-	Queue.Empty();
+	if (Queue.Num())
+	{
+		RHICmdList.BeginTransitions(Queue);
+		Queue.Empty();
+	}
 }
 
 void FRDGTransitionQueue::End(FRHIComputeCommandList& RHICmdList)
 {
-	RHICmdList.EndTransitions(Queue);
-	Queue.Empty();
+	if (Queue.Num())
+	{
+		RHICmdList.EndTransitions(Queue);
+		Queue.Empty();
+	}
 }
 
 FRDGBarrierBatchBegin::FRDGBarrierBatchBegin(ERHIPipeline InPipelinesToBegin, ERHIPipeline InPipelinesToEnd, const TCHAR* InDebugName, const FRDGPass* InDebugPass)
