Fix duplicated RefCount between RDGPooledResource and RDGPooledBuffer.

diff --git a/Engine/Source/Runtime/RenderCore/Public/RenderGraphResources.h b/Engine/Source/Runtime/RenderCore/Public/RenderGraphResources.h
--- a/Engine/Source/Runtime/RenderCore/Public/RenderGraphResources.h
+++ b/Engine/Source/Runtime/RenderCore/Public/RenderGraphResources.h
@@ -1007,7 +1007,6 @@ private:
 	FRDGBufferRef Owner = nullptr;
 	FRDGSubresourceState State;
 
-	mutable uint32 RefCount = 0;
 	uint32 LastUsedFrame = 0;
 
 	friend FRenderGraphResourcePool;
