Added EVertexFactoryFlags::SupportsRayTracingDynamicGeometry to avoid FindVertexFactoryType() calls in IsSupportedDynamicVertexFactoryType().

diff --git a/Engine/Plugins/FX/Niagara/Source/NiagaraVertexFactories/Private/NiagaraRibbonVertexFactory.cpp b/Engine/Plugins/FX/Niagara/Source/NiagaraVertexFactories/Private/NiagaraRibbonVertexFactory.cpp
--- a/Engine/Plugins/FX/Niagara/Source/NiagaraVertexFactories/Private/NiagaraRibbonVertexFactory.cpp
+++ b/Engine/Plugins/FX/Niagara/Source/NiagaraVertexFactories/Private/NiagaraRibbonVertexFactory.cpp
@@ -187,4 +187,5 @@ IMPLEMENT_VERTEX_FACTORY_TYPE(FNiagaraRibbonVertexFactory, "/Plugin/FX/Niagara/P
 	  EVertexFactoryFlags::UsedWithMaterials 
 	| EVertexFactoryFlags::SupportsDynamicLighting
 	| EVertexFactoryFlags::SupportsRayTracing
+	| EVertexFactoryFlags::SupportsRayTracingDynamicGeometry
 );
diff --git a/Engine/Plugins/FX/Niagara/Source/NiagaraVertexFactories/Private/NiagaraSpriteVertexFactory.cpp b/Engine/Plugins/FX/Niagara/Source/NiagaraVertexFactories/Private/NiagaraSpriteVertexFactory.cpp
--- a/Engine/Plugins/FX/Niagara/Source/NiagaraVertexFactories/Private/NiagaraSpriteVertexFactory.cpp
+++ b/Engine/Plugins/FX/Niagara/Source/NiagaraVertexFactories/Private/NiagaraSpriteVertexFactory.cpp
@@ -210,6 +210,7 @@ IMPLEMENT_VERTEX_FACTORY_TYPE(FNiagaraSpriteVertexFactory,"/Plugin/FX/Niagara/Pr
 	  EVertexFactoryFlags::UsedWithMaterials
 	| EVertexFactoryFlags::SupportsDynamicLighting
 	| EVertexFactoryFlags::SupportsRayTracing
+	| EVertexFactoryFlags::SupportsRayTracingDynamicGeometry
 );
 
 /////////////////////////////////////////////////////////////////////////////
diff --git a/Engine/Source/Runtime/Engine/Private/GPUSkinVertexFactory.cpp b/Engine/Source/Runtime/Engine/Private/GPUSkinVertexFactory.cpp
--- a/Engine/Source/Runtime/Engine/Private/GPUSkinVertexFactory.cpp
+++ b/Engine/Source/Runtime/Engine/Private/GPUSkinVertexFactory.cpp
@@ -895,6 +895,7 @@ IMPLEMENT_VERTEX_FACTORY_TYPE(FGPUSkinPassthroughVertexFactory, "/Engine/Private
 	  EVertexFactoryFlags::UsedWithMaterials
 	| EVertexFactoryFlags::SupportsDynamicLighting
 	| EVertexFactoryFlags::SupportsRayTracing
+	| EVertexFactoryFlags::SupportsRayTracingDynamicGeometry
 	| (GPUCULL_TODO ? EVertexFactoryFlags::SupportsPrimitiveIdStream : EVertexFactoryFlags::None)
 );
 
diff --git a/Engine/Source/Runtime/Engine/Private/InstancedStaticMesh.cpp b/Engine/Source/Runtime/Engine/Private/InstancedStaticMesh.cpp
--- a/Engine/Source/Runtime/Engine/Private/InstancedStaticMesh.cpp
+++ b/Engine/Source/Runtime/Engine/Private/InstancedStaticMesh.cpp
@@ -874,6 +874,7 @@ IMPLEMENT_VERTEX_FACTORY_TYPE(FInstancedStaticMeshVertexFactory,"/Engine/Private
 	| EVertexFactoryFlags::SupportsPositionOnly
 	| EVertexFactoryFlags::SupportsCachingMeshDrawCommands
 	| EVertexFactoryFlags::SupportsRayTracing
+	| EVertexFactoryFlags::SupportsRayTracingDynamicGeometry
 	| (GPUCULL_TODO ? EVertexFactoryFlags::SupportsPrimitiveIdStream : EVertexFactoryFlags::None)
 );
 
diff --git a/Engine/Source/Runtime/Engine/Private/LocalVertexFactory.cpp b/Engine/Source/Runtime/Engine/Private/LocalVertexFactory.cpp
--- a/Engine/Source/Runtime/Engine/Private/LocalVertexFactory.cpp
+++ b/Engine/Source/Runtime/Engine/Private/LocalVertexFactory.cpp
@@ -405,4 +405,5 @@ IMPLEMENT_VERTEX_FACTORY_TYPE(FLocalVertexFactory,"/Engine/Private/LocalVertexFa
 	| EVertexFactoryFlags::SupportsCachingMeshDrawCommands
 	| EVertexFactoryFlags::SupportsPrimitiveIdStream
 	| EVertexFactoryFlags::SupportsRayTracing
+	| EVertexFactoryFlags::SupportsRayTracingDynamicGeometry
 );
diff --git a/Engine/Source/Runtime/Landscape/Private/LandscapeRender.cpp b/Engine/Source/Runtime/Landscape/Private/LandscapeRender.cpp
--- a/Engine/Source/Runtime/Landscape/Private/LandscapeRender.cpp
+++ b/Engine/Source/Runtime/Landscape/Private/LandscapeRender.cpp
@@ -3233,6 +3233,7 @@ IMPLEMENT_VERTEX_FACTORY_TYPE(FLandscapeVertexFactory, "/Engine/Private/Landscap
 	| EVertexFactoryFlags::SupportsDynamicLighting
 	| EVertexFactoryFlags::SupportsCachingMeshDrawCommands
 	| EVertexFactoryFlags::SupportsRayTracing
+	| EVertexFactoryFlags::SupportsRayTracingDynamicGeometry
 	| (GPUCULL_TODO ? EVertexFactoryFlags::SupportsPrimitiveIdStream : EVertexFactoryFlags::None)
 );
 
@@ -3276,6 +3277,7 @@ IMPLEMENT_VERTEX_FACTORY_TYPE(FLandscapeXYOffsetVertexFactory, "/Engine/Private/
 	| EVertexFactoryFlags::SupportsDynamicLighting
 	| EVertexFactoryFlags::SupportsCachingMeshDrawCommands
 	| EVertexFactoryFlags::SupportsRayTracing
+	| EVertexFactoryFlags::SupportsRayTracingDynamicGeometry
 	| (GPUCULL_TODO ? EVertexFactoryFlags::SupportsPrimitiveIdStream : EVertexFactoryFlags::None)
 );
 
@@ -3302,6 +3304,7 @@ IMPLEMENT_VERTEX_FACTORY_TYPE(FLandscapeFixedGridVertexFactory, "/Engine/Private
 	| EVertexFactoryFlags::SupportsDynamicLighting
 	| EVertexFactoryFlags::SupportsCachingMeshDrawCommands
 	| EVertexFactoryFlags::SupportsRayTracing
+	| EVertexFactoryFlags::SupportsRayTracingDynamicGeometry
 	| (GPUCULL_TODO ? EVertexFactoryFlags::SupportsPrimitiveIdStream : EVertexFactoryFlags::None)
 );
 
diff --git a/Engine/Source/Runtime/RenderCore/Public/VertexFactory.h b/Engine/Source/Runtime/RenderCore/Public/VertexFactory.h
--- a/Engine/Source/Runtime/RenderCore/Public/VertexFactory.h
+++ b/Engine/Source/Runtime/RenderCore/Public/VertexFactory.h
@@ -91,16 +91,17 @@ enum class EVertexInputStreamType : uint8
 
 enum class EVertexFactoryFlags : uint32
 {
-	None                            = 0u,
-	UsedWithMaterials               = 1u << 1,
-	SupportsStaticLighting          = 1u << 2,
-	SupportsDynamicLighting         = 1u << 3,
-	SupportsPrecisePrevWorldPos     = 1u << 4,
-	SupportsPositionOnly            = 1u << 5,
-	SupportsCachingMeshDrawCommands = 1u << 6,
-	SupportsPrimitiveIdStream       = 1u << 7,
-	SupportsNaniteRendering         = 1u << 8,
-	SupportsRayTracing              = 1u << 9,
+	None                              = 0u,
+	UsedWithMaterials                 = 1u << 1,
+	SupportsStaticLighting            = 1u << 2,
+	SupportsDynamicLighting           = 1u << 3,
+	SupportsPrecisePrevWorldPos       = 1u << 4,
+	SupportsPositionOnly              = 1u << 5,
+	SupportsCachingMeshDrawCommands   = 1u << 6,
+	SupportsPrimitiveIdStream         = 1u << 7,
+	SupportsNaniteRendering           = 1u << 8,
+	SupportsRayTracing                = 1u << 9,
+	SupportsRayTracingDynamicGeometry = 1u << 10,
 };
 ENUM_CLASS_FLAGS(EVertexFactoryFlags);
 
@@ -398,15 +399,16 @@ public:
 		return EnumHasAllFlags(Flags, FlagsToCheck);
 	}
 
-	bool IsUsedWithMaterials() const             { return HasFlags(EVertexFactoryFlags::UsedWithMaterials); }
-	bool SupportsStaticLighting() const          { return HasFlags(EVertexFactoryFlags::SupportsStaticLighting); }
-	bool SupportsDynamicLighting() const         { return HasFlags(EVertexFactoryFlags::SupportsDynamicLighting); }
-	bool SupportsPrecisePrevWorldPos() const     { return HasFlags(EVertexFactoryFlags::SupportsPrecisePrevWorldPos); }
-	bool SupportsPositionOnly() const            { return HasFlags(EVertexFactoryFlags::SupportsPositionOnly); }
-	bool SupportsCachingMeshDrawCommands() const { return HasFlags(EVertexFactoryFlags::SupportsCachingMeshDrawCommands); }
-	bool SupportsPrimitiveIdStream() const       { return HasFlags(EVertexFactoryFlags::SupportsPrimitiveIdStream); }
-	bool SupportsNaniteRendering() const         { return HasFlags(EVertexFactoryFlags::SupportsNaniteRendering); }
-	bool SupportsRayTracing() const              { return HasFlags(EVertexFactoryFlags::SupportsRayTracing); }
+	bool IsUsedWithMaterials() const                { return HasFlags(EVertexFactoryFlags::UsedWithMaterials); }
+	bool SupportsStaticLighting() const             { return HasFlags(EVertexFactoryFlags::SupportsStaticLighting); }
+	bool SupportsDynamicLighting() const            { return HasFlags(EVertexFactoryFlags::SupportsDynamicLighting); }
+	bool SupportsPrecisePrevWorldPos() const        { return HasFlags(EVertexFactoryFlags::SupportsPrecisePrevWorldPos); }
+	bool SupportsPositionOnly() const               { return HasFlags(EVertexFactoryFlags::SupportsPositionOnly); }
+	bool SupportsCachingMeshDrawCommands() const    { return HasFlags(EVertexFactoryFlags::SupportsCachingMeshDrawCommands); }
+	bool SupportsPrimitiveIdStream() const          { return HasFlags(EVertexFactoryFlags::SupportsPrimitiveIdStream); }
+	bool SupportsNaniteRendering() const            { return HasFlags(EVertexFactoryFlags::SupportsNaniteRendering); }
+	bool SupportsRayTracing() const                 { return HasFlags(EVertexFactoryFlags::SupportsRayTracing); }
+	bool SupportsRayTracingDynamicGeometry() const  { return HasFlags(EVertexFactoryFlags::SupportsRayTracingDynamicGeometry); }
 
 	// Hash function.
 	friend uint32 GetTypeHash(const FVertexFactoryType* Type)
diff --git a/Engine/Source/Runtime/Renderer/Private/RayTracing/RayTracingDynamicGeometry.cpp b/Engine/Source/Runtime/Renderer/Private/RayTracing/RayTracingDynamicGeometry.cpp
--- a/Engine/Source/Runtime/Renderer/Private/RayTracing/RayTracingDynamicGeometry.cpp
+++ b/Engine/Source/Runtime/Renderer/Private/RayTracing/RayTracingDynamicGeometry.cpp
@@ -14,14 +14,7 @@ DECLARE_CYCLE_STAT(TEXT("RTDynGeomBuild"), STAT_CLM_RTDynGeomBuild, STATGROUP_Pa
 
 static bool IsSupportedDynamicVertexFactoryType(const FVertexFactoryType* VertexFactoryType)
 {
-	return VertexFactoryType == FindVertexFactoryType(FName(TEXT("FNiagaraSpriteVertexFactory"), FNAME_Find))
-		|| VertexFactoryType == FindVertexFactoryType(FName(TEXT("FNiagaraRibbonVertexFactory"), FNAME_Find))
-		|| VertexFactoryType == FindVertexFactoryType(FName(TEXT("FLocalVertexFactory"), FNAME_Find))
-		|| VertexFactoryType == FindVertexFactoryType(FName(TEXT("FLandscapeVertexFactory"), FNAME_Find))
-		|| VertexFactoryType == FindVertexFactoryType(FName(TEXT("FLandscapeFixedGridVertexFactory"), FNAME_Find))
-		|| VertexFactoryType == FindVertexFactoryType(FName(TEXT("FLandscapeXYOffsetVertexFactory"), FNAME_Find))
-		|| VertexFactoryType == FindVertexFactoryType(FName(TEXT("FGPUSkinPassthroughVertexFactory"), FNAME_Find))
-		|| VertexFactoryType == FindVertexFactoryType(FName(TEXT("FInstancedStaticMeshVertexFactory"), FNAME_Find));
+	return VertexFactoryType->SupportsRayTracingDynamicGeometry();
 }
 
 class FRayTracingDynamicGeometryConverterCS : public FMeshMaterialShader
