Fix crashes when dumping shaderlib debug info (UE-106302).

#rnx

diff --git a/Engine/Source/Runtime/RenderCore/Private/ShaderCodeArchive.cpp b/Engine/Source/Runtime/RenderCore/Private/ShaderCodeArchive.cpp
--- a/Engine/Source/Runtime/RenderCore/Private/ShaderCodeArchive.cpp
+++ b/Engine/Source/Runtime/RenderCore/Private/ShaderCodeArchive.cpp
@@ -570,7 +570,7 @@ void FSerializedShaderArchive::DumpContentsInPlaintext(FString& OutText) const
 
 	{
 		Out << TEXT("\tShaderHashes\n\t{\n");
-		for (int32 IdxHash = 0, NumHashes = ShaderMapHashes.Num(); IdxHash < NumHashes; ++IdxHash)
+		for (int32 IdxHash = 0, NumHashes = ShaderHashes.Num(); IdxHash < NumHashes; ++IdxHash)
 		{
 			Out << TEXT("\t\t");
 			Out << ShaderHashes[IdxHash].ToString();
