Added RDG support for cube volume textures.

diff --git a/Engine/Source/Runtime/RenderCore/Private/RenderUtils.cpp b/Engine/Source/Runtime/RenderCore/Private/RenderUtils.cpp
--- a/Engine/Source/Runtime/RenderCore/Private/RenderUtils.cpp
+++ b/Engine/Source/Runtime/RenderCore/Private/RenderUtils.cpp
@@ -24,7 +24,7 @@ FTextureWithRDG::~FTextureWithRDG() = default;
 FRDGTexture* FTextureWithRDG::GetRDG(FRDGBuilder& GraphBuilder) const
 {
 	checkf(RenderTarget, TEXT("InitRDG was not called before use."));
-	return GraphBuilder.RegisterExternalTexture(RenderTarget);
+	return GraphBuilder.RegisterExternalTexture(RenderTarget, ERenderTargetTexture::ShaderResource, ERDGTextureFlags::ReadOnly);
 }
 
 void FTextureWithRDG::ReleaseRHI()
@@ -468,7 +468,7 @@ RENDERCORE_API int32 GMipColorTextureMipLevels = FMipColorTexture::NumMips;
 RENDERCORE_API const uint32 GDiffuseConvolveMipLevel = 4;
 
 /** A solid color cube texture. */
-class FSolidColorTextureCube : public FTexture
+class FSolidColorTextureCube : public FTextureWithRDG
 {
 public:
 	FSolidColorTextureCube(const FColor& InColor)
@@ -487,7 +487,9 @@ public:
 	virtual void InitRHI() override
 	{
 		// Create the texture RHI.
-		FRHIResourceCreateInfo CreateInfo(TEXT("SolidColorCube"));
+		const TCHAR* Name = TEXT("SolidColorCube");
+
+		FRHIResourceCreateInfo CreateInfo(Name);
 		FTextureCubeRHIRef TextureCube = RHICreateTextureCube(1, PixelFormat, 1, TexCreate_ShaderResource, CreateInfo);
 		TextureRHI = TextureCube;
 
@@ -510,6 +512,8 @@ public:
 		// Create the sampler state RHI resource.
 		FSamplerStateInitializerRHI SamplerStateInitializer(SF_Point, AM_Wrap, AM_Wrap, AM_Wrap);
 		SamplerStateRHI = GetOrCreateSamplerState(SamplerStateInitializer);
+
+		FTextureWithRDG::InitRDG(Name);
 	}
 
 	/** Returns the width of the texture in pixels. */
@@ -536,7 +540,7 @@ class FWhiteTextureCube : public FSolidColorTextureCube
 public:
 	FWhiteTextureCube() : FSolidColorTextureCube(FColor::White) {}
 };
-FTexture* GWhiteTextureCube = new TGlobalResource<FWhiteTextureCube>;
+FTextureWithRDG* GWhiteTextureCube = new TGlobalResource<FWhiteTextureCube>;
 
 /** A black cube texture. */
 class FBlackTextureCube : public FSolidColorTextureCube
@@ -544,7 +548,7 @@ class FBlackTextureCube : public FSolidColorTextureCube
 public:
 	FBlackTextureCube() : FSolidColorTextureCube(FColor::Black) {}
 };
-FTexture* GBlackTextureCube = new TGlobalResource<FBlackTextureCube>;
+FTextureWithRDG* GBlackTextureCube = new TGlobalResource<FBlackTextureCube>;
 
 /** A black cube texture. */
 class FBlackTextureDepthCube : public FSolidColorTextureCube
@@ -552,9 +556,9 @@ class FBlackTextureDepthCube : public FSolidColorTextureCube
 public:
 	FBlackTextureDepthCube() : FSolidColorTextureCube(PF_ShadowDepth) {}
 };
-FTexture* GBlackTextureDepthCube = new TGlobalResource<FBlackTextureDepthCube>;
+FTextureWithRDG* GBlackTextureDepthCube = new TGlobalResource<FBlackTextureDepthCube>;
 
-class FBlackCubeArrayTexture : public FTexture
+class FBlackCubeArrayTexture : public FTextureWithRDG
 {
 public:
 	// FResource interface.
@@ -562,8 +566,10 @@ public:
 	{
 		if (SupportsTextureCubeArray(GetFeatureLevel() ))
 		{
+			const TCHAR* Name = TEXT("BlackCubeArray");
+
 			// Create the texture RHI.
-			FRHIResourceCreateInfo CreateInfo(TEXT("BlackCubeArray"));
+			FRHIResourceCreateInfo CreateInfo(Name);
 			FTextureCubeRHIRef TextureCubeArray = RHICreateTextureCubeArray(1,1,PF_B8G8R8A8,1,TexCreate_ShaderResource,CreateInfo);
 			TextureRHI = TextureCubeArray;
 
@@ -579,6 +585,8 @@ public:
 			// Create the sampler state RHI resource.
 			FSamplerStateInitializerRHI SamplerStateInitializer(SF_Point,AM_Wrap,AM_Wrap,AM_Wrap);
 			SamplerStateRHI = GetOrCreateSamplerState(SamplerStateInitializer);
+
+			FTextureWithRDG::InitRDG(Name);
 		}
 	}
 
@@ -594,7 +602,7 @@ public:
 		return 1;
 	}
 };
-FTexture* GBlackCubeArrayTexture = new TGlobalResource<FBlackCubeArrayTexture>;
+FTextureWithRDG* GBlackCubeArrayTexture = new TGlobalResource<FBlackCubeArrayTexture>;
 
 /**
  * A UINT 1x1 texture.
diff --git a/Engine/Source/Runtime/RenderCore/Public/RenderUtils.h b/Engine/Source/Runtime/RenderCore/Public/RenderUtils.h
--- a/Engine/Source/Runtime/RenderCore/Public/RenderUtils.h
+++ b/Engine/Source/Runtime/RenderCore/Public/RenderUtils.h
@@ -171,16 +171,16 @@ extern RENDERCORE_API class FTexture* GBlackUintTexture;
 extern RENDERCORE_API class FTextureWithRDG* GBlackUintVolumeTexture;
 
 /** A global white cube texture. */
-extern RENDERCORE_API class FTexture* GWhiteTextureCube;
+extern RENDERCORE_API class FTextureWithRDG* GWhiteTextureCube;
 
 /** A global black cube texture. */
-extern RENDERCORE_API class FTexture* GBlackTextureCube;
+extern RENDERCORE_API class FTextureWithRDG* GBlackTextureCube;
 
 /** A global black cube depth texture. */
-extern RENDERCORE_API class FTexture* GBlackTextureDepthCube;
+extern RENDERCORE_API class FTextureWithRDG* GBlackTextureDepthCube;
 
 /** A global black cube array texture. */
-extern RENDERCORE_API class FTexture* GBlackCubeArrayTexture;
+extern RENDERCORE_API class FTextureWithRDG* GBlackCubeArrayTexture;
 
 /** A global texture that has a different solid color in each mip-level. */
 extern RENDERCORE_API class FTexture* GMipColorTexture;
