Added initial support for RDG_{TEXTURE, BUFFER}_ACCESS_ARRAY, which enables a dynamic number of inputs into an RDG pass. Removed unused / deprecated SHADER_PARAMETER_RDG_BUFFER.

diff --git a/Engine/Shaders/Public/ShaderVersion.ush b/Engine/Shaders/Public/ShaderVersion.ush
--- a/Engine/Shaders/Public/ShaderVersion.ush
+++ b/Engine/Shaders/Public/ShaderVersion.ush
@@ -2,4 +2,4 @@
 // Each time the console command is executed it generates a new GUID. As this file is included
 // in Platform.ush (which should be included in any shader) it allows to invalidate the shader DDC.
 // 
-// GUID = 630C3067-81FD-451B-8767-DD5F261F288C
+// GUID = 196E4485-47A1-477A-8434-7000E7CA65CC
\ No newline at end of file
diff --git a/Engine/Source/Developer/Apple/MetalShaderFormat/Private/MetalDerivedData.cpp b/Engine/Source/Developer/Apple/MetalShaderFormat/Private/MetalDerivedData.cpp
--- a/Engine/Source/Developer/Apple/MetalShaderFormat/Private/MetalDerivedData.cpp
+++ b/Engine/Source/Developer/Apple/MetalShaderFormat/Private/MetalDerivedData.cpp
@@ -1386,7 +1386,6 @@ bool FMetalShaderOutputCooker::Build(TArray<uint8>& OutData)
 							case UBMT_RDG_TEXTURE_SRV:
 							case UBMT_SRV:
 							case UBMT_SAMPLER:
-							case UBMT_RDG_BUFFER:
 							case UBMT_RDG_BUFFER_SRV:
 							case UBMT_UAV:
 							case UBMT_RDG_TEXTURE_UAV:
@@ -1450,7 +1449,6 @@ bool FMetalShaderOutputCooker::Build(TArray<uint8>& OutData)
 											EntryString += ")]];\n";
 											break;
 										}
-										case UBMT_RDG_BUFFER:
 										case UBMT_RDG_BUFFER_SRV:
 										{
 											FCStringAnsi::Snprintf(BufferIdx, 3, "%d", Entry.ResourceIndex + 1);
diff --git a/Engine/Source/Runtime/RHI/Public/RHIDefinitions.h b/Engine/Source/Runtime/RHI/Public/RHIDefinitions.h
--- a/Engine/Source/Runtime/RHI/Public/RHIDefinitions.h
+++ b/Engine/Source/Runtime/RHI/Public/RHIDefinitions.h
@@ -1062,10 +1062,11 @@ enum EUniformBufferBaseType : uint8
 	// Resources tracked by render graph.
 	UBMT_RDG_TEXTURE,
 	UBMT_RDG_TEXTURE_ACCESS,
+	UBMT_RDG_TEXTURE_ACCESS_ARRAY,
 	UBMT_RDG_TEXTURE_SRV,
 	UBMT_RDG_TEXTURE_UAV,
-	UBMT_RDG_BUFFER,
 	UBMT_RDG_BUFFER_ACCESS,
+	UBMT_RDG_BUFFER_ACCESS_ARRAY,
 	UBMT_RDG_BUFFER_SRV,
 	UBMT_RDG_BUFFER_UAV,
 	UBMT_RDG_UNIFORM_BUFFER,
@@ -1904,17 +1905,28 @@ inline bool IsRDGTextureReferenceShaderParameterType(EUniformBufferBaseType Base
 		BaseType == UBMT_RDG_TEXTURE ||
 		BaseType == UBMT_RDG_TEXTURE_SRV ||
 		BaseType == UBMT_RDG_TEXTURE_UAV ||
-		BaseType == UBMT_RDG_TEXTURE_ACCESS;
+		BaseType == UBMT_RDG_TEXTURE_ACCESS ||
+		BaseType == UBMT_RDG_TEXTURE_ACCESS_ARRAY;
 }
 
 /** Returns whether the shader parameter type references an RDG buffer. */
 inline bool IsRDGBufferReferenceShaderParameterType(EUniformBufferBaseType BaseType)
 {
 	return
-		BaseType == UBMT_RDG_BUFFER ||
 		BaseType == UBMT_RDG_BUFFER_SRV ||
 		BaseType == UBMT_RDG_BUFFER_UAV ||
-		BaseType == UBMT_RDG_BUFFER_ACCESS;
+		BaseType == UBMT_RDG_BUFFER_ACCESS ||
+		BaseType == UBMT_RDG_BUFFER_ACCESS_ARRAY;
+}
+
+/** Returns whether the shader parameter type is for RDG access and not actually for shaders. */
+inline bool IsRDGResourceAccessType(EUniformBufferBaseType BaseType)
+{
+	return
+		BaseType == UBMT_RDG_TEXTURE_ACCESS ||
+		BaseType == UBMT_RDG_TEXTURE_ACCESS_ARRAY ||
+		BaseType == UBMT_RDG_BUFFER_ACCESS ||
+		BaseType == UBMT_RDG_BUFFER_ACCESS_ARRAY;
 }
 
 /** Returns whether the shader parameter type is a reference onto a RDG resource. */
@@ -1949,8 +1961,7 @@ inline bool IsShaderParameterTypeIgnoredByRHI(EUniformBufferBaseType BaseType)
 		BaseType == UBMT_RENDER_TARGET_BINDING_SLOTS ||
 
 		// Custom access states are used by the render graph.
-		BaseType == UBMT_RDG_TEXTURE_ACCESS ||
-		BaseType == UBMT_RDG_BUFFER_ACCESS ||
+		IsRDGResourceAccessType(BaseType) ||
 
 		// #yuriy_todo: RHI is able to dereference uniform buffer in root shader parameter structures
 		BaseType == UBMT_REFERENCED_STRUCT ||
diff --git a/Engine/Source/Runtime/RenderCore/Private/RenderGraphBuilder.cpp b/Engine/Source/Runtime/RenderCore/Private/RenderGraphBuilder.cpp
--- a/Engine/Source/Runtime/RenderCore/Private/RenderGraphBuilder.cpp
+++ b/Engine/Source/Runtime/RenderCore/Private/RenderGraphBuilder.cpp
@@ -135,7 +135,17 @@ void EnumerateTextureAccess(FRDGParameterStruct PassParameters, ERDGPassFlags Pa
 		{
 			if (FRDGTextureAccess TextureAccess = Parameter.GetAsTextureAccess())
 			{
-				AccessFunction(nullptr, TextureAccess.GetTexture(), TextureAccess.GetAccess(), NoneFlags, TextureAccess.GetTexture()->GetSubresourceRange());
+				AccessFunction(nullptr, TextureAccess.GetTexture(), TextureAccess.GetAccess(), NoneFlags, TextureAccess->GetSubresourceRange());
+			}
+		}
+		break;
+		case UBMT_RDG_TEXTURE_ACCESS_ARRAY:
+		{
+			const FRDGTextureAccessArray& TextureAccessArray = Parameter.GetAsTextureAccessArray();
+
+			for (FRDGTextureRef Texture : TextureAccessArray)
+			{
+				AccessFunction(nullptr, Texture, TextureAccessArray.GetAccess(), NoneFlags, Texture->GetSubresourceRange());
 			}
 		}
 		break;
@@ -224,12 +234,23 @@ void EnumerateTextureParameters(FRDGParameterStruct PassParameters, TParameterFu
 		switch (Parameter.GetType())
 		{
 		case UBMT_RDG_TEXTURE:
-		case UBMT_RDG_TEXTURE_ACCESS:
 			if (FRDGTextureRef Texture = Parameter.GetAsTexture())
 			{
 				ParameterFunction(Texture);
 			}
 		break;
+		case UBMT_RDG_TEXTURE_ACCESS:
+			if (FRDGTextureAccess Texture = Parameter.GetAsTextureAccess())
+			{
+				ParameterFunction(Texture);
+			}
+		break;
+		case UBMT_RDG_TEXTURE_ACCESS_ARRAY:
+			for (FRDGTextureRef Texture : Parameter.GetAsTextureAccessArray())
+			{
+				ParameterFunction(Texture);
+			}
+		break;
 		case UBMT_RDG_TEXTURE_SRV:
 			if (FRDGTextureSRVRef SRV = Parameter.GetAsTextureSRV())
 			{
@@ -282,25 +303,22 @@ void EnumerateBufferAccess(FRDGParameterStruct PassParameters, ERDGPassFlags Pas
 	{
 		switch (Parameter.GetType())
 		{
-		case UBMT_RDG_BUFFER:
-			if (FRDGBufferRef Buffer = Parameter.GetAsBuffer())
-			{
-				ERHIAccess BufferAccess = SRVAccess;
-
-				if (EnumHasAnyFlags(Buffer->Desc.Usage, BUF_DrawIndirect))
-				{
-					BufferAccess |= ERHIAccess::IndirectArgs;
-				}
-
-				AccessFunction(nullptr, Buffer, BufferAccess);
-			}
-		break;
 		case UBMT_RDG_BUFFER_ACCESS:
 			if (FRDGBufferAccess BufferAccess = Parameter.GetAsBufferAccess())
 			{
 				AccessFunction(nullptr, BufferAccess.GetBuffer(), BufferAccess.GetAccess());
 			}
 		break;
+		case UBMT_RDG_BUFFER_ACCESS_ARRAY:
+		{
+			const FRDGBufferAccessArray& BufferAccessArray = Parameter.GetAsBufferAccessArray();
+
+			for (FRDGBufferRef Buffer : BufferAccessArray)
+			{
+				AccessFunction(nullptr, Buffer, BufferAccessArray.GetAccess());
+			}
+		}
+		break;
 		case UBMT_RDG_BUFFER_SRV:
 			if (FRDGBufferSRVRef SRV = Parameter.GetAsBufferSRV())
 			{
@@ -335,9 +353,14 @@ void EnumerateBufferParameters(FRDGParameterStruct PassParameters, TParameterFun
 	{
 		switch (Parameter.GetType())
 		{
-		case UBMT_RDG_BUFFER:
 		case UBMT_RDG_BUFFER_ACCESS:
-			if (FRDGBufferRef Buffer = Parameter.GetAsBuffer())
+			if (FRDGBufferRef Buffer = Parameter.GetAsBufferAccess())
+			{
+				ParameterFunction(Buffer);
+			}
+		break;
+		case UBMT_RDG_BUFFER_ACCESS_ARRAY:
+			for (FRDGBufferRef Buffer : Parameter.GetAsBufferAccessArray())
 			{
 				ParameterFunction(Buffer);
 			}
@@ -2781,7 +2804,7 @@ void FRDGBuilder::ClobberPassOutputs(const FRDGPass* Pass)
 					{
 						for (int32 MipLevel = 0; MipLevel < Texture->Desc.NumMips; MipLevel++)
 						{
-							AddClearUAVPass(*this, this->CreateUAV(FRDGTextureUAVDesc(Texture, MipLevel)), ClobberColor);
+							AddClearUAVPass(*this, CreateUAV(FRDGTextureUAVDesc(Texture, MipLevel)), ClobberColor);
 						}
 					}
 					else if (EnumHasAnyFlags(TextureAccess.GetAccess(), ERHIAccess::RTV))
@@ -2808,7 +2831,7 @@ void FRDGBuilder::ClobberPassOutputs(const FRDGPass* Pass)
 					{
 						for (int32 MipLevel = 0; MipLevel < Texture->Desc.NumMips; MipLevel++)
 						{
-							AddClearUAVPass(*this, this->CreateUAV(FRDGTextureUAVDesc(Texture, MipLevel)), ClobberColor);
+							AddClearUAVPass(*this, CreateUAV(FRDGTextureUAVDesc(Texture, MipLevel)), ClobberColor);
 						}
 					}
 				}
diff --git a/Engine/Source/Runtime/RenderCore/Private/RenderGraphValidation.cpp b/Engine/Source/Runtime/RenderCore/Private/RenderGraphValidation.cpp
--- a/Engine/Source/Runtime/RenderCore/Private/RenderGraphValidation.cpp
+++ b/Engine/Source/Runtime/RenderCore/Private/RenderGraphValidation.cpp
@@ -562,6 +562,34 @@ void FRDGUserValidation::ValidateAddPass(const FRDGPass* Pass, bool bSkipPassAcc
 
 	bool bCanProduce = false;
 
+	const auto CheckResourceAccess = [&](FRDGParentResourceRef Resource, ERHIAccess Access)
+	{
+		checkf(bIsCopy || !EnumHasAnyFlags(Access, AccessMaskCopy), TEXT("Pass '%s' uses resource '%s' with access '%s' containing states which require the 'ERDGPass::Copy' flag."), Pass->GetName(), Resource->Name, *GetRHIAccessName(Access));
+		checkf(bIsAnyCompute || !EnumHasAnyFlags(Access, AccessMaskCompute), TEXT("Pass '%s' uses resource '%s' with access '%s' containing states which require the 'ERDGPass::Compute' or 'ERDGPassFlags::AsyncCompute' flag."), Pass->GetName(), Resource->Name, *GetRHIAccessName(Access));
+		checkf(bIsRaster || !EnumHasAnyFlags(Access, AccessMaskRaster), TEXT("Pass '%s' uses resource '%s' with access '%s' containing states which require the 'ERDGPass::Raster' flag."), Pass->GetName(), Resource->Name, *GetRHIAccessName(Access));
+		checkf(bIsAnyCompute || bIsRaster || !EnumHasAnyFlags(Access, AccessMaskComputeOrRaster), TEXT("Pass '%s' uses resource '%s' with access '%s' containing states which require the 'ERDGPassFlags::Compute' or 'ERDGPassFlags::AsyncCompute' or 'ERDGPass::Raster' flag."), Pass->GetName(), Resource->Name, *GetRHIAccessName(Access));
+	};
+
+	const auto CheckBufferAccess = [&](FRDGBufferRef Buffer, ERHIAccess Access)
+	{
+		CheckResourceAccess(Buffer, Access);
+
+		if (IsWritableAccess(Access))
+		{
+			MarkBufferAsProduced(Buffer);
+		}
+	};
+
+	const auto CheckTextureAccess = [&](FRDGTextureRef Texture, ERHIAccess Access)
+	{
+		CheckResourceAccess(Texture, Access);
+
+		if (IsWritableAccess(Access))
+		{
+			MarkTextureAsProduced(Texture);
+		}
+	};
+
 	PassParameters.Enumerate([&](FRDGParameter Parameter)
 	{
 		if (Parameter.IsParentResource())
@@ -610,14 +638,6 @@ void FRDGUserValidation::ValidateAddPass(const FRDGPass* Pass, bool bSkipPassAcc
 			}
 		}
 		break;
-		case UBMT_RDG_BUFFER:
-		{
-			if (FRDGBufferRef Buffer = Parameter.GetAsBuffer())
-			{
-				MarkAsConsumed(Buffer);
-			}
-		}
-		break;
 		case UBMT_RDG_BUFFER_SRV:
 		{
 			if (FRDGBufferSRVRef SRV = Parameter.GetAsBufferSRV())
@@ -641,43 +661,45 @@ void FRDGUserValidation::ValidateAddPass(const FRDGPass* Pass, bool bSkipPassAcc
 		break;
 		case UBMT_RDG_TEXTURE_ACCESS:
 		{
-			const FRDGTextureAccess TextureAccess = Parameter.GetAsTextureAccess();
-			const ERHIAccess Access = TextureAccess.GetAccess();
-			const bool bIsWriteAccess = IsWritableAccess(Access);
-			bCanProduce |= bIsWriteAccess;
+			FRDGTextureAccess TextureAccess = Parameter.GetAsTextureAccess();
+			bCanProduce |= IsWritableAccess(TextureAccess.GetAccess());
 
-			if (FRDGTextureRef Texture = TextureAccess.GetTexture())
+			if (TextureAccess)
 			{
-				checkf(bIsCopy       || !EnumHasAnyFlags(Access, AccessMaskCopy),    TEXT("Pass '%s' uses texture '%s' with access '%s' containing states which require the 'ERDGPass::Copy' flag."), Pass->GetName(), Texture->Name, *GetRHIAccessName(Access));
-				checkf(bIsAnyCompute || !EnumHasAnyFlags(Access, AccessMaskCompute), TEXT("Pass '%s' uses texture '%s' with access '%s' containing states which require the 'ERDGPass::Compute' or 'ERDGPassFlags::AsyncCompute' flag."), Pass->GetName(), Texture->Name, *GetRHIAccessName(Access));
-				checkf(bIsRaster     || !EnumHasAnyFlags(Access, AccessMaskRaster),  TEXT("Pass '%s' uses texture '%s' with access '%s' containing states which require the 'ERDGPass::Raster' flag."), Pass->GetName(), Texture->Name, *GetRHIAccessName(Access));
-				checkf(bIsAnyCompute || bIsRaster || !EnumHasAnyFlags(Access, AccessMaskComputeOrRaster), TEXT("Pass '%s' uses texture '%s' with access '%s' containing states which require the 'ERDGPassFlags::Compute' or 'ERDGPassFlags::AsyncCompute' or 'ERDGPass::Raster' flag."), Pass->GetName(), Texture->Name, *GetRHIAccessName(Access));
+				CheckTextureAccess(TextureAccess.GetTexture(), TextureAccess.GetAccess());
+			}
+		}
+		break;
+		case UBMT_RDG_TEXTURE_ACCESS_ARRAY:
+		{
+			const FRDGTextureAccessArray& TextureAccessArray = Parameter.GetAsTextureAccessArray();
+			bCanProduce |= IsWritableAccess(TextureAccessArray.GetAccess());
 
-				if (bIsWriteAccess)
-				{
-					MarkTextureAsProduced(Texture);
-				}
+			for (FRDGTextureRef Texture : TextureAccessArray)
+			{
+				CheckTextureAccess(Texture, TextureAccessArray.GetAccess());
 			}
 		}
 		break;
 		case UBMT_RDG_BUFFER_ACCESS:
 		{
-			const FRDGBufferAccess BufferAccess = Parameter.GetAsBufferAccess();
-			const ERHIAccess Access = BufferAccess.GetAccess();
-			const bool bIsWriteAccess = IsWritableAccess(Access);
-			bCanProduce |= bIsWriteAccess;
+			FRDGBufferAccess BufferAccess = Parameter.GetAsBufferAccess();
+			bCanProduce |= IsWritableAccess(BufferAccess.GetAccess());
 
-			if (FRDGBufferRef Buffer = BufferAccess.GetBuffer())
+			if (BufferAccess)
 			{
-				checkf(bIsCopy       || !EnumHasAnyFlags(Access, AccessMaskCopy),    TEXT("Pass '%s' uses buffer '%s' with access '%s' containing states which require the 'ERDGPass::Copy' flag."), Pass->GetName(), Buffer->Name, *GetRHIAccessName(Access));
-				checkf(bIsAnyCompute || !EnumHasAnyFlags(Access, AccessMaskCompute), TEXT("Pass '%s' uses buffer '%s' with access '%s' containing states which require the 'ERDGPass::Compute' or 'ERDGPassFlags::AsyncCompute' flag."), Pass->GetName(), Buffer->Name, *GetRHIAccessName(Access));
-				checkf(bIsRaster     || !EnumHasAnyFlags(Access, AccessMaskRaster),  TEXT("Pass '%s' uses buffer '%s' with access '%s' containing states which require the 'ERDGPass::Raster' flag."), Pass->GetName(), Buffer->Name, *GetRHIAccessName(Access));
-				checkf(bIsAnyCompute || bIsRaster || !EnumHasAnyFlags(Access, AccessMaskComputeOrRaster), TEXT("Pass '%s' uses buffer '%s' with access '%s' containing states which require the 'ERDGPassFlags::Compute' or 'ERDGPassFlags::AsyncCompute' or 'ERDGPass::Raster' flag."), Pass->GetName(), Buffer->Name, *GetRHIAccessName(Access));
+				CheckBufferAccess(BufferAccess.GetBuffer(), BufferAccess.GetAccess());
+			}
+		}
+		break;
+		case UBMT_RDG_BUFFER_ACCESS_ARRAY:
+		{
+			const FRDGBufferAccessArray& BufferAccessArray = Parameter.GetAsBufferAccessArray();
+			bCanProduce |= IsWritableAccess(BufferAccessArray.GetAccess());
 
-				if (IsWritableAccess(BufferAccess.GetAccess()))
-				{
-					MarkBufferAsProduced(Buffer);
-				}
+			for (FRDGBufferRef Buffer : BufferAccessArray)
+			{
+				CheckBufferAccess(Buffer, BufferAccessArray.GetAccess());
 			}
 		}
 		break;
@@ -881,6 +903,19 @@ void FRDGUserValidation::ValidateExecutePassBegin(const FRDGPass* Pass)
 			}
 		});
 
+		const auto ValidateTextureAccess = [](FRDGTextureRef Texture, ERHIAccess Access)
+		{
+			if (EnumHasAnyFlags(Access, ERHIAccess::UAVMask))
+			{
+				Texture->GetTextureDebugData().bHasNeededUAV = true;
+			}
+			if (EnumHasAnyFlags(Access, ERHIAccess::RTV | ERHIAccess::DSVRead | ERHIAccess::DSVWrite))
+			{
+				Texture->GetTextureDebugData().bHasBeenBoundAsRenderTarget = true;
+			}
+			Texture->MarkResourceAsUsed();
+		};
+
 		Pass->GetParameters().Enumerate([&](FRDGParameter Parameter)
 		{
 			switch (Parameter.GetType())
@@ -893,20 +928,18 @@ void FRDGUserValidation::ValidateExecutePassBegin(const FRDGPass* Pass)
 				}
 			break;
 			case UBMT_RDG_TEXTURE_ACCESS:
+				if (FRDGTextureAccess TextureAccess = Parameter.GetAsTextureAccess())
+				{
+					ValidateTextureAccess(TextureAccess.GetTexture(), TextureAccess.GetAccess());
+				}
+			break;
+			case UBMT_RDG_TEXTURE_ACCESS_ARRAY:
 			{
-				const FRDGTextureAccess TextureAccess = Parameter.GetAsTextureAccess();
-				if (FRDGTextureRef Texture = TextureAccess.GetTexture())
+				const FRDGTextureAccessArray& TextureAccessArray = Parameter.GetAsTextureAccessArray();
+
+				for (FRDGTextureRef Texture : TextureAccessArray)
 				{
-					const ERHIAccess Access = TextureAccess.GetAccess();
-					if (EnumHasAnyFlags(Access, ERHIAccess::UAVMask))
-					{
-						Texture->GetTextureDebugData().bHasNeededUAV = true;
-					}
-					if (EnumHasAnyFlags(Access, ERHIAccess::RTV | ERHIAccess::DSVRead | ERHIAccess::DSVWrite))
-					{
-						Texture->GetTextureDebugData().bHasBeenBoundAsRenderTarget = true;
-					}
-					Texture->MarkResourceAsUsed();
+					ValidateTextureAccess(Texture, TextureAccessArray.GetAccess());
 				}
 			}
 			break;
@@ -916,6 +949,12 @@ void FRDGUserValidation::ValidateExecutePassBegin(const FRDGPass* Pass)
 					Buffer->MarkResourceAsUsed();
 				}
 			break;
+			case UBMT_RDG_BUFFER_ACCESS_ARRAY:
+				for (FRDGBufferRef Buffer : Parameter.GetAsBufferAccessArray())
+				{
+					Buffer->MarkResourceAsUsed();
+				}
+			break;
 			case UBMT_RENDER_TARGET_BINDING_SLOTS:
 			{
 				const FRenderTargetBindingSlots& RenderTargets = Parameter.GetAsRenderTargetBindingSlots();
diff --git a/Engine/Source/Runtime/RenderCore/Private/ShaderParameterMetadata.cpp b/Engine/Source/Runtime/RenderCore/Private/ShaderParameterMetadata.cpp
--- a/Engine/Source/Runtime/RenderCore/Private/ShaderParameterMetadata.cpp
+++ b/Engine/Source/Runtime/RenderCore/Private/ShaderParameterMetadata.cpp
@@ -121,11 +121,12 @@ const TCHAR* const kShaderParameterMacroNames[] = {
 
 	// Resources tracked by render graph.
 	TEXT("SHADER_PARAMETER_RDG_TEXTURE"), // UBMT_RDG_TEXTURE,
-	TEXT("SHADER_PARAMETER_RDG_TEXTURE_ACCESS"), // UBMT_RDG_TEXTURE_ACCESS,
+	TEXT("RDG_TEXTURE_ACCESS"), // UBMT_RDG_TEXTURE_ACCESS,
+	TEXT("RDG_TEXTURE_ACCESS_ARRAY"), // UBMT_RDG_TEXTURE_ACCESS,
 	TEXT("SHADER_PARAMETER_RDG_TEXTURE_SRV"), // UBMT_RDG_TEXTURE_SRV,
 	TEXT("SHADER_PARAMETER_RDG_TEXTURE_UAV"), // UBMT_RDG_TEXTURE_UAV,
-	TEXT("SHADER_PARAMETER_RDG_BUFFER"), // UBMT_RDG_BUFFER,
-	TEXT("SHADER_PARAMETER_RDG_BUFFER_ACCESS"), // UBMT_RDG_BUFFER_ACCESS,
+	TEXT("RDG_BUFFER_ACCESS"), // UBMT_RDG_BUFFER_ACCESS,
+	TEXT("RDG_BUFFER_ACCESS_ARRAY"), // UBMT_RDG_BUFFER_ACCESS_ARRAY,
 	TEXT("SHADER_PARAMETER_RDG_BUFFER_SRV"), // UBMT_RDG_BUFFER_SRV,
 	TEXT("SHADER_PARAMETER_RDG_BUFFER_UAV"), // UBMT_RDG_BUFFER_UAV,
 	TEXT("SHADER_PARAMETER_RDG_UNIFORM_BUFFER"), // UBMT_RDG_UNIFORM_BUFFER,
@@ -506,7 +507,7 @@ void FShaderParametersMetadata::InitializeLayout()
 			}
 
 			// Validate the shader binding type.
-			if (bIsRHIResource || (bIsRDGResource && (BaseType != UBMT_RDG_BUFFER && BaseType != UBMT_RDG_TEXTURE_ACCESS && BaseType != UBMT_RDG_BUFFER_ACCESS && BaseType != UBMT_RDG_UNIFORM_BUFFER)))
+			if (bIsRHIResource || (bIsRDGResource && !IsRDGResourceAccessType(BaseType) && BaseType != UBMT_RDG_UNIFORM_BUFFER))
 			{
 				EShaderCodeResourceBindingType BindingType = ParseShaderResourceBindingType(ShaderType);
 
diff --git a/Engine/Source/Runtime/RenderCore/Private/ShaderParameterStruct.cpp b/Engine/Source/Runtime/RenderCore/Private/ShaderParameterStruct.cpp
--- a/Engine/Source/Runtime/RenderCore/Private/ShaderParameterStruct.cpp
+++ b/Engine/Source/Runtime/RenderCore/Private/ShaderParameterStruct.cpp
@@ -72,10 +72,7 @@ struct FShaderParameterStructBindingContext
 				BaseType == UBMT_SAMPLER);
 
 			const bool bIsRDGResource =
-				IsRDGResourceReferenceShaderParameterType(BaseType) &&
-				BaseType != UBMT_RDG_BUFFER &&
-				BaseType != UBMT_RDG_BUFFER_ACCESS &&
-				BaseType != UBMT_RDG_TEXTURE_ACCESS;
+				IsRDGResourceReferenceShaderParameterType(BaseType) && !IsRDGResourceAccessType(BaseType);
 
 			const bool bIsVariableNativeType = (
 				BaseType == UBMT_INT32 ||
@@ -122,16 +119,6 @@ struct FShaderParameterStructBindingContext
 				// The member name of a globally referenced struct is the not name on the struct.
 				ShaderBindingName = Member.GetStructMetadata()->GetShaderVariableName();
 			}
-			else if (BaseType == UBMT_RDG_BUFFER)
-			{
-				// RHI does not support setting a buffer as a shader parameter.
-				check(!bIsArray);
-				if( ParametersMap->ContainsParameterAllocation(*ShaderBindingName) )
-				{
-					UE_LOG(LogShaders, Fatal, TEXT("%s can't bind shader parameter %s as buffer. Use buffer SRV for reading in shader."), *CppName, *ShaderBindingName);
-				}
-				continue;
-			}
 			else if (bUseRootShaderParameters && bIsVariableNativeType)
 			{
 				// Constants are stored in the root shader parameter cbuffer when bUseRootShaderParameters == true.
diff --git a/Engine/Source/Runtime/RenderCore/Private/ShaderParameters.cpp b/Engine/Source/Runtime/RenderCore/Private/ShaderParameters.cpp
--- a/Engine/Source/Runtime/RenderCore/Private/ShaderParameters.cpp
+++ b/Engine/Source/Runtime/RenderCore/Private/ShaderParameters.cpp
@@ -242,7 +242,7 @@ static void CreateHLSLUniformBufferStructMembersDeclaration(
 
 		if (IsShaderParameterTypeForUniformBufferLayout(Member.GetBaseType()))
 		{
-			checkf(Member.GetBaseType() != UBMT_RDG_TEXTURE_ACCESS && Member.GetBaseType() != UBMT_RDG_BUFFER_ACCESS, TEXT("Copy destination usage is not supported in uniform buffers."));
+			checkf(!IsRDGResourceAccessType(Member.GetBaseType()), TEXT("RDG access parameter types (e.g. RDG_TEXTURE_ACCESS) are not allowed in uniform buffers."));
 			if (Member.GetBaseType() == UBMT_SRV)
 			{
 				// TODO: handle arrays?
diff --git a/Engine/Source/Runtime/RenderCore/Public/RenderGraphParameter.h b/Engine/Source/Runtime/RenderCore/Public/RenderGraphParameter.h
--- a/Engine/Source/Runtime/RenderCore/Public/RenderGraphParameter.h
+++ b/Engine/Source/Runtime/RenderCore/Public/RenderGraphParameter.h
@@ -35,7 +35,8 @@ public:
 	{
 		return
 			MemberType == UBMT_RDG_TEXTURE ||
-			MemberType == UBMT_RDG_TEXTURE_ACCESS;
+			MemberType == UBMT_RDG_TEXTURE_ACCESS ||
+			MemberType == UBMT_RDG_TEXTURE_ACCESS_ARRAY;
 	}
 
 	bool IsTextureAccess() const
@@ -43,11 +44,16 @@ public:
 		return MemberType == UBMT_RDG_TEXTURE_ACCESS;
 	}
 
+	bool IsTextureAccessArray() const
+	{
+		return MemberType == UBMT_RDG_TEXTURE_ACCESS_ARRAY;
+	}
+
 	bool IsBuffer() const
 	{
 		return
-			MemberType == UBMT_RDG_BUFFER ||
-			MemberType == UBMT_RDG_BUFFER_ACCESS;
+			MemberType == UBMT_RDG_BUFFER_ACCESS ||
+			MemberType == UBMT_RDG_BUFFER_ACCESS_ARRAY;
 	}
 
 	bool IsBufferAccess() const
@@ -55,6 +61,11 @@ public:
 		return MemberType == UBMT_RDG_BUFFER_ACCESS;
 	}
 
+	bool IsBufferAccessArray() const
+	{
+		return MemberType == UBMT_RDG_BUFFER_ACCESS_ARRAY;
+	}
+
 	bool IsUniformBuffer() const
 	{
 		return MemberType == UBMT_RDG_UNIFORM_BUFFER;
@@ -123,6 +134,12 @@ public:
 		return *GetAs<FRDGTextureAccess>();
 	}
 
+	const FRDGTextureAccessArray& GetAsTextureAccessArray() const
+	{
+		check(MemberType == UBMT_RDG_TEXTURE_ACCESS_ARRAY);
+		return *GetAs<FRDGTextureAccessArray>();
+	}
+
 	FRDGBufferRef GetAsBuffer() const
 	{
 		check(IsBuffer());
@@ -135,6 +152,12 @@ public:
 		return *GetAs<FRDGBufferAccess>();
 	}
 
+	const FRDGBufferAccessArray& GetAsBufferAccessArray() const
+	{
+		check(MemberType == UBMT_RDG_BUFFER_ACCESS_ARRAY);
+		return *GetAs<FRDGBufferAccessArray>();
+	}
+
 	FRDGTextureSRVRef GetAsTextureSRV() const
 	{
 		check(MemberType == UBMT_RDG_TEXTURE_SRV);
diff --git a/Engine/Source/Runtime/RenderCore/Public/ShaderParameterMacros.h b/Engine/Source/Runtime/RenderCore/Public/ShaderParameterMacros.h
--- a/Engine/Source/Runtime/RenderCore/Public/ShaderParameterMacros.h
+++ b/Engine/Source/Runtime/RenderCore/Public/ShaderParameterMacros.h
@@ -8,6 +8,7 @@
 #pragma once
 
 #include "ShaderParameterMetadata.h"
+#include "RenderGraphAllocator.h"
 #include "Algo/Reverse.h"
 
 
@@ -422,6 +423,76 @@ public:
 	{}
 };
 
+template <typename ResourceType>
+class alignas(SHADER_PARAMETER_POINTER_ALIGNMENT) TRDGResourceAccessArray
+{
+public:
+	using ArrayType = TArray<ResourceType*, FRDGArrayAllocator>;
+	using ArrayViewType = TArrayView<ResourceType* const>;
+
+	TRDGResourceAccessArray() = default;
+	TRDGResourceAccessArray(ArrayViewType InArrayView, ERHIAccess InAccess)
+		: Array(InArrayView)
+		, Access(InAccess)
+	{}
+
+	FORCEINLINE ArrayViewType GetArray() const { return Array; }
+	FORCEINLINE ERHIAccess GetAccess() const { return Access; }
+
+	// Forward commonly used functions to the array.
+	FORCEINLINE void Add(ResourceType* Resource) { Array.Add(Resource); }
+	FORCEINLINE void Reserve(uint32 Count) { Array.Reserve(Count); }
+	FORCEINLINE bool IsEmpty() const { return Array.IsEmpty(); }
+	FORCEINLINE uint32 Num() const { return Array.Num(); }
+	FORCEINLINE typename ArrayType::RangedForIteratorType      begin()       { return Array.begin(); }
+	FORCEINLINE typename ArrayType::RangedForConstIteratorType begin() const { return Array.begin(); }
+	FORCEINLINE typename ArrayType::RangedForIteratorType      end()         { return Array.end(); }
+	FORCEINLINE typename ArrayType::RangedForConstIteratorType end() const   { return Array.end(); }
+
+private:
+	ArrayType Array;
+#if !PLATFORM_64BITS
+	uint32 _Padding;
+#endif
+	ERHIAccess Access = ERHIAccess::Unknown;
+};
+
+using FRDGBufferAccessArray = TRDGResourceAccessArray<FRDGBuffer>;
+
+template <ERHIAccess InAccess>
+class alignas(SHADER_PARAMETER_POINTER_ALIGNMENT) TRDGBufferAccessArray
+	: public FRDGBufferAccessArray
+{
+public:
+	static_assert(IsValidAccess(InAccess), "Buffer access is invalid.");
+
+	TRDGBufferAccessArray()
+		: TRDGBufferAccessArray({}, InAccess)
+	{}
+
+	TRDGBufferAccessArray(ArrayViewType InBufferView)
+		: TRDGBufferAccessArray(InBufferView, InAccess)
+	{}
+};
+
+using FRDGTextureAccessArray = TRDGResourceAccessArray<FRDGTexture>;
+
+template <ERHIAccess InAccess>
+class alignas(SHADER_PARAMETER_POINTER_ALIGNMENT) TRDGTextureAccessArray
+	: public FRDGTextureAccessArray
+{
+public:
+	static_assert(IsValidAccess(InAccess), "Texture access is invalid.");
+
+	TRDGTextureAccessArray()
+		: FRDGTextureAccessArray({}, InAccess)
+	{}
+
+	TRDGTextureAccessArray(ArrayViewType InTextureArray)
+		: FRDGTextureAccessArray(InTextureArray, InAccess)
+	{}
+};
+
 /** Render graph information about how to bind a render target. */
 struct FRenderTargetBinding
 {
@@ -992,7 +1063,7 @@ struct TShaderParameterTypeInfo<FMatrix>
 	static const FShaderParametersMetadata* GetStructMetadata() { return nullptr; }
 };
 
-template <typename RDGResourceType>
+template <typename ResourceAccessType>
 struct TRDGResourceAccessTypeInfo
 {
 	static constexpr int32 NumRows = 1;
@@ -1001,13 +1072,29 @@ struct TRDGResourceAccessTypeInfo
 	static constexpr int32 Alignment = SHADER_PARAMETER_POINTER_ALIGNMENT;
 	static constexpr bool bIsStoredInConstantBuffer = false;
 
-	using TAlignedType = RDGResourceType;
+	using TAlignedType = ResourceAccessType;
 
 	static const FShaderParametersMetadata* GetStructMetadata() { return nullptr; }
 
 	static_assert(sizeof(TAlignedType) == SHADER_PARAMETER_POINTER_ALIGNMENT * 2, "Uniform buffer layout must not be platform dependent.");
 };
 
+template <typename ResourceAccessArrayType>
+struct TRDGResourceAccessArrayTypeInfo
+{
+	static constexpr int32 NumRows = 1;
+	static constexpr int32 NumColumns = 1;
+	static constexpr int32 NumElements = 0;
+	static constexpr int32 Alignment = SHADER_PARAMETER_POINTER_ALIGNMENT;
+	static constexpr bool bIsStoredInConstantBuffer = false;
+
+	using TAlignedType = ResourceAccessArrayType;
+
+	static const FShaderParametersMetadata* GetStructMetadata() { return nullptr; }
+
+	static_assert(sizeof(TAlignedType) == SHADER_PARAMETER_POINTER_ALIGNMENT * 3, "Uniform buffer layout must not be platform dependent.");
+};
+
 template<typename T, size_t InNumElements>
 struct TShaderParameterTypeInfo<T[InNumElements]>
 {
@@ -1472,19 +1559,6 @@ extern RENDERCORE_API FShaderParametersMetadata* FindUniformBufferStructByShader
 #define SHADER_PARAMETER_RDG_TEXTURE_UAV_ARRAY(ShaderType,MemberName, ArrayDecl) \
 	INTERNAL_SHADER_PARAMETER_EXPLICIT(UBMT_RDG_TEXTURE_UAV, TShaderResourceParameterTypeInfo<FRDGTextureUAV* ArrayDecl>, FRDGTextureUAV*,MemberName,ArrayDecl,,EShaderPrecisionModifier::Float,TEXT(#ShaderType),false)
 
-/** Adds a render graph tracked buffer.
- *
- * Example:
- *	SHADER_PARAMETER_RDG_BUFFER(Buffer<float4>, MyBuffer)
- *	SHADER_PARAMETER_RDG_BUFFER_ARRAY(Buffer<float4>, MyArrayOfBuffers, [4])
- */
-// TODO: ShaderType is unnecessary, because the RHI does not support binding a buffer as a shader parameter.
-#define SHADER_PARAMETER_RDG_BUFFER(ShaderType,MemberName) \
-	INTERNAL_SHADER_PARAMETER_EXPLICIT(UBMT_RDG_BUFFER, TShaderResourceParameterTypeInfo<FRDGBuffer*>, FRDGBuffer*,MemberName,, = nullptr,EShaderPrecisionModifier::Float,TEXT(#ShaderType),false)
-
-#define SHADER_PARAMETER_RDG_BUFFER_ARRAY(ShaderType,MemberName, ArrayDecl) \
-	INTERNAL_SHADER_PARAMETER_EXPLICIT(UBMT_RDG_BUFFER, TShaderResourceParameterTypeInfo<FRDGBuffer* ArrayDecl>, FRDGBuffer*,MemberName,ArrayDecl,,EShaderPrecisionModifier::Float,TEXT(#ShaderType),false)
-
 /** Adds a shader resource view for a render graph tracked buffer.
  *
  * Example:
@@ -1617,18 +1691,27 @@ extern RENDERCORE_API FShaderParametersMetadata* FindUniformBufferStructByShader
 #define RDG_BUFFER_ACCESS(MemberName, Access) \
 	INTERNAL_SHADER_PARAMETER_EXPLICIT(UBMT_RDG_BUFFER_ACCESS, TRDGResourceAccessTypeInfo<TRDGBufferAccess<Access>>, TRDGBufferAccess<Access>,MemberName,,,EShaderPrecisionModifier::Float,TEXT(""),false)
 
+#define RDG_BUFFER_ACCESS_ARRAY(MemberName, Access) \
+	INTERNAL_SHADER_PARAMETER_EXPLICIT(UBMT_RDG_BUFFER_ACCESS_ARRAY, TRDGResourceAccessArrayTypeInfo<TRDGBufferAccessArray<Access>>, TRDGBufferAccessArray<Access>,MemberName,,,EShaderPrecisionModifier::Float,TEXT(""),false)
+
 #define RDG_BUFFER_ACCESS_DYNAMIC(MemberName) \
 	INTERNAL_SHADER_PARAMETER_EXPLICIT(UBMT_RDG_BUFFER_ACCESS, TRDGResourceAccessTypeInfo<FRDGBufferAccess>, FRDGBufferAccess,MemberName,,,EShaderPrecisionModifier::Float,TEXT(""),false)
 
+#define RDG_BUFFER_ACCESS_ARRAY_DYNAMIC(MemberName) \
+	INTERNAL_SHADER_PARAMETER_EXPLICIT(UBMT_RDG_BUFFER_ACCESS_ARRAY, TRDGResourceAccessArrayTypeInfo<FRDGBufferAccessArray>, FRDGBufferAccessArray,MemberName,,,EShaderPrecisionModifier::Float,TEXT(""),false)
+
 /** Informs the RDG pass to transition the texture into the requested state. */
 #define RDG_TEXTURE_ACCESS(MemberName, Access) \
 	INTERNAL_SHADER_PARAMETER_EXPLICIT(UBMT_RDG_TEXTURE_ACCESS, TRDGResourceAccessTypeInfo<TRDGTextureAccess<Access>>, TRDGTextureAccess<Access>,MemberName,,,EShaderPrecisionModifier::Float,TEXT(""),false)
 
+#define RDG_TEXTURE_ACCESS_ARRAY(MemberName, Access) \
+	INTERNAL_SHADER_PARAMETER_EXPLICIT(UBMT_RDG_TEXTURE_ACCESS_ARRAY, TRDGResourceAccessArrayTypeInfo<TRDGTextureAccessArray<Access>>, TRDGTextureAccessArray<Access>,MemberName,,,EShaderPrecisionModifier::Float,TEXT(""),false)
+
 #define RDG_TEXTURE_ACCESS_DYNAMIC(MemberName) \
 	INTERNAL_SHADER_PARAMETER_EXPLICIT(UBMT_RDG_TEXTURE_ACCESS, TRDGResourceAccessTypeInfo<FRDGTextureAccess>, FRDGTextureAccess,MemberName,,,EShaderPrecisionModifier::Float,TEXT(""),false)
 
-#define SHADER_PARAMETER_RDG_BUFFER_UPLOAD(MemberName) \
-	RDG_BUFFER_ACCESS(MemberName, ERHIAccess::CopyDest)
+#define RDG_TEXTURE_ACCESS_ARRAY_DYNAMIC(MemberName) \
+	INTERNAL_SHADER_PARAMETER_EXPLICIT(UBMT_RDG_TEXTURE_ACCESS_ARRAY, TRDGResourceAccessArrayTypeInfo<FRDGTextureAccessArray>, FRDGTextureAccessArray,MemberName,,,EShaderPrecisionModifier::Float,TEXT(""),false)
 
 /** Adds bindings slots for render targets on the structure. This is important for rasterizer based pass bind the
  * render target at the RHI pass creation. The name of the struct member will forced to RenderTargets, and
@@ -1662,3 +1745,15 @@ END_SHADER_PARAMETER_STRUCT()
 
 /** Returns the name of the macro that should be used for a given shader parameter base type. */
 const TCHAR* GetShaderParameterMacroName(EUniformBufferBaseType ShaderParameterBaseType);
+
+/** Deprecated buffer macros. */
+#define SHADER_PARAMETER_RDG_BUFFER_UPLOAD(MemberName) \
+	DEPRECATED_MACRO(5.0, "SHADER_PARAMETER_RDG_BUFFER_UPLOAD has been deprecated. Use RDG_BUFFER_ACCESS(Buffer, ERHIAccess::CopyDest) instead.") \
+	RDG_BUFFER_ACCESS(MemberName, ERHIAccess::CopyDest)
+
+#define SHADER_PARAMETER_RDG_BUFFER(ShaderType,MemberName) \
+	DEPRECATED_MACRO(5.0, "SHADER_PARAMETER_RDG_BUFFER has been deprecated. Use RDG_BUFFER_ACCESS with an explicit RHI state instead.") \
+	RDG_BUFFER_ACCESS(MemberName, ERHIAccess::SRVMask | ERHIAccess::IndirectArgs)
+
+#define SHADER_PARAMETER_RDG_BUFFER_ARRAY(ShaderType,MemberName, ArrayDecl) \
+	DEPRECATED_MACRO(5.0, "SHADER_PARAMETER_RDG_BUFFER_ARRAY has been deprecated. Use RDG_BUFFER_ACCESS_ARRAY instead.")
\ No newline at end of file
diff --git a/Engine/Source/Runtime/Renderer/Private/HairStrands/HairStrandsLUT.cpp b/Engine/Source/Runtime/Renderer/Private/HairStrands/HairStrandsLUT.cpp
--- a/Engine/Source/Runtime/Renderer/Private/HairStrands/HairStrandsLUT.cpp
+++ b/Engine/Source/Runtime/Renderer/Private/HairStrands/HairStrandsLUT.cpp
@@ -263,7 +263,7 @@ class FHairCoverageLUTCS : public FGlobalShader
 	END_SHADER_PARAMETER_STRUCT()
 
 	BEGIN_SHADER_PARAMETER_STRUCT(FUploadParameters, )
-		SHADER_PARAMETER_RDG_BUFFER_UPLOAD(UploadBuffer)
+		RDG_BUFFER_ACCESS(UploadBuffer, ERHIAccess::CopyDest)
 	END_SHADER_PARAMETER_STRUCT()
 
 public:
