Fix for N-to-N pipe transition check in RDG.

diff --git a/Engine/Source/Runtime/RenderCore/Private/RenderGraphBuilder.cpp b/Engine/Source/Runtime/RenderCore/Private/RenderGraphBuilder.cpp
--- a/Engine/Source/Runtime/RenderCore/Private/RenderGraphBuilder.cpp
+++ b/Engine/Source/Runtime/RenderCore/Private/RenderGraphBuilder.cpp
@@ -2299,18 +2299,13 @@ void FRDGBuilder::AddTransition(
 			}
 		}
 	}
-	// N-to-1 transition.
+	// N-to-1 or N-to-N transition.
 	else
 	{
-		check(PipelinesAfterCount == 1);
 		checkf(StateBefore.GetLastPass() != StateAfter.GetFirstPass() || BarrierLocation == EBarrierLocation::Epilogue,
 			TEXT("Attempted to queue a transition for resource '%s' from '%s' to '%s', but previous and next passes are the same on one pipe."),
 			Resource->Name, *GetRHIPipelineName(PipelinesBefore), *GetRHIPipelineName(PipelinesAfter));
 
-		/** We have multiple begin passes which are being issued from All pipes to 1 end pass on pipe. So, we create a unique barrier batch
-		 *  instance for the requested configuration and add it to all begin passes. The 'End' side of things is slightly simpler than the
-		 *  other path, because there's only one of them.
-		 */
 		FRDGBarrierBatchBeginId Id;
 		Id.PipelinesAfter = PipelinesAfter;
 		for (ERHIPipeline Pipeline : GetRHIPipelines())
