Added RDG validation to check for dangling pointers and assert at AddPass time. Tracks resources actively used in the graph and will emit an assert if an incoming resource isn't part of the set. This would catch, for example, if the user kept around a resource from a previous builder and tried to reuse it in a new builder.

diff --git a/Engine/Source/Runtime/RenderCore/Private/RenderGraphValidation.cpp b/Engine/Source/Runtime/RenderCore/Private/RenderGraphValidation.cpp
--- a/Engine/Source/Runtime/RenderCore/Private/RenderGraphValidation.cpp
+++ b/Engine/Source/Runtime/RenderCore/Private/RenderGraphValidation.cpp
@@ -73,22 +73,16 @@ void FRDGResource::MarkResourceAsUsed()
 {
 	ValidateRHIAccess();
 
-	if (DebugData)
-	{
-		DebugData->bIsActuallyUsedByPass = true;
-	}
+	GetDebugData().bIsActuallyUsedByPass = true;
 }
 
 void FRDGResource::ValidateRHIAccess() const
 {
-	// Passthrough resources will not have debug data, since they are not tied to a graph instance.
-	if (DebugData)
-	{
-		checkf(DebugData->bAllowRHIAccess || GRDGAllowRHIAccess,
-			TEXT("Accessing the RHI resource of %s at this time is not allowed. If you hit this check in pass, ")
-			TEXT("that is due to this resource not being referenced in the parameters of your pass."),
-			Name);
-	}
+	check(DebugData);
+	checkf(DebugData->bAllowRHIAccess || GRDGAllowRHIAccess,
+		TEXT("Accessing the RHI resource of %s at this time is not allowed. If you hit this check in pass, ")
+		TEXT("that is due to this resource not being referenced in the parameters of your pass."),
+		Name);
 }
 
 FRDGResourceDebugData& FRDGResource::GetDebugData() const
@@ -118,11 +112,6 @@ FRDGParentResourceDebugData& FRDGParentResource::GetParentDebugData() const
 	return *ParentDebugData;
 }
 
-bool FRDGResource::IsPassthrough() const
-{
-	return DebugData == nullptr;
-}
-
 struct FRDGTextureDebugData
 {
 	/** Tracks whether a UAV has ever been allocated to catch when TexCreate_UAV was unneeded. */
@@ -194,6 +183,10 @@ void FRDGUserValidation::ValidateCreateResource(FRDGResourceRef Resource)
 	MemStackGuard();
 	check(Resource);
 	Resource->DebugData = Allocator.Alloc<FRDGResourceDebugData>();
+
+	bool bAlreadyInSet = false;
+	ResourceMap.Emplace(Resource, &bAlreadyInSet);
+	check(!bAlreadyInSet);
 }
 
 void FRDGUserValidation::ValidateCreateParentResource(FRDGParentResourceRef Resource)
@@ -319,7 +312,6 @@ void FRDGUserValidation::ValidateCreateSRV(const FRDGTextureSRVDesc& Desc)
 {
 	FRDGTextureRef Texture = Desc.Texture;
 	checkf(Texture, TEXT("Texture SRV created with a null texture."));
-	checkf(!Texture->IsPassthrough(), TEXT("Texture SRV created with passthrough texture '%s'."), Texture->Name);
 	ExecuteGuard(TEXT("CreateSRV"), Texture->Name);
 	checkf(Texture->Desc.Flags & TexCreate_ShaderResource, TEXT("Attempted to create SRV from texture %s which was not created with TexCreate_ShaderResource"), Desc.Texture->Name);
 
@@ -362,7 +354,6 @@ void FRDGUserValidation::ValidateCreateSRV(const FRDGBufferSRVDesc& Desc)
 {
 	FRDGBufferRef Buffer = Desc.Buffer;
 	checkf(Buffer, TEXT("Buffer SRV created with a null buffer."));
-	checkf(!Buffer->IsPassthrough(), TEXT("Buffer SRV created with passthrough buffer '%s'."), Buffer->Name);
 	ExecuteGuard(TEXT("CreateSRV"), Buffer->Name);
 }
 
@@ -371,7 +362,6 @@ void FRDGUserValidation::ValidateCreateUAV(const FRDGTextureUAVDesc& Desc)
 	FRDGTextureRef Texture = Desc.Texture;
 
 	checkf(Texture, TEXT("Texture UAV created with a null texture."));
-	checkf(!Texture->IsPassthrough(), TEXT("Texture UAV created with passthrough texture '%s'."), Texture->Name);
 	ExecuteGuard(TEXT("CreateUAV"), Texture->Name);
 	checkf(Texture->Desc.Flags & TexCreate_UAV, TEXT("Attempted to create UAV from texture %s which was not created with TexCreate_UAV"), Texture->Name);
 	checkf(Desc.MipLevel < Texture->Desc.NumMips, TEXT("Failed to create UAV at mip %d: the texture %s has only %d mip levels."), Desc.MipLevel, Texture->Name, Texture->Desc.NumMips);
@@ -381,7 +371,6 @@ void FRDGUserValidation::ValidateCreateUAV(const FRDGBufferUAVDesc& Desc)
 {
 	FRDGBufferRef Buffer = Desc.Buffer;
 	checkf(Buffer, TEXT("Buffer UAV created with a null buffer."));
-	checkf(!Buffer->IsPassthrough(), TEXT("Buffer UAV created with passthrough buffer '%s'."), Buffer->Name);
 	ExecuteGuard(TEXT("CreateUAV"), Buffer->Name);
 }
 
@@ -567,9 +556,9 @@ void FRDGUserValidation::ValidateAddPass(const FRDGPass* Pass, bool bSkipPassAcc
 		}
 	};
 
-	const auto CheckNotPassthrough = [&](FRDGParentResourceRef Resource)
+	const auto CheckValidResource = [&](FRDGResourceRef Resource)
 	{
-		checkf(!Resource->IsPassthrough(), TEXT("Resource '%s' was created as a passthrough resource but is attached to pass '%s'."), Resource->Name, Pass->GetName());
+		checkf(ResourceMap.Contains(Resource), TEXT("Resource at %p registered with pass %s is not part of the graph and is likely a dangling pointer or garbage value."), Resource, Pass->GetName());
 	};
 
 	const auto CheckNotCopy = [&](FRDGResourceRef Resource)
@@ -611,18 +600,11 @@ void FRDGUserValidation::ValidateAddPass(const FRDGPass* Pass, bool bSkipPassAcc
 
 	PassParameters.Enumerate([&](FRDGParameter Parameter)
 	{
-		if (Parameter.IsParentResource())
-		{
-			if (FRDGParentResourceRef Resource = Parameter.GetAsParentResource())
-			{
-				CheckNotPassthrough(Resource);
-			}
-		}
-		else if (Parameter.IsView())
+		if (Parameter.IsResource())
 		{
-			if (FRDGViewRef View = Parameter.GetAsView())
+			if (FRDGResourceRef Resource = Parameter.GetAsResource())
 			{
-				CheckNotPassthrough(View->GetParent());
+				CheckValidResource(Resource);
 			}
 		}
 
@@ -748,6 +730,7 @@ void FRDGUserValidation::ValidateAddPass(const FRDGPass* Pass, bool bSkipPassAcc
 		{
 			if (FRDGTextureRef Texture = RenderTargetBindingSlots->ShadingRateTexture)
 			{
+				CheckValidResource(Texture);
 				MarkAsConsumed(Texture);
 			}
 
@@ -757,7 +740,7 @@ void FRDGUserValidation::ValidateAddPass(const FRDGPass* Pass, bool bSkipPassAcc
 			{
 				// Depth stencil only supports one mip, since there isn't actually a way to select the mip level.
 				check(Texture->Desc.NumMips == 1);
-				CheckNotPassthrough(Texture);
+				CheckValidResource(Texture);
 				if (DepthStencil.GetDepthStencilAccess().IsAnyWrite())
 				{
 					MarkTextureAsProduced(Texture);
@@ -803,7 +786,7 @@ void FRDGUserValidation::ValidateAddPass(const FRDGPass* Pass, bool bSkipPassAcc
 						TEXT("Pass '%s' attempted to bind texture '%s' as a render target, but the texture has not been created with TexCreate_ResolveTargetable."),
 						PassName, ResolveTexture->Name);
 
-					CheckNotPassthrough(ResolveTexture);
+					CheckValidResource(Texture);
 					MarkTextureAsProduced(ResolveTexture);
 				}
 
@@ -814,7 +797,7 @@ void FRDGUserValidation::ValidateAddPass(const FRDGPass* Pass, bool bSkipPassAcc
 						TEXT("Pass '%s' attempted to bind texture '%s' as a render target, but the texture has not been created with TexCreate_RenderTargetable."),
 						PassName, Texture->Name);
 
-					CheckNotPassthrough(Texture);
+					CheckValidResource(Texture);
 
 					/** Mark the pass as a producer for render targets with a store action. */
 					MarkTextureAsProduced(Texture);
diff --git a/Engine/Source/Runtime/RenderCore/Public/RenderGraphResources.h b/Engine/Source/Runtime/RenderCore/Public/RenderGraphResources.h
--- a/Engine/Source/Runtime/RenderCore/Public/RenderGraphResources.h
+++ b/Engine/Source/Runtime/RenderCore/Public/RenderGraphResources.h
@@ -124,7 +124,6 @@ private:
 #if RDG_ENABLE_DEBUG
 	struct FRDGResourceDebugData* DebugData = nullptr;
 	FRDGResourceDebugData& GetDebugData() const;
-	bool IsPassthrough() const;
 #endif
 
 	friend FRDGBuilder;
diff --git a/Engine/Source/Runtime/RenderCore/Public/RenderGraphValidation.h b/Engine/Source/Runtime/RenderCore/Public/RenderGraphValidation.h
--- a/Engine/Source/Runtime/RenderCore/Public/RenderGraphValidation.h
+++ b/Engine/Source/Runtime/RenderCore/Public/RenderGraphValidation.h
@@ -106,6 +106,9 @@ private:
 	TArray<FRDGTextureRef, FRDGArrayAllocator> TrackedTextures;
 	TArray<FRDGBufferRef, FRDGArrayAllocator> TrackedBuffers;
 
+	/** Map tracking all active resources in the graph. */
+	TSet<FRDGResourceRef, DefaultKeyFuncs<FRDGResourceRef>, FRDGSetAllocator> ResourceMap;
+
 	/** Whether the Execute() has already been called. */
 	bool bHasExecuted = false;
 
