Fixed RDG builder to propagate transient Discard flag to RHI.

#rnx

diff --git a/Engine/Source/Runtime/RenderCore/Private/RenderGraphBuilder.cpp b/Engine/Source/Runtime/RenderCore/Private/RenderGraphBuilder.cpp
--- a/Engine/Source/Runtime/RenderCore/Private/RenderGraphBuilder.cpp
+++ b/Engine/Source/Runtime/RenderCore/Private/RenderGraphBuilder.cpp
@@ -2153,6 +2153,11 @@ void FRDGBuilder::AddTransition(FRDGPassHandle PassHandle, FRDGTextureRef Textur
 			Info.AccessBefore = SubresourceStateBefore.Access;
 			Info.AccessAfter = SubresourceStateAfter.Access;
 
+			if (Info.AccessBefore == ERHIAccess::Discard)
+			{
+				Info.Flags |= EResourceTransitionFlags::Discard;
+			}
+
 			if (Subresource)
 			{
 				Info.MipIndex = Subresource->MipIndex;
@@ -2484,7 +2489,6 @@ void FRDGBuilder::BeginResourceRHI(FTransientResourceAllocator* TransientResourc
 			FRDGSubresourceState InitialState;
 			InitialState.SetPass(ERHIPipeline::Graphics, PassHandle);
 			InitialState.Access = ERHIAccess::Discard;
-			InitialState.Flags = EResourceTransitionFlags::Discard;
 			InitAsWholeResource(Texture->GetState(), InitialState);
 
 			AddToPrologueBarriers(PassHandle, [&](FRDGBarrierBatchBegin& Barriers)
