Allow extracted resources to be re-used within the same graph if the reference is no longer externally held. Allows overlap of shadow HZB, saving 45MB.

diff --git a/Engine/Source/Runtime/RenderCore/Private/RenderGraphBuilder.cpp b/Engine/Source/Runtime/RenderCore/Private/RenderGraphBuilder.cpp
--- a/Engine/Source/Runtime/RenderCore/Private/RenderGraphBuilder.cpp
+++ b/Engine/Source/Runtime/RenderCore/Private/RenderGraphBuilder.cpp
@@ -2709,12 +2709,6 @@ void FRDGBuilder::EndResourceRHI(FTransientResourceAllocator& TransientResourceA
 
 	if (Texture->ReferenceCount == 0)
 	{
-		// External and transient textures never release the reference.
-		if (!Texture->bExternal && !Texture->bTransient)
-		{
-			Texture->Allocation = nullptr;
-		}
-
 		if (Texture->bTransient)
 		{
 			TransientResourceAllocator->DeallocateMemory(Texture->GetRHIUnchecked());
@@ -2724,6 +2718,10 @@ void FRDGBuilder::EndResourceRHI(FTransientResourceAllocator& TransientResourceA
 				Barriers.AddAlias(Texture, FRHITransientAliasingInfo::Discard(Texture->GetRHIUnchecked()));
 			});
 		}
+		else if (Texture->PooledRenderTarget && Texture->PooledRenderTarget->IsTracked())
+		{
+			Texture->Allocation = nullptr;
+		}
 
 		Texture->LastPass = PassHandle;
 	}
@@ -2737,12 +2735,6 @@ void FRDGBuilder::EndResourceRHI(FTransientResourceAllocator& TransientResourceA
 
 	if (Buffer->ReferenceCount == 0)
 	{
-		// External and transient buffers never release the reference.
-		if (!Buffer->bExternal && !Buffer->bTransient)
-		{
-			Buffer->Allocation = nullptr;
-		}
-
 		if (Buffer->bTransient)
 		{
 			TransientResourceAllocator->DeallocateMemory(Buffer->GetRHIUnchecked());
@@ -2752,6 +2744,10 @@ void FRDGBuilder::EndResourceRHI(FTransientResourceAllocator& TransientResourceA
 				Barriers.AddAlias(Buffer, FRHITransientAliasingInfo::Discard(Buffer->GetRHIUnchecked()));
 			});
 		}
+		else
+		{
+			Buffer->Allocation = nullptr;
+		}
 
 		Buffer->LastPass = PassHandle;
 	}
diff --git a/Engine/Source/Runtime/RenderCore/Private/RenderGraphResources.cpp b/Engine/Source/Runtime/RenderCore/Private/RenderGraphResources.cpp
--- a/Engine/Source/Runtime/RenderCore/Private/RenderGraphResources.cpp
+++ b/Engine/Source/Runtime/RenderCore/Private/RenderGraphResources.cpp
@@ -257,7 +257,7 @@ void FRDGTexture::SetRHI(FRDGPooledTexture* InPooledTexture)
 void FRDGTexture::Finalize()
 {
 	checkf(NextOwner.IsNull() == !!bLastOwner, TEXT("NextOwner must match bLastOwner."));
-	checkf(((bExternal || bExtracted) && !bLastOwner) == false, TEXT("Both external and extracted resources must be the last owner of a resource."));
+	checkf(!bExtracted || bLastOwner, TEXT("Extracted resources must be the last owner of a resource."));
 
 	if (bLastOwner)
 	{
@@ -304,8 +304,8 @@ void FRDGBuffer::SetRHI(FRDGPooledBuffer* InPooledBuffer)
 void FRDGBuffer::Finalize()
 {
 	// If these fire, the graph is not tracking state properly.
-	check(NextOwner.IsNull() == !!bLastOwner);
-	check(!((bExternal || bExtracted) && !bLastOwner));
+	checkf(NextOwner.IsNull() == !!bLastOwner, TEXT("NextOwner must match bLastOwner."));
+	checkf(!bExtracted || bLastOwner, TEXT("Extracted resources must be the last owner of a resource."));
 
 	if (bLastOwner)
 	{
