Allow overriding LogShaders verbosity via logcmds when r.ShaderDevelopmentMode isn't set.

- E.g. now possible to get more info with -logcmds="logshaders veryverbose"

diff --git a/Engine/Source/Runtime/RenderCore/Private/ShaderCore.cpp b/Engine/Source/Runtime/RenderCore/Private/ShaderCore.cpp
--- a/Engine/Source/Runtime/RenderCore/Private/ShaderCore.cpp
+++ b/Engine/Source/Runtime/RenderCore/Private/ShaderCore.cpp
@@ -36,7 +36,8 @@ static TAutoConsoleVariable<int32> CVarShaderDevelopmentMode(
 void UpdateShaderDevelopmentMode()
 {
 	// Keep LogShaders verbosity in sync with r.ShaderDevelopmentMode
-	// r.ShaderDevelopmentMode==1 results in all LogShaders log messages being displayed
+	// r.ShaderDevelopmentMode==1 results in all LogShaders log messages being displayed.
+	// if r.ShaderDevelopmentMode isn't set, we leave the category alone (it defaults to Error, but we can be overriding it to something higher)
 	bool bLogShadersUnsuppressed = UE_LOG_ACTIVE(LogShaders, Log);
 	bool bDesiredLogShadersUnsuppressed = CVarShaderDevelopmentMode.GetValueOnGameThread() == 1;
 
@@ -46,10 +47,6 @@ void UpdateShaderDevelopmentMode()
 		{
 			UE_SET_LOG_VERBOSITY(LogShaders, Log);
 		}
-		else
-		{
-			UE_SET_LOG_VERBOSITY(LogShaders, Error);
-		}
 	}
 }
 
