Support a pass with the ability to explicitly skip pass merging.

diff --git a/Engine/Source/Runtime/RenderCore/Private/RenderGraphBuilder.cpp b/Engine/Source/Runtime/RenderCore/Private/RenderGraphBuilder.cpp
--- a/Engine/Source/Runtime/RenderCore/Private/RenderGraphBuilder.cpp
+++ b/Engine/Source/Runtime/RenderCore/Private/RenderGraphBuilder.cpp
@@ -1164,8 +1164,8 @@ void FRDGBuilder::Compile()
 			{
 				FRDGPass* NextPass = Passes[PassHandle];
 
-				// A pass where the user controls the render pass can't merge with other passes
-				if (EnumHasAnyFlags(NextPass->GetFlags(), ERDGPassFlags::SkipRenderPass))
+				// A pass where the user controls the render pass or it is forced to skip pass merging can't merge with other passes
+				if (EnumHasAnyFlags(NextPass->GetFlags(), ERDGPassFlags::SkipRenderPass | ERDGPassFlags::NeverMerge))
 				{
 					CommitMerge();
 					continue;
diff --git a/Engine/Source/Runtime/RenderCore/Private/RenderGraphValidation.cpp b/Engine/Source/Runtime/RenderCore/Private/RenderGraphValidation.cpp
--- a/Engine/Source/Runtime/RenderCore/Private/RenderGraphValidation.cpp
+++ b/Engine/Source/Runtime/RenderCore/Private/RenderGraphValidation.cpp
@@ -501,6 +501,9 @@ void FRDGUserValidation::ValidateAddPass(const void* ParameterStruct, const FSha
 
 	checkf(!EnumHasAllFlags(Flags, ERDGPassFlags::SkipRenderPass) || EnumHasAllFlags(Flags, ERDGPassFlags::Raster),
 		TEXT("Pass %s specified SkipRenderPass without Raster. Only raster passes support this flag."));
+
+	checkf(!EnumHasAllFlags(Flags, ERDGPassFlags::NeverMerge) || EnumHasAllFlags(Flags, ERDGPassFlags::Raster),
+		TEXT("Pass %s specified NeverMerge without Raster. Only raster passes support this flag."));
 }
 
 void FRDGUserValidation::ValidateAddPass(const FRDGPass* Pass, bool bSkipPassAccessMarking)
diff --git a/Engine/Source/Runtime/RenderCore/Public/RenderGraphDefinitions.h b/Engine/Source/Runtime/RenderCore/Public/RenderGraphDefinitions.h
--- a/Engine/Source/Runtime/RenderCore/Public/RenderGraphDefinitions.h
+++ b/Engine/Source/Runtime/RenderCore/Public/RenderGraphDefinitions.h
@@ -95,6 +95,9 @@ enum class ERDGPassFlags : uint8
 	/** Render pass begin / end is skipped and left to the user. Only valid when combined with 'Raster'. Disables render pass merging for the pass. */
 	SkipRenderPass = 1 << 5,
 
+	/** Pass will never have its render pass merged with other passes. */
+	NeverMerge = 1 << 6,
+
 	/** Pass uses copy commands but writes to a staging resource. */
 	Readback = Copy | NeverCull
 };
