Reduce shader compilation allocations during early editor startup

diff --git a/Engine/Source/Runtime/RenderCore/Private/ShaderCore.cpp b/Engine/Source/Runtime/RenderCore/Private/ShaderCore.cpp
--- a/Engine/Source/Runtime/RenderCore/Private/ShaderCore.cpp
+++ b/Engine/Source/Runtime/RenderCore/Private/ShaderCore.cpp
@@ -8,12 +8,14 @@
 #include "HAL/FileManager.h"
 #include "Misc/FileHelper.h"
 #include "Misc/Paths.h"
+#include "Misc/PathViews.h"
 #include "Misc/ScopeLock.h"
 #include "Stats/StatsMisc.h"
 #include "Misc/ScopedSlowTask.h"
 #include "Misc/StringBuilder.h"
 #include "Shader.h"
 #include "ShaderCompilerCore.h"
+#include "String/Find.h"
 #include "VertexFactory.h"
 #include "Modules/ModuleManager.h"
 #include "Interfaces/IShaderFormat.h"
@@ -442,37 +444,43 @@ static void ReportVirtualShaderFilePathError(TArray<FShaderCompilerError>* Compi
 	UE_LOG(LogShaders, Error, TEXT("%s"), *ErrorString);
 }
 
-bool CheckVirtualShaderFilePath(const FString& VirtualFilePath, TArray<FShaderCompilerError>* CompileErrors /*= nullptr*/)
+
+static bool Contains(FStringView View, FStringView Search)
+{
+	return UE::String::FindFirst(View, Search) != INDEX_NONE;
+}
+
+bool CheckVirtualShaderFilePath(FStringView VirtualFilePath, TArray<FShaderCompilerError>* CompileErrors /*= nullptr*/)
 {
 	bool bSuccess = true;
 
-	if (!VirtualFilePath.StartsWith(TEXT("/")))
+	if (!VirtualFilePath.StartsWith('/'))
 	{
-		FString Error = FString::Printf(TEXT("Virtual shader source file name \"%s\" should be absolute from the virtual root directory \"/\"."), *VirtualFilePath);
+		FString Error = FString::Printf(TEXT("Virtual shader source file name \"%s\" should be absolute from the virtual root directory \"/\"."), *FString(VirtualFilePath));
 		ReportVirtualShaderFilePathError(CompileErrors, Error);
 		bSuccess = false;
 	}
 
-	if (VirtualFilePath.Contains(TEXT("..")))
+	if (Contains(VirtualFilePath, TEXT(".."_SV)))
 	{
-		FString Error = FString::Printf(TEXT("Virtual shader source file name \"%s\" should have relative directories (\"../\") collapsed."), *VirtualFilePath);
+		FString Error = FString::Printf(TEXT("Virtual shader source file name \"%s\" should have relative directories (\"../\") collapsed."), *FString(VirtualFilePath));
 		ReportVirtualShaderFilePathError(CompileErrors, Error);
 		bSuccess = false;
 	}
 
-	if (VirtualFilePath.Contains(TEXT("\\")))
+	if (Contains(VirtualFilePath, TEXT("\\"_SV)))
 	{
-		FString Error = FString::Printf(TEXT("Backslashes are not permitted in virtual shader source file name \"%s\""), *VirtualFilePath);
+		FString Error = FString::Printf(TEXT("Backslashes are not permitted in virtual shader source file name \"%s\""), *FString(VirtualFilePath));
 		ReportVirtualShaderFilePathError(CompileErrors, Error);
 		bSuccess = false;
 	}
 
-	FString Extension = FPaths::GetExtension(VirtualFilePath);
+	FStringView Extension = FPathViews::GetExtension(VirtualFilePath);
 	if (VirtualFilePath.StartsWith(TEXT("/Engine/Shared/")))
 	{
-		if ((Extension != TEXT("h")))
+		if ((Extension != TEXT("h"_SV)))
 		{
-			FString Error = FString::Printf(TEXT("Extension on virtual shader source file name \"%s\" is wrong. Only .h is allowed for shared headers that are shared between C++ and shader code."), *VirtualFilePath);
+			FString Error = FString::Printf(TEXT("Extension on virtual shader source file name \"%s\" is wrong. Only .h is allowed for shared headers that are shared between C++ and shader code."), *FString(VirtualFilePath));
 			ReportVirtualShaderFilePathError(CompileErrors, Error);
 			bSuccess = false;
 		}	
@@ -485,7 +493,7 @@ bool CheckVirtualShaderFilePath(const FString& VirtualFilePath, TArray<FShaderCo
 	{
 		if ((Extension != TEXT("usf") && Extension != TEXT("ush")) || VirtualFilePath.EndsWith(TEXT(".usf.usf")))
 		{
-			FString Error = FString::Printf(TEXT("Extension on virtual shader source file name \"%s\" is wrong. Only .usf or .ush allowed."), *VirtualFilePath);
+			FString Error = FString::Printf(TEXT("Extension on virtual shader source file name \"%s\" is wrong. Only .usf or .ush allowed."), *FString(VirtualFilePath));
 			ReportVirtualShaderFilePathError(CompileErrors, Error);
 			bSuccess = false;
 		}
diff --git a/Engine/Source/Runtime/RenderCore/Public/ShaderCompilerCore.h b/Engine/Source/Runtime/RenderCore/Public/ShaderCompilerCore.h
--- a/Engine/Source/Runtime/RenderCore/Public/ShaderCompilerCore.h
+++ b/Engine/Source/Runtime/RenderCore/Public/ShaderCompilerCore.h
@@ -484,7 +484,7 @@ enum class ESCWErrorCode
  * Meant to be use as such: check(CheckVirtualShaderFilePath(VirtualFilePath));
  * CompileErrors output array is optional.
  */
-extern RENDERCORE_API bool CheckVirtualShaderFilePath(const FString& VirtualPath, TArray<FShaderCompilerError>* CompileErrors = nullptr);
+extern RENDERCORE_API bool CheckVirtualShaderFilePath(FStringView VirtualPath, TArray<FShaderCompilerError>* CompileErrors = nullptr);
 
 /**
  * Loads the shader file with the given name.
