Filling a few gaps in Mesh shader support. Adding ShouldCompilePermutation support to shader pipeline types to guard against unsupported shader types.

diff --git a/Engine/Plugins/FX/Niagara/Source/NiagaraEditor/Private/NiagaraCompiler.cpp b/Engine/Plugins/FX/Niagara/Source/NiagaraEditor/Private/NiagaraCompiler.cpp
--- a/Engine/Plugins/FX/Niagara/Source/NiagaraEditor/Private/NiagaraCompiler.cpp
+++ b/Engine/Plugins/FX/Niagara/Source/NiagaraEditor/Private/NiagaraCompiler.cpp
@@ -1068,6 +1068,8 @@ void FNiagaraEditorModule::TestCompileScriptFromConsole(const TArray<FString>& A
 			Input.Environment.SetDefine(TEXT("HULLSHADER"), 0);
 			Input.Environment.SetDefine(TEXT("VERTEXSHADER"), 0);
 			Input.Environment.SetDefine(TEXT("GEOMETRYSHADER"), 0);
+			Input.Environment.SetDefine(TEXT("MESHSHADER"), 0);
+			Input.Environment.SetDefine(TEXT("AMPLIFICATIONSHADER"), 0);
 			Input.Environment.IncludeVirtualPathToContentsMap.Add(TEXT("/Engine/Generated/NiagaraEmitterInstance.ush"), TranslatedHLSL);
 
 			FShaderCompilerOutput Output;
@@ -1164,6 +1166,8 @@ int32 FHlslNiagaraCompiler::CompileScript(const FNiagaraCompileRequestData* InCo
 	Input.Environment.SetDefine(TEXT("HULLSHADER"), 0);
 	Input.Environment.SetDefine(TEXT("VERTEXSHADER"), 0);
 	Input.Environment.SetDefine(TEXT("GEOMETRYSHADER"), 0);
+	Input.Environment.SetDefine(TEXT("MESHSHADER"), 0);
+	Input.Environment.SetDefine(TEXT("AMPLIFICATIONSHADER"), 0);
 	Input.Environment.IncludeVirtualPathToContentsMap.Add(TEXT("/Engine/Generated/NiagaraEmitterInstance.ush"), TranslatedHLSL);
 	Input.bGenerateDirectCompileFile = false;
 	Input.DumpDebugInfoRootPath = GShaderCompilingManager->GetAbsoluteShaderDebugInfoDirectory() / TEXT("VM");
diff --git a/Engine/Source/Developer/ShaderCompilerCommon/Private/ShaderCompilerCommon.cpp b/Engine/Source/Developer/ShaderCompilerCommon/Private/ShaderCompilerCommon.cpp
--- a/Engine/Source/Developer/ShaderCompilerCommon/Private/ShaderCompilerCommon.cpp
+++ b/Engine/Source/Developer/ShaderCompilerCommon/Private/ShaderCompilerCommon.cpp
@@ -1225,12 +1225,14 @@ FString CreateShaderCompilerWorkerDirectCommandLine(const FShaderCompilerInput&
 	Text += Input.EntryPointName;
 	switch (Input.Target.Frequency)
 	{
-	case SF_Vertex:		Text += TEXT(" -vs"); break;
-	case SF_Hull:		Text += TEXT(" -hs"); break;
-	case SF_Domain:		Text += TEXT(" -ds"); break;
-	case SF_Geometry:	Text += TEXT(" -gs"); break;
-	case SF_Pixel:		Text += TEXT(" -ps"); break;
-	case SF_Compute:	Text += TEXT(" -cs"); break;
+	case SF_Vertex:			Text += TEXT(" -vs"); break;
+	case SF_Mesh:			Text += TEXT(" -ms"); break;
+	case SF_Amplification:	Text += TEXT(" -as"); break;
+	case SF_Hull:			Text += TEXT(" -hs"); break;
+	case SF_Domain:			Text += TEXT(" -ds"); break;
+	case SF_Geometry:		Text += TEXT(" -gs"); break;
+	case SF_Pixel:			Text += TEXT(" -ps"); break;
+	case SF_Compute:		Text += TEXT(" -cs"); break;
 #if RHI_RAYTRACING
 	case SF_RayGen:			Text += TEXT(" -rgs"); break;
 	case SF_RayMiss:		Text += TEXT(" -rms"); break;
diff --git a/Engine/Source/Programs/ShaderCompileWorker/Private/ShaderCompileWorker.cpp b/Engine/Source/Programs/ShaderCompileWorker/Private/ShaderCompileWorker.cpp
--- a/Engine/Source/Programs/ShaderCompileWorker/Private/ShaderCompileWorker.cpp
+++ b/Engine/Source/Programs/ShaderCompileWorker/Private/ShaderCompileWorker.cpp
@@ -886,6 +886,14 @@ static void DirectCompile(const TArray<const class IShaderFormat*>& ShaderFormat
 			{
 				Frequency = SF_Vertex;
 			}
+			else if (!FCString::Strcmp(*Token, TEXT("ms")))
+			{
+				Frequency = SF_Mesh;
+			}
+			else if (!FCString::Strcmp(*Token, TEXT("as")))
+			{
+				Frequency = SF_Amplification;
+			}
 			else if (!FCString::Strcmp(*Token, TEXT("hs")))
 			{
 				Frequency = SF_Hull;
diff --git a/Engine/Source/Runtime/Engine/Private/ShaderCompiler/ShaderCompiler.cpp b/Engine/Source/Runtime/Engine/Private/ShaderCompiler/ShaderCompiler.cpp
--- a/Engine/Source/Runtime/Engine/Private/ShaderCompiler/ShaderCompiler.cpp
+++ b/Engine/Source/Runtime/Engine/Private/ShaderCompiler/ShaderCompiler.cpp
@@ -4598,6 +4598,8 @@ void GlobalBeginCompileShader(
 		Input.Environment.SetDefine(TEXT("DOMAINSHADER"), Target.Frequency == SF_Domain);
 		Input.Environment.SetDefine(TEXT("HULLSHADER"), Target.Frequency == SF_Hull);
 		Input.Environment.SetDefine(TEXT("VERTEXSHADER"), Target.Frequency == SF_Vertex);
+		Input.Environment.SetDefine(TEXT("MESHSHADER"), Target.Frequency == SF_Mesh);
+		Input.Environment.SetDefine(TEXT("AMPLIFICATIONSHADER"), Target.Frequency == SF_Amplification);
 		Input.Environment.SetDefine(TEXT("GEOMETRYSHADER"), Target.Frequency == SF_Geometry);
 		Input.Environment.SetDefine(TEXT("COMPUTESHADER"), Target.Frequency == SF_Compute);
 		Input.Environment.SetDefine(TEXT("RAYCALLABLESHADER"), Target.Frequency == SF_RayCallable);
@@ -5609,10 +5611,9 @@ void VerifyGlobalShaders(EShaderPlatform Platform, const ITargetPlatform* Target
 		const FShaderPipelineType* Pipeline = *ShaderPipelineIt;
 		if (Pipeline->IsGlobalTypePipeline())
 		{
-			if (!GlobalShaderMap->HasShaderPipeline(Pipeline) || (OutdatedShaderPipelineTypes && OutdatedShaderPipelineTypes->Contains(Pipeline)))
+			if (Pipeline->ShouldCompilePermutation(FGlobalShaderPermutationParameters(Platform, kUniqueShaderPermutationId, PermutationFlags))
+				&& (!GlobalShaderMap->HasShaderPipeline(Pipeline) || (OutdatedShaderPipelineTypes && OutdatedShaderPipelineTypes->Contains(Pipeline))))
 			{
-				auto& StageTypes = Pipeline->GetStages();
-
 				if (OutdatedShaderPipelineTypes)
 				{
 					// Remove old pipeline
@@ -5637,7 +5638,7 @@ void VerifyGlobalShaders(EShaderPlatform Platform, const ITargetPlatform* Target
 				else
 				{
 					// If sharing shaders amongst pipelines, add this pipeline as a dependency of an existing individual job
-					for (const FShaderType* ShaderType : StageTypes)
+					for (const FShaderType* ShaderType : Pipeline->GetStages())
 					{
 						TShaderTypePermutation<const FShaderType> ShaderTypePermutation(ShaderType, kUniqueShaderPermutationId);
 
diff --git a/Engine/Source/Runtime/RHI/Private/PipelineFileCache.cpp b/Engine/Source/Runtime/RHI/Private/PipelineFileCache.cpp
--- a/Engine/Source/Runtime/RHI/Private/PipelineFileCache.cpp
+++ b/Engine/Source/Runtime/RHI/Private/PipelineFileCache.cpp
@@ -666,7 +666,7 @@ bool FPipelineCacheFileFormatPSO::Verify() const
 			return false;
 		}
 
-#if PLATFORM_SUPPORTS_TESSELLATION_SHADERS
+#if PLATFORM_SUPPORTS_MESH_SHADERS
 		if (GraphicsDesc.VertexShader != FSHAHash() && GraphicsDesc.MeshShader != FSHAHash())
 		{
 			// Vertex shader and mesh shader are mutually exclusive
@@ -3408,6 +3408,16 @@ struct FPipelineCacheFileData
 									{
 										Entry.Value.Shaders.Add(PSO.GraphicsDesc.DomainShader);
 									}
+
+									if (PSO.GraphicsDesc.MeshShader != FSHAHash())
+									{
+										Entry.Value.Shaders.Add(PSO.GraphicsDesc.MeshShader);
+									}
+
+									if (PSO.GraphicsDesc.AmplificationShader != FSHAHash())
+									{
+										Entry.Value.Shaders.Add(PSO.GraphicsDesc.AmplificationShader);
+									}
 									break;
 								case FPipelineCacheFileFormatPSO::DescriptorType::RayTracing:
 									Entry.Value.Shaders.Add(PSO.RayTracingDesc.ShaderHash);
diff --git a/Engine/Source/Runtime/RHI/Public/RHIValidation.h b/Engine/Source/Runtime/RHI/Public/RHIValidation.h
--- a/Engine/Source/Runtime/RHI/Public/RHIValidation.h
+++ b/Engine/Source/Runtime/RHI/Public/RHIValidation.h
@@ -172,6 +172,20 @@ public:
 		return RHI->RHICreateGeometryShader(Code, Hash);
 	}
 
+	// FlushType: Wait RHI Thread
+	virtual FMeshShaderRHIRef RHICreateMeshShader(TArrayView<const uint8> Code, const FSHAHash& Hash)
+	{
+		check(RHISupportsMeshShaders(GMaxRHIShaderPlatform));
+		return RHI->RHICreateMeshShader(Code, Hash);
+	}
+
+	// FlushType: Wait RHI Thread
+	virtual FAmplificationShaderRHIRef RHICreateAmplificationShader(TArrayView<const uint8> Code, const FSHAHash& Hash)
+	{
+		check(RHISupportsMeshShaders(GMaxRHIShaderPlatform));
+		return RHI->RHICreateAmplificationShader(Code, Hash);
+	}
+
 	// Some RHIs can have pending messages/logs for error tracking, or debug modes
 	virtual void FlushPendingLogs() override final
 	{
@@ -1385,6 +1399,18 @@ public:
 		return RHI->CreateGeometryShader_RenderThread(RHICmdList, Code, Hash);
 	}
 
+	virtual FMeshShaderRHIRef CreateMeshShader_RenderThread(class FRHICommandListImmediate& RHICmdList, TArrayView<const uint8> Code, const FSHAHash& Hash) override final
+	{
+		check(RHISupportsMeshShaders(GMaxRHIShaderPlatform));
+		return RHI->CreateMeshShader_RenderThread(RHICmdList, Code, Hash);
+	}
+
+	virtual FAmplificationShaderRHIRef CreateAmplificationShader_RenderThread(class FRHICommandListImmediate& RHICmdList, TArrayView<const uint8> Code, const FSHAHash& Hash) override final
+	{
+		check(RHISupportsMeshShaders(GMaxRHIShaderPlatform));
+		return RHI->CreateAmplificationShader_RenderThread(RHICmdList, Code, Hash);
+	}
+
 	virtual FComputeShaderRHIRef CreateComputeShader_RenderThread(class FRHICommandListImmediate& RHICmdList, TArrayView<const uint8> Code, const FSHAHash& Hash) override final
 	{
 		check(RHISupportsComputeShaders(GMaxRHIShaderPlatform));
diff --git a/Engine/Source/Runtime/RenderCore/Private/GlobalShader.cpp b/Engine/Source/Runtime/RenderCore/Private/GlobalShader.cpp
--- a/Engine/Source/Runtime/RenderCore/Private/GlobalShader.cpp
+++ b/Engine/Source/Runtime/RenderCore/Private/GlobalShader.cpp
@@ -377,7 +377,9 @@ bool FGlobalShaderMap::IsComplete(const ITargetPlatform* TargetPlatform) const
 	for (TLinkedList<FShaderPipelineType*>::TIterator ShaderPipelineIt(FShaderPipelineType::GetTypeList()); ShaderPipelineIt; ShaderPipelineIt.Next())
 	{
 		const FShaderPipelineType* Pipeline = *ShaderPipelineIt;
-		if (Pipeline->IsGlobalTypePipeline() && !HasShaderPipeline(Pipeline))
+		if (Pipeline->IsGlobalTypePipeline()
+			&& Pipeline->ShouldCompilePermutation(FGlobalShaderPermutationParameters(Platform, kUniqueShaderPermutationId, PermutationFlags))
+			&& !HasShaderPipeline(Pipeline))
 		{
 			return false;
 		}
diff --git a/Engine/Source/Runtime/RenderCore/Private/Shader.cpp b/Engine/Source/Runtime/RenderCore/Private/Shader.cpp
--- a/Engine/Source/Runtime/RenderCore/Private/Shader.cpp
+++ b/Engine/Source/Runtime/RenderCore/Private/Shader.cpp
@@ -800,24 +800,21 @@ static TArray<FShaderPipelineType*>& GetSortedShaderPipelineTypes(FShaderType::E
 
 FShaderPipelineType::FShaderPipelineType(
 	const TCHAR* InName,
-	const FShaderType* InVertexShader,
-	const FShaderType* InHullShader,
-	const FShaderType* InDomainShader,
-	const FShaderType* InGeometryShader,
+	const FShaderType* InVertexOrMeshShader,
+	const FShaderType* InGeometryOrAmplificationShader,
 	const FShaderType* InPixelShader,
+	bool bInIsMeshPipeline,
 	bool bInShouldOptimizeUnusedOutputs) :
 	Name(InName),
 	TypeName(Name),
 	HashedName(TypeName),
-	HashedPrimaryShaderFilename(InVertexShader->GetShaderFilename()),
+	HashedPrimaryShaderFilename(InVertexOrMeshShader->GetShaderFilename()),
 	GlobalListLink(this),
 	bShouldOptimizeUnusedOutputs(bInShouldOptimizeUnusedOutputs)
 {
 	checkf(Name && *Name, TEXT("Shader Pipeline Type requires a valid Name!"));
 
-	checkf(InVertexShader, TEXT("A Shader Pipeline always requires a Vertex Shader"));
-
-	checkf((InHullShader == nullptr && InDomainShader == nullptr) || (InHullShader != nullptr && InDomainShader != nullptr), TEXT("Both Hull & Domain shaders are needed for tessellation on Pipeline %s"), Name);
+	checkf(InVertexOrMeshShader, TEXT("A Shader Pipeline always requires a Vertex or Mesh Shader"));
 
 	//make sure the name is shorter than the maximum serializable length
 	check(FCString::Strlen(InName) < NAME_SIZE);
@@ -826,28 +823,19 @@ FShaderPipelineType::FShaderPipelineType(
 
 	if (InPixelShader)
 	{
-		check(InPixelShader->GetTypeForDynamicCast() == InVertexShader->GetTypeForDynamicCast());
+		check(InPixelShader->GetTypeForDynamicCast() == InVertexOrMeshShader->GetTypeForDynamicCast());
 		Stages.Add(InPixelShader);
 		AllStages[SF_Pixel] = InPixelShader;
 	}
-	if (InGeometryShader)
-	{
-		check(InGeometryShader->GetTypeForDynamicCast() == InVertexShader->GetTypeForDynamicCast());
-		Stages.Add(InGeometryShader);
-		AllStages[SF_Geometry] = InGeometryShader;
-	}
-	if (InDomainShader)
-	{
-		check(InDomainShader->GetTypeForDynamicCast() == InVertexShader->GetTypeForDynamicCast());
-		check(InHullShader->GetTypeForDynamicCast() == InVertexShader->GetTypeForDynamicCast());
-		Stages.Add(InDomainShader);
-		AllStages[SF_Domain] = InDomainShader;
 
-		Stages.Add(InHullShader);
-		AllStages[SF_Hull] = InHullShader;
+	if (InGeometryOrAmplificationShader)
+	{
+		check(InGeometryOrAmplificationShader->GetTypeForDynamicCast() == InVertexOrMeshShader->GetTypeForDynamicCast());
+		Stages.Add(InGeometryOrAmplificationShader);
+		AllStages[bInIsMeshPipeline ? SF_Amplification : SF_Geometry] = InGeometryOrAmplificationShader;
 	}
-	Stages.Add(InVertexShader);
-	AllStages[SF_Vertex] = InVertexShader;
+	Stages.Add(InVertexOrMeshShader);
+	AllStages[bInIsMeshPipeline ? SF_Mesh : SF_Vertex] = InVertexOrMeshShader;
 
 	for (uint32 FrequencyIndex = 0; FrequencyIndex < SF_NumStandardFrequencies; ++FrequencyIndex)
 	{
@@ -864,7 +852,7 @@ FShaderPipelineType::FShaderPipelineType(
 	GlobalListLink.LinkHead(GetTypeList());
 	GetNameToTypeMap().Add(HashedName, this);
 
-	TArray<FShaderPipelineType*>& SortedTypes = GetSortedShaderPipelineTypes(InVertexShader->GetTypeForDynamicCast());
+	TArray<FShaderPipelineType*>& SortedTypes = GetSortedShaderPipelineTypes(InVertexOrMeshShader->GetTypeForDynamicCast());
 	const int32 SortedIndex = Algo::LowerBoundBy(SortedTypes, HashedName, [](const FShaderPipelineType* InType) { return InType->GetHashedName(); });
 	SortedTypes.Insert(this, SortedIndex);
 
@@ -878,7 +866,7 @@ FShaderPipelineType::~FShaderPipelineType()
 	GetNameToTypeMap().Remove(HashedName);
 	GlobalListLink.Unlink();
 
-	TArray<FShaderPipelineType*>& SortedTypes = GetSortedShaderPipelineTypes(AllStages[SF_Vertex]->GetTypeForDynamicCast());
+	TArray<FShaderPipelineType*>& SortedTypes = GetSortedShaderPipelineTypes(AllStages[HasMeshShader() ? SF_Mesh : SF_Vertex]->GetTypeForDynamicCast());
 	const int32 SortedIndex = Algo::BinarySearchBy(SortedTypes, HashedName, [](const FShaderPipelineType* InType) { return InType->GetHashedName(); });
 	check(SortedIndex != INDEX_NONE);
 	SortedTypes.RemoveAt(SortedIndex);
@@ -1004,6 +992,18 @@ const FSHAHash& FShaderPipelineType::GetSourceHash(EShaderPlatform ShaderPlatfor
 	return GetShaderFilesHash(Filenames, ShaderPlatform);
 }
 
+bool FShaderPipelineType::ShouldCompilePermutation(const FShaderPermutationParameters& Parameters) const
+{
+	for (const FShaderType* ShaderType : Stages)
+	{
+		if (!ShaderType->ShouldCompilePermutation(Parameters))
+		{
+			return false;
+		}
+	}
+	return true;
+}
+
 void FShaderPipeline::AddShader(FShader* Shader, int32 PermutationId)
 {
 	const EShaderFrequency Frequency = Shader->GetFrequency();
diff --git a/Engine/Source/Runtime/RenderCore/Private/ShaderCodeLibrary.cpp b/Engine/Source/Runtime/RenderCore/Private/ShaderCodeLibrary.cpp
--- a/Engine/Source/Runtime/RenderCore/Private/ShaderCodeLibrary.cpp
+++ b/Engine/Source/Runtime/RenderCore/Private/ShaderCodeLibrary.cpp
@@ -673,11 +673,11 @@ void FShaderCodeLibraryPipeline::Initialize(const FShaderPipeline* Pipeline)
 void FShaderCodeLibraryPipeline::GetPipelineHash(FSHAHash& Output)
 {
 	FSHA1 Hasher;
-	Hasher.Update(&Shaders[SF_Vertex].Hash[0], sizeof(FSHAHash));
-	Hasher.Update(&Shaders[SF_Pixel].Hash[0], sizeof(FSHAHash));
-	Hasher.Update(&Shaders[SF_Geometry].Hash[0], sizeof(FSHAHash));
-	Hasher.Update(&Shaders[SF_Hull].Hash[0], sizeof(FSHAHash));
-	Hasher.Update(&Shaders[SF_Domain].Hash[0], sizeof(FSHAHash));
+
+	for (int32 Frequency = 0; Frequency < SF_NumGraphicsFrequencies; Frequency++)
+	{
+		Hasher.Update(&Shaders[Frequency].Hash[0], sizeof(FSHAHash));
+	}
 
 	Hasher.Final();
 	Hasher.GetHash(&Output.Hash[0]);
diff --git a/Engine/Source/Runtime/RenderCore/Public/Shader.h b/Engine/Source/Runtime/RenderCore/Public/Shader.h
--- a/Engine/Source/Runtime/RenderCore/Public/Shader.h
+++ b/Engine/Source/Runtime/RenderCore/Public/Shader.h
@@ -1457,15 +1457,14 @@ public:
 	// Set bShouldOptimizeUnusedOutputs to false if the FShaders will point to the individual shaders in the map
 	FShaderPipelineType(
 		const TCHAR* InName,
-		const FShaderType* InVertexShader,
-		const FShaderType* InHullShader,
-		const FShaderType* InDomainShader,
-		const FShaderType* InGeometryShader,
+		const FShaderType* InVertexOrMeshShader,
+		const FShaderType* InGeometryOrAmplificationShader,
 		const FShaderType* InPixelShader,
+		bool bInIsMeshPipeline,
 		bool bInShouldOptimizeUnusedOutputs);
 	~FShaderPipelineType();
 
-	FORCEINLINE bool HasTessellation() const { return AllStages[SF_Domain] != nullptr; }
+	FORCEINLINE bool HasMeshShader() const { return AllStages[SF_Mesh] != nullptr; }
 	FORCEINLINE bool HasGeometry() const { return AllStages[SF_Geometry] != nullptr; }
 	FORCEINLINE bool HasPixelShader() const { return AllStages[SF_Pixel] != nullptr; }
 
@@ -1517,6 +1516,8 @@ public:
 	/** Calculates a Hash based on this shader pipeline type stages' source code and includes */
 	const FSHAHash& GetSourceHash(EShaderPlatform ShaderPlatform) const;
 
+	bool ShouldCompilePermutation(const FShaderPermutationParameters& Parameters) const;
+
 protected:
 	const TCHAR* const Name;
 	FName TypeName;
@@ -1539,28 +1540,23 @@ protected:
 #if !UE_BUILD_DOCS
 // Vertex+Pixel
 #define IMPLEMENT_SHADERPIPELINE_TYPE_VSPS(PipelineName, VertexShaderType, PixelShaderType, bRemoveUnused)	\
-	static FShaderPipelineType PipelineName(TEXT(PREPROCESSOR_TO_STRING(PipelineName)), &VertexShaderType::StaticType, nullptr, nullptr, nullptr, &PixelShaderType::StaticType, bRemoveUnused);
+	static FShaderPipelineType PipelineName(TEXT(PREPROCESSOR_TO_STRING(PipelineName)), &VertexShaderType::StaticType, nullptr, &PixelShaderType::StaticType, false, bRemoveUnused);
 // Only VS
 #define IMPLEMENT_SHADERPIPELINE_TYPE_VS(PipelineName, VertexShaderType, bRemoveUnused)	\
-	static FShaderPipelineType PipelineName(TEXT(PREPROCESSOR_TO_STRING(PipelineName)), &VertexShaderType::StaticType, nullptr, nullptr, nullptr, nullptr, bRemoveUnused);
+	static FShaderPipelineType PipelineName(TEXT(PREPROCESSOR_TO_STRING(PipelineName)), &VertexShaderType::StaticType, nullptr, nullptr, false, bRemoveUnused);
 // Vertex+Geometry+Pixel
 #define IMPLEMENT_SHADERPIPELINE_TYPE_VSGSPS(PipelineName, VertexShaderType, GeometryShaderType, PixelShaderType, bRemoveUnused)	\
-	static FShaderPipelineType PipelineName(TEXT(PREPROCESSOR_TO_STRING(PipelineName)), &VertexShaderType::StaticType, nullptr, nullptr, &GeometryShaderType::StaticType, &PixelShaderType::StaticType, bRemoveUnused);
+	static FShaderPipelineType PipelineName(TEXT(PREPROCESSOR_TO_STRING(PipelineName)), &VertexShaderType::StaticType, &GeometryShaderType::StaticType, &PixelShaderType::StaticType, false, bRemoveUnused);
 // Vertex+Geometry
 #define IMPLEMENT_SHADERPIPELINE_TYPE_VSGS(PipelineName, VertexShaderType, GeometryShaderType, bRemoveUnused)	\
-	static FShaderPipelineType PipelineName(TEXT(PREPROCESSOR_TO_STRING(PipelineName)), &VertexShaderType::StaticType, nullptr, nullptr, &GeometryShaderType::StaticType, nullptr, bRemoveUnused);
-// Vertex+Hull+Domain+Pixel
-#define IMPLEMENT_SHADERPIPELINE_TYPE_VSHSDSPS(PipelineName, VertexShaderType, HullShaderType, DomainShaderType, PixelShaderType, bRemoveUnused)	\
-	static FShaderPipelineType PipelineName(TEXT(PREPROCESSOR_TO_STRING(PipelineName)), &VertexShaderType::StaticType, &HullShaderType::StaticType, &DomainShaderType::StaticType, nullptr, &PixelShaderType::StaticType, bRemoveUnused);
-// Vertex+Hull+Domain+Geometry+Pixel
-#define IMPLEMENT_SHADERPIPELINE_TYPE_VSHSDSGSPS(PipelineName, VertexShaderType, HullShaderType, DomainShaderType, GeometryShaderType, PixelShaderType, bRemoveUnused)	\
-	static FShaderPipelineType PipelineName(TEXT(PREPROCESSOR_TO_STRING(PipelineName)), &VertexShaderType::StaticType, &HullShaderType::StaticType, &DomainShaderType::StaticType, &GeometryShaderType::StaticType, &PixelShaderType::StaticType, bRemoveUnused);
-// Vertex+Hull+Domain
-#define IMPLEMENT_SHADERPIPELINE_TYPE_VSHSDS(PipelineName, VertexShaderType, HullShaderType, DomainShaderType, bRemoveUnused)	\
-	static FShaderPipelineType PipelineName(TEXT(PREPROCESSOR_TO_STRING(PipelineName)), &VertexShaderType::StaticType, &HullShaderType::StaticType, &DomainShaderType::StaticType, nullptr, nullptr, bRemoveUnused);
-// Vertex+Hull+Domain+Geometry
-#define IMPLEMENT_SHADERPIPELINE_TYPE_VSHSDSGS(PipelineName, VertexShaderType, HullShaderType, DomainShaderType, GeometryShaderType, bRemoveUnused)	\
-	static FShaderPipelineType PipelineName(TEXT(PREPROCESSOR_TO_STRING(PipelineName)), &VertexShaderType::StaticType, &HullShaderType::StaticType, &DomainShaderType::StaticType, &GeometryShaderType::StaticType, nullptr, bRemoveUnused);
+	static FShaderPipelineType PipelineName(TEXT(PREPROCESSOR_TO_STRING(PipelineName)), &VertexShaderType::StaticType, &GeometryShaderType::StaticType, nullptr, false, bRemoveUnused);
+
+// Mesh+Pixel
+#define IMPLEMENT_SHADERPIPELINE_TYPE_MSPS(PipelineName, MeshShaderType, PixelShaderType, bRemoveUnused)	\
+	static FShaderPipelineType PipelineName(TEXT(PREPROCESSOR_TO_STRING(PipelineName)), &MeshShaderType::StaticType, nullptr, &PixelShaderType::StaticType, true, bRemoveUnused);
+// Mesh+Amplification+Pixel
+#define IMPLEMENT_SHADERPIPELINE_TYPE_MSASPS(PipelineName, MeshShaderType, AmplificationShaderType, PixelShaderType, bRemoveUnused)	\
+	static FShaderPipelineType PipelineName(TEXT(PREPROCESSOR_TO_STRING(PipelineName)), &MeshShaderType::StaticType, &AmplificationShaderType::StaticType, &PixelShaderType::StaticType, true, bRemoveUnused);
 #endif
 
 /** Encapsulates a dependency on a shader type and saved state from that shader type. */
