Renamed STAT_RenderQueryResultTime to make it clear you need to look at ProfileGPU to understand this time

diff --git a/Engine/Source/Runtime/RenderCore/Public/RenderCore.h b/Engine/Source/Runtime/RenderCore/Public/RenderCore.h
--- a/Engine/Source/Runtime/RenderCore/Public/RenderCore.h
+++ b/Engine/Source/Runtime/RenderCore/Public/RenderCore.h
@@ -16,7 +16,7 @@ DECLARE_LOG_CATEGORY_EXTERN(LogRendererCore, Log, All);
  */
 
 DECLARE_CYCLE_STAT_EXTERN(TEXT("BeginOcclusionTests"),STAT_BeginOcclusionTestsTime,STATGROUP_SceneRendering, RENDERCORE_API);
-DECLARE_CYCLE_STAT_EXTERN(TEXT("RenderQuery Result"),STAT_RenderQueryResultTime,STATGROUP_SceneRendering, RENDERCORE_API);
+DECLARE_CYCLE_STAT_EXTERN(TEXT("Waiting for GPU - check ProfileGPU"),STAT_RenderQueryResultTime,STATGROUP_SceneRendering, RENDERCORE_API);
 DECLARE_CYCLE_STAT_EXTERN(TEXT("InitViews"),STAT_InitViewsTime,STATGROUP_SceneRendering, RENDERCORE_API);
 DECLARE_CYCLE_STAT_EXTERN(TEXT("GatherRayTracingWorldInstances"), STAT_GatherRayTracingWorldInstances, STATGROUP_SceneRendering, RENDERCORE_API);
 DECLARE_CYCLE_STAT_EXTERN(TEXT("InitViewsPossiblyAfterPrepass"), STAT_InitViewsPossiblyAfterPrepass, STATGROUP_SceneRendering, RENDERCORE_API);
