Deprecation of TStaticArray's multi-element init constructor, to avoid confusion of something like TStaticArray<FVector4>(NoInit) and TStaticArray<FVector4>(InPlace, NoInit) producing very different codegen (worse in the original case).

diff --git a/Engine/Source/Runtime/Core/Public/Containers/StaticArray.h b/Engine/Source/Runtime/Core/Public/Containers/StaticArray.h
--- a/Engine/Source/Runtime/Core/Public/Containers/StaticArray.h
+++ b/Engine/Source/Runtime/Core/Public/Containers/StaticArray.h
@@ -19,6 +19,7 @@ public:
 	{
 	}
 
+	UE_DEPRECATED(5.0, "Please call TStaticArray(InPlace, DefaultElement) instead.")
 	explicit TStaticArray(const InElementType& DefaultElement)
 		: Storage(InPlace, TMakeIntegerSequence<uint32, NumElements>(), DefaultElement)
 	{
diff --git a/Engine/Source/Runtime/D3D12RHI/Private/D3D12RHICommon.h b/Engine/Source/Runtime/D3D12RHI/Private/D3D12RHICommon.h
--- a/Engine/Source/Runtime/D3D12RHI/Private/D3D12RHICommon.h
+++ b/Engine/Source/Runtime/D3D12RHI/Private/D3D12RHICommon.h
@@ -302,7 +302,7 @@ private:
 	struct FLinkedObjects
 	{
 		FLinkedObjects() 
-			: Objects(nullptr)
+			: Objects(InPlace, nullptr)
 		{}
 
 		FRHIGPUMask GPUMask;
diff --git a/Engine/Source/Runtime/D3D12RHI/Private/D3D12Stats.cpp b/Engine/Source/Runtime/D3D12RHI/Private/D3D12Stats.cpp
--- a/Engine/Source/Runtime/D3D12RHI/Private/D3D12Stats.cpp
+++ b/Engine/Source/Runtime/D3D12RHI/Private/D3D12Stats.cpp
@@ -8,7 +8,7 @@ D3D12Stats.cpp:RHI Stats and timing implementation.
 #include "Engine/Engine.h"
 #include "Engine/GameViewportClient.h"
 
-TStaticArray<uint32, MAX_NUM_GPUS> D3D12RHI::FD3DGPUProfiler::GGPUFrameCycles(0);
+TStaticArray<uint32, MAX_NUM_GPUS> D3D12RHI::FD3DGPUProfiler::GGPUFrameCycles(InPlace, 0);
 
 void D3D12RHI::FD3DGPUProfiler::BeginFrame(FD3D12DynamicRHI* InRHI)
 {
diff --git a/Engine/Source/Runtime/Engine/Private/PhysicsField/PhysicsFieldComponent.cpp b/Engine/Source/Runtime/Engine/Private/PhysicsField/PhysicsFieldComponent.cpp
--- a/Engine/Source/Runtime/Engine/Private/PhysicsField/PhysicsFieldComponent.cpp
+++ b/Engine/Source/Runtime/Engine/Private/PhysicsField/PhysicsFieldComponent.cpp
@@ -592,7 +592,7 @@ void FPhysicsFieldResource::UpdateResource(FRHICommandListImmediate& RHICmdList,
 
 void FPhysicsFieldInstance::InitInstance( const TArray<EFieldPhysicsType>& TargetTypes)
 {
-	FPhysicsFieldInfos::BufferOffsets VectorTargets(-1), ScalarTargets(-1), IntegerTargets(-1), PhysicsTargets(-1);
+	FPhysicsFieldInfos::BufferOffsets VectorTargets(InPlace, -1), ScalarTargets(InPlace, -1), IntegerTargets(InPlace, -1), PhysicsTargets(InPlace, -1);
 
 	static const TArray<EFieldPhysicsType> VectorTypes = GetFieldTargetTypes(EFieldOutputType::Field_Output_Vector);
 	static const TArray<EFieldPhysicsType> ScalarTypes = GetFieldTargetTypes(EFieldOutputType::Field_Output_Scalar);
diff --git a/Engine/Source/Runtime/MeshDescription/Public/AttributeArrayContainer.h b/Engine/Source/Runtime/MeshDescription/Public/AttributeArrayContainer.h
--- a/Engine/Source/Runtime/MeshDescription/Public/AttributeArrayContainer.h
+++ b/Engine/Source/Runtime/MeshDescription/Public/AttributeArrayContainer.h
@@ -350,9 +350,9 @@ private:
 		// Start, count and allocated count in the Data array for each element in the chunk.
 		// Arranged as SoA for cache optimization, since the most frequent operation is
 		// adding a fixed amount to all the start indices when a value is inserted.
-		TStaticArray<int32, ChunkSize> StartIndex = TStaticArray<int32, ChunkSize>(0);
-		TStaticArray<int32, ChunkSize> Count = TStaticArray<int32, ChunkSize>(0);
-		TStaticArray<int32, ChunkSize> MaxCount = TStaticArray<int32, ChunkSize>(0);
+		TStaticArray<int32, ChunkSize> StartIndex{InPlace, 0};
+		TStaticArray<int32, ChunkSize> Count     {InPlace, 0};
+		TStaticArray<int32, ChunkSize> MaxCount  {InPlace, 0};
 
 		// Specify number of valid elements in this chunk.
 		// Default to maximum number so that we can add chunks in bulk, and all put the final
diff --git a/Engine/Source/Runtime/RHI/Private/GPUProfiler.cpp b/Engine/Source/Runtime/RHI/Private/GPUProfiler.cpp
--- a/Engine/Source/Runtime/RHI/Private/GPUProfiler.cpp
+++ b/Engine/Source/Runtime/RHI/Private/GPUProfiler.cpp
@@ -653,7 +653,7 @@ void FGPUProfiler::PopEvent()
 bool FGPUTiming::GIsSupported = false;
 
 /** Frequency for the timing values, in number of ticks per seconds, or 0 if the feature isn't supported. */
-TStaticArray<uint64, MAX_NUM_GPUS> FGPUTiming::GTimingFrequency(0);
+TStaticArray<uint64, MAX_NUM_GPUS> FGPUTiming::GTimingFrequency(InPlace, 0);
 
 /**
 * Two timestamps performed on GPU and CPU at nearly the same time.
diff --git a/Engine/Source/Runtime/RHI/Public/RHIResources.h b/Engine/Source/Runtime/RHI/Public/RHIResources.h
--- a/Engine/Source/Runtime/RHI/Public/RHIResources.h
+++ b/Engine/Source/Runtime/RHI/Public/RHIResources.h
@@ -2089,7 +2089,7 @@ struct FImmutableSamplerState
 	using TImmutableSamplers = TStaticArray<FRHISamplerState*, MaxImmutableSamplers>;
 
 	FImmutableSamplerState()
-		: ImmutableSamplers(nullptr)
+		: ImmutableSamplers(InPlace, nullptr)
 	{}
 
 	void Reset()
@@ -2144,8 +2144,8 @@ public:
 		, RasterizerState(nullptr)
 		, DepthStencilState(nullptr)
 		, RenderTargetsEnabled(0)
-		, RenderTargetFormats(PF_Unknown)
-		, RenderTargetFlags(0)
+		, RenderTargetFormats(InPlace, PF_Unknown)
+		, RenderTargetFlags(InPlace, 0)
 		, DepthStencilTargetFormat(PF_Unknown)
 		, DepthStencilTargetFlag(0)
 		, DepthTargetLoadAction(ERenderTargetLoadAction::ENoAction)
diff --git a/Engine/Source/Runtime/RenderCore/Private/ProfilingDebugging/RealtimeGPUProfiler.cpp b/Engine/Source/Runtime/RenderCore/Private/ProfilingDebugging/RealtimeGPUProfiler.cpp
--- a/Engine/Source/Runtime/RenderCore/Private/ProfilingDebugging/RealtimeGPUProfiler.cpp
+++ b/Engine/Source/Runtime/RenderCore/Private/ProfilingDebugging/RealtimeGPUProfiler.cpp
@@ -108,8 +108,8 @@ class FRealtimeGPUProfilerEvent
 {
 public:
 	FRealtimeGPUProfilerEvent(FRHIRenderQueryPool& RenderQueryPool)
-		: StartResultMicroseconds(InvalidQueryResult)
-		, EndResultMicroseconds(InvalidQueryResult)
+		: StartResultMicroseconds(InPlace, InvalidQueryResult)
+		, EndResultMicroseconds(InPlace, InvalidQueryResult)
 		, StartQuery(RenderQueryPool.AllocateQuery())
 		, EndQuery(RenderQueryPool.AllocateQuery())
 		, FrameNumber(-1)
@@ -132,8 +132,8 @@ public:
 
 		Name = NewName;
 		STAT(StatName = NewStatName;)
-		StartResultMicroseconds = TStaticArray<uint64, MAX_NUM_GPUS>(InvalidQueryResult);
-		EndResultMicroseconds = TStaticArray<uint64, MAX_NUM_GPUS>(InvalidQueryResult);
+		StartResultMicroseconds = TStaticArray<uint64, MAX_NUM_GPUS>(InPlace, InvalidQueryResult);
+		EndResultMicroseconds = TStaticArray<uint64, MAX_NUM_GPUS>(InPlace, InvalidQueryResult);
 		FrameNumber = GFrameNumberRenderThread;
 	}
 
diff --git a/Engine/Source/Runtime/Renderer/Public/MeshPassProcessor.h b/Engine/Source/Runtime/Renderer/Public/MeshPassProcessor.h
--- a/Engine/Source/Runtime/Renderer/Public/MeshPassProcessor.h
+++ b/Engine/Source/Runtime/Renderer/Public/MeshPassProcessor.h
@@ -261,8 +261,8 @@ public:
 			, ImmutableSamplerState
 			, PrimitiveType
 			, 0
-			, FGraphicsPipelineStateInitializer::TRenderTargetFormats(PF_Unknown)
-			, FGraphicsPipelineStateInitializer::TRenderTargetFlags(0)
+			, FGraphicsPipelineStateInitializer::TRenderTargetFormats(InPlace, PF_Unknown)
+			, FGraphicsPipelineStateInitializer::TRenderTargetFlags(InPlace, 0)
 			, PF_Unknown
 			, TexCreate_None
 			, ERenderTargetLoadAction::ENoAction
