Fix memory leak of shader code during shader compilation (during cooks, deployments, or editor).  Depending on settings and project this can save a considerable amount.  For example, this saved ~39GB on the test project I was cooking.

- FShaderMapResourceCode derives from FThreadSafeRefCountedObject which has a member that keeps track of the number of references to the object.  We were doing deep copies using the Copy Constructor of FShaderMapResourceCode which resulted in a member wise copy of the num refs.  This is not the behavior you'd expect.
- The offending line of code that caused the memory leak was in FShaderMapBase::AssignCopy:
    - Code = new FShaderMapResourceCode(*Source.Code);
- Like FRefCountBase I've deleted the copy constructor and assignment operator for FRefCountedObject and FThreadSafeRefCountedObject to catch issues at compile time where we are copying across NumRefs incorrectly.
- To fix the issue I've manually created copy constructors for FGPUProfilerEventNodeStats and FShaderMapResourceCode.

diff --git a/Engine/Source/Runtime/Core/Public/Templates/RefCounting.h b/Engine/Source/Runtime/Core/Public/Templates/RefCounting.h
--- a/Engine/Source/Runtime/Core/Public/Templates/RefCounting.h
+++ b/Engine/Source/Runtime/Core/Public/Templates/RefCounting.h
@@ -76,6 +76,8 @@ class CORE_API FRefCountedObject
 public:
 	FRefCountedObject(): NumRefs(0) {}
 	virtual ~FRefCountedObject() { check(!NumRefs); }
+	FRefCountedObject(const FRefCountedObject& Rhs) = delete;
+	FRefCountedObject& operator=(const FRefCountedObject& Rhs) = delete;
 	uint32 AddRef() const
 	{
 		return uint32(++NumRefs);
@@ -105,6 +107,8 @@ class CORE_API FThreadSafeRefCountedObject
 public:
 	FThreadSafeRefCountedObject() : NumRefs(0) {}
 	virtual ~FThreadSafeRefCountedObject() { check(NumRefs.GetValue() == 0); }
+	FThreadSafeRefCountedObject(const FThreadSafeRefCountedObject& Rhs) = delete;
+	FThreadSafeRefCountedObject& operator=(const FThreadSafeRefCountedObject& Rhs) = delete;
 	uint32 AddRef() const
 	{
 		return uint32(NumRefs.Increment());
diff --git a/Engine/Source/Runtime/RHI/Public/GPUProfiler.h b/Engine/Source/Runtime/RHI/Public/GPUProfiler.h
--- a/Engine/Source/Runtime/RHI/Public/GPUProfiler.h
+++ b/Engine/Source/Runtime/RHI/Public/GPUProfiler.h
@@ -29,6 +29,20 @@ public:
 	{
 	}
 
+	FGPUProfilerEventNodeStats(const FGPUProfilerEventNodeStats& rhs)
+	{
+		NumDraws = rhs.NumDraws;
+		NumPrimitives = rhs.NumPrimitives;
+		NumVertices = rhs.NumVertices;
+		NumDispatches = rhs.NumDispatches;
+		NumTotalDispatches = rhs.NumTotalDispatches;
+		NumTotalDraws = rhs.NumDraws;
+		NumTotalPrimitives = rhs.NumPrimitives;
+		NumTotalVertices = rhs.NumVertices;
+		TimingResult = rhs.TimingResult;
+		NumEvents = rhs.NumEvents;
+	}
+
 	/** Exclusive number of draw calls rendered in this event. */
 	uint32 NumDraws;
 
diff --git a/Engine/Source/Runtime/RenderCore/Private/ShaderResource.cpp b/Engine/Source/Runtime/RenderCore/Private/ShaderResource.cpp
--- a/Engine/Source/Runtime/RenderCore/Private/ShaderResource.cpp
+++ b/Engine/Source/Runtime/RenderCore/Private/ShaderResource.cpp
@@ -150,6 +150,18 @@ static void RemoveResourceStats(FShaderMapResourceCode& Resource)
 #endif // STATS
 }
 
+FShaderMapResourceCode::FShaderMapResourceCode(const FShaderMapResourceCode& Other)
+{
+	ResourceHash = Other.ResourceHash;
+	ShaderHashes = Other.ShaderHashes;
+	ShaderEntries = Other.ShaderEntries;
+
+#if WITH_EDITORONLY_DATA
+	PlatformDebugData = Other.PlatformDebugData;
+	PlatformDebugDataHashes = Other.PlatformDebugDataHashes;
+#endif // WITH_EDITORONLY_DATA
+}
+
 FShaderMapResourceCode::~FShaderMapResourceCode()
 {
 	RemoveResourceStats(*this);
@@ -424,6 +436,8 @@ FRHIShader* FShaderMapResource::CreateShader(int32 ShaderIndex)
 
 TRefCountPtr<FRHIShader> FShaderMapResource_InlineCode::CreateRHIShader(int32 ShaderIndex)
 {
+	TRACE_CPUPROFILER_EVENT_SCOPE(FShaderMapResource_InlineCode::CreateRHIShader);
+
 	// we can't have this called on the wrong platform's shaders
 	if (!ArePlatformsCompatible(GMaxRHIShaderPlatform, GetPlatform()))
 	{
diff --git a/Engine/Source/Runtime/RenderCore/Public/Shader.h b/Engine/Source/Runtime/RenderCore/Public/Shader.h
--- a/Engine/Source/Runtime/RenderCore/Public/Shader.h
+++ b/Engine/Source/Runtime/RenderCore/Public/Shader.h
@@ -344,7 +344,9 @@ public:
 	};
 #endif // WITH_EDITORONLY_DATA
 
-	RENDERCORE_API ~FShaderMapResourceCode();
+	RENDERCORE_API FShaderMapResourceCode() {}
+	RENDERCORE_API FShaderMapResourceCode(const FShaderMapResourceCode& Other);
+    RENDERCORE_API ~FShaderMapResourceCode();
 
 	RENDERCORE_API void Finalize();
 
