Join the last async compute pass to its first consumer pass instead of the graph epilogue

diff --git a/Engine/Source/Runtime/RenderCore/Private/RenderGraphBuilder.cpp b/Engine/Source/Runtime/RenderCore/Private/RenderGraphBuilder.cpp
--- a/Engine/Source/Runtime/RenderCore/Private/RenderGraphBuilder.cpp
+++ b/Engine/Source/Runtime/RenderCore/Private/RenderGraphBuilder.cpp
@@ -1431,7 +1431,7 @@ void FRDGBuilder::Compile()
 		// Last async compute pass in the graph needs to be manually joined back to the epilogue pass.
 		if (PrevAsyncComputePass)
 		{
-			InsertAsyncToGraphicsComputeJoin(PrevAsyncComputePass, EpiloguePass);
+			InsertAsyncToGraphicsComputeJoin(PrevAsyncComputePass, PrevGraphicsJoinPass);
 			PrevAsyncComputePass->bAsyncComputeEndExecute = 1;
 		}
 	}
