Restoring functionality that allowed platforms to override shader debug settings. Updating some inis to include the default values as examples of how to override them.

diff --git a/Engine/Config/BaseEngine.ini b/Engine/Config/BaseEngine.ini
--- a/Engine/Config/BaseEngine.ini
+++ b/Engine/Config/BaseEngine.ini
@@ -1413,6 +1413,15 @@ bUseVirtualCores = False
 ; Note: CookerMemoryUsedInGB, MemoryToLeaveForTheOSInGB, MemoryUsedPerSCWProcessInGB must all be set to enable
 bForceUseSCWMemoryPressureLimits = False
 
+[ShaderCompiler]
+; Shader compiler debug info cvars that can be overridden per platform.
+; Build shader debug info and export any external debug info.
+r.Shaders.KeepDebugInfo=0
+; Build, but do not export external debug info.
+r.Shaders.PrepareExportedDebugInfo=0
+; Export external debug info if it is built.
+r.Shaders.ExportDebugInfo=0
+
 [LogFiles]
 PurgeLogsDays=5
 MaxLogFilesOnDisk=10
diff --git a/Engine/Source/Runtime/RenderCore/Private/ShaderCore.cpp b/Engine/Source/Runtime/RenderCore/Private/ShaderCore.cpp
--- a/Engine/Source/Runtime/RenderCore/Private/ShaderCore.cpp
+++ b/Engine/Source/Runtime/RenderCore/Private/ShaderCore.cpp
@@ -292,18 +292,17 @@ static bool IsShaderCompilerConfigEnabledForPlatform(const IConsoleVariable* Glo
 	{
 		bEnabled = true;
 	}
-	// TODO: christopher.waters - find out a better way to load these to prevent excessive ini load times on startup without cached ini files.
-#if WITH_EDITOR && 0
+#if WITH_EDITOR
 	else
 	{
 		// Then check the per platform settings.
 		ITargetPlatform* TargetPlatform = GetTargetPlatformManager()->FindTargetPlatformWithSupport(TEXT("ShaderFormat"), ShaderFormat);
 		if (TargetPlatform)
 		{
-			FConfigFile EngineSettings;
-			FConfigCacheIni::LoadLocalIniFile(EngineSettings, TEXT("Engine"), true, *TargetPlatform->IniPlatformName());
-
-			EngineSettings.GetBool(TEXT("ShaderCompiler"), Key, bEnabled);
+			if (FConfigCacheIni* PlatformConfig = TargetPlatform->GetConfigSystem())
+			{
+				PlatformConfig->GetBool(TEXT("ShaderCompiler"), Key, bEnabled, GEngineIni);
+			}
 		}
 	}
 #endif
