Force IsUsingDBuffers to return false on mobile.
UE5 decal code has already been modified so that material compilation and rendering will remain correct on both mobile and mobile deferred after this change.

diff --git a/Engine/Source/Runtime/RenderCore/Private/RenderUtils.cpp b/Engine/Source/Runtime/RenderCore/Private/RenderUtils.cpp
--- a/Engine/Source/Runtime/RenderCore/Private/RenderUtils.cpp
+++ b/Engine/Source/Runtime/RenderCore/Private/RenderUtils.cpp
@@ -1337,7 +1337,7 @@ RENDERCORE_API void RenderUtilsInit()
 					GForwardShadingPlatformMask &= ~Mask;
 				}
 
-				if (TargetPlatform->UsesDBuffer())
+				if (TargetPlatform->UsesDBuffer() && !IsMobilePlatform(ShaderPlatform))
 				{
 					GDBufferPlatformMask |= Mask;
 				}
@@ -1414,6 +1414,7 @@ RENDERCORE_API void RenderUtilsInit()
 #else
 	if (IsMobilePlatform(GMaxRHIShaderPlatform))
 	{
+		GDBufferPlatformMask = 0;
 		GBasePassVelocityPlatformMask = 0;
 	}
 #endif // WITH_EDITOR
