Disabling SDF shadowing on D3D mobile for now as it adds one more sampler

diff --git a/Engine/Source/Runtime/RenderCore/Private/RenderUtils.cpp b/Engine/Source/Runtime/RenderCore/Private/RenderUtils.cpp
--- a/Engine/Source/Runtime/RenderCore/Private/RenderUtils.cpp
+++ b/Engine/Source/Runtime/RenderCore/Private/RenderUtils.cpp
@@ -1160,7 +1160,7 @@ RENDERCORE_API bool UseMobileAmbientOcclusion(const FStaticShaderPlatform Platfo
 
 RENDERCORE_API bool IsMobileDistanceFieldEnabled(const FStaticShaderPlatform Platform)
 {
-	return IsMobilePlatform(Platform) && (IsSwitchPlatform(Platform) || IsD3DPlatform(Platform)) && IsUsingDistanceFields(Platform);
+	return IsMobilePlatform(Platform) && (IsSwitchPlatform(Platform)/* || IsD3DPlatform(Platform)*/) && IsUsingDistanceFields(Platform);
 }
 
 RENDERCORE_API bool IsMobileDistanceFieldShadowingEnabled(const FStaticShaderPlatform Platform)
