RenderGraph - allow to create an aliased UAV into a BC surface

diff --git a/Engine/Source/Runtime/RHI/Public/RHIUtilities.h b/Engine/Source/Runtime/RHI/Public/RHIUtilities.h
--- a/Engine/Source/Runtime/RHI/Public/RHIUtilities.h
+++ b/Engine/Source/Runtime/RHI/Public/RHIUtilities.h
@@ -47,6 +47,23 @@ static inline bool IsStencilFormat(EPixelFormat Format)
 	return false;
 }
 
+static inline bool IsBlockCompressedFormat(EPixelFormat Format)
+{
+	switch (Format)
+	{
+	case PF_DXT1:
+	case PF_DXT3:
+	case PF_DXT5:
+	case PF_BC4:
+	case PF_BC5:
+	case PF_BC6H:
+	case PF_BC7:
+		return true;
+	}
+
+	return false;
+}
+
 inline bool RHISupportsTextureBuffers(const FStaticShaderPlatform Platform)
 {
 	return IsFeatureLevelSupported(Platform, ERHIFeatureLevel::SM5)
diff --git a/Engine/Source/Runtime/RenderCore/Private/RenderGraphValidation.cpp b/Engine/Source/Runtime/RenderCore/Private/RenderGraphValidation.cpp
--- a/Engine/Source/Runtime/RenderCore/Private/RenderGraphValidation.cpp
+++ b/Engine/Source/Runtime/RenderCore/Private/RenderGraphValidation.cpp
@@ -363,7 +363,14 @@ void FRDGUserValidation::ValidateCreateUAV(const FRDGTextureUAVDesc& Desc)
 
 	checkf(Texture, TEXT("Texture UAV created with a null texture."));
 	ExecuteGuard(TEXT("CreateUAV"), Texture->Name);
-	checkf(Texture->Desc.Flags & TexCreate_UAV, TEXT("Attempted to create UAV from texture %s which was not created with TexCreate_UAV"), Texture->Name);
+
+	// BC surfaces can't have TexCreate_UAV, but it's valid to create UINT2/UINT4 UAV for them if GRHISupportsUAVFormatAliasing
+	const bool BCFormatUAVAliasing = IsBlockCompressedFormat(Texture->Desc.Format) && GRHISupportsUAVFormatAliasing;
+	if (!BCFormatUAVAliasing)
+	{
+		checkf(Texture->Desc.Flags & TexCreate_UAV, TEXT("Attempted to create UAV from texture %s which was not created with TexCreate_UAV"), Texture->Name);
+	}
+
 	checkf(Desc.MipLevel < Texture->Desc.NumMips, TEXT("Failed to create UAV at mip %d: the texture %s has only %d mip levels."), Desc.MipLevel, Texture->Name, Texture->Desc.NumMips);
 }
 
