Skip RDG trace work when RDG tracing is disabled.

diff --git a/Engine/Source/Runtime/RenderCore/Private/RenderGraphTrace.cpp b/Engine/Source/Runtime/RenderCore/Private/RenderGraphTrace.cpp
--- a/Engine/Source/Runtime/RenderCore/Private/RenderGraphTrace.cpp
+++ b/Engine/Source/Runtime/RenderCore/Private/RenderGraphTrace.cpp
@@ -274,12 +274,22 @@ void FRDGTrace::AddResource(FRDGParentResource* Resource)
 
 void FRDGTrace::AddTexturePassDependency(FRDGTexture* Texture, FRDGPass* Pass)
 {
+	if (!IsTraceEnabled())
+	{
+		return;
+	}
+
 	Pass->TraceTextures.Add(Texture->Handle);
 	Texture->TracePasses.Add(Pass->Handle);
 }
 
 void FRDGTrace::AddBufferPassDependency(FRDGBuffer* Buffer, FRDGPass* Pass)
 {
+	if (!IsTraceEnabled())
+	{
+		return;
+	}
+
 	Pass->TraceBuffers.Add(Buffer->Handle);
 	Buffer->TracePasses.Add(Pass->Handle);
 }
