RDG validation fix for RDG_{TEXTURE, BUFFER}_ACCESS_ARRAY.

diff --git a/Engine/Source/Runtime/RenderCore/Private/RenderGraphValidation.cpp b/Engine/Source/Runtime/RenderCore/Private/RenderGraphValidation.cpp
--- a/Engine/Source/Runtime/RenderCore/Private/RenderGraphValidation.cpp
+++ b/Engine/Source/Runtime/RenderCore/Private/RenderGraphValidation.cpp
@@ -1056,6 +1056,20 @@ void FRDGUserValidation::SetAllowRHIAccess(const FRDGPass* Pass, bool bAllowAcce
 				Resource->GetDebugData().bAllowRHIAccess = bAllowAccess;
 			}
 		}
+		else if (Parameter.IsBufferAccessArray())
+		{
+			for (FRDGBufferAccess BufferAccess : Parameter.GetAsBufferAccessArray())
+			{
+				BufferAccess->GetDebugData().bAllowRHIAccess = bAllowAccess;
+			}
+		}
+		else if (Parameter.IsTextureAccessArray())
+		{
+			for (FRDGTextureAccess TextureAccess : Parameter.GetAsTextureAccessArray())
+			{
+				TextureAccess->GetDebugData().bAllowRHIAccess = bAllowAccess;
+			}
+		}
 		else if (Parameter.IsRenderTargetBindingSlots())
 		{
 			const FRenderTargetBindingSlots& RenderTargets = Parameter.GetAsRenderTargetBindingSlots();
