Minor RDG optimizations.
 - Simplified texture subresource tracking.
 - Removed map lookup for each resource in SetupPass.
 - Improved Compile / CollectPassResources to reduce cache misses.
 - Added some container reservations to reduce reallocation costs.
 - Added snapping of buffers to page boundaries to improve re-use.

diff --git a/Engine/Source/Runtime/RenderCore/Private/RenderGraphBuilder.cpp b/Engine/Source/Runtime/RenderCore/Private/RenderGraphBuilder.cpp
--- a/Engine/Source/Runtime/RenderCore/Private/RenderGraphBuilder.cpp
+++ b/Engine/Source/Runtime/RenderCore/Private/RenderGraphBuilder.cpp
@@ -223,75 +223,6 @@ void EnumerateTextureAccess(FRDGParameterStruct PassParameters, ERDGPassFlags Pa
 	});
 }
 
-/** Enumerates all texture pass parameters and calls the provided function. The input function must accept an FRDGResource*,
- *  since either views or textures are provided.
- */
-template <typename TParameterFunction>
-void EnumerateTextureParameters(FRDGParameterStruct PassParameters, TParameterFunction ParameterFunction)
-{
-	PassParameters.EnumerateTextures([&](FRDGParameter Parameter)
-	{
-		switch (Parameter.GetType())
-		{
-		case UBMT_RDG_TEXTURE:
-			if (FRDGTextureRef Texture = Parameter.GetAsTexture())
-			{
-				ParameterFunction(Texture);
-			}
-		break;
-		case UBMT_RDG_TEXTURE_ACCESS:
-			if (FRDGTextureAccess Texture = Parameter.GetAsTextureAccess())
-			{
-				ParameterFunction(Texture);
-			}
-		break;
-		case UBMT_RDG_TEXTURE_ACCESS_ARRAY:
-			for (FRDGTextureAccess Texture : Parameter.GetAsTextureAccessArray())
-			{
-				ParameterFunction(Texture);
-			}
-		break;
-		case UBMT_RDG_TEXTURE_SRV:
-			if (FRDGTextureSRVRef SRV = Parameter.GetAsTextureSRV())
-			{
-				ParameterFunction(SRV);
-			}
-		break;
-		case UBMT_RDG_TEXTURE_UAV:
-			if (FRDGTextureUAVRef UAV = Parameter.GetAsTextureUAV())
-			{
-				ParameterFunction(UAV);
-			}
-		break;
-		case UBMT_RENDER_TARGET_BINDING_SLOTS:
-		{
-			const FRenderTargetBindingSlots& RenderTargets = Parameter.GetAsRenderTargetBindingSlots();
-
-			RenderTargets.Enumerate([&](FRenderTargetBinding RenderTarget)
-			{
-				ParameterFunction(RenderTarget.GetTexture());
-
-				if (RenderTarget.GetResolveTexture())
-				{
-					ParameterFunction(RenderTarget.GetResolveTexture());
-				}
-			});
-
-			if (FRDGTextureRef Texture = RenderTargets.DepthStencil.GetTexture())
-			{
-				ParameterFunction(Texture);
-			}
-
-			if (FRDGTextureRef Texture = RenderTargets.ShadingRateTexture)
-			{
-				ParameterFunction(Texture);
-			}
-		}
-		break;
-		}
-	});
-}
-
 /** Enumerates all buffer accesses and provides the access info. */
 template <typename TAccessFunction>
 void EnumerateBufferAccess(FRDGParameterStruct PassParameters, ERDGPassFlags PassFlags, TAccessFunction AccessFunction)
@@ -343,44 +274,6 @@ void EnumerateBufferAccess(FRDGParameterStruct PassParameters, ERDGPassFlags Pas
 	});
 }
 
-/** Enumerates all buffer pass parameters and calls the provided function. The input function must accept an FRDGResource*,
- *  since either views or textures are provided.
- */
-template <typename TParameterFunction>
-void EnumerateBufferParameters(FRDGParameterStruct PassParameters, TParameterFunction ParameterFunction)
-{
-	PassParameters.EnumerateBuffers([&](FRDGParameter Parameter)
-	{
-		switch (Parameter.GetType())
-		{
-		case UBMT_RDG_BUFFER_ACCESS:
-			if (FRDGBufferRef Buffer = Parameter.GetAsBufferAccess())
-			{
-				ParameterFunction(Buffer);
-			}
-		break;
-		case UBMT_RDG_BUFFER_ACCESS_ARRAY:
-			for (FRDGBufferAccess Buffer : Parameter.GetAsBufferAccessArray())
-			{
-				ParameterFunction(Buffer);
-			}
-		break;
-		case UBMT_RDG_BUFFER_SRV:
-			if (FRDGBufferSRVRef SRV = Parameter.GetAsBufferSRV())
-			{
-				ParameterFunction(SRV);
-			}
-		break;
-		case UBMT_RDG_BUFFER_UAV:
-			if (FRDGBufferUAVRef UAV = Parameter.GetAsBufferUAV())
-			{
-				ParameterFunction(UAV);
-			}
-		break;
-		}
-	});
-}
-
 inline FRDGViewHandle GetHandleIfNoUAVBarrier(FRDGViewRef Resource)
 {
 	if (Resource && (Resource->Type == ERDGViewType::BufferUAV || Resource->Type == ERDGViewType::TextureUAV))
@@ -452,7 +345,12 @@ void FRDGBuilder::TickPoolElements()
 	SET_DWORD_STAT(STAT_RDG_RenderPassMergeCount, GRDGStatRenderPassMergeCount);
 	SET_DWORD_STAT(STAT_RDG_PassDependencyCount, GRDGStatPassDependencyCount);
 	SET_DWORD_STAT(STAT_RDG_TextureCount, GRDGStatTextureCount);
+	SET_DWORD_STAT(STAT_RDG_TextureReferenceCount, GRDGStatTextureReferenceCount);
+	SET_FLOAT_STAT(STAT_RDG_TextureReferenceAverage, (float)(GRDGStatTextureReferenceCount / FMath::Max<float>(GRDGStatTextureCount, 1.0f)));
 	SET_DWORD_STAT(STAT_RDG_BufferCount, GRDGStatBufferCount);
+	SET_DWORD_STAT(STAT_RDG_BufferReferenceCount, GRDGStatBufferReferenceCount);
+	SET_FLOAT_STAT(STAT_RDG_BufferReferenceAverage, (float)(GRDGStatBufferReferenceCount / FMath::Max<float>(GRDGStatBufferCount, 1.0f)));
+	SET_DWORD_STAT(STAT_RDG_ViewCount, GRDGStatViewCount);
 	SET_DWORD_STAT(STAT_RDG_TransientTextureCount, GRDGStatTransientTextureCount);
 	SET_DWORD_STAT(STAT_RDG_TransientBufferCount, GRDGStatTransientBufferCount);
 	SET_DWORD_STAT(STAT_RDG_TransitionCount, GRDGStatTransitionCount);
@@ -465,7 +363,10 @@ void FRDGBuilder::TickPoolElements()
 	GRDGStatRenderPassMergeCount = 0;
 	GRDGStatPassDependencyCount = 0;
 	GRDGStatTextureCount = 0;
+	GRDGStatTextureReferenceCount = 0;
 	GRDGStatBufferCount = 0;
+	GRDGStatBufferReferenceCount = 0;
+	GRDGStatViewCount = 0;
 	GRDGStatTransientTextureCount = 0;
 	GRDGStatTransientBufferCount = 0;
 	GRDGStatTransitionCount = 0;
@@ -554,9 +455,10 @@ FRDGBuilder::FTransientResourceAllocator::~FTransientResourceAllocator()
 IRHITransientResourceAllocator* FRDGBuilder::FTransientResourceAllocator::GetOrCreate()
 {
 	check(!IsImmediateMode());
-	if (!Allocator)
+	if (!Allocator && !bCreateAttempted)
 	{
 		Allocator = RHICreateTransientResourceAllocator();
+		bCreateAttempted = true;
 	}
 	return Allocator;
 }
@@ -839,11 +741,6 @@ void FRDGBuilder::Compile()
 
 		const auto AddCullingDependency = [&](FRDGProducerStatesByPipeline& LastProducers, const FRDGProducerState& NextState, ERHIPipeline NextPipeline)
 		{
-			if (NextState.Access == ERHIAccess::Unknown)
-			{
-				return;
-			}
-
 			for (ERHIPipeline LastPipeline : GetRHIPipelines())
 			{
 				FRDGProducerState& LastProducer = LastProducers[LastPipeline];
@@ -870,21 +767,21 @@ void FRDGBuilder::Compile()
 			FRDGPass* Pass = Passes[PassHandle];
 			const ERHIPipeline PassPipeline = Pass->GetPipeline();
 
-			bool bUntrackedOutputs = Pass->GetParameters().HasExternalOutputs();
+			bool bUntrackedOutputs = Pass->bHasExternalOutputs;
 
-			for (auto& TexturePair : Pass->TextureStates)
+			for (auto& PassState : Pass->TextureStates)
 			{
-				FRDGTextureRef Texture = TexturePair.Key;
+				FRDGTextureRef Texture = PassState.Texture;
 				auto& LastProducers = Texture->LastProducers;
-				auto& PassState = TexturePair.Value.State;
-
-				const bool bWholePassState = IsWholeResource(PassState);
-				const bool bWholeProducers = IsWholeResource(LastProducers);
-				checkf(!bWholeProducers || bWholePassState == bWholeProducers, TEXT("The producer array needs to be at least as large as the pass state array."));
 
 				for (uint32 Index = 0, Count = LastProducers.Num(); Index < Count; ++Index)
 				{
-					const auto& SubresourceState = PassState[bWholePassState ? 0 : Index];
+					const auto& SubresourceState = PassState.State[Index];
+
+					if (SubresourceState.Access == ERHIAccess::Unknown)
+					{
+						continue;
+					}
 
 					FRDGProducerState ProducerState;
 					ProducerState.Access = SubresourceState.Access;
@@ -897,10 +794,10 @@ void FRDGBuilder::Compile()
 				bUntrackedOutputs |= Texture->bExternal;
 			}
 
-			for (auto& BufferPair : Pass->BufferStates)
+			for (auto& PassState : Pass->BufferStates)
 			{
-				FRDGBufferRef Buffer = BufferPair.Key;
-				const auto& SubresourceState = BufferPair.Value.State;
+				FRDGBufferRef Buffer = PassState.Buffer;
+				const auto& SubresourceState = PassState.State;
 
 				FRDGProducerState ProducerState;
 				ProducerState.Access = SubresourceState.Access;
@@ -963,6 +860,10 @@ void FRDGBuilder::Compile()
 		{
 			Buffer->ReferenceCount++;
 		});
+
+		// Prologue / Epilogue passes get a lot of transitions added for external / extracted resources.
+		ProloguePass->GetEpilogueBarriersToBeginForGraphics(Allocator).Reserve(ExternalTextures.Num() + ExternalBuffers.Num());
+		EpiloguePass->GetPrologueBarriersToEnd(Allocator).Reserve(32);
 	}
 
 	// All dependencies in the raw graph have been specified; if enabled, all passes are marked as culled and a
@@ -972,8 +873,9 @@ void FRDGBuilder::Compile()
 	if (GRDGCullPasses)
 	{
 		SCOPED_NAMED_EVENT(FRDGBuilder_Compile_Cull_Passes, FColor::Emerald);
-		TArray<FRDGPassHandle, TInlineAllocator<32, FRDGArrayAllocator>> PassStack;
+		TArray<FRDGPassHandle, FRDGArrayAllocator> PassStack;
 		PassesToCull.Init(true, Passes.Num());
+		PassStack.Reserve(Passes.Num());
 
 	#if STATS
 		GRDGStatPassCullCount += Passes.Num();
@@ -1030,11 +932,6 @@ void FRDGBuilder::Compile()
 
 			const auto MergeSubresourceStates = [&](ERDGParentResourceType ResourceType, FRDGSubresourceState*& PassMergeState, FRDGSubresourceState*& ResourceMergeState, const FRDGSubresourceState& PassState)
 			{
-				if (PassState.Access == ERHIAccess::Unknown)
-				{
-					return;
-				}
-
 				if (!ResourceMergeState || !FRDGSubresourceState::IsMergeAllowed(ResourceType, *ResourceMergeState, PassState))
 				{
 					// Cross-pipeline, non-mergable state changes require a new pass dependency for fencing purposes.
@@ -1071,39 +968,37 @@ void FRDGBuilder::Compile()
 				PassMergeState = ResourceMergeState;
 			};
 
-			for (auto& TexturePair : Pass->TextureStates)
+			for (auto& PassState : Pass->TextureStates)
 			{
-				FRDGTextureRef Texture = TexturePair.Key;
-				auto& PassState = TexturePair.Value;
-
+				FRDGTextureRef Texture = PassState.Texture;
 				Texture->ReferenceCount += PassState.ReferenceCount;
 				Texture->bUsedByAsyncComputePass |= bAsyncComputePass;
 
-				// For simplicity, the merge / pass state dimensionality should match.
-				if (Texture->MergeState.Num() != PassState.State.Num())
-				{
-					check(IsSubresources(Texture->MergeState));
-					InitAsSubresources(PassState.State, Texture->Layout, GetWholeResource(PassState.State));
-				}
-
-				const uint32 SubresourceCount = PassState.State.Num();
-				check(Texture->MergeState.Num() == SubresourceCount);
-				PassState.MergeState.SetNum(SubresourceCount);
+			#if STATS
+				GRDGStatTextureReferenceCount += PassState.ReferenceCount;
+			#endif
 
-				for (uint32 Index = 0; Index < SubresourceCount; ++Index)
+				for (int32 Index = 0; Index < PassState.State.Num(); ++Index)
 				{
+					if (PassState.State[Index].Access == ERHIAccess::Unknown)
+					{
+						continue;
+					}
+
 					MergeSubresourceStates(ERDGParentResourceType::Texture, PassState.MergeState[Index], Texture->MergeState[Index], PassState.State[Index]);
 				}
 			}
 
-			for (auto& BufferPair : Pass->BufferStates)
+			for (auto& PassState : Pass->BufferStates)
 			{
-				FRDGBufferRef Buffer = BufferPair.Key;
-				auto& PassState = BufferPair.Value;
-
+				FRDGBufferRef Buffer = PassState.Buffer;
 				Buffer->ReferenceCount += PassState.ReferenceCount;
 				Buffer->bUsedByAsyncComputePass |= bAsyncComputePass;
 
+			#if STATS
+				GRDGStatBufferReferenceCount += PassState.ReferenceCount;
+			#endif
+
 				MergeSubresourceStates(ERDGParentResourceType::Buffer, PassState.MergeState, Buffer->MergeState, PassState.State);
 			}
 		}
@@ -1122,8 +1017,24 @@ void FRDGBuilder::Compile()
 		{
 			if (PassesToMerge.Num())
 			{
+				const auto SetEpilogueBarrierPass = [&](FRDGPass* Pass, FRDGPassHandle EpilogueBarrierPassHandle)
+				{
+					Pass->EpilogueBarrierPass = EpilogueBarrierPassHandle;
+					Pass->ResourcesToEnd.Reset();
+					Passes[EpilogueBarrierPassHandle]->ResourcesToEnd.Add(Pass);
+				};
+
+				const auto SetPrologueBarrierPass = [&](FRDGPass* Pass, FRDGPassHandle PrologueBarrierPassHandle)
+				{
+					Pass->PrologueBarrierPass = PrologueBarrierPassHandle;
+					Pass->ResourcesToBegin.Reset();
+					Passes[PrologueBarrierPassHandle]->ResourcesToBegin.Add(Pass);
+				};
+
 				const FRDGPassHandle FirstPassHandle = PassesToMerge[0];
 				const FRDGPassHandle LastPassHandle = PassesToMerge.Last();
+				Passes[FirstPassHandle]->ResourcesToBegin.Reserve(PassesToMerge.Num());
+				Passes[LastPassHandle]->ResourcesToEnd.Reserve(PassesToMerge.Num());
 
 				// Given an interval of passes to merge into a single render pass: [B, X, X, X, X, E]
 				//
@@ -1137,7 +1048,7 @@ void FRDGBuilder::Compile()
 				{
 					FRDGPass* Pass = Passes[FirstPassHandle];
 					Pass->bSkipRenderPassEnd = 1;
-					Pass->EpilogueBarrierPass = LastPassHandle;
+					SetEpilogueBarrierPass(Pass, LastPassHandle);
 				}
 
 				// (X) Intermediate passes.
@@ -1147,15 +1058,15 @@ void FRDGBuilder::Compile()
 					FRDGPass* Pass = Passes[PassHandle];
 					Pass->bSkipRenderPassBegin = 1;
 					Pass->bSkipRenderPassEnd = 1;
-					Pass->PrologueBarrierPass = FirstPassHandle;
-					Pass->EpilogueBarrierPass = LastPassHandle;
+					SetPrologueBarrierPass(Pass, FirstPassHandle);
+					SetEpilogueBarrierPass(Pass, LastPassHandle);
 				}
 
 				// (E) Last pass in the merge sequence.
 				{
 					FRDGPass* Pass = Passes[LastPassHandle];
 					Pass->bSkipRenderPassBegin = 1;
-					Pass->PrologueBarrierPass = FirstPassHandle;
+					SetPrologueBarrierPass(Pass, FirstPassHandle);
 				}
 
 #if STATS
@@ -1192,12 +1103,6 @@ void FRDGBuilder::Compile()
 					continue;
 				}
 
-				// A graphics fork pass can't merge with a previous raster pass.
-				if (NextPass->bGraphicsFork)
-				{
-					CommitMerge();
-				}
-
 				const FRenderTargetBindingSlots& RenderTargets = NextPass->GetParameters().GetRenderTargets();
 
 				if (PrevPass)
@@ -1235,20 +1140,6 @@ void FRDGBuilder::Compile()
 		CommitMerge();
 	}
 
-	for (FRDGPassHandle PassHandle = Passes.Begin(); PassHandle != Passes.End(); ++PassHandle)
-	{
-		if (PassesToCull[PassHandle] || PassesWithEmptyParameters[PassHandle] || PassesOnAsyncCompute[PassHandle])
-		{
-			continue;
-		}
-
-		FRDGPass* Pass = Passes[PassHandle];
-		FRDGPass* PrologueBarrierPass = GetPrologueBarrierPass(PassHandle);
-		FRDGPass* EpilogueBarrierPass = GetEpilogueBarrierPass(PassHandle);
-		PrologueBarrierPass->ResourcesToBegin.Add(Pass);
-		EpilogueBarrierPass->ResourcesToEnd.Add(Pass);
-	}
-
 	if (AsyncComputePassCount > 0)
 	{
 		SCOPED_NAMED_EVENT(FRDGBuilder_Compile_AsyncCompute, FColor::Emerald);
@@ -1456,6 +1347,7 @@ void FRDGBuilder::Execute()
 
 	if (!IsImmediateMode())
 	{
+		RDG_ALLOW_RHI_ACCESS_SCOPE();
 		Compile();
 
 		IF_RDG_ENABLE_DEBUG(LogFile.Begin(BuilderName, &Passes, PassesToCull, ProloguePassHandle, EpiloguePassHandle));
@@ -1499,6 +1391,17 @@ void FRDGBuilder::Execute()
 			}
 		}
 
+		{
+			TRACE_CPUPROFILER_EVENT_SCOPE(FRDGBuilder::CreateUniformBuffers);
+
+			for (TConstSetBitIterator It(UniformBuffersToCreate); It; ++It)
+			{
+				const FRDGUniformBufferHandle UniformBufferHandle(It.GetIndex());
+
+				BeginResourceRHI(UniformBuffers[UniformBufferHandle]);
+			}
+		}
+
 		{
 			TRACE_CPUPROFILER_EVENT_SCOPE(FRDGBuilder::CollectBarriers);
 			SCOPE_CYCLE_COUNTER(STAT_RDG_CollectBarriersTime);
@@ -1514,44 +1417,48 @@ void FRDGBuilder::Execute()
 		}
 	}
 
-#if RDG_ENABLE_DEBUG
-	const auto LogResource = [&](auto* Resource, auto& Registry)
 	{
-		if (!Resource->bCulled)
+		TRACE_CPUPROFILER_EVENT_SCOPE(FRDGBuilder::Finalize);
+
+#if RDG_ENABLE_DEBUG
+		const auto LogResource = [&](auto* Resource, auto& Registry)
 		{
-			if (!Resource->bLastOwner)
+			if (!Resource->bCulled)
 			{
-				auto* NextOwner = Registry[Resource->NextOwner];
-				LogFile.AddAliasEdge(Resource, Resource->LastPass, NextOwner, NextOwner->FirstPass);
+				if (!Resource->bLastOwner)
+				{
+					auto* NextOwner = Registry[Resource->NextOwner];
+					LogFile.AddAliasEdge(Resource, Resource->LastPass, NextOwner, NextOwner->FirstPass);
+				}
+				LogFile.AddFirstEdge(Resource, Resource->FirstPass);
 			}
-			LogFile.AddFirstEdge(Resource, Resource->FirstPass);
-		}
-	};
+		};
 #endif
 
-	for (FRDGTextureHandle TextureHandle = Textures.Begin(); TextureHandle != Textures.End(); ++TextureHandle)
-	{
-		FRDGTextureRef Texture = Textures[TextureHandle];
-
-		if (Texture->GetRHIUnchecked())
+		for (FRDGTextureHandle TextureHandle = Textures.Begin(); TextureHandle != Textures.End(); ++TextureHandle)
 		{
-			AddEpilogueTransition(Texture);
-			Texture->Finalize();
+			FRDGTextureRef Texture = Textures[TextureHandle];
 
-			IF_RDG_ENABLE_DEBUG(LogResource(Texture, Textures));
-		}
-	}
+			if (Texture->GetRHIUnchecked())
+			{
+				AddEpilogueTransition(Texture);
+				Texture->Finalize();
 
-	for (FRDGBufferHandle BufferHandle = Buffers.Begin(); BufferHandle != Buffers.End(); ++BufferHandle)
-	{
-		FRDGBufferRef Buffer = Buffers[BufferHandle];
+				IF_RDG_ENABLE_DEBUG(LogResource(Texture, Textures));
+			}
+		}
 
-		if (Buffer->GetRHIUnchecked())
+		for (FRDGBufferHandle BufferHandle = Buffers.Begin(); BufferHandle != Buffers.End(); ++BufferHandle)
 		{
-			AddEpilogueTransition(Buffer);
-			Buffer->Finalize();
+			FRDGBufferRef Buffer = Buffers[BufferHandle];
 
-			IF_RDG_ENABLE_DEBUG(LogResource(Buffer, Buffers));
+			if (Buffer->GetRHIUnchecked())
+			{
+				AddEpilogueTransition(Buffer);
+				Buffer->Finalize();
+
+				IF_RDG_ENABLE_DEBUG(LogResource(Buffer, Buffers));
+			}
 		}
 	}
 
@@ -1619,6 +1526,7 @@ void FRDGBuilder::Execute()
 	}
 	GRDGStatBufferCount += Buffers.Num();
 	GRDGStatTextureCount += Textures.Num();
+	GRDGStatViewCount += Views.Num();
 	GRDGStatMemoryWatermark = FMath::Max(GRDGStatMemoryWatermark, Allocator.GetByteCount());
 #endif
 
@@ -1643,6 +1551,8 @@ void FRDGBuilder::Clear()
 void FRDGBuilder::SetupPass(FRDGPass* Pass)
 {
 	IF_RDG_ENABLE_DEBUG(UserValidation.ValidateAddPass(Pass, bInDebugPassScope));
+	SCOPE_CYCLE_COUNTER(STAT_RDG_SetupTime);
+	CSV_SCOPED_TIMING_STAT_EXCLUSIVE(RDG_SetupPass);
 
 	const FRDGParameterStruct PassParameters = Pass->GetParameters();
 	const FRDGPassHandle PassHandle = Pass->GetHandle();
@@ -1651,9 +1561,21 @@ void FRDGBuilder::SetupPass(FRDGPass* Pass)
 
 	bool bRenderPassOnlyWrites = true;
 
+	const auto TryAddView = [&](FRDGViewRef View)
+	{
+		if (View && View->LastPass != PassHandle)
+		{
+			View->LastPass = PassHandle;
+			Pass->Views.Add(View->Handle);
+		}
+	};
+
+	Pass->Views.Reserve(PassParameters.GetBufferParameterCount() + PassParameters.GetTextureParameterCount());
 	Pass->TextureStates.Reserve(PassParameters.GetTextureParameterCount() + (PassParameters.HasRenderTargets() ? (MaxSimultaneousRenderTargets + 1) : 0));
 	EnumerateTextureAccess(PassParameters, PassFlags, [&](FRDGViewRef TextureView, FRDGTextureRef Texture, ERHIAccess Access, ERDGTextureAccessFlags AccessFlags, FRDGTextureSubresourceRange Range)
 	{
+		TryAddView(TextureView);
+
 		if (Texture->bFinalizedAccess)
 		{
 			// Finalized resources expected to remain in the same state, so are ignored by the graph.
@@ -1664,30 +1586,24 @@ void FRDGBuilder::SetupPass(FRDGPass* Pass)
 			return;
 		}
 
-		check(Access != ERHIAccess::Unknown);
 		const FRDGViewHandle NoUAVBarrierHandle = GetHandleIfNoUAVBarrier(TextureView);
 		const EResourceTransitionFlags TransitionFlags = GetTextureViewTransitionFlags(TextureView, Texture);
 
-		auto& PassState = Pass->TextureStates.FindOrAdd(Texture);
-		PassState.ReferenceCount++;
+		FRDGPass::FTextureState* PassState;
 
-		const bool bWholeTextureRange = Range.IsWholeResource(Texture->GetSubresourceLayout());
-		bool bWholePassState = IsWholeResource(PassState.State);
-
-		// Convert the pass state to subresource dimensionality if we've found a subresource range.
-		if (!bWholeTextureRange && bWholePassState)
+		if (Texture->LastPass != PassHandle)
 		{
-			InitAsSubresources(PassState.State, Texture->Layout, GetWholeResource(PassState.State));
-			bWholePassState = false;
+			Texture->LastPass = PassHandle;
+			Texture->PassStateIndex = Pass->TextureStates.Num();
 
-			// Also convert the texture's internal arrays to subresources. We'll process everything as subresources from now on.
-			if (IsWholeResource(Texture->MergeState))
-			{
-				check(IsWholeResource(Texture->LastProducers));
-				InitAsSubresources(Texture->MergeState, Texture->Layout);
-				InitAsSubresources(Texture->LastProducers, Texture->Layout);
-			}
+			PassState = &Pass->TextureStates.Emplace_GetRef(Texture);
 		}
+		else
+		{
+			PassState = &Pass->TextureStates[Texture->PassStateIndex];
+		}
+
+		PassState->ReferenceCount++;
 
 		const auto AddSubresourceAccess = [&](FRDGSubresourceState& State)
 		{
@@ -1697,13 +1613,13 @@ void FRDGBuilder::SetupPass(FRDGPass* Pass)
 			State.SetPass(PassPipeline, PassHandle);
 		};
 
-		if (bWholePassState)
+		if (IsWholeResource(PassState->State))
 		{
-			AddSubresourceAccess(GetWholeResource(PassState.State));
+			AddSubresourceAccess(GetWholeResource(PassState->State));
 		}
 		else
 		{
-			EnumerateSubresourceRange(PassState.State, Texture->Layout, Range, AddSubresourceAccess);
+			EnumerateSubresourceRange(PassState->State, Texture->Layout, Range, AddSubresourceAccess);
 		}
 
 		const bool bWritableAccess = IsWritableAccess(Access);
@@ -1714,7 +1630,7 @@ void FRDGBuilder::SetupPass(FRDGPass* Pass)
 	Pass->BufferStates.Reserve(PassParameters.GetBufferParameterCount());
 	EnumerateBufferAccess(PassParameters, PassFlags, [&](FRDGViewRef BufferView, FRDGBufferRef Buffer, ERHIAccess Access)
 	{
-		check(Access != ERHIAccess::Unknown);
+		TryAddView(BufferView);
 
 		if (Buffer->bFinalizedAccess)
 		{
@@ -1728,18 +1644,37 @@ void FRDGBuilder::SetupPass(FRDGPass* Pass)
 
 		const FRDGViewHandle NoUAVBarrierHandle = GetHandleIfNoUAVBarrier(BufferView);
 
-		auto& PassState = Pass->BufferStates.FindOrAdd(Buffer);
-		PassState.ReferenceCount++;
-		PassState.State.Access = MakeValidAccess(PassState.State.Access | Access);
-		PassState.State.NoUAVBarrierFilter.AddHandle(NoUAVBarrierHandle);
-		PassState.State.SetPass(PassPipeline, PassHandle);
+		FRDGPass::FBufferState* PassState;
+
+		if (Buffer->LastPass != PassHandle)
+		{
+			Buffer->LastPass = PassHandle;
+			Buffer->PassStateIndex = Pass->BufferStates.Num();
+
+			PassState = &Pass->BufferStates.Emplace_GetRef(Buffer);
+		}
+		else
+		{
+			PassState = &Pass->BufferStates[Buffer->PassStateIndex];
+		}
+
+		PassState->ReferenceCount++;
+		PassState->State.Access = MakeValidAccess(PassState->State.Access | Access);
+		PassState->State.NoUAVBarrierFilter.AddHandle(NoUAVBarrierHandle);
+		PassState->State.SetPass(PassPipeline, PassHandle);
 
 		const bool bWritableAccess = IsWritableAccess(Access);
 		bRenderPassOnlyWrites &= !bWritableAccess;
 		Buffer->bProduced |= bWritableAccess;
 	});
 
+	PassParameters.EnumerateUniformBuffers([&](FRDGUniformBufferBinding UniformBuffer)
+	{
+		UniformBuffersToCreate[UniformBuffer->Handle] = true;
+	});
+
 	Pass->bRenderPassOnlyWrites = bRenderPassOnlyWrites;
+	Pass->bHasExternalOutputs = PassParameters.HasExternalOutputs();
 
 	const bool bEmptyParameters = !Pass->TextureStates.Num() && !Pass->BufferStates.Num();
 	PassesWithEmptyParameters.Add(bEmptyParameters);
@@ -1763,6 +1698,12 @@ void FRDGBuilder::SetupPassInternal(FRDGPass* Pass, FRDGPassHandle PassHandle, E
 	Pass->PrologueBarrierPass = PassHandle;
 	Pass->EpilogueBarrierPass = PassHandle;
 
+	if (!EnumHasAnyFlags(Pass->GetFlags(), ERDGPassFlags::AsyncCompute))
+	{
+		Pass->ResourcesToBegin = { Pass };
+		Pass->ResourcesToEnd   = { Pass };
+	}
+
 #if WITH_MGPU
 	Pass->GPUMask = RHICmdList.GetGPUMask();
 #endif
@@ -1787,13 +1728,11 @@ void FRDGBuilder::SetupPassInternal(FRDGPass* Pass, FRDGPassHandle PassHandle, E
 
 	if (IsImmediateMode() && Pass != EpiloguePass)
 	{
-		Pass->ResourcesToBegin.Add(Pass);
-		Pass->ResourcesToEnd.Add(Pass);
+		RDG_ALLOW_RHI_ACCESS_SCOPE();
 
 		// Trivially redirect the merge states to the pass states, since we won't be compiling the graph.
-		for (auto& TexturePair : Pass->TextureStates)
+		for (auto& PassState : Pass->TextureStates)
 		{
-			auto& PassState = TexturePair.Value;
 			const uint32 SubresourceCount = PassState.State.Num();
 			PassState.MergeState.SetNum(SubresourceCount);
 			for (uint32 Index = 0; Index < SubresourceCount; ++Index)
@@ -1805,9 +1744,8 @@ void FRDGBuilder::SetupPassInternal(FRDGPass* Pass, FRDGPassHandle PassHandle, E
 			}
 		}
 
-		for (auto& BufferPair : Pass->BufferStates)
+		for (auto& PassState : Pass->BufferStates)
 		{
-			auto& PassState = BufferPair.Value;
 			PassState.MergeState = &PassState.State;
 		}
 
@@ -1820,6 +1758,11 @@ void FRDGBuilder::SetupPassInternal(FRDGPass* Pass, FRDGPassHandle PassHandle, E
 			check(!TransientResourceAllocator);
 		}
 
+		Pass->GetParameters().EnumerateUniformBuffers([&](FRDGUniformBufferBinding UniformBuffer)
+		{
+			BeginResourceRHI(UniformBuffer.GetUniformBuffer());
+		});
+
 		CollectPassBarriers(PassHandle);
 		ExecutePass(Pass);
 	}
@@ -1985,40 +1928,36 @@ void FRDGBuilder::ExecutePass(FRDGPass* Pass)
 
 void FRDGBuilder::CollectPassResources(FTransientResourceAllocator& TransientResourceAllocator, FRDGPassHandle PassHandle)
 {
-	RDG_ALLOW_RHI_ACCESS_SCOPE();
-
 	FRDGPass* Pass = Passes[PassHandle];
 
 	for (FRDGPass* PassToBegin : Pass->ResourcesToBegin)
 	{
-		const FRDGParameterStruct PassParameters = PassToBegin->GetParameters();
-
-		EnumerateTextureParameters(PassParameters, [&](auto Resource)
+		for (const auto& PassState : PassToBegin->TextureStates)
 		{
-			BeginResourceRHI(&TransientResourceAllocator, PassHandle, Resource);
-		});
+			BeginResourceRHI(&TransientResourceAllocator, PassHandle, PassState.Texture);
+		}
 
-		EnumerateBufferParameters(PassParameters, [&](auto Resource)
+		for (const auto& PassState : PassToBegin->BufferStates)
 		{
-			BeginResourceRHI(&TransientResourceAllocator, PassHandle, Resource);
-		});
+			BeginResourceRHI(&TransientResourceAllocator, PassHandle, PassState.Buffer);
+		}
 
-		PassParameters.EnumerateUniformBuffers([&](FRDGUniformBufferBinding UniformBuffer)
+		for (FRDGViewHandle ViewHandle : PassToBegin->Views)
 		{
-			BeginResourceRHI(UniformBuffer.GetUniformBuffer());
-		});
+			BeginResourceRHI(PassHandle, Views[ViewHandle]);
+		}
 	}
 
 	for (FRDGPass* PassToEnd : Pass->ResourcesToEnd)
 	{
-		for (const auto& TexturePair : PassToEnd->TextureStates)
+		for (const auto& PassState : PassToEnd->TextureStates)
 		{
-			EndResourceRHI(TransientResourceAllocator, PassHandle, TexturePair.Key, TexturePair.Value.ReferenceCount);
+			EndResourceRHI(TransientResourceAllocator, PassHandle, PassState.Texture, PassState.ReferenceCount);
 		}
 
-		for (const auto& BufferPair : PassToEnd->BufferStates)
+		for (const auto& PassState : PassToEnd->BufferStates)
 		{
-			EndResourceRHI(TransientResourceAllocator, PassHandle, BufferPair.Key, BufferPair.Value.ReferenceCount);
+			EndResourceRHI(TransientResourceAllocator, PassHandle, PassState.Buffer, PassState.ReferenceCount);
 		}
 	}
 }
@@ -2033,19 +1972,19 @@ void FRDGBuilder::CollectPassBarriers(FRDGPassHandle PassHandle)
 		return;
 	}
 
-	for (const auto& TexturePair : Pass->TextureStates)
+	for (const auto& PassState : Pass->TextureStates)
 	{
-		FRDGTextureRef Texture = TexturePair.Key;
-		AddTransition(PassHandle, Texture, TexturePair.Value.MergeState);
+		FRDGTextureRef Texture = PassState.Texture;
+		AddTransition(PassHandle, Texture, PassState.MergeState);
 		Texture->bCulled = false;
 
 		IF_RDG_ENABLE_TRACE(Trace.AddTexturePassDependency(Texture, Pass));
 	}
 
-	for (const auto& BufferPair : Pass->BufferStates)
+	for (const auto& PassState : Pass->BufferStates)
 	{
-		FRDGBufferRef Buffer = BufferPair.Key;
-		AddTransition(PassHandle, Buffer, *BufferPair.Value.MergeState);
+		FRDGBufferRef Buffer = PassState.Buffer;
+		AddTransition(PassHandle, Buffer, *PassState.MergeState);
 		Buffer->bCulled = false;
 
 		IF_RDG_ENABLE_TRACE(Trace.AddBufferPassDependency(Buffer, Pass));
@@ -2534,7 +2473,7 @@ void FRDGBuilder::BeginResourceRHI(FTransientResourceAllocator* TransientResourc
 	check(Texture->GetRHIUnchecked());
 }
 
-void FRDGBuilder::BeginResourceRHI(FTransientResourceAllocator*, FRDGPassHandle PassHandle, FRDGTextureSRVRef SRV)
+void FRDGBuilder::BeginResourceRHI(FRDGPassHandle PassHandle, FRDGTextureSRVRef SRV)
 {
 	check(SRV);
 
@@ -2550,7 +2489,7 @@ void FRDGBuilder::BeginResourceRHI(FTransientResourceAllocator*, FRDGPassHandle
 	SRV->ResourceRHI = Texture->ViewCache->GetOrCreateSRV(TextureRHI, SRV->Desc);
 }
 
-void FRDGBuilder::BeginResourceRHI(FTransientResourceAllocator* TransientResourceAllocator, FRDGPassHandle PassHandle, FRDGTextureUAVRef UAV)
+void FRDGBuilder::BeginResourceRHI(FRDGPassHandle PassHandle, FRDGTextureUAVRef UAV)
 {
 	check(UAV);
 
@@ -2559,8 +2498,6 @@ void FRDGBuilder::BeginResourceRHI(FTransientResourceAllocator* TransientResourc
 		return;
 	}
 
-	BeginResourceRHI(TransientResourceAllocator, PassHandle, UAV->Desc.Texture);
-
 	FRDGTextureRef Texture = UAV->Desc.Texture;
 	FRHITexture* TextureRHI = Texture->GetRHIUnchecked();
 	check(TextureRHI);
@@ -2580,6 +2517,11 @@ void FRDGBuilder::BeginResourceRHI(FTransientResourceAllocator* TransientResourc
 	check(TransientResourceAllocator || Buffer->bExternal);
 	check(Buffer->ReferenceCount > 0 || Buffer->bExternal || IsImmediateMode());
 
+	const uint64 BufferPageSize = 64 * 1024;
+
+	FRDGBufferDesc SnappedBufferDesc = Buffer->Desc;
+	SnappedBufferDesc.NumElements = Align(SnappedBufferDesc.BytesPerElement * SnappedBufferDesc.NumElements, BufferPageSize) / SnappedBufferDesc.BytesPerElement;
+
 	// If transient then create the resource on the transient allocator. External or extracted resource can't be transient because of lifetime tracking issues.
 	if (IsTransient(Buffer))
 	{
@@ -2587,7 +2529,7 @@ void FRDGBuilder::BeginResourceRHI(FTransientResourceAllocator* TransientResourc
 
 		if (IRHITransientResourceAllocator* AllocatorRHI = TransientResourceAllocator->GetOrCreate())
 		{
-			if (FRHITransientBuffer* TransientBuffer = AllocatorRHI->CreateBuffer(Translate(Buffer->Desc), Buffer->Name))
+			if (FRHITransientBuffer* TransientBuffer = AllocatorRHI->CreateBuffer(Translate(SnappedBufferDesc), Buffer->Name))
 			{
 				Buffer->SetRHI(TransientBuffer, Allocator);
 
@@ -2611,7 +2553,7 @@ void FRDGBuilder::BeginResourceRHI(FTransientResourceAllocator* TransientResourc
 
 	if (!Buffer->bTransient)
 	{
-		Buffer->SetRHI(GRenderGraphResourcePool.FindFreeBufferInternal(RHICmdList, Buffer->Desc, Buffer->Name));
+		Buffer->SetRHI(GRenderGraphResourcePool.FindFreeBufferInternal(RHICmdList, SnappedBufferDesc, Buffer->Name));
 	}
 
 	Buffer->FirstPass = PassHandle;
@@ -2619,7 +2561,7 @@ void FRDGBuilder::BeginResourceRHI(FTransientResourceAllocator* TransientResourc
 	check(Buffer->GetRHIUnchecked());
 }
 
-void FRDGBuilder::BeginResourceRHI(FTransientResourceAllocator*, FRDGPassHandle PassHandle, FRDGBufferSRVRef SRV)
+void FRDGBuilder::BeginResourceRHI(FRDGPassHandle PassHandle, FRDGBufferSRVRef SRV)
 {
 	check(SRV);
 
@@ -2645,7 +2587,7 @@ void FRDGBuilder::BeginResourceRHI(FTransientResourceAllocator*, FRDGPassHandle
 	SRV->ResourceRHI = Buffer->ViewCache->GetOrCreateSRV(BufferRHI, SRVCreateInfo);
 }
 
-void FRDGBuilder::BeginResourceRHI(FTransientResourceAllocator* TransientResourceAllocator, FRDGPassHandle PassHandle, FRDGBufferUAV* UAV)
+void FRDGBuilder::BeginResourceRHI(FRDGPassHandle PassHandle, FRDGBufferUAV* UAV)
 {
 	check(UAV);
 
@@ -2655,7 +2597,7 @@ void FRDGBuilder::BeginResourceRHI(FTransientResourceAllocator* TransientResourc
 	}
 
 	FRDGBufferRef Buffer = UAV->Desc.Buffer;
-	BeginResourceRHI(TransientResourceAllocator, PassHandle, Buffer);
+	check(Buffer);
 
 	FRHIBufferUAVCreateInfo UAVCreateInfo = UAV->Desc;
 
@@ -2670,6 +2612,30 @@ void FRDGBuilder::BeginResourceRHI(FTransientResourceAllocator* TransientResourc
 	UAV->ResourceRHI = Buffer->ViewCache->GetOrCreateUAV(Buffer->GetRHIUnchecked(), UAVCreateInfo);
 }
 
+void FRDGBuilder::BeginResourceRHI(FRDGPassHandle PassHandle, FRDGView* View)
+{
+	if (View->GetRHIUnchecked())
+	{
+		return;
+	}
+
+	switch (View->Type)
+	{
+	case ERDGViewType::TextureUAV:
+		BeginResourceRHI(PassHandle, static_cast<FRDGTextureUAV*>(View));
+		break;
+	case ERDGViewType::TextureSRV:
+		BeginResourceRHI(PassHandle, static_cast<FRDGTextureSRV*>(View));
+		break;
+	case ERDGViewType::BufferUAV:
+		BeginResourceRHI(PassHandle, static_cast<FRDGBufferUAV*>(View));
+		break;
+	case ERDGViewType::BufferSRV:
+		BeginResourceRHI(PassHandle, static_cast<FRDGBufferSRV*>(View));
+		break;
+	}
+}
+
 void FRDGBuilder::EndResourceRHI(FTransientResourceAllocator& TransientResourceAllocator, FRDGPassHandle PassHandle, FRDGTextureRef Texture, uint32 ReferenceCount)
 {
 	check(Texture);
diff --git a/Engine/Source/Runtime/RenderCore/Private/RenderGraphEvent.cpp b/Engine/Source/Runtime/RenderCore/Private/RenderGraphEvent.cpp
--- a/Engine/Source/Runtime/RenderCore/Private/RenderGraphEvent.cpp
+++ b/Engine/Source/Runtime/RenderCore/Private/RenderGraphEvent.cpp
@@ -26,7 +26,6 @@ FRDGEventName::FRDGEventName(const TCHAR* InEventFormat, ...)
 {
 	check(InEventFormat);
 
-	if (GetEmitRDGEvents())
 	{
 		va_list VAList;
 		va_start(VAList, InEventFormat);
diff --git a/Engine/Source/Runtime/RenderCore/Private/RenderGraphPrivate.cpp b/Engine/Source/Runtime/RenderCore/Private/RenderGraphPrivate.cpp
--- a/Engine/Source/Runtime/RenderCore/Private/RenderGraphPrivate.cpp
+++ b/Engine/Source/Runtime/RenderCore/Private/RenderGraphPrivate.cpp
@@ -306,7 +306,7 @@ FAutoConsoleVariableRef CVarRDGUseTransientAllocator(
 	ECVF_RenderThreadSafe);
 
 #if CSV_PROFILER
-int32 GRDGVerboseCSVStats = 0;
+int32 GRDGVerboseCSVStats = 1;
 FAutoConsoleVariableRef CVarRDGVerboseCSVStats(
 	TEXT("r.RDG.VerboseCSVStats"),
 	GRDGVerboseCSVStats,
@@ -323,7 +323,10 @@ int32 GRDGStatPassCullCount = 0;
 int32 GRDGStatPassDependencyCount = 0;
 int32 GRDGStatRenderPassMergeCount = 0;
 int32 GRDGStatTextureCount = 0;
+int32 GRDGStatTextureReferenceCount = 0;
 int32 GRDGStatBufferCount = 0;
+int32 GRDGStatBufferReferenceCount = 0;
+int32 GRDGStatViewCount = 0;
 int32 GRDGStatTransientTextureCount = 0;
 int32 GRDGStatTransientBufferCount = 0;
 int32 GRDGStatTransitionCount = 0;
@@ -337,12 +340,18 @@ DEFINE_STAT(STAT_RDG_PassCullCount);
 DEFINE_STAT(STAT_RDG_RenderPassMergeCount);
 DEFINE_STAT(STAT_RDG_PassDependencyCount);
 DEFINE_STAT(STAT_RDG_TextureCount);
+DEFINE_STAT(STAT_RDG_TextureReferenceCount);
+DEFINE_STAT(STAT_RDG_TextureReferenceAverage);
 DEFINE_STAT(STAT_RDG_BufferCount);
+DEFINE_STAT(STAT_RDG_BufferReferenceCount);
+DEFINE_STAT(STAT_RDG_BufferReferenceAverage);
+DEFINE_STAT(STAT_RDG_ViewCount);
 DEFINE_STAT(STAT_RDG_TransientTextureCount);
 DEFINE_STAT(STAT_RDG_TransientBufferCount);
 DEFINE_STAT(STAT_RDG_TransitionCount);
 DEFINE_STAT(STAT_RDG_AliasingCount);
 DEFINE_STAT(STAT_RDG_TransitionBatchCount);
+DEFINE_STAT(STAT_RDG_SetupTime);
 DEFINE_STAT(STAT_RDG_CompileTime);
 DEFINE_STAT(STAT_RDG_CollectResourcesTime);
 DEFINE_STAT(STAT_RDG_CollectBarriersTime);
diff --git a/Engine/Source/Runtime/RenderCore/Private/RenderGraphPrivate.h b/Engine/Source/Runtime/RenderCore/Private/RenderGraphPrivate.h
--- a/Engine/Source/Runtime/RenderCore/Private/RenderGraphPrivate.h
+++ b/Engine/Source/Runtime/RenderCore/Private/RenderGraphPrivate.h
@@ -122,7 +122,10 @@ extern int32 GRDGStatPassCullCount;
 extern int32 GRDGStatRenderPassMergeCount;
 extern int32 GRDGStatPassDependencyCount;
 extern int32 GRDGStatTextureCount;
+extern int32 GRDGStatTextureReferenceCount;
 extern int32 GRDGStatBufferCount;
+extern int32 GRDGStatBufferReferenceCount;
+extern int32 GRDGStatViewCount;
 extern int32 GRDGStatTransientTextureCount;
 extern int32 GRDGStatTransientBufferCount;
 extern int32 GRDGStatTransitionCount;
@@ -136,13 +139,19 @@ DECLARE_DWORD_COUNTER_STAT_EXTERN(TEXT("Passes Culled"), STAT_RDG_PassCullCount,
 DECLARE_DWORD_COUNTER_STAT_EXTERN(TEXT("Render Passes Merged"), STAT_RDG_RenderPassMergeCount, STATGROUP_RDG, RENDERCORE_API);
 DECLARE_DWORD_COUNTER_STAT_EXTERN(TEXT("Pass Dependencies"), STAT_RDG_PassDependencyCount, STATGROUP_RDG, RENDERCORE_API);
 DECLARE_DWORD_COUNTER_STAT_EXTERN(TEXT("Textures"), STAT_RDG_TextureCount, STATGROUP_RDG, RENDERCORE_API);
+DECLARE_DWORD_COUNTER_STAT_EXTERN(TEXT("Texture References"), STAT_RDG_TextureReferenceCount, STATGROUP_RDG, RENDERCORE_API);
+DECLARE_FLOAT_COUNTER_STAT_EXTERN(TEXT("Texture References Average"), STAT_RDG_TextureReferenceAverage, STATGROUP_RDG, RENDERCORE_API);
 DECLARE_DWORD_COUNTER_STAT_EXTERN(TEXT("Buffers"), STAT_RDG_BufferCount, STATGROUP_RDG, RENDERCORE_API);
+DECLARE_DWORD_COUNTER_STAT_EXTERN(TEXT("Buffer References"), STAT_RDG_BufferReferenceCount, STATGROUP_RDG, RENDERCORE_API);
+DECLARE_FLOAT_COUNTER_STAT_EXTERN(TEXT("Buffer References Average"), STAT_RDG_BufferReferenceAverage, STATGROUP_RDG, RENDERCORE_API);
+DECLARE_DWORD_COUNTER_STAT_EXTERN(TEXT("Views"), STAT_RDG_ViewCount, STATGROUP_RDG, RENDERCORE_API);
 DECLARE_DWORD_COUNTER_STAT_EXTERN(TEXT("Transient Textures"), STAT_RDG_TransientTextureCount, STATGROUP_RDG, RENDERCORE_API);
 DECLARE_DWORD_COUNTER_STAT_EXTERN(TEXT("Transient Buffers"), STAT_RDG_TransientBufferCount, STATGROUP_RDG, RENDERCORE_API);
 DECLARE_DWORD_COUNTER_STAT_EXTERN(TEXT("Resource Transitions"), STAT_RDG_TransitionCount, STATGROUP_RDG, RENDERCORE_API);
 DECLARE_DWORD_COUNTER_STAT_EXTERN(TEXT("Resource Acquires and Discards"), STAT_RDG_AliasingCount, STATGROUP_RDG, RENDERCORE_API);
 DECLARE_DWORD_COUNTER_STAT_EXTERN(TEXT("Resource Transition Batches"), STAT_RDG_TransitionBatchCount, STATGROUP_RDG, RENDERCORE_API);
 
+DECLARE_CYCLE_STAT_EXTERN(TEXT("Setup"), STAT_RDG_SetupTime, STATGROUP_RDG, RENDERCORE_API);
 DECLARE_CYCLE_STAT_EXTERN(TEXT("Compile"), STAT_RDG_CompileTime, STATGROUP_RDG, RENDERCORE_API);
 DECLARE_CYCLE_STAT_EXTERN(TEXT("Collect Resources"), STAT_RDG_CollectResourcesTime, STATGROUP_RDG, RENDERCORE_API);
 DECLARE_CYCLE_STAT_EXTERN(TEXT("Collect Barriers"), STAT_RDG_CollectBarriersTime, STATGROUP_RDG, RENDERCORE_API);
@@ -171,11 +180,16 @@ inline const TCHAR* GetEpilogueBarriersToBeginDebugName(ERHIPipeline Pipelines)
 	return TEXT("");
 }
 
-inline bool IsImmediateMode()
+FORCEINLINE bool IsImmediateMode()
 {
 	return GRDGImmediateMode != 0;
 }
 
+FORCEINLINE bool IsRenderPassMergeEnabled()
+{
+	return GRDGMergeRenderPasses != 0 && !IsImmediateMode();
+}
+
 template <typename ResourceRegistryType, typename FunctionType>
 inline void EnumerateExtendedLifetimeResources(ResourceRegistryType& Registry, FunctionType Function)
 {
diff --git a/Engine/Source/Runtime/RenderCore/Private/RenderGraphResourcePool.cpp b/Engine/Source/Runtime/RenderCore/Private/RenderGraphResourcePool.cpp
--- a/Engine/Source/Runtime/RenderCore/Private/RenderGraphResourcePool.cpp
+++ b/Engine/Source/Runtime/RenderCore/Private/RenderGraphResourcePool.cpp
@@ -7,6 +7,11 @@
 #include "RenderGraphResourcePool.h"
 #include "RenderGraphResources.h"
 
+uint64 ComputeHash(const FRDGBufferDesc& Desc)
+{
+	return CityHash64((const char*)&Desc, sizeof(FRDGBufferDesc));
+}
+
 FRenderGraphResourcePool::FRenderGraphResourcePool()
 { }
 
@@ -26,23 +31,31 @@ TRefCountPtr<FRDGPooledBuffer> FRenderGraphResourcePool::FindFreeBufferInternal(
 	const FRDGBufferDesc& Desc,
 	const TCHAR* InDebugName)
 {
+	const uint64 BufferHash = ComputeHash(Desc);
+
 	// First find if available.
-	for (auto& PooledBuffer : AllocatedBuffers)
+	for (int32 Index = 0; Index < AllocatedBufferHashes.Num(); ++Index)
 	{
+		if (AllocatedBufferHashes[Index] != BufferHash)
+		{
+			continue;
+		}
+
+		const auto& PooledBuffer = AllocatedBuffers[Index];
+
 		// Still being used outside the pool.
 		if (PooledBuffer->GetRefCount() > 1)
 		{
 			continue;
 		}
 
-		if (PooledBuffer->Desc == Desc)
-		{
-			PooledBuffer->LastUsedFrame = FrameCounter;
-			PooledBuffer->ViewCache.SetDebugName(InDebugName);
-			PooledBuffer->Name = InDebugName;
+		check(PooledBuffer->Desc == Desc);
 
-			return PooledBuffer;
-		}
+		PooledBuffer->LastUsedFrame = FrameCounter;
+		PooledBuffer->ViewCache.SetDebugName(InDebugName);
+		PooledBuffer->Name = InDebugName;
+
+		return PooledBuffer;
 	}
 
 	// Allocate new one
@@ -71,11 +84,11 @@ TRefCountPtr<FRDGPooledBuffer> FRenderGraphResourcePool::FindFreeBufferInternal(
 			check(0);
 		}
 
-		TRefCountPtr<FRDGPooledBuffer> PooledBuffer = new FRDGPooledBuffer(MoveTemp(BufferRHI), Desc);
+		TRefCountPtr<FRDGPooledBuffer> PooledBuffer = new FRDGPooledBuffer(MoveTemp(BufferRHI), Desc, InDebugName);
 		AllocatedBuffers.Add(PooledBuffer);
+		AllocatedBufferHashes.Add(BufferHash);
 		check(PooledBuffer->GetRefCount() == 2);
 
-		PooledBuffer->Name = InDebugName;
 		PooledBuffer->LastUsedFrame = FrameCounter;
 
 		return PooledBuffer;
@@ -103,8 +116,8 @@ void FRenderGraphResourcePool::TickPoolElements()
 
 		if (bIsUnused && bNotRequestedRecently)
 		{
-			Swap(Buffer, AllocatedBuffers.Last());
-			AllocatedBuffers.Pop();
+			AllocatedBuffers.RemoveAtSwap(BufferIndex);
+			AllocatedBufferHashes.RemoveAtSwap(BufferIndex);
 		}
 		else
 		{
diff --git a/Engine/Source/Runtime/RenderCore/Private/RenderGraphResourcePool.h b/Engine/Source/Runtime/RenderCore/Private/RenderGraphResourcePool.h
--- a/Engine/Source/Runtime/RenderCore/Private/RenderGraphResourcePool.h
+++ b/Engine/Source/Runtime/RenderCore/Private/RenderGraphResourcePool.h
@@ -36,6 +36,7 @@ private:
 
 	/** Elements can be 0, we compact the buffer later. */
 	TArray<TRefCountPtr<FRDGPooledBuffer>> AllocatedBuffers;
+	TArray<uint64> AllocatedBufferHashes;
 
 	uint32 FrameCounter = 0;
 
diff --git a/Engine/Source/Runtime/RenderCore/Private/RenderGraphResources.cpp b/Engine/Source/Runtime/RenderCore/Private/RenderGraphResources.cpp
--- a/Engine/Source/Runtime/RenderCore/Private/RenderGraphResources.cpp
+++ b/Engine/Source/Runtime/RenderCore/Private/RenderGraphResources.cpp
@@ -40,11 +40,6 @@ FRDGParentResource::~FRDGParentResource()
 	check(IsImmediateMode() || ReferenceCount == 0);
 }
 
-void FRDGParentResource::SetPassthroughRHI(FRHIResource* InResourceRHI)
-{
-	ResourceRHI = InResourceRHI;
-}
-
 bool FRDGProducerState::IsDependencyRequired(FRDGProducerState LastProducer, ERHIPipeline LastPipeline, FRDGProducerState NextState, ERHIPipeline NextPipeline)
 {
 	/** This function determines whether a producer-consumer relationship exists in the graph, which is used for culling and
diff --git a/Engine/Source/Runtime/RenderCore/Public/RenderGraphBuilder.h b/Engine/Source/Runtime/RenderCore/Public/RenderGraphBuilder.h
--- a/Engine/Source/Runtime/RenderCore/Public/RenderGraphBuilder.h
+++ b/Engine/Source/Runtime/RenderCore/Public/RenderGraphBuilder.h
@@ -388,6 +388,9 @@ private:
 	/** Passes that don't have any parameters. */
 	FRDGPassBitArray PassesWithEmptyParameters;
 
+	/** Uniform buffers which were used in a pass. */
+	FRDGUniformBufferBitArray UniformBuffersToCreate;
+
 	/** Tracks external resources to their registered render graph counterparts for de-duplication. */
 	TSortedMap<FRHITexture*, FRDGTexture*, FRDGArrayAllocator> ExternalTextures;
 	TSortedMap<const FRDGPooledBuffer*, FRDGBuffer*, FRDGArrayAllocator> ExternalBuffers;
@@ -458,6 +461,7 @@ private:
 	private:
 		FRHICommandListImmediate& RHICmdList;
 		IRHITransientResourceAllocator* Allocator = nullptr;
+		bool bCreateAttempted = false;
 	};
 
 	void Compile();
@@ -465,11 +469,12 @@ private:
 
 	void BeginResourceRHI(FRDGUniformBuffer* UniformBuffer);
 	void BeginResourceRHI(FTransientResourceAllocator*, FRDGPassHandle, FRDGTexture* Texture);
-	void BeginResourceRHI(FTransientResourceAllocator*, FRDGPassHandle, FRDGTextureSRV* SRV);
-	void BeginResourceRHI(FTransientResourceAllocator*, FRDGPassHandle, FRDGTextureUAV* UAV);
+	void BeginResourceRHI(FRDGPassHandle, FRDGTextureSRV* SRV);
+	void BeginResourceRHI(FRDGPassHandle, FRDGTextureUAV* UAV);
 	void BeginResourceRHI(FTransientResourceAllocator*, FRDGPassHandle, FRDGBuffer* Buffer);
-	void BeginResourceRHI(FTransientResourceAllocator*, FRDGPassHandle, FRDGBufferSRV* SRV);
-	void BeginResourceRHI(FTransientResourceAllocator*, FRDGPassHandle, FRDGBufferUAV* UAV);
+	void BeginResourceRHI(FRDGPassHandle, FRDGBufferSRV* SRV);
+	void BeginResourceRHI(FRDGPassHandle, FRDGBufferUAV* UAV);
+	void BeginResourceRHI(FRDGPassHandle, FRDGView* View);
 
 	void EndResourceRHI(FTransientResourceAllocator&, FRDGPassHandle, FRDGTexture* Texture, uint32 ReferenceCount);
 	void EndResourceRHI(FTransientResourceAllocator&, FRDGPassHandle, FRDGBuffer* Buffer, uint32 ReferenceCount);
diff --git a/Engine/Source/Runtime/RenderCore/Public/RenderGraphBuilder.inl b/Engine/Source/Runtime/RenderCore/Public/RenderGraphBuilder.inl
--- a/Engine/Source/Runtime/RenderCore/Public/RenderGraphBuilder.inl
+++ b/Engine/Source/Runtime/RenderCore/Public/RenderGraphBuilder.inl
@@ -117,6 +117,7 @@ TRDGUniformBufferRef<ParameterStructType> FRDGBuilder::CreateUniformBuffer(const
 {
 	IF_RDG_ENABLE_DEBUG(UserValidation.ValidateCreateUniformBuffer(ParameterStruct, &ParameterStructType::StaticStructMetadata));
 	auto* UniformBuffer = UniformBuffers.Allocate<TRDGUniformBuffer<ParameterStructType>>(Allocator, ParameterStruct, ParameterStructType::StaticStructMetadata.GetShaderVariableName());
+	UniformBuffersToCreate.Add(false);
 	IF_RDG_ENABLE_DEBUG(UserValidation.ValidateCreateUniformBuffer(UniformBuffer));
 	return UniformBuffer;
 }
diff --git a/Engine/Source/Runtime/RenderCore/Public/RenderGraphParameter.h b/Engine/Source/Runtime/RenderCore/Public/RenderGraphParameter.h
--- a/Engine/Source/Runtime/RenderCore/Public/RenderGraphParameter.h
+++ b/Engine/Source/Runtime/RenderCore/Public/RenderGraphParameter.h
@@ -241,6 +241,9 @@ public:
 	/** Returns the number of texture parameters present on the layout. */
 	uint32 GetTextureParameterCount() const { return Layout->GraphTextures.Num(); }
 
+	/** Returns the number of RDG uniform buffers present in the layout. */
+	uint32 GetUniformBufferParameterCount() const { return Layout->GraphUniformBuffers.Num(); }
+
 	/** Returns the render target binding slots. Asserts if they don't exist. */
 	const FRenderTargetBindingSlots& GetRenderTargets() const
 	{
diff --git a/Engine/Source/Runtime/RenderCore/Public/RenderGraphPass.h b/Engine/Source/Runtime/RenderCore/Public/RenderGraphPass.h
--- a/Engine/Source/Runtime/RenderCore/Public/RenderGraphPass.h
+++ b/Engine/Source/Runtime/RenderCore/Public/RenderGraphPass.h
@@ -59,6 +59,11 @@ public:
 	void Submit(FRHIComputeCommandList& RHICmdList, ERHIPipeline Pipeline);
 	void Submit(FRHIComputeCommandList& RHICmdList, ERHIPipeline Pipeline, FRDGTransitionQueue& TransitionsToBegin);
 
+	void Reserve(uint32 TransitionCount)
+	{
+		Transitions.Reserve(TransitionCount);
+	}
+
 private:
 	const FRHITransition* Transition = nullptr;
 	TArray<FRHITransitionInfo, TInlineAllocator<1, FRDGArrayAllocator>> Transitions;
@@ -98,6 +103,11 @@ public:
 
 	void Submit(FRHIComputeCommandList& RHICmdList, ERHIPipeline Pipeline);
 
+	void Reserve(uint32 TransitionBatchCount)
+	{
+		Dependencies.Reserve(TransitionBatchCount);
+	}
+
 private:
 	TArray<FRDGBarrierBatchBegin*, TInlineAllocator<1, FRDGArrayAllocator>> Dependencies;
 
@@ -309,6 +319,9 @@ protected:
 			/** Whether the pass uses the immediate command list. */
 			uint32 bImmediateCommandList : 1;
 
+			/** Whether this pass has non-RDG UAV outputs. */
+			uint32 bHasExternalOutputs : 1;
+
 			/** Whether this pass allocated a texture through the pool. */
 			IF_RDG_ENABLE_DEBUG(uint32 bFirstTextureAllocated : 1);
 		};
@@ -332,26 +345,44 @@ protected:
 
 	struct FTextureState
 	{
-		FTextureState()
+		FTextureState() = default;
+
+		FTextureState(FRDGTextureRef InTexture)
+			: Texture(InTexture)
+			, TextureHandle(InTexture->GetHandle())
 		{
-			InitAsWholeResource(State);
+			const uint32 SubresourceCount = Texture->GetSubresourceCount();
+			State.SetNum(SubresourceCount);
+			MergeState.SetNum(SubresourceCount);
 		}
 
+		FRDGTextureRef Texture = nullptr;
 		FRDGTextureTransientSubresourceState State;
 		FRDGTextureTransientSubresourceStateIndirect MergeState;
+		FRDGTextureHandle TextureHandle;
 		uint16 ReferenceCount = 0;
 	};
 
 	struct FBufferState
 	{
+		FBufferState() = default;
+
+		FBufferState(FRDGBufferRef InBuffer)
+			: Buffer(InBuffer)
+			, BufferHandle(InBuffer->GetHandle())
+		{}
+
+		FRDGBufferRef Buffer = nullptr;
 		FRDGSubresourceState State;
 		FRDGSubresourceState* MergeState = nullptr;
+		FRDGBufferHandle BufferHandle;
 		uint16 ReferenceCount = 0;
 	};
 
 	/** Maps textures / buffers to information on how they are used in the pass. */
-	TSortedMap<FRDGTexture*, FTextureState, FRDGArrayAllocator> TextureStates;
-	TSortedMap<FRDGBuffer*, FBufferState, FRDGArrayAllocator> BufferStates;
+	TArray<FTextureState, FRDGArrayAllocator> TextureStates;
+	TArray<FBufferState, FRDGArrayAllocator> BufferStates;
+	TArray<FRDGViewHandle, FRDGArrayAllocator> Views;
 
 	/** Lists of pass parameters scheduled for begin during execution of this pass. */
 	TArray<FRDGPass*, TInlineAllocator<1, FRDGArrayAllocator>> ResourcesToBegin;
@@ -363,7 +394,7 @@ protected:
 	FRDGBarrierBatchBegin* EpilogueBarriersToBeginForGraphics = nullptr;
 	FRDGBarrierBatchBegin* EpilogueBarriersToBeginForAsyncCompute = nullptr;
 	FRDGBarrierBatchBegin* EpilogueBarriersToBeginForAll = nullptr;
-	TArray<FRDGBarrierBatchBegin*> SharedEpilogueBarriersToBegin;
+	TArray<FRDGBarrierBatchBegin*, FRDGArrayAllocator> SharedEpilogueBarriersToBegin;
 	FRDGBarrierBatchEnd* EpilogueBarriersToEnd = nullptr;
 
 	EAsyncComputeBudget AsyncComputeBudget = EAsyncComputeBudget::EAll_4;
diff --git a/Engine/Source/Runtime/RenderCore/Public/RenderGraphResources.h b/Engine/Source/Runtime/RenderCore/Public/RenderGraphResources.h
--- a/Engine/Source/Runtime/RenderCore/Public/RenderGraphResources.h
+++ b/Engine/Source/Runtime/RenderCore/Public/RenderGraphResources.h
@@ -273,13 +273,13 @@ protected:
 	/** If true, the resource has been used on an async compute pass and may have async compute states. */
 	uint8 bUsedByAsyncComputePass : 1;
 
-	/** Assigns this resource as a simple passthrough container for an RHI resource. */
-	void SetPassthroughRHI(FRHIResource* InResourceRHI);
-
 private:
 	/** Number of references in passes and deferred queries. */
 	uint16 ReferenceCount = 0;
 
+	/** Scratch index allocated for the resource in the pass being setup. */
+	uint16 PassStateIndex = 0;
+
 	/** The initial and final states of the resource assigned by the user, if known. */
 	ERHIAccess AccessInitial = ERHIAccess::Unknown;
 	ERHIAccess AccessFinal = ERHIAccess::Unknown;
@@ -327,6 +327,7 @@ protected:
 
 private:
 	FRDGViewHandle Handle;
+	FRDGPassHandle LastPass;
 
 	friend FRDGBuilder;
 	friend FRDGViewRegistry;
@@ -421,6 +422,11 @@ public:
 		return FRDGTextureSubresourceRange(Layout);
 	}
 
+	FORCEINLINE uint32 GetSubresourceCount() const
+	{
+		return Layout.GetSubresourceCount();
+	}
+
 	FRDGTextureSubresourceRange GetSubresourceRangeSRV() const;
 
 private:
@@ -431,8 +437,9 @@ private:
 		, RenderTargetTexture(InRenderTargetTexture)
 		, Layout(InDesc)
 	{
-		InitAsWholeResource(MergeState);
-		InitAsWholeResource(LastProducers);
+		const uint32 SubresourceCount = GetSubresourceCount();
+		MergeState.SetNum(SubresourceCount);
+		LastProducers.SetNum(SubresourceCount);
 	}
 
 	/** Assigns a pooled render target as the backing RHI resource. */
@@ -839,9 +846,10 @@ class RENDERCORE_API FRDGPooledBuffer final
 	: public FRefCountedObject
 {
 public:
-	FRDGPooledBuffer(TRefCountPtr<FRHIBuffer> InBuffer, const FRDGBufferDesc& InDesc)
+	FRDGPooledBuffer(TRefCountPtr<FRHIBuffer> InBuffer, const FRDGBufferDesc& InDesc, const TCHAR* InName)
 		: Desc(InDesc)
 		, Buffer(MoveTemp(InBuffer))
+		, Name(InName)
 	{}
 
 	const FRDGBufferDesc Desc;
