Fix compile error

diff --git a/Engine/Source/Runtime/RenderCore/Public/VirtualTexturing.h b/Engine/Source/Runtime/RenderCore/Public/VirtualTexturing.h
--- a/Engine/Source/Runtime/RenderCore/Public/VirtualTexturing.h
+++ b/Engine/Source/Runtime/RenderCore/Public/VirtualTexturing.h
@@ -384,7 +384,7 @@ public:
 
 protected:
 	friend class FVirtualTextureSystem;
-	virtual void Destroy(FVirtualTextureSystem* InSystem) = 0;
+	virtual void Destroy(class FVirtualTextureSystem* InSystem) = 0;
 	virtual ~IAllocatedVirtualTexture() {}
 
 	FAllocatedVTDescription Description;
