Fix for Desktop Forward on Metal devices
    - Disabled GPUScene

diff --git a/Engine/Source/Runtime/RHI/Public/RHIDefinitions.h b/Engine/Source/Runtime/RHI/Public/RHIDefinitions.h
--- a/Engine/Source/Runtime/RHI/Public/RHIDefinitions.h
+++ b/Engine/Source/Runtime/RHI/Public/RHIDefinitions.h
@@ -1780,6 +1780,11 @@ inline bool IsVulkanMobileSM5Platform(const EShaderPlatform Platform)
 	return Platform == SP_VULKAN_SM5_ANDROID;
 }
 
+inline bool IsMetalMobileSM5Platform(const EShaderPlatform Platform)
+{
+    return Platform == SP_METAL_MRT;
+}
+
 inline bool IsAndroidOpenGLESPlatform(const FStaticShaderPlatform Platform)
 {
 	return Platform == SP_OPENGL_ES3_1_ANDROID || FDataDrivenShaderPlatformInfo::GetIsAndroidOpenGLES(Platform);
diff --git a/Engine/Source/Runtime/RenderCore/Public/RenderUtils.h b/Engine/Source/Runtime/RenderCore/Public/RenderUtils.h
--- a/Engine/Source/Runtime/RenderCore/Public/RenderUtils.h
+++ b/Engine/Source/Runtime/RenderCore/Public/RenderUtils.h
@@ -624,6 +624,7 @@ inline bool UseGPUScene(const FStaticShaderPlatform Platform, const FStaticFeatu
 		&& !IsOpenGLPlatform(Platform)
 		&& !IsSwitchPlatform(Platform)
 		&& !IsVulkanMobileSM5Platform(Platform)
+        && !IsMetalMobileSM5Platform(Platform)
 		// we only check DDSPI for platforms that have been read in - IsValid() can go away once ALL platforms are converted over to this system
 		&& (!FDataDrivenShaderPlatformInfo::IsValid(Platform) || FDataDrivenShaderPlatformInfo::GetSupportsGPUScene(Platform));
 }
