Reworked RDG pooled buffer size alignment so that it no longer affects the external-facing descriptor.

diff --git a/Engine/Source/Runtime/RenderCore/Private/RenderGraphBuilder.cpp b/Engine/Source/Runtime/RenderCore/Private/RenderGraphBuilder.cpp
--- a/Engine/Source/Runtime/RenderCore/Private/RenderGraphBuilder.cpp
+++ b/Engine/Source/Runtime/RenderCore/Private/RenderGraphBuilder.cpp
@@ -2517,11 +2517,6 @@ void FRDGBuilder::BeginResourceRHI(FTransientResourceAllocator* TransientResourc
 	check(TransientResourceAllocator || Buffer->bExternal);
 	check(Buffer->ReferenceCount > 0 || Buffer->bExternal || IsImmediateMode());
 
-	const uint64 BufferPageSize = 64 * 1024;
-
-	FRDGBufferDesc SnappedBufferDesc = Buffer->Desc;
-	SnappedBufferDesc.NumElements = Align(SnappedBufferDesc.BytesPerElement * SnappedBufferDesc.NumElements, BufferPageSize) / SnappedBufferDesc.BytesPerElement;
-
 	// If transient then create the resource on the transient allocator. External or extracted resource can't be transient because of lifetime tracking issues.
 	if (IsTransient(Buffer))
 	{
@@ -2529,7 +2524,7 @@ void FRDGBuilder::BeginResourceRHI(FTransientResourceAllocator* TransientResourc
 
 		if (IRHITransientResourceAllocator* AllocatorRHI = TransientResourceAllocator->GetOrCreate())
 		{
-			if (FRHITransientBuffer* TransientBuffer = AllocatorRHI->CreateBuffer(Translate(SnappedBufferDesc), Buffer->Name))
+			if (FRHITransientBuffer* TransientBuffer = AllocatorRHI->CreateBuffer(Translate(Buffer->Desc), Buffer->Name))
 			{
 				Buffer->SetRHI(TransientBuffer, Allocator);
 
@@ -2553,7 +2548,7 @@ void FRDGBuilder::BeginResourceRHI(FTransientResourceAllocator* TransientResourc
 
 	if (!Buffer->bTransient)
 	{
-		Buffer->SetRHI(GRenderGraphResourcePool.FindFreeBufferInternal(RHICmdList, SnappedBufferDesc, Buffer->Name));
+		Buffer->SetRHI(GRenderGraphResourcePool.FindFreeBufferInternal(RHICmdList, Buffer->Desc, Buffer->Name));
 	}
 
 	Buffer->FirstPass = PassHandle;
diff --git a/Engine/Source/Runtime/RenderCore/Private/RenderGraphResourcePool.cpp b/Engine/Source/Runtime/RenderCore/Private/RenderGraphResourcePool.cpp
--- a/Engine/Source/Runtime/RenderCore/Private/RenderGraphResourcePool.cpp
+++ b/Engine/Source/Runtime/RenderCore/Private/RenderGraphResourcePool.cpp
@@ -31,7 +31,11 @@ TRefCountPtr<FRDGPooledBuffer> FRenderGraphResourcePool::FindFreeBufferInternal(
 	const FRDGBufferDesc& Desc,
 	const TCHAR* InDebugName)
 {
-	const uint64 BufferHash = ComputeHash(Desc);
+	const uint64 BufferPageSize = 64 * 1024;
+
+	FRDGBufferDesc AlignedDesc = Desc;
+	AlignedDesc.NumElements = Align(AlignedDesc.BytesPerElement * AlignedDesc.NumElements, BufferPageSize) / AlignedDesc.BytesPerElement;
+	const uint64 BufferHash = ComputeHash(AlignedDesc);
 
 	// First find if available.
 	for (int32 Index = 0; Index < AllocatedBufferHashes.Num(); ++Index)
@@ -49,12 +53,15 @@ TRefCountPtr<FRDGPooledBuffer> FRenderGraphResourcePool::FindFreeBufferInternal(
 			continue;
 		}
 
-		check(PooledBuffer->Desc == Desc);
+		check(PooledBuffer->GetAlignedDesc() == AlignedDesc);
 
 		PooledBuffer->LastUsedFrame = FrameCounter;
 		PooledBuffer->ViewCache.SetDebugName(InDebugName);
 		PooledBuffer->Name = InDebugName;
 
+		// We need the external-facing desc to match what the user requested.
+		const_cast<FRDGBufferDesc&>(PooledBuffer->Desc).NumElements = Desc.NumElements;
+
 		return PooledBuffer;
 	}
 
@@ -62,7 +69,7 @@ TRefCountPtr<FRDGPooledBuffer> FRenderGraphResourcePool::FindFreeBufferInternal(
 	{
 		TRACE_CPUPROFILER_EVENT_SCOPE(FRenderGraphResourcePool::CreateBuffer);
 
-		uint32 NumBytes = Desc.GetTotalNumBytes();
+		const uint32 NumBytes = AlignedDesc.GetTotalNumBytes();
 
 		FRHIResourceCreateInfo CreateInfo(InDebugName);
 		TRefCountPtr<FRHIBuffer> BufferRHI;
@@ -84,7 +91,7 @@ TRefCountPtr<FRDGPooledBuffer> FRenderGraphResourcePool::FindFreeBufferInternal(
 			check(0);
 		}
 
-		TRefCountPtr<FRDGPooledBuffer> PooledBuffer = new FRDGPooledBuffer(MoveTemp(BufferRHI), Desc, InDebugName);
+		TRefCountPtr<FRDGPooledBuffer> PooledBuffer = new FRDGPooledBuffer(MoveTemp(BufferRHI), Desc, AlignedDesc.NumElements, InDebugName);
 		AllocatedBuffers.Add(PooledBuffer);
 		AllocatedBufferHashes.Add(BufferHash);
 		check(PooledBuffer->GetRefCount() == 2);
diff --git a/Engine/Source/Runtime/RenderCore/Public/RenderGraphResources.h b/Engine/Source/Runtime/RenderCore/Public/RenderGraphResources.h
--- a/Engine/Source/Runtime/RenderCore/Public/RenderGraphResources.h
+++ b/Engine/Source/Runtime/RenderCore/Public/RenderGraphResources.h
@@ -846,10 +846,11 @@ class RENDERCORE_API FRDGPooledBuffer final
 	: public FRefCountedObject
 {
 public:
-	FRDGPooledBuffer(TRefCountPtr<FRHIBuffer> InBuffer, const FRDGBufferDesc& InDesc, const TCHAR* InName)
+	FRDGPooledBuffer(TRefCountPtr<FRHIBuffer> InBuffer, const FRDGBufferDesc& InDesc, uint32 InNumAllocatedElements, const TCHAR* InName)
 		: Desc(InDesc)
 		, Buffer(MoveTemp(InBuffer))
 		, Name(InName)
+		, NumAllocatedElements(InNumAllocatedElements)
 	{}
 
 	const FRDGBufferDesc Desc;
@@ -885,10 +886,18 @@ private:
 		State.Finalize();
 	}
 
+	FRDGBufferDesc GetAlignedDesc() const
+	{
+		FRDGBufferDesc AlignedDesc = Desc;
+		AlignedDesc.NumElements = NumAllocatedElements;
+		return AlignedDesc;
+	}
+
 	const TCHAR* Name = nullptr;
 	FRDGBufferRef Owner = nullptr;
 	FRDGSubresourceState State;
 
+	const uint32 NumAllocatedElements;
 	uint32 LastUsedFrame = 0;
 
 	friend FRenderGraphResourcePool;
