Fix memory leaks.  The memory leak check was hit while running a light bake in a QA test level.  This lead to the discovery of several other memory leaks that needed to be addressed

diff --git a/Engine/Plugins/Runtime/HairStrands/Source/HairStrandsCore/Private/GroomAsset.cpp b/Engine/Plugins/Runtime/HairStrands/Source/HairStrandsCore/Private/GroomAsset.cpp
--- a/Engine/Plugins/Runtime/HairStrands/Source/HairStrandsCore/Private/GroomAsset.cpp
+++ b/Engine/Plugins/Runtime/HairStrands/Source/HairStrandsCore/Private/GroomAsset.cpp
@@ -2787,6 +2787,7 @@ void UGroomAsset::CreateDebugData()
 			ENQUEUE_RENDER_COMMAND(HairStrandsDebugResourceCommand)(
 				[InData, InResource](FRHICommandListImmediate& RHICmdList)
 				{
+					FMemMark Mark(FMemStack::Get());
 					FRDGBuilder GraphBuilder(RHICmdList);
 					CreateHairStrandsDebugResources(GraphBuilder, InData, InResource);
 					GraphBuilder.Execute();
diff --git a/Engine/Source/Editor/UnrealEd/Private/Lightmass/LightmassLandscapeRender.cpp b/Engine/Source/Editor/UnrealEd/Private/Lightmass/LightmassLandscapeRender.cpp
--- a/Engine/Source/Editor/UnrealEd/Private/Lightmass/LightmassLandscapeRender.cpp
+++ b/Engine/Source/Editor/UnrealEd/Private/Lightmass/LightmassLandscapeRender.cpp
@@ -113,6 +113,8 @@ void RenderLandscapeMaterialForLightmass(const FLandscapeStaticLightingMesh* Lan
 	ENQUEUE_RENDER_COMMAND(CanvasFlushSetupCommand)(
 		[RenderTarget, &DynamicMeshBuilder, ViewInitOptions, MaterialProxy](FRHICommandListImmediate& RHICmdList)
 		{
+			FMemMark Mark(FMemStack::Get());
+
 			FMeshBatch Mesh;
 			FMeshBuilderOneFrameResources OneFrameResource;
 			DynamicMeshBuilder.GetMeshElement(FMatrix::Identity, MaterialProxy, SDPG_Foreground, true, false, 0, OneFrameResource, Mesh);
diff --git a/Engine/Source/Runtime/Engine/Private/DistanceFieldAtlas.cpp b/Engine/Source/Runtime/Engine/Private/DistanceFieldAtlas.cpp
--- a/Engine/Source/Runtime/Engine/Private/DistanceFieldAtlas.cpp
+++ b/Engine/Source/Runtime/Engine/Private/DistanceFieldAtlas.cpp
@@ -1152,6 +1152,7 @@ void FLandscapeTextureAtlas::UpdateAllocations(FRDGBuilder& GraphBuilder, ERHIFe
 
 void FLandscapeTextureAtlas::UpdateAllocations(FRHICommandListImmediate& RHICmdList, ERHIFeatureLevel::Type InFeatureLevel)
 {
+	FMemMark Mark(FMemStack::Get());
 	FRDGBuilder GraphBuilder(RHICmdList);
 	UpdateAllocations(GraphBuilder, InFeatureLevel);
 	GraphBuilder.Execute();
diff --git a/Engine/Source/Runtime/RenderCore/Private/GenerateMips.cpp b/Engine/Source/Runtime/RenderCore/Private/GenerateMips.cpp
--- a/Engine/Source/Runtime/RenderCore/Private/GenerateMips.cpp
+++ b/Engine/Source/Runtime/RenderCore/Private/GenerateMips.cpp
@@ -215,6 +215,7 @@ void FGenerateMips::Execute(FRHICommandListImmediate& RHICmdList, FRHITexture* T
 {
 	TRefCountPtr<IPooledRenderTarget> PooledRenderTarget = CreateRenderTarget(Texture, TEXT("MipGeneration"));
 
+	FMemMark MemMark(FMemStack::Get());
 	FRDGBuilder GraphBuilder(RHICmdList);
 	Execute(GraphBuilder, GraphBuilder.RegisterExternalTexture(PooledRenderTarget), Params, bAllowRenderBasedGeneration ? EGenerateMipsPass::Raster : EGenerateMipsPass::Compute);
 	GraphBuilder.Execute();
@@ -224,6 +225,7 @@ void FGenerateMips::Execute(FRHICommandListImmediate& RHICmdList, FRHITexture* T
 {
 	TRefCountPtr<IPooledRenderTarget> PooledRenderTarget = CreateRenderTarget(Texture, TEXT("MipGeneration"));
 
+	FMemMark MemMark(FMemStack::Get());
 	FRDGBuilder GraphBuilder(RHICmdList);
 	Execute(GraphBuilder, GraphBuilder.RegisterExternalTexture(PooledRenderTarget), Params, bAllowRenderBasedGeneration ? EGenerateMipsPass::Raster : EGenerateMipsPass::Compute);
 	GraphBuilder.Execute();
diff --git a/Engine/Source/Runtime/RenderCore/Private/VisualizeTexture.cpp b/Engine/Source/Runtime/RenderCore/Private/VisualizeTexture.cpp
--- a/Engine/Source/Runtime/RenderCore/Private/VisualizeTexture.cpp
+++ b/Engine/Source/Runtime/RenderCore/Private/VisualizeTexture.cpp
@@ -878,6 +878,7 @@ void FVisualizeTexture::SetCheckPoint(FRDGBuilder& GraphBuilder, IPooledRenderTa
 
 void FVisualizeTexture::SetCheckPoint(FRHICommandListImmediate& RHICmdList, IPooledRenderTarget* PooledRenderTarget)
 {
+	FMemMark MemMark(FMemStack::Get());
 	FRDGBuilder GraphBuilder(RHICmdList);
 	SetCheckPoint(GraphBuilder, PooledRenderTarget);
 	GraphBuilder.Execute();
diff --git a/Engine/Source/Runtime/Renderer/Private/RendererUtils.cpp b/Engine/Source/Runtime/Renderer/Private/RendererUtils.cpp
--- a/Engine/Source/Runtime/Renderer/Private/RendererUtils.cpp
+++ b/Engine/Source/Runtime/Renderer/Private/RendererUtils.cpp
@@ -76,6 +76,7 @@ void FRenderTargetWriteMask::Decode(
 	ETextureCreateFlags RTWriteMaskFastVRamConfig,
 	const TCHAR* RTWriteMaskDebugName)
 {
+	FMemMark MemMark(FMemStack::Get());
 	FRDGBuilder GraphBuilder(RHICmdList);
 
 	TArray<FRDGTextureRef, SceneRenderingAllocator> InputTextures;
diff --git a/Engine/Source/Runtime/Renderer/Private/SceneCaptureRendering.cpp b/Engine/Source/Runtime/Renderer/Private/SceneCaptureRendering.cpp
--- a/Engine/Source/Runtime/Renderer/Private/SceneCaptureRendering.cpp
+++ b/Engine/Source/Runtime/Renderer/Private/SceneCaptureRendering.cpp
@@ -924,6 +924,7 @@ void FScene::UpdateSceneCaptureContents(USceneCaptureComponentCube* CaptureCompo
 		{
 			const ERHIAccess FinalAccess = ERHIAccess::RTV;
 
+			FMemMark MemMark(FMemStack::Get());
 			FRDGBuilder GraphBuilder(RHICmdList);
 			FRDGTextureRef OutputTexture = GraphBuilder.RegisterExternalTexture(CreateRenderTarget(RenderTarget->GetRenderTargetTexture(), TEXT("Output")));
 			FRDGTextureRef LeftEyeTexture = GraphBuilder.RegisterExternalTexture(CreateRenderTarget(LeftEye->TextureRHI, TEXT("LeftEye")));
