U5 - Fix D3d11 validation error when clearing large buffers that would run into a ComputeShader Threadgroup limitation.

diff --git a/Engine/Shaders/Private/ClearReplacementShaders.usf b/Engine/Shaders/Private/ClearReplacementShaders.usf
--- a/Engine/Shaders/Private/ClearReplacementShaders.usf
+++ b/Engine/Shaders/Private/ClearReplacementShaders.usf
@@ -55,6 +55,12 @@
 		#define ADDR(p) (p).x
 		#define COMP(a,b) ((a).x < (b).x)
 
+	#elif RESOURCE_TYPE == 5
+	
+		RWBuffer        <VALUE_TYPE> ClearResource;
+		#define ADDR(p) (p).x
+		#define COMP(a,b) ((a).x < (b).x)
+
 	#endif
 
 	[numthreads(THREADGROUPSIZE_X, THREADGROUPSIZE_Y, THREADGROUPSIZE_Z)]
diff --git a/Engine/Source/Runtime/RenderCore/Private/ClearReplacementShaders.cpp b/Engine/Source/Runtime/RenderCore/Private/ClearReplacementShaders.cpp
--- a/Engine/Source/Runtime/RenderCore/Private/ClearReplacementShaders.cpp
+++ b/Engine/Source/Runtime/RenderCore/Private/ClearReplacementShaders.cpp
@@ -16,6 +16,11 @@ IMPLEMENT_SHADER_TYPE4_WITH_TEMPLATE_PREFIX(template<>, RENDERCORE_API, FClearRe
 IMPLEMENT_SHADER_TYPE4_WITH_TEMPLATE_PREFIX(template<>, RENDERCORE_API, FClearReplacementCS_Buffer_Uint4					, SF_Compute);
 IMPLEMENT_SHADER_TYPE4_WITH_TEMPLATE_PREFIX(template<>, RENDERCORE_API, FClearReplacementCS_Buffer_Float_Bounds				, SF_Compute);
 IMPLEMENT_SHADER_TYPE4_WITH_TEMPLATE_PREFIX(template<>, RENDERCORE_API, FClearReplacementCS_Buffer_Float4_Bounds			, SF_Compute);
+IMPLEMENT_SHADER_TYPE4_WITH_TEMPLATE_PREFIX(template<>, RENDERCORE_API, FClearReplacementCS_LargeBuffer_Uint_Bounds			, SF_Compute);
+IMPLEMENT_SHADER_TYPE4_WITH_TEMPLATE_PREFIX(template<>, RENDERCORE_API, FClearReplacementCS_LargeBuffer_Uint4_Bounds		, SF_Compute);
+IMPLEMENT_SHADER_TYPE4_WITH_TEMPLATE_PREFIX(template<>, RENDERCORE_API, FClearReplacementCS_LargeBuffer_Uint4				, SF_Compute);
+IMPLEMENT_SHADER_TYPE4_WITH_TEMPLATE_PREFIX(template<>, RENDERCORE_API, FClearReplacementCS_LargeBuffer_Float_Bounds		, SF_Compute);
+IMPLEMENT_SHADER_TYPE4_WITH_TEMPLATE_PREFIX(template<>, RENDERCORE_API, FClearReplacementCS_LargeBuffer_Float4_Bounds		, SF_Compute);
 IMPLEMENT_SHADER_TYPE4_WITH_TEMPLATE_PREFIX(template<>, RENDERCORE_API, FClearReplacementCS_StructuredBuffer_Uint_Bounds	, SF_Compute);
 IMPLEMENT_SHADER_TYPE4_WITH_TEMPLATE_PREFIX(template<>, RENDERCORE_API, FClearReplacementCS_StructuredBuffer_Uint4_Bounds	, SF_Compute);
 IMPLEMENT_SHADER_TYPE4_WITH_TEMPLATE_PREFIX(template<>, RENDERCORE_API, FClearReplacementCS_StructuredBuffer_Uint4			, SF_Compute);
@@ -28,6 +33,7 @@ IMPLEMENT_SHADER_TYPE4_WITH_TEMPLATE_PREFIX(template<>, RENDERCORE_API, FClearRe
 IMPLEMENT_SHADER_TYPE4_WITH_TEMPLATE_PREFIX(template<>, RENDERCORE_API, FClearReplacementCS_StructuredBuffer_Uint_Zero		, SF_Compute);
 IMPLEMENT_SHADER_TYPE4_WITH_TEMPLATE_PREFIX(template<>, RENDERCORE_API, FClearReplacementCS_Texture2DArray_Uint_Zero		, SF_Compute);
 IMPLEMENT_SHADER_TYPE4_WITH_TEMPLATE_PREFIX(template<>, RENDERCORE_API, FClearReplacementCS_Buffer_Uint						, SF_Compute);
+IMPLEMENT_SHADER_TYPE4_WITH_TEMPLATE_PREFIX(template<>, RENDERCORE_API, FClearReplacementCS_LargeBuffer_Uint				, SF_Compute);
 IMPLEMENT_SHADER_TYPE4_WITH_TEMPLATE_PREFIX(template<>, RENDERCORE_API, FClearReplacementCS_StructuredBuffer_Uint			, SF_Compute);
 IMPLEMENT_SHADER_TYPE4_WITH_TEMPLATE_PREFIX(template<>, RENDERCORE_API, FClearReplacementCS_Texture2DArray_Uint				, SF_Compute);
 IMPLEMENT_SHADER_TYPE4_WITH_TEMPLATE_PREFIX(template<>, RENDERCORE_API, FClearReplacementCS_Texture3D_Float4				, SF_Compute);
@@ -40,6 +46,7 @@ IMPLEMENT_SHADER_TYPE4_WITH_TEMPLATE_PREFIX(template<>, RENDERCORE_API, FClearRe
 IMPLEMENT_SHADER_TYPE4_WITH_TEMPLATE_PREFIX(template<>, RENDERCORE_API, FClearReplacementCS_Texture2D_Uint4_Bounds			, SF_Compute);
 IMPLEMENT_SHADER_TYPE4_WITH_TEMPLATE_PREFIX(template<>, RENDERCORE_API, FClearReplacementCS_Texture2DArray_Uint4_Bounds		, SF_Compute);
 IMPLEMENT_SHADER_TYPE4_WITH_TEMPLATE_PREFIX(template<>, RENDERCORE_API, FClearReplacementCS_Buffer_Sint4_Bounds				, SF_Compute);
+IMPLEMENT_SHADER_TYPE4_WITH_TEMPLATE_PREFIX(template<>, RENDERCORE_API, FClearReplacementCS_LargeBuffer_Sint4_Bounds		, SF_Compute);
 IMPLEMENT_SHADER_TYPE4_WITH_TEMPLATE_PREFIX(template<>, RENDERCORE_API, FClearReplacementCS_StructuredBuffer_Sint4_Bounds	, SF_Compute);
 IMPLEMENT_SHADER_TYPE4_WITH_TEMPLATE_PREFIX(template<>, RENDERCORE_API, FClearReplacementCS_Texture3D_Sint4_Bounds			, SF_Compute);
 IMPLEMENT_SHADER_TYPE4_WITH_TEMPLATE_PREFIX(template<>, RENDERCORE_API, FClearReplacementCS_Texture2D_Sint4_Bounds			, SF_Compute);
diff --git a/Engine/Source/Runtime/RenderCore/Public/ClearReplacementShaders.h b/Engine/Source/Runtime/RenderCore/Public/ClearReplacementShaders.h
--- a/Engine/Source/Runtime/RenderCore/Public/ClearReplacementShaders.h
+++ b/Engine/Source/Runtime/RenderCore/Public/ClearReplacementShaders.h
@@ -15,7 +15,8 @@ enum class EClearReplacementResourceType
 	Texture2D = 1,
 	Texture2DArray = 2,
 	Texture3D = 3,
-	StructuredBuffer = 4
+	StructuredBuffer = 4,
+	LargeBuffer = 5
 };
 
 enum class EClearReplacementValueType
@@ -110,11 +111,12 @@ private:
 namespace ClearReplacementCS
 {
 	template <EClearReplacementResourceType> struct TThreadGroupSize {};
-	template <> struct TThreadGroupSize<EClearReplacementResourceType::Buffer>				{ static constexpr int32 X = 64, Y = 1, Z = 1; };
-	template <> struct TThreadGroupSize<EClearReplacementResourceType::Texture2D>			{ static constexpr int32 X =  8, Y = 8, Z = 1; };
-	template <> struct TThreadGroupSize<EClearReplacementResourceType::Texture2DArray>		{ static constexpr int32 X =  8, Y = 8, Z = 1; };
-	template <> struct TThreadGroupSize<EClearReplacementResourceType::Texture3D>			{ static constexpr int32 X =  4, Y = 4, Z = 4; };
-	template <> struct TThreadGroupSize<EClearReplacementResourceType::StructuredBuffer>	{ static constexpr int32 X = 64, Y = 1, Z = 1; };
+	template <> struct TThreadGroupSize<EClearReplacementResourceType::Buffer>				{ static constexpr int32 X =  64, Y = 1, Z = 1; };
+	template <> struct TThreadGroupSize<EClearReplacementResourceType::Texture2D>			{ static constexpr int32 X =   8, Y = 8, Z = 1; };
+	template <> struct TThreadGroupSize<EClearReplacementResourceType::Texture2DArray>		{ static constexpr int32 X =   8, Y = 8, Z = 1; };
+	template <> struct TThreadGroupSize<EClearReplacementResourceType::Texture3D>			{ static constexpr int32 X =   4, Y = 4, Z = 4; };
+	template <> struct TThreadGroupSize<EClearReplacementResourceType::StructuredBuffer>	{ static constexpr int32 X =  64, Y = 1, Z = 1; };
+	template <> struct TThreadGroupSize<EClearReplacementResourceType::LargeBuffer>			{ static constexpr int32 X = 512, Y = 1, Z = 1; };
 }
 
 template <EClearReplacementResourceType ResourceType, typename BaseType>
@@ -260,6 +262,9 @@ typedef TClearReplacementPS<false, FClearReplacementBase_Float4_Zero>        FCl
 typedef TClearReplacementCS<EClearReplacementResourceType::Buffer,           FClearReplacementBase_Uint_Bounds>   FClearReplacementCS_Buffer_Uint_Bounds;
 typedef TClearReplacementCS<EClearReplacementResourceType::Buffer,           FClearReplacementBase_Float_Bounds>  FClearReplacementCS_Buffer_Float_Bounds;
 typedef TClearReplacementCS<EClearReplacementResourceType::Buffer,           FClearReplacementBase_Float4_Bounds> FClearReplacementCS_Buffer_Float4_Bounds;
+typedef TClearReplacementCS<EClearReplacementResourceType::LargeBuffer,      FClearReplacementBase_Uint_Bounds>   FClearReplacementCS_LargeBuffer_Uint_Bounds;
+typedef TClearReplacementCS<EClearReplacementResourceType::LargeBuffer,      FClearReplacementBase_Float_Bounds>  FClearReplacementCS_LargeBuffer_Float_Bounds;
+typedef TClearReplacementCS<EClearReplacementResourceType::LargeBuffer,      FClearReplacementBase_Float4_Bounds> FClearReplacementCS_LargeBuffer_Float4_Bounds;
 typedef TClearReplacementCS<EClearReplacementResourceType::StructuredBuffer, FClearReplacementBase_Uint_Bounds>   FClearReplacementCS_StructuredBuffer_Uint_Bounds;
 typedef TClearReplacementCS<EClearReplacementResourceType::StructuredBuffer, FClearReplacementBase_Float_Bounds>  FClearReplacementCS_StructuredBuffer_Float_Bounds;
 typedef TClearReplacementCS<EClearReplacementResourceType::StructuredBuffer, FClearReplacementBase_Float4_Bounds> FClearReplacementCS_StructuredBuffer_Float4_Bounds;
@@ -269,9 +274,11 @@ typedef TClearReplacementCS<EClearReplacementResourceType::Texture2DArray,	 FCle
 
 // Compute shaders for clearing each resource type. No bounds checks enabled.
 typedef TClearReplacementCS<EClearReplacementResourceType::Buffer,           FClearReplacementBase_Uint_Zero>     FClearReplacementCS_Buffer_Uint_Zero;
+typedef TClearReplacementCS<EClearReplacementResourceType::LargeBuffer,      FClearReplacementBase_Uint_Zero>     FClearReplacementCS_LargeBuffer_Uint_Zero;
 typedef TClearReplacementCS<EClearReplacementResourceType::StructuredBuffer, FClearReplacementBase_Uint_Zero>     FClearReplacementCS_StructuredBuffer_Uint_Zero;
 typedef TClearReplacementCS<EClearReplacementResourceType::Texture2DArray,   FClearReplacementBase_Uint_Zero>     FClearReplacementCS_Texture2DArray_Uint_Zero;
 typedef TClearReplacementCS<EClearReplacementResourceType::Buffer,           FClearReplacementBase_Uint>          FClearReplacementCS_Buffer_Uint;
+typedef TClearReplacementCS<EClearReplacementResourceType::LargeBuffer,      FClearReplacementBase_Uint>          FClearReplacementCS_LargeBuffer_Uint;
 typedef TClearReplacementCS<EClearReplacementResourceType::StructuredBuffer, FClearReplacementBase_Uint>          FClearReplacementCS_StructuredBuffer_Uint;
 typedef TClearReplacementCS<EClearReplacementResourceType::Texture2DArray,   FClearReplacementBase_Uint>          FClearReplacementCS_Texture2DArray_Uint;
 	
@@ -284,15 +291,18 @@ typedef TClearReplacementCS<EClearReplacementResourceType::Texture3D,        FCl
 typedef TClearReplacementCS<EClearReplacementResourceType::Texture2D,        FClearReplacementBase_Uint4>         FClearReplacementCS_Texture2D_Uint4;
 typedef TClearReplacementCS<EClearReplacementResourceType::Texture2DArray,   FClearReplacementBase_Uint4>         FClearReplacementCS_Texture2DArray_Uint4;
 typedef TClearReplacementCS<EClearReplacementResourceType::Buffer,			 FClearReplacementBase_Uint4>         FClearReplacementCS_Buffer_Uint4;
+typedef TClearReplacementCS<EClearReplacementResourceType::LargeBuffer,		 FClearReplacementBase_Uint4>         FClearReplacementCS_LargeBuffer_Uint4;
 typedef TClearReplacementCS<EClearReplacementResourceType::StructuredBuffer, FClearReplacementBase_Uint4>         FClearReplacementCS_StructuredBuffer_Uint4;
 
 typedef TClearReplacementCS<EClearReplacementResourceType::Buffer,		     FClearReplacementBase_Uint4_Bounds>  FClearReplacementCS_Buffer_Uint4_Bounds;
+typedef TClearReplacementCS<EClearReplacementResourceType::LargeBuffer,		 FClearReplacementBase_Uint4_Bounds>  FClearReplacementCS_LargeBuffer_Uint4_Bounds;
 typedef TClearReplacementCS<EClearReplacementResourceType::StructuredBuffer, FClearReplacementBase_Uint4_Bounds>  FClearReplacementCS_StructuredBuffer_Uint4_Bounds;
 typedef TClearReplacementCS<EClearReplacementResourceType::Texture3D,        FClearReplacementBase_Uint4_Bounds>  FClearReplacementCS_Texture3D_Uint4_Bounds;
 typedef TClearReplacementCS<EClearReplacementResourceType::Texture2D,        FClearReplacementBase_Uint4_Bounds>  FClearReplacementCS_Texture2D_Uint4_Bounds;
 typedef TClearReplacementCS<EClearReplacementResourceType::Texture2DArray,   FClearReplacementBase_Uint4_Bounds>  FClearReplacementCS_Texture2DArray_Uint4_Bounds;
 
 typedef TClearReplacementCS<EClearReplacementResourceType::Buffer,           FClearReplacementBase_Sint4_Bounds>  FClearReplacementCS_Buffer_Sint4_Bounds;
+typedef TClearReplacementCS<EClearReplacementResourceType::LargeBuffer,      FClearReplacementBase_Sint4_Bounds>  FClearReplacementCS_LargeBuffer_Sint4_Bounds;
 typedef TClearReplacementCS<EClearReplacementResourceType::StructuredBuffer, FClearReplacementBase_Sint4_Bounds>  FClearReplacementCS_StructuredBuffer_Sint4_Bounds;
 typedef TClearReplacementCS<EClearReplacementResourceType::Texture3D,        FClearReplacementBase_Sint4_Bounds>  FClearReplacementCS_Texture3D_Sint4_Bounds;
 typedef TClearReplacementCS<EClearReplacementResourceType::Texture2D,        FClearReplacementBase_Sint4_Bounds>  FClearReplacementCS_Texture2D_Sint4_Bounds;
diff --git a/Engine/Source/Runtime/Windows/D3D11RHI/Private/D3D11UAV.cpp b/Engine/Source/Runtime/Windows/D3D11RHI/Private/D3D11UAV.cpp
--- a/Engine/Source/Runtime/Windows/D3D11RHI/Private/D3D11UAV.cpp
+++ b/Engine/Source/Runtime/Windows/D3D11RHI/Private/D3D11UAV.cpp
@@ -416,7 +416,15 @@ void FD3D11DynamicRHI::ClearUAV(TRHICommandList_RecursiveHazardous<FD3D11Dynamic
 		}
 		else
 		{
-			ClearUAVShader_T<EClearReplacementResourceType::Buffer, 4, false>(RHICmdList, UnorderedAccessView, UAVDesc.Buffer.NumElements, 1, 1, ClearValues, ValueType);
+			if (UAVDesc.Buffer.NumElements < 65536 * 64)
+			{
+				ClearUAVShader_T<EClearReplacementResourceType::Buffer, 4, false>(RHICmdList, UnorderedAccessView, UAVDesc.Buffer.NumElements, 1, 1, ClearValues, ValueType);
+			}
+			else
+			{
+				ClearUAVShader_T<EClearReplacementResourceType::LargeBuffer, 4, false>(RHICmdList, UnorderedAccessView, UAVDesc.Buffer.NumElements, 1, 1, ClearValues, ValueType);
+			}
+			
 		}
 	}
 	else
