Fixed crash when drawing landscape BP brushes :
* Added game thread-only version of draw events in order to avoid the mistake of pushing a draw event on the main thread using an immediate command list, and popping it on the main thread later on, while the original command list may have been destroyed already.
* FDrawEvent is now able to be started/stopped on the game thread (no command list provided : a render command using the current immediate command list will be issued on both start and end) or the rendering/RHI threads (same as now : the command list passed in the start function will be the one used in the end function)
* Fixed potential similar issue with Begin/EndDrawCanvasToRenderTarget and DrawMaterialToRenderTarget BP functions
* Fixed potential similar issue with FGPUSkinCache end of frame updates (+ fixed FGPUSkinCache::BeginBatchDispatch being ifdefed-out by #if WANTS_DRAW_MESH_EVENTS)

#tests editor, PF

diff --git a/Engine/Source/Runtime/Engine/Classes/Kismet/KismetRenderingLibrary.h b/Engine/Source/Runtime/Engine/Classes/Kismet/KismetRenderingLibrary.h
--- a/Engine/Source/Runtime/Engine/Classes/Kismet/KismetRenderingLibrary.h
+++ b/Engine/Source/Runtime/Engine/Classes/Kismet/KismetRenderingLibrary.h
@@ -8,6 +8,7 @@
 #include "Kismet/BlueprintFunctionLibrary.h"
 #include "Components/SkinnedMeshComponent.h"
 #include "Engine/TextureRenderTarget2D.h"
+#include "ProfilingDebugging/RealtimeGPUProfiler.h"
 #include "Camera/CameraTypes.h"
 #include "KismetRenderingLibrary.generated.h"
 
@@ -21,15 +22,12 @@ struct FDrawToRenderTargetContext
 {
 	GENERATED_USTRUCT_BODY()
 
-	FDrawToRenderTargetContext() :
-		RenderTarget(NULL),
-		DrawEvent(NULL)
-	{}
-
 	UPROPERTY()
 	TObjectPtr<UTextureRenderTarget2D> RenderTarget;
 
-	FDrawEvent* DrawEvent;
+#if WANTS_DRAW_MESH_EVENTS
+	FDrawEvent* DrawEvent = nullptr;
+#endif // WANTS_DRAW_MESH_EVENTS
 };
 
 UCLASS(MinimalAPI, meta=(ScriptName="RenderingLibrary"))
@@ -121,7 +119,7 @@ class UKismetRenderingLibrary : public UBlueprintFunctionLibrary
 
 	/**
     * Incredibly inefficient and slow operation! Read an area of values as-is from a render target using a rectangle defined by integer pixel coordinates.
-    */
+	*/
 	UFUNCTION(BlueprintCallable, Category = "Rendering", meta = (Keywords = "ReadRenderTarget", WorldContext = "WorldContextObject"))
 	static ENGINE_API TArray<FLinearColor> ReadRenderTargetRawPixelArea(UObject* WorldContextObject, UTextureRenderTarget2D* TextureRenderTarget, int32 MinX, int32 MinY, int32 MaxX, int32 MaxY, bool bNormalize = true);
 
diff --git a/Engine/Source/Runtime/Engine/Private/KismetRenderingLibrary.cpp b/Engine/Source/Runtime/Engine/Private/KismetRenderingLibrary.cpp
--- a/Engine/Source/Runtime/Engine/Private/KismetRenderingLibrary.cpp
+++ b/Engine/Source/Runtime/Engine/Private/KismetRenderingLibrary.cpp
@@ -172,35 +172,23 @@ void UKismetRenderingLibrary::DrawMaterialToRenderTarget(UObject* WorldContextOb
 		Canvas->Init(TextureRenderTarget->SizeX, TextureRenderTarget->SizeY, nullptr, &RenderCanvas);
 		Canvas->Update();
 
-		FDrawEvent* DrawMaterialToTargetEvent = new FDrawEvent();
-
-		FName RTName = TextureRenderTarget->GetFName();
-		ENQUEUE_RENDER_COMMAND(BeginDrawEventCommand)(
-			[RTName, DrawMaterialToTargetEvent, RenderTargetResource](FRHICommandListImmediate& RHICmdList)
-			{
-				RenderTargetResource->FlushDeferredResourceUpdate(RHICmdList);
+		{
+			SCOPED_DRAW_EVENTF_GAMETHREAD(DrawMaterialToRenderTarget, *TextureRenderTarget->GetFName().ToString());
 
-				BEGIN_DRAW_EVENTF(
-					RHICmdList, 
-					DrawCanvasToTarget, 
-					(*DrawMaterialToTargetEvent), 
-					*RTName.ToString());
-			});
+			ENQUEUE_RENDER_COMMAND(FlushDeferredResourceUpdateCommand)(
+				[RenderTargetResource](FRHICommandListImmediate& RHICmdList)
+				{
+					RenderTargetResource->FlushDeferredResourceUpdate(RHICmdList);
+				});
 
-		Canvas->K2_DrawMaterial(Material, FVector2D(0, 0), FVector2D(TextureRenderTarget->SizeX, TextureRenderTarget->SizeY), FVector2D(0, 0));
+			Canvas->K2_DrawMaterial(Material, FVector2D(0, 0), FVector2D(TextureRenderTarget->SizeX, TextureRenderTarget->SizeY), FVector2D(0, 0));
 
-		RenderCanvas.Flush_GameThread();
-		Canvas->Canvas = NULL;
+			RenderCanvas.Flush_GameThread();
+			Canvas->Canvas = nullptr;
 
-		//UpdateResourceImmediate must be called here to ensure mips are generated.
-		TextureRenderTarget->UpdateResourceImmediate(false);
-		ENQUEUE_RENDER_COMMAND(CanvasRenderTargetResolveCommand)(
-			[DrawMaterialToTargetEvent](FRHICommandList& RHICmdList)
-			{
-				STOP_DRAW_EVENT((*DrawMaterialToTargetEvent));
-				delete DrawMaterialToTargetEvent;
-			}
-		);
+			//UpdateResourceImmediate must be called here to ensure mips are generated.
+			TextureRenderTarget->UpdateResourceImmediate(false);
+		}
 	}
 }
 
@@ -452,7 +440,7 @@ void UKismetRenderingLibrary::CreateTexture2DFromRenderTarget(UObject* WorldCont
 
 		//FImageUtils::CreateTexture2D
 
-		UTexture2D* NewTexture = RenderTarget->ConstructTexture2D(Texture->GetOuter(), Texture->GetName(), RenderTarget->GetMaskedFlags(), CTF_Default, NULL);
+		UTexture2D* NewTexture = RenderTarget->ConstructTexture2D(Texture->GetOuter(), Texture->GetName(), RenderTarget->GetMaskedFlags(), CTF_Default, nullptr);
 
 		check(NewTexture == Texture);
 		NewTexture->UpdateResource();
@@ -506,7 +494,7 @@ UTexture2D* UKismetRenderingLibrary::RenderTargetCreateStaticTexture2DEditorOnly
 		UObject* NewObj = nullptr;
 
 		// create a static 2d texture
-		NewObj = RenderTarget->ConstructTexture2D(CreatePackage( *PackageName), Name, RenderTarget->GetMaskedFlags() | RF_Public | RF_Standalone, CTF_Default | CTF_AllowMips, NULL);
+		NewObj = RenderTarget->ConstructTexture2D(CreatePackage( *PackageName), Name, RenderTarget->GetMaskedFlags() | RF_Public | RF_Standalone, CTF_Default | CTF_AllowMips, nullptr);
 		UTexture2D* NewTex = Cast<UTexture2D>(NewObj);
 
 		if (NewTex != nullptr)
@@ -550,7 +538,7 @@ void UKismetRenderingLibrary::ConvertRenderTargetToTexture2DEditorOnly( UObject*
 	}
 	else
 	{
-		UTexture2D* NewTexture = RenderTarget->ConstructTexture2D(Texture->GetOuter(), Texture->GetName(), RenderTarget->GetMaskedFlags() | RF_Public | RF_Standalone, CTF_Default, NULL);
+		UTexture2D* NewTexture = RenderTarget->ConstructTexture2D(Texture->GetOuter(), Texture->GetName(), RenderTarget->GetMaskedFlags() | RF_Public | RF_Standalone, CTF_Default, nullptr);
 
 		check(NewTexture == Texture);
 
@@ -619,7 +607,7 @@ UTexture2D* UKismetRenderingLibrary::ImportBufferAsTexture2D(UObject* WorldConte
 
 void UKismetRenderingLibrary::BeginDrawCanvasToRenderTarget(UObject* WorldContextObject, UTextureRenderTarget2D* TextureRenderTarget, UCanvas*& Canvas, FVector2D& Size, FDrawToRenderTargetContext& Context)
 {
-	Canvas = NULL;
+	Canvas = nullptr;
 	Size = FVector2D(0, 0);
 	Context = FDrawToRenderTargetContext();
 	
@@ -664,20 +652,15 @@ void UKismetRenderingLibrary::BeginDrawCanvasToRenderTarget(UObject* WorldContex
 		Canvas->Init(TextureRenderTarget->SizeX, TextureRenderTarget->SizeY, nullptr, NewCanvas);
 		Canvas->Update();
 
+#if  WANTS_DRAW_MESH_EVENTS
 		Context.DrawEvent = new FDrawEvent();
+		BEGIN_DRAW_EVENTF_GAMETHREAD(DrawCanvasToTarget, (*Context.DrawEvent), *TextureRenderTarget->GetFName().ToString())
+#endif // WANTS_DRAW_MESH_EVENTS
 
-		FName RTName = TextureRenderTarget->GetFName();
-		FDrawEvent* DrawEvent = Context.DrawEvent;
-		ENQUEUE_RENDER_COMMAND(BeginDrawEventCommand)(
-			[RTName, DrawEvent, RenderTargetResource](FRHICommandListImmediate& RHICmdList)
+		ENQUEUE_RENDER_COMMAND(FlushDeferredResourceUpdateCommand)(
+			[RenderTargetResource](FRHICommandListImmediate& RHICmdList)
 			{
 				RenderTargetResource->FlushDeferredResourceUpdate(RHICmdList);
-
-				BEGIN_DRAW_EVENTF(
-					RHICmdList, 
-					DrawCanvasToTarget, 
-					(*DrawEvent), 
-					*RTName.ToString());
 			});
 	}
 }
@@ -706,16 +689,20 @@ void UKismetRenderingLibrary::EndDrawCanvasToRenderTarget(UObject* WorldContextO
 		if (Context.RenderTarget)
 		{
 			FTextureRenderTargetResource* RenderTargetResource = Context.RenderTarget->GameThread_GetRenderTargetResource();
-			FDrawEvent* DrawEvent = Context.DrawEvent;
+
 			ENQUEUE_RENDER_COMMAND(CanvasRenderTargetResolveCommand)(
-				[RenderTargetResource, DrawEvent](FRHICommandList& RHICmdList)
+				[RenderTargetResource](FRHICommandList& RHICmdList)
 				{
 					RHICmdList.CopyToResolveTarget(RenderTargetResource->GetRenderTargetTexture(), RenderTargetResource->TextureRHI, FResolveParams());
-					STOP_DRAW_EVENT((*DrawEvent));
-					delete DrawEvent;
 				}
 			);
 
+#if WANTS_DRAW_MESH_EVENTS
+			STOP_DRAW_EVENT_GAMETHREAD(*Context.DrawEvent);
+			delete Context.DrawEvent;
+#endif // WANTS_DRAW_MESH_EVENTS
+
+
 			// Remove references to the context now that we've resolved it, to avoid a crash when EndDrawCanvasToRenderTarget is called multiple times with the same context
 			// const cast required, as BP will treat Context as an output without the const
 			const_cast<FDrawToRenderTargetContext&>(Context) = FDrawToRenderTargetContext();
diff --git a/Engine/Source/Runtime/Engine/Private/LevelTick.cpp b/Engine/Source/Runtime/Engine/Private/LevelTick.cpp
--- a/Engine/Source/Runtime/Engine/Private/LevelTick.cpp
+++ b/Engine/Source/Runtime/Engine/Private/LevelTick.cpp
@@ -41,6 +41,8 @@
 #include "TimerManager.h"
 #include "Camera/CameraPhotography.h"
 #include "HAL/LowLevelMemTracker.h"
+#include "ProfilingDebugging/RealtimeGPUProfiler.h"
+
 #if ENABLE_COLLISION_ANALYZER
 #include "PhysicsEngine/CollisionAnalyzerCapture.h"
 #endif
@@ -66,7 +68,7 @@
 #include "ComputeFramework/ComputeFramework.h"
 
 #if WITH_EDITOR
-#include "Editor.h"
+	#include "Editor.h"
 #include "LevelInstance/LevelInstanceSubsystem.h"
 #include "ObjectCacheEventSink.h"
 #endif
@@ -970,83 +972,74 @@ bool UWorld::HasEndOfFrameUpdates() const
 
 struct FSendAllEndOfFrameUpdates
 {
+	FSendAllEndOfFrameUpdates(FSceneInterface* InScene)
+	{
+		if (InScene != nullptr)
+		{
+			GPUSkinCache = InScene->GetGPUSkinCache();
+			ComputeFramework = InScene->GetComputeFramework();
+			FeatureLevel = InScene->GetFeatureLevel();
+		}
+	}
+	
 	FGPUSkinCache* GPUSkinCache = nullptr;
 	FComputeFramework* ComputeFramework = nullptr;
 	ERHIFeatureLevel::Type FeatureLevel = ERHIFeatureLevel::Num;
 
 #if WANTS_DRAW_MESH_EVENTS
 	FDrawEvent DrawEvent;
-#endif
+#endif // WANTS_DRAW_MESH_EVENTS
 };
 
-FSendAllEndOfFrameUpdates* BeginSendEndOfFrameUpdatesDrawEvent(
-	FGPUSkinCache* GPUSkinCache,
-	FComputeFramework* ComputeFramework,
-	ERHIFeatureLevel::Type FeatureLevel
-	)
+void BeginSendEndOfFrameUpdatesDrawEvent(FSendAllEndOfFrameUpdates& SendAllEndOfFrameUpdates)
 {
-	FSendAllEndOfFrameUpdates* SendAllEndOfFrameUpdates = new FSendAllEndOfFrameUpdates;
-	SendAllEndOfFrameUpdates->GPUSkinCache = GPUSkinCache;
-	SendAllEndOfFrameUpdates->ComputeFramework = ComputeFramework;
-	SendAllEndOfFrameUpdates->FeatureLevel = FeatureLevel;
-
-#if WANTS_DRAW_MESH_EVENTS
+	BEGIN_DRAW_EVENTF_GAMETHREAD(SendAllEndOfFrameUpdates, SendAllEndOfFrameUpdates.DrawEvent, TEXT("SendAllEndOfFrameUpdates"));
+	
 	ENQUEUE_RENDER_COMMAND(BeginDrawEventCommand)(
-		[SendAllEndOfFrameUpdates](FRHICommandListImmediate& RHICmdList)
+		[GPUSkinCache = SendAllEndOfFrameUpdates.GPUSkinCache](FRHICommandListImmediate& RHICmdList)
 		{
-			BEGIN_DRAW_EVENTF(
-				RHICmdList, 
-				SendAllEndOfFrameUpdates, 
-				SendAllEndOfFrameUpdates->DrawEvent,
-				TEXT("SendAllEndOfFrameUpdates"));
-
-			if (SendAllEndOfFrameUpdates->GPUSkinCache)
+			if (GPUSkinCache != nullptr)
 			{
-				SendAllEndOfFrameUpdates->GPUSkinCache->BeginBatchDispatch(RHICmdList);
+				GPUSkinCache->BeginBatchDispatch(RHICmdList);
 			}
 		});
-#endif
-
-	return SendAllEndOfFrameUpdates;
 }
 
 DECLARE_GPU_STAT(EndOfFrameUpdates);
 DECLARE_GPU_STAT(GPUSkinCacheRayTracingGeometry);
 DECLARE_GPU_STAT(ComputeFrameworkExecuteBatches);
-void EndSendEndOfFrameUpdatesDrawEvent(FSendAllEndOfFrameUpdates* SendAllEndOfFrameUpdates)
+void EndSendEndOfFrameUpdatesDrawEvent(FSendAllEndOfFrameUpdates& SendAllEndOfFrameUpdates)
 {
 	ENQUEUE_RENDER_COMMAND(EndDrawEventCommand)(
-		[SendAllEndOfFrameUpdates](FRHICommandListImmediate& RHICmdList)
-	{
-		SCOPED_GPU_STAT(RHICmdList, EndOfFrameUpdates);
-
-		if (SendAllEndOfFrameUpdates->GPUSkinCache)
+		[GPUSkinCache = SendAllEndOfFrameUpdates.GPUSkinCache, ComputeFramework = SendAllEndOfFrameUpdates.ComputeFramework, FeatureLevel = SendAllEndOfFrameUpdates.FeatureLevel](FRHICommandListImmediate& RHICmdList)
 		{
+			SCOPED_GPU_STAT(RHICmdList, EndOfFrameUpdates);
+
+			if (GPUSkinCache != nullptr)
 			{
 				// Once all the individual components have received their DoDeferredRenderUpdates_Concurrent()
 				// allow the GPU Skin Cache system to update.
-				SendAllEndOfFrameUpdates->GPUSkinCache->EndBatchDispatch(RHICmdList);
+				GPUSkinCache->EndBatchDispatch(RHICmdList);
 
 				// Flush any remaining pending resource barriers.
-				SendAllEndOfFrameUpdates->GPUSkinCache->TransitionAllToReadable(RHICmdList);
+				GPUSkinCache->TransitionAllToReadable(RHICmdList);
+
+#if RHI_RAYTRACING
+				{
+					SCOPED_GPU_STAT(RHICmdList, GPUSkinCacheRayTracingGeometry);
+					GPUSkinCache->CommitRayTracingGeometryUpdates(RHICmdList);
+				}
+#endif // RHI_RAYTRACING
 			}
 
-		#if RHI_RAYTRACING
+			if (ComputeFramework != nullptr)
 			{
-				SCOPED_GPU_STAT(RHICmdList, GPUSkinCacheRayTracingGeometry);
-				SendAllEndOfFrameUpdates->GPUSkinCache->CommitRayTracingGeometryUpdates(RHICmdList);
+				SCOPED_GPU_STAT(RHICmdList, ComputeFrameworkExecuteBatches);
+				ComputeFramework->ExecuteBatches(RHICmdList, FeatureLevel);
 			}
-		#endif // RHI_RAYTRACING
-		}
-
-		if (SendAllEndOfFrameUpdates->ComputeFramework)
-		{
-			SCOPED_GPU_STAT(RHICmdList, ComputeFrameworkExecuteBatches);
-			SendAllEndOfFrameUpdates->ComputeFramework->ExecuteBatches(RHICmdList, SendAllEndOfFrameUpdates->FeatureLevel);
-		}
+		});
 
-		delete SendAllEndOfFrameUpdates;
-	});
+	STOP_DRAW_EVENT_GAMETHREAD(SendAllEndOfFrameUpdates.DrawEvent);
 }
 
 /**
@@ -1106,11 +1099,8 @@ void UWorld::SendAllEndOfFrameUpdates()
 	}
 
 	// Issue a GPU event to wrap GPU work done during SendAllEndOfFrameUpdates, like skin cache updates
-	FSendAllEndOfFrameUpdates* SendAllEndOfFrameUpdates = BeginSendEndOfFrameUpdatesDrawEvent(
-		Scene ? Scene->GetGPUSkinCache() : nullptr,
-		Scene ? Scene->GetComputeFramework() : nullptr,
-		Scene ? Scene->GetFeatureLevel() : ERHIFeatureLevel::Num
-		);
+	FSendAllEndOfFrameUpdates SendAllEndOfFrameUpdates(Scene);
+	BeginSendEndOfFrameUpdatesDrawEvent(SendAllEndOfFrameUpdates);
 
 	// update all dirty components. 
 	FGuardValue_Bitfield(bPostTickComponentUpdate, true); 
@@ -1331,34 +1321,6 @@ DECLARE_CYCLE_STAT(TEXT("TG_LastDemotable"), STAT_TG_LastDemotable, STATGROUP_Ti
 
 #include "GameFramework/SpawnActorTimer.h"
 
-FDrawEvent* BeginTickDrawEvent()
-{
-	FDrawEvent* TickDrawEvent = new FDrawEvent();
-	FDrawEvent* InTickDrawEvent = TickDrawEvent;
-	ENQUEUE_RENDER_COMMAND(BeginDrawEventCommand)(
-		[InTickDrawEvent](FRHICommandList& RHICmdList)
-		{
-			BEGIN_DRAW_EVENTF(
-				RHICmdList, 
-				WorldTick, 
-				(*InTickDrawEvent),
-				TEXT("WorldTick"));
-		});
-
-	return TickDrawEvent;
-}
-
-void EndTickDrawEvent(FDrawEvent* TickDrawEvent)
-{
-	FDrawEvent* InTickDrawEvent = TickDrawEvent;
-	ENQUEUE_RENDER_COMMAND(EndDrawEventCommand)(
-		[InTickDrawEvent](FRHICommandList& RHICmdList)
-		{
-			STOP_DRAW_EVENT((*InTickDrawEvent));
-			delete InTickDrawEvent;
-		});
-}
-
 /**
  * Update the level after a variable amount of time, DeltaSeconds, has passed.
  * All child actors are ticked after their owners have been ticked.
@@ -1375,7 +1337,7 @@ void UWorld::Tick( ELevelTick TickType, float DeltaSeconds )
 		return;
 	}
 
-	FDrawEvent* TickDrawEvent = BeginTickDrawEvent();
+	SCOPED_DRAW_EVENT_GAMETHREAD(WorldTick);
 
 	FWorldDelegates::OnWorldTickStart.Broadcast(this, TickType, DeltaSeconds);
 
@@ -1853,8 +1815,6 @@ void UWorld::Tick( ELevelTick TickType, float DeltaSeconds )
 				WorldParam->PerfTrackers->GetInGamePerformanceTracker((EInGamePerfTrackers)Tracker, EInGamePerfTrackerThreads::RenderThread).Tick();
 			}
 		});
-
-	EndTickDrawEvent(TickDrawEvent);
 }
 
 void UWorld::CleanupActors()
diff --git a/Engine/Source/Runtime/Landscape/Private/LandscapeEditLayers.cpp b/Engine/Source/Runtime/Landscape/Private/LandscapeEditLayers.cpp
--- a/Engine/Source/Runtime/Landscape/Private/LandscapeEditLayers.cpp
+++ b/Engine/Source/Runtime/Landscape/Private/LandscapeEditLayers.cpp
@@ -6741,23 +6741,11 @@ UTextureRenderTarget2D* FLandscapeLayerBrush::Render(bool InIsHeightmap, const F
 	if (Initialize(InLandscapeExtent, InLandscapeRenderTarget))
 	{
 		FString ProfilingEventName = FString::Format(TEXT("LandscapeLayers_RenderLayerBrush {0}: {1}"), { InIsHeightmap ? TEXT("LS Height") : TEXT("LS Weight"), BlueprintBrush->GetName() });
-		TSharedPtr<FDrawEvent> ProfilingEvent = MakeShared<FDrawEvent>();
-
-		ENQUEUE_RENDER_COMMAND(LandscapeLayers_Cmd_BeginRenderLayerBrush)(
-			[ProfilingEvent, ProfilingEventName](FRHICommandListImmediate& RHICmdList) mutable
-		{
-			BEGIN_DRAW_EVENTF(RHICmdList, LandscapeLayers, *ProfilingEvent.Get(), *ProfilingEventName);
-		});
+		SCOPED_DRAW_EVENTF_GAMETHREAD(LandscapeLayers, *ProfilingEventName);
 
 		TGuardValue<bool> AutoRestore(GAllowActorScriptExecutionInEditor, true);
 		UTextureRenderTarget2D* Result = BlueprintBrush->Render(InIsHeightmap, InLandscapeRenderTarget, InWeightmapLayerName);
 
-		ENQUEUE_RENDER_COMMAND(LandscapeLayers_Cmd_StopRenderLayerBrush)(
-			[ProfilingEvent](FRHICommandList& RHICmdList)
-		{
-			STOP_DRAW_EVENT(*ProfilingEvent.Get());
-		});
-
 		return Result;
 	}
 #endif
diff --git a/Engine/Source/Runtime/RenderCore/Private/ProfilingDebugging/RealtimeGPUProfiler.cpp b/Engine/Source/Runtime/RenderCore/Private/ProfilingDebugging/RealtimeGPUProfiler.cpp
--- a/Engine/Source/Runtime/RenderCore/Private/ProfilingDebugging/RealtimeGPUProfiler.cpp
+++ b/Engine/Source/Runtime/RenderCore/Private/ProfilingDebugging/RealtimeGPUProfiler.cpp
@@ -4,6 +4,7 @@
 #include "ProfilingDebugging/CsvProfiler.h"
 #include "ProfilingDebugging/TracingProfiler.h"
 #include "RenderCore.h"
+#include "RenderingThread.h"
 #include "GPUProfiler.h"
 
 // Only exposed for debugging. Disabling this carries a severe performance penalty
@@ -64,14 +65,61 @@ void FDrawEvent::Start(FRHIComputeCommandList& InRHICmdList, FColor Color, const
 		RHICmdList = &InRHICmdList;
 		va_end(ptr);
 	}
+	bStarted = true;
+}
+
+void FDrawEvent::Start(FRHIComputeCommandList* InRHICmdList, FColor Color, const TCHAR* Fmt, ...)
+{
+	bool bIsRenderingOrRHIThread = IsInParallelRenderingThread() || IsInRHIThread();
+	// A command list must be passed if on the rendering or RHI thread, otherwise (game thread), a command will be enqueued on the immediate command list
+	{
+		va_list ptr;
+		va_start(ptr, Fmt);
+		TCHAR TempStr[256];
+		// Build the string in the temp buffer
+		FCString::GetVarArgs(TempStr, UE_ARRAY_COUNT(TempStr), Fmt, ptr);
+
+		if (InRHICmdList != nullptr)
+		{
+			check(!GUseThreadedRendering || bIsRenderingOrRHIThread); // A command list is needed on rendering/RHI threads (see comment above)
+			InRHICmdList->PushEvent(TempStr, Color);
+			RHICmdList = InRHICmdList;
+		}
+		else
+		{
+			check(!GUseThreadedRendering || !bIsRenderingOrRHIThread); // The command list should be null on game thread (see comment above)
+			ENQUEUE_RENDER_COMMAND(PushEventCommand)([EventName = FString(TempStr), Color](FRHICommandListImmediate& RHICommandListLocal)
+			{
+				RHICommandListLocal.PushEvent(*EventName, Color);
+			});
+		}
+
+		va_end(ptr);
+	}
+	bStarted = true;
 }
 
 void FDrawEvent::Stop()
 {
-	if (RHICmdList)
+	if (bStarted)
 	{
+		bool bIsRenderingOrRHIThread = IsInParallelRenderingThread() || IsInRHIThread();
+		// if we have a command list, we must be on the rendering or RHI thread, otherwise (game thread), a command will be enqueued on the immediate command list :
+		if (RHICmdList != nullptr)
+	{
+			check(!GUseThreadedRendering || bIsRenderingOrRHIThread); // A command list is needed on rendering/RHI threads (see comment above)
 		RHICmdList->PopEvent();
-		RHICmdList = NULL;
+			RHICmdList = nullptr;
+		}
+		else
+		{
+			check(!GUseThreadedRendering || !bIsRenderingOrRHIThread); // The command list should be null on game thread (see comment above)
+			ENQUEUE_RENDER_COMMAND(PopEventCommand)([](FRHICommandListImmediate& RHICommandListLocal)
+			{
+				RHICommandListLocal.PopEvent();
+			});
+		}
+		bStarted = false;
 	}
 }
 
diff --git a/Engine/Source/Runtime/RenderCore/Public/ProfilingDebugging/RealtimeGPUProfiler.h b/Engine/Source/Runtime/RenderCore/Public/ProfilingDebugging/RealtimeGPUProfiler.h
--- a/Engine/Source/Runtime/RenderCore/Public/ProfilingDebugging/RealtimeGPUProfiler.h
+++ b/Engine/Source/Runtime/RenderCore/Public/ProfilingDebugging/RealtimeGPUProfiler.h
@@ -65,9 +65,13 @@ struct RENDERCORE_API FBreadcrumbEvent
 		/** Cmdlist to push onto. */
 		FRHIComputeCommandList* RHICmdList;
 
+		/** Indicates whether the event has actually been fired or not. */
+		bool bStarted;
+
 		/** Default constructor, initializing all member variables. */
 		FORCEINLINE FDrawEvent()
 			: RHICmdList(nullptr)
+			, bStarted(false)
 		{}
 
 		/**
@@ -75,15 +79,21 @@ struct RENDERCORE_API FBreadcrumbEvent
 		 */
 		FORCEINLINE ~FDrawEvent()
 		{
-			if (RHICmdList)
+			if (bStarted)
 			{
 				Stop();
 			}
 		}
 
 		/**
-		 * Function for logging a PIX event with var args
+		 * Functions for logging a PIX event with var args. 
+		 * - If Start is called on the rendering or RHI threads, RHICmdList must be non-null. 
+		 * - On the game thread, RHICmdList must be nullptr and a render command will be enqueued on the immediate command list
+		 * - Stop can be called on any thread but must be called on the same thread Start was called on and will either use the same 
+		 *  command list (rendering / RHI thread) or enqueue a command on the current immediate command list (game thread)
 		 */
+		void CDECL Start(FRHIComputeCommandList* RHICmdList, FColor Color, const TCHAR* Fmt, ...);
+		UE_DEPRECATED(4.27, "This function has been deprecated in favor of the other version that takes a pointer in parameter")
 		void CDECL Start(FRHIComputeCommandList& RHICmdList, FColor Color, const TCHAR* Fmt, ...);
 		void Stop();
 	};
@@ -120,30 +130,44 @@ struct RENDERCORE_API FBreadcrumbEvent
 		void Stop();
 	};
 
-	#define SCOPED_GPU_EVENT(RHICmdList, Name) BREADCRUMB_EVENT(RHICmdList, Name); FDrawEvent PREPROCESSOR_JOIN(Event_##Name,__LINE__); if(GetEmitDrawEvents()) PREPROCESSOR_JOIN(Event_##Name,__LINE__).Start(RHICmdList, FColor(0), TEXT(#Name));
-	#define SCOPED_GPU_EVENT_COLOR(RHICmdList, Color, Name) BREADCRUMB_EVENT(RHICmdList, Name); FDrawEvent PREPROCESSOR_JOIN(Event_##Name,__LINE__); if(GetEmitDrawEvents()) PREPROCESSOR_JOIN(Event_##Name,__LINE__).Start(RHICmdList, Color, TEXT(#Name));
-	#define SCOPED_GPU_EVENTF(RHICmdList, Name, Format, ...) BREADCRUMB_EVENT(RHICmdList, Name); FDrawEvent PREPROCESSOR_JOIN(Event_##Name,__LINE__); if(GetEmitDrawEvents()) PREPROCESSOR_JOIN(Event_##Name,__LINE__).Start(RHICmdList, FColor(0), Format, ##__VA_ARGS__);
-	#define SCOPED_GPU_EVENTF_COLOR(RHICmdList, Color, Name, Format, ...) BREADCRUMB_EVENT(RHICmdList, Name); FDrawEvent PREPROCESSOR_JOIN(Event_##Name,__LINE__); if(GetEmitDrawEvents()) PREPROCESSOR_JOIN(Event_##Name,__LINE__).Start(RHICmdList, Color, Format, ##__VA_ARGS__);
-	#define SCOPED_CONDITIONAL_GPU_EVENT(RHICmdList, Name, Condition) FDrawEvent PREPROCESSOR_JOIN(Event_##Name,__LINE__); if(GetEmitDrawEvents() && (Condition)) PREPROCESSOR_JOIN(Event_##Name,__LINE__).Start(RHICmdList, FColor(0), TEXT(#Name));
-	#define SCOPED_CONDITIONAL_GPU_EVENT_COLOR(RHICmdList, Name, Color, Condition) FDrawEvent PREPROCESSOR_JOIN(Event_##Name,__LINE__); if(GetEmitDrawEvents() && (Condition)) PREPROCESSOR_JOIN(Event_##Name,__LINE__).Start(RHICmdList, Color, TEXT(#Name));
-	#define SCOPED_CONDITIONAL_GPU_EVENTF(RHICmdList, Name, Condition, Format, ...) FDrawEvent PREPROCESSOR_JOIN(Event_##Name,__LINE__); if(GetEmitDrawEvents() && (Condition)) PREPROCESSOR_JOIN(Event_##Name,__LINE__).Start(RHICmdList, FColor(0), Format, ##__VA_ARGS__);
-	#define SCOPED_CONDITIONAL_GPU_EVENTF_COLOR(RHICmdList, Color, Name, Condition, Format, ...) FDrawEvent PREPROCESSOR_JOIN(Event_##Name,__LINE__); if(GetEmitDrawEvents() && (Condition)) PREPROCESSOR_JOIN(Event_##Name,__LINE__).Start(RHICmdList, Color, Format, ##__VA_ARGS__);
-	#define BEGIN_GPU_EVENTF(RHICmdList, Name, Event, Format, ...) if(GetEmitDrawEvents()) (Event).Start(RHICmdList, FColor(0), Format, ##__VA_ARGS__);
-	#define BEGIN_GPU_EVENTF_COLOR(RHICmdList, Color, Name, Event, Format, ...) if(GetEmitDrawEvents()) (Event).Start(RHICmdList, Color, Format, ##__VA_ARGS__);
-	#define STOP_GPU_EVENT(Event) (Event).Stop();
-
 	// Macros to allow for scoping of draw events outside of RHI function implementations
-	#define SCOPED_DRAW_EVENT(RHICmdList, Name) BREADCRUMB_EVENT(RHICmdList, Name); FDrawEvent PREPROCESSOR_JOIN(Event_##Name,__LINE__); if(GetEmitDrawEvents()) PREPROCESSOR_JOIN(Event_##Name,__LINE__).Start(RHICmdList, FColor(0), TEXT(#Name));
-	#define SCOPED_DRAW_EVENT_COLOR(RHICmdList, Color, Name) BREADCRUMB_EVENT(RHICmdList, Name); FDrawEvent PREPROCESSOR_JOIN(Event_##Name,__LINE__); if(GetEmitDrawEvents()) PREPROCESSOR_JOIN(Event_##Name,__LINE__).Start(RHICmdList, Color, TEXT(#Name));
-	#define SCOPED_DRAW_EVENTF(RHICmdList, Name, Format, ...) BREADCRUMB_EVENT(RHICmdList, Name); FDrawEvent PREPROCESSOR_JOIN(Event_##Name,__LINE__); if(GetEmitDrawEvents()) PREPROCESSOR_JOIN(Event_##Name,__LINE__).Start(RHICmdList, FColor(0), Format, ##__VA_ARGS__);
-	#define SCOPED_DRAW_EVENTF_COLOR(RHICmdList, Color, Name, Format, ...) BREADCRUMB_EVENT(RHICmdList, Name); FDrawEvent PREPROCESSOR_JOIN(Event_##Name,__LINE__); if(GetEmitDrawEvents()) PREPROCESSOR_JOIN(Event_##Name,__LINE__).Start(RHICmdList, Color, Format, ##__VA_ARGS__);
-	#define SCOPED_CONDITIONAL_DRAW_EVENT(RHICmdList, Name, Condition) FDrawEvent PREPROCESSOR_JOIN(Event_##Name,__LINE__); if(GetEmitDrawEvents() && (Condition)) PREPROCESSOR_JOIN(Event_##Name,__LINE__).Start(RHICmdList, FColor(0), TEXT(#Name));
-	#define SCOPED_CONDITIONAL_DRAW_EVENT_COLOR(RHICmdList, Name, Color, Condition) FDrawEvent PREPROCESSOR_JOIN(Event_##Name,__LINE__); if(GetEmitDrawEvents() && (Condition)) PREPROCESSOR_JOIN(Event_##Name,__LINE__).Start(RHICmdList, Color, TEXT(#Name));
-	#define SCOPED_CONDITIONAL_DRAW_EVENTF(RHICmdList, Name, Condition, Format, ...) FDrawEvent PREPROCESSOR_JOIN(Event_##Name,__LINE__); if(GetEmitDrawEvents() && (Condition)) PREPROCESSOR_JOIN(Event_##Name,__LINE__).Start(RHICmdList, FColor(0), Format, ##__VA_ARGS__);
-	#define SCOPED_CONDITIONAL_DRAW_EVENTF_COLOR(RHICmdList, Color, Name, Condition, Format, ...) FDrawEvent PREPROCESSOR_JOIN(Event_##Name,__LINE__); if(GetEmitDrawEvents() && (Condition)) PREPROCESSOR_JOIN(Event_##Name,__LINE__).Start(RHICmdList, Color, Format, ##__VA_ARGS__);
-	#define BEGIN_DRAW_EVENTF(RHICmdList, Name, Event, Format, ...) if(GetEmitDrawEvents()) (Event).Start(RHICmdList, FColor(0), Format, ##__VA_ARGS__);
-	#define BEGIN_DRAW_EVENTF_COLOR(RHICmdList, Color, Name, Event, Format, ...) if(GetEmitDrawEvents()) (Event).Start(RHICmdList, Color, Format, ##__VA_ARGS__);
+	// Render-thread event macros:
+	#define SCOPED_DRAW_EVENT(RHICmdList, Name) BREADCRUMB_EVENT(RHICmdList, Name); FDrawEvent PREPROCESSOR_JOIN(Event_##Name,__LINE__); if(GetEmitDrawEvents()) PREPROCESSOR_JOIN(Event_##Name,__LINE__).Start(&RHICmdList, FColor(0), TEXT(#Name));
+	#define SCOPED_DRAW_EVENT_COLOR(RHICmdList, Color, Name) BREADCRUMB_EVENT(RHICmdList, Name); FDrawEvent PREPROCESSOR_JOIN(Event_##Name,__LINE__); if(GetEmitDrawEvents()) PREPROCESSOR_JOIN(Event_##Name,__LINE__).Start(&RHICmdList, Color, TEXT(#Name));
+	#define SCOPED_DRAW_EVENTF(RHICmdList, Name, Format, ...) BREADCRUMB_EVENT(RHICmdList, Name); FDrawEvent PREPROCESSOR_JOIN(Event_##Name,__LINE__); if(GetEmitDrawEvents()) PREPROCESSOR_JOIN(Event_##Name,__LINE__).Start(&RHICmdList, FColor(0), Format, ##__VA_ARGS__);
+	#define SCOPED_DRAW_EVENTF_COLOR(RHICmdList, Color, Name, Format, ...) BREADCRUMB_EVENT(RHICmdList, Name); FDrawEvent PREPROCESSOR_JOIN(Event_##Name,__LINE__); if(GetEmitDrawEvents()) PREPROCESSOR_JOIN(Event_##Name,__LINE__).Start(&RHICmdList, Color, Format, ##__VA_ARGS__);
+	#define SCOPED_CONDITIONAL_DRAW_EVENT(RHICmdList, Name, Condition) FDrawEvent PREPROCESSOR_JOIN(Event_##Name,__LINE__); if(GetEmitDrawEvents() && (Condition)) PREPROCESSOR_JOIN(Event_##Name,__LINE__).Start(&RHICmdList, FColor(0), TEXT(#Name));
+	#define SCOPED_CONDITIONAL_DRAW_EVENT_COLOR(RHICmdList, Name, Color, Condition) FDrawEvent PREPROCESSOR_JOIN(Event_##Name,__LINE__); if(GetEmitDrawEvents() && (Condition)) PREPROCESSOR_JOIN(Event_##Name,__LINE__).Start(&RHICmdList, Color, TEXT(#Name));
+	#define SCOPED_CONDITIONAL_DRAW_EVENTF(RHICmdList, Name, Condition, Format, ...) FDrawEvent PREPROCESSOR_JOIN(Event_##Name,__LINE__); if(GetEmitDrawEvents() && (Condition)) PREPROCESSOR_JOIN(Event_##Name,__LINE__).Start(&RHICmdList, FColor(0), Format, ##__VA_ARGS__);
+	#define SCOPED_CONDITIONAL_DRAW_EVENTF_COLOR(RHICmdList, Color, Name, Condition, Format, ...) FDrawEvent PREPROCESSOR_JOIN(Event_##Name,__LINE__); if(GetEmitDrawEvents() && (Condition)) PREPROCESSOR_JOIN(Event_##Name,__LINE__).Start(&RHICmdList, Color, Format, ##__VA_ARGS__);
+	#define BEGIN_DRAW_EVENTF(RHICmdList, Name, Event, Format, ...) if(GetEmitDrawEvents()) (Event).Start(&RHICmdList, FColor(0), Format, ##__VA_ARGS__);
+	#define BEGIN_DRAW_EVENTF_COLOR(RHICmdList, Color, Name, Event, Format, ...) if(GetEmitDrawEvents()) (Event).Start&(RHICmdList, Color, Format, ##__VA_ARGS__);
 	#define STOP_DRAW_EVENT(Event) (Event).Stop();
+	// Non-render-thread event macros:
+	#define SCOPED_DRAW_EVENT_GAMETHREAD(Name) FDrawEvent PREPROCESSOR_JOIN(Event_##Name,__LINE__); if(GetEmitDrawEvents()) PREPROCESSOR_JOIN(Event_##Name,__LINE__).Start(nullptr, FColor(0), TEXT(#Name));
+	#define SCOPED_DRAW_EVENT_COLOR_GAMETHREAD(Color, Name) FDrawEvent PREPROCESSOR_JOIN(Event_##Name,__LINE__); if(GetEmitDrawEvents()) PREPROCESSOR_JOIN(Event_##Name,__LINE__).Start(nullptr, Color, TEXT(#Name));
+	#define SCOPED_DRAW_EVENTF_GAMETHREAD(Name, Format, ...) FDrawEvent PREPROCESSOR_JOIN(Event_##Name,__LINE__); if(GetEmitDrawEvents()) PREPROCESSOR_JOIN(Event_##Name,__LINE__).Start(nullptr, FColor(0), Format, ##__VA_ARGS__);
+	#define SCOPED_DRAW_EVENTF_COLOR_GAMETHREAD(Color, Name, Format, ...) FDrawEvent PREPROCESSOR_JOIN(Event_##Name,__LINE__); if(GetEmitDrawEvents()) PREPROCESSOR_JOIN(Event_##Name,__LINE__).Start(nullptr, Color, Format, ##__VA_ARGS__);
+	#define SCOPED_CONDITIONAL_DRAW_EVENT_GAMETHREAD(Name, Condition) FDrawEvent PREPROCESSOR_JOIN(Event_##Name,__LINE__); if(GetEmitDrawEvents() && (Condition)) PREPROCESSOR_JOIN(Event_##Name,__LINE__).Start(nullptr, FColor(0), TEXT(#Name));
+	#define SCOPED_CONDITIONAL_DRAW_EVENT_COLOR_GAMETHREAD(Name, Color, Condition) FDrawEvent PREPROCESSOR_JOIN(Event_##Name,__LINE__); if(GetEmitDrawEvents() && (Condition)) PREPROCESSOR_JOIN(Event_##Name,__LINE__).Start(nullptr, Color, TEXT(#Name));
+	#define SCOPED_CONDITIONAL_DRAW_EVENTF_GAMETHREAD(Name, Condition, Format, ...) FDrawEvent PREPROCESSOR_JOIN(Event_##Name,__LINE__); if(GetEmitDrawEvents() && (Condition)) PREPROCESSOR_JOIN(Event_##Name,__LINE__).Start(nullptr, FColor(0), Format, ##__VA_ARGS__);
+	#define SCOPED_CONDITIONAL_DRAW_EVENTF_COLOR_GAMETHREAD(Color, Name, Condition, Format, ...) FDrawEvent PREPROCESSOR_JOIN(Event_##Name,__LINE__); if(GetEmitDrawEvents() && (Condition)) PREPROCESSOR_JOIN(Event_##Name,__LINE__).Start(nullptr, Color, Format, ##__VA_ARGS__);
+	#define BEGIN_DRAW_EVENTF_GAMETHREAD(Name, Event, Format, ...) if(GetEmitDrawEvents()) (Event).Start(nullptr, FColor(0), Format, ##__VA_ARGS__);
+	#define BEGIN_DRAW_EVENTF_COLOR_GAMETHREAD(Color, Name, Event, Format, ...) if(GetEmitDrawEvents()) (Event).Start(nullptr, Color, Format, ##__VA_ARGS__);
+	#define STOP_DRAW_EVENT_GAMETHREAD(Event) (Event).Stop();
+
+	// Deprecated version : use SCOPED_DRAW_... instead:
+	#define SCOPED_GPU_EVENT(RHICmdList, Name) SCOPED_DRAW_EVENT(RHICmdList, Name)
+	#define SCOPED_GPU_EVENT_COLOR(RHICmdList, Color, Name) SCOPED_DRAW_EVENT_COLOR(RHICmdList, Color, Name)
+	#define SCOPED_GPU_EVENTF(RHICmdList, Name, Format, ...) SCOPED_DRAW_EVENTF(RHICmdList, Name, Format, ##__VA_ARGS__)
+	#define SCOPED_GPU_EVENTF_COLOR(RHICmdList, Color, Name, Format, ...) SCOPED_DRAW_EVENTF_COLOR(RHICmdList, Color, Name, Format, ##__VA_ARGS__)
+	#define SCOPED_CONDITIONAL_GPU_EVENT(RHICmdList, Name, Condition) SCOPED_CONDITIONAL_DRAW_EVENT(RHICmdList, Name, Condition)
+	#define SCOPED_CONDITIONAL_GPU_EVENT_COLOR(RHICmdList, Name, Color, Condition) SCOPED_CONDITIONAL_DRAW_EVENT_COLOR(RHICmdList, Name, Color, Condition)
+	#define SCOPED_CONDITIONAL_GPU_EVENTF(RHICmdList, Name, Condition, Format, ...) SCOPED_CONDITIONAL_DRAW_EVENTF(RHICmdList, Name, Condition, Format, ##__VA_ARGS__)
+	#define SCOPED_CONDITIONAL_GPU_EVENTF_COLOR(RHICmdList, Color, Name, Condition, Format, ...) SCOPED_CONDITIONAL_DRAW_EVENTF_COLOR(RHICmdList, Color, Name, Condition, Format, ##__VA_ARGS__)
+	#define BEGIN_GPU_EVENTF(RHICmdList, Name, Event, Format, ...) BEGIN_DRAW_EVENTF(RHICmdList, Name, Event, Format, ##__VA_ARGS__)
+	#define BEGIN_GPU_EVENTF_COLOR(RHICmdList, Color, Name, Event, Format, ...) BEGIN_DRAW_EVENTF_COLOR(RHICmdList, Color, Name, Event, Format, ##__VA_ARGS__)
+	#define STOP_GPU_EVENT(Event) STOP_DRAW_EVENT(Event)
 
 	// Macros to allow for scoping of draw events within RHI function implementations
 	#define SCOPED_RHI_DRAW_EVENT(RHICmdContext, Name) FDrawEventRHIExecute PREPROCESSOR_JOIN(Event_##Name,__LINE__); if(GetEmitDrawEvents()) PREPROCESSOR_JOIN(Event_##Name,__LINE__).Start(RHICmdContext, FColor(0), TEXT(#Name));
@@ -161,18 +185,6 @@ struct RENDERCORE_API FBreadcrumbEvent
 	{
 	};
 
-	#define SCOPED_GPU_EVENT(RHICmdList, Name) BREADCRUMB_EVENT(RHICmdList, Name);
-	#define SCOPED_GPU_EVENT_COLOR(RHICmdList, Color, Name) BREADCRUMB_EVENT(RHICmdList, Name);
-	#define SCOPED_GPU_EVENTF(RHICmdList, Name, Format, ...) BREADCRUMB_EVENT(RHICmdList, Name);
-	#define SCOPED_GPU_EVENTF_COLOR(RHICmdList, Color, Name, Format, ...) BREADCRUMB_EVENT(RHICmdList, Name); 
-	#define SCOPED_CONDITIONAL_GPU_EVENT(...)
-	#define SCOPED_CONDITIONAL_GPU_EVENT_COLOR(...)
-	#define SCOPED_CONDITIONAL_GPU_EVENTF(...)
-	#define SCOPED_CONDITIONAL_GPU_EVENTF_COLOR(...)
-	#define BEGIN_GPU_EVENTF(...)
-	#define BEGIN_GPU_EVENTF_COLOR(...)
-	#define STOP_GPU_EVENT(...)
-
 	#define SCOPED_DRAW_EVENT(RHICmdList, Name) BREADCRUMB_EVENT(RHICmdList, Name);
 	#define SCOPED_DRAW_EVENT_COLOR(RHICmdList, Color, Name) BREADCRUMB_EVENT(RHICmdList, Name);
 	#define SCOPED_DRAW_EVENTF(RHICmdList, Name, Format, ...) BREADCRUMB_EVENT(RHICmdList, Name);
@@ -185,6 +197,30 @@ struct RENDERCORE_API FBreadcrumbEvent
 	#define BEGIN_DRAW_EVENTF_COLOR(...)
 	#define STOP_DRAW_EVENT(...)
 
+	#define SCOPED_DRAW_EVENT_GAMETHREAD(Name)
+	#define SCOPED_DRAW_EVENT_COLOR_GAMETHREAD(Color, Name)
+	#define SCOPED_DRAW_EVENTF_GAMETHREAD(Name, Format, ...)
+	#define SCOPED_DRAW_EVENTF_COLOR_GAMETHREAD(Color, Name, Format, ...)
+	#define SCOPED_CONDITIONAL_DRAW_EVENT_GAMETHREAD(Name, Condition)
+	#define SCOPED_CONDITIONAL_DRAW_EVENT_COLOR_GAMETHREAD(Name, Color, Condition)
+	#define SCOPED_CONDITIONAL_DRAW_EVENTF_GAMETHREAD(Name, Condition, Format, ...)
+	#define SCOPED_CONDITIONAL_DRAW_EVENTF_COLOR_GAMETHREAD(Color, Name, Condition, Format, ...)
+	#define BEGIN_DRAW_EVENTF_GAMETHREAD(Name, Event, Format, ...)
+	#define BEGIN_DRAW_EVENTF_COLOR_GAMETHREAD(Color, Name, Event, Format, ...)
+	#define STOP_DRAW_EVENT_GAMETHREAD(Event)
+
+	#define SCOPED_GPU_EVENT(RHICmdList, Name) BREADCRUMB_EVENT(RHICmdList, Name);
+	#define SCOPED_GPU_EVENT_COLOR(RHICmdList, Color, Name) BREADCRUMB_EVENT(RHICmdList, Name);
+	#define SCOPED_GPU_EVENTF(RHICmdList, Name, Format, ...) BREADCRUMB_EVENT(RHICmdList, Name);
+	#define SCOPED_GPU_EVENTF_COLOR(RHICmdList, Color, Name, Format, ...) BREADCRUMB_EVENT(RHICmdList, Name); 
+	#define SCOPED_CONDITIONAL_GPU_EVENT(...)
+	#define SCOPED_CONDITIONAL_GPU_EVENT_COLOR(...)
+	#define SCOPED_CONDITIONAL_GPU_EVENTF(...)
+	#define SCOPED_CONDITIONAL_GPU_EVENTF_COLOR(...)
+	#define BEGIN_GPU_EVENTF(...)
+	#define BEGIN_GPU_EVENTF_COLOR(...)
+	#define STOP_GPU_EVENT(...)
+
 	#define SCOPED_RHI_DRAW_EVENT(...)
 	#define SCOPED_RHI_DRAW_EVENT_COLOR(...)
 	#define SCOPED_RHI_DRAW_EVENTF(...)
diff --git a/Engine/Source/Runtime/Renderer/Private/RayTracing/RayTracingDynamicGeometry.cpp b/Engine/Source/Runtime/Renderer/Private/RayTracing/RayTracingDynamicGeometry.cpp
--- a/Engine/Source/Runtime/Renderer/Private/RayTracing/RayTracingDynamicGeometry.cpp
+++ b/Engine/Source/Runtime/Renderer/Private/RayTracing/RayTracingDynamicGeometry.cpp
@@ -333,7 +333,7 @@ void FRayTracingDynamicGeometryCollection::AddDynamicMeshBatchForGeometryUpdate(
 void FRayTracingDynamicGeometryCollection::DispatchUpdates(FRHIComputeCommandList& ParentCmdList)
 {
 #if WANTS_DRAW_MESH_EVENTS
-#define SCOPED_DRAW_OR_COMPUTE_EVENT(ParentCmdList, Name) FDrawEvent PREPROCESSOR_JOIN(Event_##Name,__LINE__); if(GetEmitDrawEvents()) PREPROCESSOR_JOIN(Event_##Name,__LINE__).Start(ParentCmdList, FColor(0), TEXT(#Name));
+#define SCOPED_DRAW_OR_COMPUTE_EVENT(ParentCmdList, Name) FDrawEvent PREPROCESSOR_JOIN(Event_##Name,__LINE__); if(GetEmitDrawEvents()) PREPROCESSOR_JOIN(Event_##Name,__LINE__).Start(&ParentCmdList, FColor(0), TEXT(#Name));
 #else
 #define SCOPED_DRAW_OR_COMPUTE_EVENT(...)
 #endif
