Removed high-frequency trace event.

diff --git a/Engine/Source/Runtime/RenderCore/Private/RenderResource.cpp b/Engine/Source/Runtime/RenderCore/Private/RenderResource.cpp
--- a/Engine/Source/Runtime/RenderCore/Private/RenderResource.cpp
+++ b/Engine/Source/Runtime/RenderCore/Private/RenderResource.cpp
@@ -76,8 +76,6 @@ void FRenderResource::ChangeFeatureLevel(ERHIFeatureLevel::Type NewFeatureLevel)
 
 void FRenderResource::InitResource()
 {
-	TRACE_CPUPROFILER_EVENT_SCOPE(FRenderResource::InitResource);
-	
 	check(IsInRenderingThread());
 	if (ListIndex == INDEX_NONE)
 	{
