Use "GetErrorStringWithLineMarker()" by default for FShaderCompilerError.

#rnx

diff --git a/Engine/Source/Runtime/RenderCore/Public/ShaderCompilerCore.h b/Engine/Source/Runtime/RenderCore/Public/ShaderCompilerCore.h
--- a/Engine/Source/Runtime/RenderCore/Public/ShaderCompilerCore.h
+++ b/Engine/Source/Runtime/RenderCore/Public/ShaderCompilerCore.h
@@ -362,8 +362,8 @@ struct FShaderCompilerError
 	FString HighlightedLine;
 	FString HighlightedLineMarker;
 
-	/** Returns the error message with source file and source line (if present). */
-	FString GetErrorString() const
+	/** Returns the error message with source file and source line (if present), as well as a line marker separated with a LINE_TERMINATOR. */
+	FString GetErrorStringWithoutLineMarker() const
 	{
 		if (ErrorVirtualFilePath.IsEmpty())
 		{
@@ -374,7 +374,7 @@ struct FShaderCompilerError
 			return ErrorVirtualFilePath + TEXT("(") + ErrorLineString + TEXT("): ") + StrippedErrorMessage;
 		}
 	}
-
+	
 	/** Returns the error message with source file and source line (if present), as well as a line marker separated with a LINE_TERMINATOR. */
 	FString GetErrorStringWithLineMarker() const
 	{
@@ -382,14 +382,20 @@ struct FShaderCompilerError
 		{
 			// Append highlighted line and its marker to the same error message with line terminators
 			// to get a similar multiline error output as with DXC
-			return (GetErrorString() + LINE_TERMINATOR + TEXT("\t") + HighlightedLine + LINE_TERMINATOR + TEXT("\t") + HighlightedLineMarker);
+			return (GetErrorStringWithoutLineMarker() + LINE_TERMINATOR + TEXT("\t") + HighlightedLine + LINE_TERMINATOR + TEXT("\t") + HighlightedLineMarker);
 		}
 		else
 		{
-			return GetErrorString();
+			return GetErrorStringWithoutLineMarker();
 		}
 	}
 
+	/** Returns the error message with source file and source line (if present). */
+	FString GetErrorString(bool bOmitLineMarker = false) const
+	{
+		return bOmitLineMarker ? GetErrorStringWithoutLineMarker() : GetErrorStringWithLineMarker();
+	}
+
 	/**
 	Returns true if this error message has a marker string for the highlighted source line where the error occurred. Example:
 		/Engine/Private/MySourceFile.usf(120): error: undeclared identifier 'a'
