Fix crash in editor where the format aliasing overload of RHICreateUnorderedAccessView is not implemented.
 - We shouldn't call this function when GRHISupportsUAVFormatAliasing is false

diff --git a/Engine/Source/Runtime/RenderCore/Private/RenderTargetPool.cpp b/Engine/Source/Runtime/RenderCore/Private/RenderTargetPool.cpp
--- a/Engine/Source/Runtime/RenderCore/Private/RenderTargetPool.cpp
+++ b/Engine/Source/Runtime/RenderCore/Private/RenderTargetPool.cpp
@@ -504,13 +504,22 @@ Done:
 
 		if ((Desc.TargetableFlags & TexCreate_UAV))
 		{
-			EPixelFormat AliasFormat = Desc.UAVFormat != PF_Unknown
-				? Desc.UAVFormat
-				: Desc.Format;
-
 			// The render target desc is invalid if a UAV is requested with an RHI that doesn't support the high-end feature level.
 			check(GMaxRHIFeatureLevel >= ERHIFeatureLevel::SM5 || GMaxRHIFeatureLevel == ERHIFeatureLevel::ES3_1);
-			Found->RenderTargetItem.UAV = RHICreateUnorderedAccessView(Found->RenderTargetItem.TargetableTexture, 0, AliasFormat);
+
+			if (GRHISupportsUAVFormatAliasing)
+			{
+				EPixelFormat AliasFormat = Desc.UAVFormat != PF_Unknown
+					? Desc.UAVFormat
+					: Desc.Format;
+
+				Found->RenderTargetItem.UAV = RHICreateUnorderedAccessView(Found->RenderTargetItem.TargetableTexture, 0, AliasFormat);
+			}
+			else
+			{
+				checkf(Desc.UAVFormat == PF_Unknown || Desc.UAVFormat == Desc.Format, TEXT("UAV aliasing is not supported by the current RHI."));
+				Found->RenderTargetItem.UAV = RHICreateUnorderedAccessView(Found->RenderTargetItem.TargetableTexture, 0);
+			}
 		}
 
 		AllocationLevelInKB += ComputeSizeInKB(*Found);
