RDG refactor to enable draining of work after issuing occlusion queries.
 - New Drain() method on FRDGBuilder; will flush all pending work.
 - Drained passes are not culled; resource lifetimes are extended; async compute fences are optimized as best as possible but fence joining may occur after the drain.
 - Batch up and pre-build all resource transitions. This is a prerequisite for parallel command lists.
 - Removed ServiceLocalQueue passes with built-in RDG AddDispatchHint().

diff --git a/Engine/Source/Runtime/RHI/Private/RHIValidation.cpp b/Engine/Source/Runtime/RHI/Private/RHIValidation.cpp
--- a/Engine/Source/Runtime/RHI/Private/RHIValidation.cpp
+++ b/Engine/Source/Runtime/RHI/Private/RHIValidation.cpp
@@ -1500,29 +1500,29 @@ FRHITransientTexture* FValidationTransientResourceAllocator::CreateTexture(const
 
 	FRHITexture* RHITexture = TransientTexture->GetRHI();
 
-	bool bReinitializeTrackerResource = true;
-
-	if (RHIValidation::FResource* Resource = RHITexture->GetTrackerResource())
-	{
-		if (!Resource->IsBarrierTrackingInitialized())
-		{
-			RHITexture->InitBarrierTracking(InCreateInfo.NumMips, InCreateInfo.ArraySize, InCreateInfo.Format, InCreateInfo.Flags, ERHIAccess::Discard, InDebugName);
-			bReinitializeTrackerResource = false;
-		}
-	}
-
 	// Store allocation data
 	FAllocatedResourceData ResourceData;
 	ResourceData.DebugName = InDebugName;
 	ResourceData.ResourceType = FAllocatedResourceData::EType::Texture;
 	ResourceData.bMemoryAllocated = true;
-	ResourceData.bReinitializeBarrierTracking = bReinitializeTrackerResource;
 	ResourceData.Texture.Flags = InCreateInfo.Flags;
 	ResourceData.Texture.Format = InCreateInfo.Format;
 	ResourceData.Texture.ArraySize = InCreateInfo.ArraySize;
 	ResourceData.Texture.NumMips = InCreateInfo.NumMips;
 	AllocatedResourceMap.Add(RHITexture, ResourceData);
 
+	if (RHIValidation::FResource* Resource = RHITexture->GetTrackerResource())
+	{
+		if (!Resource->IsBarrierTrackingInitialized())
+		{
+			RHITexture->InitBarrierTracking(InCreateInfo.NumMips, InCreateInfo.ArraySize, InCreateInfo.Format, InCreateInfo.Flags, ERHIAccess::Discard, InDebugName);
+		}
+		else
+		{
+			AllocatedResourcesToInit.Emplace(RHITexture, ResourceData);
+		}
+	}
+
 	return TransientTexture;
 }
 
@@ -1540,22 +1540,22 @@ FRHITransientBuffer* FValidationTransientResourceAllocator::CreateBuffer(const F
 
 	FRHIBuffer* RHIBuffer = TransientBuffer->GetRHI();
 
-	bool bReinitializeTrackerResource = true;
-
-	if (!RHIBuffer->IsBarrierTrackingInitialized())
-	{
-		RHIBuffer->InitBarrierTracking(ERHIAccess::Discard, InDebugName);
-		bReinitializeTrackerResource = false;
-	}
-
 	// Store allocation data
 	FAllocatedResourceData ResourceData;
 	ResourceData.DebugName = InDebugName;
 	ResourceData.ResourceType = FAllocatedResourceData::EType::Buffer;
 	ResourceData.bMemoryAllocated = true;
-	ResourceData.bReinitializeBarrierTracking = bReinitializeTrackerResource;
 	AllocatedResourceMap.Add(RHIBuffer, ResourceData);
 
+	if (!RHIBuffer->IsBarrierTrackingInitialized())
+	{
+		RHIBuffer->InitBarrierTracking(ERHIAccess::Discard, InDebugName);
+	}
+	else
+	{
+		AllocatedResourcesToInit.Emplace(RHIBuffer, ResourceData);
+	}
+
 	return TransientBuffer;
 }
 
@@ -1583,6 +1583,20 @@ void FValidationTransientResourceAllocator::DeallocateMemory(FRHITransientBuffer
 	ResourceData->bMemoryAllocated = false;
 }
 
+void FValidationTransientResourceAllocator::Flush(FRHICommandListImmediate& RHICmdList)
+{
+	check(!bFrozen);
+
+	RHICmdList.EnqueueLambda([AllocatedResourcesToInit = MoveTemp(AllocatedResourcesToInit)](FRHICommandListImmediate& InRHICmdList)
+	{
+		// Tracking will be re-initialized, so we need to flush any remaining references.
+		static_cast<FValidationContext&>(InRHICmdList.GetContext()).FlushValidationOps();
+		InitBarrierTracking(AllocatedResourcesToInit);
+	});
+
+	RHIAllocator->Flush(RHICmdList);
+}
+
 void FValidationTransientResourceAllocator::Freeze(FRHICommandListImmediate& RHICmdList)
 {
 	check(!bFrozen);
@@ -1616,12 +1630,11 @@ void FValidationTransientResourceAllocator::Freeze(FRHICommandListImmediate& RHI
 		}
 	}
 
-	RHICmdList.EnqueueLambda([this](FRHICommandListImmediate& InRHICmdList)
+	RHICmdList.EnqueueLambda([AllocatedResourcesToInit = MoveTemp(AllocatedResourcesToInit)](FRHICommandListImmediate& InRHICmdList)
 	{
 		// Tracking will be re-initialized, so we need to flush any remaining references.
 		static_cast<FValidationContext&>(InRHICmdList.GetContext()).FlushValidationOps();
-
-		InitBarrierTracking();
+		InitBarrierTracking(AllocatedResourcesToInit);
 	});
 
 	bFrozen = true;
@@ -1638,20 +1651,15 @@ void FValidationTransientResourceAllocator::Release(FRHICommandListImmediate& RH
 	delete this;
 }
 
-void FValidationTransientResourceAllocator::InitBarrierTracking()
+void FValidationTransientResourceAllocator::InitBarrierTracking(const FAllocatedResourceDataArray& AllocatedResourcesToInit)
 {
 	// Barrier tracking initialization has to happen on the RHI thread, because RHI resources are pooled and reused.
 
-	for (const auto& Entry : AllocatedResourceMap)
+	for (const auto& Entry : AllocatedResourcesToInit)
 	{
 		FRHIResource* Resource = Entry.Key;
 		const FAllocatedResourceData& ResourceData = Entry.Value;
 
-		if (!ResourceData.bReinitializeBarrierTracking)
-		{
-			continue;
-		}
-
 		switch (ResourceData.ResourceType)
 		{
 		case FAllocatedResourceData::EType::Texture:
diff --git a/Engine/Source/Runtime/RHI/Public/RHI.h b/Engine/Source/Runtime/RHI/Public/RHI.h
--- a/Engine/Source/Runtime/RHI/Public/RHI.h
+++ b/Engine/Source/Runtime/RHI/Public/RHI.h
@@ -810,12 +810,11 @@ enum class ERHIPipeline : uint8
 };
 ENUM_CLASS_FLAGS(ERHIPipeline)
 
-inline uint32 GetRHIPipelineIndex(ERHIPipeline Pipeline)
+inline constexpr uint32 GetRHIPipelineIndex(ERHIPipeline Pipeline)
 {
 	switch (Pipeline)
 	{
 	default:
-		checkf(false, TEXT("Attempted to get a pipeline index but the pipeline mask is not exclusively Graphics or AsyncCompute."));
 	case ERHIPipeline::Graphics:
 		return 0;
 	case ERHIPipeline::AsyncCompute:
@@ -852,9 +851,7 @@ class TRHIPipelineArray : public TStaticArray<ElementType, GetRHIPipelineCount()
 {
 	using Base = TStaticArray<ElementType, GetRHIPipelineCount()>;
 public:
-	TRHIPipelineArray() = default;
-	TRHIPipelineArray(const TRHIPipelineArray&) = default;
-	TRHIPipelineArray& operator=(const TRHIPipelineArray&) = default;
+	using Base::Base;
 
 	FORCEINLINE ElementType& operator[](ERHIPipeline Pipeline)
 	{
diff --git a/Engine/Source/Runtime/RHI/Public/RHITransientResourceAllocator.h b/Engine/Source/Runtime/RHI/Public/RHITransientResourceAllocator.h
--- a/Engine/Source/Runtime/RHI/Public/RHITransientResourceAllocator.h
+++ b/Engine/Source/Runtime/RHI/Public/RHITransientResourceAllocator.h
@@ -148,6 +148,9 @@ public:
 	virtual void DeallocateMemory(FRHITransientTexture* InTexture) = 0;
 	virtual void DeallocateMemory(FRHITransientBuffer* InBuffer) = 0;
 
+	// Flushes any pending allocations in preparation for rendering. Resources are not required to be deallocated yet. Use when interleaving execution with allocation.
+	virtual void Flush(FRHICommandListImmediate& RHICmdList) {}
+
 	// Freezes all allocations and validates that all resources have their memory deallocated.
 	virtual void Freeze(FRHICommandListImmediate& RHICmdList) = 0;
 
diff --git a/Engine/Source/Runtime/RHI/Public/RHIValidationTransientResourceAllocator.h b/Engine/Source/Runtime/RHI/Public/RHIValidationTransientResourceAllocator.h
--- a/Engine/Source/Runtime/RHI/Public/RHIValidationTransientResourceAllocator.h
+++ b/Engine/Source/Runtime/RHI/Public/RHIValidationTransientResourceAllocator.h
@@ -20,14 +20,11 @@ public:
 	virtual FRHITransientBuffer* CreateBuffer(const FRHIBufferCreateInfo& InCreateInfo, const TCHAR* InDebugName) override final;
 	virtual void DeallocateMemory(FRHITransientTexture* InTexture) override final;
 	virtual void DeallocateMemory(FRHITransientBuffer* InBuffer) override final;
+	virtual void Flush(FRHICommandListImmediate&) override final;
 	virtual void Freeze(FRHICommandListImmediate&) override final;
 	virtual void Release(FRHICommandListImmediate&) override final;
 
 private:
-
-	friend class FValidationContext;
-	void InitBarrierTracking();
-
 	// Actual RHI transient allocator which will get all functions forwarded
 	IRHITransientResourceAllocator* RHIAllocator = nullptr;
 
@@ -46,7 +43,6 @@ private:
 		const TCHAR* DebugName = nullptr;
 		EType ResourceType = EType::Texture;
 		bool bMemoryAllocated = false;
-		bool bReinitializeBarrierTracking = false;
 
 		struct FTexture
 		{
@@ -56,7 +52,15 @@ private:
 			uint8 NumMips = 0;
 		} Texture;
 	};
-	TMap<FRHIResource*, FAllocatedResourceData> AllocatedResourceMap;
+
+	using FAllocatedResourceDataMap = TMap<FRHIResource*, FAllocatedResourceData>;
+	using FAllocatedResourceDataArray = TArray<TPair<FRHIResource*, FAllocatedResourceData>>;
+
+	friend class FValidationContext;
+	static void InitBarrierTracking(const FAllocatedResourceDataArray& AllocatedResourcesToInit);
+
+	FAllocatedResourceDataMap AllocatedResourceMap;
+	FAllocatedResourceDataArray AllocatedResourcesToInit;
 };
 
 #endif	// ENABLE_RHI_VALIDATION
diff --git a/Engine/Source/Runtime/RenderCore/Private/RenderGraphAllocator.cpp b/Engine/Source/Runtime/RenderCore/Private/RenderGraphAllocator.cpp
--- a/Engine/Source/Runtime/RenderCore/Private/RenderGraphAllocator.cpp
+++ b/Engine/Source/Runtime/RenderCore/Private/RenderGraphAllocator.cpp
@@ -16,7 +16,7 @@ FRDGAllocator::~FRDGAllocator()
 
 void FRDGAllocator::ReleaseAll()
 {
-	SCOPE_CYCLE_COUNTER(STAT_RDG_ClearTime);
+	TRACE_CPUPROFILER_EVENT_SCOPE(FRDGAllocator::ReleaseAll);
 	CSV_SCOPED_TIMING_STAT_EXCLUSIVE_CONDITIONAL(RDGAllocator_Clear, GRDGVerboseCSVStats != 0);
 	for (int32 Index = TrackedAllocs.Num() - 1; Index >= 0; --Index)
 	{
diff --git a/Engine/Source/Runtime/RenderCore/Private/RenderGraphBuilder.cpp b/Engine/Source/Runtime/RenderCore/Private/RenderGraphBuilder.cpp
--- a/Engine/Source/Runtime/RenderCore/Private/RenderGraphBuilder.cpp
+++ b/Engine/Source/Runtime/RenderCore/Private/RenderGraphBuilder.cpp
@@ -438,9 +438,6 @@ bool FRDGBuilder::IsTransientInternal(FRDGParentResourceRef Resource) const
 	}
 #endif
 
-	// This resource cannot be extracted or made external later.
-	check(bSetupComplete);
-
 	return true;
 }
 
@@ -463,15 +460,14 @@ IRHITransientResourceAllocator* FRDGBuilder::FTransientResourceAllocator::GetOrC
 	return Allocator;
 }
 
-const char* const FRDGBuilder::kDefaultUnaccountedCSVStat = "RDG_Pass";
-
 FRDGBuilder::FRDGBuilder(FRHICommandListImmediate& InRHICmdList, FRDGEventName InName)
 	: RHICmdList(InRHICmdList)
 	, Blackboard(Allocator)
 	, RHICmdListAsyncCompute(FRHICommandListExecutor::GetImmediateAsyncComputeCommandList())
 	, BuilderName(InName)
+	, ProloguePasses(InPlace, nullptr)
 #if RDG_CPU_SCOPES
-	, CPUScopeStacks(RHICmdList, Allocator, kDefaultUnaccountedCSVStat)
+	, CPUScopeStacks(RHICmdList, Allocator)
 #endif
 #if RDG_GPU_SCOPES
 	, GPUScopeStacks(RHICmdList, RHICmdListAsyncCompute, Allocator)
@@ -479,25 +475,30 @@ FRDGBuilder::FRDGBuilder(FRHICommandListImmediate& InRHICmdList, FRDGEventName I
 #if RDG_ENABLE_DEBUG
 	, UserValidation(Allocator)
 	, BarrierValidation(&Passes, BuilderName)
+	, LogFile(Passes)
 #endif
+	, TransientResourceAllocator(RHICmdList)
 {
-	ProloguePass = Passes.Allocate<FRDGSentinelPass>(Allocator, RDG_EVENT_NAME("Graph Prologue"));
-	SetupEmptyPass(ProloguePass);
+	AddProloguePass();
 
 #if RDG_EVENTS != RDG_EVENTS_NONE
 	// This is polled once as a workaround for a race condition since the underlying global is not always changed on the render thread.
 	GRDGEmitEvents = GetEmitDrawEvents();
 #endif
+
+	IF_RDG_ENABLE_DEBUG(LogFile.Begin(BuilderName));
 }
 
+FRDGBuilder::~FRDGBuilder() {}
+
 const TRefCountPtr<FRDGPooledBuffer>& FRDGBuilder::ConvertToExternalBuffer(FRDGBufferRef Buffer)
 {
-	checkf(Buffer, TEXT("ConvertToExternalBuffer called with a null buffer."));
+	IF_RDG_ENABLE_DEBUG(UserValidation.ValidateConvertToExternalResource(Buffer));
 	if (!Buffer->bExternal)
 	{
 		Buffer->bExternal = 1;
 		Buffer->AccessFinal = kDefaultAccessFinal;
-		BeginResourceRHI(nullptr, GetProloguePassHandle(), Buffer);
+		BeginResourceRHI(GetProloguePassHandle(), Buffer);
 		ExternalBuffers.Add(Buffer->PooledBuffer, Buffer);
 	}
 	return GetPooledBuffer(Buffer);
@@ -505,12 +506,12 @@ const TRefCountPtr<FRDGPooledBuffer>& FRDGBuilder::ConvertToExternalBuffer(FRDGB
 
 const TRefCountPtr<IPooledRenderTarget>& FRDGBuilder::ConvertToExternalTexture(FRDGTextureRef Texture)
 {
-	checkf(Texture, TEXT("ConvertToExternalTexture called with a null texture."));
+	IF_RDG_ENABLE_DEBUG(UserValidation.ValidateConvertToExternalResource(Texture));
 	if (!Texture->bExternal)
 	{
 		Texture->bExternal = 1;
 		Texture->AccessFinal = kDefaultAccessFinal;
-		BeginResourceRHI(nullptr, GetProloguePassHandle(), Texture);
+		BeginResourceRHI(GetProloguePassHandle(), Texture);
 		ExternalTextures.Add(Texture->GetRHIUnchecked(), Texture);
 	}
 	return GetPooledTexture(Texture);
@@ -697,37 +698,31 @@ void FRDGBuilder::AddPassDependency(FRDGPassHandle ProducerHandle, FRDGPassHandl
 	}
 };
 
-void FRDGBuilder::Compile()
+void FRDGBuilder::Compile(EExecuteMode ExecuteMode)
 {
 	SCOPE_CYCLE_COUNTER(STAT_RDG_CompileTime);
 	CSV_SCOPED_TIMING_STAT_EXCLUSIVE_CONDITIONAL(RDG_Compile, GRDGVerboseCSVStats != 0);
 	SCOPED_NAMED_EVENT(FRDGBuilder_Compile, FColor::Emerald);
 
-	uint32 RasterPassCount = 0;
-	uint32 AsyncComputePassCount = 0;
-
-	FRDGPassBitArray PassesOnAsyncCompute(false, Passes.Num());
-	FRDGPassBitArray PassesOnRaster(false, Passes.Num());
-	FRDGPassBitArray PassesWithUntrackedOutputs(false, Passes.Num());
-	FRDGPassBitArray PassesToNeverCull(false, Passes.Num());
+	ExtendPassBitArray(PassesOnAsyncCompute, false);
 
 	const FRDGPassHandle ProloguePassHandle = GetProloguePassHandle();
 	const FRDGPassHandle EpiloguePassHandle = GetEpiloguePassHandle();
 
-	const auto IsCrossPipeline = [&](FRDGPassHandle A, FRDGPassHandle B)
-	{
-		return PassesOnAsyncCompute[A] != PassesOnAsyncCompute[B];
-	};
+	// This will differ from Passes.Num() when drains are performed.
+	const uint32 CompilePassCount = EpiloguePassHandle.GetIndex() - ProloguePassHandle.GetIndex() + 1;
 
-	const auto IsSortedBefore = [&](FRDGPassHandle A, FRDGPassHandle B)
-	{
-		return A < B;
-	};
+	// Culling is only performed in the final execution mode.
+	const bool bCullPasses = GRDGCullPasses && ExecuteMode == EExecuteMode::Final;
+
+	TArray<FRDGPassHandle, FRDGArrayAllocator> PassStack;
 
-	const auto IsSortedAfter = [&](FRDGPassHandle A, FRDGPassHandle B)
+	if (bCullPasses)
 	{
-		return A > B;
-	};
+		PassStack.Reserve(CompilePassCount);
+	}
+
+	TransitionCreateQueue.Reserve(CompilePassCount);
 
 	// Build producer / consumer dependencies across the graph and construct packed bit-arrays of metadata
 	// for better cache coherency when searching for passes meeting specific criteria. Search roots are also
@@ -736,6 +731,7 @@ void FRDGBuilder::Compile()
 	// lifetime to the epilogue pass which is always a root of the graph. The prologue and epilogue are helper
 	// passes and therefore never culled.
 
+	if (bCullPasses || AsyncComputePassCount > 0)
 	{
 		SCOPED_NAMED_EVENT(FRDGBuilder_Compile_Culling_Dependencies, FColor::Emerald);
 
@@ -762,10 +758,10 @@ void FRDGBuilder::Compile()
 			}
 		};
 
-		for (FRDGPassHandle PassHandle = Passes.Begin(); PassHandle != Passes.End(); ++PassHandle)
+		for (FRDGPassHandle PassHandle = ProloguePassHandle + 1; PassHandle < EpiloguePassHandle; ++PassHandle)
 		{
 			FRDGPass* Pass = Passes[PassHandle];
-			const ERHIPipeline PassPipeline = Pass->GetPipeline();
+			const ERHIPipeline PassPipeline = Pass->Pipeline;
 
 			bool bUntrackedOutputs = Pass->bHasExternalOutputs;
 
@@ -808,106 +804,72 @@ void FRDGBuilder::Compile()
 				bUntrackedOutputs |= Buffer->bExternal;
 			}
 
-			const ERDGPassFlags PassFlags = Pass->GetFlags();
-			const bool bAsyncCompute = EnumHasAnyFlags(PassFlags, ERDGPassFlags::AsyncCompute);
-			const bool bRaster = EnumHasAnyFlags(PassFlags, ERDGPassFlags::Raster);
-			const bool bNeverCull = EnumHasAnyFlags(PassFlags, ERDGPassFlags::NeverCull);
+			PassesOnAsyncCompute[PassHandle] = EnumHasAnyFlags(Pass->Flags, ERDGPassFlags::AsyncCompute);
+			Pass->bCulled = bCullPasses;
 
-			PassesOnRaster[PassHandle] = bRaster;
-			PassesOnAsyncCompute[PassHandle] = bAsyncCompute;
-			PassesToNeverCull[PassHandle] = bNeverCull;
-			PassesWithUntrackedOutputs[PassHandle] = bUntrackedOutputs;
-			AsyncComputePassCount += bAsyncCompute ? 1 : 0;
-			RasterPassCount += bRaster ? 1 : 0;
+			if (bCullPasses && (bUntrackedOutputs || EnumHasAnyFlags(Pass->Flags, ERDGPassFlags::NeverCull)))
+			{
+				PassStack.Emplace(PassHandle);
+			}
 		}
 
-		// The prologue / epilogue is responsible for external resource import / export, respectively.
-		PassesWithUntrackedOutputs[ProloguePassHandle] = true;
-		PassesWithUntrackedOutputs[EpiloguePassHandle] = true;
-
-		for (const auto& Query : ExtractedTextures)
+		if (ExecuteMode == EExecuteMode::Final)
 		{
-			FRDGTextureRef Texture = Query.Key;
-			for (auto& LastProducer : Texture->LastProducers)
+			for (const auto& Query : ExtractedTextures)
 			{
-				FRDGProducerState StateFinal;
-				StateFinal.Access = Texture->AccessFinal;
-				StateFinal.PassHandle = EpiloguePassHandle;
+				FRDGTextureRef Texture = Query.Key;
+				for (auto& LastProducer : Texture->LastProducers)
+				{
+					FRDGProducerState StateFinal;
+					StateFinal.Access = Texture->AccessFinal;
+					StateFinal.PassHandle = EpiloguePassHandle;
 
-				AddCullingDependency(LastProducer, StateFinal, ERHIPipeline::Graphics);
+					AddCullingDependency(LastProducer, StateFinal, ERHIPipeline::Graphics);
+				}
 			}
-			Texture->ReferenceCount++;
-		}
 
-		for (const auto& Query : ExtractedBuffers)
-		{
-			FRDGBufferRef Buffer = Query.Key;
+			for (const auto& Query : ExtractedBuffers)
+			{
+				FRDGBufferRef Buffer = Query.Key;
 
-			FRDGProducerState StateFinal;
-			StateFinal.Access = Buffer->AccessFinal;
-			StateFinal.PassHandle = EpiloguePassHandle;
+				FRDGProducerState StateFinal;
+				StateFinal.Access = Buffer->AccessFinal;
+				StateFinal.PassHandle = EpiloguePassHandle;
 
-			AddCullingDependency(Buffer->LastProducer, StateFinal, ERHIPipeline::Graphics);
-			Buffer->ReferenceCount++;
+				AddCullingDependency(Buffer->LastProducer, StateFinal, ERHIPipeline::Graphics);
+			}
 		}
-
-		EnumerateExtendedLifetimeResources(Textures, [](FRDGTextureRef Texture)
-		{
-			Texture->ReferenceCount++;
-		});
-
-		EnumerateExtendedLifetimeResources(Buffers, [](FRDGBufferRef Buffer)
-		{
-			Buffer->ReferenceCount++;
-		});
-
-		// Prologue / Epilogue passes get a lot of transitions added for external / extracted resources.
-		ProloguePass->GetEpilogueBarriersToBeginForGraphics(Allocator).Reserve(ExternalTextures.Num() + ExternalBuffers.Num());
-		EpiloguePass->GetPrologueBarriersToEnd(Allocator).Reserve(32);
 	}
 
 	// All dependencies in the raw graph have been specified; if enabled, all passes are marked as culled and a
 	// depth first search is employed to find reachable regions of the graph. Roots of the search are those passes
 	// with outputs leaving the graph or those marked to never cull.
 
-	if (GRDGCullPasses)
+	if (bCullPasses)
 	{
 		SCOPED_NAMED_EVENT(FRDGBuilder_Compile_Cull_Passes, FColor::Emerald);
-		TArray<FRDGPassHandle, FRDGArrayAllocator> PassStack;
-		PassesToCull.Init(true, Passes.Num());
-		PassStack.Reserve(Passes.Num());
 
-	#if STATS
-		GRDGStatPassCullCount += Passes.Num();
-	#endif
+		PassStack.Emplace(EpiloguePassHandle);
+
+		// Manually mark the prologue passes as not culled.
+		ProloguePasses[ERHIPipeline::Graphics]->bCulled = 0;
 
-		for (FRDGPassHandle PassHandle = Passes.Begin(); PassHandle != Passes.End(); ++PassHandle)
+		if (FRDGPass* Pass = ProloguePasses[ERHIPipeline::AsyncCompute])
 		{
-			if (PassesWithUntrackedOutputs[PassHandle] || PassesToNeverCull[PassHandle])
-			{
-				PassStack.Add(PassHandle);
-			}
+			Pass->bCulled = 0;
 		}
 
 		while (PassStack.Num())
 		{
-			const FRDGPassHandle PassHandle = PassStack.Pop();
+			FRDGPass* Pass = Passes[PassStack.Pop()];
 
-			if (PassesToCull[PassHandle])
+			if (Pass->bCulled)
 			{
-				PassesToCull[PassHandle] = false;
-				PassStack.Append(Passes[PassHandle]->Producers);
-
-			#if STATS
-				--GRDGStatPassCullCount;
-			#endif
+				Pass->bCulled = 0;
+				PassStack.Append(Pass->Producers);
 			}
 		}
 	}
-	else
-	{
-		PassesToCull.Init(false, Passes.Num());
-	}
 
 	// Walk the culled graph and compile barriers for each subresource. Certain transitions are redundant; read-to-read, for example.
 	// We can avoid them by traversing and merging compatible states together. The merging states removes a transition, but the merging
@@ -918,17 +880,18 @@ void FRDGBuilder::Compile()
 	{
 		SCOPED_NAMED_EVENT(FRDGBuilder_Compile_Barriers, FColor::Emerald);
 
-		for (FRDGPassHandle PassHandle = Passes.Begin(); PassHandle != Passes.End(); ++PassHandle)
+		for (FRDGPassHandle PassHandle = ProloguePassHandle + 1; PassHandle < EpiloguePassHandle; ++PassHandle)
 		{
-			if (PassesToCull[PassHandle] || PassesWithEmptyParameters[PassHandle])
+			FRDGPass* Pass = Passes[PassHandle];
+
+			if (Pass->bCulled || Pass->bEmptyParameters)
 			{
 				continue;
 			}
 
 			const bool bAsyncComputePass = PassesOnAsyncCompute[PassHandle];
 
-			FRDGPass* Pass = Passes[PassHandle];
-			const ERHIPipeline PassPipeline = Pass->GetPipeline();
+			const ERHIPipeline PassPipeline = Pass->Pipeline;
 
 			const auto MergeSubresourceStates = [&](ERDGParentResourceType ResourceType, FRDGSubresourceState*& PassMergeState, FRDGSubresourceState*& ResourceMergeState, const FRDGSubresourceState& PassState)
 			{
@@ -1005,11 +968,11 @@ void FRDGBuilder::Compile()
 	}
 
 	// Traverses passes on the graphics pipe and merges raster passes with the same render targets into a single RHI render pass.
-	if (GRDGMergeRenderPasses && RasterPassCount > 0)
+	if (IsRenderPassMergeEnabled() && RasterPassCount > 0)
 	{
 		SCOPED_NAMED_EVENT(FRDGBuilder_Compile_RenderPassMerge, FColor::Emerald);
 
-		TArray<FRDGPassHandle, FRDGArrayAllocator> PassesToMerge;
+		TArray<FRDGPassHandle, TInlineAllocator<32, FRDGArrayAllocator>> PassesToMerge;
 		FRDGPass* PrevPass = nullptr;
 		const FRenderTargetBindingSlots* PrevRenderTargets = nullptr;
 
@@ -1078,19 +1041,19 @@ void FRDGBuilder::Compile()
 			PrevRenderTargets = nullptr;
 		};
 
-		for (FRDGPassHandle PassHandle = Passes.Begin(); PassHandle != Passes.End(); ++PassHandle)
+		for (FRDGPassHandle PassHandle = ProloguePassHandle + 1; PassHandle < EpiloguePassHandle; ++PassHandle)
 		{
-			if (PassesToCull[PassHandle] || PassesWithEmptyParameters[PassHandle])
+			FRDGPass* NextPass = Passes[PassHandle];
+
+			if (NextPass->bCulled || NextPass->bEmptyParameters)
 			{
 				continue;
 			}
 
-			if (PassesOnRaster[PassHandle])
+			if (EnumHasAnyFlags(NextPass->Flags, ERDGPassFlags::Raster))
 			{
-				FRDGPass* NextPass = Passes[PassHandle];
-
 				// A pass where the user controls the render pass or it is forced to skip pass merging can't merge with other passes
-				if (EnumHasAnyFlags(NextPass->GetFlags(), ERDGPassFlags::SkipRenderPass | ERDGPassFlags::NeverMerge))
+				if (EnumHasAnyFlags(NextPass->Flags, ERDGPassFlags::SkipRenderPass | ERDGPassFlags::NeverMerge))
 				{
 					CommitMerge();
 					continue;
@@ -1130,7 +1093,7 @@ void FRDGBuilder::Compile()
 				PrevPass = NextPass;
 				PrevRenderTargets = &RenderTargets;
 			}
-			else if (!PassesOnAsyncCompute[PassHandle])
+			else if (!EnumHasAnyFlags(NextPass->Flags, ERDGPassFlags::AsyncCompute))
 			{
 				// A non-raster pass on the graphics pipe will invalidate the render target merge.
 				CommitMerge();
@@ -1148,18 +1111,23 @@ void FRDGBuilder::Compile()
 		// cross-pipeline consumer for each pass. This helps narrow the search space later when building async
 		// compute overlap regions.
 
+		const auto IsCrossPipeline = [&](FRDGPassHandle A, FRDGPassHandle B)
+		{
+			return PassesOnAsyncCompute[A] != PassesOnAsyncCompute[B];
+		};
+
 		FRDGPassBitArray PassesWithCrossPipelineProducer(false, Passes.Num());
 		FRDGPassBitArray PassesWithCrossPipelineConsumer(false, Passes.Num());
 
-		for (FRDGPassHandle PassHandle = Passes.Begin(); PassHandle != Passes.End(); ++PassHandle)
+		for (FRDGPassHandle PassHandle = ProloguePassHandle + 1; PassHandle < EpiloguePassHandle; ++PassHandle)
 		{
-			if (PassesToCull[PassHandle] || PassesWithEmptyParameters[PassHandle])
+			FRDGPass* Pass = Passes[PassHandle];
+
+			if (Pass->bCulled || Pass->bEmptyParameters)
 			{
 				continue;
 			}
 
-			FRDGPass* Pass = Passes[PassHandle];
-
 			for (FRDGPassHandle ProducerHandle : Pass->GetProducers())
 			{
 				const FRDGPassHandle ConsumerHandle = PassHandle;
@@ -1173,14 +1141,14 @@ void FRDGBuilder::Compile()
 				FRDGPass* Producer = Passes[ProducerHandle];
 
 				// Finds the earliest consumer on the other pipeline for the producer.
-				if (Producer->CrossPipelineConsumer.IsNull() || IsSortedBefore(ConsumerHandle, Producer->CrossPipelineConsumer))
+				if (Producer->CrossPipelineConsumer.IsNull() || ConsumerHandle < Producer->CrossPipelineConsumer)
 				{
 					Producer->CrossPipelineConsumer = PassHandle;
 					PassesWithCrossPipelineConsumer[ProducerHandle] = true;
 				}
 
 				// Finds the latest producer on the other pipeline for the consumer.
-				if (Consumer->CrossPipelineProducer.IsNull() || IsSortedAfter(ProducerHandle, Consumer->CrossPipelineProducer))
+				if (Consumer->CrossPipelineProducer.IsNull() || ProducerHandle > Consumer->CrossPipelineProducer)
 				{
 					Consumer->CrossPipelineProducer = ProducerHandle;
 					PassesWithCrossPipelineProducer[ConsumerHandle] = true;
@@ -1205,20 +1173,18 @@ void FRDGBuilder::Compile()
 
 		const auto FindCrossPipelineProducer = [&](FRDGPassHandle PassHandle)
 		{
-			check(PassHandle != ProloguePassHandle);
-
 			FRDGPassHandle LatestProducerHandle = ProloguePassHandle;
 			FRDGPassHandle ConsumerHandle = PassHandle;
 
 			// We want to find the latest producer on the other pipeline in order to establish a fork point.
 			// Since we could be consuming N resources with N producer passes, we only care about the last one.
-			while (ConsumerHandle != Passes.Begin())
+			while (ConsumerHandle != ProloguePassHandle)
 			{
-				if (!PassesToCull[ConsumerHandle] && !IsCrossPipeline(ConsumerHandle, PassHandle) && IsCrossPipelineConsumer(ConsumerHandle))
+				if (IsCrossPipelineConsumer(ConsumerHandle) && !IsCrossPipeline(ConsumerHandle, PassHandle))
 				{
 					const FRDGPass* Consumer = Passes[ConsumerHandle];
 
-					if (IsSortedAfter(Consumer->CrossPipelineProducer, LatestProducerHandle))
+					if (Consumer->CrossPipelineProducer > LatestProducerHandle && !Consumer->bCulled)
 					{
 						LatestProducerHandle = Consumer->CrossPipelineProducer;
 					}
@@ -1231,21 +1197,19 @@ void FRDGBuilder::Compile()
 
 		const auto FindCrossPipelineConsumer = [&](FRDGPassHandle PassHandle)
 		{
-			check(PassHandle != EpiloguePassHandle);
-
 			FRDGPassHandle EarliestConsumerHandle = EpiloguePassHandle;
 			FRDGPassHandle ProducerHandle = PassHandle;
 
 			// We want to find the earliest consumer on the other pipeline, as this establishes a join point
 			// between the pipes. Since we could be producing for N consumers on the other pipeline, we only
 			// care about the first one to execute.
-			while (ProducerHandle != Passes.End())
+			while (ProducerHandle != EpiloguePassHandle)
 			{
-				if (!PassesToCull[ProducerHandle] && !IsCrossPipeline(ProducerHandle, PassHandle) && IsCrossPipelineProducer(ProducerHandle))
+				if (IsCrossPipelineProducer(ProducerHandle) && !IsCrossPipeline(ProducerHandle, PassHandle))
 				{
 					const FRDGPass* Producer = Passes[ProducerHandle];
 
-					if (IsSortedBefore(Producer->CrossPipelineConsumer, EarliestConsumerHandle))
+					if (Producer->CrossPipelineConsumer < EarliestConsumerHandle && !Producer->bCulled)
 					{
 						EarliestConsumerHandle = Producer->CrossPipelineConsumer;
 					}
@@ -1258,7 +1222,7 @@ void FRDGBuilder::Compile()
 
 		const auto InsertGraphicsToAsyncComputeFork = [&](FRDGPass* GraphicsPass, FRDGPass* AsyncComputePass)
 		{
-			FRDGBarrierBatchBegin& EpilogueBarriersToBeginForAsyncCompute = GraphicsPass->GetEpilogueBarriersToBeginForAsyncCompute(Allocator);
+			FRDGBarrierBatchBegin& EpilogueBarriersToBeginForAsyncCompute = GraphicsPass->GetEpilogueBarriersToBeginForAsyncCompute(Allocator, TransitionCreateQueue);
 
 			GraphicsPass->bGraphicsFork = 1;
 			EpilogueBarriersToBeginForAsyncCompute.SetUseCrossPipelineFence();
@@ -1267,9 +1231,9 @@ void FRDGBuilder::Compile()
 			AsyncComputePass->GetPrologueBarriersToEnd(Allocator).AddDependency(&EpilogueBarriersToBeginForAsyncCompute);
 		};
 
-		const auto InsertAsyncToGraphicsComputeJoin = [&](FRDGPass* AsyncComputePass, FRDGPass* GraphicsPass)
+		const auto InsertAsyncComputeToGraphicsJoin = [&](FRDGPass* AsyncComputePass, FRDGPass* GraphicsPass)
 		{
-			FRDGBarrierBatchBegin& EpilogueBarriersToBeginForGraphics = AsyncComputePass->GetEpilogueBarriersToBeginForGraphics(Allocator);
+			FRDGBarrierBatchBegin& EpilogueBarriersToBeginForGraphics = AsyncComputePass->GetEpilogueBarriersToBeginForGraphics(Allocator, TransitionCreateQueue);
 
 			AsyncComputePass->bAsyncComputeEnd = 1;
 			EpilogueBarriersToBeginForGraphics.SetUseCrossPipelineFence();
@@ -1278,145 +1242,259 @@ void FRDGBuilder::Compile()
 			GraphicsPass->GetPrologueBarriersToEnd(Allocator).AddDependency(&EpilogueBarriersToBeginForGraphics);
 		};
 
-		FRDGPass* PrevGraphicsForkPass = nullptr;
-		FRDGPass* PrevGraphicsJoinPass = nullptr;
-		FRDGPass* PrevAsyncComputePass = nullptr;
+		FRDGPassHandle CurrentGraphicsForkPassHandle;
 
-		for (FRDGPassHandle PassHandle = Passes.Begin(); PassHandle != Passes.End(); ++PassHandle)
+		for (FRDGPassHandle PassHandle = ProloguePassHandle + 1; PassHandle < EpiloguePassHandle; ++PassHandle)
 		{
-			if (!PassesOnAsyncCompute[PassHandle] || PassesToCull[PassHandle])
+			if (!PassesOnAsyncCompute[PassHandle])
 			{
 				continue;
 			}
 
 			FRDGPass* AsyncComputePass = Passes[PassHandle];
 
+			if (AsyncComputePass->bCulled)
+			{
+				continue;
+			}
+
 			const FRDGPassHandle GraphicsForkPassHandle = FindCrossPipelineProducer(PassHandle);
-			const FRDGPassHandle GraphicsJoinPassHandle = FindCrossPipelineConsumer(PassHandle);
+
+			FRDGPass* GraphicsForkPass = Passes[GraphicsForkPassHandle];
 
 			AsyncComputePass->GraphicsForkPass = GraphicsForkPassHandle;
-			AsyncComputePass->GraphicsJoinPass = GraphicsJoinPassHandle;
+			GraphicsForkPass->ResourcesToBegin.Add(AsyncComputePass);
 
-			FRDGPass* GraphicsForkPass = Passes[GraphicsForkPassHandle];
+			if (CurrentGraphicsForkPassHandle != GraphicsForkPassHandle)
+			{
+				CurrentGraphicsForkPassHandle = GraphicsForkPassHandle;
+				InsertGraphicsToAsyncComputeFork(GraphicsForkPass, AsyncComputePass);
+			}
+		}
+
+		for (FRDGPassHandle PassHandle : PassesOnAsyncComputeToJoin)
+		{
+			const FRDGPassHandle GraphicsJoinPassHandle = FindCrossPipelineConsumer(PassHandle);
+
+			FRDGPass* AsyncComputePass = Passes[PassHandle];
 			FRDGPass* GraphicsJoinPass = Passes[GraphicsJoinPassHandle];
+			FRDGPass* AsyncComputeProloguePass = ProloguePasses[ERHIPipeline::AsyncCompute];
 
-			// Extend the lifetime of resources used on async compute to the fork / join graphics passes.
-			GraphicsForkPass->ResourcesToBegin.Add(AsyncComputePass);
+			AsyncComputePass->GraphicsJoinPass = GraphicsJoinPassHandle;
 			GraphicsJoinPass->ResourcesToEnd.Add(AsyncComputePass);
 
-			if (PrevGraphicsForkPass != GraphicsForkPass)
+			if (AsyncComputeProloguePass->GraphicsJoinPass == AsyncComputeProloguePass->Handle)
 			{
-				InsertGraphicsToAsyncComputeFork(GraphicsForkPass, AsyncComputePass);
+				AsyncComputeProloguePass->GraphicsJoinPass = GraphicsJoinPassHandle;
 			}
+			else
+			{
+				AsyncComputeProloguePass->GraphicsJoinPass = FMath::Min(AsyncComputeProloguePass->GraphicsJoinPass, GraphicsJoinPassHandle);
+			}
+		}
+
+		// Flush the prologue pass fence if it exists.
+		if (!PassesOnAsyncComputeToJoin.IsEmpty())
+		{
+			FRDGPass* AsyncComputePass = ProloguePasses[ERHIPipeline::AsyncCompute];
+			FRDGPass* GraphicsJoinPass = Passes[AsyncComputePass->GraphicsJoinPass];
+
+			check(AsyncComputePass && AsyncComputePass->GraphicsJoinPass && AsyncComputePass->Pipeline == ERHIPipeline::AsyncCompute);
+			InsertAsyncComputeToGraphicsJoin(AsyncComputePass, GraphicsJoinPass);
+		}
+
+		PassesOnAsyncComputeToJoin.Reset();
+
+		FRDGPassHandle CurrentGraphicsJoinPassHandle;
 
-			if (PrevGraphicsJoinPass != GraphicsJoinPass && PrevAsyncComputePass)
+		bool bFoundAsyncComputeEndExecutePass = false;
+
+		for (FRDGPassHandle PassHandle = EpiloguePassHandle - 1; PassHandle > ProloguePassHandle; --PassHandle)
+		{
+			if (!PassesOnAsyncCompute[PassHandle])
 			{
-				InsertAsyncToGraphicsComputeJoin(PrevAsyncComputePass, PrevGraphicsJoinPass);
+				continue;
 			}
 
-			PrevAsyncComputePass = AsyncComputePass;
-			PrevGraphicsForkPass = GraphicsForkPass;
-			PrevGraphicsJoinPass = GraphicsJoinPass;
+			FRDGPass* AsyncComputePass = Passes[PassHandle];
+
+			if (AsyncComputePass->bCulled)
+			{
+				continue;
+			}
+
+			const FRDGPassHandle GraphicsJoinPassHandle = FindCrossPipelineConsumer(PassHandle);
+
+			// When draining and the epilogue pass is found, mark it for the next cycle and continue.
+			if (GraphicsJoinPassHandle == EpiloguePassHandle && ExecuteMode == EExecuteMode::Drain)
+			{
+				PassesOnAsyncComputeToJoin.Add(PassHandle);
+				continue;
+			}
+
+			FRDGPass* GraphicsJoinPass = Passes[GraphicsJoinPassHandle];
+
+			AsyncComputePass->GraphicsJoinPass = GraphicsJoinPassHandle;
+			GraphicsJoinPass->ResourcesToEnd.Add(AsyncComputePass);
+
+			if (CurrentGraphicsJoinPassHandle != GraphicsJoinPassHandle)
+			{
+				CurrentGraphicsJoinPassHandle = GraphicsJoinPassHandle;
+				InsertAsyncComputeToGraphicsJoin(AsyncComputePass, GraphicsJoinPass);
+
+				if (!bFoundAsyncComputeEndExecutePass)
+				{
+					bFoundAsyncComputeEndExecutePass = true;
+
+					AsyncComputePass->bAsyncComputeEndExecute = 1;
+				}
+			}
 		}
 
-		// Last async compute pass in the graph needs to be manually joined back to the epilogue pass.
-		if (PrevAsyncComputePass)
+		if (!bFoundAsyncComputeEndExecutePass && ExecuteMode == EExecuteMode::Final)
 		{
-			InsertAsyncToGraphicsComputeJoin(PrevAsyncComputePass, PrevGraphicsJoinPass);
-			PrevAsyncComputePass->bAsyncComputeEndExecute = 1;
+			ProloguePasses[ERHIPipeline::AsyncCompute]->bAsyncComputeEndExecute = 1;
 		}
 	}
 }
 
+void FRDGBuilder::Drain()
+{
+	if (IsImmediateMode() || !GRDGDrain)
+	{
+		return;
+	}
+
+	Execute(EExecuteMode::Drain);
+}
+
 void FRDGBuilder::Execute()
+{
+	Execute(EExecuteMode::Final);
+}
+
+void FRDGBuilder::Execute(EExecuteMode ExecuteMode)
 {
 	CSV_SCOPED_TIMING_STAT_EXCLUSIVE(RDG);
-	SCOPED_NAMED_EVENT(FRDGBuilder_Execute, FColor::Emerald);
+	TRACE_CPUPROFILER_EVENT_SCOPE(FRDGBuilder::Execute);
 
 	// Create the epilogue pass at the end of the graph just prior to compilation.
-	EpiloguePass = Passes.Allocate<FRDGSentinelPass>(Allocator, RDG_EVENT_NAME("Graph Epilogue"));
-	SetupEmptyPass(EpiloguePass);
-	bSetupComplete = true;
-
-	IF_RDG_ENABLE_DEBUG(UserValidation.ValidateExecuteBegin());
+	SetupEmptyPass(EpiloguePass = Passes.Allocate<FRDGSentinelPass>(Allocator, RDG_EVENT_NAME("Graph Epilogue%s", (ExecuteMode == EExecuteMode::Drain) ? TEXT(" (Drain)") : TEXT(""))));
 
 	const FRDGPassHandle ProloguePassHandle = GetProloguePassHandle();
 	const FRDGPassHandle EpiloguePassHandle = GetEpiloguePassHandle();
 
-	FTransientResourceAllocator TransientResourceAllocator(RHICmdList);
+	const bool bFirstExecute = ExecuteCount == 0;
+	const bool bFinalExecute = ExecuteMode == EExecuteMode::Final;
 
-	if (!IsImmediateMode())
+	if (bFirstExecute)
 	{
-		RDG_ALLOW_RHI_ACCESS_SCOPE();
-		Compile();
+		IF_RDG_ENABLE_DEBUG(UserValidation.ValidateExecuteBegin());
+	}
 
-		IF_RDG_ENABLE_DEBUG(LogFile.Begin(BuilderName, &Passes, PassesToCull, ProloguePassHandle, EpiloguePassHandle));
+	IF_RDG_ENABLE_DEBUG(GRDGAllowRHIAccess = true);
+
+	if (!IsImmediateMode())
+	{
+		Compile(ExecuteMode);
 
 		{
-			TRACE_CPUPROFILER_EVENT_SCOPE(FRDGBuilder::CollectResources);
 			SCOPE_CYCLE_COUNTER(STAT_RDG_CollectResourcesTime);
 			CSV_SCOPED_TIMING_STAT_EXCLUSIVE(RDG_CollectResources);
+			TRACE_CPUPROFILER_EVENT_SCOPE(FRDGBuilder::CollectResources);
 
-			for (FRDGPassHandle PassHandle = Passes.Begin(); PassHandle != Passes.End(); ++PassHandle)
+			if (bFinalExecute)
 			{
-				if (!PassesToCull[PassHandle])
+				EnumerateExtendedLifetimeResources(Textures, [](FRDGTexture* Texture)
+				{
+					++Texture->ReferenceCount;
+				});
+
+				EnumerateExtendedLifetimeResources(Buffers, [](FRDGBuffer* Buffer)
+				{
+					++Buffer->ReferenceCount;
+				});
+
+				for (FRDGPassHandle PassHandle = Passes.Begin(); PassHandle < ProloguePassHandle; ++PassHandle)
 				{
-					CollectPassResources(TransientResourceAllocator, PassHandle);
+					FRDGPass* Pass = Passes[PassHandle];
+
+					if (!Pass->bCulled)
+					{
+						EndResourcesRHI(Pass, ProloguePassHandle);
+					}
 				}
-			}
 
-			for (const auto& Query : ExtractedTextures)
-			{
-				EndResourceRHI(TransientResourceAllocator, EpiloguePassHandle, Query.Key, 1);
-			}
+				for (FRDGPassHandle PassHandle = ProloguePassHandle; PassHandle <= EpiloguePassHandle; ++PassHandle)
+				{
+					FRDGPass* Pass = Passes[PassHandle];
 
-			for (const auto& Query : ExtractedBuffers)
-			{
-				EndResourceRHI(TransientResourceAllocator, EpiloguePassHandle, Query.Key, 1);
-			}
+					if (!Pass->bCulled)
+					{
+						BeginResourcesRHI(Pass, PassHandle);
+						EndResourcesRHI(Pass, PassHandle);
+					}
+				}
 
-			EnumerateExtendedLifetimeResources(Textures, [&](FRDGTextureRef Texture)
-			{
-				EndResourceRHI(TransientResourceAllocator, EpiloguePassHandle, Texture, 1);
-			});
+				for (const auto& Query : ExtractedTextures)
+				{
+					EndResourceRHI(EpiloguePassHandle, Query.Key, 1);
+				}
 
-			EnumerateExtendedLifetimeResources(Buffers, [&](FRDGBufferRef Buffer)
-			{
-				EndResourceRHI(TransientResourceAllocator, EpiloguePassHandle, Buffer, 1);
-			});
+				for (const auto& Query : ExtractedBuffers)
+				{
+					EndResourceRHI(EpiloguePassHandle, Query.Key, 1);
+				}
 
-			if (TransientResourceAllocator)
-			{
-				TransientResourceAllocator->Freeze(RHICmdList);
-			}
-		}
+				EnumerateExtendedLifetimeResources(Textures, [&](FRDGTextureRef Texture)
+				{
+					EndResourceRHI(EpiloguePassHandle, Texture, 1);
+				});
 
-		{
-			TRACE_CPUPROFILER_EVENT_SCOPE(FRDGBuilder::CreateUniformBuffers);
+				EnumerateExtendedLifetimeResources(Buffers, [&](FRDGBufferRef Buffer)
+				{
+					EndResourceRHI(EpiloguePassHandle, Buffer, 1);
+				});
 
-			for (TConstSetBitIterator It(UniformBuffersToCreate); It; ++It)
+				if (TransientResourceAllocator)
+				{
+					TransientResourceAllocator->Freeze(RHICmdList);
+				}
+			}
+			else
 			{
-				const FRDGUniformBufferHandle UniformBufferHandle(It.GetIndex());
+				for (FRDGPassHandle PassHandle = ProloguePassHandle; PassHandle < EpiloguePassHandle; ++PassHandle)
+				{
+					BeginResourcesRHI(Passes[PassHandle], PassHandle);
+				}
 
-				BeginResourceRHI(UniformBuffers[UniformBufferHandle]);
+				if (TransientResourceAllocator)
+				{
+					TransientResourceAllocator->Flush(RHICmdList);
+				}
 			}
 		}
 
+		CreateUniformBuffers();
+
 		{
 			TRACE_CPUPROFILER_EVENT_SCOPE(FRDGBuilder::CollectBarriers);
 			SCOPE_CYCLE_COUNTER(STAT_RDG_CollectBarriersTime);
 			CSV_SCOPED_TIMING_STAT_EXCLUSIVE_CONDITIONAL(RDG_CollectBarriers, GRDGVerboseCSVStats != 0);
 
-			for (FRDGPassHandle PassHandle = Passes.Begin(); PassHandle != Passes.End(); ++PassHandle)
+			for (FRDGPassHandle PassHandle = ProloguePassHandle + 1; PassHandle < EpiloguePassHandle; ++PassHandle)
 			{
-				if (!PassesToCull[PassHandle])
+				FRDGPass* Pass = Passes[PassHandle];
+
+				if (!Pass->bCulled && !Pass->bEmptyParameters)
 				{
-					CollectPassBarriers(PassHandle);
+					CollectPassBarriers(Pass, PassHandle);
 				}
 			}
 		}
 	}
 
+	if (bFinalExecute)
 	{
 		TRACE_CPUPROFILER_EVENT_SCOPE(FRDGBuilder::Finalize);
 
@@ -1435,129 +1513,156 @@ void FRDGBuilder::Execute()
 		};
 #endif
 
-		for (FRDGTextureHandle TextureHandle = Textures.Begin(); TextureHandle != Textures.End(); ++TextureHandle)
+		ActivePooledTextures.Reserve(Textures.Num());
+		Textures.Enumerate([&](FRDGTextureRef Texture)
 		{
-			FRDGTextureRef Texture = Textures[TextureHandle];
-
-			if (Texture->GetRHIUnchecked())
+			if (Texture->HasRHI())
 			{
 				AddEpilogueTransition(Texture);
-				Texture->Finalize();
+				Texture->Finalize(ActivePooledTextures);
 
 				IF_RDG_ENABLE_DEBUG(LogResource(Texture, Textures));
 			}
-		}
+		});
 
-		for (FRDGBufferHandle BufferHandle = Buffers.Begin(); BufferHandle != Buffers.End(); ++BufferHandle)
+		ActivePooledBuffers.Reserve(Buffers.Num());
+		Buffers.Enumerate([&](FRDGBufferRef Buffer)
 		{
-			FRDGBufferRef Buffer = Buffers[BufferHandle];
-
-			if (Buffer->GetRHIUnchecked())
+			if (Buffer->HasRHI())
 			{
 				AddEpilogueTransition(Buffer);
-				Buffer->Finalize();
+				Buffer->Finalize(ActivePooledBuffers);
 
 				IF_RDG_ENABLE_DEBUG(LogResource(Buffer, Buffers));
 			}
-		}
+		});
 	}
+	else
+	{
+		TRACE_CPUPROFILER_EVENT_SCOPE(FRDGBuilder::Finalize);
 
-	IF_RDG_CPU_SCOPES(CPUScopeStacks.BeginExecute());
-	IF_RDG_GPU_SCOPES(GPUScopeStacks.BeginExecute());
-	IF_RDG_ENABLE_TRACE(Trace.OutputGraphBegin());
+		Textures.Enumerate([&](FRDGTextureRef Texture)
+		{
+			FMemory::Memzero(Texture->MergeState.GetData(), Texture->MergeState.Num() * sizeof(FRDGSubresourceState*));
+		});
+
+		Buffers.Enumerate([&](FRDGBufferRef Buffer)
+		{
+			Buffer->MergeState = nullptr;
+		});
+	}
+
+	IF_RDG_ENABLE_DEBUG(GRDGAllowRHIAccess = false);
+
+	CreatePassBarriers();
+
+	if (bFirstExecute)
+	{
+		IF_RDG_CPU_SCOPES(CPUScopeStacks.BeginExecute());
+		IF_RDG_GPU_SCOPES(GPUScopeStacks.BeginExecute());
+		IF_RDG_ENABLE_TRACE(Trace.OutputGraphBegin());
+	}
 
 	if (!IsImmediateMode())
 	{
 		QUICK_SCOPE_CYCLE_COUNTER(STAT_FRDGBuilder_Execute_Passes);
+		CSV_SCOPED_TIMING_STAT_EXCLUSIVE(RenderOther);
 
-		for (FRDGPassHandle PassHandle = Passes.Begin(); PassHandle != Passes.End(); ++PassHandle)
+		if (FRDGPass* AsyncComputeProloguePass = GetProloguePass(ERHIPipeline::AsyncCompute))
 		{
-			if (!PassesToCull[PassHandle])
+			ExecutePass(AsyncComputeProloguePass);
+		}
+
+		for (FRDGPassHandle PassHandle = ProloguePassHandle; PassHandle <= EpiloguePassHandle; ++PassHandle)
+		{
+			FRDGPass* Pass = Passes[PassHandle];
+
+			if (!Pass->bCulled)
 			{
-				ExecutePass(Passes[PassHandle]);
+				ExecutePass(Pass);
 			}
-		}
 
-		IF_RDG_ENABLE_DEBUG(LogFile.End());
+		#if STATS
+			GRDGStatPassCullCount += Pass->bCulled;
+		#endif
+		}
 	}
 	else
 	{
 		ExecutePass(EpiloguePass);
 	}
 
-	RHICmdList.SetStaticUniformBuffers({});
-
-#if WITH_MGPU
-	if (NameForTemporalEffect != NAME_None)
+	if (bFinalExecute)
 	{
-		TArray<FRHITexture*> BroadcastTexturesForTemporalEffect;
-		for (const auto& Query : ExtractedTextures)
+		RHICmdList.SetStaticUniformBuffers({});
+
+	#if WITH_MGPU
+		if (NameForTemporalEffect != NAME_None)
 		{
-			if (EnumHasAnyFlags(Query.Key->Flags, ERDGTextureFlags::MultiFrame))
+			TArray<FRHITexture*> BroadcastTexturesForTemporalEffect;
+			for (const auto& Query : ExtractedTextures)
 			{
-				BroadcastTexturesForTemporalEffect.Add(Query.Key->GetRHIUnchecked());
+				if (EnumHasAnyFlags(Query.Key->Flags, ERDGTextureFlags::MultiFrame))
+				{
+					BroadcastTexturesForTemporalEffect.Add(Query.Key->GetRHIUnchecked());
+				}
 			}
+			RHICmdList.BroadcastTemporalEffect(NameForTemporalEffect, BroadcastTexturesForTemporalEffect);
 		}
-		RHICmdList.BroadcastTemporalEffect(NameForTemporalEffect, BroadcastTexturesForTemporalEffect);
-	}
-#endif
+	#endif
 
-	for (const auto& Query : ExtractedTextures)
-	{
-		*Query.Value = Query.Key->PooledRenderTarget;
-	}
+		for (const auto& Query : ExtractedTextures)
+		{
+			*Query.Value = Query.Key->PooledRenderTarget;
+		}
 
-	for (const auto& Query : ExtractedBuffers)
-	{
-		*Query.Value = Query.Key->PooledBuffer;
-	}
+		for (const auto& Query : ExtractedBuffers)
+		{
+			*Query.Value = Query.Key->PooledBuffer;
+		}
 
-	IF_RDG_ENABLE_TRACE(Trace.OutputGraphEnd(*this));
-	IF_RDG_GPU_SCOPES(GPUScopeStacks.Graphics.EndExecute());
-	IF_RDG_CPU_SCOPES(CPUScopeStacks.EndExecute());
-	IF_RDG_ENABLE_DEBUG(UserValidation.ValidateExecuteEnd());
+		IF_RDG_ENABLE_TRACE(Trace.OutputGraphEnd(*this));
+		IF_RDG_GPU_SCOPES(GPUScopeStacks.Graphics.EndExecute());
+		IF_RDG_CPU_SCOPES(CPUScopeStacks.EndExecute());
+		IF_RDG_ENABLE_DEBUG(UserValidation.ValidateExecuteEnd());
+		IF_RDG_ENABLE_DEBUG(LogFile.End());
 
-#if STATS
-	GRDGStatPassCount += Passes.Num();
-	GRDGStatPassWithParameterCount = GRDGStatPassCount;
-	for (FRDGPassHandle PassHandle = Passes.Begin(); PassHandle != Passes.End(); ++PassHandle)
+	#if STATS
+		GRDGStatBufferCount += Buffers.Num();
+		GRDGStatTextureCount += Textures.Num();
+		GRDGStatViewCount += Views.Num();
+		GRDGStatMemoryWatermark = FMath::Max(GRDGStatMemoryWatermark, Allocator.GetByteCount());
+	#endif
+
+		Clear();
+	}
+	else
 	{
-		GRDGStatPassWithParameterCount -= PassesWithEmptyParameters[PassHandle] ? 1 : 0;
+		AddProloguePass();
 	}
-	GRDGStatBufferCount += Buffers.Num();
-	GRDGStatTextureCount += Textures.Num();
-	GRDGStatViewCount += Views.Num();
-	GRDGStatMemoryWatermark = FMath::Max(GRDGStatMemoryWatermark, Allocator.GetByteCount());
-#endif
 
-	Clear();
+	ExecuteCount++;
 }
 
 void FRDGBuilder::Clear()
 {
 	TRACE_CPUPROFILER_EVENT_SCOPE(FRDGBuilder::Clear);
-	ExternalTextures.Empty();
-	ExternalBuffers.Empty();
-	ExtractedTextures.Empty();
-	ExtractedBuffers.Empty();
 	Passes.Clear();
-	Views.Clear();
-	Textures.Clear();
-	Buffers.Clear();
 	UniformBuffers.Clear();
 	Blackboard.Clear();
+	ActivePooledTextures.Empty();
+	ActivePooledBuffers.Empty();
 }
 
-void FRDGBuilder::SetupPass(FRDGPass* Pass)
+FRDGPass* FRDGBuilder::SetupPass(FRDGPass* Pass)
 {
 	IF_RDG_ENABLE_DEBUG(UserValidation.ValidateAddPass(Pass, bInDebugPassScope));
-	SCOPE_CYCLE_COUNTER(STAT_RDG_SetupTime);
-	CSV_SCOPED_TIMING_STAT_EXCLUSIVE(RDG_SetupPass);
+	CSV_SCOPED_TIMING_STAT_EXCLUSIVE_CONDITIONAL(RDGBuilder_SetupPass, GRDGVerboseCSVStats != 0);
 
 	const FRDGParameterStruct PassParameters = Pass->GetParameters();
-	const FRDGPassHandle PassHandle = Pass->GetHandle();
-	const ERDGPassFlags PassFlags = Pass->GetFlags();
-	const ERHIPipeline PassPipeline = Pass->GetPipeline();
+	const FRDGPassHandle PassHandle = Pass->Handle;
+	const ERDGPassFlags PassFlags = Pass->Flags;
+	const ERHIPipeline PassPipeline = Pass->Pipeline;
 
 	bool bRenderPassOnlyWrites = true;
 
@@ -1670,46 +1775,60 @@ void FRDGBuilder::SetupPass(FRDGPass* Pass)
 
 	PassParameters.EnumerateUniformBuffers([&](FRDGUniformBufferBinding UniformBuffer)
 	{
-		UniformBuffersToCreate[UniformBuffer->Handle] = true;
+		if (!UniformBuffer->bQueuedForCreate)
+		{
+			UniformBuffersToCreate.Emplace(UniformBuffer->Handle);
+			UniformBuffer->bQueuedForCreate = true;
+		}
 	});
 
 	Pass->bRenderPassOnlyWrites = bRenderPassOnlyWrites;
 	Pass->bHasExternalOutputs = PassParameters.HasExternalOutputs();
 
 	const bool bEmptyParameters = !Pass->TextureStates.Num() && !Pass->BufferStates.Num();
-	PassesWithEmptyParameters.Add(bEmptyParameters);
-
-	SetupPassInternal(Pass, PassHandle, PassPipeline);
+	SetupPassInternal(Pass, PassHandle, PassPipeline, bEmptyParameters);
+	return Pass;
 }
 
-void FRDGBuilder::SetupEmptyPass(FRDGPass* Pass)
+FRDGPass* FRDGBuilder::SetupEmptyPass(FRDGPass* Pass)
 {
-	PassesWithEmptyParameters.Add(true);
-	SetupPassInternal(Pass, Pass->GetHandle(), ERHIPipeline::Graphics);
+	const bool bEmptyParameters = true;
+	SetupPassInternal(Pass, Pass->Handle, Pass->Pipeline, bEmptyParameters);
+	return Pass;
 }
 
-void FRDGBuilder::SetupPassInternal(FRDGPass* Pass, FRDGPassHandle PassHandle, ERHIPipeline PassPipeline)
+void FRDGBuilder::SetupPassInternal(FRDGPass* Pass, FRDGPassHandle PassHandle, ERHIPipeline PassPipeline, bool bEmptyParameters)
 {
-	check(Pass->GetHandle() == PassHandle);
-	check(Pass->GetPipeline() == PassPipeline);
+	check(Pass->Handle == PassHandle);
+	check(Pass->Pipeline == PassPipeline);
 
-	Pass->GraphicsJoinPass = PassHandle;
+	Pass->bEmptyParameters = bEmptyParameters;
+	Pass->bDispatchAfterExecute = bDispatchHint;
 	Pass->GraphicsForkPass = PassHandle;
+	Pass->GraphicsJoinPass = PassHandle;
 	Pass->PrologueBarrierPass = PassHandle;
 	Pass->EpilogueBarrierPass = PassHandle;
 
-	if (!EnumHasAnyFlags(Pass->GetFlags(), ERDGPassFlags::AsyncCompute))
+	bDispatchHint = false;
+
+	if (!EnumHasAnyFlags(Pass->Flags, ERDGPassFlags::AsyncCompute))
 	{
-		Pass->ResourcesToBegin = { Pass };
-		Pass->ResourcesToEnd   = { Pass };
+		Pass->ResourcesToBegin.Add(Pass);
+		Pass->ResourcesToEnd.Add(Pass);
 	}
 
+	AsyncComputePassCount += EnumHasAnyFlags(Pass->Flags, ERDGPassFlags::AsyncCompute) ? 1 : 0;
+	RasterPassCount       += EnumHasAnyFlags(Pass->Flags, ERDGPassFlags::Raster)       ? 1 : 0;
+
 #if WITH_MGPU
 	Pass->GPUMask = RHICmdList.GetGPUMask();
 #endif
 
 #if STATS
-	Pass->CommandListStat = CommandListStat;
+	Pass->CommandListStat = CommandListStatScope;
+
+	GRDGStatPassCount++;
+	GRDGStatPassWithParameterCount += !bEmptyParameters ? 1 : 0;
 #endif
 
 	IF_RDG_CPU_SCOPES(Pass->CPUScopes = CPUScopeStacks.GetCurrentScopes());
@@ -1726,7 +1845,7 @@ void FRDGBuilder::SetupPassInternal(FRDGPass* Pass, FRDGPassHandle PassHandle, E
 	}
 #endif
 
-	if (IsImmediateMode() && Pass != EpiloguePass)
+	if (IsImmediateMode() && !Pass->bSentinel)
 	{
 		RDG_ALLOW_RHI_ACCESS_SCOPE();
 
@@ -1751,35 +1870,58 @@ void FRDGBuilder::SetupPassInternal(FRDGPass* Pass, FRDGPassHandle PassHandle, E
 
 		check(!EnumHasAnyFlags(PassPipeline, ERHIPipeline::AsyncCompute));
 
-		// Transient allocator should not be initialized in immediate mode.
-		{
-			FTransientResourceAllocator TransientResourceAllocator(RHICmdList);
-			CollectPassResources(TransientResourceAllocator, PassHandle);
-			check(!TransientResourceAllocator);
-		}
-
-		Pass->GetParameters().EnumerateUniformBuffers([&](FRDGUniformBufferBinding UniformBuffer)
-		{
-			BeginResourceRHI(UniformBuffer.GetUniformBuffer());
-		});
-
-		CollectPassBarriers(PassHandle);
+		BeginResourcesRHI(Pass, PassHandle);
+		EndResourcesRHI(Pass, PassHandle);
+		CollectPassBarriers(Pass, PassHandle);
+		CreatePassBarriers();
+		CreateUniformBuffers();
 		ExecutePass(Pass);
 	}
 
 	IF_RDG_ENABLE_DEBUG(VisualizePassOutputs(Pass));
 }
 
+void FRDGBuilder::CreateUniformBuffers()
+{
+	TRACE_CPUPROFILER_EVENT_SCOPE(FRDGBuilder::CreateUniformBuffers);
+
+	for (FRDGUniformBufferHandle UniformBufferHandle : UniformBuffersToCreate)
+	{
+		UniformBuffers[UniformBufferHandle]->InitRHI();
+	}
+	UniformBuffersToCreate.Reset();
+}
+
+void FRDGBuilder::AddProloguePass()
+{
+	if (!PassesOnAsyncComputeToJoin.IsEmpty())
+	{
+		FRDGPass* AsyncComputePass = SetupEmptyPass(Passes.Allocate<FRDGSentinelPass>(Allocator, RDG_EVENT_NAME("Graph Prologue (AsyncCompute)"), ERDGPassFlags::AsyncCompute));
+		ProloguePasses[ERHIPipeline::AsyncCompute] = AsyncComputePass;
+		ProloguePassHandles[ERHIPipeline::AsyncCompute] = AsyncComputePass->Handle;
+	}
+
+	FRDGPass* GraphicsPass = SetupEmptyPass(Passes.Allocate<FRDGSentinelPass>(Allocator, RDG_EVENT_NAME("Graph Prologue (Graphics)")));
+	ProloguePasses[ERHIPipeline::Graphics] = GraphicsPass;
+	ProloguePassHandles[ERHIPipeline::Graphics] = GraphicsPass->Handle;
+}
+
 void FRDGBuilder::ExecutePassPrologue(FRHIComputeCommandList& RHICmdListPass, FRDGPass* Pass)
 {
-	QUICK_SCOPE_CYCLE_COUNTER(STAT_FRDGBuilder_ExecutePassPrologue);
 	CSV_SCOPED_TIMING_STAT_EXCLUSIVE_CONDITIONAL(RDGBuilder_ExecutePassPrologue, GRDGVerboseCSVStats != 0);
 
 	IF_RDG_ENABLE_DEBUG(UserValidation.ValidateExecutePassBegin(Pass));
-	IF_RDG_CMDLIST_STATS(RHICmdList.SetCurrentStat(Pass->CommandListStat));
 
-	const ERDGPassFlags PassFlags = Pass->GetFlags();
-	const ERHIPipeline PassPipeline = Pass->GetPipeline();
+#if RDG_CMDLIST_STATS
+	if (!Pass->bSentinel && CommandListStatState != Pass->CommandListStat)
+	{
+		CommandListStatState = Pass->CommandListStat;
+		RHICmdList.SetCurrentStat(Pass->CommandListStat);
+	}
+#endif
+
+	const ERDGPassFlags PassFlags = Pass->Flags;
+	const ERHIPipeline PassPipeline = Pass->Pipeline;
 
 	if (Pass->PrologueBarriersToBegin)
 	{
@@ -1787,14 +1929,12 @@ void FRDGBuilder::ExecutePassPrologue(FRHIComputeCommandList& RHICmdListPass, FR
 		Pass->PrologueBarriersToBegin->Submit(RHICmdListPass, PassPipeline);
 	}
 
-	if (Pass->PrologueBarriersToEnd)
-	{
-		IF_RDG_ENABLE_DEBUG(BarrierValidation.ValidateBarrierBatchEnd(Pass, *Pass->PrologueBarriersToEnd));
-		Pass->PrologueBarriersToEnd->Submit(RHICmdListPass, PassPipeline);
-	}
+	IF_RDG_ENABLE_DEBUG(BarrierValidation.ValidateBarrierBatchEnd(Pass, Pass->PrologueBarriersToEnd));
+	Pass->PrologueBarriersToEnd.Submit(RHICmdListPass, PassPipeline);
 
-	if (PassPipeline == ERHIPipeline::AsyncCompute)
+	if (PassPipeline == ERHIPipeline::AsyncCompute && !Pass->bSentinel && AsyncComputeBudgetState != Pass->AsyncComputeBudget)
 	{
+		AsyncComputeBudgetState = Pass->AsyncComputeBudget;
 		RHICmdListPass.SetAsyncComputeBudget(Pass->AsyncComputeBudget);
 	}
 
@@ -1811,13 +1951,12 @@ void FRDGBuilder::ExecutePassPrologue(FRHIComputeCommandList& RHICmdListPass, FR
 
 void FRDGBuilder::ExecutePassEpilogue(FRHIComputeCommandList& RHICmdListPass, FRDGPass* Pass)
 {
-	QUICK_SCOPE_CYCLE_COUNTER(STAT_FRDGBuilder_ExecutePassEpilogue);
 	CSV_SCOPED_TIMING_STAT_EXCLUSIVE_CONDITIONAL(RDGBuilder_ExecutePassEpilogue, GRDGVerboseCSVStats != 0);
 
 	EndUAVOverlap(Pass, RHICmdListPass);
 
-	const ERDGPassFlags PassFlags = Pass->GetFlags();
-	const ERHIPipeline PassPipeline = Pass->GetPipeline();
+	const ERDGPassFlags PassFlags = Pass->Flags;
+	const ERHIPipeline PassPipeline = Pass->Pipeline;
 	const FRDGParameterStruct PassParameters = Pass->GetParameters();
 
 	if (EnumHasAnyFlags(PassFlags, ERDGPassFlags::Raster) && !EnumHasAnyFlags(PassFlags, ERDGPassFlags::SkipRenderPass) && !Pass->SkipRenderPassEnd())
@@ -1827,11 +1966,8 @@ void FRDGBuilder::ExecutePassEpilogue(FRHIComputeCommandList& RHICmdListPass, FR
 
 	FRDGTransitionQueue Transitions;
 
-	if (Pass->EpilogueBarriersToBeginForGraphics)
-	{
-		IF_RDG_ENABLE_DEBUG(BarrierValidation.ValidateBarrierBatchBegin(Pass, *Pass->EpilogueBarriersToBeginForGraphics));
-		Pass->EpilogueBarriersToBeginForGraphics->Submit(RHICmdListPass, PassPipeline, Transitions);
-	}
+	IF_RDG_ENABLE_DEBUG(BarrierValidation.ValidateBarrierBatchBegin(Pass, Pass->EpilogueBarriersToBeginForGraphics));
+	Pass->EpilogueBarriersToBeginForGraphics.Submit(RHICmdListPass, PassPipeline, Transitions);
 
 	if (Pass->EpilogueBarriersToBeginForAsyncCompute)
 	{
@@ -1851,7 +1987,10 @@ void FRDGBuilder::ExecutePassEpilogue(FRHIComputeCommandList& RHICmdListPass, FR
 		BarriersToBegin->Submit(RHICmdListPass, PassPipeline, Transitions);
 	}
 
-	Transitions.Begin(RHICmdListPass);
+	if (!Transitions.IsEmpty())
+	{
+		RHICmdListPass.BeginTransitions(Transitions);
+	}
 
 	if (Pass->EpilogueBarriersToEnd)
 	{
@@ -1864,8 +2003,6 @@ void FRDGBuilder::ExecutePassEpilogue(FRHIComputeCommandList& RHICmdListPass, FR
 
 void FRDGBuilder::ExecutePass(FRDGPass* Pass)
 {
-	QUICK_SCOPE_CYCLE_COUNTER(STAT_FRDGBuilder_ExecutePass);
-
 	// Note that we must do this before doing anything with RHICmdList for the pass.
 	// For example, if this pass only executes on GPU 1 we want to avoid adding a
 	// 0-duration event for this pass on GPU 0's time line.
@@ -1885,15 +2022,14 @@ void FRDGBuilder::ExecutePass(FRDGPass* Pass)
 
 	// Execute the pass by invoking the prologue, then the pass body, then the epilogue.
 	// The entire pass is executed using the command list on the specified pipeline.
-	FRHIComputeCommandList& RHICmdListPass = (Pass->GetPipeline() == ERHIPipeline::AsyncCompute)
+	FRHIComputeCommandList& RHICmdListPass = (Pass->Pipeline == ERHIPipeline::AsyncCompute)
 		? static_cast<FRHIComputeCommandList&>(RHICmdListAsyncCompute)
 		: RHICmdList;
 
 	ExecutePassPrologue(RHICmdListPass, Pass);
 
 #if RDG_GPU_SCOPES
-	const bool bUsePassEventScope = Pass != EpiloguePass && Pass != ProloguePass;
-	if (bUsePassEventScope)
+	if (!Pass->bSentinel)
 	{
 		GPUScopeStacks.BeginExecutePass(Pass);
 	}
@@ -1902,7 +2038,7 @@ void FRDGBuilder::ExecutePass(FRDGPass* Pass)
 	Pass->Execute(RHICmdListPass);
 
 #if RDG_GPU_SCOPES
-	if (bUsePassEventScope)
+	if (!Pass->bSentinel)
 	{
 		GPUScopeStacks.EndExecutePass(Pass);
 	}
@@ -1924,53 +2060,55 @@ void FRDGBuilder::ExecutePass(FRDGPass* Pass)
 		RHICmdList.SubmitCommandsAndFlushGPU();
 		RHICmdList.BlockUntilGPUIdle();
 	}
+
+	if (Pass->bDispatchAfterExecute && IsRunningRHIInSeparateThread())
+	{
+		RHICmdList.SubmitCommandsHint();
+		FTaskGraphInterface::Get().ProcessThreadUntilIdle(ENamedThreads::GetRenderThread_Local());
+		RHICmdList.ImmediateFlush(EImmediateFlushType::DispatchToRHIThread);
+	}
 }
 
-void FRDGBuilder::CollectPassResources(FTransientResourceAllocator& TransientResourceAllocator, FRDGPassHandle PassHandle)
+void FRDGBuilder::BeginResourcesRHI(FRDGPass* ResourcePass, FRDGPassHandle ExecutePassHandle)
 {
-	FRDGPass* Pass = Passes[PassHandle];
-
-	for (FRDGPass* PassToBegin : Pass->ResourcesToBegin)
+	for (FRDGPass* PassToBegin : ResourcePass->ResourcesToBegin)
 	{
 		for (const auto& PassState : PassToBegin->TextureStates)
 		{
-			BeginResourceRHI(&TransientResourceAllocator, PassHandle, PassState.Texture);
+			BeginResourceRHI(ExecutePassHandle, PassState.Texture);
 		}
 
 		for (const auto& PassState : PassToBegin->BufferStates)
 		{
-			BeginResourceRHI(&TransientResourceAllocator, PassHandle, PassState.Buffer);
+			BeginResourceRHI(ExecutePassHandle, PassState.Buffer);
 		}
 
 		for (FRDGViewHandle ViewHandle : PassToBegin->Views)
 		{
-			BeginResourceRHI(PassHandle, Views[ViewHandle]);
+			BeginResourceRHI(ExecutePassHandle, Views[ViewHandle]);
 		}
 	}
+}
 
-	for (FRDGPass* PassToEnd : Pass->ResourcesToEnd)
+void FRDGBuilder::EndResourcesRHI(FRDGPass* ResourcePass, FRDGPassHandle ExecutePassHandle)
+{
+	for (FRDGPass* PassToEnd : ResourcePass->ResourcesToEnd)
 	{
 		for (const auto& PassState : PassToEnd->TextureStates)
 		{
-			EndResourceRHI(TransientResourceAllocator, PassHandle, PassState.Texture, PassState.ReferenceCount);
+			EndResourceRHI(ExecutePassHandle, PassState.Texture, PassState.ReferenceCount);
 		}
 
 		for (const auto& PassState : PassToEnd->BufferStates)
 		{
-			EndResourceRHI(TransientResourceAllocator, PassHandle, PassState.Buffer, PassState.ReferenceCount);
+			EndResourceRHI(ExecutePassHandle, PassState.Buffer, PassState.ReferenceCount);
 		}
 	}
 }
 
-void FRDGBuilder::CollectPassBarriers(FRDGPassHandle PassHandle)
+void FRDGBuilder::CollectPassBarriers(FRDGPass* Pass, FRDGPassHandle PassHandle)
 {
-	FRDGPass* Pass = Passes[PassHandle];
-
 	IF_RDG_ENABLE_DEBUG(ConditionalDebugBreak(RDG_BREAKPOINT_PASS_COMPILE, BuilderName.GetTCHAR(), Pass->GetName()));
-	if (PassesWithEmptyParameters[PassHandle])
-	{
-		return;
-	}
 
 	for (const auto& PassState : Pass->TextureStates)
 	{
@@ -1991,6 +2129,17 @@ void FRDGBuilder::CollectPassBarriers(FRDGPassHandle PassHandle)
 	}
 }
 
+void FRDGBuilder::CreatePassBarriers()
+{
+	TRACE_CPUPROFILER_EVENT_SCOPE(FRDGBuilder::CreatePassBarriers);
+
+	for (FRDGBarrierBatchBegin* BarrierBatchBegin : TransitionCreateQueue)
+	{
+		BarrierBatchBegin->CreateTransition();
+	}
+	TransitionCreateQueue.Reset();
+}
+
 void FRDGBuilder::AddEpilogueTransition(FRDGTextureRef Texture)
 {
 	if (!Texture->bLastOwner || Texture->bCulled || Texture->bFinalizedAccess)
@@ -2006,7 +2155,7 @@ void FRDGBuilder::AddEpilogueTransition(FRDGTextureRef Texture)
 	// Texture is using the RHI transient allocator. Transition it back to Discard in the final pass it is used.
 	if (Texture->bTransient)
 	{
-		ScratchSubresourceState.SetPass(ERHIPipeline::Graphics, Texture->LastPass);
+		ScratchSubresourceState.SetPass(ERHIPipeline::Graphics, ClampToPrologue(Texture->LastPass));
 		ScratchSubresourceState.Access = ERHIAccess::Discard;
 		InitAsWholeResource(ScratchTextureState, &ScratchSubresourceState);
 
@@ -2065,7 +2214,7 @@ void FRDGBuilder::AddEpilogueTransition(FRDGBufferRef Buffer)
 	if (Buffer->bTransient)
 	{
 		FRDGSubresourceState StateFinal;
-		StateFinal.SetPass(ERHIPipeline::Graphics, Buffer->LastPass);
+		StateFinal.SetPass(ERHIPipeline::Graphics, ClampToPrologue(Buffer->LastPass));
 		StateFinal.Access = ERHIAccess::Discard;
 		AddTransition(Buffer->LastPass, Buffer, StateFinal, EBarrierLocation::Epilogue);
 	}
@@ -2128,7 +2277,7 @@ void FRDGBuilder::AddTransition(FRDGPassHandle PassHandle, FRDGTextureRef Textur
 				Info.PlaneSlice = Subresource->PlaneSlice;
 			}
 
-			AddTransition(Texture, SubresourceStateBefore, SubresourceStateAfter, BarrierLocation, Info);
+			AddTransition(Texture, SubresourceStateBefore, SubresourceStateAfter, Info, BarrierLocation);
 		}
 
 		if (Subresource)
@@ -2229,7 +2378,7 @@ void FRDGBuilder::AddTransition(FRDGPassHandle PassHandle, FRDGBufferRef Buffer,
 		Info.AccessBefore = StateBefore.Access;
 		Info.AccessAfter = StateAfter.Access;
 
-		AddTransition(Buffer, StateBefore, StateAfter, BarrierLocation, Info);
+		AddTransition(Buffer, StateBefore, StateAfter, Info, BarrierLocation);
 	}
 
 	IF_RDG_ENABLE_DEBUG(LogFile.AddTransitionEdge(PassHandle, StateBefore, StateAfter, Buffer));
@@ -2240,8 +2389,8 @@ void FRDGBuilder::AddTransition(
 	FRDGParentResource* Resource,
 	FRDGSubresourceState StateBefore,
 	FRDGSubresourceState StateAfter,
-	EBarrierLocation BarrierLocation,
-	const FRHITransitionInfo& TransitionInfo)
+	const FRHITransitionInfo& TransitionInfo,
+	EBarrierLocation BarrierLocation)
 {
 	const ERHIPipeline Graphics = ERHIPipeline::Graphics;
 	const ERHIPipeline AsyncCompute = ERHIPipeline::AsyncCompute;
@@ -2261,6 +2410,8 @@ void FRDGBuilder::AddTransition(
 		return;
 	}
 
+	StateBefore.LastPass = ClampToPrologue(StateBefore.LastPass);
+
 	ERHIPipeline PipelinesBefore = StateBefore.GetPipelines();
 	ERHIPipeline PipelinesAfter = StateAfter.GetPipelines();
 
@@ -2275,14 +2426,11 @@ void FRDGBuilder::AddTransition(
 	checkf(StateBefore.GetLastPass() <= StateAfter.GetFirstPass(), TEXT("Submitted a state for '%s' that begins before our previous state has ended."), Resource->Name);
 	check(!Resource->bTransient || StateBefore.GetLastPass() >= Resource->FirstPass);
 
-	FRDGPassHandlesByPipeline& PassesBefore = StateBefore.LastPass;
-	FRDGPassHandlesByPipeline& PassesAfter  = StateAfter.FirstPass;
-
-	const uint32 PipelinesBeforeCount = FMath::CountBits(uint64(PipelinesBefore));
-	const uint32 PipelinesAfterCount = FMath::CountBits(uint64(PipelinesAfter));
+	const FRDGPassHandlesByPipeline& PassesBefore = StateBefore.LastPass;
+	const FRDGPassHandlesByPipeline& PassesAfter = StateAfter.FirstPass;
 
 	// 1-to-1 or 1-to-N pipe transition.
-	if (PipelinesBeforeCount == 1)
+	if (PipelinesBefore != ERHIPipeline::All)
 	{
 		const FRDGPassHandle BeginPassHandle = StateBefore.GetLastPass();
 		const FRDGPassHandle FirstEndPassHandle = StateAfter.GetFirstPass();
@@ -2294,7 +2442,7 @@ void FRDGBuilder::AddTransition(
 		if (BeginPassHandle < FirstEndPassHandle || BarrierLocation == EBarrierLocation::Epilogue)
 		{
 			BeginPass = GetEpilogueBarrierPass(BeginPassHandle);
-			BarriersToBegin = &BeginPass->GetEpilogueBarriersToBeginFor(Allocator, PipelinesAfter);
+			BarriersToBegin = &BeginPass->GetEpilogueBarriersToBeginFor(Allocator, TransitionCreateQueue, PipelinesAfter);
 		}
 		// This is an immediate prologue transition in the same pass. Issue the begin in the prologue.
 		else
@@ -2304,7 +2452,7 @@ void FRDGBuilder::AddTransition(
 				Resource->Name, *GetRHIPipelineName(PipelinesAfter));
 
 			BeginPass = GetPrologueBarrierPass(BeginPassHandle);
-			BarriersToBegin = &BeginPass->GetPrologueBarriersToBegin(Allocator);
+			BarriersToBegin = &BeginPass->GetPrologueBarriersToBegin(Allocator, TransitionCreateQueue);
 		}
 
 		BarriersToBegin->AddTransition(Resource, TransitionInfo);
@@ -2327,7 +2475,7 @@ void FRDGBuilder::AddTransition(
 			 *                                   /            \
 			 *  Graphics Pipe:            BeginA  EndA         EndB
 			 */
-			if ((PipelinesBefore == Pipeline && PipelinesAfterCount > 1))
+			if ((PipelinesBefore == Pipeline && PipelinesAfter == ERHIPipeline::All))
 			{
 				AddToEpilogueBarriersToEnd(BeginPassHandle, *BarriersToBegin);
 			}
@@ -2362,11 +2510,16 @@ void FRDGBuilder::AddTransition(
 
 		if (!BarriersToBegin)
 		{
-			BarriersToBegin = Allocator.AllocNoDestruct<FRDGBarrierBatchBegin>(PipelinesBefore, PipelinesAfter, GetEpilogueBarriersToBeginDebugName(PipelinesAfter), Id.Passes);
+			FRDGPassesByPipeline BarrierBatchPasses;
+			BarrierBatchPasses[Graphics]     = Passes[Id.Passes[Graphics]];
+			BarrierBatchPasses[AsyncCompute] = Passes[Id.Passes[AsyncCompute]];
 
-			for (FRDGPassHandle PassHandle : Id.Passes)
+			BarriersToBegin = Allocator.AllocNoDestruct<FRDGBarrierBatchBegin>(PipelinesBefore, PipelinesAfter, GetEpilogueBarriersToBeginDebugName(PipelinesAfter), BarrierBatchPasses);
+			TransitionCreateQueue.Emplace(BarriersToBegin);
+
+			for (FRDGPass* Pass : BarrierBatchPasses)
 			{
-				Passes[PassHandle]->SharedEpilogueBarriersToBegin.Add(BarriersToBegin);
+				Pass->SharedEpilogueBarriersToBegin.Add(BarriersToBegin);
 			}
 		}
 
@@ -2389,38 +2542,15 @@ void FRDGBuilder::AddTransition(
 	}
 }
 
-void FRDGBuilder::BeginResourceRHI(FRDGUniformBuffer* UniformBuffer)
-{
-	check(UniformBuffer);
-
-	if (UniformBuffer->GetRHIUnchecked())
-	{
-		return;
-	}
-
-	const FRDGParameterStruct PassParameters = UniformBuffer->GetParameters();
-
-	const EUniformBufferValidation Validation =
-#if RDG_ENABLE_DEBUG
-		EUniformBufferValidation::ValidateResources;
-#else
-		EUniformBufferValidation::None;
-#endif
-
-	UniformBuffer->UniformBufferRHI = RHICreateUniformBuffer(PassParameters.GetContents(), PassParameters.GetLayout(), UniformBuffer_SingleFrame, Validation);
-	UniformBuffer->ResourceRHI = UniformBuffer->UniformBufferRHI;
-}
-
-void FRDGBuilder::BeginResourceRHI(FTransientResourceAllocator* TransientResourceAllocator, FRDGPassHandle PassHandle, FRDGTextureRef Texture)
+void FRDGBuilder::BeginResourceRHI(FRDGPassHandle PassHandle, FRDGTextureRef Texture)
 {
 	check(Texture);
 
-	if (Texture->GetRHIUnchecked())
+	if (Texture->HasRHI())
 	{
 		return;
 	}
 
-	check(TransientResourceAllocator  || Texture->bExternal);
 	check(Texture->ReferenceCount > 0 || Texture->bExternal || IsImmediateMode());
 
 #if RDG_ENABLE_DEBUG
@@ -2436,9 +2566,7 @@ void FRDGBuilder::BeginResourceRHI(FTransientResourceAllocator* TransientResourc
 
 	if (IsTransient(Texture))
 	{
-		check(TransientResourceAllocator);
-
-		if (IRHITransientResourceAllocator* AllocatorRHI = TransientResourceAllocator->GetOrCreate())
+		if (IRHITransientResourceAllocator* AllocatorRHI = TransientResourceAllocator.GetOrCreate())
 		{
 			if (FRHITransientTexture* TransientTexture = AllocatorRHI->CreateTexture(Texture->Desc, Texture->Name))
 			{
@@ -2470,14 +2598,14 @@ void FRDGBuilder::BeginResourceRHI(FTransientResourceAllocator* TransientResourc
 
 	Texture->FirstPass = PassHandle;
 
-	check(Texture->GetRHIUnchecked());
+	check(Texture->HasRHI());
 }
 
 void FRDGBuilder::BeginResourceRHI(FRDGPassHandle PassHandle, FRDGTextureSRVRef SRV)
 {
 	check(SRV);
 
-	if (SRV->GetRHIUnchecked())
+	if (SRV->HasRHI())
 	{
 		return;
 	}
@@ -2493,7 +2621,7 @@ void FRDGBuilder::BeginResourceRHI(FRDGPassHandle PassHandle, FRDGTextureUAVRef
 {
 	check(UAV);
 
-	if (UAV->GetRHIUnchecked())
+	if (UAV->HasRHI())
 	{
 		return;
 	}
@@ -2505,24 +2633,21 @@ void FRDGBuilder::BeginResourceRHI(FRDGPassHandle PassHandle, FRDGTextureUAVRef
 	UAV->ResourceRHI = Texture->ViewCache->GetOrCreateUAV(TextureRHI, UAV->Desc);
 }
 
-void FRDGBuilder::BeginResourceRHI(FTransientResourceAllocator* TransientResourceAllocator, FRDGPassHandle PassHandle, FRDGBufferRef Buffer)
+void FRDGBuilder::BeginResourceRHI(FRDGPassHandle PassHandle, FRDGBufferRef Buffer)
 {
 	check(Buffer);
 
-	if (Buffer->GetRHIUnchecked())
+	if (Buffer->HasRHI())
 	{
 		return;
 	}
 
-	check(TransientResourceAllocator || Buffer->bExternal);
 	check(Buffer->ReferenceCount > 0 || Buffer->bExternal || IsImmediateMode());
 
 	// If transient then create the resource on the transient allocator. External or extracted resource can't be transient because of lifetime tracking issues.
 	if (IsTransient(Buffer))
 	{
-		check(TransientResourceAllocator);
-
-		if (IRHITransientResourceAllocator* AllocatorRHI = TransientResourceAllocator->GetOrCreate())
+		if (IRHITransientResourceAllocator* AllocatorRHI = TransientResourceAllocator.GetOrCreate())
 		{
 			if (FRHITransientBuffer* TransientBuffer = AllocatorRHI->CreateBuffer(Translate(Buffer->Desc), Buffer->Name))
 			{
@@ -2553,14 +2678,14 @@ void FRDGBuilder::BeginResourceRHI(FTransientResourceAllocator* TransientResourc
 
 	Buffer->FirstPass = PassHandle;
 
-	check(Buffer->GetRHIUnchecked());
+	check(Buffer->HasRHI());
 }
 
 void FRDGBuilder::BeginResourceRHI(FRDGPassHandle PassHandle, FRDGBufferSRVRef SRV)
 {
 	check(SRV);
 
-	if (SRV->GetRHIUnchecked())
+	if (SRV->HasRHI())
 	{
 		return;
 	}
@@ -2586,7 +2711,7 @@ void FRDGBuilder::BeginResourceRHI(FRDGPassHandle PassHandle, FRDGBufferUAV* UAV
 {
 	check(UAV);
 
-	if (UAV->GetRHIUnchecked())
+	if (UAV->HasRHI())
 	{
 		return;
 	}
@@ -2609,7 +2734,7 @@ void FRDGBuilder::BeginResourceRHI(FRDGPassHandle PassHandle, FRDGBufferUAV* UAV
 
 void FRDGBuilder::BeginResourceRHI(FRDGPassHandle PassHandle, FRDGView* View)
 {
-	if (View->GetRHIUnchecked())
+	if (View->HasRHI())
 	{
 		return;
 	}
@@ -2631,7 +2756,7 @@ void FRDGBuilder::BeginResourceRHI(FRDGPassHandle PassHandle, FRDGView* View)
 	}
 }
 
-void FRDGBuilder::EndResourceRHI(FTransientResourceAllocator& TransientResourceAllocator, FRDGPassHandle PassHandle, FRDGTextureRef Texture, uint32 ReferenceCount)
+void FRDGBuilder::EndResourceRHI(FRDGPassHandle PassHandle, FRDGTextureRef Texture, uint32 ReferenceCount)
 {
 	check(Texture);
 	check(Texture->ReferenceCount >= ReferenceCount || IsImmediateMode());
@@ -2650,14 +2775,15 @@ void FRDGBuilder::EndResourceRHI(FTransientResourceAllocator& TransientResourceA
 		}
 		else if (Texture->PooledRenderTarget && Texture->PooledRenderTarget->IsTracked())
 		{
-			Texture->Allocation = nullptr;
+			// This releases the reference without invoking a virtual function call.
+			TRefCountPtr<FPooledRenderTarget>(MoveTemp(Texture->Allocation));
 		}
 
 		Texture->LastPass = PassHandle;
 	}
 }
 
-void FRDGBuilder::EndResourceRHI(FTransientResourceAllocator& TransientResourceAllocator, FRDGPassHandle PassHandle, FRDGBufferRef Buffer, uint32 ReferenceCount)
+void FRDGBuilder::EndResourceRHI(FRDGPassHandle PassHandle, FRDGBufferRef Buffer, uint32 ReferenceCount)
 {
 	check(Buffer);
 	check(Buffer->ReferenceCount >= ReferenceCount || IsImmediateMode());
diff --git a/Engine/Source/Runtime/RenderCore/Private/RenderGraphEvent.cpp b/Engine/Source/Runtime/RenderCore/Private/RenderGraphEvent.cpp
--- a/Engine/Source/Runtime/RenderCore/Private/RenderGraphEvent.cpp
+++ b/Engine/Source/Runtime/RenderCore/Private/RenderGraphEvent.cpp
@@ -381,12 +381,9 @@ static void OnPopCSVStat(FRHIComputeCommandList&, const FRDGCSVStatScope* Scope,
 #endif
 }
 
-FRDGCSVStatScopeStack::FRDGCSVStatScopeStack(FRHIComputeCommandList& InRHICmdList, FRDGAllocator& Allocator, const char* InUnaccountedStatName)
+FRDGCSVStatScopeStack::FRDGCSVStatScopeStack(FRHIComputeCommandList& InRHICmdList, FRDGAllocator& Allocator)
 	: ScopeStack(InRHICmdList, Allocator, &OnPushCSVStat, &OnPopCSVStat, GetEmitRDGEvents())
-	, UnaccountedStatName(InUnaccountedStatName)
-{
-	BeginScope(UnaccountedStatName);
-}
+{}
 
 bool FRDGCSVStatScopeStack::IsEnabled()
 {
@@ -417,7 +414,6 @@ void FRDGCSVStatScopeStack::BeginExecute()
 {
 	if (IsEnabled())
 	{
-		EndScope();
 		ScopeStack.BeginExecute();
 	}
 }
diff --git a/Engine/Source/Runtime/RenderCore/Private/RenderGraphPass.cpp b/Engine/Source/Runtime/RenderCore/Private/RenderGraphPass.cpp
--- a/Engine/Source/Runtime/RenderCore/Private/RenderGraphPass.cpp
+++ b/Engine/Source/Runtime/RenderCore/Private/RenderGraphPass.cpp
@@ -4,14 +4,6 @@
 #include "RenderGraphPrivate.h"
 #include "RenderGraphUtils.h"
 
-static FRDGPassHandlesByPipeline GetPassesByPipeline(const FRDGPass* Pass)
-{
-	check(Pass);
-	FRDGPassHandlesByPipeline Passes;
-	Passes[Pass->GetPipeline()] = Pass->GetHandle();
-	return Passes;
-}
-
 FUniformBufferStaticBindings FRDGParameterStruct::GetStaticUniformBuffers() const
 {
 	FUniformBufferStaticBindings GlobalUniformBuffers;
@@ -110,46 +102,22 @@ FRHIRenderPassInfo FRDGParameterStruct::GetRenderPassInfo() const
 	return RenderPassInfo;
 }
 
-uint32 GetTypeHash(FRDGBarrierBatchBeginId Id)
-{
-	static_assert(sizeof(Id.Passes) == 4, "Hash expects the Passes array to be 4 bytes (2 uint16's).");
-	uint32 Hash = *(const uint32*)Id.Passes.GetData();
-	return (Hash << GetRHIPipelineCount()) | uint32(Id.PipelinesAfter);
-}
-
-FRDGTransitionQueue::FRDGTransitionQueue(uint32 ReservedCount)
-{
-	Queue.Reserve(ReservedCount);
-}
-
-void FRDGTransitionQueue::Insert(const FRHITransition* Transition, ERHITransitionCreateFlags TransitionFlags)
-{
-	Queue.Add(Transition);
-}
-
-void FRDGTransitionQueue::Begin(FRHIComputeCommandList& RHICmdList)
-{
-	if (Queue.Num())
-	{
-		RHICmdList.BeginTransitions(Queue);
-		Queue.Empty();
-	}
-}
-
-void FRDGTransitionQueue::End(FRHIComputeCommandList& RHICmdList)
+FRDGBarrierBatchBegin::FRDGBarrierBatchBegin(ERHIPipeline InPipelineToBegin, ERHIPipeline InPipelinesToEnd, const TCHAR* InDebugName, FRDGPass* InDebugPass)
+	: PipelinesToBegin(InPipelineToBegin)
+	, PipelinesToEnd(InPipelinesToEnd)
+#if RDG_ENABLE_DEBUG
+	, DebugPasses(InPlace, nullptr)
+	, DebugName(InDebugName)
+	, DebugPipelinesToBegin(InPipelineToBegin)
+	, DebugPipelinesToEnd(InPipelinesToEnd)
+#endif
 {
-	if (Queue.Num())
-	{
-		RHICmdList.EndTransitions(Queue);
-		Queue.Empty();
-	}
+#if RDG_ENABLE_DEBUG
+	DebugPasses[InPipelineToBegin] = InDebugPass;
+#endif
 }
 
-FRDGBarrierBatchBegin::FRDGBarrierBatchBegin(ERHIPipeline InPipelinesToBegin, ERHIPipeline InPipelinesToEnd, const TCHAR* InDebugName, const FRDGPass* InDebugPass)
-	: FRDGBarrierBatchBegin(InPipelinesToBegin, InPipelinesToEnd, InDebugName, GetPassesByPipeline(InDebugPass))
-{}
-
-FRDGBarrierBatchBegin::FRDGBarrierBatchBegin(ERHIPipeline InPipelinesToBegin, ERHIPipeline InPipelinesToEnd, const TCHAR* InDebugName, FRDGPassHandlesByPipeline InDebugPasses)
+FRDGBarrierBatchBegin::FRDGBarrierBatchBegin(ERHIPipeline InPipelinesToBegin, ERHIPipeline InPipelinesToEnd, const TCHAR* InDebugName, FRDGPassesByPipeline InDebugPasses)
 	: PipelinesToBegin(InPipelinesToBegin)
 	, PipelinesToEnd(InPipelinesToEnd)
 #if RDG_ENABLE_DEBUG
@@ -158,15 +126,7 @@ FRDGBarrierBatchBegin::FRDGBarrierBatchBegin(ERHIPipeline InPipelinesToBegin, ER
 	, DebugPipelinesToBegin(InPipelinesToBegin)
 	, DebugPipelinesToEnd(InPipelinesToEnd)
 #endif
-{
-#if RDG_ENABLE_DEBUG
-	for (ERHIPipeline Pipeline : GetRHIPipelines())
-	{
-		// We should have provided corresponding debug passes to match the pipeline flags.
-		check(InDebugPasses[Pipeline].IsValid() == EnumHasAllFlags(InPipelinesToBegin, Pipeline));
-	}
-#endif
-}
+{}
 
 void FRDGBarrierBatchBegin::AddTransition(FRDGParentResourceRef Resource, const FRHITransitionInfo& Info)
 {
@@ -196,19 +156,19 @@ void FRDGBarrierBatchBegin::AddAlias(FRDGParentResourceRef Resource, const FRHIT
 #endif
 }
 
-void FRDGBarrierBatchBegin::Submit(FRHIComputeCommandList& RHICmdList, ERHIPipeline Pipeline, FRDGTransitionQueue& TransitionsToBegin)
+void FRDGBarrierBatchBegin::CreateTransition()
 {
-	// Submit may be called once for each pipe. The first to submit creates the transition.
-	if (!Transition && bTransitionNeeded)
-	{
-		Transition = RHICreateTransition(FRHITransitionCreateInfo(PipelinesToBegin, PipelinesToEnd, TransitionFlags, Transitions, Aliases));
-	}
+	check(bTransitionNeeded && !Transition);
+	Transition = RHICreateTransition(FRHITransitionCreateInfo(PipelinesToBegin, PipelinesToEnd, TransitionFlags, Transitions, Aliases));
+}
 
+void FRDGBarrierBatchBegin::Submit(FRHIComputeCommandList& RHICmdList, ERHIPipeline Pipeline, FRDGTransitionQueue& TransitionsToBegin)
+{
 	if (Transition)
 	{
 		check(EnumHasAnyFlags(PipelinesToBegin, Pipeline));
 		EnumRemoveFlags(PipelinesToBegin, Pipeline);
-		TransitionsToBegin.Insert(Transition, TransitionFlags);
+		TransitionsToBegin.Emplace(Transition);
 	}
 
 #if STATS
@@ -220,7 +180,11 @@ void FRDGBarrierBatchBegin::Submit(FRHIComputeCommandList& RHICmdList, ERHIPipel
 {
 	FRDGTransitionQueue TransitionsToBegin;
 	Submit(RHICmdList, Pipeline, TransitionsToBegin);
-	TransitionsToBegin.Begin(RHICmdList);
+
+	if (!TransitionsToBegin.IsEmpty())
+	{
+		RHICmdList.BeginTransitions(TransitionsToBegin);
+	}
 }
 
 void FRDGBarrierBatchEnd::AddDependency(FRDGBarrierBatchBegin* BeginBatch)
@@ -230,10 +194,10 @@ void FRDGBarrierBatchEnd::AddDependency(FRDGBarrierBatchBegin* BeginBatch)
 
 	for (ERHIPipeline Pipeline : GetRHIPipelines())
 	{
-		const FRDGPassHandle BeginPassHandle = BeginBatch->DebugPasses[Pipeline];
-		if (BeginPassHandle.IsValid())
+		const FRDGPass* BeginPass = BeginBatch->DebugPasses[Pipeline];
+		if (BeginPass)
 		{
-			checkf(BeginPassHandle <= PassHandle, TEXT("A transition end batch for pass %d is dependent on begin batch for pass %d."), PassHandle.GetIndex(), BeginPassHandle.GetIndex());
+			checkf(BeginPass->GetHandle() <= Pass->GetHandle(), TEXT("A transition end batch for pass %s is dependent on begin batch for pass %s."), Pass->GetName(), BeginPass->GetName());
 		}
 	}
 #endif
@@ -243,7 +207,8 @@ void FRDGBarrierBatchEnd::AddDependency(FRDGBarrierBatchBegin* BeginBatch)
 
 void FRDGBarrierBatchEnd::Submit(FRHIComputeCommandList& RHICmdList, ERHIPipeline Pipeline)
 {
-	FRDGTransitionQueue Transitions(Dependencies.Num());
+	FRDGTransitionQueue Transitions;
+	Transitions.Reserve(Dependencies.Num());
 
 	for (FRDGBarrierBatchBegin* Dependent : Dependencies)
 	{
@@ -251,63 +216,66 @@ void FRDGBarrierBatchEnd::Submit(FRHIComputeCommandList& RHICmdList, ERHIPipelin
 		{
 			check(Dependent->Transition);
 			EnumRemoveFlags(Dependent->PipelinesToEnd, Pipeline);
-			Transitions.Insert(Dependent->Transition, Dependent->TransitionFlags);
+			Transitions.Emplace(Dependent->Transition);
 		}
 	}
 
-	Transitions.End(RHICmdList);
+	if (!Transitions.IsEmpty())
+	{
+		RHICmdList.EndTransitions(Transitions);
+	}
 }
 
-FRDGBarrierBatchBegin& FRDGPass::GetPrologueBarriersToBegin(FRDGAllocator& Allocator)
+FRDGBarrierBatchBegin& FRDGPass::GetPrologueBarriersToBegin(FRDGAllocator& Allocator, FRDGTransitionCreateQueue& CreateQueue)
 {
 	if (!PrologueBarriersToBegin)
 	{
 		PrologueBarriersToBegin = Allocator.AllocNoDestruct<FRDGBarrierBatchBegin>(Pipeline, Pipeline, TEXT("Prologue"), this);
+		CreateQueue.Emplace(PrologueBarriersToBegin);
 	}
 	return *PrologueBarriersToBegin;
 }
 
-FRDGBarrierBatchBegin& FRDGPass::GetEpilogueBarriersToBeginForGraphics(FRDGAllocator& Allocator)
+FRDGBarrierBatchBegin& FRDGPass::GetEpilogueBarriersToBeginForGraphics(FRDGAllocator& Allocator, FRDGTransitionCreateQueue& CreateQueue)
 {
-	if (!EpilogueBarriersToBeginForGraphics)
+	if (!EpilogueBarriersToBeginForGraphics.IsTransitionNeeded())
 	{
-		EpilogueBarriersToBeginForGraphics = Allocator.AllocNoDestruct<FRDGBarrierBatchBegin>(Pipeline, ERHIPipeline::Graphics, GetEpilogueBarriersToBeginDebugName(ERHIPipeline::Graphics), this);
+		EpilogueBarriersToBeginForGraphics.Reserve(TextureStates.Num() + BufferStates.Num());
+		CreateQueue.Emplace(&EpilogueBarriersToBeginForGraphics);
 	}
-	return *EpilogueBarriersToBeginForGraphics;
+	return EpilogueBarriersToBeginForGraphics;
 }
 
-FRDGBarrierBatchBegin& FRDGPass::GetEpilogueBarriersToBeginForAsyncCompute(FRDGAllocator& Allocator)
+FRDGBarrierBatchBegin& FRDGPass::GetEpilogueBarriersToBeginForAsyncCompute(FRDGAllocator& Allocator, FRDGTransitionCreateQueue& CreateQueue)
 {
 	if (!EpilogueBarriersToBeginForAsyncCompute)
 	{
 		EpilogueBarriersToBeginForAsyncCompute = Allocator.AllocNoDestruct<FRDGBarrierBatchBegin>(Pipeline, ERHIPipeline::AsyncCompute, GetEpilogueBarriersToBeginDebugName(ERHIPipeline::AsyncCompute), this);
+		CreateQueue.Emplace(EpilogueBarriersToBeginForAsyncCompute);
 	}
 	return *EpilogueBarriersToBeginForAsyncCompute;
 }
 
-FRDGBarrierBatchBegin& FRDGPass::GetEpilogueBarriersToBeginForAll(FRDGAllocator& Allocator)
+FRDGBarrierBatchBegin& FRDGPass::GetEpilogueBarriersToBeginForAll(FRDGAllocator& Allocator, FRDGTransitionCreateQueue& CreateQueue)
 {
 	if (!EpilogueBarriersToBeginForAll)
 	{
 		EpilogueBarriersToBeginForAll = Allocator.AllocNoDestruct<FRDGBarrierBatchBegin>(Pipeline, ERHIPipeline::All, GetEpilogueBarriersToBeginDebugName(ERHIPipeline::AsyncCompute), this);
+		CreateQueue.Emplace(EpilogueBarriersToBeginForAll);
 	}
 	return *EpilogueBarriersToBeginForAll;
 }
 
 FRDGBarrierBatchEnd& FRDGPass::GetPrologueBarriersToEnd(FRDGAllocator& Allocator)
 {
-	if (!PrologueBarriersToEnd)
-	{
-		PrologueBarriersToEnd = Allocator.AllocNoDestruct<FRDGBarrierBatchEnd>(Handle);
-	}
-	return *PrologueBarriersToEnd;
+	return PrologueBarriersToEnd;
 }
 
 FRDGBarrierBatchEnd& FRDGPass::GetEpilogueBarriersToEnd(FRDGAllocator& Allocator)
 {
 	if (!EpilogueBarriersToEnd)
 	{
-		EpilogueBarriersToEnd = Allocator.AllocNoDestruct<FRDGBarrierBatchEnd>(Handle);
+		EpilogueBarriersToEnd = Allocator.AllocNoDestruct<FRDGBarrierBatchEnd>(this);
 	}
 	return *EpilogueBarriersToEnd;
 }
@@ -320,6 +288,8 @@ FRDGPass::FRDGPass(
 	, ParameterStruct(InParameterStruct)
 	, Flags(InFlags)
 	, Pipeline(EnumHasAnyFlags(Flags, ERDGPassFlags::AsyncCompute) ? ERHIPipeline::AsyncCompute : ERHIPipeline::Graphics)
+	, PrologueBarriersToEnd(this)
+	, EpilogueBarriersToBeginForGraphics(Pipeline, ERHIPipeline::Graphics, GetEpilogueBarriersToBeginDebugName(ERHIPipeline::Graphics), this)
 {}
 
 #if RDG_ENABLE_DEBUG
@@ -335,11 +305,4 @@ const TCHAR* FRDGPass::GetName() const
 		return Name.GetTCHAR();
 	}
 }
-#endif
-
-void FRDGPass::Execute(FRHIComputeCommandList& RHICmdList)
-{
-	QUICK_SCOPE_CYCLE_COUNTER(STAT_FRDGPass_Execute);
-	RHICmdList.SetStaticUniformBuffers(ParameterStruct.GetStaticUniformBuffers());
-	ExecuteImpl(RHICmdList);
-}
\ No newline at end of file
+#endif
\ No newline at end of file
diff --git a/Engine/Source/Runtime/RenderCore/Private/RenderGraphPrivate.cpp b/Engine/Source/Runtime/RenderCore/Private/RenderGraphPrivate.cpp
--- a/Engine/Source/Runtime/RenderCore/Private/RenderGraphPrivate.cpp
+++ b/Engine/Source/Runtime/RenderCore/Private/RenderGraphPrivate.cpp
@@ -305,8 +305,16 @@ FAutoConsoleVariableRef CVarRDGUseTransientAllocator(
 	TEXT(" 1: enables the transient allocator (default)"),
 	ECVF_RenderThreadSafe);
 
+int32 GRDGDrain = 1;
+FAutoConsoleVariableRef CVarRDGDrain(
+	TEXT("r.RDG.Drain"), GRDGDrain,
+	TEXT("RDG will perform drain operations when requested.")
+	TEXT(" 0: disables draining, all work is deferred until FRDGBuilder::Execute;")
+	TEXT(" 1: enables execution of the RDG subgraph when FRDGBuilder::Drain is called. (default)"),
+	ECVF_RenderThreadSafe);
+
 #if CSV_PROFILER
-int32 GRDGVerboseCSVStats = 1;
+int32 GRDGVerboseCSVStats = 0;
 FAutoConsoleVariableRef CVarRDGVerboseCSVStats(
 	TEXT("r.RDG.VerboseCSVStats"),
 	GRDGVerboseCSVStats,
diff --git a/Engine/Source/Runtime/RenderCore/Private/RenderGraphPrivate.h b/Engine/Source/Runtime/RenderCore/Private/RenderGraphPrivate.h
--- a/Engine/Source/Runtime/RenderCore/Private/RenderGraphPrivate.h
+++ b/Engine/Source/Runtime/RenderCore/Private/RenderGraphPrivate.h
@@ -107,7 +107,7 @@ extern int32 GRDGAsyncCompute;
 extern int32 GRDGCullPasses;
 extern int32 GRDGMergeRenderPasses;
 extern int32 GRDGTransientAllocator;
-extern int32 GRDGTransientMinimumBufferSize;
+extern int32 GRDGDrain;
 
 #if CSV_PROFILER
 extern int32 GRDGVerboseCSVStats;
diff --git a/Engine/Source/Runtime/RenderCore/Private/RenderGraphResources.cpp b/Engine/Source/Runtime/RenderCore/Private/RenderGraphResources.cpp
--- a/Engine/Source/Runtime/RenderCore/Private/RenderGraphResources.cpp
+++ b/Engine/Source/Runtime/RenderCore/Private/RenderGraphResources.cpp
@@ -30,16 +30,10 @@ FRDGParentResource::FRDGParentResource(const TCHAR* InName, const ERDGParentReso
 	, bUserSetNonTransient(0)
 	, bFinalizedAccess(0)
 	, bLastOwner(1)
-	// Culling logic runs only when immediate mode is off.
-	, bCulled(1)
+	, bCulled(0)
 	, bUsedByAsyncComputePass(0)
 {}
 
-FRDGParentResource::~FRDGParentResource()
-{
-	check(IsImmediateMode() || ReferenceCount == 0);
-}
-
 bool FRDGProducerState::IsDependencyRequired(FRDGProducerState LastProducer, ERHIPipeline LastPipeline, FRDGProducerState NextState, ERHIPipeline NextPipeline)
 {
 	/** This function determines whether a producer-consumer relationship exists in the graph, which is used for culling and
@@ -159,6 +153,22 @@ bool FRDGSubresourceState::IsTransitionRequired(const FRDGSubresourceState& Prev
 	return false;
 }
 
+void FRDGUniformBuffer::InitRHI()
+{
+	check(!HasRHI());
+
+	const EUniformBufferValidation Validation =
+#if RDG_ENABLE_DEBUG
+		EUniformBufferValidation::ValidateResources;
+#else
+		EUniformBufferValidation::None;
+#endif
+
+	const FRDGParameterStruct PassParameters = GetParameters();
+	UniformBufferRHI = RHICreateUniformBuffer(PassParameters.GetContents(), PassParameters.GetLayout(), UniformBuffer_SingleFrame, Validation);
+	ResourceRHI = UniformBufferRHI;
+}
+
 void FRDGPooledTexture::Finalize()
 {
 	for (FRDGSubresourceState& SubresourceState : State)
@@ -186,9 +196,15 @@ FRDGTextureSubresourceRange FRDGTexture::GetSubresourceRangeSRV() const
 	return Range;
 }
 
+IPooledRenderTarget* FRDGTexture::GetPooledRenderTarget() const
+{
+	IF_RDG_ENABLE_DEBUG(ValidateRHIAccess());
+	return PooledRenderTarget;
+}
+
 void FRDGTexture::SetRHI(FPooledRenderTarget* InPooledRenderTarget)
 {
-	Allocation = InPooledRenderTarget;
+	Allocation = TRefCountPtr<FPooledRenderTarget>(InPooledRenderTarget);
 	PooledRenderTarget = InPooledRenderTarget;
 
 	if (!InPooledRenderTarget->HasRDG())
@@ -229,7 +245,7 @@ void FRDGTexture::SetRHI(FRHITransientTexture* InTransientTexture, FRDGAllocator
 	bTransient = true;
 }
 
-void FRDGTexture::Finalize()
+void FRDGTexture::Finalize(FRDGPooledTextureArray& PooledTextureArray)
 {
 	checkf(NextOwner.IsNull() == !!bLastOwner, TEXT("NextOwner must match bLastOwner."));
 	checkf(!bExtracted || bLastOwner, TEXT("Extracted resources must be the last owner of a resource."));
@@ -257,10 +273,13 @@ void FRDGTexture::Finalize()
 			// Restore the reference to the last owner in the aliasing chain; not necessary for the transient resource allocator.
 			if (PooledRenderTarget)
 			{
-				Allocation = PooledRenderTarget;
+				PooledTextureArray.Emplace(PooledRenderTarget);
 			}
 		}
 	}
+
+	// This releases the reference without invoking a virtual function call.
+	TRefCountPtr<FPooledRenderTarget>(MoveTemp(Allocation));
 }
 
 void FRDGBuffer::SetRHI(FRDGPooledBuffer* InPooledBuffer)
@@ -293,7 +312,7 @@ void FRDGBuffer::SetRHI(FRHITransientBuffer* InTransientBuffer, FRDGAllocator& A
 	bTransient = true;;
 }
 
-void FRDGBuffer::Finalize()
+void FRDGBuffer::Finalize(FRDGPooledBufferArray& PooledBufferArray)
 {
 	// If these fire, the graph is not tracking state properly.
 	checkf(NextOwner.IsNull() == !!bLastOwner, TEXT("NextOwner must match bLastOwner."));
@@ -316,8 +335,9 @@ void FRDGBuffer::Finalize()
 				PooledBuffer->Finalize();
 			}
 
-			// Restore the reference to the last owner in the aliasing chain.
-			Allocation = PooledBuffer;
+			PooledBufferArray.Emplace(PooledBuffer);
 		}
 	}
+
+	Allocation = nullptr;
 }
\ No newline at end of file
diff --git a/Engine/Source/Runtime/RenderCore/Private/RenderGraphTrace.cpp b/Engine/Source/Runtime/RenderCore/Private/RenderGraphTrace.cpp
--- a/Engine/Source/Runtime/RenderCore/Private/RenderGraphTrace.cpp
+++ b/Engine/Source/Runtime/RenderCore/Private/RenderGraphTrace.cpp
@@ -142,7 +142,7 @@ void FRDGTrace::OutputGraphEnd(const FRDGBuilder& GraphBuilder)
 			<< PassMessage.Buffers((const uint16*)Pass->TraceBuffers.GetData(), (uint16)Pass->TraceBuffers.Num())
 			<< PassMessage.Flags(uint16(Pass->GetFlags()))
 			<< PassMessage.Pipeline(uint16(Pass->GetPipeline()))
-			<< PassMessage.IsCulled(GraphBuilder.PassesToCull[Handle])
+			<< PassMessage.IsCulled(Pass->bCulled != 0)
 			<< PassMessage.IsAsyncComputeBegin(Pass->bAsyncComputeBegin != 0)
 			<< PassMessage.IsAsyncComputeEnd(Pass->bAsyncComputeEnd != 0)
 			<< PassMessage.SkipRenderPassBegin(Pass->bSkipRenderPassBegin != 0)
diff --git a/Engine/Source/Runtime/RenderCore/Private/RenderGraphValidation.cpp b/Engine/Source/Runtime/RenderCore/Private/RenderGraphValidation.cpp
--- a/Engine/Source/Runtime/RenderCore/Private/RenderGraphValidation.cpp
+++ b/Engine/Source/Runtime/RenderCore/Private/RenderGraphValidation.cpp
@@ -155,7 +155,6 @@ void FRDGUniformBuffer::MarkResourceAsUsed()
 
 FRDGUserValidation::FRDGUserValidation(FRDGAllocator& InAllocator)
 	: Allocator(InAllocator)
-	, ExpectedNumMarks(FMemStack::Get().GetNumMarks())
 {
 	checkf(!GRDGBuilderActive, TEXT("Another FRDGBuilder already exists on the stack. Only one builder can be created at a time. This builder instance should be merged into the parent one."));
 	GRDGBuilderActive = true;
@@ -167,20 +166,13 @@ FRDGUserValidation::~FRDGUserValidation()
 	GRDGBuilderActive = false;
 }
 
-void FRDGUserValidation::MemStackGuard()
-{
-	checkf(ExpectedNumMarks == FMemStack::Get().GetNumMarks(), TEXT("A MemStack mark was added during the FRDGBuilder lifetime. This is not allowed as it will free memory still used by the builder."));
-}
-
 void FRDGUserValidation::ExecuteGuard(const TCHAR* Operation, const TCHAR* ResourceName)
 {
 	checkf(!bHasExecuted, TEXT("Render graph operation '%s' with resource '%s' must be performed prior to graph execution."), Operation, ResourceName);
-	MemStackGuard();
 }
 
 void FRDGUserValidation::ValidateCreateResource(FRDGResourceRef Resource)
 {
-	MemStackGuard();
 	check(Resource);
 	Resource->DebugData = Allocator.Alloc<FRDGResourceDebugData>();
 
@@ -397,9 +389,12 @@ void FRDGUserValidation::ValidateExtractBuffer(FRDGBufferRef Buffer, TRefCountPt
 
 void FRDGUserValidation::ValidateExtractResource(FRDGParentResourceRef Resource)
 {
-	MemStackGuard();
 	check(Resource);
 
+	checkf(!bHasExecuteBegun || !Resource->bTransient,
+		TEXT("Unable to queue the extraction of the resource %s because passes in the graph have already executed and it was made transient."),
+		Resource->Name);
+
 	checkf(Resource->bProduced || Resource->bExternal,
 		TEXT("Unable to queue the extraction of the resource %s because it has not been produced by any pass."),
 		Resource->Name);
@@ -411,6 +406,14 @@ void FRDGUserValidation::ValidateExtractResource(FRDGParentResourceRef Resource)
 	Resource->GetParentDebugData().PassAccessCount++;
 }
 
+void FRDGUserValidation::ValidateConvertToExternalResource(FRDGParentResourceRef Resource)
+{
+	check(Resource);
+	checkf(!bHasExecuteBegun || !Resource->bTransient,
+		TEXT("Unable to convert resource %s to external because passes in the graph have already executed."),
+		Resource->Name);
+}
+
 void FRDGUserValidation::RemoveUnusedWarning(FRDGParentResourceRef Resource)
 {
 	check(Resource);
@@ -509,8 +512,6 @@ void FRDGUserValidation::ValidateAddPass(const void* ParameterStruct, const FSha
 
 void FRDGUserValidation::ValidateAddPass(const FRDGPass* Pass, bool bSkipPassAccessMarking)
 {
-	MemStackGuard();
-
 	const FRenderTargetBindingSlots* RenderTargetBindingSlots = nullptr;
 
 	// Pass flags are validated as early as possible by the builder in AddPass.
@@ -826,13 +827,14 @@ void FRDGUserValidation::ValidateAddPass(const FRDGPass* Pass, bool bSkipPassAcc
 
 void FRDGUserValidation::ValidateExecuteBegin()
 {
-	MemStackGuard();
 	checkf(!bHasExecuted, TEXT("Render graph execution should only happen once to ensure consistency with immediate mode."));
+	check(!bHasExecuteBegun);
+	bHasExecuteBegun = true;
 }
 
 void FRDGUserValidation::ValidateExecuteEnd()
 {
-	MemStackGuard();
+	check(bHasExecuteBegun);
 
 	bHasExecuted = true;
 
@@ -1362,26 +1364,14 @@ namespace
 
 FString FRDGLogFile::GetProducerName(FRDGPassHandle PassHandle)
 {
-	if (PassHandle.IsValid())
-	{
-		return GetNodeName(PassHandle);
-	}
-	else
-	{
-		return GetNodeName(ProloguePassHandle);
-	}
+	check(PassHandle.IsValid());
+	return GetNodeName(PassHandle);
 }
 
 FString FRDGLogFile::GetConsumerName(FRDGPassHandle PassHandle)
 {
-	if (PassHandle.IsValid())
-	{
-		return GetNodeName(PassHandle);
-	}
-	else
-	{
-		return GetNodeName(EpiloguePassHandle);
-	}
+	check(PassHandle.IsValid());
+	return GetNodeName(PassHandle);
 }
 
 FString FRDGLogFile::GetNodeName(FRDGPassHandle PassHandle)
@@ -1419,12 +1409,7 @@ void FRDGLogFile::AddBraceEnd()
 	AddLine(TEXT("}"));
 }
 
-void FRDGLogFile::Begin(
-	const FRDGEventName& InGraphName,
-	const FRDGPassRegistry* InPasses,
-	FRDGPassBitArray InPassesCulled,
-	FRDGPassHandle InProloguePassHandle,
-	FRDGPassHandle InEpiloguePassHandle)
+void FRDGLogFile::Begin(const FRDGEventName& InGraphName)
 {
 	if (GRDGDumpGraph)
 	{
@@ -1435,12 +1420,7 @@ void FRDGLogFile::Begin(
 		}
 
 		check(File.IsEmpty());
-		check(InPasses && InEpiloguePassHandle.IsValid());
 
-		Passes = InPasses;
-		PassesCulled = InPassesCulled;
-		ProloguePassHandle = InProloguePassHandle;
-		EpiloguePassHandle = InEpiloguePassHandle;
 		GraphName = InGraphName.GetTCHAR();
 
 		if (GraphName.IsEmpty())
@@ -1467,9 +1447,9 @@ void FRDGLogFile::End()
 	TArray<FRDGPassHandle> PassesGraphics;
 	TArray<FRDGPassHandle> PassesAsyncCompute;
 
-	for (FRDGPassHandle PassHandle = Passes->Begin(); PassHandle != Passes->End(); ++PassHandle)
+	for (FRDGPassHandle PassHandle = Passes.Begin(); PassHandle != Passes.End(); ++PassHandle)
 	{
-		const FRDGPass* Pass = Passes->Get(PassHandle);
+		const FRDGPass* Pass = Passes[PassHandle];
 		const ERHIPipeline Pipeline = Pass->GetPipeline();
 
 		switch (Pipeline)
@@ -1489,9 +1469,9 @@ void FRDGLogFile::End()
 	{
 		FRDGPassHandle PrevPassesByPipeline[uint32(ERHIPipeline::Num)];
 
-		for (FRDGPassHandle PassHandle = Passes->Begin(); PassHandle != Passes->End(); ++PassHandle)
+		for (FRDGPassHandle PassHandle = Passes.Begin(); PassHandle != Passes.End(); ++PassHandle)
 		{
-			const FRDGPass* Pass = Passes->Get(PassHandle);
+			const FRDGPass* Pass = Passes[PassHandle];
 
 			if (!EnumHasAnyFlags(Pass->GetFlags(), ERDGPassFlags::Copy | ERDGPassFlags::Raster | ERDGPassFlags::Compute | ERDGPassFlags::AsyncCompute))
 			{
@@ -1534,24 +1514,16 @@ void FRDGLogFile::End()
 	}
 	else if (GRDGDumpGraph == RDG_DUMP_GRAPH_PRODUCERS)
 	{
-		for (FRDGPassHandle PassHandle = Passes->Begin(); PassHandle != Passes->End(); ++PassHandle)
+		for (FRDGPassHandle PassHandle = Passes.Begin(); PassHandle < Passes.Last(); ++PassHandle)
 		{
-			if (PassHandle == EpiloguePassHandle)
-			{
-				break;
-			}
-
-			const FRDGPass* Pass = Passes->Get(PassHandle);
+			const FRDGPass* Pass = Passes[PassHandle];
 
 			for (const FRDGPassHandle ProducerHandle : Pass->GetProducers())
 			{
-				if (ProducerHandle != ProloguePassHandle)
-				{
-					const FRDGPass* Producer = Passes->Get(ProducerHandle);
+				const FRDGPass* Producer = Passes[ProducerHandle];
 
-					File += FString::Printf(TEXT("\t\"%s\" -> \"%s\" [penwidth=2, color=\"%s:%s\"]\n"),
-						*GetNodeName(ProducerHandle), *GetNodeName(PassHandle), GetPassColorName(Pass->GetFlags()), GetPassColorName(Producer->GetFlags()));
-				}
+				File += FString::Printf(TEXT("\t\"%s\" -> \"%s\" [penwidth=2, color=\"%s:%s\"]\n"),
+					*GetNodeName(ProducerHandle), *GetNodeName(PassHandle), GetPassColorName(Pass->GetFlags()), GetPassColorName(Producer->GetFlags()));
 			}
 		}
 	}
@@ -1566,8 +1538,8 @@ void FRDGLogFile::End()
 			return;
 		}
 
-		const FRDGPass* Pass = Passes->Get(PassHandle);
-		const TCHAR* Style = PassesCulled[PassHandle] ? TEXT("dashed") : TEXT("filled");
+		const FRDGPass* Pass = Passes[PassHandle];
+		const TCHAR* Style = Pass->IsCulled() ? TEXT("dashed") : TEXT("filled");
 		FString PassName = FString::Printf(TEXT("[%d]: %s"), PassHandle.GetIndex(), Pass->GetName());
 
 		if (Pass->GetParameters().HasExternalOutputs())
@@ -1583,7 +1555,7 @@ void FRDGLogFile::End()
 
 		for (FRDGPassHandle PassHandle : PassesGraphics)
 		{
-			const FRDGPass* Pass = Passes->Get(PassHandle);
+			const FRDGPass* Pass = Passes[PassHandle];
 
 			if (Pass->IsMergedRenderPassBegin())
 			{
@@ -1640,9 +1612,10 @@ void FRDGLogFile::End()
 
 	uint32 NumPassesActive = 0;
 	uint32 NumPassesCulled = 0;
-	for (FRDGPassHandle PassHandle = Passes->Begin(); PassHandle != Passes->End(); ++PassHandle)
+
+	Passes.Enumerate([&](const FRDGPass* Pass)
 	{
-		if (PassesCulled[PassHandle])
+		if (Pass->IsCulled())
 		{
 			NumPassesCulled++;
 		}
@@ -1650,7 +1623,7 @@ void FRDGLogFile::End()
 		{
 			NumPassesActive++;
 		}
-	}
+	});
 
 	AddLine(FString::Printf(TEXT("label=\"%s [Active Passes: %d, Culled Passes: %d, Textures: %d, Buffers: %d]\""), *GraphName, NumPassesActive, NumPassesCulled, Textures.Num(), Buffers.Num()));
 
@@ -1679,7 +1652,7 @@ void FRDGLogFile::End()
 
 bool FRDGLogFile::IncludeTransitionEdgeInGraph(FRDGPassHandle Pass) const
 {
-	return Pass.IsValid() && Pass != ProloguePassHandle && Pass != EpiloguePassHandle;
+	return Pass.IsValid() && !Passes[Pass]->IsSentinel();
 }
 
 bool FRDGLogFile::IncludeTransitionEdgeInGraph(FRDGPassHandle PassBefore, FRDGPassHandle PassAfter) const
diff --git a/Engine/Source/Runtime/RenderCore/Public/RenderGraphBuilder.h b/Engine/Source/Runtime/RenderCore/Public/RenderGraphBuilder.h
--- a/Engine/Source/Runtime/RenderCore/Public/RenderGraphBuilder.h
+++ b/Engine/Source/Runtime/RenderCore/Public/RenderGraphBuilder.h
@@ -25,6 +25,7 @@ class RENDERCORE_API FRDGBuilder
 public:
 	FRDGBuilder(FRHICommandListImmediate& InRHICmdList, FRDGEventName InName = {});
 	FRDGBuilder(const FRDGBuilder&) = delete;
+	~FRDGBuilder();
 
 	/** Finds an RDG texture associated with the external texture, or returns null if none is found. */
 	FRDGTextureRef FindExternalTexture(FRHITexture* Texture) const;
@@ -166,6 +167,9 @@ public:
 	/** Sets the current command list stat for all subsequent passes. */
 	void SetCommandListStat(TStatId StatId);
 
+	/** A hint to the builder to flush work to the RHI thread after the last queued pass on the execution timeline. */
+	void AddDispatchHint();
+
 	/** Queues a pooled render target extraction to happen at the end of graph execution. For graph-created textures, this extends
 	 *  the lifetime of the GPU resource until execution, at which point the pointer is filled. If specified, the texture is transitioned
 	 *  to the AccessFinal state, or kDefaultAccessFinal otherwise.
@@ -234,6 +238,9 @@ public:
 	/** Manually ends the current GPU event scope. */
 	void EndEventScope();
 
+	/** Drains the RDG queue early. All resource lifetimes are extended to the end of the Drain and culling is disabled for any drained passes. */
+	void Drain();
+
 	/** Executes the queued passes, managing setting of render targets (RHI RenderPasses), resource transitions and queued texture extraction. */
 	void Execute();
 
@@ -274,10 +281,29 @@ private:
 
 	static ERDGPassFlags OverridePassFlags(const TCHAR* PassName, ERDGPassFlags Flags, bool bAsyncComputeSupported);
 
+	void AddProloguePass();
+
+	FORCEINLINE FRDGPassHandle ClampToPrologue(FRDGPassHandle PassHandle, ERHIPipeline Pipeline = ERHIPipeline::Graphics) const
+	{
+		return FRDGPassHandle::Max(ProloguePassHandles[Pipeline], PassHandle);
+	}
+
+	FORCEINLINE FRDGPassHandlesByPipeline ClampToPrologue(FRDGPassHandlesByPipeline PassHandles) const
+	{
+		PassHandles[ERHIPipeline::Graphics]     = ClampToPrologue(PassHandles[ERHIPipeline::Graphics],     ERHIPipeline::Graphics);
+		PassHandles[ERHIPipeline::AsyncCompute] = ClampToPrologue(PassHandles[ERHIPipeline::AsyncCompute], ERHIPipeline::AsyncCompute);
+		return PassHandles;
+	}
+
+	FORCEINLINE FRDGPass* GetProloguePass(ERHIPipeline Pipeline = ERHIPipeline::Graphics) const
+	{
+		return ProloguePasses[Pipeline];
+	}
+
 	/** Returns the graph prologue pass handle. */
-	FORCEINLINE FRDGPassHandle GetProloguePassHandle() const
+	FORCEINLINE FRDGPassHandle GetProloguePassHandle(ERHIPipeline Pipeline = ERHIPipeline::Graphics) const
 	{
-		return Passes.Begin();
+		return ProloguePassHandles[Pipeline];
 	}
 
 	/** Returns the graph epilogue pass handle. */
@@ -346,7 +372,7 @@ private:
 	void AddToPrologueBarriers(FRDGPassHandle PassHandle, FunctionType Function)
 	{
 		FRDGPass* Pass = GetPrologueBarrierPass(PassHandle);
-		FRDGBarrierBatchBegin& BarriersToBegin = Pass->GetPrologueBarriersToBegin(Allocator);
+		FRDGBarrierBatchBegin& BarriersToBegin = Pass->GetPrologueBarriersToBegin(Allocator, TransitionCreateQueue);
 		Function(BarriersToBegin);
 		Pass->GetPrologueBarriersToEnd(Allocator).AddDependency(&BarriersToBegin);
 	}
@@ -356,7 +382,7 @@ private:
 	void AddToEpilogueBarriers(FRDGPassHandle PassHandle, FunctionType Function)
 	{
 		FRDGPass* Pass = GetEpilogueBarrierPass(PassHandle);
-		FRDGBarrierBatchBegin& BarriersToBegin = Pass->GetEpilogueBarriersToBeginFor(Allocator, Pass->GetPipeline());
+		FRDGBarrierBatchBegin& BarriersToBegin = Pass->GetEpilogueBarriersToBeginFor(Allocator, TransitionCreateQueue, Pass->GetPipeline());
 		Function(BarriersToBegin);
 		Pass->GetEpilogueBarriersToEnd(Allocator).AddDependency(&BarriersToBegin);
 	}
@@ -382,14 +408,11 @@ private:
 	FRDGViewRegistry Views;
 	FRDGUniformBufferRegistry UniformBuffers;
 
-	/** Passes that have been culled from the graph. */
-	FRDGPassBitArray PassesToCull;
-
-	/** Passes that don't have any parameters. */
-	FRDGPassBitArray PassesWithEmptyParameters;
+	FRDGPassBitArray PassesOnAsyncCompute;
+	TArray<FRDGPassHandle, TInlineAllocator<8, FRDGArrayAllocator>> PassesOnAsyncComputeToJoin;
 
 	/** Uniform buffers which were used in a pass. */
-	FRDGUniformBufferBitArray UniformBuffersToCreate;
+	TArray<FRDGUniformBufferHandle, TInlineAllocator<32, FRDGArrayAllocator>> UniformBuffersToCreate;
 
 	/** Tracks external resources to their registered render graph counterparts for de-duplication. */
 	TSortedMap<FRHITexture*, FRDGTexture*, FRDGArrayAllocator> ExternalTextures;
@@ -398,30 +421,41 @@ private:
 	/** Map of barrier batches begun from more than one pipe. */
 	TMap<FRDGBarrierBatchBeginId, FRDGBarrierBatchBegin*, FRDGSetAllocator> BarrierBatchMap;
 
+	/** Set of all active barrier batch begin instances; used to create transitions. */
+	FRDGTransitionCreateQueue TransitionCreateQueue;
+
 	/** The epilogue and prologue passes are sentinels that are used to simplify graph logic around barriers
 	 *  and traversal. The prologue pass is used exclusively for barriers before the graph executes, while the
 	 *  epilogue pass is used for resource extraction barriers--a property that also makes it the main root of
 	 *  the graph for culling purposes. The epilogue pass is added to the very end of the pass array for traversal
 	 *  purposes. The prologue does not need to participate in any graph traversal behavior.
 	 */
-	FRDGPass* ProloguePass = nullptr;
+	FRDGPassHandlesByPipeline ProloguePassHandles;
+	FRDGPassesByPipeline ProloguePasses;
 	FRDGPass* EpiloguePass = nullptr;
 
 	/** Array of all requested resource extractions from the graph. */
 	TArray<TPair<FRDGTextureRef, TRefCountPtr<IPooledRenderTarget>*>, FRDGArrayAllocator> ExtractedTextures;
 	TArray<TPair<FRDGBufferRef, TRefCountPtr<FRDGPooledBuffer>*>, FRDGArrayAllocator> ExtractedBuffers;
 
+	/** Array of all pooled references held during execution. */
+	FRDGPooledTextureArray ActivePooledTextures;
+	FRDGPooledBufferArray ActivePooledBuffers;
+
 	/** Texture state used for intermediate operations. Held here to avoid re-allocating. */
 	FRDGTextureTransientSubresourceStateIndirect ScratchTextureState;
 
 	/** Current scope's async compute budget. This is passed on to every pass created. */
 	EAsyncComputeBudget AsyncComputeBudgetScope = EAsyncComputeBudget::EAll_4;
+	EAsyncComputeBudget AsyncComputeBudgetState = EAsyncComputeBudget(~0u);
 
 	IF_RDG_CPU_SCOPES(FRDGCPUScopeStacks CPUScopeStacks);
 	IF_RDG_GPU_SCOPES(FRDGGPUScopeStacksByPipeline GPUScopeStacks);
 
 	IF_RDG_ENABLE_TRACE(FRDGTrace Trace);
 
+	bool bDispatchHint = false;
+
 #if RDG_ENABLE_DEBUG
 	FRDGUserValidation UserValidation;
 	FRDGBarrierValidation BarrierValidation;
@@ -439,10 +473,12 @@ private:
 	bool bWaitedForTemporalEffect = false;
 #endif
 
-	/** Tracks whether all passes / resources have been added to the graph. */
-	bool bSetupComplete = false;
+	uint32 AsyncComputePassCount = 0;
+	uint32 RasterPassCount = 0;
+	uint32 ExecuteCount = 0;
 
-	IF_RDG_CMDLIST_STATS(TStatId CommandListStat);
+	IF_RDG_CMDLIST_STATS(TStatId CommandListStatScope);
+	IF_RDG_CMDLIST_STATS(TStatId CommandListStatState);
 
 	class FTransientResourceAllocator
 	{
@@ -462,33 +498,47 @@ private:
 		FRHICommandListImmediate& RHICmdList;
 		IRHITransientResourceAllocator* Allocator = nullptr;
 		bool bCreateAttempted = false;
+
+	} TransientResourceAllocator;
+
+	enum class EExecuteMode
+	{
+		Drain,
+		Final
 	};
 
-	void Compile();
+	void Execute(EExecuteMode ExecuteMode);
+	void Compile(EExecuteMode ExecuteMode);
 	void Clear();
 
-	void BeginResourceRHI(FRDGUniformBuffer* UniformBuffer);
-	void BeginResourceRHI(FTransientResourceAllocator*, FRDGPassHandle, FRDGTexture* Texture);
+	void BeginResourcesRHI(FRDGPass* ResourcePass, FRDGPassHandle ExecutePassHandle);
+	void BeginResourceRHI(FRDGPassHandle, FRDGTexture* Texture);
 	void BeginResourceRHI(FRDGPassHandle, FRDGTextureSRV* SRV);
 	void BeginResourceRHI(FRDGPassHandle, FRDGTextureUAV* UAV);
-	void BeginResourceRHI(FTransientResourceAllocator*, FRDGPassHandle, FRDGBuffer* Buffer);
+	void BeginResourceRHI(FRDGPassHandle, FRDGBuffer* Buffer);
 	void BeginResourceRHI(FRDGPassHandle, FRDGBufferSRV* SRV);
 	void BeginResourceRHI(FRDGPassHandle, FRDGBufferUAV* UAV);
 	void BeginResourceRHI(FRDGPassHandle, FRDGView* View);
 
-	void EndResourceRHI(FTransientResourceAllocator&, FRDGPassHandle, FRDGTexture* Texture, uint32 ReferenceCount);
-	void EndResourceRHI(FTransientResourceAllocator&, FRDGPassHandle, FRDGBuffer* Buffer, uint32 ReferenceCount);
+	void EndResourcesRHI(FRDGPass* ResourcePass, FRDGPassHandle ExecutePassHandle);
+	void EndResourceRHI(FRDGPassHandle, FRDGTexture* Texture, uint32 ReferenceCount);
+	void EndResourceRHI(FRDGPassHandle, FRDGBuffer* Buffer, uint32 ReferenceCount);
 
-	void SetupPassInternal(FRDGPass* Pass, FRDGPassHandle PassHandle, ERHIPipeline PassPipeline);
-	void SetupPass(FRDGPass* Pass);
-	void SetupEmptyPass(FRDGPass* Pass);
+	void FinalizeResourcesRHI(EExecuteMode ExecuteMode);
+
+	void SetupPassInternal(FRDGPass* Pass, FRDGPassHandle PassHandle, ERHIPipeline PassPipeline, bool bEmptyParameters);
+	FRDGPass* SetupPass(FRDGPass* Pass);
+	FRDGPass* SetupEmptyPass(FRDGPass* Pass);
 	void ExecutePass(FRDGPass* Pass);
 
 	void ExecutePassPrologue(FRHIComputeCommandList& RHICmdListPass, FRDGPass* Pass);
 	void ExecutePassEpilogue(FRHIComputeCommandList& RHICmdListPass, FRDGPass* Pass);
 
-	void CollectPassResources(FTransientResourceAllocator&, FRDGPassHandle PassHandle);
-	void CollectPassBarriers(FRDGPassHandle PassHandle);
+	void CollectPassBarriers(FRDGPass* Pass, FRDGPassHandle PassHandle);
+
+	void CreatePassBarriers();
+
+	void CreateUniformBuffers();
 
 	void AddPassDependency(FRDGPassHandle ProducerHandle, FRDGPassHandle ConsumerHandle);
 
@@ -511,8 +561,8 @@ private:
 		FRDGParentResource* Resource,
 		FRDGSubresourceState StateBefore,
 		FRDGSubresourceState StateAfter,
-		EBarrierLocation BarrierLocation,
-		const FRHITransitionInfo& TransitionInfo);
+		const FRHITransitionInfo& TransitionInfo,
+		EBarrierLocation BarrierLocation);
 
 	bool IsTransient(FRDGTextureRef Texture) const;
 	bool IsTransient(FRDGBufferRef Buffer) const;
@@ -527,6 +577,11 @@ private:
 
 	FRDGSubresourceState* AllocSubresource(const FRDGSubresourceState& Other);
 
+	inline void ExtendPassBitArray(FRDGPassBitArray& BitArray, bool bValue) const
+	{
+		BitArray.Insert(bValue, BitArray.Num(), Passes.Num() - BitArray.Num());
+	}
+
 #if RDG_ENABLE_DEBUG
 	void VisualizePassOutputs(const FRDGPass* Pass);
 	void ClobberPassOutputs(const FRDGPass* Pass);
diff --git a/Engine/Source/Runtime/RenderCore/Public/RenderGraphBuilder.inl b/Engine/Source/Runtime/RenderCore/Public/RenderGraphBuilder.inl
--- a/Engine/Source/Runtime/RenderCore/Public/RenderGraphBuilder.inl
+++ b/Engine/Source/Runtime/RenderCore/Public/RenderGraphBuilder.inl
@@ -117,7 +117,6 @@ TRDGUniformBufferRef<ParameterStructType> FRDGBuilder::CreateUniformBuffer(const
 {
 	IF_RDG_ENABLE_DEBUG(UserValidation.ValidateCreateUniformBuffer(ParameterStruct, &ParameterStructType::StaticStructMetadata));
 	auto* UniformBuffer = UniformBuffers.Allocate<TRDGUniformBuffer<ParameterStructType>>(Allocator, ParameterStruct, ParameterStructType::StaticStructMetadata.GetShaderVariableName());
-	UniformBuffersToCreate.Add(false);
 	IF_RDG_ENABLE_DEBUG(UserValidation.ValidateCreateUniformBuffer(UniformBuffer));
 	return UniformBuffer;
 }
@@ -195,6 +194,7 @@ inline void FRDGBuilder::QueueTextureExtraction(FRDGTextureRef Texture, TRefCoun
 {
 	IF_RDG_ENABLE_DEBUG(UserValidation.ValidateExtractTexture(Texture, OutTexturePtr));
 
+	Texture->ReferenceCount++;
 	Texture->bExtracted = true;
 	Texture->bCulled = false;
 	ExtractedTextures.Emplace(Texture, OutTexturePtr);
@@ -209,6 +209,7 @@ inline void FRDGBuilder::QueueBufferExtraction(FRDGBufferRef Buffer, TRefCountPt
 {
 	IF_RDG_ENABLE_DEBUG(UserValidation.ValidateExtractBuffer(Buffer, OutBufferPtr));
 
+	Buffer->ReferenceCount++;
 	Buffer->bExtracted = true;
 	Buffer->bCulled = false;
 	ExtractedBuffers.Emplace(Buffer, OutBufferPtr);
@@ -228,11 +229,16 @@ inline void FRDGBuilder::QueueBufferExtraction(FRDGBufferRef Buffer, TRefCountPt
 inline void FRDGBuilder::SetCommandListStat(TStatId StatId)
 {
 #if RDG_CMDLIST_STATS
-	CommandListStat = StatId;
+	CommandListStatScope = StatId;
 	RHICmdList.SetCurrentStat(StatId);
 #endif
 }
 
+inline void FRDGBuilder::AddDispatchHint()
+{
+	bDispatchHint = true;
+}
+
 inline const TRefCountPtr<IPooledRenderTarget>& FRDGBuilder::GetPooledTexture(FRDGTextureRef Texture) const
 {
 	IF_RDG_ENABLE_DEBUG(UserValidation.ValidateGetPooledTexture(Texture));
diff --git a/Engine/Source/Runtime/RenderCore/Public/RenderGraphDefinitions.h b/Engine/Source/Runtime/RenderCore/Public/RenderGraphDefinitions.h
--- a/Engine/Source/Runtime/RenderCore/Public/RenderGraphDefinitions.h
+++ b/Engine/Source/Runtime/RenderCore/Public/RenderGraphDefinitions.h
@@ -221,6 +221,7 @@ public:
 	FORCEINLINE IndexType GetIndexUnchecked() const { return Index; }
 	FORCEINLINE bool IsNull()  const { return Index == kNullIndex; }
 	FORCEINLINE bool IsValid() const { return Index != kNullIndex; }
+	FORCEINLINE operator bool() const { return IsValid(); }
 	FORCEINLINE bool operator==(TRDGHandle Other) const { return Index == Other.Index; }
 	FORCEINLINE bool operator!=(TRDGHandle Other) const { return Index != Other.Index; }
 	FORCEINLINE bool operator<=(TRDGHandle Other) const { check(IsValid() && Other.IsValid()); return Index <= Other.Index; }
@@ -228,6 +229,34 @@ public:
 	FORCEINLINE bool operator< (TRDGHandle Other) const { check(IsValid() && Other.IsValid()); return Index <  Other.Index; }
 	FORCEINLINE bool operator> (TRDGHandle Other) const { check(IsValid() && Other.IsValid()); return Index >  Other.Index; }
 
+	FORCEINLINE TRDGHandle& operator+=(int32 Increment)
+	{
+		check(int64(Index + Increment) <= int64(kNullIndex));
+		Index += Increment;
+		return *this;
+	}
+
+	FORCEINLINE TRDGHandle& operator-=(int32 Decrement)
+	{
+		check(int64(Index - Decrement) > 0);
+		Index -= Decrement;
+		return *this;
+	}
+
+	FORCEINLINE TRDGHandle operator-(int32 Subtract) const
+	{
+		TRDGHandle Handle = *this;
+		Handle -= Subtract;
+		return Handle;
+	}
+
+	FORCEINLINE TRDGHandle operator+(int32 Add) const
+	{
+		TRDGHandle Handle = *this;
+		Handle += Add;
+		return Handle;
+	}
+
 	FORCEINLINE TRDGHandle& operator++()
 	{
 		check(IsValid());
@@ -242,6 +271,20 @@ public:
 		return *this;
 	}
 
+	// Returns the min of two pass handles. Returns null if both are null; returns the valid handle if one is null.
+	FORCEINLINE static TRDGHandle Min(TRDGHandle A, TRDGHandle B)
+	{
+		// If either index is null is will fail the comparison.
+		return A.Index < B.Index ? A : B;
+	}
+
+	// Returns the max of two pass handles. Returns null if both are null; returns the valid handle if one is null.
+	FORCEINLINE static TRDGHandle Max(TRDGHandle A, TRDGHandle B)
+	{
+		// If either index is null, it will wrap around to 0 and fail the comparison.
+		return (IndexType)(A.Index + 1) > (IndexType)(B.Index + 1) ? A : B;
+	}
+
 private:
 	static const IndexType kNullIndex = TNumericLimits<IndexType>::Max();
 	IndexType Index = kNullIndex;
@@ -253,8 +296,15 @@ FORCEINLINE uint32 GetTypeHash(TRDGHandle<ObjectType, IndexType> Handle)
 	return Handle.GetIndex();
 }
 
+enum class ERDGHandleRegistryDestructPolicy
+{
+	Registry,
+	Allocator,
+	Never
+};
+
 /** Helper handle registry class for internal tracking of RDG types. */
-template <typename LocalHandleType>
+template <typename LocalHandleType, ERDGHandleRegistryDestructPolicy DestructPolicy = ERDGHandleRegistryDestructPolicy::Registry>
 class TRDGHandleRegistry
 {
 public:
@@ -272,20 +322,49 @@ public:
 	DerivedType* Allocate(FRDGAllocator& Allocator, TArgs&&... Args)
 	{
 		static_assert(TIsDerivedFrom<DerivedType, ObjectType>::Value, "You must specify a type that derives from ObjectType");
-		DerivedType* Object = Allocator.AllocNoDestruct<DerivedType>(Forward<TArgs>(Args)...);
+		DerivedType* Object;
+		if (DestructPolicy == ERDGHandleRegistryDestructPolicy::Allocator)
+		{
+			Object = Allocator.Alloc<DerivedType>(Forward<TArgs>(Args)...);
+		}
+		else
+		{
+			Object = Allocator.AllocNoDestruct<DerivedType>(Forward<TArgs>(Args)...);
+		}
 		Insert(Object);
 		return Object;
 	}
 
 	void Clear()
 	{
-		for (int32 Index = Array.Num() - 1; Index >= 0; --Index)
+		if (DestructPolicy == ERDGHandleRegistryDestructPolicy::Registry)
 		{
-			Array[Index]->~ObjectType();
+			for (int32 Index = Array.Num() - 1; Index >= 0; --Index)
+			{
+				Array[Index]->~ObjectType();
+			}
 		}
 		Array.Empty();
 	}
 
+	template <typename FunctionType>
+	void Enumerate(FunctionType Function)
+	{
+		for (ObjectType* Object : Array)
+		{
+			Function(Object);
+		}
+	}
+
+	template <typename FunctionType>
+	void Enumerate(FunctionType Function) const
+	{
+		for (const ObjectType* Object : Array)
+		{
+			Function(Object);
+		}
+	}
+
 	FORCEINLINE const ObjectType* Get(HandleType Handle) const
 	{
 		return Array[Handle.GetIndex()];
@@ -439,7 +518,7 @@ class FRDGBufferUAV;
 using FRDGBufferUAVRef = FRDGBufferUAV*;
 
 class FRDGPass;
-using FRDGPassRef = const FRDGPass*;
+using FRDGPassRef = FRDGPass*;
 using FRDGPassHandle = TRDGHandle<FRDGPass, uint16>;
 using FRDGPassRegistry = TRDGHandleRegistry<FRDGPassHandle>;
 using FRDGPassHandleArray = TArray<FRDGPassHandle, TInlineAllocator<4, FRDGArrayAllocator>>;
@@ -454,19 +533,19 @@ using FRDGUniformBufferBitArray = TRDGHandleBitArray<FRDGUniformBufferHandle>;
 class FRDGView;
 using FRDGViewRef = FRDGView*;
 using FRDGViewHandle = TRDGHandle<FRDGView, uint16>;
-using FRDGViewRegistry = TRDGHandleRegistry<FRDGViewHandle>;
+using FRDGViewRegistry = TRDGHandleRegistry<FRDGViewHandle, ERDGHandleRegistryDestructPolicy::Never>;
 using FRDGViewUniqueFilter = TRDGHandleUniqueFilter<FRDGViewHandle>;
 
 class FRDGTexture;
 using FRDGTextureRef = FRDGTexture*;
 using FRDGTextureHandle = TRDGHandle<FRDGTexture, uint16>;
-using FRDGTextureRegistry = TRDGHandleRegistry<FRDGTextureHandle>;
+using FRDGTextureRegistry = TRDGHandleRegistry<FRDGTextureHandle, ERDGHandleRegistryDestructPolicy::Never>;
 using FRDGTextureBitArray = TRDGHandleBitArray<FRDGTextureHandle>;
 
 class FRDGBuffer;
 using FRDGBufferRef = FRDGBuffer*;
 using FRDGBufferHandle = TRDGHandle<FRDGBuffer, uint16>;
-using FRDGBufferRegistry = TRDGHandleRegistry<FRDGBufferHandle>;
+using FRDGBufferRegistry = TRDGHandleRegistry<FRDGBufferHandle, ERDGHandleRegistryDestructPolicy::Never>;
 using FRDGBufferBitArray = TRDGHandleBitArray<FRDGBufferHandle>;
 
 class FRDGPooledTexture;
@@ -479,5 +558,6 @@ template <typename InElementType, typename InAllocatorType = FDefaultAllocator>
 using TRDGTextureSubresourceArray = TArray<InElementType, TInlineAllocator<1, InAllocatorType>>;
 
 using FRDGPassHandlesByPipeline = TRHIPipelineArray<FRDGPassHandle>;
+using FRDGPassesByPipeline = TRHIPipelineArray<FRDGPass*>;
 
 class FRDGTrace;
\ No newline at end of file
diff --git a/Engine/Source/Runtime/RenderCore/Public/RenderGraphEvent.h b/Engine/Source/Runtime/RenderCore/Public/RenderGraphEvent.h
--- a/Engine/Source/Runtime/RenderCore/Public/RenderGraphEvent.h
+++ b/Engine/Source/Runtime/RenderCore/Public/RenderGraphEvent.h
@@ -376,7 +376,7 @@ public:
 class RENDERCORE_API FRDGCSVStatScopeStack final
 {
 public:
-	FRDGCSVStatScopeStack(FRHIComputeCommandList& RHICmdList, FRDGAllocator& Allocator, const char* UnaccountedStatName);
+	FRDGCSVStatScopeStack(FRHIComputeCommandList& RHICmdList, FRDGAllocator& Allocator);
 
 	void BeginScope(const char* StatName);
 
@@ -396,7 +396,6 @@ public:
 private:
 	static bool IsEnabled();
 	TRDGScopeStack<FRDGCSVStatScope> ScopeStack;
-	const char* const UnaccountedStatName;
 };
 
 #if CSV_PROFILER
@@ -430,7 +429,7 @@ struct FRDGCPUScopes
 
 struct FRDGCPUScopeStacks
 {
-	FRDGCPUScopeStacks(FRHIComputeCommandList& RHICmdList, FRDGAllocator& Allocator, const char* UnaccountedCSVStat);
+	FRDGCPUScopeStacks(FRHIComputeCommandList& RHICmdList, FRDGAllocator& Allocator);
 
 	void BeginExecute();
 
diff --git a/Engine/Source/Runtime/RenderCore/Public/RenderGraphEvent.inl b/Engine/Source/Runtime/RenderCore/Public/RenderGraphEvent.inl
--- a/Engine/Source/Runtime/RenderCore/Public/RenderGraphEvent.inl
+++ b/Engine/Source/Runtime/RenderCore/Public/RenderGraphEvent.inl
@@ -109,7 +109,6 @@ void TRDGScopeStack<TScopeType>::EndScope()
 template <typename TScopeType>
 void TRDGScopeStack<TScopeType>::BeginExecute()
 {
-	checkf(CurrentScope == nullptr, TEXT("Render graph needs to have all scopes ended to execute."));
 }
 
 template <typename ScopeType>
@@ -130,6 +129,7 @@ void TRDGScopeStack<ScopeType>::BeginExecutePass(const ScopeType* ParentScope)
 template <typename TScopeType>
 void TRDGScopeStack<TScopeType>::EndExecute()
 {
+	checkf(CurrentScope == nullptr, TEXT("Render graph needs to have all scopes ended to execute."));
 	Helper.EndExecute([this](const TScopeType* Scope)
 	{
 		PopFunction(RHICmdList, Scope, bRDGEvents);
@@ -336,8 +336,8 @@ inline FRDGGPUScopes FRDGGPUScopeStacksByPipeline::GetCurrentScopes(ERHIPipeline
 
 #if RDG_CPU_SCOPES
 
-inline FRDGCPUScopeStacks::FRDGCPUScopeStacks(FRHIComputeCommandList& RHICmdList, FRDGAllocator& Allocator, const char* UnaccountedCSVStat)
-	: CSV(RHICmdList, Allocator, UnaccountedCSVStat)
+inline FRDGCPUScopeStacks::FRDGCPUScopeStacks(FRHIComputeCommandList& RHICmdList, FRDGAllocator& Allocator)
+	: CSV(RHICmdList, Allocator)
 {}
 
 inline void FRDGCPUScopeStacks::BeginExecute()
diff --git a/Engine/Source/Runtime/RenderCore/Public/RenderGraphPass.h b/Engine/Source/Runtime/RenderCore/Public/RenderGraphPass.h
--- a/Engine/Source/Runtime/RenderCore/Public/RenderGraphPass.h
+++ b/Engine/Source/Runtime/RenderCore/Public/RenderGraphPass.h
@@ -6,19 +6,7 @@
 #include "RenderGraphEvent.h"
 #include "Containers/SortedMap.h"
 
-class FRDGTransitionQueue
-{
-public:
-	FRDGTransitionQueue() = default;
-	FRDGTransitionQueue(uint32 ReservedCount);
-
-	void Insert(const FRHITransition* Transition, ERHITransitionCreateFlags TransitionFlags);
-	void Begin(FRHIComputeCommandList& RHICmdList);
-	void End(FRHIComputeCommandList& RHICmdList);
-
-private:
-	TArray<const FRHITransition*, TInlineAllocator<1, FRDGArrayAllocator>> Queue;
-};
+using FRDGTransitionQueue = TArray<const FRHITransition*, TInlineAllocator<4, FRDGArrayAllocator>>;
 
 struct FRDGBarrierBatchBeginId
 {
@@ -34,17 +22,22 @@ struct FRDGBarrierBatchBeginId
 		return !(*this == Other);
 	}
 
+	friend uint32 GetTypeHash(FRDGBarrierBatchBeginId Id)
+	{
+		static_assert(sizeof(Id.Passes) == 4);
+		uint32 Hash = *(const uint32*)Id.Passes.GetData();
+		return (Hash << GetRHIPipelineCount()) | uint32(Id.PipelinesAfter);
+	}
+
 	FRDGPassHandlesByPipeline Passes;
 	ERHIPipeline PipelinesAfter = ERHIPipeline::None;
 };
 
-RENDERCORE_API uint32 GetTypeHash(FRDGBarrierBatchBeginId Id);
-
 class RENDERCORE_API FRDGBarrierBatchBegin
 {
 public:
-	FRDGBarrierBatchBegin(ERHIPipeline PipelinesToBegin, ERHIPipeline PipelinesToEnd, const TCHAR* DebugName, const FRDGPass* DebugPass);
-	FRDGBarrierBatchBegin(ERHIPipeline PipelinesToBegin, ERHIPipeline PipelinesToEnd, const TCHAR* DebugName, FRDGPassHandlesByPipeline DebugPasses);
+	FRDGBarrierBatchBegin(ERHIPipeline PipelinesToBegin, ERHIPipeline PipelinesToEnd, const TCHAR* DebugName, FRDGPass* DebugPass);
+	FRDGBarrierBatchBegin(ERHIPipeline PipelinesToBegin, ERHIPipeline PipelinesToEnd, const TCHAR* DebugName, FRDGPassesByPipeline DebugPasses);
 
 	void AddTransition(FRDGParentResourceRef Resource, const FRHITransitionInfo& Info);
 
@@ -56,6 +49,8 @@ public:
 		bTransitionNeeded = true;
 	}
 
+	void CreateTransition();
+
 	void Submit(FRHIComputeCommandList& RHICmdList, ERHIPipeline Pipeline);
 	void Submit(FRHIComputeCommandList& RHICmdList, ERHIPipeline Pipeline, FRDGTransitionQueue& TransitionsToBegin);
 
@@ -64,6 +59,11 @@ public:
 		Transitions.Reserve(TransitionCount);
 	}
 
+	bool IsTransitionNeeded() const
+	{
+		return bTransitionNeeded;
+	}
+
 private:
 	const FRHITransition* Transition = nullptr;
 	TArray<FRHITransitionInfo, TInlineAllocator<1, FRDGArrayAllocator>> Transitions;
@@ -76,7 +76,7 @@ private:
 	ERHIPipeline PipelinesToEnd;
 
 #if RDG_ENABLE_DEBUG
-	FRDGPassHandlesByPipeline DebugPasses;
+	FRDGPassesByPipeline DebugPasses;
 	TArray<FRDGParentResource*, FRDGArrayAllocator> DebugTransitionResources;
 	TArray<FRDGParentResource*, FRDGArrayAllocator> DebugAliasingResources;
 	const TCHAR* DebugName;
@@ -89,12 +89,14 @@ private:
 	friend class FRDGBarrierValidation;
 };
 
+using FRDGTransitionCreateQueue = TArray<FRDGBarrierBatchBegin*, FRDGArrayAllocator>;
+
 class RENDERCORE_API FRDGBarrierBatchEnd
 {
 public:
-	FRDGBarrierBatchEnd(FRDGPassHandle InPassHandle)
+	FRDGBarrierBatchEnd(FRDGPass* InPass)
 #if RDG_ENABLE_DEBUG
-		: PassHandle(InPassHandle)
+		: Pass(InPass)
 #endif
 	{}
 
@@ -109,10 +111,10 @@ public:
 	}
 
 private:
-	TArray<FRDGBarrierBatchBegin*, TInlineAllocator<1, FRDGArrayAllocator>> Dependencies;
+	TArray<FRDGBarrierBatchBegin*, TInlineAllocator<4, FRDGArrayAllocator>> Dependencies;
 
 #if RDG_ENABLE_DEBUG
-	FRDGPassHandle PassHandle;
+	FRDGPass* Pass;
 #endif
 
 	friend class FRDGBarrierValidation;
@@ -205,6 +207,16 @@ public:
 		return bGraphicsJoin;
 	}
 
+	bool IsCulled() const
+	{
+		return bCulled;
+	}
+
+	bool IsSentinel() const
+	{
+		return bSentinel;
+	}
+
 	const FRDGPassHandleArray& GetProducers() const
 	{
 		return Producers;
@@ -249,12 +261,12 @@ public:
 #endif
 
 protected:
-	FRDGBarrierBatchBegin& GetPrologueBarriersToBegin(FRDGAllocator& Allocator);
-	FRDGBarrierBatchBegin& GetEpilogueBarriersToBeginForGraphics(FRDGAllocator& Allocator);
-	FRDGBarrierBatchBegin& GetEpilogueBarriersToBeginForAsyncCompute(FRDGAllocator& Allocator);
-	FRDGBarrierBatchBegin& GetEpilogueBarriersToBeginForAll(FRDGAllocator& Allocator);
+	FRDGBarrierBatchBegin& GetPrologueBarriersToBegin(FRDGAllocator& Allocator, FRDGTransitionCreateQueue& CreateQueue);
+	FRDGBarrierBatchBegin& GetEpilogueBarriersToBeginForGraphics(FRDGAllocator& Allocator, FRDGTransitionCreateQueue& CreateQueue);
+	FRDGBarrierBatchBegin& GetEpilogueBarriersToBeginForAsyncCompute(FRDGAllocator& Allocator, FRDGTransitionCreateQueue& CreateQueue);
+	FRDGBarrierBatchBegin& GetEpilogueBarriersToBeginForAll(FRDGAllocator& Allocator, FRDGTransitionCreateQueue& CreateQueue);
 
-	FRDGBarrierBatchBegin& GetEpilogueBarriersToBeginFor(FRDGAllocator& Allocator, ERHIPipeline PipelineForEnd)
+	FRDGBarrierBatchBegin& GetEpilogueBarriersToBeginFor(FRDGAllocator& Allocator, FRDGTransitionCreateQueue& CreateQueue, ERHIPipeline PipelineForEnd)
 	{
 		switch (PipelineForEnd)
 		{
@@ -263,27 +275,20 @@ protected:
 			// fall through
 
 		case ERHIPipeline::Graphics:
-			return GetEpilogueBarriersToBeginForGraphics(Allocator);
+			return GetEpilogueBarriersToBeginForGraphics(Allocator, CreateQueue);
 
 		case ERHIPipeline::AsyncCompute:
-			return GetEpilogueBarriersToBeginForAsyncCompute(Allocator);
+			return GetEpilogueBarriersToBeginForAsyncCompute(Allocator, CreateQueue);
 
 		case ERHIPipeline::All:
-			return GetEpilogueBarriersToBeginForAll(Allocator);
+			return GetEpilogueBarriersToBeginForAll(Allocator, CreateQueue);
 		}
 	}
 
 	FRDGBarrierBatchEnd& GetPrologueBarriersToEnd(FRDGAllocator& Allocator);
 	FRDGBarrierBatchEnd& GetEpilogueBarriersToEnd(FRDGAllocator& Allocator);
 
-	//////////////////////////////////////////////////////////////////////////
-	//! User Methods to Override
-
-	virtual void ExecuteImpl(FRHIComputeCommandList& RHICmdList) = 0;
-
-	//////////////////////////////////////////////////////////////////////////
-
-	void Execute(FRHIComputeCommandList& RHICmdList);
+	virtual void Execute(FRHIComputeCommandList& RHICmdList) {}
 
 	// When r.RDG.Debug is enabled, this will include a full namespace path with event scopes included.
 	IF_RDG_ENABLE_DEBUG(FString FullPathIfDebug);
@@ -322,6 +327,18 @@ protected:
 			/** Whether this pass has non-RDG UAV outputs. */
 			uint32 bHasExternalOutputs : 1;
 
+			/** Whether this pass is a sentinel (prologue / epilogue) pass. */
+			uint32 bSentinel : 1;
+
+			/** Whether this pass has been culled. */
+			uint32 bCulled : 1;
+
+			/** Whether this pass does not contain parameters. */
+			uint32 bEmptyParameters : 1;
+
+			/** If set, dispatches to the RHI thread before executing this pass. */
+			uint32 bDispatchAfterExecute : 1;
+
 			/** Whether this pass allocated a texture through the pool. */
 			IF_RDG_ENABLE_DEBUG(uint32 bFirstTextureAllocated : 1);
 		};
@@ -390,8 +407,8 @@ protected:
 
 	/** Split-barrier batches at various points of execution of the pass. */
 	FRDGBarrierBatchBegin* PrologueBarriersToBegin = nullptr;
-	FRDGBarrierBatchEnd* PrologueBarriersToEnd = nullptr;
-	FRDGBarrierBatchBegin* EpilogueBarriersToBeginForGraphics = nullptr;
+	FRDGBarrierBatchEnd PrologueBarriersToEnd;
+	FRDGBarrierBatchBegin EpilogueBarriersToBeginForGraphics;
 	FRDGBarrierBatchBegin* EpilogueBarriersToBeginForAsyncCompute = nullptr;
 	FRDGBarrierBatchBegin* EpilogueBarriersToBeginForAll = nullptr;
 	TArray<FRDGBarrierBatchBegin*, FRDGArrayAllocator> SharedEpilogueBarriersToBegin;
@@ -420,6 +437,7 @@ protected:
 	friend FRDGBuilder;
 	friend FRDGPassRegistry;
 	friend FRDGTrace;
+	friend FRDGUserValidation;
 };
 
 /** Render graph pass with lambda execute function. */
@@ -470,9 +488,11 @@ public:
 	}
 
 private:
-	void ExecuteImpl(FRHIComputeCommandList& RHICmdList) override
+	void Execute(FRHIComputeCommandList& RHICmdList) override
 	{
 		check(!kSupportsRaster || RHICmdList.IsImmediate());
+		QUICK_SCOPE_CYCLE_COUNTER(STAT_FRDGPass_Execute);
+		RHICmdList.SetStaticUniformBuffers(ParameterStruct.GetStaticUniformBuffers());
 		ExecuteLambda(static_cast<TRHICommandList&>(RHICmdList));
 	}
 
@@ -500,12 +520,13 @@ class FRDGSentinelPass final
 	: public FRDGPass
 {
 public:
-	FRDGSentinelPass(FRDGEventName&& Name)
-		: FRDGPass(MoveTemp(Name), FRDGParameterStruct(&EmptyShaderParameters), ERDGPassFlags::NeverCull)
-	{}
+	FRDGSentinelPass(FRDGEventName&& Name, ERDGPassFlags InPassFlagsToAdd = ERDGPassFlags::None)
+		: FRDGPass(MoveTemp(Name), FRDGParameterStruct(&EmptyShaderParameters), ERDGPassFlags::NeverCull | InPassFlagsToAdd)
+	{
+		bSentinel = 1;
+	}
 
 private:
-	void ExecuteImpl(FRHIComputeCommandList&) override {}
 	FEmptyShaderParameters EmptyShaderParameters;
 };
 
diff --git a/Engine/Source/Runtime/RenderCore/Public/RenderGraphResources.h b/Engine/Source/Runtime/RenderCore/Public/RenderGraphResources.h
--- a/Engine/Source/Runtime/RenderCore/Public/RenderGraphResources.h
+++ b/Engine/Source/Runtime/RenderCore/Public/RenderGraphResources.h
@@ -77,6 +77,9 @@ using FRDGTextureSubresourceState = TRDGTextureSubresourceArray<FRDGSubresourceS
 using FRDGTextureTransientSubresourceState = TRDGTextureSubresourceArray<FRDGSubresourceState, FRDGArrayAllocator>;
 using FRDGTextureTransientSubresourceStateIndirect = TRDGTextureSubresourceArray<FRDGSubresourceState*, FRDGArrayAllocator>;
 
+using FRDGPooledTextureArray = TArray<TRefCountPtr<FPooledRenderTarget>, FRDGArrayAllocator>;
+using FRDGPooledBufferArray = TArray<TRefCountPtr<FRDGPooledBuffer>, FRDGArrayAllocator>;
+
 /** Generic graph resource. */
 class RENDERCORE_API FRDGResource
 {
@@ -115,6 +118,11 @@ protected:
 		return ResourceRHI;
 	}
 
+	bool HasRHI() const
+	{
+		return ResourceRHI != nullptr;
+	}
+
 	FRHIResource* ResourceRHI = nullptr;
 
 #if RDG_ENABLE_DEBUG
@@ -165,10 +173,17 @@ protected:
 	{}
 
 private:
+	FRHIUniformBuffer* GetRHIUnchecked() const
+	{
+		return static_cast<FRHIUniformBuffer*>(FRDGResource::GetRHIUnchecked());
+	}
+
+	void InitRHI();
 
 	const FRDGParameterStruct ParameterStruct;
 	TRefCountPtr<FRHIUniformBuffer> UniformBufferRHI;
 	FRDGUniformBufferHandle Handle;
+	bool bQueuedForCreate = false;
 
 	friend FRDGBuilder;
 	friend FRDGUniformBufferRegistry;
@@ -244,7 +259,6 @@ public:
 
 protected:
 	FRDGParentResource(const TCHAR* InName, ERDGParentResourceType InType);
-	~FRDGParentResource();
 
 	/** Whether this is an externally registered resource. */
 	uint8 bExternal : 1;
@@ -393,11 +407,7 @@ public:
 
 	/** Returns the allocated pooled render target. */
 	UE_DEPRECATED(5.0, "Accessing the underlying pooled render target has been deprecated. Use GetRHI() instead.")
-	IPooledRenderTarget* GetPooledRenderTarget() const
-	{
-		IF_RDG_ENABLE_DEBUG(ValidateRHIAccess());
-		return PooledRenderTarget;
-	}
+	IPooledRenderTarget* GetPooledRenderTarget() const;
 
 	/** Returns the allocated RHI texture. */
 	FRHITexture* GetRHI() const
@@ -419,7 +429,7 @@ public:
 
 	FRDGTextureSubresourceRange GetSubresourceRange() const
 	{
-		return FRDGTextureSubresourceRange(Layout);
+		return WholeRange;
 	}
 
 	FORCEINLINE uint32 GetSubresourceCount() const
@@ -436,6 +446,7 @@ private:
 		, Flags(InFlags)
 		, RenderTargetTexture(InRenderTargetTexture)
 		, Layout(InDesc)
+		, WholeRange(Layout)
 	{
 		const uint32 SubresourceCount = GetSubresourceCount();
 		MergeState.SetNum(SubresourceCount);
@@ -452,7 +463,7 @@ private:
 	void SetRHI(FRHITransientTexture* TransientTexture, FRDGAllocator& Allocator);
 
 	/** Finalizes the texture for execution; no other transitions are allowed after calling this. */
-	void Finalize();
+	void Finalize(FRDGPooledTextureArray& PooledTextureArray);
 
 	/** Returns RHI texture without access checks. */
 	FRHITexture* GetRHIUnchecked() const
@@ -476,17 +487,21 @@ private:
 	/** Describes which RHI texture this RDG texture represents on a pooled texture. Must be default unless the texture is externally registered. */
 	const ERenderTargetTexture RenderTargetTexture;
 
-	/** The layout used to facilitate subresource transitions. */
-	FRDGTextureSubresourceLayout Layout;
-
 	/** The next texture to own the PooledTexture allocation during execution. */
 	FRDGTextureHandle NextOwner;
 
 	/** The handle registered with the builder. */
 	FRDGTextureHandle Handle;
 
+	/** The layout used to facilitate subresource transitions. */
+	FRDGTextureSubresourceLayout Layout;
+	FRDGTextureSubresourceRange  WholeRange;
+
+	/** Cached state pointer from the pooled texture. */
+	FRDGTextureSubresourceState* State = nullptr;
+
 	/** The assigned pooled render target to use during execution. Never reset. */
-	IPooledRenderTarget* PooledRenderTarget = nullptr;
+	FPooledRenderTarget* PooledRenderTarget = nullptr;
 
 	union
 	{
@@ -500,9 +515,6 @@ private:
 	/** The assigned view cache for this texture (sourced from transient / pooled texture). Never reset. */
 	FRHITextureViewCache* ViewCache = nullptr;
 
-	/** Cached state pointer from the pooled texture. */
-	FRDGTextureSubresourceState* State = nullptr;
-
 	/** Valid strictly when holding a strong reference; use PooledRenderTarget instead. */
 	TRefCountPtr<IPooledRenderTarget> Allocation;
 
@@ -965,7 +977,7 @@ private:
 	void SetRHI(FRHITransientBuffer* InTransientBuffer, FRDGAllocator& Allocator);
 
 	/** Finalizes the buffer for execution; no other transitions are allowed after calling this. */
-	void Finalize();
+	void Finalize(FRDGPooledBufferArray& PooledBufferArray);
 
 	FRHIBuffer* GetRHIUnchecked() const
 	{
diff --git a/Engine/Source/Runtime/RenderCore/Public/RenderGraphResources.inl b/Engine/Source/Runtime/RenderCore/Public/RenderGraphResources.inl
--- a/Engine/Source/Runtime/RenderCore/Public/RenderGraphResources.inl
+++ b/Engine/Source/Runtime/RenderCore/Public/RenderGraphResources.inl
@@ -37,38 +37,12 @@ inline bool FRDGSubresourceState::IsUsedBy(ERHIPipeline Pipeline) const
 
 inline FRDGPassHandle FRDGSubresourceState::GetLastPass() const
 {
-	const FRDGPassHandle GraphicsPass = LastPass[ERHIPipeline::Graphics];
-	const FRDGPassHandle AsyncComputePass = LastPass[ERHIPipeline::AsyncCompute];
-
-	if (GraphicsPass.IsNull())
-	{
-		return AsyncComputePass;
-	}
-	
-	if (AsyncComputePass.IsNull())
-	{
-		return GraphicsPass;
-	}
-
-	return FMath::Max(GraphicsPass, AsyncComputePass);
+	return FRDGPassHandle::Max(LastPass[ERHIPipeline::Graphics], LastPass[ERHIPipeline::AsyncCompute]);
 }
 
 inline FRDGPassHandle FRDGSubresourceState::GetFirstPass() const
 {
-	const FRDGPassHandle GraphicsPass = FirstPass[ERHIPipeline::Graphics];
-	const FRDGPassHandle AsyncComputePass = FirstPass[ERHIPipeline::AsyncCompute];
-
-	if (GraphicsPass.IsNull())
-	{
-		return AsyncComputePass;
-	}
-
-	if (AsyncComputePass.IsNull())
-	{
-		return GraphicsPass;
-	}
-
-	return FMath::Min(GraphicsPass, AsyncComputePass);
+	return FRDGPassHandle::Min(FirstPass[ERHIPipeline::Graphics], FirstPass[ERHIPipeline::AsyncCompute]);
 }
 
 FORCEINLINE ERHIPipeline FRDGSubresourceState::GetPipelines() const
diff --git a/Engine/Source/Runtime/RenderCore/Public/RenderGraphValidation.h b/Engine/Source/Runtime/RenderCore/Public/RenderGraphValidation.h
--- a/Engine/Source/Runtime/RenderCore/Public/RenderGraphValidation.h
+++ b/Engine/Source/Runtime/RenderCore/Public/RenderGraphValidation.h
@@ -63,6 +63,8 @@ public:
 	void ValidateExtractTexture(FRDGTextureRef Texture, TRefCountPtr<IPooledRenderTarget>* OutTexturePtr);
 	void ValidateExtractBuffer(FRDGBufferRef Buffer, TRefCountPtr<FRDGPooledBuffer>* OutBufferPtr);
 
+	void ValidateConvertToExternalResource(FRDGParentResourceRef Resource);
+
 	/** Tracks and validates the addition of a new pass to the graph.
 	 *  @param bSkipPassAccessMarking Skips marking the pass as a producer or incrementing the pass access. Useful when
 	 *      the builder needs to inject a pass for debugging while preserving error messages and warnings for the original
@@ -111,12 +113,11 @@ private:
 
 	/** Whether the Execute() has already been called. */
 	bool bHasExecuted = false;
+	bool bHasExecuteBegun = false;
 
-	void MemStackGuard();
 	void ExecuteGuard(const TCHAR* Operation, const TCHAR* ResourceName);
 
 	FRDGAllocator& Allocator;
-	int32 ExpectedNumMarks;
 };
 
 /** This class validates and logs barriers submitted by the graph. */
@@ -151,14 +152,11 @@ private:
 class FRDGLogFile
 {
 public:
-	FRDGLogFile() = default;
+	FRDGLogFile(const FRDGPassRegistry& InPasses)
+		: Passes(InPasses)
+	{}
 
-	void Begin(
-		const FRDGEventName& GraphName,
-		const FRDGPassRegistry* InPassRegistry,
-		FRDGPassBitArray InPassesCulled,
-		FRDGPassHandle InProloguePassHandle,
-		FRDGPassHandle InEpiloguePassHandle);
+	void Begin(const FRDGEventName& GraphName);
 
 	void AddFirstEdge(const FRDGTextureRef Texture, FRDGPassHandle FirstPass);
 
@@ -193,15 +191,12 @@ private:
 
 	bool bOpen = false;
 
+	const FRDGPassRegistry& Passes;
+
 	TSet<FRDGPassHandle> PassesReferenced;
 	TArray<const FRDGTexture*> Textures;
 	TArray<const FRDGBuffer*> Buffers;
 
-	const FRDGPassRegistry* Passes = nullptr;
-	FRDGPassBitArray PassesCulled;
-	FRDGPassHandle ProloguePassHandle;
-	FRDGPassHandle EpiloguePassHandle;
-
 	FString Indentation;
 	FString File;
 	FString GraphName;
diff --git a/Engine/Source/Runtime/Renderer/Private/CustomDepthRendering.cpp b/Engine/Source/Runtime/Renderer/Private/CustomDepthRendering.cpp
--- a/Engine/Source/Runtime/Renderer/Private/CustomDepthRendering.cpp
+++ b/Engine/Source/Runtime/Renderer/Private/CustomDepthRendering.cpp
@@ -108,12 +108,17 @@ FCustomDepthTextures FCustomDepthTextures::Create(FRDGBuilder& GraphBuilder, FIn
 		const FRDGTextureDesc MobileCustomStencilDesc = FRDGTextureDesc::Create2D(CustomDepthExtent, PF_G8, FClearValueBinding::Transparent, MobileCustomStencilFlags);
 
 		CustomDepthTextures.MobileDepth = GraphBuilder.CreateTexture(MobileCustomDepthDesc, TEXT("MobileCustomDepth"));
+		CustomDepthTextures.MobileDepth->SetNonTransient();
+
 		CustomDepthTextures.MobileStencil = GraphBuilder.CreateTexture(MobileCustomStencilDesc, TEXT("MobileCustomStencil"));
+		CustomDepthTextures.MobileStencil->SetNonTransient();
 	}
 
 	const FRDGTextureDesc CustomDepthDesc = FRDGTextureDesc::Create2D(CustomDepthExtent, PF_DepthStencil, FClearValueBinding::DepthFar, GFastVRamConfig.CustomDepth | TexCreate_NoFastClear | TexCreate_DepthStencilTargetable | TexCreate_ShaderResource);
 
 	CustomDepthTextures.Depth = GraphBuilder.CreateTexture(CustomDepthDesc, TEXT("CustomDepth"));
+	CustomDepthTextures.Depth->SetNonTransient();
+
 	CustomDepthTextures.Stencil = GraphBuilder.CreateSRV(FRDGTextureSRVDesc::CreateWithPixelFormat(CustomDepthTextures.Depth, PF_X24_G8));
 	CustomDepthTextures.DepthAction = ERenderTargetLoadAction::EClear;
 	CustomDepthTextures.StencilAction = bWritesCustomStencil ? ERenderTargetLoadAction::EClear : ERenderTargetLoadAction::ENoAction;
diff --git a/Engine/Source/Runtime/Renderer/Private/DeferredShadingRenderer.cpp b/Engine/Source/Runtime/Renderer/Private/DeferredShadingRenderer.cpp
--- a/Engine/Source/Runtime/Renderer/Private/DeferredShadingRenderer.cpp
+++ b/Engine/Source/Runtime/Renderer/Private/DeferredShadingRenderer.cpp
@@ -237,7 +237,6 @@ DECLARE_CYCLE_STAT(TEXT("InitViews Intentional Stall"), STAT_InitViews_Intention
 
 DECLARE_CYCLE_STAT(TEXT("DeferredShadingSceneRenderer UpdateDownsampledDepthSurface"), STAT_FDeferredShadingSceneRenderer_UpdateDownsampledDepthSurface, STATGROUP_SceneRendering);
 DECLARE_CYCLE_STAT(TEXT("DeferredShadingSceneRenderer Render Init"), STAT_FDeferredShadingSceneRenderer_Render_Init, STATGROUP_SceneRendering);
-DECLARE_CYCLE_STAT(TEXT("DeferredShadingSceneRenderer Render ServiceLocalQueue"), STAT_FDeferredShadingSceneRenderer_Render_ServiceLocalQueue, STATGROUP_SceneRendering);
 DECLARE_CYCLE_STAT(TEXT("DeferredShadingSceneRenderer FGlobalDynamicVertexBuffer Commit"), STAT_FDeferredShadingSceneRenderer_FGlobalDynamicVertexBuffer_Commit, STATGROUP_SceneRendering);
 DECLARE_CYCLE_STAT(TEXT("DeferredShadingSceneRenderer FXSystem PreRender"), STAT_FDeferredShadingSceneRenderer_FXSystem_PreRender, STATGROUP_SceneRendering);
 DECLARE_CYCLE_STAT(TEXT("DeferredShadingSceneRenderer AllocGBufferTargets"), STAT_FDeferredShadingSceneRenderer_AllocGBufferTargets, STATGROUP_SceneRendering);
@@ -407,32 +406,6 @@ bool FDeferredShadingSceneRenderer::RenderHzb(FRDGBuilder& GraphBuilder, FRDGTex
 	return FamilyPipelineState->bHZBOcclusion;
 }
 
-// The render thread is involved in sending stuff to the RHI, so we will periodically service that queue
-void ServiceLocalQueue()
-{
-	SCOPE_CYCLE_COUNTER(STAT_FDeferredShadingSceneRenderer_Render_ServiceLocalQueue);
-	FTaskGraphInterface::Get().ProcessThreadUntilIdle(ENamedThreads::GetRenderThread_Local());
-
-	if (IsRunningRHIInSeparateThread())
-	{
-		FRHICommandListExecutor::GetImmediateCommandList().ImmediateFlush(EImmediateFlushType::DispatchToRHIThread);
-	}
-}
-
-void AddServiceLocalQueuePass(FRDGBuilder& GraphBuilder)
-{
-	AddPass(GraphBuilder, [](FRHICommandListImmediate& RHICmdList)
-	{
-		SCOPE_CYCLE_COUNTER(STAT_FDeferredShadingSceneRenderer_Render_ServiceLocalQueue);
-		FTaskGraphInterface::Get().ProcessThreadUntilIdle(ENamedThreads::GetRenderThread_Local());
-
-		if (IsRunningRHIInSeparateThread())
-		{
-			RHICmdList.ImmediateFlush(EImmediateFlushType::DispatchToRHIThread);
-		}
-	});
-}
-
 BEGIN_SHADER_PARAMETER_STRUCT(FRenderOpaqueFXPassParameters, )
 	SHADER_PARAMETER_RDG_UNIFORM_BUFFER(FSceneTextureUniformParameters, SceneTextures)
 END_SHADER_PARAMETER_STRUCT()
@@ -472,7 +445,7 @@ static void RenderOpaqueFX(
 			GPUSortManager->OnPostRenderOpaque(GraphBuilder);
 		}
 
-		ServiceLocalQueue();
+		GraphBuilder.AddDispatchHint();
 	}
 }
 
@@ -1851,8 +1824,8 @@ void FDeferredShadingSceneRenderer::Render(FRDGBuilder& GraphBuilder)
 		return;
 	}
 
-	RDG_RHI_EVENT_SCOPE(GraphBuilder, Scene);
-	RDG_RHI_GPU_STAT_SCOPE(GraphBuilder, Unaccounted);
+	RDG_EVENT_SCOPE(GraphBuilder, "Scene");
+	RDG_GPU_STAT_SCOPE(GraphBuilder, Unaccounted);
 	
 	{
 		SCOPE_CYCLE_COUNTER(STAT_FDeferredShadingSceneRenderer_Render_Init);
@@ -2191,25 +2164,42 @@ void FDeferredShadingSceneRenderer::Render(FRDGBuilder& GraphBuilder)
 			RenderVelocities(GraphBuilder, SceneTextures, EVelocityPass::Opaque, bHairEnable);
 			GraphBuilder.SetCommandListStat(GET_STATID(STAT_CLM_AfterVelocity));
 		}
+	}
 
-		if (bDoInitViewAftersPrepass)
+	for (int32 ViewIndex = 0; ViewIndex < Views.Num(); ++ViewIndex)
+	{
+		const FViewInfo& View = Views[ViewIndex];
+
+		if (((View.FinalPostProcessSettings.DynamicGlobalIlluminationMethod == EDynamicGlobalIlluminationMethod::ScreenSpace && ScreenSpaceRayTracing::ShouldKeepBleedFreeSceneColor(View))
+			|| GetViewPipelineState(View).DiffuseIndirectMethod == EDiffuseIndirectMethod::Lumen)
+			&& !View.bStatePrevViewInfoIsReadOnly)
 		{
-			{
-				RDG_RHI_GPU_STAT_SCOPE(GraphBuilder, VisibilityCommands);
-				InitViewsPossiblyAfterPrepass(GraphBuilder, ILCTaskData, InstanceCullingManager);
-			}
+			// Keep scene color and depth for next frame screen space ray tracing.
+			FSceneViewState* ViewState = View.ViewState;
+			ViewState->PrevFrameViewInfo.DepthBuffer = GraphBuilder.ConvertToExternalTexture(SceneTextures.Depth.Resolve);
+			ViewState->PrevFrameViewInfo.ScreenSpaceRayTracingInput = GraphBuilder.ConvertToExternalTexture(SceneTextures.Color.Resolve);
+		}
+	}
 
-			{
-				RDG_RHI_GPU_STAT_SCOPE(GraphBuilder, GPUSceneUpdate);
-				PrepareDistanceFieldScene(GraphBuilder, false);
-			}
+	if (bDoInitViewAftersPrepass)
+	{
+		GraphBuilder.Drain();
 
-			{
-				SCOPE_CYCLE_COUNTER(STAT_FDeferredShadingSceneRenderer_FGlobalDynamicVertexBuffer_Commit);
-				DynamicVertexBufferForInitShadows.Commit();
-				DynamicIndexBufferForInitShadows.Commit();
-				DynamicReadBufferForInitShadows.Commit();
-			}
+		{
+			RDG_RHI_GPU_STAT_SCOPE(GraphBuilder, VisibilityCommands);
+			InitViewsPossiblyAfterPrepass(GraphBuilder, ILCTaskData, InstanceCullingManager);
+		}
+
+		{
+			RDG_RHI_GPU_STAT_SCOPE(GraphBuilder, GPUSceneUpdate);
+			PrepareDistanceFieldScene(GraphBuilder, false);
+		}
+
+		{
+			SCOPE_CYCLE_COUNTER(STAT_FDeferredShadingSceneRenderer_FGlobalDynamicVertexBuffer_Commit);
+			DynamicVertexBufferForInitShadows.Commit();
+			DynamicIndexBufferForInitShadows.Commit();
+			DynamicReadBufferForInitShadows.Commit();
 		}
 	}
 
@@ -2361,6 +2351,8 @@ void FDeferredShadingSceneRenderer::Render(FRDGBuilder& GraphBuilder)
 		{
 			CompositionLightingAsyncResults = CompositionLighting::ProcessAsync(GraphBuilder, Views, SceneTextures);
 		}
+
+		GraphBuilder.Drain();
 	};
 
 	// Early occlusion queries
@@ -2376,7 +2368,6 @@ void FDeferredShadingSceneRenderer::Render(FRDGBuilder& GraphBuilder)
 		RenderOcclusionLambda();
 	}
 
-	AddServiceLocalQueuePass(GraphBuilder);
 	// End early occlusion queries
 
 	// Early Shadow depth rendering
@@ -2517,8 +2508,8 @@ void FDeferredShadingSceneRenderer::Render(FRDGBuilder& GraphBuilder)
 
 	{
 		RenderBasePass(GraphBuilder, SceneTextures, DBufferTextures, BasePassDepthStencilAccess, ForwardScreenSpaceShadowMaskTexture, InstanceCullingManager);
-		AddServiceLocalQueuePass(GraphBuilder);
-		
+		GraphBuilder.AddDispatchHint();
+
 		if (bNaniteEnabled && bShouldApplyNaniteMaterials)
 		{
 			for (int32 ViewIndex = 0; ViewIndex < Views.Num(); ++ViewIndex)
@@ -3158,28 +3149,13 @@ void FDeferredShadingSceneRenderer::Render(FRDGBuilder& GraphBuilder)
 	DoCrossGPUTransfers(GraphBuilder, RenderTargetGPUMask, ViewFamilyTexture);
 #endif
 
-	for (int32 ViewIndex = 0; ViewIndex < Views.Num(); ++ViewIndex)
-	{
-		const FViewInfo& View = Views[ViewIndex];
-
-		if (((View.FinalPostProcessSettings.DynamicGlobalIlluminationMethod == EDynamicGlobalIlluminationMethod::ScreenSpace && ScreenSpaceRayTracing::ShouldKeepBleedFreeSceneColor(View))
-				|| GetViewPipelineState(View).DiffuseIndirectMethod == EDiffuseIndirectMethod::Lumen)
-			&& !View.bStatePrevViewInfoIsReadOnly)
-		{
-			// Keep scene color and depth for next frame screen space ray tracing.
-			FSceneViewState* ViewState = View.ViewState;
-			GraphBuilder.QueueTextureExtraction(SceneTextures.Depth.Resolve, &ViewState->PrevFrameViewInfo.DepthBuffer);
-			GraphBuilder.QueueTextureExtraction(SceneTextures.Color.Resolve, &ViewState->PrevFrameViewInfo.ScreenSpaceRayTracingInput);
-		}
-	}
-
 	{
 		SCOPE_CYCLE_COUNTER(STAT_FDeferredShadingSceneRenderer_RenderFinish);
 		RDG_GPU_STAT_SCOPE(GraphBuilder, FrameRenderFinish);
 		GraphBuilder.SetCommandListStat(GET_STATID(STAT_CLM_RenderFinish));
 		RenderFinish(GraphBuilder, ViewFamilyTexture);
 		GraphBuilder.SetCommandListStat(GET_STATID(STAT_CLM_AfterFrame));
-		AddServiceLocalQueuePass(GraphBuilder);
+		GraphBuilder.AddDispatchHint();
 	}
 
 	QueueSceneTextureExtractions(GraphBuilder, SceneTextures);
diff --git a/Engine/Source/Runtime/Renderer/Private/HairStrands/HairStrandsRendering.cpp b/Engine/Source/Runtime/Renderer/Private/HairStrands/HairStrandsRendering.cpp
--- a/Engine/Source/Runtime/Renderer/Private/HairStrands/HairStrandsRendering.cpp
+++ b/Engine/Source/Runtime/Renderer/Private/HairStrands/HairStrandsRendering.cpp
@@ -80,7 +80,6 @@ static TRDGUniformBufferRef<FHairStrandsViewUniformParameters> InternalCreateHai
 
 IMPLEMENT_GLOBAL_SHADER_PARAMETER_STRUCT(FHairStrandsViewUniformParameters, "HairStrands");
 
-void AddServiceLocalQueuePass(FRDGBuilder& GraphBuilder);
 bool GetHairStrandsSkyLightingDebugEnable();
 bool IsHairStrandsAdaptiveVoxelAllocationEnable();
 
@@ -118,13 +117,12 @@ void RenderHairPrePass(
 
 		//SCOPED_GPU_STAT(RHICmdList, HairRendering);
 		CreateHairStrandsMacroGroups(GraphBuilder, Scene, View);
-		AddServiceLocalQueuePass(GraphBuilder);
+		GraphBuilder.AddDispatchHint();
 
 		// Voxelization and Deep Opacity Maps
 		VoxelizeHairStrands(GraphBuilder, Scene, View, InstanceCullingManager);
 		RenderHairStrandsDeepShadows(GraphBuilder, Scene, View, InstanceCullingManager);
-
-		AddServiceLocalQueuePass(GraphBuilder);
+		GraphBuilder.AddDispatchHint();;
 	}
 }
 
diff --git a/Engine/Source/Runtime/Renderer/Private/SceneTextures.cpp b/Engine/Source/Runtime/Renderer/Private/SceneTextures.cpp
--- a/Engine/Source/Runtime/Renderer/Private/SceneTextures.cpp
+++ b/Engine/Source/Runtime/Renderer/Private/SceneTextures.cpp
@@ -606,6 +606,7 @@ FSceneTextures& FMinimalSceneTextures::Create(FRDGBuilder& GraphBuilder, const F
 							 FRDGTextureDesc::Create2D(SceneTextures.Config.Extent, PF_DepthStencil, Config.DepthClearValue, Flags));
 		Desc.NumSamples = Config.NumSamples;
 		SceneTextures.Depth = GraphBuilder.CreateTexture(Desc, TEXT("SceneDepthZ"));
+		SceneTextures.Depth.Target->SetNonTransient();
 
 		if (Desc.NumSamples > 1)
 		{
@@ -618,6 +619,7 @@ FSceneTextures& FMinimalSceneTextures::Create(FRDGBuilder& GraphBuilder, const F
 			else
 			{
 				SceneTextures.Depth.Resolve = GraphBuilder.CreateTexture(Desc, TEXT("SceneDepthZ"));
+				SceneTextures.Depth.Resolve->SetNonTransient();
 			}
 		}
 
