Make RDG draining opt-in--it needs additional testing.

diff --git a/Engine/Source/Runtime/RenderCore/Private/RenderGraphPrivate.cpp b/Engine/Source/Runtime/RenderCore/Private/RenderGraphPrivate.cpp
--- a/Engine/Source/Runtime/RenderCore/Private/RenderGraphPrivate.cpp
+++ b/Engine/Source/Runtime/RenderCore/Private/RenderGraphPrivate.cpp
@@ -305,12 +305,12 @@ FAutoConsoleVariableRef CVarRDGUseTransientAllocator(
 	TEXT(" 1: enables the transient allocator (default)"),
 	ECVF_RenderThreadSafe);
 
-int32 GRDGDrain = 1;
+int32 GRDGDrain = 0;
 FAutoConsoleVariableRef CVarRDGDrain(
 	TEXT("r.RDG.Drain"), GRDGDrain,
 	TEXT("RDG will perform drain operations when requested.")
-	TEXT(" 0: disables draining, all work is deferred until FRDGBuilder::Execute;")
-	TEXT(" 1: enables execution of the RDG subgraph when FRDGBuilder::Drain is called. (default)"),
+	TEXT(" 0: disables draining, all work is deferred until FRDGBuilder::Execute; (default)")
+	TEXT(" 1: enables execution of the RDG subgraph when FRDGBuilder::Drain is called."),
 	ECVF_RenderThreadSafe);
 
 #if CSV_PROFILER
