Prevent non-exclusive access to GShaderFileCache (UE-115959).

#rnx

diff --git a/Engine/Source/Runtime/RenderCore/Private/ShaderCore.cpp b/Engine/Source/Runtime/RenderCore/Private/ShaderCore.cpp
--- a/Engine/Source/Runtime/RenderCore/Private/ShaderCore.cpp
+++ b/Engine/Source/Runtime/RenderCore/Private/ShaderCore.cpp
@@ -1254,7 +1254,10 @@ void FlushShaderFileCache()
 		FScopeLock ShaderHashAccessLock(&GShaderHashAccessGuard);
 		GShaderHashCache.Empty();
 	}
-	GShaderFileCache.Empty();
+	{
+		FScopeLock ScopeLock(&FileCacheCriticalSection);
+		GShaderFileCache.Empty();
+	}
 
 	if (!FPlatformProperties::RequiresCookedData())
 	{
