Disable VSMs by default on Vulkan until they can be debugged/fixed.

diff --git a/Engine/Source/Runtime/RenderCore/Public/RenderUtils.h b/Engine/Source/Runtime/RenderCore/Public/RenderUtils.h
--- a/Engine/Source/Runtime/RenderCore/Public/RenderUtils.h
+++ b/Engine/Source/Runtime/RenderCore/Public/RenderUtils.h
@@ -695,12 +695,16 @@ inline bool UseNanite(EShaderPlatform ShaderPlatform, bool bCheckForAtomicSuppor
 inline bool UseVirtualShadowMaps(EShaderPlatform ShaderPlatform, const FStaticFeatureLevel FeatureLevel)
 {
 	static const auto EnableVirtualSMCVar = IConsoleManager::Get().FindConsoleVariable(TEXT("r.Shadow.Virtual.Enable"));
+	const bool bVirtualShadowMapsEnabled = EnableVirtualSMCVar ? (EnableVirtualSMCVar->GetInt() != 0) : false;
 
-	// TODO: Add test for required RHI platform stuff like GPUScene
-	// bool bNaniteSupported = DoesPlatformSupportNanite(ShaderPlatform);
-	bool bForwardShadingEnabled = IsForwardShadingEnabled(ShaderPlatform);
+	const bool bUseGPUScene = UseGPUScene(ShaderPlatform, FeatureLevel);
+	const bool bPlatformSupportsNanite = DoesPlatformSupportNanite(ShaderPlatform);
+	const bool bForwardShadingEnabled = IsForwardShadingEnabled(ShaderPlatform);
+
+	// TODO: VSMs not currently working on Vulkan
+	const bool bVulkanSM5 = ShaderPlatform == SP_VULKAN_SM5;
 
-	return UseGPUScene(ShaderPlatform, FeatureLevel) && !bForwardShadingEnabled && (EnableVirtualSMCVar && EnableVirtualSMCVar->GetInt() != 0) && DoesPlatformSupportNanite(ShaderPlatform);
+	return bVirtualShadowMapsEnabled && bPlatformSupportsNanite && bUseGPUScene && !bForwardShadingEnabled && !bVulkanSM5;
 }
 
 /**
