Fix for RDG_{TEXTURE, BUFFER}_ACCESS asserting when used with shader parameter bindings.

diff --git a/Engine/Source/Runtime/RenderCore/Private/ShaderParameterStruct.cpp b/Engine/Source/Runtime/RenderCore/Private/ShaderParameterStruct.cpp
--- a/Engine/Source/Runtime/RenderCore/Private/ShaderParameterStruct.cpp
+++ b/Engine/Source/Runtime/RenderCore/Private/ShaderParameterStruct.cpp
@@ -57,6 +57,12 @@ struct FShaderParameterStructBindingContext
 				continue;
 			}
 
+			// Ignore RDG access types when binding to shaders.
+			if (IsRDGResourceAccessType(BaseType))
+			{
+				continue;
+			}
+
 			// Compute the shader member name to look for according to nesting.
 			FString ShaderBindingName = FString::Printf(TEXT("%s%s"), MemberPrefix, Member.GetName());
 
@@ -72,7 +78,7 @@ struct FShaderParameterStructBindingContext
 				BaseType == UBMT_SAMPLER);
 
 			const bool bIsRDGResource =
-				IsRDGResourceReferenceShaderParameterType(BaseType) && !IsRDGResourceAccessType(BaseType);
+				IsRDGResourceReferenceShaderParameterType(BaseType);
 
 			const bool bIsVariableNativeType = (
 				BaseType == UBMT_INT32 ||
