Minor RDG fixes and support for read-only system textures.

diff --git a/Engine/Source/Runtime/RenderCore/Private/RenderGraphBuilder.cpp b/Engine/Source/Runtime/RenderCore/Private/RenderGraphBuilder.cpp
--- a/Engine/Source/Runtime/RenderCore/Private/RenderGraphBuilder.cpp
+++ b/Engine/Source/Runtime/RenderCore/Private/RenderGraphBuilder.cpp
@@ -607,28 +607,35 @@ FRDGTextureRef FRDGBuilder::RegisterExternalTexture(
 		return FoundTexture;
 	}
 
-	FRDGTextureRef Texture = Textures.Allocate(Allocator, Name, Translate(ExternalPooledTexture->GetDesc(), RenderTargetTexture), Flags, RenderTargetTexture);
-	Texture->SetRHI(static_cast<FPooledRenderTarget*>(ExternalPooledTexture.GetReference()));
+	const FRDGTextureDesc Desc = Translate(ExternalPooledTexture->GetDesc(), RenderTargetTexture);
+	bool bFinalizedAccess = false;
 
-	const ERHIAccess AccessInitial = kDefaultAccessInitial;
+	if (!EnumHasAnyFlags(Desc.Flags, TexCreate_RenderTargetable | TexCreate_ResolveTargetable | TexCreate_DepthStencilTargetable | TexCreate_UAV | TexCreate_DepthStencilResolveTarget))
+	{
+		Flags |= ERDGTextureFlags::ReadOnly;
+		bFinalizedAccess = true;
+	}
+
+	FRDGTextureRef Texture = Textures.Allocate(Allocator, Name, Desc, Flags, RenderTargetTexture);
+	Texture->SetRHI(static_cast<FPooledRenderTarget*>(ExternalPooledTexture.GetReference()));
 
 	Texture->bExternal = true;
-	Texture->AccessInitial = AccessInitial;
-	Texture->AccessFinal = EnumHasAnyFlags(ExternalTextureRHI->GetFlags(), TexCreate_Foveation)? ERHIAccess::ShadingRateSource : kDefaultAccessFinal;
+	Texture->AccessFinal = EnumHasAnyFlags(ExternalTextureRHI->GetFlags(), TexCreate_Foveation) ? ERHIAccess::ShadingRateSource : kDefaultAccessFinal;
 	Texture->FirstPass = GetProloguePassHandle();
 
+	// Textures that are created read-only are not transitioned by RDG.
+	if (bFinalizedAccess)
+	{
+		Texture->bFinalizedAccess = 1;
+		Texture->AccessFinal = ERHIAccess::ReadOnlyExclusiveMask;
+	}
+
 	FRDGTextureSubresourceState& TextureState = Texture->GetState();
 
 	checkf(IsWholeResource(TextureState) && GetWholeResource(TextureState).Access == ERHIAccess::Unknown,
 		TEXT("Externally registered texture '%s' has known RDG state. This means the graph did not sanitize it correctly, or ")
 		TEXT("an IPooledRenderTarget reference was improperly held within a pass."), Texture->Name);
 
-	{
-		FRDGSubresourceState SubresourceState;
-		SubresourceState.Access = AccessInitial;
-		InitAsWholeResource(TextureState, SubresourceState);
-	}
-
 	ExternalTextures.Add(Texture->GetRHIUnchecked(), Texture);
 
 	IF_RDG_ENABLE_DEBUG(UserValidation.ValidateRegisterExternalTexture(Texture));
@@ -662,23 +669,34 @@ FRDGBufferRef FRDGBuilder::RegisterExternalBuffer(
 		return *FoundBufferPtr;
 	}
 
+	const FRDGBufferDesc& Desc = ExternalPooledBuffer->Desc;
+	bool bFinalizedAccess = false;
+
+	if (!EnumHasAnyFlags(Desc.Usage, BUF_UnorderedAccess))
+	{
+		Flags |= ERDGBufferFlags::ReadOnly;
+		bFinalizedAccess = true;
+	}
+
 	FRDGBufferRef Buffer = Buffers.Allocate(Allocator, Name, ExternalPooledBuffer->Desc, Flags);
 	Buffer->SetRHI(ExternalPooledBuffer);
 
-	const ERHIAccess AccessInitial = kDefaultAccessInitial;
-
 	Buffer->bExternal = true;
-	Buffer->AccessInitial = AccessInitial;
 	Buffer->AccessFinal = kDefaultAccessFinal;
 	Buffer->FirstPass = GetProloguePassHandle();
 
+	// Buffers that are created read-only are not transitioned by RDG.
+	if (bFinalizedAccess)
+	{
+		Buffer->bFinalizedAccess = 1;
+		Buffer->AccessFinal = ERHIAccess::ReadOnlyExclusiveMask;
+	}
+
 	FRDGSubresourceState& BufferState = Buffer->GetState();
 	checkf(BufferState.Access == ERHIAccess::Unknown,
 		TEXT("Externally registered buffer '%s' has known RDG state. This means the graph did not sanitize it correctly, or ")
 		TEXT("an FRDGPooledBuffer reference was improperly held within a pass."), Buffer->Name);
 
-	BufferState.Access = AccessInitial;
-
 	ExternalBuffers.Add(ExternalPooledBuffer, Buffer);
 
 	IF_RDG_ENABLE_DEBUG(UserValidation.ValidateRegisterExternalBuffer(Buffer));
@@ -688,20 +706,18 @@ FRDGBufferRef FRDGBuilder::RegisterExternalBuffer(
 
 void FRDGBuilder::AddPassDependency(FRDGPassHandle ProducerHandle, FRDGPassHandle ConsumerHandle)
 {
-	checkf(ProducerHandle.IsValid(), TEXT("AddPassDependency called with null producer."));
-	checkf(ConsumerHandle.IsValid(), TEXT("AddPassDependency called with null consumer."));
 	FRDGPass* Consumer = Passes[ConsumerHandle];
 
 	auto& Producers = Consumer->Producers;
 	if (Producers.Find(ProducerHandle) == INDEX_NONE)
 	{
 		Producers.Add(ProducerHandle);
+	}
 
 #if STATS
 		GRDGStatPassDependencyCount++;
 #endif
-	}
-};
+}
 
 void FRDGBuilder::Compile(EExecuteMode ExecuteMode)
 {
@@ -900,7 +916,7 @@ void FRDGBuilder::Compile(EExecuteMode ExecuteMode)
 
 			const auto MergeSubresourceStates = [&](ERDGParentResourceType ResourceType, FRDGSubresourceState*& PassMergeState, FRDGSubresourceState*& ResourceMergeState, const FRDGSubresourceState& PassState)
 			{
-				if (!ResourceMergeState || !FRDGSubresourceState::IsMergeAllowed(ResourceType, *ResourceMergeState, PassState))
+				if (!ResourceMergeState || !FRDGSubresourceState::IsMergeAllowed(ResourceType, *ResourceMergeState, PassState) || ResourceMergeState->ExecuteGeneration != ExecuteGeneration)
 				{
 					// Cross-pipeline, non-mergable state changes require a new pass dependency for fencing purposes.
 					if (ResourceMergeState)
@@ -1312,8 +1328,6 @@ void FRDGBuilder::Compile(EExecuteMode ExecuteMode)
 
 		FRDGPassHandle CurrentGraphicsJoinPassHandle;
 
-		bool bFoundAsyncComputeEndExecutePass = false;
-
 		for (FRDGPassHandle PassHandle = EpiloguePassHandle - 1; PassHandle > ProloguePassHandle; --PassHandle)
 		{
 			if (!PassesOnAsyncCompute[PassHandle])
@@ -1346,20 +1360,8 @@ void FRDGBuilder::Compile(EExecuteMode ExecuteMode)
 			{
 				CurrentGraphicsJoinPassHandle = GraphicsJoinPassHandle;
 				InsertAsyncComputeToGraphicsJoin(AsyncComputePass, GraphicsJoinPass);
-
-				if (!bFoundAsyncComputeEndExecutePass)
-				{
-					bFoundAsyncComputeEndExecutePass = true;
-
-					AsyncComputePass->bAsyncComputeEndExecute = 1;
-				}
 			}
 		}
-
-		if (!bFoundAsyncComputeEndExecutePass && ExecuteMode == EExecuteMode::Final)
-		{
-			ProloguePasses[ERHIPipeline::AsyncCompute]->bAsyncComputeEndExecute = 1;
-		}
 	}
 }
 
@@ -1389,7 +1391,7 @@ void FRDGBuilder::Execute(EExecuteMode ExecuteMode)
 	const FRDGPassHandle ProloguePassHandle = GetProloguePassHandle();
 	const FRDGPassHandle EpiloguePassHandle = GetEpiloguePassHandle();
 
-	const bool bFirstExecute = ExecuteCount == 0;
+	const bool bFirstExecute = ExecuteGeneration == 0;
 	const bool bFinalExecute = ExecuteMode == EExecuteMode::Final;
 
 	if (bFirstExecute)
@@ -1542,20 +1544,6 @@ void FRDGBuilder::Execute(EExecuteMode ExecuteMode)
 			}
 		});
 	}
-	else
-	{
-		TRACE_CPUPROFILER_EVENT_SCOPE(FRDGBuilder::Finalize);
-
-		Textures.Enumerate([&](FRDGTextureRef Texture)
-		{
-			FMemory::Memzero(Texture->MergeState.GetData(), Texture->MergeState.Num() * sizeof(FRDGSubresourceState*));
-		});
-
-		Buffers.Enumerate([&](FRDGBufferRef Buffer)
-		{
-			Buffer->MergeState = nullptr;
-		});
-	}
 
 	IF_RDG_ENABLE_DEBUG(GRDGAllowRHIAccess = false);
 
@@ -1643,10 +1631,15 @@ void FRDGBuilder::Execute(EExecuteMode ExecuteMode)
 	}
 	else
 	{
+		if (!PassesOnAsyncComputeToJoin.IsEmpty())
+		{
+			FRHIAsyncComputeCommandListImmediate::ImmediateDispatch(RHICmdListAsyncCompute);
+		}
+
 		AddProloguePass();
 	}
 
-	ExecuteCount++;
+	ExecuteGeneration++;
 }
 
 void FRDGBuilder::Clear()
@@ -1876,7 +1869,6 @@ void FRDGBuilder::SetupPassInternal(FRDGPass* Pass, FRDGPassHandle PassHandle, E
 		check(!EnumHasAnyFlags(PassPipeline, ERHIPipeline::AsyncCompute));
 
 		BeginResourcesRHI(Pass, PassHandle);
-		EndResourcesRHI(Pass, PassHandle);
 		CollectPassBarriers(Pass, PassHandle);
 		CreatePassBarriers();
 		CreateUniformBuffers();
@@ -1909,6 +1901,12 @@ void FRDGBuilder::AddProloguePass()
 	FRDGPass* GraphicsPass = SetupEmptyPass(Passes.Allocate<FRDGSentinelPass>(Allocator, RDG_EVENT_NAME("Graph Prologue (Graphics)")));
 	ProloguePasses[ERHIPipeline::Graphics] = GraphicsPass;
 	ProloguePassHandles[ERHIPipeline::Graphics] = GraphicsPass->Handle;
+
+	// If we don't need an async compute prologue pass, we still want a valid handle to make ClampToPrologue not always return null.
+	if (PassesOnAsyncComputeToJoin.IsEmpty())
+	{
+		ProloguePassHandles[ERHIPipeline::AsyncCompute] = GraphicsPass->Handle;
+	}
 }
 
 void FRDGBuilder::ExecutePassPrologue(FRHIComputeCommandList& RHICmdListPass, FRDGPass* Pass)
@@ -2053,10 +2051,6 @@ void FRDGBuilder::ExecutePass(FRDGPass* Pass)
 
 	if (Pass->bAsyncComputeEnd)
 	{
-		if (Pass->bAsyncComputeEndExecute)
-		{
-			IF_RDG_GPU_SCOPES(GPUScopeStacks.AsyncCompute.EndExecute());
-		}
 		FRHIAsyncComputeCommandListImmediate::ImmediateDispatch(RHICmdListAsyncCompute);
 	}
 
@@ -2295,19 +2289,6 @@ void FRDGBuilder::AddTransition(FRDGPassHandle PassHandle, FRDGTextureRef Textur
 		}
 	};
 
-	const auto MergeSubresourceState = [&] (FRDGSubresourceState& SubresourceStateBefore, const FRDGSubresourceState& SubresourceStateAfter)
-	{
-		SubresourceStateBefore = SubresourceStateAfter;
-
-		for (ERHIPipeline Pipeline : GetRHIPipelines())
-		{
-			if (SubresourceStateBefore.IsUsedBy(Pipeline))
-			{
-				SubresourceStateBefore.FirstPass[Pipeline] = PassHandle;
-			}
-		}
-	};
-
 	if (IsWholeResource(StateBefore))
 	{
 		// 1 -> 1
@@ -2317,7 +2298,7 @@ void FRDGBuilder::AddTransition(FRDGPassHandle PassHandle, FRDGTextureRef Textur
 			{
 				FRDGSubresourceState& SubresourceStateBefore = GetWholeResource(StateBefore);
 				AddSubresourceTransition(SubresourceStateBefore, *SubresourceStateAfter, nullptr);
-				MergeSubresourceState(SubresourceStateBefore, *SubresourceStateAfter);
+				SubresourceStateBefore = *SubresourceStateAfter;
 			}
 		}
 		// 1 -> N
@@ -2331,7 +2312,7 @@ void FRDGBuilder::AddTransition(FRDGPassHandle PassHandle, FRDGTextureRef Textur
 				{
 					AddSubresourceTransition(SubresourceStateBeforeWhole, *SubresourceStateAfter, &Subresource);
 					FRDGSubresourceState& SubresourceStateBefore = GetSubresource(StateBefore, Layout, Subresource);
-					MergeSubresourceState(SubresourceStateBefore, *SubresourceStateAfter);
+					SubresourceStateBefore = *SubresourceStateAfter;
 				}
 			});
 		}
@@ -2349,7 +2330,7 @@ void FRDGBuilder::AddTransition(FRDGPassHandle PassHandle, FRDGTextureRef Textur
 				});
 				InitAsWholeResource(StateBefore);
 				FRDGSubresourceState& SubresourceStateBefore = GetWholeResource(StateBefore);
-				MergeSubresourceState(SubresourceStateBefore, *SubresourceStateAfter);
+				SubresourceStateBefore = *SubresourceStateAfter;
 			}
 		}
 		// N -> N
@@ -2361,7 +2342,7 @@ void FRDGBuilder::AddTransition(FRDGPassHandle PassHandle, FRDGTextureRef Textur
 				{
 					FRDGSubresourceState& SubresourceStateBefore = GetSubresource(StateBefore, Layout, Subresource);
 					AddSubresourceTransition(SubresourceStateBefore, *SubresourceStateAfter, &Subresource);
-					MergeSubresourceState(SubresourceStateBefore, *SubresourceStateAfter);
+					SubresourceStateBefore = *SubresourceStateAfter;
 				}
 			});
 		}
@@ -2561,6 +2542,13 @@ void FRDGBuilder::BeginResourceRHI(FRDGPassHandle PassHandle, FRDGTextureRef Tex
 #if RDG_ENABLE_DEBUG
 	{
 		FRDGPass* Pass = Passes[PassHandle];
+
+		// Cannot begin a resource on an async compute pass.
+		check(Pass->Pipeline == ERHIPipeline::Graphics);
+
+		// Cannot begin a resource within a merged render pass region.
+		check(GetPrologueBarrierPassHandle(PassHandle) == PassHandle);
+
 		if (!Pass->bFirstTextureAllocated)
 		{
 			GRenderTargetPool.AddPhaseEvent(Pass->GetName());
@@ -2577,8 +2565,6 @@ void FRDGBuilder::BeginResourceRHI(FRDGPassHandle PassHandle, FRDGTextureRef Tex
 			{
 				Texture->SetRHI(TransientTexture, Allocator);
 
-				check(GetPrologueBarrierPassHandle(PassHandle) == PassHandle);
-
 				FRDGSubresourceState InitialState;
 				InitialState.SetPass(ERHIPipeline::Graphics, PassHandle);
 				InitialState.Access = ERHIAccess::Discard;
@@ -2586,7 +2572,7 @@ void FRDGBuilder::BeginResourceRHI(FRDGPassHandle PassHandle, FRDGTextureRef Tex
 
 				AddToPrologueBarriers(PassHandle, [&](FRDGBarrierBatchBegin& Barriers)
 				{
-					Barriers.AddAlias(Texture, FRHITransientAliasingInfo::Acquire(Texture->GetRHIUnchecked(), TransientTexture->GetAliasingOverlaps()));
+					Barriers.AddAlias(Texture, FRHITransientAliasingInfo::Acquire(TransientTexture->GetRHI(), TransientTexture->GetAliasingOverlaps()));
 				});
 
 			#if STATS
@@ -2649,6 +2635,18 @@ void FRDGBuilder::BeginResourceRHI(FRDGPassHandle PassHandle, FRDGBufferRef Buff
 
 	check(Buffer->ReferenceCount > 0 || Buffer->bExternal || IsImmediateMode());
 
+#if RDG_ENABLE_DEBUG
+	{
+		const FRDGPass* Pass = Passes[PassHandle];
+
+		// Cannot begin a resource on an async compute pass.
+		check(Pass->Pipeline == ERHIPipeline::Graphics);
+
+		// Cannot begin a resource within a merged render pass region.
+		check(GetPrologueBarrierPassHandle(PassHandle) == PassHandle);
+	}
+#endif
+
 	// If transient then create the resource on the transient allocator. External or extracted resource can't be transient because of lifetime tracking issues.
 	if (IsTransient(Buffer))
 	{
@@ -2658,15 +2656,13 @@ void FRDGBuilder::BeginResourceRHI(FRDGPassHandle PassHandle, FRDGBufferRef Buff
 			{
 				Buffer->SetRHI(TransientBuffer, Allocator);
 
-				check(GetPrologueBarrierPassHandle(PassHandle) == PassHandle);
-
 				FRDGSubresourceState& InitialState = Buffer->GetState();
 				InitialState.SetPass(ERHIPipeline::Graphics, PassHandle);
 				InitialState.Access = ERHIAccess::Discard;
 
 				AddToPrologueBarriers(PassHandle, [&](FRDGBarrierBatchBegin& Barriers)
 				{
-					Barriers.AddAlias(Buffer, FRHITransientAliasingInfo::Acquire(Buffer->GetRHIUnchecked(), TransientBuffer->GetAliasingOverlaps()));
+					Barriers.AddAlias(Buffer, FRHITransientAliasingInfo::Acquire(TransientBuffer->GetRHI(), TransientBuffer->GetAliasingOverlaps()));
 				});
 
 			#if STATS
diff --git a/Engine/Source/Runtime/RenderCore/Private/RenderGraphEvent.cpp b/Engine/Source/Runtime/RenderCore/Private/RenderGraphEvent.cpp
--- a/Engine/Source/Runtime/RenderCore/Private/RenderGraphEvent.cpp
+++ b/Engine/Source/Runtime/RenderCore/Private/RenderGraphEvent.cpp
@@ -319,6 +319,14 @@ void FRDGGPUScopeStacksByPipeline::EndExecutePass(const FRDGPass* Pass)
 	GetScopeStacks(Pipeline).EndExecutePass();
 }
 
+ void FRDGGPUScopeStacksByPipeline::EndExecute()
+{
+	Graphics.EndExecute();
+	AsyncCompute.EndExecute();
+
+	FRHIAsyncComputeCommandListImmediate::ImmediateDispatch(static_cast<FRHIAsyncComputeCommandListImmediate&>(RHICmdListAsyncCompute));
+}
+
 #endif
 
 //////////////////////////////////////////////////////////////////////////
diff --git a/Engine/Source/Runtime/RenderCore/Private/RenderGraphResources.cpp b/Engine/Source/Runtime/RenderCore/Private/RenderGraphResources.cpp
--- a/Engine/Source/Runtime/RenderCore/Private/RenderGraphResources.cpp
+++ b/Engine/Source/Runtime/RenderCore/Private/RenderGraphResources.cpp
@@ -309,7 +309,7 @@ void FRDGBuffer::SetRHI(FRHITransientBuffer* InTransientBuffer, FRDGAllocator& A
 	ViewCache = &InTransientBuffer->ViewCache;
 	ResourceRHI = InTransientBuffer->GetRHI();
 
-	bTransient = true;;
+	bTransient = true;
 }
 
 void FRDGBuffer::Finalize(FRDGPooledBufferArray& PooledBufferArray)
diff --git a/Engine/Source/Runtime/RenderCore/Private/RenderGraphValidation.cpp b/Engine/Source/Runtime/RenderCore/Private/RenderGraphValidation.cpp
--- a/Engine/Source/Runtime/RenderCore/Private/RenderGraphValidation.cpp
+++ b/Engine/Source/Runtime/RenderCore/Private/RenderGraphValidation.cpp
@@ -1709,11 +1709,11 @@ void FRDGLogFile::AddAliasEdge(const FRDGBufferRef BufferBefore, FRDGPassHandle
 	}
 }
 
-void FRDGLogFile::AddTransitionEdge(FRDGPassHandle PassHandle, FRDGSubresourceState StateBefore, FRDGSubresourceState StateAfter, const FRDGTextureRef Texture)
+void FRDGLogFile::AddTransitionEdge(FRDGPassHandle PassHandle, const FRDGSubresourceState& StateBefore, const FRDGSubresourceState& StateAfter, const FRDGTextureRef Texture)
 {
-	if (GRDGDumpGraph == RDG_DUMP_GRAPH_RESOURCES && bOpen && IncludeTransitionEdgeInGraph(StateBefore.GetFirstPass(), PassHandle) && IsDebugAllowedForResource(Texture->Name))
+	if (GRDGDumpGraph == RDG_DUMP_GRAPH_RESOURCES && bOpen && IsDebugAllowedForResource(Texture->Name))
 	{
-		if (FRDGSubresourceState::IsTransitionRequired(StateBefore, StateAfter))
+		if (IncludeTransitionEdgeInGraph(StateBefore.GetLastPass(), StateAfter.GetFirstPass()) && FRDGSubresourceState::IsTransitionRequired(StateBefore, StateAfter))
 		{
 			AddLine(FString::Printf(TEXT("\"%s\" -> \"%s\" [%s, label=<%s: <b>%s -&gt; %s</b>>]"),
 				*GetProducerName(StateBefore.GetLastPass()),
@@ -1723,23 +1723,25 @@ void FRDGLogFile::AddTransitionEdge(FRDGPassHandle PassHandle, FRDGSubresourceSt
 				*GetSubresourceStateLabel(StateBefore),
 				*GetSubresourceStateLabel(StateAfter)));
 		}
-		else
+		else if (IncludeTransitionEdgeInGraph(StateBefore.LogFilePass, PassHandle))
 		{
 			AddLine(FString::Printf(TEXT("\"%s\" -> \"%s\" [%s, label=<%s: <b>%s</b>>]"),
-				*GetProducerName(StateBefore.GetFirstPass()),
+				*GetProducerName(StateBefore.LogFilePass),
 				*GetConsumerName(PassHandle),
 				TextureColorAttributes,
 				Texture->Name,
 				*GetSubresourceStateLabel(StateBefore)));
 		}
+
+		StateAfter.LogFilePass = PassHandle;
 	}
 }
 
-void FRDGLogFile::AddTransitionEdge(FRDGPassHandle PassHandle, FRDGSubresourceState StateBefore, FRDGSubresourceState StateAfter, const FRDGTextureRef Texture, FRDGTextureSubresource Subresource)
+void FRDGLogFile::AddTransitionEdge(FRDGPassHandle PassHandle, const FRDGSubresourceState& StateBefore, const FRDGSubresourceState& StateAfter, const FRDGTextureRef Texture, FRDGTextureSubresource Subresource)
 {
-	if (GRDGDumpGraph == RDG_DUMP_GRAPH_RESOURCES && bOpen && IncludeTransitionEdgeInGraph(StateBefore.GetFirstPass(), PassHandle) && IsDebugAllowedForResource(Texture->Name))
+	if (GRDGDumpGraph == RDG_DUMP_GRAPH_RESOURCES && bOpen && IsDebugAllowedForResource(Texture->Name))
 	{
-		if (FRDGSubresourceState::IsTransitionRequired(StateBefore, StateAfter))
+		if (IncludeTransitionEdgeInGraph(StateBefore.GetLastPass(), StateAfter.GetFirstPass()) && FRDGSubresourceState::IsTransitionRequired(StateBefore, StateAfter))
 		{
 			AddLine(FString::Printf(TEXT("\"%s\" -> \"%s\" [%s, label=<%s[%d][%d][%d]: <b>%s -&gt; %s</b>>]"),
 				*GetProducerName(StateBefore.GetLastPass()),
@@ -1750,24 +1752,26 @@ void FRDGLogFile::AddTransitionEdge(FRDGPassHandle PassHandle, FRDGSubresourceSt
 				*GetSubresourceStateLabel(StateBefore),
 				*GetSubresourceStateLabel(StateAfter)));
 		}
-		else
+		else if (IncludeTransitionEdgeInGraph(StateBefore.LogFilePass, PassHandle))
 		{
 			AddLine(FString::Printf(TEXT("\"%s\" -> \"%s\" [%s, label=<%s[%d][%d][%d]: <b>%s</b>>]"),
-				*GetProducerName(StateBefore.GetFirstPass()),
+				*GetProducerName(StateBefore.LogFilePass),
 				*GetConsumerName(PassHandle),
 				TextureColorAttributes,
 				Texture->Name,
 				Subresource.MipIndex, Subresource.ArraySlice, Subresource.PlaneSlice,
 				*GetSubresourceStateLabel(StateBefore)));
 		}
+
+		StateAfter.LogFilePass = PassHandle;
 	}
 }
 
-void FRDGLogFile::AddTransitionEdge(FRDGPassHandle PassHandle, FRDGSubresourceState StateBefore, FRDGSubresourceState StateAfter, const FRDGBufferRef Buffer)
+void FRDGLogFile::AddTransitionEdge(FRDGPassHandle PassHandle, const FRDGSubresourceState& StateBefore, const FRDGSubresourceState& StateAfter, const FRDGBufferRef Buffer)
 {
-	if (GRDGDumpGraph == RDG_DUMP_GRAPH_RESOURCES && bOpen && IncludeTransitionEdgeInGraph(StateBefore.GetFirstPass(), PassHandle) && IsDebugAllowedForResource(Buffer->Name))
+	if (GRDGDumpGraph == RDG_DUMP_GRAPH_RESOURCES && bOpen && IsDebugAllowedForResource(Buffer->Name))
 	{
-		if (FRDGSubresourceState::IsTransitionRequired(StateBefore, StateAfter))
+		if (IncludeTransitionEdgeInGraph(StateBefore.GetLastPass(), StateAfter.GetFirstPass()) && FRDGSubresourceState::IsTransitionRequired(StateBefore, StateAfter))
 		{
 			AddLine(FString::Printf(TEXT("\"%s\" -> \"%s\" [%s, label=<%s: <b>%s -&gt; %s</b>>]"),
 				*GetProducerName(StateBefore.GetLastPass()),
@@ -1777,15 +1781,17 @@ void FRDGLogFile::AddTransitionEdge(FRDGPassHandle PassHandle, FRDGSubresourceSt
 				*GetSubresourceStateLabel(StateBefore),
 				*GetSubresourceStateLabel(StateAfter)));
 		}
-		else
+		else if (IncludeTransitionEdgeInGraph(StateBefore.LogFilePass, PassHandle))
 		{
 			AddLine(FString::Printf(TEXT("\"%s\" -> \"%s\" [%s, label=<%s: <b>%s</b>>]"),
-				*GetProducerName(StateBefore.GetFirstPass()),
+				*GetProducerName(StateBefore.LogFilePass),
 				*GetConsumerName(PassHandle),
 				BufferColorAttributes,
 				Buffer->Name,
 				*GetSubresourceStateLabel(StateBefore)));
 		}
+
+		StateAfter.LogFilePass = PassHandle;
 	}
 }
 
diff --git a/Engine/Source/Runtime/RenderCore/Public/RenderGraphBuilder.h b/Engine/Source/Runtime/RenderCore/Public/RenderGraphBuilder.h
--- a/Engine/Source/Runtime/RenderCore/Public/RenderGraphBuilder.h
+++ b/Engine/Source/Runtime/RenderCore/Public/RenderGraphBuilder.h
@@ -263,7 +263,6 @@ public:
 	//////////////////////////////////////////////////////////////////////////
 
 private:
-	static const ERHIAccess kDefaultAccessInitial = ERHIAccess::Unknown;
 	static const ERHIAccess kDefaultAccessFinal = ERHIAccess::SRVMask;
 	static const char* const kDefaultUnaccountedCSVStat;
 
@@ -285,7 +284,8 @@ private:
 
 	FORCEINLINE FRDGPassHandle ClampToPrologue(FRDGPassHandle PassHandle, ERHIPipeline Pipeline = ERHIPipeline::Graphics) const
 	{
-		return FRDGPassHandle::Max(ProloguePassHandles[Pipeline], PassHandle);
+		// Preserve null inputs as outputs. Null is the highest value.
+		return ProloguePassHandles[Pipeline].GetIndexUnchecked() > PassHandle.GetIndexUnchecked() ? ProloguePassHandles[Pipeline] : PassHandle;
 	}
 
 	FORCEINLINE FRDGPassHandlesByPipeline ClampToPrologue(FRDGPassHandlesByPipeline PassHandles) const
@@ -475,7 +475,7 @@ private:
 
 	uint32 AsyncComputePassCount = 0;
 	uint32 RasterPassCount = 0;
-	uint32 ExecuteCount = 0;
+	uint32 ExecuteGeneration = 0;
 
 	IF_RDG_CMDLIST_STATS(TStatId CommandListStatScope);
 	IF_RDG_CMDLIST_STATS(TStatId CommandListStatState);
diff --git a/Engine/Source/Runtime/RenderCore/Public/RenderGraphEvent.h b/Engine/Source/Runtime/RenderCore/Public/RenderGraphEvent.h
--- a/Engine/Source/Runtime/RenderCore/Public/RenderGraphEvent.h
+++ b/Engine/Source/Runtime/RenderCore/Public/RenderGraphEvent.h
@@ -323,7 +323,7 @@ struct FRDGGPUScopeStacks
 
 struct FRDGGPUScopeStacksByPipeline
 {
-	FRDGGPUScopeStacksByPipeline(FRHICommandListImmediate& RHICmdListGraphics, FRHIComputeCommandList& RHICmdListAsyncCompute, FRDGAllocator& Allocator);
+	FRDGGPUScopeStacksByPipeline(FRHICommandListImmediate& RHICmdListGraphics, FRHIAsyncComputeCommandListImmediate& RHICmdListAsyncCompute, FRDGAllocator& Allocator);
 
 	void BeginEventScope(FRDGEventName&& ScopeName);
 
@@ -349,6 +349,7 @@ struct FRDGGPUScopeStacksByPipeline
 
 	FRDGGPUScopeStacks Graphics;
 	FRDGGPUScopeStacks AsyncCompute;
+	FRHIAsyncComputeCommandListImmediate& RHICmdListAsyncCompute;
 };
 
 #endif
diff --git a/Engine/Source/Runtime/RenderCore/Public/RenderGraphEvent.inl b/Engine/Source/Runtime/RenderCore/Public/RenderGraphEvent.inl
--- a/Engine/Source/Runtime/RenderCore/Public/RenderGraphEvent.inl
+++ b/Engine/Source/Runtime/RenderCore/Public/RenderGraphEvent.inl
@@ -253,9 +253,10 @@ inline FRDGGPUScopes FRDGGPUScopeStacks::GetCurrentScopes() const
 	return Scopes;
 }
 
-inline FRDGGPUScopeStacksByPipeline::FRDGGPUScopeStacksByPipeline(FRHICommandListImmediate& RHICmdListGraphics, FRHIComputeCommandList& RHICmdListAsyncCompute, FRDGAllocator& Allocator)
+inline FRDGGPUScopeStacksByPipeline::FRDGGPUScopeStacksByPipeline(FRHICommandListImmediate& RHICmdListGraphics, FRHIAsyncComputeCommandListImmediate& InRHICmdListAsyncCompute, FRDGAllocator& Allocator)
 	: Graphics(RHICmdListGraphics, Allocator)
-	, AsyncCompute(RHICmdListAsyncCompute, Allocator)
+	, AsyncCompute(InRHICmdListAsyncCompute, Allocator)
+	, RHICmdListAsyncCompute(InRHICmdListAsyncCompute)
 {}
 
 inline void FRDGGPUScopeStacksByPipeline::BeginEventScope(FRDGEventName&& ScopeName)
@@ -289,12 +290,6 @@ inline void FRDGGPUScopeStacksByPipeline::BeginExecute()
 	AsyncCompute.BeginExecute();
 }
 
-inline void FRDGGPUScopeStacksByPipeline::EndExecute()
-{
-	Graphics.EndExecute();
-	AsyncCompute.EndExecute();
-}
-
 inline const FRDGGPUScopeStacks& FRDGGPUScopeStacksByPipeline::GetScopeStacks(ERHIPipeline Pipeline) const
 {
 	switch (Pipeline)
diff --git a/Engine/Source/Runtime/RenderCore/Public/RenderGraphPass.h b/Engine/Source/Runtime/RenderCore/Public/RenderGraphPass.h
--- a/Engine/Source/Runtime/RenderCore/Public/RenderGraphPass.h
+++ b/Engine/Source/Runtime/RenderCore/Public/RenderGraphPass.h
@@ -311,9 +311,6 @@ protected:
 			uint32 bAsyncComputeBegin : 1;
 			uint32 bAsyncComputeEnd : 1;
 
-			/** (AsyncCompute only) Whether this is the last async compute pass in the graph. */
-			uint32 bAsyncComputeEndExecute : 1;
-
 			/** (Graphics only) Whether this is a graphics fork / join pass. */
 			uint32 bGraphicsFork : 1;
 			uint32 bGraphicsJoin : 1;
@@ -366,17 +363,17 @@ protected:
 
 		FTextureState(FRDGTextureRef InTexture)
 			: Texture(InTexture)
-			, TextureHandle(InTexture->GetHandle())
 		{
 			const uint32 SubresourceCount = Texture->GetSubresourceCount();
+			State.Reserve(SubresourceCount);
 			State.SetNum(SubresourceCount);
+			MergeState.Reserve(SubresourceCount);
 			MergeState.SetNum(SubresourceCount);
 		}
 
 		FRDGTextureRef Texture = nullptr;
 		FRDGTextureTransientSubresourceState State;
 		FRDGTextureTransientSubresourceStateIndirect MergeState;
-		FRDGTextureHandle TextureHandle;
 		uint16 ReferenceCount = 0;
 	};
 
@@ -386,13 +383,11 @@ protected:
 
 		FBufferState(FRDGBufferRef InBuffer)
 			: Buffer(InBuffer)
-			, BufferHandle(InBuffer->GetHandle())
 		{}
 
 		FRDGBufferRef Buffer = nullptr;
 		FRDGSubresourceState State;
 		FRDGSubresourceState* MergeState = nullptr;
-		FRDGBufferHandle BufferHandle;
 		uint16 ReferenceCount = 0;
 	};
 
diff --git a/Engine/Source/Runtime/RenderCore/Public/RenderGraphResources.h b/Engine/Source/Runtime/RenderCore/Public/RenderGraphResources.h
--- a/Engine/Source/Runtime/RenderCore/Public/RenderGraphResources.h
+++ b/Engine/Source/Runtime/RenderCore/Public/RenderGraphResources.h
@@ -71,6 +71,12 @@ struct FRDGSubresourceState
 
 	/** The last no-UAV barrier to be used by this subresource. */
 	FRDGViewUniqueFilter NoUAVBarrierFilter;
+
+	/** The graph execution generation to submit this subresource state. */
+	uint8 ExecuteGeneration = 0;
+
+	/** The last used pass for log file debugging. */
+	IF_RDG_ENABLE_DEBUG(mutable FRDGPassHandle LogFilePass);
 };
 
 using FRDGTextureSubresourceState = TRDGTextureSubresourceArray<FRDGSubresourceState, FDefaultAllocator>;
@@ -287,6 +293,9 @@ protected:
 	/** If true, the resource has been used on an async compute pass and may have async compute states. */
 	uint8 bUsedByAsyncComputePass : 1;
 
+	FRDGPassHandle FirstPass;
+	FRDGPassHandle LastPass;
+
 private:
 	/** Number of references in passes and deferred queries. */
 	uint16 ReferenceCount = 0;
@@ -294,13 +303,9 @@ private:
 	/** Scratch index allocated for the resource in the pass being setup. */
 	uint16 PassStateIndex = 0;
 
-	/** The initial and final states of the resource assigned by the user, if known. */
-	ERHIAccess AccessInitial = ERHIAccess::Unknown;
+	/** The final state of the resource at the end of graph execution, if known. */
 	ERHIAccess AccessFinal = ERHIAccess::Unknown;
 
-	FRDGPassHandle FirstPass;
-	FRDGPassHandle LastPass;
-
 #if RDG_ENABLE_TRACE
 	uint16 TraceOrder = 0;
 	TArray<FRDGPassHandle, FRDGArrayAllocator> TracePasses;
@@ -360,7 +365,6 @@ class RENDERCORE_API FRDGPooledTexture final
 public:
 	FRDGPooledTexture(FRHITexture* InTexture, const FRDGTextureSubresourceLayout& InLayout)
 		: Texture(InTexture)
-		, Layout(InLayout)
 	{
 		Reset();
 	}
@@ -384,7 +388,6 @@ private:
 
 	TRefCountPtr<FRHITexture> Texture;
 	FRDGTexture* Owner = nullptr;
-	FRDGTextureSubresourceLayout Layout;
 	FRDGTextureSubresourceState State;
 	FRHITextureViewCache ViewCache;
 
@@ -410,31 +413,36 @@ public:
 	IPooledRenderTarget* GetPooledRenderTarget() const;
 
 	/** Returns the allocated RHI texture. */
-	FRHITexture* GetRHI() const
+	FORCEINLINE FRHITexture* GetRHI() const
 	{
 		return static_cast<FRHITexture*>(FRDGResource::GetRHI());
 	}
 
 	//////////////////////////////////////////////////////////////////////////
 
-	FRDGTextureHandle GetHandle() const
+	FORCEINLINE FRDGTextureHandle GetHandle() const
 	{
 		return Handle;
 	}
 
-	FRDGTextureSubresourceLayout GetSubresourceLayout() const
+	FORCEINLINE FRDGTextureSubresourceLayout GetSubresourceLayout() const
 	{
 		return Layout;
 	}
 
-	FRDGTextureSubresourceRange GetSubresourceRange() const
+	FORCEINLINE FRDGTextureSubresourceRange GetSubresourceRange() const
 	{
 		return WholeRange;
 	}
 
 	FORCEINLINE uint32 GetSubresourceCount() const
 	{
-		return Layout.GetSubresourceCount();
+		return SubresourceCount;
+	}
+
+	FORCEINLINE FRDGTextureSubresource GetSubresource(uint32 SubresourceIndex) const
+	{
+		return Layout.GetSubresource(SubresourceIndex);
 	}
 
 	FRDGTextureSubresourceRange GetSubresourceRangeSRV() const;
@@ -447,9 +455,11 @@ private:
 		, RenderTargetTexture(InRenderTargetTexture)
 		, Layout(InDesc)
 		, WholeRange(Layout)
+		, SubresourceCount(Layout.GetSubresourceCount())
 	{
-		const uint32 SubresourceCount = GetSubresourceCount();
+		MergeState.Reserve(SubresourceCount);
 		MergeState.SetNum(SubresourceCount);
+		LastProducers.Reserve(SubresourceCount);
 		LastProducers.SetNum(SubresourceCount);
 	}
 
@@ -471,12 +481,6 @@ private:
 		return static_cast<FRHITexture*>(FRDGResource::GetRHIUnchecked());
 	}
 
-	/** Whether this texture is the last owner of the allocation in the graph. */
-	bool IsLastOwner() const
-	{
-		return NextOwner.IsNull();
-	}
-
 	/** Returns the current texture state. Only valid to call after SetRHI. */
 	FRDGTextureSubresourceState& GetState()
 	{
@@ -496,9 +500,7 @@ private:
 	/** The layout used to facilitate subresource transitions. */
 	FRDGTextureSubresourceLayout Layout;
 	FRDGTextureSubresourceRange  WholeRange;
-
-	/** Cached state pointer from the pooled texture. */
-	FRDGTextureSubresourceState* State = nullptr;
+	const uint32 SubresourceCount;
 
 	/** The assigned pooled render target to use during execution. Never reset. */
 	FPooledRenderTarget* PooledRenderTarget = nullptr;
@@ -518,6 +520,9 @@ private:
 	/** Valid strictly when holding a strong reference; use PooledRenderTarget instead. */
 	TRefCountPtr<IPooledRenderTarget> Allocation;
 
+	/** Cached state pointer from the pooled texture. */
+	FRDGTextureSubresourceState* State = nullptr;
+
 	/** Tracks merged subresource states as the graph is built. */
 	FRDGTextureTransientSubresourceStateIndirect MergeState;
 
@@ -994,9 +999,6 @@ private:
 	/** Registered handle set by the builder. */
 	FRDGBufferHandle Handle;
 
-	/** Tracks the last pass that produced this resource as the graph is built. */
-	FRDGProducerStatesByPipeline LastProducer;
-
 	/** The next buffer to own the PooledBuffer allocation during execution. */
 	FRDGBufferHandle NextOwner;
 
@@ -1012,15 +1014,18 @@ private:
 	/** The assigned buffer view cache (sourced from the pooled / transient buffer. Never reset. */
 	FRHIBufferViewCache* ViewCache = nullptr;
 
-	/** Cached state pointer from the pooled / transient buffer. */
-	FRDGSubresourceState* State = nullptr;
-
 	/** Valid strictly when holding a strong reference; use PooledBuffer instead. */
 	TRefCountPtr<FRDGPooledBuffer> Allocation;
 
+	/** Cached state pointer from the pooled / transient buffer. */
+	FRDGSubresourceState* State = nullptr;
+
 	/** Tracks the merged subresource state as the graph is built. */
 	FRDGSubresourceState* MergeState = nullptr;
 
+	/** Tracks the last pass that produced this resource as the graph is built. */
+	FRDGProducerStatesByPipeline LastProducer;
+
 #if RDG_ENABLE_DEBUG
 	struct FRDGBufferDebugData* BufferDebugData = nullptr;
 	FRDGBufferDebugData& GetBufferDebugData() const;
diff --git a/Engine/Source/Runtime/RenderCore/Public/RenderGraphTextureSubresource.h b/Engine/Source/Runtime/RenderCore/Public/RenderGraphTextureSubresource.h
--- a/Engine/Source/Runtime/RenderCore/Public/RenderGraphTextureSubresource.h
+++ b/Engine/Source/Runtime/RenderCore/Public/RenderGraphTextureSubresource.h
@@ -92,6 +92,15 @@ struct FRDGTextureSubresourceLayout
 		return Subresource.MipIndex + (Subresource.ArraySlice * NumMips) + (Subresource.PlaneSlice * NumMips * NumArraySlices);
 	}
 
+	inline FRDGTextureSubresource GetSubresource(uint32 Index) const
+	{
+		FRDGTextureSubresource Subresource;
+		Subresource.MipIndex = Index % NumMips;
+		Subresource.ArraySlice = (Index / NumMips) % NumArraySlices;
+		Subresource.PlaneSlice = Index / (NumMips * NumArraySlices);
+		return Subresource;
+	}
+
 	inline FRDGTextureSubresource GetMaxSubresource() const
 	{
 		return FRDGTextureSubresource(NumMips, NumArraySlices, NumPlaneSlices);
@@ -215,7 +224,8 @@ inline void InitAsSubresources(TRDGTextureSubresourceArray<ElementType, Allocato
 	const uint32 SubresourceCount = Layout.GetSubresourceCount();
 	checkf(SubresourceCount > 0, TEXT("Subresource layout has no subresources."));
 	checkf(SubresourceCount > 1, TEXT("Subresource layout has only 1 resource. Use InitAsWholeResource instead."));
-	SubresourceArray.SetNum(SubresourceCount, false);
+	SubresourceArray.Reserve(SubresourceCount);
+	SubresourceArray.SetNum(SubresourceCount);
 	for (uint32 SubresourceIndex = 0; SubresourceIndex < SubresourceCount; ++SubresourceIndex)
 	{
 		SubresourceArray[SubresourceIndex] = Element;
diff --git a/Engine/Source/Runtime/RenderCore/Public/RenderGraphValidation.h b/Engine/Source/Runtime/RenderCore/Public/RenderGraphValidation.h
--- a/Engine/Source/Runtime/RenderCore/Public/RenderGraphValidation.h
+++ b/Engine/Source/Runtime/RenderCore/Public/RenderGraphValidation.h
@@ -166,11 +166,11 @@ public:
 
 	void AddAliasEdge(const FRDGBufferRef BufferBefore, FRDGPassHandle BeforePass, const FRDGBufferRef BufferAfter, FRDGPassHandle PassAfter);
 
-	void AddTransitionEdge(FRDGPassHandle PassHandle, FRDGSubresourceState StateBefore, FRDGSubresourceState StateAfter, const FRDGTextureRef Texture);
+	void AddTransitionEdge(FRDGPassHandle PassHandle, const FRDGSubresourceState& StateBefore, const FRDGSubresourceState& StateAfter, const FRDGTextureRef Texture);
 
-	void AddTransitionEdge(FRDGPassHandle PassHandle, FRDGSubresourceState StateBefore, FRDGSubresourceState StateAfter, const FRDGTextureRef Texture, FRDGTextureSubresource Subresource);
+	void AddTransitionEdge(FRDGPassHandle PassHandle, const FRDGSubresourceState& StateBefore, const FRDGSubresourceState& StateAfter, const FRDGTextureRef Texture, FRDGTextureSubresource Subresource);
 
-	void AddTransitionEdge(FRDGPassHandle PassHandle, FRDGSubresourceState StateBefore, FRDGSubresourceState StateAfter, const FRDGBufferRef Buffer);
+	void AddTransitionEdge(FRDGPassHandle PassHandle, const FRDGSubresourceState& StateBefore, const FRDGSubresourceState& StateAfter, const FRDGBufferRef Buffer);
 
 	void End();
 
