Fixed a crash in OGL RHIAcquireThreadOwnership

diff --git a/Engine/Source/Runtime/OpenGLDrv/Private/OpenGLDevice.cpp b/Engine/Source/Runtime/OpenGLDrv/Private/OpenGLDevice.cpp
--- a/Engine/Source/Runtime/OpenGLDrv/Private/OpenGLDevice.cpp
+++ b/Engine/Source/Runtime/OpenGLDrv/Private/OpenGLDevice.cpp
@@ -1243,7 +1243,6 @@ void FOpenGLDynamicRHI::RHIAcquireThreadOwnership()
 	PlatformRenderingContextSetup(PlatformDevice);
 	PlatformRebindResources(PlatformDevice);
 	bIsRenderingContextAcquired = true;
-	FTaskTagScope Scope(ETaskTag::ERenderingThread);
 	VERIFY_GL(RHIAcquireThreadOwnership);
 	{
 		FScopeLock lock(&CustomPresentSection);
diff --git a/Engine/Source/Runtime/RenderCore/Private/RenderingThread.cpp b/Engine/Source/Runtime/RenderCore/Private/RenderingThread.cpp
--- a/Engine/Source/Runtime/RenderCore/Private/RenderingThread.cpp
+++ b/Engine/Source/Runtime/RenderCore/Private/RenderingThread.cpp
@@ -701,6 +701,7 @@ public:
 
 		if (bAcquireOwnership)
 		{
+			FTaskTagScope Scope(ETaskTag::ERhiThread);
 			GDynamicRHI->RHIAcquireThreadOwnership();
 		}
 		else
@@ -739,15 +740,17 @@ void StartRenderingThread()
 		}
 		DECLARE_CYCLE_STAT(TEXT("Wait For RHIThread"), STAT_WaitForRHIThread, STATGROUP_TaskGraphTasks);
 
-		FGraphEventRef CompletionEvent = TGraphTask<FOwnershipOfRHIThreadTask>::CreateTask(NULL, ENamedThreads::GameThread).ConstructAndDispatchWhenReady(true, GET_STATID(STAT_WaitForRHIThread));
-		QUICK_SCOPE_CYCLE_COUNTER(STAT_StartRenderingThread);
-		FTaskGraphInterface::Get().WaitUntilTaskCompletes(CompletionEvent, ENamedThreads::GameThread_Local);
 PRAGMA_DISABLE_DEPRECATION_WARNINGS
 		GRHIThread_InternalUseOnly = FRHIThread::Get().Thread;
 		GRHIThreadId = GRHIThread_InternalUseOnly->GetThreadID();
 PRAGMA_ENABLE_DEPRECATION_WARNINGS
 		GIsRunningRHIInDedicatedThread_InternalUseOnly = true;
 		GIsRunningRHIInSeparateThread_InternalUseOnly = true;
+
+		FGraphEventRef CompletionEvent = TGraphTask<FOwnershipOfRHIThreadTask>::CreateTask(NULL, ENamedThreads::GameThread).ConstructAndDispatchWhenReady(true, GET_STATID(STAT_WaitForRHIThread));
+		QUICK_SCOPE_CYCLE_COUNTER(STAT_StartRenderingThread);
+		FTaskGraphInterface::Get().WaitUntilTaskCompletes(CompletionEvent, ENamedThreads::GameThread_Local);
+
 		GRHICommandList.LatchBypass();
 	}
 	else if (GUseRHITaskThreads_InternalUseOnly)
