Removed ensure for RDG finalized resources being accessed on async compute. Needs to handle read-only resources.

diff --git a/Engine/Source/Runtime/RenderCore/Private/RenderGraphValidation.cpp b/Engine/Source/Runtime/RenderCore/Private/RenderGraphValidation.cpp
--- a/Engine/Source/Runtime/RenderCore/Private/RenderGraphValidation.cpp
+++ b/Engine/Source/Runtime/RenderCore/Private/RenderGraphValidation.cpp
@@ -483,9 +483,11 @@ void FRDGUserValidation::ValidateFinalizedAccess(FRDGParentResourceRef Resource,
 		TEXT("Resource %s was finalized with access %s, but is being used in pass %s with access %s. Any future pass must use a subset of the finalized access state."),
 		Resource->Name, *GetRHIAccessName(Resource->AccessFinal), Pass->GetName(), *GetRHIAccessName(Access), *GetRHIPipelineName(Pass->GetPipeline()));
 
+#if 0 // TODO: Need to account for read-only resources. 
 	ensureMsgf(Pass->GetPipeline() == ERHIPipeline::Graphics,
 		TEXT("Resource %s was finalized but is being used on the async compute pass %s. Only graphics pipe access is allowed for finalized resources."),
 		Resource->Name, Pass->GetName());
+#endif
 }
 
 void FRDGUserValidation::ValidateAddPass(const FRDGEventName& Name, ERDGPassFlags Flags)
