Fix for culled RDG uniform buffer being created.

diff --git a/Engine/Source/Runtime/RenderCore/Private/RenderGraphBuilder.cpp b/Engine/Source/Runtime/RenderCore/Private/RenderGraphBuilder.cpp
--- a/Engine/Source/Runtime/RenderCore/Private/RenderGraphBuilder.cpp
+++ b/Engine/Source/Runtime/RenderCore/Private/RenderGraphBuilder.cpp
@@ -1780,13 +1780,10 @@ FRDGPass* FRDGBuilder::SetupPass(FRDGPass* Pass)
 		Buffer->bProduced |= bWritableAccess;
 	});
 
+	Pass->UniformBuffers.Reserve(PassParameters.GetUniformBufferParameterCount());
 	PassParameters.EnumerateUniformBuffers([&](FRDGUniformBufferBinding UniformBuffer)
 	{
-		if (!UniformBuffer->bQueuedForCreate)
-		{
-			UniformBuffersToCreate.Emplace(UniformBuffer->Handle);
-			UniformBuffer->bQueuedForCreate = true;
-		}
+		Pass->UniformBuffers.Emplace(UniformBuffer.GetUniformBuffer()->Handle);
 	});
 
 	Pass->bRenderPassOnlyWrites = bRenderPassOnlyWrites;
@@ -2110,6 +2107,15 @@ void FRDGBuilder::BeginResourcesRHI(FRDGPass* ResourcePass, FRDGPassHandle Execu
 			BeginResourceRHI(ExecutePassHandle, PassState.Buffer);
 		}
 
+		for (FRDGUniformBufferHandle UniformBufferHandle : PassToBegin->UniformBuffers)
+		{
+			if (FRDGUniformBuffer* UniformBuffer = UniformBuffers[UniformBufferHandle]; !UniformBuffer->bQueuedForCreate)
+			{
+				UniformBuffer->bQueuedForCreate = true;
+				UniformBuffersToCreate.Add(UniformBufferHandle);
+			}
+		}
+
 		for (FRDGViewHandle ViewHandle : PassToBegin->Views)
 		{
 			BeginResourceRHI(ExecutePassHandle, Views[ViewHandle]);
diff --git a/Engine/Source/Runtime/RenderCore/Public/RenderGraphPass.h b/Engine/Source/Runtime/RenderCore/Public/RenderGraphPass.h
--- a/Engine/Source/Runtime/RenderCore/Public/RenderGraphPass.h
+++ b/Engine/Source/Runtime/RenderCore/Public/RenderGraphPass.h
@@ -395,6 +395,7 @@ protected:
 	TArray<FTextureState, FRDGArrayAllocator> TextureStates;
 	TArray<FBufferState, FRDGArrayAllocator> BufferStates;
 	TArray<FRDGViewHandle, FRDGArrayAllocator> Views;
+	TArray<FRDGUniformBufferHandle, FRDGArrayAllocator> UniformBuffers;
 
 	/** Lists of pass parameters scheduled for begin during execution of this pass. */
 	TArray<FRDGPass*, TInlineAllocator<1, FRDGArrayAllocator>> ResourcesToBegin;
