Fix for assert when readback resource is being used in a copy destination state.

diff --git a/Engine/Source/Runtime/RenderCore/Private/RenderGraphBuilder.cpp b/Engine/Source/Runtime/RenderCore/Private/RenderGraphBuilder.cpp
--- a/Engine/Source/Runtime/RenderCore/Private/RenderGraphBuilder.cpp
+++ b/Engine/Source/Runtime/RenderCore/Private/RenderGraphBuilder.cpp
@@ -612,7 +612,6 @@ FRDGTextureRef FRDGBuilder::RegisterExternalTexture(
 
 	if (!EnumHasAnyFlags(Desc.Flags, TexCreate_RenderTargetable | TexCreate_ResolveTargetable | TexCreate_DepthStencilTargetable | TexCreate_UAV | TexCreate_DepthStencilResolveTarget))
 	{
-		Flags |= ERDGTextureFlags::ReadOnly;
 		bFinalizedAccess = true;
 	}
 
@@ -626,8 +625,16 @@ FRDGTextureRef FRDGBuilder::RegisterExternalTexture(
 	// Textures that are created read-only are not transitioned by RDG.
 	if (bFinalizedAccess)
 	{
+		// When in 'finalized access' mode, the access represents the valid set of states to touch the resource for
+		// validation, not its final state after the graph executes. That's why it's okay to have a write state mixed
+		// with read states.
 		Texture->bFinalizedAccess = 1;
 		Texture->AccessFinal = ERHIAccess::ReadOnlyExclusiveMask;
+
+		if (EnumHasAnyFlags(Desc.Flags, TexCreate_CPUReadback))
+		{
+			Texture->AccessFinal |= ERHIAccess::CopyDest;
+		}
 	}
 
 	FRDGTextureSubresourceState& TextureState = Texture->GetState();
