AllowDebugViewModes() now returns true if running a commandlet with -AllowCommandletRendering

diff --git a/Engine/Source/Runtime/RenderCore/Private/ShaderCore.cpp b/Engine/Source/Runtime/RenderCore/Private/ShaderCore.cpp
--- a/Engine/Source/Runtime/RenderCore/Private/ShaderCore.cpp
+++ b/Engine/Source/Runtime/RenderCore/Private/ShaderCore.cpp
@@ -248,14 +248,34 @@ bool AllowDebugViewmodes()
 
 	// To use debug viewmodes on consoles, r.ForceDebugViewModes must be set to 1 in ConsoleVariables.ini
 	// And EngineDebugMaterials must be in the StartupPackages for the target platform.
-	bool bForceEnable = ForceDebugViewValue == 1;
-	bool bForceDisable = ForceDebugViewValue == 2;
 
-	//bool bEnabled = (!IsRunningCommandlet() && !FPlatformProperties::RequiresCookedData());	
-	//bEnabled |= bForceEnable;
-	//bEnabled &= !bForceDisable;
+	// Force enabled: r.ForceDebugViewModes 1
+	const bool bForceEnable = ForceDebugViewValue == 1;
+	if (bForceEnable)
+	{
+		return true;
+	}
+
+	// Force disabled: r.ForceDebugViewModes 2
+	const bool bForceDisable = ForceDebugViewValue == 2;
+	if (bForceDisable)
+	{
+		return false;
+	}
+
+	// Disable when running a commandlet without -AllowCommandletRendering
+	if (IsRunningCommandlet() && !IsAllowCommandletRendering())
+	{
+		return false;
+	}
+
+	// Disable if we require cooked data
+	if (FPlatformProperties::RequiresCookedData())
+	{
+		return false;
+	}
 
-	return (!bForceDisable) && (bForceEnable || (!IsRunningCommandlet() && !FPlatformProperties::RequiresCookedData()));
+	return true;
 }
 
 /** Returns true if debug viewmodes are allowed for the current platform. */
