Added UAVOverlap to UnifiedBuffer MemcpyResource to fix RHI validation errors.

diff --git a/Engine/Source/Runtime/RenderCore/Private/UnifiedBuffer.cpp b/Engine/Source/Runtime/RenderCore/Private/UnifiedBuffer.cpp
--- a/Engine/Source/Runtime/RenderCore/Private/UnifiedBuffer.cpp
+++ b/Engine/Source/Runtime/RenderCore/Private/UnifiedBuffer.cpp
@@ -249,7 +249,7 @@ void MemsetResource(FRHICommandList& RHICmdList, const ResourceType& DstBuffer,
 }
 
 template<typename ResourceType>
-void MemcpyResource(FRHICommandList& RHICmdList, const ResourceType& DstBuffer, const ResourceType& SrcBuffer, uint32 NumBytes, uint32 DstOffset, uint32 SrcOffset)
+void MemcpyResource(FRHICommandList& RHICmdList, const ResourceType& DstBuffer, const ResourceType& SrcBuffer, uint32 NumBytes, uint32 DstOffset, uint32 SrcOffset, bool bAlreadyInUAVOverlap)
 {
 	uint32 DivisorAlignment = ResourceTypeTraits<ResourceType>::Type == EResourceType::BYTEBUFFER ? 4 : 16;
 	if(ResourceTypeTraits<ResourceType>::Type == EResourceType::BYTEBUFFER)
@@ -267,6 +267,9 @@ void MemcpyResource(FRHICommandList& RHICmdList, const ResourceType& DstBuffer,
 
 	uint32 NumBytesProcessed = 0;
 
+	if (!bAlreadyInUAVOverlap)	// TODO: Get rid of this check once BeginUAVOverlap/EndUAVOverlap supports nesting.
+		RHICmdList.BeginUAVOverlap(DstBuffer.UAV);	
+
 	while (NumBytesProcessed < NumBytes)
 	{
 		const uint32 NumWaves = FMath::Max(FMath::Min<uint32>(GRHIMaxDispatchThreadGroupsPerDimension.X, FMath::DivideAndRoundUp(NumBytes / 16, 64u)), 1u);
@@ -314,6 +317,9 @@ void MemcpyResource(FRHICommandList& RHICmdList, const ResourceType& DstBuffer,
 
 		NumBytesProcessed += NumBytesPerDispatch;
 	}
+
+	if(!bAlreadyInUAVOverlap)
+		RHICmdList.EndUAVOverlap(DstBuffer.UAV);
 }
 
 template<>
@@ -437,7 +443,7 @@ RENDERCORE_API bool ResizeResourceSOAIfNeeded<FRWBufferStructured>(FRHICommandLi
 		uint32 CopyBytes = FMath::Min(NewArraySize, OldArraySize);
 		for( uint32 i = 0; i < NumArrays; i++ )
 		{
-			MemcpyResource( RHICmdList, NewBuffer, Buffer, CopyBytes, i * NewArraySize, i * OldArraySize );
+			MemcpyResource( RHICmdList, NewBuffer, Buffer, CopyBytes, i * NewArraySize, i * OldArraySize, true );
 		}
 		RHICmdList.EndUAVOverlap(NewBuffer.UAV);
 
@@ -481,7 +487,7 @@ RENDERCORE_API bool ResizeResourceSOAIfNeeded(FRDGBuilder& GraphBuilder, FRWBuff
 
 			for (uint32 i = 0; i < NumArrays; i++)
 			{
-				MemcpyResource(RHICmdList, NewBuffer, OldBuffer, CopyBytes, i * NewArraySize, i * OldArraySize);
+				MemcpyResource(RHICmdList, NewBuffer, OldBuffer, CopyBytes, i * NewArraySize, i * OldArraySize, true);
 			}
 			RHICmdList.EndUAVOverlap(NewBuffer.UAV);
 		});
@@ -689,9 +695,9 @@ void FScatterUploadBuffer::ResourceUploadTo(FRHICommandList& RHICmdList, const R
 
 template RENDERCORE_API void MemsetResource< FRWBufferStructured>(FRHICommandList& RHICmdList, const FRWBufferStructured& DstBuffer, uint32 Value, uint32 NumBytes, uint32 DstOffset);
 template RENDERCORE_API void MemsetResource< FRWByteAddressBuffer>(FRHICommandList& RHICmdList, const FRWByteAddressBuffer& DstBuffer, uint32 Value, uint32 NumBytes, uint32 DstOffset);
-template RENDERCORE_API void MemcpyResource< FTextureRWBuffer2D>(FRHICommandList& RHICmdList, const FTextureRWBuffer2D& DstBuffer, const FTextureRWBuffer2D& SrcBuffer, uint32 NumBytes, uint32 DstOffset, uint32 SrcOffset);
-template RENDERCORE_API void MemcpyResource< FRWBufferStructured>(FRHICommandList& RHICmdList, const FRWBufferStructured& DstBuffer, const FRWBufferStructured& SrcBuffer, uint32 NumBytes, uint32 DstOffset, uint32 SrcOffset);
-template RENDERCORE_API void MemcpyResource< FRWByteAddressBuffer>(FRHICommandList& RHICmdList, const FRWByteAddressBuffer& DstBuffer, const FRWByteAddressBuffer& SrcBuffer, uint32 NumBytes, uint32 DstOffset, uint32 SrcOffset);
+template RENDERCORE_API void MemcpyResource< FTextureRWBuffer2D>(FRHICommandList& RHICmdList, const FTextureRWBuffer2D& DstBuffer, const FTextureRWBuffer2D& SrcBuffer, uint32 NumBytes, uint32 DstOffset, uint32 SrcOffset, bool bAlreadyInUAVOverlap);
+template RENDERCORE_API void MemcpyResource< FRWBufferStructured>(FRHICommandList& RHICmdList, const FRWBufferStructured& DstBuffer, const FRWBufferStructured& SrcBuffer, uint32 NumBytes, uint32 DstOffset, uint32 SrcOffset, bool bAlreadyInUAVOverlap);
+template RENDERCORE_API void MemcpyResource< FRWByteAddressBuffer>(FRHICommandList& RHICmdList, const FRWByteAddressBuffer& DstBuffer, const FRWByteAddressBuffer& SrcBuffer, uint32 NumBytes, uint32 DstOffset, uint32 SrcOffset, bool bAlreadyInUAVOverlap);
 template RENDERCORE_API void FScatterUploadBuffer::ResourceUploadTo<FTextureRWBuffer2D>(FRHICommandList& RHICmdList, const FTextureRWBuffer2D& DstBuffer, bool bFlush);
 template RENDERCORE_API void FScatterUploadBuffer::ResourceUploadTo<FRWBufferStructured>(FRHICommandList& RHICmdList, const FRWBufferStructured& DstBuffer, bool bFlush);
 template RENDERCORE_API void FScatterUploadBuffer::ResourceUploadTo<FRWByteAddressBuffer>(FRHICommandList& RHICmdList, const FRWByteAddressBuffer& DstBuffer, bool bFlush);
\ No newline at end of file
diff --git a/Engine/Source/Runtime/RenderCore/Public/UnifiedBuffer.h b/Engine/Source/Runtime/RenderCore/Public/UnifiedBuffer.h
--- a/Engine/Source/Runtime/RenderCore/Public/UnifiedBuffer.h
+++ b/Engine/Source/Runtime/RenderCore/Public/UnifiedBuffer.h
@@ -17,7 +17,7 @@ class FRDGBuilder;
 template<typename ResourceType>
 extern RENDERCORE_API void MemsetResource(FRHICommandList& RHICmdList, const ResourceType& DstBuffer, uint32 Value, uint32 NumBytes, uint32 DstOffset = 0);
 template<typename ResourceType>
-extern RENDERCORE_API void MemcpyResource(FRHICommandList& RHICmdList, const ResourceType& DstBuffer, const ResourceType& SrcBuffer, uint32 NumBytes, uint32 DstOffset = 0, uint32 SrcOffset = 0);
+extern RENDERCORE_API void MemcpyResource(FRHICommandList& RHICmdList, const ResourceType& DstBuffer, const ResourceType& SrcBuffer, uint32 NumBytes, uint32 DstOffset = 0, uint32 SrcOffset = 0, bool bAlreadyInUAVOverlap = false);
 template<typename ResourceType>
 extern RENDERCORE_API bool ResizeResourceIfNeeded(FRHICommandList& RHICmdList, ResourceType& Texture, uint32 NumBytes, const TCHAR* DebugName);
 template<typename ResourceType>
