Fixed transient texture array not properly created

diff --git a/Engine/Source/Runtime/RenderCore/Private/RenderTargetPool.cpp b/Engine/Source/Runtime/RenderCore/Private/RenderTargetPool.cpp
--- a/Engine/Source/Runtime/RenderCore/Private/RenderTargetPool.cpp
+++ b/Engine/Source/Runtime/RenderCore/Private/RenderTargetPool.cpp
@@ -462,14 +462,31 @@ Done:
 			if (Desc.Is2DTexture())
 			{
 				// this is useful to get a CPU lockable texture through the same interface
-				Found->RenderTargetItem.ShaderResourceTexture = RHICreateTexture2D(
-					Desc.Extent.X,
-					Desc.Extent.Y,
-					Desc.Format,
-					Desc.NumMips,
-					Desc.NumSamples,
-					Desc.Flags,
-					CreateInfo);
+				if (!Desc.IsArray())
+				{
+					Found->RenderTargetItem.ShaderResourceTexture = RHICreateTexture2D(
+						Desc.Extent.X,
+						Desc.Extent.Y,
+						Desc.Format,
+						Desc.NumMips,
+						Desc.NumSamples,
+						Desc.Flags,
+						CreateInfo);
+				}
+				else
+				{
+					Found->RenderTargetItem.ShaderResourceTexture = RHICreateTexture2DArray(
+						Desc.Extent.X,
+						Desc.Extent.Y,
+						Desc.ArraySize,
+						Desc.Format,
+						Desc.NumMips,
+						Desc.NumSamples,
+						Desc.Flags,
+						CreateInfo);
+				}
+
+				
 			}
 			else if (Desc.Is3DTexture())
 			{
