UE5EA - Fix threading issue when updating devices with Turnkey
One should not access target platform manager outside the game thread - in this case the target platforms array was invalidated while rendering thread was requesting info which was also causing it to re-initialize

diff --git a/Engine/Source/Runtime/RenderCore/Private/RenderUtils.cpp b/Engine/Source/Runtime/RenderCore/Private/RenderUtils.cpp
--- a/Engine/Source/Runtime/RenderCore/Private/RenderUtils.cpp
+++ b/Engine/Source/Runtime/RenderCore/Private/RenderUtils.cpp
@@ -1558,48 +1558,31 @@ RENDERCORE_API void QuantizeSceneBufferSize(const FIntPoint& InBufferSize, FIntP
 
 RENDERCORE_API bool UseVirtualTexturing(const FStaticFeatureLevel InFeatureLevel, const ITargetPlatform* TargetPlatform)
 {
-#if !PLATFORM_SUPPORTS_VIRTUAL_TEXTURE_STREAMING
-	if (GIsEditor == false)
+#if PLATFORM_SUPPORTS_VIRTUAL_TEXTURE_STREAMING
+	if (!FPlatformProperties::SupportsVirtualTextureStreaming())
 	{
 		return false;
 	}
-	else
-#endif
-	{
-		// does the platform supports it.
-#if WITH_EDITOR
-		if (GIsEditor && TargetPlatform == nullptr)
-		{
-			ITargetPlatformManagerModule* TPM = GetTargetPlatformManager();
-			if (TPM)
-			{
-				TargetPlatform = TPM->GetRunningTargetPlatform();
-			}
-		}
-
-		if (TargetPlatform && TargetPlatform->SupportsFeature(ETargetPlatformFeatures::VirtualTextureStreaming) == false)
-		{
-			return false;
-		}
-#endif
-
-		// does the project has it enabled ?
-		static const auto CVarVirtualTexture = IConsoleManager::Get().FindTConsoleVariableDataInt(TEXT("r.VirtualTextures"));
-		check(CVarVirtualTexture);
-		if (CVarVirtualTexture->GetValueOnAnyThread() == 0)
-		{
-			return false;
-		}		
 
-		// mobile needs an additional switch to enable VT		
-		static const auto CVarMobileVirtualTexture = IConsoleManager::Get().FindTConsoleVariableDataInt(TEXT("r.Mobile.VirtualTextures"));
-		if (InFeatureLevel == ERHIFeatureLevel::ES3_1 && CVarMobileVirtualTexture->GetValueOnAnyThread() == 0)
-		{
-			return false;
-		}
+	// does the project has it enabled ?
+	static const auto CVarVirtualTexture = IConsoleManager::Get().FindTConsoleVariableDataInt(TEXT("r.VirtualTextures"));
+	check(CVarVirtualTexture);
+	if (CVarVirtualTexture->GetValueOnAnyThread() == 0)
+	{
+		return false;
+	}		
 
-		return true;
+	// mobile needs an additional switch to enable VT		
+	static const auto CVarMobileVirtualTexture = IConsoleManager::Get().FindTConsoleVariableDataInt(TEXT("r.Mobile.VirtualTextures"));
+	if (InFeatureLevel == ERHIFeatureLevel::ES3_1 && CVarMobileVirtualTexture->GetValueOnAnyThread() == 0)
+	{
+		return false;
 	}
+
+	return true;
+#else
+	return false;
+#endif
 }
 
 RENDERCORE_API bool UseVirtualTextureLightmap(const FStaticFeatureLevel InFeatureLevel, const ITargetPlatform* TargetPlatform)
