Moving shader symbol exporting into a shared type.
Adding shader symbol export path override cvar. This cvar supports specifying the shader platform name via "{Platform}".
Adding Luke's change for having build machine specific shader configs.

diff --git a/Engine/Config/ConsoleVariables.ini b/Engine/Config/ConsoleVariables.ini
--- a/Engine/Config/ConsoleVariables.ini
+++ b/Engine/Config/ConsoleVariables.ini
@@ -74,6 +74,9 @@ r.DumpShaderDebugInfo=2
 ;r.Shaders.KeepDebugInfo=1
 ; Uncomment to skip shader compression. Can save a significant time when using debug shaders.
 ;r.Shaders.SkipCompression=1
+; Uncomment and change to set a custom folder to write shader symbols to when they are generated. You can optionally use "{Platform}" in the path to specify where the platform folder should go.
+; The quotes are REQUIRED for the {Platform} token to work.
+;r.Shaders.SymbolPathOverride="Path\To\Symbols\{Platform}"
 
 ; If Linux editor crashes in FMallocBinned with callstack that mentions MeshUtilities, you may need to uncomment this.
 ;r.TriangleOrderOptimization=2
@@ -96,4 +99,4 @@ p.chaos.AllowCreatePhysxBodies=1
 fx.SkipVectorVMBackendOptimizations=1
  
 ; Comment this out to turn off audio submix effects.
-;au.BypassAllSubmixEffects=0
\ No newline at end of file
+;au.BypassAllSubmixEffects=0
diff --git a/Engine/Source/Developer/FileUtilities/Private/ZipArchiveWriter.cpp b/Engine/Source/Developer/FileUtilities/Private/ZipArchiveWriter.cpp
--- a/Engine/Source/Developer/FileUtilities/Private/ZipArchiveWriter.cpp
+++ b/Engine/Source/Developer/FileUtilities/Private/ZipArchiveWriter.cpp
@@ -97,7 +97,7 @@ FZipArchiveWriter::~FZipArchiveWriter()
 	}
 }
 
-void FZipArchiveWriter::AddFile(const FString& Filename, const TArray<uint8>& Data, const FDateTime& Timestamp)
+void FZipArchiveWriter::AddFile(const FString& Filename, TConstArrayView<uint8> Data, const FDateTime& Timestamp)
 {
 	uint32 Crc = FCrc::MemCrc32(Data.GetData(), Data.Num());
 
@@ -140,6 +140,11 @@ void FZipArchiveWriter::AddFile(const FString& Filename, const TArray<uint8>& Da
 	Flush();
 }
 
+void FZipArchiveWriter::AddFile(const FString& Filename, const TArray<uint8>& Data, const FDateTime& Timestamp)
+{
+	AddFile(Filename, TConstArrayView<uint8>(Data), Timestamp);
+}
+
 void FZipArchiveWriter::Flush()
 {
 	if (Buffer.Num())
diff --git a/Engine/Source/Developer/FileUtilities/Public/FileUtilities/ZipArchiveWriter.h b/Engine/Source/Developer/FileUtilities/Public/FileUtilities/ZipArchiveWriter.h
--- a/Engine/Source/Developer/FileUtilities/Public/FileUtilities/ZipArchiveWriter.h
+++ b/Engine/Source/Developer/FileUtilities/Public/FileUtilities/ZipArchiveWriter.h
@@ -49,6 +49,7 @@ public:
 	FZipArchiveWriter(IFileHandle* InFile);
 	~FZipArchiveWriter();
 
+	void AddFile(const FString& Filename, TConstArrayView<uint8> Data, const FDateTime& Timestamp);
 	void AddFile(const FString& Filename, const TArray<uint8>& Data, const FDateTime& Timestamp);
 };
 
diff --git a/Engine/Source/Developer/ShaderCompilerCommon/Public/ShaderSymbolExport.h b/Engine/Source/Developer/ShaderCompilerCommon/Public/ShaderSymbolExport.h
new file mode 100644
--- /dev/null
+++ b/Engine/Source/Developer/ShaderCompilerCommon/Public/ShaderSymbolExport.h
@@ -0,0 +1,191 @@
+// Copyright Epic Games, Inc. All Rights Reserved.
+
+#pragma once
+
+#if WITH_ENGINE
+
+#include "Serialization/MemoryReader.h"
+#include "HAL/IConsoleManager.h"
+#include "HAL/PlatformFileManager.h"
+#include "GenericPlatform/GenericPlatformFile.h"
+#include "Misc/CommandLine.h"
+#include "Misc/ConfigCacheIni.h"
+#include "Misc/Paths.h"
+#include "FileUtilities/ZipArchiveWriter.h"
+
+DECLARE_LOG_CATEGORY_CLASS(LogShaderSymbolExport, Display, Display);
+
+class FShaderSymbolExport
+{
+public:
+	FShaderSymbolExport(FName InShaderFormat);
+
+	/** Should be called from IShaderFormat::NotifyShaderCompiled implementation.
+	*   Template type is the platform specific symbol data structure.
+	*/
+	template<typename TPlatformShaderSymbolData>
+	void NotifyShaderCompiled(const TConstArrayView<uint8>& PlatformSymbolData);
+
+private:
+	void Initialize();
+
+	const FName ShaderFormat;
+
+	TUniquePtr<FZipArchiveWriter> ZipWriter;
+	TSet<FString> ExportedShaders;
+	FString ExportPath;
+
+	uint64 TotalSymbolDataBytes{ 0 };
+	uint64 TotalSymbolData{ 0 };
+	bool bExportShaderSymbols{ false };
+};
+
+inline FShaderSymbolExport::FShaderSymbolExport(FName InShaderFormat)
+	: ShaderFormat(InShaderFormat)
+{
+}
+
+inline void FShaderSymbolExport::Initialize()
+{
+	const bool bSymbolsEnabled = ShouldExportShaderDebugInfo(ShaderFormat);
+	const bool bForceSymbols = FParse::Value(FCommandLine::Get(), TEXT("-ShaderSymbolsExport="), ExportPath);
+
+	if (bSymbolsEnabled || bForceSymbols)
+	{
+		// if no command line path is provided, look to the cvar first
+		if (ExportPath.IsEmpty())
+		{
+			GetShaderDebugInfoPathOverride(ExportPath, ShaderFormat);
+		}
+
+		// if there was no path set via command line or the cvar, fall back to our default
+		if (ExportPath.IsEmpty())
+		{
+			ExportPath = IFileManager::Get().ConvertToAbsolutePathForExternalAppForWrite(
+				*(FPaths::ProjectSavedDir() / TEXT("ShaderSymbols") / ShaderFormat.ToString()));
+		}
+
+		IPlatformFile& PlatformFile = FPlatformFileManager::Get().GetPlatformFile();
+		bExportShaderSymbols = PlatformFile.CreateDirectoryTree(*ExportPath);
+
+		if (!bExportShaderSymbols)
+		{
+			UE_LOG(LogShaderSymbolExport, Error, TEXT("Failed to create shader symbols output directory. Shader symbol export will be disabled."));
+		}
+		else
+		{
+			// Check if the export mode is to an uncompressed archive or loose files.
+			bool bExportAsZip = ShouldExportShaderDebugInfoAsZip(ShaderFormat);
+			if (bExportAsZip || FParse::Param(FCommandLine::Get(), TEXT("ShaderSymbolsExportZip")))
+			{
+				FString SingleFilePath = ExportPath / TEXT("ShaderSymbols.zip");
+
+				IFileHandle* OutputZipFile = PlatformFile.OpenWrite(*SingleFilePath);
+				if (!OutputZipFile)
+				{
+					UE_LOG(LogShaderSymbolExport, Error, TEXT("Failed to create shader symbols output file \"%s\". Shader symbol export will be disabled."), *SingleFilePath);
+					bExportShaderSymbols = false;
+				}
+				else
+				{
+					ZipWriter = MakeUnique<FZipArchiveWriter>(OutputZipFile);
+				}
+			}
+		}
+	}
+
+	if (bExportShaderSymbols)
+	{
+		UE_LOG(LogShaderSymbolExport, Display, TEXT("Shader symbol export enabled. Output directory: \"%s\""), *ExportPath);
+		if (ZipWriter)
+		{
+			UE_LOG(LogShaderSymbolExport, Display, TEXT("Shader symbol zip mode enabled. Shader symbols will be archived in a single (uncompressed) zip file."));
+		}
+	}
+}
+
+template<typename TPlatformShaderSymbolData>
+inline void FShaderSymbolExport::NotifyShaderCompiled(const TConstArrayView<uint8>& PlatformSymbolData)
+{
+	static bool bFirst = true;
+	if (bFirst)
+	{
+		// If we get called, we know we're compiling. Do one time initialization
+		// which will create the output directory / open the open file stream.
+		Initialize();
+		bFirst = false;
+	}
+
+	if (bExportShaderSymbols)
+	{
+		// Deserialize the platform symbol data
+		TPlatformShaderSymbolData FullSymbolData;
+		FMemoryReaderView Ar(PlatformSymbolData);
+		Ar << FullSymbolData;
+
+		for (const auto& SymbolData : FullSymbolData.GetAllSymbolData())
+		{
+			if (TConstArrayView<uint8> Contents = SymbolData.GetContents(); !Contents.IsEmpty())
+			{
+				const FString FileName = SymbolData.GetFilename();
+
+				// Skip this symbol data if we've already exported it before.
+				bool bAlreadyInSet = false;
+				ExportedShaders.Add(FileName, &bAlreadyInSet);
+				if (bAlreadyInSet)
+				{
+					// We've already exported this shader hash
+					continue;
+				}
+
+				// Emit periodic log messages detailing the size of the shader symbols output file/directory.
+				static uint64 LastReport = 0;
+				TotalSymbolDataBytes += Contents.Num();
+				TotalSymbolData++;
+
+				if ((TotalSymbolDataBytes - LastReport) >= (64 * 1024 * 1024))
+				{
+					UE_LOG(LogShaderSymbolExport, Display, TEXT("Shader symbols export size: %.2f MB, count: %llu"),
+						double(TotalSymbolDataBytes) / (1024.0 * 1024.0), TotalSymbolData);
+					LastReport = TotalSymbolDataBytes;
+				}
+
+				if (ZipWriter)
+				{
+					// Append the platform data to the zip file
+					ZipWriter->AddFile(FileName, Contents, FDateTime::Now());
+				}
+				else
+				{
+					IPlatformFile& PlatformFile = FPlatformFileManager::Get().GetPlatformFile();
+
+					// Write the symbols to the export directory
+					const FString OutputPath = ExportPath / FileName;
+					const FString Directory = FPaths::GetPath(OutputPath);
+
+					// FileName could contain extra folders, so we need to make sure they exist first.
+					if (!PlatformFile.CreateDirectoryTree(*Directory))
+					{
+						UE_LOG(LogShaderSymbolExport, Error, TEXT("Failed to create shader symbol directory \"%s\"."), *Directory);
+					}
+					else
+					{
+						IFileHandle* File = PlatformFile.OpenWrite(*OutputPath);
+						if (!File || !File->Write(Contents.GetData(), Contents.Num()))
+						{
+							UE_LOG(LogShaderSymbolExport, Error, TEXT("Failed to export shader symbols \"%s\"."), *OutputPath);
+						}
+
+						if (File)
+						{
+							delete File;
+						}
+					}
+				}
+			}
+		}
+	}
+}
+
+
+#endif // WITH_ENGINE
diff --git a/Engine/Source/Runtime/Engine/Private/ShaderCompiler/ShaderCompiler.cpp b/Engine/Source/Runtime/Engine/Private/ShaderCompiler/ShaderCompiler.cpp
--- a/Engine/Source/Runtime/Engine/Private/ShaderCompiler/ShaderCompiler.cpp
+++ b/Engine/Source/Runtime/Engine/Private/ShaderCompiler/ShaderCompiler.cpp
@@ -838,6 +838,13 @@ static TAutoConsoleVariable<int32> CVarExportShaderDebugData(
 	TEXT("For some platforms this cvar can be overriden in the Engine.ini, under the [ShaderCompiler] section."),
 	ECVF_ReadOnly);
 
+static TAutoConsoleVariable<FString> CVarShadersSymbolPathOverride(
+	TEXT("r.Shaders.SymbolPathOverride"),
+	"",
+	TEXT("Override output location of shader symbols. If the path contains the text '{Platform}', that will be replaced with the shader platform string.\n")
+	TEXT("Empty: use default location Saved/ShaderSymbols/{Platform}"),
+	ECVF_ReadOnly);
+
 static TAutoConsoleVariable<int32> CVarAllowUniqueDebugInfo(
 	TEXT("r.Shaders.AllowUniqueDebugInfo"),
 	0,
@@ -846,8 +853,8 @@ static TAutoConsoleVariable<int32> CVarAllowUniqueDebugInfo(
 	TEXT("For some platforms this cvar can be overriden in the Engine.ini, under the [ShaderCompiler] section."),
 	ECVF_ReadOnly);
 
-static TAutoConsoleVariable<int32> CVarExportShaderDebugDataMode(
-	TEXT("r.Shaders.ExportDebugInfoMode"),
+static TAutoConsoleVariable<int32> CVarExportShaderDebugInfoZip(
+	TEXT("r.Shaders.ExportDebugInfo.Zip"),
 	0,
 	TEXT(" 0: Export as loose files.\n")
 	TEXT(" 1: Export as an uncompressed archive.\n"),
diff --git a/Engine/Source/Runtime/RenderCore/Private/ShaderCore.cpp b/Engine/Source/Runtime/RenderCore/Private/ShaderCore.cpp
--- a/Engine/Source/Runtime/RenderCore/Private/ShaderCore.cpp
+++ b/Engine/Source/Runtime/RenderCore/Private/ShaderCore.cpp
@@ -321,7 +321,12 @@ static bool IsShaderCompilerConfigEnabledForPlatform(const IConsoleVariable* Glo
 		{
 			if (FConfigCacheIni* PlatformConfig = TargetPlatform->GetConfigSystem())
 			{
-				PlatformConfig->GetBool(TEXT("ShaderCompiler"), Key, bEnabled, GEngineIni);
+				// Override with a build machine specific setting, if present.
+				bool QueryBaseConfig = GIsBuildMachine ? !PlatformConfig->GetBool(TEXT("ShaderCompiler_BuildMachine"), Key, bEnabled, GEngineIni) : true;
+				if (QueryBaseConfig)
+				{
+					PlatformConfig->GetBool(TEXT("ShaderCompiler"), Key, bEnabled, GEngineIni);
+				}
 			}
 		}
 	}
@@ -330,44 +335,81 @@ static bool IsShaderCompilerConfigEnabledForPlatform(const IConsoleVariable* Glo
 	return bEnabled;
 }
 
-bool ShouldKeepShaderDebugInfo(FName ShaderFormat)
+static bool GetShaderCompilerStringForPlatform(FString& OutputString, const IConsoleVariable* GlobalCvar, const TCHAR* Key, FName ShaderFormat)
 {
-	static const TCHAR* const KeepDebugInfoName = TEXT("r.Shaders.KeepDebugInfo");
-	static const IConsoleVariable* KeepDebugInfoCvar = IConsoleManager::Get().FindConsoleVariable(KeepDebugInfoName);
+	// First check the global cvar
+	if (GlobalCvar)
+	{
+		OutputString = GlobalCvar->GetString();
+	}
 
-	// PrepareExportedDebugInfo is only read from the platform settings, not the global cvar
-	static const TCHAR* const PrepareExportedDebugInfoName = TEXT("r.Shaders.PrepareExportedDebugInfo");
+#if WITH_EDITOR
+	if (OutputString.IsEmpty())
+	{
+		// Then check the per platform settings.
+		ITargetPlatform* TargetPlatform = GetTargetPlatformManager()->FindTargetPlatformWithSupport(TEXT("ShaderFormat"), ShaderFormat);
+		if (TargetPlatform)
+		{
+			if (FConfigCacheIni* PlatformConfig = TargetPlatform->GetConfigSystem())
+			{
+				// Override with a build machine specific setting, if present.
+				if (GIsBuildMachine)
+				{
+					OutputString = PlatformConfig->GetStr(TEXT("ShaderCompiler_BuildMachine"), Key, GEngineIni);
+				}
 
-	return IsShaderCompilerConfigEnabledForPlatform(KeepDebugInfoCvar, KeepDebugInfoName, ShaderFormat)
-		|| IsShaderCompilerConfigEnabledForPlatform(nullptr, PrepareExportedDebugInfoName, ShaderFormat);
-}
+				if (OutputString.IsEmpty())
+				{
+					OutputString = PlatformConfig->GetStr(TEXT("ShaderCompiler"), Key, GEngineIni);
+				}
+			}
+		}
+	}
+#endif
 
-bool ShouldKeepShaderDebugInfo(EShaderPlatform ShaderPlatform)
-{
-	return ShouldKeepShaderDebugInfo(LegacyShaderPlatformToShaderFormat(ShaderPlatform));
+	return !OutputString.IsEmpty();
 }
 
-bool ShouldExportShaderDebugInfo(FName ShaderFormat)
+struct FShaderSymbolSettingHelper
 {
-	// r.DumpShaderDebugInfo should also turn on ExportShaderDebugInfo
-	// The difference is that r.DumpShaderDebugInfo will also output engine debug files such as converted hlsl or SCW helper files.
-	// Where as r.Shader.ExportDebugInfo is purely to export the graphics debugging tool's debug info files.
-	static IConsoleVariable* CVarDumpDebugInfo = IConsoleManager::Get().FindConsoleVariable(TEXT("r.DumpShaderDebugInfo"));
-	if (CVarDumpDebugInfo && CVarDumpDebugInfo->GetInt())
+	const TCHAR* const SettingName;
+	const IConsoleVariable* const SettingCVar;
+
+	const bool bEnableLegacySettings = true;
+public:
+	FShaderSymbolSettingHelper(const TCHAR* InSettingName, bool bPlatformOnly = false)
+		: SettingName(InSettingName)
+		, SettingCVar(!bPlatformOnly ? IConsoleManager::Get().FindConsoleVariable(InSettingName) : nullptr)
 	{
-		return true;
 	}
 
-	static const TCHAR* const KeepDebugInfoName = TEXT("r.Shaders.KeepDebugInfo");
-	static const IConsoleVariable* KeepDebugInfoCvar = IConsoleManager::Get().FindConsoleVariable(KeepDebugInfoName);
-	if (IsShaderCompilerConfigEnabledForPlatform(KeepDebugInfoCvar, KeepDebugInfoName, ShaderFormat))
+	bool IsEnabled(FName ShaderFormat) const
 	{
-		return true;
+		return  IsShaderCompilerConfigEnabledForPlatform(SettingCVar, SettingName, ShaderFormat);
 	}
 
-	static const TCHAR* const ExportDebugInfoName = TEXT("r.Shaders.ExportDebugInfo");
-	static const IConsoleVariable* ExportDebugInfoCvar = IConsoleManager::Get().FindConsoleVariable(ExportDebugInfoName);
-	return IsShaderCompilerConfigEnabledForPlatform(ExportDebugInfoCvar, ExportDebugInfoName, ShaderFormat);
+	bool GetString(FString& OutString, FName ShaderFormat) const
+	{
+		return GetShaderCompilerStringForPlatform(OutString, SettingCVar, SettingName, ShaderFormat);
+	}
+};
+
+bool ShouldKeepShaderDebugInfo(FName ShaderFormat)
+{
+	static const FShaderSymbolSettingHelper GenerateSymbols(TEXT("r.Shaders.KeepDebugInfo"));
+	static const FShaderSymbolSettingHelper PrepareExportedSymbols(TEXT("r.Shaders.PrepareExportedDebugInfo"), true);
+	return GenerateSymbols.IsEnabled(ShaderFormat) || PrepareExportedSymbols.IsEnabled(ShaderFormat);
+}
+
+bool ShouldKeepShaderDebugInfo(EShaderPlatform ShaderPlatform)
+{
+	return ShouldKeepShaderDebugInfo(LegacyShaderPlatformToShaderFormat(ShaderPlatform));
+}
+
+bool ShouldExportShaderDebugInfo(FName ShaderFormat)
+{
+	static const FShaderSymbolSettingHelper ExportSymbols(TEXT("r.Shaders.ExportDebugInfo"));
+	return ExportSymbols.IsEnabled(ShaderFormat) || ShouldKeepShaderDebugInfo(ShaderFormat);
 }
 
 bool ShouldExportShaderDebugInfo(EShaderPlatform ShaderPlatform)
@@ -375,6 +417,27 @@ bool ShouldExportShaderDebugInfo(EShaderPlatform ShaderPlatform)
 	return ShouldExportShaderDebugInfo(LegacyShaderPlatformToShaderFormat(ShaderPlatform));
 }
 
+bool GetShaderDebugInfoPathOverride(FString& OutPathOverride, FName ShaderFormat)
+{
+	static const FShaderSymbolSettingHelper SymbolPathOverride(TEXT("r.Shaders.SymbolPathOverride"));
+	if (SymbolPathOverride.GetString(OutPathOverride, ShaderFormat))
+	{
+		if (!OutPathOverride.IsEmpty())
+		{
+			// Allow the user to specify the location of the per-platform string.
+			OutPathOverride = OutPathOverride.Replace(TEXT("{Platform}"), *ShaderFormat.ToString(), ESearchCase::IgnoreCase);
+		}
+		return !OutPathOverride.IsEmpty();
+	}
+	return false;
+}
+
+bool ShouldExportShaderDebugInfoAsZip(FName ShaderFormat)
+{
+	static const FShaderSymbolSettingHelper ExportSymbolsZip(TEXT("r.Shaders.ExportDebugInfo.Zip"));
+	return ExportSymbolsZip.IsEnabled(ShaderFormat);
+}
+
 #ifndef UE_ALLOW_SHADER_COMPILING
 // is shader compiling allowed at all? (set to 0 in a cooked editor .Target.cs if the target has no shaders available)
 #define UE_ALLOW_SHADER_COMPILING 1
@@ -415,9 +478,8 @@ bool AllowGlobalShaderLoad()
 
 bool ShouldAllowUniqueDebugInfo(FName ShaderFormat)
 {
-	static const TCHAR* const VariableName = TEXT("r.Shaders.AllowUniqueDebugInfo");
-	static const IConsoleVariable* GlobalCvar = IConsoleManager::Get().FindConsoleVariable(VariableName);
-	return IsShaderCompilerConfigEnabledForPlatform(GlobalCvar, VariableName, ShaderFormat);
+	static const FShaderSymbolSettingHelper AllowUniqueDebugInfo(TEXT("r.Shaders.AllowUniqueDebugInfo"));
+	return AllowUniqueDebugInfo.IsEnabled(ShaderFormat);
 }
 
 bool FShaderParameterMap::FindParameterAllocation(const TCHAR* ParameterName,uint16& OutBufferIndex,uint16& OutBaseIndex,uint16& OutSize) const
diff --git a/Engine/Source/Runtime/RenderCore/Public/ShaderCompilerCore.h b/Engine/Source/Runtime/RenderCore/Public/ShaderCompilerCore.h
--- a/Engine/Source/Runtime/RenderCore/Public/ShaderCompilerCore.h
+++ b/Engine/Source/Runtime/RenderCore/Public/ShaderCompilerCore.h
@@ -32,9 +32,13 @@ extern RENDERCORE_API bool ShouldKeepShaderDebugInfo(FName ShaderFormat);
 extern RENDERCORE_API bool ShouldExportShaderDebugInfo(EShaderPlatform Platform);
 extern RENDERCORE_API bool ShouldExportShaderDebugInfo(FName ShaderFormat);
 
-/** Returns true if (external) debug information should be specific to each shader rather than be deduplicated. */
-extern RENDERCORE_API bool ShouldAllowUniqueDebugInfo(FName ShaderFormat);
+extern RENDERCORE_API bool GetShaderDebugInfoPathOverride(FString& OutPathOverride, FName ShaderFormat);
+
+/** Returns true if shader symbols should be exported to a single zip file for a given platform. */
+extern RENDERCORE_API bool ShouldExportShaderDebugInfoAsZip(FName ShaderFormat);
 
+/** Returns true if (external) shader symbols should be specific to each shader rather than be de-duplicated. */
+extern RENDERCORE_API bool ShouldAllowUniqueDebugInfo(FName ShaderFormat);
 
 /** Returns true is shader compiling is allowed */
 extern RENDERCORE_API bool AllowShaderCompiling();
