Fixing an issue where PrepareExportedDebugInfo would actually force symbols to be written.
Adding {ProjectDir} and {ProjectSavedDir} to r.Shaders.SymbolPathOverride

diff --git a/Engine/Config/ConsoleVariables.ini b/Engine/Config/ConsoleVariables.ini
--- a/Engine/Config/ConsoleVariables.ini
+++ b/Engine/Config/ConsoleVariables.ini
@@ -74,9 +74,10 @@ r.DumpShaderDebugInfo=2
 ;r.Shaders.KeepDebugInfo=1
 ; Uncomment to skip shader compression. Can save a significant time when using debug shaders.
 ;r.Shaders.SkipCompression=1
-; Uncomment and change to set a custom folder to write shader symbols to when they are generated. You can optionally use "{Platform}" in the path to specify where the platform folder should go.
-; The quotes are REQUIRED for the {Platform} token to work.
-;r.Shaders.SymbolPathOverride="Path\To\Symbols\{Platform}"
+; Uncomment and change to set a custom folder to write shader symbols to when they are generated.
+; A few tokens can be used for engine paths. {ProjectDir} is the project directory, {ProjectSavedDir} is the project saved dir and {Platform} is the shader platform.
+; The quotes are REQUIRED for the {} tokens to work.
+;r.Shaders.SymbolPathOverride="{ProjectSavedDir}\ShaderSymbols\{Platform}"
 
 ; If Linux editor crashes in FMallocBinned with callstack that mentions MeshUtilities, you may need to uncomment this.
 ;r.TriangleOrderOptimization=2
diff --git a/Engine/Source/Developer/ShaderCompilerCommon/Public/ShaderSymbolExport.h b/Engine/Source/Developer/ShaderCompilerCommon/Public/ShaderSymbolExport.h
--- a/Engine/Source/Developer/ShaderCompilerCommon/Public/ShaderSymbolExport.h
+++ b/Engine/Source/Developer/ShaderCompilerCommon/Public/ShaderSymbolExport.h
@@ -55,7 +55,10 @@ inline void FShaderSymbolExport::Initialize()
 		// if no command line path is provided, look to the cvar first
 		if (ExportPath.IsEmpty())
 		{
-			GetShaderDebugInfoPathOverride(ExportPath, ShaderFormat);
+			if (GetShaderDebugInfoPathOverride(ExportPath, ShaderFormat))
+			{
+				ExportPath = IFileManager::Get().ConvertToAbsolutePathForExternalAppForWrite(*ExportPath);
+			}
 		}
 
 		// if there was no path set via command line or the cvar, fall back to our default
diff --git a/Engine/Source/Runtime/RenderCore/Private/ShaderCore.cpp b/Engine/Source/Runtime/RenderCore/Private/ShaderCore.cpp
--- a/Engine/Source/Runtime/RenderCore/Private/ShaderCore.cpp
+++ b/Engine/Source/Runtime/RenderCore/Private/ShaderCore.cpp
@@ -409,7 +409,8 @@ bool ShouldKeepShaderDebugInfo(EShaderPlatform ShaderPlatform)
 bool ShouldExportShaderDebugInfo(FName ShaderFormat)
 {
 	static const FShaderSymbolSettingHelper ExportSymbols(TEXT("r.Shaders.ExportDebugInfo"));
-	return ExportSymbols.IsEnabled(ShaderFormat) || ShouldKeepShaderDebugInfo(ShaderFormat);
+	static const FShaderSymbolSettingHelper GenerateSymbols(TEXT("r.Shaders.KeepDebugInfo"));
+	return ExportSymbols.IsEnabled(ShaderFormat) || GenerateSymbols.IsEnabled(ShaderFormat);
 }
 
 bool ShouldExportShaderDebugInfo(EShaderPlatform ShaderPlatform)
@@ -426,6 +427,10 @@ bool GetShaderDebugInfoPathOverride(FString& OutPathOverride, FName ShaderFormat
 		{
 			// Allow the user to specify the location of the per-platform string.
 			OutPathOverride = OutPathOverride.Replace(TEXT("{Platform}"), *ShaderFormat.ToString(), ESearchCase::IgnoreCase);
+			// Allow the user to specify the location of the per-project string.
+			OutPathOverride = OutPathOverride.Replace(TEXT("{ProjectDir}"), *FPaths::ProjectDir(), ESearchCase::IgnoreCase);
+			// Allow the user to specify the location of the per-project saved folder string.
+			OutPathOverride = OutPathOverride.Replace(TEXT("{ProjectSavedDir}"), *FPaths::ProjectSavedDir(), ESearchCase::IgnoreCase);
 		}
 		return !OutPathOverride.IsEmpty();
 	}
