Fixed compiling QAGame with LWC on

diff --git a/Engine/Source/Runtime/Core/Public/Math/Vector.h b/Engine/Source/Runtime/Core/Public/Math/Vector.h
--- a/Engine/Source/Runtime/Core/Public/Math/Vector.h
+++ b/Engine/Source/Runtime/Core/Public/Math/Vector.h
@@ -2451,8 +2451,8 @@ static FORCEINLINE UE::Math::TVector<T> INTEL_ORDER_VECTOR(UE::Math::TVector<T>
  * @param Point 3D position of interest
  * @return the distance from the Point to the bounding box.
  */
-template<typename T>
-FORCEINLINE T ComputeSquaredDistanceFromBoxToPoint(const UE::Math::TVector<T>& Mins, const UE::Math::TVector<T>& Maxs, const UE::Math::TVector<T>& Point)
+template<typename T, typename U>
+FORCEINLINE T ComputeSquaredDistanceFromBoxToPoint(const UE::Math::TVector<T>& Mins, const UE::Math::TVector<T>& Maxs, const UE::Math::TVector<U>& Point)
 {
 	// Accumulates the distance as we iterate axis
 	T DistSquared = 0;
diff --git a/Engine/Source/Runtime/RenderCore/Private/RenderTransform.cpp b/Engine/Source/Runtime/RenderCore/Private/RenderTransform.cpp
--- a/Engine/Source/Runtime/RenderCore/Private/RenderTransform.cpp
+++ b/Engine/Source/Runtime/RenderCore/Private/RenderTransform.cpp
@@ -35,6 +35,37 @@ FRenderBounds FRenderBounds::TransformBy(const FMatrix44f& M) const
 	return Result;
 }
 
+FRenderBounds FRenderBounds::TransformBy(const FMatrix44d& M) const
+{
+	FVector3f Origin = GetCenter();
+	FVector3f Extent = GetExtent();
+
+	const VectorRegister VecOrigin = VectorLoadFloat3(&Origin);
+	const VectorRegister VecExtent = VectorLoadFloat3(&Extent);
+
+	const VectorRegister4Double m0 = VectorLoadAligned(M.M[0]);
+	const VectorRegister4Double m1 = VectorLoadAligned(M.M[1]);
+	const VectorRegister4Double m2 = VectorLoadAligned(M.M[2]);
+	const VectorRegister4Double m3 = VectorLoadAligned(M.M[3]);
+
+	VectorRegister4Double NewOrigin = VectorMultiply(VectorReplicate(VecOrigin, 0), m0);
+	NewOrigin = VectorMultiplyAdd(VectorReplicate(VecOrigin, 1), m1, NewOrigin);
+	NewOrigin = VectorMultiplyAdd(VectorReplicate(VecOrigin, 2), m2, NewOrigin);
+	NewOrigin = VectorAdd(NewOrigin, m3);
+
+	VectorRegister4Double NewExtent = VectorAbs(VectorMultiply(VectorReplicate(VecExtent, 0), m0));
+	NewExtent = VectorAdd(NewExtent, VectorAbs(VectorMultiply(VectorReplicate(VecExtent, 1), m1)));
+	NewExtent = VectorAdd(NewExtent, VectorAbs(VectorMultiply(VectorReplicate(VecExtent, 2), m2)));
+
+	VectorStoreFloat3(NewExtent, &(Extent.X));
+	VectorStoreFloat3(NewOrigin, &(Origin.X));
+
+	FRenderBounds Result;
+	Result.Min = Origin - Extent;
+	Result.Max = Origin + Extent;
+	return Result;
+}
+
 FRenderBounds FRenderBounds::TransformBy(const FRenderTransform& T) const
 {
 	return TransformBy(T.ToMatrix44f());
diff --git a/Engine/Source/Runtime/RenderCore/Public/RenderTransform.h b/Engine/Source/Runtime/RenderCore/Public/RenderTransform.h
--- a/Engine/Source/Runtime/RenderCore/Public/RenderTransform.h
+++ b/Engine/Source/Runtime/RenderCore/Public/RenderTransform.h
@@ -26,6 +26,14 @@ public:
 		Origin				= FVector3f(M.M[3][0], M.M[3][1], M.M[3][2]);
 	}
 
+	FORCEINLINE FRenderTransform(const FMatrix44d& M)
+	{
+		TransformRows[0] = FVector3f(M.M[0][0], M.M[0][1], M.M[0][2]);
+		TransformRows[1] = FVector3f(M.M[1][0], M.M[1][1], M.M[1][2]);
+		TransformRows[2] = FVector3f(M.M[2][0], M.M[2][1], M.M[2][2]);
+		Origin = FVector3f(M.M[3][0], M.M[3][1], M.M[3][2]);
+	}
+
 	FORCEINLINE FRenderTransform& operator=(const FRenderTransform& From)
 	{
 		TransformRows[0] = From.TransformRows[0];
@@ -44,6 +52,15 @@ public:
 		return *this;
 	}
 
+	FORCEINLINE FRenderTransform& operator=(const FMatrix44d& From)
+	{
+		TransformRows[0] = FVector3f(From.M[0][0], From.M[0][1], From.M[0][2]);
+		TransformRows[1] = FVector3f(From.M[1][0], From.M[1][1], From.M[1][2]);
+		TransformRows[2] = FVector3f(From.M[2][0], From.M[2][1], From.M[2][2]);
+		Origin = FVector3f(From.M[3][0], From.M[3][1], From.M[3][2]);
+		return *this;
+	}
+
 	FORCEINLINE bool Equals(const FRenderTransform& Other, float Tolerance = KINDA_SMALL_NUMBER) const
 	{
 		return
@@ -299,6 +316,7 @@ public:
 	 * @return The transformed volume.
 	 */
 	RENDERCORE_API FRenderBounds TransformBy(const FMatrix44f& M) const;
+	RENDERCORE_API FRenderBounds TransformBy(const FMatrix44d& M) const;
 
 	/**
 	 * Gets a bounding volume transformed by a render transform.
