Add simple functionality to wrap 1D dispatches into 2D or 3D as needed with minimal overhead, and add utility shader function to retrieve the linear index.
- Makes the assumption (checked) that dimensions are at least 128 in both X and Y.

diff --git a/Engine/Shaders/Private/ComputeShaderUtils.ush b/Engine/Shaders/Private/ComputeShaderUtils.ush
new file mode 100644
--- /dev/null
+++ b/Engine/Shaders/Private/ComputeShaderUtils.ush
@@ -0,0 +1,24 @@
+// Copyright Epic Games, Inc. All Rights Reserved.
+
+#pragma once
+
+// Note: this value must match WrappedGroupStride in RenderGraphUtils.h
+#define WRAPPED_GROUP_STRIDE (128U)
+
+/**
+ * Convert wrapped (when overflowing 1D dimension) SV_GroupID created by using FComputeShaderUtils::GetGroupCountWrapped to a linear group id.
+ * NOTE: when using this you need to check the group count against whatever max you have as the wrapping means there may be indices larger than the requested (unlike when you use a regular 1D dispatch).
+ */
+uint GetUnWrappedDispatchGroupId(uint3 GroupId /*SV_GroupID*/)
+{
+	return GroupId.x + (GroupId.z * WRAPPED_GROUP_STRIDE + GroupId.y) * WRAPPED_GROUP_STRIDE;
+}
+
+/**
+ * Calculate linear dispatch thread id from a wrapped GroupId : SV_GroupID (created by using FComputeShaderUtils::GetGroupCountWrapped) and 
+ * a linear GroupThreadIndex : SV_GroupIndex. 
+ */
+uint GetUnWrappedDispatchThreadId(uint3 GroupId /*SV_GroupID*/, uint GroupThreadIndex /*SV_GroupIndex*/, uint ThreadGroupSize)
+{
+	return GetUnWrappedDispatchGroupId(GroupId) * ThreadGroupSize + GroupThreadIndex;
+}
diff --git a/Engine/Source/Runtime/RenderCore/Public/RenderGraphUtils.h b/Engine/Source/Runtime/RenderCore/Public/RenderGraphUtils.h
--- a/Engine/Source/Runtime/RenderCore/Public/RenderGraphUtils.h
+++ b/Engine/Source/Runtime/RenderCore/Public/RenderGraphUtils.h
@@ -253,7 +253,7 @@ struct RENDERCORE_API FComputeShaderUtils
 	/** Ideal size of group size 8x8 to occupy at least an entire wave on GCN, two warp on Nvidia. */
 	static constexpr int32 kGolden2DGroupSize = 8;
 
-	/** Compute the number of group to dispatch. */
+	/** Compute the number of groups to dispatch. */
 	static FIntVector GetGroupCount(const int32 ThreadCount, const int32 GroupSize)
 	{
 		return FIntVector(
@@ -291,6 +291,54 @@ struct RENDERCORE_API FComputeShaderUtils
 	}
 
 
+
+	/**
+	 * Constant stride used when wrapping too large 1D dispatches using GetGroupCountWrapped, selected as 128 appears to be the lowest common denominator 
+	 * for mobile (GLES 3.1). For PC (with ~64k groups / dimension) this yields ~8M groups (500M threads @ group size 64) before even wrapping into Z.
+	 * NOTE: this value must match WRAPPED_GROUP_STRIDE in ComputeShaderUtils.ush
+	 */
+	static constexpr int32 WrappedGroupStride = 128;
+
+	/**
+	 * Wrapping number of groups to Y and Z dimension if X group count overflows GRHIMaxDispatchThreadGroupsPerDimension.
+	 * Calculate the linear group index as (or use GetUnWrappedDispatchGroupId(GroupId) in ComputeShaderUtils.ush):
+	 *  uint LinearGroupId = GroupId.X + (GroupId.Z * WrappedGroupStride + GroupId.Y) * WrappedGroupStride;
+	 * Note that you must use an early out because LinearGroupId may be larger than the ideal due to wrapping.
+	 */
+	static FIntVector GetGroupCountWrapped(const int32 TargetGroupCount)
+	{
+		check(GRHIMaxDispatchThreadGroupsPerDimension.X >= WrappedGroupStride && GRHIMaxDispatchThreadGroupsPerDimension.Y >= WrappedGroupStride);
+
+		FIntVector GroupCount(TargetGroupCount, 1, 1);
+
+		if (GroupCount.X > GRHIMaxDispatchThreadGroupsPerDimension.X)
+		{
+			GroupCount.Y = FMath::DivideAndRoundUp(GroupCount.X, WrappedGroupStride);
+			GroupCount.X = WrappedGroupStride;
+		}
+		if (GroupCount.Y > GRHIMaxDispatchThreadGroupsPerDimension.Y)
+		{
+			GroupCount.Z = FMath::DivideAndRoundUp(GroupCount.Y, WrappedGroupStride);
+			GroupCount.Y = WrappedGroupStride;
+		}
+
+		check(TargetGroupCount <= GroupCount.X * GroupCount.Y * GroupCount.Z);
+
+		return GroupCount;
+	}
+
+	/**
+	 * Compute the number of groups to dispatch and allow wrapping to Y and Z dimension if X group count overflows. 
+	 * Calculate the linear group index as (or use GetUnWrappedDispatchGroupId(GroupId) in ComputeShaderUtils.ush):
+	 *  uint LinearGroupId = GroupId.X + (GroupId.Z * WrappedGroupStride + GroupId.Y) * WrappedGroupStride;
+	 * Note that you must use an early out because LinearGroupId may be larger than the ideal due to wrapping.
+	 */
+	static FIntVector GetGroupCountWrapped(const int32 ThreadCount, const int32 GroupSize)
+	{
+		return GetGroupCountWrapped(FMath::DivideAndRoundUp(ThreadCount, GroupSize));
+	}
+
+
 	/** Dispatch a compute shader to rhi command list with its parameters. */
 	template<typename TShaderClass>
 	static void Dispatch(
