Move RHI_API into the UE_TRACE_CHANNEL_EXTERN() macro so it compiles if the macro is a stub.

diff --git a/Engine/Source/Runtime/RenderCore/Private/ProfilingDebugging/RealtimeGPUProfiler.cpp b/Engine/Source/Runtime/RenderCore/Private/ProfilingDebugging/RealtimeGPUProfiler.cpp
--- a/Engine/Source/Runtime/RenderCore/Private/ProfilingDebugging/RealtimeGPUProfiler.cpp
+++ b/Engine/Source/Runtime/RenderCore/Private/ProfilingDebugging/RealtimeGPUProfiler.cpp
@@ -827,7 +827,7 @@ void FRealtimeGPUProfiler::Cleanup()
 #if UE_TRACE_ENABLED
 namespace GpuProfilerTrace
 {
-	RHI_API UE_TRACE_CHANNEL_EXTERN(GpuChannel)
+	UE_TRACE_CHANNEL_EXTERN(GpuChannel, RHI_API)
 }
 #endif
 
