Fix clang compiler warning/error

diff --git a/Engine/Source/Runtime/RenderCore/Private/RenderGraphEvent.cpp b/Engine/Source/Runtime/RenderCore/Private/RenderGraphEvent.cpp
--- a/Engine/Source/Runtime/RenderCore/Private/RenderGraphEvent.cpp
+++ b/Engine/Source/Runtime/RenderCore/Private/RenderGraphEvent.cpp
@@ -13,7 +13,7 @@ RENDERCORE_API bool GetEmitRDGEvents()
 #if RDG_ENABLE_TRACE
 	bRDGChannelEnabled = UE_TRACE_CHANNELEXPR_IS_ENABLED(RDGChannel);
 #endif // RDG_ENABLE_TRACE
-	return GRDGEmitEvents || GRDGDebug || bRDGChannelEnabled;
+	return GRDGEmitEvents != 0 || GRDGDebug != 0 || bRDGChannelEnabled != 0;
 #else
 	return false;
 #endif
@@ -442,4 +442,4 @@ void FRDGCSVStatScopeStack::EndExecute()
 	}
 }
 
-#endif
\ No newline at end of file
+#endif
