RDGBuffer objects actually need destructing with their new TFunction members. This fixes a slow leak when using Buffer creation callbacks.

diff --git a/Engine/Source/Runtime/RenderCore/Private/RenderGraphBuilder.cpp b/Engine/Source/Runtime/RenderCore/Private/RenderGraphBuilder.cpp
--- a/Engine/Source/Runtime/RenderCore/Private/RenderGraphBuilder.cpp
+++ b/Engine/Source/Runtime/RenderCore/Private/RenderGraphBuilder.cpp
@@ -1672,6 +1672,7 @@ void FRDGBuilder::Clear()
 {
 	TRACE_CPUPROFILER_EVENT_SCOPE(FRDGBuilder::Clear);
 	Passes.Clear();
+	Buffers.Clear();
 	UniformBuffers.Clear();
 	Blackboard.Clear();
 	ActivePooledTextures.Empty();
diff --git a/Engine/Source/Runtime/RenderCore/Public/RenderGraphDefinitions.h b/Engine/Source/Runtime/RenderCore/Public/RenderGraphDefinitions.h
--- a/Engine/Source/Runtime/RenderCore/Public/RenderGraphDefinitions.h
+++ b/Engine/Source/Runtime/RenderCore/Public/RenderGraphDefinitions.h
@@ -561,7 +561,7 @@ using FRDGTextureBitArray = TRDGHandleBitArray<FRDGTextureHandle>;
 class FRDGBuffer;
 using FRDGBufferRef = FRDGBuffer*;
 using FRDGBufferHandle = TRDGHandle<FRDGBuffer, uint16>;
-using FRDGBufferRegistry = TRDGHandleRegistry<FRDGBufferHandle, ERDGHandleRegistryDestructPolicy::Never>;
+using FRDGBufferRegistry = TRDGHandleRegistry<FRDGBufferHandle, ERDGHandleRegistryDestructPolicy::Registry>;
 using FRDGBufferBitArray = TRDGHandleBitArray<FRDGBufferHandle>;
 
 class FRDGPooledTexture;
