Added per platform setting to toggle distance field support. If it's disabled then no distance data will be cooked and distance fields will be disabled on the target platform.

Renamed FDataDrivenPlatformInfo::AudioCompressionSettingsIniSectionName to TargetSettingsIniSectionName. It's now used to populate PlatformMask at runtime, to make sure that it matches between cook and runtime. Before it was possible to cook with e.g. GRayTracingPlaformMask=0, but then the cooked game would run with GRayTracingPlaformMask=1.

diff --git a/Engine/Config/Android/DataDrivenPlatformInfo.ini b/Engine/Config/Android/DataDrivenPlatformInfo.ini
--- a/Engine/Config/Android/DataDrivenPlatformInfo.ini
+++ b/Engine/Config/Android/DataDrivenPlatformInfo.ini
@@ -1,6 +1,6 @@
 [DataDrivenPlatformInfo]
 bIsConfidential=false
-AudioCompressionSettingsIniSectionName=/Script/AndroidRuntimeSettings.AndroidRuntimeSettings
+TargetSettingsIniSectionName=/Script/AndroidRuntimeSettings.AndroidRuntimeSettings
 bHasDedicatedGamepad=false
 bInputSupportConfigurable=false
 DefaultInputType=Touch
diff --git a/Engine/Config/BaseEngine.ini b/Engine/Config/BaseEngine.ini
--- a/Engine/Config/BaseEngine.ini
+++ b/Engine/Config/BaseEngine.ini
@@ -2551,6 +2551,7 @@ bEnable2DPhysics=false
 
 [/Script/WindowsTargetPlatform.WindowsTargetSettings]
 +TargetedRHIs=PCD3D_SM5
+bEnableDistanceFields=true
 bEnableRayTracing=true
 bTarget32Bit=false
 
diff --git a/Engine/Config/IOS/DataDrivenPlatformInfo.ini b/Engine/Config/IOS/DataDrivenPlatformInfo.ini
--- a/Engine/Config/IOS/DataDrivenPlatformInfo.ini
+++ b/Engine/Config/IOS/DataDrivenPlatformInfo.ini
@@ -1,6 +1,6 @@
 [DataDrivenPlatformInfo]
 bIsConfidential=false
-AudioCompressionSettingsIniSectionName=/Script/IOSRuntimeSettings.IOSRuntimeSettings
+TargetSettingsIniSectionName=/Script/IOSRuntimeSettings.IOSRuntimeSettings
 bHasDedicatedGamepad=false
 bInputSupportConfigurable=false
 DefaultInputType=Touch
diff --git a/Engine/Config/Mac/DataDrivenPlatformInfo.ini b/Engine/Config/Mac/DataDrivenPlatformInfo.ini
--- a/Engine/Config/Mac/DataDrivenPlatformInfo.ini
+++ b/Engine/Config/Mac/DataDrivenPlatformInfo.ini
@@ -1,6 +1,6 @@
 [DataDrivenPlatformInfo]
 bIsConfidential=false
-AudioCompressionSettingsIniSectionName=/Script/MacTargetPlatform.MacTargetSettings
+TargetSettingsIniSectionName=/Script/MacTargetPlatform.MacTargetSettings
 bHasDedicatedGamepad=false
 bDefaultInputStandardKeyboard=false
 bInputSupportConfigurable=true
diff --git a/Engine/Config/TVOS/DataDrivenPlatformInfo.ini b/Engine/Config/TVOS/DataDrivenPlatformInfo.ini
--- a/Engine/Config/TVOS/DataDrivenPlatformInfo.ini
+++ b/Engine/Config/TVOS/DataDrivenPlatformInfo.ini
@@ -1,7 +1,7 @@
 [DataDrivenPlatformInfo]
 bIsConfidential=false
 IniParent=IOS
-AudioCompressionSettingsIniSectionName=/Script/IOSRuntimeSettings.IOSRuntimeSettings
+TargetSettingsIniSectionName=/Script/IOSRuntimeSettings.IOSRuntimeSettings
 GlobalIdentifier=A478294FBD0D4EC08830B6D4D219C1A4
 
 NormalIconPath=Launcher/TVOS/Platform_TVOS_24x
diff --git a/Engine/Config/Windows/DataDrivenPlatformInfo.ini b/Engine/Config/Windows/DataDrivenPlatformInfo.ini
--- a/Engine/Config/Windows/DataDrivenPlatformInfo.ini
+++ b/Engine/Config/Windows/DataDrivenPlatformInfo.ini
@@ -1,6 +1,6 @@
 [DataDrivenPlatformInfo]
 bIsConfidential=false
-AudioCompressionSettingsIniSectionName=/Script/WindowsTargetPlatform.WindowsTargetSettings
+TargetSettingsIniSectionName=/Script/WindowsTargetPlatform.WindowsTargetSettings
 Freezing_b32Bit=false
 Freezing_bWithRayTracing=true
 Freezing_MaxFieldAlignment32=4
diff --git a/Engine/Source/Developer/TargetPlatform/Private/TargetPlatformBase.cpp b/Engine/Source/Developer/TargetPlatform/Private/TargetPlatformBase.cpp
--- a/Engine/Source/Developer/TargetPlatform/Private/TargetPlatformBase.cpp
+++ b/Engine/Source/Developer/TargetPlatform/Private/TargetPlatformBase.cpp
@@ -44,7 +44,8 @@ bool FTargetPlatformBase::UsesSelectiveBasePassOutputs() const
 
 bool FTargetPlatformBase::UsesDistanceFields() const
 {
-	return true;
+	static IConsoleVariable* CVar = IConsoleManager::Get().FindConsoleVariable(TEXT("r.DistanceFields"));
+	return CVar ? (CVar->GetInt() != 0) : false;
 }
 
 bool FTargetPlatformBase::UsesRayTracing() const
diff --git a/Engine/Source/Developer/Windows/WindowsTargetPlatform/Private/GenericWindowsTargetPlatform.h b/Engine/Source/Developer/Windows/WindowsTargetPlatform/Private/GenericWindowsTargetPlatform.h
--- a/Engine/Source/Developer/Windows/WindowsTargetPlatform/Private/GenericWindowsTargetPlatform.h
+++ b/Engine/Source/Developer/Windows/WindowsTargetPlatform/Private/GenericWindowsTargetPlatform.h
@@ -464,6 +464,14 @@ public:
 		return TProperties::GetVariantPriority();
 	}
 
+	virtual bool UsesDistanceFields() const override
+	{
+		bool bEnableDistanceFields = false;
+		GConfig->GetBool(TEXT("/Script/WindowsTargetPlatform.WindowsTargetSettings"), TEXT("bEnableDistanceFields"), bEnableDistanceFields, GEngineIni);
+
+		return bEnableDistanceFields && TSuper::UsesDistanceFields();
+	}
+
 	virtual bool UsesRayTracing() const override
 	{
 		bool bEnableRayTracing = false;
diff --git a/Engine/Source/Runtime/Core/Private/Misc/DataDrivenPlatformInfoRegistry.cpp b/Engine/Source/Runtime/Core/Private/Misc/DataDrivenPlatformInfoRegistry.cpp
--- a/Engine/Source/Runtime/Core/Private/Misc/DataDrivenPlatformInfoRegistry.cpp
+++ b/Engine/Source/Runtime/Core/Private/Misc/DataDrivenPlatformInfoRegistry.cpp
@@ -271,7 +271,7 @@ static void LoadDDPIIniSettings(const FConfigFile& IniFile, FDataDrivenPlatformI
 {
 	DDPIGetBool(IniFile, TEXT("bIsConfidential"), Info.bIsConfidential);
 	DDPIGetBool(IniFile, TEXT("bIsFakePlatform"), Info.bIsFakePlatform);
-	DDPIGetString(IniFile, TEXT("AudioCompressionSettingsIniSectionName"), Info.AudioCompressionSettingsIniSectionName);
+	DDPIGetString(IniFile, TEXT("TargetSettingsIniSectionName"), Info.TargetSettingsIniSectionName);
 	DDPIGetString(IniFile, TEXT("HardwareCompressionFormat"), Info.HardwareCompressionFormat);
 	DDPIGetStringArray(IniFile, TEXT("AdditionalRestrictedFolders"), Info.AdditionalRestrictedFolders);
 
diff --git a/Engine/Source/Runtime/Core/Public/Misc/DataDrivenPlatformInfoRegistry.h b/Engine/Source/Runtime/Core/Public/Misc/DataDrivenPlatformInfoRegistry.h
--- a/Engine/Source/Runtime/Core/Public/Misc/DataDrivenPlatformInfoRegistry.h
+++ b/Engine/Source/Runtime/Core/Public/Misc/DataDrivenPlatformInfoRegistry.h
@@ -112,13 +112,13 @@ struct FDataDrivenPlatformInfo
 	// some platforms are here just for IniParentChain needs and are not concrete platforms
 	bool bIsFakePlatform = false;
 
-	// the name of the ini section to use to load audio compression settings (used at runtime and cooktime)
-	FString AudioCompressionSettingsIniSectionName;
+	// the name of the ini section to use to load target platform settings (used at runtime and cooktime)
+	FString TargetSettingsIniSectionName;
 
-	// list of additonal restricted folders
+	// list of additional restricted folders
 	TArray<FString> AdditionalRestrictedFolders;
 
-	// GUID to prepresent this platform forever
+	// GUID to represent this platform forever
 	FGuid GlobalIdentifier;
 
 	// MemoryFreezing information, matches FPlatformTypeLayoutParameters - defaults are clang, noneditor
diff --git a/Engine/Source/Runtime/Engine/Private/AudioCompressionSettingsUtils.cpp b/Engine/Source/Runtime/Engine/Private/AudioCompressionSettingsUtils.cpp
--- a/Engine/Source/Runtime/Engine/Private/AudioCompressionSettingsUtils.cpp
+++ b/Engine/Source/Runtime/Engine/Private/AudioCompressionSettingsUtils.cpp
@@ -87,7 +87,7 @@ void CacheAudioCookOverrides(FPlatformAudioCookOverrides& OutOverrides, const TC
 	
 	// now use that platform name to get the ini section out of DDPI
 	const FDataDrivenPlatformInfo& PlatformInfo = FDataDrivenPlatformInfoRegistry::GetPlatformInfo(PlatformName);
-	const FString& CategoryName = PlatformInfo.AudioCompressionSettingsIniSectionName;
+	const FString& CategoryName = PlatformInfo.TargetSettingsIniSectionName;
 
 	// if we don't support platform overrides, then return 
 	if (CategoryName.Len() == 0)
@@ -271,7 +271,7 @@ void CacheAudioCookOverrides(FPlatformAudioCookOverrides& OutOverrides, const TC
 
 static bool PlatformSupportsCompressionOverrides(const FString& PlatformName)
 {
-	return FDataDrivenPlatformInfoRegistry::GetPlatformInfo(PlatformName).AudioCompressionSettingsIniSectionName.Len() > 0;
+	return FDataDrivenPlatformInfoRegistry::GetPlatformInfo(PlatformName).TargetSettingsIniSectionName.Len() > 0;
 }
 
 static inline FString GetCookOverridePlatformName(const TCHAR* PlatformName)
diff --git a/Engine/Source/Runtime/Engine/Private/StaticMesh.cpp b/Engine/Source/Runtime/Engine/Private/StaticMesh.cpp
--- a/Engine/Source/Runtime/Engine/Private/StaticMesh.cpp
+++ b/Engine/Source/Runtime/Engine/Private/StaticMesh.cpp
@@ -1651,7 +1651,8 @@ void FStaticMeshRenderData::Serialize(FArchive& Ar, UStaticMesh* Owner, bool bCo
 		uint8 ClassDataStripFlags = 0;
 
 #if WITH_EDITOR
-		const bool bWantToStripDistanceFieldData = Ar.IsCooking() && !Ar.CookingTarget()->SupportsFeature(ETargetPlatformFeatures::DistanceFieldAO);
+		const bool bWantToStripDistanceFieldData = Ar.IsCooking() 
+			&& (!Ar.CookingTarget()->SupportsFeature(ETargetPlatformFeatures::DistanceFieldAO) || !Ar.CookingTarget()->UsesDistanceFields());
 
 		ClassDataStripFlags |= (bWantToStripDistanceFieldData ? DistanceFieldDataStripFlag : 0);
 #endif
diff --git a/Engine/Source/Runtime/RenderCore/Private/RenderUtils.cpp b/Engine/Source/Runtime/RenderCore/Private/RenderUtils.cpp
--- a/Engine/Source/Runtime/RenderCore/Private/RenderUtils.cpp
+++ b/Engine/Source/Runtime/RenderCore/Private/RenderUtils.cpp
@@ -9,6 +9,7 @@
 #include "PipelineStateCache.h"
 #include "Misc/ConfigCacheIni.h"
 #include "RenderGraphResourcePool.h"
+#include "Misc/DataDrivenPlatformInfoRegistry.h"
 
 #if WITH_EDITOR
 #include "Misc/CoreMisc.h"
@@ -1397,11 +1398,39 @@ RENDERCORE_API void RenderUtilsInit()
 		}
 	}
 #else
+
 	if (IsMobilePlatform(GMaxRHIShaderPlatform))
 	{
 		GDBufferPlatformMask = 0;
 		GBasePassVelocityPlatformMask = 0;
 	}
+
+	// Load runtime values from and *.ini file used by a current platform
+	// Should be code shared between cook and game, but unfortunately can't be done before we untangle non data driven platforms
+	const FString PlatformName(FPlatformProperties::IniPlatformName());
+	const FDataDrivenPlatformInfo& PlatformInfo = FDataDrivenPlatformInfoRegistry::GetPlatformInfo(PlatformName);
+
+	const FString CategoryName = PlatformInfo.TargetSettingsIniSectionName;
+	if (!CategoryName.IsEmpty())
+	{
+		FConfigFile PlatformIniFile;
+		if (FConfigCacheIni::LoadLocalIniFile(PlatformIniFile, TEXT("Engine"), /*bIsBaseIniName*/ true, *PlatformName))
+		{
+			bool bDistanceFields = false;
+			PlatformIniFile.GetBool(*CategoryName, TEXT("bEnableDistanceFields"), bDistanceFields);
+			if (!bDistanceFields)
+			{
+				GDistanceFieldsPlatformMask = 0;
+			}
+
+			bool bRayTracing = false;
+			PlatformIniFile.GetBool(*CategoryName, TEXT("bEnableRayTracing"), bRayTracing);
+			if (!bRayTracing)
+			{
+				GRayTracingPlaformMask = 0;
+			}
+		}
+	}
 #endif // WITH_EDITOR
 
 	// Run-time ray tracing support depends on the following factors:
