Build fix for RDG compiled without GPU events.

diff --git a/Engine/Source/Runtime/RenderCore/Public/RenderGraphEvent.inl b/Engine/Source/Runtime/RenderCore/Public/RenderGraphEvent.inl
--- a/Engine/Source/Runtime/RenderCore/Public/RenderGraphEvent.inl
+++ b/Engine/Source/Runtime/RenderCore/Public/RenderGraphEvent.inl
@@ -161,6 +161,8 @@ inline const TCHAR* FRDGEventName::GetTCHAR() const
 #endif
 }
 
+#if RDG_GPU_SCOPES
+
 inline const FRDGGPUScopeStacks& FRDGGPUScopeStacksByPipeline::GetScopeStacks(ERHIPipeline Pipeline) const
 {
 	switch (Pipeline)
@@ -187,4 +189,6 @@ inline FRDGGPUScopeStacks& FRDGGPUScopeStacksByPipeline::GetScopeStacks(ERHIPipe
 		checkNoEntry();
 		return Graphics;
 	}
-}
\ No newline at end of file
+}
+
+#endif
\ No newline at end of file
