Remove split transition for swap chain textures. This avoids cases where certain platforms will stall at the beginning of graph execution.

diff --git a/Engine/Source/Runtime/RenderCore/Private/RenderGraphBuilder.cpp b/Engine/Source/Runtime/RenderCore/Private/RenderGraphBuilder.cpp
--- a/Engine/Source/Runtime/RenderCore/Private/RenderGraphBuilder.cpp
+++ b/Engine/Source/Runtime/RenderCore/Private/RenderGraphBuilder.cpp
@@ -1041,6 +1041,13 @@ void FRDGBuilder::Compile(EExecuteMode ExecuteMode)
 				Texture->bUsedByAsyncComputePass |= bAsyncComputePass;
 				Texture->bCulled = false;
 
+				if (Texture->bSwapChain && !Texture->bSwapChainAlreadyMoved)
+				{
+					Texture->bSwapChainAlreadyMoved = 1;
+					Texture->FirstPass = PassHandle;
+					GetWholeResource(Texture->GetState()).SetPass(ERHIPipeline::Graphics, PassHandle);
+				}
+
 			#if STATS
 				GRDGStatTextureReferenceCount += PassState.ReferenceCount;
 			#endif
@@ -1467,7 +1474,7 @@ void FRDGBuilder::Execute()
 void FRDGBuilder::Execute(EExecuteMode ExecuteMode)
 {
 	CSV_SCOPED_TIMING_STAT_EXCLUSIVE(RDG);
-	SCOPED_NAMED_EVENT(FRDGBuilder_Execute, FColor::Emerald);
+	SCOPED_NAMED_EVENT_TCHAR("FRDGBuilder::Execute", FColor::Magenta);
 
 	// Create the epilogue pass at the end of the graph just prior to compilation.
 	SetupEmptyPass(EpiloguePass = Passes.Allocate<FRDGSentinelPass>(Allocator, RDG_EVENT_NAME("Graph Epilogue%s", (ExecuteMode == EExecuteMode::Drain) ? TEXT(" (Drain)") : TEXT(""))));
@@ -1498,7 +1505,7 @@ void FRDGBuilder::Execute(EExecuteMode ExecuteMode)
 		{
 			SCOPE_CYCLE_COUNTER(STAT_RDG_CollectResourcesTime);
 			CSV_SCOPED_TIMING_STAT_EXCLUSIVE(RDG_CollectResources);
-			SCOPED_NAMED_EVENT(CollectResources, FColor::Emerald);
+			SCOPED_NAMED_EVENT_TCHAR("FRDGBuilder::CollectResources", FColor::Magenta);
 
 			if (bFinalExecute)
 			{
@@ -1575,7 +1582,7 @@ void FRDGBuilder::Execute(EExecuteMode ExecuteMode)
 		CreateUniformBuffers();
 
 		{
-			SCOPED_NAMED_EVENT(CollectPassBarriers, FColor::Emerald);
+			SCOPED_NAMED_EVENT_TCHAR("FRDGBuilder::CollectBarriers", FColor::Magenta);
 			SCOPE_CYCLE_COUNTER(STAT_RDG_CollectBarriersTime);
 			CSV_SCOPED_TIMING_STAT_EXCLUSIVE_CONDITIONAL(RDG_CollectBarriers, GRDGVerboseCSVStats != 0);
 
@@ -1593,7 +1600,7 @@ void FRDGBuilder::Execute(EExecuteMode ExecuteMode)
 
 	if (bFinalExecute)
 	{
-		SCOPED_NAMED_EVENT(Finalize, FColor::Emerald);
+		SCOPED_NAMED_EVENT_TCHAR("FRDGBuilder::Finalize", FColor::Magenta);
 
 #if RDG_ENABLE_DEBUG
 		const auto LogResource = [&](auto* Resource, auto& Registry)
@@ -1646,6 +1653,7 @@ void FRDGBuilder::Execute(EExecuteMode ExecuteMode)
 
 	if (!IsImmediateMode())
 	{
+		SCOPED_NAMED_EVENT_TEXT("FRDGBuilder::ExecutePasses", FColor::Magenta);
 		SCOPE_CYCLE_COUNTER(STAT_RDG_ExecuteTime);
 		CSV_SCOPED_TIMING_STAT_EXCLUSIVE(RenderOther);
 
@@ -2160,7 +2168,7 @@ void FRDGBuilder::SetupPassInternal(FRDGPass* Pass, FRDGPassHandle PassHandle, E
 
 void FRDGBuilder::SubmitBufferUploads()
 {
-	SCOPED_NAMED_EVENT(SubmitBufferUploads, FColor::Emerald);
+	SCOPED_NAMED_EVENT_TCHAR("FRDGBuilder::SubmitBufferUploads", FColor::Magenta);
 
 	for (FUploadedBuffer& UploadedBuffer : UploadedBuffers)
 	{
@@ -2192,7 +2200,7 @@ void FRDGBuilder::SubmitBufferUploads()
 
 void FRDGBuilder::CreateUniformBuffers()
 {
-	SCOPED_NAMED_EVENT(CreateUniformBuffer, FColor::Emerald);
+	SCOPED_NAMED_EVENT_TCHAR("FRDGBuilder::CreateUniformBuffers", FColor::Magenta);
 
 	for (FRDGUniformBufferHandle UniformBufferHandle : UniformBuffersToCreate)
 	{
@@ -2317,6 +2325,10 @@ void FRDGBuilder::ExecutePassEpilogue(FRHIComputeCommandList& RHICmdListPass, FR
 
 void FRDGBuilder::ExecutePass(FRDGPass* Pass, FRHIComputeCommandList& RHICmdListPass)
 {
+#if 0 // This is compiled out for performance reasons; enable to display per-pass timings.
+	SCOPED_NAMED_EVENT_F(TEXT("%s"), FColor::Magenta, Pass->GetName());
+#endif
+
 	// Note that we must do this before doing anything with RHICmdList for the pass.
 	// For example, if this pass only executes on GPU 1 we want to avoid adding a
 	// 0-duration event for this pass on GPU 0's time line.
@@ -2444,7 +2456,7 @@ void FRDGBuilder::CollectPassBarriers(FRDGPass* Pass, FRDGPassHandle PassHandle)
 
 void FRDGBuilder::CreatePassBarriers()
 {
-	SCOPED_NAMED_EVENT(CreatePassBarriers, FColor::Emerald);
+	SCOPED_NAMED_EVENT_TCHAR("FRDGBuilder::CreatePassBarriers", FColor::Magenta);
 
 	for (FRDGBarrierBatchBegin* BarrierBatchBegin : TransitionCreateQueue)
 	{
diff --git a/Engine/Source/Runtime/RenderCore/Private/RenderGraphResources.cpp b/Engine/Source/Runtime/RenderCore/Private/RenderGraphResources.cpp
--- a/Engine/Source/Runtime/RenderCore/Private/RenderGraphResources.cpp
+++ b/Engine/Source/Runtime/RenderCore/Private/RenderGraphResources.cpp
@@ -33,6 +33,8 @@ FRDGParentResource::FRDGParentResource(const TCHAR* InName, const ERDGParentReso
 	, bCulled(1)
 	, bUsedByAsyncComputePass(0)
 	, bQueuedForUpload(0)
+	, bSwapChain(0)
+	, bSwapChainAlreadyMoved(0)
 {}
 
 bool FRDGProducerState::IsDependencyRequired(FRDGProducerState LastProducer, ERHIPipeline LastPipeline, FRDGProducerState NextState, ERHIPipeline NextPipeline)
diff --git a/Engine/Source/Runtime/RenderCore/Public/RenderGraphResources.h b/Engine/Source/Runtime/RenderCore/Public/RenderGraphResources.h
--- a/Engine/Source/Runtime/RenderCore/Public/RenderGraphResources.h
+++ b/Engine/Source/Runtime/RenderCore/Public/RenderGraphResources.h
@@ -296,6 +296,12 @@ protected:
 	/** If true, the resource has been queued for an upload operation. */
 	uint8 bQueuedForUpload : 1;
 
+	/** If true, this resource is a swap chain texture. */
+	uint8 bSwapChain : 1;
+
+	/** If true, the swap chain transition has been moved to occur at the latest possible time. */
+	uint8 bSwapChainAlreadyMoved : 1;
+
 	FRDGPassHandle FirstPass;
 	FRDGPassHandle LastPass;
 
@@ -464,6 +470,7 @@ private:
 		MergeState.SetNum(SubresourceCount);
 		LastProducers.Reserve(SubresourceCount);
 		LastProducers.SetNum(SubresourceCount);
+		bSwapChain = EnumHasAnyFlags(Desc.Flags, ETextureCreateFlags::Presentable);
 	}
 
 	/** Assigns a pooled render target as the backing RHI resource. */
