Fixed FRenderBounds::ComputeSquaredDistanceToPoint

diff --git a/Engine/Source/Runtime/RenderCore/Public/RenderTransform.h b/Engine/Source/Runtime/RenderCore/Public/RenderTransform.h
--- a/Engine/Source/Runtime/RenderCore/Public/RenderTransform.h
+++ b/Engine/Source/Runtime/RenderCore/Public/RenderTransform.h
@@ -334,7 +334,7 @@ public:
 	 */
 	FORCEINLINE float ComputeSquaredDistanceToPoint(const FVector3f& Point) const
 	{
-		FVector3f AxisDistances = (Point - GetCenter()) - GetExtent();
+		FVector3f AxisDistances = (Point - GetCenter()).GetAbs() - GetExtent();
 		AxisDistances = FVector3f::Max(AxisDistances, FVector3f(0.0f, 0.0f, 0.0f));
 		return AxisDistances | AxisDistances;
 	}
