Keep the default values of GDistanceFieldsPlatformMask and GRayTracingPlatformMask if their corresponding settings are not found in init files

diff --git a/Engine/Source/Runtime/RenderCore/Private/RenderUtils.cpp b/Engine/Source/Runtime/RenderCore/Private/RenderUtils.cpp
--- a/Engine/Source/Runtime/RenderCore/Private/RenderUtils.cpp
+++ b/Engine/Source/Runtime/RenderCore/Private/RenderUtils.cpp
@@ -1417,15 +1417,13 @@ RENDERCORE_API void RenderUtilsInit()
 		if (FConfigCacheIni::LoadLocalIniFile(PlatformIniFile, TEXT("Engine"), /*bIsBaseIniName*/ true, *PlatformName))
 		{
 			bool bDistanceFields = false;
-			PlatformIniFile.GetBool(*CategoryName, TEXT("bEnableDistanceFields"), bDistanceFields);
-			if (!bDistanceFields)
+			if (PlatformIniFile.GetBool(*CategoryName, TEXT("bEnableDistanceFields"), bDistanceFields) && !bDistanceFields)
 			{
 				GDistanceFieldsPlatformMask = 0;
 			}
 
 			bool bRayTracing = false;
-			PlatformIniFile.GetBool(*CategoryName, TEXT("bEnableRayTracing"), bRayTracing);
-			if (!bRayTracing)
+			if (PlatformIniFile.GetBool(*CategoryName, TEXT("bEnableRayTracing"), bRayTracing) && !bRayTracing)
 			{
 				GRayTracingPlaformMask = 0;
 			}
