Skip FlushRenderingCommand in slowtask tick when there is nothing to cleanup

#rnx

diff --git a/Engine/Source/Editor/UnrealEd/Private/FeedbackContextEditor.cpp b/Engine/Source/Editor/UnrealEd/Private/FeedbackContextEditor.cpp
--- a/Engine/Source/Editor/UnrealEd/Private/FeedbackContextEditor.cpp
+++ b/Engine/Source/Editor/UnrealEd/Private/FeedbackContextEditor.cpp
@@ -523,8 +523,11 @@ void FFeedbackContextEditor::ProgressReported( const float TotalProgressInterp,
 	{
 		// Get list of objects that are pending cleanup.
 		FPendingCleanupObjects* PendingCleanupObjects = GetPendingCleanupObjects();
-		// Flush rendering commands in the queue.
-		FlushRenderingCommands();
+		if (!PendingCleanupObjects->IsEmpty())
+		{
+			// Flush rendering commands in the queue.
+			FlushRenderingCommands();
+		}
 		// It is now safe to delete the pending clean objects.
 		delete PendingCleanupObjects;
 		// Keep track of time this operation was performed so we don't do it too often.
diff --git a/Engine/Source/Runtime/RenderCore/Private/RenderingThread.cpp b/Engine/Source/Runtime/RenderCore/Private/RenderingThread.cpp
--- a/Engine/Source/Runtime/RenderCore/Private/RenderingThread.cpp
+++ b/Engine/Source/Runtime/RenderCore/Private/RenderingThread.cpp
@@ -1359,20 +1359,23 @@ FPendingCleanupObjects::FPendingCleanupObjects()
 
 FPendingCleanupObjects::~FPendingCleanupObjects()
 {
-	QUICK_SCOPE_CYCLE_COUNTER(STAT_FPendingCleanupObjects_Destruct);
-
-	const bool bBatchingEnabled = bEnablePendingCleanupObjectsCommandBatching;
-	if (bBatchingEnabled)
-	{
-		StartRenderCommandFenceBundler();
-	}
-	for (int32 ObjectIndex = 0; ObjectIndex < CleanupArray.Num(); ObjectIndex++)
-	{
-		delete CleanupArray[ObjectIndex];
-	}
-	if (bBatchingEnabled)
+	if (CleanupArray.Num())
 	{
-		StopRenderCommandFenceBundler();
+		QUICK_SCOPE_CYCLE_COUNTER(STAT_FPendingCleanupObjects_Destruct);
+
+		const bool bBatchingEnabled = bEnablePendingCleanupObjectsCommandBatching;
+		if (bBatchingEnabled)
+		{
+			StartRenderCommandFenceBundler();
+		}
+		for (int32 ObjectIndex = 0; ObjectIndex < CleanupArray.Num(); ObjectIndex++)
+		{
+			delete CleanupArray[ObjectIndex];
+		}
+		if (bBatchingEnabled)
+		{
+			StopRenderCommandFenceBundler();
+		}
 	}
 }
 
diff --git a/Engine/Source/Runtime/RenderCore/Public/RenderingThread.h b/Engine/Source/Runtime/RenderCore/Public/RenderingThread.h
--- a/Engine/Source/Runtime/RenderCore/Public/RenderingThread.h
+++ b/Engine/Source/Runtime/RenderCore/Public/RenderingThread.h
@@ -298,6 +298,7 @@ class FPendingCleanupObjects
 {
 	TArray<FDeferredCleanupInterface*> CleanupArray;
 public:
+	inline bool IsEmpty() const { return CleanupArray.IsEmpty(); }
 	FPendingCleanupObjects();
 	RENDERCORE_API ~FPendingCleanupObjects();
 };
