Minor clean-up of RDG event name. Simplified the 'compiled out' name and fixed misspelling of a variable name.

diff --git a/Engine/Source/Runtime/RenderCore/Public/RenderGraphEvent.h b/Engine/Source/Runtime/RenderCore/Public/RenderGraphEvent.h
--- a/Engine/Source/Runtime/RenderCore/Public/RenderGraphEvent.h
+++ b/Engine/Source/Runtime/RenderCore/Public/RenderGraphEvent.h
@@ -273,8 +273,6 @@ public:
 
 	explicit FRDGEventName(const TCHAR* EventFormat, ...);
 
-	~FRDGEventName();
-
 	FRDGEventName(const FRDGEventName& Other);
 	FRDGEventName(FRDGEventName&& Other);
 	FRDGEventName& operator=(const FRDGEventName& Other);
@@ -288,8 +286,7 @@ private:
 	const TCHAR* EventFormat = TEXT("");
 
 #if RDG_EVENTS == RDG_EVENTS_STRING_COPY
-	// Formated event name if GetEmitRDGEvents() == true.
-	FString FormatedEventName;
+	FString FormattedEventName;
 #endif
 #endif
 };
diff --git a/Engine/Source/Runtime/RenderCore/Public/RenderGraphEvent.inl b/Engine/Source/Runtime/RenderCore/Public/RenderGraphEvent.inl
--- a/Engine/Source/Runtime/RenderCore/Public/RenderGraphEvent.inl
+++ b/Engine/Source/Runtime/RenderCore/Public/RenderGraphEvent.inl
@@ -93,19 +93,14 @@ TRDGScopeOpArray<ScopeOpType> TRDGScopeStackHelper<ScopeOpType>::EndCompile()
 	return TRDGScopeOpArray<ScopeOpType>(Ops, OffsetIndex, Ops.Num() - OffsetIndex);
 }
 
-inline FRDGEventName::~FRDGEventName()
-{
-#if RDG_EVENTS == RDG_EVENTS_STRING_REF || RDG_EVENTS == RDG_EVENTS_STRING_COPY
-	EventFormat = nullptr;
-#endif
-}
-
 #if RDG_EVENTS != RDG_EVENTS_STRING_COPY
 inline FRDGEventName::FRDGEventName(const TCHAR* InEventFormat, ...)
 #if RDG_EVENTS == RDG_EVENTS_STRING_REF
 	: EventFormat(InEventFormat)
 #endif
-{}
+{
+	check(InEventFormat != nullptr);
+}
 #endif
 
 inline FRDGEventName::FRDGEventName(const FRDGEventName& Other)
@@ -124,7 +119,7 @@ inline FRDGEventName& FRDGEventName::operator=(const FRDGEventName& Other)
 	EventFormat = Other.EventFormat;
 #elif RDG_EVENTS == RDG_EVENTS_STRING_COPY
 	EventFormat = Other.EventFormat;
-	FormatedEventName = Other.FormatedEventName;
+	FormattedEventName = Other.FormattedEventName;
 #endif
 	return *this;
 }
@@ -133,11 +128,11 @@ inline FRDGEventName& FRDGEventName::operator=(FRDGEventName&& Other)
 {
 #if RDG_EVENTS == RDG_EVENTS_STRING_REF
 	EventFormat = Other.EventFormat;
-	Other.EventFormat = nullptr;
+	Other.EventFormat = TEXT("");
 #elif RDG_EVENTS == RDG_EVENTS_STRING_COPY
 	EventFormat = Other.EventFormat;
-	Other.EventFormat = nullptr;
-	FormatedEventName = MoveTemp(Other.FormatedEventName);
+	Other.EventFormat = TEXT("");
+	FormattedEventName = MoveTemp(Other.FormattedEventName);
 #endif
 	return *this;
 }
@@ -146,9 +141,9 @@ inline const TCHAR* FRDGEventName::GetTCHAR() const
 {
 #if RDG_EVENTS == RDG_EVENTS_STRING_REF || RDG_EVENTS == RDG_EVENTS_STRING_COPY
 	#if RDG_EVENTS == RDG_EVENTS_STRING_COPY
-		if (!FormatedEventName.IsEmpty())
+		if (!FormattedEventName.IsEmpty())
 		{
-			return *FormatedEventName;
+			return *FormattedEventName;
 		}
 	#endif
 
@@ -157,7 +152,7 @@ inline const TCHAR* FRDGEventName::GetTCHAR() const
 	return EventFormat;
 #else
 	// Render graph draw events have been completely compiled for CPU performance reasons.
-	return TEXT("!!!Unavailable RDG event name: need RDG_EVENTS>=0 and r.RDG.EmitWarnings=1 or -rdgdebug!!!");
+	return TEXT("[Compiled Out]");
 #endif
 }
 
