Missing local file.

diff --git a/Engine/Source/Runtime/RenderCore/Private/RenderGraphEvent.cpp b/Engine/Source/Runtime/RenderCore/Private/RenderGraphEvent.cpp
--- a/Engine/Source/Runtime/RenderCore/Private/RenderGraphEvent.cpp
+++ b/Engine/Source/Runtime/RenderCore/Private/RenderGraphEvent.cpp
@@ -33,7 +33,7 @@ FRDGEventName::FRDGEventName(const TCHAR* InEventFormat, ...)
 		FCString::GetVarArgs(TempStr, UE_ARRAY_COUNT(TempStr), InEventFormat, VAList);
 		va_end(VAList);
 
-		FormatedEventName = TempStr;
+		FormattedEventName = TempStr;
 	}
 }
 
