Properly support r.ShaderPipelines when filtering out non-pipelined shader types.

- Also, include this in the DDC key.

diff --git a/Engine/Source/Runtime/Engine/Private/Materials/MaterialShader.cpp b/Engine/Source/Runtime/Engine/Private/Materials/MaterialShader.cpp
--- a/Engine/Source/Runtime/Engine/Private/Materials/MaterialShader.cpp
+++ b/Engine/Source/Runtime/Engine/Private/Materials/MaterialShader.cpp
@@ -105,7 +105,7 @@ public:
 
 	FPipelinedShaderFilter(EShaderPlatform ShaderPlatform, TArray<FShaderPipelineType*> Pipelines)
 	{
-		if (GMaterialExcludeNonPipelinedShaders && RHISupportsShaderPipelines(ShaderPlatform))
+		if (ExcludeNonPipelinedShaderTypes(ShaderPlatform))
 		{
 			for (FShaderPipelineType* Pipeline : Pipelines)
 			{
diff --git a/Engine/Source/Runtime/RenderCore/Private/RenderUtils.cpp b/Engine/Source/Runtime/RenderCore/Private/RenderUtils.cpp
--- a/Engine/Source/Runtime/RenderCore/Private/RenderUtils.cpp
+++ b/Engine/Source/Runtime/RenderCore/Private/RenderUtils.cpp
@@ -1626,6 +1626,24 @@ RENDERCORE_API bool UseVirtualTextureLightmap(const FStaticFeatureLevel InFeatur
 	return bUseVirtualTextureLightmap;
 }
 
+RENDERCORE_API bool ExcludeNonPipelinedShaderTypes(EShaderPlatform ShaderPlatform)
+{
+	if (RHISupportsShaderPipelines(ShaderPlatform))
+	{
+		static const TConsoleVariableData<int32>* CVarShaderPipelines = IConsoleManager::Get().FindTConsoleVariableDataInt(TEXT("r.ShaderPipelines"));
+		bool bShaderPipelinesAreEnabled = CVarShaderPipelines && CVarShaderPipelines->GetValueOnAnyThread(IsInGameThread()) != 0;
+		if (bShaderPipelinesAreEnabled)
+		{
+			static const TConsoleVariableData<int32>* CVarExcludeNonPipelinedShaders = IConsoleManager::Get().FindTConsoleVariableDataInt(TEXT("r.Material.ExcludeNonPipelinedShaders"));
+			bool bExcludeNonPipelinedShaders = CVarExcludeNonPipelinedShaders && CVarExcludeNonPipelinedShaders->GetValueOnAnyThread(IsInGameThread()) != 0;
+
+			return bExcludeNonPipelinedShaders;
+		}
+	}
+
+	return false;
+}
+
 RENDERCORE_API bool PlatformSupportsVelocityRendering(const FStaticShaderPlatform Platform)
 {
 	if (IsMobilePlatform(Platform))
diff --git a/Engine/Source/Runtime/RenderCore/Private/Shader.cpp b/Engine/Source/Runtime/RenderCore/Private/Shader.cpp
--- a/Engine/Source/Runtime/RenderCore/Private/Shader.cpp
+++ b/Engine/Source/Runtime/RenderCore/Private/Shader.cpp
@@ -1855,4 +1855,9 @@ void ShaderMapAppendKeyString(EShaderPlatform Platform, FString& KeyString)
 		KeyString += TEXT("_Compr");
 		KeyString += GetShaderCompressionFormat().ToString();
 	}
+
+	{
+		// add whether or not non-pipelined shader types are included
+		KeyString += FString::Printf(TEXT("_ExclNonPipSh-%d"), ExcludeNonPipelinedShaderTypes(Platform));
+	}
 }
diff --git a/Engine/Source/Runtime/RenderCore/Public/RenderUtils.h b/Engine/Source/Runtime/RenderCore/Public/RenderUtils.h
--- a/Engine/Source/Runtime/RenderCore/Public/RenderUtils.h
+++ b/Engine/Source/Runtime/RenderCore/Public/RenderUtils.h
@@ -691,4 +691,9 @@ inline bool UseNonNaniteVirtualShadowMaps(EShaderPlatform ShaderPlatform, const
 /**
 *	Checks if virtual texturing lightmap enabled and supported
 */
-RENDERCORE_API bool UseVirtualTextureLightmap(const FStaticFeatureLevel InFeatureLevel, const class ITargetPlatform* TargetPlatform = nullptr);
\ No newline at end of file
+RENDERCORE_API bool UseVirtualTextureLightmap(const FStaticFeatureLevel InFeatureLevel, const class ITargetPlatform* TargetPlatform = nullptr);
+
+/**
+ *  Checks if the non-pipeline shaders will not be compild and ones from FShaderPipeline used instead.
+ */
+RENDERCORE_API bool ExcludeNonPipelinedShaderTypes(EShaderPlatform ShaderPlatform);
