Fix cvar access

diff --git a/Engine/Source/Runtime/RenderCore/Private/RenderUtils.cpp b/Engine/Source/Runtime/RenderCore/Private/RenderUtils.cpp
--- a/Engine/Source/Runtime/RenderCore/Private/RenderUtils.cpp
+++ b/Engine/Source/Runtime/RenderCore/Private/RenderUtils.cpp
@@ -1634,8 +1634,8 @@ RENDERCORE_API bool ExcludeNonPipelinedShaderTypes(EShaderPlatform ShaderPlatfor
 		bool bShaderPipelinesAreEnabled = CVarShaderPipelines && CVarShaderPipelines->GetValueOnAnyThread(IsInGameThread()) != 0;
 		if (bShaderPipelinesAreEnabled)
 		{
-			static const TConsoleVariableData<int32>* CVarExcludeNonPipelinedShaders = IConsoleManager::Get().FindTConsoleVariableDataInt(TEXT("r.Material.ExcludeNonPipelinedShaders"));
-			bool bExcludeNonPipelinedShaders = CVarExcludeNonPipelinedShaders && CVarExcludeNonPipelinedShaders->GetValueOnAnyThread(IsInGameThread()) != 0;
+			static const IConsoleVariable* CVarExcludeNonPipelinedShaders = IConsoleManager::Get().FindConsoleVariable(TEXT("r.Material.ExcludeNonPipelinedShaders"));
+			bool bExcludeNonPipelinedShaders = CVarExcludeNonPipelinedShaders && CVarExcludeNonPipelinedShaders->GetInt() != 0;
 
 			return bExcludeNonPipelinedShaders;
 		}
