Rendering: Temporary fix for draw events asseting if called while the rendering thread is stopped (e.g., during a map rebuild)

diff --git a/Engine/Source/Runtime/RenderCore/Private/ProfilingDebugging/RealtimeGPUProfiler.cpp b/Engine/Source/Runtime/RenderCore/Private/ProfilingDebugging/RealtimeGPUProfiler.cpp
--- a/Engine/Source/Runtime/RenderCore/Private/ProfilingDebugging/RealtimeGPUProfiler.cpp
+++ b/Engine/Source/Runtime/RenderCore/Private/ProfilingDebugging/RealtimeGPUProfiler.cpp
@@ -87,7 +87,9 @@ void FDrawEvent::Start(FRHIComputeCommandList* InRHICmdList, FColor Color, const
 		}
 		else
 		{
-			check(!GUseThreadedRendering || !bIsRenderingOrRHIThread); // The command list should be null on game thread (see comment above)
+			PRAGMA_DISABLE_DEPRECATION_WARNINGS
+			check((GRenderThreadId == 0) || !GUseThreadedRendering || !bIsRenderingOrRHIThread); // The command list should be null on game thread (see comment above)
+			PRAGMA_DISABLE_DEPRECATION_WARNINGS
 			ENQUEUE_RENDER_COMMAND(PushEventCommand)([EventName = FString(TempStr), Color](FRHICommandListImmediate& RHICommandListLocal)
 			{
 				RHICommandListLocal.PushEvent(*EventName, Color);
@@ -106,14 +108,16 @@ void FDrawEvent::Stop()
 		bool bIsRenderingOrRHIThread = IsInParallelRenderingThread() || IsInRHIThread();
 		// if we have a command list, we must be on the rendering or RHI thread, otherwise (game thread), a command will be enqueued on the immediate command list :
 		if (RHICmdList != nullptr)
-	{
+		{
 			check(!GUseThreadedRendering || bIsRenderingOrRHIThread); // A command list is needed on rendering/RHI threads (see comment above)
-		RHICmdList->PopEvent();
+			RHICmdList->PopEvent();
 			RHICmdList = nullptr;
 		}
 		else
 		{
-			check(!GUseThreadedRendering || !bIsRenderingOrRHIThread); // The command list should be null on game thread (see comment above)
+			PRAGMA_DISABLE_DEPRECATION_WARNINGS
+			check((GRenderThreadId == 0) || !GUseThreadedRendering || !bIsRenderingOrRHIThread); // The command list should be null on game thread (see comment above)
+			PRAGMA_DISABLE_DEPRECATION_WARNINGS
 			ENQUEUE_RENDER_COMMAND(PopEventCommand)([](FRHICommandListImmediate& RHICommandListLocal)
 			{
 				RHICommandListLocal.PopEvent();
