Rendering: Fixed mismatched PRAGMA_DISABLE_DEPRECATION_WARNINGS macros

diff --git a/Engine/Source/Runtime/RenderCore/Private/ProfilingDebugging/RealtimeGPUProfiler.cpp b/Engine/Source/Runtime/RenderCore/Private/ProfilingDebugging/RealtimeGPUProfiler.cpp
--- a/Engine/Source/Runtime/RenderCore/Private/ProfilingDebugging/RealtimeGPUProfiler.cpp
+++ b/Engine/Source/Runtime/RenderCore/Private/ProfilingDebugging/RealtimeGPUProfiler.cpp
@@ -89,7 +89,7 @@ void FDrawEvent::Start(FRHIComputeCommandList* InRHICmdList, FColor Color, const
 		{
 			PRAGMA_DISABLE_DEPRECATION_WARNINGS
 			check((GRenderThreadId == 0) || !GUseThreadedRendering || !bIsRenderingOrRHIThread); // The command list should be null on game thread (see comment above)
-			PRAGMA_DISABLE_DEPRECATION_WARNINGS
+			PRAGMA_ENABLE_DEPRECATION_WARNINGS
 			ENQUEUE_RENDER_COMMAND(PushEventCommand)([EventName = FString(TempStr), Color](FRHICommandListImmediate& RHICommandListLocal)
 			{
 				RHICommandListLocal.PushEvent(*EventName, Color);
@@ -117,7 +117,7 @@ void FDrawEvent::Stop()
 		{
 			PRAGMA_DISABLE_DEPRECATION_WARNINGS
 			check((GRenderThreadId == 0) || !GUseThreadedRendering || !bIsRenderingOrRHIThread); // The command list should be null on game thread (see comment above)
-			PRAGMA_DISABLE_DEPRECATION_WARNINGS
+			PRAGMA_ENABLE_DEPRECATION_WARNINGS
 			ENQUEUE_RENDER_COMMAND(PopEventCommand)([](FRHICommandListImmediate& RHICommandListLocal)
 			{
 				RHICommandListLocal.PopEvent();
