Added support for separate memstack for tasks in RDG.

diff --git a/Engine/Source/Runtime/RenderCore/Private/RenderGraphAllocator.cpp b/Engine/Source/Runtime/RenderCore/Private/RenderGraphAllocator.cpp
--- a/Engine/Source/Runtime/RenderCore/Private/RenderGraphAllocator.cpp
+++ b/Engine/Source/Runtime/RenderCore/Private/RenderGraphAllocator.cpp
@@ -11,17 +11,23 @@ FRDGAllocator& FRDGAllocator::Get()
 
 FRDGAllocator::~FRDGAllocator()
 {
-	check(MemStack.IsEmpty());
+	check(Context.MemStack.IsEmpty() && ContextForTasks.MemStack.IsEmpty());
 }
 
-void FRDGAllocator::ReleaseAll()
+void FRDGAllocator::FContext::ReleaseAll()
 {
-	TRACE_CPUPROFILER_EVENT_SCOPE(FRDGAllocator::ReleaseAll);
-	CSV_SCOPED_TIMING_STAT_EXCLUSIVE_CONDITIONAL(RDGAllocator_Clear, GRDGVerboseCSVStats != 0);
 	for (int32 Index = TrackedAllocs.Num() - 1; Index >= 0; --Index)
 	{
 		TrackedAllocs[Index]->~FTrackedAlloc();
 	}
 	TrackedAllocs.Reset();
 	MemStack.Flush();
+}
+
+void FRDGAllocator::ReleaseAll()
+{
+	TRACE_CPUPROFILER_EVENT_SCOPE(FRDGAllocator::ReleaseAll);
+	CSV_SCOPED_TIMING_STAT_EXCLUSIVE_CONDITIONAL(RDGAllocator_Clear, GRDGVerboseCSVStats != 0);
+	Context.ReleaseAll();
+	ContextForTasks.ReleaseAll();
 }
\ No newline at end of file
diff --git a/Engine/Source/Runtime/RenderCore/Public/RenderGraphAllocator.h b/Engine/Source/Runtime/RenderCore/Public/RenderGraphAllocator.h
--- a/Engine/Source/Runtime/RenderCore/Public/RenderGraphAllocator.h
+++ b/Engine/Source/Runtime/RenderCore/Public/RenderGraphAllocator.h
@@ -14,7 +14,7 @@ public:
 	/** Allocates raw memory. */
 	FORCEINLINE void* Alloc(uint32 SizeInBytes, uint32 AlignInBytes)
 	{
-		return MemStack.Alloc(SizeInBytes, AlignInBytes);
+		return GetContext().MemStack.Alloc(SizeInBytes, AlignInBytes);
 	}
 
 	/** Allocates an uninitialized type without destructor tracking. */
@@ -28,9 +28,10 @@ public:
 	template <typename T, typename... TArgs>
 	FORCEINLINE T* Alloc(TArgs&&... Args)
 	{
-		TTrackedAlloc<T>* TrackedAlloc = new(MemStack) TTrackedAlloc<T>(Forward<TArgs&&>(Args)...);
+		FContext& LocalContext = GetContext();
+		TTrackedAlloc<T>* TrackedAlloc = new(LocalContext.MemStack) TTrackedAlloc<T>(Forward<TArgs&&>(Args)...);
 		check(TrackedAlloc);
-		TrackedAllocs.Add(TrackedAlloc);
+		LocalContext.TrackedAllocs.Add(TrackedAlloc);
 		return &TrackedAlloc->Alloc;
 	}
 
@@ -38,12 +39,12 @@ public:
 	template <typename T, typename... TArgs>
 	FORCEINLINE T* AllocNoDestruct(TArgs&&... Args)
 	{
-		return new (MemStack) T(Forward<TArgs&&>(Args)...);
+		return new (GetContext().MemStack) T(Forward<TArgs&&>(Args)...);
 	}
 
 	FORCEINLINE int32 GetByteCount() const
 	{
-		return MemStack.GetByteCount();
+		return Context.MemStack.GetByteCount() + ContextForTasks.MemStack.GetByteCount();
 	}
 
 private:
@@ -67,8 +68,21 @@ private:
 		T Alloc;
 	};
 
-	FMemStackBase MemStack;
-	TArray<FTrackedAlloc*> TrackedAllocs;
+	struct FContext
+	{
+		FMemStackBase MemStack;
+		TArray<FTrackedAlloc*> TrackedAllocs;
+
+		void ReleaseAll();
+	};
+
+	FContext Context;
+	FContext ContextForTasks;
+
+	FORCEINLINE FContext& GetContext()
+	{
+		return FTaskTagScope::IsCurrentTag(ETaskTag::ERenderingThread) ? Context : ContextForTasks;
+	}
 
 	template <uint32>
 	friend class TRDGArrayAllocator;
