Modify RHI query checks so they can be done in a parallel rendering task.

diff --git a/Engine/Source/Runtime/RHI/Private/DynamicRHI.cpp b/Engine/Source/Runtime/RHI/Private/DynamicRHI.cpp
--- a/Engine/Source/Runtime/RHI/Private/DynamicRHI.cpp
+++ b/Engine/Source/Runtime/RHI/Private/DynamicRHI.cpp
@@ -506,7 +506,7 @@ FDefaultRHIRenderQueryPool::~FDefaultRHIRenderQueryPool()
 
 FRHIPooledRenderQuery FDefaultRHIRenderQueryPool::AllocateQuery()
 {
-	check(IsInRenderingThread());
+	check(IsInParallelRenderingThread());
 	if (Queries.Num() > 0)
 	{
 		return FRHIPooledRenderQuery(this, Queries.Pop());
@@ -530,7 +530,7 @@ void FDefaultRHIRenderQueryPool::ReleaseQuery(TRefCountPtr<FRHIRenderQuery>&& Qu
 		static int dbg = 0;
 		dbg++;
 	}
-	check(IsInRenderingThread());
+	check(IsInParallelRenderingThread());
 	//Hard to validate because of Resource resurrection, better to remove GetQueryRef entirely
 	//checkf(Query.IsValid() && Query.GetRefCount() <= 2, TEXT("Query has been released but reference still held: use FRHIPooledRenderQuery::GetQueryRef() with extreme caution"));
 	
diff --git a/Engine/Source/Runtime/RHI/Public/RHIResources.h b/Engine/Source/Runtime/RHI/Public/RHIResources.h
--- a/Engine/Source/Runtime/RHI/Public/RHIResources.h
+++ b/Engine/Source/Runtime/RHI/Public/RHIResources.h
@@ -1642,7 +1642,7 @@ inline FRHIPooledRenderQuery::FRHIPooledRenderQuery(FRHIRenderQueryPool* InQuery
 	: Query(MoveTemp(InQuery))
 	, QueryPool(InQueryPool)
 {
-	check(IsInRenderingThread());
+	check(IsInParallelRenderingThread());
 }
 
 inline void FRHIPooledRenderQuery::ReleaseQuery()
diff --git a/Engine/Source/Runtime/RenderCore/Private/ProfilingDebugging/RealtimeGPUProfiler.cpp b/Engine/Source/Runtime/RenderCore/Private/ProfilingDebugging/RealtimeGPUProfiler.cpp
--- a/Engine/Source/Runtime/RenderCore/Private/ProfilingDebugging/RealtimeGPUProfiler.cpp
+++ b/Engine/Source/Runtime/RenderCore/Private/ProfilingDebugging/RealtimeGPUProfiler.cpp
@@ -174,7 +174,7 @@ public:
 
 	FRealtimeGPUProfilerQuery Begin(FRHIGPUMask InGPUMask, const FName& NewName, const FName& NewStatName)
 	{
-		check(IsInRenderingThread());
+		check(IsInParallelRenderingThread());
 		check(!bInsideQuery && StartQuery.IsValid());
 #if DO_CHECK || USING_CODE_ANALYSIS
 		bInsideQuery = true;
@@ -192,7 +192,7 @@ public:
 
 	FRealtimeGPUProfilerQuery End()
 	{
-		check(IsInRenderingThread());
+		check(IsInParallelRenderingThread());
 		check(bInsideQuery && EndQuery.IsValid());
 #if DO_CHECK || USING_CODE_ANALYSIS
 		bInsideQuery = false;
@@ -934,7 +934,7 @@ void FRealtimeGPUProfiler::EndFrame(FRHICommandListImmediate& RHICmdList)
 
 FRealtimeGPUProfilerQuery FRealtimeGPUProfiler::PushEvent(FRHIGPUMask GPUMask, const FName& Name, const FName& StatName)
 {
-	check(IsInRenderingThread());
+	check(IsInParallelRenderingThread());
 	if (bStatGatheringPaused || !bInBeginEndBlock)
 	{
 		return {};
@@ -949,7 +949,7 @@ FRealtimeGPUProfilerQuery FRealtimeGPUProfiler::PushEvent(FRHIGPUMask GPUMask, c
 
 FRealtimeGPUProfilerQuery FRealtimeGPUProfiler::PopEvent()
 {
-	check(IsInRenderingThread());
+	check(IsInParallelRenderingThread());
 	if (bStatGatheringPaused || !bInBeginEndBlock)
 	{
 		return {};
