Fix for potential deadlock in shader preloads

FPreloadShaderTask takes the ShaderPreloadLock in FShaderCodeArchive::OnShaderPreloadFinished() and then issues an async read.  Another task can be picked up by the thread while waiting on that read.  If that new task is also a FPreloadShaderTask it will try to aquire the ShaderPreloadLock again and deadlock.

Introduce an EnsureReadNonBlocking method to IMemoryReadStream which can be called outside the lock to ensure the data is ready.  We can then lock and copy out the data without incuring a wait and the potential for another FPreloadShaderTask being run while we hold the lock.

diff --git a/Engine/Source/Runtime/Core/Private/FileCache/FileCache.cpp b/Engine/Source/Runtime/Core/Private/FileCache/FileCache.cpp
--- a/Engine/Source/Runtime/Core/Private/FileCache/FileCache.cpp
+++ b/Engine/Source/Runtime/Core/Private/FileCache/FileCache.cpp
@@ -478,6 +478,15 @@ public:
 		return ResultData + InOffset;
 	}
 
+	virtual void EnsureReadNonBlocking() override
+	{
+		if (Request)
+		{
+			check(!Memory);
+			Request->WaitCompletion();
+		}
+	}
+
 	virtual int64 GetSize() override
 	{
 		return Size;
diff --git a/Engine/Source/Runtime/Core/Public/Misc/MemoryReadStream.h b/Engine/Source/Runtime/Core/Public/Misc/MemoryReadStream.h
--- a/Engine/Source/Runtime/Core/Public/Misc/MemoryReadStream.h
+++ b/Engine/Source/Runtime/Core/Public/Misc/MemoryReadStream.h
@@ -18,6 +18,7 @@ public:
 	virtual const void* Read(int64& OutSize, int64 InOffset, int64 InSize) = 0;
 	virtual int64 GetSize() = 0;
 	virtual void CopyTo(void* Buffer, int64 InOffset, int64 InSize);
+	virtual void EnsureReadNonBlocking() {};		// Incur any potential wait such that future Read() or CopyTo() calls are wait and thread switch free.
 
 	FORCEINLINE uint32 AddRef() const { return uint32(NumRefs.Increment()); }
 
diff --git a/Engine/Source/Runtime/RenderCore/Private/ShaderCodeArchive.cpp b/Engine/Source/Runtime/RenderCore/Private/ShaderCodeArchive.cpp
--- a/Engine/Source/Runtime/RenderCore/Private/ShaderCodeArchive.cpp
+++ b/Engine/Source/Runtime/RenderCore/Private/ShaderCodeArchive.cpp
@@ -771,11 +771,14 @@ void FShaderCodeArchive::Teardown()
 
 void FShaderCodeArchive::OnShaderPreloadFinished(int32 ShaderIndex, const IMemoryReadStreamRef& PreloadData)
 {
-	FWriteScopeLock Lock(ShaderPreloadLock);
 	const FShaderCodeEntry& ShaderEntry = SerializedShaders.ShaderEntries[ShaderIndex];
-	FShaderPreloadEntry& ShaderPreloadEntry = ShaderPreloads[ShaderIndex];
-	PreloadData->CopyTo(ShaderPreloadEntry.Code, 0, ShaderEntry.Size);
-	ShaderPreloadEntry.PreloadEvent.SafeRelease();
+	PreloadData->EnsureReadNonBlocking();		// Ensure data is ready before taking the lock
+	{
+		FWriteScopeLock Lock(ShaderPreloadLock);
+		FShaderPreloadEntry& ShaderPreloadEntry = ShaderPreloads[ShaderIndex];
+		PreloadData->CopyTo(ShaderPreloadEntry.Code, 0, ShaderEntry.Size);
+		ShaderPreloadEntry.PreloadEvent.SafeRelease();
+	}
 }
 
 struct FPreloadShaderTask
