Parallel RDG execution improvements.
 - Added ERDGBuilderFlags::AllowParallelExecute to tag specific builders to attempt parallel execution. This avoids cases where small graphs fork tasks and end up causing contention. Only the main scene render graphs are tagged.
 - Moved RHI transition creation to an async task.
 - Moved parallel execute setup and dispatch to an async task.
 - Fixed RDG draining asserts using a short-term workaround by tagging relevant scene textures as non-transient.
 - Deprecated RDG AddPass utilities without names and fixed up last remnants.
 - Enabled parallel RDG execution by default.

diff --git a/Engine/Source/Runtime/RenderCore/Private/RenderGraphBuilder.cpp b/Engine/Source/Runtime/RenderCore/Private/RenderGraphBuilder.cpp
--- a/Engine/Source/Runtime/RenderCore/Private/RenderGraphBuilder.cpp
+++ b/Engine/Source/Runtime/RenderCore/Private/RenderGraphBuilder.cpp
@@ -424,6 +424,11 @@ bool FRDGBuilder::IsImmediateMode()
 	return ::IsImmediateMode();
 }
 
+bool FRDGBuilder::IsDrainEnabled()
+{
+	return GRDGDrain != 0;
+}
+
 ERDGPassFlags FRDGBuilder::OverridePassFlags(const TCHAR* PassName, ERDGPassFlags PassFlags, bool bAsyncComputeSupported)
 {
 	const bool bDebugAllowedForPass =
@@ -518,7 +523,7 @@ IRHITransientResourceAllocator* FRDGBuilder::FTransientResourceAllocator::GetOrC
 	return Allocator;
 }
 
-FRDGBuilder::FRDGBuilder(FRHICommandListImmediate& InRHICmdList, FRDGEventName InName)
+FRDGBuilder::FRDGBuilder(FRHICommandListImmediate& InRHICmdList, FRDGEventName InName, ERDGBuilderFlags InFlags)
 	: RHICmdList(InRHICmdList)
 	, Blackboard(Allocator)
 	, RHICmdListAsyncCompute(FRHICommandListExecutor::GetImmediateAsyncComputeCommandList())
@@ -530,7 +535,7 @@ FRDGBuilder::FRDGBuilder(FRHICommandListImmediate& InRHICmdList, FRDGEventName I
 #if RDG_GPU_SCOPES
 	, GPUScopeStacks(Allocator)
 #endif
-	, bParallelExecuteEnabled(IsParallelExecuteEnabled())
+	, bParallelExecuteEnabled(IsParallelExecuteEnabled() && EnumHasAnyFlags(InFlags, ERDGBuilderFlags::AllowParallelExecute))
 #if RDG_ENABLE_DEBUG
 	, UserValidation(Allocator)
 	, BarrierValidation(&Passes, BuilderName)
@@ -552,14 +557,6 @@ FRDGBuilder::~FRDGBuilder()
 {
 	SCOPED_NAMED_EVENT(FRDGBuilder_Clear, FColor::Emerald);
 
-	if (bParallelExecuteEnabled)
-	{
-		if (!ParallelPassEvents.IsEmpty())
-		{
-			FTaskGraphInterface::Get().WaitUntilTasksComplete(ParallelPassEvents, ENamedThreads::GetRenderThread_Local());
-		}
-	}
-
 	Passes.Clear();
 	Buffers.Clear();
 	UniformBuffers.Clear();
@@ -1474,7 +1471,7 @@ void FRDGBuilder::Execute()
 void FRDGBuilder::Execute(EExecuteMode ExecuteMode)
 {
 	CSV_SCOPED_TIMING_STAT_EXCLUSIVE(RDG);
-	SCOPED_NAMED_EVENT_TCHAR("FRDGBuilder::Execute", FColor::Magenta);
+	SCOPED_NAMED_EVENT_TEXT("FRDGBuilder::Execute", FColor::Magenta);
 
 	// Create the epilogue pass at the end of the graph just prior to compilation.
 	SetupEmptyPass(EpiloguePass = Passes.Allocate<FRDGSentinelPass>(Allocator, RDG_EVENT_NAME("Graph Epilogue%s", (ExecuteMode == EExecuteMode::Drain) ? TEXT(" (Drain)") : TEXT(""))));
@@ -1485,6 +1482,8 @@ void FRDGBuilder::Execute(EExecuteMode ExecuteMode)
 	const bool bFirstExecute = ExecuteGeneration == 0;
 	const bool bFinalExecute = ExecuteMode == EExecuteMode::Final;
 
+	FGraphEventArray AsyncCompileEvents;
+
 	if (bFirstExecute)
 	{
 		IF_RDG_ENABLE_DEBUG(UserValidation.ValidateExecuteBegin());
@@ -1496,16 +1495,28 @@ void FRDGBuilder::Execute(EExecuteMode ExecuteMode)
 	{
 		BeginFlushResourcesRHI();
 
+		SetupBufferUploads();
+
 		Compile(ExecuteMode);
 
-		EndFlushResourcesRHI();
+		IF_RDG_GPU_SCOPES(GPUScopeStacks.ReserveOps(Passes.Num()));
+		IF_RDG_CPU_SCOPES(CPUScopeStacks.ReserveOps());
+
+		if (bParallelExecuteEnabled)
+		{
+			// Parallel execute setup can be done off the render thread and synced prior to dispatch.
+			AsyncCompileEvents.Emplace(FFunctionGraphTask::CreateAndDispatchWhenReady(
+				[this](ENamedThreads::Type, const FGraphEventRef&)
+			{
+				SetupParallelExecute();
 
-		SubmitBufferUploads();
+			}, TStatId(), nullptr, ENamedThreads::AnyHiPriThreadHiPriTask));
+		}
 
 		{
 			SCOPE_CYCLE_COUNTER(STAT_RDG_CollectResourcesTime);
 			CSV_SCOPED_TIMING_STAT_EXCLUSIVE(RDG_CollectResources);
-			SCOPED_NAMED_EVENT_TCHAR("FRDGBuilder::CollectResources", FColor::Magenta);
+			SCOPED_NAMED_EVENT_TEXT("FRDGBuilder::CollectResources", FColor::Magenta);
 
 			if (bFinalExecute)
 			{
@@ -1579,10 +1590,8 @@ void FRDGBuilder::Execute(EExecuteMode ExecuteMode)
 			}
 		}
 
-		CreateUniformBuffers();
-
 		{
-			SCOPED_NAMED_EVENT_TCHAR("FRDGBuilder::CollectBarriers", FColor::Magenta);
+			SCOPED_NAMED_EVENT_TEXT("FRDGBuilder::CollectBarriers", FColor::Magenta);
 			SCOPE_CYCLE_COUNTER(STAT_RDG_CollectBarriersTime);
 			CSV_SCOPED_TIMING_STAT_EXCLUSIVE_CONDITIONAL(RDG_CollectBarriers, GRDGVerboseCSVStats != 0);
 
@@ -1600,7 +1609,7 @@ void FRDGBuilder::Execute(EExecuteMode ExecuteMode)
 
 	if (bFinalExecute)
 	{
-		SCOPED_NAMED_EVENT_TCHAR("FRDGBuilder::Finalize", FColor::Magenta);
+		SCOPED_NAMED_EVENT_TEXT("FRDGBuilder::Finalize", FColor::Magenta);
 
 #if RDG_ENABLE_DEBUG
 		const auto LogResource = [&](auto* Resource, auto& Registry)
@@ -1642,153 +1651,58 @@ void FRDGBuilder::Execute(EExecuteMode ExecuteMode)
 		});
 	}
 
-	IF_RDG_ENABLE_DEBUG(GRDGAllowRHIAccess = false);
+	if (bParallelExecuteEnabled)
+	{
+		// Overlap pass barrier creation with other compilation tasks, since it's not required to run on the render thread.
+		AsyncCompileEvents.Emplace(FFunctionGraphTask::CreateAndDispatchWhenReady(
+			[this](ENamedThreads::Type, const FGraphEventRef&)
+		{
+			CreatePassBarriers();
 
-	CreatePassBarriers();
+		}, TStatId(), nullptr, ENamedThreads::AnyHiPriThreadHiPriTask));
+	}
+	else
+	{
+		CreatePassBarriers();
+	}
+
+	SubmitBufferUploads();
+
+	CreateUniformBuffers();
+
+	EndFlushResourcesRHI();
 
 	if (bFirstExecute)
 	{
 		IF_RDG_ENABLE_TRACE(Trace.OutputGraphBegin());
 	}
 
+	IF_RDG_ENABLE_DEBUG(GRDGAllowRHIAccess = false);
+
+	const ENamedThreads::Type RenderThread = ENamedThreads::GetRenderThread_Local();
+
+	FGraphEventRef DispatchParallelExecuteEvent;
+
 	if (!IsImmediateMode())
 	{
 		SCOPED_NAMED_EVENT_TEXT("FRDGBuilder::ExecutePasses", FColor::Magenta);
 		SCOPE_CYCLE_COUNTER(STAT_RDG_ExecuteTime);
 		CSV_SCOPED_TIMING_STAT_EXCLUSIVE(RenderOther);
 
-		IF_RDG_GPU_SCOPES(GPUScopeStacks.ReserveOps(Passes.Num()));
-		IF_RDG_CPU_SCOPES(CPUScopeStacks.ReserveOps());
+		// Wait on all async compilation tasks before executing any passes.
+		if (!AsyncCompileEvents.IsEmpty())
+		{
+			FTaskGraphInterface::Get().WaitUntilTasksComplete(AsyncCompileEvents, RenderThread);
+		}
 
 		if (bParallelExecuteEnabled)
 		{
+			DispatchParallelExecuteEvent = FFunctionGraphTask::CreateAndDispatchWhenReady(
+				[this, &RHICmdContext = RHICmdList.GetContext()](ENamedThreads::Type, const FGraphEventRef&)
 			{
-				SCOPED_NAMED_EVENT(ParallelExecutePassSetup, FColor::Emerald);
-				TArray<FRDGPass*, TInlineAllocator<128, FRDGArrayAllocator>> ParallelPassCandidates;
-
-				const auto FlushParallelPassCandidates = [&]()
-				{
-					if (ParallelPassCandidates.IsEmpty())
-					{
-						return;
-					}
-
-					int32 PassBeginIndex = 0;
-					int32 PassEndIndex = ParallelPassCandidates.Num();
-
-					// It's possible that the first pass is inside a merged render pass region. If so, we must push it forward until after the render pass ends.
-					if (const FRDGPass* FirstPass = ParallelPassCandidates[PassBeginIndex]; FirstPass->PrologueBarrierPass < FirstPass->Handle)
-					{
-						const FRDGPass* EpilogueBarrierPass = Passes[FirstPass->EpilogueBarrierPass];
-
-						for (; PassBeginIndex < ParallelPassCandidates.Num(); ++PassBeginIndex)
-						{
-							if (ParallelPassCandidates[PassBeginIndex] == EpilogueBarrierPass)
-							{
-								++PassBeginIndex;
-								break;
-							}
-						}
-					}
-
-					if (PassBeginIndex < PassEndIndex)
-					{
-						// It's possible that the last pass is inside a merged render pass region. If so, we must push it backwards until after the render pass begins.
-						if (FRDGPass* LastPass = ParallelPassCandidates.Last(); LastPass->EpilogueBarrierPass > LastPass->Handle)
-						{
-							FRDGPass* PrologueBarrierPass = Passes[LastPass->PrologueBarrierPass];
-
-							while (PassEndIndex > PassBeginIndex)
-							{
-								if (ParallelPassCandidates[--PassEndIndex] == PrologueBarrierPass)
-								{
-									break;
-								}
-							}
-						}
-					}
-
-					const int32 ParallelPassCandidateCount = PassEndIndex - PassBeginIndex;
-
-					if (ParallelPassCandidateCount >= GRDGParallelExecutePassMin)
-					{
-						FRDGPass* PassBegin = ParallelPassCandidates[PassBeginIndex];
-						PassBegin->bParallelExecuteBegin = 1;
-
-						FRDGPass* PassEnd = ParallelPassCandidates[PassEndIndex - 1];
-						PassEnd->bParallelExecuteEnd = 1;
-						PassEnd->ParallelPassSetIndex = ParallelPassSets.Num();
-
-						for (int32 PassIndex = PassBeginIndex; PassIndex < PassEndIndex; ++PassIndex)
-						{
-							ParallelPassCandidates[PassIndex]->bParallelExecute = 1;
-						}
-
-						FParallelPassSet& ParallelPassSet = ParallelPassSets.Emplace_GetRef(ParallelPassCandidates.GetData() + PassBeginIndex, ParallelPassCandidateCount);
-						ParallelPassSet.RHICmdList = new FRHICommandList(FRHIGPUMask::All());
-						ParallelPassSet.RHICmdList->CopyContext(RHICmdList);
-					}
-
-					ParallelPassCandidates.Reset();
-				};
-
-				ParallelPassSets.Reserve(32);
-				ParallelPassCandidates.Emplace(GetProloguePass(ERHIPipeline::Graphics));
-
-				for (FRDGPassHandle PassHandle = ProloguePassHandle + 1; PassHandle < EpiloguePassHandle; ++PassHandle)
-				{
-					FRDGPass* Pass = Passes[PassHandle];
-
-					if (Pass->bCulled)
-					{
-						continue;
-					}
-
-					CompilePassOps(Pass);
-
-					if (Pass->Pipeline == ERHIPipeline::AsyncCompute)
-					{
-						continue;
-					}
-
-					if (Pass->bImmediateCommandList)
-					{
-						FlushParallelPassCandidates();
-						continue;
-					}
-
-					ParallelPassCandidates.Emplace(Pass);
-				}
-
-				FlushParallelPassCandidates();
-			}
-
-			{
-				SCOPED_NAMED_EVENT(ParallelExecutePassDispatch, FColor::Emerald);
-
-				ParallelPassEvents.Reserve(ParallelPassEvents.Num() + ParallelPassSets.Num());
-
-				for (FParallelPassSet& ParallelPassSet : ParallelPassSets)
-				{
-					ParallelPassSet.Event = FFunctionGraphTask::CreateAndDispatchWhenReady(
-						[this, ParallelPassSet](ENamedThreads::Type, const FGraphEventRef& MyCompletionGraphEvent)
-					{
-						SCOPED_NAMED_EVENT(ParallelExecute, FColor::Emerald);
-						FTaskTagScope Scope(ETaskTag::EParallelRenderingThread);
-						FMemMark MemMark(FMemStack::Get());
-
-						for (FRDGPass* Pass : ParallelPassSet.Passes)
-						{
-							ExecutePass(Pass, *ParallelPassSet.RHICmdList);
-						}
-
-						ParallelPassSet.RHICmdList->HandleRTThreadTaskCompletion(MyCompletionGraphEvent);
+				DispatchParallelExecute(&RHICmdContext);
 
-					}, TStatId(), nullptr, ENamedThreads::AnyHiPriThreadHiPriTask);
-
-					ParallelPassEvents.Emplace(ParallelPassSet.Event);
-				}
-			}
+			}, TStatId(), nullptr, ENamedThreads::AnyHiPriThreadHiPriTask);
 		}
 
 		if (FRDGPass* AsyncComputeProloguePass = GetProloguePass(ERHIPipeline::AsyncCompute))
@@ -1820,6 +1734,11 @@ void FRDGBuilder::Execute(EExecuteMode ExecuteMode)
 					if (Pass->bParallelExecuteEnd)
 					{
 						FParallelPassSet& ParallelPassSet = ParallelPassSets[Pass->ParallelPassSetIndex];
+
+						// Busy wait until our pass set is ready. This will be set by the dispatch task.
+						while (!FPlatformAtomics::AtomicRead(&ParallelPassSet.bInitialized)) {};
+
+						check(ParallelPassSet.Event != nullptr && ParallelPassSet.RHICmdList != nullptr);
 						RHICmdList.QueueRenderThreadCommandListSubmit(ParallelPassSet.Event, ParallelPassSet.RHICmdList);
 					}
 
@@ -1843,6 +1762,12 @@ void FRDGBuilder::Execute(EExecuteMode ExecuteMode)
 		ExecutePass(EpiloguePass, RHICmdList);
 	}
 
+	// Wait for the parallel dispatch task before attempting to wait on the execute event array (the former mutates the array).
+	if (DispatchParallelExecuteEvent)
+	{
+		DispatchParallelExecuteEvent->Wait(RenderThread);
+	}
+
 	if (bFinalExecute)
 	{
 		RHICmdList.SetStaticUniformBuffers({});
@@ -1862,6 +1787,13 @@ void FRDGBuilder::Execute(EExecuteMode ExecuteMode)
 		}
 	#endif
 
+		// Wait on the actual parallel execute tasks in the Execute call. When draining is okay to let them overlap with other graph setup.
+		// This also needs to be done before extraction of external resources to be consistent with non-parallel rendering.
+		if (!ParallelExecuteEvents.IsEmpty())
+		{
+			FTaskGraphInterface::Get().WaitUntilTasksComplete(ParallelExecuteEvents, RenderThread);
+		}
+
 		for (const FExtractedTexture& ExtractedTexture : ExtractedTextures)
 		{
 			check(ExtractedTexture.Texture->PooledRenderTarget);
@@ -1896,8 +1828,8 @@ void FRDGBuilder::Execute(EExecuteMode ExecuteMode)
 	}
 
 	ExecuteGeneration++;
-	ParallelPassSets.Reset();
 
+	// Flush any outstanding async compute commands at the end to get things moving down the pipe.
 	if (RHICmdListAsyncCompute.HasCommands())
 	{
 		FRHIAsyncComputeCommandListImmediate::ImmediateDispatch(RHICmdListAsyncCompute);
@@ -2154,6 +2086,7 @@ void FRDGBuilder::SetupPassInternal(FRDGPass* Pass, FRDGPassHandle PassHandle, E
 
 		check(!EnumHasAnyFlags(PassPipeline, ERHIPipeline::AsyncCompute));
 
+		SetupBufferUploads();
 		SubmitBufferUploads();
 		CompilePassOps(Pass);
 		BeginResourcesRHI(Pass, PassHandle);
@@ -2166,41 +2099,200 @@ void FRDGBuilder::SetupPassInternal(FRDGPass* Pass, FRDGPassHandle PassHandle, E
 	IF_RDG_ENABLE_DEBUG(VisualizePassOutputs(Pass));
 }
 
-void FRDGBuilder::SubmitBufferUploads()
+void FRDGBuilder::SetupBufferUploads()
 {
-	SCOPED_NAMED_EVENT_TCHAR("FRDGBuilder::SubmitBufferUploads", FColor::Magenta);
+	SCOPED_NAMED_EVENT_TEXT("FRDGBuilder::PrepareBufferUploads", FColor::Magenta);
 
 	for (FUploadedBuffer& UploadedBuffer : UploadedBuffers)
 	{
-		const void* InitialData;
-		uint64 InitialDataSize;
-
 		if (UploadedBuffer.bUseCallbacks)
 		{
-			InitialData = UploadedBuffer.DataCallback();
-			InitialDataSize = UploadedBuffer.DataSizeCallback();
+			UploadedBuffer.Data = UploadedBuffer.DataCallback();
+			UploadedBuffer.DataSize = UploadedBuffer.DataSizeCallback();
 		}
-		else
+
+		if (UploadedBuffer.Data && UploadedBuffer.DataSize)
 		{
-			InitialData = UploadedBuffer.Data;
-			InitialDataSize = UploadedBuffer.DataSize;
+			ConvertToExternalBuffer(UploadedBuffer.Buffer);
+			check(UploadedBuffer.DataSize <= UploadedBuffer.Buffer->Desc.GetTotalNumBytes());
 		}
+	}
+}
 
-		if (InitialData && InitialDataSize)
+void FRDGBuilder::SubmitBufferUploads()
+{
+	SCOPED_NAMED_EVENT_TEXT("FRDGBuilder::SubmitBufferUploads", FColor::Magenta);
+
+	for (const FUploadedBuffer& UploadedBuffer : UploadedBuffers)
+	{
+		if (UploadedBuffer.Data && UploadedBuffer.DataSize)
 		{
-			const TRefCountPtr<FRDGPooledBuffer>& PooledBuffer = ConvertToExternalBuffer(UploadedBuffer.Buffer);
-			check(InitialDataSize <= UploadedBuffer.Buffer->Desc.GetTotalNumBytes());
-			void* DestPtr = RHICmdList.LockBuffer(PooledBuffer->GetRHI(), 0, InitialDataSize, RLM_WriteOnly);
-			FMemory::Memcpy(DestPtr, InitialData, InitialDataSize);
-			RHICmdList.UnlockBuffer(PooledBuffer->GetRHI());
+			void* DestPtr = RHICmdList.LockBuffer(UploadedBuffer.Buffer->GetRHI(), 0, UploadedBuffer.DataSize, RLM_WriteOnly);
+			FMemory::Memcpy(DestPtr, UploadedBuffer.Data, UploadedBuffer.DataSize);
+			RHICmdList.UnlockBuffer(UploadedBuffer.Buffer->GetRHI());
 		}
 	}
 	UploadedBuffers.Reset();
 }
 
+void FRDGBuilder::SetupParallelExecute()
+{
+	SCOPED_NAMED_EVENT(SetupParallelExecute, FColor::Emerald);
+	FTaskTagScope Scope(ETaskTag::EParallelRenderingThread);
+
+	TArray<FRDGPass*, TInlineAllocator<64, FRDGArrayAllocator>> ParallelPassCandidates;
+	int32 MergedRenderPassCandidates = 0;
+
+	const auto FlushParallelPassCandidates = [&]()
+	{
+		if (ParallelPassCandidates.IsEmpty())
+		{
+			return;
+		}
+
+		int32 PassBeginIndex = 0;
+		int32 PassEndIndex = ParallelPassCandidates.Num();
+
+		// It's possible that the first pass is inside a merged RHI render pass region. If so, we must push it forward until after the render pass ends.
+		if (const FRDGPass* FirstPass = ParallelPassCandidates[PassBeginIndex]; FirstPass->PrologueBarrierPass < FirstPass->Handle)
+		{
+			const FRDGPass* EpilogueBarrierPass = Passes[FirstPass->EpilogueBarrierPass];
+
+			for (; PassBeginIndex < ParallelPassCandidates.Num(); ++PassBeginIndex)
+			{
+				if (ParallelPassCandidates[PassBeginIndex] == EpilogueBarrierPass)
+				{
+					++PassBeginIndex;
+					break;
+				}
+			}
+		}
+
+		if (PassBeginIndex < PassEndIndex)
+		{
+			// It's possible that the last pass is inside a merged RHI render pass region. If so, we must push it backwards until after the render pass begins.
+			if (FRDGPass* LastPass = ParallelPassCandidates.Last(); LastPass->EpilogueBarrierPass > LastPass->Handle)
+			{
+				FRDGPass* PrologueBarrierPass = Passes[LastPass->PrologueBarrierPass];
+
+				while (PassEndIndex > PassBeginIndex)
+				{
+					if (ParallelPassCandidates[--PassEndIndex] == PrologueBarrierPass)
+					{
+						break;
+					}
+				}
+			}
+		}
+
+		const int32 ParallelPassCandidateCount = PassEndIndex - PassBeginIndex;
+
+		if (ParallelPassCandidateCount >= GRDGParallelExecutePassMin)
+		{
+			FRDGPass* PassBegin = ParallelPassCandidates[PassBeginIndex];
+			PassBegin->bParallelExecuteBegin = 1;
+
+			FRDGPass* PassEnd = ParallelPassCandidates[PassEndIndex - 1];
+			PassEnd->bParallelExecuteEnd = 1;
+			PassEnd->ParallelPassSetIndex = ParallelPassSets.Num();
+
+			for (int32 PassIndex = PassBeginIndex; PassIndex < PassEndIndex; ++PassIndex)
+			{
+				ParallelPassCandidates[PassIndex]->bParallelExecute = 1;
+			}
+
+			ParallelPassSets.Emplace(ParallelPassCandidates.GetData() + PassBeginIndex, ParallelPassCandidateCount);
+		}
+
+		ParallelPassCandidates.Reset();
+		MergedRenderPassCandidates = 0;
+	};
+
+	ParallelPassSets.Reserve(32);
+	ParallelPassCandidates.Emplace(GetProloguePass(ERHIPipeline::Graphics));
+
+	for (FRDGPassHandle PassHandle = GetProloguePassHandle() + 1; PassHandle < GetEpiloguePassHandle(); ++PassHandle)
+	{
+		FRDGPass* Pass = Passes[PassHandle];
+
+		if (Pass->bCulled)
+		{
+			continue;
+		}
+
+		CompilePassOps(Pass);
+
+		if (Pass->Pipeline == ERHIPipeline::AsyncCompute)
+		{
+			continue;
+		}
+
+		if (Pass->bImmediateCommandList)
+		{
+			FlushParallelPassCandidates();
+			continue;
+		}
+
+		ParallelPassCandidates.Emplace(Pass);
+
+		// Don't count merged render passes for the maximum pass threshold. This avoids the case where
+		// a large merged render pass span could end up forcing it back onto the render thread, since
+		// it's not possible to launch a task for a subset of passes within a merged render pass.
+		MergedRenderPassCandidates += Pass->bSkipRenderPassBegin | Pass->bSkipRenderPassEnd;
+
+		if (ParallelPassCandidates.Num() - MergedRenderPassCandidates >= GRDGParallelExecutePassMax)
+		{
+			FlushParallelPassCandidates();
+		}
+	}
+
+	FlushParallelPassCandidates();
+}
+
+void FRDGBuilder::DispatchParallelExecute(IRHICommandContext* RHICmdContext)
+{
+	SCOPED_NAMED_EVENT(SetupParallelExecute, FColor::Emerald);
+	ParallelExecuteEvents.Reserve(ParallelExecuteEvents.Num() + ParallelPassSets.Num());
+
+	for (int32 Index = ParallelPassSetsOffset; Index < ParallelPassSets.Num(); ++Index)
+	{
+		FParallelPassSet& ParallelPassSet = ParallelPassSets[Index];
+		ParallelPassSet.RHICmdList = new FRHICommandList(FRHIGPUMask::All());
+		ParallelPassSet.RHICmdList->SetContext(RHICmdContext);
+
+		// Avoid referencing the parallel pass struct directly in the task, as the set can resize.
+		TArrayView<FRDGPass*> ParallelPasses = ParallelPassSet.Passes;
+
+		ParallelPassSet.Event = FFunctionGraphTask::CreateAndDispatchWhenReady(
+			[this, ParallelPasses, &RHICmdListPass = *ParallelPassSet.RHICmdList](ENamedThreads::Type, const FGraphEventRef& MyCompletionGraphEvent)
+		{
+			SCOPED_NAMED_EVENT(ParallelExecute, FColor::Emerald);
+			FTaskTagScope Scope(ETaskTag::EParallelRenderingThread);
+			FMemMark MemMark(FMemStack::Get());
+
+			for (FRDGPass* Pass : ParallelPasses)
+			{
+				ExecutePass(Pass, RHICmdListPass);
+			}
+
+			RHICmdListPass.HandleRTThreadTaskCompletion(MyCompletionGraphEvent);
+
+		}, TStatId(), nullptr, ENamedThreads::AnyHiPriThreadHiPriTask);
+
+		// Mark this set as initialized so that it can be submitted.
+		FPlatformAtomics::AtomicStore(&ParallelPassSet.bInitialized, 1);
+
+		// Enqueue the event to be synced at the end of RDG execution.
+		ParallelExecuteEvents.Emplace(ParallelPassSet.Event);
+	}
+
+	// Update the offset to avoid reprocessing the same parallel pass sets when drain is enabled.
+	ParallelPassSetsOffset = ParallelPassSets.Num();
+}
+
 void FRDGBuilder::CreateUniformBuffers()
 {
-	SCOPED_NAMED_EVENT_TCHAR("FRDGBuilder::CreateUniformBuffers", FColor::Magenta);
+	SCOPED_NAMED_EVENT_TEXT("FRDGBuilder::CreateUniformBuffers", FColor::Magenta);
 
 	for (FRDGUniformBufferHandle UniformBufferHandle : UniformBuffersToCreate)
 	{
@@ -2325,8 +2417,8 @@ void FRDGBuilder::ExecutePassEpilogue(FRHIComputeCommandList& RHICmdListPass, FR
 
 void FRDGBuilder::ExecutePass(FRDGPass* Pass, FRHIComputeCommandList& RHICmdListPass)
 {
-#if 0 // This is compiled out for performance reasons; enable to display per-pass timings.
-	SCOPED_NAMED_EVENT_F(TEXT("%s"), FColor::Magenta, Pass->GetName());
+#if RDG_EVENTS != RDG_EVENTS_NONE
+	SCOPED_NAMED_EVENT_TCHAR(Pass->GetName(), FColor::Magenta);
 #endif
 
 	// Note that we must do this before doing anything with RHICmdList for the pass.
@@ -2456,7 +2548,7 @@ void FRDGBuilder::CollectPassBarriers(FRDGPass* Pass, FRDGPassHandle PassHandle)
 
 void FRDGBuilder::CreatePassBarriers()
 {
-	SCOPED_NAMED_EVENT_TCHAR("FRDGBuilder::CreatePassBarriers", FColor::Magenta);
+	SCOPED_NAMED_EVENT_TEXT("FRDGBuilder::CreatePassBarriers", FColor::Magenta);
 
 	for (FRDGBarrierBatchBegin* BarrierBatchBegin : TransitionCreateQueue)
 	{
diff --git a/Engine/Source/Runtime/RenderCore/Private/RenderGraphPrivate.cpp b/Engine/Source/Runtime/RenderCore/Private/RenderGraphPrivate.cpp
--- a/Engine/Source/Runtime/RenderCore/Private/RenderGraphPrivate.cpp
+++ b/Engine/Source/Runtime/RenderCore/Private/RenderGraphPrivate.cpp
@@ -306,12 +306,18 @@ FAutoConsoleVariableRef CVarRDGUseTransientAllocator(
 	TEXT(" 2: enables the transient allocator for resources with FastVRAM flag only"),
 	ECVF_RenderThreadSafe);
 
-int32 GRDGParallelExecutePassMin = 0;
-FAutoConsoleVariableRef CVarRDGParallelExecute(
+int32 GRDGParallelExecutePassMin = 1;
+FAutoConsoleVariableRef CVarRDGParallelExecutePassMin(
 	TEXT("r.RDG.ParallelExecutePassMin"), GRDGParallelExecutePassMin,
 	TEXT("The minimum span of contiguous passes eligible for parallel execution for the span to be offloaded to a task."),
 	ECVF_RenderThreadSafe);
 
+int32 GRDGParallelExecutePassMax = 32;
+FAutoConsoleVariableRef CVarRDGParallelExecutePassMax(
+	TEXT("r.RDG.ParallelExecutePassMax"), GRDGParallelExecutePassMax,
+	TEXT("The maximum span of contiguous passes eligible for parallel execution for the span to be offloaded to a task."),
+	ECVF_RenderThreadSafe);
+
 // Fix for random GPU crashes on draw indirects on multiple IHVs. Force all indirect arg buffers as non transient (see UE-115982)
 int32 GRDGTransientIndirectArgBuffers = 0;
 FAutoConsoleVariableRef CVarRDGIndirectArgBufferTransientAllocated(
@@ -327,6 +333,43 @@ FAutoConsoleVariableRef CVarRDGDrain(
 	TEXT(" 1: enables execution of the RDG subgraph when FRDGBuilder::Drain is called."),
 	ECVF_RenderThreadSafe);
 
+int32 GRDGParallelExecuteStress = 0;
+FAutoConsoleVariableRef CVarRDGDebugParallelExecute(
+	TEXT("r.RDG.ParallelExecuteStress"),
+	GRDGParallelExecuteStress,
+	TEXT("Stress tests the parallel execution path by launching one task per pass. Render pass merging is also disabled."),
+	FConsoleVariableDelegate::CreateLambda([](IConsoleVariable* Variable)
+	{
+		static int32 GRDGMergeRenderPassesHistory = GRDGMergeRenderPasses;
+		static int32 GRDGParallelExecutePassMinHistory = GRDGParallelExecutePassMin;
+		static int32 GRDGParallelExecutePassMaxHistory = GRDGParallelExecutePassMax;
+
+		const int32 CurrentValue = Variable->GetInt();
+
+		if (GRDGParallelExecuteStress == CurrentValue)
+		{
+			return;
+		}
+
+		if (CurrentValue)
+		{
+			GRDGMergeRenderPassesHistory = GRDGMergeRenderPasses;
+			GRDGParallelExecutePassMinHistory = GRDGParallelExecutePassMin;
+			GRDGParallelExecutePassMaxHistory = GRDGParallelExecutePassMax;
+
+			GRDGMergeRenderPasses = 0;
+			GRDGParallelExecutePassMin = 1;
+			GRDGParallelExecutePassMax = 1;
+		}
+		else
+		{
+			GRDGMergeRenderPasses = GRDGMergeRenderPassesHistory;
+			GRDGParallelExecutePassMin = GRDGParallelExecutePassMinHistory;
+			GRDGParallelExecutePassMax = GRDGParallelExecutePassMaxHistory;
+		}
+	}),
+	ECVF_RenderThreadSafe);
+
 #if CSV_PROFILER
 int32 GRDGVerboseCSVStats = 0;
 FAutoConsoleVariableRef CVarRDGVerboseCSVStats(
diff --git a/Engine/Source/Runtime/RenderCore/Private/RenderGraphPrivate.h b/Engine/Source/Runtime/RenderCore/Private/RenderGraphPrivate.h
--- a/Engine/Source/Runtime/RenderCore/Private/RenderGraphPrivate.h
+++ b/Engine/Source/Runtime/RenderCore/Private/RenderGraphPrivate.h
@@ -110,6 +110,7 @@ extern int32 GRDGTransientAllocator;
 extern int32 GRDGTransientIndirectArgBuffers;
 extern int32 GRDGDrain;
 extern int32 GRDGParallelExecutePassMin;
+extern int32 GRDGParallelExecutePassMax;
 
 #if CSV_PROFILER
 extern int32 GRDGVerboseCSVStats;
diff --git a/Engine/Source/Runtime/RenderCore/Public/RenderGraphBuilder.h b/Engine/Source/Runtime/RenderCore/Public/RenderGraphBuilder.h
--- a/Engine/Source/Runtime/RenderCore/Public/RenderGraphBuilder.h
+++ b/Engine/Source/Runtime/RenderCore/Public/RenderGraphBuilder.h
@@ -23,7 +23,7 @@ class RENDERCORE_API FRDGBuilder
 	: FRDGAllocatorScope
 {
 public:
-	FRDGBuilder(FRHICommandListImmediate& InRHICmdList, FRDGEventName InName = {});
+	FRDGBuilder(FRHICommandListImmediate& RHICmdList, FRDGEventName Name = {}, ERDGBuilderFlags Flags = ERDGBuilderFlags::None);
 	FRDGBuilder(const FRDGBuilder&) = delete;
 	~FRDGBuilder();
 
@@ -274,6 +274,9 @@ public:
 	/** Whether RDG is running in immediate mode. */
 	static bool IsImmediateMode();
 
+	/** Whether RDG will actually perform a drain when Drain() is called. */
+	static bool IsDrainEnabled();
+
 	/** The RHI command list used for the render graph. */
 	FRHICommandListImmediate& RHICmdList;
 
@@ -515,10 +518,14 @@ private:
 		TArray<FRDGPass*, FRDGArrayAllocator> Passes;
 		FGraphEventRef Event;
 		FRHICommandList* RHICmdList{};
+		int8 bInitialized = 0;
 	};
 
 	TArray<FParallelPassSet, FRDGArrayAllocator> ParallelPassSets;
-	FGraphEventArray ParallelPassEvents;
+	uint32 ParallelPassSetsOffset = 0;
+
+	/** Array of all active parallel execute tasks. */
+	FGraphEventArray ParallelExecuteEvents;
 
 	/** Array of all pooled references held during execution. */
 	FRDGPooledTextureArray ActivePooledTextures;
@@ -608,6 +615,10 @@ private:
 	void EndResourceRHI(FRDGPassHandle, FRDGTexture* Texture, uint32 ReferenceCount);
 	void EndResourceRHI(FRDGPassHandle, FRDGBuffer* Buffer, uint32 ReferenceCount);
 
+	void SetupParallelExecute();
+	void DispatchParallelExecute(IRHICommandContext* RHICmdContext);
+
+	void SetupBufferUploads();
 	void SubmitBufferUploads();
 	void BeginFlushResourcesRHI();
 	void EndFlushResourcesRHI();
diff --git a/Engine/Source/Runtime/RenderCore/Public/RenderGraphBuilder.inl b/Engine/Source/Runtime/RenderCore/Public/RenderGraphBuilder.inl
--- a/Engine/Source/Runtime/RenderCore/Public/RenderGraphBuilder.inl
+++ b/Engine/Source/Runtime/RenderCore/Public/RenderGraphBuilder.inl
@@ -238,6 +238,12 @@ inline void FRDGBuilder::QueueBufferUpload(FRDGBufferRef Buffer, FRDGBufferIniti
 inline void FRDGBuilder::QueueTextureExtraction(FRDGTextureRef Texture, TRefCountPtr<IPooledRenderTarget>* OutTexturePtr, ERHIAccess AccessFinal)
 {
 	QueueTextureExtraction(Texture, OutTexturePtr);
+
+	if (Texture->bTransient)
+	{
+		return;
+	}
+
 	SetTextureAccessFinal(Texture, AccessFinal);
 }
 
@@ -245,6 +251,11 @@ inline void FRDGBuilder::QueueTextureExtraction(FRDGTextureRef Texture, TRefCoun
 {
 	IF_RDG_ENABLE_DEBUG(UserValidation.ValidateExtractTexture(Texture, OutTexturePtr));
 
+	if (Texture->bTransient)
+	{
+		return;
+	}
+
 	Texture->ReferenceCount++;
 	Texture->bExtracted = true;
 	Texture->bCulled = false;
diff --git a/Engine/Source/Runtime/RenderCore/Public/RenderGraphDefinitions.h b/Engine/Source/Runtime/RenderCore/Public/RenderGraphDefinitions.h
--- a/Engine/Source/Runtime/RenderCore/Public/RenderGraphDefinitions.h
+++ b/Engine/Source/Runtime/RenderCore/Public/RenderGraphDefinitions.h
@@ -73,6 +73,15 @@
 
 /** ENUMS */
 
+enum class ERDGBuilderFlags
+{
+	None = 0,
+
+	/** Allows the builder to parallelize execution of passes. Without this flag, all passes execute on the render thread. */
+	AllowParallelExecute = 1 << 0
+};
+ENUM_CLASS_FLAGS(ERDGBuilderFlags);
+
 /** Flags to annotate a pass with when calling AddPass. */
 enum class ERDGPassFlags : uint8
 {
diff --git a/Engine/Source/Runtime/RenderCore/Public/RenderGraphPass.h b/Engine/Source/Runtime/RenderCore/Public/RenderGraphPass.h
--- a/Engine/Source/Runtime/RenderCore/Public/RenderGraphPass.h
+++ b/Engine/Source/Runtime/RenderCore/Public/RenderGraphPass.h
@@ -373,13 +373,6 @@ protected:
 			/** If set, dispatches to the RHI thread before executing this pass. */
 			uint32 bDispatchAfterExecute : 1;
 
-			/** If set, marks the begin / end of a span of passes executed in parallel in a task. */
-			uint32 bParallelExecuteBegin : 1;
-			uint32 bParallelExecuteEnd : 1;
-
-			/** If set, marks that a pass is executing in parallel. */
-			uint32 bParallelExecute : 1;
-
 			/** If set, the pass should set its command list stat. */
 			uint32 bSetCommandListStat : 1;
 
@@ -389,7 +382,22 @@ protected:
 			/** Whether this pass allocated a texture through the pool. */
 			IF_RDG_ENABLE_DEBUG(uint32 bFirstTextureAllocated : 1);
 		};
-		uint32 PackedBits = 0;
+		uint32 PackedBits1 = 0;
+	};
+
+	union
+	{
+		// Task-specific bits which are written in a task in parallel with reads from the other set.
+		struct
+		{
+			/** If set, marks the begin / end of a span of passes executed in parallel in a task. */
+			uint32 bParallelExecuteBegin : 1;
+			uint32 bParallelExecuteEnd : 1;
+
+			/** If set, marks that a pass is executing in parallel. */
+			uint32 bParallelExecute : 1;
+		};
+		uint32 PacketBits2 = 0;
 	};
 
 	/** Handle of the latest cross-pipeline producer and earliest cross-pipeline consumer. */
diff --git a/Engine/Source/Runtime/RenderCore/Public/RenderGraphUtils.h b/Engine/Source/Runtime/RenderCore/Public/RenderGraphUtils.h
--- a/Engine/Source/Runtime/RenderCore/Public/RenderGraphUtils.h
+++ b/Engine/Source/Runtime/RenderCore/Public/RenderGraphUtils.h
@@ -807,6 +807,7 @@ FORCEINLINE void AddPass(FRDGBuilder& GraphBuilder, FRDGEventName&& Name, Execut
 }
 
 template <typename ExecuteLambdaType>
+UE_DEPRECATED(5.0, "AddPass without an RDG_EVENT_NAME is deprecated. Use the named version instead.")
 FORCEINLINE void AddPass(FRDGBuilder& GraphBuilder, ExecuteLambdaType&& ExecuteLambda)
 {
 	AddPass(GraphBuilder, {}, MoveTemp(ExecuteLambda));
@@ -821,6 +822,7 @@ FORCEINLINE void AddPassIfDebug(FRDGBuilder& GraphBuilder, FRDGEventName&& Name,
 }
 
 template <typename ExecuteLambdaType>
+UE_DEPRECATED(5.0, "AddPassIfDebug without an RDG_EVENT_NAME is deprecated. Use the named version instead.")
 FORCEINLINE void AddPassIfDebug(FRDGBuilder& GraphBuilder, ExecuteLambdaType&& ExecuteLambda)
 {
 	AddPassIfDebug(GraphBuilder, {}, MoveTemp(ExecuteLambda));
@@ -834,63 +836,69 @@ FORCEINLINE void AddSetCurrentStatPass(FRDGBuilder& GraphBuilder, TStatId StatId
 
 FORCEINLINE void AddDispatchToRHIThreadPass(FRDGBuilder& GraphBuilder)
 {
-	AddPass(GraphBuilder, [](FRHICommandListImmediate& RHICmdList)
+	AddPass(GraphBuilder, RDG_EVENT_NAME("DispatchToRHI"), [](FRHICommandListImmediate& RHICmdList)
 	{
 		RHICmdList.ImmediateFlush(EImmediateFlushType::DispatchToRHIThread);
 	});
 }
 
+UE_DEPRECATED(5.0, "AddBeginUAVOverlapPass is deprecated.")
 FORCEINLINE void AddBeginUAVOverlapPass(FRDGBuilder& GraphBuilder)
 {
-	AddPass(GraphBuilder, [](FRHICommandListImmediate& RHICmdList)
+	AddPass(GraphBuilder, RDG_EVENT_NAME("BeginUAVOverlap"), [](FRHICommandList& RHICmdList)
 	{
 		RHICmdList.BeginUAVOverlap();
 	});
 }
 
+UE_DEPRECATED(5.0, "AddEndUAVOverlapPass is deprecated.")
 FORCEINLINE void AddEndUAVOverlapPass(FRDGBuilder& GraphBuilder)
 {
-	AddPass(GraphBuilder, [](FRHICommandListImmediate& RHICmdList)
+	AddPass(GraphBuilder, RDG_EVENT_NAME("EndUAVOverlap"), [](FRHICommandList& RHICmdList)
 	{
 		RHICmdList.EndUAVOverlap();
 	});
 }
 
+UE_DEPRECATED(5.0, "AddBeginUAVOverlapPass is deprecated.")
 FORCEINLINE void AddBeginUAVOverlapPass(FRDGBuilder& GraphBuilder, FRHIUnorderedAccessView* UAV)
 {
-	AddPass(GraphBuilder, [UAV](FRHICommandListImmediate& RHICmdList)
+	AddPass(GraphBuilder, RDG_EVENT_NAME("BeginUAVOverlap"), [UAV](FRHICommandList& RHICmdList)
 	{
 		RHICmdList.BeginUAVOverlap(UAV);
 	});
 }
 
+UE_DEPRECATED(5.0, "AddEndUAVOverlapPass is deprecated.")
 FORCEINLINE void AddEndUAVOverlapPass(FRDGBuilder& GraphBuilder, FRHIUnorderedAccessView* UAV)
 {
-	AddPass(GraphBuilder, [UAV](FRHICommandListImmediate& RHICmdList)
+	AddPass(GraphBuilder, RDG_EVENT_NAME("EndUAVOverlap"), [UAV](FRHICommandList& RHICmdList)
 	{
 		RHICmdList.EndUAVOverlap(UAV);
 	});
 }
 
+UE_DEPRECATED(5.0, "AddBeginUAVOverlapPass is deprecated.")
 FORCEINLINE void AddBeginUAVOverlapPass(FRDGBuilder& GraphBuilder, TArrayView<FRHIUnorderedAccessView*> UAVs)
 {
 	uint32 AllocSize = UAVs.Num() * sizeof(FRHIUnorderedAccessView*);
 	FRHIUnorderedAccessView** LocalUAVs = (FRHIUnorderedAccessView**)GraphBuilder.Alloc(AllocSize, alignof(FRHIUnorderedAccessView*));
 	FMemory::Memcpy(LocalUAVs, UAVs.GetData(), AllocSize);
 	TArrayView<FRHIUnorderedAccessView*> LocalView(LocalUAVs, UAVs.Num());
-	AddPass(GraphBuilder, [LocalView](FRHICommandListImmediate& RHICmdList)
+	AddPass(GraphBuilder, RDG_EVENT_NAME("BeginUAVOverlap"), [LocalView](FRHICommandList& RHICmdList)
 	{
 		RHICmdList.BeginUAVOverlap(LocalView);
 	});
 }
 
+UE_DEPRECATED(5.0, "AddEndUAVOverlapPass is deprecated.")
 FORCEINLINE void AddEndUAVOverlapPass(FRDGBuilder& GraphBuilder, TArrayView<FRHIUnorderedAccessView*> UAVs)
 {
 	uint32 AllocSize = UAVs.Num() * sizeof(FRHIUnorderedAccessView*);
 	FRHIUnorderedAccessView** LocalUAVs = (FRHIUnorderedAccessView**)GraphBuilder.Alloc(AllocSize, alignof(FRHIUnorderedAccessView*));
 	FMemory::Memcpy(LocalUAVs, UAVs.GetData(), AllocSize);
 	TArrayView<FRHIUnorderedAccessView*> LocalView(LocalUAVs, UAVs.Num());
-	AddPass(GraphBuilder, [LocalView](FRHICommandListImmediate& RHICmdList)
+	AddPass(GraphBuilder, RDG_EVENT_NAME("EndUAVOverlap"), [LocalView](FRHICommandList& RHICmdList)
 	{
 		RHICmdList.EndUAVOverlap(LocalView);
 	});
diff --git a/Engine/Source/Runtime/Renderer/Private/InstanceCulling/InstanceCullingContext.cpp b/Engine/Source/Runtime/Renderer/Private/InstanceCulling/InstanceCullingContext.cpp
--- a/Engine/Source/Runtime/Renderer/Private/InstanceCulling/InstanceCullingContext.cpp
+++ b/Engine/Source/Runtime/Renderer/Private/InstanceCulling/InstanceCullingContext.cpp
@@ -617,9 +617,7 @@ void FInstanceCullingContext::BuildRenderingCommandsDeferred(
 
 bool FInstanceCullingContext::AllowBatchedBuildRenderingCommands(const FGPUScene& GPUScene)
 {
-	static IConsoleVariable* CVarRDGDrain = IConsoleManager::Get().FindConsoleVariable(TEXT("r.RDG.Drain"));
-	check(CVarRDGDrain);
-	return GPUScene.IsEnabled() && !!GAllowBatchedBuildRenderingCommands && !FRDGBuilder::IsImmediateMode() && !CVarRDGDrain->GetInt();
+	return GPUScene.IsEnabled() && !!GAllowBatchedBuildRenderingCommands && !FRDGBuilder::IsImmediateMode() && !FRDGBuilder::IsDrainEnabled();
 }
 
 
diff --git a/Engine/Source/Runtime/Renderer/Private/PlanarReflectionRendering.cpp b/Engine/Source/Runtime/Renderer/Private/PlanarReflectionRendering.cpp
--- a/Engine/Source/Runtime/Renderer/Private/PlanarReflectionRendering.cpp
+++ b/Engine/Source/Runtime/Renderer/Private/PlanarReflectionRendering.cpp
@@ -347,7 +347,7 @@ static void UpdatePlanarReflectionContents_RenderThread(
 				SCOPED_DRAW_EVENT(RHICmdList, UpdatePlanarReflectionContent_RenderThread);
 #endif
 
-				FRDGBuilder GraphBuilder(RHICmdList, RDG_EVENT_NAME("PlanarReflection"));
+				FRDGBuilder GraphBuilder(RHICmdList, RDG_EVENT_NAME("PlanarReflection"), ERDGBuilderFlags::AllowParallelExecute);
 
 				// Reflection view late update
 				if (SceneRenderer->Views.Num() > 1)
diff --git a/Engine/Source/Runtime/Renderer/Private/ReflectionEnvironmentCapture.cpp b/Engine/Source/Runtime/Renderer/Private/ReflectionEnvironmentCapture.cpp
--- a/Engine/Source/Runtime/Renderer/Private/ReflectionEnvironmentCapture.cpp
+++ b/Engine/Source/Runtime/Renderer/Private/ReflectionEnvironmentCapture.cpp
@@ -723,12 +723,12 @@ void CaptureSceneToScratchCubemap(FRHICommandListImmediate& RHICmdList, FSceneRe
 	{
 		SCOPED_DRAW_EVENT(RHICmdList, CubeMapCapture);
 
-		FRDGBuilder GraphBuilder(RHICmdList, RDG_EVENT_NAME("CubeMapCapture"));
+		FRDGBuilder GraphBuilder(RHICmdList, RDG_EVENT_NAME("CubeMapCapture"), ERDGBuilderFlags::AllowParallelExecute);
 
 		// Render the scene normally for one face of the cubemap
 		SceneRenderer->Render(GraphBuilder);
 
-		AddPass(GraphBuilder, [](FRHICommandListImmediate& InRHICmdList)
+		AddPass(GraphBuilder, RDG_EVENT_NAME("FlushGPU"), [](FRHICommandListImmediate& InRHICmdList)
 		{
 			QUICK_SCOPE_CYCLE_COUNTER(STAT_CaptureSceneToScratchCubemap_Flush);
 			FRHICommandListExecutor::GetImmediateCommandList().ImmediateFlush(EImmediateFlushType::FlushRHIThread);
diff --git a/Engine/Source/Runtime/Renderer/Private/SceneCaptureRendering.cpp b/Engine/Source/Runtime/Renderer/Private/SceneCaptureRendering.cpp
--- a/Engine/Source/Runtime/Renderer/Private/SceneCaptureRendering.cpp
+++ b/Engine/Source/Runtime/Renderer/Private/SceneCaptureRendering.cpp
@@ -285,10 +285,10 @@ static void UpdateSceneCaptureContentDeferred_RenderThread(
 	{
 #if WANTS_DRAW_MESH_EVENTS
 		SCOPED_DRAW_EVENTF(RHICmdList, SceneCapture, TEXT("SceneCapture %s"), *EventName);
-		FRDGBuilder GraphBuilder(RHICmdList, RDG_EVENT_NAME("SceneCapture %s", *EventName));
+		FRDGBuilder GraphBuilder(RHICmdList, RDG_EVENT_NAME("SceneCapture %s", *EventName), ERDGBuilderFlags::AllowParallelExecute);
 #else
 		SCOPED_DRAW_EVENT(RHICmdList, UpdateSceneCaptureContent_RenderThread);
-		FRDGBuilder GraphBuilder(RHICmdList, RDG_EVENT_NAME("SceneCapture"));
+		FRDGBuilder GraphBuilder(RHICmdList, RDG_EVENT_NAME("SceneCapture"), ERDGBuilderFlags::AllowParallelExecute);
 #endif
 
 		FRDGTextureRef TargetTexture = RegisterExternalTexture(GraphBuilder, RenderTarget->GetRenderTargetTexture(), TEXT("SceneCaptureTarget"));
diff --git a/Engine/Source/Runtime/Renderer/Private/SceneRendering.cpp b/Engine/Source/Runtime/Renderer/Private/SceneRendering.cpp
--- a/Engine/Source/Runtime/Renderer/Private/SceneRendering.cpp
+++ b/Engine/Source/Runtime/Renderer/Private/SceneRendering.cpp
@@ -3784,7 +3784,8 @@ static void RenderViewFamily_RenderThread(FRHICommandListImmediate& RHICmdList,
 				RDG_EVENT_NAME("SceneRenderer_%s(ViewFamily=%s)",
 					ViewFamily.EngineShowFlags.HitProxies ? TEXT("RenderHitProxies") : TEXT("Render"),
 					ViewFamily.bResolveScene ? TEXT("Primary") : TEXT("Auxiliary")
-				)
+				),
+				ERDGBuilderFlags::AllowParallelExecute
 			);
 
 			if (ViewFamily.EngineShowFlags.HitProxies)
diff --git a/Engine/Source/Runtime/Renderer/Private/SceneTextures.cpp b/Engine/Source/Runtime/Renderer/Private/SceneTextures.cpp
--- a/Engine/Source/Runtime/Renderer/Private/SceneTextures.cpp
+++ b/Engine/Source/Runtime/Renderer/Private/SceneTextures.cpp
@@ -680,6 +680,11 @@ FSceneTextures& FMinimalSceneTextures::Create(FRDGBuilder& GraphBuilder, const F
 
 			SceneTextures.Color.Resolve = GraphBuilder.CreateTexture(Desc, SceneColorName);
 		}
+
+		if (FRDGBuilder::IsDrainEnabled())
+		{
+			SceneTextures.Color.Resolve->SetNonTransient();
+		}
 	}
 
 	// Custom Depth
@@ -780,6 +785,14 @@ FSceneTextures& FSceneTextures::Create(FRDGBuilder& GraphBuilder, const FSceneTe
 			const FRDGTextureDesc Desc = FRDGTextureDesc::Create2D(Config.Extent, GetGBufferFFormat(), FClearValueBinding({ 0.5f, 0.5f, 0.5f, 0.5f }), TexCreate_RenderTargetable | TexCreate_ShaderResource | FlagsToAdd | GFastVRamConfig.GBufferF);
 			SceneTextures.GBufferF = GraphBuilder.CreateTexture(Desc, TEXT("GBufferF"));
 		}
+
+		// Lighting features may require GBuffer history.
+		if (FRDGBuilder::IsDrainEnabled())
+		{
+			SceneTextures.GBufferA->SetNonTransient();
+			SceneTextures.GBufferB->SetNonTransient();
+			SceneTextures.GBufferC->SetNonTransient();
+		}
 	}
 
 
@@ -931,7 +944,7 @@ void FSceneTextureExtracts::QueueExtractions(FRDGBuilder& GraphBuilder, const FS
 	// We want these textures in a SRV Compute | Raster state.
 	const ERDGPassFlags PassFlags = ERDGPassFlags::Raster | ERDGPassFlags::SkipRenderPass | ERDGPassFlags::Compute | ERDGPassFlags::NeverCull;
 
-	GraphBuilder.AddPass({}, PassParameters, PassFlags,
+	GraphBuilder.AddPass(RDG_EVENT_NAME("ExtractUniformBuffer"), PassParameters, PassFlags,
 		[this, PassParameters, ShadingPath = SceneTextures.Config.ShadingPath](FRHICommandList&)
 	{
 		if (ShadingPath == EShadingPath::Deferred)
