Conditional GPU copies for RDG buffer uploads. This fixes broken GPU captures on some platforms when CPU-initialized GPU-modified buffers are used

diff --git a/Engine/Source/Runtime/Core/Public/HAL/Platform.h b/Engine/Source/Runtime/Core/Public/HAL/Platform.h
--- a/Engine/Source/Runtime/Core/Public/HAL/Platform.h
+++ b/Engine/Source/Runtime/Core/Public/HAL/Platform.h
@@ -619,6 +619,10 @@
 	#define PLATFORM_REQUIRES_UAV_TO_RTV_TEXTURE_CACHE_FLUSH_WORKAROUND 0
 #endif // #ifndef PLATFORM_REQUIRES_UAV_TO_RTV_TEXTURE_CACHE_FLUSH_WORKAROUND
 
+#ifndef PLATFORM_NEEDS_GPU_UAV_RESOURCE_INIT_WORKAROUND
+	#define PLATFORM_NEEDS_GPU_UAV_RESOURCE_INIT_WORKAROUND 0
+#endif
+
 // deprecated, do not use
 #define PLATFORM_HAS_THREADSAFE_RHIGetRenderQueryResult	#
 #define PLATFORM_SUPPORTS_RHI_THREAD #
diff --git a/Engine/Source/Runtime/RenderCore/Private/RenderGraphBuilder.cpp b/Engine/Source/Runtime/RenderCore/Private/RenderGraphBuilder.cpp
--- a/Engine/Source/Runtime/RenderCore/Private/RenderGraphBuilder.cpp
+++ b/Engine/Source/Runtime/RenderCore/Private/RenderGraphBuilder.cpp
@@ -2127,9 +2127,25 @@ void FRDGBuilder::SubmitBufferUploads()
 	{
 		if (UploadedBuffer.Data && UploadedBuffer.DataSize)
 		{
-			void* DestPtr = RHICmdList.LockBuffer(UploadedBuffer.Buffer->GetRHI(), 0, UploadedBuffer.DataSize, RLM_WriteOnly);
-			FMemory::Memcpy(DestPtr, UploadedBuffer.Data, UploadedBuffer.DataSize);
-			RHICmdList.UnlockBuffer(UploadedBuffer.Buffer->GetRHI());
+#if PLATFORM_NEEDS_GPU_UAV_RESOURCE_INIT_WORKAROUND
+			if (UploadedBuffer.Buffer->bUAVAccessed)
+			{
+				FRHIResourceCreateInfo CreateInfo(UploadedBuffer.Buffer->Name);
+				FBufferRHIRef TempBuffer = RHICreateVertexBuffer(UploadedBuffer.DataSize, BUF_Static | BUF_ShaderResource, CreateInfo);
+				void* DestPtr = RHICmdList.LockBuffer(TempBuffer, 0, UploadedBuffer.DataSize, RLM_WriteOnly);
+				FMemory::Memcpy(DestPtr, UploadedBuffer.Data, UploadedBuffer.DataSize);
+				RHICmdList.UnlockBuffer(TempBuffer);
+				RHICmdList.Transition(FRHITransitionInfo(TempBuffer, ERHIAccess::Unknown, ERHIAccess::CopySrc | ERHIAccess::SRVMask));
+				RHICmdList.Transition(FRHITransitionInfo(UploadedBuffer.Buffer->GetRHI(), ERHIAccess::Unknown, ERHIAccess::CopyDest | ERHIAccess::UAVMask));
+				RHICmdList.CopyBufferRegion(UploadedBuffer.Buffer->GetRHI(), 0, TempBuffer, 0, UploadedBuffer.DataSize);
+			}
+			else
+#endif
+			{
+				void* DestPtr = RHICmdList.LockBuffer(UploadedBuffer.Buffer->GetRHI(), 0, UploadedBuffer.DataSize, RLM_WriteOnly);
+				FMemory::Memcpy(DestPtr, UploadedBuffer.Data, UploadedBuffer.DataSize);
+				RHICmdList.UnlockBuffer(UploadedBuffer.Buffer->GetRHI());
+			}
 		}
 	}
 	UploadedBuffers.Reset();
diff --git a/Engine/Source/Runtime/RenderCore/Private/RenderGraphResources.cpp b/Engine/Source/Runtime/RenderCore/Private/RenderGraphResources.cpp
--- a/Engine/Source/Runtime/RenderCore/Private/RenderGraphResources.cpp
+++ b/Engine/Source/Runtime/RenderCore/Private/RenderGraphResources.cpp
@@ -35,6 +35,7 @@ FRDGParentResource::FRDGParentResource(const TCHAR* InName, const ERDGParentReso
 	, bQueuedForUpload(0)
 	, bSwapChain(0)
 	, bSwapChainAlreadyMoved(0)
+	, bUAVAccessed(0)
 {}
 
 bool FRDGProducerState::IsDependencyRequired(FRDGProducerState LastProducer, ERHIPipeline LastPipeline, FRDGProducerState NextState, ERHIPipeline NextPipeline)
diff --git a/Engine/Source/Runtime/RenderCore/Public/RenderGraphBuilder.inl b/Engine/Source/Runtime/RenderCore/Public/RenderGraphBuilder.inl
--- a/Engine/Source/Runtime/RenderCore/Public/RenderGraphBuilder.inl
+++ b/Engine/Source/Runtime/RenderCore/Public/RenderGraphBuilder.inl
@@ -89,6 +89,7 @@ inline FRDGTextureUAVRef FRDGBuilder::CreateUAV(const FRDGTextureUAVDesc& Desc,
 {
 	IF_RDG_ENABLE_DEBUG(UserValidation.ValidateCreateUAV(Desc));
 	FRDGTextureUAVRef UAV = Views.Allocate<FRDGTextureUAV>(Allocator, Desc.Texture->Name, Desc, InFlags);
+	Desc.Texture->bUAVAccessed = true;
 	IF_RDG_ENABLE_DEBUG(UserValidation.ValidateCreateUAV(UAV));
 	return UAV;
 }
@@ -97,6 +98,7 @@ inline FRDGBufferUAVRef FRDGBuilder::CreateUAV(const FRDGBufferUAVDesc& Desc, ER
 {
 	IF_RDG_ENABLE_DEBUG(UserValidation.ValidateCreateUAV(Desc));
 	FRDGBufferUAVRef UAV = Views.Allocate<FRDGBufferUAV>(Allocator, Desc.Buffer->Name, Desc, InFlags);
+	Desc.Buffer->bUAVAccessed = true;
 	IF_RDG_ENABLE_DEBUG(UserValidation.ValidateCreateUAV(UAV));
 	return UAV;
 }
diff --git a/Engine/Source/Runtime/RenderCore/Public/RenderGraphResources.h b/Engine/Source/Runtime/RenderCore/Public/RenderGraphResources.h
--- a/Engine/Source/Runtime/RenderCore/Public/RenderGraphResources.h
+++ b/Engine/Source/Runtime/RenderCore/Public/RenderGraphResources.h
@@ -302,6 +302,9 @@ protected:
 	/** If true, the swap chain transition has been moved to occur at the latest possible time. */
 	uint8 bSwapChainAlreadyMoved : 1;
 
+	/** If true, the resource is access through at least one UAV. */
+	uint8 bUAVAccessed : 1;
+
 	FRDGPassHandle FirstPass;
 	FRDGPassHandle LastPass;
 
