Fix recursing lock called during the release of the preloaded shaders.

diff --git a/Engine/Source/Runtime/RenderCore/Private/ShaderCodeLibrary.cpp b/Engine/Source/Runtime/RenderCore/Private/ShaderCodeLibrary.cpp
--- a/Engine/Source/Runtime/RenderCore/Private/ShaderCodeLibrary.cpp
+++ b/Engine/Source/Runtime/RenderCore/Private/ShaderCodeLibrary.cpp
@@ -941,8 +941,16 @@ public:
 
 	void ReleasePreloadedPackageShaderMap(int32 ShaderMapIndex)
 	{
-		FRWScopeLock Locker(ResourceLock, SLT_Write);
-		Resources[ShaderMapIndex]->Release();
+		FShaderMapResource_SharedCode* Resource = nullptr;
+		{
+			FRWScopeLock Locker(ResourceLock, SLT_Write);
+			Resource = Resources[ShaderMapIndex];
+		}
+
+		if (Resource)
+		{
+			Resource->Release();
+		}
 	}
 
 public:
