r.RDG.ParallelExecutePassMin temporary set to 0 to avoid a crash in ray tracing

diff --git a/Engine/Source/Runtime/RenderCore/Private/RenderGraphPrivate.cpp b/Engine/Source/Runtime/RenderCore/Private/RenderGraphPrivate.cpp
--- a/Engine/Source/Runtime/RenderCore/Private/RenderGraphPrivate.cpp
+++ b/Engine/Source/Runtime/RenderCore/Private/RenderGraphPrivate.cpp
@@ -306,7 +306,7 @@ FAutoConsoleVariableRef CVarRDGUseTransientAllocator(
 	TEXT(" 2: enables the transient allocator for resources with FastVRAM flag only"),
 	ECVF_RenderThreadSafe);
 
-int32 GRDGParallelExecutePassMin = 1;
+int32 GRDGParallelExecutePassMin = 0;
 FAutoConsoleVariableRef CVarRDGParallelExecutePassMin(
 	TEXT("r.RDG.ParallelExecutePassMin"), GRDGParallelExecutePassMin,
 	TEXT("The minimum span of contiguous passes eligible for parallel execution for the span to be offloaded to a task."),
