Fixed state merging when r.RDG.Drain is enabled.

diff --git a/Engine/Source/Runtime/RenderCore/Private/RenderGraphBuilder.cpp b/Engine/Source/Runtime/RenderCore/Private/RenderGraphBuilder.cpp
--- a/Engine/Source/Runtime/RenderCore/Private/RenderGraphBuilder.cpp
+++ b/Engine/Source/Runtime/RenderCore/Private/RenderGraphBuilder.cpp
@@ -1899,6 +1899,7 @@ FRDGPass* FRDGBuilder::SetupPass(FRDGPass* Pass)
 
 		const auto AddSubresourceAccess = [&](FRDGSubresourceState& State)
 		{
+			State.ExecuteGeneration = ExecuteGeneration;
 			State.Access = MakeValidAccess(State.Access | Access);
 			State.Flags |= TransitionFlags;
 			State.NoUAVBarrierFilter.AddHandle(NoUAVBarrierHandle);
@@ -1951,6 +1952,7 @@ FRDGPass* FRDGBuilder::SetupPass(FRDGPass* Pass)
 		}
 
 		PassState->ReferenceCount++;
+		PassState->State.ExecuteGeneration = ExecuteGeneration;
 		PassState->State.Access = MakeValidAccess(PassState->State.Access | Access);
 		PassState->State.NoUAVBarrierFilter.AddHandle(NoUAVBarrierHandle);
 		PassState->State.SetPass(PassPipeline, PassHandle);
