Added r.RDG.ParallelExecute console variable to control parallel RDG execution for simplicity.

diff --git a/Engine/Source/Runtime/RenderCore/Private/RenderGraphPrivate.cpp b/Engine/Source/Runtime/RenderCore/Private/RenderGraphPrivate.cpp
--- a/Engine/Source/Runtime/RenderCore/Private/RenderGraphPrivate.cpp
+++ b/Engine/Source/Runtime/RenderCore/Private/RenderGraphPrivate.cpp
@@ -306,15 +306,38 @@ FAutoConsoleVariableRef CVarRDGUseTransientAllocator(
 	TEXT(" 2: enables the transient allocator for resources with FastVRAM flag only"),
 	ECVF_RenderThreadSafe);
 
+int32 GRDGParallelExecute = 1;
+FAutoConsoleVariableRef CVarRDGParallelExecute(
+	TEXT("r.RDG.ParallelExecute"), GRDGParallelExecute,
+	TEXT("Whether to enable parallel execution of passes when supported.")
+	TEXT(" 0: off;")
+	TEXT(" 1: on (default)"),
+	FConsoleVariableDelegate::CreateLambda([](IConsoleVariable* Variable)
+	{
+		if (Variable->GetInt())
+		{
+			if (GRDGParallelExecutePassMax <= 1)
+			{
+				GRDGParallelExecutePassMax = 1;
+			}
+
+			if (GRDGParallelExecutePassMin < GRDGParallelExecutePassMax)
+			{
+				GRDGParallelExecutePassMin = GRDGParallelExecutePassMax;
+			}
+		}
+	}),
+	ECVF_RenderThreadSafe);
+
 int32 GRDGParallelExecutePassMin = 1;
 FAutoConsoleVariableRef CVarRDGParallelExecutePassMin(
-	TEXT("r.RDG.ParallelExecutePassMin"), GRDGParallelExecutePassMin,
+	TEXT("r.RDG.ParallelExecute.PassMin"), GRDGParallelExecutePassMin,
 	TEXT("The minimum span of contiguous passes eligible for parallel execution for the span to be offloaded to a task."),
 	ECVF_RenderThreadSafe);
 
 int32 GRDGParallelExecutePassMax = 32;
 FAutoConsoleVariableRef CVarRDGParallelExecutePassMax(
-	TEXT("r.RDG.ParallelExecutePassMax"), GRDGParallelExecutePassMax,
+	TEXT("r.RDG.ParallelExecute.PassMax"), GRDGParallelExecutePassMax,
 	TEXT("The maximum span of contiguous passes eligible for parallel execution for the span to be offloaded to a task."),
 	ECVF_RenderThreadSafe);
 
diff --git a/Engine/Source/Runtime/RenderCore/Private/RenderGraphPrivate.h b/Engine/Source/Runtime/RenderCore/Private/RenderGraphPrivate.h
--- a/Engine/Source/Runtime/RenderCore/Private/RenderGraphPrivate.h
+++ b/Engine/Source/Runtime/RenderCore/Private/RenderGraphPrivate.h
@@ -109,6 +109,7 @@ extern int32 GRDGMergeRenderPasses;
 extern int32 GRDGTransientAllocator;
 extern int32 GRDGTransientIndirectArgBuffers;
 extern int32 GRDGDrain;
+extern int32 GRDGParallelExecute;
 extern int32 GRDGParallelExecutePassMin;
 extern int32 GRDGParallelExecutePassMax;
 
@@ -197,7 +198,7 @@ FORCEINLINE bool IsRenderPassMergeEnabled()
 
 FORCEINLINE bool IsParallelExecuteEnabled()
 {
-	return GRDGParallelExecutePassMin > 0
+	return GRDGParallelExecute > 0
 		&& !GRHICommandList.Bypass()
 		&& !IsImmediateMode()
 		&& !GRDGDebug
