Don't allow creation of presentable textures through RDG

diff --git a/Engine/Source/Runtime/RenderCore/Private/RenderGraphValidation.cpp b/Engine/Source/Runtime/RenderCore/Private/RenderGraphValidation.cpp
--- a/Engine/Source/Runtime/RenderCore/Private/RenderGraphValidation.cpp
+++ b/Engine/Source/Runtime/RenderCore/Private/RenderGraphValidation.cpp
@@ -272,6 +272,9 @@ void FRDGUserValidation::ValidateCreateTexture(const FRDGTextureDesc& Desc, cons
 		TEXT("Failed to create texture %s with pixel format %s because it is not supported."), Name, GPixelFormats[Desc.Format].Name);
 	checkf(Desc.IsValid(), TEXT("Texture %s was created with an invalid descriptor."), Name);
 
+	// Can't create back buffer textures
+	checkf(!EnumHasAnyFlags(Desc.Flags, ETextureCreateFlags::Presentable), TEXT("Illegal to create texture %s with presentable flag."), Name);
+
 	const bool bCanHaveUAV = EnumHasAnyFlags(Desc.Flags, TexCreate_UAV);
 	const bool bIsMSAA = Desc.NumSamples > 1;
 
diff --git a/Engine/Source/Runtime/Renderer/Private/PostProcess/PostProcessDownsample.cpp b/Engine/Source/Runtime/Renderer/Private/PostProcess/PostProcessDownsample.cpp
--- a/Engine/Source/Runtime/Renderer/Private/PostProcess/PostProcessDownsample.cpp
+++ b/Engine/Source/Runtime/Renderer/Private/PostProcess/PostProcessDownsample.cpp
@@ -125,7 +125,7 @@ FScreenPassTexture AddDownsamplePass(
 		Desc.Extent = FIntPoint::DivideAndRoundUp(Desc.Extent, 2);
 		Desc.Extent.X = FMath::Max(1, Desc.Extent.X);
 		Desc.Extent.Y = FMath::Max(1, Desc.Extent.Y);
-		Desc.Flags &= ~(TexCreate_RenderTargetable | TexCreate_UAV);
+		Desc.Flags &= ~(TexCreate_RenderTargetable | TexCreate_UAV | TexCreate_Presentable);
 		Desc.Flags |= bIsComputePass ? TexCreate_UAV : TexCreate_RenderTargetable;
 		Desc.Flags |= GFastVRamConfig.Downsample;
 		Desc.ClearValue = FClearValueBinding(FLinearColor(0, 0, 0, 0));
