Disable Parallel RDG for GL

diff --git a/Engine/Source/Runtime/RenderCore/Private/RenderGraphPrivate.h b/Engine/Source/Runtime/RenderCore/Private/RenderGraphPrivate.h
--- a/Engine/Source/Runtime/RenderCore/Private/RenderGraphPrivate.h
+++ b/Engine/Source/Runtime/RenderCore/Private/RenderGraphPrivate.h
@@ -202,7 +202,8 @@ FORCEINLINE bool IsParallelExecuteEnabled()
 		&& !GRHICommandList.Bypass()
 		&& !IsImmediateMode()
 		&& !GRDGDebug
-		&& !GRDGTransitionLog;
+		&& !GRDGTransitionLog
+		&& !IsOpenGLPlatform(GMaxRHIShaderPlatform);
 }
 
 template <typename ResourceRegistryType, typename FunctionType>
