Add copy buffer pass to RDG utility functions.

diff --git a/Engine/Source/Runtime/RenderCore/Private/RenderGraphUtils.cpp b/Engine/Source/Runtime/RenderCore/Private/RenderGraphUtils.cpp
--- a/Engine/Source/Runtime/RenderCore/Private/RenderGraphUtils.cpp
+++ b/Engine/Source/Runtime/RenderCore/Private/RenderGraphUtils.cpp
@@ -293,6 +293,31 @@ void AddCopyToResolveTargetPass(
 	});
 }
 
+BEGIN_SHADER_PARAMETER_STRUCT(FCopyBufferParameters, )
+	RDG_BUFFER_ACCESS(SrcBuffer, ERHIAccess::CopySrc)
+	RDG_BUFFER_ACCESS(DstBuffer, ERHIAccess::CopyDest)
+END_SHADER_PARAMETER_STRUCT()
+
+void AddCopyBufferPass(FRDGBuilder& GraphBuilder, FRDGBufferRef DstBuffer, FRDGBufferRef SrcBuffer)
+{
+	check(SrcBuffer);
+	check(DstBuffer);
+
+	FCopyBufferParameters* Parameters = GraphBuilder.AllocParameters<FCopyBufferParameters>();
+	Parameters->SrcBuffer = SrcBuffer;
+	Parameters->DstBuffer = DstBuffer;
+	const uint64 NumBytes = Parameters->SrcBuffer->Desc.NumElements * Parameters->SrcBuffer->Desc.BytesPerElement;
+
+	GraphBuilder.AddPass(
+		RDG_EVENT_NAME("CopyBuffer(%s Size=%ubytes)", SrcBuffer->Name, SrcBuffer->Desc.GetTotalNumBytes()),
+		Parameters,
+		ERDGPassFlags::Copy,
+		[&Parameters, SrcBuffer, DstBuffer, NumBytes](FRHICommandList& RHICmdList)
+		{
+			RHICmdList.CopyBufferRegion(DstBuffer->GetRHI(), 0, SrcBuffer->GetRHI(), 0, NumBytes);
+		});
+}
+
 BEGIN_SHADER_PARAMETER_STRUCT(FClearBufferUAVParameters, )
 	SHADER_PARAMETER_RDG_BUFFER_UAV(RWBuffer<uint>, BufferUAV)
 END_SHADER_PARAMETER_STRUCT()
diff --git a/Engine/Source/Runtime/RenderCore/Public/RenderGraphUtils.h b/Engine/Source/Runtime/RenderCore/Public/RenderGraphUtils.h
--- a/Engine/Source/Runtime/RenderCore/Public/RenderGraphUtils.h
+++ b/Engine/Source/Runtime/RenderCore/Public/RenderGraphUtils.h
@@ -675,6 +675,9 @@ RENDERCORE_API void AddClearStencilPass(FRDGBuilder& GraphBuilder, FRDGTextureRe
 /** Adds a render graph pass to resummarize the htile plane. */
 RENDERCORE_API void AddResummarizeHTilePass(FRDGBuilder& GraphBuilder, FRDGTextureRef DepthTexture);
 
+/** Adds a render graph pass to copy SrcBuffer content into DstBuffer. */
+RENDERCORE_API void AddCopyBufferPass(FRDGBuilder& GraphBuilder, FRDGBufferRef DstBuffer, FRDGBufferRef SrcBuffer);
+
 /** Adds a pass to readback contents of an RDG texture. */
 RENDERCORE_API void AddEnqueueCopyPass(FRDGBuilder& GraphBuilder, FRHIGPUTextureReadback* Readback, FRDGTextureRef SourceTexture, FResolveRect Rect = FResolveRect());
 
