Fixed bug with RDG drain and async compute. Fixed transition bug with UAV workaround in upload buffers. Added some command line arguments. Optimiized checks to disable RDG validation if parallel.

diff --git a/Engine/Source/Runtime/RenderCore/Private/RenderGraphBuilder.cpp b/Engine/Source/Runtime/RenderCore/Private/RenderGraphBuilder.cpp
--- a/Engine/Source/Runtime/RenderCore/Private/RenderGraphBuilder.cpp
+++ b/Engine/Source/Runtime/RenderCore/Private/RenderGraphBuilder.cpp
@@ -537,7 +537,7 @@ FRDGBuilder::FRDGBuilder(FRHICommandListImmediate& InRHICmdList, FRDGEventName I
 #endif
 	, bParallelExecuteEnabled(IsParallelExecuteEnabled() && EnumHasAnyFlags(InFlags, ERDGBuilderFlags::AllowParallelExecute))
 #if RDG_ENABLE_DEBUG
-	, UserValidation(Allocator)
+	, UserValidation(Allocator, bParallelExecuteEnabled)
 	, BarrierValidation(&Passes, BuilderName)
 	, LogFile(Passes)
 #endif
@@ -1677,7 +1677,7 @@ void FRDGBuilder::Execute(EExecuteMode ExecuteMode)
 		IF_RDG_ENABLE_TRACE(Trace.OutputGraphBegin());
 	}
 
-	IF_RDG_ENABLE_DEBUG(GRDGAllowRHIAccess = false);
+	IF_RDG_ENABLE_DEBUG(GRDGAllowRHIAccess = bParallelExecuteEnabled);
 
 	const ENamedThreads::Type RenderThread = ENamedThreads::GetRenderThread_Local();
 
@@ -1819,6 +1819,7 @@ void FRDGBuilder::Execute(EExecuteMode ExecuteMode)
 
 		IF_RDG_ENABLE_DEBUG(UserValidation.ValidateExecuteEnd());
 		IF_RDG_ENABLE_DEBUG(LogFile.End());
+		IF_RDG_ENABLE_DEBUG(GRDGAllowRHIAccess = false);
 
 	#if STATS
 		GRDGStatBufferCount += Buffers.Num();
@@ -1832,6 +1833,8 @@ void FRDGBuilder::Execute(EExecuteMode ExecuteMode)
 		AddProloguePass();
 	}
 
+	RasterPassCount = 0;
+	AsyncComputePassCount = 0;
 	ExecuteGeneration++;
 
 	// Flush any outstanding async compute commands at the end to get things moving down the pipe.
@@ -2142,8 +2145,11 @@ void FRDGBuilder::SubmitBufferUploads()
 				void* DestPtr = RHICmdList.LockBuffer(TempBuffer, 0, UploadedBuffer.DataSize, RLM_WriteOnly);
 				FMemory::Memcpy(DestPtr, UploadedBuffer.Data, UploadedBuffer.DataSize);
 				RHICmdList.UnlockBuffer(TempBuffer);
-				RHICmdList.Transition(FRHITransitionInfo(TempBuffer, ERHIAccess::Unknown, ERHIAccess::CopySrc | ERHIAccess::SRVMask));
-				RHICmdList.Transition(FRHITransitionInfo(UploadedBuffer.Buffer->GetRHI(), ERHIAccess::Unknown, ERHIAccess::CopyDest | ERHIAccess::UAVMask));
+				RHICmdList.Transition(
+				{
+					FRHITransitionInfo(TempBuffer, ERHIAccess::Unknown, ERHIAccess::CopySrc | ERHIAccess::SRVMask),
+					FRHITransitionInfo(UploadedBuffer.Buffer->GetRHI(), ERHIAccess::Unknown, ERHIAccess::CopyDest)
+				});
 				RHICmdList.CopyBufferRegion(UploadedBuffer.Buffer->GetRHI(), 0, TempBuffer, 0, UploadedBuffer.DataSize);
 			}
 			else
@@ -2277,12 +2283,13 @@ void FRDGBuilder::SetupParallelExecute()
 		}
 	}
 
+	ParallelPassCandidates.Emplace(EpiloguePass);
 	FlushParallelPassCandidates();
 }
 
 void FRDGBuilder::DispatchParallelExecute(IRHICommandContext* RHICmdContext)
 {
-	SCOPED_NAMED_EVENT(SetupParallelExecute, FColor::Emerald);
+	SCOPED_NAMED_EVENT(DispatchParallelExecute, FColor::Emerald);
 	ParallelExecuteEvents.Reserve(ParallelExecuteEvents.Num() + ParallelPassSets.Num());
 
 	for (int32 Index = ParallelPassSetsOffset; Index < ParallelPassSets.Num(); ++Index)
@@ -2334,7 +2341,7 @@ void FRDGBuilder::CreateUniformBuffers()
 
 void FRDGBuilder::AddProloguePass()
 {
-	if (!PassesOnAsyncComputeToJoin.IsEmpty())
+	if (AsyncComputePassCount > 0)
 	{
 		FRDGPass* AsyncComputePass = SetupEmptyPass(Passes.Allocate<FRDGSentinelPass>(Allocator, RDG_EVENT_NAME("Graph Prologue (AsyncCompute)"), ERDGPassFlags::AsyncCompute));
 		ProloguePasses[ERHIPipeline::AsyncCompute] = AsyncComputePass;
@@ -2346,7 +2353,7 @@ void FRDGBuilder::AddProloguePass()
 	ProloguePassHandles[ERHIPipeline::Graphics] = GraphicsPass->Handle;
 
 	// If we don't need an async compute prologue pass, we still want a valid handle to make ClampToPrologue not always return null.
-	if (PassesOnAsyncComputeToJoin.IsEmpty())
+	if (!AsyncComputePassCount)
 	{
 		ProloguePassHandles[ERHIPipeline::AsyncCompute] = GraphicsPass->Handle;
 	}
diff --git a/Engine/Source/Runtime/RenderCore/Private/RenderGraphPrivate.cpp b/Engine/Source/Runtime/RenderCore/Private/RenderGraphPrivate.cpp
--- a/Engine/Source/Runtime/RenderCore/Private/RenderGraphPrivate.cpp
+++ b/Engine/Source/Runtime/RenderCore/Private/RenderGraphPrivate.cpp
@@ -488,12 +488,30 @@ void InitRenderGraph()
 		GRDGClobberResources = 1;
 	}
 
+	int32 TransientAllocatorValue = 0;
+	if (FParse::Value(FCommandLine::Get(), TEXT("rdgtransientallocator="), TransientAllocatorValue))
+	{
+		GRDGTransientAllocator = TransientAllocatorValue;
+	}
+
 	int32 CullPassesValue = 0;
 	if (FParse::Value(FCommandLine::Get(), TEXT("rdgcullpasses="), CullPassesValue))
 	{
 		GRDGCullPasses = CullPassesValue;
 	}
 
+	int32 DrainValue = 0;
+	if (FParse::Value(FCommandLine::Get(), TEXT("rdgdrain="), DrainValue))
+	{
+		GRDGDrain = DrainValue;
+	}
+
+	int32 ParallelExecuteValue = 0;
+	if (FParse::Value(FCommandLine::Get(), TEXT("rdgparallelxecute="), ParallelExecuteValue))
+	{
+		GRDGParallelExecute = ParallelExecuteValue;
+	}
+
 	int32 MergeRenderPassesValue = 0;
 	if (FParse::Value(FCommandLine::Get(), TEXT("rdgmergerenderpasses="), MergeRenderPassesValue))
 	{
diff --git a/Engine/Source/Runtime/RenderCore/Private/RenderGraphResources.cpp b/Engine/Source/Runtime/RenderCore/Private/RenderGraphResources.cpp
--- a/Engine/Source/Runtime/RenderCore/Private/RenderGraphResources.cpp
+++ b/Engine/Source/Runtime/RenderCore/Private/RenderGraphResources.cpp
@@ -304,6 +304,14 @@ void FRDGBuffer::SetRHI(FRDGPooledBuffer* InPooledBuffer)
 	State = &PooledBuffer->State;
 	ViewCache = &PooledBuffer->ViewCache;
 	ResourceRHI = InPooledBuffer->GetRHI();
+
+	// The upload with UAV workaround performs its own transitions outside of RDG, so fall back to Unknown for simplicity.
+#if PLATFORM_NEEDS_GPU_UAV_RESOURCE_INIT_WORKAROUND
+	if (bUAVAccessed && bQueuedForUpload)
+	{
+		*State = FRDGSubresourceState();
+	}
+#endif
 }
 
 void FRDGBuffer::SetRHI(FRHITransientBuffer* InTransientBuffer, FRDGAllocator& Allocator)
diff --git a/Engine/Source/Runtime/RenderCore/Private/RenderGraphValidation.cpp b/Engine/Source/Runtime/RenderCore/Private/RenderGraphValidation.cpp
--- a/Engine/Source/Runtime/RenderCore/Private/RenderGraphValidation.cpp
+++ b/Engine/Source/Runtime/RenderCore/Private/RenderGraphValidation.cpp
@@ -79,7 +79,7 @@ void FRDGResource::MarkResourceAsUsed()
 void FRDGResource::ValidateRHIAccess() const
 {
 	check(DebugData);
-	checkf(DebugData->bAllowRHIAccess || GRDGAllowRHIAccess || IsParallelExecuteEnabled(),
+	checkf(DebugData->bAllowRHIAccess || GRDGAllowRHIAccess,
 		TEXT("Accessing the RHI resource of %s at this time is not allowed. If you hit this check in pass, ")
 		TEXT("that is due to this resource not being referenced in the parameters of your pass."),
 		Name);
@@ -153,8 +153,9 @@ void FRDGUniformBuffer::MarkResourceAsUsed()
 	});
 }
 
-FRDGUserValidation::FRDGUserValidation(FRDGAllocator& InAllocator)
+FRDGUserValidation::FRDGUserValidation(FRDGAllocator& InAllocator, bool bInParallelExecuteEnabled)
 	: Allocator(InAllocator)
+	, bParallelExecuteEnabled(bInParallelExecuteEnabled)
 {
 	checkf(!GRDGBuilderActive, TEXT("Another FRDGBuilder already exists on the stack. Only one builder can be created at a time. This builder instance should be merged into the parent one."));
 	GRDGBuilderActive = true;
@@ -912,7 +913,7 @@ void FRDGUserValidation::ValidateExecuteEnd()
 
 void FRDGUserValidation::ValidateExecutePassBegin(const FRDGPass* Pass)
 {
-	if (IsParallelExecuteEnabled())
+	if (bParallelExecuteEnabled)
 	{
 		return;
 	}
@@ -1012,7 +1013,7 @@ void FRDGUserValidation::ValidateExecutePassBegin(const FRDGPass* Pass)
 
 void FRDGUserValidation::ValidateExecutePassEnd(const FRDGPass* Pass)
 {
-	if (IsParallelExecuteEnabled())
+	if (bParallelExecuteEnabled)
 	{
 		return;
 	}
diff --git a/Engine/Source/Runtime/RenderCore/Public/RenderGraphAllocator.h b/Engine/Source/Runtime/RenderCore/Public/RenderGraphAllocator.h
--- a/Engine/Source/Runtime/RenderCore/Public/RenderGraphAllocator.h
+++ b/Engine/Source/Runtime/RenderCore/Public/RenderGraphAllocator.h
@@ -77,6 +77,7 @@ private:
 	};
 
 	FContext Context;
+	uint8 Pad[PLATFORM_CACHE_LINE_SIZE];
 	FContext ContextForTasks;
 
 	FORCEINLINE FContext& GetContext()
diff --git a/Engine/Source/Runtime/RenderCore/Public/RenderGraphValidation.h b/Engine/Source/Runtime/RenderCore/Public/RenderGraphValidation.h
--- a/Engine/Source/Runtime/RenderCore/Public/RenderGraphValidation.h
+++ b/Engine/Source/Runtime/RenderCore/Public/RenderGraphValidation.h
@@ -23,7 +23,7 @@
 class RENDERCORE_API FRDGUserValidation final
 {
 public:
-	FRDGUserValidation(FRDGAllocator& Allocator);
+	FRDGUserValidation(FRDGAllocator& Allocator, bool bParallelExecuteEnabled);
 	FRDGUserValidation(const FRDGUserValidation&) = delete;
 	~FRDGUserValidation();
 
@@ -107,6 +107,8 @@ private:
 	void ValidateCreateResource(FRDGResourceRef Resource);
 	void ValidateExtractResource(FRDGParentResourceRef Resource);
 
+	FRDGAllocator& Allocator;
+
 	/** List of tracked resources for validation prior to shutdown. */
 	TArray<FRDGTextureRef, FRDGArrayAllocator> TrackedTextures;
 	TArray<FRDGBufferRef, FRDGArrayAllocator> TrackedBuffers;
@@ -118,9 +120,9 @@ private:
 	bool bHasExecuted = false;
 	bool bHasExecuteBegun = false;
 
-	void ExecuteGuard(const TCHAR* Operation, const TCHAR* ResourceName);
+	bool bParallelExecuteEnabled;
 
-	FRDGAllocator& Allocator;
+	void ExecuteGuard(const TCHAR* Operation, const TCHAR* ResourceName);
 };
 
 /** This class validates and logs barriers submitted by the graph. */
