Add RenderCaptureInterface support for RDG deferred work.

diff --git a/Engine/Source/Runtime/RenderCore/Private/RenderCaptureInterface.cpp b/Engine/Source/Runtime/RenderCore/Private/RenderCaptureInterface.cpp
--- a/Engine/Source/Runtime/RenderCore/Private/RenderCaptureInterface.cpp
+++ b/Engine/Source/Runtime/RenderCore/Private/RenderCaptureInterface.cpp
@@ -2,6 +2,7 @@
 
 #include "RenderCaptureInterface.h"
 #include "IRenderCaptureProvider.h"
+#include "RenderGraphBuilder.h"
 #include "RenderingThread.h"
 #include "RHI.h"
 
@@ -9,21 +10,21 @@ namespace RenderCaptureInterface
 {
 	FScopedCapture::FScopedCapture(bool bEnable, TCHAR const* InEventName, TCHAR const* InFileName)
 		: bCapture(bEnable && IRenderCaptureProvider::IsAvailable())
-		, bEvent(false)
+		, bEvent(InEventName != nullptr)
 		, RHICommandList(nullptr)
+		, GraphBuilder(nullptr)
 	{
 		check(!GIsThreadedRendering || !IsInRenderingThread());
 
 		if (bCapture)
 		{
-			ENQUEUE_RENDER_COMMAND(BeginCaptureCommand)([this, EventName = FString(InEventName), FileName = FString(InFileName)](FRHICommandListImmediate& RHICommandListLocal)
+			ENQUEUE_RENDER_COMMAND(BeginCaptureCommand)([bPushEvent = bEvent, EventName = FString(InEventName), FileName = FString(InFileName)](FRHICommandListImmediate& RHICommandListLocal)
 			{
 				IRenderCaptureProvider::Get().BeginCapture(&RHICommandListLocal, IRenderCaptureProvider::ECaptureFlags_Launch, FileName);
 
-				if (!EventName.IsEmpty())
+				if (bPushEvent)
 				{
 					RHICommandListLocal.PushEvent(*EventName, FColor::White);
-					bEvent = true;
 				}
 			});
 		}
@@ -31,7 +32,9 @@ namespace RenderCaptureInterface
 
 	FScopedCapture::FScopedCapture(bool bEnable, FRHICommandListImmediate* InRHICommandList, TCHAR const* InEventName, TCHAR const* InFileName)
 		: bCapture(bEnable && IRenderCaptureProvider::IsAvailable())
+		, bEvent(InEventName != nullptr)
 		, RHICommandList(InRHICommandList)
+		, GraphBuilder(nullptr)
 	{
 		check(!GIsThreadedRendering || IsInRenderingThread());
 
@@ -39,10 +42,34 @@ namespace RenderCaptureInterface
 		{
 			IRenderCaptureProvider::Get().BeginCapture(RHICommandList, IRenderCaptureProvider::ECaptureFlags_Launch, FString(InFileName));
 		
-			if (InEventName != nullptr)
+			if (bEvent)
 			{
 				RHICommandList->PushEvent(InEventName, FColor::White);
-				bEvent = true;
+			}
+		}
+	}
+
+	FScopedCapture::FScopedCapture(bool bEnable, FRDGBuilder& InGraphBuilder, TCHAR const* InEventName, TCHAR const* InFileName)
+		: bCapture(bEnable&& IRenderCaptureProvider::IsAvailable())
+		, bEvent(InEventName != nullptr)
+		, RHICommandList(nullptr)
+		, GraphBuilder(&InGraphBuilder) 
+	{
+		check(!GIsThreadedRendering || IsInRenderingThread());
+
+		if (bCapture)
+		{
+			GraphBuilder->AddPass(
+				RDG_EVENT_NAME("BeginCapture"),
+				ERDGPassFlags::None, 
+				[FileName = FString(InFileName)](FRHICommandListImmediate& RHICommandListLocal)
+				{
+					IRenderCaptureProvider::Get().BeginCapture(&RHICommandListLocal, IRenderCaptureProvider::ECaptureFlags_Launch, FString(FileName));
+				});
+
+			if (bEvent)
+			{
+				GraphBuilder->BeginEventScope(RDG_EVENT_NAME("%s", InEventName));
 			}
 		}
 	}
@@ -51,7 +78,24 @@ namespace RenderCaptureInterface
 	{
 		if (bCapture)
 		{
-			if (RHICommandList != nullptr)
+			if (GraphBuilder != nullptr)
+			{
+				check(!GIsThreadedRendering || IsInRenderingThread());
+
+				if (bEvent)
+				{
+					GraphBuilder->EndEventScope();
+				}
+
+				GraphBuilder->AddPass(
+					RDG_EVENT_NAME("EndCapture"), 
+					ERDGPassFlags::None, 
+					[](FRHICommandListImmediate& RHICommandListLocal)
+					{
+						IRenderCaptureProvider::Get().EndCapture(&RHICommandListLocal);
+					});
+			}
+			else if (RHICommandList != nullptr)
 			{
 				check(!GIsThreadedRendering || IsInRenderingThread());
 				
diff --git a/Engine/Source/Runtime/RenderCore/Public/RenderCaptureInterface.h b/Engine/Source/Runtime/RenderCore/Public/RenderCaptureInterface.h
--- a/Engine/Source/Runtime/RenderCore/Public/RenderCaptureInterface.h
+++ b/Engine/Source/Runtime/RenderCore/Public/RenderCaptureInterface.h
@@ -4,6 +4,7 @@
 
 #include "CoreMinimal.h"
 
+class FRDGBuilder;
 class FRHICommandListImmediate;
 
 /** Easy to use interface for IRenderCaptureProvider. */
@@ -20,12 +21,15 @@ namespace RenderCaptureInterface
 		RENDERCORE_API FScopedCapture(bool bEnable, TCHAR const* InEventName = nullptr, TCHAR const* InFileName = nullptr);
 		/** Use this constructor if on rendering thread. Use bEnable to allow control over the capture frequency. */
 		RENDERCORE_API FScopedCapture(bool bEnable, FRHICommandListImmediate* InRHICommandList, TCHAR const* InEventName = nullptr, TCHAR const* InFileName = nullptr);
-		
+		/** Use this constructor if using RenderGraph to schedule work. Use bEnable to allow control over the capture frequency. */
+		RENDERCORE_API FScopedCapture(bool bEnable, FRDGBuilder& InGraphBuilder, TCHAR const* InEventName = nullptr, TCHAR const* InFileName = nullptr);
+
 		RENDERCORE_API ~FScopedCapture();
 	
 	private:
 		bool bCapture;
 		bool bEvent;
 		FRHICommandListImmediate* RHICommandList;
+		FRDGBuilder* GraphBuilder;
 	};
 }
