Cleanup UnifiedBuffer/ByteBuffer:
- Unified FMemcpyCS and FScatterCopyCS to handle both buffers and textures.
- Simplified FByteBufferShader::FPermutationDomain.

diff --git a/Engine/Shaders/Private/ByteBuffer.usf b/Engine/Shaders/Private/ByteBuffer.usf
--- a/Engine/Shaders/Private/ByteBuffer.usf
+++ b/Engine/Shaders/Private/ByteBuffer.usf
@@ -2,6 +2,11 @@
 
 #include "Common.ush"
 
+#define RESOURCE_TYPE_FLOAT4_STRUCTURED_BUFFER	(0)
+#define RESOURCE_TYPE_UINT_BUFFER				(1)
+#define RESOURCE_TYPE_UINT4_ALIGNED_BUFFER		(2)
+#define RESOURCE_TYPE_FLOAT4_TEXTURE			(3)
+
 uint Value;
 uint Size;
 uint Float4sPerLine;
@@ -9,34 +14,33 @@ uint NumScatters;
 uint SrcOffset;
 uint DstOffset;
 
-#if FLOAT4_BUFFER
-	StructuredBuffer< uint >		ScatterStructuredBuffer;
-	StructuredBuffer< float4 >		SrcStructuredBuffer;
-	StructuredBuffer< float4 >		UploadStructuredBuffer;
-	RWStructuredBuffer< float4 >	DstStructuredBuffer;
-#elif FLOAT4_TEXTURE
-	StructuredBuffer< uint >		ScatterStructuredBuffer;
+#if RESOURCE_TYPE == RESOURCE_TYPE_FLOAT4_STRUCTURED_BUFFER
+	StructuredBuffer<uint>			ScatterStructuredBuffer;
+	StructuredBuffer<float4>		SrcStructuredBuffer;
+	StructuredBuffer<float4>		UploadStructuredBuffer;
+	RWStructuredBuffer<float4>		DstStructuredBuffer;
+#elif RESOURCE_TYPE == RESOURCE_TYPE_FLOAT4_TEXTURE
+	StructuredBuffer<uint>			ScatterStructuredBuffer;
 	Texture2D<float4>				SrcTexture;
-	StructuredBuffer< float4 >		UploadStructuredBuffer;
+	StructuredBuffer<float4>		UploadStructuredBuffer;
 	RWTexture2D<float4>				DstTexture;
-#else
+#elif RESOURCE_TYPE == RESOURCE_TYPE_UINT_BUFFER || RESOURCE_TYPE == RESOURCE_TYPE_UINT4_ALIGNED_BUFFER
 	ByteAddressBuffer				ScatterByteAddressBuffer;
 	ByteAddressBuffer				SrcByteAddressBuffer;
 	ByteAddressBuffer				UploadByteAddressBuffer;
 	RWByteAddressBuffer				DstByteAddressBuffer;
 #endif
 
-#if FLOAT4_TEXTURE == false
 [numthreads(64, 1, 1)]
 void MemsetBufferCS( uint ThreadId : SV_DispatchThreadID ) 
 {
-#if FLOAT4_BUFFER
+#if RESOURCE_TYPE == RESOURCE_TYPE_FLOAT4_STRUCTURED_BUFFER
 	// Size is in float4s
 	if( ThreadId < Size )
 	{
 		DstStructuredBuffer[ DstOffset + ThreadId ] = asfloat( Value );
 	}
-#else
+#elif RESOURCE_TYPE == RESOURCE_TYPE_UINT_BUFFER || RESOURCE_TYPE == RESOURCE_TYPE_UINT4_ALIGNED_BUFFER
 	// Size and offset is in dwords
 	uint DstIndex = DstOffset + ThreadId * 4;
 
@@ -56,20 +60,30 @@ void MemsetBufferCS( uint ThreadId : SV_DispatchThreadID )
 	{
 		DstByteAddressBuffer.Store( DstIndex * 4, Value );
 	}
+#elif RESOURCE_TYPE == RESOURCE_TYPE_FLOAT4_TEXTURE
+	uint2 IndexTexture;
+	IndexTexture.y = ThreadId / (Float4sPerLine);
+	IndexTexture.x = ThreadId % (Float4sPerLine);
+
+	if (ThreadId < Size)
+	{
+		DstTexture[IndexTexture.xy] = asfloat(Value);
+	}
+#else
+	#error "Not implemented"
 #endif
 }
-#endif
 
 [numthreads(64, 1, 1)]
 void MemcpyCS( uint ThreadId : SV_DispatchThreadID ) 
 {
-#if FLOAT4_BUFFER
+#if RESOURCE_TYPE == RESOURCE_TYPE_FLOAT4_STRUCTURED_BUFFER
 	// Size is in float4s
 	if( ThreadId < Size )
 	{
 		DstStructuredBuffer[ DstOffset + ThreadId ] = SrcStructuredBuffer[ SrcOffset + ThreadId ];
 	}
-#elif FLOAT4_TEXTURE
+#elif RESOURCE_TYPE == RESOURCE_TYPE_FLOAT4_TEXTURE
 	uint2 IndexTexture;
 	IndexTexture.y = ThreadId / (Float4sPerLine);
 	IndexTexture.x = ThreadId % (Float4sPerLine);
@@ -79,7 +93,7 @@ void MemcpyCS( uint ThreadId : SV_DispatchThreadID )
 		float4 SrcValue = SrcTexture.Load(float3(IndexTexture.x, IndexTexture.y, 0));
 		DstTexture[IndexTexture.xy] = SrcValue;
 	}
-#else
+#elif RESOURCE_TYPE == RESOURCE_TYPE_UINT_BUFFER || RESOURCE_TYPE == RESOURCE_TYPE_UINT4_ALIGNED_BUFFER
 	// Size and offsets are in dwords
 	uint SrcIndex = SrcOffset + ThreadId * 4;
 	uint DstIndex = DstOffset + ThreadId * 4;
@@ -104,6 +118,8 @@ void MemcpyCS( uint ThreadId : SV_DispatchThreadID )
 		uint SrcData = SrcByteAddressBuffer.Load( SrcIndex * 4 );
 		DstByteAddressBuffer.Store( DstIndex * 4, SrcData );
 	}
+#else
+#error "Not implemented"
 #endif
 }
 
@@ -117,16 +133,16 @@ void ScatterCopyCS( uint DispatchThreadId : SV_DispatchThreadID )
 
 	if( ScatterIndex < NumScatters )
 	{
-#if FLOAT4_BUFFER
+#if RESOURCE_TYPE == RESOURCE_TYPE_FLOAT4_STRUCTURED_BUFFER
 		uint DstIndex = ScatterStructuredBuffer[ ScatterIndex ] * Size + ScatterOffset;
 		uint SrcIndex = ThreadId;
 		DstStructuredBuffer[ DstIndex ] = UploadStructuredBuffer[ SrcIndex ];
-#elif UINT4_ALIGNED
+#elif RESOURCE_TYPE == RESOURCE_TYPE_UINT4_ALIGNED_BUFFER
 		uint DstIndex = ScatterByteAddressBuffer.Load( ScatterIndex * 4 ) * Size + ScatterOffset;
 		uint SrcIndex = ThreadId;
 		uint4 SrcData = UploadByteAddressBuffer.Load4( SrcIndex * 16 );
 		DstByteAddressBuffer.Store4( DstIndex * 16, SrcData );
-#elif FLOAT4_TEXTURE
+#elif RESOURCE_TYPE == RESOURCE_TYPE_FLOAT4_TEXTURE
 		uint DstIndex = ScatterStructuredBuffer[ScatterIndex] * Size + ScatterOffset;
 		uint SrcIndex = ThreadId;
 		uint2 IndexTexture;
@@ -135,11 +151,13 @@ void ScatterCopyCS( uint DispatchThreadId : SV_DispatchThreadID )
 
 		float4 srvResourceVal = UploadStructuredBuffer[SrcIndex];
 		DstTexture[IndexTexture.xy] = srvResourceVal;
-#else
+#elif RESOURCE_TYPE == RESOURCE_TYPE_UINT_BUFFER
 		uint DstIndex = ScatterByteAddressBuffer.Load( ScatterIndex * 4 ) * Size + ScatterOffset;
 		uint SrcIndex = ThreadId;
 		uint SrcData = UploadByteAddressBuffer.Load( SrcIndex * 4 );
-		DstByteAddressBuffer.Store( DstIndex * 4, SrcData );
+		DstByteAddressBuffer.Store(DstIndex * 4, SrcData);
+#else
+#error "Not implemented"
 #endif
 	}
 }
diff --git a/Engine/Source/Runtime/RenderCore/Private/UnifiedBuffer.cpp b/Engine/Source/Runtime/RenderCore/Private/UnifiedBuffer.cpp
--- a/Engine/Source/Runtime/RenderCore/Private/UnifiedBuffer.cpp
+++ b/Engine/Source/Runtime/RenderCore/Private/UnifiedBuffer.cpp
@@ -10,46 +10,47 @@
 #include "RenderUtils.h"
 #include "RenderGraphUtils.h"
 
+enum class EByteBufferResourceType
+{
+	Float4_StructuredBuffer,
+	Uint_Buffer,
+	Uint4Aligned_Buffer,
+	Float4_Texture,
+	Count
+};
+
 class FByteBufferShader : public FGlobalShader
 {
-	DECLARE_INLINE_TYPE_LAYOUT( FByteBufferShader, NonVirtual );
+	DECLARE_INLINE_TYPE_LAYOUT(FByteBufferShader, NonVirtual);
 
 	FByteBufferShader() {}
-	FByteBufferShader( const ShaderMetaType::CompiledShaderInitializerType& Initializer )
-		: FGlobalShader( Initializer )
+	FByteBufferShader(const ShaderMetaType::CompiledShaderInitializerType& Initializer)
+		: FGlobalShader(Initializer)
 	{}
 
-	class FFloat4BufferDim : SHADER_PERMUTATION_BOOL("FLOAT4_BUFFER");
-	class FUint4AlignedDim : SHADER_PERMUTATION_BOOL("UINT4_ALIGNED");
+	class ResourceTypeDim : SHADER_PERMUTATION_INT("RESOURCE_TYPE", (int)EByteBufferResourceType::Count);
 
-	using FPermutationDomain = TShaderPermutationDomain<
-		FFloat4BufferDim,
-		FUint4AlignedDim
-	>;
+	using FPermutationDomain = TShaderPermutationDomain<ResourceTypeDim>;
 
 	static bool ShouldCompilePermutation( const FGlobalShaderPermutationParameters& Parameters )
 	{
 		FPermutationDomain PermutationVector( Parameters.PermutationId );
 
-		if( PermutationVector.Get< FFloat4BufferDim >() )
-		{
-			return RHISupportsComputeShaders(Parameters.Platform);
-		}
-		else
+		EByteBufferResourceType ResourceType = (EByteBufferResourceType)PermutationVector.Get<ResourceTypeDim>();
+
+		if (ResourceType == EByteBufferResourceType::Uint_Buffer || ResourceType == EByteBufferResourceType::Uint4Aligned_Buffer)
 		{
 			/*
 			return RHISupportsComputeShaders(Parameters.Platform)
 				&& FDataDrivenShaderPlatformInfo::GetInfo(Parameters.Platform).bSupportsByteBufferComputeShaders;
 				*/
-			// TODO: Workaround for FDataDrivenShaderPlatformInfo::GetInfo not being properly filled out yet.
+				// TODO: Workaround for FDataDrivenShaderPlatformInfo::GetInfo not being properly filled out yet.
 			return FDataDrivenShaderPlatformInfo::GetSupportsByteBufferComputeShaders(Parameters.Platform) || Parameters.Platform == SP_PCD3D_SM5;
 		}
-	}
-
-	static void ModifyCompilationEnvironment(const FGlobalShaderPermutationParameters& Parameters, FShaderCompilerEnvironment& OutEnvironment)
-	{
-		FGlobalShader::ModifyCompilationEnvironment(Parameters, OutEnvironment);
-		OutEnvironment.SetDefine(TEXT("FLOAT4_TEXTURE"), false);
+		else
+		{
+			return RHISupportsComputeShaders(Parameters.Platform);
+		}
 	}
 
 	BEGIN_SHADER_PARAMETER_STRUCT( FParameters, )
@@ -71,25 +72,19 @@ class FMemsetBufferCS : public FByteBufferShader
 };
 IMPLEMENT_GLOBAL_SHADER( FMemsetBufferCS, "/Engine/Private/ByteBuffer.usf", "MemsetBufferCS", SF_Compute );
 
-class FMemcpyBufferCS : public FByteBufferShader
+class FMemcpyCS : public FByteBufferShader
 {
-	DECLARE_GLOBAL_SHADER( FMemcpyBufferCS );
-	SHADER_USE_PARAMETER_STRUCT( FMemcpyBufferCS, FByteBufferShader );
+	DECLARE_GLOBAL_SHADER( FMemcpyCS );
+	SHADER_USE_PARAMETER_STRUCT( FMemcpyCS, FByteBufferShader );
+
 	BEGIN_SHADER_PARAMETER_STRUCT(FParameters, )
 		SHADER_PARAMETER_STRUCT_INCLUDE(FByteBufferShader::FParameters, Common)
-		SHADER_PARAMETER_SRV(ByteAddressBuffer, SrcByteAddressBuffer)
 		SHADER_PARAMETER_SRV(StructuredBuffer<float4>, SrcStructuredBuffer)
+		SHADER_PARAMETER_SRV(ByteAddressBuffer, SrcByteAddressBuffer)
 		SHADER_PARAMETER_SRV(Texture2D<float4>, SrcTexture)
 	END_SHADER_PARAMETER_STRUCT()
-
-
-	static void ModifyCompilationEnvironment(const FGlobalShaderPermutationParameters& Parameters, FShaderCompilerEnvironment& OutEnvironment)
-	{
-		FByteBufferShader::ModifyCompilationEnvironment(Parameters, OutEnvironment);
-		OutEnvironment.SetDefine(TEXT("FLOAT4_TEXTURE"), false);
-	}
 };
-IMPLEMENT_GLOBAL_SHADER( FMemcpyBufferCS, "/Engine/Private/ByteBuffer.usf", "MemcpyCS", SF_Compute );
+IMPLEMENT_GLOBAL_SHADER( FMemcpyCS, "/Engine/Private/ByteBuffer.usf", "MemcpyCS", SF_Compute );
 
 class FScatterCopyCS : public FByteBufferShader
 {
@@ -104,66 +99,12 @@ class FScatterCopyCS : public FByteBufferShader
 		SHADER_PARAMETER_SRV(ByteAddressBuffer, ScatterByteAddressBuffer)
 		SHADER_PARAMETER_SRV(StructuredBuffer<uint>, ScatterStructuredBuffer)
 	END_SHADER_PARAMETER_STRUCT()
-
-	static void ModifyCompilationEnvironment(const FGlobalShaderPermutationParameters& Parameters, FShaderCompilerEnvironment& OutEnvironment)
-	{
-		FByteBufferShader::ModifyCompilationEnvironment(Parameters, OutEnvironment);
-		OutEnvironment.SetDefine(TEXT("FLOAT4_TEXTURE"), false);
-	}
 };
 IMPLEMENT_GLOBAL_SHADER( FScatterCopyCS, "/Engine/Private/ByteBuffer.usf", "ScatterCopyCS", SF_Compute );
 
-class FMemcpyTextureToTextureCS : public FMemcpyBufferCS
-{
-	DECLARE_GLOBAL_SHADER( FMemcpyTextureToTextureCS );
-
-	SHADER_USE_PARAMETER_STRUCT( FMemcpyTextureToTextureCS, FMemcpyBufferCS);
-
-
-	static bool ShouldCompilePermutation( const FGlobalShaderPermutationParameters& Parameters )
-	{
-		FPermutationDomain PermutationVector(Parameters.PermutationId);
-		return RHISupportsComputeShaders( Parameters.Platform );
-	}
-
-	static void ModifyCompilationEnvironment(const FGlobalShaderPermutationParameters& Parameters, FShaderCompilerEnvironment& OutEnvironment)
-	{
-		FByteBufferShader::ModifyCompilationEnvironment(Parameters, OutEnvironment);
-		OutEnvironment.SetDefine(TEXT("FLOAT4_BUFFER"), false);
-		OutEnvironment.SetDefine(TEXT("UINT4_ALIGNED"), false);
-		OutEnvironment.SetDefine(TEXT("FLOAT4_TEXTURE"), true);
-	}
-};
-
-IMPLEMENT_GLOBAL_SHADER( FMemcpyTextureToTextureCS, "/Engine/Private/ByteBuffer.usf", "MemcpyCS", SF_Compute );
-
-
-class FScatterCopyTextureCS : public FScatterCopyCS
-{
-	DECLARE_GLOBAL_SHADER( FScatterCopyTextureCS );
-
-	SHADER_USE_PARAMETER_STRUCT( FScatterCopyTextureCS, FScatterCopyCS);
-
-	static bool ShouldCompilePermutation(const FGlobalShaderPermutationParameters& Parameters)
-	{
-		FPermutationDomain PermutationVector(Parameters.PermutationId);
-		return RHISupportsComputeShaders(Parameters.Platform);
-	}
-	
-	static void ModifyCompilationEnvironment(const FGlobalShaderPermutationParameters& Parameters, FShaderCompilerEnvironment& OutEnvironment)
-	{
-		FByteBufferShader::ModifyCompilationEnvironment(Parameters, OutEnvironment);
-		OutEnvironment.SetDefine(TEXT("FLOAT4_BUFFER"), false);
-		OutEnvironment.SetDefine(TEXT("UINT4_ALIGNED"), false);
-		OutEnvironment.SetDefine(TEXT("FLOAT4_TEXTURE"), true);
-	}
-};
-
-IMPLEMENT_GLOBAL_SHADER( FScatterCopyTextureCS, "/Engine/Private/ByteBuffer.usf", "ScatterCopyCS", SF_Compute );
-
 enum class EResourceType
 {
-	BUFFER,
+	STRUCTURED_BUFFER,
 	BYTEBUFFER,
 	TEXTURE
 };
@@ -174,38 +115,30 @@ struct ResourceTypeTraits;
 template<>
 struct ResourceTypeTraits<FRWBufferStructured>
 {
-	typedef FScatterCopyCS FScatterCS;
-	typedef FMemcpyBufferCS FMemcypCS;
-	typedef FMemsetBufferCS FMemsetCS;
-	static const EResourceType Type = EResourceType::BUFFER;
+	static const EResourceType Type = EResourceType::STRUCTURED_BUFFER;
 };
 
 template<>
 struct ResourceTypeTraits<FTextureRWBuffer2D>
 {
-	typedef FScatterCopyTextureCS FScatterCS;
-	typedef FMemcpyTextureToTextureCS FMemcypCS;
 	static const EResourceType Type = EResourceType::TEXTURE;
 };
 
 template<>
 struct ResourceTypeTraits<FRWByteAddressBuffer>
 {
-	typedef FScatterCopyCS FScatterCS;
-	typedef FMemcpyBufferCS FMemcypCS;
-	typedef FMemsetBufferCS FMemsetCS;
 	static const EResourceType Type = EResourceType::BYTEBUFFER;
 };
 
+static uint32 CalculateFloat4sPerLine()
+{
+	uint16 PrimitivesPerTextureLine = FMath::Min((int32)MAX_uint16, (int32)GMaxTextureDimensions) / FScatterUploadBuffer::PrimitiveDataStrideInFloat4s;
+	return PrimitivesPerTextureLine * FScatterUploadBuffer::PrimitiveDataStrideInFloat4s;
+}
+
 template<typename ResourceType>
 void MemsetResource(FRHICommandList& RHICmdList, const ResourceType& DstBuffer, uint32 Value, uint32 NumBytes, uint32 DstOffset)
 {
-	if (ResourceTypeTraits<ResourceType>::Type == EResourceType::TEXTURE)
-	{
-		check(false && TEXT("TEXTURE memset not yet implemented"));
-		return;
-	}
-
 	uint32 DivisorAlignment = ResourceTypeTraits<ResourceType>::Type == EResourceType::BYTEBUFFER ? 4 : 16;
 	if (ResourceTypeTraits<ResourceType>::Type == EResourceType::BYTEBUFFER)
 	{
@@ -218,32 +151,39 @@ void MemsetResource(FRHICommandList& RHICmdList, const ResourceType& DstBuffer,
 		check((NumBytes & 15) == 0);
 	}
 
-	typename ResourceTypeTraits<ResourceType>::FMemsetCS::FParameters Parameters;
+	EByteBufferResourceType ResourceTypeEnum;
+
+	FMemsetBufferCS::FParameters Parameters;
 	Parameters.Value = Value;
 	Parameters.Size = NumBytes / DivisorAlignment;
 	Parameters.DstOffset = DstOffset / DivisorAlignment;
+
 	if (ResourceTypeTraits<ResourceType>::Type == EResourceType::BYTEBUFFER)
 	{
+		ResourceTypeEnum = EByteBufferResourceType::Uint_Buffer;
+
 		Parameters.DstByteAddressBuffer = DstBuffer.UAV;
 	}
-	else if (ResourceTypeTraits<ResourceType>::Type == EResourceType::BUFFER)
+	else if (ResourceTypeTraits<ResourceType>::Type == EResourceType::STRUCTURED_BUFFER)
 	{
-		Parameters.DstStructuredBuffer = DstBuffer.UAV;
-	}
+		ResourceTypeEnum = EByteBufferResourceType::Float4_StructuredBuffer;
 
-	typename ResourceTypeTraits<ResourceType>::FMemsetCS::FPermutationDomain PermutationVector;
-	if (ResourceTypeTraits<ResourceType>::Type == EResourceType::BYTEBUFFER)
-	{
-		PermutationVector.template Set<typename ResourceTypeTraits<ResourceType>::FMemsetCS::FFloat4BufferDim >(false);
+		Parameters.DstStructuredBuffer = DstBuffer.UAV;
 	}
-	else
+	else if (ResourceTypeTraits<ResourceType>::Type == EResourceType::TEXTURE)
 	{
-		PermutationVector.template Set<typename ResourceTypeTraits<ResourceType>::FMemsetCS::FFloat4BufferDim >(true);
+		ResourceTypeEnum = EByteBufferResourceType::Float4_Texture;
+
+		Parameters.DstTexture = DstBuffer.UAV;
+		Parameters.Float4sPerLine = CalculateFloat4sPerLine();
 	}
 
+	FMemcpyCS::FPermutationDomain PermutationVector;
+	PermutationVector.Set<FMemcpyCS::ResourceTypeDim >((int)ResourceTypeEnum);
+
 	auto ShaderMap = GetGlobalShaderMap(GMaxRHIFeatureLevel);
 
-	auto ComputeShader = GetGlobalShaderMap(GMaxRHIFeatureLevel)->GetShader<typename ResourceTypeTraits<ResourceType>::FMemsetCS >(PermutationVector);
+	auto ComputeShader = GetGlobalShaderMap(GMaxRHIFeatureLevel)->GetShader<FMemsetBufferCS>(PermutationVector);
 
 	FComputeShaderUtils::Dispatch(RHICmdList, ComputeShader, Parameters, FIntVector(FMath::DivideAndRoundUp(NumBytes / 16, 64u), 1, 1));
 }
@@ -275,44 +215,44 @@ void MemcpyResource(FRHICommandList& RHICmdList, const ResourceType& DstBuffer,
 		const uint32 NumWaves = FMath::Max(FMath::Min<uint32>(GRHIMaxDispatchThreadGroupsPerDimension.X, FMath::DivideAndRoundUp(NumBytes / 16, 64u)), 1u);
 		const uint32 NumBytesPerDispatch = FMath::Min(FMath::Max(NumWaves, 1u) * 16 * 64, NumBytes);
 
-		typename ResourceTypeTraits<ResourceType>::FMemcypCS::FParameters Parameters;
+		EByteBufferResourceType ResourceTypeEnum;
+
+		FMemcpyCS::FParameters Parameters;
 		Parameters.Common.Size = NumBytesPerDispatch / DivisorAlignment;
 		Parameters.Common.SrcOffset = (SrcOffset + NumBytesProcessed) / DivisorAlignment;
 		Parameters.Common.DstOffset = (DstOffset + NumBytesProcessed) / DivisorAlignment;
+
 		if (ResourceTypeTraits<ResourceType>::Type == EResourceType::BYTEBUFFER)
 		{
+			ResourceTypeEnum = EByteBufferResourceType::Uint_Buffer;
+
 			Parameters.SrcByteAddressBuffer = SrcBuffer.SRV;
 			Parameters.Common.DstByteAddressBuffer = DstBuffer.UAV;
 		}
-		else if (ResourceTypeTraits<ResourceType>::Type == EResourceType::BUFFER)
+		else if (ResourceTypeTraits<ResourceType>::Type == EResourceType::STRUCTURED_BUFFER)
 		{
+			ResourceTypeEnum = EByteBufferResourceType::Float4_StructuredBuffer;
+
 			Parameters.SrcStructuredBuffer = SrcBuffer.SRV;
 			Parameters.Common.DstStructuredBuffer = DstBuffer.UAV;
 		}
 		else if (ResourceTypeTraits<ResourceType>::Type == EResourceType::TEXTURE)
 		{
+			ResourceTypeEnum = EByteBufferResourceType::Float4_Texture;
+
 			Parameters.SrcTexture = SrcBuffer.SRV;
 			Parameters.Common.DstTexture = DstBuffer.UAV;
-		}
-
-		if (ResourceTypeTraits<ResourceType>::Type == EResourceType::TEXTURE)
-		{
-			// TODO: Is this even used?
-			uint16 PrimitivesPerTextureLine = FMath::Min((int32)MAX_uint16, (int32)GMaxTextureDimensions) / (FScatterUploadBuffer::PrimitiveDataStrideInFloat4s);
-			Parameters.Common.Float4sPerLine = PrimitivesPerTextureLine * FScatterUploadBuffer::PrimitiveDataStrideInFloat4s;
-		}
-
-		typename ResourceTypeTraits<ResourceType>::FMemcypCS::FPermutationDomain PermutationVector;
-		if (ResourceTypeTraits<ResourceType>::Type == EResourceType::BYTEBUFFER)
-		{
-			PermutationVector.template Set<typename ResourceTypeTraits<ResourceType>::FMemcypCS::FFloat4BufferDim >(false);
+			Parameters.Common.Float4sPerLine = CalculateFloat4sPerLine();
 		}
 		else
 		{
-			PermutationVector.template Set<typename ResourceTypeTraits<ResourceType>::FMemcypCS::FFloat4BufferDim >(true);
+			check(false);
 		}
 
-		auto ComputeShader = GetGlobalShaderMap(GMaxRHIFeatureLevel)->GetShader<typename ResourceTypeTraits<ResourceType>::FMemcypCS >(PermutationVector);
+		FMemcpyCS::FPermutationDomain PermutationVector;
+		PermutationVector.Set<FMemcpyCS::ResourceTypeDim >((int)ResourceTypeEnum);
+
+		auto ComputeShader = GetGlobalShaderMap(GMaxRHIFeatureLevel)->GetShader<FMemcpyCS >(PermutationVector);
 
 		FComputeShaderUtils::Dispatch(RHICmdList, ComputeShader, Parameters, FIntVector(NumWaves, 1, 1));
 
@@ -327,8 +267,7 @@ template<>
 RENDERCORE_API bool ResizeResourceIfNeeded<FTextureRWBuffer2D>(FRHICommandList& RHICmdList, FTextureRWBuffer2D& Texture, uint32 NumBytes, const TCHAR* DebugName)
 {
 	check((NumBytes & 15) == 0);
-	uint16 PrimitivesPerTextureLine = FMath::Min((int32)MAX_uint16, (int32)GMaxTextureDimensions) / (FScatterUploadBuffer::PrimitiveDataStrideInFloat4s);
-	uint32 Float4sPerLine = PrimitivesPerTextureLine * FScatterUploadBuffer::PrimitiveDataStrideInFloat4s;
+	uint32 Float4sPerLine = CalculateFloat4sPerLine();
 	uint32 BytesPerLine = Float4sPerLine * 16;
 
 	EPixelFormat BufferFormat = PF_A32B32G32R32F;
@@ -631,49 +570,51 @@ void FScatterUploadBuffer::ResourceUploadTo(FRHICommandList& RHICmdList, const R
 	uint32 NumDispatches = FMath::DivideAndRoundUp(NumThreads, ThreadGroupSize);
 	uint32 NumLoops = FMath::DivideAndRoundUp(NumDispatches, (uint32)GMaxComputeDispatchDimension);
 
-	typename ResourceTypeTraits<ResourceType>::FScatterCS::FParameters Parameters;
-	
-	if (ResourceTypeTraits<ResourceType>::Type == EResourceType::TEXTURE)
-	{
-		uint16 PrimitivesPerTextureLine = FMath::Min((int32)MAX_uint16, (int32)GMaxTextureDimensions) / (FScatterUploadBuffer::PrimitiveDataStrideInFloat4s);
-		Parameters.Common.Float4sPerLine = PrimitivesPerTextureLine * FScatterUploadBuffer::PrimitiveDataStrideInFloat4s;;
-	}
+	EByteBufferResourceType ResourceTypeEnum;
 
+	FScatterCopyCS::FParameters Parameters;
 	Parameters.Common.Size = NumThreadsPerScatter;
 	Parameters.NumScatters = NumScatters;
 
+	check(bFloat4Buffer || ResourceTypeTraits<ResourceType>::Type == EResourceType::BYTEBUFFER);
+
 	if (ResourceTypeTraits<ResourceType>::Type == EResourceType::BYTEBUFFER)
 	{
+		if (NumBytesPerThread == 16)
+		{
+			ResourceTypeEnum = EByteBufferResourceType::Uint4Aligned_Buffer;
+		}
+		else
+		{
+			ResourceTypeEnum = EByteBufferResourceType::Uint_Buffer;
+		}
 		Parameters.UploadByteAddressBuffer = UploadBuffer.SRV;
 		Parameters.ScatterByteAddressBuffer = ScatterBuffer.SRV;
 		Parameters.Common.DstByteAddressBuffer = DstBuffer.UAV;
 	}
-	else if (ResourceTypeTraits<ResourceType>::Type == EResourceType::BUFFER)
+	else if (ResourceTypeTraits<ResourceType>::Type == EResourceType::STRUCTURED_BUFFER)
 	{
+		ResourceTypeEnum = EByteBufferResourceType::Float4_StructuredBuffer;
+
 		Parameters.UploadStructuredBuffer = UploadBuffer.SRV;
 		Parameters.ScatterStructuredBuffer = ScatterBuffer.SRV;
 		Parameters.Common.DstStructuredBuffer = DstBuffer.UAV;
 	}
 	else if (ResourceTypeTraits<ResourceType>::Type == EResourceType::TEXTURE)
 	{
+		ResourceTypeEnum = EByteBufferResourceType::Float4_Texture;
+
 		Parameters.UploadStructuredBuffer = UploadBuffer.SRV;
 		Parameters.ScatterStructuredBuffer = ScatterBuffer.SRV;
 		Parameters.Common.DstTexture = DstBuffer.UAV;
-	}
 
-	typename ResourceTypeTraits<ResourceType>::FMemcypCS::FPermutationDomain PermutationVector;
-	if (ResourceTypeTraits<ResourceType>::Type == EResourceType::TEXTURE)
-	{
-		PermutationVector.template Set<typename ResourceTypeTraits<ResourceType>::FMemcypCS::FFloat4BufferDim >(false);
-		PermutationVector.template Set<typename ResourceTypeTraits<ResourceType>::FMemcypCS::FUint4AlignedDim >(false);
-	}
-	else
-	{
-		PermutationVector.template Set< FScatterCopyCS::FFloat4BufferDim >(bFloat4Buffer);
-		PermutationVector.template Set< FScatterCopyCS::FUint4AlignedDim >(NumBytesPerThread == 16);
+		Parameters.Common.Float4sPerLine = CalculateFloat4sPerLine();
 	}
 
-	auto ComputeShader = GetGlobalShaderMap(GMaxRHIFeatureLevel)->GetShader<typename ResourceTypeTraits<ResourceType>::FScatterCS >(PermutationVector);
+	FByteBufferShader::FPermutationDomain PermutationVector;
+	PermutationVector.Set<FByteBufferShader::ResourceTypeDim>((int)ResourceTypeEnum);
+
+	auto ComputeShader = GetGlobalShaderMap(GMaxRHIFeatureLevel)->GetShader<FScatterCopyCS>(PermutationVector);
 
 	RHICmdList.BeginUAVOverlap(DstBuffer.UAV);
 
@@ -694,11 +635,13 @@ void FScatterUploadBuffer::ResourceUploadTo(FRHICommandList& RHICmdList, const R
 	}
 }
 
-template RENDERCORE_API void MemsetResource< FRWBufferStructured>(FRHICommandList& RHICmdList, const FRWBufferStructured& DstBuffer, uint32 Value, uint32 NumBytes, uint32 DstOffset);
-template RENDERCORE_API void MemsetResource< FRWByteAddressBuffer>(FRHICommandList& RHICmdList, const FRWByteAddressBuffer& DstBuffer, uint32 Value, uint32 NumBytes, uint32 DstOffset);
-template RENDERCORE_API void MemcpyResource< FTextureRWBuffer2D>(FRHICommandList& RHICmdList, const FTextureRWBuffer2D& DstBuffer, const FTextureRWBuffer2D& SrcBuffer, uint32 NumBytes, uint32 DstOffset, uint32 SrcOffset, bool bAlreadyInUAVOverlap);
-template RENDERCORE_API void MemcpyResource< FRWBufferStructured>(FRHICommandList& RHICmdList, const FRWBufferStructured& DstBuffer, const FRWBufferStructured& SrcBuffer, uint32 NumBytes, uint32 DstOffset, uint32 SrcOffset, bool bAlreadyInUAVOverlap);
-template RENDERCORE_API void MemcpyResource< FRWByteAddressBuffer>(FRHICommandList& RHICmdList, const FRWByteAddressBuffer& DstBuffer, const FRWByteAddressBuffer& SrcBuffer, uint32 NumBytes, uint32 DstOffset, uint32 SrcOffset, bool bAlreadyInUAVOverlap);
+template RENDERCORE_API void MemsetResource<FRWBufferStructured>(FRHICommandList& RHICmdList, const FRWBufferStructured& DstBuffer, uint32 Value, uint32 NumBytes, uint32 DstOffset);
+template RENDERCORE_API void MemsetResource<FRWByteAddressBuffer>(FRHICommandList& RHICmdList, const FRWByteAddressBuffer& DstBuffer, uint32 Value, uint32 NumBytes, uint32 DstOffset);
+
+template RENDERCORE_API void MemcpyResource<FTextureRWBuffer2D>(FRHICommandList& RHICmdList, const FTextureRWBuffer2D& DstBuffer, const FTextureRWBuffer2D& SrcBuffer, uint32 NumBytes, uint32 DstOffset, uint32 SrcOffset, bool bAlreadyInUAVOverlap);
+template RENDERCORE_API void MemcpyResource<FRWBufferStructured>(FRHICommandList& RHICmdList, const FRWBufferStructured& DstBuffer, const FRWBufferStructured& SrcBuffer, uint32 NumBytes, uint32 DstOffset, uint32 SrcOffset, bool bAlreadyInUAVOverlap);
+template RENDERCORE_API void MemcpyResource<FRWByteAddressBuffer>(FRHICommandList& RHICmdList, const FRWByteAddressBuffer& DstBuffer, const FRWByteAddressBuffer& SrcBuffer, uint32 NumBytes, uint32 DstOffset, uint32 SrcOffset, bool bAlreadyInUAVOverlap);
+
 template RENDERCORE_API void FScatterUploadBuffer::ResourceUploadTo<FTextureRWBuffer2D>(FRHICommandList& RHICmdList, const FTextureRWBuffer2D& DstBuffer, bool bFlush);
 template RENDERCORE_API void FScatterUploadBuffer::ResourceUploadTo<FRWBufferStructured>(FRHICommandList& RHICmdList, const FRWBufferStructured& DstBuffer, bool bFlush);
 template RENDERCORE_API void FScatterUploadBuffer::ResourceUploadTo<FRWByteAddressBuffer>(FRHICommandList& RHICmdList, const FRWByteAddressBuffer& DstBuffer, bool bFlush);
\ No newline at end of file
