Remove MOBILE_PROPAGATE_ALPHA shader define
Temporary remove premultiply alpha
FXAA luminance calculation moved from tonemapper to FXAA when MobilePropagateAlpha enabled

diff --git a/Engine/Shaders/Private/FXAAShader.usf b/Engine/Shaders/Private/FXAAShader.usf
--- a/Engine/Shaders/Private/FXAAShader.usf
+++ b/Engine/Shaders/Private/FXAAShader.usf
@@ -106,7 +106,7 @@ void FxaaVS(
 			fxaaConsoleEdgeThresholdMin,
 			fxaaConsole360ConstDir);
 
-		#if (POST_PROCESS_ALPHA != 2) && !MOBILE_PROPAGATE_ALPHA
+		#if (POST_PROCESS_ALPHA != 2)
 			OutColor.a = 1.0;
 		#endif
 	}
diff --git a/Engine/Shaders/Private/MobileBasePassPixelShader.usf b/Engine/Shaders/Private/MobileBasePassPixelShader.usf
--- a/Engine/Shaders/Private/MobileBasePassPixelShader.usf
+++ b/Engine/Shaders/Private/MobileBasePassPixelShader.usf
@@ -108,7 +108,6 @@ int NumDynamicPointLights;
 #define TRANSLUCENCY_NON_DIRECTIONAL (MATERIALBLENDING_ANY_TRANSLUCENT && (TRANSLUCENCY_LIGHTING_VOLUMETRIC_NONDIRECTIONAL || TRANSLUCENCY_LIGHTING_VOLUMETRIC_PERVERTEX_NONDIRECTIONAL) && !MATERIAL_SHADINGMODEL_SINGLELAYERWATER)
 
 #define ALLOW_CUBE_REFLECTIONS HQ_REFLECTIONS
-#define MOBILE_PROPAGATE_ALPHA_PREMULTIPLY 2
 
 #include "SHCommon.ush"
 #include "MobileBasePassCommon.ush"
@@ -1097,7 +1096,7 @@ void Main(
 		OutColor.rgb = sqrt(OutColor.rgb);
 	#endif
 #endif
-#if (MATERIALBLENDING_TRANSLUCENT || MATERIALBLENDING_ALPHACOMPOSITE) && (MOBILE_PROPAGATE_ALPHA == MOBILE_PROPAGATE_ALPHA_PREMULTIPLY)
+#if (MATERIALBLENDING_TRANSLUCENT || MATERIALBLENDING_ALPHACOMPOSITE)
 		OutColor.rgb = OutColor.rgb * OutColor.a;
 #endif
 #if NUM_VIRTUALTEXTURE_SAMPLES || LIGHTMAP_VT_ENABLED
diff --git a/Engine/Shaders/Private/PostProcessTonemap.usf b/Engine/Shaders/Private/PostProcessTonemap.usf
--- a/Engine/Shaders/Private/PostProcessTonemap.usf
+++ b/Engine/Shaders/Private/PostProcessTonemap.usf
@@ -699,7 +699,7 @@ float4 TonemapCommonPS(
 	}
 	#endif
 	
-	#if (FEATURE_LEVEL <= FEATURE_LEVEL_ES3_1 && (USE_FXAA == 0 || MOBILE_PROPAGATE_ALPHA > 0)) || POST_PROCESS_ALPHA == 2 || (POST_PROCESS_ALPHA == 1 && DIM_OUTPUT_DEVICE == TONEMAPPER_OUTPUT_NoToneCurve)  || (POST_PROCESS_ALPHA == 1 && DIM_OUTPUT_DEVICE == TONEMAPPER_OUTPUT_WithToneCurve)
+	#if POST_PROCESS_ALPHA == 2 || (POST_PROCESS_ALPHA == 1 && DIM_OUTPUT_DEVICE == TONEMAPPER_OUTPUT_NoToneCurve)  || (POST_PROCESS_ALPHA == 1 && DIM_OUTPUT_DEVICE == TONEMAPPER_OUTPUT_WithToneCurve)
 		OutColor.a = SceneColor.a;
 	#endif
 
diff --git a/Engine/Shaders/Private/PostProcessUpscale.usf b/Engine/Shaders/Private/PostProcessUpscale.usf
--- a/Engine/Shaders/Private/PostProcessUpscale.usf
+++ b/Engine/Shaders/Private/PostProcessUpscale.usf
@@ -290,7 +290,7 @@ void MainPS(noperspective float4 UVAndScreenPos : TEXCOORD0, float4 SvPosition :
 	}
 #endif
 
-#if !POST_PROCESS_ALPHA && !MOBILE_PROPAGATE_ALPHA
+#if !POST_PROCESS_ALPHA
 	OutColor.a = 0; // Skip all computations related to alpha
 #endif
 }
\ No newline at end of file
diff --git a/Engine/Source/Runtime/Engine/Private/SceneUtils.cpp b/Engine/Source/Runtime/Engine/Private/SceneUtils.cpp
--- a/Engine/Source/Runtime/Engine/Private/SceneUtils.cpp
+++ b/Engine/Source/Runtime/Engine/Private/SceneUtils.cpp
@@ -18,11 +18,6 @@ bool IsMobilePropagateAlphaEnabled(EShaderPlatform Platform)
 	return IsMobilePlatform(Platform) && (FPlatformMisc::GetMobilePropagateAlphaSetting() > 0);
 }
 
-bool IsMobilePropagateAlphaPremultiply(EShaderPlatform Platform)
-{
-	return IsMobilePlatform(Platform) && (FPlatformMisc::GetMobilePropagateAlphaSetting() == 2);
-}
-
 ENGINE_API EMobileHDRMode GetMobileHDRMode()
 {
 	EMobileHDRMode HDRMode = EMobileHDRMode::EnabledFloat16;
diff --git a/Engine/Source/Runtime/Engine/Private/ShaderCompiler/ShaderCompiler.cpp b/Engine/Source/Runtime/Engine/Private/ShaderCompiler/ShaderCompiler.cpp
--- a/Engine/Source/Runtime/Engine/Private/ShaderCompiler/ShaderCompiler.cpp
+++ b/Engine/Source/Runtime/Engine/Private/ShaderCompiler/ShaderCompiler.cpp
@@ -4876,11 +4876,6 @@ void GlobalBeginCompileShader(
 			bool bGLESDeferredShading = Target.Platform == SP_OPENGL_ES3_1_ANDROID;
 			Input.Environment.SetDefine(TEXT("USE_GLES_FBF_DEFERRED"), bGLESDeferredShading ? 1 : 0);
 		}
-
-		static FShaderPlatformCachedIniValue<int32> MobilePropagateAlphaIniValue(TEXT("/Script/Engine.RendererSettings"), TEXT("r.Mobile.PropagateAlpha"));
-		int MobilePropagateAlphaIniValueInt = MobilePropagateAlphaIniValue.Get((EShaderPlatform)ShaderPlatform);
-
-		Input.Environment.SetDefine(TEXT("MOBILE_PROPAGATE_ALPHA"), MobilePropagateAlphaIniValueInt);
 	}
 
 	// Set VR definitions
@@ -5269,6 +5264,14 @@ void GlobalBeginCompileShader(
 	{
 		static IConsoleVariable* CVar = IConsoleManager::Get().FindConsoleVariable(TEXT("r.PostProcessing.PropagateAlpha"));
 		int32 PropagateAlpha = CVar->GetInt();
+
+		if (bIsMobilePlatform)
+		{
+			static FShaderPlatformCachedIniValue<int32> MobilePropagateAlphaIniValue(TEXT("/Script/Engine.RendererSettings"), TEXT("r.Mobile.PropagateAlpha"));
+			int MobilePropagateAlphaIniValueInt = MobilePropagateAlphaIniValue.Get((EShaderPlatform)ShaderPlatform);
+			PropagateAlpha = MobilePropagateAlphaIniValueInt > 0 ? 2 : 0;
+		}
+
 		if (PropagateAlpha < 0 || PropagateAlpha > 2)
 		{
 			PropagateAlpha = 0;
diff --git a/Engine/Source/Runtime/Engine/Public/SceneUtils.h b/Engine/Source/Runtime/Engine/Public/SceneUtils.h
--- a/Engine/Source/Runtime/Engine/Public/SceneUtils.h
+++ b/Engine/Source/Runtime/Engine/Public/SceneUtils.h
@@ -54,7 +54,6 @@ ENGINE_API bool IsMobileHDR();
 
 /** True if Alpha Propagate is enabled for the mobile renderer. */
 ENGINE_API bool IsMobilePropagateAlphaEnabled(EShaderPlatform Platform);
-ENGINE_API bool IsMobilePropagateAlphaPremultiply(EShaderPlatform Platform);
 
 ENGINE_API EMobileHDRMode GetMobileHDRMode();
 
diff --git a/Engine/Source/Runtime/RenderCore/Private/Shader.cpp b/Engine/Source/Runtime/RenderCore/Private/Shader.cpp
--- a/Engine/Source/Runtime/RenderCore/Private/Shader.cpp
+++ b/Engine/Source/Runtime/RenderCore/Private/Shader.cpp
@@ -1530,20 +1530,6 @@ void ShaderMapAppendKeyString(EShaderPlatform Platform, FString& KeyString)
 			KeyString += (CVar && CVar->GetInt() != 0) ? TEXT("_MobileHDR") : TEXT("");
 		}
 
-		{
-			static FShaderPlatformCachedIniValue<int32> MobilePropagateAlphaIniValue(TEXT("/Script/Engine.RendererSettings"), TEXT("r.Mobile.PropagateAlpha"));
-			int MobilePropagateAlphaIniValueInt = MobilePropagateAlphaIniValue.Get((EShaderPlatform)Platform);
-
-			if (MobilePropagateAlphaIniValueInt == 1)
-			{
-				KeyString += TEXT("_MbPropA");
-			}
-			else if (MobilePropagateAlphaIniValueInt == 2)
-			{
-				KeyString += TEXT("_MbPrePropA");
-			}
-		}
-
 		{
 			static IConsoleVariable* CVar = IConsoleManager::Get().FindConsoleVariable(TEXT("r.Mobile.ShadingPath"));
 			KeyString += (CVar && CVar->GetInt() != 0) ? TEXT("_MobDSh") : TEXT("");
@@ -1675,10 +1661,25 @@ void ShaderMapAppendKeyString(EShaderPlatform Platform, FString& KeyString)
 	}
 
 	{
-		static const auto CVar = IConsoleManager::Get().FindTConsoleVariableDataInt(TEXT("r.PostProcessing.PropagateAlpha"));
-		if (CVar && CVar->GetValueOnAnyThread() > 0)
+		int PropagateAlphaType = 0;
+		if (IsMobilePlatform(Platform))
+		{
+			static FShaderPlatformCachedIniValue<int32> MobilePropagateAlphaIniValue(TEXT("/Script/Engine.RendererSettings"), TEXT("r.Mobile.PropagateAlpha"));
+			int MobilePropagateAlphaIniValueInt = MobilePropagateAlphaIniValue.Get((EShaderPlatform)Platform);
+			PropagateAlphaType = MobilePropagateAlphaIniValueInt > 0 ? 2 : 0;
+		}
+		else
+		{
+			static const auto CVar = IConsoleManager::Get().FindTConsoleVariableDataInt(TEXT("r.PostProcessing.PropagateAlpha"));
+			if (CVar)
+			{
+				PropagateAlphaType = CVar->GetValueOnAnyThread();
+			}
+		}
+
+		if (PropagateAlphaType > 0)
 		{
-			if (CVar->GetValueOnAnyThread() == 2)
+			if (PropagateAlphaType == 2)
 			{
 				KeyString += TEXT("_SA2");
 			}
diff --git a/Engine/Source/Runtime/Renderer/Private/MobileBasePassRendering.cpp b/Engine/Source/Runtime/Renderer/Private/MobileBasePassRendering.cpp
--- a/Engine/Source/Runtime/Renderer/Private/MobileBasePassRendering.cpp
+++ b/Engine/Source/Runtime/Renderer/Private/MobileBasePassRendering.cpp
@@ -55,8 +55,7 @@ static TAutoConsoleVariable<int32> CVarMobilePropagateAlpha(
 	TEXT("r.Mobile.PropagateAlpha"),
 	0,
 	TEXT("0: Disabled")
-	TEXT("1: Propagate Full Alpha Propagate")
-	TEXT("2: Propagate Full Alpha Propagate + Premultiply"),
+	TEXT("1: Propagate Full Alpha Propagate"),
 	ECVF_ReadOnly | ECVF_RenderThreadSafe);
 
 
diff --git a/Engine/Source/Runtime/Renderer/Private/PostProcess/PostProcessTonemap.cpp b/Engine/Source/Runtime/Renderer/Private/PostProcess/PostProcessTonemap.cpp
--- a/Engine/Source/Runtime/Renderer/Private/PostProcess/PostProcessTonemap.cpp
+++ b/Engine/Source/Runtime/Renderer/Private/PostProcess/PostProcessTonemap.cpp
@@ -118,7 +118,6 @@ class FTonemapperFilmGrainDim      : SHADER_PERMUTATION_BOOL("USE_FILM_GRAIN");
 class FTonemapperSwitchAxis        : SHADER_PERMUTATION_BOOL("NEEDTOSWITCHVERTICLEAXIS");
 class FTonemapperMsaaDim           : SHADER_PERMUTATION_BOOL("METAL_MSAA_HDR_DECODE");
 class FTonemapperEyeAdaptationDim  : SHADER_PERMUTATION_BOOL("EYEADAPTATION_EXPOSURE_FIX");
-class FTonemapperUseFXAA           : SHADER_PERMUTATION_BOOL("USE_FXAA");
 
 using FCommonDomain = TShaderPermutationDomain<
 	FTonemapperBloomDim,
@@ -127,8 +126,7 @@ using FCommonDomain = TShaderPermutationDomain<
 	FTonemapperSharpenDim,
 	FTonemapperFilmGrainDim,
 	FTonemapperSwitchAxis,
-	FTonemapperMsaaDim,
-	FTonemapperUseFXAA>;
+	FTonemapperMsaaDim>;
 
 bool ShouldCompileCommonPermutation(const FGlobalShaderPermutationParameters& Parameters, const FCommonDomain& PermutationVector)
 {
@@ -144,11 +142,6 @@ bool ShouldCompileCommonPermutation(const FGlobalShaderPermutationParameters& Pa
 		return false;
 	}
 
-	if (PermutationVector.Get<FTonemapperUseFXAA>() && !IsMobilePlatform(Parameters.Platform))
-	{
-		return false;
-	}
-
 	// If GammaOnly, don't compile any other dimmension == true.
 	if (PermutationVector.Get<FTonemapperGammaOnlyDim>())
 	{
@@ -184,10 +177,6 @@ FCommonDomain BuildCommonPermutationDomain(const FViewInfo& View, bool bGammaOnl
 	PermutationVector.Set<FTonemapperSharpenDim>(CVarTonemapperSharpen.GetValueOnRenderThread() > 0.0f);	
 	PermutationVector.Set<FTonemapperSwitchAxis>(bSwitchVerticalAxis);
 	PermutationVector.Set<FTonemapperMsaaDim>(bMetalMSAAHDRDecode);
-	if (IsMobilePlatform(View.GetShaderPlatform()))
-	{
-		PermutationVector.Set<FTonemapperUseFXAA>(View.AntiAliasingMethod == AAM_FXAA);
-	}
 	return PermutationVector;
 }
 
