Added FCompressedTransform.

Compressed version of FRenderTransform. Used in GPU scene for instance transforms.

diff --git a/Engine/Shaders/Private/SceneData.ush b/Engine/Shaders/Private/SceneData.ush
--- a/Engine/Shaders/Private/SceneData.ush
+++ b/Engine/Shaders/Private/SceneData.ush
@@ -2,6 +2,8 @@
 
 #pragma once
 
+#include "OctahedralCommon.ush"
+
 #ifndef USE_GLOBAL_GPU_SCENE_DATA
 	#define USE_GLOBAL_GPU_SCENE_DATA 0
 #endif
@@ -299,6 +301,90 @@ float4 LoadInstanceDataElement(uint Index)
 #endif
 }
 
+// [Frisvad 2012, "Building an Orthonormal Basis from a 3D Unit Vector Without Normalization"]
+void GetHemiOrthoBasis( inout float3 BasisX, inout float3 BasisY, float3 BasisZ )
+{
+	float A = 1.0f / ( 1.0f + BasisZ.z );
+	float B = -BasisZ.x * BasisZ.y * A;
+	BasisX = float3( 1.0f - BasisZ.x * BasisZ.x * A, B, -BasisZ.x );
+	BasisY = float3( B, 1.0f - BasisZ.y * BasisZ.y * A, -BasisZ.y );
+}
+
+float4x4 DecodeTransform( uint4 RotationScale, float3 Translation, inout float3 Scale )
+{
+	float4x4 M = 0.0;
+	M[3].xyz = Translation;
+	M[3].w = 1.0;
+
+	// Rotation
+	{
+		float3 Rotation =
+		{
+			( RotationScale[0] >>  0 ) & 0xffff,
+			( RotationScale[0] >> 16 ) & 0xffff,
+			( RotationScale[1] >>  0 ) & 0x7fff
+		};
+
+		float2 OctZ  = ( Rotation.xy - 32768 ) * (1.0f / 32767.0f);
+		float Spin0  = ( Rotation.z  - 16384 ) * (0.70710678f / 16383.0f);	// rsqrt(2)
+		bool bSpinIsX = RotationScale[1] & 0x8000;
+
+		M[2].xyz = HemiOctahedronToUnitVector( OctZ );
+
+		float3 BasisX, BasisY;
+		GetHemiOrthoBasis( BasisX, BasisY, M[2].xyz );
+
+		float Spin1 = sqrt( 1.0f - Spin0 * Spin0 );
+		float X = bSpinIsX ? Spin0 : Spin1;
+		float Y = bSpinIsX ? Spin1 : Spin0;
+
+		M[0].xyz = BasisX * X + BasisY * Y;
+		M[1].xyz = cross( M[2].xyz, M[0].xyz );
+	}
+
+	// Scale
+	{
+		const uint SignMantissaBits	= 16;
+		const uint SignMantissaMask	= (1u << SignMantissaBits) - 1;
+		const uint MantissaBits		= SignMantissaBits - 1;
+
+#if 0
+		uint SharedExp = RotationScale[3] >> 22;
+
+		float ExpScale = asfloat( ( SharedExp - MantissaBits ) << 23 );
+
+		int3 Mantissa =
+		{
+			( RotationScale[2] >> 0 ),
+			( RotationScale[2] >> 18 ) | ( RotationScale[3] << 14 ),
+			( RotationScale[3] >> 4 )
+		};
+#else
+		uint SharedExp = RotationScale[3] >> 16;
+
+		float ExpScale = asfloat( ( SharedExp - MantissaBits ) << 23 );
+
+		uint3 Mantissa =
+		{
+			RotationScale[2] >> 0,
+			RotationScale[2] >> 16,
+			RotationScale[3] >> 0
+		};
+		
+#endif
+		Mantissa &= SignMantissaMask;
+		Scale = Mantissa;
+		Scale -= 1u << MantissaBits;
+		Scale *= ExpScale;
+
+		M[0] *= Scale[0];
+		M[1] *= Scale[1];
+		M[2] *= Scale[2];
+	}
+
+	return M;
+}
+
 // Fetch from scene primitive buffer
 FInstanceSceneData GetInstanceSceneData(uint InstanceId, uint SOAStride)
 {
@@ -320,6 +406,7 @@ FInstanceSceneData GetInstanceSceneData(uint InstanceId, uint SOAStride)
 	BRANCH
 	if (InstanceData.ValidInstance)
 	{
+#if 0
 		InstanceData.LocalToWorld			= transpose(float4x4(LoadInstanceDataElement(1 * SOAStride + InstanceId),
 																 LoadInstanceDataElement(2 * SOAStride + InstanceId),
 																 LoadInstanceDataElement(3 * SOAStride + InstanceId),
@@ -351,9 +438,35 @@ FInstanceSceneData GetInstanceSceneData(uint InstanceId, uint SOAStride)
 		InstanceData.DeterminantSign							= CondMask(InstanceData.Flags & INSTANCE_SCENE_DATA_FLAG_DETERMINANT_SIGN, -1.0f, 1.0f);
 
 		InstanceData.NonUniformScale.xyz						= Scale2 * InstanceData.InvNonUniformScale;
+#else
+		uint4  RotationScale									= asuint( LoadInstanceDataElement(1 * SOAStride + InstanceId) );
+		float3 Translation										= LoadInstanceDataElement(2 * SOAStride + InstanceId).xyz;
+		float3 Scale = 0;
+		InstanceData.LocalToWorld = DecodeTransform( RotationScale, Translation, Scale );
+
+		uint4  PrevRotationScale								= asuint( LoadInstanceDataElement(3 * SOAStride + InstanceId) );
+		float3 PrevTranslation									= LoadInstanceDataElement(4 * SOAStride + InstanceId).xyz;
+		float3 PrevScale = 0;
+		InstanceData.PrevLocalToWorld = DecodeTransform( PrevRotationScale, PrevTranslation, PrevScale );
+
+		InstanceData.LocalBoundsCenter							= LoadInstanceDataElement(5 * SOAStride + InstanceId).xyz;
+		InstanceData.LocalBoundsExtent.x						= LoadInstanceDataElement(5 * SOAStride + InstanceId).w;
+
+		InstanceData.LocalBoundsExtent.yz						= LoadInstanceDataElement(6 * SOAStride + InstanceId).xy;
+		InstanceData.PayloadDataOffset							= asuint(LoadInstanceDataElement(6 * SOAStride + InstanceId).z);
+
+	#if USES_PER_INSTANCE_RANDOM
+		InstanceData.PerInstanceRandom							= LoadInstanceDataElement(6 * SOAStride + InstanceId).w;
+	#endif
+
+		InstanceData.LightMapAndShadowMapUVBias					= LoadInstanceDataElement(7 * SOAStride + InstanceId);
+
+		InstanceData.NonUniformScale.xyz						= abs( Scale );
+		InstanceData.InvNonUniformScale							= rcp( InstanceData.NonUniformScale.xyz );
+		InstanceData.DeterminantSign							= CondMask(InstanceData.Flags & INSTANCE_SCENE_DATA_FLAG_DETERMINANT_SIGN, -1.0f, 1.0f);	// Scale.x * Scale.y * Scale.z < 0.0 ? -1.0 : 1.0;
+#endif
 
-		float3 AbsNonUniformScale								= abs(InstanceData.NonUniformScale.xyz);
-		InstanceData.NonUniformScale.w							= max3(AbsNonUniformScale.x, AbsNonUniformScale.y, AbsNonUniformScale.z);
+		InstanceData.NonUniformScale.w							= max3( InstanceData.NonUniformScale.x, InstanceData.NonUniformScale.y, InstanceData.NonUniformScale.z );
 	
 		InstanceData.WorldToLocal								= InstanceData.LocalToWorld;
 		InstanceData.WorldToLocal[0].xyz					   *= Pow2(InstanceData.InvNonUniformScale.x);
diff --git a/Engine/Source/Runtime/Engine/Private/InstanceUniformShaderParameters.cpp b/Engine/Source/Runtime/Engine/Private/InstanceUniformShaderParameters.cpp
--- a/Engine/Source/Runtime/Engine/Private/InstanceUniformShaderParameters.cpp
+++ b/Engine/Source/Runtime/Engine/Private/InstanceUniformShaderParameters.cpp
@@ -60,8 +60,11 @@ void FInstanceSceneShaderData::Setup(
 	}
 
 	// Remove shear
-	LocalToWorld.Orthonormalize();
-	PrevLocalToWorld.Orthonormalize();
+	LocalToWorld.Orthogonalize();
+	PrevLocalToWorld.Orthogonalize();
+
+	FCompressedTransform CompressedLocalToWorld( LocalToWorld );
+	FCompressedTransform CompressedPrevLocalToWorld( PrevLocalToWorld );
 
 	uint32 InstanceFlags = Instance.Flags;
 	if (LocalToWorld.RotDeterminant() < 0.0f)
@@ -82,6 +85,7 @@ void FInstanceSceneShaderData::Setup(
 	Data[0].Z  = *(const     float*)&Instance.NaniteHierarchyOffset;
 	Data[0].W  = *(const     float*)&LastUpdateFrame;
 
+#if 0
 	LocalToWorld.To3x4MatrixTranspose((float*)&Data[1]);
 	PrevLocalToWorld.To3x4MatrixTranspose((float*)&Data[4]);
 
@@ -97,6 +101,26 @@ void FInstanceSceneShaderData::Setup(
 	Data[8].W  = *(const     float*)&RandomID;
 
 	Data[9]    = *(const  FVector4*)&LightMapShadowMapUVBias;
+#else
+	Data[1]    = *(const FVector4* )&CompressedLocalToWorld.Rotation[0];
+	Data[2]    = *(const FVector3f*)&CompressedLocalToWorld.Translation;
+
+	Data[3]    = *(const FVector4* )&CompressedPrevLocalToWorld.Rotation[0];
+	Data[4]    = *(const FVector3f*)&CompressedPrevLocalToWorld.Translation;
+
+	const FVector3f BoundsOrigin = Instance.LocalBounds.GetCenter();
+	const FVector3f BoundsExtent = Instance.LocalBounds.GetExtent();
+	
+	Data[5]    = *(const FVector3f*)&BoundsOrigin;
+	Data[5].W  = *(const     float*)&BoundsExtent.X;
+	
+	Data[6].X  = *(const     float*)&BoundsExtent.Y;
+	Data[6].Y  = *(const     float*)&BoundsExtent.Z;
+	Data[6].Z  = *(const     float*)&PayloadDataOffset;
+	Data[6].W  = *(const     float*)&RandomID;
+
+	Data[7]    = *(const  FVector4*)&LightMapShadowMapUVBias;
+#endif
 }
 
 ENGINE_API const FInstanceSceneShaderData& GetDummyInstanceSceneShaderData()
diff --git a/Engine/Source/Runtime/Engine/Public/InstanceUniformShaderParameters.h b/Engine/Source/Runtime/Engine/Public/InstanceUniformShaderParameters.h
--- a/Engine/Source/Runtime/Engine/Public/InstanceUniformShaderParameters.h
+++ b/Engine/Source/Runtime/Engine/Public/InstanceUniformShaderParameters.h
@@ -44,7 +44,7 @@ struct FPrimitiveInstance
 		if (PrimitiveToWorld.IsScaleNonUniform())
 	#endif
 		{
-			LocalToWorld.Orthonormalize();
+			LocalToWorld.Orthogonalize();
 		}
 
 		return LocalToWorld;
@@ -70,7 +70,7 @@ struct FPrimitiveInstanceDynamicData
 		if (PrevPrimitiveToWorld.IsScaleNonUniform())
 	#endif
 		{
-			PrevLocalToWorld.Orthonormalize();
+			PrevLocalToWorld.Orthogonalize();
 		}
 
 		return PrevLocalToWorld;
@@ -95,7 +95,7 @@ FORCEINLINE FPrimitiveInstance ConstructPrimitiveInstance(
 struct FInstanceSceneShaderData
 {
 	// Must match GetInstanceSceneData() in SceneData.ush
-	enum { DataStrideInFloat4s = 10 };
+	enum { DataStrideInFloat4s = 8 };
 
 	TStaticArray<FVector4, DataStrideInFloat4s> Data;
 
diff --git a/Engine/Source/Runtime/RenderCore/Public/RenderTransform.h b/Engine/Source/Runtime/RenderCore/Public/RenderTransform.h
--- a/Engine/Source/Runtime/RenderCore/Public/RenderTransform.h
+++ b/Engine/Source/Runtime/RenderCore/Public/RenderTransform.h
@@ -173,7 +173,7 @@ public:
 		);
 	}
 
-	FORCEINLINE void Orthonormalize()
+	FORCEINLINE void Orthogonalize()
 	{
 		FVector3f X = TransformRows[0];
 		FVector3f Y = TransformRows[1];
@@ -353,3 +353,197 @@ public:
 		return Ar;
 	}
 };
+
+
+// [Frisvad 2012, "Building an Orthonormal Basis from a 3D Unit Vector Without Normalization"]
+FORCEINLINE void GetHemiOrthoBasis( FVector3f& BasisX, FVector3f& BasisY, const FVector3f& BasisZ )
+{
+	float A = 1.0f / ( 1.0f + BasisZ.Z );
+	float B = -BasisZ.X * BasisZ.Y * A;
+	BasisX = FVector3f( 1.0f - BasisZ.X * BasisZ.X * A, B, -BasisZ.X );
+	BasisY = FVector3f( B, 1.0f - BasisZ.Y * BasisZ.Y * A, -BasisZ.Y );
+}
+
+FORCEINLINE FVector2D UnitVectorToHemiOctahedron( const FVector3f& N )
+{
+	float Sum = FMath::Abs( N.X ) + FMath::Abs( N.Y ) + FMath::Abs( N.Z );
+	return FVector2D( N.X + N.Y, N.X - N.Y ) / Sum;
+}
+
+FORCEINLINE FVector3f HemiOctahedronToUnitVector( const FVector2D& Oct )
+{
+	FVector3f N;
+	N.X = Oct.X + Oct.Y;
+	N.Y = Oct.X - Oct.Y;
+	N.Z = 2.0f - FMath::Abs( N.X ) - FMath::Abs( N.Y );
+	return N.GetUnsafeNormal();
+}
+
+struct FCompressedTransform
+{
+	static constexpr uint32 ExpBits				= 8;
+	static constexpr uint32 ExpBias				= ( 1 << (ExpBits - 1) ) - 1;
+	static constexpr uint32 SignMantissaBits	= 16;
+	static constexpr uint32 SignMantissaMask	= (1 << SignMantissaBits) - 1;
+	static constexpr uint32 MantissaBits		= SignMantissaBits - 1;
+
+	FVector3f	Translation;		// 4B padding
+	uint16		Rotation[4];		// 2B padding
+	uint16		Scale_SharedExp[4];	// 1B padding if SignMantissaBits == 16
+
+	FORCEINLINE FCompressedTransform() {}
+	FORCEINLINE FCompressedTransform( const FRenderTransform& In )
+	{
+		Translation = In.Origin;
+
+		FVector3f Scale;
+		FVector3f Axis[3];
+		for( int i = 0; i < 3; i++ )
+		{
+			Axis[i] = In.TransformRows[i];
+			Scale[i] = Axis[i].Size();
+			Axis[i] /= Scale[i];
+		}
+
+		// Rotation
+		{
+			if( Axis[2].Z < 0.0f )
+			{
+				Axis[2]  *= -1.0f;
+				Scale[2] *= -1.0f;
+			}
+
+			FVector2D OctZ = UnitVectorToHemiOctahedron( Axis[2] );
+
+			FVector3f BasisX, BasisY;
+			GetHemiOrthoBasis( BasisX, BasisY, Axis[2] );
+
+			float X = Axis[0] | BasisX;
+			float Y = Axis[0] | BasisY;
+
+			float aX = FMath::Abs( X );
+			float aY = FMath::Abs( Y );
+
+			bool bSpinIsX = aX < aY;
+			float Spin0 = bSpinIsX ? X : Y;
+			float Spin1 = bSpinIsX ? Y : X;
+			float Sign1 = Spin1 < 0.0f ? -1.0f : 1.0f;
+		
+			//Axis[0]	*= Sign1;
+			Scale[0]*= Sign1;
+			Spin0	*= Sign1;
+
+			FVector3f GeneratedY = Axis[2] ^ Axis[0];
+			Scale[1] *= ( Axis[1] | GeneratedY ) < 0.0f ? -Sign1 : Sign1;
+
+#if 1
+			// Avoid sign extension in shader by biasing
+			Rotation[0] = FMath::RoundToInt( OctZ.X * 32767.0f ) + 32768;
+			Rotation[1] = FMath::RoundToInt( OctZ.Y * 32767.0f ) + 32768;
+			Rotation[2] = FMath::RoundToInt( Spin0  * 16383.0f * 1.41421356f );	// sqrt(2)
+			
+			Rotation[2] = ( Rotation[2] + 16384 ) & 0x7fff;
+			Rotation[2] |= bSpinIsX ? (1 << 15) : 0;
+#else
+			Rotation[0]  = ( FMath::RoundToInt( OctZ.X * 32767.0f ) + 32768 ) << 0;
+			Rotation[0] |= ( FMath::RoundToInt( OctZ.Y * 32767.0f ) + 32768 ) << 16;
+			Rotation[1] = FMath::RoundToInt( Spin0  * 16383.0f * 1.41421356f );	// sqrt(2)
+			
+			Rotation[1] = ( Rotation[1] + 16384 ) & 0x7fff;
+			Rotation[1] |= bSpinIsX ? (1 << 15) : 0;
+#endif
+		}
+
+		// Scale
+		{
+			FFloat32 MaxComponent = Scale.GetAbsMax();
+
+			// Need +1 because of losing the implicit leading bit of mantissa
+			// TODO assumes ExpBits == 8
+			// TODO clamp to expressable range
+			uint32 SharedExp = MaxComponent.Components.Exponent + 1;
+	
+			FFloat32 ExpScale( 1.0f );
+			ExpScale.Components.Exponent = 127 + ExpBias + MantissaBits - SharedExp;
+
+			if( FMath::RoundToInt( MaxComponent.FloatValue * ExpScale.FloatValue ) == (1 << MantissaBits) )
+			{
+				// Mantissa rounded up
+				SharedExp++;
+				ExpScale.FloatValue *= 0.5f;
+			}
+
+#if 0
+			Scale_SharedExp = (uint64)SharedExp << ( SignMantissaBits * 3 );
+			for( int i = 0; i < 3; i++ )
+			{
+				uint32 Mantissa = FMath::RoundToInt( Scale[i] * ExpScale.FloatValue ) + (1 << MantissaBits);
+				Scale_SharedExp |= (uint64)Mantissa << ( SignMantissaBits * i );
+			}
+#else
+			Scale_SharedExp[3] = SharedExp;
+			for( int i = 0; i < 3; i++ )
+			{
+				Scale_SharedExp[i] = FMath::RoundToInt( Scale[i] * ExpScale.FloatValue ) + (1 << MantissaBits);
+			}
+#endif
+		}
+	}
+
+	FORCEINLINE FRenderTransform ToRenderTransform() const
+	{
+		FRenderTransform Out;
+
+		Out.Origin = Translation;
+
+		// Rotation
+		{
+			FVector2D OctZ;
+			float Spin0;
+			OctZ.X = ( (int32)Rotation[0] - 32768 ) * (1.0f / 32767.0f);
+			OctZ.Y = ( (int32)Rotation[1] - 32768 ) * (1.0f / 32767.0f);
+			Spin0  = ( (int32)( Rotation[2] & 0x7fff ) - 16384 ) * (0.70710678f / 16383.0f);	// rsqrt(2)
+			bool bSpinIsX = Rotation[2] > 0x7fff;
+
+			Out.TransformRows[2] = HemiOctahedronToUnitVector( OctZ );
+
+			FVector3f BasisX, BasisY;
+			GetHemiOrthoBasis( BasisX, BasisY, Out.TransformRows[2] );
+
+			float Spin1 = FMath::Sqrt( 1.0f - Spin0 * Spin0 );
+			float X = bSpinIsX ? Spin0 : Spin1;
+			float Y = bSpinIsX ? Spin1 : Spin0;
+
+			Out.TransformRows[0] = BasisX * X + BasisY * Y;
+			Out.TransformRows[1] = Out.TransformRows[2] ^ Out.TransformRows[0];
+		}
+
+		// Scale
+		{
+			//uint32 SharedExp = Scale_SharedExp >> ( SignMantissaBits * 3 );
+			uint32 SharedExp = Scale_SharedExp[3];
+
+			FFloat32 ExpScale( 1.0f );
+			ExpScale.Components.Exponent = 127 - ExpBias - MantissaBits + SharedExp;
+
+#if 0
+			for( int i = 0; i < 3; i++ )
+			{
+				int32 Mantissa = ( Scale_SharedExp >> ( SignMantissaBits * i ) ) & SignMantissaMask;
+				float Scale = Mantissa - (1 << MantissaBits);
+				Scale *= ExpScale.FloatValue;
+				Out.TransformRows[i] *= Scale;
+			}
+#else
+			for( int i = 0; i < 3; i++ )
+			{
+				float Scale = (int32)Scale_SharedExp[i] - (1 << MantissaBits);
+				Scale *= ExpScale.FloatValue;
+				Out.TransformRows[i] *= Scale;
+			}
+#endif
+		}
+
+		return Out;
+	}
+};
\ No newline at end of file
