Iterate over all the TargetPlatforms to find which ShaderPlatform they support to flip the correct flags

Currently no platform handles support different but there may be a case later in which some TargetPlatform has a different Support/NoSupport then another
This would cause the Mask to possibly be Enable/Disabled for a specific platform that could be incorrect. Overall a global mask is not an ideal setup for this

diff --git a/Engine/Source/Runtime/RenderCore/Private/RenderUtils.cpp b/Engine/Source/Runtime/RenderCore/Private/RenderUtils.cpp
--- a/Engine/Source/Runtime/RenderCore/Private/RenderUtils.cpp
+++ b/Engine/Source/Runtime/RenderCore/Private/RenderUtils.cpp
@@ -15,6 +15,7 @@
 #include "Misc/CoreMisc.h"
 #include "Interfaces/ITargetPlatform.h"
 #include "Interfaces/ITargetPlatformManagerModule.h"
+#include "RHIShaderFormatDefinitions.inl"
 #endif
 
 FBufferWithRDG::FBufferWithRDG() = default;
@@ -1306,13 +1307,16 @@ RENDERCORE_API void RenderUtilsInit()
 	ITargetPlatformManagerModule* TargetPlatformManager = GetTargetPlatformManager();
 	if (TargetPlatformManager)
 	{
-		for (uint32 ShaderPlatformIndex = 0; ShaderPlatformIndex < SP_NumPlatforms; ++ShaderPlatformIndex)
+		for (ITargetPlatform* TargetPlatform : TargetPlatformManager->GetTargetPlatforms())
 		{
-			EShaderPlatform ShaderPlatform = EShaderPlatform(ShaderPlatformIndex);
-			FName PlatformName = ShaderPlatformToPlatformName(ShaderPlatform);
-			ITargetPlatform* TargetPlatform = TargetPlatformManager->FindTargetPlatform(PlatformName.ToString());
-			if (TargetPlatform)
+			TArray<FName> PlaformShaderFormats;
+			TargetPlatform->GetAllPossibleShaderFormats(PlaformShaderFormats);
+
+			for (FName Format : PlaformShaderFormats)
 			{
+				EShaderPlatform ShaderPlatform = ShaderFormatNameToShaderPlatform(Format);
+				uint32 ShaderPlatformIndex = static_cast<uint32>(ShaderPlatform);
+
 				uint64 Mask = 1ull << ShaderPlatformIndex;
 
 				if (TargetPlatform->UsesForwardShading())
