Add ability to set shader defines for global shaders from config files. This allows platform-specific tuning of global shader permutations for performance reasons.

diff --git a/Engine/Plugins/Compositing/OpenColorIO/Source/OpenColorIO/Public/OpenColorIOShaderType.h b/Engine/Plugins/Compositing/OpenColorIO/Source/OpenColorIO/Public/OpenColorIOShaderType.h
--- a/Engine/Plugins/Compositing/OpenColorIO/Source/OpenColorIO/Public/OpenColorIOShaderType.h
+++ b/Engine/Plugins/Compositing/OpenColorIO/Source/OpenColorIO/Public/OpenColorIOShaderType.h
@@ -30,12 +30,12 @@ class FUniformExpressionSet;
 /** Called for every OpenColorIO shader to update the appropriate stats. */
 extern void UpdateOpenColorIOShaderCompilingStats(const FOpenColorIOTransformResource* InShader);
 
-struct FOpenColorIOShaderPermutationParameters : public FShaderPermutationParameters
+struct FOpenColorIOShaderPermutationParameters : public FGlobalShaderPermutationParameters
 {
 	const FOpenColorIOTransformResource* Transform;
 
-	FOpenColorIOShaderPermutationParameters(EShaderPlatform InPlatform, const FOpenColorIOTransformResource* InTransform)
-		: FShaderPermutationParameters(InPlatform)
+	FOpenColorIOShaderPermutationParameters(FName InGlobalShaderName, EShaderPlatform InPlatform, const FOpenColorIOTransformResource* InTransform)
+		: FGlobalShaderPermutationParameters(InGlobalShaderName, InPlatform)
 		, Transform(InTransform)
 	{}
 };
@@ -120,7 +120,7 @@ public:
 	 */
 	bool ShouldCache(EShaderPlatform InPlatform, const FOpenColorIOTransformResource* InColorTransform) const
 	{
-		return ShouldCompilePermutation(FOpenColorIOShaderPermutationParameters(InPlatform, InColorTransform));
+		return ShouldCompilePermutation(FOpenColorIOShaderPermutationParameters(GetFName(), InPlatform, InColorTransform));
 	}
 
 
@@ -132,6 +132,6 @@ protected:
 	 */
 	void SetupCompileEnvironment(EShaderPlatform InPlatform, const FOpenColorIOTransformResource* InColorTransform, FShaderCompilerEnvironment& OutEnvironment) const
 	{
-		ModifyCompilationEnvironment(FOpenColorIOShaderPermutationParameters(InPlatform, InColorTransform), OutEnvironment);
+		ModifyCompilationEnvironment(FOpenColorIOShaderPermutationParameters(GetFName(), InPlatform, InColorTransform), OutEnvironment);
 	}
 };
diff --git a/Engine/Source/Runtime/Engine/Private/ShaderCompiler/ShaderCompiler.cpp b/Engine/Source/Runtime/Engine/Private/ShaderCompiler/ShaderCompiler.cpp
--- a/Engine/Source/Runtime/Engine/Private/ShaderCompiler/ShaderCompiler.cpp
+++ b/Engine/Source/Runtime/Engine/Private/ShaderCompiler/ShaderCompiler.cpp
@@ -5872,7 +5872,7 @@ void VerifyGlobalShaders(EShaderPlatform Platform, const ITargetPlatform* Target
 		const FShaderPipelineType* Pipeline = *ShaderPipelineIt;
 		if (Pipeline->IsGlobalTypePipeline())
 		{
-			if (Pipeline->ShouldCompilePermutation(FGlobalShaderPermutationParameters(Platform, kUniqueShaderPermutationId, PermutationFlags))
+			if (FGlobalShaderType::ShouldCompilePipeline(Pipeline, Platform, PermutationFlags)
 				&& (!GlobalShaderMap->HasShaderPipeline(Pipeline) || (OutdatedShaderPipelineTypes && OutdatedShaderPipelineTypes->Contains(Pipeline))))
 			{
 				if (OutdatedShaderPipelineTypes)
diff --git a/Engine/Source/Runtime/RenderCore/Private/GlobalShader.cpp b/Engine/Source/Runtime/RenderCore/Private/GlobalShader.cpp
--- a/Engine/Source/Runtime/RenderCore/Private/GlobalShader.cpp
+++ b/Engine/Source/Runtime/RenderCore/Private/GlobalShader.cpp
@@ -15,9 +15,12 @@
 #include "Serialization/MemoryReader.h"
 #include "Misc/ScopedSlowTask.h"
 #include "Misc/App.h"
+#include "Misc/DataDrivenPlatformInfoRegistry.h"
+#include "Misc/CoreMisc.h"
+#include "Interfaces/ITargetPlatformManagerModule.h"
+#include "Containers/StaticBitArray.h"
 #include "StaticBoundShaderState.h"
 
-
 /** The global shader map. */
 FGlobalShaderMap* GGlobalShaderMap[SP_NumPlatforms] = {};
 
@@ -26,11 +29,341 @@ IMPLEMENT_TYPE_LAYOUT(FGlobalShaderMapContent);
 
 IMPLEMENT_SHADER_TYPE(,FNULLPS,TEXT("/Engine/Private/NullPixelShader.usf"),TEXT("Main"),SF_Pixel);
 
+/** A per-platform map of global shader defines that come from config .ini */
+class FGlobalShaderConfigDefines
+{
+public:
+
+	static void ApplyConfigDefines(const FGlobalShaderPermutationParameters& Parameters, FShaderCompilerEnvironment& OutEnvironment)
+	{
+		FName ShaderFormat;
+		if (const FPlatformInfo* Platform = GetPlatformInfoAndErrorCheck(Parameters.Platform, ShaderFormat))
+		{
+			if (const TArray<FDefine>* ShaderDefines = Platform->ShaderDefines.Find(Parameters.GlobalShaderName))
+			{			
+				for (const FDefine& Define : *ShaderDefines)
+				{
+					if (Define.IsRelevant(ShaderFormat, Parameters.PermutationId))
+					{
+						OutEnvironment.SetDefine(*Define.Name, Define.Value);
+					}
+				}
+			}
+		}
+	}
+
+	static void AppendKeyString(FString& KeyString, FName ShaderName, int32 PermutationID, FName IniPlatform, EShaderPlatform ShaderPlatform)
+	{
+		FName ShaderFormat;
+		if (const FPlatformInfo* Platform = GetPlatformInfoAndErrorCheck(ShaderPlatform, ShaderFormat))
+		{
+			if (const TArray<FDefine>* ShaderDefines = Platform->ShaderDefines.Find(ShaderName))
+			{
+				TArray<FDefine> FilteredDefines;
+				for (const FDefine& Define : *ShaderDefines)
+				{
+					if (Define.IsRelevant(ShaderFormat, PermutationID))
+					{
+						FilteredDefines.Add(Define);
+					}
+				}
+
+				if (FilteredDefines.Num() > 0)
+				{
+					// Sort and hash the applicable defines
+					FilteredDefines.Sort();
+
+					FSHA1 HashState;
+					for (const FDefine& Define : FilteredDefines)
+					{
+						HashState.UpdateWithString(*Define.Name, Define.Name.Len());
+						HashState.UpdateWithString(*Define.Value, Define.Value.Len());
+
+						const uint32 ShaderFormatHash = GetTypeHash(Define.ShaderFormat);
+						HashState.Update((const uint8*)&ShaderFormatHash, sizeof(ShaderFormatHash));
+
+						HashState.Update((const uint8*)&Define.PermutationID, sizeof(Define.PermutationID));
+					}
+
+					KeyString += HashState.Finalize().ToString();
+				}
+			}
+		}
+	}
+
+private:
+	struct FDefine
+	{
+		FString Name;
+		FString Value;
+		FName ShaderFormat = NAME_None;
+		int32 PermutationID = INDEX_NONE;
+
+		bool IsRelevant(FName InShaderFormat, int32 InPermutationID) const
+		{
+			if (ShaderFormat != NAME_None && ShaderFormat != InShaderFormat)
+			{
+				// This define is inappropriate for the target shader format
+				return false;
+			}
+
+			if (PermutationID != INDEX_NONE && PermutationID != InPermutationID)
+			{
+				// This define is inappropriate for the permutation being compiled
+				return false;
+			}
+
+			return true;
+		}
+
+		int32 Compare(const FDefine& Other) const
+		{
+			int Cmp = Name.Compare(Other.Name, ESearchCase::IgnoreCase);
+			if (Cmp == 0)
+			{
+				Cmp = Value.Compare(Other.Value, ESearchCase::CaseSensitive);
+				if (Cmp == 0)
+				{
+					Cmp = ShaderFormat.Compare(Other.ShaderFormat);
+					if (Cmp == 0)
+					{
+						Cmp = Other.PermutationID - PermutationID;
+					}
+				}
+			}
+
+			return Cmp;
+		}
+
+		bool operator==(const FDefine& Other) const
+		{
+			return Name.Equals(Other.Name, ESearchCase::IgnoreCase)
+				&& Value == Other.Value
+				&& ShaderFormat == Other.ShaderFormat
+				&& PermutationID == Other.PermutationID;
+		}
+
+		bool operator!=(const FDefine& Other) const
+		{
+			return !(*this == Other);
+		}
+
+		bool operator<(const FDefine& Other) const
+		{
+			return Compare(Other) < 0;
+		}
+	};
+
+	using FShaderDefines = TMap<FName, TArray<FDefine>>;
+
+	struct FPlatformInfo
+	{
+		FName Name;
+		bool bInitializedFromConfig = false;
+		FShaderDefines ShaderDefines;
+	};
+	
+	static TMap<FName, FPlatformInfo> ConfigDefines;
+	static TStaticBitArray<EShaderPlatform::SP_NumPlatforms> ErrorCheckedPlatforms;
+
+	static const FPlatformInfo* GetPlatformInfoAndErrorCheck(EShaderPlatform ShaderPlatform, FName& OutShaderFormat)
+	{
+		OutShaderFormat = LegacyShaderPlatformToShaderFormat(ShaderPlatform);
+
+		// Search for all target platforms that support this shader platform
+		TArray<FName, TInlineAllocator<16>> IniPlatforms;
+		for(const ITargetPlatform* TP : GetTargetPlatformManagerRef().GetTargetPlatforms())
+		{
+			check(TP);
+			TArray<FName> PlatformShaderFormats;
+			TP->GetAllPossibleShaderFormats(PlatformShaderFormats);
+			if (PlatformShaderFormats.Contains(OutShaderFormat))
+			{
+				IniPlatforms.AddUnique(FName(TP->IniPlatformName()));
+			}
+		}
+		
+		if (IniPlatforms.Num() == 0)
+		{
+			// No found platforms that support this shader format
+			return nullptr;
+		}
+
+		const FPlatformInfo& Platform = GetOrLoadPlatformInfo(IniPlatforms[0]);
+		
+		if (!ErrorCheckedPlatforms[ShaderPlatform])
+		{
+			ErrorCheckedPlatforms[ShaderPlatform] = true;
+
+			if (IniPlatforms.Num() > 1)
+			{
+				// This shader platform is shared by multiple target platforms that can be configured independently. We need to make sure all config defines
+				// match up, and that no platform-specific defines exist that might introduce shader compiler output that diverges between target platforms
+				for (int32 PlatformIndex = 1; PlatformIndex < IniPlatforms.Num(); ++PlatformIndex)
+				{
+					const FPlatformInfo& OtherPlatform = GetOrLoadPlatformInfo(IniPlatforms[PlatformIndex]);
+					ErrorCheckPlatformsForShaderFormat(Platform, OtherPlatform, OutShaderFormat);
+				}
+			}
+		}
+
+		return &Platform;
+	}
+
+	static FPlatformInfo& GetOrLoadPlatformInfo(FName PlatformName)
+	{
+		FPlatformInfo& Platform = ConfigDefines.FindOrAdd(PlatformName);
+
+		if (!Platform.bInitializedFromConfig)
+		{
+			InitializePlatform(Platform, PlatformName);
+		}
+
+		return Platform;
+	}
+
+	static void InitializePlatform(FPlatformInfo& Platform, FName PlatformName)
+	{
+		if (FConfigCacheIni* ConfigCache = FConfigCacheIni::ForPlatform(PlatformName))
+		{
+			TArray<FString> DefineStrings;
+			ConfigCache->GetArray(TEXT("GlobalShaderDefines"), TEXT("Definitions"), DefineStrings, GEngineIni);
+
+			for (const FString& DefineString : DefineStrings)
+			{				
+				FString ShaderName, DefineName;
+				if (FParse::Value(*DefineString, TEXT("Shader="), ShaderName) &&
+					FParse::Value(*DefineString, TEXT("Define="), DefineName))
+				{
+					FName ShaderFormat = NAME_None;
+					FParse::Value(*DefineString, TEXT("ShaderFormat="), ShaderFormat);
+
+					int32 PermutationID = INDEX_NONE;
+					FParse::Value(*DefineString, TEXT("PermutationID="), PermutationID);
+
+					FString DefineValue;
+					int32 EqualsIndex = INDEX_NONE;
+					if (DefineName.FindChar(TCHAR('='), EqualsIndex))
+					{
+						DefineValue = DefineName.Mid(EqualsIndex + 1);
+						DefineName.MidInline(0, EqualsIndex, false);
+					}
+
+					const FShaderType* ShaderType = FindShaderTypeByName(ShaderName);
+					if (ShaderType == nullptr || ShaderType->GetGlobalShaderType() == nullptr)
+					{
+						// This global shader doesn't actually exist
+						UE_LOG(LogShaders, Error, TEXT("Global shader definition '%s' found in engine config for global shader '%s', which does not exist"), *DefineName, *ShaderName);
+						continue;
+					}
+
+					TArray<FDefine>& ShaderDefines = Platform.ShaderDefines.FindOrAdd(FName(ShaderName));
+					FDefine* Define = ShaderDefines.FindByPredicate(
+						[&](const FDefine& Define)
+						{
+							return Define.Name.Equals(DefineName, ESearchCase::IgnoreCase) &&
+								Define.ShaderFormat == ShaderFormat &&
+								Define.PermutationID == PermutationID;
+						});
+
+					if (!Define)
+					{
+						Define = &ShaderDefines.AddDefaulted_GetRef();
+						Define->Name = DefineName;
+						Define->ShaderFormat = ShaderFormat;
+						Define->PermutationID = PermutationID;
+					}
+
+					Define->Value = DefineValue;
+				}
+			}			
+		}		
+
+		Platform.Name = PlatformName;
+		Platform.bInitializedFromConfig = true;
+	}
+
+	static void ErrorCheckPlatformsForShaderFormat(const FPlatformInfo& PlatformA, const FPlatformInfo& PlatformB, FName ShaderFormat)
+	{
+		auto FilterShaderDefines = [ShaderFormat](const FShaderDefines& ShaderDefines) -> FShaderDefines
+		{
+			FShaderDefines FilteredShaderDefines;
+
+			for (const auto& ShaderIt : ShaderDefines)
+			{
+				TArray<FDefine> FilteredDefines = ShaderIt.Value.FilterByPredicate(
+					[ShaderFormat](const FDefine& Define)
+					{
+						return Define.ShaderFormat == NAME_None || Define.ShaderFormat == ShaderFormat;
+					});
+
+				if (FilteredDefines.Num() > 0)
+				{
+					// Sort the filtered defines for every shader for the purposes of comparison in error checking
+					FilteredDefines.Sort();
+					FilteredShaderDefines.Add(ShaderIt.Key, MoveTemp(FilteredDefines));
+				}
+			}
+
+			return MoveTemp(FilteredShaderDefines);
+		};
+
+		FShaderDefines FilteredDefinesA = FilterShaderDefines(PlatformA.ShaderDefines);
+		FShaderDefines FilteredDefinesB = FilterShaderDefines(PlatformB.ShaderDefines);
+
+		TSet<FName> ShadersWithError;
+
+		// Check if there are any keys in A that aren't in B, and vice versa
+		{
+			TSet<FName> ShaderNamesA, ShaderNamesB;
+			FilteredDefinesA.GetKeys(ShaderNamesA);
+			FilteredDefinesB.GetKeys(ShaderNamesB);
+
+			ShadersWithError.Append(ShaderNamesA.Difference(ShaderNamesB));
+			ShadersWithError.Append(ShaderNamesB.Difference(ShaderNamesA));
+		}
+
+		// Check if the defines (that are relevant to this shader format) are different between A or B
+		for (const auto& ShaderIt : FilteredDefinesA)
+		{
+			const TArray<FDefine>& DefinesA = ShaderIt.Value;
+			const TArray<FDefine>* DefinesB = FilteredDefinesB.Find(ShaderIt.Key);
+
+			if (DefinesB && DefinesA != *DefinesB)
+			{
+				ShadersWithError.Add(ShaderIt.Key);
+			}
+		}
+
+		if (ShadersWithError.Num() > 0)
+		{
+			FString ShaderListString = *FString::JoinBy(ShadersWithError, TEXT("\n    "), [](const FName& A) { return A.ToString(); });
+			UE_LOG(LogShaders, Fatal, TEXT("It has been detected that one or more global shaders are configured differently between platform '%s' and '%s' for ")
+				TEXT("shader format '%s'. This is unsupported, and could result in data mismatches in the Derived Data Cache. Please check that all entries ")
+				TEXT("of '[GlobalShaderDefines]` in the various engine config .ini files are the same for these shaders on all platforms that use this shader ")
+				TEXT("model.\n\nGlobal shaders with errors:\n    %s\n"),
+				*PlatformA.Name.ToString(), *PlatformB.Name.ToString(), *ShaderFormat.ToString(), *ShaderListString);
+		}
+	}
+};
+
+TMap<FName, FGlobalShaderConfigDefines::FPlatformInfo> FGlobalShaderConfigDefines::ConfigDefines;
+TStaticBitArray<EShaderPlatform::SP_NumPlatforms> FGlobalShaderConfigDefines::ErrorCheckedPlatforms;
+
 /** Used to identify the global shader map in compile queues. */
 const int32 GlobalShaderMapId = 0;
 
 FGlobalShaderMapId::FGlobalShaderMapId(EShaderPlatform Platform, const ITargetPlatform* TargetPlatform)
 {
+	if (TargetPlatform == nullptr)
+	{
+		TargetPlatform = GetTargetPlatformManagerRef().GetRunningTargetPlatform();
+		check(TargetPlatform);
+	}
+
+	ShaderPlatform = Platform;
+	IniPlatformName = FName(TargetPlatform->IniPlatformName());
+
 	LayoutParams.InitializeForPlatform(TargetPlatform);
 	const EShaderPermutationFlags PermutationFlags = GetShaderPermutationFlags(LayoutParams);
 	TArray<FShaderType*> ShaderTypes;
@@ -136,6 +469,9 @@ void FGlobalShaderMapId::AppendKeyString(FString& KeyString, const TArray<FShade
 		// Add the type's source hash so that we can invalidate cached shaders when .usf changes are made
 		KeyString += ShaderTypeDependency.SourceHash.ToString();
 
+		// Add the config define hash, if any defines exist in config
+		FGlobalShaderConfigDefines::AppendKeyString(KeyString, ShaderType->GetFName(), ShaderTypeDependency.PermutationId, IniPlatformName, ShaderPlatform);
+
 		if (const FShaderParametersMetadata* ParameterStructMetadata = ShaderType->GetRootParametersMetadata())
 		{
 			KeyString += FString::Printf(TEXT("%08x"), ParameterStructMetadata->GetLayoutHash());
@@ -196,10 +532,31 @@ void FGlobalShaderMapId::AppendKeyString(FString& KeyString, const TArray<FShade
 #endif // WITH_EDITOR
 }
 
+bool FGlobalShaderType::ShouldCompilePipeline(const FShaderPipelineType* ShaderPipelineType, EShaderPlatform Platform, EShaderPermutationFlags Flags)
+{
+	for (const FShaderType* ShaderType : ShaderPipelineType->GetStages())
+	{
+		const FGlobalShaderPermutationParameters Parameters(ShaderType->GetFName(), Platform, kUniqueShaderPermutationId, Flags);
+		checkSlow(ShaderType->GetGlobalShaderType());
+		if (!ShaderType->ShouldCompilePermutation(Parameters))
+		{
+			return false;
+		}
+	}
+	return true;
+}
+
 FGlobalShader::FGlobalShader(const ShaderMetaType::CompiledShaderInitializerType& Initializer)
 :	FShader(Initializer)
 {}
 
+void FGlobalShader::ModifyCompilationEnvironment(const FGlobalShaderPermutationParameters& Parameters, FShaderCompilerEnvironment& Environment)
+{
+	FShader::ModifyCompilationEnvironment(Parameters, Environment);
+
+	FGlobalShaderConfigDefines::ApplyConfigDefines(Parameters, Environment);
+}
+
 void BackupGlobalShaderMap(FGlobalShaderBackupData& OutGlobalShaderBackup)
 {
 #if 0
@@ -378,7 +735,7 @@ bool FGlobalShaderMap::IsComplete(const ITargetPlatform* TargetPlatform) const
 	{
 		const FShaderPipelineType* Pipeline = *ShaderPipelineIt;
 		if (Pipeline->IsGlobalTypePipeline()
-			&& Pipeline->ShouldCompilePermutation(FGlobalShaderPermutationParameters(Platform, kUniqueShaderPermutationId, PermutationFlags))
+			&& FGlobalShaderType::ShouldCompilePipeline(Pipeline, Platform, PermutationFlags)
 			&& !HasShaderPipeline(Pipeline))
 		{
 			return false;
diff --git a/Engine/Source/Runtime/RenderCore/Public/GlobalShader.h b/Engine/Source/Runtime/RenderCore/Public/GlobalShader.h
--- a/Engine/Source/Runtime/RenderCore/Public/GlobalShader.h
+++ b/Engine/Source/Runtime/RenderCore/Public/GlobalShader.h
@@ -35,6 +35,8 @@ public:
 
 private:
 	FPlatformTypeLayoutParameters LayoutParams;
+	EShaderPlatform ShaderPlatform;
+	FName IniPlatformName;
 
 	/** Shader types that this shader map is dependent on and their stored state. Mapped by shader filename, so every filename can have it's own DDC key. */
 	TMap<FString, TArray<FShaderTypeDependency>> ShaderFilenameToDependenciesMap;
@@ -43,7 +45,15 @@ private:
 	TArray<FShaderPipelineTypeDependency> ShaderPipelineTypeDependencies;
 };
 
-using FGlobalShaderPermutationParameters = FShaderPermutationParameters;
+struct FGlobalShaderPermutationParameters : public FShaderPermutationParameters
+{
+	const FName GlobalShaderName;
+
+	FGlobalShaderPermutationParameters(FName InGlobalShaderName, EShaderPlatform InShaderPlatform, int32 InPermutationId = 0, EShaderPermutationFlags InFlags = EShaderPermutationFlags::HasEditorOnlyData)
+		: FShaderPermutationParameters(InShaderPlatform, InPermutationId, InFlags)
+		, GlobalShaderName(InGlobalShaderName)
+	{}
+};
 
 /**
  * A shader meta type for the simplest shaders; shaders which are not material or vertex factory linked.
@@ -93,7 +103,7 @@ public:
 	 */
 	bool ShouldCompilePermutation(EShaderPlatform Platform, int32 PermutationId, EShaderPermutationFlags Flags) const
 	{
-		return FShaderType::ShouldCompilePermutation(FGlobalShaderPermutationParameters(Platform, PermutationId, Flags));
+		return FShaderType::ShouldCompilePermutation(FGlobalShaderPermutationParameters(GetFName(), Platform, PermutationId, Flags));
 	}
 
 	/**
@@ -104,8 +114,10 @@ public:
 	void SetupCompileEnvironment(EShaderPlatform Platform, int32 PermutationId, EShaderPermutationFlags Flags, FShaderCompilerEnvironment& Environment) const
 	{
 		// Allow the shader type to modify its compile environment.
-		ModifyCompilationEnvironment(FGlobalShaderPermutationParameters(Platform, PermutationId, Flags), Environment);
+		ModifyCompilationEnvironment(FGlobalShaderPermutationParameters(GetFName(), Platform, PermutationId, Flags), Environment);
 	}
+
+	static RENDERCORE_API bool ShouldCompilePipeline(const FShaderPipelineType* ShaderPipelineType, EShaderPlatform Platform, EShaderPermutationFlags Flags);
 };
 
 class RENDERCORE_API FGlobalShaderMapContent : public FShaderMapContent
@@ -239,6 +251,7 @@ class FGlobalShader : public FShader
 	DECLARE_EXPORTED_TYPE_LAYOUT(FGlobalShader, RENDERCORE_API, NonVirtual);
 public:
 	using ShaderMetaType = FGlobalShaderType;
+	using FPermutationParameters = FGlobalShaderPermutationParameters;
 
 	FGlobalShader() : FShader() {}
 
@@ -251,8 +264,12 @@ public:
 		SetUniformBufferParameter(RHICmdList, ShaderRHI, ViewUniformBufferParameter, ViewUniformBuffer);
 	}
 
+	static inline bool ShouldCompilePermutation(const FGlobalShaderPermutationParameters& Parameters)
+	{
+		return FShader::ShouldCompilePermutation(Parameters);
+	}
 	
-	
+	static RENDERCORE_API void ModifyCompilationEnvironment(const FGlobalShaderPermutationParameters& Parameters, FShaderCompilerEnvironment& Environment);
 };
 
 /**
