Non-unity compile fix

diff --git a/Engine/Source/Runtime/RenderCore/Private/GlobalShader.cpp b/Engine/Source/Runtime/RenderCore/Private/GlobalShader.cpp
--- a/Engine/Source/Runtime/RenderCore/Private/GlobalShader.cpp
+++ b/Engine/Source/Runtime/RenderCore/Private/GlobalShader.cpp
@@ -17,6 +17,7 @@
 #include "Misc/App.h"
 #include "Misc/DataDrivenPlatformInfoRegistry.h"
 #include "Misc/CoreMisc.h"
+#include "Misc/ConfigCacheIni.h"
 #include "Interfaces/ITargetPlatformManagerModule.h"
 #include "Containers/StaticBitArray.h"
 #include "StaticBoundShaderState.h"
