Reworked acquire / discard transitions in RDG to use split barriers to improve overlap.

diff --git a/Engine/Source/Runtime/D3D12RHI/Private/D3D12TransientResourceAllocator.cpp b/Engine/Source/Runtime/D3D12RHI/Private/D3D12TransientResourceAllocator.cpp
--- a/Engine/Source/Runtime/D3D12RHI/Private/D3D12TransientResourceAllocator.cpp
+++ b/Engine/Source/Runtime/D3D12RHI/Private/D3D12TransientResourceAllocator.cpp
@@ -135,14 +135,14 @@ FD3D12TransientResourceAllocator::FD3D12TransientResourceAllocator(FD3D12Transie
 	, AllocationInfoQueryDevice(GetParentAdapter()->GetDevice(0))
 {}
 
-FRHITransientTexture* FD3D12TransientResourceAllocator::CreateTexture(const FRHITextureCreateInfo& InCreateInfo, const TCHAR* InDebugName)
+FRHITransientTexture* FD3D12TransientResourceAllocator::CreateTexture(const FRHITextureCreateInfo& InCreateInfo, const TCHAR* InDebugName, uint32 InPassIndex)
 {
 	FD3D12DynamicRHI* DynamicRHI = FD3D12DynamicRHI::GetD3DRHI();
 
 	const D3D12_RESOURCE_DESC Desc = DynamicRHI->GetResourceDesc(InCreateInfo);
 	const D3D12_RESOURCE_ALLOCATION_INFO Info = AllocationInfoQueryDevice->GetResourceAllocationInfo(Desc);
 
-	return Allocator.CreateTexture(InCreateInfo, InDebugName, Info.SizeInBytes, Info.Alignment,
+	return Allocator.CreateTexture(InCreateInfo, InDebugName, InPassIndex, Info.SizeInBytes, Info.Alignment,
 		[&](const FRHITransientResourceAllocator::FResourceInitializer& Initializer)
 	{
 		ERHIAccess InitialState = ERHIAccess::UAVMask;
@@ -195,7 +195,7 @@ void FD3D12TransientResourceAllocator::FResourceAllocatorAdapter::AllocateResour
 #endif
 }
 
-FRHITransientBuffer* FD3D12TransientResourceAllocator::CreateBuffer(const FRHIBufferCreateInfo& InCreateInfo, const TCHAR* InDebugName)
+FRHITransientBuffer* FD3D12TransientResourceAllocator::CreateBuffer(const FRHIBufferCreateInfo& InCreateInfo, const TCHAR* InDebugName, uint32 InPassIndex)
 {
 	D3D12_RESOURCE_DESC Desc;
 	uint32 Alignment;
@@ -205,7 +205,7 @@ FRHITransientBuffer* FD3D12TransientResourceAllocator::CreateBuffer(const FRHIBu
 	Alignment = D3D12_DEFAULT_RESOURCE_PLACEMENT_ALIGNMENT;
 	uint64 Size = Align(Desc.Width, Alignment);
 
-	return Allocator.CreateBuffer(InCreateInfo, InDebugName, Size, Alignment,
+	return Allocator.CreateBuffer(InCreateInfo, InDebugName, InPassIndex, Size, Alignment,
 		[&](const FRHITransientResourceAllocator::FResourceInitializer& Initializer)
 	{
 		ED3D12ResourceTransientMode TransientMode = ED3D12ResourceTransientMode::Transient;
@@ -215,14 +215,14 @@ FRHITransientBuffer* FD3D12TransientResourceAllocator::CreateBuffer(const FRHIBu
 	});
 }
 
-void FD3D12TransientResourceAllocator::DeallocateMemory(FRHITransientTexture* InTexture)
+void FD3D12TransientResourceAllocator::DeallocateMemory(FRHITransientTexture* InTexture, uint32 InPassIndex)
 {
-	Allocator.DeallocateMemory(InTexture);
+	Allocator.DeallocateMemory(InTexture, InPassIndex);
 }
 
-void FD3D12TransientResourceAllocator::DeallocateMemory(FRHITransientBuffer* InBuffer)
+void FD3D12TransientResourceAllocator::DeallocateMemory(FRHITransientBuffer* InBuffer, uint32 InPassIndex)
 {
-	Allocator.DeallocateMemory(InBuffer);
+	Allocator.DeallocateMemory(InBuffer, InPassIndex);
 }
 
 void FD3D12TransientResourceAllocator::Freeze(FRHICommandListImmediate& RHICmdList)
diff --git a/Engine/Source/Runtime/D3D12RHI/Private/D3D12TransientResourceAllocator.h b/Engine/Source/Runtime/D3D12RHI/Private/D3D12TransientResourceAllocator.h
--- a/Engine/Source/Runtime/D3D12RHI/Private/D3D12TransientResourceAllocator.h
+++ b/Engine/Source/Runtime/D3D12RHI/Private/D3D12TransientResourceAllocator.h
@@ -46,10 +46,10 @@ public:
 	FD3D12TransientResourceAllocator(FD3D12TransientResourceSystem& InParentSystem);
 
 	//! IRHITransientResourceAllocator Overrides
-	FRHITransientTexture* CreateTexture(const FRHITextureCreateInfo& InCreateInfo, const TCHAR* InDebugName) override;
-	FRHITransientBuffer* CreateBuffer(const FRHIBufferCreateInfo& InCreateInfo, const TCHAR* InDebugName) override;
-	void DeallocateMemory(FRHITransientTexture* InTexture) override;
-	void DeallocateMemory(FRHITransientBuffer* InBuffer) override;
+	FRHITransientTexture* CreateTexture(const FRHITextureCreateInfo& InCreateInfo, const TCHAR* InDebugName, uint32 InPassIndex) override;
+	FRHITransientBuffer* CreateBuffer(const FRHIBufferCreateInfo& InCreateInfo, const TCHAR* InDebugName, uint32 InPassIndex) override;
+	void DeallocateMemory(FRHITransientTexture* InTexture, uint32 InPassIndex) override;
+	void DeallocateMemory(FRHITransientBuffer* InBuffer, uint32 InPassIndex) override;
 	void Freeze(FRHICommandListImmediate& RHICmdList) override;
 
 private:
diff --git a/Engine/Source/Runtime/RHI/Private/RHI.cpp b/Engine/Source/Runtime/RHI/Private/RHI.cpp
--- a/Engine/Source/Runtime/RHI/Private/RHI.cpp
+++ b/Engine/Source/Runtime/RHI/Private/RHI.cpp
@@ -219,7 +219,8 @@ FString GetTextureCreateFlagsName(ETextureCreateFlags TextureCreateFlags)
 		case TexCreate_VideoDecode:						return TEXT("TexCreate_VideoDecode");
 		case TexCreate_Dynamic:							return TEXT("TexCreate_Dynamic");
 		case TexCreate_InputAttachmentRead:				return TEXT("TexCreate_InputAttachmentRead");
-		case TexCreate_Foveation:						return TEXT("(TexCreate_Foveation | TexCreate_3DTiling");
+		case TexCreate_Foveation:						return TEXT("TexCreate_Foveation");
+		case TexCreate_3DTiling:						return TEXT("TexCreate_3DTiling");
 		case TexCreate_Memoryless:						return TEXT("TexCreate_Memoryless");
 		case TexCreate_GenerateMipCapable:				return TEXT("TexCreate_GenerateMipCapable");
 		case TexCreate_FastVRAMPartialAlloc:			return TEXT("TexCreate_FastVRAMPartialAlloc");
@@ -2687,9 +2688,9 @@ void FRHIBufferViewCache::SetDebugName(const TCHAR* DebugName)
 
 #endif
 
-void FRHITransientTexture::Init(const TCHAR* InName, uint32 InAllocationIndex)
+void FRHITransientTexture::Init(const TCHAR* InName, uint32 InAllocationIndex, uint32 InPassIndex)
 {
-	FRHITransientResource::Init(InName, InAllocationIndex);
+	FRHITransientResource::Init(InName, InAllocationIndex, InPassIndex);
 	ViewCache.SetDebugName(InName);
 
 #if !(UE_BUILD_SHIPPING || UE_BUILD_TEST)
@@ -2697,9 +2698,9 @@ void FRHITransientTexture::Init(const TCHAR* InName, uint32 InAllocationIndex)
 #endif
 }
 
-void FRHITransientBuffer::Init(const TCHAR* InName, uint32 InAllocationIndex)
+void FRHITransientBuffer::Init(const TCHAR* InName, uint32 InAllocationIndex, uint32 InPassIndex)
 {
-	FRHITransientResource::Init(InName, InAllocationIndex);
+	FRHITransientResource::Init(InName, InAllocationIndex, InPassIndex);
 	ViewCache.SetDebugName(InName);
 
 	// TODO: Add method to rename a buffer.
diff --git a/Engine/Source/Runtime/RHI/Private/RHIValidation.cpp b/Engine/Source/Runtime/RHI/Private/RHIValidation.cpp
--- a/Engine/Source/Runtime/RHI/Private/RHIValidation.cpp
+++ b/Engine/Source/Runtime/RHI/Private/RHIValidation.cpp
@@ -214,11 +214,11 @@ void FValidationRHI::RHICreateTransition(FRHITransition* Transition, const FRHIT
 		}
 	}
 
-	TArray<FOperation> SignalOps, WaitOps, AcquireOps, DiscardOps, BeginOps, EndOps;
+	TArray<FOperation> SignalOps, WaitOps, AliasingOps, AliasingOverlapOps, BeginOps, EndOps;
 	SignalOps .Reserve(Fences.Num());
 	WaitOps   .Reserve(Fences.Num());
-	AcquireOps.Reserve(CreateInfo.AliasingInfos.Num());
-	DiscardOps.Reserve(CreateInfo.AliasingInfos.Num());
+	AliasingOps.Reserve(CreateInfo.AliasingInfos.Num());
+	AliasingOverlapOps.Reserve(CreateInfo.AliasingInfos.Num());
 	BeginOps  .Reserve(CreateInfo.TransitionInfos.Num());
 	EndOps    .Reserve(CreateInfo.TransitionInfos.Num());
 
@@ -255,7 +255,7 @@ void FValidationRHI::RHICreateTransition(FRHITransition* Transition, const FRHIT
 			checkf(Resource->TransientState.bTransient, TEXT("Acquiring resource %s which is not transient. Only transient resources can be acquired."), Resource->GetDebugName());
 			checkf(SrcPipelines == ERHIPipeline::Graphics, TEXT("Acquiring a transient resource (%s) must begin on the graphics pipe."), Resource->GetDebugName());
 
-			AcquireOps.Emplace(FOperation::AcquireTransientResource(Resource, nullptr));
+			AliasingOps.Emplace(FOperation::AcquireTransientResource(Resource, nullptr));
 
 			for (const FRHITransientAliasingOverlap& Overlap : Info.Overlaps)
 			{
@@ -272,7 +272,7 @@ void FValidationRHI::RHICreateTransition(FRHITransition* Transition, const FRHIT
 
 				checkf(ResourceBefore, TEXT("Null resource provided as an aliasing overlap of %s"), Resource->GetDebugName());
 
-				AcquireOps.Emplace(FOperation::AliasingOverlap(ResourceBefore, Resource, nullptr));
+				AliasingOverlapOps.Emplace(FOperation::AliasingOverlap(ResourceBefore, Resource, nullptr));
 			}
 		}
 		else
@@ -281,7 +281,7 @@ void FValidationRHI::RHICreateTransition(FRHITransition* Transition, const FRHIT
 			checkf(Resource->TransientState.bTransient, TEXT("Discarding resource %s which is not transient. Only transient resources can be discarded."), Resource->GetDebugName());
 			checkf(DstPipelines == ERHIPipeline::Graphics, TEXT("Discarding a transient resource (%s) must end on the graphics pipe."), Resource->GetDebugName());
 
-			DiscardOps.Emplace(FOperation::DiscardTransientResource(Resource, nullptr));
+			AliasingOps.Emplace(FOperation::DiscardTransientResource(Resource, nullptr));
 		}
 	}
 
@@ -328,22 +328,20 @@ void FValidationRHI::RHICreateTransition(FRHITransition* Transition, const FRHIT
 	{
 		void* Backtrace = CaptureBacktrace();
 
-		for (FOperation& Op : AcquireOps)
+		for (FOperation& Op : AliasingOps)
 		{
 			switch (Op.Type)
 			{
 			case EOpType::AcquireTransient:
 				Op.Data_AcquireTransient.CreateBacktrace = Backtrace;
 				break;
-			case EOpType::AliasingOverlap:
-				Op.Data_AliasingOverlap.CreateBacktrace = Backtrace;
+			case EOpType::DiscardTransient:
+				Op.Data_DiscardTransient.CreateBacktrace = Backtrace;
 				break;
-			default:
-				checkNoEntry();
 			}
 		}
 
-		for (FOperation& Op : DiscardOps) { Op.Data_DiscardTransient.CreateBacktrace = Backtrace; }
+		for (FOperation& Op : AliasingOverlapOps) { Op.Data_AliasingOverlap.CreateBacktrace = Backtrace; }
 		for (FOperation& Op : BeginOps) { Op.Data_BeginTransition.CreateBacktrace = Backtrace; }
 		for (FOperation& Op : EndOps) { Op.Data_EndTransition.CreateBacktrace = Backtrace; }
 	}
@@ -355,8 +353,8 @@ void FValidationRHI::RHICreateTransition(FRHITransition* Transition, const FRHIT
 
 	Transition->PendingSignals.Operations = MoveTemp(SignalOps);
 	Transition->PendingWaits.Operations = MoveTemp(WaitOps);
-	Transition->PendingAcquires.Operations = MoveTemp(AcquireOps);
-	Transition->PendingDiscards.Operations = MoveTemp(DiscardOps);
+	Transition->PendingAliases.Operations = MoveTemp(AliasingOps);
+	Transition->PendingAliasingOverlaps.Operations = MoveTemp(AliasingOverlapOps);
 	Transition->PendingOperationsBegin.Operations = MoveTemp(BeginOps);
 	Transition->PendingOperationsEnd.Operations = MoveTemp(EndOps);
 
@@ -1090,7 +1088,7 @@ namespace RHIValidation
 
 		if (Resource->TransientState.bTransient)
 		{
-			RHI_VALIDATION_CHECK(Resource->TransientState.IsAcquired(), *GetReasonString_TransitionWithoutAcquire(Resource));
+			RHI_VALIDATION_CHECK(Resource->TransientState.IsAcquired() || (Resource->TransientState.IsDiscarded() && TargetState.Access == ERHIAccess::Discard), *GetReasonString_TransitionWithoutAcquire(Resource));
 		}
 
 		// Check we're not already transitioning
@@ -1154,7 +1152,7 @@ namespace RHIValidation
 
 		if (Resource->TransientState.bTransient)
 		{
-			RHI_VALIDATION_CHECK(Resource->TransientState.IsAcquired(), *GetReasonString_TransitionWithoutAcquire(Resource));
+			RHI_VALIDATION_CHECK(Resource->TransientState.IsAcquired() || (Resource->TransientState.IsDiscarded() && TargetState.Access == ERHIAccess::Discard), *GetReasonString_TransitionWithoutAcquire(Resource));
 		}
 
 		// Check that the end matches the begin.
@@ -1509,12 +1507,12 @@ FValidationTransientResourceAllocator::~FValidationTransientResourceAllocator()
 	check(bReleased);
 }
 
-FRHITransientTexture* FValidationTransientResourceAllocator::CreateTexture(const FRHITextureCreateInfo& InCreateInfo, const TCHAR* InDebugName)
+FRHITransientTexture* FValidationTransientResourceAllocator::CreateTexture(const FRHITextureCreateInfo& InCreateInfo, const TCHAR* InDebugName, uint32 InPassIndex)
 {
 	check(!bFrozen);
 	checkf(!EnumHasAnyFlags(InCreateInfo.Flags, TexCreate_Transient), TEXT("Attempted to create transient texture %s using legacy TexCreate_Transient flag."), InDebugName);
 
-	FRHITransientTexture* TransientTexture = RHIAllocator->CreateTexture(InCreateInfo, InDebugName);
+	FRHITransientTexture* TransientTexture = RHIAllocator->CreateTexture(InCreateInfo, InDebugName, InPassIndex);
 
 	if (!TransientTexture)
 	{
@@ -1549,12 +1547,12 @@ FRHITransientTexture* FValidationTransientResourceAllocator::CreateTexture(const
 	return TransientTexture;
 }
 
-FRHITransientBuffer* FValidationTransientResourceAllocator::CreateBuffer(const FRHIBufferCreateInfo& InCreateInfo, const TCHAR* InDebugName)
+FRHITransientBuffer* FValidationTransientResourceAllocator::CreateBuffer(const FRHIBufferCreateInfo& InCreateInfo, const TCHAR* InDebugName, uint32 InPassIndex)
 {
 	check(!bFrozen);
 	checkf(!EnumHasAnyFlags(InCreateInfo.Usage, BUF_Transient), TEXT("Attempted to create transient buffer %s using legacy BUF_Transient flag."), InDebugName);
 
-	FRHITransientBuffer* TransientBuffer = RHIAllocator->CreateBuffer(InCreateInfo, InDebugName);
+	FRHITransientBuffer* TransientBuffer = RHIAllocator->CreateBuffer(InCreateInfo, InDebugName, InPassIndex);
 
 	if (!TransientBuffer)
 	{
@@ -1582,11 +1580,11 @@ FRHITransientBuffer* FValidationTransientResourceAllocator::CreateBuffer(const F
 	return TransientBuffer;
 }
 
-void FValidationTransientResourceAllocator::DeallocateMemory(FRHITransientTexture* InTransientTexture)
+void FValidationTransientResourceAllocator::DeallocateMemory(FRHITransientTexture* InTransientTexture, uint32 InPassIndex)
 {
 	check(!bFrozen);
 
-	RHIAllocator->DeallocateMemory(InTransientTexture);
+	RHIAllocator->DeallocateMemory(InTransientTexture, InPassIndex);
 
 	// Mark memory as freed
 	FAllocatedResourceData* ResourceData = AllocatedResourceMap.Find(InTransientTexture->GetRHI());
@@ -1594,11 +1592,11 @@ void FValidationTransientResourceAllocator::DeallocateMemory(FRHITransientTextur
 	ResourceData->bMemoryAllocated = false;
 }
 
-void FValidationTransientResourceAllocator::DeallocateMemory(FRHITransientBuffer* InTransientBuffer)
+void FValidationTransientResourceAllocator::DeallocateMemory(FRHITransientBuffer* InTransientBuffer, uint32 InPassIndex)
 {
 	check(!bFrozen);
 
-	RHIAllocator->DeallocateMemory(InTransientBuffer);
+	RHIAllocator->DeallocateMemory(InTransientBuffer, InPassIndex);
 
 	// Mark memory as freed
 	FAllocatedResourceData* ResourceData = AllocatedResourceMap.Find(InTransientBuffer->GetRHI());
diff --git a/Engine/Source/Runtime/RHI/Public/RHI.h b/Engine/Source/Runtime/RHI/Public/RHI.h
--- a/Engine/Source/Runtime/RHI/Public/RHI.h
+++ b/Engine/Source/Runtime/RHI/Public/RHI.h
@@ -2300,8 +2300,8 @@ private:
 
 	RHIValidation::FOperationsList PendingSignals;
 	RHIValidation::FOperationsList PendingWaits;
-	RHIValidation::FOperationsList PendingAcquires;
-	RHIValidation::FOperationsList PendingDiscards;
+	RHIValidation::FOperationsList PendingAliases;
+	RHIValidation::FOperationsList PendingAliasingOverlaps;
 	RHIValidation::FOperationsList PendingOperationsBegin;
 	RHIValidation::FOperationsList PendingOperationsEnd;
 #endif
diff --git a/Engine/Source/Runtime/RHI/Public/RHITransientResourceAllocator.h b/Engine/Source/Runtime/RHI/Public/RHITransientResourceAllocator.h
--- a/Engine/Source/Runtime/RHI/Public/RHITransientResourceAllocator.h
+++ b/Engine/Source/Runtime/RHI/Public/RHITransientResourceAllocator.h
@@ -25,17 +25,25 @@ public:
 	virtual ~FRHITransientResource() = default;
 
 	// (Internal) Initializes the transient resource with a new allocation / name.
-	virtual void Init(const TCHAR* InName, uint32 InAllocationIndex)
+	virtual void Init(const TCHAR* InName, uint32 InAllocationIndex, uint32 InAcquirePassIndex)
 	{
 		Name = InName;
 		AllocationIndex = InAllocationIndex;
+		AcquirePasses = TInterval<uint32>(0, InAcquirePassIndex);
+		DiscardPasses = TInterval<uint32>(0, TNumericLimits<uint32>::Max());
 		AliasingOverlaps.Reset();
 	}
 
+	// (Internal) Assigns the discard pass index.
+	FORCEINLINE void SetDiscardPass(uint32 PassIndex) { DiscardPasses.Min = PassIndex; }
+
 	// (Internal) Adds a new transient resource overlap.
 	FORCEINLINE void AddAliasingOverlap(FRHITransientResource* InResource)
 	{
 		AliasingOverlaps.Emplace(InResource->GetRHI(), InResource->IsTexture() ? FRHITransientAliasingOverlap::EType::Texture : FRHITransientAliasingOverlap::EType::Buffer);
+
+		InResource->DiscardPasses.Max = FMath::Min(InResource->DiscardPasses.Max,             AcquirePasses.Max);
+		            AcquirePasses.Min = FMath::Max(            AcquirePasses.Min, InResource->DiscardPasses.Min);
 	}
 
 	// Returns the underlying RHI resource.
@@ -50,6 +58,12 @@ public:
 	// (Internal) Returns the platform-specific allocation index.
 	FORCEINLINE uint32 GetAllocationIndex() const { return AllocationIndex; }
 
+	// Returns the pass index which may end acquiring this resource.
+	FORCEINLINE TInterval<uint32> GetAcquirePasses() const { return AcquirePasses; }
+
+	// Returns the pass index which discarded this resource.
+	FORCEINLINE TInterval<uint32> GetDiscardPasses() const { return DiscardPasses; }
+
 	// Returns the aliasing overlaps for this resource.
 	FORCEINLINE TConstArrayView<FRHITransientAliasingOverlap> GetAliasingOverlaps() const { return AliasingOverlaps; }
 
@@ -75,6 +89,8 @@ private:
 
 	// An index to the underlying allocation info on the internal platform allocator.
 	uint32 AllocationIndex = ~0u;
+	TInterval<uint32> AcquirePasses = TInterval<uint32>(0, 0);
+	TInterval<uint32> DiscardPasses = TInterval<uint32>(0, 0);
 };
 
 class RHI_API FRHITransientTexture final : public FRHITransientResource
@@ -85,7 +101,7 @@ public:
 		, CreateInfo(InCreateInfo)
 	{}
 
-	void Init(const TCHAR* InName, uint32 InAllocationIndex) override;
+	void Init(const TCHAR* InName, uint32 InAllocationIndex, uint32 InPassIndex) override;
 
 	// Returns the underlying RHI texture.
 	FRHITexture* GetRHI() const { return static_cast<FRHITexture*>(FRHITransientResource::GetRHI()); }
@@ -114,7 +130,7 @@ public:
 		, CreateInfo(InCreateInfo)
 	{}
 
-	void Init(const TCHAR* InName, uint32 InAllocationIndex) override;
+	void Init(const TCHAR* InName, uint32 InAllocationIndex, uint32 InPassIndex) override;
 
 	// Returns the underlying RHI buffer.
 	FORCEINLINE FRHIBuffer* GetRHI() const { return static_cast<FRHIBuffer*>(FRHITransientResource::GetRHI()); }
@@ -141,12 +157,12 @@ public:
 	virtual ~IRHITransientResourceAllocator() = default;
 
 	// Allocates a new transient resource with memory backed by the transient allocator.
-	virtual FRHITransientTexture* CreateTexture(const FRHITextureCreateInfo& InCreateInfo, const TCHAR* InDebugName) = 0;
-	virtual FRHITransientBuffer* CreateBuffer(const FRHIBufferCreateInfo& InCreateInfo, const TCHAR* InDebugName) = 0;
+	virtual FRHITransientTexture* CreateTexture(const FRHITextureCreateInfo& InCreateInfo, const TCHAR* InDebugName, uint32 InPassIndex) = 0;
+	virtual FRHITransientBuffer* CreateBuffer(const FRHIBufferCreateInfo& InCreateInfo, const TCHAR* InDebugName, uint32 InPassIndex) = 0;
 
 	// Deallocates the underlying memory for use by a future resource creation call.
-	virtual void DeallocateMemory(FRHITransientTexture* InTexture) = 0;
-	virtual void DeallocateMemory(FRHITransientBuffer* InBuffer) = 0;
+	virtual void DeallocateMemory(FRHITransientTexture* InTexture, uint32 InPassIndex) = 0;
+	virtual void DeallocateMemory(FRHITransientBuffer* InBuffer, uint32 InPassIndex) = 0;
 
 	// Flushes any pending allocations in preparation for rendering. Resources are not required to be deallocated yet. Use when interleaving execution with allocation.
 	virtual void Flush(FRHICommandListImmediate& RHICmdList) {}
diff --git a/Engine/Source/Runtime/RHI/Public/RHIValidationContext.h b/Engine/Source/Runtime/RHI/Public/RHIValidationContext.h
--- a/Engine/Source/Runtime/RHI/Public/RHIValidationContext.h
+++ b/Engine/Source/Runtime/RHI/Public/RHIValidationContext.h
@@ -60,7 +60,7 @@ public:
 	{
 		for (const FRHITransition* Transition : Transitions)
 		{
-			Tracker->AddOps(Transition->PendingAcquires);
+			Tracker->AddOps(Transition->PendingAliases);
 		}
 
 		for (const FRHITransition* Transition : Transitions)
@@ -90,7 +90,7 @@ public:
 
 		for (const FRHITransition* Transition : Transitions)
 		{
-			Tracker->AddOps(Transition->PendingDiscards);
+			Tracker->AddOps(Transition->PendingAliasingOverlaps);
 		}
 
 		RHIContext->RHIEndTransitions(Transitions);
@@ -378,7 +378,7 @@ public:
 
 		for (const FRHITransition* Transition : Transitions)
 		{
-			Tracker->AddOps(Transition->PendingAcquires);
+			Tracker->AddOps(Transition->PendingAliases);
 		}
 
 		for (const FRHITransition* Transition : Transitions)
@@ -410,7 +410,7 @@ public:
 
 		for (const FRHITransition* Transition : Transitions)
 		{
-			Tracker->AddOps(Transition->PendingDiscards);
+			Tracker->AddOps(Transition->PendingAliasingOverlaps);
 		}
 
 		RHIContext->RHIEndTransitions(Transitions);
diff --git a/Engine/Source/Runtime/RHI/Public/RHIValidationTransientResourceAllocator.h b/Engine/Source/Runtime/RHI/Public/RHIValidationTransientResourceAllocator.h
--- a/Engine/Source/Runtime/RHI/Public/RHIValidationTransientResourceAllocator.h
+++ b/Engine/Source/Runtime/RHI/Public/RHIValidationTransientResourceAllocator.h
@@ -16,10 +16,10 @@ public:
 	virtual ~FValidationTransientResourceAllocator();
 
 	// Implementation of FRHITransientResourceAllocator interface
-	virtual FRHITransientTexture* CreateTexture(const FRHITextureCreateInfo& InCreateInfo, const TCHAR* InDebugName) override final;
-	virtual FRHITransientBuffer* CreateBuffer(const FRHIBufferCreateInfo& InCreateInfo, const TCHAR* InDebugName) override final;
-	virtual void DeallocateMemory(FRHITransientTexture* InTexture) override final;
-	virtual void DeallocateMemory(FRHITransientBuffer* InBuffer) override final;
+	virtual FRHITransientTexture* CreateTexture(const FRHITextureCreateInfo& InCreateInfo, const TCHAR* InDebugName, uint32 InPassIndex) override final;
+	virtual FRHITransientBuffer* CreateBuffer(const FRHIBufferCreateInfo& InCreateInfo, const TCHAR* InDebugName, uint32 InPassIndex) override final;
+	virtual void DeallocateMemory(FRHITransientTexture* InTexture, uint32 InPassIndex) override final;
+	virtual void DeallocateMemory(FRHITransientBuffer* InBuffer, uint32 InPassIndex) override final;
 	virtual void Flush(FRHICommandListImmediate&) override final;
 	virtual void Freeze(FRHICommandListImmediate&) override final;
 	virtual void Release(FRHICommandListImmediate&) override final;
diff --git a/Engine/Source/Runtime/RHICore/Private/RHICoreTransientResourceAllocator.cpp b/Engine/Source/Runtime/RHICore/Private/RHICoreTransientResourceAllocator.cpp
--- a/Engine/Source/Runtime/RHICore/Private/RHICoreTransientResourceAllocator.cpp
+++ b/Engine/Source/Runtime/RHICore/Private/RHICoreTransientResourceAllocator.cpp
@@ -526,6 +526,7 @@ FRHITransientResourceAllocator::FRHITransientResourceAllocator(FRHITransientReso
 FRHITransientTexture* FRHITransientResourceAllocator::CreateTexture(
 	const FRHITextureCreateInfo& CreateInfo,
 	const TCHAR* DebugName,
+	uint32 PassIndex,
 	uint64 TextureSize,
 	uint32 TextureAlignment,
 	FCreateTextureFunction CreateTextureFunction)
@@ -550,7 +551,7 @@ FRHITransientTexture* FRHITransientResourceAllocator::CreateTexture(
 		return CreateTextureFunction(ResourceInitializer);
 	});
 
-	InitResource(TransientTexture, Allocation, DebugName);
+	InitResource(TransientTexture, PassIndex, Allocation, DebugName);
 
 #if RHICORE_TRANSIENT_ALLOCATOR_DEBUG
 	DebugTextures.Emplace(TransientTexture);
@@ -563,6 +564,7 @@ FRHITransientTexture* FRHITransientResourceAllocator::CreateTexture(
 FRHITransientBuffer* FRHITransientResourceAllocator::CreateBuffer(
 	const FRHIBufferCreateInfo& CreateInfo,
 	const TCHAR* DebugName,
+	uint32 PassIndex,
 	uint32 BufferSize,
 	uint32 BufferAlignment,
 	FCreateBufferFunction CreateBufferFunction)
@@ -582,7 +584,7 @@ FRHITransientBuffer* FRHITransientResourceAllocator::CreateBuffer(
 		return CreateBufferFunction(ResourceInitializer);
 	});
 
-	InitResource(TransientBuffer, Allocation, DebugName);
+	InitResource(TransientBuffer, PassIndex, Allocation, DebugName);
 
 #if RHICORE_TRANSIENT_ALLOCATOR_DEBUG
 	DebugBuffers.Emplace(TransientBuffer);
@@ -592,8 +594,10 @@ FRHITransientBuffer* FRHITransientResourceAllocator::CreateBuffer(
 	return TransientBuffer;
 }
 
-void FRHITransientResourceAllocator::DeallocateMemoryInternal(FRHITransientResource* InResource, FMemoryStats& StatsToUpdate)
+void FRHITransientResourceAllocator::DeallocateMemoryInternal(FRHITransientResource* InResource, uint32 PassIndex, FMemoryStats& StatsToUpdate)
 {
+	InResource->SetDiscardPass(PassIndex);
+
 	const FRHITransientHeapAllocation& Allocation = HeapAllocations[InResource->GetAllocationIndex()];
 
 	HeapAllocators[Allocation.HeapIndex].Deallocate(Allocation);
@@ -666,8 +670,8 @@ FRHITransientHeapAllocation FRHITransientResourceAllocator::Allocate(FMemoryStat
 	return Allocation;
 }
 
-void FRHITransientResourceAllocator::InitResource(FRHITransientResource* TransientResource, const FRHITransientHeapAllocation& Allocation, const TCHAR* Name)
+void FRHITransientResourceAllocator::InitResource(FRHITransientResource* TransientResource, uint32 PassIndex, const FRHITransientHeapAllocation& Allocation, const TCHAR* Name)
 {
-	TransientResource->Init(Name, HeapAllocations.Emplace(Allocation));
+	TransientResource->Init(Name, HeapAllocations.Emplace(Allocation), PassIndex);
 	HeapAllocators[Allocation.HeapIndex].TrackOverlap(TransientResource, Allocation);
 }
\ No newline at end of file
diff --git a/Engine/Source/Runtime/RHICore/Public/RHICoreTransientResourceAllocator.h b/Engine/Source/Runtime/RHICore/Public/RHICoreTransientResourceAllocator.h
--- a/Engine/Source/Runtime/RHICore/Public/RHICoreTransientResourceAllocator.h
+++ b/Engine/Source/Runtime/RHICore/Public/RHICoreTransientResourceAllocator.h
@@ -569,6 +569,7 @@ public:
 	FRHITransientTexture* CreateTexture(
 		const FRHITextureCreateInfo& CreateInfo,
 		const TCHAR* DebugName,
+		uint32 PassIndex,
 		uint64 TextureSize,
 		uint32 TextureAlignment,
 		FCreateTextureFunction CreateTextureFunction);
@@ -583,20 +584,21 @@ public:
 	FRHITransientBuffer* CreateBuffer(
 		const FRHIBufferCreateInfo& CreateInfo,
 		const TCHAR* DebugName,
+		uint32 PassIndex,
 		uint32 BufferSize,
 		uint32 BufferAlignment,
 		FCreateBufferFunction CreateBufferFunction);
 
 	// Deallocates a texture from its parent heap. Provide the current platform fence value used to update the heap.
-	FORCEINLINE void DeallocateMemory(FRHITransientTexture* Texture)
+	FORCEINLINE void DeallocateMemory(FRHITransientTexture* Texture, uint32 PassIndex)
 	{
-		DeallocateMemoryInternal(Texture, TextureStats);
+		DeallocateMemoryInternal(Texture, PassIndex, TextureStats);
 	}
 
 	// Deallocates a buffer from its parent heap. Provide the current platform fence value used to update the heap.
-	FORCEINLINE void DeallocateMemory(FRHITransientBuffer* Buffer)
+	FORCEINLINE void DeallocateMemory(FRHITransientBuffer* Buffer, uint32 PassIndex)
 	{
-		DeallocateMemoryInternal(Buffer, BufferStats);
+		DeallocateMemoryInternal(Buffer, PassIndex, BufferStats);
 	}
 
 	// Called to signify all allocations have completed. Forfeits all resources / heaps back to the parent system.
@@ -619,9 +621,9 @@ private:
 
 	FRHITransientHeapAllocation Allocate(FMemoryStats& StatsToUpdate, uint64 Size, uint32 Alignment, ERHITransientHeapFlags ResourceHeapFlags);
 
-	void DeallocateMemoryInternal(FRHITransientResource* InResource, FMemoryStats& StatsToUpdate);
+	void DeallocateMemoryInternal(FRHITransientResource* InResource, uint32 PassIndex, FMemoryStats& StatsToUpdate);
 
-	void InitResource(FRHITransientResource* TransientResource, const FRHITransientHeapAllocation& Allocation, const TCHAR* Name);
+	void InitResource(FRHITransientResource* TransientResource, uint32 PassIndex, const FRHITransientHeapAllocation& Allocation, const TCHAR* Name);
 
 	FRHITransientResourceSystem& ParentSystem;
 
diff --git a/Engine/Source/Runtime/RenderCore/Private/RenderGraphBuilder.cpp b/Engine/Source/Runtime/RenderCore/Private/RenderGraphBuilder.cpp
--- a/Engine/Source/Runtime/RenderCore/Private/RenderGraphBuilder.cpp
+++ b/Engine/Source/Runtime/RenderCore/Private/RenderGraphBuilder.cpp
@@ -2603,16 +2603,19 @@ void FRDGBuilder::AddEpilogueTransition(FRDGTextureRef Texture)
 	const FRDGPassHandle EpiloguePassHandle = GetEpiloguePassHandle();
 
 	FRDGSubresourceState ScratchSubresourceState;
-	ERDGBarrierLocation BarrierLocation = ERDGBarrierLocation::Prologue;
 
 	// Texture is using the RHI transient allocator. Transition it back to Discard in the final pass it is used.
 	if (Texture->bTransient)
 	{
-		ScratchSubresourceState.SetPass(ERHIPipeline::Graphics, ClampToPrologue(Texture->LastPass));
+		const TInterval<uint32> DiscardPasses = Texture->TransientTexture->GetDiscardPasses();
+		const FRDGPassHandle MinDiscardPassHandle(DiscardPasses.Min);
+		const FRDGPassHandle MaxDiscardPassHandle(FMath::Min<uint32>(DiscardPasses.Max, EpiloguePassHandle.GetIndex()));
+
+		AddAliasingTransition(MinDiscardPassHandle, MaxDiscardPassHandle, Texture, FRHITransientAliasingInfo::Discard(Texture->GetRHIUnchecked()));
+
+		ScratchSubresourceState.SetPass(ERHIPipeline::Graphics, MaxDiscardPassHandle);
 		ScratchSubresourceState.Access = ERHIAccess::Discard;
 		InitAsWholeResource(ScratchTextureState, &ScratchSubresourceState);
-
-		BarrierLocation = ERDGBarrierLocation::Epilogue;
 	}
 	// A known final state means extraction from the graph (or an external texture).
 	else if(Texture->AccessFinal != ERHIAccess::Unknown)
@@ -2651,7 +2654,7 @@ void FRDGBuilder::AddEpilogueTransition(FRDGTextureRef Texture)
 		return;
 	}
 
-	AddTransition(EpiloguePassHandle, Texture, ScratchTextureState, BarrierLocation);
+	AddTransition(EpiloguePassHandle, Texture, ScratchTextureState);
 	ScratchTextureState.Reset();
 }
 
@@ -2666,10 +2669,16 @@ void FRDGBuilder::AddEpilogueTransition(FRDGBufferRef Buffer)
 
 	if (Buffer->bTransient)
 	{
+		const TInterval<uint32> DiscardPasses = Buffer->TransientBuffer->GetDiscardPasses();
+		const FRDGPassHandle MinDiscardPassHandle(DiscardPasses.Min);
+		const FRDGPassHandle MaxDiscardPassHandle(FMath::Min<uint32>(DiscardPasses.Max, EpiloguePassHandle.GetIndex()));
+
+		AddAliasingTransition(MinDiscardPassHandle, MaxDiscardPassHandle, Buffer, FRHITransientAliasingInfo::Discard(Buffer->GetRHIUnchecked()));
+
 		FRDGSubresourceState StateFinal;
-		StateFinal.SetPass(ERHIPipeline::Graphics, ClampToPrologue(Buffer->LastPass));
+		StateFinal.SetPass(ERHIPipeline::Graphics, MaxDiscardPassHandle);
 		StateFinal.Access = ERHIAccess::Discard;
-		AddTransition(Buffer->LastPass, Buffer, StateFinal, ERDGBarrierLocation::Epilogue);
+		AddTransition(Buffer->LastPass, Buffer, StateFinal);
 	}
 	else
 	{
@@ -2696,7 +2705,7 @@ void FRDGBuilder::AddEpilogueTransition(FRDGBufferRef Buffer)
 	}
 }
 
-void FRDGBuilder::AddTransition(FRDGPassHandle PassHandle, FRDGTextureRef Texture, const FRDGTextureTransientSubresourceStateIndirect& StateAfter, ERDGBarrierLocation BarrierLocation)
+void FRDGBuilder::AddTransition(FRDGPassHandle PassHandle, FRDGTextureRef Texture, const FRDGTextureTransientSubresourceStateIndirect& StateAfter)
 {
 	const FRDGTextureSubresourceRange WholeRange = Texture->GetSubresourceRange();
 	const FRDGTextureSubresourceLayout Layout = Texture->Layout;
@@ -2730,7 +2739,7 @@ void FRDGBuilder::AddTransition(FRDGPassHandle PassHandle, FRDGTextureRef Textur
 				Info.PlaneSlice = Subresource->PlaneSlice;
 			}
 
-			AddTransition(Texture, SubresourceStateBefore, SubresourceStateAfter, Info, BarrierLocation);
+			AddTransition(Texture, SubresourceStateBefore, SubresourceStateAfter, Info);
 		}
 
 		if (Subresource)
@@ -2803,7 +2812,7 @@ void FRDGBuilder::AddTransition(FRDGPassHandle PassHandle, FRDGTextureRef Textur
 	}
 }
 
-void FRDGBuilder::AddTransition(FRDGPassHandle PassHandle, FRDGBufferRef Buffer, FRDGSubresourceState StateAfter, ERDGBarrierLocation BarrierLocation)
+void FRDGBuilder::AddTransition(FRDGPassHandle PassHandle, FRDGBufferRef Buffer, FRDGSubresourceState StateAfter)
 {
 	check(StateAfter.Access != ERHIAccess::Unknown);
 
@@ -2818,7 +2827,7 @@ void FRDGBuilder::AddTransition(FRDGPassHandle PassHandle, FRDGBufferRef Buffer,
 		Info.AccessBefore = StateBefore.Access;
 		Info.AccessAfter = StateAfter.Access;
 
-		AddTransition(Buffer, StateBefore, StateAfter, Info, BarrierLocation);
+		AddTransition(Buffer, StateBefore, StateAfter, Info);
 	}
 
 	IF_RDG_ENABLE_DEBUG(LogFile.AddTransitionEdge(PassHandle, StateBefore, StateAfter, Buffer));
@@ -2829,8 +2838,7 @@ void FRDGBuilder::AddTransition(
 	FRDGParentResource* Resource,
 	FRDGSubresourceState StateBefore,
 	FRDGSubresourceState StateAfter,
-	const FRHITransitionInfo& TransitionInfo,
-	ERDGBarrierLocation BarrierLocation)
+	const FRHITransitionInfo& TransitionInfo)
 {
 	const ERHIPipeline Graphics = ERHIPipeline::Graphics;
 	const ERHIPipeline AsyncCompute = ERHIPipeline::AsyncCompute;
@@ -2843,7 +2851,7 @@ void FRDGBuilder::AddTransition(
 	if (IsImmediateMode())
 	{
 		// Immediate mode simply enqueues the barrier into the 'after' pass. Everything is on the graphics pipe.
-		AddToBarriers(StateAfter.FirstPass[Graphics], BarrierLocation, [&](FRDGBarrierBatchBegin& Barriers)
+		AddToPrologueBarriers(StateAfter.FirstPass[Graphics], [&](FRDGBarrierBatchBegin& Barriers)
 		{
 			Barriers.AddTransition(Resource, TransitionInfo);
 		});
@@ -2864,7 +2872,6 @@ void FRDGBuilder::AddTransition(
 
 	check(PipelinesBefore != ERHIPipeline::None && PipelinesAfter != ERHIPipeline::None);
 	checkf(StateBefore.GetLastPass() <= StateAfter.GetFirstPass(), TEXT("Submitted a state for '%s' that begins before our previous state has ended."), Resource->Name);
-	check(!Resource->bTransient || StateBefore.GetLastPass() >= Resource->FirstPass);
 
 	const FRDGPassHandlesByPipeline& PassesBefore = StateBefore.LastPass;
 	const FRDGPassHandlesByPipeline& PassesAfter = StateAfter.FirstPass;
@@ -2879,7 +2886,7 @@ void FRDGBuilder::AddTransition(
 		FRDGBarrierBatchBegin* BarriersToBegin = nullptr;
 
 		// Issue the begin in the epilogue of the begin pass if the barrier is being split across multiple passes or the barrier end is in the epilogue.
-		if (BeginPassHandle < FirstEndPassHandle || BarrierLocation == ERDGBarrierLocation::Epilogue)
+		if (BeginPassHandle < FirstEndPassHandle)
 		{
 			BeginPass = GetEpilogueBarrierPass(BeginPassHandle);
 			BarriersToBegin = &BeginPass->GetEpilogueBarriersToBeginFor(Allocator, TransitionCreateQueue, PipelinesAfter);
@@ -2921,21 +2928,14 @@ void FRDGBuilder::AddTransition(
 			}
 			else if (EnumHasAnyFlags(PipelinesAfter, Pipeline))
 			{
-				if (BarrierLocation == ERDGBarrierLocation::Prologue)
-				{
-					AddToPrologueBarriersToEnd(PassesAfter[Pipeline], *BarriersToBegin);
-				}
-				else
-				{
-					AddToEpilogueBarriersToEnd(PassesAfter[Pipeline], *BarriersToBegin);
-				}
+				AddToPrologueBarriersToEnd(PassesAfter[Pipeline], *BarriersToBegin);
 			}
 		}
 	}
 	// N-to-1 or N-to-N transition.
 	else
 	{
-		checkf(StateBefore.GetLastPass() != StateAfter.GetFirstPass() || BarrierLocation == ERDGBarrierLocation::Epilogue,
+		checkf(StateBefore.GetLastPass() != StateAfter.GetFirstPass(),
 			TEXT("Attempted to queue a transition for resource '%s' from '%s' to '%s', but previous and next passes are the same on one pipe."),
 			Resource->Name, *GetRHIPipelineName(PipelinesBefore), *GetRHIPipelineName(PipelinesAfter));
 
@@ -2969,19 +2969,45 @@ void FRDGBuilder::AddTransition(
 		{
 			if (EnumHasAnyFlags(PipelinesAfter, Pipeline))
 			{
-				if (BarrierLocation == ERDGBarrierLocation::Prologue)
-				{
-					AddToPrologueBarriersToEnd(PassesAfter[Pipeline], *BarriersToBegin);
-				}
-				else
-				{
-					AddToEpilogueBarriersToEnd(PassesAfter[Pipeline], *BarriersToBegin);
-				}
+				AddToPrologueBarriersToEnd(PassesAfter[Pipeline], *BarriersToBegin);
 			}
 		}
 	}
 }
 
+void FRDGBuilder::AddAliasingTransition(FRDGPassHandle BeginPassHandle, FRDGPassHandle EndPassHandle, FRDGParentResourceRef Resource, const FRHITransientAliasingInfo& Info)
+{
+	check(BeginPassHandle <= EndPassHandle);
+
+	FRDGBarrierBatchBegin* BarriersToBegin{};
+	FRDGPass* EndPass{};
+
+	if (BeginPassHandle == EndPassHandle)
+	{
+		FRDGPass* BeginPass = Passes[BeginPassHandle];
+		EndPass = BeginPass;
+
+		check(GetPrologueBarrierPassHandle(BeginPassHandle) == BeginPassHandle);
+		check(BeginPass->GetPipeline() == ERHIPipeline::Graphics);
+
+		BarriersToBegin = &BeginPass->GetPrologueBarriersToBegin(Allocator, TransitionCreateQueue);
+	}
+	else
+	{
+		FRDGPass* BeginPass = GetEpilogueBarrierPass(BeginPassHandle);
+		EndPass = Passes[EndPassHandle];
+
+		check(GetPrologueBarrierPassHandle(EndPassHandle) == EndPassHandle);
+		check(BeginPass->GetPipeline() == ERHIPipeline::Graphics);
+		check(EndPass->GetPipeline() == ERHIPipeline::Graphics);
+
+		BarriersToBegin = &BeginPass->GetEpilogueBarriersToBeginForGraphics(Allocator, TransitionCreateQueue);
+	}
+
+	BarriersToBegin->AddAlias(Resource, Info);
+	EndPass->GetPrologueBarriersToEnd(Allocator).AddDependency(BarriersToBegin);
+}
+
 void FRDGBuilder::BeginResourceRHI(FRDGPassHandle PassHandle, FRDGTextureRef Texture)
 {
 	check(Texture);
@@ -3015,20 +3041,19 @@ void FRDGBuilder::BeginResourceRHI(FRDGPassHandle PassHandle, FRDGTextureRef Tex
 	{
 		if (IRHITransientResourceAllocator* AllocatorRHI = TransientResourceAllocator.GetOrCreate())
 		{
-			if (FRHITransientTexture* TransientTexture = AllocatorRHI->CreateTexture(Texture->Desc, Texture->Name))
+			if (FRHITransientTexture* TransientTexture = AllocatorRHI->CreateTexture(Texture->Desc, Texture->Name, PassHandle.GetIndex()))
 			{
 				Texture->SetRHI(TransientTexture, Allocator);
 
+				const FRDGPassHandle MinAcquirePassHandle = ClampToPrologue(FRDGPassHandle(TransientTexture->GetAcquirePasses().Min));
+
+				AddAliasingTransition(MinAcquirePassHandle, PassHandle, Texture, FRHITransientAliasingInfo::Acquire(TransientTexture->GetRHI(), TransientTexture->GetAliasingOverlaps()));
+
 				FRDGSubresourceState InitialState;
-				InitialState.SetPass(ERHIPipeline::Graphics, PassHandle);
+				InitialState.SetPass(ERHIPipeline::Graphics, MinAcquirePassHandle);
 				InitialState.Access = ERHIAccess::Discard;
 				InitAsWholeResource(Texture->GetState(), InitialState);
 
-				AddToPrologueBarriers(PassHandle, [&](FRDGBarrierBatchBegin& Barriers)
-				{
-					Barriers.AddAlias(Texture, FRHITransientAliasingInfo::Acquire(TransientTexture->GetRHI(), TransientTexture->GetAliasingOverlaps()));
-				});
-
 			#if STATS
 				GRDGStatTransientTextureCount++;
 			#endif
@@ -3107,19 +3132,18 @@ void FRDGBuilder::BeginResourceRHI(FRDGPassHandle PassHandle, FRDGBufferRef Buff
 	{
 		if (IRHITransientResourceAllocator* AllocatorRHI = TransientResourceAllocator.GetOrCreate())
 		{
-			if (FRHITransientBuffer* TransientBuffer = AllocatorRHI->CreateBuffer(Translate(Buffer->Desc), Buffer->Name))
+			if (FRHITransientBuffer* TransientBuffer = AllocatorRHI->CreateBuffer(Translate(Buffer->Desc), Buffer->Name, PassHandle.GetIndex()))
 			{
 				Buffer->SetRHI(TransientBuffer, Allocator);
 
+				const FRDGPassHandle MinAcquirePassHandle = ClampToPrologue(FRDGPassHandle(TransientBuffer->GetAcquirePasses().Min));
+
+				AddAliasingTransition(MinAcquirePassHandle, PassHandle, Buffer, FRHITransientAliasingInfo::Acquire(TransientBuffer->GetRHI(), TransientBuffer->GetAliasingOverlaps()));
+
 				FRDGSubresourceState& InitialState = Buffer->GetState();
-				InitialState.SetPass(ERHIPipeline::Graphics, PassHandle);
+				InitialState.SetPass(ERHIPipeline::Graphics, MinAcquirePassHandle);
 				InitialState.Access = ERHIAccess::Discard;
 
-				AddToPrologueBarriers(PassHandle, [&](FRDGBarrierBatchBegin& Barriers)
-				{
-					Barriers.AddAlias(Buffer, FRHITransientAliasingInfo::Acquire(TransientBuffer->GetRHI(), TransientBuffer->GetAliasingOverlaps()));
-				});
-
 			#if STATS
 				GRDGStatTransientBufferCount++;
 			#endif
@@ -3222,12 +3246,7 @@ void FRDGBuilder::EndResourceRHI(FRDGPassHandle PassHandle, FRDGTextureRef Textu
 	{
 		if (Texture->bTransient)
 		{
-			TransientResourceAllocator->DeallocateMemory(Texture->TransientTexture);
-
-			AddToEpilogueBarriers(PassHandle, [&](FRDGBarrierBatchBegin& Barriers)
-			{
-				Barriers.AddAlias(Texture, FRHITransientAliasingInfo::Discard(Texture->GetRHIUnchecked()));
-			});
+			TransientResourceAllocator->DeallocateMemory(Texture->TransientTexture, PassHandle.GetIndex());
 		}
 		else if (Texture->PooledRenderTarget && Texture->PooledRenderTarget->IsTracked())
 		{
@@ -3249,12 +3268,7 @@ void FRDGBuilder::EndResourceRHI(FRDGPassHandle PassHandle, FRDGBufferRef Buffer
 	{
 		if (Buffer->bTransient)
 		{
-			TransientResourceAllocator->DeallocateMemory(Buffer->TransientBuffer);
-
-			AddToEpilogueBarriers(PassHandle, [&](FRDGBarrierBatchBegin& Barriers)
-			{
-				Barriers.AddAlias(Buffer, FRHITransientAliasingInfo::Discard(Buffer->GetRHIUnchecked()));
-			});
+			TransientResourceAllocator->DeallocateMemory(Buffer->TransientBuffer, PassHandle.GetIndex());
 		}
 		else
 		{
diff --git a/Engine/Source/Runtime/RenderCore/Public/RenderGraphBuilder.h b/Engine/Source/Runtime/RenderCore/Public/RenderGraphBuilder.h
--- a/Engine/Source/Runtime/RenderCore/Public/RenderGraphBuilder.h
+++ b/Engine/Source/Runtime/RenderCore/Public/RenderGraphBuilder.h
@@ -409,20 +409,6 @@ private:
 		Pass->GetEpilogueBarriersToEnd(Allocator).AddDependency(&BarriersToBegin);
 	}
 
-	/** Utility function to add an immediate barrier dependency in either the prologue or epilogue of the provided pass, depending on BarrierLocation. */
-	template <typename FunctionType>
-	void AddToBarriers(FRDGPassHandle PassHandle, ERDGBarrierLocation BarrierLocation, FunctionType Function)
-	{
-		if (BarrierLocation == ERDGBarrierLocation::Prologue)
-		{
-			AddToPrologueBarriers(PassHandle, Function);
-		}
-		else
-		{
-			AddToEpilogueBarriers(PassHandle, Function);
-		}
-	}
-
 	/** Registry of graph objects. */
 	FRDGPassRegistry Passes;
 	FRDGTextureRegistry Textures;
@@ -654,31 +640,29 @@ private:
 	void AddTransition(
 		FRDGPassHandle PassHandle,
 		FRDGTextureRef Texture,
-		const FRDGTextureTransientSubresourceStateIndirect& StateAfter,
-		ERDGBarrierLocation BarrierLocation = ERDGBarrierLocation::Prologue);
+		const FRDGTextureTransientSubresourceStateIndirect& StateAfter);
 
 	void AddTransition(
 		FRDGPassHandle PassHandle,
 		FRDGBufferRef Buffer,
-		FRDGSubresourceState StateAfter,
-		ERDGBarrierLocation BarrierLocation = ERDGBarrierLocation::Prologue);
+		FRDGSubresourceState StateAfter);
 
 	void AddTransition(
 		FRDGParentResource* Resource,
 		FRDGSubresourceState StateBefore,
 		FRDGSubresourceState StateAfter,
-		const FRHITransitionInfo& TransitionInfo,
-		ERDGBarrierLocation BarrierLocation);
+		const FRHITransitionInfo& TransitionInfo);
+
+	void AddAliasingTransition(
+		FRDGPassHandle BeginPassHandle,
+		FRDGPassHandle EndPassHandle,
+		FRDGParentResourceRef Resource,
+		const FRHITransientAliasingInfo& Info);
 
 	bool IsTransient(FRDGTextureRef Texture) const;
 	bool IsTransient(FRDGBufferRef Buffer) const;
 	bool IsTransientInternal(FRDGParentResourceRef Resource) const;
 
-	void AddAcquireResourceRHI(FRDGTexture* Texture, FRDGPassHandle PassHandle);
-	void AddAcquireResourceRHI(FRDGBuffer* Buffer, FRDGPassHandle PassHandle);
-	void AddDiscardResourceRHI(FRDGTexture* Texture, FRDGPassHandle PassHandle);
-	void AddDiscardResourceRHI(FRDGBuffer* Buffer, FRDGPassHandle PassHandle);
-
 	FRHIRenderPassInfo GetRenderPassInfo(const FRDGPass* Pass) const;
 
 	FRDGSubresourceState* AllocSubresource(const FRDGSubresourceState& Other);
