Small improvement to RDG debug resource name filtering to allow exclusion.

diff --git a/Engine/Source/Runtime/RenderCore/Private/RenderGraphPrivate.cpp b/Engine/Source/Runtime/RenderCore/Private/RenderGraphPrivate.cpp
--- a/Engine/Source/Runtime/RenderCore/Private/RenderGraphPrivate.cpp
+++ b/Engine/Source/Runtime/RenderCore/Private/RenderGraphPrivate.cpp
@@ -140,13 +140,21 @@ inline bool IsDebugAllowed(const FString& FilterString, const TCHAR* Name)
 		return true;
 	}
 
-	const bool bFound = FCString::Strifind(Name, *FilterString) != nullptr;
-	if (!bFound)
+	const bool bInverted = FilterString[0] == '!';
+	if (FilterString.Len() == 1 && bInverted)
 	{
-		return false;
+		return true;
+	}
+
+	const TCHAR* FilterStringRaw = *FilterString;
+
+	if (bInverted)
+	{
+		FilterStringRaw++;
 	}
 
-	return FilterString[0] != TEXT('!');
+	const bool bFound = FCString::Strifind(Name, FilterStringRaw) != nullptr;
+	return bFound ^ bInverted;
 }
 
 bool IsDebugAllowedForGraph(const TCHAR* GraphName)
