Enable resummarize on HTile data for hair visibility pass.

diff --git a/Engine/Source/Runtime/RenderCore/Private/RenderGraphUtils.cpp b/Engine/Source/Runtime/RenderCore/Private/RenderGraphUtils.cpp
--- a/Engine/Source/Runtime/RenderCore/Private/RenderGraphUtils.cpp
+++ b/Engine/Source/Runtime/RenderCore/Private/RenderGraphUtils.cpp
@@ -639,7 +639,10 @@ void AddClearStencilPass(FRDGBuilder& GraphBuilder, FRDGTextureRef Texture)
 void AddResummarizeHTilePass(FRDGBuilder& GraphBuilder, FRDGTextureRef Texture)
 {
 	auto* PassParameters = GraphBuilder.AllocParameters<FRenderTargetParameters>();
-	PassParameters->RenderTargets.DepthStencil = FDepthStencilBinding(Texture, ERenderTargetLoadAction::ELoad, ERenderTargetLoadAction::ELoad, FExclusiveDepthStencil::DepthWrite_StencilWrite);
+	const bool bHasStencil = Texture->Desc.Format == PF_DepthStencil;
+	PassParameters->RenderTargets.DepthStencil = bHasStencil ? 
+		FDepthStencilBinding(Texture, ERenderTargetLoadAction::ELoad, ERenderTargetLoadAction::ELoad, FExclusiveDepthStencil::DepthWrite_StencilWrite) :
+		FDepthStencilBinding(Texture, ERenderTargetLoadAction::ELoad, ERenderTargetLoadAction::ENoAction, FExclusiveDepthStencil::DepthWrite_StencilNop);
 	GraphBuilder.AddPass(RDG_EVENT_NAME("ResummarizeHTile (%s)", Texture->Name), PassParameters, ERDGPassFlags::Raster,
 		[Texture](FRHICommandList& RHICmdList)
 	{
diff --git a/Engine/Source/Runtime/Renderer/Private/HairStrands/HairStrandsVisibility.cpp b/Engine/Source/Runtime/Renderer/Private/HairStrands/HairStrandsVisibility.cpp
--- a/Engine/Source/Runtime/Renderer/Private/HairStrands/HairStrandsVisibility.cpp
+++ b/Engine/Source/Runtime/Renderer/Private/HairStrands/HairStrandsVisibility.cpp
@@ -2130,6 +2130,8 @@ static FRDGTextureRef AddHairVisibilityFillOpaqueDepth(
 		});
 	}
 
+	// Ensure HTile is valid after manually feeding the scene depth value
+	AddResummarizeHTilePass(GraphBuilder, OutVisibilityDepthTexture);
 
 	return OutVisibilityDepthTexture;
 }
