Fixed RDG debug crash.

diff --git a/Engine/Source/Runtime/RenderCore/Private/RenderGraphValidation.cpp b/Engine/Source/Runtime/RenderCore/Private/RenderGraphValidation.cpp
--- a/Engine/Source/Runtime/RenderCore/Private/RenderGraphValidation.cpp
+++ b/Engine/Source/Runtime/RenderCore/Private/RenderGraphValidation.cpp
@@ -862,7 +862,7 @@ void FRDGUserValidation::ValidateExecuteEnd()
 	{
 		auto ValidateResourceAtExecuteEnd = [](const FRDGParentResourceRef Resource)
 		{
-			check(Resource->ReferenceCount == 0);
+			check(Resource->ReferenceCount == Resource->bExtracted ? 1 : 0);
 
 			const auto& ParentDebugData = Resource->GetParentDebugData();
 			const bool bProducedButNeverUsed = ParentDebugData.PassAccessCount == 1 && ParentDebugData.FirstProducer;
