Add RDG Upload variant which take a lambda function for freeing the CPU memory once the data are uploaded.

diff --git a/Engine/Source/Runtime/RenderCore/Private/RenderGraphBuilder.cpp b/Engine/Source/Runtime/RenderCore/Private/RenderGraphBuilder.cpp
--- a/Engine/Source/Runtime/RenderCore/Private/RenderGraphBuilder.cpp
+++ b/Engine/Source/Runtime/RenderCore/Private/RenderGraphBuilder.cpp
@@ -2115,7 +2115,7 @@ void FRDGBuilder::SetupBufferUploads()
 
 	for (FUploadedBuffer& UploadedBuffer : UploadedBuffers)
 	{
-		if (UploadedBuffer.bUseCallbacks)
+		if (UploadedBuffer.bUseDataCallbacks)
 		{
 			UploadedBuffer.Data = UploadedBuffer.DataCallback();
 			UploadedBuffer.DataSize = UploadedBuffer.DataSizeCallback();
@@ -2159,6 +2159,11 @@ void FRDGBuilder::SubmitBufferUploads()
 				FMemory::Memcpy(DestPtr, UploadedBuffer.Data, UploadedBuffer.DataSize);
 				RHICmdList.UnlockBuffer(UploadedBuffer.Buffer->GetRHI());
 			}
+
+			if (UploadedBuffer.bUseFreeCallbacks)
+			{
+				UploadedBuffer.DataFreeCallback(UploadedBuffer.Data);
+			}
 		}
 	}
 	UploadedBuffers.Reset();
diff --git a/Engine/Source/Runtime/RenderCore/Private/RenderGraphValidation.cpp b/Engine/Source/Runtime/RenderCore/Private/RenderGraphValidation.cpp
--- a/Engine/Source/Runtime/RenderCore/Private/RenderGraphValidation.cpp
+++ b/Engine/Source/Runtime/RenderCore/Private/RenderGraphValidation.cpp
@@ -384,6 +384,12 @@ void FRDGUserValidation::ValidateUploadBuffer(FRDGBufferRef Buffer, const void*
 	checkf(!Buffer->bQueuedForUpload, TEXT("Buffer %s already has an upload queued. Only one upload can be done for each graph."), Buffer->Name);
 	check(InitialData || InitialDataSize == 0);
 }
+void FRDGUserValidation::ValidateUploadBuffer(FRDGBufferRef Buffer, const void* InitialData, uint64 InitialDataSize, const FRDGBufferInitialDataFreeCallback& InitialDataFreeCallback)
+{
+	check(Buffer);
+	checkf(!Buffer->bQueuedForUpload, TEXT("Buffer %s already has an upload queued. Only one upload can be done for each graph."), Buffer->Name);
+	check((InitialData || InitialDataSize == 0) && InitialDataFreeCallback);
+}
 
 void FRDGUserValidation::ValidateUploadBuffer(FRDGBufferRef Buffer, const FRDGBufferInitialDataCallback& InitialDataCallback, const FRDGBufferInitialDataSizeCallback& InitialDataSizeCallback)
 {
@@ -392,6 +398,13 @@ void FRDGUserValidation::ValidateUploadBuffer(FRDGBufferRef Buffer, const FRDGBu
 	check(InitialDataCallback && InitialDataSizeCallback);
 }
 
+void FRDGUserValidation::ValidateUploadBuffer(FRDGBufferRef Buffer, const FRDGBufferInitialDataCallback& InitialDataCallback, const FRDGBufferInitialDataSizeCallback& InitialDataSizeCallback, const FRDGBufferInitialDataFreeCallback& InitialDataFreeCallback)
+{
+	check(Buffer);
+	checkf(!Buffer->bQueuedForUpload, TEXT("Buffer %s already has an upload queued. Only one upload can be done for each graph."), Buffer->Name);
+	check(InitialDataCallback && InitialDataSizeCallback && InitialDataFreeCallback);
+}
+
 void FRDGUserValidation::ValidateExtractTexture(FRDGTextureRef Texture, TRefCountPtr<IPooledRenderTarget>* OutTexturePtr)
 {
 	ValidateExtractResource(Texture);
diff --git a/Engine/Source/Runtime/RenderCore/Public/RenderGraphBuilder.h b/Engine/Source/Runtime/RenderCore/Public/RenderGraphBuilder.h
--- a/Engine/Source/Runtime/RenderCore/Public/RenderGraphBuilder.h
+++ b/Engine/Source/Runtime/RenderCore/Public/RenderGraphBuilder.h
@@ -195,10 +195,20 @@ public:
 		QueueBufferUpload(Buffer, Container.GetData(), Container.Num() * sizeof(ElementType), InitialDataFlags);
 	}
 
+	/** Queues a buffer upload operation prior to execution. The resource lifetime is extended and the data is uploaded prior to executing passes. */
+	void QueueBufferUpload(FRDGBufferRef Buffer, const void* InitialData, uint64 InitialDataSize, FRDGBufferInitialDataFreeCallback&& InitialDataFreeCallback);
+
+	template <typename ElementType>
+	inline void QueueBufferUpload(FRDGBufferRef Buffer, TArrayView<ElementType, int32> Container, FRDGBufferInitialDataFreeCallback&& InitialDataFreeCallback)
+	{
+		QueueBufferUpload(Buffer, Container.GetData(), Container.Num() * sizeof(ElementType), InitialDataFreeCallback);
+	}
+
 	/** A variant where InitialData and InitialDataSize are supplied through callbacks. This allows queuing an upload with information unknown at
 	 *  creation time. The callbacks are called before RDG pass execution so data must be ready before that.
 	 */
 	void QueueBufferUpload(FRDGBufferRef Buffer, FRDGBufferInitialDataCallback&& InitialDataCallback, FRDGBufferInitialDataSizeCallback&& InitialDataSizeCallback);
+	void QueueBufferUpload(FRDGBufferRef Buffer, FRDGBufferInitialDataCallback&& InitialDataCallback, FRDGBufferInitialDataSizeCallback&& InitialDataSizeCallback, FRDGBufferInitialDataFreeCallback&& InitialDataFreeCallback);
 
 	/** Queues a pooled render target extraction to happen at the end of graph execution. For graph-created textures, this extends
 	 *  the lifetime of the GPU resource until execution, at which point the pointer is filled. If specified, the texture is transitioned
@@ -477,25 +487,47 @@ private:
 		FUploadedBuffer() = default;
 
 		FUploadedBuffer(FRDGBuffer* InBuffer, const void* InData, uint64 InDataSize)
-			: bUseCallbacks(false)
+			: bUseDataCallbacks(false)
+			, bUseFreeCallbacks(false)
+			, Buffer(InBuffer)
+			, Data(InData)
+			, DataSize(InDataSize)
+		{}
+
+		FUploadedBuffer(FRDGBuffer* InBuffer, const void* InData, uint64 InDataSize, FRDGBufferInitialDataFreeCallback&& InDataFreeCallback)
+			: bUseDataCallbacks(false)
+			, bUseFreeCallbacks(true)
 			, Buffer(InBuffer)
 			, Data(InData)
 			, DataSize(InDataSize)
+			, DataFreeCallback(MoveTemp(InDataFreeCallback))
 		{}
 
 		FUploadedBuffer(FRDGBuffer* InBuffer, FRDGBufferInitialDataCallback&& InDataCallback, FRDGBufferInitialDataSizeCallback&& InDataSizeCallback)
-			: bUseCallbacks(true)
+			: bUseDataCallbacks(true)
+			, bUseFreeCallbacks(false)
+			, Buffer(InBuffer)
+			, DataCallback(MoveTemp(InDataCallback))
+			, DataSizeCallback(MoveTemp(InDataSizeCallback))
+		{}
+
+		FUploadedBuffer(FRDGBuffer* InBuffer, FRDGBufferInitialDataCallback&& InDataCallback, FRDGBufferInitialDataSizeCallback&& InDataSizeCallback, FRDGBufferInitialDataFreeCallback&& InDataFreeCallback)
+			: bUseDataCallbacks(true)
+			, bUseFreeCallbacks(true)
 			, Buffer(InBuffer)
 			, DataCallback(MoveTemp(InDataCallback))
 			, DataSizeCallback(MoveTemp(InDataSizeCallback))
+			, DataFreeCallback(MoveTemp(InDataFreeCallback))
 		{}
 
-		bool bUseCallbacks = false;
+		bool bUseDataCallbacks = false;
+		bool bUseFreeCallbacks = false;
 		FRDGBuffer* Buffer{};
 		const void* Data{};
 		uint64 DataSize{};
 		FRDGBufferInitialDataCallback DataCallback;
 		FRDGBufferInitialDataSizeCallback DataSizeCallback;
+		FRDGBufferInitialDataFreeCallback DataFreeCallback;
 	};
 
 	TArray<FUploadedBuffer, FRDGArrayAllocator> UploadedBuffers;
diff --git a/Engine/Source/Runtime/RenderCore/Public/RenderGraphBuilder.inl b/Engine/Source/Runtime/RenderCore/Public/RenderGraphBuilder.inl
--- a/Engine/Source/Runtime/RenderCore/Public/RenderGraphBuilder.inl
+++ b/Engine/Source/Runtime/RenderCore/Public/RenderGraphBuilder.inl
@@ -229,6 +229,19 @@ inline void FRDGBuilder::QueueBufferUpload(FRDGBufferRef Buffer, const void* Ini
 	Buffer->bQueuedForUpload = 1;
 }
 
+inline void FRDGBuilder::QueueBufferUpload(FRDGBufferRef Buffer, const void* InitialData, uint64 InitialDataSize, FRDGBufferInitialDataFreeCallback&& InitialDataFreeCallback)
+{
+	IF_RDG_ENABLE_DEBUG(UserValidation.ValidateUploadBuffer(Buffer, InitialData, InitialDataSize));
+
+	if (InitialDataSize == 0)
+	{
+		return;
+	}
+
+	UploadedBuffers.Emplace(Buffer, InitialData, InitialDataSize, MoveTemp(InitialDataFreeCallback));
+	Buffer->bQueuedForUpload = 1;
+}
+
 inline void FRDGBuilder::QueueBufferUpload(FRDGBufferRef Buffer, FRDGBufferInitialDataCallback&& InitialDataCallback, FRDGBufferInitialDataSizeCallback&& InitialDataSizeCallback)
 {
 	IF_RDG_ENABLE_DEBUG(UserValidation.ValidateUploadBuffer(Buffer, InitialDataCallback, InitialDataSizeCallback));
@@ -237,6 +250,14 @@ inline void FRDGBuilder::QueueBufferUpload(FRDGBufferRef Buffer, FRDGBufferIniti
 	Buffer->bQueuedForUpload = 1;
 }
 
+inline void FRDGBuilder::QueueBufferUpload(FRDGBufferRef Buffer, FRDGBufferInitialDataCallback&& InitialDataCallback, FRDGBufferInitialDataSizeCallback&& InitialDataSizeCallback, FRDGBufferInitialDataFreeCallback&& InitialDataFreeCallback)
+{
+	IF_RDG_ENABLE_DEBUG(UserValidation.ValidateUploadBuffer(Buffer, InitialDataCallback, InitialDataSizeCallback, InitialDataFreeCallback));
+
+	UploadedBuffers.Emplace(Buffer, MoveTemp(InitialDataCallback), MoveTemp(InitialDataSizeCallback), MoveTemp(InitialDataFreeCallback));
+	Buffer->bQueuedForUpload = 1;
+}
+
 inline void FRDGBuilder::QueueTextureExtraction(FRDGTextureRef Texture, TRefCountPtr<IPooledRenderTarget>* OutTexturePtr, ERHIAccess AccessFinal)
 {
 	QueueTextureExtraction(Texture, OutTexturePtr);
diff --git a/Engine/Source/Runtime/RenderCore/Public/RenderGraphDefinitions.h b/Engine/Source/Runtime/RenderCore/Public/RenderGraphDefinitions.h
--- a/Engine/Source/Runtime/RenderCore/Public/RenderGraphDefinitions.h
+++ b/Engine/Source/Runtime/RenderCore/Public/RenderGraphDefinitions.h
@@ -590,4 +590,5 @@ class FRDGTrace;
 using FRDGBufferNumElementsCallback = TFunction<uint32()>;
 using FRDGBufferInitialDataCallback = TFunction<const void*()>;
 using FRDGBufferInitialDataSizeCallback = TFunction<uint64()>;
+using FRDGBufferInitialDataFreeCallback = TFunction<void(const void* InData)>;
 using FRDGDispatchGroupCountCallback = TFunction<FIntVector()>;
\ No newline at end of file
diff --git a/Engine/Source/Runtime/RenderCore/Public/RenderGraphValidation.h b/Engine/Source/Runtime/RenderCore/Public/RenderGraphValidation.h
--- a/Engine/Source/Runtime/RenderCore/Public/RenderGraphValidation.h
+++ b/Engine/Source/Runtime/RenderCore/Public/RenderGraphValidation.h
@@ -60,7 +60,9 @@ public:
 	void ValidateRegisterExternalBuffer(FRDGBufferRef Buffer);
 
 	void ValidateUploadBuffer(FRDGBufferRef Buffer, const void* InitialData, uint64 InitialDataSize);
+	void ValidateUploadBuffer(FRDGBufferRef Buffer, const void* InitialData, uint64 InitialDataSize, const FRDGBufferInitialDataFreeCallback& InitialDataFreeCallback);
 	void ValidateUploadBuffer(FRDGBufferRef Buffer, const FRDGBufferInitialDataCallback& InitialDataCallback, const FRDGBufferInitialDataSizeCallback& InitialDataSizeCallback);
+	void ValidateUploadBuffer(FRDGBufferRef Buffer, const FRDGBufferInitialDataCallback& InitialDataCallback, const FRDGBufferInitialDataSizeCallback& InitialDataSizeCallback, const FRDGBufferInitialDataFreeCallback& InitialDataFreeCallback);
 
 	/** Validates a resource extraction operation. */
 	void ValidateExtractTexture(FRDGTextureRef Texture, TRefCountPtr<IPooledRenderTarget>* OutTexturePtr);
