Support FRWBuffer scattered updates.
- Modified MemsetResource and MemcpyResource to take number of elements instead of number of bytes as input, since FRWBuffer can use different formats.

diff --git a/Engine/Shaders/Private/ByteBuffer.usf b/Engine/Shaders/Private/ByteBuffer.usf
--- a/Engine/Shaders/Private/ByteBuffer.usf
+++ b/Engine/Shaders/Private/ByteBuffer.usf
@@ -2,10 +2,11 @@
 
 #include "Common.ush"
 
-#define RESOURCE_TYPE_FLOAT4_STRUCTURED_BUFFER	(0)
-#define RESOURCE_TYPE_UINT_BUFFER				(1)
-#define RESOURCE_TYPE_UINT4_ALIGNED_BUFFER		(2)
-#define RESOURCE_TYPE_FLOAT4_TEXTURE			(3)
+#define RESOURCE_TYPE_FLOAT4_BUFFER				(0)
+#define RESOURCE_TYPE_FLOAT4_STRUCTURED_BUFFER	(1)
+#define RESOURCE_TYPE_UINT_BUFFER				(2)
+#define RESOURCE_TYPE_UINT4_ALIGNED_BUFFER		(3)
+#define RESOURCE_TYPE_FLOAT4_TEXTURE			(4)
 
 uint Value;
 uint Size;
@@ -14,7 +15,12 @@ uint NumScatters;
 uint SrcOffset;
 uint DstOffset;
 
-#if RESOURCE_TYPE == RESOURCE_TYPE_FLOAT4_STRUCTURED_BUFFER
+#if RESOURCE_TYPE == RESOURCE_TYPE_FLOAT4_BUFFER
+	StructuredBuffer<uint>			ScatterStructuredBuffer;
+	Buffer<float4>					SrcBuffer;
+	StructuredBuffer<float4>		UploadStructuredBuffer;
+	RWBuffer<float4>				DstBuffer;
+#elif RESOURCE_TYPE == RESOURCE_TYPE_FLOAT4_STRUCTURED_BUFFER
 	StructuredBuffer<uint>			ScatterStructuredBuffer;
 	StructuredBuffer<float4>		SrcStructuredBuffer;
 	StructuredBuffer<float4>		UploadStructuredBuffer;
@@ -34,7 +40,13 @@ uint DstOffset;
 [numthreads(64, 1, 1)]
 void MemsetBufferCS( uint ThreadId : SV_DispatchThreadID ) 
 {
-#if RESOURCE_TYPE == RESOURCE_TYPE_FLOAT4_STRUCTURED_BUFFER
+#if RESOURCE_TYPE == RESOURCE_TYPE_FLOAT4_BUFFER
+	// Size is in float4s
+	if (ThreadId < Size)
+	{
+		DstBuffer[DstOffset + ThreadId] = asfloat(Value);
+	}
+#elif RESOURCE_TYPE == RESOURCE_TYPE_FLOAT4_STRUCTURED_BUFFER
 	// Size is in float4s
 	if( ThreadId < Size )
 	{
@@ -77,7 +89,13 @@ void MemsetBufferCS( uint ThreadId : SV_DispatchThreadID )
 [numthreads(64, 1, 1)]
 void MemcpyCS( uint ThreadId : SV_DispatchThreadID ) 
 {
-#if RESOURCE_TYPE == RESOURCE_TYPE_FLOAT4_STRUCTURED_BUFFER
+#if RESOURCE_TYPE == RESOURCE_TYPE_FLOAT4_BUFFER
+	// Size is in float4s
+	if (ThreadId < Size)
+	{
+		DstBuffer[DstOffset + ThreadId] = SrcBuffer[SrcOffset + ThreadId];
+	}
+#elif RESOURCE_TYPE == RESOURCE_TYPE_FLOAT4_STRUCTURED_BUFFER
 	// Size is in float4s
 	if( ThreadId < Size )
 	{
@@ -119,7 +137,7 @@ void MemcpyCS( uint ThreadId : SV_DispatchThreadID )
 		DstByteAddressBuffer.Store( DstIndex * 4, SrcData );
 	}
 #else
-#error "Not implemented"
+	#error "Not implemented"
 #endif
 }
 
@@ -133,7 +151,11 @@ void ScatterCopyCS( uint DispatchThreadId : SV_DispatchThreadID )
 
 	if( ScatterIndex < NumScatters )
 	{
-#if RESOURCE_TYPE == RESOURCE_TYPE_FLOAT4_STRUCTURED_BUFFER
+#if RESOURCE_TYPE == RESOURCE_TYPE_FLOAT4_BUFFER
+		uint DstIndex = ScatterStructuredBuffer[ScatterIndex] * Size + ScatterOffset;
+		uint SrcIndex = ThreadId;
+		DstBuffer[DstIndex] = UploadStructuredBuffer[SrcIndex];
+#elif RESOURCE_TYPE == RESOURCE_TYPE_FLOAT4_STRUCTURED_BUFFER
 		uint DstIndex = ScatterStructuredBuffer[ ScatterIndex ] * Size + ScatterOffset;
 		uint SrcIndex = ThreadId;
 		DstStructuredBuffer[ DstIndex ] = UploadStructuredBuffer[ SrcIndex ];
@@ -157,7 +179,7 @@ void ScatterCopyCS( uint DispatchThreadId : SV_DispatchThreadID )
 		uint SrcData = UploadByteAddressBuffer.Load( SrcIndex * 4 );
 		DstByteAddressBuffer.Store(DstIndex * 4, SrcData);
 #else
-#error "Not implemented"
+	#error "Not implemented"
 #endif
 	}
 }
diff --git a/Engine/Source/Runtime/RenderCore/Private/UnifiedBuffer.cpp b/Engine/Source/Runtime/RenderCore/Private/UnifiedBuffer.cpp
--- a/Engine/Source/Runtime/RenderCore/Private/UnifiedBuffer.cpp
+++ b/Engine/Source/Runtime/RenderCore/Private/UnifiedBuffer.cpp
@@ -12,6 +12,7 @@
 
 enum class EByteBufferResourceType
 {
+	Float4_Buffer,
 	Float4_StructuredBuffer,
 	Uint_Buffer,
 	Uint4Aligned_Buffer,
@@ -59,6 +60,7 @@ class FByteBufferShader : public FGlobalShader
 		SHADER_PARAMETER(uint32, SrcOffset)
 		SHADER_PARAMETER(uint32, DstOffset)
 		SHADER_PARAMETER(uint32, Float4sPerLine)
+		SHADER_PARAMETER_UAV(RWBuffer<float4>, DstBuffer)
 		SHADER_PARAMETER_UAV(RWStructuredBuffer<float4>, DstStructuredBuffer)
 		SHADER_PARAMETER_UAV(RWByteAddressBuffer, DstByteAddressBuffer)
 		SHADER_PARAMETER_UAV(RWTexture2D<float4>, DstTexture)
@@ -79,6 +81,7 @@ class FMemcpyCS : public FByteBufferShader
 
 	BEGIN_SHADER_PARAMETER_STRUCT(FParameters, )
 		SHADER_PARAMETER_STRUCT_INCLUDE(FByteBufferShader::FParameters, Common)
+		SHADER_PARAMETER_SRV(Buffer<float4>, SrcBuffer)
 		SHADER_PARAMETER_SRV(StructuredBuffer<float4>, SrcStructuredBuffer)
 		SHADER_PARAMETER_SRV(ByteAddressBuffer, SrcByteAddressBuffer)
 		SHADER_PARAMETER_SRV(Texture2D<float4>, SrcTexture)
@@ -104,6 +107,7 @@ IMPLEMENT_GLOBAL_SHADER( FScatterCopyCS, "/Engine/Private/ByteBuffer.usf", "Scat
 
 enum class EResourceType
 {
+	BUFFER,
 	STRUCTURED_BUFFER,
 	BYTEBUFFER,
 	TEXTURE
@@ -112,6 +116,12 @@ enum class EResourceType
 template<typename ResourceType>
 struct ResourceTypeTraits;
 
+template<>
+struct ResourceTypeTraits<FRWBuffer>
+{
+	static const EResourceType Type = EResourceType::BUFFER;
+};
+
 template<>
 struct ResourceTypeTraits<FRWBufferStructured>
 {
@@ -137,26 +147,14 @@ static uint32 CalculateFloat4sPerLine()
 }
 
 template<typename ResourceType>
-void MemsetResource(FRHICommandList& RHICmdList, const ResourceType& DstBuffer, uint32 Value, uint32 NumBytes, uint32 DstOffset)
+void MemsetResource(FRHICommandList& RHICmdList, const ResourceType& DstBuffer, const FMemsetResourceParams& Params)
 {
-	uint32 DivisorAlignment = ResourceTypeTraits<ResourceType>::Type == EResourceType::BYTEBUFFER ? 4 : 16;
-	if (ResourceTypeTraits<ResourceType>::Type == EResourceType::BYTEBUFFER)
-	{
-		check((NumBytes & 3) == 0);
-		check((DstOffset & 3) == 0);
-	}
-	else
-	{
-		check((DstOffset & 15) == 0);
-		check((NumBytes & 15) == 0);
-	}
-
 	EByteBufferResourceType ResourceTypeEnum;
 
 	FMemsetBufferCS::FParameters Parameters;
-	Parameters.Value = Value;
-	Parameters.Size = NumBytes / DivisorAlignment;
-	Parameters.DstOffset = DstOffset / DivisorAlignment;
+	Parameters.Value = Params.Value;
+	Parameters.Size = Params.Count;
+	Parameters.DstOffset = Params.DstOffset;
 
 	if (ResourceTypeTraits<ResourceType>::Type == EResourceType::BYTEBUFFER)
 	{
@@ -164,6 +162,12 @@ void MemsetResource(FRHICommandList& RHICmdList, const ResourceType& DstBuffer,
 
 		Parameters.DstByteAddressBuffer = DstBuffer.UAV;
 	}
+	else if (ResourceTypeTraits<ResourceType>::Type == EResourceType::BUFFER)
+	{
+		ResourceTypeEnum = EByteBufferResourceType::Float4_Buffer;
+
+		Parameters.DstBuffer = DstBuffer.UAV;
+	}
 	else if (ResourceTypeTraits<ResourceType>::Type == EResourceType::STRUCTURED_BUFFER)
 	{
 		ResourceTypeEnum = EByteBufferResourceType::Float4_StructuredBuffer;
@@ -185,42 +189,34 @@ void MemsetResource(FRHICommandList& RHICmdList, const ResourceType& DstBuffer,
 
 	auto ComputeShader = GetGlobalShaderMap(GMaxRHIFeatureLevel)->GetShader<FMemsetBufferCS>(PermutationVector);
 
-	FComputeShaderUtils::Dispatch(RHICmdList, ComputeShader, Parameters, FIntVector(FMath::DivideAndRoundUp(NumBytes / 16, 64u), 1, 1));
+	// each thread will set 4 floats / uints
+	const uint32 Divisor = ResourceTypeTraits<ResourceType>::Type == EResourceType::BYTEBUFFER ? 4 : 1;
+
+	FComputeShaderUtils::Dispatch(RHICmdList, ComputeShader, Parameters, FIntVector(FMath::DivideAndRoundUp(Params.Count / Divisor, 64u), 1, 1));
 }
 
 template<typename ResourceType>
-void MemcpyResource(FRHICommandList& RHICmdList, const ResourceType& DstBuffer, const ResourceType& SrcBuffer, uint32 NumBytes, uint32 DstOffset, uint32 SrcOffset, bool bAlreadyInUAVOverlap)
+void MemcpyResource(FRHICommandList& RHICmdList, const ResourceType& DstBuffer, const ResourceType& SrcBuffer, const FMemcpyResourceParams& Params, bool bAlreadyInUAVOverlap)
 {
-	uint32 DivisorAlignment = ResourceTypeTraits<ResourceType>::Type == EResourceType::BYTEBUFFER ? 4 : 16;
-	if(ResourceTypeTraits<ResourceType>::Type == EResourceType::BYTEBUFFER)
-	{
-		check((NumBytes & 3) == 0);
-		check((SrcOffset & 3) == 0);
-		check((DstOffset & 3) == 0);
-	}
-	else
-	{
-		check((SrcOffset & 15) == 0);
-		check((DstOffset & 15) == 0);
-		check((NumBytes & 15) == 0);
-	}
-
-	uint32 NumBytesProcessed = 0;
+	// each thread will copy 4 floats / uints
+	const uint32 Divisor = ResourceTypeTraits<ResourceType>::Type == EResourceType::BYTEBUFFER ? 4 : 1;
 
 	if (!bAlreadyInUAVOverlap)	// TODO: Get rid of this check once BeginUAVOverlap/EndUAVOverlap supports nesting.
-		RHICmdList.BeginUAVOverlap(DstBuffer.UAV);	
+		RHICmdList.BeginUAVOverlap(DstBuffer.UAV);
+
+	uint32 NumElementsProcessed = 0;
 
-	while (NumBytesProcessed < NumBytes)
+	while (NumElementsProcessed < Params.Count)
 	{
-		const uint32 NumWaves = FMath::Max(FMath::Min<uint32>(GRHIMaxDispatchThreadGroupsPerDimension.X, FMath::DivideAndRoundUp(NumBytes / 16, 64u)), 1u);
-		const uint32 NumBytesPerDispatch = FMath::Min(FMath::Max(NumWaves, 1u) * 16 * 64, NumBytes);
+		const uint32 NumWaves = FMath::Max(FMath::Min<uint32>(GRHIMaxDispatchThreadGroupsPerDimension.X, FMath::DivideAndRoundUp(Params.Count / Divisor, 64u)), 1u);
+		const uint32 NumElementsPerDispatch = FMath::Min(FMath::Max(NumWaves, 1u) * Divisor * 64, Params.Count - NumElementsProcessed);
 
 		EByteBufferResourceType ResourceTypeEnum;
 
 		FMemcpyCS::FParameters Parameters;
-		Parameters.Common.Size = NumBytesPerDispatch / DivisorAlignment;
-		Parameters.Common.SrcOffset = (SrcOffset + NumBytesProcessed) / DivisorAlignment;
-		Parameters.Common.DstOffset = (DstOffset + NumBytesProcessed) / DivisorAlignment;
+		Parameters.Common.Size = NumElementsPerDispatch;
+		Parameters.Common.SrcOffset = (Params.SrcOffset + NumElementsProcessed);
+		Parameters.Common.DstOffset = (Params.DstOffset + NumElementsProcessed);
 
 		if (ResourceTypeTraits<ResourceType>::Type == EResourceType::BYTEBUFFER)
 		{
@@ -236,6 +232,13 @@ void MemcpyResource(FRHICommandList& RHICmdList, const ResourceType& DstBuffer,
 			Parameters.SrcStructuredBuffer = SrcBuffer.SRV;
 			Parameters.Common.DstStructuredBuffer = DstBuffer.UAV;
 		}
+		else if (ResourceTypeTraits<ResourceType>::Type == EResourceType::BUFFER)
+		{
+			ResourceTypeEnum = EByteBufferResourceType::Float4_Buffer;
+
+			Parameters.SrcBuffer = SrcBuffer.SRV;
+			Parameters.Common.DstBuffer = DstBuffer.UAV;
+		}
 		else if (ResourceTypeTraits<ResourceType>::Type == EResourceType::TEXTURE)
 		{
 			ResourceTypeEnum = EByteBufferResourceType::Float4_Texture;
@@ -256,7 +259,7 @@ void MemcpyResource(FRHICommandList& RHICmdList, const ResourceType& DstBuffer,
 
 		FComputeShaderUtils::Dispatch(RHICmdList, ComputeShader, Parameters, FIntVector(NumWaves, 1, 1));
 
-		NumBytesProcessed += NumBytesPerDispatch;
+		NumElementsProcessed += NumElementsPerDispatch;
 	}
 
 	if(!bAlreadyInUAVOverlap)
@@ -284,7 +287,12 @@ RENDERCORE_API bool ResizeResourceIfNeeded<FTextureRWBuffer2D>(FRHICommandList&
 	{
 		FTextureRWBuffer2D NewTexture;
 		NewTexture.Initialize(DebugName, BytesPerElement, Float4sPerLine, NumLines, PF_A32B32G32R32F, TexCreate_RenderTargetable | TexCreate_UAV);
-		MemcpyResource(RHICmdList, NewTexture, Texture, NumBytes, 0, 0);
+
+		FMemcpyResourceParams Params;
+		Params.Count = NumBytes / BytesPerElement;
+		Params.SrcOffset = 0;
+		Params.DstOffset = 0;
+		MemcpyResource(RHICmdList, NewTexture, Texture, Params);
 		Texture = NewTexture;
 		return true;
 	}
@@ -295,24 +303,31 @@ RENDERCORE_API bool ResizeResourceIfNeeded<FTextureRWBuffer2D>(FRHICommandList&
 template<>
 RENDERCORE_API bool ResizeResourceIfNeeded<FRWBufferStructured>(FRHICommandList& RHICmdList, FRWBufferStructured& Buffer, uint32 NumBytes, const TCHAR* DebugName)
 {
-	check((NumBytes & 15) == 0);
+	const uint32 BytesPerElement = 16;
+
+	check((NumBytes & (BytesPerElement - 1)) == 0);
+
+	uint32 NumElements = NumBytes / BytesPerElement;
 
 	if (Buffer.NumBytes == 0)
 	{
-		Buffer.Initialize(DebugName, 16, NumBytes / 16);
+		Buffer.Initialize(DebugName, BytesPerElement, NumElements);
 		return true;
 	}
 	else if (NumBytes != Buffer.NumBytes)
 	{
 		FRWBufferStructured NewBuffer;
-		NewBuffer.Initialize(DebugName, 16, NumBytes / 16);
+		NewBuffer.Initialize(DebugName, BytesPerElement, NumElements);
 
 		RHICmdList.Transition(FRHITransitionInfo(Buffer.UAV, ERHIAccess::Unknown, ERHIAccess::SRVCompute));
 		RHICmdList.Transition(FRHITransitionInfo(NewBuffer.UAV, ERHIAccess::Unknown, ERHIAccess::UAVCompute));
 
 		// Copy data to new buffer
-		uint32 CopyBytes = FMath::Min(NumBytes, Buffer.NumBytes);
-		MemcpyResource(RHICmdList, NewBuffer, Buffer, CopyBytes);
+		FMemcpyResourceParams Params;
+		Params.Count = FMath::Min(NumBytes, Buffer.NumBytes) / BytesPerElement;
+		Params.SrcOffset = 0;
+		Params.DstOffset = 0;
+		MemcpyResource(RHICmdList, NewBuffer, Buffer, Params);
 
 		Buffer = NewBuffer;
 		return true;
@@ -324,6 +339,8 @@ RENDERCORE_API bool ResizeResourceIfNeeded<FRWBufferStructured>(FRHICommandList&
 template<>
 RENDERCORE_API bool ResizeResourceIfNeeded<FRWByteAddressBuffer>(FRHICommandList& RHICmdList, FRWByteAddressBuffer& Buffer, uint32 NumBytes, const TCHAR* DebugName)
 {
+	const uint32 BytesPerElement = 4;
+
 	// Needs to be aligned to 16 bytes to MemcpyResource to work correctly (otherwise it skips last unaligned elements of the buffer during resize)
 	check((NumBytes & 15) == 0);
 
@@ -343,8 +360,11 @@ RENDERCORE_API bool ResizeResourceIfNeeded<FRWByteAddressBuffer>(FRHICommandList
 		});
 
 		// Copy data to new buffer
-		uint32 CopyBytes = FMath::Min(NumBytes, Buffer.NumBytes);
-		MemcpyResource(RHICmdList, NewBuffer, Buffer, CopyBytes);
+		FMemcpyResourceParams Params;
+		Params.Count = FMath::Min(NumBytes, Buffer.NumBytes) / BytesPerElement;
+		Params.SrcOffset = 0;
+		Params.DstOffset = 0;
+		MemcpyResource(RHICmdList, NewBuffer, Buffer, Params);
 
 		Buffer = NewBuffer;
 		return true;
@@ -353,22 +373,61 @@ RENDERCORE_API bool ResizeResourceIfNeeded<FRWByteAddressBuffer>(FRHICommandList
 	return false;
 }
 
-template<>
-RENDERCORE_API bool ResizeResourceSOAIfNeeded<FRWBufferStructured>(FRHICommandList& RHICmdList, FRWBufferStructured& Buffer, uint32 NumBytes, uint32 NumArrays, const TCHAR* DebugName)
+RENDERCORE_API bool ResizeResourceIfNeeded(FRHICommandList& RHICmdList, FRWBuffer& Buffer, EPixelFormat Format, uint32 NumElements, const TCHAR* DebugName)
 {
-	check((NumBytes & 15) == 0);
-	check(NumBytes % NumArrays == 0);
-	check(Buffer.NumBytes % NumArrays == 0);
+	const uint32 BytesPerElement = GPixelFormats[Format].BlockBytes;
+	const uint32 NumBytes = BytesPerElement * NumElements;
 
 	if (Buffer.NumBytes == 0)
 	{
-		Buffer.Initialize(DebugName, 16, NumBytes / 16);
+		Buffer.Initialize(DebugName, BytesPerElement, NumElements, Format);
 		return true;
 	}
 	else if (NumBytes != Buffer.NumBytes)
+	{
+		FRWBuffer NewBuffer;
+		NewBuffer.Initialize(DebugName, BytesPerElement, NumElements, Format);
+
+		RHICmdList.Transition(FRHITransitionInfo(Buffer.UAV, ERHIAccess::Unknown, ERHIAccess::SRVCompute));
+		RHICmdList.Transition(FRHITransitionInfo(NewBuffer.UAV, ERHIAccess::Unknown, ERHIAccess::UAVCompute));
+
+		// Copy data to new buffer
+		FMemcpyResourceParams MemcpyParams;
+		MemcpyParams.Count = FMath::Min(NumBytes, Buffer.NumBytes) / BytesPerElement;
+		MemcpyParams.SrcOffset = 0;
+		MemcpyParams.DstOffset = 0;
+		MemcpyResource(RHICmdList, NewBuffer, Buffer, MemcpyParams);
+
+		Buffer = NewBuffer;
+		return true;
+	}
+
+	return false;
+}
+
+template<>
+RENDERCORE_API bool ResizeResourceSOAIfNeeded<FRWBufferStructured>(FRHICommandList& RHICmdList, FRWBufferStructured& Buffer, const FResizeResourceSOAParams& Params, const TCHAR* DebugName)
+{
+	const uint32 BytesPerElement = 16;
+
+	checkf(Params.NumBytes % BytesPerElement == 0, TEXT("NumBytes (%s) must be a multiple of BytesPerElement (%s)"), Params.NumBytes, BytesPerElement);
+	checkf(Buffer.NumBytes % BytesPerElement == 0, TEXT("NumBytes (%s) must be a multiple of BytesPerElement (%s)"), Buffer.NumBytes, BytesPerElement);
+
+	uint32 NumElements = Params.NumBytes / BytesPerElement;
+	uint32 NumElementsOld = Buffer.NumBytes / BytesPerElement;
+
+	checkf(NumElements % Params.NumArrays == 0, TEXT("NumElements (%s) must be a multiple of NumArrays (%s)"), NumElements, Params.NumArrays);
+	checkf(NumElementsOld % Params.NumArrays == 0, TEXT("NumElements (%s) must be a multiple of NumArrays (%s)"), NumElementsOld, Params.NumArrays);
+
+	if (Buffer.NumBytes == 0)
+	{
+		Buffer.Initialize(DebugName, BytesPerElement, NumElements);
+		return true;
+	}
+	else if (Params.NumBytes != Buffer.NumBytes)
 	{
 		FRWBufferStructured NewBuffer;
-		NewBuffer.Initialize(DebugName, 16, NumBytes / 16);
+		NewBuffer.Initialize(DebugName, BytesPerElement, NumElements);
 
 		RHICmdList.Transition({
 			FRHITransitionInfo(Buffer.UAV, ERHIAccess::Unknown, ERHIAccess::SRVCompute),
@@ -376,15 +435,21 @@ RENDERCORE_API bool ResizeResourceSOAIfNeeded<FRWBufferStructured>(FRHICommandLi
 		});
 
 		// Copy data to new buffer
-		uint32 OldArraySize = Buffer.NumBytes / NumArrays;
-		uint32 NewArraySize = NumBytes / NumArrays;
+		uint32 OldArraySize = NumElementsOld / Params.NumArrays;
+		uint32 NewArraySize = NumElements / Params.NumArrays;
 
 		RHICmdList.BeginUAVOverlap(NewBuffer.UAV);
-		uint32 CopyBytes = FMath::Min(NewArraySize, OldArraySize);
-		for( uint32 i = 0; i < NumArrays; i++ )
+
+		FMemcpyResourceParams MemcpyParams;
+		MemcpyParams.Count = FMath::Min(NewArraySize, OldArraySize);
+
+		for( uint32 i = 0; i < Params.NumArrays; i++ )
 		{
-			MemcpyResource( RHICmdList, NewBuffer, Buffer, CopyBytes, i * NewArraySize, i * OldArraySize, true );
+			MemcpyParams.SrcOffset = i * OldArraySize;
+			MemcpyParams.DstOffset = i * NewArraySize;
+			MemcpyResource( RHICmdList, NewBuffer, Buffer, MemcpyParams, true );
 		}
+
 		RHICmdList.EndUAVOverlap(NewBuffer.UAV);
 
 		Buffer = NewBuffer;
@@ -394,24 +459,32 @@ RENDERCORE_API bool ResizeResourceSOAIfNeeded<FRWBufferStructured>(FRHICommandLi
 	return false;
 }
 
-RENDERCORE_API bool ResizeResourceSOAIfNeeded(FRDGBuilder& GraphBuilder, FRWBufferStructured& Buffer, uint32 NumBytes, uint32 NumArrays, const TCHAR* DebugName)
+RENDERCORE_API bool ResizeResourceSOAIfNeeded(FRDGBuilder& GraphBuilder, FRWBufferStructured& Buffer, const FResizeResourceSOAParams& Params, const TCHAR* DebugName)
 {
-	check((NumBytes & 15) == 0);
-	check(NumBytes % NumArrays == 0);
-	check(Buffer.NumBytes % NumArrays == 0);
+	const uint32 BytesPerElement = 16;
+
+	checkf(Params.NumBytes % BytesPerElement == 0, TEXT("NumBytes (%s) must be a multiple of BytesPerElement (%s)"), Params.NumBytes, BytesPerElement);
+	checkf(Buffer.NumBytes % BytesPerElement == 0, TEXT("NumBytes (%s) must be a multiple of BytesPerElement (%s)"), Buffer.NumBytes, BytesPerElement);
+
+	uint32 NumElements = Params.NumBytes / BytesPerElement;
+	uint32 NumElementsOld = Buffer.NumBytes / BytesPerElement;
+
+	checkf(NumElements % Params.NumArrays == 0, TEXT("NumElements (%s) must be a multiple of NumArrays (%s)"), NumElements, Params.NumArrays);
+	checkf(NumElementsOld % Params.NumArrays == 0, TEXT("NumElements (%s) must be a multiple of NumArrays (%s)"), NumElementsOld, Params.NumArrays);
 
 	if (Buffer.NumBytes == 0)
 	{
-		Buffer.Initialize(DebugName, 16, NumBytes / 16);
+		Buffer.Initialize(DebugName, BytesPerElement, NumElements);
 		return true;
 	}
-	else if (NumBytes != Buffer.NumBytes)
+	else if (Params.NumBytes != Buffer.NumBytes)
 	{
 		FRWBufferStructured NewBuffer;
 		FRWBufferStructured OldBuffer = Buffer;
-		NewBuffer.Initialize(DebugName, 16, NumBytes / 16);
+		NewBuffer.Initialize(DebugName, BytesPerElement, NumElements);
+
 		AddPass(GraphBuilder, RDG_EVENT_NAME("ResizeResourceSOAIfNeeded"), 
-			[OldBuffer, NewBuffer, NumBytes, NumArrays](FRHICommandListImmediate& RHICmdList)
+			[OldBuffer, NewBuffer, NumElements, NumElementsOld, Params](FRHICommandListImmediate& RHICmdList)
 		{
 			RHICmdList.Transition({
 				FRHITransitionInfo(OldBuffer.UAV, ERHIAccess::Unknown, ERHIAccess::SRVCompute),
@@ -419,15 +492,19 @@ RENDERCORE_API bool ResizeResourceSOAIfNeeded(FRDGBuilder& GraphBuilder, FRWBuff
 			});
 
 			// Copy data to new buffer
-			uint32 OldArraySize = OldBuffer.NumBytes / NumArrays;
-			uint32 NewArraySize = NumBytes / NumArrays;
+			uint32 OldArraySize = NumElementsOld / Params.NumArrays;
+			uint32 NewArraySize = NumElements / Params.NumArrays;
 
 			RHICmdList.BeginUAVOverlap(NewBuffer.UAV);
-			uint32 CopyBytes = FMath::Min(NewArraySize, OldArraySize);
 
-			for (uint32 i = 0; i < NumArrays; i++)
+			FMemcpyResourceParams MemcpyParams;
+			MemcpyParams.Count = FMath::Min(NewArraySize, OldArraySize);
+
+			for (uint32 i = 0; i < Params.NumArrays; i++)
 			{
-				MemcpyResource(RHICmdList, NewBuffer, OldBuffer, CopyBytes, i * NewArraySize, i * OldArraySize, true);
+				MemcpyParams.SrcOffset = i * OldArraySize;
+				MemcpyParams.DstOffset = i * NewArraySize;
+				MemcpyResource(RHICmdList, NewBuffer, OldBuffer, MemcpyParams, true);
 			}
 			RHICmdList.EndUAVOverlap(NewBuffer.UAV);
 		});
@@ -440,10 +517,10 @@ RENDERCORE_API bool ResizeResourceSOAIfNeeded(FRDGBuilder& GraphBuilder, FRWBuff
 }
 
 template <typename FBufferType>
-void AddCopyBufferPass(FRDGBuilder& GraphBuilder, const FBufferType &NewBuffer, const FBufferType &OldBuffer)
+void AddCopyBufferPass(FRDGBuilder& GraphBuilder, const FBufferType &NewBuffer, const FBufferType &OldBuffer, uint32 ElementSize)
 {
 	AddPass(GraphBuilder, RDG_EVENT_NAME("ResizeResourceIfNeeded-Copy"), 
-		[OldBuffer, NewBuffer](FRHICommandListImmediate& RHICmdList)
+		[OldBuffer, NewBuffer, ElementSize](FRHICommandListImmediate& RHICmdList)
 	{
 		RHICmdList.Transition({
 			FRHITransitionInfo(OldBuffer.UAV, ERHIAccess::Unknown, ERHIAccess::SRVCompute),
@@ -451,26 +528,34 @@ void AddCopyBufferPass(FRDGBuilder& GraphBuilder, const FBufferType &NewBuffer,
 		});
 
 		// Copy data to new buffer
-		uint32 CopyBytes = FMath::Min(NewBuffer.NumBytes, OldBuffer.NumBytes);
-		MemcpyResource(RHICmdList, NewBuffer, OldBuffer, CopyBytes);
+		FMemcpyResourceParams MemcpyParams;
+		MemcpyParams.Count = FMath::Min(NewBuffer.NumBytes, OldBuffer.NumBytes) / ElementSize;
+		MemcpyParams.SrcOffset = 0;
+		MemcpyParams.DstOffset = 0;
+
+		MemcpyResource(RHICmdList, NewBuffer, OldBuffer, MemcpyParams);
 	});
 }
 
 RENDERCORE_API bool ResizeResourceIfNeeded(FRDGBuilder& GraphBuilder, FRWBufferStructured& Buffer, uint32 NumBytes, const TCHAR* DebugName)
 {
-	check((NumBytes & 15) == 0);
+	const uint32 BytesPerElement = 16;
+
+	checkf((NumBytes % BytesPerElement) == 0, TEXT("NumBytes (%s) must be a multiple of BytesPerElement (%s)"), NumBytes, BytesPerElement);
+
+	uint32 NumElements = NumBytes / BytesPerElement;
 
 	if (Buffer.NumBytes == 0)
 	{
-		Buffer.Initialize(DebugName, 16, NumBytes / 16);
+		Buffer.Initialize(DebugName, BytesPerElement, NumElements);
 		return true;
 	}
 	else if (NumBytes != Buffer.NumBytes)
 	{
 		FRWBufferStructured NewBuffer;
-		NewBuffer.Initialize(DebugName, 16, NumBytes / 16);
+		NewBuffer.Initialize(DebugName, BytesPerElement, NumElements);
 
-		AddCopyBufferPass(GraphBuilder, NewBuffer, Buffer);
+		AddCopyBufferPass(GraphBuilder, NewBuffer, Buffer, BytesPerElement);
 
 		Buffer = NewBuffer;
 		return true;
@@ -494,7 +579,31 @@ RENDERCORE_API bool ResizeResourceIfNeeded(FRDGBuilder& GraphBuilder, FRWByteAdd
 		FRWByteAddressBuffer NewBuffer;
 		NewBuffer.Initialize(DebugName, NumBytes);
 
-		AddCopyBufferPass(GraphBuilder, NewBuffer, Buffer);
+		AddCopyBufferPass(GraphBuilder, NewBuffer, Buffer, 4);
+
+		Buffer = NewBuffer;
+		return true;
+	}
+
+	return false;
+}
+
+RENDERCORE_API bool ResizeResourceIfNeeded(FRDGBuilder& GraphBuilder, FRWBuffer& Buffer, EPixelFormat Format, uint32 NumElements, const TCHAR* DebugName)
+{
+	const uint32 BytesPerElement = GPixelFormats[Format].BlockBytes;
+	const uint32 NumBytes = BytesPerElement * NumElements;
+
+	if (Buffer.NumBytes == 0)
+	{
+		Buffer.Initialize(DebugName, BytesPerElement, NumElements, Format);
+		return true;
+	}
+	else if (NumBytes != Buffer.NumBytes)
+	{
+		FRWBuffer NewBuffer;
+		NewBuffer.Initialize(DebugName, BytesPerElement, NumElements, Format);
+
+		AddCopyBufferPass(GraphBuilder, NewBuffer, Buffer, BytesPerElement);
 
 		Buffer = NewBuffer;
 		return true;
@@ -600,6 +709,14 @@ void FScatterUploadBuffer::ResourceUploadTo(FRHICommandList& RHICmdList, const R
 		Parameters.ScatterStructuredBuffer = ScatterBuffer.SRV;
 		Parameters.Common.DstStructuredBuffer = DstBuffer.UAV;
 	}
+	else if (ResourceTypeTraits<ResourceType>::Type == EResourceType::BUFFER)
+	{
+		ResourceTypeEnum = EByteBufferResourceType::Float4_Buffer;
+
+		Parameters.UploadStructuredBuffer = UploadBuffer.SRV;
+		Parameters.ScatterStructuredBuffer = ScatterBuffer.SRV;
+		Parameters.Common.DstBuffer = DstBuffer.UAV;
+	}
 	else if (ResourceTypeTraits<ResourceType>::Type == EResourceType::TEXTURE)
 	{
 		ResourceTypeEnum = EByteBufferResourceType::Float4_Texture;
@@ -635,13 +752,15 @@ void FScatterUploadBuffer::ResourceUploadTo(FRHICommandList& RHICmdList, const R
 	}
 }
 
-template RENDERCORE_API void MemsetResource<FRWBufferStructured>(FRHICommandList& RHICmdList, const FRWBufferStructured& DstBuffer, uint32 Value, uint32 NumBytes, uint32 DstOffset);
-template RENDERCORE_API void MemsetResource<FRWByteAddressBuffer>(FRHICommandList& RHICmdList, const FRWByteAddressBuffer& DstBuffer, uint32 Value, uint32 NumBytes, uint32 DstOffset);
+template RENDERCORE_API void MemsetResource<FRWBufferStructured>(FRHICommandList& RHICmdList, const FRWBufferStructured& DstBuffer, const FMemsetResourceParams& Params);
+template RENDERCORE_API void MemsetResource<FRWByteAddressBuffer>(FRHICommandList& RHICmdList, const FRWByteAddressBuffer& DstBuffer, const FMemsetResourceParams& Params);
 
-template RENDERCORE_API void MemcpyResource<FTextureRWBuffer2D>(FRHICommandList& RHICmdList, const FTextureRWBuffer2D& DstBuffer, const FTextureRWBuffer2D& SrcBuffer, uint32 NumBytes, uint32 DstOffset, uint32 SrcOffset, bool bAlreadyInUAVOverlap);
-template RENDERCORE_API void MemcpyResource<FRWBufferStructured>(FRHICommandList& RHICmdList, const FRWBufferStructured& DstBuffer, const FRWBufferStructured& SrcBuffer, uint32 NumBytes, uint32 DstOffset, uint32 SrcOffset, bool bAlreadyInUAVOverlap);
-template RENDERCORE_API void MemcpyResource<FRWByteAddressBuffer>(FRHICommandList& RHICmdList, const FRWByteAddressBuffer& DstBuffer, const FRWByteAddressBuffer& SrcBuffer, uint32 NumBytes, uint32 DstOffset, uint32 SrcOffset, bool bAlreadyInUAVOverlap);
+template RENDERCORE_API void MemcpyResource<FTextureRWBuffer2D>(FRHICommandList& RHICmdList, const FTextureRWBuffer2D& DstBuffer, const FTextureRWBuffer2D& SrcBuffer, const FMemcpyResourceParams& Params, bool bAlreadyInUAVOverlap);
+template RENDERCORE_API void MemcpyResource<FRWBuffer>(FRHICommandList& RHICmdList, const FRWBuffer& DstBuffer, const FRWBuffer& SrcBuffer, const FMemcpyResourceParams& Params, bool bAlreadyInUAVOverlap);
+template RENDERCORE_API void MemcpyResource<FRWBufferStructured>(FRHICommandList& RHICmdList, const FRWBufferStructured& DstBuffer, const FRWBufferStructured& SrcBuffer, const FMemcpyResourceParams& Params, bool bAlreadyInUAVOverlap);
+template RENDERCORE_API void MemcpyResource<FRWByteAddressBuffer>(FRHICommandList& RHICmdList, const FRWByteAddressBuffer& DstBuffer, const FRWByteAddressBuffer& SrcBuffer, const FMemcpyResourceParams& Params, bool bAlreadyInUAVOverlap);
 
 template RENDERCORE_API void FScatterUploadBuffer::ResourceUploadTo<FTextureRWBuffer2D>(FRHICommandList& RHICmdList, const FTextureRWBuffer2D& DstBuffer, bool bFlush);
+template RENDERCORE_API void FScatterUploadBuffer::ResourceUploadTo<FRWBuffer>(FRHICommandList& RHICmdList, const FRWBuffer& DstBuffer, bool bFlush);
 template RENDERCORE_API void FScatterUploadBuffer::ResourceUploadTo<FRWBufferStructured>(FRHICommandList& RHICmdList, const FRWBufferStructured& DstBuffer, bool bFlush);
 template RENDERCORE_API void FScatterUploadBuffer::ResourceUploadTo<FRWByteAddressBuffer>(FRHICommandList& RHICmdList, const FRWByteAddressBuffer& DstBuffer, bool bFlush);
\ No newline at end of file
diff --git a/Engine/Source/Runtime/RenderCore/Public/UnifiedBuffer.h b/Engine/Source/Runtime/RenderCore/Public/UnifiedBuffer.h
--- a/Engine/Source/Runtime/RenderCore/Public/UnifiedBuffer.h
+++ b/Engine/Source/Runtime/RenderCore/Public/UnifiedBuffer.h
@@ -14,23 +14,45 @@ class FRDGBuilder;
  * Float4 versions exist for platforms that don't yet support byte address buffers.
  */
 
+struct FMemsetResourceParams
+{
+	uint32 Value;
+	uint32 Count;
+	uint32 DstOffset;
+};
+
+struct FMemcpyResourceParams
+{
+	uint32 Count;
+	uint32 SrcOffset;
+	uint32 DstOffset;
+};
+
+struct FResizeResourceSOAParams
+{
+	uint32 NumBytes;
+	uint32 NumArrays;
+};
+
 template<typename ResourceType>
-extern RENDERCORE_API void MemsetResource(FRHICommandList& RHICmdList, const ResourceType& DstBuffer, uint32 Value, uint32 NumBytes, uint32 DstOffset = 0);
+extern RENDERCORE_API void MemsetResource(FRHICommandList& RHICmdList, const ResourceType& DstBuffer, const FMemsetResourceParams& Params);
 template<typename ResourceType>
-extern RENDERCORE_API void MemcpyResource(FRHICommandList& RHICmdList, const ResourceType& DstBuffer, const ResourceType& SrcBuffer, uint32 NumBytes, uint32 DstOffset = 0, uint32 SrcOffset = 0, bool bAlreadyInUAVOverlap = false);
+extern RENDERCORE_API void MemcpyResource(FRHICommandList& RHICmdList, const ResourceType& DstBuffer, const ResourceType& SrcBuffer, const FMemcpyResourceParams& Params, bool bAlreadyInUAVOverlap = false);
+
 template<typename ResourceType>
-extern RENDERCORE_API bool ResizeResourceIfNeeded(FRHICommandList& RHICmdList, ResourceType& Texture, uint32 NumBytes, const TCHAR* DebugName);
+extern RENDERCORE_API bool ResizeResourceSOAIfNeeded(FRHICommandList& RHICmdList, ResourceType& Texture, const FResizeResourceSOAParams& Params, const TCHAR* DebugName);
 template<typename ResourceType>
-extern RENDERCORE_API bool ResizeResourceSOAIfNeeded( FRHICommandList& RHICmdList, ResourceType& Texture, uint32 NumBytes, uint32 NumArrays, const TCHAR* DebugName );
+extern RENDERCORE_API bool ResizeResourceIfNeeded(FRHICommandList& RHICmdList, ResourceType& Texture, uint32 NumBytes, const TCHAR* DebugName);
+RENDERCORE_API bool ResizeResourceIfNeeded(FRHICommandList& RHICmdList, FRWBuffer& Buffer, EPixelFormat Format, uint32 NumElements, const TCHAR* DebugName);
 
 
 /**
  * This version will resize/allocate the buffer at once and add a RDG pass to perform the copy on the RDG time-line if there was previous data).
  */
-RENDERCORE_API bool ResizeResourceSOAIfNeeded(FRDGBuilder& GraphBuilder, FRWBufferStructured& Buffer, uint32 NumBytes, uint32 NumArrays, const TCHAR* DebugName);
+RENDERCORE_API bool ResizeResourceSOAIfNeeded(FRDGBuilder& GraphBuilder, FRWBufferStructured& Buffer, const FResizeResourceSOAParams& Params, const TCHAR* DebugName);
 RENDERCORE_API bool ResizeResourceIfNeeded(FRDGBuilder& GraphBuilder, FRWBufferStructured& Buffer, uint32 NumBytes, const TCHAR* DebugName);
 RENDERCORE_API bool ResizeResourceIfNeeded(FRDGBuilder& GraphBuilder, FRWByteAddressBuffer& Buffer, uint32 NumBytes, const TCHAR* DebugName);
-
+RENDERCORE_API bool ResizeResourceIfNeeded(FRDGBuilder& GraphBuilder, FRWBuffer& Buffer, EPixelFormat Format, uint32 NumElements, const TCHAR* DebugName);
 
 class FScatterUploadBuffer
 {
diff --git a/Engine/Source/Runtime/Renderer/Private/GPUScene.cpp b/Engine/Source/Runtime/Renderer/Private/GPUScene.cpp
--- a/Engine/Source/Runtime/Renderer/Private/GPUScene.cpp
+++ b/Engine/Source/Runtime/Renderer/Private/GPUScene.cpp
@@ -761,9 +761,12 @@ FGPUSceneBufferState FGPUScene::UpdateBufferState(FRDGBuilder& GraphBuilder, FSc
 	BufferState.bResizedPrimitiveData = ResizeResourceIfNeeded(GraphBuilder, PrimitiveBuffer, SizeReserve * sizeof(FPrimitiveSceneShaderData::Data), TEXT("GPUScene.PrimitiveData"));
 	BufferState.PrimitiveBuffer = PrimitiveBuffer;
 
-	const uint32 InstanceSceneDataNumArrays = FInstanceSceneShaderData::DataStrideInFloat4s;
 	const uint32 InstanceSceneDataSizeReserve = FMath::RoundUpToPowerOfTwo(FMath::Max(InstanceSceneDataAllocator.GetMaxSize(), InitialBufferSize));
-	BufferState.bResizedInstanceSceneData = ResizeResourceSOAIfNeeded(GraphBuilder, InstanceSceneDataBuffer, InstanceSceneDataSizeReserve * sizeof(FInstanceSceneShaderData::Data), InstanceSceneDataNumArrays, TEXT("GPUScene.InstanceSceneData"));
+	FResizeResourceSOAParams ResizeParams;
+	ResizeParams.NumBytes = InstanceSceneDataSizeReserve * sizeof(FInstanceSceneShaderData::Data);
+	ResizeParams.NumArrays = FInstanceSceneShaderData::DataStrideInFloat4s;
+
+	BufferState.bResizedInstanceSceneData = ResizeResourceSOAIfNeeded(GraphBuilder, InstanceSceneDataBuffer, ResizeParams, TEXT("GPUScene.InstanceSceneData"));
 	BufferState.InstanceSceneDataBuffer = InstanceSceneDataBuffer;
 	InstanceSceneDataSOAStride = InstanceSceneDataSizeReserve;
 	BufferState.InstanceSceneDataSOAStride = InstanceSceneDataSizeReserve;
