Extract source location from shader compiler errors if backend doesn't provide it separately.

#rnx

diff --git a/Engine/Source/Runtime/Engine/Private/ShaderCompiler/ShaderCompiler.cpp b/Engine/Source/Runtime/Engine/Private/ShaderCompiler/ShaderCompiler.cpp
--- a/Engine/Source/Runtime/Engine/Private/ShaderCompiler/ShaderCompiler.cpp
+++ b/Engine/Source/Runtime/Engine/Private/ShaderCompiler/ShaderCompiler.cpp
@@ -1317,6 +1317,9 @@ static void ProcessErrors(const FShaderCompileJob& CurrentJob, TArray<FString>&
 			// This unique error is being processed, remove it from the array
 			UniqueErrors.RemoveAt(UniqueError);
 
+			// Extract source location from error message if the shader backend doesn't provide it separated from the stripped message
+			CurrentError.ExtractSourceLocation();
+
 			// Remap filenames
 			if (CurrentError.ErrorVirtualFilePath == TEXT("/Engine/Generated/Material.ush"))
 			{
@@ -1372,12 +1375,12 @@ static void ProcessErrors(const FShaderCompileJob& CurrentJob, TArray<FString>&
 					*CurrentJob.Input.VirtualSourceFilePath);
 			}
 
-			FString UniqueErrorString = UniqueErrorPrefix + CurrentError.StrippedErrorMessage + TEXT("\n");
+			FString UniqueErrorString = UniqueErrorPrefix + CurrentError.GetErrorStringWithLineMarker() + TEXT("\n");
 
 			if (GIsBuildMachine)
 			{
 				// Format everything on one line, and with the correct verbosity, so we can display proper errors in the failure logs.
-				UE_LOG(LogShaderCompilers, Error, TEXT("%s%s"), *UniqueErrorPrefix.Replace(TEXT("\n"), TEXT("")), *CurrentError.StrippedErrorMessage);
+				UE_LOG(LogShaderCompilers, Error, TEXT("%s%s"), *UniqueErrorPrefix.Replace(TEXT("\n"), TEXT("")), *CurrentError.GetErrorStringWithLineMarker());
 			}
 			else if (FPlatformMisc::IsDebuggerPresent())
 			{
diff --git a/Engine/Source/Runtime/RenderCore/Private/ShaderCore.cpp b/Engine/Source/Runtime/RenderCore/Private/ShaderCore.cpp
--- a/Engine/Source/Runtime/RenderCore/Private/ShaderCore.cpp
+++ b/Engine/Source/Runtime/RenderCore/Private/ShaderCore.cpp
@@ -1451,6 +1451,105 @@ FString MakeInjectedShaderCodeBlock(const TCHAR* BlockName, const FString& CodeT
 	return FString::Printf(TEXT("#line 1 \"%s\"\n%s"), BlockName, *CodeToInject);
 }
 
+FString FShaderCompilerError::GetErrorStringWithSourceLocation() const
+{
+	if (!ErrorVirtualFilePath.IsEmpty() && !ErrorLineString.IsEmpty())
+	{
+		return ErrorVirtualFilePath + TEXT("(") + ErrorLineString + TEXT("): ") + StrippedErrorMessage;
+	}
+	else
+	{
+		return StrippedErrorMessage;
+	}
+}
+
+FString FShaderCompilerError::GetErrorStringWithLineMarker() const
+{
+	if (HasLineMarker())
+	{
+		// Append highlighted line and its marker to the same error message with line terminators
+		// to get a similar multiline error output as with DXC
+		return (GetErrorStringWithSourceLocation() + LINE_TERMINATOR + TEXT("\t") + HighlightedLine + LINE_TERMINATOR + TEXT("\t") + HighlightedLineMarker);
+	}
+	else
+	{
+		return GetErrorStringWithSourceLocation();
+	}
+}
+
+FString FShaderCompilerError::GetErrorString(bool bOmitLineMarker) const
+{
+	return bOmitLineMarker ? GetErrorStringWithSourceLocation() : GetErrorStringWithLineMarker();
+}
+
+static bool ExtractSourceLocationFromCompilerMessage(FString& CompilerMessage, FString& OutFilePath, int32& OutRow, int32& OutColumn, const TCHAR* LeftBracket, const TCHAR* MiddleBracket, const TCHAR* RightBracket)
+{
+	// Ignore ':' character from absolute paths in Windows format
+	const int32 StartPosition = (CompilerMessage.Len() >= 3 && FChar::IsAlpha(CompilerMessage[0]) && CompilerMessage[1] == TEXT(':') && (CompilerMessage[2] == TEXT('/') || CompilerMessage[2] == TEXT('\\')) ? 3 : 0);
+
+	const int32 LeftBracketLen = FCString::Strlen(LeftBracket);
+	const int32 LeftPosition = CompilerMessage.Find(LeftBracket, ESearchCase::IgnoreCase, ESearchDir::FromStart, StartPosition);
+	if (LeftPosition == INDEX_NONE || LeftPosition == StartPosition || LeftPosition + LeftBracketLen >= CompilerMessage.Len() || !FChar::IsDigit(CompilerMessage[LeftPosition + LeftBracketLen]))
+	{
+		return false;
+	}
+
+	const int32 MiddleBracketLen = FCString::Strlen(MiddleBracket);
+	const int32 MiddlePosition = CompilerMessage.Find(MiddleBracket, ESearchCase::IgnoreCase, ESearchDir::FromStart, LeftPosition + LeftBracketLen);
+	if (MiddlePosition == INDEX_NONE || MiddlePosition + MiddleBracketLen >= CompilerMessage.Len() || !FChar::IsDigit(CompilerMessage[MiddlePosition + MiddleBracketLen]))
+	{
+		return false;
+	}
+
+	const int32 RightBracketLen = FCString::Strlen(RightBracket);
+	const int32 RightPosition = CompilerMessage.Find(RightBracket, ESearchCase::IgnoreCase, ESearchDir::FromStart, MiddlePosition + MiddleBracketLen);
+	if (RightPosition == INDEX_NONE || RightPosition >= CompilerMessage.Len())
+	{
+		return false;
+	}
+
+	// Extract file path, row, and column from compiler message
+	OutFilePath = CompilerMessage.Left(LeftPosition);
+	LexFromString(OutRow, *CompilerMessage.Mid(LeftPosition + LeftBracketLen, MiddlePosition - LeftPosition - LeftBracketLen));
+	LexFromString(OutColumn, *CompilerMessage.Mid(MiddlePosition + MiddleBracketLen, RightPosition - MiddlePosition - MiddleBracketLen));
+
+	// Remove extracted information from compiler message
+	CompilerMessage = CompilerMessage.Right(CompilerMessage.Len() - RightPosition - RightBracketLen);
+
+	return true;
+}
+
+bool FShaderCompilerError::ExtractSourceLocation()
+{
+	// Ignore this call if a file path and line string is already provided
+	if (!StrippedErrorMessage.IsEmpty() && ErrorVirtualFilePath.IsEmpty() && ErrorLineString.IsEmpty())
+	{
+		auto ExtractSourceLineFromStrippedErrorMessage = [this](const TCHAR* LeftBracket, const TCHAR* MiddleBracket, const TCHAR* RightBracket) -> bool
+		{
+			int32 Row = 0, Column = 0;
+			if (ExtractSourceLocationFromCompilerMessage(StrippedErrorMessage, ErrorVirtualFilePath, Row, Column, LeftBracket, MiddleBracket, RightBracket))
+			{
+				// Format error line string to MSVC format to be able to jump to the source location with a double click in VisualStudio
+				ErrorLineString = FString::Printf(TEXT("%d,%d"), Row, Column);
+				return true;
+			}
+			return false;
+		};
+
+		// Extract from Clang format, e.g. "Filename:3:12: error:"
+		if (ExtractSourceLineFromStrippedErrorMessage(TEXT(":"), TEXT(":"), TEXT(": ")))
+		{
+			return true;
+		}
+
+		// Extract from MSVC format, e.g. "Filename(3,12) : error: "
+		if (ExtractSourceLineFromStrippedErrorMessage(TEXT("("), TEXT(","), TEXT(") : ")))
+		{
+			return true;
+		}
+	}
+	return false;
+}
 
 FString FShaderCompilerError::GetShaderSourceFilePath() const
 {
diff --git a/Engine/Source/Runtime/RenderCore/Public/ShaderCompilerCore.h b/Engine/Source/Runtime/RenderCore/Public/ShaderCompilerCore.h
--- a/Engine/Source/Runtime/RenderCore/Public/ShaderCompilerCore.h
+++ b/Engine/Source/Runtime/RenderCore/Public/ShaderCompilerCore.h
@@ -380,38 +380,13 @@ struct FShaderCompilerError
 	FString HighlightedLineMarker;
 
 	/** Returns the error message with source file and source line (if present), as well as a line marker separated with a LINE_TERMINATOR. */
-	FString GetErrorStringWithoutLineMarker() const
-	{
-		if (ErrorVirtualFilePath.IsEmpty())
-		{
-			return StrippedErrorMessage;
-		}
-		else
-		{
-			return ErrorVirtualFilePath + TEXT("(") + ErrorLineString + TEXT("): ") + StrippedErrorMessage;
-		}
-	}
+	FString RENDERCORE_API GetErrorStringWithSourceLocation() const;
 	
 	/** Returns the error message with source file and source line (if present), as well as a line marker separated with a LINE_TERMINATOR. */
-	FString GetErrorStringWithLineMarker() const
-	{
-		if (HasLineMarker())
-		{
-			// Append highlighted line and its marker to the same error message with line terminators
-			// to get a similar multiline error output as with DXC
-			return (GetErrorStringWithoutLineMarker() + LINE_TERMINATOR + TEXT("\t") + HighlightedLine + LINE_TERMINATOR + TEXT("\t") + HighlightedLineMarker);
-		}
-		else
-		{
-			return GetErrorStringWithoutLineMarker();
-		}
-	}
+	FString RENDERCORE_API GetErrorStringWithLineMarker() const;
 
 	/** Returns the error message with source file and source line (if present). */
-	FString GetErrorString(bool bOmitLineMarker = false) const
-	{
-		return bOmitLineMarker ? GetErrorStringWithoutLineMarker() : GetErrorStringWithLineMarker();
-	}
+	FString RENDERCORE_API GetErrorString(bool bOmitLineMarker = false) const;
 
 	/**
 	Returns true if this error message has a marker string for the highlighted source line where the error occurred. Example:
@@ -419,11 +394,14 @@ struct FShaderCompilerError
 		float b = a;
 				  ^
 	*/
-	bool HasLineMarker() const
+	FORCEINLINE bool HasLineMarker() const
 	{
 		return !HighlightedLine.IsEmpty() && !HighlightedLineMarker.IsEmpty();
 	}
 
+	/** Extracts the file path and source line from StrippedErrorMessage to ErrorVirtualFilePath and ErrorLineString. */
+	bool RENDERCORE_API ExtractSourceLocation();
+
 	/** Returns the path of the underlying source file relative to the process base dir. */
 	FString RENDERCORE_API GetShaderSourceFilePath() const;
 
