Added r.Raytracing.UseOfflineBVH to control whether to use offline built BVHs.

diff --git a/Engine/Source/Runtime/RenderCore/Private/RenderResource.cpp b/Engine/Source/Runtime/RenderCore/Private/RenderResource.cpp
--- a/Engine/Source/Runtime/RenderCore/Private/RenderResource.cpp
+++ b/Engine/Source/Runtime/RenderCore/Private/RenderResource.cpp
@@ -29,6 +29,13 @@ FAutoConsoleVariableRef CVarReadBufferNumFramesUnusedThresold(
 	GGlobalBufferNumFramesUnusedThresold ,
 	TEXT("Number of frames after which unused global resource allocations will be discarded. Set 0 to ignore. (default=30)"));
 
+int32 GVarDebugForceRuntimeBLAS = 0;
+FAutoConsoleVariableRef CVarDebugForceRuntimeBLAS(
+	TEXT("r.Raytracing.DebugForceRuntimeBLAS"),
+	GVarDebugForceRuntimeBLAS,
+	TEXT("Force building BLAS at runtime."),
+	ECVF_ReadOnly);
+
 FThreadSafeCounter FRenderResource::ResourceListIterationActive;
 
 TArray<int32>& GetFreeIndicesList()
@@ -383,6 +390,12 @@ void FRayTracingGeometry::CreateRayTracingGeometry(ERTAccelerationStructureBuild
 		Initializer.OfflineData = &RawData;
 	}
 
+	if (GVarDebugForceRuntimeBLAS && Initializer.OfflineData != nullptr)
+	{
+		Initializer.OfflineData->Discard();
+		Initializer.OfflineData = nullptr;
+	}
+
 	bool bAllSegmentsAreValid = Initializer.Segments.Num() > 0 || Initializer.OfflineData;
 	for (const FRayTracingGeometrySegment& Segment : Initializer.Segments)
 	{
