Read r.MobileNumDynamicPointLights from platform ini instead of project cvar
Add r.MobileNumDynamicPointLights to the shader keystring

diff --git a/Engine/Source/Runtime/RenderCore/Private/Shader.cpp b/Engine/Source/Runtime/RenderCore/Private/Shader.cpp
--- a/Engine/Source/Runtime/RenderCore/Private/Shader.cpp
+++ b/Engine/Source/Runtime/RenderCore/Private/Shader.cpp
@@ -1516,12 +1516,16 @@ void ShaderMapAppendKeyString(EShaderPlatform Platform, FString& KeyString)
 		{
 			static FShaderPlatformCachedIniValue<bool> MobileEnableMovableSpotlightsIniValue(TEXT("/Script/Engine.RendererSettings"), TEXT("r.Mobile.EnableMovableSpotlights"));
 			static FShaderPlatformCachedIniValue<bool> MobileEnableMovableSpotlightsShadowIniValue(TEXT("/Script/Engine.RendererSettings"), TEXT("r.Mobile.EnableMovableSpotlightsShadow"));
+			static FShaderPlatformCachedIniValue<int32> MobileNumDynamicPointLightsIniValue(TEXT("/Script/Engine.RendererSettings"), TEXT("r.MobileNumDynamicPointLights"));
 
 			bool bMobileEnableMovableSpotlights = (MobileEnableMovableSpotlightsIniValue.Get(Platform) != 0);
 			KeyString += (bMobileEnableMovableSpotlights) ? TEXT("_MSPTL") : TEXT("");
 
 			bool bMobileEnableMovableSpotlightsShadow = (MobileEnableMovableSpotlightsShadowIniValue.Get(Platform) != 0);
 			KeyString += (bMobileEnableMovableSpotlights && bMobileEnableMovableSpotlightsShadow) ? TEXT("S") : TEXT("");
+
+			const int32 MobileNumDynamicPointLights = MobileNumDynamicPointLightsIniValue.Get(Platform);
+			KeyString += (MobileNumDynamicPointLights != 0)? FString::Printf(TEXT("_MobDynPL_%d"), MobileNumDynamicPointLights): TEXT("");
 		}
 		
 		{
diff --git a/Engine/Source/Runtime/Renderer/Private/MobileBasePassRendering.h b/Engine/Source/Runtime/Renderer/Private/MobileBasePassRendering.h
--- a/Engine/Source/Runtime/Renderer/Private/MobileBasePassRendering.h
+++ b/Engine/Source/Runtime/Renderer/Private/MobileBasePassRendering.h
@@ -438,8 +438,7 @@ public:
 		OutEnvironment.SetDefine(TEXT("OUTPUT_MOBILE_HDR"), OutputFormat == HDR_LINEAR_64 ? 1u : 0u);
 		if (NumMovablePointLights == INT32_MAX)
 		{
-			static auto* MobileNumDynamicPointLightsCVar = IConsoleManager::Get().FindTConsoleVariableDataInt(TEXT("r.MobileNumDynamicPointLights"));
-			const int32 MaxDynamicPointLights = FMath::Clamp(MobileNumDynamicPointLightsCVar->GetValueOnAnyThread(), 0, MAX_BASEPASS_DYNAMIC_POINT_LIGHTS);
+			const int32 MaxDynamicPointLights = FMath::Clamp(MobileBasePass::MobileNumDynamicPointLightsIniValue.Get(Parameters.Platform), 0, MAX_BASEPASS_DYNAMIC_POINT_LIGHTS);
 			
 			OutEnvironment.SetDefine(TEXT("MAX_DYNAMIC_POINT_LIGHTS"), (uint32)MaxDynamicPointLights);
 			OutEnvironment.SetDefine(TEXT("VARIABLE_NUM_DYNAMIC_POINT_LIGHTS"), (uint32)1);
