FScatterUploadBuffer::Init now respects GetMaxBufferDimension()
Power of two rounding to mitigate fragmentation now happens in bytes, instead of elements, so we get fewer unique buffer sizes.

diff --git a/Engine/Source/Runtime/RenderCore/Private/UnifiedBuffer.cpp b/Engine/Source/Runtime/RenderCore/Private/UnifiedBuffer.cpp
--- a/Engine/Source/Runtime/RenderCore/Private/UnifiedBuffer.cpp
+++ b/Engine/Source/Runtime/RenderCore/Private/UnifiedBuffer.cpp
@@ -619,7 +619,6 @@ void FScatterUploadBuffer::Init( uint32 NumElements, uint32 InNumBytesPerElement
 
 	NumScatters = 0;
 	MaxScatters = NumElements;
-	NumScattersAllocated = FMath::RoundUpToPowerOfTwo( NumElements );
 	NumBytesPerElement = InNumBytesPerElement;
 	bFloat4Buffer = bInFloat4Buffer;
 
@@ -627,7 +626,8 @@ void FScatterUploadBuffer::Init( uint32 NumElements, uint32 InNumBytesPerElement
 	const uint32 TypeSize = bInFloat4Buffer ? 16 : 4;
 
 	uint32 ScatterBytes = NumElements * sizeof( uint32 );
-	uint32 ScatterBufferSize = NumScattersAllocated * sizeof(uint32);
+	uint32 ScatterBufferSize = (uint32)FMath::Min( (uint64)FMath::RoundUpToPowerOfTwo( ScatterBytes ), GetMaxBufferDimension() * sizeof( uint32 ) );
+	check( ScatterBufferSize >= ScatterBytes );
 
 	if( ScatterBytes > ScatterBuffer.NumBytes || ScatterBufferSize < ScatterBuffer.NumBytes / 2)
 	{
@@ -641,7 +641,8 @@ void FScatterUploadBuffer::Init( uint32 NumElements, uint32 InNumBytesPerElement
 	}
 
 	uint32 UploadBytes = NumElements * NumBytesPerElement;
-	uint32 UploadBufferSize = NumScattersAllocated * NumBytesPerElement;
+	uint32 UploadBufferSize = (uint32)FMath::Min( (uint64)FMath::RoundUpToPowerOfTwo( UploadBytes ), GetMaxBufferDimension() * TypeSize );
+	check( UploadBufferSize >= UploadBytes );
 
 	if( UploadBytes > UploadBuffer.NumBytes || UploadBufferSize < UploadBuffer.NumBytes / 2)
 	{
diff --git a/Engine/Source/Runtime/RenderCore/Public/UnifiedBuffer.h b/Engine/Source/Runtime/RenderCore/Public/UnifiedBuffer.h
--- a/Engine/Source/Runtime/RenderCore/Public/UnifiedBuffer.h
+++ b/Engine/Source/Runtime/RenderCore/Public/UnifiedBuffer.h
@@ -67,7 +67,6 @@ public:
 
 	uint32 	NumScatters = 0;
 	uint32 	MaxScatters = 0;
-	uint32	NumScattersAllocated = 0;
 	uint32	NumBytesPerElement = 0;
 
 	bool	bFloat4Buffer = false;
