ShaderCodeLibrary: a hopeful fix for the PVS warning.

diff --git a/Engine/Source/Runtime/RenderCore/Private/ShaderCodeLibrary.cpp b/Engine/Source/Runtime/RenderCore/Private/ShaderCodeLibrary.cpp
--- a/Engine/Source/Runtime/RenderCore/Private/ShaderCodeLibrary.cpp
+++ b/Engine/Source/Runtime/RenderCore/Private/ShaderCodeLibrary.cpp
@@ -3052,7 +3052,7 @@ void FShaderCodeLibrary::UnregisterSharedShaderCodeRequestDelegate_Handle(FDeleg
 bool UE::ShaderLibrary::Private::FNamedShaderLibrary::OpenShaderCode(const FString& ShaderCodeDir, FString const& Library)
 {
 	FShaderLibraryInstance* LibraryInstance = FShaderLibraryInstance::Create(ShaderPlatform, ShaderCodeDir, Library);
-	if (!LibraryInstance)
+	if (LibraryInstance == nullptr)
 	{
 		UE_LOG(LogShaderLibrary, Verbose, TEXT("Cooked Context: No Shared Shader Library for: %s and native library not supported."), *Library);
 		return false;
