Fixed cyclic dependency IoStoreUtiltities->RenderCore->PakFileUtilities->IoStoreUtilities by removing unused pak order file functionality from ShaderCodeLibrary
#rnx

diff --git a/Engine/Source/Runtime/RenderCore/Private/ShaderCodeLibrary.cpp b/Engine/Source/Runtime/RenderCore/Private/ShaderCodeLibrary.cpp
--- a/Engine/Source/Runtime/RenderCore/Private/ShaderCodeLibrary.cpp
+++ b/Engine/Source/Runtime/RenderCore/Private/ShaderCodeLibrary.cpp
@@ -45,7 +45,6 @@ ShaderCodeLibrary.cpp: Bound shader state cache implementation.
 #endif
 
 #if WITH_EDITOR
-#include "PakFileUtilities.h"
 #include "PipelineCacheUtilities.h"
 #include "UObject/Object.h"
 #include "UObject/Class.h"
@@ -1834,8 +1833,6 @@ class FShaderLibrariesCollection
 	FShaderCodeStats EditorShaderCodeStats[EShaderPlatform::SP_NumPlatforms];
 	// At cook time, whether the shader archive supports pipelines (only OpenGL should)
 	bool EditorArchivePipelines[EShaderPlatform::SP_NumPlatforms];
-	// At cook time, the file open order map
-	FPakOrderMap* OpenOrderMap;
 #endif //WITH_EDITOR
 	bool bSupportsPipelines;
 	bool bNativeFormat;
@@ -1872,8 +1869,6 @@ public:
 		{
 			ChunksSaved[Idx].Empty();
 		}
-
-		OpenOrderMap = nullptr;
 #endif
 
 #if UE_SHADERLIB_WITH_INTROSPECTION
@@ -2266,88 +2261,6 @@ public:
 				static_assert(SP_NumPlatforms < 64u, "ShaderPlatform will no longer fit into bitfield.");
 			}
 		}
-
-		if (bAtLeastOneFormatNeedsDeterminism)
-		{
-			LoadFileOpenOrderFiles();
-		}
-	}
-
-	void LoadFileOpenOrderFiles()
-	{
-		// attempt to open the open order map file
-		FString OrderFile;
-		UE_LOG(LogShaderLibrary, Display, TEXT("Shader library set to be deterministic, looking for the order file"));
-		// first, use the override, if any
-		if (FParse::Value(FCommandLine::Get(), TEXT("FileOpenOrderPrimary="), OrderFile))
-		{
-			UE_LOG(LogShaderLibrary, Display, TEXT("Using '%s' as a source for the file open order (passed on the command line)"), *OrderFile);
-		}
-		else
-		{
-			FString PlatformStr;
-			if (FParse::Value(FCommandLine::Get(), TEXT("TARGETPLATFORM="), PlatformStr))
-			{
-				TArray<FString> PlatformNames;
-				if (!(PlatformStr == TEXT("None") || PlatformStr == TEXT("All")))
-				{
-					PlatformStr.ParseIntoArray(PlatformNames, TEXT("+"), true);
-				}
-
-				// only take the first and issue a warning if there's more than one
-				if (PlatformNames.Num() > 1)
-				{
-					UE_LOG(LogShaderLibrary, Warning, TEXT("More than one platform is being targeted, only one of them will be considered for the file open order."));
-				}
-
-				for (int32 Platform = 0; Platform < PlatformNames.Num(); ++Platform)
-				{
-					FString LogFileDirectory = FPaths::Combine(FPlatformMisc::ProjectDir(), TEXT("Platforms"), *PlatformNames[Platform], TEXT("Build"), TEXT("FileOpenOrder"));
-					if (!FPaths::DirectoryExists(LogFileDirectory))
-					{
-						LogFileDirectory = FPaths::Combine(FPlatformMisc::ProjectDir(), TEXT("Build"), *PlatformNames[Platform], TEXT("FileOpenOrder"));
-					}
-					FString LogFilePath = FPaths::Combine(*LogFileDirectory, TEXT("GameOpenOrder.log"));
-					UE_LOG(LogShaderLibrary, Display, TEXT("Checking if '%s' exists..."), *LogFilePath);
-					if (FPaths::FileExists(LogFilePath))
-					{
-						OrderFile = LogFilePath;
-						UE_LOG(LogShaderLibrary, Display, TEXT("Using '%s' as a source for the file open order (inferred from -targetplatform switch)"), *OrderFile);
-						break;
-					}
-				}
-			}
-		}
-
-		if (!OrderFile.IsEmpty())
-		{
-			OpenOrderMap = new FPakOrderMap;
-
-			if (!OpenOrderMap->ProcessOrderFile(*OrderFile))
-			{
-				UE_LOG(LogShaderLibrary, Error, TEXT("Unable to use file open order file '%s', the shader library will not be deterministic."), *OrderFile);
-				delete OpenOrderMap;
-				OpenOrderMap = nullptr;
-			}
-			else
-			{
-				// check for a secondary file, if any
-				FString OrderFileSecondary;
-				if (FParse::Value(FCommandLine::Get(), TEXT("FileOpenOrderSecondary="), OrderFileSecondary))
-				{
-					UE_LOG(LogShaderLibrary, Display, TEXT("Using '%s' as a secondary source for the file open order (passed on the command line)"), *OrderFileSecondary);
-
-					if (!OpenOrderMap->ProcessOrderFile(*OrderFileSecondary))
-					{
-						UE_LOG(LogShaderLibrary, Warning, TEXT("Unable to use secondary file open order file '%s', only the primary one will be used."), *OrderFileSecondary);
-					}
-				}
-			}
-		}
-		else
-		{
-			UE_LOG(LogShaderLibrary, Error, TEXT("Unable to find any file open order file, the shader library will not be deterministic."));
-		}
 	}
 
 	bool NeedsShaderStableKeys(EShaderPlatform Platform) 
diff --git a/Engine/Source/Runtime/RenderCore/RenderCore.Build.cs b/Engine/Source/Runtime/RenderCore/RenderCore.Build.cs
--- a/Engine/Source/Runtime/RenderCore/RenderCore.Build.cs
+++ b/Engine/Source/Runtime/RenderCore/RenderCore.Build.cs
@@ -13,8 +13,6 @@ public class RenderCore : ModuleRules
 		if (Target.bBuildEditor == true)
         {
 			DynamicallyLoadedModuleNames.Add("TargetPlatform");
-			// PakFileUtitilities due to file open order usage by the shader library
-			PrivateDependencyModuleNames.Add("PakFileUtilities");
 			// JSON is used for the asset info in the shader library
 			PrivateDependencyModuleNames.Add("Json");
 			// UObjects are used to produce the full path of the asset by which the shaders are identified
