Enable VSMs on Vulkan

diff --git a/Engine/Source/Runtime/RenderCore/Public/RenderUtils.h b/Engine/Source/Runtime/RenderCore/Public/RenderUtils.h
--- a/Engine/Source/Runtime/RenderCore/Public/RenderUtils.h
+++ b/Engine/Source/Runtime/RenderCore/Public/RenderUtils.h
@@ -673,10 +673,7 @@ inline bool UseVirtualShadowMaps(EShaderPlatform ShaderPlatform, const FStaticFe
 	const bool bPlatformSupportsNanite = DoesPlatformSupportNanite(ShaderPlatform);
 	const bool bForwardShadingEnabled = IsForwardShadingEnabled(ShaderPlatform);
 
-	// TODO: VSMs not currently working on Vulkan
-	const bool bVulkanSM5 = ShaderPlatform == SP_VULKAN_SM5;
-
-	return bVirtualShadowMapsEnabled && bPlatformSupportsNanite && bUseGPUScene && !bForwardShadingEnabled && !bVulkanSM5;
+	return bVirtualShadowMapsEnabled && bPlatformSupportsNanite && bUseGPUScene && !bForwardShadingEnabled;
 }
 
 /**
