Prepare for fixing of 16 byte alignment of array elements for Vulkan.  Make a few macros to help with conversion of cases where a single scalar are used in shader parameter arays (and be can't be packed with other scalars):
[C++] SHADER_PARAMETER_SCALAR_ARRAY(MemberType,MemberName,NumElements): adds an automagically packed array of scalars to shader parameters
[C++,HLSL] GET_SCALAR_ARRAY_ELEMENT(PackedArray, ElementIndex): returns the value at that index in the scalar array
[HLSL] DECLARE_SCALAR_ARRAY(ScalarType, ScalarName, ElementCount): declares an automagically packed array of scalars

diff --git a/Engine/Shaders/Public/Platform.ush b/Engine/Shaders/Public/Platform.ush
--- a/Engine/Shaders/Public/Platform.ush
+++ b/Engine/Shaders/Public/Platform.ush
@@ -942,3 +942,14 @@ uint WaveGetActiveLaneIndexLast()
 #ifndef PLATFORM_NEEDS_PRECISE_SHADOW_DEPTH
 #define PLATFORM_NEEDS_PRECISE_SHADOW_DEPTH 0
 #endif
+
+// Little tools to help with packing scalars arrays
+#ifndef CALC_SCALAR_ARRAY_SIZE
+#define CALC_SCALAR_ARRAY_SIZE(ElementCount) ((ElementCount+3)/4)
+#endif
+#ifndef DECLARE_SCALAR_ARRAY
+#define DECLARE_SCALAR_ARRAY(ScalarType, ScalarName, ElementCount) ScalarType##4 ScalarName[CALC_SCALAR_ARRAY_SIZE(ElementCount)]
+#endif
+#ifndef GET_SCALAR_ARRAY_ELEMENT
+#define GET_SCALAR_ARRAY_ELEMENT(PackedArray, ElementIndex) PackedArray[ElementIndex>>2u][ElementIndex&3u]
+#endif
diff --git a/Engine/Source/Runtime/RenderCore/Public/ShaderParameterMacros.h b/Engine/Source/Runtime/RenderCore/Public/ShaderParameterMacros.h
--- a/Engine/Source/Runtime/RenderCore/Public/ShaderParameterMacros.h
+++ b/Engine/Source/Runtime/RenderCore/Public/ShaderParameterMacros.h
@@ -1421,9 +1421,9 @@ extern RENDERCORE_API FShaderParametersMetadata* FindUniformBufferStructByShader
 	INTERNAL_SHADER_PARAMETER_EXPLICIT(TShaderParameterTypeInfo<MemberType>::BaseType, TShaderParameterTypeInfo<MemberType>, MemberType,MemberName,,,Precision,TEXT(""),false)
 
 /** Adds a constant-buffer stored array of values.
+ *  MemberType size must be a multiple of 16byte (see SHADER_PARAMETER_SCALAR_ARRAY for scalar arrays).
  *
  * Example:
- *	SHADER_PARAMETER_ARRAY(float, MyScalarArray, [8])
  *	SHADER_PARAMETER_ARRAY(FMatrix44f, MyMatrixArray, [2])
  */
 #define SHADER_PARAMETER_ARRAY(MemberType,MemberName,ArrayDecl) \
@@ -1703,4 +1703,49 @@ const TCHAR* GetShaderParameterMacroName(EUniformBufferBaseType ShaderParameterB
 	RDG_BUFFER_ACCESS(MemberName, ERHIAccess::SRVMask | ERHIAccess::IndirectArgs)
 
 #define SHADER_PARAMETER_RDG_BUFFER_ARRAY(ShaderType,MemberName, ArrayDecl) \
-	DEPRECATED_MACRO(5.0, "SHADER_PARAMETER_RDG_BUFFER_ARRAY has been deprecated. Use RDG_BUFFER_ACCESS_ARRAY instead.")
\ No newline at end of file
+	DEPRECATED_MACRO(5.0, "SHADER_PARAMETER_RDG_BUFFER_ARRAY has been deprecated. Use RDG_BUFFER_ACCESS_ARRAY instead.")
+
+
+
+
+/** Upgrade the scalar type to a vector for storage. */
+
+template<typename TypeParameter>
+struct TShaderParameterScalarArrayTypeInfo
+{
+	static_assert(sizeof(TypeParameter) == 0, "This type is not supported for automatic packing.");
+	using PackedArrayType = void;
+};
+
+template<>
+struct TShaderParameterScalarArrayTypeInfo<uint32>
+{
+	using PackedArrayType = FUintVector4;
+};
+
+template<>
+struct TShaderParameterScalarArrayTypeInfo<int32>
+{
+	using PackedArrayType = FIntVector4;
+};
+
+template<>
+struct TShaderParameterScalarArrayTypeInfo<float>
+{
+	using PackedArrayType = FVector4f;
+};
+
+constexpr uint32 CalcPackedArraySize(uint32 NumElements)		{ return (NumElements + 3) / 4; }
+constexpr uint32 CalcPackedArrayIndex(uint32 ElementIndex)		{ return (ElementIndex >> 2);	}
+constexpr uint32 CalcPackedComponentIndex(uint32 ElementIndex)	{ return (ElementIndex & 3);	}
+
+/** Adds a packed constant-buffer stored array of values.
+ *
+ * Example:
+ *	SHADER_PARAMETER_SCALAR_ARRAY(float, MyScalarArray, 8)
+ */
+#define SHADER_PARAMETER_SCALAR_ARRAY(MemberType,MemberName,NumElements) \
+	SHADER_PARAMETER_ARRAY(TShaderParameterScalarArrayTypeInfo<MemberType>::PackedArrayType, MemberName, [CalcPackedArraySize(NumElements)])
+
+#define GET_SCALAR_ARRAY_ELEMENT(PackedArray, ElementIndex) \
+	PackedArray[CalcPackedArrayIndex(ElementIndex)][CalcPackedComponentIndex(ElementIndex)]
