Add bSupportsDxc to DataDrivenPlatformInfo.

#rnx

diff --git a/Engine/Config/Mac/DataDrivenPlatformInfo.ini b/Engine/Config/Mac/DataDrivenPlatformInfo.ini
--- a/Engine/Config/Mac/DataDrivenPlatformInfo.ini
+++ b/Engine/Config/Mac/DataDrivenPlatformInfo.ini
@@ -33,3 +33,4 @@ bSupportsSSDIndirect=true
 bSupportsWaveOperations=false
 bSupportsGPUScene=true
 bSupportsFFTBloom=true
+bSupportsDxc=true
diff --git a/Engine/Config/VulkanPC/DataDrivenPlatformInfo.ini b/Engine/Config/VulkanPC/DataDrivenPlatformInfo.ini
--- a/Engine/Config/VulkanPC/DataDrivenPlatformInfo.ini
+++ b/Engine/Config/VulkanPC/DataDrivenPlatformInfo.ini
@@ -14,6 +14,7 @@ bIsMetalMRT = false
 bIsPC = true
 bIsConsole = false
 bIsAndroidOpenGLES = false
+bSupportsDxc = true
 
 ; //#todo-vulkanddpi
 bSupportsMobileMultiView = false
diff --git a/Engine/Config/Windows/DataDrivenPlatformInfo.ini b/Engine/Config/Windows/DataDrivenPlatformInfo.ini
--- a/Engine/Config/Windows/DataDrivenPlatformInfo.ini
+++ b/Engine/Config/Windows/DataDrivenPlatformInfo.ini
@@ -99,6 +99,7 @@ bSupportsLumenGI=true
 bSupportsDistanceFields=true
 bSupportsSSDIndirect=true
 bSupportsWaveOperations=true
+bSupportsDxc=true
 ; TODO: SM6.6
 bRequiresVendorExtensionsForAtomics=true
 bSupportsAnisotropicMaterials=true
diff --git a/Engine/Source/Runtime/RHI/Private/RHI.cpp b/Engine/Source/Runtime/RHI/Private/RHI.cpp
--- a/Engine/Source/Runtime/RHI/Private/RHI.cpp
+++ b/Engine/Source/Runtime/RHI/Private/RHI.cpp
@@ -2259,6 +2259,7 @@ void FGenericDataDrivenShaderPlatformInfo::ParseDataDrivenShaderInfo(const FConf
 	GET_SECTION_INT_HELPER(MaxMeshShaderThreadGroupSize);
 	GET_SECTION_BOOL_HELPER(bSupportsPerPixelDBufferMask);
 	GET_SECTION_BOOL_HELPER(bIsHlslcc);
+	GET_SECTION_BOOL_HELPER(bSupportsDxc);
 	GET_SECTION_BOOL_HELPER(bSupportsVariableRateShading);
 	GET_SECTION_INT_HELPER(NumberOfComputeThreads);
 
diff --git a/Engine/Source/Runtime/RHI/Public/RHIDefinitions.h b/Engine/Source/Runtime/RHI/Public/RHIDefinitions.h
--- a/Engine/Source/Runtime/RHI/Public/RHIDefinitions.h
+++ b/Engine/Source/Runtime/RHI/Public/RHIDefinitions.h
@@ -359,6 +359,7 @@ class RHI_API FGenericDataDrivenShaderPlatformInfo
 	uint32 MaxMeshShaderThreadGroupSize : 10;
 	uint32 bSupportsPerPixelDBufferMask : 1;
 	uint32 bIsHlslcc : 1;
+	uint32 bSupportsDxc : 1; // Whether DirectXShaderCompiler (DXC) is supported
 	uint32 bSupportsVariableRateShading : 1;
 	uint32 NumberOfComputeThreads : 10;
 	uint32 bWaterUsesSimpleForwardShading : 1;
@@ -691,6 +692,11 @@ public:
 		return Infos[Platform].bIsHlslcc;
 	}
 
+	static FORCEINLINE_DEBUGGABLE const bool GetSupportsDxc(const FStaticShaderPlatform Platform)
+	{
+		return Infos[Platform].bSupportsDxc;
+	}
+
 	static FORCEINLINE_DEBUGGABLE const bool GetSupportsVariableRateShading(const FStaticShaderPlatform Platform)
 	{
 		return Infos[Platform].bSupportsVariableRateShading;
diff --git a/Engine/Source/Runtime/RenderCore/Private/Shader.cpp b/Engine/Source/Runtime/RenderCore/Private/Shader.cpp
--- a/Engine/Source/Runtime/RenderCore/Private/Shader.cpp
+++ b/Engine/Source/Runtime/RenderCore/Private/Shader.cpp
@@ -98,6 +98,14 @@ static TAutoConsoleVariable<int32> CVarShaderCompilerEmitWarningsOnLoad(
 	ECVF_Default
 );
 
+static TAutoConsoleVariable<int32> CVarShadersForceDXC(
+	TEXT("r.Shaders.ForceDXC"),
+	0,
+	TEXT("Forces DirectX Shader Compiler (DXC) to be used for all shaders instead of HLSLcc if supported.\n")
+	TEXT(" 0: Disable (default)\n")
+	TEXT(" 1: Force new compiler for all shaders"),
+	ECVF_ReadOnly);
+
 static TLinkedList<FShaderType*>*			GShaderTypeList = nullptr;
 static TLinkedList<FShaderPipelineType*>*	GShaderPipelineList = nullptr;
 
@@ -1292,14 +1300,24 @@ void DispatchIndirectComputeShader(
 
 bool IsDxcEnabledForPlatform(EShaderPlatform Platform)
 {
+	// Check the generic console variable first (if DXC is supported)
+	if (FDataDrivenShaderPlatformInfo::GetSupportsDxc(Platform))
+	{
+		static const IConsoleVariable* CVar = IConsoleManager::Get().FindConsoleVariable(TEXT("r.Shaders.ForceDXC"));
+		if (CVar && CVar->GetInt() != 0)
+		{
+			return true;
+		}
+	}
+	// Check backend specific console variables next
 	if (IsD3DPlatform(Platform) && IsPCPlatform(Platform))
 	{
-		static const auto CVar = IConsoleManager::Get().FindConsoleVariable(TEXT("r.D3D.ForceDXC"));
+		static const IConsoleVariable* CVar = IConsoleManager::Get().FindConsoleVariable(TEXT("r.D3D.ForceDXC"));
 		return (CVar && CVar->GetInt() != 0);
 	}
 	if (IsOpenGLPlatform(Platform))
 	{
-		static const auto CVar = IConsoleManager::Get().FindConsoleVariable(TEXT("r.OpenGL.ForceDXC"));
+		static const IConsoleVariable* CVar = IConsoleManager::Get().FindConsoleVariable(TEXT("r.OpenGL.ForceDXC"));
 		return (CVar && CVar->GetInt() != 0);
 	}
 	// Hlslcc has been removed for Metal and Vulkan. There is only DXC now.
