Do not precreate clear replacement shaders in the editor (UE-127489).

- The benefit isn't worth having to deal with multiple fallouts from this.

diff --git a/Engine/Source/Runtime/RenderCore/Private/ClearReplacementShaders.cpp b/Engine/Source/Runtime/RenderCore/Private/ClearReplacementShaders.cpp
--- a/Engine/Source/Runtime/RenderCore/Private/ClearReplacementShaders.cpp
+++ b/Engine/Source/Runtime/RenderCore/Private/ClearReplacementShaders.cpp
@@ -54,11 +54,10 @@ IMPLEMENT_SHADER_TYPE4_WITH_TEMPLATE_PREFIX(template<>, RENDERCORE_API, FClearRe
 
 void CreateClearReplacementShaders()
 {
+#if !WITH_EDITOR	// in editor, this function overall has a lesser sence due to the ability to switch the preview modes on the fly and recompile the shaders, which both can execute the code on GT while RT is already created.
 	// if RHI supports MT shader creation, then we don't care to init here
 	ensureMsgf(!GRHISupportsMultithreadedShaderCreation, TEXT("InitClearReplacementShaders() is called while GRHISupportsMultithreadedShaderCreation is true. This is an unnecessary call."));
-#if !WITH_EDITOR	// in editor, this ensure (and the function overall) has a lesser sence due to the ability to switch the preview modes on the fly
 	ensureMsgf(IsInRenderingThread() || GRHISupportsMultithreadedShaderCreation, TEXT("InitClearReplacementShaders() is expected to be called from the render thread if GRHISupportsMultithreadedShaderCreation is false."));
-#endif
 
 #define UE_CREATE_SHADER(TShaderType) \
 	{ \
@@ -116,4 +115,5 @@ void CreateClearReplacementShaders()
 	UE_CREATE_SHADER(FClearReplacementCS_Texture2DArray_Sint4_Bounds);
 
 #undef UE_CREATE_SHADER
+#endif
 }
