Fixed ray tracing is not enabled by default in cooked games with r.raytracing set to 1.

diff --git a/Engine/Source/Runtime/RenderCore/Private/RenderUtils.cpp b/Engine/Source/Runtime/RenderCore/Private/RenderUtils.cpp
--- a/Engine/Source/Runtime/RenderCore/Private/RenderUtils.cpp
+++ b/Engine/Source/Runtime/RenderCore/Private/RenderUtils.cpp
@@ -1523,13 +1523,13 @@ RENDERCORE_API void RenderUtilsInit()
 				}
 				else
 				{
-					//GUseRayTracing = true;
+					GUseRayTracing = true;
 
-					//UE_LOG(LogRendererCore, Log, TEXT("Ray tracing is enabled for the game. Reason: r.RayTracing=1, and r.RayTracing.EnableInGame is not present (default true)."));
+					UE_LOG(LogRendererCore, Log, TEXT("Ray tracing is enabled for the game. Reason: r.RayTracing=1, and r.RayTracing.EnableInGame is not present (default true)."));
 
-					GUseRayTracing = false;
+					//GUseRayTracing = false;
 
-					UE_LOG(LogRendererCore, Log, TEXT("Ray tracing is disabled for the game. Reason: r.RayTracing=1, and r.RayTracing.EnableInGame is not present (default false)."));
+					//UE_LOG(LogRendererCore, Log, TEXT("Ray tracing is disabled for the game. Reason: r.RayTracing=1, and r.RayTracing.EnableInGame is not present (default false)."));
 				}
 			}
 
