- Sped up the shader library scanning across all plugins/pak files, by checking if the plugins have any .ushaderbytecode files before trying to mount anything (and used ParallelFor to parallelize that checking)

diff --git a/Engine/Source/Runtime/RenderCore/Private/ShaderCodeLibrary.cpp b/Engine/Source/Runtime/RenderCore/Private/ShaderCodeLibrary.cpp
--- a/Engine/Source/Runtime/RenderCore/Private/ShaderCodeLibrary.cpp
+++ b/Engine/Source/Runtime/RenderCore/Private/ShaderCodeLibrary.cpp
@@ -26,6 +26,7 @@ ShaderCodeLibrary.cpp: Bound shader state cache implementation.
 #include "Containers/HashTable.h"
 #include "FileCache/FileCache.h"
 #include "Misc/CoreDelegates.h"
+#include "Async/ParallelFor.h"
 
 #include "ShaderPipelineCache.h"
 #include "Misc/FileHelper.h"
@@ -2486,8 +2487,15 @@ static void FShaderCodeLibraryPluginMountedCallback(IPlugin& Plugin)
 {
 	if (Plugin.CanContainContent() && Plugin.IsEnabled())
 	{
-		FShaderCodeLibrary::OpenLibrary(Plugin.GetName(), Plugin.GetBaseDir());
-		FShaderCodeLibrary::OpenLibrary(Plugin.GetName(), Plugin.GetContentDir());
+		// make sure we have some libraries before bothering with this plugin
+		TArray<FString> Found;
+		IFileManager::Get().FindFiles(Found, *Plugin.GetBaseDir(), TEXT(".ushaderbytecode"));
+		IFileManager::Get().FindFiles(Found, *Plugin.GetContentDir(), TEXT(".ushaderbytecode"));
+		if (Found.Num() > 0)
+		{
+			FShaderCodeLibrary::OpenLibrary(Plugin.GetName(), Plugin.GetBaseDir());
+			FShaderCodeLibrary::OpenLibrary(Plugin.GetName(), Plugin.GetContentDir());
+		}
 	}
 }
 
@@ -2556,10 +2564,10 @@ void FShaderCodeLibrary::InitForRuntime(EShaderPlatform ShaderPlatform)
 
 			// mount shader library from the plugins as they may also have global shaders
 			auto Plugins = IPluginManager::Get().GetEnabledPluginsWithContent();
-			for (auto Plugin : Plugins)
+			ParallelFor(Plugins.Num(), [&](int32 Index)
 			{
-				FShaderCodeLibraryPluginMountedCallback(*Plugin);
-			}
+				FShaderCodeLibraryPluginMountedCallback(*Plugins[Index]);
+			});
 		}
 		else
 		{
