Fix win64 clang linker error

diff --git a/Engine/Source/Runtime/RenderCore/RenderCore.Build.cs b/Engine/Source/Runtime/RenderCore/RenderCore.Build.cs
--- a/Engine/Source/Runtime/RenderCore/RenderCore.Build.cs
+++ b/Engine/Source/Runtime/RenderCore/RenderCore.Build.cs
@@ -18,7 +18,7 @@ public class RenderCore : ModuleRules
 			// JSON is used for the asset info in the shader library
 			PrivateDependencyModuleNames.Add("Json");
 			// UObjects are used to produce the full path of the asset by which the shaders are identified
-			PrivateIncludePathModuleNames.Add("CoreUObject");
+			PrivateDependencyModuleNames.Add("CoreUObject");
 		}
 
         PrivateDependencyModuleNames.AddRange(new string[] { "Core", "Projects", "RHI", "ApplicationCore", "TraceLog" });
