SceneDepthAux removed from certain platforms

Resolved Mobile Texture Lookup Scene Depth outputting DeviceZ instead of WorldZ
#[robomerge] main

diff --git a/Engine/Shaders/Private/MaterialTemplate.ush b/Engine/Shaders/Private/MaterialTemplate.ush
--- a/Engine/Shaders/Private/MaterialTemplate.ush
+++ b/Engine/Shaders/Private/MaterialTemplate.ush
@@ -2079,11 +2079,7 @@ MaterialFloat4 MobileSceneTextureLookup(inout FMaterialPixelParameters Parameter
 #else
 	if (SceneTextureId == PPI_SceneDepth)
 	{
-		#if MOBILE_DEFERRED_SHADING
-			MaterialFloat Depth = ConvertFromDeviceZ(Texture2DSample(MobileSceneTextures.SceneDepthTexture, MobileSceneTextures.SceneDepthTextureSampler, UV).r);
-		#else
-			MaterialFloat Depth = ConvertFromDeviceZ(Texture2DSample(MobileSceneTextures.SceneDepthAuxTexture, MobileSceneTextures.SceneDepthAuxTextureSampler, UV).r);
-		#endif
+		MaterialFloat Depth = CalcSceneDepth(UV);
 		return MaterialFloat4(Depth.rrr, 0.f);
 	}
 	else if (SceneTextureId == PPI_CustomDepth)
diff --git a/Engine/Shaders/Private/MobileBasePassPixelShader.usf b/Engine/Shaders/Private/MobileBasePassPixelShader.usf
--- a/Engine/Shaders/Private/MobileBasePassPixelShader.usf
+++ b/Engine/Shaders/Private/MobileBasePassPixelShader.usf
@@ -99,7 +99,7 @@ int NumDynamicPointLights;
 #define FORCE_VERTEX_NORMAL (MOBILE_QL_DISABLE_MATERIAL_NORMAL)
 #define SUPPORT_SPOTLIGHTS_SHADOW (MATERIALBLENDING_SOLID || MATERIALBLENDING_MASKED) && PROJECT_MOBILE_ENABLE_MOVABLE_SPOTLIGHTS_SHADOW
 #define DEFERRED_SHADING_PATH (MOBILE_DEFERRED_SHADING && ((MATERIALBLENDING_SOLID || MATERIALBLENDING_MASKED) && !MATERIAL_SHADINGMODEL_SINGLELAYERWATER))
-#define USE_SCENE_DEPTH_AUX ((DEFERRED_SHADING_PATH && METAL_PROFILE) || !DEFERRED_SHADING_PATH) && (MATERIALBLENDING_SOLID || MATERIALBLENDING_MASKED)
+#define USE_SCENE_DEPTH_AUX ((DEFERRED_SHADING_PATH && METAL_PROFILE) || !DEFERRED_SHADING_PATH) && (MATERIALBLENDING_SOLID || MATERIALBLENDING_MASKED) && !COMPILER_SWITCH
 #if DEFERRED_SHADING_PATH
 	#define SV_TargetDepthAux SV_Target4
 #else
diff --git a/Engine/Shaders/Private/ResolvePixelShader.usf b/Engine/Shaders/Private/ResolvePixelShader.usf
--- a/Engine/Shaders/Private/ResolvePixelShader.usf
+++ b/Engine/Shaders/Private/ResolvePixelShader.usf
@@ -5,9 +5,6 @@
 =============================================================================*/
 
 #include "Common.ush"
-
-#if FEATURE_LEVEL >= FEATURE_LEVEL_SM5 || SIMULATED_PLATFORM  
-
 Texture2DMS<float4> UnresolvedSurface;
 
 
@@ -52,19 +49,4 @@ void MainSingleSample(
 	int2 IntUV = trunc(InUV);
 
 	OutColor = UnresolvedSurface.Load(IntUV,SingleSampleIndex);
-}
-
-#else
-
-Texture2D UnresolvedSurfaceNonMS;
-void MainDepthNonMS(
-	float2 InUV : TEXCOORD0,
-	out float OutDepth : SV_DEPTH
-	)
-{
-	int3 IntUV = int3(trunc(InUV),0);
-	float Sample = UnresolvedSurfaceNonMS.Load(IntUV).r;
-	OutDepth = Sample;
-}
-
-#endif
\ No newline at end of file
+}
\ No newline at end of file
diff --git a/Engine/Shaders/Private/SceneTexturesCommon.ush b/Engine/Shaders/Private/SceneTexturesCommon.ush
--- a/Engine/Shaders/Private/SceneTexturesCommon.ush
+++ b/Engine/Shaders/Private/SceneTexturesCommon.ush
@@ -139,7 +139,7 @@ float LookupDeviceZ( float2 ScreenUV )
 {
 #if	SCENE_TEXTURES_DISABLED
 	return FarDepthValue;
-#elif FORCE_DEPTH_TEXTURE_READS 
+#elif FORCE_DEPTH_TEXTURE_READS || COMPILER_SWITCH
 	// native Depth buffer lookup
 	return Texture2DSampleLevel(MobileSceneTextures.SceneDepthTexture, MobileSceneTextures.SceneDepthTextureSampler, ScreenUV, 0).r;
 #elif (POST_PROCESS_MATERIAL || POST_PROCESS_MATERIAL_MOBILE) && !POST_PROCESS_AR_PASSTHROUGH
diff --git a/Engine/Source/Runtime/D3D12RHI/Private/D3D12Commands.cpp b/Engine/Source/Runtime/D3D12RHI/Private/D3D12Commands.cpp
--- a/Engine/Source/Runtime/D3D12RHI/Private/D3D12Commands.cpp
+++ b/Engine/Source/Runtime/D3D12RHI/Private/D3D12Commands.cpp
@@ -94,7 +94,7 @@ void FD3D12DynamicRHI::SetupRecursiveResources()
 		return;
 
 	TShaderMapRef<FResolveVS> ResolveVertexShader(ShaderMap);
-	if (GMaxRHIShaderPlatform == SP_PCD3D_SM6 || GMaxRHIShaderPlatform == SP_PCD3D_SM5)
+
 	{
 		TShaderMapRef<FResolveDepthPS> ResolvePixelShader_Depth(ShaderMap);
 		ResolvePixelShader_Depth.GetPixelShader();
@@ -102,11 +102,6 @@ void FD3D12DynamicRHI::SetupRecursiveResources()
 		TShaderMapRef<FResolveDepthPS> ResolvePixelShader_SingleSample(ShaderMap);
 		ResolvePixelShader_SingleSample.GetPixelShader();
 	}
-	else
-	{
-		TShaderMapRef<FResolveDepthNonMSPS> ResolvePixelShader_DepthNonMS(ShaderMap);
-		ResolvePixelShader_DepthNonMS.GetPixelShader();
-	}
 }
 
 // Vertex state.
diff --git a/Engine/Source/Runtime/RenderCore/Private/RenderUtils.cpp b/Engine/Source/Runtime/RenderCore/Private/RenderUtils.cpp
--- a/Engine/Source/Runtime/RenderCore/Private/RenderUtils.cpp
+++ b/Engine/Source/Runtime/RenderCore/Private/RenderUtils.cpp
@@ -1117,11 +1117,15 @@ RENDERCORE_API bool IsMobileDeferredShadingEnabled(const FStaticShaderPlatform P
 
 RENDERCORE_API bool MobileRequiresSceneDepthAux(const FStaticShaderPlatform Platform)
 {
-	// SceneDepth is used on all mobile platforms when the forward shading is enabled and on IOS with deferred shading.
-	if (!IsMobileDeferredShadingEnabled(Platform) || IsMetalMobilePlatform(Platform))
+	// SceneDepth is used on most mobile platforms when forward shading is enabled and always on IOS.
+	if (IsMetalMobilePlatform(Platform))
 	{
 		return true;
 	}
+	else if (!IsMobileDeferredShadingEnabled(Platform))
+	{
+		return IsAndroidOpenGLESPlatform(Platform) || IsVulkanMobilePlatform(Platform) || IsSimulatedPlatform(Platform);
+	}
 	return false;
 }
 
diff --git a/Engine/Source/Runtime/RenderCore/Private/ResolveShader.cpp b/Engine/Source/Runtime/RenderCore/Private/ResolveShader.cpp
--- a/Engine/Source/Runtime/RenderCore/Private/ResolveShader.cpp
+++ b/Engine/Source/Runtime/RenderCore/Private/ResolveShader.cpp
@@ -8,7 +8,6 @@ IMPLEMENT_SHADER_TYPE(, FResolveDepthPS, TEXT("/Engine/Private/ResolvePixelShade
 IMPLEMENT_SHADER_TYPE(, FResolveDepth2XPS, TEXT("/Engine/Private/ResolvePixelShader.usf"), TEXT("MainDepth"), SF_Pixel);
 IMPLEMENT_SHADER_TYPE(, FResolveDepth4XPS, TEXT("/Engine/Private/ResolvePixelShader.usf"), TEXT("MainDepth"), SF_Pixel);
 IMPLEMENT_SHADER_TYPE(, FResolveDepth8XPS, TEXT("/Engine/Private/ResolvePixelShader.usf"), TEXT("MainDepth"), SF_Pixel);
-IMPLEMENT_SHADER_TYPE(, FResolveDepthNonMSPS, TEXT("/Engine/Private/ResolvePixelShader.usf"), TEXT("MainDepthNonMS"), SF_Pixel);
 IMPLEMENT_SHADER_TYPE(, FResolveSingleSamplePS, TEXT("/Engine/Private/ResolvePixelShader.usf"), TEXT("MainSingleSample"), SF_Pixel);
 IMPLEMENT_SHADER_TYPE(, FResolveVS, TEXT("/Engine/Private/ResolveVertexShader.usf"), TEXT("Main"), SF_Vertex);
 
diff --git a/Engine/Source/Runtime/RenderCore/Public/ResolveShader.h b/Engine/Source/Runtime/RenderCore/Public/ResolveShader.h
--- a/Engine/Source/Runtime/RenderCore/Public/ResolveShader.h
+++ b/Engine/Source/Runtime/RenderCore/Public/ResolveShader.h
@@ -16,17 +16,10 @@ class FResolveDepthPS : public FGlobalShader
 public:
 	
 	typedef FDummyResolveParameter FParameter;
-	
-	static bool ShouldCompilePermutation(const FGlobalShaderPermutationParameters& Parameters) 
-	{ 
-		return GetMaxSupportedFeatureLevel(Parameters.Platform) >= ERHIFeatureLevel::SM5
-			|| IsSimulatedPlatform(Parameters.Platform); // support resolving MSAA depth in mobile emulation
-	}
 
 	static void ModifyCompilationEnvironment(const FGlobalShaderPermutationParameters& Parameters, FShaderCompilerEnvironment& OutEnvironment)
 	{
 		FGlobalShader::ModifyCompilationEnvironment(Parameters, OutEnvironment);
-		OutEnvironment.SetDefine(TEXT("SIMULATED_PLATFORM"), IsSimulatedPlatform(Parameters.Platform) ? 1 : 0);
 	}
 		
 	FResolveDepthPS(const ShaderMetaType::CompiledShaderInitializerType& Initializer):
@@ -50,17 +43,10 @@ public:
 
 	typedef FDummyResolveParameter FParameter;
 
-	static bool ShouldCompilePermutation(const FGlobalShaderPermutationParameters& Parameters) 
-	{ 
-		return GetMaxSupportedFeatureLevel(Parameters.Platform) >= ERHIFeatureLevel::SM5 
-			|| IsSimulatedPlatform(Parameters.Platform); // support resolving MSAA depth in mobile emulation
-	}
-
 	static void ModifyCompilationEnvironment(const FGlobalShaderPermutationParameters& Parameters, FShaderCompilerEnvironment& OutEnvironment)
 	{
 		FGlobalShader::ModifyCompilationEnvironment(Parameters, OutEnvironment);
 		OutEnvironment.SetDefine(TEXT("DEPTH_RESOLVE_NUM_SAMPLES"), 2);
-		OutEnvironment.SetDefine(TEXT("SIMULATED_PLATFORM"), IsSimulatedPlatform(Parameters.Platform) ? 1 : 0);
 	}
 
 	FResolveDepth2XPS(const ShaderMetaType::CompiledShaderInitializerType& Initializer) :
@@ -85,17 +71,10 @@ public:
 
 	typedef FDummyResolveParameter FParameter;
 
-	static bool ShouldCompilePermutation(const FGlobalShaderPermutationParameters& Parameters) 
-	{ 
-		return GetMaxSupportedFeatureLevel(Parameters.Platform) >= ERHIFeatureLevel::SM5
-			|| IsSimulatedPlatform(Parameters.Platform); // support resolving MSAA depth in mobile emulation
-	}
-
 	static void ModifyCompilationEnvironment(const FGlobalShaderPermutationParameters& Parameters, FShaderCompilerEnvironment& OutEnvironment)
 	{
 		FGlobalShader::ModifyCompilationEnvironment(Parameters, OutEnvironment);
 		OutEnvironment.SetDefine(TEXT("DEPTH_RESOLVE_NUM_SAMPLES"), 4);
-		OutEnvironment.SetDefine(TEXT("SIMULATED_PLATFORM"), IsSimulatedPlatform(Parameters.Platform) ? 1 : 0);
 	}
 
 	FResolveDepth4XPS(const ShaderMetaType::CompiledShaderInitializerType& Initializer) :
@@ -122,17 +101,10 @@ public:
 
 	static bool ShouldCache(EShaderPlatform Platform) { return GetMaxSupportedFeatureLevel(Platform) >= ERHIFeatureLevel::SM5; }
 
-	static bool ShouldCompilePermutation(const FGlobalShaderPermutationParameters& Parameters)
-	{
-		return GetMaxSupportedFeatureLevel(Parameters.Platform) >= ERHIFeatureLevel::SM5
-			|| IsSimulatedPlatform(Parameters.Platform); // support resolving MSAA depth in mobile emulation
-	}
-
 	static void ModifyCompilationEnvironment(const FGlobalShaderPermutationParameters& Parameters, FShaderCompilerEnvironment& OutEnvironment)
 	{
 		FGlobalShader::ModifyCompilationEnvironment(Parameters, OutEnvironment);
 		OutEnvironment.SetDefine(TEXT("DEPTH_RESOLVE_NUM_SAMPLES"), 8);
-		OutEnvironment.SetDefine(TEXT("SIMULATED_PLATFORM"), IsSimulatedPlatform(Parameters.Platform) ? 1 : 0);
 	}
 
 	FResolveDepth8XPS(const ShaderMetaType::CompiledShaderInitializerType& Initializer) :
@@ -149,30 +121,6 @@ public:
 	LAYOUT_FIELD(FShaderResourceParameter, UnresolvedSurface);
 };
 
-
-class FResolveDepthNonMSPS : public FGlobalShader
-{
-	DECLARE_EXPORTED_SHADER_TYPE(FResolveDepthNonMSPS, Global, RENDERCORE_API);
-public:
-	
-	typedef FDummyResolveParameter FParameter;
-	
-	static bool ShouldCompilePermutation(const FGlobalShaderPermutationParameters& Parameters) { return GetMaxSupportedFeatureLevel(Parameters.Platform) <= ERHIFeatureLevel::ES3_1; }
-	
-	FResolveDepthNonMSPS(const ShaderMetaType::CompiledShaderInitializerType& Initializer):
-	FGlobalShader(Initializer)
-	{
-		UnresolvedSurface.Bind(Initializer.ParameterMap,TEXT("UnresolvedSurfaceNonMS"), SPF_Mandatory);
-	}
-	FResolveDepthNonMSPS() {}
-	
-	void SetParameters(FRHICommandList& RHICmdList, FParameter)
-	{
-	}
-
-	LAYOUT_FIELD(FShaderResourceParameter, UnresolvedSurface);
-};
-
 class FResolveSingleSamplePS : public FGlobalShader
 {
 	DECLARE_EXPORTED_SHADER_TYPE(FResolveSingleSamplePS, Global, RENDERCORE_API);
diff --git a/Engine/Source/Runtime/Renderer/Private/MobileShadingRenderer.cpp b/Engine/Source/Runtime/Renderer/Private/MobileShadingRenderer.cpp
--- a/Engine/Source/Runtime/Renderer/Private/MobileShadingRenderer.cpp
+++ b/Engine/Source/Runtime/Renderer/Private/MobileShadingRenderer.cpp
@@ -456,8 +456,8 @@ void FMobileSceneRenderer::InitViews(FRDGBuilder& GraphBuilder, FSceneTexturesCo
 		bShouldRenderVelocities ||
 		bRequireSeparateViewPass ||
 		bIsFullDepthPrepassEnabled;
-	// never keep MSAA depth
-	bKeepDepthContent = (NumMSAASamples > 1 ? false : bKeepDepthContent);
+	// never keep MSAA depth if SceneDepthAux is enabled
+	bKeepDepthContent = ((NumMSAASamples > 1) && MobileRequiresSceneDepthAux(ShaderPlatform)) ? false : bKeepDepthContent;
 
 	// Update the bKeepDepthContent based on the mobile renderer status.
 	SceneTexturesConfig.bKeepDepthContent = bKeepDepthContent;
@@ -823,6 +823,11 @@ void FMobileSceneRenderer::Render(FRDGBuilder& GraphBuilder)
 	{
 		RenderFullDepthPrepass(GraphBuilder, SceneTextures);
 
+		if (!MobileRequiresSceneDepthAux(ShaderPlatform))
+		{
+			AddResolveSceneDepthPass(GraphBuilder, Views, SceneTextures.Depth);
+		}
+
 		SceneTextures.MobileSetupMode = EMobileSceneTextureSetupMode::SceneDepth;
 		SceneTextures.MobileUniformBuffer = CreateMobileSceneTextureUniformBuffer(GraphBuilder, SceneTextures.MobileSetupMode);
 
@@ -986,7 +991,10 @@ void FMobileSceneRenderer::RenderForward(FRDGBuilder& GraphBuilder, FRDGTextureR
 
 	FRenderTargetBindingSlots BasePassRenderTargets;
 	BasePassRenderTargets[0] = FRenderTargetBinding(SceneColor, SceneColorResolve, ERenderTargetLoadAction::EClear);
-	BasePassRenderTargets[1] = FRenderTargetBinding(SceneTextures.DepthAux.Target, SceneTextures.DepthAux.Resolve, ERenderTargetLoadAction::EClear);
+	if (MobileRequiresSceneDepthAux(ShaderPlatform))
+	{
+		BasePassRenderTargets[1] = FRenderTargetBinding(SceneTextures.DepthAux.Target, SceneTextures.DepthAux.Resolve, ERenderTargetLoadAction::EClear);
+	}
 	BasePassRenderTargets.DepthStencil = bIsFullDepthPrepassEnabled ? FDepthStencilBinding(SceneDepth, ERenderTargetLoadAction::ELoad, FExclusiveDepthStencil::DepthRead_StencilWrite) : FDepthStencilBinding(SceneDepth, ERenderTargetLoadAction::EClear, FExclusiveDepthStencil::DepthWrite_StencilWrite);
 	BasePassRenderTargets.ShadingRateTexture = (!MainView.bIsSceneCapture && !MainView.bIsReflectionCapture && ShadingRateTarget.IsValid()) ? RegisterExternalTexture(GraphBuilder, ShadingRateTarget->GetRenderTargetItem().ShaderResourceTexture, TEXT("ShadingRateTexture")) : nullptr;
 	BasePassRenderTargets.SubpassHint = ESubpassHint::DepthReadSubpass;
@@ -1142,7 +1150,10 @@ void FMobileSceneRenderer::RenderForwardMultiPass(FRDGBuilder& GraphBuilder, FMo
 	}
 
 	BasePassRenderTargets[0].SetLoadAction(ERenderTargetLoadAction::ELoad);
-	BasePassRenderTargets[1].SetLoadAction(ERenderTargetLoadAction::ELoad);
+	if (MobileRequiresSceneDepthAux(ShaderPlatform))
+	{
+		BasePassRenderTargets[1].SetLoadAction(ERenderTargetLoadAction::ELoad);
+	}
 	BasePassRenderTargets.DepthStencil.SetDepthLoadAction(ERenderTargetLoadAction::ELoad);
 	BasePassRenderTargets.DepthStencil.SetStencilLoadAction(ERenderTargetLoadAction::ELoad);
 	BasePassRenderTargets.DepthStencil.SetDepthStencilAccess(ExclusiveDepthStencil);
