Supports black texture array on mobile.

diff --git a/Engine/Source/Runtime/RenderCore/Private/RenderUtils.cpp b/Engine/Source/Runtime/RenderCore/Private/RenderUtils.cpp
--- a/Engine/Source/Runtime/RenderCore/Private/RenderUtils.cpp
+++ b/Engine/Source/Runtime/RenderCore/Private/RenderUtils.cpp
@@ -328,18 +328,15 @@ public:
 	// FResource interface.
 	virtual void InitRHI() override
 	{
-		if (GetFeatureLevel() >= ERHIFeatureLevel::SM5)
-		{
-			// Create the texture RHI.
-			FBlackVolumeTextureResourceBulkDataInterface BlackTextureBulkData(0);
-			FRHIResourceCreateInfo CreateInfo(TEXT("BlackArrayTexture"), &BlackTextureBulkData);
-			FTexture2DArrayRHIRef TextureArray = RHICreateTexture2DArray(1, 1, 1, PF_B8G8R8A8, 1, 1, TexCreate_ShaderResource, CreateInfo);
-			TextureRHI = TextureArray;
+		// Create the texture RHI.
+		FBlackVolumeTextureResourceBulkDataInterface BlackTextureBulkData(0);
+		FRHIResourceCreateInfo CreateInfo(TEXT("BlackArrayTexture"), &BlackTextureBulkData);
+		FTexture2DArrayRHIRef TextureArray = RHICreateTexture2DArray(1, 1, 1, PF_B8G8R8A8, 1, 1, TexCreate_ShaderResource, CreateInfo);
+		TextureRHI = TextureArray;
 
-			// Create the sampler state RHI resource.
-			FSamplerStateInitializerRHI SamplerStateInitializer(SF_Point,AM_Wrap,AM_Wrap,AM_Wrap);
-			SamplerStateRHI = GetOrCreateSamplerState(SamplerStateInitializer);
-		}
+		// Create the sampler state RHI resource.
+		FSamplerStateInitializerRHI SamplerStateInitializer(SF_Point, AM_Wrap, AM_Wrap, AM_Wrap);
+		SamplerStateRHI = GetOrCreateSamplerState(SamplerStateInitializer);
 	}
 
 	/** Returns the width of the texture in pixels. */
