Fix CVars not being picked up from other sections (UE-135567).

- Some projects put rendering CVars in other sections than RendererSettings.

diff --git a/Engine/Source/Runtime/Engine/Private/ShaderCompiler/ShaderCompiler.cpp b/Engine/Source/Runtime/Engine/Private/ShaderCompiler/ShaderCompiler.cpp
--- a/Engine/Source/Runtime/Engine/Private/ShaderCompiler/ShaderCompiler.cpp
+++ b/Engine/Source/Runtime/Engine/Private/ShaderCompiler/ShaderCompiler.cpp
@@ -5809,7 +5809,7 @@ void GlobalBeginCompileShader(
 
 		if (bIsMobilePlatform)
 		{
-			static FShaderPlatformCachedIniValue<int32> MobilePropagateAlphaIniValue(TEXT("/Script/Engine.RendererSettings"), TEXT("r.Mobile.PropagateAlpha"));
+			static FShaderPlatformCachedIniValue<int32> MobilePropagateAlphaIniValue(TEXT("r.Mobile.PropagateAlpha"));
 			int MobilePropagateAlphaIniValueInt = MobilePropagateAlphaIniValue.Get((EShaderPlatform)ShaderPlatform);
 			PropagateAlpha = MobilePropagateAlphaIniValueInt > 0 ? 2 : 0;
 		}
@@ -5835,8 +5835,8 @@ void GlobalBeginCompileShader(
 
 	if (IsMobilePlatform((EShaderPlatform)Target.Platform))
 	{
-		static FShaderPlatformCachedIniValue<bool> MobileEnableMovableSpotlightsIniValue(TEXT("/Script/Engine.RendererSettings"), TEXT("r.Mobile.EnableMovableSpotlights"));
-		static FShaderPlatformCachedIniValue<bool> MobileEnableMovableSpotlightsShadowIniValue(TEXT("/Script/Engine.RendererSettings"), TEXT("r.Mobile.EnableMovableSpotlightsShadow"));
+		static FShaderPlatformCachedIniValue<bool> MobileEnableMovableSpotlightsIniValue(TEXT("r.Mobile.EnableMovableSpotlights"));
+		static FShaderPlatformCachedIniValue<bool> MobileEnableMovableSpotlightsShadowIniValue(TEXT("r.Mobile.EnableMovableSpotlightsShadow"));
 
 		bool bMobileEnableMovableSpotlights = (MobileEnableMovableSpotlightsIniValue.Get((EShaderPlatform)Target.Platform) != 0);
 		Input.Environment.SetDefine(TEXT("PROJECT_MOBILE_ENABLE_MOVABLE_SPOTLIGHTS"), bMobileEnableMovableSpotlights ? 1 : 0);
diff --git a/Engine/Source/Runtime/RenderCore/Private/RenderUtils.cpp b/Engine/Source/Runtime/RenderCore/Private/RenderUtils.cpp
--- a/Engine/Source/Runtime/RenderCore/Private/RenderUtils.cpp
+++ b/Engine/Source/Runtime/RenderCore/Private/RenderUtils.cpp
@@ -1202,7 +1202,7 @@ RENDERCORE_API bool SupportsGen4TAA(const FStaticShaderPlatform Platform)
 {
 	if (IsMobilePlatform(Platform))
 	{
-		static FShaderPlatformCachedIniValue<bool> MobileSupportsGen4TAAIniValue(TEXT("/Script/Engine.RendererSettings"), TEXT("r.Mobile.SupportsGen4TAA"));
+		static FShaderPlatformCachedIniValue<bool> MobileSupportsGen4TAAIniValue(TEXT("r.Mobile.SupportsGen4TAA"));
 		return (MobileSupportsGen4TAAIniValue.Get(Platform) != 0);
 	}
 
@@ -1226,14 +1226,32 @@ Type FShaderPlatformCachedIniValue<Type>::Get(EShaderPlatform ShaderPlatform)
 		return *ExistingEntry;
 	}
 
+	bool bTestedIni = TestedIni.Contains(ShaderPlatform);
 	if (!bTestedIni)
 	{
-		bTestedIni = true;
+		TestedIni.Add(ShaderPlatform);
 		FConfigCacheIni* PlatformIni = FConfigCacheIni::ForPlatform(ShaderPlatformToPlatformName(ShaderPlatform));
-		if (PlatformIni != nullptr && PlatformIni->GetValue(Section, Key, Value, TEXT("Engine")))
+
+		if (PlatformIni != nullptr)
 		{
-			CachedValues.Add(ShaderPlatform, Value);
-			return Value;
+			// Rendering CVars can be set from a number of sections - see UDeviceProfileManager::ExpandDeviceProfileCVars for their ordering.
+			// (Here they are in reverse order because we will use the first found setting).
+			const TCHAR* ConsoleVariablesSections[] =
+			{
+				TEXT("ConsoleVariables"),
+				TEXT("SystemSettings"),
+				TEXT("/Script/Engine.RendererOverrideSettings"),
+				TEXT("/Script/Engine.RendererSettings")
+			};
+
+			for (const TCHAR* SectionThatCanSetAVar : ConsoleVariablesSections)
+			{
+				if (PlatformIni->GetValue(SectionThatCanSetAVar, Key, Value, TEXT("Engine")))
+				{
+					CachedValues.Add(ShaderPlatform, Value);
+					return Value;
+				}
+			}
 		}
 	}
 #endif
@@ -1751,6 +1769,6 @@ RENDERCORE_API bool IsUsingDBuffers(const FStaticShaderPlatform Platform)
 
 RENDERCORE_API bool AreSkinCacheShadersEnabled(EShaderPlatform Platform)
 {
-	static FShaderPlatformCachedIniValue<bool> PerPlatformCVar(TEXT("/Script/Engine.RendererSettings"), TEXT("r.SkinCache.CompileShaders"));
+	static FShaderPlatformCachedIniValue<bool> PerPlatformCVar(TEXT("r.SkinCache.CompileShaders"));
 	return (PerPlatformCVar.Get(Platform) != 0);
 }
diff --git a/Engine/Source/Runtime/RenderCore/Private/Shader.cpp b/Engine/Source/Runtime/RenderCore/Private/Shader.cpp
--- a/Engine/Source/Runtime/RenderCore/Private/Shader.cpp
+++ b/Engine/Source/Runtime/RenderCore/Private/Shader.cpp
@@ -1537,9 +1537,9 @@ void ShaderMapAppendKeyString(EShaderPlatform Platform, FString& KeyString)
 		KeyString += IsUsingEmulatedUniformBuffers(Platform) ? TEXT("_NoUB") : TEXT("");
 
 		{
-			static FShaderPlatformCachedIniValue<bool> MobileEnableMovableSpotlightsIniValue(TEXT("/Script/Engine.RendererSettings"), TEXT("r.Mobile.EnableMovableSpotlights"));
-			static FShaderPlatformCachedIniValue<bool> MobileEnableMovableSpotlightsShadowIniValue(TEXT("/Script/Engine.RendererSettings"), TEXT("r.Mobile.EnableMovableSpotlightsShadow"));
-			static FShaderPlatformCachedIniValue<int32> MobileNumDynamicPointLightsIniValue(TEXT("/Script/Engine.RendererSettings"), TEXT("r.MobileNumDynamicPointLights"));
+			static FShaderPlatformCachedIniValue<bool> MobileEnableMovableSpotlightsIniValue(TEXT("r.Mobile.EnableMovableSpotlights"));
+			static FShaderPlatformCachedIniValue<bool> MobileEnableMovableSpotlightsShadowIniValue(TEXT("r.Mobile.EnableMovableSpotlightsShadow"));
+			static FShaderPlatformCachedIniValue<int32> MobileNumDynamicPointLightsIniValue(TEXT("r.MobileNumDynamicPointLights"));
 
 			bool bMobileEnableMovableSpotlights = (MobileEnableMovableSpotlightsIniValue.Get(Platform) != 0);
 			KeyString += (bMobileEnableMovableSpotlights) ? TEXT("_MSPTL") : TEXT("");
@@ -1698,7 +1698,7 @@ void ShaderMapAppendKeyString(EShaderPlatform Platform, FString& KeyString)
 		int PropagateAlphaType = 0;
 		if (IsMobilePlatform(Platform))
 		{
-			static FShaderPlatformCachedIniValue<int32> MobilePropagateAlphaIniValue(TEXT("/Script/Engine.RendererSettings"), TEXT("r.Mobile.PropagateAlpha"));
+			static FShaderPlatformCachedIniValue<int32> MobilePropagateAlphaIniValue(TEXT("r.Mobile.PropagateAlpha"));
 			int MobilePropagateAlphaIniValueInt = MobilePropagateAlphaIniValue.Get((EShaderPlatform)Platform);
 			PropagateAlphaType = MobilePropagateAlphaIniValueInt > 0 ? 2 : 0;
 		}
diff --git a/Engine/Source/Runtime/RenderCore/Public/RenderUtils.h b/Engine/Source/Runtime/RenderCore/Public/RenderUtils.h
--- a/Engine/Source/Runtime/RenderCore/Public/RenderUtils.h
+++ b/Engine/Source/Runtime/RenderCore/Public/RenderUtils.h
@@ -520,21 +520,22 @@ RENDERCORE_API bool IsUsingDBuffers(const FStaticShaderPlatform Platform);
 template<typename Type>
 struct RENDERCORE_API FShaderPlatformCachedIniValue
 {
-	FShaderPlatformCachedIniValue(const TCHAR* InSection, const TCHAR* InKey)
-		: Section(InSection)
-		, Key(InKey)
-		, bTestedIni(false)
+	FShaderPlatformCachedIniValue(const TCHAR* InKey)
+		: Key(InKey)
 		, CVar(nullptr)
 	{}
 
 	Type Get(EShaderPlatform ShaderPlatform);
 
 private:
-	const TCHAR* Section;
 	const TCHAR* Key;
-	bool bTestedIni;
 	IConsoleVariable* CVar;
+#if WITH_EDITOR
+	/** Set that holds shader platforms for which the ini check was done. */
+	TSet<EShaderPlatform> TestedIni;
+	/** Cached values from the inis. */
 	TMap<EShaderPlatform, Type> CachedValues;
+#endif
 };
 
 /** Returns if ForwardShading is enabled. Only valid for the current platform (otherwise call ITargetPlatform::UsesForwardShading()). */
diff --git a/Engine/Source/Runtime/Renderer/Private/MobileBasePassRendering.cpp b/Engine/Source/Runtime/Renderer/Private/MobileBasePassRendering.cpp
--- a/Engine/Source/Runtime/Renderer/Private/MobileBasePassRendering.cpp
+++ b/Engine/Source/Runtime/Renderer/Private/MobileBasePassRendering.cpp
@@ -70,7 +70,7 @@ static TAutoConsoleVariable<int32> CVarMobileUseHWsRGBEncoding(
 
 namespace MobileBasePass
 {
-	FShaderPlatformCachedIniValue<int32> MobileNumDynamicPointLightsIniValue(TEXT("/Script/Engine.RendererSettings"), TEXT("r.MobileNumDynamicPointLights"));
+	FShaderPlatformCachedIniValue<int32> MobileNumDynamicPointLightsIniValue(TEXT("r.MobileNumDynamicPointLights"));
 };
 
 bool ShouldCacheShaderByPlatformAndOutputFormat(EShaderPlatform Platform, EOutputFormat OutputFormat)
diff --git a/Engine/Source/Runtime/Renderer/Private/PostProcess/TemporalAA.cpp b/Engine/Source/Runtime/Renderer/Private/PostProcess/TemporalAA.cpp
--- a/Engine/Source/Runtime/Renderer/Private/PostProcess/TemporalAA.cpp
+++ b/Engine/Source/Runtime/Renderer/Private/PostProcess/TemporalAA.cpp
@@ -238,7 +238,7 @@ class FTemporalAACS : public FGlobalShader
 		OutEnvironment.SetDefine(TEXT("THREADGROUP_SIZEX"), GTemporalAATileSizeX);
 		OutEnvironment.SetDefine(TEXT("THREADGROUP_SIZEY"), GTemporalAATileSizeY);
 
-		static FShaderPlatformCachedIniValue<bool> UseMobileConfig(TEXT("/Script/Engine.RendererSettings"), TEXT("r.TemporalAA.UseMobileConfig"));
+		static FShaderPlatformCachedIniValue<bool> UseMobileConfig(TEXT("r.TemporalAA.UseMobileConfig"));
 		bool bUseMobileConfig = (UseMobileConfig.Get((EShaderPlatform)Parameters.Platform) != 0);
 
 		bool bIsMobileTiledGPU = RHIHasTiledGPU(Parameters.Platform) || IsSimulatedPlatform(Parameters.Platform);
diff --git a/Engine/Source/Runtime/Renderer/Private/ShadowDepthRendering.cpp b/Engine/Source/Runtime/Renderer/Private/ShadowDepthRendering.cpp
--- a/Engine/Source/Runtime/Renderer/Private/ShadowDepthRendering.cpp
+++ b/Engine/Source/Runtime/Renderer/Private/ShadowDepthRendering.cpp
@@ -278,7 +278,7 @@ static TAutoConsoleVariable<int32> CVarSupportPointLightWholeSceneShadows(
 static bool MobileUsesPerspectiveCorrectShadowPermutation(EShaderPlatform ShaderPlatform)
 {
 	// Required only for spotlight shadows on mobile
-	static FShaderPlatformCachedIniValue<bool> MobileEnableMovableSpotlightsShadowIniValue(TEXT("/Script/Engine.RendererSettings"), TEXT("r.Mobile.EnableMovableSpotlightsShadow"));
+	static FShaderPlatformCachedIniValue<bool> MobileEnableMovableSpotlightsShadowIniValue(TEXT("r.Mobile.EnableMovableSpotlightsShadow"));
 	const bool bMobileEnableMovableSpotlightsShadow = (MobileEnableMovableSpotlightsShadowIniValue.Get(ShaderPlatform) != 0);
 	return bMobileEnableMovableSpotlightsShadow;
 }
