Working around D3D11 hitches caused by unpaged memory returned from Map. The workaround is to have UnifiedBuffer create brand new resources instead of updating existing ones.

diff --git a/Engine/Source/Runtime/RHI/Private/RHI.cpp b/Engine/Source/Runtime/RHI/Private/RHI.cpp
--- a/Engine/Source/Runtime/RHI/Private/RHI.cpp
+++ b/Engine/Source/Runtime/RHI/Private/RHI.cpp
@@ -1460,6 +1460,8 @@ bool GRHISupportsMeshShadersTier1 = false;
 
 bool GRHISupportsShaderTimestamp = false;
 
+bool GRHISupportsEfficientUploadOnResourceCreation = false;
+
 FVertexElementTypeSupportInfo GVertexElementTypeSupport;
 
 RHI_API int32 volatile GCurrentTextureMemorySize = 0;
diff --git a/Engine/Source/Runtime/RHI/Public/RHI.h b/Engine/Source/Runtime/RHI/Public/RHI.h
--- a/Engine/Source/Runtime/RHI/Public/RHI.h
+++ b/Engine/Source/Runtime/RHI/Public/RHI.h
@@ -748,6 +748,8 @@ extern RHI_API bool GRHISupportsMeshShadersTier1;
 */
 extern RHI_API bool GRHISupportsShaderTimestamp;
 
+extern RHI_API bool GRHISupportsEfficientUploadOnResourceCreation;
+
 /** Tables of all MSAA sample offset for all MSAA supported. Use GetMSAASampleOffsets() to read it. */
 extern RHI_API FVector2f GRHIDefaultMSAASampleOffsets[1 + 2 + 4 + 8 + 16];
 
diff --git a/Engine/Source/Runtime/RenderCore/Private/UnifiedBuffer.cpp b/Engine/Source/Runtime/RenderCore/Private/UnifiedBuffer.cpp
--- a/Engine/Source/Runtime/RenderCore/Private/UnifiedBuffer.cpp
+++ b/Engine/Source/Runtime/RenderCore/Private/UnifiedBuffer.cpp
@@ -609,9 +609,6 @@ RENDERCORE_API bool ResizeResourceIfNeeded(FRDGBuilder& GraphBuilder, FRWBuffer&
 
 void FScatterUploadBuffer::Init( uint32 NumElements, uint32 InNumBytesPerElement, bool bInFloat4Buffer, const TCHAR* DebugName )
 {
-	check( ScatterData == nullptr );
-	check( UploadData == nullptr );
-
 	NumScatters = 0;
 	MaxScatters = NumElements;
 	NumBytesPerElement = InNumBytesPerElement;
@@ -624,46 +621,117 @@ void FScatterUploadBuffer::Init( uint32 NumElements, uint32 InNumBytesPerElement
 	uint32 ScatterBufferSize = (uint32)FMath::Min( (uint64)FMath::RoundUpToPowerOfTwo( ScatterBytes ), GetMaxBufferDimension() * sizeof( uint32 ) );
 	check( ScatterBufferSize >= ScatterBytes );
 
-	if( ScatterBytes > ScatterBuffer.NumBytes || ScatterBufferSize < ScatterBuffer.NumBytes / 2)
-	{
-		// Resize Scatter Buffer
-		ScatterBuffer.Release();
-		ScatterBuffer.NumBytes = ScatterBufferSize;
-
-		FRHIResourceCreateInfo CreateInfo(DebugName);
-		ScatterBuffer.Buffer = RHICreateStructuredBuffer( sizeof( uint32 ), ScatterBuffer.NumBytes, BUF_ShaderResource | BUF_Volatile | Usage, CreateInfo );
-		ScatterBuffer.SRV = RHICreateShaderResourceView( ScatterBuffer.Buffer );
-	}
-
 	uint32 UploadBytes = NumElements * NumBytesPerElement;
 	uint32 UploadBufferSize = (uint32)FMath::Min( (uint64)FMath::RoundUpToPowerOfTwo( UploadBytes ), GetMaxBufferDimension() * TypeSize );
 	check( UploadBufferSize >= UploadBytes );
 
-	if( UploadBytes > UploadBuffer.NumBytes || UploadBufferSize < UploadBuffer.NumBytes / 2)
+	if (bUploadViaCreate)
 	{
-		// Resize Upload Buffer
-		UploadBuffer.Release();
-		UploadBuffer.NumBytes = UploadBufferSize;
+		if (ScatterBytes > ScatterDataSize || ScatterBufferSize < ScatterDataSize / 2)
+		{
+			FMemory::Free(ScatterData);
+			ScatterData = (uint32*)FMemory::Malloc(ScatterBufferSize);
+			ScatterDataSize = ScatterBufferSize;
+		}
 
-		FRHIResourceCreateInfo CreateInfo(DebugName);
-		UploadBuffer.Buffer = RHICreateStructuredBuffer( TypeSize, UploadBuffer.NumBytes, BUF_ShaderResource | BUF_Volatile | Usage, CreateInfo );
-		UploadBuffer.SRV = RHICreateShaderResourceView( UploadBuffer.Buffer );
+		if (UploadBytes > UploadDataSize || UploadBufferSize < UploadDataSize / 2)
+		{
+			FMemory::Free(UploadData);
+			UploadData = (uint8*)FMemory::Malloc(UploadBufferSize);
+			UploadDataSize = UploadBufferSize;
+		}
 	}
+	else
+	{
+		check(ScatterData == nullptr);
+		check(UploadData == nullptr);
+
+		if (ScatterBytes > ScatterBuffer.NumBytes || ScatterBufferSize < ScatterBuffer.NumBytes / 2)
+		{
+			// Resize Scatter Buffer
+			ScatterBuffer.Release();
+			ScatterBuffer.NumBytes = ScatterBufferSize;
 
-	ScatterData = (uint32*)RHILockBuffer( ScatterBuffer.Buffer, 0, ScatterBytes, RLM_WriteOnly );
-	UploadData = (uint8*)RHILockBuffer( UploadBuffer.Buffer, 0, UploadBytes, RLM_WriteOnly );
+			FRHIResourceCreateInfo CreateInfo(DebugName);
+			ScatterBuffer.Buffer = RHICreateStructuredBuffer(sizeof(uint32), ScatterBuffer.NumBytes, BUF_ShaderResource | BUF_Volatile | Usage, CreateInfo);
+			ScatterBuffer.SRV = RHICreateShaderResourceView(ScatterBuffer.Buffer);
+		}
+
+		if (UploadBytes > UploadBuffer.NumBytes || UploadBufferSize < UploadBuffer.NumBytes / 2)
+		{
+			// Resize Upload Buffer
+			UploadBuffer.Release();
+			UploadBuffer.NumBytes = UploadBufferSize;
+
+			FRHIResourceCreateInfo CreateInfo(DebugName);
+			UploadBuffer.Buffer = RHICreateStructuredBuffer(TypeSize, UploadBuffer.NumBytes, BUF_ShaderResource | BUF_Volatile | Usage, CreateInfo);
+			UploadBuffer.SRV = RHICreateShaderResourceView(UploadBuffer.Buffer);
+		}
+
+		ScatterData = (uint32*)RHILockBuffer(ScatterBuffer.Buffer, 0, ScatterBytes, RLM_WriteOnly);
+		UploadData = (uint8*)RHILockBuffer(UploadBuffer.Buffer, 0, UploadBytes, RLM_WriteOnly);
+	}
 }
 
+// Helper type used to initialize the buffer data on creation
+struct FScatterUploadBufferResourceArray : public FResourceArrayInterface
+{
+	const void* const DataPtr;
+	const int32 DataSize;
+
+	FScatterUploadBufferResourceArray(void* InDataPtr, int32 InDataSize)
+		: DataPtr(InDataPtr)
+		, DataSize(InDataSize)
+	{
+	}
+
+	const void* GetResourceData() const override { return DataPtr; }
+	uint32 GetResourceDataSize() const override { return DataSize; }
+
+	// Not necessary for our purposes
+	void Discard() override { }
+	bool IsStatic() const override { return false; }
+	bool GetAllowCPUAccess() const override { return true; }
+	void SetAllowCPUAccess(bool bInNeedsCPUAccess) override { }
+};
+
 template<typename ResourceType>
 void FScatterUploadBuffer::ResourceUploadTo(FRHICommandList& RHICmdList, const ResourceType& DstBuffer, bool bFlush)
 {
 	TRACE_CPUPROFILER_EVENT_SCOPE(FScatterUploadBuffer::ResourceUploadTo);
-	
-	RHIUnlockBuffer(ScatterBuffer.Buffer);
-	RHIUnlockBuffer(UploadBuffer.Buffer);
 
-	ScatterData = nullptr;
-	UploadData = nullptr;
+	if (bUploadViaCreate)
+	{
+		ScatterBuffer.Release();
+		UploadBuffer.Release();
+
+		ScatterBuffer.NumBytes = ScatterDataSize;
+		UploadBuffer.NumBytes = UploadDataSize;
+
+		const uint32 TypeSize = bFloat4Buffer ? 16 : 4;
+		const EBufferUsageFlags Usage = bFloat4Buffer ? BUF_None : BUF_ByteAddressBuffer;
+
+		{
+			FScatterUploadBufferResourceArray ScatterResourceArray(ScatterData, ScatterDataSize);
+			FRHIResourceCreateInfo CreateInfo(TEXT("ScatterResourceArray"), &ScatterResourceArray);
+			ScatterBuffer.Buffer = RHICreateStructuredBuffer(sizeof(uint32), ScatterDataSize, BUF_ShaderResource | BUF_Volatile | Usage, CreateInfo);
+			ScatterBuffer.SRV = RHICreateShaderResourceView(ScatterBuffer.Buffer);
+		}
+		{
+			FScatterUploadBufferResourceArray UploadResourceArray(UploadData, UploadDataSize);
+			FRHIResourceCreateInfo CreateInfo(TEXT("ScatterUploadBuffer"), &UploadResourceArray);
+			UploadBuffer.Buffer = RHICreateStructuredBuffer(TypeSize, UploadDataSize, BUF_ShaderResource | BUF_Volatile | Usage, CreateInfo);
+			UploadBuffer.SRV = RHICreateShaderResourceView(UploadBuffer.Buffer);
+		}
+	}
+	else
+	{
+		RHIUnlockBuffer(ScatterBuffer.Buffer);
+		RHIUnlockBuffer(UploadBuffer.Buffer);
+
+		ScatterData = nullptr;
+		UploadData = nullptr;
+	}
 
 	if (NumScatters == 0)
 		return;
diff --git a/Engine/Source/Runtime/RenderCore/Public/UnifiedBuffer.h b/Engine/Source/Runtime/RenderCore/Public/UnifiedBuffer.h
--- a/Engine/Source/Runtime/RenderCore/Public/UnifiedBuffer.h
+++ b/Engine/Source/Runtime/RenderCore/Public/UnifiedBuffer.h
@@ -65,11 +65,19 @@ public:
 	uint32*	ScatterData	= nullptr;
 	uint8*	UploadData	= nullptr;
 
+	uint32	ScatterDataSize = 0;
+	uint32	UploadDataSize = 0;
 	uint32 	NumScatters = 0;
 	uint32 	MaxScatters = 0;
 	uint32	NumBytesPerElement = 0;
 
 	bool	bFloat4Buffer = false;
+	bool    bUploadViaCreate = false;
+
+	~FScatterUploadBuffer()
+	{
+		Release();
+	}
 
 	RENDERCORE_API void Init( uint32 NumElements, uint32 InNumBytesPerElement, bool bInFloat4Buffer, const TCHAR* DebugName );
 
@@ -88,15 +96,14 @@ public:
 		checkSlow( ScatterData != nullptr );
 		checkSlow( UploadData != nullptr );
 
+		uint32* ScatterWriteData = ScatterData + NumScatters;
+
 		for( uint32 i = 0; i < Num; i++ )
 		{
-			ScatterData[ i ] = Index + i;
+			ScatterWriteData[ i ] = Index + i;
 		}
 
-		void* Result = UploadData;
-
-		ScatterData += Num;
-		UploadData += Num * NumBytesPerElement;
+		void* Result = UploadData + NumScatters * NumBytesPerElement;
 		NumScatters += Num;
 		return Result;
 	}
@@ -105,10 +112,37 @@ public:
 	{
 		ScatterBuffer.Release();
 		UploadBuffer.Release();
+
+		if (bUploadViaCreate)
+		{
+			if (ScatterData)
+			{
+				FMemory::Free(ScatterData);
+				ScatterData = nullptr;
+			}
+			if (UploadData)
+			{
+				FMemory::Free(UploadData);
+				UploadData = nullptr;
+			}
+			ScatterDataSize = 0;
+			UploadDataSize = 0;
+		}
 	}
 
 	uint32 GetNumBytes() const
 	{
 		return ScatterBuffer.NumBytes + UploadBuffer.NumBytes;
 	}
+
+	void SetUploadViaCreate(bool bInUploadViaCreate)
+	{
+		if (bInUploadViaCreate != bUploadViaCreate)
+		{
+			// When switching the upload path, just free everything.
+			Release();
+
+			bUploadViaCreate = bInUploadViaCreate;
+		}
+	}
 };
\ No newline at end of file
diff --git a/Engine/Source/Runtime/Renderer/Private/GPUScene.cpp b/Engine/Source/Runtime/Renderer/Private/GPUScene.cpp
--- a/Engine/Source/Runtime/Renderer/Private/GPUScene.cpp
+++ b/Engine/Source/Runtime/Renderer/Private/GPUScene.cpp
@@ -90,6 +90,14 @@ static TAutoConsoleVariable<float> CVarGPUSceneDebugDrawRange(
 	ECVF_RenderThreadSafe
 );
 
+int32 GGPUSceneInstanceUploadViaCreate = 1;
+FAutoConsoleVariableRef CVarGPUSceneInstanceUploadViaCreate(
+	TEXT("r.GPUScene.InstanceUploadViaCreate"),
+	GGPUSceneInstanceUploadViaCreate,
+	TEXT("When uploading GPUScene InstanceData, upload via resource creation when the RHI supports it efficiently."),
+	ECVF_RenderThreadSafe
+);
+
 LLM_DECLARE_TAG_API(GPUScene, RENDERER_API);
 DECLARE_LLM_MEMORY_STAT(TEXT("GPUScene"), STAT_GPUSceneLLM, STATGROUP_LLMFULL);
 DECLARE_LLM_MEMORY_STAT(TEXT("GPUScene"), STAT_GPUSceneSummaryLLM, STATGROUP_LLM);
@@ -973,6 +981,10 @@ void FGPUScene::UploadGeneral(FRHICommandListImmediate& RHICmdList, FScene *Scen
 				}
 			};
 
+			const bool bShouldUploadViaCreate = GRHISupportsEfficientUploadOnResourceCreation && GGPUSceneInstanceUploadViaCreate != 0;
+			PrimitiveUploadBuffer.SetUploadViaCreate(bShouldUploadViaCreate);
+			InstanceSceneUploadBuffer.SetUploadViaCreate(bShouldUploadViaCreate);
+
 			FUAVTransitionStateScopeHelper PrimitiveDataTransitionHelper(RHICmdList, BufferState.PrimitiveBuffer.UAV, ERHIAccess::Unknown, ERHIAccess::SRVMask);
 
 			const int32 MaxPrimitivesUploads = GetMaxPrimitivesUpdate(NumPrimitiveDataUploads, FPrimitiveSceneShaderData::DataStrideInFloat4s);
diff --git a/Engine/Source/Runtime/Windows/D3D11RHI/Private/Windows/WindowsD3D11Device.cpp b/Engine/Source/Runtime/Windows/D3D11RHI/Private/Windows/WindowsD3D11Device.cpp
--- a/Engine/Source/Runtime/Windows/D3D11RHI/Private/Windows/WindowsD3D11Device.cpp
+++ b/Engine/Source/Runtime/Windows/D3D11RHI/Private/Windows/WindowsD3D11Device.cpp
@@ -2208,6 +2208,7 @@ void FD3D11DynamicRHI::InitD3DDevice()
 		GRHISupportsTextureStreaming = true;
 		GRHISupportsFirstInstance = true;
 		GRHINeedsExtraDeletionLatency = false;
+		GRHISupportsEfficientUploadOnResourceCreation = true;
 
 		// Command lists need the validation RHI context if enabled, so call the global scope version of RHIGetDefaultContext() and RHIGetDefaultAsyncComputeContext().
 		GRHICommandList.GetImmediateCommandList().SetContext(::RHIGetDefaultContext());
