Fix a static code analysis warning in ShaderCodeLibrary.cpp

- Exempted, the warning is bogus.

diff --git a/Engine/Source/Runtime/RenderCore/Private/ShaderCodeLibrary.cpp b/Engine/Source/Runtime/RenderCore/Private/ShaderCodeLibrary.cpp
--- a/Engine/Source/Runtime/RenderCore/Private/ShaderCodeLibrary.cpp
+++ b/Engine/Source/Runtime/RenderCore/Private/ShaderCodeLibrary.cpp
@@ -2968,6 +2968,7 @@ bool UE::ShaderLibrary::Private::FNamedShaderLibrary::OpenShaderCode(const FStri
 	if (LibraryInstance == nullptr)
 	{
 		UE_LOG(LogShaderLibrary, Verbose, TEXT("Cooked Context: No Shared Shader Library for: %s and native library not supported."), *Library);
+		// PVS reports "The function was exited without releasing the 'LibraryInstance' pointer. A memory leak is possible", which is totally bogus here
 		return false; //-V773
 	}
 
