Csv proviles: disable RenderTargetPool category on the server side by default.

diff --git a/Engine/Source/Runtime/RenderCore/Private/RenderTargetPool.cpp b/Engine/Source/Runtime/RenderCore/Private/RenderTargetPool.cpp
--- a/Engine/Source/Runtime/RenderCore/Private/RenderTargetPool.cpp
+++ b/Engine/Source/Runtime/RenderCore/Private/RenderTargetPool.cpp
@@ -14,7 +14,7 @@ TGlobalResource<FRenderTargetPool> GRenderTargetPool;
 
 DEFINE_LOG_CATEGORY_STATIC(LogRenderTargetPool, Warning, All);
 
-CSV_DEFINE_CATEGORY(RenderTargetPool, true);
+CSV_DEFINE_CATEGORY(RenderTargetPool, !UE_SERVER);
 
 
 TRefCountPtr<IPooledRenderTarget> CreateRenderTarget(FRHITexture* Texture, const TCHAR* Name)
