Fix 1.2MB of untagged memory when initializing the ResourceList array

diff --git a/Engine/Source/Runtime/RenderCore/Private/RenderResource.cpp b/Engine/Source/Runtime/RenderCore/Private/RenderResource.cpp
--- a/Engine/Source/Runtime/RenderCore/Private/RenderResource.cpp
+++ b/Engine/Source/Runtime/RenderCore/Private/RenderResource.cpp
@@ -90,6 +90,7 @@ void FRenderResource::InitResource()
 
 		if (PLATFORM_NEEDS_RHIRESOURCELIST || !GIsRHIInitialized)
 		{
+			LLM_SCOPE(ELLMTag::SceneRender);
 			TArray<FRenderResource*>& ResourceList = GetResourceList();
 			TArray<int32>& FreeIndicesList = GetFreeIndicesList();
 
