Fix Strata material buffer being cleared twice.

This ensure that the clearing happen only through the dedicated UAV pass.
This fixes classification issue for tile at the horizon.

diff --git a/Engine/Source/Runtime/RenderCore/Public/RenderGraphUtils.h b/Engine/Source/Runtime/RenderCore/Public/RenderGraphUtils.h
--- a/Engine/Source/Runtime/RenderCore/Public/RenderGraphUtils.h
+++ b/Engine/Source/Runtime/RenderCore/Public/RenderGraphUtils.h
@@ -73,15 +73,24 @@ struct FTextureRenderTargetBinding
 {
 	FRDGTextureRef Texture;
 	int16 ArraySlice;
+	ERenderTargetLoadAction OverrideLoadAction;
 
 	FTextureRenderTargetBinding()
 		: Texture(nullptr)
 		, ArraySlice(-1)
+		, OverrideLoadAction(ERenderTargetLoadAction::Num)
 	{}
 
-	FTextureRenderTargetBinding(FRDGTextureRef InTexture, int16 InArraySlice = -1)
+	FTextureRenderTargetBinding(FRDGTextureRef InTexture, ERenderTargetLoadAction InOverrideLoadAction)
+		: Texture(InTexture)
+		, ArraySlice(-1)
+		, OverrideLoadAction(InOverrideLoadAction)
+	{}
+
+	FTextureRenderTargetBinding(FRDGTextureRef InTexture, int16 InArraySlice = -1, ERenderTargetLoadAction InOverrideLoadAction = ERenderTargetLoadAction::Num)
 		: Texture(InTexture)
 		, ArraySlice(InArraySlice)
+		, OverrideLoadAction(InOverrideLoadAction)
 	{}
 };
 inline FRenderTargetBindingSlots GetRenderTargetBindings(ERenderTargetLoadAction ColorLoadAction, TArrayView<FTextureRenderTargetBinding> ColorTextures)
@@ -93,6 +102,10 @@ inline FRenderTargetBindingSlots GetRenderTargetBindings(ERenderTargetLoadAction
 	{
 		check(ColorTextures[Index].Texture);
 		BindingSlots[Index] = FRenderTargetBinding(ColorTextures[Index].Texture, ColorLoadAction, 0, ColorTextures[Index].ArraySlice);
+		if (ColorTextures[Index].OverrideLoadAction != ERenderTargetLoadAction::Num)
+		{
+			BindingSlots[Index].SetLoadAction(ERenderTargetLoadAction::ELoad);
+		}
 	}
 	return BindingSlots;
 }
diff --git a/Engine/Source/Runtime/Renderer/Private/Strata/Strata.cpp b/Engine/Source/Runtime/Renderer/Private/Strata/Strata.cpp
--- a/Engine/Source/Runtime/Renderer/Private/Strata/Strata.cpp
+++ b/Engine/Source/Runtime/Renderer/Private/Strata/Strata.cpp
@@ -635,7 +635,7 @@ void ApprendStrataMRTs(FSceneRenderer& SceneRenderer, uint32& RenderTargetCount,
 		// Add 2 uint for Strata fast path
 		auto AddStrataOutputTarget = [&](int16 StrataMaterialArraySlice)
 		{
-			RenderTargets[RenderTargetCount] = FTextureRenderTargetBinding(SceneRenderer.Scene->StrataSceneData.MaterialTextureArrayMRTs, StrataMaterialArraySlice);
+			RenderTargets[RenderTargetCount] = FTextureRenderTargetBinding(SceneRenderer.Scene->StrataSceneData.MaterialTextureArrayMRTs, StrataMaterialArraySlice, ERenderTargetLoadAction::ELoad);
 			RenderTargetCount++;
 		};
 		AddStrataOutputTarget(0);
@@ -643,7 +643,7 @@ void ApprendStrataMRTs(FSceneRenderer& SceneRenderer, uint32& RenderTargetCount,
 
 		// Add another MRT for Strata top layer information
 		{
-			RenderTargets[RenderTargetCount] = FTextureRenderTargetBinding(SceneRenderer.Scene->StrataSceneData.TopLayerTexture);
+			RenderTargets[RenderTargetCount] = FTextureRenderTargetBinding(SceneRenderer.Scene->StrataSceneData.TopLayerTexture, ERenderTargetLoadAction::ELoad);
 			RenderTargetCount++;
 		};
 	}
