Fixed mistake from 18313006.

diff --git a/Engine/Source/Runtime/RenderCore/Public/RenderGraphUtils.h b/Engine/Source/Runtime/RenderCore/Public/RenderGraphUtils.h
--- a/Engine/Source/Runtime/RenderCore/Public/RenderGraphUtils.h
+++ b/Engine/Source/Runtime/RenderCore/Public/RenderGraphUtils.h
@@ -104,7 +104,7 @@ inline FRenderTargetBindingSlots GetRenderTargetBindings(ERenderTargetLoadAction
 		BindingSlots[Index] = FRenderTargetBinding(ColorTextures[Index].Texture, ColorLoadAction, 0, ColorTextures[Index].ArraySlice);
 		if (ColorTextures[Index].OverrideLoadAction != ERenderTargetLoadAction::Num)
 		{
-			BindingSlots[Index].SetLoadAction(ERenderTargetLoadAction::ELoad);
+			BindingSlots[Index].SetLoadAction(ColorTextures[Index].OverrideLoadAction);
 		}
 	}
 	return BindingSlots;
