Engine changes to enable RDG transient heap visualization mode.

diff --git a/Engine/Source/Runtime/D3D12RHI/Private/D3D12TransientResourceAllocator.cpp b/Engine/Source/Runtime/D3D12RHI/Private/D3D12TransientResourceAllocator.cpp
--- a/Engine/Source/Runtime/D3D12RHI/Private/D3D12TransientResourceAllocator.cpp
+++ b/Engine/Source/Runtime/D3D12RHI/Private/D3D12TransientResourceAllocator.cpp
@@ -221,7 +221,7 @@ void FD3D12TransientResourceAllocator::DeallocateMemory(FRHITransientBuffer* InB
 	Allocator.DeallocateMemory(InBuffer, InPassIndex);
 }
 
-void FD3D12TransientResourceAllocator::Freeze(FRHICommandListImmediate& RHICmdList)
+void FD3D12TransientResourceAllocator::Freeze(FRHICommandListImmediate& RHICmdList, FRHITransientHeapStats& OutHeapStats)
 {
-	Allocator.Freeze(RHICmdList);
+	Allocator.Freeze(RHICmdList, OutHeapStats);
 }
diff --git a/Engine/Source/Runtime/D3D12RHI/Private/D3D12TransientResourceAllocator.h b/Engine/Source/Runtime/D3D12RHI/Private/D3D12TransientResourceAllocator.h
--- a/Engine/Source/Runtime/D3D12RHI/Private/D3D12TransientResourceAllocator.h
+++ b/Engine/Source/Runtime/D3D12RHI/Private/D3D12TransientResourceAllocator.h
@@ -50,7 +50,7 @@ public:
 	FRHITransientBuffer* CreateBuffer(const FRHIBufferCreateInfo& InCreateInfo, const TCHAR* InDebugName, uint32 InPassIndex) override;
 	void DeallocateMemory(FRHITransientTexture* InTexture, uint32 InPassIndex) override;
 	void DeallocateMemory(FRHITransientBuffer* InBuffer, uint32 InPassIndex) override;
-	void Freeze(FRHICommandListImmediate& RHICmdList) override;
+	void Freeze(FRHICommandListImmediate& RHICmdList, FRHITransientHeapStats& OutHeapStats) override;
 
 private:
 
diff --git a/Engine/Source/Runtime/RHI/Private/RHI.cpp b/Engine/Source/Runtime/RHI/Private/RHI.cpp
--- a/Engine/Source/Runtime/RHI/Private/RHI.cpp
+++ b/Engine/Source/Runtime/RHI/Private/RHI.cpp
@@ -2726,9 +2726,9 @@ void FRHIBufferViewCache::SetDebugName(const TCHAR* DebugName)
 
 #endif
 
-void FRHITransientTexture::Init(const TCHAR* InName, uint32 InAllocationIndex, uint32 InPassIndex)
+void FRHITransientTexture::Init(const TCHAR* InName, uint32 InAllocationIndex, uint32 InPassIndex, const FRHITransientResourceStats& InStats)
 {
-	FRHITransientResource::Init(InName, InAllocationIndex, InPassIndex);
+	FRHITransientResource::Init(InName, InAllocationIndex, InPassIndex, InStats);
 	ViewCache.SetDebugName(InName);
 
 #if !(UE_BUILD_SHIPPING || UE_BUILD_TEST)
@@ -2736,9 +2736,9 @@ void FRHITransientTexture::Init(const TCHAR* InName, uint32 InAllocationIndex, u
 #endif
 }
 
-void FRHITransientBuffer::Init(const TCHAR* InName, uint32 InAllocationIndex, uint32 InPassIndex)
+void FRHITransientBuffer::Init(const TCHAR* InName, uint32 InAllocationIndex, uint32 InPassIndex, const FRHITransientResourceStats& InStats)
 {
-	FRHITransientResource::Init(InName, InAllocationIndex, InPassIndex);
+	FRHITransientResource::Init(InName, InAllocationIndex, InPassIndex, InStats);
 	ViewCache.SetDebugName(InName);
 
 	// TODO: Add method to rename a buffer.
diff --git a/Engine/Source/Runtime/RHI/Private/RHIValidation.cpp b/Engine/Source/Runtime/RHI/Private/RHIValidation.cpp
--- a/Engine/Source/Runtime/RHI/Private/RHIValidation.cpp
+++ b/Engine/Source/Runtime/RHI/Private/RHIValidation.cpp
@@ -1643,7 +1643,7 @@ void FValidationTransientResourceAllocator::Flush(FRHICommandListImmediate& RHIC
 	RHIAllocator->Flush(RHICmdList);
 }
 
-void FValidationTransientResourceAllocator::Freeze(FRHICommandListImmediate& RHICmdList)
+void FValidationTransientResourceAllocator::Freeze(FRHICommandListImmediate& RHICmdList, FRHITransientHeapStats& OutHeapStats)
 {
 	check(!bFrozen);
 
@@ -1685,7 +1685,7 @@ void FValidationTransientResourceAllocator::Freeze(FRHICommandListImmediate& RHI
 
 	bFrozen = true;
 
-	RHIAllocator->Freeze(RHICmdList);
+	RHIAllocator->Freeze(RHICmdList, OutHeapStats);
 }
 
 void FValidationTransientResourceAllocator::Release(FRHICommandListImmediate& RHICmdList)
diff --git a/Engine/Source/Runtime/RHI/Public/RHITransientResourceAllocator.h b/Engine/Source/Runtime/RHI/Public/RHITransientResourceAllocator.h
--- a/Engine/Source/Runtime/RHI/Public/RHITransientResourceAllocator.h
+++ b/Engine/Source/Runtime/RHI/Public/RHITransientResourceAllocator.h
@@ -13,6 +13,13 @@ enum class ERHITransientResourceType
 	Buffer
 };
 
+struct FRHITransientResourceStats
+{
+	uint64 HeapOffsetMin = 0;
+	uint64 HeapOffsetMax = 0;
+	uint32 HeapIndex = 0;
+};
+
 class FRHITransientResource
 {
 public:
@@ -25,13 +32,14 @@ public:
 	virtual ~FRHITransientResource() = default;
 
 	// (Internal) Initializes the transient resource with a new allocation / name.
-	virtual void Init(const TCHAR* InName, uint32 InAllocationIndex, uint32 InAcquirePassIndex)
+	virtual void Init(const TCHAR* InName, uint32 InAllocationIndex, uint32 InAcquirePassIndex, const FRHITransientResourceStats& InStats)
 	{
 		Name = InName;
 		AllocationIndex = InAllocationIndex;
 		AcquirePasses = TInterval<uint32>(0, InAcquirePassIndex);
 		DiscardPasses = TInterval<uint32>(0, TNumericLimits<uint32>::Max());
 		AliasingOverlaps.Reset();
+		Stats = InStats;
 	}
 
 	// (Internal) Assigns the discard pass index.
@@ -71,6 +79,8 @@ public:
 	FORCEINLINE bool IsTexture() const { return Type == ERHITransientResourceType::Texture; }
 	FORCEINLINE bool IsBuffer() const { return Type == ERHITransientResourceType::Buffer; }
 
+	FORCEINLINE const FRHITransientResourceStats& GetStats() const { return Stats; }
+
 private:
 	// Underlying RHI resource.
 	TRefCountPtr<FRHIResource> Resource;
@@ -91,6 +101,8 @@ private:
 	uint32 AllocationIndex = ~0u;
 	TInterval<uint32> AcquirePasses = TInterval<uint32>(0, 0);
 	TInterval<uint32> DiscardPasses = TInterval<uint32>(0, 0);
+
+	FRHITransientResourceStats Stats;
 };
 
 class RHI_API FRHITransientTexture final : public FRHITransientResource
@@ -101,7 +113,7 @@ public:
 		, CreateInfo(InCreateInfo)
 	{}
 
-	void Init(const TCHAR* InName, uint32 InAllocationIndex, uint32 InPassIndex) override;
+	void Init(const TCHAR* InName, uint32 InAllocationIndex, uint32 InPassIndex, const FRHITransientResourceStats& InStats) override;
 
 	// Returns the underlying RHI texture.
 	FRHITexture* GetRHI() const { return static_cast<FRHITexture*>(FRHITransientResource::GetRHI()); }
@@ -130,7 +142,7 @@ public:
 		, CreateInfo(InCreateInfo)
 	{}
 
-	void Init(const TCHAR* InName, uint32 InAllocationIndex, uint32 InPassIndex) override;
+	void Init(const TCHAR* InName, uint32 InAllocationIndex, uint32 InPassIndex, const FRHITransientResourceStats& InStats) override;
 
 	// Returns the underlying RHI buffer.
 	FORCEINLINE FRHIBuffer* GetRHI() const { return static_cast<FRHIBuffer*>(FRHITransientResource::GetRHI()); }
@@ -151,6 +163,20 @@ public:
 	FRHIBufferViewCache ViewCache;
 };
 
+struct FRHITransientHeapStats
+{
+	struct FHeap
+	{
+		// Maximum number of bytes used in the heap.
+		uint64 WatermarkSize{};
+
+		// Number of bytes allocated for use in the heap.
+		uint64 Capacity{};
+	};
+
+	TArray<FHeap, TInlineAllocator<4>> Heaps;
+};
+
 class IRHITransientResourceAllocator
 {
 public:
@@ -168,7 +194,7 @@ public:
 	virtual void Flush(FRHICommandListImmediate& RHICmdList) {}
 
 	// Freezes all allocations and validates that all resources have their memory deallocated.
-	virtual void Freeze(FRHICommandListImmediate& RHICmdList) = 0;
+	virtual void Freeze(FRHICommandListImmediate& RHICmdList, FRHITransientHeapStats& OutHeapStats) = 0;
 
 	// Releases the transient allocator and deletes the instance. Any FRHITransientResource* access after this call is not allowed.
 	virtual void Release(FRHICommandListImmediate& RHICmdList)
diff --git a/Engine/Source/Runtime/RHI/Public/RHIValidationTransientResourceAllocator.h b/Engine/Source/Runtime/RHI/Public/RHIValidationTransientResourceAllocator.h
--- a/Engine/Source/Runtime/RHI/Public/RHIValidationTransientResourceAllocator.h
+++ b/Engine/Source/Runtime/RHI/Public/RHIValidationTransientResourceAllocator.h
@@ -21,7 +21,7 @@ public:
 	virtual void DeallocateMemory(FRHITransientTexture* InTexture, uint32 InPassIndex) override final;
 	virtual void DeallocateMemory(FRHITransientBuffer* InBuffer, uint32 InPassIndex) override final;
 	virtual void Flush(FRHICommandListImmediate&) override final;
-	virtual void Freeze(FRHICommandListImmediate&) override final;
+	virtual void Freeze(FRHICommandListImmediate&, FRHITransientHeapStats&) override final;
 	virtual void Release(FRHICommandListImmediate&) override final;
 
 private:
diff --git a/Engine/Source/Runtime/RHICore/Private/RHICoreTransientResourceAllocator.cpp b/Engine/Source/Runtime/RHICore/Private/RHICoreTransientResourceAllocator.cpp
--- a/Engine/Source/Runtime/RHICore/Private/RHICoreTransientResourceAllocator.cpp
+++ b/Engine/Source/Runtime/RHICore/Private/RHICoreTransientResourceAllocator.cpp
@@ -612,8 +612,15 @@ void FRHITransientResourceAllocator::DeallocateMemoryInternal(FRHITransientResou
 	StatsToUpdate.CurrentSizeWithAliasing -= Allocation.Size;
 }
 
-void FRHITransientResourceAllocator::Freeze(FRHICommandListImmediate& RHICmdList)
+void FRHITransientResourceAllocator::Freeze(FRHICommandListImmediate& RHICmdList, FRHITransientHeapStats& OutHeapStats)
 {
+	for (FRHITransientHeapState& State : Heaps)
+	{
+		auto& Heap = OutHeapStats.Heaps.AddDefaulted_GetRef();
+		Heap.WatermarkSize = State.CommitSize;
+		Heap.Capacity = State.Heap->GetCapacity();
+	}
+
 	ParentSystem.ForfeitHeaps(Heaps);
 
 	RHICmdList.EnqueueLambda([&ParentSystem = ParentSystem, BufferStats = BufferStats, TextureStats = TextureStats](FRHICommandListImmediate&)
@@ -679,7 +686,12 @@ FRHITransientHeapAllocation FRHITransientResourceAllocator::Allocate(FMemoryStat
 
 void FRHITransientResourceAllocator::InitResource(FRHITransientResource* TransientResource, uint32 PassIndex, const FRHITransientHeapAllocation& Allocation, const TCHAR* Name)
 {
-	TransientResource->Init(Name, HeapAllocations.Emplace(Allocation), PassIndex);
+	FRHITransientResourceStats Stats;
+	Stats.HeapIndex = Allocation.HeapIndex;
+	Stats.HeapOffsetMin = Allocation.Offset;
+	Stats.HeapOffsetMax = Allocation.Offset + Allocation.Size;
+
+	TransientResource->Init(Name, HeapAllocations.Emplace(Allocation), PassIndex, Stats);
 
 	HeapAllocators[Allocation.HeapIndex].TrackOverlap(TransientResource, Allocation);
 }
\ No newline at end of file
diff --git a/Engine/Source/Runtime/RHICore/Public/RHICoreTransientResourceAllocator.h b/Engine/Source/Runtime/RHICore/Public/RHICoreTransientResourceAllocator.h
--- a/Engine/Source/Runtime/RHICore/Public/RHICoreTransientResourceAllocator.h
+++ b/Engine/Source/Runtime/RHICore/Public/RHICoreTransientResourceAllocator.h
@@ -602,7 +602,7 @@ public:
 	}
 
 	// Called to signify all allocations have completed. Forfeits all resources / heaps back to the parent system.
-	void Freeze(FRHICommandListImmediate& RHICmdList);
+	void Freeze(FRHICommandListImmediate& RHICmdList, FRHITransientHeapStats& OutHeapStats);
 
 	// Called to flush any active allocations prior to Freeze.
 	void Flush();
diff --git a/Engine/Source/Runtime/RenderCore/Private/RenderGraphBuilder.cpp b/Engine/Source/Runtime/RenderCore/Private/RenderGraphBuilder.cpp
--- a/Engine/Source/Runtime/RenderCore/Private/RenderGraphBuilder.cpp
+++ b/Engine/Source/Runtime/RenderCore/Private/RenderGraphBuilder.cpp
@@ -1589,7 +1589,9 @@ void FRDGBuilder::Execute(EExecuteMode ExecuteMode)
 
 				if (TransientResourceAllocator)
 				{
-					TransientResourceAllocator->Freeze(RHICmdList);
+					FRHITransientHeapStats TransientHeapStats;
+					TransientResourceAllocator->Freeze(RHICmdList, TransientHeapStats);
+					IF_RDG_ENABLE_TRACE(Trace.SetTransientHeapStats(TransientHeapStats));
 				}
 			}
 			else
diff --git a/Engine/Source/Runtime/RenderCore/Private/RenderGraphTrace.cpp b/Engine/Source/Runtime/RenderCore/Private/RenderGraphTrace.cpp
--- a/Engine/Source/Runtime/RenderCore/Private/RenderGraphTrace.cpp
+++ b/Engine/Source/Runtime/RenderCore/Private/RenderGraphTrace.cpp
@@ -14,6 +14,8 @@ UE_TRACE_EVENT_BEGIN(RDGTrace, GraphMessage)
 	UE_TRACE_EVENT_FIELD(uint64, StartCycles)
 	UE_TRACE_EVENT_FIELD(uint64, EndCycles)
 	UE_TRACE_EVENT_FIELD(uint16, PassCount)
+	UE_TRACE_EVENT_FIELD(uint64[], TransientHeapWatermarkSizes)
+	UE_TRACE_EVENT_FIELD(uint64[], TransientHeapCapacities)
 UE_TRACE_EVENT_END()
 
 UE_TRACE_EVENT_BEGIN(RDGTrace, GraphEndMessage)
@@ -50,6 +52,9 @@ UE_TRACE_EVENT_BEGIN(RDGTrace, BufferMessage)
 	UE_TRACE_EVENT_FIELD(uint16, NextOwnerHandle)
 	UE_TRACE_EVENT_FIELD(uint16, Order)
 	UE_TRACE_EVENT_FIELD(uint16[], Passes)
+	UE_TRACE_EVENT_FIELD(uint16, TransientHeapIndex)
+	UE_TRACE_EVENT_FIELD(uint64, TransientHeapOffsetMin)
+	UE_TRACE_EVENT_FIELD(uint64, TransientHeapOffsetMax)
 	UE_TRACE_EVENT_FIELD(bool, IsExternal)
 	UE_TRACE_EVENT_FIELD(bool, IsExtracted)
 	UE_TRACE_EVENT_FIELD(bool, IsCulled)
@@ -64,6 +69,9 @@ UE_TRACE_EVENT_BEGIN(RDGTrace, TextureMessage)
 	UE_TRACE_EVENT_FIELD(uint16, NextOwnerHandle)
 	UE_TRACE_EVENT_FIELD(uint16, Order)
 	UE_TRACE_EVENT_FIELD(uint16[], Passes)
+	UE_TRACE_EVENT_FIELD(uint16, TransientHeapIndex)
+	UE_TRACE_EVENT_FIELD(uint64, TransientHeapOffsetMin)
+	UE_TRACE_EVENT_FIELD(uint64, TransientHeapOffsetMax)
 	UE_TRACE_EVENT_FIELD(uint64, SizeInBytes)
 	UE_TRACE_EVENT_FIELD(uint64, CreateFlags)
 	UE_TRACE_EVENT_FIELD(uint32, Dimension)
@@ -124,11 +132,22 @@ void FRDGTrace::OutputGraphEnd(const FRDGBuilder& GraphBuilder)
 	{
 		const TCHAR* Name = GraphBuilder.BuilderName.GetTCHAR();
 
+		TArray<uint64, TInlineAllocator<8>> TransientHeapWatermarkSizes;
+		TArray<uint64, TInlineAllocator<8>> TransientHeapCapacities;
+
+		for (const auto& Heap : TransientHeapStats.Heaps)
+		{
+			TransientHeapWatermarkSizes.Emplace(Heap.WatermarkSize);
+			TransientHeapCapacities.Emplace(Heap.Capacity);
+		}
+
 		UE_TRACE_LOG(RDGTrace, GraphMessage, RDGChannel)
 			<< GraphMessage.Name(Name, uint16(FCString::Strlen(Name)))
 			<< GraphMessage.StartCycles(GraphStartCycles)
 			<< GraphMessage.EndCycles(FPlatformTime::Cycles64())
-			<< GraphMessage.PassCount(uint16(Passes.Num()));
+			<< GraphMessage.PassCount(uint16(Passes.Num()))
+			<< GraphMessage.TransientHeapWatermarkSizes(TransientHeapWatermarkSizes.GetData(), (uint16)TransientHeapWatermarkSizes.Num())
+			<< GraphMessage.TransientHeapCapacities(TransientHeapCapacities.GetData(), (uint16)TransientHeapCapacities.Num());
 	}
 
 	for (FRDGPassHandle Handle = Passes.Begin(); Handle != Passes.End(); ++Handle)
@@ -239,12 +258,22 @@ void FRDGTrace::OutputGraphEnd(const FRDGBuilder& GraphBuilder)
 			SizeInBytes = RHIComputeMemorySize(TextureRHI);
 		}
 
+		FRHITransientResourceStats TransientStats;
+
+		if (Texture->bTransient)
+		{
+			TransientStats = Texture->TransientTexture->GetStats();
+		}
+
 		UE_TRACE_LOG(RDGTrace, TextureMessage, RDGChannel)
 			<< TextureMessage.Name(Texture->Name, uint16(FCString::Strlen(Texture->Name)))
 			<< TextureMessage.Handle(Handle.GetIndex())
 			<< TextureMessage.NextOwnerHandle(Texture->NextOwner.GetIndexUnchecked())
 			<< TextureMessage.Order(Texture->TraceOrder)
 			<< TextureMessage.Passes((const uint16*)Texture->TracePasses.GetData(), (uint16)Texture->TracePasses.Num())
+			<< TextureMessage.TransientHeapIndex(TransientStats.HeapIndex)
+			<< TextureMessage.TransientHeapOffsetMin(TransientStats.HeapOffsetMin)
+			<< TextureMessage.TransientHeapOffsetMax(TransientStats.HeapOffsetMax)
 			<< TextureMessage.SizeInBytes(SizeInBytes)
 			<< TextureMessage.CreateFlags(uint32(Texture->Desc.Flags))
 			<< TextureMessage.Dimension(uint32(Texture->Desc.Dimension))
@@ -265,12 +294,22 @@ void FRDGTrace::OutputGraphEnd(const FRDGBuilder& GraphBuilder)
 	{
 		const FRDGBuffer* Buffer = Buffers[Handle];
 
+		FRHITransientResourceStats TransientStats;
+
+		if (Buffer->bTransient)
+		{
+			TransientStats = Buffer->TransientBuffer->GetStats();
+		}
+
 		UE_TRACE_LOG(RDGTrace, BufferMessage, RDGChannel)
 			<< BufferMessage.Name(Buffer->Name, uint16(FCString::Strlen(Buffer->Name)))
 			<< BufferMessage.Handle(Buffer->Handle.GetIndex())
 			<< BufferMessage.NextOwnerHandle(Buffer->NextOwner.GetIndexUnchecked())
 			<< BufferMessage.Order(Buffer->TraceOrder)
 			<< BufferMessage.Passes((const uint16*)Buffer->TracePasses.GetData(), (uint16)Buffer->TracePasses.Num())
+			<< BufferMessage.TransientHeapIndex(TransientStats.HeapIndex)
+			<< BufferMessage.TransientHeapOffsetMin(TransientStats.HeapOffsetMin)
+			<< BufferMessage.TransientHeapOffsetMax(TransientStats.HeapOffsetMax)
 			<< BufferMessage.UsageFlags(uint32(Buffer->Desc.Usage))
 			<< BufferMessage.BytesPerElement(Buffer->Desc.BytesPerElement)
 			<< BufferMessage.NumElements(Buffer->Desc.NumElements)
diff --git a/Engine/Source/Runtime/RenderCore/Public/RenderGraphTrace.h b/Engine/Source/Runtime/RenderCore/Public/RenderGraphTrace.h
--- a/Engine/Source/Runtime/RenderCore/Public/RenderGraphTrace.h
+++ b/Engine/Source/Runtime/RenderCore/Public/RenderGraphTrace.h
@@ -3,6 +3,7 @@
 #pragma once
 
 #include "RenderGraphDefinitions.h"
+#include "RHITransientResourceAllocator.h"
 
 #if RDG_ENABLE_TRACE
 
@@ -18,7 +19,10 @@ public:
 	void AddTexturePassDependency(FRDGTexture* Texture, FRDGPass* Pass);
 	void AddBufferPassDependency(FRDGBuffer* Buffer, FRDGPass* Pass);
 
+	void SetTransientHeapStats(const FRHITransientHeapStats& Stats) { TransientHeapStats = Stats; }
+
 private:
+	FRHITransientHeapStats TransientHeapStats;
 	uint64 GraphStartCycles{};
 	uint32 ResourceOrder{};
 };
