GPU Messaging API

diff --git a/Engine/Shaders/Private/GPUMessaging.usf b/Engine/Shaders/Private/GPUMessaging.usf
new file mode 100644
--- /dev/null
+++ b/Engine/Shaders/Private/GPUMessaging.usf
@@ -0,0 +1,10 @@
+// Copyright Epic Games, Inc. All Rights Reserved.
+
+#include "Common.ush"
+#include "GPUMessaging.ush"
+
+[numthreads(1, 1, 1)]
+void ClearMessageBuffer()
+{
+	GPUMessageDataBuffer[0] = 0;
+}
\ No newline at end of file
diff --git a/Engine/Shaders/Private/GPUMessaging.ush b/Engine/Shaders/Private/GPUMessaging.ush
new file mode 100644
--- /dev/null
+++ b/Engine/Shaders/Private/GPUMessaging.ush
@@ -0,0 +1,105 @@
+// Copyright Epic Games, Inc. All Rights Reserved.
+
+#pragma once
+
+// Buffer containing both headers, message payloads and the global offset in element 0.
+RWStructuredBuffer<uint> GPUMessageDataBuffer;
+uint GPUMessageDataBufferSize;
+
+struct FGPUMessageWriter
+{
+	bool bValid;
+	uint PayloadOffset;
+	uint PayloadNumUints;
+	uint CurrentWriteOffset;
+};
+
+FGPUMessageWriter GPUMessageBegin(uint ID, uint PayloadNumUints)
+{
+	uint MessageStartOffset = 0U;
+	InterlockedAdd(GPUMessageDataBuffer[0], 2U + PayloadNumUints, MessageStartOffset);
+
+	// Always offset 1 for the GPUMessageDataBuffer[0] is used for the counter
+	MessageStartOffset += 1U;
+
+	FGPUMessageWriter GPUMessageWriter;
+	GPUMessageWriter.PayloadOffset = MessageStartOffset + 2U;
+	GPUMessageWriter.PayloadNumUints = PayloadNumUints;
+	GPUMessageWriter.bValid = MessageStartOffset + 2U + PayloadNumUints <= GPUMessageDataBufferSize;
+	GPUMessageWriter.CurrentWriteOffset = 0U;
+	if (GPUMessageWriter.bValid)
+	{
+		GPUMessageDataBuffer[MessageStartOffset] = ID;
+		GPUMessageDataBuffer[MessageStartOffset + 1] = PayloadNumUints;
+	}
+	return GPUMessageWriter;	
+}
+
+void GPUMessageWriteItem(inout FGPUMessageWriter Mw, uint Value)
+{
+	if (Mw.bValid && Mw.CurrentWriteOffset < Mw.PayloadNumUints)
+	{
+		GPUMessageDataBuffer[Mw.PayloadOffset + Mw.CurrentWriteOffset] = Value;
+		Mw.CurrentWriteOffset += 1U;
+	}
+}
+
+void GPUMessageWriteItem(inout FGPUMessageWriter Mw, uint2 Value)
+{
+	GPUMessageWriteItem(Mw, Value.x);
+	GPUMessageWriteItem(Mw, Value.y);
+}
+
+void GPUMessageWriteItem(inout FGPUMessageWriter Mw, uint3 Value)
+{
+	GPUMessageWriteItem(Mw, Value.x);
+	GPUMessageWriteItem(Mw, Value.y);
+	GPUMessageWriteItem(Mw, Value.z);
+}
+
+void GPUMessageWriteItem(inout FGPUMessageWriter Mw, uint4 Value)
+{
+	GPUMessageWriteItem(Mw, Value.x);
+	GPUMessageWriteItem(Mw, Value.y);
+	GPUMessageWriteItem(Mw, Value.z);
+	GPUMessageWriteItem(Mw, Value.w);
+}
+
+void GPUMessageWriteItem(inout FGPUMessageWriter Mw, float Value)
+{
+	GPUMessageWriteItem(Mw, asuint(Value));
+}
+
+void GPUMessageWriteItem(inout FGPUMessageWriter Mw, float2 Value)
+{
+	GPUMessageWriteItem(Mw, Value.x);
+	GPUMessageWriteItem(Mw, Value.y);
+}
+
+void GPUMessageWriteItem(inout FGPUMessageWriter Mw, float3 Value)
+{
+	GPUMessageWriteItem(Mw, Value.x);
+	GPUMessageWriteItem(Mw, Value.y);
+	GPUMessageWriteItem(Mw, Value.z);
+}
+
+void GPUMessageWriteItem(inout FGPUMessageWriter Mw, float4 Value)
+{
+	GPUMessageWriteItem(Mw, Value.x);
+	GPUMessageWriteItem(Mw, Value.y);
+	GPUMessageWriteItem(Mw, Value.z);
+	GPUMessageWriteItem(Mw, Value.w);
+}
+
+void GPUMessagePost(uint ID, uint PayloadItem)
+{
+	FGPUMessageWriter Mw = GPUMessageBegin(ID, 1);
+	GPUMessageWriteItem(Mw, PayloadItem);
+}
+
+void GPUMessagePost(uint ID, uint PayloadItem0, uint PayloadItem1)
+{
+	FGPUMessageWriter Mw = GPUMessageBegin(ID, 2);
+	GPUMessageWriteItem(Mw, PayloadItem0);
+	GPUMessageWriteItem(Mw, PayloadItem1);
+}
\ No newline at end of file
diff --git a/Engine/Source/Runtime/RenderCore/Private/GPUMessaging.cpp b/Engine/Source/Runtime/RenderCore/Private/GPUMessaging.cpp
new file mode 100644
--- /dev/null
+++ b/Engine/Source/Runtime/RenderCore/Private/GPUMessaging.cpp
@@ -0,0 +1,282 @@
+// Copyright Epic Games, Inc. All Rights Reserved.
+
+#include "GPUMessaging.h"
+#include "RenderResource.h"
+#include "RenderGraphUtils.h"
+#include "RHIGPUReadback.h"
+#include "GlobalShader.h"
+
+DEFINE_LOG_CATEGORY_STATIC(LogGPUMessaging, Log, All);
+
+RDG_REGISTER_BLACKBOARD_STRUCT(GPUMessage::FParameters)
+
+class FClearMessageBufferCS : public FGlobalShader
+{
+public:
+	DECLARE_GLOBAL_SHADER(FClearMessageBufferCS);
+	SHADER_USE_PARAMETER_STRUCT(FClearMessageBufferCS, FGlobalShader)
+
+	BEGIN_SHADER_PARAMETER_STRUCT(FParameters, )
+		SHADER_PARAMETER_RDG_BUFFER_UAV(StructuredBuffer<uint>, GPUMessageDataBuffer)
+	END_SHADER_PARAMETER_STRUCT()
+};
+
+IMPLEMENT_GLOBAL_SHADER(FClearMessageBufferCS, "/Engine/Private/GPUMessaging.usf", "ClearMessageBuffer", SF_Compute);
+
+namespace GPUMessage
+{
+
+int32 GLogAllMessages = 0;
+static FAutoConsoleVariableRef CVarLogAllMessages(
+	TEXT("r.GPUMessage.LogAllMessages"),
+	GLogAllMessages,
+	TEXT("Log all messages to the console.\n")
+	TEXT("0: Disabled\n")
+	TEXT("1: Enabled\n"),
+	ECVF_RenderThreadSafe
+);
+
+int32 GMaxBufferSize = 64;
+static FAutoConsoleVariableRef CVarMaxBufferSize(
+	TEXT("r.GPUMessage.MaxBufferSize"),
+	GMaxBufferSize,
+	TEXT("Specifies the maximum size of the GPU message buffer, in KiB.\n")
+	TEXT("default: 64\n"),
+	ECVF_RenderThreadSafe | ECVF_ReadOnly
+);
+
+class FSystem : public FRenderResource
+{
+public:
+	FSocket RegisterHandler(const TSharedPtr<FHandler>& Handler)
+	{
+		checkf(!MessageBuffer, TEXT("Handlers must be registered outside of an active GPU message scope."));
+
+		const FMessageId MessageId(NextMessageId++);
+		check(!MessageHandlers.Contains(MessageId));
+		MessageHandlers.Add(MessageId, Handler);
+		return FSocket(MessageId);
+	}
+
+	void RemoveHandler(FMessageId MessageId)
+	{
+		checkf(!MessageBuffer, TEXT("Handlers must be unregistered outside of an active GPU message scope."));
+
+		check(MessageId.IsValid());
+		check(MessageHandlers.Contains(MessageId));
+		MessageHandlers.Remove(MessageId);
+	}
+
+	void BeginMessageScope(FRDGBuilder& GraphBuilder)
+	{
+		check(MessageBuffer == nullptr);
+		PollMessages();
+
+		if (MessageHandlers.IsEmpty())
+		{
+			return;
+		}
+
+		const uint32 MaxBufferSize = GetMaxBufferSize();
+
+		FRDGBufferDesc Desc = FRDGBufferDesc::CreateStructuredDesc(sizeof(uint32), MaxBufferSize);
+		Desc.Usage = EBufferUsageFlags(Desc.Usage | BUF_SourceCopy);
+		MessageBuffer = GraphBuilder.CreateBuffer(Desc, TEXT("GPUMessageManager.MessageBuffer"));
+
+		FParameters& Parameters = GraphBuilder.Blackboard.Create<FParameters>();
+		Parameters.GPUMessageDataBuffer = GraphBuilder.CreateUAV(MessageBuffer);
+		Parameters.GPUMessageDataBufferSize = MaxBufferSize;
+
+		TShaderMapRef<FClearMessageBufferCS> ComputeShader(GetGlobalShaderMap(GMaxRHIFeatureLevel));
+
+		FClearMessageBufferCS::FParameters* PassParameters = GraphBuilder.AllocParameters<FClearMessageBufferCS::FParameters>();
+		PassParameters->GPUMessageDataBuffer = Parameters.GPUMessageDataBuffer;
+
+		FComputeShaderUtils::AddPass(GraphBuilder, RDG_EVENT_NAME("ClearGPUMessageBuffer"), ComputeShader, PassParameters, FIntVector(1, 1, 1));
+	}
+
+	void EndMessageScope(FRDGBuilder& GraphBuilder)
+	{
+		if (MessageHandlers.IsEmpty())
+		{
+			check(MessageBuffer == nullptr);
+			return;
+		}
+
+		check(MessageBuffer != nullptr);
+
+		FRHIGPUBufferReadback* GPUBufferReadback = nullptr;
+		if (!MessageReadbackBuffersReady.IsEmpty())
+		{
+			GPUBufferReadback = MessageReadbackBuffersReady.Pop();
+		}
+		else
+		{
+			GPUBufferReadback = new FRHIGPUBufferReadback(TEXT("GPUMessageManager.Readback"));
+		}
+
+		AddEnqueueCopyPass(GraphBuilder, GPUBufferReadback, MessageBuffer, 0u);
+		MessageReadbackBuffersInUse.Add(GPUBufferReadback);
+
+		MessageBuffer = nullptr;
+	}
+
+private:
+	void PollMessages()
+	{
+		for (int32 Index = 0; Index < MessageReadbackBuffersInUse.Num();)
+		{
+			FRHIGPUBufferReadback* GPUBufferReadback = MessageReadbackBuffersInUse[Index];
+
+			if (GPUBufferReadback->IsReady())
+			{
+				const uint32 MaxBufferSize = GetMaxBufferSize();
+
+				const uint32* BufferPtr = (const uint32*)GPUBufferReadback->Lock(MaxBufferSize * sizeof(uint32));
+				uint32 ValidRangeEnd = BufferPtr[0] + 1U;
+
+				FString LogMessage;
+
+				const uint32 OriginalValidRangeEnd = ValidRangeEnd;
+				ValidRangeEnd = FMath::Min(ValidRangeEnd, MaxBufferSize);
+				ensureMsgf(ValidRangeEnd == OriginalValidRangeEnd, TEXT("GPU messages size %u exceeded maximum size %u. Results have been truncated."), OriginalValidRangeEnd, MaxBufferSize);
+
+				if (GLogAllMessages != 0)
+				{
+					if (ValidRangeEnd != OriginalValidRangeEnd)
+					{
+						LogMessage.Appendf(TEXT("GPU messages size %u exceeded maximum size %u. Results have been truncated!\n\n"), OriginalValidRangeEnd, MaxBufferSize);
+					}
+
+					LogMessage.Appendf(TEXT("\nMessage Buffer[0] = %d;\n"), BufferPtr[0]);
+
+					for (uint32 ReadOffset = 1U; ReadOffset < ValidRangeEnd; )
+					{
+						const FMessageId MessageId(BufferPtr[ReadOffset]);
+						uint32 PayloadNumUints = BufferPtr[ReadOffset + 1];
+
+						LogMessage.Appendf(TEXT("  Message[Offset: %d, ID: %d, PayloadSize: %d, "), ReadOffset, MessageId, PayloadNumUints);
+
+						if (auto Handler = MessageHandlers.Find(MessageId))
+						{
+							LogMessage.Appendf(TEXT("Name: %s]"), (*Handler)->GetName());
+						}
+						else
+						{
+							LogMessage.Append(TEXT("<Unknown>]"));
+						}
+
+						if (PayloadNumUints > 0)
+						{
+							LogMessage.Append(TEXT(" Payload: { "));
+							for (uint32 PayloadIndex = 0; PayloadIndex < PayloadNumUints; ++PayloadIndex)
+							{
+								LogMessage.Appendf(TEXT("%d%s"), BufferPtr[ReadOffset + 2U + PayloadIndex], PayloadIndex + 1U == PayloadNumUints ? TEXT("") : TEXT(", "));
+							}
+							LogMessage.Append(TEXT(" }"));
+						}
+
+						LogMessage.Append(TEXT("\n"));
+
+						// Step past the message to the next
+						ReadOffset += PayloadNumUints + 2U;
+					}
+
+					UE_LOG(LogGPUMessaging, Log, TEXT("%s"), *LogMessage);
+				}
+
+				for (uint32 ReadOffset = 1U; ReadOffset < ValidRangeEnd; )
+				{
+					const FMessageId MessageId(BufferPtr[ReadOffset]);
+					uint32 PayloadNumUints = BufferPtr[ReadOffset + 1];
+
+					if (auto Handler = MessageHandlers.Find(MessageId))
+					{
+						(*Handler)->Execute(FReader(MessageId, PayloadNumUints, &BufferPtr[ReadOffset + 2U]));
+					}
+
+					// Step past the message to the next
+					ReadOffset += PayloadNumUints + 2U;
+				}
+
+				GPUBufferReadback->Unlock();
+
+				MessageReadbackBuffersInUse.RemoveAt(Index);
+				MessageReadbackBuffersReady.Push(GPUBufferReadback);
+			}
+			else
+			{
+				++Index;
+			}
+		}
+	}
+
+	void ReleaseDynamicRHI() override
+	{
+		check(MessageBuffer == nullptr);
+
+		for (FRHIGPUBufferReadback* ReadbackBuffer : MessageReadbackBuffersInUse)
+		{
+			delete ReadbackBuffer;
+		}
+		MessageReadbackBuffersInUse.Empty();
+
+		for (FRHIGPUBufferReadback* ReadbackBuffer : MessageReadbackBuffersReady)
+		{
+			delete ReadbackBuffer;
+		}
+		MessageReadbackBuffersReady.Empty();
+	}
+
+	uint32 GetMaxBufferSize() const { return GMaxBufferSize * 1024; }
+
+	TMap<FMessageId, TSharedPtr<FHandler>> MessageHandlers;
+
+	TArray<FRHIGPUBufferReadback*, TInlineAllocator<8>> MessageReadbackBuffersInUse;
+	TArray<FRHIGPUBufferReadback*, TInlineAllocator<8>> MessageReadbackBuffersReady;
+	FRDGBuffer* MessageBuffer{};
+
+	uint32 NextMessageId = 0;
+};
+
+static TGlobalResource<FSystem> GSystem;
+
+void FSocket::Reset()
+{
+	if (MessageId.IsValid())
+	{
+		GSystem.RemoveHandler(MessageId);
+	}
+
+	MessageId = FMessageId::Null;
+}
+
+bool FScope::bRecursionCheck = false;
+
+FScope::FScope(FRDGBuilder& InGraphBuilder)
+	: GraphBuilder(InGraphBuilder)
+{
+	checkf(!bRecursionCheck, TEXT("GPUMessage::FScope (GPU_MESSAGE_SCOPE) has already been pushed. Only one scope can be active at a time."));
+	bRecursionCheck = true;
+	GSystem.BeginMessageScope(InGraphBuilder);
+}
+
+FScope::~FScope()
+{
+	GSystem.EndMessageScope(GraphBuilder);
+	bRecursionCheck = false;
+}
+
+FSocket RegisterHandler(const TSharedPtr<FHandler>& Handler)
+{
+	return GSystem.RegisterHandler(Handler);
+}
+
+FParameters GetShaderParameters(FRDGBuilder& GraphBuilder)
+{
+	const FParameters* Parameters = GraphBuilder.Blackboard.Get<FParameters>();
+	checkf(Parameters, TEXT("Shader parameters not initialized. GPUMessage::GetShaderParameters may only be called within a GPU_MESSAGE_SCOPE."));
+	return *Parameters;
+}
+
+} // GPUMessage
\ No newline at end of file
diff --git a/Engine/Source/Runtime/RenderCore/Public/GPUMessaging.h b/Engine/Source/Runtime/RenderCore/Public/GPUMessaging.h
new file mode 100644
--- /dev/null
+++ b/Engine/Source/Runtime/RenderCore/Public/GPUMessaging.h
@@ -0,0 +1,161 @@
+// Copyright Epic Games, Inc. All Rights Reserved.
+
+#pragma once
+
+#include "RenderGraphBuilder.h"
+
+#define GPU_MESSAGE_SCOPE(GraphBuilder) GPUMessage::FScope GPUMessageScope_(GraphBuilder);
+
+namespace GPUMessage
+{
+
+class FSystem;
+class FSocket;
+
+/** RAII object for controlling a new GPU message scope. */
+class RENDERCORE_API FScope
+{
+public:
+	FScope(FRDGBuilder& GraphBuilder);
+	~FScope();
+
+private:
+	FRDGBuilder& GraphBuilder;
+
+	static bool bRecursionCheck;
+};
+
+using FMessageId = TRDGHandle<FSocket, uint32>;
+
+/** A socket object mapped to a persistently registered handler. Supports only move construction. Will unregister the handler on destruction. */
+class RENDERCORE_API FSocket
+{
+public:
+	FSocket() = default;
+
+	FSocket(FSocket&& Other)
+	{
+		*this = MoveTemp(Other);
+	}
+
+	FSocket& operator=(FSocket&& Other)
+	{
+		Reset();
+		MessageId = Other.MessageId;
+		Other.MessageId = FMessageId::Null;
+		return *this;
+	}
+
+	~FSocket()
+	{
+		Reset();
+	}
+
+	const FMessageId GetMessageId() const { return MessageId; }
+
+	void Reset();
+
+private:
+	FSocket(FMessageId InMessageId)
+		: MessageId(InMessageId)
+	{}
+
+	FMessageId MessageId;
+
+	friend FSystem;
+};
+
+/** Interface for the handelr to read data from the GPU buffer during readback. */
+class FReader
+{
+public:
+	const FMessageId MessageId;
+
+	template <typename DataType>
+	DataType Read(const DataType& DefaultValue = {})
+	{
+		if (ensure(ReadOffset < PayloadSize))
+		{
+			return PayloadData[ReadOffset++];
+		}
+		return DefaultValue;
+	}
+
+private:
+	FReader(FMessageId InMessageId, uint32 InPayloadSize, const uint32* InPayloadData)
+		: MessageId(InMessageId)
+		, PayloadSize(InPayloadSize)
+		, PayloadData(InPayloadData)
+	{}
+
+	const uint32 PayloadSize;
+	const uint32* PayloadData;
+
+	uint32 ReadOffset = 0;
+
+	friend FSystem;
+};
+
+/** Base class persistent handler for receiving messages. */
+class FHandler
+{
+public:
+	FHandler(const TCHAR* InName)
+		: Name(InName)
+	{}
+
+	virtual ~FHandler() = default;
+
+	virtual void Execute(FReader Message) = 0;
+
+	const TCHAR* GetName() const { return Name; }
+
+private:
+	const TCHAR* Name;
+};
+
+/** Lambda persistent handler which calls the provided lambda when a message is received. */
+template <typename LambdaType>
+class TLambdaHandler : public FHandler
+{
+public:
+	TLambdaHandler(const TCHAR* InDebugId, LambdaType&& InLambda)
+		: FHandler(InDebugId)
+		, Lambda(MoveTemp(InLambda))
+	{}
+
+	virtual void Execute(FReader Message) override
+	{
+		Lambda(Message);
+	}
+
+private:
+	LambdaType Lambda;
+};
+
+///////////////////////////////////////////////////////////////////////////////
+
+/** Registers a persistent handler with the system. The FSocket object controls the lifetime of the handler. The handler
+ *  execute function is called when a message is received. Messages may be delayed by several frames.
+ */
+extern RENDERCORE_API FSocket RegisterHandler(const TSharedPtr<FHandler>& Handler);
+
+/** Registers a lambda function as a message handler. */
+template <typename LambdaType>
+FSocket RegisterHandler(const TCHAR* MessageDebugId, LambdaType&& Lambda)
+{
+	return RegisterHandler(TSharedPtr<FHandler>(new TLambdaHandler(MessageDebugId, MoveTemp(Lambda))));
+}
+
+/** Shader parameters used to write to write a GPU message. */
+BEGIN_SHADER_PARAMETER_STRUCT(FParameters, RENDERCORE_API)
+	SHADER_PARAMETER_RDG_BUFFER_UAV(StructuredBuffer<uint>, GPUMessageDataBuffer)
+	SHADER_PARAMETER(uint32, GPUMessageDataBufferSize)
+END_SHADER_PARAMETER_STRUCT()
+
+/** Returns the shader parameters for the active GPU_MESSAGE_SCOPE. Will assert if called outside of that scope. */
+extern RENDERCORE_API FParameters GetShaderParameters(FRDGBuilder& GraphBuilder);
+
+///////////////////////////////////////////////////////////////////////////////
+
+} // GPUMessage
\ No newline at end of file
diff --git a/Engine/Source/Runtime/Renderer/Private/DeferredShadingRenderer.cpp b/Engine/Source/Runtime/Renderer/Private/DeferredShadingRenderer.cpp
--- a/Engine/Source/Runtime/Renderer/Private/DeferredShadingRenderer.cpp
+++ b/Engine/Source/Runtime/Renderer/Private/DeferredShadingRenderer.cpp
@@ -69,6 +69,7 @@
 #include "SceneCaptureRendering.h"
 #include "NaniteSceneProxy.h"
 #include "RayTracing/RayTracingInstanceCulling.h"
+#include "GPUMessaging.h"
 
 extern int32 GNaniteShowStats;
 
@@ -1876,6 +1877,8 @@ void FDeferredShadingSceneRenderer::Render(FRDGBuilder& GraphBuilder)
 {
 	const bool bNaniteEnabled = IsNaniteEnabled();
 
+	GPU_MESSAGE_SCOPE(GraphBuilder);
+
 #if RHI_RAYTRACING
 	{
 		ERayTracingMeshCommandsMode CurrentMode = ViewFamily.EngineShowFlags.PathTracing ? ERayTracingMeshCommandsMode::PATH_TRACING : ERayTracingMeshCommandsMode::RAY_TRACING;
diff --git a/Engine/Source/Runtime/Renderer/Private/MobileShadingRenderer.cpp b/Engine/Source/Runtime/Renderer/Private/MobileShadingRenderer.cpp
--- a/Engine/Source/Runtime/Renderer/Private/MobileShadingRenderer.cpp
+++ b/Engine/Source/Runtime/Renderer/Private/MobileShadingRenderer.cpp
@@ -57,6 +57,7 @@
 #include "SceneOcclusion.h"
 #include "VariableRateShadingImageManager.h"
 #include "SceneTextureReductions.h"
+#include "GPUMessaging.h"
 
 uint32 GetShadowQuality();
 
@@ -692,6 +693,8 @@ void FMobileSceneRenderer::Render(FRDGBuilder& GraphBuilder)
 
 	Scene->UpdateAllPrimitiveSceneInfos(GraphBuilder);
 
+	GPU_MESSAGE_SCOPE(GraphBuilder);
+
 	// Establish scene primitive count (must be done after UpdateAllPrimitiveSceneInfos)
 	FGPUSceneScopeBeginEndHelper GPUSceneScopeBeginEndHelper(Scene->GPUScene, GPUSceneDynamicContext, Scene);
 
diff --git a/Engine/Source/Runtime/Renderer/Private/SceneHitProxyRendering.cpp b/Engine/Source/Runtime/Renderer/Private/SceneHitProxyRendering.cpp
--- a/Engine/Source/Runtime/Renderer/Private/SceneHitProxyRendering.cpp
+++ b/Engine/Source/Runtime/Renderer/Private/SceneHitProxyRendering.cpp
@@ -29,6 +29,7 @@
 #include "VT/VirtualTextureSystem.h"
 #include "SceneRenderingUtils.h"
 #include "InstanceCulling/InstanceCullingManager.h"
+#include "GPUMessaging.h"
 #include "HairStrands/HairStrandsData.h"
 
 class FHitProxyShaderElementData : public FMeshMaterialShaderElementData
@@ -551,6 +552,8 @@ void FMobileSceneRenderer::RenderHitProxies(FRDGBuilder& GraphBuilder)
 {
 	Scene->UpdateAllPrimitiveSceneInfos(GraphBuilder);
 
+	GPU_MESSAGE_SCOPE(GraphBuilder);
+
 	FGPUSceneScopeBeginEndHelper GPUSceneScopeBeginEndHelper(Scene->GPUScene, GPUSceneDynamicContext, Scene);
 
 	PrepareViewRectsForRendering(GraphBuilder.RHICmdList);
@@ -589,6 +592,8 @@ void FDeferredShadingSceneRenderer::RenderHitProxies(FRDGBuilder& GraphBuilder)
 
 	Scene->UpdateAllPrimitiveSceneInfos(GraphBuilder);
 
+	GPU_MESSAGE_SCOPE(GraphBuilder);
+
 	FGPUSceneScopeBeginEndHelper GPUSceneScopeBeginEndHelper(Scene->GPUScene, GPUSceneDynamicContext, Scene);
 
 	PrepareViewRectsForRendering(GraphBuilder.RHICmdList);
