Ensures FGlobalShaderConfigDefines is applied without requiring the shader type's ModifyCompilationEnvironement() to call its parent function.

diff --git a/Engine/Source/Runtime/RenderCore/Private/GlobalShader.cpp b/Engine/Source/Runtime/RenderCore/Private/GlobalShader.cpp
--- a/Engine/Source/Runtime/RenderCore/Private/GlobalShader.cpp
+++ b/Engine/Source/Runtime/RenderCore/Private/GlobalShader.cpp
@@ -551,11 +551,15 @@ FGlobalShader::FGlobalShader(const ShaderMetaType::CompiledShaderInitializerType
 :	FShader(Initializer)
 {}
 
-void FGlobalShader::ModifyCompilationEnvironment(const FGlobalShaderPermutationParameters& Parameters, FShaderCompilerEnvironment& Environment)
+void FGlobalShaderType::SetupCompileEnvironment(EShaderPlatform Platform, int32 PermutationId, EShaderPermutationFlags Flags, FShaderCompilerEnvironment& Environment) const
 {
-	FShader::ModifyCompilationEnvironment(Parameters, Environment);
+	FGlobalShaderPermutationParameters Parameters(GetFName(), Platform, PermutationId, Flags);
 
+	// Modify the compilation environment based on platform.
 	FGlobalShaderConfigDefines::ApplyConfigDefines(Parameters, Environment);
+
+	// Allow the shader type to modify its compile environment.
+	ModifyCompilationEnvironment(Parameters, Environment);
 }
 
 void BackupGlobalShaderMap(FGlobalShaderBackupData& OutGlobalShaderBackup)
diff --git a/Engine/Source/Runtime/RenderCore/Public/GlobalShader.h b/Engine/Source/Runtime/RenderCore/Public/GlobalShader.h
--- a/Engine/Source/Runtime/RenderCore/Public/GlobalShader.h
+++ b/Engine/Source/Runtime/RenderCore/Public/GlobalShader.h
@@ -59,7 +59,7 @@ struct FGlobalShaderPermutationParameters : public FShaderPermutationParameters
  * A shader meta type for the simplest shaders; shaders which are not material or vertex factory linked.
  * There should only a single instance of each simple shader type.
  */
-class FGlobalShaderType : public FShaderType
+class RENDERCORE_API FGlobalShaderType : public FShaderType
 {
 	friend class FGlobalShaderTypeCompiler;
 public:
@@ -111,13 +111,9 @@ public:
 	 * @param Platform - Platform to compile for.
 	 * @param Environment - The shader compile environment that the function modifies.
 	 */
-	void SetupCompileEnvironment(EShaderPlatform Platform, int32 PermutationId, EShaderPermutationFlags Flags, FShaderCompilerEnvironment& Environment) const
-	{
-		// Allow the shader type to modify its compile environment.
-		ModifyCompilationEnvironment(FGlobalShaderPermutationParameters(GetFName(), Platform, PermutationId, Flags), Environment);
-	}
+	void SetupCompileEnvironment(EShaderPlatform Platform, int32 PermutationId, EShaderPermutationFlags Flags, FShaderCompilerEnvironment& Environment) const;
 
-	static RENDERCORE_API bool ShouldCompilePipeline(const FShaderPipelineType* ShaderPipelineType, EShaderPlatform Platform, EShaderPermutationFlags Flags);
+	static bool ShouldCompilePipeline(const FShaderPipelineType* ShaderPipelineType, EShaderPlatform Platform, EShaderPermutationFlags Flags);
 };
 
 class RENDERCORE_API FGlobalShaderMapContent : public FShaderMapContent
@@ -278,7 +274,7 @@ public:
 		return FShader::ShouldCompilePermutation(Parameters);
 	}
 	
-	static RENDERCORE_API void ModifyCompilationEnvironment(const FGlobalShaderPermutationParameters& Parameters, FShaderCompilerEnvironment& Environment);
+	static inline void ModifyCompilationEnvironment(const FGlobalShaderPermutationParameters& Parameters, FShaderCompilerEnvironment& Environment) { };
 };
 
 /**
