Packages the GPUDumpViewer in non shipping cooked builds

diff --git a/Engine/Source/Runtime/RenderCore/Private/RenderGraphDumpResources.cpp b/Engine/Source/Runtime/RenderCore/Private/RenderGraphDumpResources.cpp
--- a/Engine/Source/Runtime/RenderCore/Private/RenderGraphDumpResources.cpp
+++ b/Engine/Source/Runtime/RenderCore/Private/RenderGraphDumpResources.cpp
@@ -329,9 +329,6 @@ void FRDGBuilder::BeginResourceDump(const TArray<FString>& Args)
 	}
 
 	// Copy the viewer
-	#if PLATFORM_DESKTOP
-	// TODO: Have the viewer available in cooked build.
-	if (!FPlatformProperties::RequiresCookedData())
 	{
 		const TCHAR* OpenGPUDumpViewerWindowsName = TEXT("OpenGPUDumpViewer.bat");
 		const TCHAR* ViewerHTML = TEXT("GPUDumpViewer.html");
@@ -340,7 +337,6 @@ void FRDGBuilder::BeginResourceDump(const TArray<FString>& Args)
 		PlatformFile.CopyFile(*(GRDGResourceDumpContext.DumpingDirectoryPath / ViewerHTML), *(DumpGPUViewerSourcePath / ViewerHTML));
 		PlatformFile.CopyFile(*(GRDGResourceDumpContext.DumpingDirectoryPath / OpenGPUDumpViewerWindowsName), *(DumpGPUViewerSourcePath / OpenGPUDumpViewerWindowsName));
 	}
-	#endif
 }
 
 void FRDGBuilder::EndResourceDump()
diff --git a/Engine/Source/Runtime/RenderCore/RenderCore.Build.cs b/Engine/Source/Runtime/RenderCore/RenderCore.Build.cs
--- a/Engine/Source/Runtime/RenderCore/RenderCore.Build.cs
+++ b/Engine/Source/Runtime/RenderCore/RenderCore.Build.cs
@@ -1,5 +1,7 @@
 // Copyright Epic Games, Inc. All Rights Reserved.
 
+using System.IO;
+using UnrealBuildBase;
 using UnrealBuildTool;
 
 public class RenderCore : ModuleRules
@@ -24,7 +26,13 @@ public class RenderCore : ModuleRules
 			PrivateDependencyModuleNames.Add("CoreUObject");
 		}
 
-        PrivateDependencyModuleNames.AddRange(new string[] { "Core", "Projects", "RHI", "ApplicationCore", "TraceLog", "CookOnTheFly" });
+		// Copy the GPUDumpViewer's source code for the r.DumpGPU command.
+		if (Target.Configuration != UnrealTargetConfiguration.Shipping)
+		{
+			RuntimeDependencies.Add(Path.Combine(Unreal.EngineDirectory.ToString(), "Extras/GPUDumpViewer/..."), StagedFileType.DebugNonUFS);
+		}
+
+		PrivateDependencyModuleNames.AddRange(new string[] { "Core", "Projects", "RHI", "ApplicationCore", "TraceLog", "CookOnTheFly" });
 
         PrivateIncludePathModuleNames.AddRange(new string[] { "DerivedDataCache" });
 		
