Fixes VisualizeTexture not listing all RDG resources

diff --git a/Engine/Source/Runtime/RenderCore/Private/RenderGraphBuilder.cpp b/Engine/Source/Runtime/RenderCore/Private/RenderGraphBuilder.cpp
--- a/Engine/Source/Runtime/RenderCore/Private/RenderGraphBuilder.cpp
+++ b/Engine/Source/Runtime/RenderCore/Private/RenderGraphBuilder.cpp
@@ -3365,12 +3365,6 @@ void FRDGBuilder::VisualizePassOutputs(const FRDGPass* Pass)
 		return;
 	}
 
-	// Don't scan any of the pass output if not texture is requested to be visualized.
-	if (GVisualizeTexture.Requested.Name.IsEmpty())
-	{
-		return;
-	}
-
 	bInDebugPassScope = true;
 
 
