Dumps previously never dumped input RDG resources with r.DumpGPU

diff --git a/Engine/Extras/GPUDumpViewer/GPUDumpViewer.html b/Engine/Extras/GPUDumpViewer/GPUDumpViewer.html
--- a/Engine/Extras/GPUDumpViewer/GPUDumpViewer.html
+++ b/Engine/Extras/GPUDumpViewer/GPUDumpViewer.html
@@ -193,7 +193,7 @@ function display_input_resource(pass_id, resource_name)
 
 function display_output_resource(pass_id, resource_name)
 {
-	var resource_version_name = resource_name;
+	var resource_version_name = `${resource_name}.v0000000000000000`;
 	if (pass_id >= 0)
 	{
 		var pass_data = g_passes[pass_id];
diff --git a/Engine/Extras/GPUDumpViewer/OpenGPUDumpViewer.bat b/Engine/Extras/GPUDumpViewer/OpenGPUDumpViewer.bat
--- a/Engine/Extras/GPUDumpViewer/OpenGPUDumpViewer.bat
+++ b/Engine/Extras/GPUDumpViewer/OpenGPUDumpViewer.bat
@@ -1,5 +1,18 @@
+
+@echo off
+
 set CWD=%cd%
+set CHROME_USER_DATA=%CWD%/.tmp_chrome_data/
+
+echo "Opening chrome..."
 
 REM --allow-file-access-from-files allow to load a file from a file:// webpage required for GPUDumpViewer.html to work.
 REM --user-data-dir is required to force chrome to open a new instance so that --allow-file-access-from-files is honored.
-"C:\Program Files\Google\Chrome\Application\chrome.exe" file://%CWD%/GPUDumpViewer.html --allow-file-access-from-files --new-window --incognito --user-data-dir=%CWD%/.tmp_chrome_data/
+"C:\Program Files\Google\Chrome\Application\chrome.exe" file://%CWD%/GPUDumpViewer.html --allow-file-access-from-files --new-window --incognito --user-data-dir=%CHROME_USER_DATA%
+
+echo "Closing chrome..."
+
+REM Wait for 2s to shut down so that CHROME_USER_DATA can be deleted completly
+timeout /t 2 /nobreak > NUL
+
+rmdir /S /Q "%CHROME_USER_DATA%"
diff --git a/Engine/Source/Runtime/RenderCore/Private/RenderGraphDumpResources.cpp b/Engine/Source/Runtime/RenderCore/Private/RenderGraphDumpResources.cpp
--- a/Engine/Source/Runtime/RenderCore/Private/RenderGraphDumpResources.cpp
+++ b/Engine/Source/Runtime/RenderCore/Private/RenderGraphDumpResources.cpp
@@ -34,6 +34,7 @@ struct FRDGResourceDumpContext
 	FDateTime Time;
 	int32 ResourcesDumped = 0;
 	int32 PassesCount = 0;
+	TMap<const FRDGResource*, const FRDGPass*> LastResourceVersion;
 
 	bool bShowInExplore = false;
 
@@ -48,7 +49,7 @@ struct FRDGResourceDumpContext
 		return FString::Printf(TEXT("%s.%p"), Resource->Name, Resource);
 	}
 
-	static FString GetResourceDumpName(const FRDGPass* Pass, const FRDGResource* Resource)
+	static FString GetResourceVersionDumpName(const FRDGPass* Pass, const FRDGResource* Resource)
 	{
 		return FString::Printf(TEXT("%s.%p.v%p.bin"), Resource->Name, Resource, Pass);
 	}
@@ -58,6 +59,43 @@ struct FRDGResourceDumpContext
 		return DumpingDirectoryPath + TEXT("Resources/");
 	}
 
+	void GetResourceDumpInfo(
+		const FRDGPass* Pass,
+		const FRDGResource* Resource,
+		bool bIsOutputResource,
+		bool* bOutDumpResourceInfos,
+		FString* OutResourceVersionDumpName)
+	{
+		*bOutDumpResourceInfos = false;
+
+		if (bIsOutputResource)
+		{
+			*OutResourceVersionDumpName = GetResourceVersionDumpName(Pass, Resource);
+		}
+
+		if (!LastResourceVersion.Contains(Resource))
+		{
+			// First time we ever see this resource, so dump it's info to disk
+			*bOutDumpResourceInfos = true;
+
+			// If not an output, it might be a resource undumped by r.DumpGPU.Root or external texture so still dump it as v0.
+			if (!bIsOutputResource)
+			{
+				*OutResourceVersionDumpName = GetResourceVersionDumpName(/* Pass = */ nullptr, Resource);
+			}
+
+			if (LastResourceVersion.Num() % 1024 == 0)
+			{
+				LastResourceVersion.Reserve(LastResourceVersion.Num() + 1024);
+			}
+			LastResourceVersion.Add(Resource, Pass);
+		}
+		else
+		{
+			LastResourceVersion[Resource] = Pass;
+		}
+	}
+
 	void AddDumpTexturePass(
 		FRDGBuilder& GraphBuilder,
 		TArray<TSharedPtr<FJsonValue>>& InputResourceNames,
@@ -100,9 +138,14 @@ struct FRDGResourceDumpContext
 			}
 		}
 
+		bool bDumpResourceInfos;
+		FString ResourceVersionDumpName;
+		GetResourceDumpInfo(Pass, Texture, bIsOutputResource, &bDumpResourceInfos, &ResourceVersionDumpName);
+
 		FString ResourceDumpDirectory = GetResourceDumpDirectory();
 
 		// Dump the information of the texture to json file.
+		if (bDumpResourceInfos)
 		{
 			FString PixelFormat = FString::Printf(TEXT("PF_%s"), GPixelFormats[Desc.Format].Name);
 
@@ -132,9 +175,9 @@ struct FRDGResourceDumpContext
 		}
 
 		// Dump the resource's binary to a .bin file.
-		if (bIsOutputResource)
+		if (!ResourceVersionDumpName.IsEmpty())
 		{
-			FString DumpFilePath = ResourceDumpDirectory + GetResourceDumpName(Pass, Texture);
+			FString DumpFilePath = ResourceDumpDirectory + ResourceVersionDumpName;
 
 			FDumpTexturePass* PassParameters = GraphBuilder.AllocParameters<FDumpTexturePass>();
 			PassParameters->Texture = Texture;
@@ -206,7 +249,12 @@ struct FRDGResourceDumpContext
 
 		FString ResourceDumpDirectory = GetResourceDumpDirectory();
 
+		bool bDumpResourceInfos;
+		FString ResourceVersionDumpName;
+		GetResourceDumpInfo(Pass, Buffer, bIsOutputResource, &bDumpResourceInfos, &ResourceVersionDumpName);
+
 		// Dump the information of the buffer to json file.
+		if (bDumpResourceInfos)
 		{
 			TSharedPtr<FJsonObject> JsonObject = MakeShareable(new FJsonObject);
 			JsonObject->SetStringField(TEXT("Name"), Buffer->Name);
@@ -229,9 +277,9 @@ struct FRDGResourceDumpContext
 		}
 
 		// Dump the resource's binary to a .bin file.
-		if (bIsOutputResource)
+		if (!ResourceVersionDumpName.IsEmpty())
 		{
-			FString DumpFilePath = ResourceDumpDirectory + GetResourceDumpName(Pass, Buffer);
+			FString DumpFilePath = ResourceDumpDirectory + ResourceVersionDumpName;
 
 			FDumpBufferPass* PassParameters = GraphBuilder.AllocParameters<FDumpBufferPass>();
 			PassParameters->Buffer = Buffer;
@@ -296,7 +344,7 @@ void FRDGBuilder::BeginResourceDump(const TArray<FString>& Args)
 	}
 
 	GRDGResourceDumpContext.Time = FDateTime::Now();
-	GRDGResourceDumpContext.DumpingDirectoryPath = FPaths::ProjectSavedDir() + TEXT("GPUDumps/") + FPlatformProperties::PlatformName() + TEXT("/") + FApp::GetProjectName() + TEXT("-") + GRDGResourceDumpContext.Time.ToString() + TEXT("/");
+	GRDGResourceDumpContext.DumpingDirectoryPath = FPaths::ProjectSavedDir() + TEXT("GPUDumps/") + FApp::GetProjectName() + TEXT("-") + FPlatformProperties::PlatformName() + TEXT("-") + GRDGResourceDumpContext.Time.ToString() + TEXT("/");
 
 	GRDGResourceDumpContext.bShowInExplore = Args.Contains(TEXT("-show"));
 
