Fix compile issue introduced in CL-17986013
Update FScatterUploadBuffer::PrimitiveDataStrideInFloat4s to match FPrimitiveSceneShaderData

diff --git a/Engine/Source/Runtime/RenderCore/Public/UnifiedBuffer.h b/Engine/Source/Runtime/RenderCore/Public/UnifiedBuffer.h
--- a/Engine/Source/Runtime/RenderCore/Public/UnifiedBuffer.h
+++ b/Engine/Source/Runtime/RenderCore/Public/UnifiedBuffer.h
@@ -57,7 +57,7 @@ RENDERCORE_API bool ResizeResourceIfNeeded(FRDGBuilder& GraphBuilder, FRWBuffer&
 class FScatterUploadBuffer
 {
 public:
-	enum { PrimitiveDataStrideInFloat4s = 38 };
+	enum { PrimitiveDataStrideInFloat4s = 39 };
 
 	FByteAddressBuffer ScatterBuffer;
 	FByteAddressBuffer UploadBuffer;
