Implements GRHIValidateBufferSourceCopy and use it in DumpGPU

diff --git a/Engine/Source/Runtime/RHI/Private/RHIValidation.cpp b/Engine/Source/Runtime/RHI/Private/RHIValidation.cpp
--- a/Engine/Source/Runtime/RHI/Private/RHIValidation.cpp
+++ b/Engine/Source/Runtime/RHI/Private/RHIValidation.cpp
@@ -14,6 +14,8 @@
 
 bool GRHIValidationEnabled = false;
 
+bool GRHIValidateBufferSourceCopy = true;
+
 // When set to 1, callstack for each uniform buffer allocation will be tracked 
 // (slow and leaks memory, but can be handy to find the location where an invalid
 // allocation has been made)
@@ -490,7 +492,10 @@ void FValidationContext::RHICopyToStagingBuffer(FRHIBuffer* SourceBufferRHI, FRH
 {
 	using namespace RHIValidation;
 	Tracker->Assert(SourceBufferRHI->GetWholeResourceIdentity(), ERHIAccess::CopySrc);
-	RHI_VALIDATION_CHECK(EnumHasAnyFlags(SourceBufferRHI->GetUsage(), BUF_SourceCopy), *GetReasonString_SourceCopyFlagMissing(SourceBufferRHI));
+	if (GRHIValidateBufferSourceCopy)
+	{
+		RHI_VALIDATION_CHECK(EnumHasAnyFlags(SourceBufferRHI->GetUsage(), BUF_SourceCopy), *GetReasonString_SourceCopyFlagMissing(SourceBufferRHI));
+	}
 	RHIContext->RHICopyToStagingBuffer(SourceBufferRHI, DestinationStagingBufferRHI, InOffset, InNumBytes);
 }
 
@@ -498,7 +503,10 @@ void FValidationComputeContext::RHICopyToStagingBuffer(FRHIBuffer* SourceBufferR
 {
 	using namespace RHIValidation;
 	Tracker->Assert(SourceBufferRHI->GetWholeResourceIdentity(), ERHIAccess::CopySrc);
-	RHI_VALIDATION_CHECK(EnumHasAnyFlags(SourceBufferRHI->GetUsage(), BUF_SourceCopy), *GetReasonString_SourceCopyFlagMissing(SourceBufferRHI));
+	if (GRHIValidateBufferSourceCopy)
+	{
+		RHI_VALIDATION_CHECK(EnumHasAnyFlags(SourceBufferRHI->GetUsage(), BUF_SourceCopy), *GetReasonString_SourceCopyFlagMissing(SourceBufferRHI));
+	}
 	RHIContext->RHICopyToStagingBuffer(SourceBufferRHI, DestinationStagingBufferRHI, InOffset, InNumBytes);
 }
 
diff --git a/Engine/Source/Runtime/RHI/Public/RHIValidation.h b/Engine/Source/Runtime/RHI/Public/RHIValidation.h
--- a/Engine/Source/Runtime/RHI/Public/RHIValidation.h
+++ b/Engine/Source/Runtime/RHI/Public/RHIValidation.h
@@ -15,6 +15,10 @@
 class FValidationComputeContext;
 class FValidationContext;
 
+// Controls whether BUF_SourceCopy should be validate or not.
+extern RHI_API bool GRHIValidateBufferSourceCopy;
+
+
 // This is a macro because we only want to evaluate the message expression if the checked expression is false.
 #define RHI_VALIDATION_CHECK(InExpression, InMessage) \
 	do \
diff --git a/Engine/Source/Runtime/RenderCore/Private/RenderGraphDumpResources.cpp b/Engine/Source/Runtime/RenderCore/Private/RenderGraphDumpResources.cpp
--- a/Engine/Source/Runtime/RenderCore/Private/RenderGraphDumpResources.cpp
+++ b/Engine/Source/Runtime/RenderCore/Private/RenderGraphDumpResources.cpp
@@ -15,6 +15,7 @@
 #include "GenericPlatform/GenericPlatformDriver.h"
 #include "RenderGraphUtils.h"
 #include "GlobalShader.h"
+#include "RHIValidation.h"
 
 #if RDG_DUMP_RESOURCES
 
@@ -1086,7 +1087,7 @@ struct FRDGResourceDumpContext
 				}
 				else
 				{
-					UE_LOG(LogRendererCore, Warning, TEXT("RHICmdList.MapStagingSurface() to dump buffer %s failed."), Buffer->Name);
+					UE_LOG(LogRendererCore, Warning, TEXT("RHICmdList.LockStagingBuffer() to dump buffer %s failed."), Buffer->Name);
 				}
 
 				StagingBuffer = nullptr;
@@ -1241,6 +1242,13 @@ FString FRDGBuilder::BeginResourceDump(const TArray<FString>& Args)
 	{
 		check(IsInRenderingThread());
 		GRDGResourceDumpContext = NewResourceDumpContext;
+
+		ImmediateRHICmdList.SubmitCommandsAndFlushGPU();
+
+		// Disable the validation for BUF_SourceCopy so that all buffers can be copied into staging buffer for CPU readback.
+		#if ENABLE_RHI_VALIDATION
+			GRHIValidateBufferSourceCopy = false;
+		#endif
 	});
 
 	bIsDumpingFrame_GameThread = true;
@@ -1264,6 +1272,16 @@ void FRDGBuilder::EndResourceDump()
 	// Wait all rendering commands are completed to finish with GRDGResourceDumpContext.
 	{
 		UE_LOG(LogRendererCore, Display, TEXT("Stalling game thread until render thread finishes to dump resources"));
+
+		ENQUEUE_RENDER_COMMAND(FEndGPUDump)(
+			[](FRHICommandListImmediate& ImmediateRHICmdList)
+		{
+			ImmediateRHICmdList.SubmitCommandsAndFlushGPU();
+			#if ENABLE_RHI_VALIDATION
+				GRHIValidateBufferSourceCopy = true;
+			#endif
+		});
+
 		FlushRenderingCommands();
 	}
 
