Implements r.DumpGPU.Test.PrettifyResourceFileNames to reduce risk of running into windows file path limits

diff --git a/Engine/Source/Runtime/RenderCore/Private/RenderGraphDumpResources.cpp b/Engine/Source/Runtime/RenderCore/Private/RenderGraphDumpResources.cpp
--- a/Engine/Source/Runtime/RenderCore/Private/RenderGraphDumpResources.cpp
+++ b/Engine/Source/Runtime/RenderCore/Private/RenderGraphDumpResources.cpp
@@ -69,6 +69,11 @@ static TAutoConsoleVariable<int32> GDumpTestEnableDiskWrite(
 	TEXT("Master switch whether any files should be written to disk, used for r.DumpGPU automation tests to not fill up workers' hard drive."),
 	ECVF_Default);
 
+static TAutoConsoleVariable<int32> GDumpTestPrettifyResourceFileNames(
+	TEXT("r.DumpGPU.Test.PrettifyResourceFileNames"), 0,
+	TEXT("Whether the resource file names should include resource name. May increase the likelyness of running into Windows' filepath limit."),
+	ECVF_RenderThreadSafe);
+
 static TAutoConsoleVariable<FString> GDumpGPUDirectoryCVar(
 	TEXT("r.DumpGPU.Directory"), TEXT(""),
 	TEXT("Directory to dump to."),
@@ -216,10 +221,14 @@ struct FRDGResourceDumpContext
 
 	static FString GetUniqueResourceName(const FRDGResource* Resource)
 	{
-		FString UniqueResourceName = FString::Printf(TEXT("%s.%016x"), Resource->Name, PtrToUint(Resource));
-		UniqueResourceName.ReplaceInline(TEXT("/"), TEXT(""));
-		UniqueResourceName.ReplaceInline(TEXT("\\"), TEXT(""));
-		return UniqueResourceName;
+		if (GDumpTestPrettifyResourceFileNames.GetValueOnRenderThread())
+		{
+			FString UniqueResourceName = FString::Printf(TEXT("%s.%016x"), Resource->Name, PtrToUint(Resource));
+			UniqueResourceName.ReplaceInline(TEXT("/"), TEXT(""));
+			UniqueResourceName.ReplaceInline(TEXT("\\"), TEXT(""));
+			return UniqueResourceName;
+		}
+		return PtrToString(Resource);
 	}
 
 	static FString GetUniqueSubResourceName(FRDGTextureSRVDesc SubResourceDesc)
