Fixes r.RDG.ClobberResources not clobbering RDGBuffer with the said value

diff --git a/Engine/Source/Runtime/RenderCore/Private/RenderGraphPrivate.cpp b/Engine/Source/Runtime/RenderCore/Private/RenderGraphPrivate.cpp
--- a/Engine/Source/Runtime/RenderCore/Private/RenderGraphPrivate.cpp
+++ b/Engine/Source/Runtime/RenderCore/Private/RenderGraphPrivate.cpp
@@ -196,26 +196,31 @@ bool IsDebugAllowedForResource(const TCHAR* ResourceName)
 	return IsDebugAllowed(GRDGDebugResourceFilterName, ResourceName);
 }
 
-FLinearColor GetClobberColor()
+static float GetClobberValue()
 {
 	switch (GRDGClobberResources)
 	{
 	case 1:
-		return FLinearColor(1000, 1000, 1000, 1000);
+		return 1000.0f;
 	case 2:
-		return FLinearColor(NAN, NAN, NAN, NAN);
+		return NAN;
 	case 3:
-		return FLinearColor(INFINITY, INFINITY, INFINITY, INFINITY);
-	case 4:
-		return FLinearColor(0, 0, 0, 0);
-	default:
-		return FLinearColor::Black;
+		return INFINITY;
 	}
+	return 0.0f;
+}
+
+FLinearColor GetClobberColor()
+{
+	float ClobberValue = GetClobberValue();
+	return FLinearColor(ClobberValue, ClobberValue, ClobberValue, ClobberValue);
 }
 
 uint32 GetClobberBufferValue()
 {
-	return 1000;
+	float ClobberValue = GetClobberValue();
+	uint32 ClobberValueUint = reinterpret_cast<const uint32*>(&ClobberValue)[0];
+	return ClobberValueUint;
 }
 
 float GetClobberDepth()
