Allows to automatically trigger any command (such as DumpGPU) when doing a high resolution screenshot

diff --git a/Engine/Source/Runtime/Engine/Private/HighResScreenshot.cpp b/Engine/Source/Runtime/Engine/Private/HighResScreenshot.cpp
--- a/Engine/Source/Runtime/Engine/Private/HighResScreenshot.cpp
+++ b/Engine/Source/Runtime/Engine/Private/HighResScreenshot.cpp
@@ -17,6 +17,21 @@ static TAutoConsoleVariable<int32> CVarSaveEXRCompressionQuality(
 	TEXT(" 1: default compression which can be slow (default)"),
 	ECVF_RenderThreadSafe);
 
+static TAutoConsoleVariable<FString> CVarHighResScreenshotCmd(
+	TEXT("r.HighResScreenshot.AdditionalCmds"), TEXT(""),
+	TEXT("Additional command to execute when a high res screenshot is requested."),
+	ECVF_Default);
+
+static void RunHighResScreenshotAdditionalCommands()
+{
+	FString Cmds = CVarHighResScreenshotCmd.GetValueOnGameThread();
+
+	if (!Cmds.IsEmpty())
+	{
+		GEngine->Exec(GWorld, *Cmds);
+	}
+}
+
 DEFINE_LOG_CATEGORY(LogHighResScreenshot);
 
 FHighResScreenshotConfig& GetHighResScreenshotConfig()
@@ -121,7 +136,7 @@ bool FHighResScreenshotConfig::ParseConsoleCommand(const FString& InCmd, FOutput
 		}
 
 		GIsHighResScreenshot = true;
-
+		RunHighResScreenshotAdditionalCommands();
 		return true;
 	}
 
@@ -183,6 +198,7 @@ bool FHighResScreenshotConfig::SetResolution(uint32 ResolutionX, uint32 Resoluti
 	GScreenshotResolutionX = (ResolutionX * ResolutionScale);
 	GScreenshotResolutionY = (ResolutionY * ResolutionScale);
 	GIsHighResScreenshot = true;
+	RunHighResScreenshotAdditionalCommands();
 
 	return true;
 }
diff --git a/Engine/Source/Runtime/RenderCore/Private/RenderGraphDumpResources.cpp b/Engine/Source/Runtime/RenderCore/Private/RenderGraphDumpResources.cpp
--- a/Engine/Source/Runtime/RenderCore/Private/RenderGraphDumpResources.cpp
+++ b/Engine/Source/Runtime/RenderCore/Private/RenderGraphDumpResources.cpp
@@ -199,7 +199,19 @@ struct FRDGResourceDumpContext
 		return FFileHelper::SaveArrayToFile(ArrayView, *FullPath);
 	}
 
-	bool IsUnsafeToDumpResource(uint32 ResourceByteSize, float DumpMemoryMultiplier) const
+	bool DumpBinaryToFile(const TArray64<uint8>& Array, const FString& FileName)
+	{
+		// Make it has if the write happened and was successful.
+		if (!bEnableDiskWrite)
+		{
+			return true;
+		}
+
+		FString FullPath = GetDumpFullPath(FileName);
+		return FFileHelper::SaveArrayToFile(Array, *FullPath);
+	}
+
+	bool IsUnsafeToDumpResource(SIZE_T ResourceByteSize, float DumpMemoryMultiplier) const
 	{
 		uint64 AproximatedStagingMemoryRequired = uint64(double(ResourceByteSize) * DumpMemoryMultiplier);
 		uint64 MaxMemoryAvailable = FMath::Min(MemoryStats.AvailablePhysical, MemoryStats.AvailableVirtual);
@@ -602,14 +614,14 @@ struct FRDGResourceDumpContext
 	struct FTextureSubresourceDumpDesc
 	{
 		FIntPoint SubResourceExtent = FIntPoint(0, 0);
-		int32 ByteSize = 0;
+		SIZE_T ByteSize = 0;
 		bool bPreprocessForStaging = false;
 		FDumpTextureCS::ETextureType DumpTextureType = FDumpTextureCS::ETextureType::MAX;
 		EPixelFormat PreprocessedPixelFormat;
 
 		bool IsDumpSupported() const
 		{
-			return ByteSize != 0;
+			return ByteSize != SIZE_T(0);
 		}
 	};
 	
@@ -665,7 +677,7 @@ struct FRDGResourceDumpContext
 			}
 		}
 
-		SubresourceDumpDesc.ByteSize = SubresourceDumpDesc.SubResourceExtent.X * SubresourceDumpDesc.SubResourceExtent.Y * GPixelFormats[SubresourceDumpDesc.PreprocessedPixelFormat].BlockBytes;
+		SubresourceDumpDesc.ByteSize = SIZE_T(SubresourceDumpDesc.SubResourceExtent.X) * SIZE_T(SubresourceDumpDesc.SubResourceExtent.Y) * SIZE_T(GPixelFormats[SubresourceDumpDesc.PreprocessedPixelFormat].BlockBytes);
 
 		// Whether the subresource need preprocessing pass before copy into staging.
 		SubresourceDumpDesc.bPreprocessForStaging = SubresourceDumpDesc.PreprocessedPixelFormat != Desc.Format || SubresourceDesc.MipLevel != 0;
@@ -766,26 +778,25 @@ struct FRDGResourceDumpContext
 
 		if (Content)
 		{
-			TArray<uint8> Array;
+			TArray64<uint8> Array;
 			Array.SetNumUninitialized(SubresourceDumpDesc.ByteSize);
 
-			int32 BytePerPixel = GPixelFormats[SubresourceDumpDesc.PreprocessedPixelFormat].BlockBytes;
+			SIZE_T BytePerPixel = SIZE_T(GPixelFormats[SubresourceDumpDesc.PreprocessedPixelFormat].BlockBytes);
 
-			const uint8* SrcData = reinterpret_cast<const uint8*>(Content);
+			const uint8* SrcData = static_cast<const uint8*>(Content);
 
 			for (int32 y = 0; y < SubresourceDumpDesc.SubResourceExtent.Y; y++)
 			{
 				// Flip the data to be bottom left corner for the WebGL viewer.
-				const uint8* SrcPos = SrcData + (SubresourceDumpDesc.SubResourceExtent.Y - 1 - y) * RowPitchInPixels * BytePerPixel;
-				uint8* DstPos = (&Array[0]) + y * SubresourceDumpDesc.SubResourceExtent.X * BytePerPixel;
+				const uint8* SrcPos = SrcData + SIZE_T(SubresourceDumpDesc.SubResourceExtent.Y - 1 - y) * SIZE_T(RowPitchInPixels) * BytePerPixel;
+				uint8* DstPos = (&Array[0]) + SIZE_T(y) * SIZE_T(SubresourceDumpDesc.SubResourceExtent.X) * BytePerPixel;
 
-				FPlatformMemory::Memmove(DstPos, SrcPos, SubresourceDumpDesc.SubResourceExtent.X * BytePerPixel);
+				FPlatformMemory::Memmove(DstPos, SrcPos, SIZE_T(SubresourceDumpDesc.SubResourceExtent.X) * BytePerPixel);
 			}
 
 			RHICmdList.UnmapStagingSurface(StagingTexture);
 
-			TArrayView<const uint8> ArrayView(reinterpret_cast<const uint8*>(&Array[0]), SubresourceDumpDesc.ByteSize);
-			DumpBinaryToFile(ArrayView, DumpFilePath);
+			DumpBinaryToFile(Array, DumpFilePath);
 		}
 		else
 		{
