Removing duplicate definition due to bad automerge
#rnx

diff --git a/Engine/Source/Runtime/RenderCore/Public/UnifiedBuffer.h b/Engine/Source/Runtime/RenderCore/Public/UnifiedBuffer.h
--- a/Engine/Source/Runtime/RenderCore/Public/UnifiedBuffer.h
+++ b/Engine/Source/Runtime/RenderCore/Public/UnifiedBuffer.h
@@ -147,17 +147,6 @@ public:
 		return ScatterBuffer.NumBytes + UploadBuffer.NumBytes;
 	}
 
-	void SetUploadViaCreate(bool bInUploadViaCreate)
-	{
-		if (bInUploadViaCreate != bUploadViaCreate)
-		{
-			// When switching the upload path, just free everything.
-			Release();
-
-			bUploadViaCreate = bInUploadViaCreate;
-		}
-	}
-	
 	/**
 	 * Init with presized num scatters, expecting each to be set at a later point. Requires the user to keep track of the offsets to use.
 	 */
