Non-Unity Fix For RHI_RAYTRACING define in RenderCore.h

diff --git a/Engine/Source/Runtime/RenderCore/Public/RenderCore.h b/Engine/Source/Runtime/RenderCore/Public/RenderCore.h
--- a/Engine/Source/Runtime/RenderCore/Public/RenderCore.h
+++ b/Engine/Source/Runtime/RenderCore/Public/RenderCore.h
@@ -8,6 +8,7 @@
 
 #include "CoreMinimal.h"
 #include "Stats/Stats.h"
+#include "RHIDefinitions.h"
 
 DECLARE_LOG_CATEGORY_EXTERN(LogRendererCore, Log, All);
 
