Fix compile error on Mac
- RHI_RAYTRACING is not visible in RenderCore.h
- the other RT related stats are not wrapped in preprocessor branches either

diff --git a/Engine/Source/Runtime/RenderCore/Private/RenderCore.cpp b/Engine/Source/Runtime/RenderCore/Private/RenderCore.cpp
--- a/Engine/Source/Runtime/RenderCore/Private/RenderCore.cpp
+++ b/Engine/Source/Runtime/RenderCore/Private/RenderCore.cpp
@@ -87,9 +87,7 @@ DEFINE_STAT(STAT_ViewStateMemory);
 DEFINE_STAT(STAT_RenderingMemStackMemory);
 DEFINE_STAT(STAT_LightInteractionMemory);
 
-#if RHI_RAYTRACING
 DEFINE_STAT(STAT_CachedRayTracingInstancesMemory);
-#endif
 
 // The InitViews stats group contains information on how long visibility culling took and how effective it was
 
diff --git a/Engine/Source/Runtime/RenderCore/Public/RenderCore.h b/Engine/Source/Runtime/RenderCore/Public/RenderCore.h
--- a/Engine/Source/Runtime/RenderCore/Public/RenderCore.h
+++ b/Engine/Source/Runtime/RenderCore/Public/RenderCore.h
@@ -149,9 +149,7 @@ DECLARE_MEMORY_STAT_EXTERN(TEXT("ViewState memory"),STAT_ViewStateMemory,STATGRO
 DECLARE_MEMORY_STAT_EXTERN(TEXT("Rendering mem stack memory"),STAT_RenderingMemStackMemory,STATGROUP_SceneMemory, RENDERCORE_API);
 DECLARE_MEMORY_STAT_EXTERN(TEXT("Light interaction memory"),STAT_LightInteractionMemory,STATGROUP_SceneMemory, RENDERCORE_API);
 
-#if RHI_RAYTRACING
 DECLARE_MEMORY_STAT_EXTERN(TEXT("Cached Ray Tracing Instances"), STAT_CachedRayTracingInstancesMemory, STATGROUP_SceneMemory, RENDERCORE_API);
-#endif
 
 DECLARE_DWORD_COUNTER_STAT_EXTERN(TEXT("Num reflective shadow map lights"),STAT_NumReflectiveShadowMapLights,STATGROUP_LightRendering, RENDERCORE_API);
 
