Adds missing SHADER_USE_PARAMETER_STRUCT() on some shader classes

diff --git a/Engine/Source/Runtime/RenderCore/Private/RenderGraphUtils.cpp b/Engine/Source/Runtime/RenderCore/Private/RenderGraphUtils.cpp
--- a/Engine/Source/Runtime/RenderCore/Private/RenderGraphUtils.cpp
+++ b/Engine/Source/Runtime/RenderCore/Private/RenderGraphUtils.cpp
@@ -705,6 +705,7 @@ void AddEnqueueCopyPass(FRDGBuilder& GraphBuilder, FRHIGPUBufferReadback* Readba
 class FClearUAVUIntCS : public FGlobalShader
 {
 	DECLARE_GLOBAL_SHADER(FClearUAVUIntCS)
+	SHADER_USE_PARAMETER_STRUCT(FClearUAVUIntCS, FGlobalShader)
 
 	BEGIN_SHADER_PARAMETER_STRUCT(FParameters, )
 		SHADER_PARAMETER_RDG_BUFFER_UAV(RWBuffer<uint>, UAV)
@@ -712,24 +713,10 @@ class FClearUAVUIntCS : public FGlobalShader
 		SHADER_PARAMETER(uint32, NumEntries)
 	END_SHADER_PARAMETER_STRUCT()
 
-	using FPermutationDomain = TShaderPermutationDomain<>;
-
-public:
-
 	static bool ShouldCompilePermutation(const FGlobalShaderPermutationParameters& Parameters)
 	{
 		return GetMaxSupportedFeatureLevel(Parameters.Platform) >= ERHIFeatureLevel::SM5;
 	}
-
-	FClearUAVUIntCS(const ShaderMetaType::CompiledShaderInitializerType& Initializer)
-		: FGlobalShader(Initializer)
-	{
-		Bindings.BindForLegacyShaderParameters(this, Initializer.PermutationId, Initializer.ParameterMap, *FParameters::FTypeInfo::GetStructMetadata());
-	}
-
-	FClearUAVUIntCS()
-	{
-	}
 };
 
 IMPLEMENT_GLOBAL_SHADER(FClearUAVUIntCS, "/Engine/Private/Tools/ClearUAV.usf", "ClearUAVUIntCS", SF_Compute);
@@ -756,6 +743,7 @@ void FComputeShaderUtils::ClearUAV(FRDGBuilder& GraphBuilder, FGlobalShaderMap*
 class FClearUAVFloatCS : public FGlobalShader
 {
 	DECLARE_GLOBAL_SHADER(FClearUAVFloatCS)
+	SHADER_USE_PARAMETER_STRUCT(FClearUAVFloatCS, FGlobalShader)
 
 	BEGIN_SHADER_PARAMETER_STRUCT(FParameters, )
 		SHADER_PARAMETER_RDG_BUFFER_UAV(RWBuffer<float4>, UAVFloat)
@@ -763,24 +751,10 @@ class FClearUAVFloatCS : public FGlobalShader
 		SHADER_PARAMETER(uint32, NumEntries)
 	END_SHADER_PARAMETER_STRUCT()
 
-	using FPermutationDomain = TShaderPermutationDomain<>;
-
-public:
-
 	static bool ShouldCompilePermutation(const FGlobalShaderPermutationParameters& Parameters)
 	{
 		return GetMaxSupportedFeatureLevel(Parameters.Platform) >= ERHIFeatureLevel::SM5;
 	}
-
-	FClearUAVFloatCS(const ShaderMetaType::CompiledShaderInitializerType& Initializer)
-		: FGlobalShader(Initializer)
-	{
-		Bindings.BindForLegacyShaderParameters(this, Initializer.PermutationId, Initializer.ParameterMap, *FParameters::FTypeInfo::GetStructMetadata());
-	}
-
-	FClearUAVFloatCS()
-	{
-	}
 };
 
 IMPLEMENT_GLOBAL_SHADER(FClearUAVFloatCS, "/Engine/Private/Tools/ClearUAV.usf", "ClearUAVFloatCS", SF_Compute);
diff --git a/Engine/Source/Runtime/Renderer/Private/ShaderPrint.cpp b/Engine/Source/Runtime/Renderer/Private/ShaderPrint.cpp
--- a/Engine/Source/Runtime/Renderer/Private/ShaderPrint.cpp
+++ b/Engine/Source/Runtime/Renderer/Private/ShaderPrint.cpp
@@ -242,7 +242,12 @@ namespace ShaderPrint
 	class FShaderDrawSymbols : public FGlobalShader
 	{
 	public:
-		SHADER_USE_PARAMETER_STRUCT(FShaderDrawSymbols, FGlobalShader);
+		FShaderDrawSymbols()
+		{}
+
+		FShaderDrawSymbols(const ShaderMetaType::CompiledShaderInitializerType& Initializer)
+			: FGlobalShader(Initializer)
+		{}
 
 		BEGIN_SHADER_PARAMETER_STRUCT(FParameters, )
 			RENDER_TARGET_BINDING_SLOTS()
@@ -260,30 +265,16 @@ namespace ShaderPrint
 
 	class FShaderDrawSymbolsVS : public FShaderDrawSymbols
 	{
-	public:
 		DECLARE_GLOBAL_SHADER(FShaderDrawSymbolsVS);
-
-		FShaderDrawSymbolsVS()
-		{}
-
-		FShaderDrawSymbolsVS(const ShaderMetaType::CompiledShaderInitializerType& Initializer)
-			: FShaderDrawSymbols(Initializer)
-		{}
+		SHADER_USE_PARAMETER_STRUCT(FShaderDrawSymbolsVS, FShaderDrawSymbols);
 	};
 
 	IMPLEMENT_GLOBAL_SHADER(FShaderDrawSymbolsVS, "/Engine/Private/ShaderPrintDraw.usf", "DrawSymbolsVS", SF_Vertex);
 
 	class FShaderDrawSymbolsPS : public FShaderDrawSymbols
 	{
-	public:
 		DECLARE_GLOBAL_SHADER(FShaderDrawSymbolsPS);
-
-		FShaderDrawSymbolsPS()
-		{}
-
-		FShaderDrawSymbolsPS(const ShaderMetaType::CompiledShaderInitializerType& Initializer)
-			: FShaderDrawSymbols(Initializer)
-		{}
+		SHADER_USE_PARAMETER_STRUCT(FShaderDrawSymbolsPS, FShaderDrawSymbols);
 	};
 
 	IMPLEMENT_GLOBAL_SHADER(FShaderDrawSymbolsPS, "/Engine/Private/ShaderPrintDraw.usf", "DrawSymbolsPS", SF_Pixel);
