PR #8679: regenerate UBO structs when platform changes (Contributed by [REDACTED_PERSON])

diff --git a/Engine/Plugins/FX/Niagara/Source/NiagaraShader/Private/NiagaraShader.cpp b/Engine/Plugins/FX/Niagara/Source/NiagaraShader/Private/NiagaraShader.cpp
--- a/Engine/Plugins/FX/Niagara/Source/NiagaraShader/Private/NiagaraShader.cpp
+++ b/Engine/Plugins/FX/Niagara/Source/NiagaraShader/Private/NiagaraShader.cpp
@@ -487,10 +487,10 @@ void FNiagaraShaderType::CacheUniformBufferIncludes(TMap<const TCHAR*, FCachedUn
 void FNiagaraShaderType::AddReferencedUniformBufferIncludes(FShaderCompilerEnvironment& OutEnvironment, FString& OutSourceFilePrefix, EShaderPlatform Platform) const
 {
 	// Cache uniform buffer struct declarations referenced by this shader type's files
-	if (!bCachedUniformBufferStructDeclarations)
+	if (CachedUniformBufferPlatform != Platform)
 	{
 		CacheUniformBufferIncludes(ReferencedUniformBufferStructsCache, Platform);
-		bCachedUniformBufferStructDeclarations = true;
+		CachedUniformBufferPlatform = Platform;
 	}
 
 	FString UniformBufferIncludes;
diff --git a/Engine/Source/Runtime/RenderCore/Private/Shader.cpp b/Engine/Source/Runtime/RenderCore/Private/Shader.cpp
--- a/Engine/Source/Runtime/RenderCore/Private/Shader.cpp
+++ b/Engine/Source/Runtime/RenderCore/Private/Shader.cpp
@@ -242,7 +242,7 @@ FShaderType::FShaderType(
 {
 	FTypeLayoutDesc::Register(InTypeLayout);
 
-	bCachedUniformBufferStructDeclarations = false;
+	CachedUniformBufferPlatform = SP_NumPlatforms;
 
 	// This will trigger if an IMPLEMENT_SHADER_TYPE was in a module not loaded before InitializeShaderTypes
 	// Shader types need to be implemented in modules that are loaded before that
diff --git a/Engine/Source/Runtime/RenderCore/Private/ShaderParameters.cpp b/Engine/Source/Runtime/RenderCore/Private/ShaderParameters.cpp
--- a/Engine/Source/Runtime/RenderCore/Private/ShaderParameters.cpp
+++ b/Engine/Source/Runtime/RenderCore/Private/ShaderParameters.cpp
@@ -332,15 +332,15 @@ static FCriticalSection UniformBufferLocks[NumUniformBufferLocks];
 
 void FShaderType::FlushShaderFileCache(const TMap<FString, TArray<const TCHAR*> >& ShaderFileToUniformBufferVariables)
 {
-	if (bCachedUniformBufferStructDeclarations)
+	if (CachedUniformBufferPlatform != SP_NumPlatforms)
 	{
 		const uint32 LockIndex = HashedName.GetHash() % NumUniformBufferLocks;
 		FScopeLock Lock(&UniformBufferLocks[LockIndex]);
-		if (bCachedUniformBufferStructDeclarations)
+		if (CachedUniformBufferPlatform != SP_NumPlatforms)
 		{
 			ReferencedUniformBufferStructsCache.Empty();
 			GenerateReferencedUniformBuffers(SourceFilename, Name, ShaderFileToUniformBufferVariables, ReferencedUniformBufferStructsCache);
-			bCachedUniformBufferStructDeclarations = false;
+			CachedUniformBufferPlatform = SP_NumPlatforms;
 		}
 	}
 }
@@ -348,14 +348,22 @@ void FShaderType::FlushShaderFileCache(const TMap<FString, TArray<const TCHAR*>
 void FShaderType::AddReferencedUniformBufferIncludes(FShaderCompilerEnvironment& OutEnvironment, FString& OutSourceFilePrefix, EShaderPlatform Platform) const
 {
 	// Cache uniform buffer struct declarations referenced by this shader type's files
-	if (!bCachedUniformBufferStructDeclarations)
+	if (CachedUniformBufferPlatform != Platform)
 	{
 		const uint32 LockIndex = HashedName.GetHash() % NumUniformBufferLocks;
 		FScopeLock Lock(&UniformBufferLocks[LockIndex]);
-		if (!bCachedUniformBufferStructDeclarations)
+		if (CachedUniformBufferPlatform != Platform)
 		{
+			// if there is already a cache but for another platform, keep the keys but reset the values
+			if (CachedUniformBufferPlatform != SP_NumPlatforms)
+			{
+				for (TMap<const TCHAR*, FCachedUniformBufferDeclaration>::TIterator It(ReferencedUniformBufferStructsCache); It; ++It)
+				{
+					It.Value() = FCachedUniformBufferDeclaration();
+				}
+			}
 			CacheUniformBufferIncludes(ReferencedUniformBufferStructsCache, Platform);
-			bCachedUniformBufferStructDeclarations = true;
+			CachedUniformBufferPlatform = Platform;
 		}
 	}
 
@@ -450,15 +458,15 @@ void FShaderType::GetShaderStableKeyParts(FStableShaderKeyAndValue& SaveKeyVal)
 
 void FVertexFactoryType::FlushShaderFileCache(const TMap<FString, TArray<const TCHAR*> >& ShaderFileToUniformBufferVariables)
 {
-	if (bCachedUniformBufferStructDeclarations)
+	if (CachedUniformBufferPlatform != SP_NumPlatforms)
 	{
 		const uint32 LockIndex = HashedName.GetHash() % NumUniformBufferLocks;
 		FScopeLock Lock(&UniformBufferLocks[LockIndex]);
-		if (bCachedUniformBufferStructDeclarations)
+		if (CachedUniformBufferPlatform != SP_NumPlatforms)
 		{
 			ReferencedUniformBufferStructsCache.Empty();
 			GenerateReferencedUniformBuffers(ShaderFilename, Name, ShaderFileToUniformBufferVariables, ReferencedUniformBufferStructsCache);
-			bCachedUniformBufferStructDeclarations = false;
+			CachedUniformBufferPlatform = SP_NumPlatforms;
 		}
 	}
 }
@@ -466,15 +474,23 @@ void FVertexFactoryType::FlushShaderFileCache(const TMap<FString, TArray<const T
 void FVertexFactoryType::AddReferencedUniformBufferIncludes(FShaderCompilerEnvironment& OutEnvironment, FString& OutSourceFilePrefix, EShaderPlatform Platform) const
 {
 	// Cache uniform buffer struct declarations referenced by this shader type's files
-	if (!bCachedUniformBufferStructDeclarations)
+	if (CachedUniformBufferPlatform != Platform)
 	{
 		const uint32 LockIndex = HashedName.GetHash() % NumUniformBufferLocks;
 		FScopeLock Lock(&UniformBufferLocks[LockIndex]);
-		if (!bCachedUniformBufferStructDeclarations)
+		if (CachedUniformBufferPlatform != Platform)
 		{
+			// if there is already a cache but for another platform, keep the keys but reset the values
+			if (CachedUniformBufferPlatform != SP_NumPlatforms)
+			{
+				for (TMap<const TCHAR*, FCachedUniformBufferDeclaration>::TIterator It(ReferencedUniformBufferStructsCache); It; ++It)
+				{
+					It.Value() = FCachedUniformBufferDeclaration();
+				}
+			}
 			CacheUniformBufferIncludes(ReferencedUniformBufferStructsCache, Platform);
-			bCachedUniformBufferStructDeclarations = true;
-		}
+			CachedUniformBufferPlatform = Platform;
+		} 
 	}
 
 	FString UniformBufferIncludes;
diff --git a/Engine/Source/Runtime/RenderCore/Private/VertexFactory.cpp b/Engine/Source/Runtime/RenderCore/Private/VertexFactory.cpp
--- a/Engine/Source/Runtime/RenderCore/Private/VertexFactory.cpp
+++ b/Engine/Source/Runtime/RenderCore/Private/VertexFactory.cpp
@@ -104,7 +104,7 @@ FVertexFactoryType::FVertexFactoryType(
 		TEXT("Incorrect virtual shader path extension for vertex factory shader header '%s': Only .ush files should be included."),
 		InShaderFilename);
 
-	bCachedUniformBufferStructDeclarations = false;
+	CachedUniformBufferPlatform = SP_NumPlatforms;
 
 	// This will trigger if an IMPLEMENT_VERTEX_FACTORY_TYPE was in a module not loaded before InitializeShaderTypes
 	// Vertex factory types need to be implemented in modules that are loaded before that
diff --git a/Engine/Source/Runtime/RenderCore/Public/Shader.h b/Engine/Source/Runtime/RenderCore/Public/Shader.h
--- a/Engine/Source/Runtime/RenderCore/Public/Shader.h
+++ b/Engine/Source/Runtime/RenderCore/Public/Shader.h
@@ -1275,7 +1275,7 @@ private:
 
 protected:
 	/** Tracks what platforms ReferencedUniformBufferStructsCache has had declarations cached for. */
-	mutable FThreadSafeBool bCachedUniformBufferStructDeclarations;
+	mutable std::atomic<EShaderPlatform> CachedUniformBufferPlatform;
 
 	/**
 	* Cache of referenced uniform buffer includes.
diff --git a/Engine/Source/Runtime/RenderCore/Public/VertexFactory.h b/Engine/Source/Runtime/RenderCore/Public/VertexFactory.h
--- a/Engine/Source/Runtime/RenderCore/Public/VertexFactory.h
+++ b/Engine/Source/Runtime/RenderCore/Public/VertexFactory.h
@@ -480,7 +480,7 @@ private:
 	mutable TMap<const TCHAR*, FCachedUniformBufferDeclaration> ReferencedUniformBufferStructsCache;
 
 	/** Tracks what platforms ReferencedUniformBufferStructsCache has had declarations cached for. */
-	mutable FThreadSafeBool bCachedUniformBufferStructDeclarations;
+	mutable std::atomic<EShaderPlatform> CachedUniformBufferPlatform;
 };
 
 /**
