FRDGTextureSRVDesc : added support for:
* SRV targeting all slices of a texture array
* SRV targeting a single mip/single slice of a texture array

diff --git a/Engine/Source/Runtime/RenderCore/Public/RenderGraphResources.h b/Engine/Source/Runtime/RenderCore/Public/RenderGraphResources.h
--- a/Engine/Source/Runtime/RenderCore/Public/RenderGraphResources.h
+++ b/Engine/Source/Runtime/RenderCore/Public/RenderGraphResources.h
@@ -626,20 +626,23 @@ public:
 		FRDGTextureSRVDesc Desc;
 		Desc.Texture = Texture;
 		Desc.NumMipLevels = Texture->Desc.NumMips;
+		Desc.NumArraySlices = Texture->Desc.ArraySize;
 		return Desc;
 	}
 
 	/** Create SRV that access one specific mip level. */
-	static FRDGTextureSRVDesc CreateForMipLevel(FRDGTextureRef Texture, int32 MipLevel)
+	static FRDGTextureSRVDesc CreateForMipLevel(FRDGTextureRef Texture, int32 MipLevel, int32 SliceIndex = 0)
 	{
 		FRDGTextureSRVDesc Desc;
 		Desc.Texture = Texture;
 		Desc.MipLevel = MipLevel;
 		Desc.NumMipLevels = 1;
+		Desc.FirstArraySlice = SliceIndex;
+		Desc.NumArraySlices = 1;
 		return Desc;
 	}
 
-	/** Create SRV that access one specific mip level. */
+	/** Create SRV that access all sub resources of texture with a specific pixel format. */
 	static FRDGTextureSRVDesc CreateWithPixelFormat(FRDGTextureRef Texture, EPixelFormat PixelFormat)
 	{
 		FRDGTextureSRVDesc Desc = FRDGTextureSRVDesc::Create(Texture);
