Adding "Auto" setting to GenerateMips that detects format support before blindly using the compute version. Adding helper to GenerateMips to tell if you need to add UAV support to a texture based on its texture format.

diff --git a/Engine/Source/Runtime/Engine/Private/TextureRenderTarget2D.cpp b/Engine/Source/Runtime/Engine/Private/TextureRenderTarget2D.cpp
--- a/Engine/Source/Runtime/Engine/Private/TextureRenderTarget2D.cpp
+++ b/Engine/Source/Runtime/Engine/Private/TextureRenderTarget2D.cpp
@@ -557,7 +557,11 @@ void FTextureRenderTarget2DResource::InitDynamicRHI()
 
 		if (Owner->bAutoGenerateMips)
 		{
-			TexCreateFlags |= (TexCreate_GenerateMipCapable | TexCreate_UAV);
+			TexCreateFlags |= TexCreate_GenerateMipCapable;
+			if (FGenerateMips::WillFormatSupportCompute(Format))
+			{
+				TexCreateFlags |= TexCreate_UAV;
+			}
 		}
 
 		if (Owner->bCanCreateUAV)
diff --git a/Engine/Source/Runtime/RenderCore/Private/GenerateMips.cpp b/Engine/Source/Runtime/RenderCore/Private/GenerateMips.cpp
--- a/Engine/Source/Runtime/RenderCore/Private/GenerateMips.cpp
+++ b/Engine/Source/Runtime/RenderCore/Private/GenerateMips.cpp
@@ -302,6 +302,11 @@ BEGIN_SHADER_PARAMETER_STRUCT(FCopyDestParameters, )
 	RDG_TEXTURE_ACCESS(Texture, ERHIAccess::CopyDest)
 END_SHADER_PARAMETER_STRUCT()
 
+bool FGenerateMips::WillFormatSupportCompute(EPixelFormat InPixelFormat)
+{
+	return RHIRequiresComputeGenerateMips() && GDynamicRHI->RHIIsTypedUAVStoreSupported(InPixelFormat);
+}
+
 void FGenerateMips::Execute(FRDGBuilder& GraphBuilder, FRDGTextureRef Texture, FGenerateMipsParams Params, EGenerateMipsPass Pass)
 {
 	if (Texture->Desc.NumMips > 1)
@@ -328,6 +333,23 @@ void FGenerateMips::Execute(FRDGBuilder& GraphBuilder, FRDGTextureRef Texture, F
 	}
 }
 
+void FGenerateMips::Execute(FRDGBuilder& GraphBuilder, FRDGTextureRef Texture, FRHISamplerState* Sampler, EGenerateMipsPass Pass)
+{
+	if (Pass == EGenerateMipsPass::AutoDetect)
+	{
+		Pass = WillFormatSupportCompute(Texture->Desc.Format) ? EGenerateMipsPass::Compute : EGenerateMipsPass::Raster;
+	}
+
+	if (Pass == EGenerateMipsPass::Compute)
+	{
+		ExecuteCompute(GraphBuilder, Texture, Sampler);
+	}
+	else
+	{
+		ExecuteRaster(GraphBuilder, Texture, Sampler);
+	}
+}
+
 void FGenerateMips::Execute(FRHICommandListImmediate& RHICmdList, FRHITexture* Texture, TSharedPtr<FGenerateMipsStruct>& ExternalMipsStructCache, FGenerateMipsParams Params, bool bAllowRenderBasedGeneration)
 {
 	TRefCountPtr<IPooledRenderTarget> PooledRenderTarget = CreateRenderTarget(Texture, TEXT("MipGeneration"));
diff --git a/Engine/Source/Runtime/RenderCore/Public/GenerateMips.h b/Engine/Source/Runtime/RenderCore/Public/GenerateMips.h
--- a/Engine/Source/Runtime/RenderCore/Public/GenerateMips.h
+++ b/Engine/Source/Runtime/RenderCore/Public/GenerateMips.h
@@ -16,6 +16,7 @@ struct FGenerateMipsParams
 
 enum class EGenerateMipsPass
 {
+	AutoDetect,
 	Compute,
 	Raster
 };
@@ -23,29 +24,21 @@ enum class EGenerateMipsPass
 class RENDERCORE_API FGenerateMips
 {
 public:
+	static bool WillFormatSupportCompute(EPixelFormat InPixelFormat);
+
 	/** (ES3.1+) Generates mips for the requested RHI texture using the feature-level appropriate means (Compute, Raster, or Fixed-Function). */
 	static void Execute(
 		FRDGBuilder& GraphBuilder,
 		FRDGTextureRef Texture,
 		FGenerateMipsParams Params = {},
-		EGenerateMipsPass Pass = EGenerateMipsPass::Compute);
+		EGenerateMipsPass Pass = EGenerateMipsPass::AutoDetect);
 
 	/** (SM5+) Generates mips for the requested RDG texture using the requested compute / raster pass. */
 	static void Execute(
 		FRDGBuilder& GraphBuilder,
 		FRDGTextureRef Texture,
 		FRHISamplerState* Sampler,
-		EGenerateMipsPass Pass = EGenerateMipsPass::Compute)
-	{
-		if (Pass == EGenerateMipsPass::Compute)
-		{
-			ExecuteCompute(GraphBuilder, Texture, Sampler);
-		}
-		else
-		{
-			ExecuteRaster(GraphBuilder, Texture, Sampler);
-		}
-	}
+		EGenerateMipsPass Pass = EGenerateMipsPass::AutoDetect);
 
 	static void ExecuteCompute(FRDGBuilder& GraphBuilder, FRDGTextureRef Texture, FRHISamplerState* Sampler);
 	
