Disables formated draw events in test builds

diff --git a/Engine/Source/Runtime/RenderCore/Public/RenderGraphDefinitions.h b/Engine/Source/Runtime/RenderCore/Public/RenderGraphDefinitions.h
--- a/Engine/Source/Runtime/RenderCore/Public/RenderGraphDefinitions.h
+++ b/Engine/Source/Runtime/RenderCore/Public/RenderGraphDefinitions.h
@@ -8,7 +8,7 @@
 /** DEFINES */
 
 /** Whether render graph debugging is enabled. */
-#define RDG_ENABLE_DEBUG (!UE_BUILD_SHIPPING)
+#define RDG_ENABLE_DEBUG (!UE_BUILD_SHIPPING && !UE_BUILD_TEST)
 
 /** Performs the operation if RDG_ENABLE_DEBUG is enabled. Useful for one-line checks without explicitly wrapping in #if. */
 #if RDG_ENABLE_DEBUG
@@ -45,7 +45,7 @@
 #define RDG_EVENTS_STRING_COPY 2
 
 /** Whether render graph GPU events are enabled. */
-#if WITH_PROFILEGPU || RDG_DUMP_RESOURCES
+#if WITH_PROFILEGPU
 	#define RDG_EVENTS RDG_EVENTS_STRING_COPY
 #elif RHI_WANT_BREADCRUMB_EVENTS
 	#define RDG_EVENTS RDG_EVENTS_STRING_REF
