REL - Move RenderingThread TaskTagScope

diff --git a/Engine/Source/Runtime/RenderCore/Private/RenderingThread.cpp b/Engine/Source/Runtime/RenderCore/Private/RenderingThread.cpp
--- a/Engine/Source/Runtime/RenderCore/Private/RenderingThread.cpp
+++ b/Engine/Source/Runtime/RenderCore/Private/RenderingThread.cpp
@@ -372,7 +372,10 @@ void RenderingThreadMain( FEvent* TaskGraphBoundSyncEvent )
 
 	FCoreDelegates::PostRenderingThreadCreated.Broadcast();
 	check(GIsThreadedRendering);
-	FTaskGraphInterface::Get().ProcessThreadUntilRequestReturn(RenderThread);
+	{
+		FTaskTagScope Scope(ETaskTag::ERenderingThread);
+		FTaskGraphInterface::Get().ProcessThreadUntilRequestReturn(RenderThread);
+	}
 	FPlatformMisc::MemoryBarrier();
 	check(!GIsThreadedRendering);
 	FCoreDelegates::PreRenderingThreadDestroyed.Broadcast();
diff --git a/Engine/Source/Runtime/RenderCore/Public/RenderingThread.h b/Engine/Source/Runtime/RenderCore/Public/RenderingThread.h
--- a/Engine/Source/Runtime/RenderCore/Public/RenderingThread.h
+++ b/Engine/Source/Runtime/RenderCore/Public/RenderingThread.h
@@ -179,7 +179,6 @@ public:
 
 	void DoTask(ENamedThreads::Type CurrentThread, const FGraphEventRef& MyCompletionGraphEvent)
 	{
-		FTaskTagScope Scope(ETaskTag::ERenderingThread);
 		TRACE_CPUPROFILER_EVENT_SCOPE_ON_CHANNEL_STR(TSTR::TStr(), RenderCommandsChannel);
 		FRHICommandListImmediate& RHICmdList = GetImmediateCommandList_ForRenderCommand();
 		Lambda(RHICmdList);
