Allocate RHI shader cache on demand (UE-136845).

- Reduces shader memory usage by avoiding upfront allocations, most of which are unused.

diff --git a/Engine/Source/Runtime/RenderCore/Private/ShaderCodeLibrary.cpp b/Engine/Source/Runtime/RenderCore/Private/ShaderCodeLibrary.cpp
--- a/Engine/Source/Runtime/RenderCore/Private/ShaderCodeLibrary.cpp
+++ b/Engine/Source/Runtime/RenderCore/Private/ShaderCodeLibrary.cpp
@@ -778,9 +778,6 @@ public:
 		const int32 NumResources = Library->GetNumShaderMaps();
 		Instance->Resources.AddZeroed(NumResources);
 
-		const int32 NumShaders = Library->GetNumShaders();
-		Instance->RHIShaders.AddZeroed(NumShaders);
-
 		INC_DWORD_STAT_BY(STAT_Shaders_ShaderResourceMemory, Instance->GetSizeBytes());
 
 		return Instance;
@@ -803,11 +800,17 @@ public:
 
 	EShaderPlatform GetPlatform() const { return Library->GetPlatform(); }
 	const int32 GetNumResources() const { return Resources.Num(); }
-	const int32 GetNumShaders() const { return RHIShaders.Num(); }
+	const int32 GetNumShaders() const { return Library->GetNumShaders(); }
 
-	uint32 GetSizeBytes() const
+	uint32 GetSizeBytes()
 	{
-		return sizeof(*this) + RHIShaders.GetAllocatedSize() + Resources.GetAllocatedSize();
+		uint32 ShaderBucketsSize = 0;
+		for (int32 IdxBucket = 0, NumBuckets = UE_ARRAY_COUNT(RHIShaders); IdxBucket < NumBuckets; ++IdxBucket)
+		{
+			FRWScopeLock Locker(ShaderLocks[IdxBucket], SLT_ReadOnly);
+			ShaderBucketsSize += RHIShaders[IdxBucket].GetAllocatedSize();
+		}
+		return sizeof(*this) + ShaderBucketsSize + Resources.GetAllocatedSize();
 	}
 
 	const int32 GetNumShadersForShaderMap(int32 ShaderMapIndex) const
@@ -891,27 +894,32 @@ public:
 
 	TRefCountPtr<FRHIShader> GetOrCreateShader(int32 ShaderIndex)
 	{
-		const int32 LockIndex = ShaderIndex % NumShaderLocks;
+		const int32 BucketIndex = ShaderIndex % NumShaderLocks;
 		TRefCountPtr<FRHIShader> Shader;
 		{
-			FRWScopeLock Locker(ShaderLocks[LockIndex], SLT_ReadOnly);
-			Shader = RHIShaders[ShaderIndex];
+			FRWScopeLock Locker(ShaderLocks[BucketIndex], SLT_ReadOnly);
+			TRefCountPtr<FRHIShader>* ShaderPtr = RHIShaders[BucketIndex].Find(ShaderIndex);
+			if (ShaderPtr)
+			{
+				Shader = *ShaderPtr;
+			}
 		}
 		if (!Shader)
 		{
-			FRWScopeLock Locker(ShaderLocks[LockIndex], SLT_Write);
+			FRWScopeLock Locker(ShaderLocks[BucketIndex], SLT_Write);
 			Shader = Library->CreateShader(ShaderIndex);
-			RHIShaders[ShaderIndex] = Shader;
+			RHIShaders[BucketIndex].Add(ShaderIndex, Shader);
 		}
 		return Shader;
 	}
 
 	void ReleaseShader(int32 ShaderIndex)
 	{
-		const int32 LockIndex = ShaderIndex % NumShaderLocks;
-		FRWScopeLock Locker(ShaderLocks[LockIndex], SLT_Write);
-		FRHIShader* Shader = RHIShaders[ShaderIndex];
-		if(Shader)
+		const int32 BucketIndex = ShaderIndex % NumShaderLocks;
+		FRWScopeLock Locker(ShaderLocks[BucketIndex], SLT_Write);
+		TRefCountPtr<FRHIShader>* ShaderPtr = RHIShaders[BucketIndex].Find(ShaderIndex);
+		FRHIShader* Shader = ShaderPtr ? ShaderPtr->GetReference() : nullptr;
+		if (Shader)
 		{
 			// The library instance is holding one ref
 			// External caller of this method must be holding a ref as well, so there must be at least 2 refs
@@ -920,7 +928,7 @@ public:
 			check(NumRefs > 1u);
 			if(NumRefs == 2u)
 			{
-				RHIShaders[ShaderIndex].SafeRelease();
+				RHIShaders[BucketIndex].Remove(ShaderIndex);
 			}
 		}
 	}
@@ -976,8 +984,13 @@ private:
 		return true;
 	}
 
-	TArray<TRefCountPtr<FRHIShader>> RHIShaders;
+	/** Number of shaders can be pretty large (several hundred thousands). Do not allocate memory for them upfront, but instead store them in a map. 
+	    There's number of maps to reduce the lock contention. */
+	TMap<int32, TRefCountPtr<FRHIShader>> RHIShaders[NumShaderLocks];
+
 	TArray<FShaderMapResource_SharedCode*> Resources;
+
+	/** Locks that guard access to particular shader buckets. */
 	FRWLock ShaderLocks[NumShaderLocks];
 	FRWLock ResourceLock;
 };
