Fix RDG event scope race condition.
Start or stop trace RDG channel could lead to mismatched scope begin/end.

diff --git a/Engine/Source/Runtime/RenderCore/Public/RenderGraphEvent.h b/Engine/Source/Runtime/RenderCore/Public/RenderGraphEvent.h
--- a/Engine/Source/Runtime/RenderCore/Public/RenderGraphEvent.h
+++ b/Engine/Source/Runtime/RenderCore/Public/RenderGraphEvent.h
@@ -408,20 +408,18 @@ public:
 	}
 
 private:
-	inline static bool IsEnabled()
+	inline bool IsEnabled()
 	{
 #if RHI_WANT_BREADCRUMB_EVENTS
 		return true;
 #elif RDG_EVENTS
-		return GetEmitRDGEvents();
+		return bRDGEvents;
 #else
 		return false;
 #endif
 	}
 
 	TRDGScopeStack<FRDGEventScopeOp> ScopeStack;
-	bool bEventPushed = false;
-
 	/** Are RDG Events enabled for these scopes */
 	bool bRDGEvents;
 };
