Fix GPU stat scope race condition.
Start or stop trace GPU channel could lead to mismatched scope begin/end.

diff --git a/Engine/Source/Runtime/RenderCore/Public/ProfilingDebugging/RealtimeGPUProfiler.h b/Engine/Source/Runtime/RenderCore/Public/ProfilingDebugging/RealtimeGPUProfiler.h
--- a/Engine/Source/Runtime/RenderCore/Public/ProfilingDebugging/RealtimeGPUProfiler.h
+++ b/Engine/Source/Runtime/RenderCore/Public/ProfilingDebugging/RealtimeGPUProfiler.h
@@ -296,7 +296,7 @@ struct RENDERCORE_API FBreadcrumbEvent
  #define GPU_STATS_ENDFRAME(RHICmdList) 
 #endif
 
-bool AreGPUStatsEnabled();
+RENDERCORE_API bool AreGPUStatsEnabled();
 
 #if HAS_GPU_STATS
 
diff --git a/Engine/Source/Runtime/RenderCore/Public/RenderGraphEvent.h b/Engine/Source/Runtime/RenderCore/Public/RenderGraphEvent.h
--- a/Engine/Source/Runtime/RenderCore/Public/RenderGraphEvent.h
+++ b/Engine/Source/Runtime/RenderCore/Public/RenderGraphEvent.h
@@ -500,6 +500,9 @@ class RENDERCORE_API FRDGGPUStatScopeStack final
 public:
 	FRDGGPUStatScopeStack(FRDGAllocator& Allocator)
 		: ScopeStack(Allocator)
+#if HAS_GPU_STATS
+		, bGPUStats(AreGPUStatsEnabled())
+#endif
 	{}
 
 	inline void BeginScope(const FName& Name, const FName& StatName, int32(*DrawCallCounter)[MAX_NUM_GPUS])
@@ -545,16 +548,18 @@ public:
 	}
 
 private:
-	inline static bool IsEnabled()
+	inline bool IsEnabled()
 	{
 #if HAS_GPU_STATS
-		return AreGPUStatsEnabled();
+		return bGPUStats;
 #else
 		return false;
 #endif
 	}
 	TRDGScopeStack<FRDGGPUStatScopeOp> ScopeStack;
 	int32 OverrideEventIndex = FRDGGPUStatScopeOpArray::kInvalidEventIndex;
+	/** Are GPU Stats enabled for these scopes */
+	bool bGPUStats = false;
 };
 
 class RENDERCORE_API FRDGGPUStatScopeGuard final
