- Fixing compile errors after robomerge happened

diff --git a/Engine/Source/Runtime/RenderCore/Public/RenderUtils.h b/Engine/Source/Runtime/RenderCore/Public/RenderUtils.h
--- a/Engine/Source/Runtime/RenderCore/Public/RenderUtils.h
+++ b/Engine/Source/Runtime/RenderCore/Public/RenderUtils.h
@@ -518,7 +518,7 @@ RENDERCORE_API bool IsUsingDBuffers(const FStaticShaderPlatform Platform);
 template<typename Type>
 struct RENDERCORE_API FShaderPlatformCachedIniValue
 {
-	FShaderPlatformCachedIniValue(const TCHAR* /*InSection*/, const TCHAR* InCVarName)
+	FShaderPlatformCachedIniValue(const TCHAR* InCVarName)
 		: CVarName(InCVarName)
 		, CVar(nullptr)
 	{
