Fix crash starting Editor when trying to retrieve an inexisting platform cvar.

diff --git a/Engine/Source/Runtime/RenderCore/Public/RenderUtils.h b/Engine/Source/Runtime/RenderCore/Public/RenderUtils.h
--- a/Engine/Source/Runtime/RenderCore/Public/RenderUtils.h
+++ b/Engine/Source/Runtime/RenderCore/Public/RenderUtils.h
@@ -544,8 +544,14 @@ struct RENDERCORE_API FShaderPlatformCachedIniValue
 			}
 		}
 
-		// now get the value from ther platform that makes sense for this shader platform
-		CVar->GetPlatformValueVariable(ShaderPlatformToPlatformName(ShaderPlatform))->GetValue(Value);
+		if (CVar != nullptr)
+		{
+			// now get the value from the platform that makes sense for this shader platform
+			if (IConsoleVariable* ConsoleVariable = CVar->GetPlatformValueVariable(ShaderPlatformToPlatformName(ShaderPlatform)))
+			{
+				ConsoleVariable->GetValue(Value);
+			}
+		}
 #else
 		unimplemented();
 #endif
