Change shader compression method to Oodle Selkie6 by default (UE-136845).

- Provides about 8% runtime memory savings (in local tests).
- Also, adds more compression types for shaders.
- Impact on the shader compilation (in SCWs) seems to be negligible in local tests.

diff --git a/Engine/Source/Programs/ShaderCompileWorker/Private/ShaderCompileWorker.cpp b/Engine/Source/Programs/ShaderCompileWorker/Private/ShaderCompileWorker.cpp
--- a/Engine/Source/Programs/ShaderCompileWorker/Private/ShaderCompileWorker.cpp
+++ b/Engine/Source/Programs/ShaderCompileWorker/Private/ShaderCompileWorker.cpp
@@ -134,7 +134,7 @@ static void ProcessCompilationJob(const FShaderCompilerInput& Input,FShaderCompi
 		Output.GenerateOutputHash();
 		if (Input.CompressionFormat != NAME_None)
 		{
-			Output.CompressOutput(Input.CompressionFormat);
+			Output.CompressOutput(Input.CompressionFormat, Input.OodleCompressor, Input.OodleLevel);
 		}
 	}
 	Output.CompileTime = FPlatformTime::Seconds() - TimeStart;
diff --git a/Engine/Source/Runtime/Core/Private/Compression/OodleDataCompression.cpp b/Engine/Source/Runtime/Core/Private/Compression/OodleDataCompression.cpp
--- a/Engine/Source/Runtime/Core/Private/Compression/OodleDataCompression.cpp
+++ b/Engine/Source/Runtime/Core/Private/Compression/OodleDataCompression.cpp
@@ -234,7 +234,7 @@ struct OodleScratchBuffers
 			OodleScratches = new OodleScratchBuffer[OodleScratchBufferCount];
 		}
 		
-		UE_LOG(OodleDataCompression, Display, TEXT("OodleScratchMemorySize=%d x %d"),(int)OodleScratchMemorySize,OodleScratchBufferCount);
+		UE_LOG(OodleDataCompression, Verbose, TEXT("OodleScratchMemorySize=%d x %d"),(int)OodleScratchMemorySize,OodleScratchBufferCount);
 	}
 	
 	~OodleScratchBuffers()
diff --git a/Engine/Source/Runtime/Engine/Private/ShaderCompiler/ShaderCompiler.cpp b/Engine/Source/Runtime/Engine/Private/ShaderCompiler/ShaderCompiler.cpp
--- a/Engine/Source/Runtime/Engine/Private/ShaderCompiler/ShaderCompiler.cpp
+++ b/Engine/Source/Runtime/Engine/Private/ShaderCompiler/ShaderCompiler.cpp
@@ -1844,7 +1844,7 @@ static bool CheckSingleJob(const FShaderCompileJob& SingleJob, TArray<FString>&
 {
 	if (SingleJob.bSucceeded)
 	{
-		check(SingleJob.Output.ShaderCode.GetShaderCodeSize() > 0);
+		checkf(SingleJob.Output.ShaderCode.GetShaderCodeSize() > 0, TEXT("Abnormal shader code size for a succeded job: %d bytes"), SingleJob.Output.ShaderCode.GetShaderCodeSize());
 	}
 
 	if (GShowShaderWarnings || !SingleJob.bSucceeded)
@@ -2537,7 +2537,7 @@ void FShaderCompileUtilities::ExecuteShaderCompileJob(FShaderCommonCompileJob& J
 			SingleJob->Output.GenerateOutputHash();
 			if (SingleJob->Input.CompressionFormat != NAME_None)
 			{
-				SingleJob->Output.CompressOutput(SingleJob->Input.CompressionFormat);
+				SingleJob->Output.CompressOutput(SingleJob->Input.CompressionFormat, SingleJob->Input.OodleCompressor, SingleJob->Input.OodleLevel);
 			}
 		}
 	}
@@ -5241,6 +5241,7 @@ void GlobalBeginCompileShader(
 	Input.Target = Target;
 	Input.ShaderFormat = ShaderFormatName;
 	Input.CompressionFormat = GetShaderCompressionFormat();
+	GetShaderCompressionOodleSettings(Input.OodleCompressor, Input.OodleLevel);
 	Input.VirtualSourceFilePath = SourceFilename;
 	Input.EntryPointName = FunctionName;
 	Input.bCompilingForShaderPipeline = false;
diff --git a/Engine/Source/Runtime/RenderCore/Private/Shader.cpp b/Engine/Source/Runtime/RenderCore/Private/Shader.cpp
--- a/Engine/Source/Runtime/RenderCore/Private/Shader.cpp
+++ b/Engine/Source/Runtime/RenderCore/Private/Shader.cpp
@@ -1909,7 +1909,15 @@ void ShaderMapAppendKeyString(EShaderPlatform Platform, FString& KeyString)
 	{
 		// add shader compression format
 		KeyString += TEXT("_Compr");
-		KeyString += GetShaderCompressionFormat().ToString();
+		FName CompressionFormat = GetShaderCompressionFormat();
+		KeyString += CompressionFormat.ToString();
+		if (CompressionFormat == NAME_Oodle)
+		{
+			FOodleDataCompression::ECompressor OodleCompressor;
+			FOodleDataCompression::ECompressionLevel OodleLevel;
+			GetShaderCompressionOodleSettings(OodleCompressor, OodleLevel);
+			KeyString += FString::Printf(TEXT("_Compr%d_Lev%d"), static_cast<int32>(OodleCompressor), static_cast<int32>(OodleLevel));
+		}
 	}
 
 	{
diff --git a/Engine/Source/Runtime/RenderCore/Private/ShaderCodeArchive.cpp b/Engine/Source/Runtime/RenderCore/Private/ShaderCodeArchive.cpp
--- a/Engine/Source/Runtime/RenderCore/Private/ShaderCodeArchive.cpp
+++ b/Engine/Source/Runtime/RenderCore/Private/ShaderCodeArchive.cpp
@@ -10,6 +10,7 @@
 #include "Serialization/JsonSerializer.h"
 #include "Misc/FileHelper.h"
 #include "Serialization/MemoryReader.h"
+#include "Compression/OodleDataCompression.h"
 #if WITH_EDITOR
 #include "Misc/StringBuilder.h"
 #include "Templates/Greater.h"
@@ -128,6 +129,22 @@ bool FSerializedShaderArchive::FindOrAddShader(const FSHAHash& Hash, int32& OutI
 	return bAdded;
 }
 
+namespace ShaderCodeArchive
+{
+	/** Decompresses the shader into caller-provided memory. Caller is assumed to allocate at least ShaderEntry uncompressed size value.
+	 * The engine will crash (LogFatal) if this function fails.
+	 */
+	void DecompressShader(uint8* OutDecompressedShader, int64 UncompressedSize, const uint8* CompressedShaderCode, int64 CompressedSize)
+	{
+		bool bSucceed = FCompression::UncompressMemory(GetShaderCompressionFormat(), OutDecompressedShader, UncompressedSize, CompressedShaderCode, CompressedSize);
+		check(bSucceed);
+		if (!bSucceed)
+		{
+			UE_LOG(LogShaderLibrary, Fatal, TEXT("ShaderCodeArchive::DecompressShader(): Could not decompress shader (GetShaderCompressionFormat=%s)"), *GetShaderCompressionFormat().ToString());
+		}
+	}
+}
+
 void FSerializedShaderArchive::DecompressShader(int32 Index, const TArray<TArray<uint8>>& ShaderCode, TArray<uint8>& OutDecompressedShader) const
 {
 	const FShaderCodeEntry& Entry = ShaderEntries[Index];
@@ -138,8 +155,7 @@ void FSerializedShaderArchive::DecompressShader(int32 Index, const TArray<TArray
 	}
 	else
 	{
-		bool bSucceed = FCompression::UncompressMemory(GetShaderCompressionFormat(), OutDecompressedShader.GetData(), Entry.UncompressedSize, ShaderCode[Index].GetData(), Entry.Size);
-		check(bSucceed);
+		ShaderCodeArchive::DecompressShader(OutDecompressedShader.GetData(), Entry.UncompressedSize, ShaderCode[Index].GetData(), Entry.Size);
 	}
 }
 
@@ -981,9 +997,8 @@ TRefCountPtr<FRHIShader> FShaderCodeArchive::CreateShader(int32 Index)
 
 	if (ShaderEntry.UncompressedSize != ShaderEntry.Size)
 	{
-		void* UncompressedCode = MemStack.Alloc(ShaderEntry.UncompressedSize, 16);
-		const bool bDecompressResult = FCompression::UncompressMemory(GetShaderCompressionFormat(), UncompressedCode, ShaderEntry.UncompressedSize, ShaderCode, ShaderEntry.Size);
-		check(bDecompressResult);
+		uint8* UncompressedCode = reinterpret_cast<uint8*>(MemStack.Alloc(ShaderEntry.UncompressedSize, 16));
+		ShaderCodeArchive::DecompressShader(UncompressedCode, ShaderEntry.UncompressedSize, ShaderCode, ShaderEntry.Size);
 		ShaderCode = (uint8*)UncompressedCode;
 	}
 
@@ -1357,9 +1372,8 @@ TRefCountPtr<FRHIShader> FIoStoreShaderCodeArchive::CreateShader(int32 Index)
 	FMemMark Mark(MemStack);
 	if (ShaderEntry.UncompressedSize != ShaderEntry.CompressedSize)
 	{
-		void* UncompressedCode = MemStack.Alloc(ShaderEntry.UncompressedSize, 16);
-		const bool bDecompressResult = FCompression::UncompressMemory(GetShaderCompressionFormat(), UncompressedCode, ShaderEntry.UncompressedSize, ShaderCode, ShaderEntry.CompressedSize);
-		check(bDecompressResult);
+		uint8* UncompressedCode = reinterpret_cast<uint8*>(MemStack.Alloc(ShaderEntry.UncompressedSize, 16));
+		ShaderCodeArchive::DecompressShader(UncompressedCode, ShaderEntry.UncompressedSize, ShaderCode, ShaderEntry.CompressedSize);
 		ShaderCode = (uint8*)UncompressedCode;
 	}
 
diff --git a/Engine/Source/Runtime/RenderCore/Private/ShaderCore.cpp b/Engine/Source/Runtime/RenderCore/Private/ShaderCore.cpp
--- a/Engine/Source/Runtime/RenderCore/Private/ShaderCore.cpp
+++ b/Engine/Source/Runtime/RenderCore/Private/ShaderCore.cpp
@@ -21,6 +21,7 @@
 #include "Interfaces/IShaderFormat.h"
 #include "Interfaces/IShaderFormatModule.h"
 #include "RHIShaderFormatDefinitions.inl"
+#include "Compression/OodleDataCompression.h"
 #if WITH_EDITOR
 #include "Misc/CoreMisc.h"
 #include "Interfaces/ITargetPlatform.h"
@@ -559,12 +560,12 @@ void FShaderCompilerOutput::GenerateOutputHash()
 	HashState.GetHash(&OutputHash.Hash[0]);
 }
 
-void FShaderCompilerOutput::CompressOutput(FName ShaderCompressionFormat)
+void FShaderCompilerOutput::CompressOutput(FName ShaderCompressionFormat, FOodleDataCompression::ECompressor OodleCompressor, FOodleDataCompression::ECompressionLevel OodleLevel)
 {
 	// make sure the hash has been generated
 	checkf(OutputHash != FSHAHash(), TEXT("Output hash must be generated before compressing the shader code."));
 	checkf(ShaderCompressionFormat != NAME_None, TEXT("Compression format should be valid"));
-	ShaderCode.Compress(ShaderCompressionFormat);
+	ShaderCode.Compress(ShaderCompressionFormat, OodleCompressor, OodleLevel);
 }
 
 static void ReportVirtualShaderFilePathError(TArray<FShaderCompilerError>* CompileErrors, FString ErrorString)
@@ -1614,23 +1615,40 @@ void AddShaderSourceDirectoryMapping(const FString& VirtualShaderDirectory, cons
 	GShaderSourceDirectoryMappings.Add(VirtualShaderDirectory, RealShaderDirectory);
 }
 
-void FShaderCode::Compress(FName ShaderCompressionFormat)
+void FShaderCode::Compress(FName ShaderCompressionFormat, FOodleDataCompression::ECompressor InOodleCompressor, FOodleDataCompression::ECompressionLevel InOodleLevel)
 {
 	checkf(OptionalDataSize == -1, TEXT("FShaderCode::Compress() was called before calling FShaderCode::FinalizeShaderCode()"));
 
 	TArray<uint8> Compressed;
-	int32 CompressedSize = ShaderCodeWithOptionalData.Num();
+	// conventional formats will fail if the compressed size isn't enough, Oodle needs a more precise estimate
+	int32 CompressedSize = (ShaderCompressionFormat != NAME_Oodle) ? ShaderCodeWithOptionalData.Num() : FOodleDataCompression::CompressedBufferSizeNeeded(ShaderCodeWithOptionalData.Num());
 	Compressed.AddUninitialized(CompressedSize);
 
+	// non-Oodle format names use the old API, for NAME_Oodle we replace the call with the custom invocation
+	bool bCompressed = false;
+	if (ShaderCompressionFormat != NAME_Oodle)
+	{
+		bCompressed = FCompression::CompressMemory(ShaderCompressionFormat, Compressed.GetData(), CompressedSize, ShaderCodeWithOptionalData.GetData(), ShaderCodeWithOptionalData.Num(), COMPRESS_BiasSize);
+	}
+	else
+	{
+		CompressedSize = FOodleDataCompression::Compress(Compressed.GetData(), CompressedSize, ShaderCodeWithOptionalData.GetData(), ShaderCodeWithOptionalData.Num(),
+			InOodleCompressor, InOodleLevel);
+		bCompressed = CompressedSize != 0;
+	}
+
 	// there is code that assumes that if CompressedSize == CodeSize, the shader isn't compressed. Because of that, do not accept equal compressed size (very unlikely anyway)
-	if (FCompression::CompressMemory(ShaderCompressionFormat, Compressed.GetData(), CompressedSize, ShaderCodeWithOptionalData.GetData(), ShaderCodeWithOptionalData.Num()) && CompressedSize < ShaderCodeWithOptionalData.Num())
+	if (bCompressed && CompressedSize < ShaderCodeWithOptionalData.Num())
 	{
 		// cache the ShaderCodeSize since it will no longer possible to get it as the reader will fail to parse the compressed data
 		FShaderCodeReader Wrapper(ShaderCodeWithOptionalData);
 		ShaderCodeSize = Wrapper.GetShaderCodeSize();
+		checkf(ShaderCodeSize >= 0, TEXT("Unable to determine ShaderCodeSize from uncompressed code"), ShaderCodeSize);
 
 		// finalize the compression
 		CompressionFormat = ShaderCompressionFormat;
+		OodleCompressor = InOodleCompressor;
+		OodleLevel = InOodleLevel;
 		UncompressedSize = ShaderCodeWithOptionalData.Num();
 
 		Compressed.SetNum(CompressedSize);
@@ -1657,7 +1675,10 @@ FArchive& operator<<(FArchive& Ar, FShaderCode& Output)
 		Ar << CompressionFormatString;
 		Output.CompressionFormat = FName(*CompressionFormatString);
 	}
+	Ar << reinterpret_cast<uint8&>(Output.OodleCompressor);
+	Ar << reinterpret_cast<uint8&>(Output.OodleLevel);
 	Ar << Output.ShaderCodeSize;
+	checkf(Output.UncompressedSize == 0 || Output.ShaderCodeSize > 0, TEXT("FShaderCode::operator<<(): invalid shader code size for a compressed shader: ShaderCodeSize=%d, UncompressedSize=%d"), Output.ShaderCodeSize, Output.UncompressedSize);
 	return Ar;
 }
 
@@ -1691,6 +1712,8 @@ FArchive& operator<<(FArchive& Ar, FShaderCompilerInput& Input)
 	Ar << Input.DebugDescription;
 	Ar << Input.Environment;
 	Ar << Input.ExtraSettings;
+	Ar << reinterpret_cast<uint8&>(Input.OodleCompressor);
+	Ar << reinterpret_cast<uint8&>(Input.OodleLevel);
 
 	// Note: skipping Input.SharedEnvironment, which is handled by FShaderCompileUtilities::DoWriteTasks in order to maintain sharing
 
diff --git a/Engine/Source/Runtime/RenderCore/Private/ShaderResource.cpp b/Engine/Source/Runtime/RenderCore/Private/ShaderResource.cpp
--- a/Engine/Source/Runtime/RenderCore/Private/ShaderResource.cpp
+++ b/Engine/Source/Runtime/RenderCore/Private/ShaderResource.cpp
@@ -24,18 +24,53 @@
 #include "ProfilingDebugging/LoadTimeTracker.h"
 #include "Misc/MemStack.h"
 #include "ShaderCompilerCore.h"
+#include "Compression/OodleDataCompression.h"
 
 #if WITH_EDITORONLY_DATA
 #include "Interfaces/IShaderFormat.h"
 #endif
 
-int32 GShaderCompressionFormatChoice = 1;
+int32 GShaderCompressionFormatChoice = 2;
 static FAutoConsoleVariableRef CVarShaderCompressionFormatChoice(
 	TEXT("r.Shaders.CompressionFormat"),
 	GShaderCompressionFormatChoice,
 	TEXT("Select the compression methods for the shader code.\n")
 	TEXT(" 0: None (uncompressed)\n")
-	TEXT(" 1: LZ4 (default)"),
+	TEXT(" 1: LZ4\n")
+	TEXT(" 2: Oodle (default)\n")
+	TEXT(" 3: ZLib\n"),
+	ECVF_ReadOnly);
+
+int32 GShaderCompressionOodleAlgo = 1;
+static FAutoConsoleVariableRef CVarShaderCompressionOodleAlgo(
+	TEXT("r.Shaders.CompressionFormat.Oodle.Algo"),
+	GShaderCompressionOodleAlgo,
+	TEXT("Oodle compression method for the shader code, from fastest to slowest to decode.\n")
+	TEXT(" 0: None (invalid setting)\n")
+	TEXT(" 1: Selkie (fastest to decode)\n")
+	TEXT(" 2: Mermaid\n")
+	TEXT(" 3: Kraken\n")
+	TEXT(" 4: Leviathan (slowest to decode)\n"),
+	ECVF_ReadOnly);
+
+int32 GShaderCompressionOodleLevel = 6;
+static FAutoConsoleVariableRef CVarShaderCompressionOodleAlgoChoice(
+	TEXT("r.Shaders.CompressionFormat.Oodle.Level"),
+	GShaderCompressionOodleLevel,
+	TEXT("Oodle compression level. This mostly trades encode speed vs compression ratio, decode speed is determined by r.Shaders.CompressionFormat.Oodle.Algo\n")
+	TEXT(" -4 : HyperFast4\n")
+	TEXT(" -3 : HyperFast3\n")
+	TEXT(" -2 : HyperFast2\n")
+	TEXT(" -1 : HyperFast1\n")
+	TEXT("  0 : None\n")
+	TEXT("  1 : SuperFast\n")
+	TEXT("  2 : VeryFast\n")
+	TEXT("  3 : Fast\n")
+	TEXT("  4 : Normal\n")
+	TEXT("  5 : Optimal1\n")
+	TEXT("  6 : Optimal2\n")
+	TEXT("  7 : Optimal3\n")
+	TEXT("  8 : Optimal4\n"),
 	ECVF_ReadOnly);
 
 FName GetShaderCompressionFormat(const FName& ShaderFormat)
@@ -50,7 +85,38 @@ FName GetShaderCompressionFormat(const FName& ShaderFormat)
 	}
 #endif
 
-	return (GShaderCompressionFormatChoice == 1) ? NAME_LZ4 : NAME_None;
+	static FName Formats[]
+	{
+		NAME_None,
+		NAME_LZ4,
+		NAME_Oodle,
+		NAME_Zlib
+	};
+	
+	//GShaderCompressionFormatChoice = (GShaderCompressionFormatChoice < 0) ? 0 : GShaderCompressionFormatChoice;
+	GShaderCompressionFormatChoice = FMath::Clamp(GShaderCompressionFormatChoice, 0, UE_ARRAY_COUNT(Formats) - 1);
+	return Formats[GShaderCompressionFormatChoice];
+}
+
+void GetShaderCompressionOodleSettings(FOodleDataCompression::ECompressor& OutCompressor, FOodleDataCompression::ECompressionLevel& OutLevel, const FName& ShaderFormat)
+{
+	// support an older developer-only CVar for compatibility and make it preempt
+#if !(UE_BUILD_SHIPPING || UE_BUILD_TEST)
+	static const IConsoleVariable* CVarSkipCompression = IConsoleManager::Get().FindConsoleVariable(TEXT("r.Shaders.SkipCompression"));
+	static bool bSkipCompression = (CVarSkipCompression && CVarSkipCompression->GetInt() != 0);
+	if (UNLIKELY(bSkipCompression))
+	{
+		OutCompressor = FOodleDataCompression::ECompressor::Selkie;
+		OutLevel = FOodleDataCompression::ECompressionLevel::None;
+		return;
+	}
+#endif
+
+	GShaderCompressionOodleAlgo = FMath::Clamp(GShaderCompressionOodleAlgo, static_cast<int32>(FOodleDataCompression::ECompressor::NotSet), static_cast<int32>(FOodleDataCompression::ECompressor::Leviathan));
+	OutCompressor = static_cast<FOodleDataCompression::ECompressor>(GShaderCompressionOodleAlgo);
+
+	GShaderCompressionOodleLevel = FMath::Clamp(GShaderCompressionOodleLevel, static_cast<int32>(FOodleDataCompression::ECompressionLevel::HyperFast4), static_cast<int32>(FOodleDataCompression::ECompressionLevel::Optimal4));
+	OutLevel = static_cast<FOodleDataCompression::ECompressionLevel>(GShaderCompressionOodleLevel);
 }
 
 bool FShaderMapResource::ArePlatformsCompatible(EShaderPlatform CurrentPlatform, EShaderPlatform TargetPlatform)
@@ -240,6 +306,26 @@ void FShaderMapResourceCode::AddShaderCode(EShaderFrequency InFrequency, const F
 					*ShaderCompressionFormat.ToString()
 				);
 			}
+			else if (ShaderCompressionFormat == NAME_Oodle)
+			{
+				// check if Oodle-specific settings match
+				FOodleDataCompression::ECompressor OodleCompressor;
+				FOodleDataCompression::ECompressionLevel OodleLevel;
+				GetShaderCompressionOodleSettings(OodleCompressor, OodleLevel);
+
+				if (InCode.GetOodleCompressor() != OodleCompressor || InCode.GetOodleLevel() != OodleLevel)
+				{
+					UE_LOG(LogShaders, Fatal, TEXT("Shader %s is expected to be compressed with Oodle compressor %d level %d, but it is compressed with compressor %d level %d instead."),
+						*InHash.ToString(),
+						static_cast<int32>(OodleCompressor),
+						static_cast<int32>(OodleLevel),
+						static_cast<int32>(InCode.GetOodleCompressor()),
+						static_cast<int32>(InCode.GetOodleLevel())
+					);
+					// unreachable
+					return;
+				}
+			}
 		}
 		else
 		{
@@ -473,7 +559,7 @@ TRefCountPtr<FRHIShader> FShaderMapResource_InlineCode::CreateRHIShader(int32 Sh
 	if (ShaderEntry.Code.Num() != ShaderEntry.UncompressedSize)
 	{
 		void* UncompressedCode = MemStack.Alloc(ShaderEntry.UncompressedSize, 16);
-		auto bSucceed = FCompression::UncompressMemory(GetShaderCompressionFormat(), UncompressedCode, ShaderEntry.UncompressedSize, ShaderCode, ShaderEntry.Code.Num());
+		bool bSucceed = FCompression::UncompressMemory(GetShaderCompressionFormat(), UncompressedCode, ShaderEntry.UncompressedSize, ShaderCode, ShaderEntry.Code.Num());
 		check(bSucceed);
 		ShaderCode = (uint8*)UncompressedCode;
 	}
diff --git a/Engine/Source/Runtime/RenderCore/Public/ShaderCompilerCore.h b/Engine/Source/Runtime/RenderCore/Public/ShaderCompilerCore.h
--- a/Engine/Source/Runtime/RenderCore/Public/ShaderCompilerCore.h
+++ b/Engine/Source/Runtime/RenderCore/Public/ShaderCompilerCore.h
@@ -163,6 +163,12 @@ struct FExtraShaderCompilerSettings
 	}
 };
 
+namespace FOodleDataCompression
+{
+	enum class ECompressor : uint8;
+	enum class ECompressionLevel : int8;
+}
+
 /** Struct that gathers all readonly inputs needed for the compilation of a single shader. */
 struct FShaderCompilerInput
 {
@@ -209,6 +215,12 @@ struct FShaderCompilerInput
 	// FMaterial::SetupExtaCompilationSettings is usually called by each (*)MaterialShaderType::BeginCompileShader() function
 	FExtraShaderCompilerSettings ExtraSettings;
 
+	/** Oodle-specific compression algorithm - used if CompressionFormat is set to NAME_Oodle. */
+	FOodleDataCompression::ECompressor OodleCompressor;
+
+	/** Oodle-specific compression level - used if CompressionFormat is set to NAME_Oodle. */
+	FOodleDataCompression::ECompressionLevel OodleLevel;
+
 	FShaderCompilerInput() :
 		Target(SF_NumFrequencies, SP_NumPlatforms),
 		bSkipPreprocessedCache(false),
@@ -451,7 +463,7 @@ struct FShaderCompilerOutput
 	RENDERCORE_API void GenerateOutputHash();
 
 	/** Calls GenerateOutputHash() before the compression, replaces FShaderCode with the compressed data (if compression result was smaller). */
-	RENDERCORE_API void CompressOutput(FName ShaderCompressionFormat);
+	RENDERCORE_API void CompressOutput(FName ShaderCompressionFormat, FOodleDataCompression::ECompressor OodleCompressor, FOodleDataCompression::ECompressionLevel OodleLevel);
 
 	friend FArchive& operator<<(FArchive& Ar, FShaderCompilerOutput& Output)
 	{
diff --git a/Engine/Source/Runtime/RenderCore/Public/ShaderCore.h b/Engine/Source/Runtime/RenderCore/Public/ShaderCore.h
--- a/Engine/Source/Runtime/RenderCore/Public/ShaderCore.h
+++ b/Engine/Source/Runtime/RenderCore/Public/ShaderCore.h
@@ -104,6 +104,15 @@ extern RENDERCORE_API bool AllowDebugViewmodes(EShaderPlatform Platform);
 /** Returns the shader compression format (passing ShaderFormat for future proofing, but as of now the setting is global for all formats). */
 extern RENDERCORE_API FName GetShaderCompressionFormat(const FName& ShaderFormat = NAME_None);
 
+namespace FOodleDataCompression
+{
+	enum class ECompressor : uint8;
+	enum class ECompressionLevel : int8;
+}
+
+/** Returns Oodle-specific shader compression format settings (passing ShaderFormat for future proofing, but as of now the setting is global for all formats). */
+extern RENDERCORE_API void GetShaderCompressionOodleSettings(FOodleDataCompression::ECompressor& OutCompressor, FOodleDataCompression::ECompressionLevel& OutLevel, const FName& ShaderFormat = NAME_None);
+
 struct FShaderTarget
 {
 	// The rest of uint32 holding the bitfields can be left unitialized. Union with a uint32 serves to prevent that to be able to set the whole uint32 value
@@ -836,6 +845,12 @@ class FShaderCode
 	/** Compression algo */
 	mutable FName CompressionFormat;
 
+	/** Oodle-specific compression algorithm - used if CompressionFormat is set to NAME_Oodle. */
+	FOodleDataCompression::ECompressor OodleCompressor;
+
+	/** Oodle-specific compression level - used if CompressionFormat is set to NAME_Oodle. */
+	FOodleDataCompression::ECompressionLevel OodleLevel;
+
 	/** We cannot get the code size after the compression, so store it here */
 	mutable int32 ShaderCodeSize;
 
@@ -861,7 +876,7 @@ public:
 		}
 	}
 
-	void Compress(FName ShaderCompressionFormat);
+	void Compress(FName ShaderCompressionFormat, FOodleDataCompression::ECompressor InOodleCompressor, FOodleDataCompression::ECompressionLevel InOodleLevel);
 
 	// Write access for regular microcode: Optional Data must be added AFTER regular microcode and BEFORE Finalize
 	TArray<uint8>& GetWriteAccess()
@@ -905,6 +920,16 @@ public:
 		return CompressionFormat;
 	}
 
+	FOodleDataCompression::ECompressor GetOodleCompressor() const
+	{
+		return OodleCompressor;
+	}
+
+	FOodleDataCompression::ECompressionLevel GetOodleLevel() const
+	{
+		return OodleLevel;
+	}
+
 	int32 GetUncompressedSize() const
 	{
 		return UncompressedSize;
