Restore CL15587695: Fixed dangling RHI sampler state in FGenerateMips::Execute

diff --git a/Engine/Source/Runtime/RenderCore/Private/GenerateMips.cpp b/Engine/Source/Runtime/RenderCore/Private/GenerateMips.cpp
--- a/Engine/Source/Runtime/RenderCore/Private/GenerateMips.cpp
+++ b/Engine/Source/Runtime/RenderCore/Private/GenerateMips.cpp
@@ -314,7 +314,8 @@ void FGenerateMips::Execute(FRDGBuilder& GraphBuilder, FRDGTextureRef Texture, F
 		if (RHIRequiresComputeGenerateMips())
 		{
 			FSamplerStateInitializerRHI SamplerInit(Params.Filter, Params.AddressU, Params.AddressV, Params.AddressW);
-			Execute(GraphBuilder, Texture, RHICreateSamplerState(SamplerInit), Pass);
+			FSamplerStateRHIRef Sampler = *GraphBuilder.AllocObject<FSamplerStateRHIRef>(RHICreateSamplerState(SamplerInit));
+			Execute(GraphBuilder, Texture, Sampler, Pass);
 		}
 		else
 		{
