- Allow for FShaderPlatformCachedIniValue on non-editor platforms - it will enforce that it can only look up it's own value
- Whenever looking up the running platform, just use the CVar's value directly instead of parsing ini values
#p4v-preflight-copy 18520954

diff --git a/Engine/Source/Runtime/RenderCore/Public/RenderUtils.h b/Engine/Source/Runtime/RenderCore/Public/RenderUtils.h
--- a/Engine/Source/Runtime/RenderCore/Public/RenderUtils.h
+++ b/Engine/Source/Runtime/RenderCore/Public/RenderUtils.h
@@ -531,21 +531,34 @@ struct RENDERCORE_API FShaderPlatformCachedIniValue
 
 	Type Get(EShaderPlatform ShaderPlatform)
 	{
-		Type Value {};
-#if ALLOW_OTHER_PLATFORM_CONFIG
-		// find/create a cvar if needed
+		Type Value{};
+
+		FName IniPlatformName = ShaderPlatformToPlatformName(ShaderPlatform);
+		// find the cvar if needed
 		if (CVar == nullptr)
 		{
 			CVar = IConsoleManager::Get().FindConsoleVariable(*CVarName);
-			if (CVar == nullptr)
-			{
-				// this could be a cvar that only exists on the target platform so create a dummy one
-				CVar = IConsoleManager::Get().RegisterConsoleVariable(*CVarName, Type(), TEXT(""), ECVF_ReadOnly);
-			}
+		}
+
+		// if we are looking up our own platform, just use the current value
+		if (IniPlatformName == FPlatformProperties::IniPlatformName())
+		{
+			checkf(CVar != nullptr, TEXT("Failed to find CVar %s when getting current value for FShaderPlatformCachedIniValue"));
+
+			CVar->GetValue(Value);
+			return Value;
+		}
+
+#if ALLOW_OTHER_PLATFORM_CONFIG
+		// create a dummy cvar if needed
+		if (CVar == nullptr)
+		{
+			// this could be a cvar that only exists on the target platform so create a dummy one
+			CVar = IConsoleManager::Get().RegisterConsoleVariable(*CVarName, Type(), TEXT(""), ECVF_ReadOnly);
 		}
 
 		// now get the value from the platform that makes sense for this shader platform
-		TSharedPtr<IConsoleVariable> OtherPlatformVar = CVar->GetPlatformValueVariable(ShaderPlatformToPlatformName(ShaderPlatform));
+		TSharedPtr<IConsoleVariable> OtherPlatformVar = CVar->GetPlatformValueVariable(IniPlatformName);
 		ensureMsgf(OtherPlatformVar.IsValid(), TEXT("Failed to get another platform's version of a cvar (possible name: '%s'). It is probably an esoteric subclass that needs to implement GetPlatformValueVariable."), *CVarName);
 		if (OtherPlatformVar.IsValid())
 		{
@@ -557,7 +570,7 @@ struct RENDERCORE_API FShaderPlatformCachedIniValue
 			CVar->GetValue(Value);
 		}
 #else
-		unimplemented();
+		checkf(IniPlatformName == FName(FPlatformProperties::IniPlatformName()), TEXT("FShaderPlatformCachedIniValue can only look up the current platform when ALLOW_OTHER_PLATFORM_CONFIG is false"));
 #endif
 		return Value;
 	}
