Add visualization of the runtime shader usage.

- Produces an image in Saved/Profiling with the preloaded and created shaders, helping to understand preloading / creation patterns.
- Only works on desktop platforms in Development. Disabled by default (set r.ShaderCodeLibrary.VisualizeShaderUsage=1 in ini to enable).

diff --git a/Engine/Source/Runtime/RenderCore/Private/ShaderCodeArchive.cpp b/Engine/Source/Runtime/RenderCore/Private/ShaderCodeArchive.cpp
--- a/Engine/Source/Runtime/RenderCore/Private/ShaderCodeArchive.cpp
+++ b/Engine/Source/Runtime/RenderCore/Private/ShaderCodeArchive.cpp
@@ -16,6 +16,12 @@
 #include "Templates/Greater.h"
 #endif
 
+#if UE_SCA_VISUALIZE_SHADER_USAGE
+#include "IImageWrapperModule.h"
+#include "IImageWrapper.h"
+#endif // UE_SCA_VISUALIZE_SHADER_USAGE
+
+
 int32 GShaderCodeLibraryAsyncLoadingPriority = int32(AIOP_Normal);
 static FAutoConsoleVariableRef CVarShaderCodeLibraryAsyncLoadingPriority(
 	TEXT("r.ShaderCodeLibrary.DefaultAsyncIOPriority"),
@@ -32,6 +38,14 @@ static FAutoConsoleVariableRef CVarShaderCodeLibraryAsyncLoadingAllowDontCache(
 	ECVF_Default
 );
 
+int32 GShaderCodeLibraryVisualizeShaderUsage = 0;
+static FAutoConsoleVariableRef CVarShaderCodeLibraryVisualizeShaderUsage(
+	TEXT("r.ShaderCodeLibrary.VisualizeShaderUsage"),
+	GShaderCodeLibraryVisualizeShaderUsage,
+	TEXT("If 1, a bitmap with the used shaders (for each shader library chunk) will be saved at the exit. Works in standalone games only."),
+	ECVF_RenderThreadSafe | ECVF_ReadOnly
+);
+
 int32 FSerializedShaderArchive::FindShaderMapWithKey(const FSHAHash& Hash, uint32 Key) const
 {
 	for (uint32 Index = ShaderMapHashTable.First(Key); ShaderMapHashTable.IsValid(Index); Index = ShaderMapHashTable.Next(Index))
@@ -129,6 +143,100 @@ bool FSerializedShaderArchive::FindOrAddShader(const FSHAHash& Hash, int32& OutI
 	return bAdded;
 }
 
+#if UE_SCA_VISUALIZE_SHADER_USAGE
+void FShaderUsageVisualizer::Initialize(const int32 InNumShaders)
+{
+	FScopeLock Lock(&VisualizeLock);
+	NumShaders = InNumShaders;
+}
+
+void FShaderUsageVisualizer::SaveShaderUsageBitmap(const FString& Name, EShaderPlatform ShaderPlatform)
+{
+	if (GShaderCodeLibraryVisualizeShaderUsage)
+	{
+		if (NumShaders)
+		{
+			if (IImageWrapperModule* ImageWrapperModule = FModuleManager::Get().GetModulePtr<IImageWrapperModule>(TEXT("ImageWrapper")))
+			{
+				if (TSharedPtr<IImageWrapper> PNGImageWrapper = ImageWrapperModule->CreateImageWrapper(EImageFormat::PNG))
+				{
+					UE_LOG(LogShaderLibrary, Log, TEXT("Creating shader usage bitmap for archive %s (NumShaders: %d, preloaded %d, created %d)"), *Name, NumShaders, PreloadedShadersForVis.Num(), CreatedShadersForVis.Num());
+
+					// find a value close to sqrt(NumShaders)
+					int32 ImageDimension = static_cast<int32>(FMath::Sqrt(static_cast<float>(NumShaders))) + 1;
+					TArray<FColor> ShaderUsageBitmap;
+					ShaderUsageBitmap.Reserve(ImageDimension * ImageDimension);
+
+					// map legend:
+					FColor UnusedShaderColor(128, 128, 128);	// unused shaders - this is the majority of the bitmap content
+					FColor PreloadedShaderColor(225, 225, 225);	// preloaded shaders - they can become the majority under certain circumstances. Pure white can blend with some viewer's background
+					FColor CreatedShaderColor(0, 0, 255);	// created shaders - in practice, always few and far between. Blue is more noticeable on a largely bright background than magenta
+
+					for (int32 Idx = 0; Idx < NumShaders; ++Idx)
+					{
+						ShaderUsageBitmap.Add(UnusedShaderColor);
+					}
+					// the rest can be zero/transparent
+					ShaderUsageBitmap.AddZeroed(ImageDimension * ImageDimension - NumShaders);
+					check(ShaderUsageBitmap.Num() == ImageDimension * ImageDimension);
+
+					{
+						// in case this ever gets called runtime
+						FScopeLock Lock(&VisualizeLock);
+
+						// fill preloaded ones first
+						for (int32 ShaderIdx : PreloadedShadersForVis)
+						{
+							ShaderUsageBitmap[ShaderIdx] = PreloadedShaderColor;
+						}
+
+						for (int32 ShaderIdx : CreatedShadersForVis)
+						{
+							ShaderUsageBitmap[ShaderIdx] = CreatedShaderColor;
+						}
+					}
+
+					bool bSet = PNGImageWrapper->SetRaw(ShaderUsageBitmap.GetData(), ShaderUsageBitmap.Num() * sizeof(FColor),
+						ImageDimension, ImageDimension, ERGBFormat::BGRA, 8);
+
+					if (bSet)
+					{
+						TArray64<uint8> CompressedData = PNGImageWrapper->GetCompressed(100);
+
+						const FString Filename = FString::Printf(TEXT("%s_%s_RuntimeShaderUsage_%s.png"), *Name, *LexToString(ShaderPlatform), *FDateTime::Now().ToString());
+						const FString SaveDir = FPaths::Combine(FPaths::ProjectSavedDir(), TEXT("Profiling"));
+						const FString FilePath = FPaths::Combine(SaveDir, Filename);
+
+						if (!FFileHelper::SaveArrayToFile(CompressedData, *FilePath))
+						{
+							UE_LOG(LogShaderLibrary, Warning, TEXT("Couldn't write shader usage bitmap %s"), *FilePath);
+						}
+						else
+						{
+							// technically not every created might have been preloaded, but we ignore this here
+							UE_LOG(LogShaderLibrary, Log, TEXT("Saved shader usage bitmap %s. Legend: unused shaders: %d (dark grey), preloaded (compressed) to RAM: %d (white), uncompressed and created: %d (blue)"), 
+								*FilePath, NumShaders - PreloadedShadersForVis.Num(), PreloadedShadersForVis.Num(), CreatedShadersForVis.Num());
+						}
+					}
+					else
+					{
+						UE_LOG(LogShaderLibrary, Warning, TEXT("Error creating shader usage bitmap for archive %s (NumShaders: %d, preloaded %d, created %d) - cannot create a PNG image"), *Name, NumShaders, PreloadedShadersForVis.Num(), CreatedShadersForVis.Num());
+					}
+				}
+				else
+				{
+					UE_LOG(LogShaderLibrary, Warning, TEXT("Couldn't create shader usage bitmap for archive %s (NumShaders: %d, preloaded %d, created %d) - cannot create ImageWrapper for PNG format"), *Name, NumShaders, PreloadedShadersForVis.Num(), CreatedShadersForVis.Num());
+				}
+			}
+			else
+			{
+				UE_LOG(LogShaderLibrary, Warning, TEXT("Couldn't create shader usage bitmap for archive %s (NumShaders: %d, preloaded %d, created %d) - no ImageWrapper module"), *Name, NumShaders, PreloadedShadersForVis.Num(), CreatedShadersForVis.Num());
+			}
+		}
+	}
+}
+#endif
+
 namespace ShaderCodeArchive
 {
 	/** Decompresses the shader into caller-provided memory. Caller is assumed to allocate at least ShaderEntry uncompressed size value.
@@ -723,6 +831,8 @@ FShaderCodeArchive* FShaderCodeArchive::Create(EShaderPlatform InPlatform, FArch
 	// Open library for async reads
 	Library->FileCacheHandle = IFileCacheHandle::CreateFileCacheHandle(*InDestFilePath);
 
+	Library->DebugVisualizer.Initialize(Library->SerializedShaders.ShaderEntries.Num());
+
 	UE_LOG(LogShaderLibrary, Display, TEXT("Using %s for material shader code. Total %d unique shaders."), *InDestFilePath, Library->SerializedShaders.ShaderEntries.Num());
 
 	INC_DWORD_STAT_BY(STAT_Shaders_ShaderResourceMemory, Library->GetSizeBytes());
@@ -763,6 +873,8 @@ void FShaderCodeArchive::Teardown()
 			DEC_DWORD_STAT_BY(STAT_Shaders_ShaderPreloadMemory, ShaderEntry.Size);
 		}
 	}
+
+	DebugVisualizer.SaveShaderUsageBitmap(GetName(), GetPlatform());
 }
 
 void FShaderCodeArchive::OnShaderPreloadFinished(int32 ShaderIndex, const IMemoryReadStreamRef& PreloadData)
@@ -814,6 +926,7 @@ bool FShaderCodeArchive::PreloadShader(int32 ShaderIndex, FGraphEventArray& OutC
 		const FShaderCodeEntry& ShaderEntry = SerializedShaders.ShaderEntries[ShaderIndex];
 		ShaderPreloadEntry.Code = FMemory::Malloc(ShaderEntry.Size);
 		ShaderPreloadEntry.FramePreloadStarted = GFrameNumber;
+		DebugVisualizer.MarkPreloadedForVisualization(ShaderIndex);
 
 		const EAsyncIOPriorityAndFlags IOPriority = (EAsyncIOPriorityAndFlags)GShaderCodeLibraryAsyncLoadingPriority;
 
@@ -866,6 +979,7 @@ bool FShaderCodeArchive::PreloadShaderMap(int32 ShaderMapIndex, FGraphEventArray
 			check(!ShaderPreloadEntry.PreloadEvent);
 			ShaderPreloadEntry.Code = FMemory::Malloc(ShaderEntry.Size);
 			ShaderPreloadEntry.FramePreloadStarted = FrameNumber;
+			DebugVisualizer.MarkPreloadedForVisualization(ShaderIndex);
 			PreloadMemory += ShaderEntry.Size;
 
 			FGraphEventArray ReadCompletionEvents;
@@ -1026,6 +1140,7 @@ TRefCountPtr<FRHIShader> FShaderCodeArchive::CreateShader(int32 Index)
 		break;
 	default: checkNoEntry(); break;
 	}
+	DebugVisualizer.MarkCreatedForVisualization(Index);
 
 	// Release the refernece we were holding
 	if (PreloadedShaderCode)
@@ -1134,7 +1249,9 @@ FIoStoreShaderCodeArchive* FIoStoreShaderCodeArchive::Create(EShaderPlatform InP
 					Library->ShaderHashTable.Add(Key, Index);
 				}
 			}
-			
+
+			Library->DebugVisualizer.Initialize(Library->ShaderEntries.Num());
+
 			UE_LOG(LogShaderLibrary, Display, TEXT("Using IoDispatcher for shader code library %s. Total %d unique shaders."), *InLibraryName, Library->ShaderEntries.Num());
 			INC_DWORD_STAT_BY(STAT_Shaders_ShaderResourceMemory, Library->GetSizeBytes());
 			return Library;
@@ -1162,6 +1279,8 @@ void FIoStoreShaderCodeArchive::Teardown()
 		FShaderPreloadEntry& ShaderPreloadEntry = ShaderPreloads[ShaderIndex];
 		check(!ShaderPreloadEntry.NumRefs);
 	}
+
+	DebugVisualizer.SaveShaderUsageBitmap(GetName(), GetPlatform());
 }
 
 bool FIoStoreShaderCodeArchive::PreloadShader(int32 ShaderIndex, FGraphEventArray& OutCompletionEvents)
@@ -1178,6 +1297,7 @@ bool FIoStoreShaderCodeArchive::PreloadShader(int32 ShaderIndex, FGraphEventArra
 
 		const FIoStoreShaderCodeEntry& ShaderEntry = ShaderEntries[ShaderIndex];
 		ShaderPreloadEntry.FramePreloadStarted = GFrameNumber;
+		DebugVisualizer.MarkPreloadedForVisualization(ShaderIndex);
 
 		ShaderPreloadEntry.PreloadEvent = FGraphEvent::CreateGraphEvent();
 		FIoBatch IoBatch = IoDispatcher.NewBatch();
@@ -1223,6 +1343,7 @@ bool FIoStoreShaderCodeArchive::PreloadShaderMap(int32 ShaderMapIndex, FGraphEve
 		{
 			check(!ShaderPreloadEntry.PreloadEvent);
 			ShaderPreloadEntry.FramePreloadStarted = FrameNumber;
+			DebugVisualizer.MarkPreloadedForVisualization(ShaderIndex);
 			PreloadMemory += ShaderEntry.CompressedSize;
 			++PreloadCount;
 			ShaderPreloadEntry.PreloadEvent = FGraphEvent::CreateGraphEvent();
@@ -1270,6 +1391,7 @@ bool FIoStoreShaderCodeArchive::PreloadShaderMap(int32 ShaderMapIndex, FCoreDele
 		{
 			check(!ShaderPreloadEntry.PreloadEvent);
 			ShaderPreloadEntry.FramePreloadStarted = FrameNumber;
+			DebugVisualizer.MarkPreloadedForVisualization(ShaderIndex);
 			PreloadMemory += ShaderEntry.CompressedSize;
 			ShaderPreloadEntry.PreloadEvent = FGraphEvent::CreateGraphEvent();
 			ShaderPreloadEntry.IoRequest = AttachShaderReadRequestFunc(GetShaderCodeChunkId(ShaderHashes[ShaderIndex]), ShaderPreloadEntry.PreloadEvent);
@@ -1397,6 +1519,7 @@ TRefCountPtr<FRHIShader> FIoStoreShaderCodeArchive::CreateShader(int32 Index)
 		break;
 	default: checkNoEntry(); break;
 	}
+	DebugVisualizer.MarkCreatedForVisualization(Index);
 
 	{
 		FWriteScopeLock Lock(ShaderPreloadLock);
diff --git a/Engine/Source/Runtime/RenderCore/Public/ShaderCodeArchive.h b/Engine/Source/Runtime/RenderCore/Public/ShaderCodeArchive.h
--- a/Engine/Source/Runtime/RenderCore/Public/ShaderCodeArchive.h
+++ b/Engine/Source/Runtime/RenderCore/Public/ShaderCodeArchive.h
@@ -9,16 +9,29 @@
 #include "Shader.h"
 #include "IO/IoDispatcher.h"
 
-struct FShaderMapEntry
+// enable visualization in the desktop Development builds only as it has a memory hit and writes files
+#define UE_SCA_VISUALIZE_SHADER_USAGE			(!WITH_EDITOR && UE_BUILD_DEVELOPMENT && PLATFORM_DESKTOP)
+
+struct FIoStoreShaderMapEntry
 {
 	uint32 ShaderIndicesOffset = 0u;
 	uint32 NumShaders = 0u;
+
+	friend FArchive& operator <<(FArchive& Ar, FIoStoreShaderMapEntry& Ref)
+	{
+		return Ar << Ref.ShaderIndicesOffset << Ref.NumShaders;
+	}
+};
+
+struct FShaderMapEntry : public FIoStoreShaderMapEntry
+{
 	uint32 FirstPreloadIndex = 0u;
 	uint32 NumPreloadEntries = 0u;
 
 	friend FArchive& operator <<(FArchive& Ar, FShaderMapEntry& Ref)
 	{
-		return Ar << Ref.ShaderIndicesOffset << Ref.NumShaders << Ref.FirstPreloadIndex << Ref.NumPreloadEntries;
+		operator<<(Ar, static_cast<FIoStoreShaderMapEntry&>(Ref));
+		return Ar << Ref.FirstPreloadIndex << Ref.NumPreloadEntries;
 	}
 };
 
@@ -125,7 +138,7 @@ public:
 			ShaderIndices.GetAllocatedSize()
 #if WITH_EDITOR
 			+ ShaderCodeToAssets.GetAllocatedSize()
-#endif
+#endif // WITH_EDITOR
 			;
 	}
 
@@ -141,7 +154,7 @@ public:
 		ShaderHashTable.Clear();
 #if WITH_EDITOR
 		ShaderCodeToAssets.Empty();
-#endif
+#endif // WITH_EDITOR
 	}
 
 	int32 GetNumShaderMaps() const
@@ -181,6 +194,53 @@ public:
 	}
 };
 
+// run-time only debugging facility
+struct FShaderUsageVisualizer
+{
+#if UE_SCA_VISUALIZE_SHADER_USAGE
+	/** Lock guarding access to visualization structures. */
+	FCriticalSection VisualizeLock;
+
+	/** Total number of shaders. */
+	int32 NumShaders;
+
+	/** Shader indices that were preloaded. */
+	TSet<int32> PreloadedShadersForVis;
+
+	/** Shader indices that were created. */
+	TSet<int32> CreatedShadersForVis;
+
+	void Initialize(const int32 InNumShaders);
+
+	void MarkPreloadedForVisualization(int32 ShaderIndex)
+	{
+		extern int32 GShaderCodeLibraryVisualizeShaderUsage;
+		if (LIKELY(GShaderCodeLibraryVisualizeShaderUsage))
+		{
+			FScopeLock Lock(&VisualizeLock);
+			PreloadedShadersForVis.Add(ShaderIndex);
+		}
+	}
+
+	void MarkCreatedForVisualization(int32 ShaderIndex)
+	{
+		extern int32 GShaderCodeLibraryVisualizeShaderUsage;
+		if (LIKELY(GShaderCodeLibraryVisualizeShaderUsage))
+		{
+			FScopeLock Lock(&VisualizeLock);
+			CreatedShadersForVis.Add(ShaderIndex);
+		}
+	}
+
+	void SaveShaderUsageBitmap(const FString& Name, EShaderPlatform ShaderPlatform);
+#else
+	inline void Initialize(const int32 InNumShaders) {}
+	inline void MarkPreloadedForVisualization(int32 ShaderIndex) {}
+	inline void MarkCreatedForVisualization(int32 ShaderIndex) {}
+	inline void SaveShaderUsageBitmap(const FString& Name, EShaderPlatform ShaderPlatform) {}
+#endif // UE_SCA_VISUALIZE_SHADER_USAGE
+};
+
 class FShaderCodeArchive : public FRHIShaderLibrary
 {
 public:
@@ -264,21 +324,11 @@ protected:
 	// The shader code present in the library
 	FSerializedShaderArchive SerializedShaders;
 
-	TArray<FGraphEventRef> ShaderMapPreloadEvents;
-
 	TArray<FShaderPreloadEntry> ShaderPreloads;
 	FRWLock ShaderPreloadLock;
-};
 
-struct FIoStoreShaderMapEntry
-{
-	uint32 ShaderIndicesOffset = 0u;
-	uint32 NumShaders = 0u;
-
-	friend FArchive& operator <<(FArchive& Ar, FIoStoreShaderMapEntry& Ref)
-	{
-		return Ar << Ref.ShaderIndicesOffset << Ref.NumShaders;
-	}
+	/** debug visualizer - in Shipping compiles out to an empty struct with no-op functions */
+	FShaderUsageVisualizer DebugVisualizer;
 };
 
 struct FIoStoreShaderCodeEntry
@@ -372,5 +422,7 @@ private:
 	FHashTable ShaderHashTable;
 	TArray<FShaderPreloadEntry> ShaderPreloads;
 	FRWLock ShaderPreloadLock;
-};
 
+	/** debug visualizer - in Shipping compiles out to an empty struct with no-op functions */
+	FShaderUsageVisualizer DebugVisualizer;
+};
diff --git a/Engine/Source/Runtime/RenderCore/RenderCore.Build.cs b/Engine/Source/Runtime/RenderCore/RenderCore.Build.cs
--- a/Engine/Source/Runtime/RenderCore/RenderCore.Build.cs
+++ b/Engine/Source/Runtime/RenderCore/RenderCore.Build.cs
@@ -3,6 +3,7 @@
 using System.IO;
 using UnrealBuildBase;
 using UnrealBuildTool;
+using System;
 
 public class RenderCore : ModuleRules
 {
@@ -25,6 +26,11 @@ public class RenderCore : ModuleRules
 			// UObjects are used to produce the full path of the asset by which the shaders are identified
 			PrivateDependencyModuleNames.Add("CoreUObject");
 		}
+		// shader runtime usage visualization requires ability to create images - it is only used in non-editor desktop development targets
+		else if (Target.Configuration == UnrealTargetConfiguration.Development && Array.IndexOf(Utils.GetPlatformsInClass(UnrealPlatformClass.Desktop), Target.Platform) >= 0)
+		{
+			PrivateDependencyModuleNames.Add("ImageWrapper");
+		}
 
 		// Copy the GPUDumpViewer's source code for the r.DumpGPU command.
 		if (Target.Configuration != UnrealTargetConfiguration.Shipping)
