Fix non-unity CIS (ShaderCodeArchive.cpp)

diff --git a/Engine/Source/Runtime/RenderCore/Private/ShaderCodeArchive.cpp b/Engine/Source/Runtime/RenderCore/Private/ShaderCodeArchive.cpp
--- a/Engine/Source/Runtime/RenderCore/Private/ShaderCodeArchive.cpp
+++ b/Engine/Source/Runtime/RenderCore/Private/ShaderCodeArchive.cpp
@@ -17,6 +17,7 @@
 #endif
 
 #if UE_SCA_VISUALIZE_SHADER_USAGE
+#include "Modules/ModuleManager.h"
 #include "IImageWrapperModule.h"
 #include "IImageWrapper.h"
 #endif // UE_SCA_VISUALIZE_SHADER_USAGE
