Change shader compression to Mermaid6 to increase memory savings (UE-136845).

diff --git a/Engine/Source/Runtime/RenderCore/Private/ShaderResource.cpp b/Engine/Source/Runtime/RenderCore/Private/ShaderResource.cpp
--- a/Engine/Source/Runtime/RenderCore/Private/ShaderResource.cpp
+++ b/Engine/Source/Runtime/RenderCore/Private/ShaderResource.cpp
@@ -41,7 +41,7 @@ static FAutoConsoleVariableRef CVarShaderCompressionFormatChoice(
 	TEXT(" 3: ZLib\n"),
 	ECVF_ReadOnly);
 
-int32 GShaderCompressionOodleAlgo = 1;
+int32 GShaderCompressionOodleAlgo = 2;
 static FAutoConsoleVariableRef CVarShaderCompressionOodleAlgo(
 	TEXT("r.Shaders.CompressionFormat.Oodle.Algo"),
 	GShaderCompressionOodleAlgo,
