UE5_RELEASE:  MGPU, include GPU mask per pass in DumpGPU.  The GPU mask is added at the front of the pass name, making it easy to visually see in the pass browser, and also to text search for globally.  The mask is only added when running on a system with multiple GPUs present (very few users), and is enabled by default for discoverability.

#rnx

diff --git a/Engine/Source/Runtime/RenderCore/Private/RenderGraphDumpResources.cpp b/Engine/Source/Runtime/RenderCore/Private/RenderGraphDumpResources.cpp
--- a/Engine/Source/Runtime/RenderCore/Private/RenderGraphDumpResources.cpp
+++ b/Engine/Source/Runtime/RenderCore/Private/RenderGraphDumpResources.cpp
@@ -54,6 +54,11 @@ static TAutoConsoleVariable<int32> GDumpGPUDraws(
 	TEXT("Whether to dump resource after each individual draw call (disabled by default)."),
 	ECVF_Default);
 
+static TAutoConsoleVariable<int32> GDumpGPUMask(
+	TEXT("r.DumpGPU.Mask"), 1,
+	TEXT("Whether to include GPU mask in the name of each Pass (has no effect unless system has multiple GPUs)."),
+	ECVF_Default);
+
 static TAutoConsoleVariable<int32> GDumpExploreCVar(
 	TEXT("r.DumpGPU.Explore"), 1,
 	TEXT("Whether to open file explorer to where the GPU dump on completion (enabled by default)."),
@@ -1341,6 +1346,24 @@ void FRDGBuilder::EndResourceDump()
 	bIsDumpingFrame_GameThread = false;
 }
 
+static const TCHAR* GetPassEventNameWithGPUMask(const FRDGPass* Pass, FString& OutNameStorage)
+{
+#if WITH_MGPU
+	if ((GNumExplicitGPUsForRendering > 1) && GDumpGPUMask.GetValueOnRenderThread())
+	{
+		// Prepend GPU mask on the event name of each pass, so you can see which GPUs the pass ran on.  Putting the mask at the
+		// front rather than the back makes all the masks line up, and easier to read (or ignore if you don't care about them).
+		// Also, it's easy to globally search for passes with a particular GPU mask using name search in the dump browser.
+		OutNameStorage = FString::Printf(TEXT("[%x] %s"), Pass->GetGPUMask().GetNative(), Pass->GetEventName().GetTCHAR());
+		return *OutNameStorage;
+	}
+	else
+#endif  // WITH_MGPU
+	{
+		return Pass->GetEventName().GetTCHAR();
+	}
+}
+
 void FRDGBuilder::DumpResourcePassOutputs(const FRDGPass* Pass)
 {
 	if (bInDebugPassScope)
@@ -1539,8 +1562,10 @@ void FRDGBuilder::DumpResourcePassOutputs(const FRDGPass* Pass)
 			ParentEventScopeNames.Add(MakeShareable(new FJsonValueString(FString::Printf(TEXT("Frame %llu"), GFrameCounterRenderThread))));
 		}
 
+		FString EventNameStorage;
+
 		TSharedPtr<FJsonObject> JsonObject = MakeShareable(new FJsonObject);
-		JsonObject->SetStringField(TEXT("EventName"), Pass->GetEventName().GetTCHAR());
+		JsonObject->SetStringField(TEXT("EventName"), GetPassEventNameWithGPUMask(Pass, EventNameStorage));
 		JsonObject->SetStringField(TEXT("ParametersName"), Pass->GetParameters().GetLayout().GetDebugName());
 		JsonObject->SetStringField(TEXT("Parameters"), FRDGResourceDumpContext::PtrToString(Pass->GetParameters().GetContents()));
 		JsonObject->SetStringField(TEXT("ParametersMetadata"), FRDGResourceDumpContext::PtrToString(Pass->GetParameters().GetMetadata()));
@@ -1740,8 +1765,10 @@ void FRDGBuilder::EndPassDump(const FRDGPass* Pass)
 	// Output how many draw has been dump for this pass.
 	if (GRDGResourceDumpContext.DrawDumpCount > 0)
 	{
+		FString EventNameStorage;
+
 		TSharedPtr<FJsonObject> JsonObject = MakeShareable(new FJsonObject);
-		JsonObject->SetStringField(TEXT("EventName"), Pass->GetEventName().GetTCHAR());
+		JsonObject->SetStringField(TEXT("EventName"), GetPassEventNameWithGPUMask(Pass, EventNameStorage));
 		JsonObject->SetStringField(TEXT("Pointer"), FString::Printf(TEXT("%016x"), FRDGResourceDumpContext::PtrToUint(Pass)));
 		JsonObject->SetNumberField(TEXT("DrawCount"), GRDGResourceDumpContext.DrawDumpCount);
 
@@ -1774,4 +1801,4 @@ bool IsDumpingRDGResources()
 	return false;
 }
 
-#endif //! RDG_DUMP_RESOURCES
\ No newline at end of file
+#endif //! RDG_DUMP_RESOURCES
diff --git a/Engine/Source/Runtime/RenderCore/Public/RenderGraphPass.h b/Engine/Source/Runtime/RenderCore/Public/RenderGraphPass.h
--- a/Engine/Source/Runtime/RenderCore/Public/RenderGraphPass.h
+++ b/Engine/Source/Runtime/RenderCore/Public/RenderGraphPass.h
@@ -297,6 +297,13 @@ public:
 	}
 #endif
 
+#if WITH_MGPU
+	FRHIGPUMask GetGPUMask() const
+	{
+		return GPUMask;
+	}
+#endif
+
 protected:
 	FRDGBarrierBatchBegin& GetPrologueBarriersToBegin(FRDGAllocator& Allocator, FRDGTransitionCreateQueue& CreateQueue);
 	FRDGBarrierBatchBegin& GetEpilogueBarriersToBeginForGraphics(FRDGAllocator& Allocator, FRDGTransitionCreateQueue& CreateQueue);
