Type truncation fixes - Renderer

diff --git a/Engine/Source/Runtime/RHI/Public/DynamicRHI.h b/Engine/Source/Runtime/RHI/Public/DynamicRHI.h
--- a/Engine/Source/Runtime/RHI/Public/DynamicRHI.h
+++ b/Engine/Source/Runtime/RHI/Public/DynamicRHI.h
@@ -1795,7 +1795,7 @@ RHI_API FRenderQueryPoolRHIRef RHICreateRenderQueryPool(ERenderQueryType QueryTy
 FORCEINLINE const FRHITransition* RHICreateTransition(const FRHITransitionCreateInfo& CreateInfo)
 {
 	// Placement create the transition on the heap.
-	FRHITransition* Transition = new (FMemory::Malloc(FRHITransition::GetTotalAllocationSize(), FRHITransition::GetAlignment())) FRHITransition(CreateInfo.SrcPipelines, CreateInfo.DstPipelines);
+	FRHITransition* Transition = new (FMemory::Malloc(FRHITransition::GetTotalAllocationSize(), (uint32)FRHITransition::GetAlignment())) FRHITransition(CreateInfo.SrcPipelines, CreateInfo.DstPipelines);
 	GDynamicRHI->RHICreateTransition(Transition, CreateInfo);
 	return Transition;
 }
diff --git a/Engine/Source/Runtime/RHI/Public/RHICommandList.h b/Engine/Source/Runtime/RHI/Public/RHICommandList.h
--- a/Engine/Source/Runtime/RHI/Public/RHICommandList.h
+++ b/Engine/Source/Runtime/RHI/Public/RHICommandList.h
@@ -428,7 +428,7 @@ public:
 
 	void SetCurrentStat(TStatId Stat);
 
-	FORCEINLINE_DEBUGGABLE void* Alloc(int32 AllocSize, int32 Alignment)
+	FORCEINLINE_DEBUGGABLE void* Alloc(int64 AllocSize, int64 Alignment)
 	{
 		checkSlow(!Bypass() && "Can't use RHICommandList in bypass mode.");
 		return MemManager.Alloc(AllocSize, Alignment);
@@ -3377,7 +3377,7 @@ public:
 
 			if (GraphicsPSOInit.RenderTargetFormats[i] != PF_Unknown)
 			{
-				GraphicsPSOInit.NumSamples = PSOContext.CachedRenderTargets[i].Texture->GetNumSamples();
+				GraphicsPSOInit.NumSamples = static_cast<uint16>(PSOContext.CachedRenderTargets[i].Texture->GetNumSamples());
 			}
 		}
 
@@ -3400,7 +3400,7 @@ public:
 
 		if (GraphicsPSOInit.DepthStencilTargetFormat != PF_Unknown)
 		{
-			GraphicsPSOInit.NumSamples = PSOContext.CachedDepthStencilTarget.Texture->GetNumSamples();
+			GraphicsPSOInit.NumSamples =  static_cast<uint16>(PSOContext.CachedDepthStencilTarget.Texture->GetNumSamples());
 		}
 
 		GraphicsPSOInit.SubpassHint = PSOContext.SubpassHint;
diff --git a/Engine/Source/Runtime/RHI/Public/RHIResources.h b/Engine/Source/Runtime/RHI/Public/RHIResources.h
--- a/Engine/Source/Runtime/RHI/Public/RHIResources.h
+++ b/Engine/Source/Runtime/RHI/Public/RHIResources.h
@@ -3312,7 +3312,7 @@ struct FRHIRenderPassInfo
 				break;
 			}
 
-			RenderTargetsInfo.RenderTargetFormats[RenderTargetIndex] = RenderTarget->GetFormat();
+			RenderTargetsInfo.RenderTargetFormats[RenderTargetIndex] = (uint8)RenderTarget->GetFormat();
 			RenderTargetsInfo.RenderTargetFlags[RenderTargetIndex] = RenderTarget->GetFlags();
 			RenderTargetsInfo.NumSamples |= RenderTarget->GetNumSamples();
 		}
@@ -3402,7 +3402,8 @@ struct RHI_API FRHITextureCreateInfo
 		ETextureCreateFlags InFlags,
 		uint8 InNumMips = 1)
 	{
-		return FRHITextureCreateInfo(ETextureDimension::Texture3D, InFlags, InFormat, FIntPoint(InSize.X, InSize.Y), InClearValue, InSize.Z, 1, InNumMips);
+		checkf(InSize.Z >= 0 && InSize.Z <= TNumericLimits<uint16>::Max(), TEXT("Depth parameter (InSize.Z) exceeds valid range"));
+		return FRHITextureCreateInfo(ETextureDimension::Texture3D, InFlags, InFormat, FIntPoint(InSize.X, InSize.Y), InClearValue, (uint16)InSize.Z, 1, InNumMips);
 	}
 
 	static FRHITextureCreateInfo CreateCube(
diff --git a/Engine/Source/Runtime/RenderCore/Private/DynamicBufferAllocator.cpp b/Engine/Source/Runtime/RenderCore/Private/DynamicBufferAllocator.cpp
--- a/Engine/Source/Runtime/RenderCore/Private/DynamicBufferAllocator.cpp
+++ b/Engine/Source/Runtime/RenderCore/Private/DynamicBufferAllocator.cpp
@@ -109,7 +109,7 @@ FGlobalDynamicReadBuffer::FAllocation FGlobalDynamicReadBuffer_AllocateInternal(
 	uint32 SizeInBytes = sizeof(Type) * Num;
 	FDynamicAllocReadBuffer* Buffer = BufferPool->CurrentBuffer;
 
-	uint32 BufferAlignment = RHIGetMinimumAlignmentForBufferBackedSRV(Format);
+	uint64 BufferAlignment = RHIGetMinimumAlignmentForBufferBackedSRV(Format);
 	uint32 ByteOffset = Buffer == nullptr ? 0 : Align(Buffer->AllocatedByteCount, BufferAlignment);
 
 	if (Buffer == nullptr || ByteOffset + SizeInBytes > Buffer->NumBytes)
diff --git a/Engine/Source/Runtime/RenderCore/Private/GenerateMips.cpp b/Engine/Source/Runtime/RenderCore/Private/GenerateMips.cpp
--- a/Engine/Source/Runtime/RenderCore/Private/GenerateMips.cpp
+++ b/Engine/Source/Runtime/RenderCore/Private/GenerateMips.cpp
@@ -156,7 +156,7 @@ void FGenerateMips::ExecuteRaster(FRDGBuilder& GraphBuilder, FRDGTextureRef Text
 
 	const FRDGTextureDesc& TextureDesc = Texture->Desc;
 
-	for (uint32 MipLevel = 1, MipCount = TextureDesc.NumMips; MipLevel < MipCount; ++MipLevel)
+	for (uint8 MipLevel = 1, MipCount = TextureDesc.NumMips; MipLevel < MipCount; ++MipLevel)
 	{
 		const uint32 InputMipLevel = MipLevel - 1;
 
@@ -213,7 +213,7 @@ void FGenerateMips::ExecuteCompute(FRDGBuilder& GraphBuilder, FRDGTextureRef Tex
 	TShaderMapRef<FGenerateMipsCS> ComputeShader(GetGlobalShaderMap(GMaxRHIFeatureLevel), PermutationVector);
 
 	// Loop through each level of the mips that require creation and add a dispatch pass per level.
-	for (uint32 MipLevel = 1, MipCount = TextureDesc.NumMips; MipLevel < MipCount; ++MipLevel)
+	for (uint8 MipLevel = 1, MipCount = TextureDesc.NumMips; MipLevel < MipCount; ++MipLevel)
 	{
 		const FIntPoint DestTextureSize(
 			FMath::Max(TextureDesc.Extent.X >> MipLevel, 1),
@@ -275,7 +275,7 @@ void FGenerateMips::ExecuteCompute(FRDGBuilder& GraphBuilder, FRDGTextureRef Tex
 	TShaderMapRef<FGenerateMipsIndirectCS> ComputeShader(GetGlobalShaderMap(GMaxRHIFeatureLevel), PermutationVector);
 
 	// Loop through each level of the mips that require creation and add a dispatch pass per level.
-	for (uint32 MipLevel = 1, MipCount = TextureDesc.NumMips; MipLevel < MipCount; ++MipLevel)
+	for (uint8 MipLevel = 1, MipCount = TextureDesc.NumMips; MipLevel < MipCount; ++MipLevel)
 	{
 		const FIntPoint DestTextureSize(
 			FMath::Max(TextureDesc.Extent.X >> MipLevel, 1),
diff --git a/Engine/Source/Runtime/RenderCore/Private/MediaShaders.cpp b/Engine/Source/Runtime/RenderCore/Private/MediaShaders.cpp
--- a/Engine/Source/Runtime/RenderCore/Private/MediaShaders.cpp
+++ b/Engine/Source/Runtime/RenderCore/Private/MediaShaders.cpp
@@ -588,7 +588,7 @@ void FRGB10toYUVv210ConvertPS::SetParameters(FRHICommandList& CommandList, TRefC
 		//Output texture will be based on a size dividable by 48 (i.e 1280 -> 1296) and divided by 6 (i.e 1296 / 6 = 216)
 		//To map output texture UVs, we get a scale from the source texture original size to the mapped output size
 		//And use the source texture size to get the pixel delta
-		const float PaddedResolution = uint32((RGBATexture->GetSizeX() + 47) / 48) * 48;
+		const float PaddedResolution = float(uint32((RGBATexture->GetSizeX() + 47) / 48) * 48);
 		UB.OnePixelDeltaX = 1.0f / (float)RGBATexture->GetSizeX();
 		UB.PaddingScale = PaddedResolution / (float)RGBATexture->GetSizeX();
 	}
diff --git a/Engine/Source/Runtime/RenderCore/Private/ProfilingDebugging/RealtimeGPUProfiler.cpp b/Engine/Source/Runtime/RenderCore/Private/ProfilingDebugging/RealtimeGPUProfiler.cpp
--- a/Engine/Source/Runtime/RenderCore/Private/ProfilingDebugging/RealtimeGPUProfiler.cpp
+++ b/Engine/Source/Runtime/RenderCore/Private/ProfilingDebugging/RealtimeGPUProfiler.cpp
@@ -602,7 +602,7 @@ PRAGMA_ENABLE_DEPRECATION_WARNINGS
 				{
 					// Align CPU and GPU frames
 					Timestamp.GPUMicroseconds = GpuProfilerEvents[1].GetStartResultMicroseconds(GPUIndex);
-					Timestamp.CPUMicroseconds = FPlatformTime::ToSeconds64(CPUFrameStartTimestamp) * 1000 * 1000;
+					Timestamp.CPUMicroseconds = static_cast<uint64>(FPlatformTime::ToSeconds64(CPUFrameStartTimestamp) * 1000 * 1000);
 				}
 				else
 				{
diff --git a/Engine/Source/Runtime/RenderCore/Private/RenderCore.cpp b/Engine/Source/Runtime/RenderCore/Private/RenderCore.cpp
--- a/Engine/Source/Runtime/RenderCore/Private/RenderCore.cpp
+++ b/Engine/Source/Runtime/RenderCore/Private/RenderCore.cpp
@@ -243,7 +243,7 @@ void FInputLatencyTimer::GameThreadTick()
 			LastCaptureTime	= FPlatformTime::Seconds();
 			bInitialized	= true;
 		}
-		float CurrentTimeInSeconds = FPlatformTime::Seconds();
+		const double CurrentTimeInSeconds = FPlatformTime::Seconds();
 		if ( (CurrentTimeInSeconds - LastCaptureTime) > UpdateFrequency )
 		{
 			LastCaptureTime		= CurrentTimeInSeconds;
diff --git a/Engine/Source/Runtime/RenderCore/Private/RenderGraphBuilder.cpp b/Engine/Source/Runtime/RenderCore/Private/RenderGraphBuilder.cpp
--- a/Engine/Source/Runtime/RenderCore/Private/RenderGraphBuilder.cpp
+++ b/Engine/Source/Runtime/RenderCore/Private/RenderGraphBuilder.cpp
@@ -390,10 +390,10 @@ void FRDGBuilder::TickPoolElements()
 	SET_DWORD_STAT(STAT_RDG_PassDependencyCount, GRDGStatPassDependencyCount);
 	SET_DWORD_STAT(STAT_RDG_TextureCount, GRDGStatTextureCount);
 	SET_DWORD_STAT(STAT_RDG_TextureReferenceCount, GRDGStatTextureReferenceCount);
-	SET_FLOAT_STAT(STAT_RDG_TextureReferenceAverage, (float)(GRDGStatTextureReferenceCount / FMath::Max<float>(GRDGStatTextureCount, 1.0f)));
+	SET_FLOAT_STAT(STAT_RDG_TextureReferenceAverage, (float)(GRDGStatTextureReferenceCount / FMath::Max((float)GRDGStatTextureCount, 1.0f)));
 	SET_DWORD_STAT(STAT_RDG_BufferCount, GRDGStatBufferCount);
 	SET_DWORD_STAT(STAT_RDG_BufferReferenceCount, GRDGStatBufferReferenceCount);
-	SET_FLOAT_STAT(STAT_RDG_BufferReferenceAverage, (float)(GRDGStatBufferReferenceCount / FMath::Max<float>(GRDGStatBufferCount, 1.0f)));
+	SET_FLOAT_STAT(STAT_RDG_BufferReferenceAverage, (float)(GRDGStatBufferReferenceCount / FMath::Max((float)GRDGStatBufferCount, 1.0f)));
 	SET_DWORD_STAT(STAT_RDG_ViewCount, GRDGStatViewCount);
 	SET_DWORD_STAT(STAT_RDG_TransientTextureCount, GRDGStatTransientTextureCount);
 	SET_DWORD_STAT(STAT_RDG_TransientBufferCount, GRDGStatTransientBufferCount);
@@ -1919,7 +1919,7 @@ FRDGPass* FRDGBuilder::SetupPass(FRDGPass* Pass)
 		if (Texture->LastPass != PassHandle)
 		{
 			Texture->LastPass = PassHandle;
-			Texture->PassStateIndex = Pass->TextureStates.Num();
+			Texture->PassStateIndex = static_cast<uint16>(Pass->TextureStates.Num());
 
 			PassState = &Pass->TextureStates.Emplace_GetRef(Texture);
 		}
diff --git a/Engine/Source/Runtime/RenderCore/Private/RenderGraphDumpResources.cpp b/Engine/Source/Runtime/RenderCore/Private/RenderGraphDumpResources.cpp
--- a/Engine/Source/Runtime/RenderCore/Private/RenderGraphDumpResources.cpp
+++ b/Engine/Source/Runtime/RenderCore/Private/RenderGraphDumpResources.cpp
@@ -709,7 +709,7 @@ struct FRDGResourceDumpContext
 				StagingSrcTexture = RHICreateTexture2D(
 					SubresourceDumpDesc.SubResourceExtent.X,
 					SubresourceDumpDesc.SubResourceExtent.Y,
-					SubresourceDumpDesc.PreprocessedPixelFormat,
+					(uint8)SubresourceDumpDesc.PreprocessedPixelFormat,
 					/* NumMips = */ 1,
 					/* NumSamples = */ 1,
 					TexCreate_UAV | TexCreate_ShaderResource | TexCreate_HideInVisualizeTexture,
@@ -747,7 +747,7 @@ struct FRDGResourceDumpContext
 			StagingTexture = RHICreateTexture2D(
 				SubresourceDumpDesc.SubResourceExtent.X,
 				SubresourceDumpDesc.SubResourceExtent.Y,
-				SubresourceDumpDesc.PreprocessedPixelFormat,
+				(uint8)SubresourceDumpDesc.PreprocessedPixelFormat,
 				/* NumMips = */ 1,
 				/* NumSamples = */ 1,
 				TexCreate_CPUReadback | TexCreate_HideInVisualizeTexture,
diff --git a/Engine/Source/Runtime/RenderCore/Private/RenderGraphTrace.cpp b/Engine/Source/Runtime/RenderCore/Private/RenderGraphTrace.cpp
--- a/Engine/Source/Runtime/RenderCore/Private/RenderGraphTrace.cpp
+++ b/Engine/Source/Runtime/RenderCore/Private/RenderGraphTrace.cpp
@@ -182,7 +182,7 @@ void FRDGTrace::OutputGraphEnd(const FRDGBuilder& GraphBuilder)
 			const TCHAR* Name{};
 			FRDGPassHandle FirstPass;
 			FRDGPassHandle LastPass;
-			uint32 Depth{};
+			uint16 Depth{};
 		};
 		TArray<FScopeInfo> Scopes;
 		TMap<const FRDGEventScope*, int32> ScopeToIndex;
@@ -207,7 +207,8 @@ void FRDGTrace::OutputGraphEnd(const FRDGBuilder& GraphBuilder)
 							FScopeInfo ScopeInfo;
 							ScopeInfo.Name = Op.Scope->Name.GetTCHAR();
 							ScopeInfo.FirstPass = Handle;
-							ScopeInfo.Depth = Depth;
+							check(Depth >= 0 && Depth <= TNumericLimits<uint16>::Max());
+							ScopeInfo.Depth = static_cast<uint16>(Depth);
 							Scopes.Add(ScopeInfo);
 
 							Depth++;
diff --git a/Engine/Source/Runtime/RenderCore/Private/RenderGraphUtils.cpp b/Engine/Source/Runtime/RenderCore/Private/RenderGraphUtils.cpp
--- a/Engine/Source/Runtime/RenderCore/Private/RenderGraphUtils.cpp
+++ b/Engine/Source/Runtime/RenderCore/Private/RenderGraphUtils.cpp
@@ -579,7 +579,7 @@ void AddClearRenderTargetPass(FRDGBuilder& GraphBuilder, FRDGTextureRef Texture,
 		ERDGPassFlags::Raster,
 		[Parameters, ClearColor, Viewport](FRHICommandList& RHICmdList)
 	{
-		RHICmdList.SetViewport(Viewport.Min.X, Viewport.Min.Y, 0.0f, Viewport.Max.X, Viewport.Max.Y, 1.0f);
+		RHICmdList.SetViewport((float)Viewport.Min.X, (float)Viewport.Min.Y, 0.0f, (float)Viewport.Max.X, (float)Viewport.Max.Y, 1.0f);
 		DrawClearQuad(RHICmdList, ClearColor);
 	});
 }
diff --git a/Engine/Source/Runtime/RenderCore/Private/RenderTargetPool.cpp b/Engine/Source/Runtime/RenderCore/Private/RenderTargetPool.cpp
--- a/Engine/Source/Runtime/RenderCore/Private/RenderTargetPool.cpp
+++ b/Engine/Source/Runtime/RenderCore/Private/RenderTargetPool.cpp
@@ -365,7 +365,7 @@ Done:
 					RHICreateTargetableShaderResource2D(
 						Desc.Extent.X,
 						Desc.Extent.Y,
-						Desc.Format,
+						(uint8)Desc.Format,
 						Desc.NumMips,
 						Desc.Flags,
 						Desc.TargetableFlags,
@@ -383,7 +383,7 @@ Done:
 						Desc.Extent.X,
 						Desc.Extent.Y,
 						Desc.ArraySize,
-						Desc.Format,
+						(uint8)Desc.Format,
 						Desc.NumMips,
 						Desc.Flags,
 						Desc.TargetableFlags,
@@ -411,7 +411,7 @@ Done:
 					Desc.Extent.X,
 					Desc.Extent.Y,
 					Desc.Depth,
-					Desc.Format,
+					(uint8)Desc.Format,
 					Desc.NumMips,
 					Desc.Flags | Desc.TargetableFlags,
 					CreateInfo);
@@ -427,7 +427,7 @@ Done:
 					RHICreateTargetableShaderResourceCubeArray(
 						Desc.Extent.X,
 						Desc.ArraySize,
-						Desc.Format,
+						(uint8)Desc.Format,
 						Desc.NumMips,
 						Desc.Flags,
 						Desc.TargetableFlags,
@@ -441,7 +441,7 @@ Done:
 				{
 					RHICreateTargetableShaderResourceCube(
 						Desc.Extent.X,
-						Desc.Format,
+						(uint8)Desc.Format,
 						Desc.NumMips,
 						Desc.Flags,
 						Desc.TargetableFlags,
@@ -467,7 +467,7 @@ Done:
 					Found->RenderTargetItem.ShaderResourceTexture = RHICreateTexture2D(
 						Desc.Extent.X,
 						Desc.Extent.Y,
-						Desc.Format,
+						(uint8)Desc.Format,
 						Desc.NumMips,
 						Desc.NumSamples,
 						Desc.Flags,
@@ -479,7 +479,7 @@ Done:
 						Desc.Extent.X,
 						Desc.Extent.Y,
 						Desc.ArraySize,
-						Desc.Format,
+						(uint8)Desc.Format,
 						Desc.NumMips,
 						Desc.NumSamples,
 						Desc.Flags,
@@ -494,7 +494,7 @@ Done:
 					Desc.Extent.X,
 					Desc.Extent.Y,
 					Desc.Depth,
-					Desc.Format,
+					(uint8)Desc.Format,
 					Desc.NumMips,
 					Desc.Flags,
 					CreateInfo);
@@ -504,12 +504,12 @@ Done:
 				check(Desc.IsCubemap());
 				if (Desc.IsArray())
 				{
-					FTextureCubeRHIRef CubeTexture = RHICreateTextureCubeArray(Desc.Extent.X, Desc.ArraySize, Desc.Format, Desc.NumMips, Desc.Flags | Desc.TargetableFlags | TexCreate_ShaderResource, CreateInfo);
+					FTextureCubeRHIRef CubeTexture = RHICreateTextureCubeArray(Desc.Extent.X, Desc.ArraySize, (uint8)Desc.Format, Desc.NumMips, Desc.Flags | Desc.TargetableFlags | TexCreate_ShaderResource, CreateInfo);
 					Found->RenderTargetItem.TargetableTexture = Found->RenderTargetItem.ShaderResourceTexture = CubeTexture;
 				}
 				else
 				{
-					FTextureCubeRHIRef CubeTexture = RHICreateTextureCube(Desc.Extent.X, Desc.Format, Desc.NumMips, Desc.Flags | Desc.TargetableFlags | TexCreate_ShaderResource, CreateInfo);
+					FTextureCubeRHIRef CubeTexture = RHICreateTextureCube(Desc.Extent.X, (uint8)Desc.Format, Desc.NumMips, Desc.Flags | Desc.TargetableFlags | TexCreate_ShaderResource, CreateInfo);
 					Found->RenderTargetItem.TargetableTexture = Found->RenderTargetItem.ShaderResourceTexture = CubeTexture;
 				}
 			}
@@ -530,7 +530,7 @@ Done:
 					? Desc.UAVFormat
 					: Desc.Format;
 
-				Found->RenderTargetItem.UAV = RHICreateUnorderedAccessView(Found->RenderTargetItem.TargetableTexture, 0, AliasFormat, 0, 0);
+				Found->RenderTargetItem.UAV = RHICreateUnorderedAccessView(Found->RenderTargetItem.TargetableTexture, 0, (uint8)AliasFormat, 0, 0);
 			}
 			else
 			{
diff --git a/Engine/Source/Runtime/RenderCore/Private/RenderUtils.cpp b/Engine/Source/Runtime/RenderCore/Private/RenderUtils.cpp
--- a/Engine/Source/Runtime/RenderCore/Private/RenderUtils.cpp
+++ b/Engine/Source/Runtime/RenderCore/Private/RenderUtils.cpp
@@ -464,7 +464,7 @@ public:
 		const TCHAR* Name = TEXT("SolidColorCube");
 
 		FRHIResourceCreateInfo CreateInfo(Name);
-		FTextureCubeRHIRef TextureCube = RHICreateTextureCube(1, PixelFormat, 1, TexCreate_ShaderResource, CreateInfo);
+		FTextureCubeRHIRef TextureCube = RHICreateTextureCube(1, (uint8)PixelFormat, 1, TexCreate_ShaderResource, CreateInfo);
 		TextureRHI = TextureCube;
 
 		// Write the contents of the texture.
@@ -789,9 +789,9 @@ SIZE_T CalcTextureMipHeightInBlocks(uint32 TextureSizeY, EPixelFormat Format, ui
 
 SIZE_T CalcTextureMipMapSize( uint32 TextureSizeX, uint32 TextureSizeY, EPixelFormat Format, uint32 MipIndex )
 {
-	const uint32 WidthInBlocks = CalcTextureMipWidthInBlocks(TextureSizeX, Format, MipIndex);
-	const uint32 HeightInBlocks = CalcTextureMipHeightInBlocks(TextureSizeY, Format, MipIndex);
-	return static_cast<SIZE_T>(WidthInBlocks) * HeightInBlocks * GPixelFormats[Format].BlockBytes;
+	const SIZE_T WidthInBlocks = CalcTextureMipWidthInBlocks(TextureSizeX, Format, MipIndex);
+	const SIZE_T HeightInBlocks = CalcTextureMipHeightInBlocks(TextureSizeY, Format, MipIndex);
+	return WidthInBlocks * HeightInBlocks * GPixelFormats[Format].BlockBytes;
 }
 
 SIZE_T CalcTextureSize( uint32 SizeX, uint32 SizeY, EPixelFormat Format, uint32 MipCount )
@@ -1547,10 +1547,10 @@ public:
 				for (uint32 X = 0; X < 2; X++)
 				{
 					const FVector4f Vertex = FVector4f(
-					  (X ? -1 : 1),
-					  (Y ? -1 : 1),
-					  (Z ? -1 : 1),
-					  1.0f
+					  (X ? -1.f : 1.f),
+					  (Y ? -1.f : 1.f),
+					  (Z ? -1.f : 1.f),
+					  1.f
 					);
 
 					Verts[GetCubeVertexIndex(X, Y, Z)] = Vertex;
diff --git a/Engine/Source/Runtime/RenderCore/Private/RenderingThread.cpp b/Engine/Source/Runtime/RenderCore/Private/RenderingThread.cpp
--- a/Engine/Source/Runtime/RenderCore/Private/RenderingThread.cpp
+++ b/Engine/Source/Runtime/RenderCore/Private/RenderingThread.cpp
@@ -236,7 +236,7 @@ FTickableObjectRenderThread::FRenderingThreadTickableObjectsArray FTickableObjec
 void TickHighFrequencyTickables(double CurTime)
 {
 	static double LastHighFreqTime = FPlatformTime::Seconds();
-	float DeltaSecondsHighFreq = CurTime - LastHighFreqTime;
+	float DeltaSecondsHighFreq = float(CurTime - LastHighFreqTime);
 
 	// tick any high frequency rendering thread tickables.
 	for (int32 ObjectIndex = 0; ObjectIndex < FTickableObjectRenderThread::RenderingThreadHighFrequencyTickableObjects.Num(); ObjectIndex++)
@@ -260,7 +260,7 @@ void TickRenderingTickables()
 
 	// calc how long has passed since last tick
 	double CurTime = FPlatformTime::Seconds();
-	float DeltaSeconds = CurTime - LastTickTime;
+	float DeltaSeconds = float(CurTime - LastTickTime);
 		
 	TickHighFrequencyTickables(CurTime);
 
diff --git a/Engine/Source/Runtime/RenderCore/Private/ResolveShader.cpp b/Engine/Source/Runtime/RenderCore/Private/ResolveShader.cpp
--- a/Engine/Source/Runtime/RenderCore/Private/ResolveShader.cpp
+++ b/Engine/Source/Runtime/RenderCore/Private/ResolveShader.cpp
@@ -19,10 +19,10 @@ void FResolveSingleSamplePS::SetParameters(FRHICommandList& RHICmdList, uint32 S
 void FResolveVS::SetParameters(FRHICommandList& RHICmdList, const FResolveRect& SrcBounds, const FResolveRect& DstBounds, uint32 DstSurfaceWidth, uint32 DstSurfaceHeight)
 {
 	// Generate the vertices used to copy from the source surface to the destination surface.
-	const float MinU = SrcBounds.X1;
-	const float MinV = SrcBounds.Y1;
-	const float MaxU = SrcBounds.X2;
-	const float MaxV = SrcBounds.Y2;
+	const float MinU = (float)SrcBounds.X1;
+	const float MinV = (float)SrcBounds.Y1;
+	const float MaxU = (float)SrcBounds.X2;
+	const float MaxV = (float)SrcBounds.Y2;
 	const float MinX = -1.f + DstBounds.X1 / ((float)DstSurfaceWidth * 0.5f);
 	const float MinY = +1.f - DstBounds.Y1 / ((float)DstSurfaceHeight * 0.5f);
 	const float MaxX = -1.f + DstBounds.X2 / ((float)DstSurfaceWidth * 0.5f);
diff --git a/Engine/Source/Runtime/RenderCore/Private/ShaderCore.cpp b/Engine/Source/Runtime/RenderCore/Private/ShaderCore.cpp
--- a/Engine/Source/Runtime/RenderCore/Private/ShaderCore.cpp
+++ b/Engine/Source/Runtime/RenderCore/Private/ShaderCore.cpp
@@ -705,7 +705,7 @@ void VerifyShaderSourceFiles(EShaderPlatform ShaderPlatform)
 		// get the list of shader files that can be used
 		TArray<FString> VirtualShaderSourcePaths;
 		GetAllVirtualShaderSourcePaths(VirtualShaderSourcePaths, ShaderPlatform);
-		FScopedSlowTask SlowTask(VirtualShaderSourcePaths.Num());
+		FScopedSlowTask SlowTask((float)VirtualShaderSourcePaths.Num());
 		for( int32 ShaderFileIdx=0; ShaderFileIdx < VirtualShaderSourcePaths.Num(); ShaderFileIdx++ )
 		{
 			SlowTask.EnterProgressFrame(1);
@@ -1259,7 +1259,7 @@ void BuildShaderFileToUniformBufferMap(TMap<FString, TArray<const TCHAR*> >& Sha
 		TArray<FString> ShaderSourceFiles;
 		GetAllVirtualShaderSourcePaths(ShaderSourceFiles, GMaxRHIShaderPlatform);
 
-		FScopedSlowTask SlowTask(ShaderSourceFiles.Num());
+		FScopedSlowTask SlowTask((float)ShaderSourceFiles.Num());
 
 		// Cache UB access strings, make it case sensitive for faster search
 		struct FShaderVariable
diff --git a/Engine/Source/Runtime/RenderCore/Private/ShaderPipelineCache.cpp b/Engine/Source/Runtime/RenderCore/Private/ShaderPipelineCache.cpp
--- a/Engine/Source/Runtime/RenderCore/Private/ShaderPipelineCache.cpp
+++ b/Engine/Source/Runtime/RenderCore/Private/ShaderPipelineCache.cpp
@@ -523,7 +523,7 @@ uint32 FShaderPipelineCache::NumPrecompilesRemaining()
 		if (CVarPSOFileCacheMaxPrecompileTime.GetValueOnAnyThread() > 0.0 && FPlatformAtomics::AtomicRead(&ShaderPipelineCache->TotalPrecompileTasks) > 0)
 		{
 			float Mult = ShaderPipelineCache->TotalPrecompileWallTime / CVarPSOFileCacheMaxPrecompileTime.GetValueOnAnyThread();
-			NumRemaining = FMath::Max(0.f, 1.f - Mult) * FPlatformAtomics::AtomicRead(&ShaderPipelineCache->TotalPrecompileTasks);
+			NumRemaining = uint32(FMath::Max(0.f, 1.f - Mult) * FPlatformAtomics::AtomicRead(&ShaderPipelineCache->TotalPrecompileTasks));
 		}
 		else
 		{
@@ -550,7 +550,7 @@ uint32 FShaderPipelineCache::NumPrecompilesActive()
 		if (CVarPSOFileCacheMaxPrecompileTime.GetValueOnAnyThread() > 0.0 && FPlatformAtomics::AtomicRead(&ShaderPipelineCache->TotalPrecompileTasks) > 0)
 		{
 			float Mult = ShaderPipelineCache->TotalPrecompileWallTime / CVarPSOFileCacheMaxPrecompileTime.GetValueOnAnyThread();
-			NumRemaining = FMath::Max(0.f, 1.f - Mult) * FPlatformAtomics::AtomicRead(&ShaderPipelineCache->TotalPrecompileTasks);
+			NumRemaining = uint32(FMath::Max(0.f, 1.f - Mult) * FPlatformAtomics::AtomicRead(&ShaderPipelineCache->TotalPrecompileTasks));
 		}
     }
     
@@ -1081,7 +1081,7 @@ bool FShaderPipelineCache::ReadyForAutoSave() const
 	uint32 SaveAfterNum = CVarPSOFileCacheSaveAfterPSOsLogged.GetValueOnAnyThread();
 	uint32 NumLogged = FPipelineFileCache::NumPSOsLogged();
 
-	const float TimeSinceSave = FPlatformTime::Seconds() - LastAutoSaveTime;
+	const double TimeSinceSave = FPlatformTime::Seconds() - LastAutoSaveTime;
 
 	// autosave if its enabled, and we have more than the desired number, or it's been a while since our last save
 	if (SaveAfterNum > 0 && 
@@ -1331,7 +1331,7 @@ void FShaderPipelineCache::Tick( float DeltaTime )
 	{
 		if (LastAutoSaveNum < int32(FPipelineFileCache::NumPSOsLogged()))
 		{
-			const float TimeSinceSave = FPlatformTime::Seconds() - LastAutoSaveTimeLogBoundPSO;
+			const double TimeSinceSave = FPlatformTime::Seconds() - LastAutoSaveTimeLogBoundPSO;
 
 			if (TimeSinceSave >= CVarPSOFileCacheAutoSaveTimeBoundPSO.GetValueOnAnyThread())
 			{
diff --git a/Engine/Source/Runtime/RenderCore/Private/ShaderResource.cpp b/Engine/Source/Runtime/RenderCore/Private/ShaderResource.cpp
--- a/Engine/Source/Runtime/RenderCore/Private/ShaderResource.cpp
+++ b/Engine/Source/Runtime/RenderCore/Private/ShaderResource.cpp
@@ -244,12 +244,13 @@ void FShaderMapResourceCode::Finalize()
 
 uint32 FShaderMapResourceCode::GetSizeBytes() const
 {
-	uint32 Size = sizeof(*this) + ShaderHashes.GetAllocatedSize() + ShaderEntries.GetAllocatedSize();
+	uint64 Size = sizeof(*this) + ShaderHashes.GetAllocatedSize() + ShaderEntries.GetAllocatedSize();
 	for (const FShaderEntry& Entry : ShaderEntries)
 	{
 		Size += Entry.Code.GetAllocatedSize();
 	}
-	return Size;
+	check(Size <= TNumericLimits<uint32>::Max());
+	return static_cast<uint32>(Size);
 }
 
 int32 FShaderMapResourceCode::FindShaderIndex(const FSHAHash& InHash) const
diff --git a/Engine/Source/Runtime/RenderCore/Private/UnifiedBuffer.cpp b/Engine/Source/Runtime/RenderCore/Private/UnifiedBuffer.cpp
--- a/Engine/Source/Runtime/RenderCore/Private/UnifiedBuffer.cpp
+++ b/Engine/Source/Runtime/RenderCore/Private/UnifiedBuffer.cpp
@@ -137,7 +137,7 @@ struct ResourceTypeTraits<FRWByteAddressBuffer>
 
 static uint32 CalculateFloat4sPerLine()
 {
-	uint16 PrimitivesPerTextureLine = FMath::Min((int32)MAX_uint16, (int32)GMaxTextureDimensions) / FScatterUploadBuffer::PrimitiveDataStrideInFloat4s;
+	uint16 PrimitivesPerTextureLine = (uint16)FMath::Min((int32)MAX_uint16, (int32)GMaxTextureDimensions) / FScatterUploadBuffer::PrimitiveDataStrideInFloat4s;
 	return PrimitivesPerTextureLine * FScatterUploadBuffer::PrimitiveDataStrideInFloat4s;
 }
 
diff --git a/Engine/Source/Runtime/RenderCore/Private/VertexFactory.cpp b/Engine/Source/Runtime/RenderCore/Private/VertexFactory.cpp
--- a/Engine/Source/Runtime/RenderCore/Private/VertexFactory.cpp
+++ b/Engine/Source/Runtime/RenderCore/Private/VertexFactory.cpp
@@ -298,7 +298,7 @@ FVertexElement FVertexFactory::AccessStreamComponent(const FVertexStreamComponen
 	VertexStream.Offset = Component.StreamOffset;
 	VertexStream.VertexStreamUsage = Component.VertexStreamUsage;
 
-	return FVertexElement(Streams.AddUnique(VertexStream),Component.Offset,Component.Type,AttributeIndex,VertexStream.Stride, EnumHasAnyFlags(EVertexStreamUsage::Instancing, VertexStream.VertexStreamUsage));
+	return FVertexElement((uint8)Streams.AddUnique(VertexStream),Component.Offset,Component.Type,AttributeIndex,VertexStream.Stride, EnumHasAnyFlags(EVertexStreamUsage::Instancing, VertexStream.VertexStreamUsage));
 }
 
 FVertexElement FVertexFactory::AccessStreamComponent(const FVertexStreamComponent& Component,uint8 AttributeIndex, EVertexInputStreamType InputStreamType)
@@ -310,11 +310,11 @@ FVertexElement FVertexFactory::AccessStreamComponent(const FVertexStreamComponen
 	VertexStream.VertexStreamUsage = Component.VertexStreamUsage;
 
 	if (InputStreamType == EVertexInputStreamType::PositionOnly)
-		return FVertexElement(PositionStream.AddUnique(VertexStream), Component.Offset, Component.Type, AttributeIndex, VertexStream.Stride, EnumHasAnyFlags(EVertexStreamUsage::Instancing, VertexStream.VertexStreamUsage));
+		return FVertexElement((uint8)PositionStream.AddUnique(VertexStream), Component.Offset, Component.Type, AttributeIndex, VertexStream.Stride, EnumHasAnyFlags(EVertexStreamUsage::Instancing, VertexStream.VertexStreamUsage));
 	else if (InputStreamType == EVertexInputStreamType::PositionAndNormalOnly)
-		return FVertexElement(PositionAndNormalStream.AddUnique(VertexStream), Component.Offset, Component.Type, AttributeIndex, VertexStream.Stride, EnumHasAnyFlags(EVertexStreamUsage::Instancing, VertexStream.VertexStreamUsage));
+		return FVertexElement((uint8)PositionAndNormalStream.AddUnique(VertexStream), Component.Offset, Component.Type, AttributeIndex, VertexStream.Stride, EnumHasAnyFlags(EVertexStreamUsage::Instancing, VertexStream.VertexStreamUsage));
 	else /* (InputStreamType == EVertexInputStreamType::Default) */
-		return FVertexElement(Streams.AddUnique(VertexStream), Component.Offset, Component.Type, AttributeIndex, VertexStream.Stride, EnumHasAnyFlags(EVertexStreamUsage::Instancing, VertexStream.VertexStreamUsage));
+		return FVertexElement((uint8)Streams.AddUnique(VertexStream), Component.Offset, Component.Type, AttributeIndex, VertexStream.Stride, EnumHasAnyFlags(EVertexStreamUsage::Instancing, VertexStream.VertexStreamUsage));
 }
 
 void FVertexFactory::InitDeclaration(const FVertexDeclarationElementList& Elements, EVertexInputStreamType StreamType)
diff --git a/Engine/Source/Runtime/RenderCore/Public/CommonRenderResources.h b/Engine/Source/Runtime/RenderCore/Public/CommonRenderResources.h
--- a/Engine/Source/Runtime/RenderCore/Public/CommonRenderResources.h
+++ b/Engine/Source/Runtime/RenderCore/Public/CommonRenderResources.h
@@ -32,7 +32,7 @@ public:
 	virtual void InitRHI()
 	{
 		FVertexDeclarationElementList Elements;
-		uint32 Stride = sizeof(FFilterVertex);
+		uint16 Stride = sizeof(FFilterVertex);
 		Elements.Add(FVertexElement(0, STRUCT_OFFSET(FFilterVertex, Position), VET_Float4, 0, Stride));
 		Elements.Add(FVertexElement(0, STRUCT_OFFSET(FFilterVertex, UV), VET_Float2, 1, Stride));
 		VertexDeclarationRHI = PipelineStateCache::GetOrCreateVertexDeclaration(Elements);
diff --git a/Engine/Source/Runtime/RenderCore/Public/PixelShaderUtils.h b/Engine/Source/Runtime/RenderCore/Public/PixelShaderUtils.h
--- a/Engine/Source/Runtime/RenderCore/Public/PixelShaderUtils.h
+++ b/Engine/Source/Runtime/RenderCore/Public/PixelShaderUtils.h
@@ -65,7 +65,7 @@ struct RENDERCORE_API FPixelShaderUtils
 		uint32 StencilRef = 0)
 	{
 		check(PixelShader.IsValid());
-		RHICmdList.SetViewport(Viewport.Min.X, Viewport.Min.Y, 0.0f, Viewport.Max.X, Viewport.Max.Y, 1.0f);
+		RHICmdList.SetViewport((float)Viewport.Min.X, (float)Viewport.Min.Y, 0.0f, (float)Viewport.Max.X, (float)Viewport.Max.Y, 1.0f);
 		
 		FGraphicsPipelineStateInitializer GraphicsPSOInit;
 		InitFullscreenPipelineState(RHICmdList, GlobalShaderMap, PixelShader, /* out */ GraphicsPSOInit);
diff --git a/Engine/Source/Runtime/RenderCore/Public/RenderCore.h b/Engine/Source/Runtime/RenderCore/Public/RenderCore.h
--- a/Engine/Source/Runtime/RenderCore/Public/RenderCore.h
+++ b/Engine/Source/Runtime/RenderCore/Public/RenderCore.h
@@ -257,7 +257,7 @@ struct FInputLatencyTimer
 	uint32	DeltaTime;
 
 	/** Last time we did a measurement (in seconds). */
-	float	LastCaptureTime;
+	double	LastCaptureTime;
 
 	/** How often we should do a measurement (in seconds). */
 	float	UpdateFrequency;
diff --git a/Engine/Source/Runtime/RenderCore/Public/RenderGraphAllocator.h b/Engine/Source/Runtime/RenderCore/Public/RenderGraphAllocator.h
--- a/Engine/Source/Runtime/RenderCore/Public/RenderGraphAllocator.h
+++ b/Engine/Source/Runtime/RenderCore/Public/RenderGraphAllocator.h
@@ -12,14 +12,14 @@ public:
 	~FRDGAllocator();
 
 	/** Allocates raw memory. */
-	FORCEINLINE void* Alloc(uint32 SizeInBytes, uint32 AlignInBytes)
+	FORCEINLINE void* Alloc(uint64 SizeInBytes, uint32 AlignInBytes)
 	{
 		return GetContext().MemStack.Alloc(SizeInBytes, AlignInBytes);
 	}
 
 	/** Allocates an uninitialized type without destructor tracking. */
 	template <typename PODType>
-	FORCEINLINE PODType* AllocUninitialized(uint32 Count = 1)
+	FORCEINLINE PODType* AllocUninitialized(uint64 Count = 1)
 	{
 		return reinterpret_cast<PODType*>(Alloc(sizeof(PODType) * Count, alignof(PODType)));
 	}
diff --git a/Engine/Source/Runtime/RenderCore/Public/RenderGraphBuilder.h b/Engine/Source/Runtime/RenderCore/Public/RenderGraphBuilder.h
--- a/Engine/Source/Runtime/RenderCore/Public/RenderGraphBuilder.h
+++ b/Engine/Source/Runtime/RenderCore/Public/RenderGraphBuilder.h
@@ -115,7 +115,7 @@ public:
 	// Allocation Methods
 
 	/** Allocates raw memory using an allocator tied to the lifetime of the graph. */
-	void* Alloc(uint32 SizeInBytes, uint32 AlignInBytes);
+	void* Alloc(uint64 SizeInBytes, uint32 AlignInBytes);
 
 	/** Allocates POD memory using an allocator tied to the lifetime of the graph. Does not construct / destruct. */
 	template <typename PODType>
diff --git a/Engine/Source/Runtime/RenderCore/Public/RenderGraphBuilder.inl b/Engine/Source/Runtime/RenderCore/Public/RenderGraphBuilder.inl
--- a/Engine/Source/Runtime/RenderCore/Public/RenderGraphBuilder.inl
+++ b/Engine/Source/Runtime/RenderCore/Public/RenderGraphBuilder.inl
@@ -113,7 +113,7 @@ inline FRDGBufferUAVRef FRDGBuilder::CreateUAV(const FRDGBufferUAVDesc& Desc, ER
 	return UAV;
 }
 
-FORCEINLINE void* FRDGBuilder::Alloc(uint32 SizeInBytes, uint32 AlignInBytes)
+FORCEINLINE void* FRDGBuilder::Alloc(uint64 SizeInBytes, uint32 AlignInBytes)
 {
 	return Allocator.Alloc(SizeInBytes, AlignInBytes);
 }
diff --git a/Engine/Source/Runtime/RenderCore/Public/RenderGraphDefinitions.h b/Engine/Source/Runtime/RenderCore/Public/RenderGraphDefinitions.h
--- a/Engine/Source/Runtime/RenderCore/Public/RenderGraphDefinitions.h
+++ b/Engine/Source/Runtime/RenderCore/Public/RenderGraphDefinitions.h
@@ -263,14 +263,14 @@ public:
 	FORCEINLINE TRDGHandle& operator+=(int32 Increment)
 	{
 		check(int64(Index + Increment) <= int64(kNullIndex));
-		Index += Increment;
+		Index += (IndexType)Increment;
 		return *this;
 	}
 
 	FORCEINLINE TRDGHandle& operator-=(int32 Decrement)
 	{
 		check(int64(Index - Decrement) > 0);
-		Index -= Decrement;
+		Index -= (IndexType)Decrement;
 		return *this;
 	}
 
diff --git a/Engine/Source/Runtime/RenderCore/Public/RenderGraphResources.h b/Engine/Source/Runtime/RenderCore/Public/RenderGraphResources.h
--- a/Engine/Source/Runtime/RenderCore/Public/RenderGraphResources.h
+++ b/Engine/Source/Runtime/RenderCore/Public/RenderGraphResources.h
@@ -635,7 +635,8 @@ public:
 	{
 		FRDGTextureSRVDesc Desc;
 		Desc.Texture = Texture;
-		Desc.MipLevel = MipLevel;
+		check(MipLevel >= -1 && MipLevel <= TNumericLimits<int8>::Max()); 
+		Desc.MipLevel = (int8)MipLevel;
 		Desc.NumMipLevels = 1;
 		return Desc;
 	}
diff --git a/Engine/Source/Runtime/RenderCore/Public/RenderGraphResources.inl b/Engine/Source/Runtime/RenderCore/Public/RenderGraphResources.inl
--- a/Engine/Source/Runtime/RenderCore/Public/RenderGraphResources.inl
+++ b/Engine/Source/Runtime/RenderCore/Public/RenderGraphResources.inl
@@ -259,7 +259,7 @@ inline FGraphicsPipelineRenderTargetsInfo ExtractRenderTargetsInfo(const FRDGPar
 	{
 		FRDGTextureRef Texture = RenderTarget.GetTexture();
 
-		RenderTargetsInfo.RenderTargetFormats[RenderTargetIndex] = Texture->Desc.Format;
+		RenderTargetsInfo.RenderTargetFormats[RenderTargetIndex] = (uint8)Texture->Desc.Format;
 		RenderTargetsInfo.RenderTargetFlags[RenderTargetIndex] = Texture->Desc.Flags;
 		RenderTargetsInfo.NumSamples |= Texture->Desc.NumSamples;
 		++RenderTargetIndex;
diff --git a/Engine/Source/Runtime/RenderCore/Public/RenderTransform.h b/Engine/Source/Runtime/RenderCore/Public/RenderTransform.h
--- a/Engine/Source/Runtime/RenderCore/Public/RenderTransform.h
+++ b/Engine/Source/Runtime/RenderCore/Public/RenderTransform.h
@@ -37,10 +37,10 @@ public:
 	FORCEINLINE FRenderTransform(const FMatrix44d& M)
 	{
 		// LWC_TODO: Precision loss
-		TransformRows[0] = FVector3f(M.M[0][0], M.M[0][1], M.M[0][2]);
-		TransformRows[1] = FVector3f(M.M[1][0], M.M[1][1], M.M[1][2]);
-		TransformRows[2] = FVector3f(M.M[2][0], M.M[2][1], M.M[2][2]);
-		Origin = FVector3f(M.M[3][0], M.M[3][1], M.M[3][2]);
+		TransformRows[0] = FVector3f((float)M.M[0][0], (float)M.M[0][1], (float)M.M[0][2]);
+		TransformRows[1] = FVector3f((float)M.M[1][0], (float)M.M[1][1], (float)M.M[1][2]);
+		TransformRows[2] = FVector3f((float)M.M[2][0], (float)M.M[2][1], (float)M.M[2][2]);
+		Origin = FVector3f((float)M.M[3][0], (float)M.M[3][1], (float)M.M[3][2]);
 	}
 
 	FORCEINLINE FRenderTransform& operator=(const FRenderTransform& From)
@@ -63,10 +63,11 @@ public:
 
 	FORCEINLINE FRenderTransform& operator=(const FMatrix44d& From)
 	{
-		TransformRows[0] = FVector3f(From.M[0][0], From.M[0][1], From.M[0][2]);
-		TransformRows[1] = FVector3f(From.M[1][0], From.M[1][1], From.M[1][2]);
-		TransformRows[2] = FVector3f(From.M[2][0], From.M[2][1], From.M[2][2]);
-		Origin = FVector3f(From.M[3][0], From.M[3][1], From.M[3][2]);
+		// LWC_TODO: Precision loss
+		TransformRows[0] = FVector3f((float)From.M[0][0], (float)From.M[0][1], (float)From.M[0][2]);
+		TransformRows[1] = FVector3f((float)From.M[1][0], (float)From.M[1][1], (float)From.M[1][2]);
+		TransformRows[2] = FVector3f((float)From.M[2][0], (float)From.M[2][1], (float)From.M[2][2]);
+		Origin = FVector3f((float)From.M[3][0], (float)From.M[3][1], (float)From.M[3][2]);
 		return *this;
 	}
 
@@ -452,9 +453,9 @@ struct FCompressedTransform
 
 #if 1
 			// Avoid sign extension in shader by biasing
-			Rotation[0] = FMath::RoundToInt( OctZ.X * 32767.0f ) + 32768;
-			Rotation[1] = FMath::RoundToInt( OctZ.Y * 32767.0f ) + 32768;
-			Rotation[2] = FMath::RoundToInt( Spin0  * 16383.0f * 1.41421356f );	// sqrt(2)
+			Rotation[0] = static_cast<uint16>(FMath::RoundToInt( OctZ.X * 32767.0f ) + 32768);
+			Rotation[1] = static_cast<uint16>(FMath::RoundToInt( OctZ.Y * 32767.0f ) + 32768);
+			Rotation[2] = static_cast<uint16>(FMath::RoundToInt( Spin0  * 16383.0f * 1.41421356f ));	// sqrt(2)
 			
 			Rotation[2] = ( Rotation[2] + 16384 ) & 0x7fff;
 			Rotation[2] |= bSpinIsX ? (1 << 15) : 0;
@@ -475,7 +476,7 @@ struct FCompressedTransform
 			// Need +1 because of losing the implicit leading bit of mantissa
 			// TODO assumes ExpBits == 8
 			// TODO clamp to expressable range
-			uint32 SharedExp = MaxComponent.Components.Exponent + 1;
+			uint16 SharedExp = (uint16)MaxComponent.Components.Exponent + 1;
 	
 			FFloat32 ExpScale( 1.0f );
 			ExpScale.Components.Exponent = 127 + ExpBias + MantissaBits - SharedExp;
@@ -498,7 +499,7 @@ struct FCompressedTransform
 			Scale_SharedExp[3] = SharedExp;
 			for( int i = 0; i < 3; i++ )
 			{
-				Scale_SharedExp[i] = FMath::RoundToInt( Scale[i] * ExpScale.FloatValue ) + (1 << MantissaBits);
+				Scale_SharedExp[i] = static_cast<uint16>(FMath::RoundToInt( Scale[i] * ExpScale.FloatValue ) + (1 << MantissaBits));
 			}
 #endif
 		}
@@ -551,7 +552,7 @@ struct FCompressedTransform
 #else
 			for( int i = 0; i < 3; i++ )
 			{
-				float Scale = (int32)Scale_SharedExp[i] - (1 << MantissaBits);
+				float Scale = static_cast<float>(int32(Scale_SharedExp[i]) - (1 << MantissaBits));
 				Scale *= ExpScale.FloatValue;
 				Out.TransformRows[i] *= Scale;
 			}
diff --git a/Engine/Source/Runtime/RenderCore/Public/RendererInterface.h b/Engine/Source/Runtime/RenderCore/Public/RendererInterface.h
--- a/Engine/Source/Runtime/RenderCore/Public/RendererInterface.h
+++ b/Engine/Source/Runtime/RenderCore/Public/RendererInterface.h
@@ -83,7 +83,7 @@ public:
 		ETextureCreateFlags InFlags,
 		ETextureCreateFlags InTargetableFlags,
 		bool bInForceSeparateTargetAndShaderResource,
-		uint16 InNumMips = 1,
+		uint8 InNumMips = 1,
 		bool InAutowritable = true,
 		bool InCreateRTWriteMask = false,
 		bool InCreateFmask = false)
@@ -123,8 +123,8 @@ public:
 		ETextureCreateFlags InFlags,
 		ETextureCreateFlags InTargetableFlags,
 		bool bInForceSeparateTargetAndShaderResource,
-		uint32 InArraySize,
-		uint16 InNumMips = 1,
+		uint16 InArraySize,
+		uint8 InNumMips = 1,
 		bool InAutowritable = true,
 		bool InCreateRTWriteMask = false,
 		bool InCreateFmask = false)
@@ -161,13 +161,13 @@ public:
 	static FPooledRenderTargetDesc CreateVolumeDesc(
 		uint32 InSizeX,
 		uint32 InSizeY,
-		uint32 InSizeZ,
+		uint16 InSizeZ,
 		EPixelFormat InFormat,
 		const FClearValueBinding& InClearValue,
 		ETextureCreateFlags InFlags,
 		ETextureCreateFlags InTargetableFlags,
 		bool bInForceSeparateTargetAndShaderResource,
-		uint16 InNumMips = 1,
+		uint8 InNumMips = 1,
 		bool InAutowritable = true)
 	{
 		check(InSizeX);
@@ -204,7 +204,7 @@ public:
 		ETextureCreateFlags InTargetableFlags,
 		bool bInForceSeparateTargetAndShaderResource,
 		uint16 InArraySize = 1,
-		uint16 InNumMips = 1,
+		uint8 InNumMips = 1,
 		bool InAutowritable = true)
 	{
 		check(InExtent);
@@ -242,7 +242,7 @@ public:
 		ETextureCreateFlags InTargetableFlags,
 		bool bInForceSeparateTargetAndShaderResource,
 		uint16 InArraySize,
-		uint16 InNumMips = 1,
+		uint8 InNumMips = 1,
 		bool InAutowritable = true)
 	{
 		check(InExtent);
diff --git a/Engine/Source/Runtime/RenderCore/Public/Shader.h b/Engine/Source/Runtime/RenderCore/Public/Shader.h
--- a/Engine/Source/Runtime/RenderCore/Public/Shader.h
+++ b/Engine/Source/Runtime/RenderCore/Public/Shader.h
@@ -2267,7 +2267,8 @@ public:
 		else
 		{
 			// Appending to the archive, track the size of this serialization
-			History.AddValue(Length);
+			check(Length >= 0 && Length <= TNumericLimits<uint32>::Max());
+			History.AddValue((uint32)Length);
 		}
 		HistoryTraversalState.StepForward();
 		
diff --git a/Engine/Source/Runtime/RenderCore/Public/ShaderCodeArchive.h b/Engine/Source/Runtime/RenderCore/Public/ShaderCodeArchive.h
--- a/Engine/Source/Runtime/RenderCore/Public/ShaderCodeArchive.h
+++ b/Engine/Source/Runtime/RenderCore/Public/ShaderCodeArchive.h
@@ -128,7 +128,7 @@ public:
 	{
 	}
 
-	uint32 GetAllocatedSize() const
+	int64 GetAllocatedSize() const
 	{
 		return ShaderHashes.GetAllocatedSize() +
 			ShaderEntries.GetAllocatedSize() +
@@ -250,7 +250,7 @@ public:
 
 	virtual bool IsNativeLibrary() const override { return false; }
 
-	uint32 GetSizeBytes() const
+	int64 GetSizeBytes() const
 	{
 		return sizeof(*this) +
 			SerializedShaders.GetAllocatedSize() +
diff --git a/Engine/Source/Runtime/RenderCore/Public/ShaderParameterUtils.h b/Engine/Source/Runtime/RenderCore/Public/ShaderParameterUtils.h
--- a/Engine/Source/Runtime/RenderCore/Public/ShaderParameterUtils.h
+++ b/Engine/Source/Runtime/RenderCore/Public/ShaderParameterUtils.h
@@ -32,7 +32,7 @@ void SetShaderValue(
 {
 	static_assert(!TIsPointer<ParameterType>::Value, "Passing by value is not valid.");
 
-	const uint32 AlignedTypeSize = Align(sizeof(ParameterType), SHADER_PARAMETER_ARRAY_ELEMENT_ALIGNMENT);
+	const uint32 AlignedTypeSize = (uint32)Align(sizeof(ParameterType), SHADER_PARAMETER_ARRAY_ELEMENT_ALIGNMENT);
 	const int32 NumBytesToSet = FMath::Min<int32>(sizeof(ParameterType),Parameter.GetNumBytes() - ElementIndex * AlignedTypeSize);
 
 	// This will trigger if the parameter was not serialized
@@ -142,7 +142,7 @@ void SetShaderValueArray(
 	uint32 BaseElementIndex = 0
 	)
 {
-	const uint32 AlignedTypeSize = Align(sizeof(ParameterType), SHADER_PARAMETER_ARRAY_ELEMENT_ALIGNMENT);
+	const uint32 AlignedTypeSize = (uint32)Align(sizeof(ParameterType), SHADER_PARAMETER_ARRAY_ELEMENT_ALIGNMENT);
 	const int32 NumBytesToSet = FMath::Min<int32>(NumElements * AlignedTypeSize,Parameter.GetNumBytes() - BaseElementIndex * AlignedTypeSize);
 
 	// This will trigger if the parameter was not serialized
diff --git a/Engine/Source/Runtime/RenderCore/Public/VertexFactory.h b/Engine/Source/Runtime/RenderCore/Public/VertexFactory.h
--- a/Engine/Source/Runtime/RenderCore/Public/VertexFactory.h
+++ b/Engine/Source/Runtime/RenderCore/Public/VertexFactory.h
@@ -659,7 +659,8 @@ public:
 protected:
 	void SetPrimitiveIdStreamIndex(const FStaticFeatureLevel InFeatureLevel, EVertexInputStreamType InputStreamType, int32 StreamIndex)
 	{
-		PrimitiveIdStreamIndex[TranslatePrimitiveIdStreamIndex(InFeatureLevel, InputStreamType)] = StreamIndex;
+		check(StreamIndex >= TNumericLimits<int8>::Min() && StreamIndex <= TNumericLimits<int8>::Max());
+		PrimitiveIdStreamIndex[TranslatePrimitiveIdStreamIndex(InFeatureLevel, InputStreamType)] = static_cast<int8>(StreamIndex);
 	}
 	
 	bool AddPrimitiveIdStreamElement(EVertexInputStreamType InputStreamType, FVertexDeclarationElementList& Elements, uint8 AttributeIndex, uint8 AttributeIndex_Mobile);
