Fix CIS (DebugGame builds) (UE-138641).

diff --git a/Engine/Source/Runtime/RenderCore/RenderCore.Build.cs b/Engine/Source/Runtime/RenderCore/RenderCore.Build.cs
--- a/Engine/Source/Runtime/RenderCore/RenderCore.Build.cs
+++ b/Engine/Source/Runtime/RenderCore/RenderCore.Build.cs
@@ -27,7 +27,8 @@ public class RenderCore : ModuleRules
 			PrivateDependencyModuleNames.Add("CoreUObject");
 		}
 		// shader runtime usage visualization requires ability to create images - it is only used in non-editor desktop development targets
-		else if (Target.Configuration == UnrealTargetConfiguration.Development && Array.IndexOf(Utils.GetPlatformsInClass(UnrealPlatformClass.Desktop), Target.Platform) >= 0)
+		// UE_BUILD_DEVELOPMENT is also defined for DebugGame
+		else if ((Target.Configuration == UnrealTargetConfiguration.Development || Target.Configuration == UnrealTargetConfiguration.DebugGame) && Array.IndexOf(Utils.GetPlatformsInClass(UnrealPlatformClass.Desktop), Target.Platform) >= 0)
 		{
 			PrivateDependencyModuleNames.Add("ImageWrapper");
 		}
