Fixed RDG culling behavior differences between Clang and MSVC due to static initialization of shader parameter metadata.

diff --git a/Engine/Source/Runtime/RenderCore/Private/ShaderParameterMetadata.cpp b/Engine/Source/Runtime/RenderCore/Private/ShaderParameterMetadata.cpp
--- a/Engine/Source/Runtime/RenderCore/Private/ShaderParameterMetadata.cpp
+++ b/Engine/Source/Runtime/RenderCore/Private/ShaderParameterMetadata.cpp
@@ -481,6 +481,8 @@ void FShaderParametersMetadata::InitializeLayout(FRHIUniformBufferLayoutInitiali
 			BaseType == UBMT_UINT32 ||
 			BaseType == UBMT_FLOAT32);
 
+		LayoutInitializer.bHasNonGraphOutputs |= BaseType == UBMT_UAV;
+
 		if (DO_CHECK)
 		{
 			auto GetMemberErrorPrefix = [&]()
@@ -660,24 +662,6 @@ void FShaderParametersMetadata::InitializeLayout(FRHIUniformBufferLayoutInitiali
 			}
 		}
 
-		if (BaseType == UBMT_UAV)
-		{
-			LayoutInitializer.bHasNonGraphOutputs = true;
-		}
-		else if (BaseType == UBMT_REFERENCED_STRUCT || BaseType == UBMT_RDG_UNIFORM_BUFFER)
-		{
-			if (ChildStruct)
-			{
-				for (const FShaderParametersMetadata::FMember& Member : ChildStruct->GetMembers())
-				{
-					if (Member.GetBaseType() == UBMT_UAV)
-					{
-						LayoutInitializer.bHasNonGraphOutputs = true;
-					}
-				}
-			}
-		}
-
 		if (ChildStruct && BaseType != UBMT_REFERENCED_STRUCT && BaseType != UBMT_RDG_UNIFORM_BUFFER)
 		{
 			for (uint32 ArrayElementId = 0; ArrayElementId < (bIsArray ? ArraySize : 1u); ArrayElementId++)
