Fixed:  DumpGPU command on mobile platforms

diff --git a/Engine/Extras/GPUDumpViewer/GPUDumpViewer.html b/Engine/Extras/GPUDumpViewer/GPUDumpViewer.html
--- a/Engine/Extras/GPUDumpViewer/GPUDumpViewer.html
+++ b/Engine/Extras/GPUDumpViewer/GPUDumpViewer.html
@@ -2316,7 +2316,7 @@ void main(void)
 			raw_channel_count = 4;
 			raw_channel_format = 'unorm';
 
-			if (ue_pixel_format === 'PF_B8G8R8A8')
+			if (ue_pixel_format === 'PF_B8G8R8A8' && g_infos['RHI'] !== 'OpenGL')
 			{
 				raw_channel_swizzling = [2, 1, 0, 3];
 			}
diff --git a/Engine/Source/Runtime/OpenGLDrv/Private/OpenGLDevice.cpp b/Engine/Source/Runtime/OpenGLDrv/Private/OpenGLDevice.cpp
--- a/Engine/Source/Runtime/OpenGLDrv/Private/OpenGLDevice.cpp
+++ b/Engine/Source/Runtime/OpenGLDrv/Private/OpenGLDevice.cpp
@@ -958,6 +958,11 @@ static void InitRHICapabilitiesForGL()
 	SetupTextureFormat( PF_R32G32B32F,			FOpenGLTextureFormat( GL_RGB32F,				GL_RGB32F,				GL_RGB,				GL_FLOAT,							false, false));
 	SetupTextureFormat( PF_R8_SINT,				FOpenGLTextureFormat( GL_R8I,					GL_R8I,					GL_RED_INTEGER,		GL_BYTE,							false, false));
 
+	SetupTextureFormat( PF_B8G8R8A8,			FOpenGLTextureFormat( GL_RGBA8,					GL_SRGB8_ALPHA8,		GL_RGBA,			GL_UNSIGNED_BYTE,					false, true));
+	SetupTextureFormat( PF_R8G8B8A8,			FOpenGLTextureFormat( GL_RGBA8,					GL_SRGB8_ALPHA8,		GL_RGBA,			GL_UNSIGNED_BYTE,					false, false));
+	SetupTextureFormat( PF_R8G8B8A8_UINT,		FOpenGLTextureFormat( GL_RGBA8UI,				GL_RGBA8UI,				GL_RGBA_INTEGER,	GL_UNSIGNED_BYTE,					false, false));
+	SetupTextureFormat( PF_G8,					FOpenGLTextureFormat( GL_R8,					GL_R8,					GL_RED,				GL_UNSIGNED_BYTE,					false, false));
+
 #if PLATFORM_DESKTOP
 	CA_SUPPRESS(6286);
 	if (PLATFORM_DESKTOP || FOpenGL::GetFeatureLevel() >= ERHIFeatureLevel::SM5)
@@ -966,23 +971,13 @@ static void InitRHICapabilitiesForGL()
 		SetupTextureFormat( PF_BC5,				FOpenGLTextureFormat( GL_COMPRESSED_RG_RGTC2,	GL_COMPRESSED_RG_RGTC2,	GL_RG,			GL_UNSIGNED_BYTE,					true,	false));
 		SetupTextureFormat( PF_BC4,				FOpenGLTextureFormat( GL_COMPRESSED_RED_RGTC1,	GL_COMPRESSED_RED_RGTC1,GL_RED,			GL_UNSIGNED_BYTE,					true,	false));
 
-		SetupTextureFormat(PF_G8, FOpenGLTextureFormat(GL_R8, GL_SRGB8, GL_RED, GL_UNSIGNED_BYTE, false, false));
-		SetupTextureFormat(PF_B8G8R8A8, FOpenGLTextureFormat(GL_RGBA8, GL_SRGB8_ALPHA8, GL_BGRA, GL_UNSIGNED_INT_8_8_8_8_REV, false, false));
-		SetupTextureFormat(PF_R8G8B8A8, FOpenGLTextureFormat(GL_RGBA8, GL_SRGB8_ALPHA8, GL_RGBA, GL_UNSIGNED_INT_8_8_8_8_REV, false, false));
-		SetupTextureFormat(PF_R8G8B8A8_UINT, FOpenGLTextureFormat(GL_RGBA8UI, GL_RGBA8UI, GL_RGBA_INTEGER, GL_UNSIGNED_BYTE, false, false));
-
-		SetupTextureFormat(PF_G16R16, FOpenGLTextureFormat(GL_RG16, GL_RG16, GL_RG, GL_UNSIGNED_SHORT, false, false));
-		SetupTextureFormat(PF_G16R16_SNORM, FOpenGLTextureFormat(GL_RG16_SNORM, GL_RG16_SNORM, GL_RG, GL_SHORT, false, false));
+		SetupTextureFormat( PF_G16R16,			FOpenGLTextureFormat(GL_RG16, GL_RG16, GL_RG, GL_UNSIGNED_SHORT, false, false));
+		SetupTextureFormat( PF_G16R16_SNORM,	FOpenGLTextureFormat(GL_RG16_SNORM, GL_RG16_SNORM, GL_RG, GL_SHORT, false, false));
 	}
 	else
 #endif // PLATFORM_DESKTOP
 	{
 #if !PLATFORM_DESKTOP
-		SetupTextureFormat(PF_B8G8R8A8, FOpenGLTextureFormat(GL_RGBA8, GL_SRGB8_ALPHA8, GL_RGBA, GL_UNSIGNED_BYTE, false, true));
-		SetupTextureFormat(PF_R8G8B8A8, FOpenGLTextureFormat(GL_RGBA8, GL_SRGB8_ALPHA8, GL_RGBA, GL_UNSIGNED_BYTE, false, false));
-		SetupTextureFormat(PF_R8G8B8A8_UINT, FOpenGLTextureFormat(GL_RGBA8UI, GL_RGBA8UI, GL_RGBA_INTEGER, GL_UNSIGNED_BYTE, false, false));
-		SetupTextureFormat(PF_G8, FOpenGLTextureFormat(GL_R8, GL_R8, GL_RED, GL_UNSIGNED_BYTE, false, false));
-		
 		FOpenGL::PE_SetupTextureFormat(&SetupTextureFormat); // platform extension
 #endif // !PLATFORM_DESKTOP
 	}
diff --git a/Engine/Source/Runtime/OpenGLDrv/Private/OpenGLRenderTarget.cpp b/Engine/Source/Runtime/OpenGLDrv/Private/OpenGLRenderTarget.cpp
--- a/Engine/Source/Runtime/OpenGLDrv/Private/OpenGLRenderTarget.cpp
+++ b/Engine/Source/Runtime/OpenGLDrv/Private/OpenGLRenderTarget.cpp
@@ -472,9 +472,7 @@ void FOpenGLDynamicRHI::RHICopyToResolveTarget(FRHITexture* SourceTextureRHI, FR
 			&& SrcRect.Size() == DestRect.Size()
 			&& SrcRect.Width() > 0
 			&& SrcRect.Height() > 0
-#if PLATFORM_ANDROID
 			&& SourceTexture->Target == DestTexture->Target // glCopyImageSubData() doesn't like copying from a texture to a renderbuffer on Android
-#endif
 			;
 		
 		if ( !bTrueBlit || !FOpenGL::SupportsCopyImage() )
diff --git a/Engine/Source/Runtime/OpenGLDrv/Private/OpenGLTexture.cpp b/Engine/Source/Runtime/OpenGLDrv/Private/OpenGLTexture.cpp
--- a/Engine/Source/Runtime/OpenGLDrv/Private/OpenGLTexture.cpp
+++ b/Engine/Source/Runtime/OpenGLDrv/Private/OpenGLTexture.cpp
@@ -835,14 +835,18 @@ void TOpenGLTexture<RHIResourceType>::Resolve(uint32 MipIndex,uint32 ArrayIndex)
 			}
 			else
 			{
+				// Get framebuffer for texture
+				FOpenGLTextureBase* Texture = this;
+				GLuint SourceFramebuffer = OpenGLRHI->GetOpenGLFramebuffer(1, &Texture, (bCubemap ? &ArrayIndex : nullptr), &MipIndex, nullptr);
+				// Bind the framebuffer
+				glBindFramebuffer(UGL_READ_FRAMEBUFFER, SourceFramebuffer);
+				FOpenGL::ReadBuffer(GL_COLOR_ATTACHMENT0);
+
 				glPixelStorei(GL_PACK_ALIGNMENT, 1);
-				FOpenGL::GetTexImage(
-									 bCubemap ? GL_TEXTURE_CUBE_MAP_POSITIVE_X + ArrayIndex : Target,
-									 MipIndex,
-									 GLFormat.Format,
-									 GLFormat.Type,
-									 0);	// offset into PBO
+				glReadPixels(0, 0, MipSizeX, MipSizeY, GLFormat.Format, GLFormat.Type, 0);
 				glPixelStorei(GL_PACK_ALIGNMENT, 4);
+
+				OpenGLRHI->GetContextStateForCurrentContext().Framebuffer = (GLuint)-1;
 			}
 		}
 	}
diff --git a/Engine/Source/Runtime/RenderCore/Private/RenderGraphDumpResources.cpp b/Engine/Source/Runtime/RenderCore/Private/RenderGraphDumpResources.cpp
--- a/Engine/Source/Runtime/RenderCore/Private/RenderGraphDumpResources.cpp
+++ b/Engine/Source/Runtime/RenderCore/Private/RenderGraphDumpResources.cpp
@@ -702,14 +702,25 @@ struct FRDGResourceDumpContext
 
 		// Preprocess
 		FTextureRHIRef StagingSrcTexture;
+		EPixelFormat PreprocessedPixelFormat = SubresourceDumpDesc.PreprocessedPixelFormat;
+		SIZE_T SubresourceByteSize = SubresourceDumpDesc.ByteSize;
 		if (SubresourceDumpDesc.bPreprocessForStaging)
 		{
+			// Some RHIs (GL) only support 32Bit single channel images as CS output
+			if (IsOpenGLPlatform(GMaxRHIShaderPlatform) &&
+				GPixelFormats[PreprocessedPixelFormat].NumComponents == 1 && 
+				GPixelFormats[PreprocessedPixelFormat].BlockBytes < 4)
+			{
+				SubresourceByteSize*= (4 / GPixelFormats[PreprocessedPixelFormat].BlockBytes);
+				PreprocessedPixelFormat = PF_R32_UINT;
+			}
+						
 			{
 				FRHIResourceCreateInfo CreateInfo(TEXT("DumpGPU.PreprocessTexture"));
 				StagingSrcTexture = RHICreateTexture2D(
 					SubresourceDumpDesc.SubResourceExtent.X,
 					SubresourceDumpDesc.SubResourceExtent.Y,
-					(uint8)SubresourceDumpDesc.PreprocessedPixelFormat,
+					(uint8)PreprocessedPixelFormat,
 					/* NumMips = */ 1,
 					/* NumSamples = */ 1,
 					TexCreate_UAV | TexCreate_ShaderResource | TexCreate_HideInVisualizeTexture,
@@ -747,7 +758,7 @@ struct FRDGResourceDumpContext
 			StagingTexture = RHICreateTexture2D(
 				SubresourceDumpDesc.SubResourceExtent.X,
 				SubresourceDumpDesc.SubResourceExtent.Y,
-				(uint8)SubresourceDumpDesc.PreprocessedPixelFormat,
+				(uint8)PreprocessedPixelFormat,
 				/* NumMips = */ 1,
 				/* NumSamples = */ 1,
 				TexCreate_CPUReadback | TexCreate_HideInVisualizeTexture,
@@ -792,9 +803,9 @@ struct FRDGResourceDumpContext
 		if (Content)
 		{
 			TArray64<uint8> Array;
-			Array.SetNumUninitialized(SubresourceDumpDesc.ByteSize);
+			Array.SetNumUninitialized(SubresourceByteSize);
 
-			SIZE_T BytePerPixel = SIZE_T(GPixelFormats[SubresourceDumpDesc.PreprocessedPixelFormat].BlockBytes);
+			SIZE_T BytePerPixel = SIZE_T(GPixelFormats[PreprocessedPixelFormat].BlockBytes);
 
 			const uint8* SrcData = static_cast<const uint8*>(Content);
 
@@ -809,6 +820,30 @@ struct FRDGResourceDumpContext
 
 			RHICmdList.UnmapStagingSurface(StagingTexture, GPUIndex);
 
+			if (PreprocessedPixelFormat != SubresourceDumpDesc.PreprocessedPixelFormat)
+			{
+				// Convert 32Bit values back to 16 or 8bit
+				const int32 DstPixelNumBytes = GPixelFormats[SubresourceDumpDesc.PreprocessedPixelFormat].BlockBytes;
+				const uint32* SrcData32 = (const uint32*)Array.GetData();
+				uint8* DstData8 = Array.GetData();
+				uint16* DstData16 = (uint16*)Array.GetData();
+											
+				for (int32 Index = 0; Index < Array.Num()/4; Index++)
+				{
+					uint32 Value32 = SrcData32[Index];
+					if (DstPixelNumBytes == 2)
+					{
+						DstData16[Index] = (uint16)Value32;
+					}
+					else
+					{
+						DstData8[Index] = (uint8)Value32;
+					}
+				}
+
+				Array.SetNum(Array.Num() / (4 / DstPixelNumBytes));
+			}
+
 			DumpBinaryToFile(Array, DumpFilePath);
 		}
 		else
@@ -1162,8 +1197,7 @@ struct FRDGResourceDumpContext
 		return bDumpPass;
 	}
 };
-
-bool bIsDumpingFrame_GameThread = false;
+static uint64 DumpingFrameCounter_GameThread = 0;
 FRDGResourceDumpContext GRDGResourceDumpContext;
 
 }
@@ -1177,7 +1211,7 @@ FString FRDGBuilder::BeginResourceDump(const TArray<FString>& Args)
 {
 	check(IsInGameThread());
 
-	if (bIsDumpingFrame_GameThread)
+	if (DumpingFrameCounter_GameThread != 0)
 	{
 		return FString();
 	}
@@ -1284,7 +1318,7 @@ FString FRDGBuilder::BeginResourceDump(const TArray<FString>& Args)
 		#endif
 	});
 
-	bIsDumpingFrame_GameThread = true;
+	DumpingFrameCounter_GameThread = GFrameCounter;
 
 	if (NewResourceDumpContext.bEnableDiskWrite)
 	{
@@ -1297,7 +1331,8 @@ void FRDGBuilder::EndResourceDump()
 {
 	check(IsInGameThread());
 
-	if (!bIsDumpingFrame_GameThread)
+	if (DumpingFrameCounter_GameThread == 0 ||
+		DumpingFrameCounter_GameThread >= GFrameCounter) // make sure at least one frame has passed since we start a resource dump
 	{
 		return;
 	}
@@ -1343,7 +1378,7 @@ void FRDGBuilder::EndResourceDump()
 	#endif
 
 	GRDGResourceDumpContext = FRDGResourceDumpContext();
-	bIsDumpingFrame_GameThread = false;
+	DumpingFrameCounter_GameThread = 0;
 }
 
 static const TCHAR* GetPassEventNameWithGPUMask(const FRDGPass* Pass, FString& OutNameStorage)
@@ -1782,6 +1817,11 @@ void FRDGBuilder::EndPassDump(const FRDGPass* Pass)
 }
 
 // static
+bool FRDGBuilder::IsDumpingFrame()
+{
+	return GRDGResourceDumpContext.IsDumpingFrame();
+}
+
 bool FRDGBuilder::IsDumpingDraws()
 {
 	if (!GRDGResourceDumpContext.IsDumpingFrame())
diff --git a/Engine/Source/Runtime/RenderCore/Public/RenderGraphBuilder.h b/Engine/Source/Runtime/RenderCore/Public/RenderGraphBuilder.h
--- a/Engine/Source/Runtime/RenderCore/Public/RenderGraphBuilder.h
+++ b/Engine/Source/Runtime/RenderCore/Public/RenderGraphBuilder.h
@@ -302,6 +302,9 @@ public:
 #if RDG_DUMP_RESOURCES
 	static FString BeginResourceDump(const TArray<FString>& Args);
 	static void EndResourceDump();
+	static bool IsDumpingFrame();
+#else
+	static bool IsDumpingFrame() { return false; }
 #endif
 
 #if RDG_DUMP_RESOURCES_AT_EACH_DRAW
diff --git a/Engine/Source/Runtime/Renderer/Private/MobileShadingRenderer.cpp b/Engine/Source/Runtime/Renderer/Private/MobileShadingRenderer.cpp
--- a/Engine/Source/Runtime/Renderer/Private/MobileShadingRenderer.cpp
+++ b/Engine/Source/Runtime/Renderer/Private/MobileShadingRenderer.cpp
@@ -456,7 +456,8 @@ void FMobileSceneRenderer::InitViews(FRDGBuilder& GraphBuilder, FSceneTexturesCo
 		(bDeferredShading && bPostProcessUsesSceneDepth) ||
 		bShouldRenderVelocities ||
 		bRequireSeparateViewPass ||
-		bIsFullDepthPrepassEnabled;
+		bIsFullDepthPrepassEnabled ||
+		GraphBuilder.IsDumpingFrame();
 	// never keep MSAA depth if SceneDepthAux is enabled
 	bKeepDepthContent = ((NumMSAASamples > 1) && bRequiresSceneDepthAux) ? false : bKeepDepthContent;
 
