Make use of TEXTVIEW in ShaderCore.cpp.

#rnx

diff --git a/Engine/Source/Runtime/RenderCore/Private/ShaderCore.cpp b/Engine/Source/Runtime/RenderCore/Private/ShaderCore.cpp
--- a/Engine/Source/Runtime/RenderCore/Private/ShaderCore.cpp
+++ b/Engine/Source/Runtime/RenderCore/Private/ShaderCore.cpp
@@ -595,14 +595,14 @@ bool CheckVirtualShaderFilePath(FStringView VirtualFilePath, TArray<FShaderCompi
 		bSuccess = false;
 	}
 
-	if (Contains(VirtualFilePath, TEXT(".."_SV)))
+	if (Contains(VirtualFilePath, TEXTVIEW("..")))
 	{
 		FString Error = FString::Printf(TEXT("Virtual shader source file name \"%s\" should have relative directories (\"../\") collapsed."), *FString(VirtualFilePath));
 		ReportVirtualShaderFilePathError(CompileErrors, Error);
 		bSuccess = false;
 	}
 
-	if (Contains(VirtualFilePath, TEXT("\\"_SV)))
+	if (Contains(VirtualFilePath, TEXTVIEW("\\")))
 	{
 		FString Error = FString::Printf(TEXT("Backslashes are not permitted in virtual shader source file name \"%s\""), *FString(VirtualFilePath));
 		ReportVirtualShaderFilePathError(CompileErrors, Error);
@@ -612,7 +612,7 @@ bool CheckVirtualShaderFilePath(FStringView VirtualFilePath, TArray<FShaderCompi
 	FStringView Extension = FPathViews::GetExtension(VirtualFilePath);
 	if (VirtualFilePath.StartsWith(TEXT("/Engine/Shared/")))
 	{
-		if ((Extension != TEXT("h"_SV)))
+		if ((Extension != TEXTVIEW("h")))
 		{
 			FString Error = FString::Printf(TEXT("Extension on virtual shader source file name \"%s\" is wrong. Only .h is allowed for shared headers that are shared between C++ and shader code."), *FString(VirtualFilePath));
 			ReportVirtualShaderFilePathError(CompileErrors, Error);
@@ -877,7 +877,7 @@ bool ReplaceVirtualFilePathForShaderPlatform(FString& InOutVirtualFilePath, ESha
 
 bool ReplaceVirtualFilePathForShaderAutogen(FString& InOutVirtualFilePath, EShaderPlatform ShaderPlatform)
 {
-	const FStringView ShaderAutogenStem = TEXT("/Engine/Generated/ShaderAutogen/"_SV);
+	const FStringView ShaderAutogenStem = TEXTVIEW("/Engine/Generated/ShaderAutogen/");
 
 	// Tweak the autogen path
 	// for examples, if it starts with "/Engine/Generated/ShaderAutogen/" change it to "ShaderAutogen/PCD3D_SM5/"
@@ -886,7 +886,7 @@ bool ReplaceVirtualFilePathForShaderAutogen(FString& InOutVirtualFilePath, EShad
 		TStringBuilder<1024> OutputShaderName;
 
 		// Append the prefix.
-		OutputShaderName.Append(TEXT("/ShaderAutogen/"_SV)); // using FStringView (_SV suffix) in order to avoid a strlen call
+		OutputShaderName.Append(TEXTVIEW("/ShaderAutogen/"));
 
 		// Append the platform name.
 		const FName PlatformName = LegacyShaderPlatformToShaderFormat(ShaderPlatform);
