UE5_RELEASE:  Virtual Production, improved profiling for multiple view family scenarios.
* Increased limit for GPU profiler events on D3D12.  A typical frame with 4 view families would easily exhaust the old limit of 1K events resulting in many events not showing up in the profile -- limit increased to 8K, and confirmed negligible difference in memory.  Fixed bug where we were only using half the queries allocated as well (unnecessary multiplication by 2 when calling RHICreateRenderQueryPool).
* Added optional "ProfileDescription" to FSceneViewFamily, which is appended to the Unreal Insights profile event names for both CPU and GPU, so you can differentiate view families.
* Infrastructure added to CPU profiler to allow a verbose name to be added to a Stat timing block.
* Infrastructure added to GPU profiler to allow verbose names to be added to timing blocks in general, both Event and Stat timing blocks.
* Added stats display for Display Cluster listing the CPU and GPU cost, and GPU assignment per view family by name, with the goal of allowing clients to tune features enabled per view family to adjust performance, without needing to gather performance information from Unreal Insights.
* FSceneRenderer::DoCrossGPUTransfers uses push transfer without lockstep (logic copied from 4.27), which significantly improves perf.

#rnx

diff --git a/Engine/Plugins/Runtime/nDisplay/Source/DisplayCluster/Private/Game/EngineClasses/Basics/DisplayClusterViewportClient.cpp b/Engine/Plugins/Runtime/nDisplay/Source/DisplayCluster/Private/Game/EngineClasses/Basics/DisplayClusterViewportClient.cpp
--- a/Engine/Plugins/Runtime/nDisplay/Source/DisplayCluster/Private/Game/EngineClasses/Basics/DisplayClusterViewportClient.cpp
+++ b/Engine/Plugins/Runtime/nDisplay/Source/DisplayCluster/Private/Game/EngineClasses/Basics/DisplayClusterViewportClient.cpp
@@ -28,6 +28,7 @@
 #include "Components/LineBatchComponent.h"
 #include "GameFramework/GameUserSettings.h"
 #include "Framework/Application/SlateApplication.h"
+#include "Slate/SceneViewport.h"
 
 #include "LegacyScreenPercentageDriver.h"
 #include "DynamicResolutionState.h"
@@ -58,6 +59,14 @@ static FAutoConsoleVariableRef CVarDisplayClusterForceCopyCrossGPU(
 	ECVF_RenderThreadSafe
 );
 
+int32 GDisplayClusterShowStats = 0;
+static FAutoConsoleVariableRef CVarDisplayClusterShowStats(
+	TEXT("DC.Stats"),
+	GDisplayClusterShowStats,
+	TEXT("Show per-view profiling stats for display cluster rendering."),
+	ECVF_RenderThreadSafe
+);
+
 UDisplayClusterViewportClient::UDisplayClusterViewportClient(FVTableHelper& Helper)
 	: Super(Helper)
 {
@@ -89,6 +98,183 @@ static UCanvas* GetCanvasByName(FName CanvasName)
 	return *FoundCanvas;
 }
 
+// Wrapper for FSceneViewport to allow us to add custom stats specific to display cluster (per-view-family CPU and GPU perf)
+class FDisplayClusterSceneViewport : public FSceneViewport
+{
+public:
+	FDisplayClusterSceneViewport(FViewportClient* InViewportClient, TSharedPtr<SViewport> InViewportWidget)
+		: FSceneViewport(InViewportClient, InViewportWidget)
+	{}
+
+	~FDisplayClusterSceneViewport()
+	{
+		for (auto It = CpuHistoryByDescription.CreateIterator(); It; ++It)
+		{
+			delete It.Value();
+		}
+	}
+
+	virtual int32 DrawStatsHUD(FCanvas* InCanvas, int32 InX, int32 InY) override
+	{
+#if GPUPROFILERTRACE_ENABLED
+		if (GDisplayClusterShowStats)
+		{
+			// Get GPU perf results
+			TArray<FRealtimeGPUProfilerDescriptionResult> PerfResults;
+			FRealtimeGPUProfiler::Get()->FetchPerfByDescription(PerfResults);
+
+			UFont* StatsFont = GetStatsFont();
+
+			const FLinearColor HeaderColor = FLinearColor(1.f, 0.2f, 0.f);
+
+			if (PerfResults.Num())
+			{
+				// Get CPU perf results
+				TArray<float> CpuPerfResults;
+				CpuPerfResults.AddUninitialized(PerfResults.Num());
+				{
+					FRWScopeLock Lock(CpuHistoryMutex, SLT_Write);
+
+					for (int32 ResultIndex = 0; ResultIndex < PerfResults.Num(); ResultIndex++)
+					{
+						CpuPerfResults[ResultIndex] = FetchHistoryAverage(PerfResults[ResultIndex].Description);
+					}
+				}
+
+				// Compute column sizes
+				int32 YIgnore;
+
+				const TCHAR* DescriptionHeader = TEXT("Display Cluster Stats");
+				int32 DescriptionColumnWidth;
+				StringSize(StatsFont, DescriptionColumnWidth, YIgnore, DescriptionHeader);
+
+				for (const FRealtimeGPUProfilerDescriptionResult& PerfResult : PerfResults)
+				{
+					int32 XL;
+					StringSize(StatsFont, XL, YIgnore, *PerfResult.Description);
+
+					DescriptionColumnWidth = FMath::Max(DescriptionColumnWidth, XL);
+				}
+
+				int32 NumberColumnWidth;
+				StringSize(StatsFont, NumberColumnWidth, YIgnore, *FString::ChrN(7, 'W'));
+
+				// Render header
+				InCanvas->DrawShadowedString(InX, InY, DescriptionHeader, StatsFont, HeaderColor);
+				RightJustify(InCanvas, StatsFont, InX + DescriptionColumnWidth + 1 * NumberColumnWidth, InY, TEXT("GPUs"), HeaderColor);
+				RightJustify(InCanvas, StatsFont, InX + DescriptionColumnWidth + 2 * NumberColumnWidth, InY, TEXT("Average"), HeaderColor);
+				RightJustify(InCanvas, StatsFont, InX + DescriptionColumnWidth + 3 * NumberColumnWidth, InY, TEXT("CPU"), HeaderColor);
+				InY += StatsFont->GetMaxCharHeight();
+
+				// Render rows
+				int32 ResultIndex = 0;
+				const FLinearColor StatColor = FLinearColor(0.f, 1.f, 0.f);
+
+				for (const FRealtimeGPUProfilerDescriptionResult& PerfResult : PerfResults)
+				{
+					InCanvas->DrawTile(InX, InY, DescriptionColumnWidth + 3 * NumberColumnWidth, StatsFont->GetMaxCharHeight(),
+						0, 0, 1, 1,
+						(ResultIndex & 1) ? FLinearColor(0.02f, 0.02f, 0.02f, 0.88f) : FLinearColor(0.05f, 0.05f, 0.05f, 0.92f),
+						GWhiteTexture, true);
+
+					// Source GPU times are in microseconds, CPU times in seconds, so we need to divide one by 1000, and multiply the other by 1000
+					InCanvas->DrawShadowedString(InX, InY, *PerfResult.Description, StatsFont, StatColor);
+					RightJustify(InCanvas, StatsFont, InX + DescriptionColumnWidth + 1 * NumberColumnWidth, InY, *FString::Printf(TEXT("%d"), PerfResult.GPUMask.GetNative()), StatColor);
+					RightJustify(InCanvas, StatsFont, InX + DescriptionColumnWidth + 2 * NumberColumnWidth, InY, *FString::Printf(TEXT("%.2f"), PerfResult.AverageTime / 1000.f), StatColor);
+					RightJustify(InCanvas, StatsFont, InX + DescriptionColumnWidth + 3 * NumberColumnWidth, InY, *FString::Printf(TEXT("%.2f"), CpuPerfResults[ResultIndex] * 1000.f), StatColor);
+
+					InY += StatsFont->GetMaxCharHeight();
+
+					ResultIndex++;
+				}
+			}
+			else
+			{
+				InCanvas->DrawShadowedString(InX, InY, TEXT("Display Cluster Stats [NO DATA]"), StatsFont, HeaderColor);
+				InY += StatsFont->GetMaxCharHeight();
+			}
+
+			InY += StatsFont->GetMaxCharHeight();
+		}
+#endif  // GPUPROFILERTRACE_ENABLED
+
+		return InY;
+	}
+
+	float* GetNextHistoryWriteAddress(const FString& Description)
+	{
+		FRWScopeLock Lock(CpuHistoryMutex, SLT_Write);
+
+		FCpuProfileHistory*& History = CpuHistoryByDescription.FindOrAdd(Description);
+		if (!History)
+		{
+			History = new FCpuProfileHistory;
+		}
+
+		return &History->Times[(History->HistoryIndex++) % FCpuProfileHistory::HistoryCount];
+	}
+
+private:
+	static void RightJustify(FCanvas* Canvas, UFont* StatsFont, const int32 X, const int32 Y, TCHAR const* Text, FLinearColor const& Color)
+	{
+		int32 ColumnSizeX, ColumnSizeY;
+		StringSize(StatsFont, ColumnSizeX, ColumnSizeY, Text);
+		Canvas->DrawShadowedString(X - ColumnSizeX, Y, Text, StatsFont, Color);
+	}
+
+	// Only callable when the CpuHistoryMutex is locked!
+	float FetchHistoryAverage(const FString& Description) const
+	{
+		const FCpuProfileHistory* const* History = CpuHistoryByDescription.Find(Description);
+
+		float Average = 0.f;
+		if (History)
+		{
+			float ValidResultCount = 0.f;
+			for (uint32 HistoryIndex = 0; HistoryIndex < FCpuProfileHistory::HistoryCount; HistoryIndex++)
+			{
+				float HistoryTime = (*History)->Times[HistoryIndex];
+				if (HistoryTime > 0.f)
+				{
+					Average += HistoryTime;
+					ValidResultCount += 1.f;
+				}
+			}
+			if (ValidResultCount > 0.f)
+			{
+				Average /= ValidResultCount;
+			}
+		}
+		return Average;
+	}
+
+	struct FCpuProfileHistory
+	{
+		FCpuProfileHistory()
+		{
+			FMemory::Memset(*this, 0);
+		}
+
+		static const uint32 HistoryCount = 64;
+
+		// Constructor memsets everything to zero, assuming structure is Plain Old Data.  If any dynamic structures are
+		// added, you'll need a more generalized constructor that zeroes out all the uninitialized data.
+		uint32 HistoryIndex;
+		float Times[HistoryCount];
+	};
+
+	// History payload is separately allocated in memory, as it's written to asynchronously by the Render Thread, and we
+	// can't have it moved if the Map storage gets reallocated when new view families are added.
+	TMap<FString, FCpuProfileHistory*> CpuHistoryByDescription;
+	FRWLock CpuHistoryMutex;
+};
+
+// Override to allocate our custom viewport class
+FSceneViewport* UDisplayClusterViewportClient::CreateGameViewport(TSharedPtr<SViewport> InViewportWidget)
+{
+	return new FDisplayClusterSceneViewport(this, InViewportWidget);
+}
+
 void UDisplayClusterViewportClient::Init(struct FWorldContext& WorldContext, UGameInstance* OwningGameInstance, bool bCreateNewAudioDevice)
 {
 	const bool bIsNDisplayClusterMode = (GEngine->StereoRenderingDevice.IsValid() && GDisplayCluster->GetOperationMode() == EDisplayClusterOperationMode::Cluster);
@@ -545,6 +731,12 @@ void UDisplayClusterViewportClient::Draw(FViewport* InViewport, FCanvas* SceneCa
 				ViewFamily.bForceCopyCrossGPU = GDisplayClusterForceCopyCrossGPU != 0;
 #endif
 
+				ViewFamily.ProfileDescription = DCViewFamily.Views[0].Viewport->GetId();
+				if (!ViewFamily.ProfileDescription.IsEmpty())
+				{
+					ViewFamily.ProfileSceneRenderTime = ((FDisplayClusterSceneViewport*)InViewport)->GetNextHistoryWriteAddress(ViewFamily.ProfileDescription);
+				}
+
 				// Draw the player views.
 				if (!bDisableWorldRendering && PlayerViewMap.Num() > 0 && FSlateApplication::Get().GetPlatformApplication()->IsAllowedToRender()) //-V560
 				{
diff --git a/Engine/Plugins/Runtime/nDisplay/Source/DisplayCluster/Public/DisplayClusterViewportClient.h b/Engine/Plugins/Runtime/nDisplay/Source/DisplayCluster/Public/DisplayClusterViewportClient.h
--- a/Engine/Plugins/Runtime/nDisplay/Source/DisplayCluster/Public/DisplayClusterViewportClient.h
+++ b/Engine/Plugins/Runtime/nDisplay/Source/DisplayCluster/Public/DisplayClusterViewportClient.h
@@ -19,6 +19,8 @@ public:
 	virtual void Init(struct FWorldContext& WorldContext, UGameInstance* OwningGameInstance, bool bCreateNewAudioDevice = true) override;
 	virtual void Draw(FViewport* Viewport, FCanvas* SceneCanvas) override;
 
+	virtual FSceneViewport* CreateGameViewport(TSharedPtr<SViewport> InViewportWidget) override;
+
 protected:
 #if WITH_EDITOR
 	bool Draw_PIE(FViewport* InViewport, FCanvas* SceneCanvas);
diff --git a/Engine/Source/Runtime/Core/Public/Stats/Stats.h b/Engine/Source/Runtime/Core/Public/Stats/Stats.h
--- a/Engine/Source/Runtime/Core/Public/Stats/Stats.h
+++ b/Engine/Source/Runtime/Core/Public/Stats/Stats.h
@@ -59,9 +59,19 @@ public:
 	 * Pushes the specified stat onto the hierarchy for this thread. Starts
 	 * the timing of the cycles used
 	 */
-	FORCEINLINE_STATS FScopeCycleCounter( TStatId StatId, EStatFlags StatFlags, bool bAlways = false )
+	FORCEINLINE_STATS FScopeCycleCounter( TStatId StatId, EStatFlags StatFlags, bool bAlways = false
+#if CPUPROFILERTRACE_ENABLED
+		// Optional verbose description added for CPU profiler, without affecting the stat name.  Should be invariant
+		// across frames, to avoid bloating memory in the name map used in the CPU profiler.
+		, const TCHAR* OptionalVerboseDescription = nullptr
+#endif
+	)
 	{
-		Start( StatId, StatFlags, bAlways );
+		Start( StatId, StatFlags, bAlways
+#if CPUPROFILERTRACE_ENABLED
+			, OptionalVerboseDescription
+#endif
+			);
 	}
 
 	FORCEINLINE_STATS FScopeCycleCounter(TStatId StatId, bool bAlways = false)
@@ -245,6 +255,10 @@ FORCEINLINE void StatsMasterEnableSubtract(int32 Value = 1)
 	FScopeCycleCounter StatNamedEventsScope_##Stat(bCondition ? ANSI_TO_PROFILING(#Stat) : nullptr); \
 	SCOPE_CYCLE_COUNTER_TO_TRACE(#Stat, Stat, bCondition);
 
+#define SCOPE_CYCLE_COUNTER_VERBOSE(Stat, VerboseDescription) \
+	FScopeCycleCounter StatNamedEventsScope_##Stat(TStatId(ANSI_TO_PROFILING(#Stat))); \
+	SCOPE_CYCLE_COUNTER_TO_TRACE(#Stat, Stat, true);
+
 #define RETURN_QUICK_DECLARE_CYCLE_STAT(StatId,GroupId) return TStatId(ANSI_TO_PROFILING(#StatId));
 
 #define GET_STATID(Stat) (TStatId(ANSI_TO_PROFILING(#Stat)))
@@ -285,6 +299,9 @@ public:
 #define CONDITIONAL_SCOPE_CYCLE_COUNTER(Stat,bCondition) \
 	FLightweightStatScope LightweightStatScope_##Stat(bCondition ? TEXT(#Stat) : nullptr);
 
+#define SCOPE_CYCLE_COUNTER_VERBOSE(Stat,VerboseDescription) \
+	FLightweightStatScope LightweightStatScope_##Stat(TEXT(#Stat));
+
 #define RETURN_QUICK_DECLARE_CYCLE_STAT(StatId,GroupId) return TStatId();
 #define GET_STATID(Stat) (TStatId())
 
@@ -293,6 +310,7 @@ public:
 #define QUICK_SCOPE_CYCLE_COUNTER(Stat)
 #define DECLARE_SCOPE_CYCLE_COUNTER(CounterName,StatId,GroupId)
 #define CONDITIONAL_SCOPE_CYCLE_COUNTER(Stat,bCondition)
+#define SCOPE_CYCLE_COUNTER_VERBOSE(Stat,VerboseDescription)
 #define RETURN_QUICK_DECLARE_CYCLE_STAT(StatId,GroupId) return TStatId();
 #define GET_STATID(Stat) (TStatId())
 
diff --git a/Engine/Source/Runtime/Core/Public/Stats/Stats2.h b/Engine/Source/Runtime/Core/Public/Stats/Stats2.h
--- a/Engine/Source/Runtime/Core/Public/Stats/Stats2.h
+++ b/Engine/Source/Runtime/Core/Public/Stats/Stats2.h
@@ -1564,7 +1564,13 @@ public:
 	 * Pushes the specified stat onto the hierarchy for this thread. Starts
 	 * the timing of the cycles used
 	 */
-	FORCEINLINE_STATS void Start(TStatId InStatId, EStatFlags InStatFlags, bool bAlways = false)
+	FORCEINLINE_STATS void Start(TStatId InStatId, EStatFlags InStatFlags, bool bAlways = false
+#if CPUPROFILERTRACE_ENABLED
+		// Optional verbose description added for CPU profiler, without affecting the stat name.  Should be invariant
+		// across frames, to avoid bloating memory in the name map used in the CPU profiler.
+		, const TCHAR* OptionalVerboseDescription = nullptr
+#endif
+		)
 	{
 		FMinimalName StatMinimalName = InStatId.GetMinimalName(EMemoryOrder::Relaxed);
 		if (StatMinimalName.IsNone())
@@ -1586,7 +1592,37 @@ public:
 #if CPUPROFILERTRACE_ENABLED
 			if (UE_TRACE_CHANNELEXPR_IS_ENABLED(CpuChannel))
 			{
-				FCpuProfilerTrace::OutputBeginDynamicEvent(InStatId.GetStatDescriptionANSI()); //todo: Could we use FName index as event id?
+				// We only support ANSI strings, since the CPU profiler internally uses that in its event type.  We could
+				// in the future add a separate wide string profiler event, but it involves duplicating a fair amount of
+				// code, so this works for now.
+				if (OptionalVerboseDescription && FCString::IsPureAnsi(OptionalVerboseDescription))
+				{
+					TArray<ANSICHAR, TNonRelocatableInlineAllocator<256>> VerboseNameANSI;
+
+					const char* StatDescriptionANSI = InStatId.GetStatDescriptionANSI();
+					int32 StatLen = FCStringAnsi::Strlen(StatDescriptionANSI);
+					int32 DescriptionLen = FCString::Strlen(OptionalVerboseDescription);
+
+					VerboseNameANSI.AddUninitialized(StatLen + DescriptionLen + 4);		// 3 character separator plus null terminator
+					char* Buffer = VerboseNameANSI.GetData();
+					char* BufferEnd = Buffer + VerboseNameANSI.Num();
+
+					FMemory::Memcpy(Buffer, StatDescriptionANSI, StatLen*sizeof(ANSICHAR));	Buffer += StatLen;
+					FMemory::Memcpy(Buffer, " - ", 3*sizeof(ANSICHAR));						Buffer += 3;
+					for (const TCHAR* DescriptionLetter = OptionalVerboseDescription; *DescriptionLetter; DescriptionLetter++, Buffer++)
+					{
+						*Buffer = (ANSICHAR)*DescriptionLetter;
+					}
+					check(Buffer + 1 == BufferEnd);
+
+					*Buffer = 0;		// null terminate
+
+					FCpuProfilerTrace::OutputBeginDynamicEvent(VerboseNameANSI.GetData());
+				}
+				else
+				{
+					FCpuProfilerTrace::OutputBeginDynamicEvent(InStatId.GetStatDescriptionANSI()); //todo: Could we use FName index as event id?
+				}
 				EmittedEvent |= TraceEvent;
 			}
 #endif
@@ -2003,6 +2039,14 @@ struct FStat_##StatName\
 #define CONDITIONAL_SCOPE_CYCLE_COUNTER(Stat,bCondition) \
 	FScopeCycleCounter CycleCount_##Stat(bCondition ? GET_STATID(Stat) : TStatId(), GET_STATFLAGS(Stat));
 
+#if CPUPROFILERTRACE_ENABLED
+#define SCOPE_CYCLE_COUNTER_VERBOSE(Stat,VerboseDescription) \
+	FScopeCycleCounter CycleCount_##Stat(GET_STATID(Stat), GET_STATFLAGS(Stat), false, VerboseDescription);
+#else
+#define SCOPE_CYCLE_COUNTER_VERBOSE(Stat,VerboseDescription) \
+	FScopeCycleCounter CycleCount_##Stat(GET_STATID(Stat), GET_STATFLAGS(Stat));
+#endif
+
 #define SCOPE_SECONDS_ACCUMULATOR(Stat) \
 	FSimpleScopeSecondsStat SecondsAccum_##Stat(GET_STATID(Stat));
 
diff --git a/Engine/Source/Runtime/D3D12RHI/Private/D3D12RHI.cpp b/Engine/Source/Runtime/D3D12RHI/Private/D3D12RHI.cpp
--- a/Engine/Source/Runtime/D3D12RHI/Private/D3D12RHI.cpp
+++ b/Engine/Source/Runtime/D3D12RHI/Private/D3D12RHI.cpp
@@ -215,11 +215,23 @@ FD3D12DynamicRHI::FD3D12DynamicRHI(const TArray<TSharedPtr<FD3D12Adapter>>& Chos
 
 	{
 		// Workaround for 4.14. Limit the number of GPU stats on D3D12 due to an issue with high memory overhead with render queries (Jira UE-38139)
-		//@TODO: Remove this when render query issues are fixed
+		// @TODO: Remove this when render query issues are fixed.
+		// 
+		// I think the high memory issues were fixed in 5.0 or some time before.  In the current implementation, a single 64K padded readback buffer is
+		// used for the query pool, with 8 bytes per entry.  So an increase from 1K to 8K should take no extra graphics memory, and this was confirmed
+		// by zero observed change to the D3D12RHI::Memory statistics.  Technically, this is the number of queries for four frames
+		// (NumGPUProfilerBufferedFrames == 4), and a real world test case with multiple views was hitting 880 queries per frame, so the number needed
+		// to be at least 3600 to not run out...  There will be some memory increase from FD3D12RenderQuery structures on the CPU, but it's fairly
+		// trivial (under 300K).
+		// 
+		// I considered removing the limit entirely, but that causes a different code path to run, which hasn't run in a long time on PC (over 5 years
+		// it has been in place), and I wasn't sure how to fully unit test the change, so I figured it was safer to simply increase the existing limit
+		// a bit.
+		//
 		static IConsoleVariable* GPUStatsEnabledCVar = IConsoleManager::Get().FindConsoleVariable(TEXT("r.GPUStatsMaxQueriesPerFrame"));
 		if (GPUStatsEnabledCVar)
 		{
-			GPUStatsEnabledCVar->Set(1024); // 1024*64KB = 64MB
+			GPUStatsEnabledCVar->Set(8192);
 		}
 	}
 
diff --git a/Engine/Source/Runtime/Engine/Classes/Engine/GameViewportClient.h b/Engine/Source/Runtime/Engine/Classes/Engine/GameViewportClient.h
--- a/Engine/Source/Runtime/Engine/Classes/Engine/GameViewportClient.h
+++ b/Engine/Source/Runtime/Engine/Classes/Engine/GameViewportClient.h
@@ -145,6 +145,9 @@ public:
 	/** Returns a relative world context for this viewport.	 */
 	virtual UWorld* GetWorld() const override;
 
+	/* Create the game viewport */
+	virtual FSceneViewport* CreateGameViewport(TSharedPtr<SViewport> InViewportWidget);
+
 	/* Returns the game viewport */
 	FSceneViewport* GetGameViewport();
 	const FSceneViewport* GetGameViewport() const;
diff --git a/Engine/Source/Runtime/Engine/Private/GameEngine.cpp b/Engine/Source/Runtime/Engine/Private/GameEngine.cpp
--- a/Engine/Source/Runtime/Engine/Private/GameEngine.cpp
+++ b/Engine/Source/Runtime/Engine/Private/GameEngine.cpp
@@ -268,7 +268,7 @@ void UGameEngine::CreateGameViewport( UGameViewportClient* GameViewportClient )
 		Window->SetOnWindowMoved( FOnWindowMoved::CreateUObject( this, &UGameEngine::OnGameWindowMoved ) );
 	}
 
-	SceneViewport = MakeShareable( new FSceneViewport( GameViewportClient, GameViewportWidgetRef ) );
+	SceneViewport = MakeShareable( GameViewportClient->CreateGameViewport(GameViewportWidgetRef) );
 	GameViewportClient->Viewport = SceneViewport.Get();
 	//GameViewportClient->CreateHighresScreenshotCaptureRegionWidget(); //  Disabled until mouse based input system can be made to work correctly.
 
diff --git a/Engine/Source/Runtime/Engine/Private/GameViewportClient.cpp b/Engine/Source/Runtime/Engine/Private/GameViewportClient.cpp
--- a/Engine/Source/Runtime/Engine/Private/GameViewportClient.cpp
+++ b/Engine/Source/Runtime/Engine/Private/GameViewportClient.cpp
@@ -398,6 +398,11 @@ void UGameViewportClient::DetachViewportClient()
 	RemoveFromRoot();
 }
 
+FSceneViewport* UGameViewportClient::CreateGameViewport(TSharedPtr<SViewport> InViewportWidget)
+{
+	return new FSceneViewport(this, InViewportWidget);
+}
+
 FSceneViewport* UGameViewportClient::GetGameViewport()
 {
 	if (Viewport && Viewport->GetViewportType() == NAME_SceneViewport)
diff --git a/Engine/Source/Runtime/Engine/Private/SceneView.cpp b/Engine/Source/Runtime/Engine/Private/SceneView.cpp
--- a/Engine/Source/Runtime/Engine/Private/SceneView.cpp
+++ b/Engine/Source/Runtime/Engine/Private/SceneView.cpp
@@ -2697,6 +2697,7 @@ FSceneViewFamily::FSceneViewFamily(const ConstructionValues& CVS)
 	GammaCorrection(CVS.GammaCorrection),
 	SecondaryViewFraction(1.0f),
 	SecondaryScreenPercentageMethod(ESecondaryScreenPercentageMethod::LowerPixelDensitySimulation),
+	ProfileSceneRenderTime(nullptr),
 	ScreenPercentageInterface(nullptr),
 	TemporalUpscalerInterface(nullptr),
 	PrimarySpatialUpscalerInterface(nullptr),
diff --git a/Engine/Source/Runtime/Engine/Private/UnrealEngine.cpp b/Engine/Source/Runtime/Engine/Private/UnrealEngine.cpp
--- a/Engine/Source/Runtime/Engine/Private/UnrealEngine.cpp
+++ b/Engine/Source/Runtime/Engine/Private/UnrealEngine.cpp
@@ -11576,7 +11576,7 @@ PRAGMA_ENABLE_DEPRECATION_WARNINGS
 		int32 Y = FMath::TruncToInt(TextureSize.Y * 0.20f) / Canvas->GetDPIScale();
 
 		// give the viewport first shot at drawing stats
-		Y = Viewport->DrawStatsHUD(Canvas, X, Y);
+		Y = Viewport->DrawStatsHUD(Canvas, StatsXOffset, Y);
 
 		// Render all the simple stats
 		GEngine->RenderEngineStats(World, Viewport, Canvas, StatsXOffset, MessageY, X, Y, &ViewLocation, &ViewRotation);
diff --git a/Engine/Source/Runtime/Engine/Public/SceneView.h b/Engine/Source/Runtime/Engine/Public/SceneView.h
--- a/Engine/Source/Runtime/Engine/Public/SceneView.h
+++ b/Engine/Source/Runtime/Engine/Public/SceneView.h
@@ -1768,6 +1768,16 @@ public:
 	bool bNullifyWorldSpacePosition;
 #endif
 
+	// Optional description of view family for Unreal Insights profiling
+	FString ProfileDescription;
+
+	/**
+	 * Optional tracking of scene render time in seconds, useful when rendering multiple scenes (such as via UDisplayClusterViewportClient)
+	 * where the user may want a live breakdown of where render time is going per scene, which the stat system otherwise doesn't provide.
+	 * The destination memory location needs to be persistent and safe to access in any thread (for example, a standalone allocated buffer). 
+	 */
+	float* ProfileSceneRenderTime;
+
 	/** Initialization constructor. */
 	FSceneViewFamily( const ConstructionValues& CVS );
 	~FSceneViewFamily();
diff --git a/Engine/Source/Runtime/RenderCore/Private/ProfilingDebugging/RealtimeGPUProfiler.cpp b/Engine/Source/Runtime/RenderCore/Private/ProfilingDebugging/RealtimeGPUProfiler.cpp
--- a/Engine/Source/Runtime/RenderCore/Private/ProfilingDebugging/RealtimeGPUProfiler.cpp
+++ b/Engine/Source/Runtime/RenderCore/Private/ProfilingDebugging/RealtimeGPUProfiler.cpp
@@ -6,6 +6,9 @@
 #include "RenderCore.h"
 #include "RenderingThread.h"
 #include "GPUProfiler.h"
+#include "Misc/ScopeRWLock.h"
+
+typedef TArray<TCHAR, TInlineAllocator<4096u>> FDescriptionStringBuffer;
 
 // Only exposed for debugging. Disabling this carries a severe performance penalty
 #define RENDER_QUERY_POOLING_ENABLED 1
@@ -165,6 +168,7 @@ public:
 		, StartQuery(RenderQueryPool.AllocateQuery())
 		, EndQuery(RenderQueryPool.AllocateQuery())
 		, FrameNumber(-1)
+		, DescriptionLength(0)
 #if DO_CHECK || USING_CODE_ANALYSIS
 		, bInsideQuery(false)
 #endif
@@ -272,6 +276,25 @@ public:
 		return Name;
 	}
 
+	const TCHAR* GetDescription(const FDescriptionStringBuffer& DescriptionStringBuffer, uint32& OutDescriptionLength) const
+	{
+		OutDescriptionLength = DescriptionLength;
+		return DescriptionLength ? &DescriptionStringBuffer[DescriptionOffset] : TEXT("");
+	}
+
+	void SetDescription(const TCHAR* Description, FDescriptionStringBuffer& DescriptionStringBuffer)
+	{
+		check(Description);
+		uint32 TestDescriptionLength = FCString::Strlen(Description);
+		if (TestDescriptionLength && (DescriptionStringBuffer.Num() + TestDescriptionLength <= UINT16_MAX))
+		{
+			DescriptionLength = (uint16)TestDescriptionLength;
+			DescriptionOffset = (uint16)DescriptionStringBuffer.Num();
+			DescriptionStringBuffer.AddUninitialized(DescriptionLength);
+			FMemory::Memcpy(&DescriptionStringBuffer[DescriptionOffset], Description, DescriptionLength * sizeof(TCHAR));
+		}
+	}
+
 	FRHIGPUMask GetGPUMask() const
 	{
 		return GPUMask;
@@ -311,6 +334,9 @@ private:
 
 	uint32 FrameNumber;
 
+	uint16 DescriptionOffset;		// Offset in DescriptionStringBuffer
+	uint16 DescriptionLength;
+
 #if DO_CHECK || USING_CODE_ANALYSIS
 	bool bInsideQuery;
 #endif
@@ -320,17 +346,36 @@ private:
 void TraverseEventTree(
 	const TArray<FRealtimeGPUProfilerEvent, TInlineAllocator<100u>>& GpuProfilerEvents,
 	const TArray<TArray<int32>>& GpuProfilerEventChildrenIndices,
+	const FDescriptionStringBuffer& DescriptionStringBuffer,
 	int32 Root,
 	uint32 GPUIndex)
 {
 	uint64 lastStartTime = 0;
 	uint64 lastEndTime = 0;
 
+	FName EventName;
+
 	if (Root != 0)
 	{
+		uint32 DescriptionLength;
+		const TCHAR* DescriptionData = GpuProfilerEvents[Root].GetDescription(DescriptionStringBuffer, DescriptionLength);
+		if (DescriptionLength)
+		{
+			FString NameWithDescription;
+			NameWithDescription = GpuProfilerEvents[Root].GetName().ToString();
+			NameWithDescription.Append(TEXT(" - "));
+			NameWithDescription.AppendChars(DescriptionData, DescriptionLength);
+
+			EventName = FName(NameWithDescription);
+		}
+		else
+		{
+			EventName = GpuProfilerEvents[Root].GetName();
+		}
+
 		check(GpuProfilerEvents[Root].GetGPUMask().Contains(GPUIndex));
-		FGpuProfilerTrace::SpecifyEventByName(GpuProfilerEvents[Root].GetName());
-		FGpuProfilerTrace::BeginEventByName(GpuProfilerEvents[Root].GetName(), GpuProfilerEvents[Root].GetFrameNumber(), GpuProfilerEvents[Root].GetStartResultMicroseconds(GPUIndex));
+		FGpuProfilerTrace::SpecifyEventByName(EventName);
+		FGpuProfilerTrace::BeginEventByName(EventName, GpuProfilerEvents[Root].GetFrameNumber(), GpuProfilerEvents[Root].GetStartResultMicroseconds(GPUIndex));
 	}
 
 	for (int32 Subroot : GpuProfilerEventChildrenIndices[Root])
@@ -348,14 +393,14 @@ void TraverseEventTree(
 				check(lastStartTime >= GpuProfilerEvents[Root].GetStartResultMicroseconds(GPUIndex));
 				check(lastEndTime <= GpuProfilerEvents[Root].GetEndResultMicroseconds(GPUIndex));
 			}
-			TraverseEventTree(GpuProfilerEvents, GpuProfilerEventChildrenIndices, Subroot, GPUIndex);
+			TraverseEventTree(GpuProfilerEvents, GpuProfilerEventChildrenIndices, DescriptionStringBuffer, Subroot, GPUIndex);
 		}
 	}
 
 	if (Root != 0)
 	{
 		check(GpuProfilerEvents[Root].GetGPUMask().Contains(GPUIndex));
-		FGpuProfilerTrace::SpecifyEventByName(GpuProfilerEvents[Root].GetName());
+		FGpuProfilerTrace::SpecifyEventByName(EventName);
 		FGpuProfilerTrace::EndEvent(GpuProfilerEvents[Root].GetEndResultMicroseconds(GPUIndex));
 	}
 }
@@ -396,6 +441,8 @@ public:
 		EventAggregates.AddUninitialized();
 
 		CPUFrameStartTimestamp = FPlatformTime::Cycles64();
+
+		DescriptionStringBuffer.Empty();
 	}
 
 	~FRealtimeGPUProfilerFrame()
@@ -418,6 +465,8 @@ public:
 
 		EventAggregates.Reset();
 		EventAggregates.AddUninitialized();
+
+		DescriptionStringBuffer.Empty();
 	}
 
 	int32 GetCurrentEventIndex() const
@@ -435,7 +484,7 @@ public:
 		EventStack.Pop(false);
 	}
 
-	FRealtimeGPUProfilerQuery PushEvent(FRHIGPUMask GPUMask, const FName& Name, const FName& StatName)
+	FRealtimeGPUProfilerQuery PushEvent(FRHIGPUMask GPUMask, const FName& Name, const FName& StatName, const TCHAR* Description)
 	{
 		if (NextEventIdx >= GpuProfilerEvents.Num())
 		{
@@ -457,6 +506,10 @@ public:
 
 		GpuProfilerEventParentIndices.Add(EventStack.Last());
 		EventStack.Push(EventIdx);
+		if (Description)
+		{
+			GpuProfilerEvents[EventIdx].SetDescription(Description, DescriptionStringBuffer);
+		}
 		return GpuProfilerEvents[EventIdx].Begin(GPUMask, Name, StatName);
 	}
 
@@ -473,7 +526,11 @@ public:
 		return GpuProfilerEvents[EventIdx].End();
 	}
 
-	bool UpdateStats(FRHICommandListImmediate& RHICmdList)
+	bool UpdateStats(FRHICommandListImmediate& RHICmdList
+#if GPUPROFILERTRACE_ENABLED
+		, FRealtimeGPUProfilerHistoryByDescription& HistoryByDescription
+#endif
+		)
 	{
 		// Gather any remaining results and check all the results are ready
 		const int32 NumEventsThisFramePlusOne = NextEventIdx;
@@ -648,9 +705,81 @@ PRAGMA_ENABLE_DEPRECATION_WARNINGS
 		for (uint32 GPUIndex = 0; GPUIndex < GNumExplicitGPUsForRendering; ++GPUIndex)
 		{
 			FGpuProfilerTrace::BeginFrame(Timestamps[GPUIndex]);
-			TraverseEventTree(GpuProfilerEvents, GpuProfilerEventChildrenIndices, 0, GPUIndex);
+			TraverseEventTree(GpuProfilerEvents, GpuProfilerEventChildrenIndices, DescriptionStringBuffer, 0, GPUIndex);
 			FGpuProfilerTrace::EndFrame(GPUIndex);
 		}
+
+		// Logic to track performance by description for root level items.  For example, if rendering multiple view families via
+		// UDisplayClusterViewportClient, each view family will have a description, and clients may want to turn render features
+		// on or off per view family to tune performance, or select which GPU each view family renders on to balance performance.
+		// The regular GPU render stats screen shows the sum of performance across all view families, and across all GPUs, which
+		// isn't terribly useful for this purpose.  The alternative would be to use Unreal Insights, but it takes a lot of work
+		// to get clean measurements there due to noise, which history averaging smooths out (plus it requires more knowledge to
+		// know how to interpret Insights).
+		{
+			FRWScopeLock Lock(HistoryByDescription.Mutex, SLT_Write);
+
+			// To clean up old descriptions, we first want to mark all existing descriptions as not updated this frame.
+			for (auto Iterator = HistoryByDescription.History.CreateIterator(); Iterator; ++Iterator)
+			{
+				Iterator.Value().UpdatedThisFrame = false;
+			}
+
+			// Then scan for root items with descriptions and add history entries for them
+			for (int32 Subroot : GpuProfilerEventChildrenIndices[0])
+			{
+				uint32 DescriptionLength;
+				const TCHAR* DescriptionData = GpuProfilerEvents[Subroot].GetDescription(DescriptionStringBuffer, DescriptionLength);
+
+				if (DescriptionLength)
+				{
+					FString Description;
+					Description.AppendChars(DescriptionData, DescriptionLength);
+
+					FRealtimeGPUProfilerHistoryItem& HistoryItem = HistoryByDescription.History.FindOrAdd(Description);
+
+					// We could have more than one root entry for a given view -- advance history and subtract out previously accumulated time
+					// the first time the given item is accessed on a frame, then accumulate from there.
+					uint64* HistoryTime;
+					if (!HistoryItem.UpdatedThisFrame)
+					{
+						HistoryItem.UpdatedThisFrame = true;
+						HistoryItem.LastGPUMask = GpuProfilerEvents[Subroot].GetGPUMask();
+
+						HistoryItem.NextWriteIndex++;
+						HistoryTime = &HistoryItem.Times[(HistoryItem.NextWriteIndex - 1) % FRealtimeGPUProfilerHistoryItem::HistoryCount];
+
+						HistoryItem.AccumulatedTime -= *HistoryTime;
+						*HistoryTime = 0;
+					}
+					else
+					{
+						HistoryTime = &HistoryItem.Times[(HistoryItem.NextWriteIndex - 1) % FRealtimeGPUProfilerHistoryItem::HistoryCount];
+					}
+
+					// If multiple GPU masks, get the one with the largest time span
+					uint64 MaxGpuTime = 0;
+
+					for (uint32 GPUIndex : GpuProfilerEvents[Subroot].GetGPUMask())
+					{
+						MaxGpuTime = FMath::Max(MaxGpuTime, GpuProfilerEvents[Subroot].GetEndResultMicroseconds(GPUIndex) - GpuProfilerEvents[Subroot].GetStartResultMicroseconds(GPUIndex));
+					}
+
+					// Add that to the accumulated and history result
+					HistoryItem.AccumulatedTime += MaxGpuTime;
+					*HistoryTime += MaxGpuTime;
+				}
+			}
+
+			// Finally, clean up any items that weren't updated this frame
+			for (auto Iterator = HistoryByDescription.History.CreateIterator(); Iterator; ++Iterator)
+			{
+				if (!Iterator.Value().UpdatedThisFrame)
+				{
+					Iterator.RemoveCurrent();
+				}
+			}
+		}
 #endif
 
 		return true;
@@ -785,6 +914,7 @@ private:
 	TArray<int32, TInlineAllocator<GPredictedMaxNumEvents>> GpuProfilerEventParentIndices;
 	TArray<int32, TInlineAllocator<GPredictedMaxStackDepth>> EventStack;
 	TArray<FGPUEventTimeAggregate, TInlineAllocator<GPredictedMaxNumEvents>> EventAggregates;
+	FDescriptionStringBuffer DescriptionStringBuffer;
 };
 
 /*-----------------------------------------------------------------------------
@@ -820,7 +950,7 @@ FRealtimeGPUProfiler::FRealtimeGPUProfiler()
 	if (GSupportsTimestampRenderQueries)
 	{
 		const int MaxGPUQueries = CVarGPUStatsMaxQueriesPerFrame.GetValueOnRenderThread();
-		RenderQueryPool = RHICreateRenderQueryPool(RQT_AbsoluteTime, (MaxGPUQueries > 0) ? MaxGPUQueries * 2 : UINT32_MAX);
+		RenderQueryPool = RHICreateRenderQueryPool(RQT_AbsoluteTime, (MaxGPUQueries > 0) ? MaxGPUQueries : UINT32_MAX);
 		for (int Index = 0; Index < NumGPUProfilerBufferedFrames; Index++)
 		{
 			Frames.Add(new FRealtimeGPUProfilerFrame(RenderQueryPool, QueryCount));
@@ -917,7 +1047,11 @@ void FRealtimeGPUProfiler::EndFrame(FRHICommandListImmediate& RHICmdList)
 	check(bInBeginEndBlock == true);
 	bInBeginEndBlock = false;
 
-	if (Frames[ReadBufferIndex]->UpdateStats(RHICmdList))
+	if (Frames[ReadBufferIndex]->UpdateStats(RHICmdList
+#if GPUPROFILERTRACE_ENABLED
+		, HistoryByDescription
+#endif
+		))
 	{
 		// On a successful read, advance the ReadBufferIndex and WriteBufferIndex and clear the frame we just read
 		Frames[ReadBufferIndex]->Clear(&RHICmdList);
@@ -933,7 +1067,7 @@ void FRealtimeGPUProfiler::EndFrame(FRHICommandListImmediate& RHICmdList)
 	}
 }
 
-FRealtimeGPUProfilerQuery FRealtimeGPUProfiler::PushEvent(FRHIGPUMask GPUMask, const FName& Name, const FName& StatName)
+FRealtimeGPUProfilerQuery FRealtimeGPUProfiler::PushEvent(FRHIGPUMask GPUMask, const FName& Name, const FName& StatName, const TCHAR* Description)
 {
 	check(IsInParallelRenderingThread());
 	if (bStatGatheringPaused || !bInBeginEndBlock)
@@ -943,7 +1077,7 @@ FRealtimeGPUProfilerQuery FRealtimeGPUProfiler::PushEvent(FRHIGPUMask GPUMask, c
 	check(Frames.Num() > 0);
 	if (WriteBufferIndex >= 0)
 	{
-		return Frames[WriteBufferIndex]->PushEvent(GPUMask, Name, StatName);
+		return Frames[WriteBufferIndex]->PushEvent(GPUMask, Name, StatName, Description);
 	}
 	return {};
 }
@@ -988,9 +1122,9 @@ void FRealtimeGPUProfiler::PopEventOverride()
 	}
 }
 
-void FRealtimeGPUProfiler::PushStat(FRHICommandListImmediate& RHICmdList, const FName& Name, const FName& StatName, int32 (*InNumDrawCallsPtr)[MAX_NUM_GPUS])
+void FRealtimeGPUProfiler::PushStat(FRHICommandListImmediate& RHICmdList, const FName& Name, const FName& StatName, const TCHAR* Description, int32 (*InNumDrawCallsPtr)[MAX_NUM_GPUS])
 {
-	PushEvent(RHICmdList.GetGPUMask(), Name, StatName).Submit(RHICmdList);
+	PushEvent(RHICmdList.GetGPUMask(), Name, StatName, Description).Submit(RHICmdList);
 
 	if (InNumDrawCallsPtr && (**InNumDrawCallsPtr) != -1)
 	{
@@ -1017,7 +1151,7 @@ void FRealtimeGPUProfiler::PopStat(FRHICommandListImmediate& RHICmdList, int32 (
 /*-----------------------------------------------------------------------------
 FScopedGPUStatEvent
 -----------------------------------------------------------------------------*/
-void FScopedGPUStatEvent::Begin(FRHICommandList& InRHICmdList, const FName& Name, const FName& StatName, int32 (*InNumDrawCallsPtr)[MAX_NUM_GPUS])
+void FScopedGPUStatEvent::Begin(FRHICommandList& InRHICmdList, const FName& Name, const FName& StatName, const TCHAR* Description, int32 (*InNumDrawCallsPtr)[MAX_NUM_GPUS])
 {
 	check(IsInRenderingThread());
 	if (!AreGPUStatsEnabled())
@@ -1030,7 +1164,7 @@ void FScopedGPUStatEvent::Begin(FRHICommandList& InRHICmdList, const FName& Name
 	{
 		NumDrawCallsPtr = InNumDrawCallsPtr;
 		RHICmdList = &static_cast<FRHICommandListImmediate&>(InRHICmdList);
-		FRealtimeGPUProfiler::Get()->PushStat(*RHICmdList, Name, StatName, InNumDrawCallsPtr);
+		FRealtimeGPUProfiler::Get()->PushStat(*RHICmdList, Name, StatName, Description, InNumDrawCallsPtr);
 	}
 }
 
@@ -1047,3 +1181,40 @@ void FScopedGPUStatEvent::End()
 	}
 }
 #endif // HAS_GPU_STATS
+
+#if GPUPROFILERTRACE_ENABLED
+FRealtimeGPUProfilerHistoryItem::FRealtimeGPUProfilerHistoryItem()
+{
+	FMemory::Memset(*this, 0);
+}
+
+void FRealtimeGPUProfiler::FetchPerfByDescription(TArray<FRealtimeGPUProfilerDescriptionResult>& OutResults) const
+{
+	FRWScopeLock Lock(HistoryByDescription.Mutex, SLT_ReadOnly);
+
+	OutResults.Empty(HistoryByDescription.History.Num());
+
+	for (auto Iterator = HistoryByDescription.History.CreateConstIterator(); Iterator; ++Iterator)
+	{
+		FRealtimeGPUProfilerDescriptionResult Result;
+		Result.Description = Iterator.Key();
+
+		const FRealtimeGPUProfilerHistoryItem& HistoryValue = Iterator.Value();
+		const uint64 ClampedTimeCount = FMath::Min(HistoryValue.NextWriteIndex, FRealtimeGPUProfilerHistoryItem::HistoryCount);
+
+		Result.GPUMask = HistoryValue.LastGPUMask;
+		Result.AverageTime = HistoryValue.AccumulatedTime / ClampedTimeCount;
+		Result.MinTime = INT64_MAX;
+		Result.MaxTime = 0;
+
+		for (uint64 TimeIndex = 0; TimeIndex < ClampedTimeCount; TimeIndex++)
+		{
+			Result.MinTime = FMath::Min(Result.MinTime, HistoryValue.Times[TimeIndex]);
+			Result.MaxTime = FMath::Max(Result.MaxTime, HistoryValue.Times[TimeIndex]);
+		}
+
+		OutResults.Add(Result);
+	}
+}
+#endif  // GPUPROFILERTRACE_ENABLED
+
diff --git a/Engine/Source/Runtime/RenderCore/Private/RenderGraphEvent.cpp b/Engine/Source/Runtime/RenderCore/Private/RenderGraphEvent.cpp
--- a/Engine/Source/Runtime/RenderCore/Private/RenderGraphEvent.cpp
+++ b/Engine/Source/Runtime/RenderCore/Private/RenderGraphEvent.cpp
@@ -186,10 +186,10 @@ FRDGEventScopeOpArray FRDGEventScopeStack::CompilePassEpilogue()
 	return MoveTemp(Ops);
 }
 
-FRDGGPUStatScopeGuard::FRDGGPUStatScopeGuard(FRDGBuilder& InGraphBuilder, const FName& Name, const FName& StatName, int32(*NumDrawCallsPtr)[MAX_NUM_GPUS])
+FRDGGPUStatScopeGuard::FRDGGPUStatScopeGuard(FRDGBuilder& InGraphBuilder, const FName& Name, const FName& StatName, const TCHAR* Description, int32(*NumDrawCallsPtr)[MAX_NUM_GPUS])
 	: GraphBuilder(InGraphBuilder)
 {
-	GraphBuilder.GPUScopeStacks.BeginStatScope(Name, StatName, NumDrawCallsPtr);
+	GraphBuilder.GPUScopeStacks.BeginStatScope(Name, StatName, Description, NumDrawCallsPtr);
 }
 
 FRDGGPUStatScopeGuard::~FRDGGPUStatScopeGuard()
@@ -208,7 +208,7 @@ FRDGGPUStatScopeOpArray::FRDGGPUStatScopeOpArray(TRDGScopeOpArray<FRDGGPUStatSco
 
 		if (Op.IsPush())
 		{
-			Op.Query = FRealtimeGPUProfiler::Get()->PushEvent(GPUMask, Op.Scope->Name, Op.Scope->StatName);
+			Op.Query = FRealtimeGPUProfiler::Get()->PushEvent(GPUMask, Op.Scope->Name, Op.Scope->StatName, *Op.Scope->Description);
 		}
 		else
 		{
diff --git a/Engine/Source/Runtime/RenderCore/Public/ProfilingDebugging/RealtimeGPUProfiler.h b/Engine/Source/Runtime/RenderCore/Public/ProfilingDebugging/RealtimeGPUProfiler.h
--- a/Engine/Source/Runtime/RenderCore/Public/ProfilingDebugging/RealtimeGPUProfiler.h
+++ b/Engine/Source/Runtime/RenderCore/Public/ProfilingDebugging/RealtimeGPUProfiler.h
@@ -276,9 +276,11 @@ struct RENDERCORE_API FBreadcrumbEvent
  #define DEFINE_GPU_STAT(StatName) DEFINE_STAT(Stat_GPU_##StatName); CSV_DEFINE_STAT(GPU,StatName); FDrawCallCategoryName DrawcallCountCategory_##StatName;
  #define DEFINE_GPU_DRAWCALL_STAT(StatName) DEFINE_STAT(Stat_GPU_##StatName); CSV_DEFINE_STAT(GPU,StatName); FDrawCallCategoryName DrawcallCountCategory_##StatName((TCHAR*)TEXT(#StatName));
 #if STATS
-  #define SCOPED_GPU_STAT(RHICmdList, StatName) FScopedGPUStatEvent PREPROCESSOR_JOIN(GPUStatEvent_##StatName,__LINE__); PREPROCESSOR_JOIN(GPUStatEvent_##StatName,__LINE__).Begin(RHICmdList, CSV_STAT_FNAME(StatName), GET_STATID( Stat_GPU_##StatName ).GetName(), &DrawcallCountCategory_##StatName.Counters);
- #else
-  #define SCOPED_GPU_STAT(RHICmdList, StatName) FScopedGPUStatEvent PREPROCESSOR_JOIN(GPUStatEvent_##StatName,__LINE__); PREPROCESSOR_JOIN(GPUStatEvent_##StatName,__LINE__).Begin(RHICmdList, CSV_STAT_FNAME(StatName), FName(), &DrawcallCountCategory_##StatName.Counters);
+  #define SCOPED_GPU_STAT(RHICmdList, StatName) FScopedGPUStatEvent PREPROCESSOR_JOIN(GPUStatEvent_##StatName,__LINE__); PREPROCESSOR_JOIN(GPUStatEvent_##StatName,__LINE__).Begin(RHICmdList, CSV_STAT_FNAME(StatName), GET_STATID( Stat_GPU_##StatName ).GetName(), nullptr, &DrawcallCountCategory_##StatName.Counters);
+  #define SCOPED_GPU_STAT_VERBOSE(RHICmdList, StatName, Description) FScopedGPUStatEvent PREPROCESSOR_JOIN(GPUStatEvent_##StatName,__LINE__); PREPROCESSOR_JOIN(GPUStatEvent_##StatName,__LINE__).Begin(RHICmdList, CSV_STAT_FNAME(StatName), GET_STATID( Stat_GPU_##StatName ).GetName(), Description, &DrawcallCountCategory_##StatName.Counters);
+#else
+  #define SCOPED_GPU_STAT(RHICmdList, StatName) FScopedGPUStatEvent PREPROCESSOR_JOIN(GPUStatEvent_##StatName,__LINE__); PREPROCESSOR_JOIN(GPUStatEvent_##StatName,__LINE__).Begin(RHICmdList, CSV_STAT_FNAME(StatName), FName(), nullptr, &DrawcallCountCategory_##StatName.Counters);
+  #define SCOPED_GPU_STAT_VERBOSE(RHICmdList, StatName, Description) FScopedGPUStatEvent PREPROCESSOR_JOIN(GPUStatEvent_##StatName,__LINE__); PREPROCESSOR_JOIN(GPUStatEvent_##StatName,__LINE__).Begin(RHICmdList, CSV_STAT_FNAME(StatName), FName(), Description, &DrawcallCountCategory_##StatName.Counters);
 #endif
  #define GPU_STATS_BEGINFRAME(RHICmdList) FRealtimeGPUProfiler::Get()->BeginFrame(RHICmdList);
  #define GPU_STATS_ENDFRAME(RHICmdList) FRealtimeGPUProfiler::Get()->EndFrame(RHICmdList);
@@ -327,6 +329,39 @@ private:
 	FRHIRenderQuery* Query{};
 };
 
+#if GPUPROFILERTRACE_ENABLED
+struct FRealtimeGPUProfilerHistoryItem
+{
+	FRealtimeGPUProfilerHistoryItem();
+
+	static const uint64 HistoryCount = 64;
+
+	// Constructor memsets everything to zero, assuming structure is Plain Old Data.  If any dynamic structures are
+	// added, you'll need a more generalized constructor that zeroes out all the uninitialized data.
+	bool UpdatedThisFrame;
+	FRHIGPUMask LastGPUMask;
+	uint64 NextWriteIndex;
+	uint64 AccumulatedTime;				// Accumulated time could be computed, but may also be useful to inspect in the debugger
+	TStaticArray<uint64, HistoryCount> Times;
+};
+
+struct FRealtimeGPUProfilerHistoryByDescription
+{
+	TMap<FString, FRealtimeGPUProfilerHistoryItem> History;
+	mutable FRWLock Mutex;
+};
+
+struct RENDERCORE_API FRealtimeGPUProfilerDescriptionResult
+{
+	// Times are in microseconds
+	FString Description;
+	FRHIGPUMask GPUMask;
+	uint64 AverageTime;
+	uint64 MinTime;
+	uint64 MaxTime;
+};
+#endif  // GPUPROFILERTRACE_ENABLED
+
 /**
 * FRealtimeGPUProfiler class. This manages recording and reporting all for GPU stats
 */
@@ -349,7 +384,7 @@ public:
 	RENDERCORE_API void Release();
 
 	/** Push/pop events */
-	FRealtimeGPUProfilerQuery PushEvent(FRHIGPUMask GPUMask, const FName& Name, const FName& StatName);
+	FRealtimeGPUProfilerQuery PushEvent(FRHIGPUMask GPUMask, const FName& Name, const FName& StatName, const TCHAR* Description);
 	FRealtimeGPUProfilerQuery PopEvent();
 
 	int32 GetCurrentEventIndex() const;
@@ -358,9 +393,13 @@ public:
 	void PopEventOverride();
 
 	/** Push/pop stats which do additional draw call tracking on top of events. */
-	void PushStat(FRHICommandListImmediate& RHICmdList, const FName& Name, const FName& StatName, int32 (*InNumDrawCallsPtr)[MAX_NUM_GPUS]);
+	void PushStat(FRHICommandListImmediate& RHICmdList, const FName& Name, const FName& StatName, const TCHAR* Description, int32 (*InNumDrawCallsPtr)[MAX_NUM_GPUS]);
 	void PopStat(FRHICommandListImmediate& RHICmdList, int32 (*InNumDrawCallsPtr)[MAX_NUM_GPUS]);
 
+#if GPUPROFILERTRACE_ENABLED
+	RENDERCORE_API void FetchPerfByDescription(TArray<FRealtimeGPUProfilerDescriptionResult> & OutResults) const;
+#endif
+
 private:
 	FRealtimeGPUProfiler();
 
@@ -379,6 +418,10 @@ private:
 	bool bStatGatheringPaused;
 	bool bInBeginEndBlock;
 	bool bLocked = false;
+
+#if GPUPROFILERTRACE_ENABLED
+	FRealtimeGPUProfilerHistoryByDescription HistoryByDescription;
+#endif
 };
 
 /**
@@ -411,7 +454,7 @@ public:
 	/**
 	* Start/Stop functions for timer stats
 	*/
-	RENDERCORE_API void Begin(FRHICommandList& InRHICmdList, const FName& Name, const FName& StatName, int32 (*InNumDrawCallsPtr)[MAX_NUM_GPUS]);
+	RENDERCORE_API void Begin(FRHICommandList& InRHICmdList, const FName& Name, const FName& StatName, const TCHAR* Description, int32 (*InNumDrawCallsPtr)[MAX_NUM_GPUS]);
 	RENDERCORE_API void End();
 };
 #endif // HAS_GPU_STATS
diff --git a/Engine/Source/Runtime/RenderCore/Public/RenderGraphEvent.h b/Engine/Source/Runtime/RenderCore/Public/RenderGraphEvent.h
--- a/Engine/Source/Runtime/RenderCore/Public/RenderGraphEvent.h
+++ b/Engine/Source/Runtime/RenderCore/Public/RenderGraphEvent.h
@@ -26,9 +26,11 @@
 
 #if HAS_GPU_STATS
 	#if STATS
-		#define RDG_GPU_STAT_SCOPE(GraphBuilder, StatName) FRDGGPUStatScopeGuard PREPROCESSOR_JOIN(__RDG_GPUStatEvent_##StatName,__LINE__) ((GraphBuilder), CSV_STAT_FNAME(StatName), GET_STATID(Stat_GPU_##StatName).GetName(), &DrawcallCountCategory_##StatName.Counters);
+		#define RDG_GPU_STAT_SCOPE(GraphBuilder, StatName) FRDGGPUStatScopeGuard PREPROCESSOR_JOIN(__RDG_GPUStatEvent_##StatName,__LINE__) ((GraphBuilder), CSV_STAT_FNAME(StatName), GET_STATID(Stat_GPU_##StatName).GetName(), nullptr, &DrawcallCountCategory_##StatName.Counters);
+		#define RDG_GPU_STAT_SCOPE_VERBOSE(GraphBuilder, StatName, Description) FRDGGPUStatScopeGuard PREPROCESSOR_JOIN(__RDG_GPUStatEvent_##StatName,__LINE__) ((GraphBuilder), CSV_STAT_FNAME(StatName), GET_STATID(Stat_GPU_##StatName).GetName(), Description, &DrawcallCountCategory_##StatName.Counters);
 	#else
-		#define RDG_GPU_STAT_SCOPE(GraphBuilder, StatName) FRDGGPUStatScopeGuard PREPROCESSOR_JOIN(__RDG_GPUStatEvent_##StatName,__LINE__) ((GraphBuilder), CSV_STAT_FNAME(StatName), FName(), &DrawcallCountCategory_##StatName.Counters);
+		#define RDG_GPU_STAT_SCOPE(GraphBuilder, StatName) FRDGGPUStatScopeGuard PREPROCESSOR_JOIN(__RDG_GPUStatEvent_##StatName,__LINE__) ((GraphBuilder), CSV_STAT_FNAME(StatName), FName(), nullptr, &DrawcallCountCategory_##StatName.Counters);
+		#define RDG_GPU_STAT_SCOPE_VERBOSE(GraphBuilder, StatName, Description) FRDGGPUStatScopeGuard PREPROCESSOR_JOIN(__RDG_GPUStatEvent_##StatName,__LINE__) ((GraphBuilder), CSV_STAT_FNAME(StatName), FName(), Description, &DrawcallCountCategory_##StatName.Counters);
 	#endif
 #else
 	#define RDG_GPU_STAT_SCOPE(GraphBuilder, StatName)
@@ -447,16 +449,22 @@ private:
 class FRDGGPUStatScope final
 {
 public:
-	FRDGGPUStatScope(const FRDGGPUStatScope* InParentScope, const FName& InName, const FName& InStatName, int32 (*InDrawCallCounter)[MAX_NUM_GPUS])
+	FRDGGPUStatScope(const FRDGGPUStatScope* InParentScope, const FName& InName, const FName& InStatName, const TCHAR* InDescription, int32 (*InDrawCallCounter)[MAX_NUM_GPUS])
 		: ParentScope(InParentScope)
 		, Name(InName)
 		, StatName(InStatName)
 		, DrawCallCounter(InDrawCallCounter)
-	{}
+	{
+		if (InDescription)
+		{
+			Description = InDescription;
+		}
+	}
 
 	const FRDGGPUStatScope* const ParentScope;
 	const FName Name;
 	const FName StatName;
+	FString Description;
 	int32 (*DrawCallCounter)[MAX_NUM_GPUS];
 };
 
@@ -505,12 +513,12 @@ public:
 #endif
 	{}
 
-	inline void BeginScope(const FName& Name, const FName& StatName, int32(*DrawCallCounter)[MAX_NUM_GPUS])
+	inline void BeginScope(const FName& Name, const FName& StatName, const TCHAR* Description, int32(*DrawCallCounter)[MAX_NUM_GPUS])
 	{
 		if (IsEnabled())
 		{
 			check(DrawCallCounter != nullptr);
-			ScopeStack.BeginScope(Name, StatName, DrawCallCounter);
+			ScopeStack.BeginScope(Name, StatName, Description, DrawCallCounter);
 		}
 	}
 
@@ -565,7 +573,7 @@ private:
 class RENDERCORE_API FRDGGPUStatScopeGuard final
 {
 public:
-	FRDGGPUStatScopeGuard(FRDGBuilder& InGraphBuilder, const FName& Name, const FName& StatName, int32 (*DrawCallCounter)[MAX_NUM_GPUS]);
+	FRDGGPUStatScopeGuard(FRDGBuilder& InGraphBuilder, const FName& Name, const FName& StatName, const TCHAR* Description, int32 (*DrawCallCounter)[MAX_NUM_GPUS]);
 	FRDGGPUStatScopeGuard(const FRDGGPUStatScopeGuard&) = delete;
 	~FRDGGPUStatScopeGuard();
 
@@ -659,9 +667,9 @@ struct RENDERCORE_API FRDGGPUScopeStacksByPipeline
 		AsyncCompute.Event.EndScope();
 	}
 
-	inline void BeginStatScope(const FName& Name, const FName& StatName, int32(*DrawCallCounter)[MAX_NUM_GPUS])
+	inline void BeginStatScope(const FName& Name, const FName& StatName, const TCHAR* Description, int32(*DrawCallCounter)[MAX_NUM_GPUS])
 	{
-		Graphics.Stat.BeginScope(Name, StatName, DrawCallCounter);
+		Graphics.Stat.BeginScope(Name, StatName, Description, DrawCallCounter);
 	}
 
 	inline void EndStatScope()
diff --git a/Engine/Source/Runtime/Renderer/Private/DeferredShadingRenderer.cpp b/Engine/Source/Runtime/Renderer/Private/DeferredShadingRenderer.cpp
--- a/Engine/Source/Runtime/Renderer/Private/DeferredShadingRenderer.cpp
+++ b/Engine/Source/Runtime/Renderer/Private/DeferredShadingRenderer.cpp
@@ -2029,7 +2029,7 @@ void FDeferredShadingSceneRenderer::Render(FRDGBuilder& GraphBuilder)
 	}
 
 	RDG_EVENT_SCOPE(GraphBuilder, "Scene");
-	RDG_GPU_STAT_SCOPE(GraphBuilder, Unaccounted);
+	RDG_GPU_STAT_SCOPE_VERBOSE(GraphBuilder, Unaccounted, *ViewFamily.ProfileDescription);
 	
 	{
 		SCOPE_CYCLE_COUNTER(STAT_FDeferredShadingSceneRenderer_Render_Init);
diff --git a/Engine/Source/Runtime/Renderer/Private/SceneRendering.cpp b/Engine/Source/Runtime/Renderer/Private/SceneRendering.cpp
--- a/Engine/Source/Runtime/Renderer/Private/SceneRendering.cpp
+++ b/Engine/Source/Runtime/Renderer/Private/SceneRendering.cpp
@@ -2979,7 +2979,7 @@ void FSceneRenderer::DoCrossGPUTransfers(FRDGBuilder& GraphBuilder, FRHIGPUMask
 						{
 							if (!ViewInfo.GPUMask.Contains(RenderTargetGPUIndex))
 							{
-								FTransferResourceParams Param(static_cast<FRHITexture2D*>(ViewFamilyTexture->GetRHI()), TransferRect, ViewInfo.GPUMask.GetFirstIndex(), RenderTargetGPUIndex, true, true);
+								FTransferResourceParams Param(static_cast<FRHITexture2D*>(ViewFamilyTexture->GetRHI()), TransferRect, ViewInfo.GPUMask.GetFirstIndex(), RenderTargetGPUIndex, false, false);
 								RHICmdList.TransferResources(TArrayView<const FTransferResourceParams>(&Param, 1));
 							}
 						}
@@ -4052,8 +4052,8 @@ static void RenderViewFamily_RenderThread(FRHICommandListImmediate& RHICmdList,
 	}
 
 	{
-		SCOPE_CYCLE_COUNTER(STAT_TotalSceneRenderingTime);
-	
+		SCOPE_CYCLE_COUNTER_VERBOSE(STAT_TotalSceneRenderingTime, ViewFamily.ProfileDescription.IsEmpty() ? nullptr : *ViewFamily.ProfileDescription);
+
 		{
 			const ERHIFeatureLevel::Type FeatureLevel = SceneRenderer->FeatureLevel;
 
@@ -4333,14 +4333,20 @@ void FRendererModule::BeginRenderingViewFamily(FCanvas* Canvas, FSceneViewFamily
 
 		SceneRenderer->ViewFamily.DisplayInternalsData.Setup(World);
 
-		const uint32 DrawSceneEnqueue = FPlatformTime::Cycles();
+		const uint64 DrawSceneEnqueue = FPlatformTime::Cycles64();
 		ENQUEUE_RENDER_COMMAND(FDrawSceneCommand)(
 			[SceneRenderer, DrawSceneEnqueue](FRHICommandListImmediate& RHICmdList)
 			{
-				const float StartDelayMillisec = FPlatformTime::ToMilliseconds(FPlatformTime::Cycles() - DrawSceneEnqueue);
+				uint64 SceneRenderStart = FPlatformTime::Cycles64();
+				const float StartDelayMillisec = FPlatformTime::ToMilliseconds64(SceneRenderStart - DrawSceneEnqueue);
 				CSV_CUSTOM_STAT_GLOBAL(DrawSceneCommand_StartDelay, StartDelayMillisec, ECsvCustomStatOp::Set);
 				RenderViewFamily_RenderThread(RHICmdList, SceneRenderer);
 				FlushPendingDeleteRHIResources_RenderThread();
+
+				if (SceneRenderer->ViewFamily.ProfileSceneRenderTime)
+				{
+					*SceneRenderer->ViewFamily.ProfileSceneRenderTime = (float)FPlatformTime::ToSeconds64(FPlatformTime::Cycles64() - SceneRenderStart);
+				}
 			});
 
 		// Force kick the RT if we've got RT polling on.
