UE5_RELEASE:  Compile error fix for CL 18591295.

diff --git a/Engine/Source/Runtime/RenderCore/Public/RenderGraphEvent.h b/Engine/Source/Runtime/RenderCore/Public/RenderGraphEvent.h
--- a/Engine/Source/Runtime/RenderCore/Public/RenderGraphEvent.h
+++ b/Engine/Source/Runtime/RenderCore/Public/RenderGraphEvent.h
@@ -34,6 +34,7 @@
 	#endif
 #else
 	#define RDG_GPU_STAT_SCOPE(GraphBuilder, StatName)
+	#define RDG_GPU_STAT_SCOPE_VERBOSE(GraphBuilder, StatName, Description)
 #endif
 
 #if CSV_PROFILER
