Misc render-related changes:
* Added single-callback version of CreateStructuredBuffer to automatically infer element size, total size, element count and data pointer (e.g. CreateStructuredBuffer(GraphBuilder, TEXT("MyBuffer"), [&]() -> auto& { return BufferSource; });, where BufferSource is a TArray)
* Added support for uint2 shader parameters
* Added ForceTracking flag to ERDGBufferFlags/ERDGTextureFlags : force the RDG to track a resource even if it can be considered as readonly (no UAV, no RTV, etc.) This allows the graph to copy from and to external textures, and handling the corresponding transitions, for example.

#tests editor

diff --git a/Engine/Source/Runtime/RenderCore/Private/RenderGraphBuilder.cpp b/Engine/Source/Runtime/RenderCore/Private/RenderGraphBuilder.cpp
--- a/Engine/Source/Runtime/RenderCore/Private/RenderGraphBuilder.cpp
+++ b/Engine/Source/Runtime/RenderCore/Private/RenderGraphBuilder.cpp
@@ -686,7 +686,8 @@ FRDGTextureRef FRDGBuilder::RegisterExternalTexture(
 	const FRDGTextureDesc Desc = Translate(ExternalPooledTexture->GetDesc(), RenderTargetTexture);
 	bool bFinalizedAccess = false;
 
-	if (!EnumHasAnyFlags(Desc.Flags, TexCreate_RenderTargetable | TexCreate_ResolveTargetable | TexCreate_DepthStencilTargetable | TexCreate_UAV | TexCreate_DepthStencilResolveTarget))
+	if (!EnumHasAnyFlags(Flags, ERDGTextureFlags::ForceTracking)
+		&& !EnumHasAnyFlags(Desc.Flags, TexCreate_RenderTargetable | TexCreate_ResolveTargetable | TexCreate_DepthStencilTargetable | TexCreate_UAV | TexCreate_DepthStencilResolveTarget))
 	{
 		bFinalizedAccess = true;
 	}
@@ -760,7 +761,8 @@ FRDGBufferRef FRDGBuilder::RegisterExternalBuffer(
 	const FRDGBufferDesc& Desc = ExternalPooledBuffer->Desc;
 	bool bFinalizedAccess = false;
 
-	if (!EnumHasAnyFlags(Desc.Usage, BUF_UnorderedAccess))
+	if (!EnumHasAnyFlags(Flags, ERDGBufferFlags::ForceTracking)
+		&& !EnumHasAnyFlags(Desc.Usage, BUF_UnorderedAccess))
 	{
 		Flags |= ERDGBufferFlags::ReadOnly;
 		bFinalizedAccess = true;
diff --git a/Engine/Source/Runtime/RenderCore/Private/RenderGraphValidation.cpp b/Engine/Source/Runtime/RenderCore/Private/RenderGraphValidation.cpp
--- a/Engine/Source/Runtime/RenderCore/Private/RenderGraphValidation.cpp
+++ b/Engine/Source/Runtime/RenderCore/Private/RenderGraphValidation.cpp
@@ -240,6 +240,7 @@ void FRDGUserValidation::ValidateRegisterExternalTexture(
 {
 	checkf(Name, TEXT("Attempted to register external texture with NULL name."));
 	checkf(ExternalPooledTexture.IsValid(), TEXT("Attempted to register NULL external texture."));
+	checkf(!EnumHasAnyFlags(Flags, ERDGTextureFlags::ReadOnly) || !EnumHasAnyFlags(Flags, ERDGTextureFlags::ForceTracking), TEXT("External texture %s cannot be ReadOnly and ForceTracking (flags are mutually exclusive)"), Name);
 	ExecuteGuard(TEXT("RegisterExternalTexture"), Name);
 }
 
@@ -247,6 +248,7 @@ void FRDGUserValidation::ValidateRegisterExternalBuffer(const TRefCountPtr<FRDGP
 {
 	checkf(Name, TEXT("Attempted to register external buffer with NULL name."));
 	checkf(ExternalPooledBuffer.IsValid(), TEXT("Attempted to register NULL external buffer."));
+	checkf(!EnumHasAnyFlags(Flags, ERDGBufferFlags::ReadOnly) || !EnumHasAnyFlags(Flags, ERDGBufferFlags::ForceTracking), TEXT("External buffer %s cannot be ReadOnly and ForceTracking (flags are mutually exclusive)"), Name);
 	ExecuteGuard(TEXT("RegisterExternalBuffer"), Name);
 }
 
diff --git a/Engine/Source/Runtime/RenderCore/Public/RenderGraphDefinitions.h b/Engine/Source/Runtime/RenderCore/Public/RenderGraphDefinitions.h
--- a/Engine/Source/Runtime/RenderCore/Public/RenderGraphDefinitions.h
+++ b/Engine/Source/Runtime/RenderCore/Public/RenderGraphDefinitions.h
@@ -129,7 +129,11 @@ enum class ERDGBufferFlags : uint8
 	MultiFrame = 1 << 0,
 
 	/** The buffer may only be used for read-only access within the graph. This flag is only allowed for registered buffers. */
-	ReadOnly = 1 << 1
+	ReadOnly = 1 << 1, 
+
+	/** Force the graph to track this resource even if it can be considered as readonly (no UAV, no RTV, etc.) This allows the graph copying from and to textures, and handling the corresponding transitions, for example.
+	 This flag is only allowed for registered buffers. Mutually exclusive with ReadOnly. */
+	ForceTracking = 1 << 2,
 };
 ENUM_CLASS_FLAGS(ERDGBufferFlags);
 
@@ -146,6 +150,10 @@ enum class ERDGTextureFlags : uint8
 
 	/** Prevents metadata decompression on this texture. */
 	MaintainCompression = 1 << 2,
+
+	/** Force the graph to track this resource even if it can be considered as readonly (no UAV, no RTV, etc.) This allows the graph copying from and to textures, and handling the corresponding transitions, for example.
+	 This flag is only allowed for registered buffers. Mutually exclusive with ReadOnly. */
+	ForceTracking = 1 << 3,
 };
 ENUM_CLASS_FLAGS(ERDGTextureFlags);
 
@@ -608,5 +616,8 @@ class FRDGTrace;
 using FRDGBufferNumElementsCallback = TFunction<uint32()>;
 using FRDGBufferInitialDataCallback = TFunction<const void*()>;
 using FRDGBufferInitialDataSizeCallback = TFunction<uint64()>;
+template <typename ArrayType, 
+	typename ArrayTypeNoRef = std::remove_reference_t<ArrayType>,
+	typename = typename TEnableIf<TIsTArray<ArrayTypeNoRef>::Value>::Type> using TRDGBufferArrayCallback = TFunction<const ArrayType&()>;
 using FRDGBufferInitialDataFreeCallback = TFunction<void(const void* InData)>;
 using FRDGDispatchGroupCountCallback = TFunction<FIntVector()>;
\ No newline at end of file
diff --git a/Engine/Source/Runtime/RenderCore/Public/RenderGraphUtils.h b/Engine/Source/Runtime/RenderCore/Public/RenderGraphUtils.h
--- a/Engine/Source/Runtime/RenderCore/Public/RenderGraphUtils.h
+++ b/Engine/Source/Runtime/RenderCore/Public/RenderGraphUtils.h
@@ -769,7 +769,7 @@ RENDERCORE_API FRDGBufferRef CreateStructuredBuffer(
 	ERDGInitialDataFlags InitialDataFlags = ERDGInitialDataFlags::None);
 
 /** A variant where NumElements, InitialData, and InitialDataSize are supplied through callbacks. This allows creating a buffer with
- *  with information unknown at creation time. Though, data must be ready before the most recent RDG pass that references the buffer
+ *  information unknown at creation time. Though, data must be ready before the most recent RDG pass that references the buffer
  *  is executed.
  */
 RENDERCORE_API FRDGBufferRef CreateStructuredBuffer(
@@ -816,6 +816,35 @@ FORCEINLINE FRDGBufferRef CreateStructuredBuffer(
 	return CreateStructuredBuffer(GraphBuilder, Name, InitialData.GetTypeSize(), InitialData.Num(), InitialData.GetData(), InitialData.Num() * InitialData.GetTypeSize(), InitialDataFlags);
 }
 
+/** A variant where the TArray is supplied through callbacks. This allows creating a buffer with
+ *  information unknown at creation time. Though, data must be ready before the most recent RDG pass that references the buffer
+ *  is executed.
+ */
+template <typename ArrayType>
+FORCEINLINE FRDGBufferRef CreateStructuredBuffer_Impl(
+	FRDGBuilder& GraphBuilder,
+	const TCHAR* Name,
+	TRDGBufferArrayCallback<ArrayType>&& ArrayCallback)
+{
+	return CreateStructuredBuffer(GraphBuilder, Name, sizeof(std::remove_reference_t<ArrayType>::ElementType),
+		/*NumElementsCallback = */[ArrayCallback]() { return ArrayCallback().Num(); },
+		/*InitialDataCallback = */[ArrayCallback]() { return ArrayCallback().GetData(); },
+		/*InitialDataSizeCallback = */[ArrayCallback]() { return ArrayCallback().Num() * ArrayCallback().GetTypeSize(); });
+}
+
+/** Same as the previous function but where the type of the array is automatically inferred, so we can do : 
+ *  TArray<FSomeType> Array;
+ *  CreateStructuredBuffer(..., [&]() -> auto&{ return Array; });
+ */
+template <typename GetArrayRefCallback, typename Type = TInvokeResult_T<GetArrayRefCallback>>
+FORCEINLINE FRDGBufferRef CreateStructuredBuffer(
+	FRDGBuilder& GraphBuilder,
+	const TCHAR* Name,
+	GetArrayRefCallback&& ArrayCallback)
+{
+	return CreateStructuredBuffer_Impl<Type>(GraphBuilder, Name, MoveTemp(ArrayCallback));
+}
+
 template <typename ElementType>
 FORCEINLINE FRDGBufferRef CreateStructuredBuffer(
 	FRDGBuilder& GraphBuilder,
diff --git a/Engine/Source/Runtime/RenderCore/Public/ShaderParameterMacros.h b/Engine/Source/Runtime/RenderCore/Public/ShaderParameterMacros.h
--- a/Engine/Source/Runtime/RenderCore/Public/ShaderParameterMacros.h
+++ b/Engine/Source/Runtime/RenderCore/Public/ShaderParameterMacros.h
@@ -957,6 +957,21 @@ struct TShaderParameterTypeInfo<FIntPoint>
 	static const FShaderParametersMetadata* GetStructMetadata() { return nullptr; }
 };
 
+template<>
+struct TShaderParameterTypeInfo<FUintVector2>
+{
+	static constexpr EUniformBufferBaseType BaseType = UBMT_UINT32;
+	static constexpr int32 NumRows = 1;
+	static constexpr int32 NumColumns = 2;
+	static constexpr int32 NumElements = 0;
+	static constexpr int32 Alignment = 8;
+	static constexpr bool bIsStoredInConstantBuffer = true;
+
+	using TAlignedType = TAlignedTypedef<FUintVector2, Alignment>::Type;
+
+	static const FShaderParametersMetadata* GetStructMetadata() { return nullptr; }
+};
+
 template<>
 struct TShaderParameterTypeInfo<FIntVector>
 {
