Fix for loop variable being copied warning.

diff --git a/Engine/Source/Runtime/RenderCore/Private/RenderGraphValidation.cpp b/Engine/Source/Runtime/RenderCore/Private/RenderGraphValidation.cpp
--- a/Engine/Source/Runtime/RenderCore/Private/RenderGraphValidation.cpp
+++ b/Engine/Source/Runtime/RenderCore/Private/RenderGraphValidation.cpp
@@ -1552,7 +1552,7 @@ void FRDGLogFile::End()
 		{
 			const FRDGPass* Pass = Passes[PassHandle];
 
-			for (const FRDGPassHandle ProducerHandle : Pass->GetProducers())
+			for (FRDGPassHandle ProducerHandle : Pass->GetProducers())
 			{
 				const FRDGPass* Producer = Passes[ProducerHandle];
 
