Dumps both GPU dumper viewer's .sh and .bat so a dump can be opened either on max, linux or windows regardless of the platform it was dumped from

diff --git a/Engine/Source/Runtime/RenderCore/Private/RenderGraphDumpResources.cpp b/Engine/Source/Runtime/RenderCore/Private/RenderGraphDumpResources.cpp
--- a/Engine/Source/Runtime/RenderCore/Private/RenderGraphDumpResources.cpp
+++ b/Engine/Source/Runtime/RenderCore/Private/RenderGraphDumpResources.cpp
@@ -1296,16 +1296,15 @@ FString FRDGBuilder::BeginResourceDump(const TArray<FString>& Args)
 	// Copy the viewer
 	if (NewResourceDumpContext.bEnableDiskWrite)
 	{
-#if PLATFORM_WINDOWS
-		const TCHAR* OpenGPUDumpViewerName = TEXT("OpenGPUDumpViewer.bat");
-#else
-		const TCHAR* OpenGPUDumpViewerName = TEXT("OpenGPUDumpViewer.sh");
-#endif
+		const TCHAR* OpenGPUDumpViewerBatName = TEXT("OpenGPUDumpViewer.bat");
+		const TCHAR* OpenGPUDumpViewerShName = TEXT("OpenGPUDumpViewer.sh");
+
 		const TCHAR* ViewerHTML = TEXT("GPUDumpViewer.html");
 		FString DumpGPUViewerSourcePath = FPaths::EngineDir() + FString(TEXT("Extras")) / TEXT("GPUDumpViewer");
 
 		PlatformFile.CopyFile(*(NewResourceDumpContext.DumpingDirectoryPath / ViewerHTML), *(DumpGPUViewerSourcePath / ViewerHTML));
-		PlatformFile.CopyFile(*(NewResourceDumpContext.DumpingDirectoryPath / OpenGPUDumpViewerName), *(DumpGPUViewerSourcePath / OpenGPUDumpViewerName));
+		PlatformFile.CopyFile(*(NewResourceDumpContext.DumpingDirectoryPath / OpenGPUDumpViewerBatName), *(DumpGPUViewerSourcePath / OpenGPUDumpViewerBatName));
+		PlatformFile.CopyFile(*(NewResourceDumpContext.DumpingDirectoryPath / OpenGPUDumpViewerShName), *(DumpGPUViewerSourcePath / OpenGPUDumpViewerShName));
 	}
 
 	ENQUEUE_RENDER_COMMAND(FStartGPUDump)(
