Removed r.RenderTargetPool.Events visualizer. It has been replaced by RDG insights.

diff --git a/Engine/Source/Runtime/RenderCore/Private/RenderGraphBuilder.cpp b/Engine/Source/Runtime/RenderCore/Private/RenderGraphBuilder.cpp
--- a/Engine/Source/Runtime/RenderCore/Private/RenderGraphBuilder.cpp
+++ b/Engine/Source/Runtime/RenderCore/Private/RenderGraphBuilder.cpp
@@ -3113,12 +3113,6 @@ void FRDGBuilder::BeginResourceRHI(FRDGPassHandle PassHandle, FRDGTextureRef Tex
 
 		// Cannot begin a resource within a merged render pass region.
 		check(GetPrologueBarrierPassHandle(PassHandle) == PassHandle);
-
-		if (!Pass->bFirstTextureAllocated)
-		{
-			GRenderTargetPool.AddPhaseEvent(Pass->GetName());
-			Pass->bFirstTextureAllocated = 1;
-		}
 	}
 #endif
 
diff --git a/Engine/Source/Runtime/RenderCore/Private/RenderTargetPool.cpp b/Engine/Source/Runtime/RenderCore/Private/RenderTargetPool.cpp
--- a/Engine/Source/Runtime/RenderCore/Private/RenderTargetPool.cpp
+++ b/Engine/Source/Runtime/RenderCore/Private/RenderTargetPool.cpp
@@ -1,9 +1,5 @@
 // Copyright Epic Games, Inc. All Rights Reserved.
 
-/*=============================================================================
-	RenderTargetPool.cpp: Scene render target pool manager.
-=============================================================================*/
-
 #include "RenderTargetPool.h"
 #include "RHIStaticStates.h"
 #include "Misc/OutputDeviceRedirector.h"
@@ -117,95 +113,11 @@ static FAutoConsoleCommandWithOutputDevice GDumpRenderTargetPoolMemoryCmd(
 	FConsoleCommandWithOutputDeviceDelegate::CreateStatic(DumpRenderTargetPoolMemory)
 );
 
-void RenderTargetPoolEvents(const TArray<FString>& Args)
-{
-	uint32 SizeInKBThreshold = -1;
-	if (Args.Num() && Args[0].IsNumeric())
-	{
-		SizeInKBThreshold = FCString::Atof(*Args[0]);
-	}
-
-	if (SizeInKBThreshold != -1)
-	{
-		UE_LOG(LogRenderTargetPool, Display, TEXT("r.DumpRenderTargetPoolEvents is now enabled, use r.DumpRenderTargetPoolEvents ? for help"));
-
-		GRenderTargetPool.EventRecordingSizeThreshold = SizeInKBThreshold;
-		GRenderTargetPool.bStartEventRecordingNextTick = true;
-	}
-	else
-	{
-		GRenderTargetPool.DisableEventDisplay();
-
-		UE_LOG(LogRenderTargetPool, Display, TEXT("r.DumpRenderTargetPoolEvents is now disabled, use r.DumpRenderTargetPoolEvents <SizeInKB> to enable or r.DumpRenderTargetPoolEvents ? for help"));
-	}
-}
-
-// CVars and commands
-static FAutoConsoleCommand GRenderTargetPoolEventsCmd(
-	TEXT("r.RenderTargetPool.Events"),
-	TEXT("Visualize the render target pool events over time in one frame. Optional parameter defines threshold in KB.\n")
-	TEXT("To disable the view use the command without any parameter"),
-	FConsoleCommandWithArgsDelegate::CreateStatic(RenderTargetPoolEvents)
-);
-
-bool FRenderTargetPool::IsEventRecordingEnabled() const
-{
-#if !(UE_BUILD_SHIPPING || UE_BUILD_TEST)
-	return bEventRecordingStarted && bEventRecordingActive;
-#else
-	return false;
-#endif
-}
-
-IPooledRenderTarget* FRenderTargetPoolEvent::GetValidatedPointer() const
-{
-	int32 Index = GRenderTargetPool.FindIndex(Pointer);
-
-	if (Index >= 0)
-	{
-		return Pointer;
-	}
-
-	return 0;
-}
-
-bool FRenderTargetPoolEvent::NeedsDeallocEvent()
-{
-	if (GetEventType() == ERTPE_Alloc)
-	{
-		if (Pointer)
-		{
-			IPooledRenderTarget* ValidPointer = GetValidatedPointer();
-			if (!ValidPointer || ValidPointer->IsFree())
-			{
-				Pointer = 0;
-				return true;
-			}
-		}
-	}
-
-	return false;
-}
-
 static uint32 ComputeSizeInKB(FPooledRenderTarget& Element)
 {
 	return (Element.ComputeMemorySize() + 1023) / 1024;
 }
 
-FRenderTargetPool::FRenderTargetPool()
-	: AllocationLevelInKB(0)
-	, bCurrentlyOverBudget(false)
-	, bStartEventRecordingNextTick(false)
-	, EventRecordingSizeThreshold(0)
-	, bEventRecordingActive(false)
-	, bEventRecordingStarted(false)
-	, CurrentEventRecordingTime(0)
-#if LOG_MAX_RENDER_TARGET_POOL_USAGE
-	, MaxUsedRenderTargetInKB(0)
-#endif
-{
-}
-
 TRefCountPtr<FPooledRenderTarget> FRenderTargetPool::FindFreeElementForRDG(
 	FRHICommandList& RHICmdList,
 	const FRDGTextureDesc& Desc,
@@ -305,15 +217,6 @@ TRefCountPtr<FPooledRenderTarget> FRenderTargetPool::FindFreeElementInternal(
 						Desc.NumSamples
 					);
 				}
-
-				if (RHISupportsRenderTargetWriteMask(GMaxRHIShaderPlatform) && Desc.bCreateRenderTargetWriteMask)
-				{
-					Found->RenderTargetItem.RTWriteMaskSRV = RHICreateShaderResourceViewWriteMask((FTexture2DRHIRef&)Found->RenderTargetItem.TargetableTexture);
-				}
-				if (Desc.bCreateRenderTargetFmask)
-				{
-					Found->RenderTargetItem.FmaskSRV = RHICreateShaderResourceViewFMask((FTexture2DRHIRef&)Found->RenderTargetItem.TargetableTexture);
-				}
 			}
 			else if (Desc.Is3DTexture())
 			{
@@ -455,20 +358,12 @@ TRefCountPtr<FPooledRenderTarget> FRenderTargetPool::FindFreeElementInternal(
 		Found->Desc.DebugName = InDebugName;
 	}
 
-	check(Found->IsFree());
-
 	Found->Desc.DebugName = InDebugName;
 	Found->UnusedForNFrames = 0;
 
-	AddAllocEvent(FoundIndex, Found);
-
-	uint32 OriginalNumRefs = Found->GetRefCount();
-
 	// assign to the reference counted variable
 	TRefCountPtr<FPooledRenderTarget> Result = Found;
 
-	check(!Found->IsFree());
-
 #if !(UE_BUILD_SHIPPING || UE_BUILD_TEST)
 	if (Found->GetRenderTargetItem().TargetableTexture)
 	{
@@ -606,235 +501,6 @@ void FRenderTargetPool::GetStats(uint32& OutWholeCount, uint32& OutWholePoolInKB
 	ensure(AllocationLevelInKB == OutWholePoolInKB);
 }
 
-void FRenderTargetPool::AddPhaseEvent(const TCHAR* InPhaseName)
-{
-	if (IsEventRecordingEnabled())
-	{
-		AddDeallocEvents();
-
-		const FString* LastName = GetLastEventPhaseName();
-
-		if (!LastName || *LastName != InPhaseName)
-		{
-			if (CurrentEventRecordingTime)
-			{
-				// put a break to former data
-				++CurrentEventRecordingTime;
-			}
-
-			FRenderTargetPoolEvent NewEvent(InPhaseName, CurrentEventRecordingTime);
-
-			RenderTargetPoolEvents.Add(NewEvent);
-		}
-	}
-}
-
-const FString* FRenderTargetPool::GetLastEventPhaseName()
-{
-	// could be optimized but this is a debug view
-
-	// start from the end for better performance
-	for (int32 i = RenderTargetPoolEvents.Num() - 1; i >= 0; --i)
-	{
-		const FRenderTargetPoolEvent* Event = &RenderTargetPoolEvents[i];
-
-		if (Event->GetEventType() == ERTPE_Phase)
-		{
-			return &Event->GetPhaseName();
-		}
-	}
-
-	return 0;
-}
-
-FRenderTargetPool::SMemoryStats FRenderTargetPool::ComputeView()
-{
-	SMemoryStats MemoryStats;
-
-#if !(UE_BUILD_SHIPPING || UE_BUILD_TEST)
-	{
-		struct FRTPColumn
-		{
-			FString Name;
-			// index into the column, -1 if this is no valid column
-			uint32 PoolEntryId;
-			// for sorting
-			uint64 SizeInBytes;
-			// for sorting
-			bool bVRam;
-
-			// default constructor
-			FRTPColumn()
-				: PoolEntryId(-1)
-				, SizeInBytes(0)
-			{
-			}
-
-			// constructor
-			FRTPColumn(const FRenderTargetPoolEvent& Event)
-				: Name(Event.GetDesc().DebugName)
-				, PoolEntryId(Event.GetPoolEntryId())
-				, bVRam(EnumHasAnyFlags(Event.GetDesc().Flags, TexCreate_FastVRAM))
-			{
-				SizeInBytes = Event.GetSizeInBytes();
-			}
-
-			// sort criteria
-			bool operator < (const FRTPColumn& RHS) const
-			{
-				if (SizeInBytes == RHS.SizeInBytes)
-				{
-					return Name < RHS.Name;
-				}
-
-				return SizeInBytes > RHS.SizeInBytes;
-			}
-		};
-
-		TArray<FRTPColumn> Colums;
-
-		// generate Colums
-		for (int32 i = 0, Num = RenderTargetPoolEvents.Num(); i < Num; i++)
-		{
-			FRenderTargetPoolEvent* Event = &RenderTargetPoolEvents[i];
-
-			if (Event->GetEventType() == ERTPE_Alloc)
-			{
-				uint32 PoolEntryId = Event->GetPoolEntryId();
-
-				if (PoolEntryId >= (uint32)Colums.Num())
-				{
-					Colums.SetNum(PoolEntryId + 1);
-				}
-
-				Colums[PoolEntryId] = FRTPColumn(*Event);
-			}
-		}
-
-		Colums.StableSort();
-
-		{
-			uint32 ColumnX = 0;
-
-			for (int32 ColumnIndex = 0, ColumnsNum = Colums.Num(); ColumnIndex < ColumnsNum; ++ColumnIndex)
-			{
-				const FRTPColumn& RTPColumn = Colums[ColumnIndex];
-
-				uint32 ColumnSize = RTPColumn.SizeInBytes;
-
-				// hide columns that are too small to make a difference (e.g. <1 MB)
-				if (RTPColumn.SizeInBytes <= EventRecordingSizeThreshold * 1024)
-				{
-					ColumnSize = 0;
-				}
-				else
-				{
-					MemoryStats.DisplayedUsageInBytes += RTPColumn.SizeInBytes;
-
-					// give an entry some size to be more UI friendly (if we get mouse UI for zooming in we might not want that any more)
-					ColumnSize = FMath::Max((uint32)(1024 * 1024), ColumnSize);
-				}
-
-				MemoryStats.TotalColumnSize += ColumnSize;
-				MemoryStats.TotalUsageInBytes += RTPColumn.SizeInBytes;
-
-				for (int32 EventIndex = 0, PoolEventsNum = RenderTargetPoolEvents.Num(); EventIndex < PoolEventsNum; EventIndex++)
-				{
-					FRenderTargetPoolEvent* Event = &RenderTargetPoolEvents[EventIndex];
-
-					if (Event->GetEventType() != ERTPE_Phase)
-					{
-						uint32 PoolEntryId = Event->GetPoolEntryId();
-
-						if (RTPColumn.PoolEntryId == PoolEntryId)
-						{
-							Event->SetColumn(ColumnIndex, ColumnX, ColumnSize);
-						}
-					}
-				}
-				ColumnX += ColumnSize;
-			}
-		}
-	}
-#endif
-
-	return MemoryStats;
-}
-
-void FRenderTargetPool::AddDeallocEvents()
-{
-	check(IsInRenderingThread());
-
-	bool bWorkWasDone = false;
-
-	for (uint32 i = 0, Num = (uint32)RenderTargetPoolEvents.Num(); i < Num; ++i)
-	{
-		FRenderTargetPoolEvent& Event = RenderTargetPoolEvents[i];
-
-		if (Event.NeedsDeallocEvent())
-		{
-			FRenderTargetPoolEvent NewEvent(Event.GetPoolEntryId(), CurrentEventRecordingTime);
-
-			// for convenience - is actually redundant
-			NewEvent.SetDesc(Event.GetDesc());
-
-			RenderTargetPoolEvents.Add(NewEvent);
-			bWorkWasDone = true;
-		}
-	}
-
-	if (bWorkWasDone)
-	{
-		++CurrentEventRecordingTime;
-	}
-}
-
-void FRenderTargetPool::AddAllocEvent(uint32 InPoolEntryId, FPooledRenderTarget* In)
-{
-	check(In);
-
-	if (IsEventRecordingEnabled())
-	{
-		AddDeallocEvents();
-
-		check(IsInRenderingThread());
-
-		FRenderTargetPoolEvent NewEvent(InPoolEntryId, CurrentEventRecordingTime++, In);
-
-		RenderTargetPoolEvents.Add(NewEvent);
-	}
-}
-
-void FRenderTargetPool::AddAllocEventsFromCurrentState()
-{
-	if (!IsEventRecordingEnabled())
-	{
-		return;
-	}
-
-	check(IsInRenderingThread());
-
-	bool bWorkWasDone = false;
-
-	for (uint32 i = 0; i < (uint32)PooledRenderTargets.Num(); ++i)
-	{
-		FPooledRenderTarget* Element = PooledRenderTargets[i];
-
-		if (Element && !Element->IsFree())
-		{
-			FRenderTargetPoolEvent NewEvent(i, CurrentEventRecordingTime, Element);
-
-			RenderTargetPoolEvents.Add(NewEvent);
-			bWorkWasDone = true;
-		}
-	}
-
-	if (bWorkWasDone)
-	{
-		++CurrentEventRecordingTime;
-	}
-}
-
 void FRenderTargetPool::TickPoolElements()
 {
 	uint32 DeferredAllocationLevelInKB = 0;
@@ -846,12 +512,6 @@ void FRenderTargetPool::TickPoolElements()
 	check(IsInRenderingThread());
 	DeferredDeleteArray.Reset();
 
-	if (bStartEventRecordingNextTick)
-	{
-		bStartEventRecordingNextTick = false;
-		bEventRecordingStarted = true;
-	}
-
 	uint32 MinimumPoolSizeInKB;
 	{
 		static const auto CVar = IConsoleManager::Get().FindTConsoleVariableDataInt(TEXT("r.RenderTargetPoolMin"));
@@ -878,21 +538,8 @@ void FRenderTargetPool::TickPoolElements()
 
 	uint32 TotalFrameUsageInKb = AllocationLevelInKB + DeferredAllocationLevelInKB ;
 
-#if LOG_MAX_RENDER_TARGET_POOL_USAGE
-	if (TotalFrameUsageInKb > MaxUsedRenderTargetInKB)
-	{
-		MaxUsedRenderTargetInKB = TotalFrameUsageInKb;
-
-		if (MaxUsedRenderTargetInKB > MinimumPoolSizeInKB)
-		{
-			DumpMemoryUsage(*GLog);
-		}
-	}
-#endif
-
 	CSV_CUSTOM_STAT(RenderTargetPool, UnusedMB, UnusedAllocationLevelInKB / 1024.0f, ECsvCustomStatOp::Set);
 	CSV_CUSTOM_STAT(RenderTargetPool, PeakUsedMB, (TotalFrameUsageInKb - UnusedAllocationLevelInKB) / 1024.f, ECsvCustomStatOp::Set);
-
 	
 	// we need to release something, take the oldest ones first
 	while (AllocationLevelInKB > MinimumPoolSizeInKB)
@@ -956,10 +603,6 @@ void FRenderTargetPool::TickPoolElements()
 		}
 	}
 
-	AddPhaseEvent(TEXT("FromLastFrame"));
-	AddAllocEventsFromCurrentState();
-	AddPhaseEvent(TEXT("Rendering"));
-
 #if STATS
 	uint32 Count, SizeKB, UsedKB;
 	GetStats(Count, SizeKB, UsedKB);
@@ -1041,10 +684,6 @@ void FRenderTargetPool::FreeUnusedResources()
 			FreeElementAtIndex(i);
 		}
 	}
-
-#if LOG_MAX_RENDER_TARGET_POOL_USAGE
-	MaxUsedRenderTargetInKB = 0;
-#endif
 }
 
 void FRenderTargetPool::DumpMemoryUsage(FOutputDevice& OutputDevice)
diff --git a/Engine/Source/Runtime/RenderCore/Public/RenderGraphPass.h b/Engine/Source/Runtime/RenderCore/Public/RenderGraphPass.h
--- a/Engine/Source/Runtime/RenderCore/Public/RenderGraphPass.h
+++ b/Engine/Source/Runtime/RenderCore/Public/RenderGraphPass.h
@@ -385,9 +385,6 @@ protected:
 
 			/** If set, the pass will wait on the assigned mGPU temporal effect. */
 			uint32 bWaitForTemporalEffect : 1;
-
-			/** Whether this pass allocated a texture through the pool. */
-			IF_RDG_ENABLE_DEBUG(uint32 bFirstTextureAllocated : 1);
 		};
 		uint32 PackedBits1 = 0;
 	};
diff --git a/Engine/Source/Runtime/RenderCore/Public/RenderTargetPool.h b/Engine/Source/Runtime/RenderCore/Public/RenderTargetPool.h
--- a/Engine/Source/Runtime/RenderCore/Public/RenderTargetPool.h
+++ b/Engine/Source/Runtime/RenderCore/Public/RenderTargetPool.h
@@ -12,8 +12,6 @@
 #include "RendererInterface.h"
 #include "RenderGraphResources.h"
 
-#define LOG_MAX_RENDER_TARGET_POOL_USAGE 0
-
 /** The reference to a pooled render target, use like this: TRefCountPtr<IPooledRenderTarget> */
 struct RENDERCORE_API FPooledRenderTarget final : public IPooledRenderTarget
 {
@@ -81,111 +79,13 @@ private:
 	friend class FRenderTargetPool;
 };
 
-enum ERenderTargetPoolEventType
-{
-	ERTPE_Alloc,
-	ERTPE_Dealloc,
-	ERTPE_Phase
-};
-
-struct RENDERCORE_API FRenderTargetPoolEvent
-{
-	// constructor for ERTPE_Alloc
-	FRenderTargetPoolEvent(uint32 InPoolEntryId, uint32 InTimeStep, FPooledRenderTarget* InPointer)
-		: PoolEntryId(InPoolEntryId)
-		, TimeStep(InTimeStep)
-		, Pointer(InPointer)
-		, EventType(ERTPE_Alloc)
-		, ColumnIndex(-1)
-		, ColumnX(0)
-		, ColumnSize(0)
-	{
-		Desc = Pointer->GetDesc();
-		SizeInBytes = Pointer->ComputeMemorySize();
-	}
-
-	// constructor for ERTPE_Alloc
-	FRenderTargetPoolEvent(uint32 InPoolEntryId, uint32 InTimeStep)
-		: PoolEntryId(InPoolEntryId)
-		, TimeStep(InTimeStep)
-		, Pointer(0)
-		, SizeInBytes(0)
-		, EventType(ERTPE_Dealloc)
-		, ColumnIndex(-1)
-		, ColumnX(0)
-		, ColumnSize(0)
-	{
-	}
-
-	// constructor for ERTPE_Alloc
-	// @param PhaseName must not be 0
-	FRenderTargetPoolEvent(const FString& InPhaseName, uint32 InTimeStep)
-		: PoolEntryId(-1)
-		, TimeStep(InTimeStep)
-		, Pointer(0)
-		, PhaseName(InPhaseName)
-		, SizeInBytes(0)
-		, EventType(ERTPE_Phase)
-		, ColumnIndex(-1)
-		, ColumnX(0)
-		, ColumnSize(0)
-	{
-	}
-
-	// @return Pointer if the object is still in the pool
-	IPooledRenderTarget* GetValidatedPointer() const;
-
-	ERenderTargetPoolEventType GetEventType() const { return EventType; }
-
-	const int32 GetPoolEntryId() const { check(EventType == ERTPE_Alloc || EventType == ERTPE_Dealloc); return PoolEntryId; }
-	const FString& GetPhaseName() const { check(EventType == ERTPE_Phase); return PhaseName; }
-	const FPooledRenderTargetDesc& GetDesc() const { check(EventType == ERTPE_Alloc || EventType == ERTPE_Dealloc); return Desc; }
-	const uint32 GetTimeStep() const { return TimeStep; }
-	const uint64 GetSizeInBytes() const { check(EventType == ERTPE_Alloc); return SizeInBytes; }
-	const void SetPoolEntryId(uint32 InPoolEntryId) { PoolEntryId = InPoolEntryId; }
-	const void SetColumn(uint32 InColumnIndex, uint32 InColumnX, uint32 InColumnSize) { check(EventType == ERTPE_Alloc || EventType == ERTPE_Dealloc); ColumnIndex = InColumnIndex; ColumnX = InColumnX; ColumnSize = InColumnSize; }
-	const uint32 GetColumnX() const { check(EventType == ERTPE_Alloc || EventType == ERTPE_Dealloc); return ColumnX; }
-	const uint32 GetColumnSize() const { check(EventType == ERTPE_Alloc || EventType == ERTPE_Dealloc); return ColumnSize; }
-	bool IsVisible() const { return EventType == ERTPE_Phase || ColumnSize > 0; }
-
-	void SetDesc(const FPooledRenderTargetDesc &InDesc) { Desc = InDesc; }
-
-	bool NeedsDeallocEvent();
-
-private:
-
-	// valid if EventType==ERTPE_Alloc || EventType==ERTPE_Dealloc, -1 if not set, was index into PooledRenderTargets[]
-	uint32 PoolEntryId;
-	//
-	uint32 TimeStep;
-	// valid EventType==ERTPE_Alloc, 0 if not set
-	FPooledRenderTarget* Pointer;
-	// valid if EventType==ERTPE_Phase TEXT("") if not set
-	FString PhaseName;
-	// valid if EventType==ERTPE_Alloc || EventType==ERTPE_Dealloc
-	FPooledRenderTargetDesc Desc;
-	// valid if EventType==ERTPE_Alloc 0 if unknown
-	uint64 SizeInBytes;
-	// e.g. ERTPE_Alloc
-	ERenderTargetPoolEventType EventType;
-
-	// for display, computed by ComputeView()
-
-	// valid if EventType==ERTPE_Alloc || EventType==ERTPE_Dealloc, -1 if not defined yet
-	uint32 ColumnIndex;
-	//
-	uint32 ColumnX;
-	//
-	uint32 ColumnSize;
-};
-
 /**
  * Encapsulates the render targets pools that allows easy sharing (mostly used on the render thread side)
  */
 class RENDERCORE_API FRenderTargetPool : public FRenderResource
 {
 public:
-	FRenderTargetPool();
+	FRenderTargetPool() = default;
 
 	/**
 	 * @param DebugName must not be 0, we only store the pointer
@@ -226,21 +126,9 @@ public:
 	// @return -1 if not found
 	int32 FindIndex(IPooledRenderTarget* In) const;
 
-	void SetObserveTarget(const FString& InObservedDebugName, uint32 InObservedDebugNameReusedGoal = 0xffffffff);
-
 	// Logs out usage information.
 	void DumpMemoryUsage(FOutputDevice& OutputDevice);
 
-	// to not have event recording for some time during rendering (e.g. thumbnail rendering)
-	void SetEventRecordingActive(bool bValue) { bEventRecordingActive = bValue; }
-
-	//
-	void DisableEventDisplay() { RenderTargetPoolEvents.Empty(); bEventRecordingStarted = false; }
-	//
-	bool IsEventRecordingEnabled() const;
-
-	void AddPhaseEvent(const TCHAR* InPhaseName);
-
 private:
 	TRefCountPtr<FPooledRenderTarget> FindFreeElementForRDG(FRHICommandList& RHICmdList, const FRDGTextureDesc& Desc, const TCHAR* Name);
 
@@ -249,8 +137,6 @@ private:
 		const FPooledRenderTargetDesc& InputDesc,
 		const TCHAR* InDebugName);
 
-	friend void RenderTargetPoolEvents(const TArray<FString>& Args);
-
 	void FreeElementAtIndex(int32 Index);
 
 	/** Elements can be 0, we compact the buffer later. */
@@ -259,65 +145,14 @@ private:
 	TArray< TRefCountPtr<FPooledRenderTarget> > DeferredDeleteArray;
 
 	// redundant, can always be computed with GetStats(), to debug "out of memory" situations and used for r.RenderTargetPoolMin
-	uint32 AllocationLevelInKB;
-
-	FGraphEventRef TransitionFence;
+	uint32 AllocationLevelInKB = 0;
 
 	// to avoid log spam
-	bool bCurrentlyOverBudget;
+	bool bCurrentlyOverBudget = false;
 
 	// could be done on the fly but that makes the RenderTargetPoolEvents harder to read
 	void CompactPool();
 
-	// the following is used for Event recording --------------------------------
-
-	struct SMemoryStats
-	{
-		SMemoryStats()
-			: DisplayedUsageInBytes(0)
-			, TotalUsageInBytes(0)
-		{
-		}
-		// for statistics
-		uint64 DisplayedUsageInBytes;
-		// for statistics
-		uint64 TotalUsageInBytes;
-		// for display purposes, to normalize the view width (Initialize to 1 to avoid a division by zero when compiled out)
-		uint64 TotalColumnSize = 1;
-	};
-
-	// if next frame we want to run with bEventRecording=true
-	bool bStartEventRecordingNextTick;
-	// in KB, e.g. 1MB = 1024, 0 to display all
-	uint32 EventRecordingSizeThreshold;
-	// true if active, to not have the event recording for some time during rendering (e.g. thumbnail rendering)
-	bool bEventRecordingActive;
-	// true meaning someone used r.RenderTargetPool.Events to start it
-	bool bEventRecordingStarted;
-	// only used if bEventRecording
-	TArray<FRenderTargetPoolEvent> RenderTargetPoolEvents;
-	//
-	uint32 CurrentEventRecordingTime;
-
-#if LOG_MAX_RENDER_TARGET_POOL_USAGE
-	uint32 MaxUsedRenderTargetInKB;
-#endif
-
-	//
-	void AddDeallocEvents();
-
-	// @param In must no be 0
-	void AddAllocEvent(uint32 InPoolEntryId, FPooledRenderTarget* In);
-
-	//
-	void AddAllocEventsFromCurrentState();
-
-	// @return 0 if none was found
-	const FString* GetLastEventPhaseName();
-
-	// sorted by size
-	SMemoryStats ComputeView();
-
 	friend struct FPooledRenderTarget;
 	friend class FVisualizeTexture;
 	friend class FVisualizeTexturePresent;
diff --git a/Engine/Source/Runtime/RenderCore/Public/RendererInterface.h b/Engine/Source/Runtime/RenderCore/Public/RendererInterface.h
--- a/Engine/Source/Runtime/RenderCore/Public/RendererInterface.h
+++ b/Engine/Source/Runtime/RenderCore/Public/RendererInterface.h
@@ -487,15 +487,18 @@ struct FSceneRenderTargetItem
 		,	UAV(InUAV)
 	{}
 
+PRAGMA_DISABLE_DEPRECATION_WARNINGS
+	FSceneRenderTargetItem(const FSceneRenderTargetItem&) = default;
+	FSceneRenderTargetItem(FSceneRenderTargetItem&&) = default;
+	FSceneRenderTargetItem& operator=(const FSceneRenderTargetItem&) = default;
+	FSceneRenderTargetItem& operator=(FSceneRenderTargetItem&&) = default;
+PRAGMA_ENABLE_DEPRECATION_WARNINGS
+
 	void SafeRelease()
 	{
 		TargetableTexture.SafeRelease();
 		ShaderResourceTexture.SafeRelease();
 		UAV.SafeRelease();
-		HTileUAV.SafeRelease();
-		HTileSRV.SafeRelease();
-		RTWriteMaskSRV.SafeRelease();
-		FmaskSRV.SafeRelease();
 	}
 
 	bool IsValid() const
@@ -519,12 +522,20 @@ struct FSceneRenderTargetItem
 	/** only created if requested through the flag. */
 	FUnorderedAccessViewRHIRef UAV;
 
+	UE_DEPRECATED(5.0, "RTWriteMaskSRV is deprecated. Use RDG instead.")
 	FShaderResourceViewRHIRef RTWriteMaskSRV;
+
+	UE_DEPRECATED(5.0, "FmaskSRV is deprecated. Use RDG instead.")
 	FShaderResourceViewRHIRef FmaskSRV;
 
 	/** only created if requested through meta data access flags */
+	UE_DEPRECATED(5.0, "HTileUAV is deprecated. Use RDG instead.")
 	FUnorderedAccessViewRHIRef HTileUAV;
+
+	UE_DEPRECATED(5.0, "HTileSRV is deprecated. Use RDG instead.")
 	FShaderResourceViewRHIRef  HTileSRV;
+
+	UE_DEPRECATED(5.0, "StencilUAV is deprecated. Use RDG instead.")
 	FUnorderedAccessViewRHIRef StencilUAV;
 };
 
diff --git a/Engine/Source/Runtime/Renderer/Private/SceneRendering.cpp b/Engine/Source/Runtime/Renderer/Private/SceneRendering.cpp
--- a/Engine/Source/Runtime/Renderer/Private/SceneRendering.cpp
+++ b/Engine/Source/Runtime/Renderer/Private/SceneRendering.cpp
@@ -4046,11 +4046,6 @@ static void RenderViewFamily_RenderThread(FRHICommandListImmediate& RHICmdList,
 
 	const FSceneViewFamily& ViewFamily = SceneRenderer->ViewFamily;
 
-	if(ViewFamily.EngineShowFlags.OnScreenDebug)
-	{
-		GRenderTargetPool.SetEventRecordingActive(true);
-	}
-
 	{
 		SCOPE_CYCLE_COUNTER_VERBOSE(STAT_TotalSceneRenderingTime, ViewFamily.ProfileDescription.IsEmpty() ? nullptr : *ViewFamily.ProfileDescription);
 
@@ -4141,8 +4136,6 @@ static void RenderViewFamily_RenderThread(FRHICommandListImmediate& RHICmdList,
 		// Update on screen notifications.
 		FRendererOnScreenNotification::Get().Broadcast();
 #endif
-
-		GRenderTargetPool.SetEventRecordingActive(false);
 	}
 
 #if STATS
diff --git a/Engine/Source/Runtime/Renderer/Private/VisualizeTexturePresent.cpp b/Engine/Source/Runtime/Renderer/Private/VisualizeTexturePresent.cpp
--- a/Engine/Source/Runtime/Renderer/Private/VisualizeTexturePresent.cpp
+++ b/Engine/Source/Runtime/Renderer/Private/VisualizeTexturePresent.cpp
@@ -8,142 +8,6 @@
 #include "Engine/Engine.h"
 #include "RenderTargetPool.h"
 
-// draw a single pixel sized rectangle using 4 sub elements
-static void DrawBorder(FCanvas& Canvas, const FIntRect Rect, FLinearColor Color)
-{
-	// top
-	Canvas.DrawTile(Rect.Min.X, Rect.Min.Y, Rect.Max.X - Rect.Min.X, 1, 0, 0, 1, 1, Color);
-	// bottom
-	Canvas.DrawTile(Rect.Min.X, Rect.Max.Y - 1, Rect.Max.X - Rect.Min.X, 1, 0, 0, 1, 1, Color);
-	// left
-	Canvas.DrawTile(Rect.Min.X, Rect.Min.Y + 1, 1, Rect.Max.Y - Rect.Min.Y - 2, 0, 0, 1, 1, Color);
-	// right
-	Canvas.DrawTile(Rect.Max.X - 1, Rect.Min.Y + 1, 1, Rect.Max.Y - Rect.Min.Y - 2, 0, 0, 1, 1, Color);
-}
-
-// helper class to get a consistent layout in multiple functions
-// MaxX and Y are the output value that can be requested during or after iteration
-// Examples usages:
-//    FRenderTargetPoolEventIterator It(RenderTargetPoolEvents, OptionalStartIndex);
-//    while(FRenderTargetPoolEvent* Event = It.Iterate()) {}
-struct FRenderTargetPoolEventIterator
-{
-	int32 Index;
-	TArray<FRenderTargetPoolEvent>& RenderTargetPoolEvents;
-	bool bLineContent;
-	uint32 TotalWidth;
-	int32 Y;
-
-	// constructor
-	FRenderTargetPoolEventIterator(TArray<FRenderTargetPoolEvent>& InRenderTargetPoolEvents, int32 InIndex = 0)
-		: Index(InIndex)
-		, RenderTargetPoolEvents(InRenderTargetPoolEvents)
-		, bLineContent(false)
-		, TotalWidth(1)
-		, Y(0)
-	{
-		Touch();
-	}
-
-	FRenderTargetPoolEvent* operator*()
-	{
-		if (Index < RenderTargetPoolEvents.Num())
-		{
-			return &RenderTargetPoolEvents[Index];
-		}
-
-		return 0;
-	}
-
-	// @return 0 if end was reached
-	FRenderTargetPoolEventIterator& operator++()
-	{
-		if (Index < RenderTargetPoolEvents.Num())
-		{
-			++Index;
-		}
-
-		Touch();
-
-		return *this;
-	}
-
-	int32 FindClosingEventY() const
-	{
-		FRenderTargetPoolEventIterator It = *this;
-
-		const ERenderTargetPoolEventType StartType = (*It)->GetEventType();
-
-		if (StartType == ERTPE_Alloc)
-		{
-			int32 PoolEntryId = RenderTargetPoolEvents[Index].GetPoolEntryId();
-
-			++It;
-
-			// search for next Dealloc of the same PoolEntryId
-			for (; *It; ++It)
-			{
-				FRenderTargetPoolEvent* Event = *It;
-
-				if (Event->GetEventType() == ERTPE_Dealloc && Event->GetPoolEntryId() == PoolEntryId)
-				{
-					break;
-				}
-			}
-		}
-		else if (StartType == ERTPE_Phase)
-		{
-			++It;
-
-			// search for next Phase
-			for (; *It; ++It)
-			{
-				FRenderTargetPoolEvent* Event = *It;
-
-				if (Event->GetEventType() == ERTPE_Phase)
-				{
-					break;
-				}
-			}
-		}
-		else
-		{
-			check(0);
-		}
-
-		return It.Y;
-	}
-
-private:
-
-	void Touch()
-	{
-		if (Index < RenderTargetPoolEvents.Num())
-		{
-			const FRenderTargetPoolEvent& Event = RenderTargetPoolEvents[Index];
-
-			const ERenderTargetPoolEventType Type = Event.GetEventType();
-
-			if (Type == ERTPE_Alloc)
-			{
-				// for now they are all equal width
-				TotalWidth = FMath::Max(TotalWidth, Event.GetColumnX() + Event.GetColumnSize());
-			}
-			Y = Event.GetTimeStep();
-		}
-	}
-};
-
-uint32 FVisualizeTexturePresent::ComputeEventDisplayHeight()
-{
-	FRenderTargetPoolEventIterator It(GRenderTargetPool.RenderTargetPoolEvents);
-	while (*It)
-	{
-		++It;
-	}
-	return It.Y;
-}
-
 void FVisualizeTexturePresent::OnStartRender(const FViewInfo& View)
 {
 #if SUPPORTS_VISUALIZE_TEXTURE
@@ -157,149 +21,6 @@ void FVisualizeTexturePresent::PresentContent(FRDGBuilder& GraphBuilder, const F
 {
 	check(Output.IsValid());
 
-	if (GRenderTargetPool.RenderTargetPoolEvents.Num())
-	{
-		FIntPoint DisplayLeftTop(20, 50);
-		// on the right we leave more space to make the mouse tooltip readable
-		FIntPoint DisplayExtent(View.ViewRect.Width() - DisplayLeftTop.X * 2 - 140, View.ViewRect.Height() - DisplayLeftTop.Y * 2);
-
-		// if the area is not too small
-		if (DisplayExtent.X > 50 && DisplayExtent.Y > 50)
-		{
-			AddDrawCanvasPass(GraphBuilder, RDG_EVENT_NAME("RenderTargetPool"), View, Output, [&View, DisplayLeftTop, DisplayExtent] (FCanvas& Canvas)
-			{
-				FRenderTargetPool::SMemoryStats MemoryStats = GRenderTargetPool.ComputeView();
-
-				// TinyFont property
-				const int32 FontHeight = 12;
-
-				FIntPoint MousePos = View.CursorPos;
-
-				FLinearColor BackgroundColor = FLinearColor(0.0f, 0.0f, 0.0f, 0.7f);
-				FLinearColor PhaseColor = FLinearColor(0.2f, 0.1f, 0.05f, 0.8f);
-				FLinearColor ElementColor = FLinearColor(0.3f, 0.3f, 0.3f, 0.9f);
-				FLinearColor ElementColorVRam = FLinearColor(0.4f, 0.25f, 0.25f, 0.9f);
-
-				UTexture2D* GradientTexture = UCanvas::StaticClass()->GetDefaultObject<UCanvas>()->GradientTexture0;
-
-				// background rectangle
-				Canvas.DrawTile(DisplayLeftTop.X, DisplayLeftTop.Y - 1 * FontHeight - 1, DisplayExtent.X, DisplayExtent.Y + FontHeight, 0, 0, 1, 1, BackgroundColor);
-
-				{
-					uint32 MB = 1024 * 1024;
-					uint32 MBm1 = MB - 1;
-
-					FString Headline = *FString::Printf(TEXT("RenderTargetPool elements(x) over time(y) >= %dKB, Displayed/Total:%d/%dMB"),
-						GRenderTargetPool.EventRecordingSizeThreshold,
-						(uint32)((MemoryStats.DisplayedUsageInBytes + MBm1) / MB),
-						(uint32)((MemoryStats.TotalUsageInBytes + MBm1) / MB));
-					Canvas.DrawShadowedString(DisplayLeftTop.X, DisplayLeftTop.Y - 1 * FontHeight - 1, *Headline, GEngine->GetTinyFont(), FLinearColor(1, 1, 1));
-				}
-
-				uint32 EventDisplayHeight = FVisualizeTexturePresent::ComputeEventDisplayHeight();
-
-				float ScaleX = DisplayExtent.X / (float)MemoryStats.TotalColumnSize;
-				float ScaleY = DisplayExtent.Y / (float)EventDisplayHeight;
-
-				// 0 if none
-				FRenderTargetPoolEvent* HighlightedEvent = 0;
-				FIntRect HighlightedRect;
-
-				// Phase events
-				for (FRenderTargetPoolEventIterator It(GRenderTargetPool.RenderTargetPoolEvents); *It; ++It)
-				{
-					FRenderTargetPoolEvent* Event = *It;
-
-					if (Event->GetEventType() == ERTPE_Phase)
-					{
-						int32 Y0 = It.Y;
-						int32 Y1 = It.FindClosingEventY();
-
-						FIntPoint PixelLeftTop((int32)(DisplayLeftTop.X), (int32)(DisplayLeftTop.Y + ScaleY * Y0));
-						FIntPoint PixelRightBottom((int32)(DisplayLeftTop.X + DisplayExtent.X), (int32)(DisplayLeftTop.Y + ScaleY * Y1));
-
-						bool bHighlight = MousePos.X >= PixelLeftTop.X && MousePos.X < PixelRightBottom.X && MousePos.Y >= PixelLeftTop.Y && MousePos.Y <= PixelRightBottom.Y;
-
-						if (bHighlight)
-						{
-							HighlightedEvent = Event;
-							HighlightedRect = FIntRect(PixelLeftTop, PixelRightBottom);
-						}
-
-						// UMax is 0.9f to avoid getting some wrap texture leaking in at the bottom
-						Canvas.DrawTile(PixelLeftTop.X, PixelLeftTop.Y, PixelRightBottom.X - PixelLeftTop.X, PixelRightBottom.Y - PixelLeftTop.Y, 0, 0, 1, 0.9f, PhaseColor, GradientTexture->GetResource());
-					}
-				}
-
-				// Alloc / Dealloc events
-				for (FRenderTargetPoolEventIterator It(GRenderTargetPool.RenderTargetPoolEvents); *It; ++It)
-				{
-					FRenderTargetPoolEvent* Event = *It;
-
-					if (Event->GetEventType() == ERTPE_Alloc && Event->GetColumnSize())
-					{
-						int32 Y0 = It.Y;
-						int32 Y1 = It.FindClosingEventY();
-
-						int32 X0 = Event->GetColumnX();
-						// for now they are all equal width
-						int32 X1 = X0 + Event->GetColumnSize();
-
-						FIntPoint PixelLeftTop((int32)(DisplayLeftTop.X + ScaleX * X0), (int32)(DisplayLeftTop.Y + ScaleY * Y0));
-						FIntPoint PixelRightBottom((int32)(DisplayLeftTop.X + ScaleX * X1), (int32)(DisplayLeftTop.Y + ScaleY * Y1));
-
-						bool bHighlight = MousePos.X >= PixelLeftTop.X && MousePos.X < PixelRightBottom.X && MousePos.Y >= PixelLeftTop.Y && MousePos.Y <= PixelRightBottom.Y;
-
-						if (bHighlight)
-						{
-							HighlightedEvent = Event;
-							HighlightedRect = FIntRect(PixelLeftTop, PixelRightBottom);
-						}
-
-						FLinearColor Color = ElementColor;
-
-						// Highlight EDRAM/FastVRAM usage
-						if (EnumHasAnyFlags(Event->GetDesc().Flags, TexCreate_FastVRAM))
-						{
-							Color = ElementColorVRam;
-						}
-
-						Canvas.DrawTile(
-							PixelLeftTop.X, PixelLeftTop.Y,
-							PixelRightBottom.X - PixelLeftTop.X - 1, PixelRightBottom.Y - PixelLeftTop.Y - 1,
-							0, 0, 1, 1, Color);
-					}
-				}
-
-				if (HighlightedEvent)
-				{
-					DrawBorder(Canvas, HighlightedRect, FLinearColor(0.8f, 0, 0, 0.5f));
-
-					// Offset to not intersect with crosshair (in editor) or arrow (in game).
-					FIntPoint Pos = MousePos + FIntPoint(12, 4);
-
-					if (HighlightedEvent->GetEventType() == ERTPE_Phase)
-					{
-						FString PhaseText = *FString::Printf(TEXT("Phase: %s"), *HighlightedEvent->GetPhaseName());
-
-						Canvas.DrawShadowedString(Pos.X, Pos.Y + 0 * FontHeight, *PhaseText, GEngine->GetTinyFont(), FLinearColor(0.5f, 0.5f, 1));
-					}
-					else
-					{
-						FString SizeString = FString::Printf(TEXT("%d KB"), (HighlightedEvent->GetSizeInBytes() + 1024) / 1024);
-
-						Canvas.DrawShadowedString(Pos.X, Pos.Y + 0 * FontHeight, HighlightedEvent->GetDesc().DebugName, GEngine->GetTinyFont(), FLinearColor(1, 1, 0));
-						Canvas.DrawShadowedString(Pos.X, Pos.Y + 1 * FontHeight, *HighlightedEvent->GetDesc().GenerateInfoString(), GEngine->GetTinyFont(), FLinearColor(1, 1, 0));
-						Canvas.DrawShadowedString(Pos.X, Pos.Y + 2 * FontHeight, *SizeString, GEngine->GetTinyFont(), FLinearColor(1, 1, 0));
-					}
-				}
-
-				GRenderTargetPool.CurrentEventRecordingTime = 0;
-				GRenderTargetPool.RenderTargetPoolEvents.Empty();
-			});
-		}
-	}
-
 #if SUPPORTS_VISUALIZE_TEXTURE
 	FVisualizeTexture::FCaptured& Captured = GVisualizeTexture.Captured;
 
diff --git a/Engine/Source/Runtime/Renderer/Private/VisualizeTexturePresent.h b/Engine/Source/Runtime/Renderer/Private/VisualizeTexturePresent.h
--- a/Engine/Source/Runtime/Renderer/Private/VisualizeTexturePresent.h
+++ b/Engine/Source/Runtime/Renderer/Private/VisualizeTexturePresent.h
@@ -12,7 +12,4 @@ public:
 
 	/** Present the visualize texture tool on screen. */
 	static void PresentContent(FRDGBuilder& GraphBuilder, const FViewInfo& View, FScreenPassRenderTarget Output);
-
-private:
-	static uint32 ComputeEventDisplayHeight();
 };
\ No newline at end of file
