RDG now reuses tracked state for external pooled textures and buffers. Added a variant of RegisterExternal{Texture, Buffer} that allows the user to specify the actual state if transitioned outside of RDG.

diff --git a/Engine/Source/Runtime/RenderCore/Private/RenderGraphBuilder.cpp b/Engine/Source/Runtime/RenderCore/Private/RenderGraphBuilder.cpp
--- a/Engine/Source/Runtime/RenderCore/Private/RenderGraphBuilder.cpp
+++ b/Engine/Source/Runtime/RenderCore/Private/RenderGraphBuilder.cpp
@@ -654,7 +654,8 @@ void FRDGBuilder::FinalizeResourceAccess(FRDGTextureAccessArray&& InTextures, FR
 FRDGTextureRef FRDGBuilder::RegisterExternalTexture(
 	const TRefCountPtr<IPooledRenderTarget>& ExternalPooledTexture,
 	ERenderTargetTexture RenderTargetTexture,
-	ERDGTextureFlags Flags)
+	ERDGTextureFlags Flags,
+	ERHIAccess AccessInitial)
 {
 #if RDG_ENABLE_DEBUG
 	checkf(ExternalPooledTexture.IsValid(), TEXT("Attempted to register NULL external texture."));
@@ -665,14 +666,15 @@ FRDGTextureRef FRDGBuilder::RegisterExternalTexture(
 	{
 		Name = TEXT("External");
 	}
-	return RegisterExternalTexture(ExternalPooledTexture, Name, RenderTargetTexture, Flags);
+	return RegisterExternalTexture(ExternalPooledTexture, Name, RenderTargetTexture, Flags, AccessInitial);
 }
 
 FRDGTextureRef FRDGBuilder::RegisterExternalTexture(
 	const TRefCountPtr<IPooledRenderTarget>& ExternalPooledTexture,
 	const TCHAR* Name,
 	ERenderTargetTexture RenderTargetTexture,
-	ERDGTextureFlags Flags)
+	ERDGTextureFlags Flags,
+	ERHIAccess AccessInitial)
 {
 	IF_RDG_ENABLE_DEBUG(UserValidation.ValidateRegisterExternalTexture(ExternalPooledTexture, Name, RenderTargetTexture, Flags));
 	FRHITexture* ExternalTextureRHI = ExternalPooledTexture->GetRenderTargetItem().GetRHI(RenderTargetTexture);
@@ -686,9 +688,10 @@ FRDGTextureRef FRDGBuilder::RegisterExternalTexture(
 	const FRDGTextureDesc Desc = Translate(ExternalPooledTexture->GetDesc(), RenderTargetTexture);
 	bool bFinalizedAccess = false;
 
-	if (!EnumHasAnyFlags(Flags, ERDGTextureFlags::ForceTracking)
-		&& !EnumHasAnyFlags(Desc.Flags, TexCreate_RenderTargetable | TexCreate_ResolveTargetable | TexCreate_DepthStencilTargetable | TexCreate_UAV | TexCreate_DepthStencilResolveTarget))
+	if (!EnumHasAnyFlags(Flags, ERDGTextureFlags::ForceTracking) &&
+		!EnumHasAnyFlags(Desc.Flags, TexCreate_RenderTargetable | TexCreate_ResolveTargetable | TexCreate_DepthStencilTargetable | TexCreate_UAV | TexCreate_DepthStencilResolveTarget))
 	{
+		Flags |= ERDGTextureFlags::ReadOnly;
 		bFinalizedAccess = true;
 	}
 
@@ -719,11 +722,13 @@ FRDGTextureRef FRDGBuilder::RegisterExternalTexture(
 		}
 	}
 
-	FRDGTextureSubresourceState& TextureState = Texture->GetState();
-
-	checkf(IsWholeResource(TextureState) && GetWholeResource(TextureState).Access == ERHIAccess::Unknown,
-		TEXT("Externally registered texture '%s' has known RDG state. This means the graph did not sanitize it correctly, or ")
-		TEXT("an IPooledRenderTarget reference was improperly held within a pass."), Texture->Name);
+	// The user provided a known initial state. Otherwise, we rely on the tracked state held by the pooled texture.
+	if (AccessInitial != ERHIAccess::Unknown)
+	{
+		FRDGSubresourceState State;
+		State.Access = AccessInitial;
+		InitAsWholeResource(Texture->GetState(), State);
+	}
 
 	ExternalTextures.Add(Texture->GetRHIUnchecked(), Texture);
 
@@ -732,7 +737,7 @@ FRDGTextureRef FRDGBuilder::RegisterExternalTexture(
 	return Texture;
 }
 
-FRDGBufferRef FRDGBuilder::RegisterExternalBuffer(const TRefCountPtr<FRDGPooledBuffer>& ExternalPooledBuffer, ERDGBufferFlags Flags)
+FRDGBufferRef FRDGBuilder::RegisterExternalBuffer(const TRefCountPtr<FRDGPooledBuffer>& ExternalPooledBuffer, ERDGBufferFlags Flags, ERHIAccess AccessInitial)
 {
 #if RDG_ENABLE_DEBUG
 	checkf(ExternalPooledBuffer.IsValid(), TEXT("Attempted to register NULL external buffer."));
@@ -743,13 +748,14 @@ FRDGBufferRef FRDGBuilder::RegisterExternalBuffer(const TRefCountPtr<FRDGPooledB
 	{
 		Name = TEXT("External");
 	}
-	return RegisterExternalBuffer(ExternalPooledBuffer, Name, Flags);
+	return RegisterExternalBuffer(ExternalPooledBuffer, Name, Flags, AccessInitial);
 }
 
 FRDGBufferRef FRDGBuilder::RegisterExternalBuffer(
 	const TRefCountPtr<FRDGPooledBuffer>& ExternalPooledBuffer,
 	const TCHAR* Name,
-	ERDGBufferFlags Flags)
+	ERDGBufferFlags Flags,
+	ERHIAccess AccessInitial)
 {
 	IF_RDG_ENABLE_DEBUG(UserValidation.ValidateRegisterExternalBuffer(ExternalPooledBuffer, Name, Flags));
 
@@ -761,8 +767,7 @@ FRDGBufferRef FRDGBuilder::RegisterExternalBuffer(
 	const FRDGBufferDesc& Desc = ExternalPooledBuffer->Desc;
 	bool bFinalizedAccess = false;
 
-	if (!EnumHasAnyFlags(Flags, ERDGBufferFlags::ForceTracking)
-		&& !EnumHasAnyFlags(Desc.Usage, BUF_UnorderedAccess))
+	if (!EnumHasAnyFlags(Flags, ERDGBufferFlags::ForceTracking) && !EnumHasAnyFlags(Desc.Usage, BUF_UnorderedAccess))
 	{
 		Flags |= ERDGBufferFlags::ReadOnly;
 		bFinalizedAccess = true;
@@ -782,10 +787,13 @@ FRDGBufferRef FRDGBuilder::RegisterExternalBuffer(
 		Buffer->AccessFinal = ERHIAccess::ReadOnlyExclusiveMask;
 	}
 
-	FRDGSubresourceState& BufferState = Buffer->GetState();
-	checkf(BufferState.Access == ERHIAccess::Unknown,
-		TEXT("Externally registered buffer '%s' has known RDG state. This means the graph did not sanitize it correctly, or ")
-		TEXT("an FRDGPooledBuffer reference was improperly held within a pass."), Buffer->Name);
+	// The user provided a known initial state. Otherwise, we rely on the tracked state held by the pooled buffer.
+	if (AccessInitial != ERHIAccess::Unknown)
+	{
+		FRDGSubresourceState State;
+		State.Access = AccessInitial;
+		Buffer->GetState() = State;
+	}
 
 	ExternalBuffers.Add(ExternalPooledBuffer, Buffer);
 
diff --git a/Engine/Source/Runtime/RenderCore/Private/RenderGraphResourcePool.cpp b/Engine/Source/Runtime/RenderCore/Private/RenderGraphResourcePool.cpp
--- a/Engine/Source/Runtime/RenderCore/Private/RenderGraphResourcePool.cpp
+++ b/Engine/Source/Runtime/RenderCore/Private/RenderGraphResourcePool.cpp
@@ -17,9 +17,7 @@ TRefCountPtr<FRDGPooledBuffer> FRDGBufferPool::FindFreeBuffer(
 	const FRDGBufferDesc& Desc,
 	const TCHAR* InDebugName)
 {
-	TRefCountPtr<FRDGPooledBuffer> Result = FindFreeBufferInternal(RHICmdList, Desc, InDebugName);
-	Result->Reset();
-	return Result;
+	return FindFreeBufferInternal(RHICmdList, Desc, InDebugName);
 }
 
 TRefCountPtr<FRDGPooledBuffer> FRDGBufferPool::FindFreeBufferInternal(
diff --git a/Engine/Source/Runtime/RenderCore/Private/RenderGraphResources.cpp b/Engine/Source/Runtime/RenderCore/Private/RenderGraphResources.cpp
--- a/Engine/Source/Runtime/RenderCore/Private/RenderGraphResources.cpp
+++ b/Engine/Source/Runtime/RenderCore/Private/RenderGraphResources.cpp
@@ -287,18 +287,9 @@ void FRDGTexture::Finalize(FRDGPooledTextureArray& PooledTextureArray)
 		{
 			if (PooledTexture)
 			{
-				// External and extracted resources are user controlled, so we cannot assume the texture stays in its final state.
-				if (bExternal || bExtracted)
-				{
-					PooledTexture->Reset();
-				}
-				else
-				{
-					PooledTexture->Finalize();
-				}
+				PooledTexture->Finalize();
 			}
 
-			// Restore the reference to the last owner in the aliasing chain; not necessary for the transient resource allocator.
 			if (PooledRenderTarget)
 			{
 				PooledTextureArray.Emplace(PooledRenderTarget);
@@ -361,15 +352,7 @@ void FRDGBuffer::Finalize(FRDGPooledBufferArray& PooledBufferArray)
 		}
 		else
 		{
-			if (bExternal || bExtracted)
-			{
-				PooledBuffer->Reset();
-			}
-			else
-			{
-				PooledBuffer->Finalize();
-			}
-
+			PooledBuffer->Finalize();
 			PooledBufferArray.Emplace(PooledBuffer);
 		}
 	}
diff --git a/Engine/Source/Runtime/RenderCore/Private/RenderTargetPool.cpp b/Engine/Source/Runtime/RenderCore/Private/RenderTargetPool.cpp
--- a/Engine/Source/Runtime/RenderCore/Private/RenderTargetPool.cpp
+++ b/Engine/Source/Runtime/RenderCore/Private/RenderTargetPool.cpp
@@ -435,27 +435,7 @@ bool FRenderTargetPool::FindFreeElement(
 		}
 	}
 
-	TRefCountPtr<FPooledRenderTarget> Result = FindFreeElementInternal(RHICmdList, Desc, InDebugName);
-
-	// Reset RDG state back to an unknown default. The resource is being handed off to a user outside of RDG, so the state is no longer valid.
-	{
-		FRDGPooledTexture* TargetableTexture = Result->TargetableTexture;
-		FRDGPooledTexture* ShaderResourceTexture = Result->ShaderResourceTexture;
-
-		if (TargetableTexture)
-		{
-			checkf(!TargetableTexture->GetOwner(), TEXT("Allocated a pooled render target that is currently owned by RDG texture %s."), TargetableTexture->GetOwner()->Name);
-			TargetableTexture->Reset();
-		}
-
-		if (ShaderResourceTexture && ShaderResourceTexture != TargetableTexture)
-		{
-			checkf(!ShaderResourceTexture->GetOwner(), TEXT("Allocated a pooled render target that is currently owned by RDG texture %s."), ShaderResourceTexture->GetOwner()->Name);
-			ShaderResourceTexture->Reset();
-		}
-	}
-
-	Out = Result;
+	Out = FindFreeElementInternal(RHICmdList, Desc, InDebugName);
 	return false;
 }
 
diff --git a/Engine/Source/Runtime/RenderCore/Public/RenderGraphBuilder.h b/Engine/Source/Runtime/RenderCore/Public/RenderGraphBuilder.h
--- a/Engine/Source/Runtime/RenderCore/Public/RenderGraphBuilder.h
+++ b/Engine/Source/Runtime/RenderCore/Public/RenderGraphBuilder.h
@@ -39,24 +39,29 @@ public:
 	FRDGTextureRef RegisterExternalTexture(
 		const TRefCountPtr<IPooledRenderTarget>& ExternalPooledTexture,
 		ERenderTargetTexture Texture = ERenderTargetTexture::ShaderResource,
-		ERDGTextureFlags Flags = ERDGTextureFlags::None);
+		ERDGTextureFlags Flags = ERDGTextureFlags::None,
+		ERHIAccess AccessInitial = ERHIAccess::Unknown);
 
 	/** Register an external texture with a custom name. The name is only used if the texture has not already been registered. */
 	FRDGTextureRef RegisterExternalTexture(
 		const TRefCountPtr<IPooledRenderTarget>& ExternalPooledTexture,
 		const TCHAR* NameIfNotRegistered,
 		ERenderTargetTexture RenderTargetTexture = ERenderTargetTexture::ShaderResource,
-		ERDGTextureFlags Flags = ERDGTextureFlags::None);
+		ERDGTextureFlags Flags = ERDGTextureFlags::None,
+		ERHIAccess AccessInitial = ERHIAccess::Unknown);
 
 	/** Register a external buffer to be tracked by the render graph. */
-	FRDGBufferRef RegisterExternalBuffer(const TRefCountPtr<FRDGPooledBuffer>& ExternalPooledBuffer, ERDGBufferFlags Flags = ERDGBufferFlags::None);
-	FRDGBufferRef RegisterExternalBuffer(const TRefCountPtr<FRDGPooledBuffer>& ExternalPooledBuffer, ERDGBufferFlags Flags, ERHIAccess AccessFinal);
+	FRDGBufferRef RegisterExternalBuffer(
+		const TRefCountPtr<FRDGPooledBuffer>& ExternalPooledBuffer,
+		ERDGBufferFlags Flags = ERDGBufferFlags::None,
+		ERHIAccess AccessInitial = ERHIAccess::Unknown);
 
 	/** Register an external buffer with a custom name. The name is only used if the buffer has not already been registered. */
 	FRDGBufferRef RegisterExternalBuffer(
 		const TRefCountPtr<FRDGPooledBuffer>& ExternalPooledBuffer,
 		const TCHAR* NameIfNotRegistered,
-		ERDGBufferFlags Flags = ERDGBufferFlags::None);
+		ERDGBufferFlags Flags = ERDGBufferFlags::None,
+		ERHIAccess AccessInitial = ERHIAccess::Unknown);
 
 	/** Create graph tracked texture from a descriptor. The CPU memory is guaranteed to be valid through execution of
 	 *  the graph, at which point it is released. The underlying RHI texture lifetime is only guaranteed for passes which
