Compress shaders as groups for cooked games.

- Saves disk size and runtime memory on all platforms that use IoStore-based shader archive (on some - significantly).
- Decompression time is small and seems to get lost in noise given how rarely shaders are actually created.
- Also changed IoStore-based archive to allocate preload entries on demand, and disabled preloading singular shaders (used by the PSO cache) as the API needs more work.

diff --git a/Engine/Source/Developer/IoStoreUtilities/Private/IoStoreUtilities.cpp b/Engine/Source/Developer/IoStoreUtilities/Private/IoStoreUtilities.cpp
--- a/Engine/Source/Developer/IoStoreUtilities/Private/IoStoreUtilities.cpp
+++ b/Engine/Source/Developer/IoStoreUtilities/Private/IoStoreUtilities.cpp
@@ -310,7 +310,8 @@ struct FShaderInfo
 
 	FIoChunkId ChunkId;
 	FIoBuffer CodeIoBuffer;
-	FSHAHash LibraryCodeHash;
+	// the smaller the order, the more likely this will be loaded
+	uint32 LoadOrderFactor;
 	uint64 DiskLayoutOrder = MAX_uint64;
 	TSet<struct FLegacyCookedPackage*> ReferencedByPackages;
 	TMap<FContainerTargetSpec*, EShaderType> TypeInContainer;
@@ -1117,7 +1118,7 @@ static void CreateDiskLayout(
 				{
 					if (A->DiskLayoutOrder == B->DiskLayoutOrder)
 					{
-						return A->LibraryCodeHash < B->LibraryCodeHash;
+						return A->LoadOrderFactor < B->LoadOrderFactor;
 					}
 					return A->DiskLayoutOrder < B->DiskLayoutOrder;
 				});
@@ -1423,7 +1424,7 @@ bool LoadShaderAssetInfo(const FString& Filename, TMap<FSHAHash, TSet<FName>>& O
 bool ConvertToIoStoreShaderLibrary(
 	const TCHAR* FileName,
 	TTuple<FIoChunkId, FIoBuffer>& OutLibraryIoChunk,
-	TArray<TTuple<FIoChunkId, FIoBuffer, FSHAHash>>& OutCodeIoChunks,
+	TArray<TTuple<FIoChunkId, FIoBuffer, uint32>>& OutCodeIoChunks,
 	TArray<TTuple<FSHAHash, TArray<FIoChunkId>>>& OutShaderMaps,
 	TArray<TTuple<FSHAHash, TSet<FName>>>& OutShaderMapAssetAssociations)
 {
@@ -1458,39 +1459,170 @@ bool ConvertToIoStoreShaderLibrary(
 		return false;
 	}
 
-	FLargeMemoryWriter IoStoreLibraryAr(0, true);
-	FIoStoreShaderCodeArchive::SaveIoStoreShaderCodeArchive(SerializedShaders, IoStoreLibraryAr);
-	OutLibraryIoChunk.Key = FIoStoreShaderCodeArchive::GetShaderCodeArchiveChunkId(LibraryName, FormatName);
-	int64 TotalSize = IoStoreLibraryAr.TotalSize();
-	OutLibraryIoChunk.Value = FIoBuffer(FIoBuffer::AssumeOwnership, IoStoreLibraryAr.ReleaseOwnership(), TotalSize);
+	FIoStoreShaderCodeArchiveHeader IoStoreLibraryHeader;
+	FIoStoreShaderCodeArchive::CreateIoStoreShaderCodeArchiveHeader(FName(FormatName), SerializedShaders, IoStoreLibraryHeader);
+	checkf(SerializedShaders.ShaderEntries.Num() == IoStoreLibraryHeader.ShaderEntries.Num(), TEXT("IoStore header has different number of shaders (%d) than the original library (%d)"), 
+		IoStoreLibraryHeader.ShaderEntries.Num(), SerializedShaders.ShaderEntries.Num());
+	checkf(SerializedShaders.ShaderMapEntries.Num() == IoStoreLibraryHeader.ShaderMapEntries.Num(), TEXT("IoStore header has different number of shadermaps (%d) than the original library (%d)"),
+		IoStoreLibraryHeader.ShaderMapEntries.Num(), SerializedShaders.ShaderMapEntries.Num());
+
+	// Load all the shaders, decompress them and recompress into groups. Each code chunk is shader group chunk now.
+	// This also updates group's compressed size, the only value left to calculate.
+	int32 GroupCount = IoStoreLibraryHeader.ShaderGroupEntries.Num();
+	OutCodeIoChunks.SetNum(GroupCount);
+	FCriticalSection DiskArchiveAccess;
+	ParallelFor(OutCodeIoChunks.Num(),
+		[&OutCodeIoChunks, &IoStoreLibraryHeader, &DiskArchiveAccess, &LibraryAr, &SerializedShaders, &OffsetToShaderCode](int32 GroupIndex)
+		{
+			FIoStoreShaderGroupEntry& Group = IoStoreLibraryHeader.ShaderGroupEntries[GroupIndex];
+			uint8* UncompressedGroupMemory = reinterpret_cast<uint8*>(FMemory::Malloc(Group.UncompressedSize));
+			check(UncompressedGroupMemory);
+
+			int32 SmallestShaderInGroupByOrdinal = MAX_int32;
+			uint8* ShaderStart = UncompressedGroupMemory;
+			// Not worth to special-case for group size of 1 (by converting to copy) - such groups are very fast to decompress and recompress.
+			for (uint32 IdxShaderInGroup = 0; IdxShaderInGroup < Group.NumShaders; ++IdxShaderInGroup)
+			{
+				int32 ShaderIndex = IoStoreLibraryHeader.ShaderIndices[Group.ShaderIndicesOffset + IdxShaderInGroup];
+				SmallestShaderInGroupByOrdinal = FMath::Min(SmallestShaderInGroupByOrdinal, ShaderIndex);
+				const FIoStoreShaderCodeEntry& Shader = IoStoreLibraryHeader.ShaderEntries[ShaderIndex];
+
+				checkf((reinterpret_cast<uint64>(ShaderStart) - reinterpret_cast<uint64>(UncompressedGroupMemory)) == Shader.UncompressedOffsetInGroup,
+					TEXT("Shader uncompressed offset in group does not agree with its actual placement: Shader.UncompressedOffsetInGroup=%llu, actual=%llu"),
+					Shader.UncompressedOffsetInGroup, (reinterpret_cast<uint64>(ShaderStart) - reinterpret_cast<uint64>(UncompressedGroupMemory))
+				);
+
+				// load and decompress the shader at the desired offset
+				const FShaderCodeEntry& IndividuallyCompressedShader = SerializedShaders.ShaderEntries[ShaderIndex];
+				{
+					// small shaders might be stored without compression, handle them here
+					if (IndividuallyCompressedShader.Size == IndividuallyCompressedShader.UncompressedSize)
+					{
+						// disk access has to be serialized between the for loops
+						FScopeLock Lock(&DiskArchiveAccess);
+						LibraryAr->Seek(OffsetToShaderCode + IndividuallyCompressedShader.Offset);
+						LibraryAr->Serialize(ShaderStart, IndividuallyCompressedShader.Size);
+					}
+					else
+					{
+						uint8* CompressedShaderMemory = reinterpret_cast<uint8*>(FMemory::Malloc(IndividuallyCompressedShader.Size));
 
-	int32 ShaderCount = SerializedShaders.ShaderEntries.Num();
-	OutCodeIoChunks.Reserve(ShaderCount);
-	for (int32 ShaderIndex = 0; ShaderIndex < ShaderCount; ++ShaderIndex)
-	{
-		const FSHAHash& ShaderHash = SerializedShaders.ShaderHashes[ShaderIndex];
-		const FShaderCodeEntry& ShaderEntry = SerializedShaders.ShaderEntries[ShaderIndex];
-		TTuple<FIoChunkId, FIoBuffer, FSHAHash>& OutCodeIoChunk = OutCodeIoChunks.Emplace_GetRef(FIoStoreShaderCodeArchive::GetShaderCodeChunkId(ShaderHash), ShaderEntry.Size, ShaderHash);
-		LibraryAr->Seek(OffsetToShaderCode + ShaderEntry.Offset);
-		LibraryAr->Serialize(OutCodeIoChunk.Get<1>().Data(), ShaderEntry.Size);
-	}
+						// disk access has to be serialized between the for loops
+						{
+							FScopeLock Lock(&DiskArchiveAccess);
+							LibraryAr->Seek(OffsetToShaderCode + IndividuallyCompressedShader.Offset);
+							LibraryAr->Serialize(CompressedShaderMemory, IndividuallyCompressedShader.Size);
+						}
+
+						// This function will crash if decompression fails.
+						ShaderCodeArchive::DecompressShader(ShaderStart, IndividuallyCompressedShader.UncompressedSize, CompressedShaderMemory, IndividuallyCompressedShader.Size);
+
+						FMemory::Free(CompressedShaderMemory);
+					}
+				}
 
-	int32 ShaderMapCount = SerializedShaders.ShaderMapHashes.Num();
-	OutShaderMaps.Reserve(ShaderMapCount);
-	for (int32 ShaderMapIndex = 0; ShaderMapIndex < ShaderMapCount; ++ShaderMapIndex)
+				ShaderStart += IndividuallyCompressedShader.UncompressedSize;
+			}
+
+			checkf((reinterpret_cast<uint64>(ShaderStart) - reinterpret_cast<uint64>(UncompressedGroupMemory)) == Group.UncompressedSize,
+				TEXT("Uncompressed shader group size does not agree with the actual results (Group.UncompressedSize=%llu, actual=%llu)"), 
+				Group.UncompressedSize, (reinterpret_cast<uint64>(ShaderStart) - reinterpret_cast<uint64>(UncompressedGroupMemory)));
+
+			// now compress the whole group
+			int64 CompressedGroupSize = 0;
+			ShaderCodeArchive::CompressShaderUsingCurrentSettings(nullptr, CompressedGroupSize, UncompressedGroupMemory, Group.UncompressedSize);
+			checkf(CompressedGroupSize > 0, TEXT("CompressedGroupSize estimate seems wrong (%lld)"), CompressedGroupSize);
+
+			uint8* CompressedShaderGroupMemory = reinterpret_cast<uint8*>(FMemory::Malloc(CompressedGroupSize));
+			bool bCompressed = ShaderCodeArchive::CompressShaderUsingCurrentSettings(CompressedShaderGroupMemory, CompressedGroupSize, UncompressedGroupMemory, Group.UncompressedSize);
+			checkf(bCompressed, TEXT("We could not compress the shader group after providing an estimated memory."));
+
+			TTuple<FIoChunkId, FIoBuffer, uint32>& OutCodeIoChunk = OutCodeIoChunks[GroupIndex];
+			OutCodeIoChunk.Get<0>() = IoStoreLibraryHeader.ShaderGroupIoHashes[GroupIndex];
+			// This value is the load order factor for the group (the smaller, the more likely), that IoStore will use to sort the chunks.
+			// We calculate it as the smallest shader index in the group, which shouldn't be bad. Arguably a better way would be to get the lowest-numbered shadermap that references _any_ shader in the group, 
+			// but it's much, much slower to calculate and the resulting order would be likely the same/similar most of the time
+			checkf(SmallestShaderInGroupByOrdinal >= 0 && SmallestShaderInGroupByOrdinal < IoStoreLibraryHeader.ShaderEntries.Num(), TEXT("SmallestShaderInGroupByOrdinal has an invalid value of %d (not within [0, %d) range as expected)"),
+				SmallestShaderInGroupByOrdinal, IoStoreLibraryHeader.ShaderEntries.Num());
+			OutCodeIoChunk.Get<2>() = static_cast<uint32>(SmallestShaderInGroupByOrdinal);
+
+			if (CompressedGroupSize < Group.UncompressedSize)
+			{
+				OutCodeIoChunk.Get<1>() = FIoBuffer(CompressedGroupSize);
+				FMemory::Memcpy(OutCodeIoChunk.Get<1>().GetData(), CompressedShaderGroupMemory, CompressedGroupSize);
+				Group.CompressedSize = CompressedGroupSize;
+			}
+			else
+			{
+				// store uncompressed (unlikely, but happens for a 200-byte sized shader that happens to get its own group)
+				OutCodeIoChunk.Get<1>() = FIoBuffer(Group.UncompressedSize);
+				FMemory::Memcpy(OutCodeIoChunk.Get<1>().GetData(), UncompressedGroupMemory, Group.UncompressedSize);
+				Group.CompressedSize = Group.UncompressedSize;
+			}
+		},
+		EParallelForFlags::Unbalanced
+	);
+
+	// calculate and log stats
+	int64 TotalUncompressedSizeViaGroups = 0;	// calculate total uncompressed size twice for a sanity check
+	int64 TotalUncompressedSizeViaShaders = 0;
+	int64 TotalIndividuallyCompressedSize = 0;
+	int64 TotalGroupCompressedSize = 0;
+	for (int32 IdxGroup = 0, NumGroups = IoStoreLibraryHeader.ShaderGroupEntries.Num(); IdxGroup < NumGroups; ++IdxGroup)
 	{
-		TTuple<FSHAHash, TArray<FIoChunkId>>& OutShaderMap = OutShaderMaps.AddDefaulted_GetRef();
-		OutShaderMap.Key = SerializedShaders.ShaderMapHashes[ShaderMapIndex];
-		const FShaderMapEntry& ShaderMapEntry = SerializedShaders.ShaderMapEntries[ShaderMapIndex];
-		int32 LookupIndexEnd = ShaderMapEntry.ShaderIndicesOffset + ShaderMapEntry.NumShaders;
-		OutShaderMap.Value.Reserve(ShaderMapEntry.NumShaders);
-		for (int32 ShaderLookupIndex = ShaderMapEntry.ShaderIndicesOffset; ShaderLookupIndex < LookupIndexEnd; ++ShaderLookupIndex)
+		FIoStoreShaderGroupEntry& Group = IoStoreLibraryHeader.ShaderGroupEntries[IdxGroup];
+		TotalGroupCompressedSize += Group.CompressedSize;
+		TotalUncompressedSizeViaGroups += Group.UncompressedSize;
+
+		// now go via shaders
+		for (uint32 IdxShaderInGroup = 0; IdxShaderInGroup < Group.NumShaders; ++IdxShaderInGroup)
 		{
-			int32 ShaderIndex = SerializedShaders.ShaderIndices[ShaderLookupIndex];
-			OutShaderMap.Value.Add(FIoStoreShaderCodeArchive::GetShaderCodeChunkId(SerializedShaders.ShaderHashes[ShaderIndex]));
+			int32 ShaderIndex = IoStoreLibraryHeader.ShaderIndices[Group.ShaderIndicesOffset + IdxShaderInGroup];
+			const FShaderCodeEntry& IndividuallyCompressedShader = SerializedShaders.ShaderEntries[ShaderIndex];
+
+			TotalIndividuallyCompressedSize += IndividuallyCompressedShader.Size;
+			TotalUncompressedSizeViaShaders += IndividuallyCompressedShader.UncompressedSize;
 		}
 	}
 
+	checkf(TotalUncompressedSizeViaGroups == TotalUncompressedSizeViaShaders, TEXT("Sanity check failure: total uncompressed shader size differs if calculated via shader groups (%lld) or individual shaders (%lld)"),
+		TotalUncompressedSizeViaGroups, TotalUncompressedSizeViaShaders
+		);
+
+	UE_LOG(LogIoStore, Display, TEXT("%s(%s): Recompressed %d shaders as %d groups. Library size changed from %lld KB (%.2f : 1 ratio) to %lld KB (%.2f : 1 ratio), %.2f%% of previous."),
+		*LibraryName, *FormatName.ToString(),
+		SerializedShaders.ShaderEntries.Num(), IoStoreLibraryHeader.ShaderGroupEntries.Num(),
+		TotalIndividuallyCompressedSize / 1024, static_cast<double>(TotalUncompressedSizeViaShaders) / static_cast<double>(TotalIndividuallyCompressedSize),
+		TotalGroupCompressedSize / 1024, static_cast<double>(TotalUncompressedSizeViaShaders) / static_cast<double>(TotalGroupCompressedSize),
+		100.0 * static_cast<double>(TotalGroupCompressedSize) / static_cast<double>(TotalIndividuallyCompressedSize)
+		);
+
+	FLargeMemoryWriter IoStoreLibraryAr(0, true);
+	FIoStoreShaderCodeArchive::SaveIoStoreShaderCodeArchive(IoStoreLibraryHeader, IoStoreLibraryAr);
+	OutLibraryIoChunk.Key = FIoStoreShaderCodeArchive::GetShaderCodeArchiveChunkId(LibraryName, FormatName);
+	int64 TotalSize = IoStoreLibraryAr.TotalSize();
+	OutLibraryIoChunk.Value = FIoBuffer(FIoBuffer::AssumeOwnership, IoStoreLibraryAr.ReleaseOwnership(), TotalSize);
+
+	int32 ShaderMapCount = IoStoreLibraryHeader.ShaderMapHashes.Num();
+	OutShaderMaps.SetNum(ShaderMapCount);
+	ParallelFor(OutShaderMaps.Num(),
+		[&OutShaderMaps, &IoStoreLibraryHeader](int32 ShaderMapIndex)
+		{
+			TTuple<FSHAHash, TArray<FIoChunkId>>& OutShaderMap = OutShaderMaps[ShaderMapIndex];
+			OutShaderMap.Key = IoStoreLibraryHeader.ShaderMapHashes[ShaderMapIndex];
+			const FIoStoreShaderMapEntry& ShaderMapEntry = IoStoreLibraryHeader.ShaderMapEntries[ShaderMapIndex];
+			int32 LookupIndexEnd = ShaderMapEntry.ShaderIndicesOffset + ShaderMapEntry.NumShaders;
+			OutShaderMap.Value.Reserve(ShaderMapEntry.NumShaders);	// worst-case, 1 group == 1 shader
+			for (int32 ShaderLookupIndex = ShaderMapEntry.ShaderIndicesOffset; ShaderLookupIndex < LookupIndexEnd; ++ShaderLookupIndex)
+			{
+				int32 ShaderIndex = IoStoreLibraryHeader.ShaderIndices[ShaderLookupIndex];
+				int32 GroupIndex = IoStoreLibraryHeader.ShaderEntries[ShaderIndex].ShaderGroupIndex;
+				OutShaderMap.Value.AddUnique(IoStoreLibraryHeader.ShaderGroupIoHashes[GroupIndex]);
+			}
+		},
+		EParallelForFlags::Unbalanced
+	);
+
 	OutShaderMapAssetAssociations.Reserve(ShaderCodeToAssets.Num());
 	for (auto& KV : ShaderCodeToAssets)
 	{
@@ -1505,7 +1637,6 @@ void ProcessShaderLibraries(const FIoStoreArguments& Arguments, TArray<FContaine
 	IOSTORE_CPU_SCOPE(ProcessShaderLibraries);
 
 	TMap<FIoChunkId, FShaderInfo*> ChunkIdToShaderInfoMap;
-	TMap<FIoChunkId, FSHAHash> AllShaderCodeHashes;
 	TMap<FSHAHash, TArray<FIoChunkId>> ShaderChunkIdsByShaderMapHash;
 	TMap<FName, TSet<FSHAHash>> PackageNameToShaderMaps;
 	TMap<FContainerTargetSpec*, TSet<FShaderInfo*>> AllContainerShaderLibraryShadersMap;
@@ -1521,7 +1652,7 @@ void ProcessShaderLibraries(const FIoStoreArguments& Arguments, TArray<FContaine
 
 				TArray<TTuple<FSHAHash, TArray<FIoChunkId>>> ShaderMaps;
 				TTuple<FIoChunkId, FIoBuffer> LibraryChunk;
-				TArray<TTuple<FIoChunkId, FIoBuffer, FSHAHash>> CodeChunks;
+				TArray<TTuple<FIoChunkId, FIoBuffer, uint32>> CodeChunks;
 				TArray<TTuple<FSHAHash, TSet<FName>>> ShaderMapAssetAssociations;
 				if (!ConvertToIoStoreShaderLibrary(*TargetFile.NormalizedSourcePath, LibraryChunk, CodeChunks, ShaderMaps, ShaderMapAssetAssociations))
 				{
@@ -1532,21 +1663,16 @@ void ProcessShaderLibraries(const FIoStoreArguments& Arguments, TArray<FContaine
 				TargetFile.SourceBuffer.Emplace(LibraryChunk.Value);
 				const bool bIsGlobalShaderLibrary = FPaths::GetCleanFilename(TargetFile.NormalizedSourcePath).StartsWith(TEXT("ShaderArchive-Global-"));
 				const FShaderInfo::EShaderType ShaderType = bIsGlobalShaderLibrary ? FShaderInfo::Global : FShaderInfo::Normal;
-				for (const TTuple<FIoChunkId, FIoBuffer, FSHAHash>& CodeChunk : CodeChunks)
+				for (const TTuple<FIoChunkId, FIoBuffer, uint32>& CodeChunk : CodeChunks)
 				{
 					const FIoChunkId& ShaderChunkId = CodeChunk.Get<0>();
-					FSHAHash* FindShaderCodeHash = AllShaderCodeHashes.Find(ShaderChunkId);
-					if (FindShaderCodeHash && *FindShaderCodeHash != CodeChunk.Get<2>())
-					{
-						UE_LOG(LogIoStore, Fatal, TEXT("Shader code chunk id collision"));
-					}
 					FShaderInfo* ShaderInfo = ChunkIdToShaderInfoMap.FindRef(ShaderChunkId);
 					if (!ShaderInfo)
 					{
 						ShaderInfo = new FShaderInfo();
 						ShaderInfo->ChunkId = ShaderChunkId;
 						ShaderInfo->CodeIoBuffer = CodeChunk.Get<1>();
-						ShaderInfo->LibraryCodeHash = CodeChunk.Get<2>();
+						ShaderInfo->LoadOrderFactor = CodeChunk.Get<2>();
 						OutShaders.Add(ShaderInfo);
 						ChunkIdToShaderInfoMap.Add(ShaderChunkId, ShaderInfo);
 					}
diff --git a/Engine/Source/Runtime/RenderCore/Private/ShaderCodeArchive.cpp b/Engine/Source/Runtime/RenderCore/Private/ShaderCodeArchive.cpp
--- a/Engine/Source/Runtime/RenderCore/Private/ShaderCodeArchive.cpp
+++ b/Engine/Source/Runtime/RenderCore/Private/ShaderCodeArchive.cpp
@@ -15,6 +15,9 @@
 #include "Misc/StringBuilder.h"
 #include "Templates/Greater.h"
 #endif
+#include "RenderUtils.h"
+#include "RHIShaderFormatDefinitions.inl"
+#include "Async/ParallelFor.h"
 
 #if UE_SCA_VISUALIZE_SHADER_USAGE
 #include "Modules/ModuleManager.h"
@@ -47,6 +50,16 @@ static FAutoConsoleVariableRef CVarShaderCodeLibraryVisualizeShaderUsage(
 	ECVF_RenderThreadSafe | ECVF_ReadOnly
 );
 
+int32 GShaderCodeLibraryMaxShaderGroupSize = 1024 * 1024;	// decompressing 1MB of shaders takes about 0.1ms on PC (TR 3970x, Oodle Mermaid6). Can be overriden per platform.
+static FAutoConsoleVariableRef CVarShaderCodeLibraryMaxShaderGroupSize(
+	TEXT("r.ShaderCodeLibrary.MaxShaderGroupSize"),
+	GShaderCodeLibraryMaxShaderGroupSize,
+	TEXT("Max (uncompressed) size of a group of shaders to be compressed/decompressed together. Can be set differently for each platform in their own *Engine.ini.")
+	TEXT("If a group exceeds it, it will be evenly split into subgroups which strive to not exceed it. However, if a shader group is down to one shader and still exceeds the limit, the limit will be ignored."),
+	ECVF_RenderThreadSafe | ECVF_ReadOnly
+);
+
+
 int32 FSerializedShaderArchive::FindShaderMapWithKey(const FSHAHash& Hash, uint32 Key) const
 {
 	for (uint32 Index = ShaderMapHashTable.First(Key); ShaderMapHashTable.IsValid(Index); Index = ShaderMapHashTable.Next(Index))
@@ -161,7 +174,7 @@ void FShaderUsageVisualizer::SaveShaderUsageBitmap(const FString& Name, EShaderP
 			{
 				if (TSharedPtr<IImageWrapper> PNGImageWrapper = ImageWrapperModule->CreateImageWrapper(EImageFormat::PNG))
 				{
-					UE_LOG(LogShaderLibrary, Log, TEXT("Creating shader usage bitmap for archive %s (NumShaders: %d, preloaded %d, created %d)"), *Name, NumShaders, PreloadedShadersForVis.Num(), CreatedShadersForVis.Num());
+					UE_LOG(LogShaderLibrary, Log, TEXT("Creating shader usage bitmap for archive %s (NumShaders: %d, preloaded %d, created %d)"), *Name, NumShaders, PreloadedShaders.Num(), CreatedShaders.Num());
 
 					// find a value close to sqrt(NumShaders)
 					int32 ImageDimension = static_cast<int32>(FMath::Sqrt(static_cast<float>(NumShaders))) + 1;
@@ -170,8 +183,11 @@ void FShaderUsageVisualizer::SaveShaderUsageBitmap(const FString& Name, EShaderP
 
 					// map legend:
 					FColor UnusedShaderColor(128, 128, 128);	// unused shaders - this is the majority of the bitmap content
-					FColor PreloadedShaderColor(225, 225, 225);	// preloaded shaders - they can become the majority under certain circumstances. Pure white can blend with some viewer's background
-					FColor CreatedShaderColor(0, 0, 255);	// created shaders - in practice, always few and far between. Blue is more noticeable on a largely bright background than magenta
+					FColor PreloadedShaderColor(192, 192, 192);	// preloaded shaders - including those that weren't explicitly so, but they happened to be grouped with shaders we needed
+					FColor ExplicitlyPreloadedShaderColor(0, 255, 0);	// shaders we explicitly wanted to preload - they can become the majority under certain circumstances. Pure white can blend with some viewer's background
+					FColor PreloadedAndDecompressedShaderColor(0, 0, 255);	// shaders that we wanted to preload and that got decompressed (as part of the creating them or their neighbor in group)
+					FColor NotPreloadedButDecompressedShaderColor(255, 0, 0);	// shaders that we decompressed just because they were grouped together with others. We did not want to preload them at all.
+					FColor CreatedShaderColor(255, 255, 255);	// created shaders - in practice, always few and far between. Blue is more noticeable on a largely bright background than magenta
 
 					for (int32 Idx = 0; Idx < NumShaders; ++Idx)
 					{
@@ -186,12 +202,25 @@ void FShaderUsageVisualizer::SaveShaderUsageBitmap(const FString& Name, EShaderP
 						FScopeLock Lock(&VisualizeLock);
 
 						// fill preloaded ones first
-						for (int32 ShaderIdx : PreloadedShadersForVis)
+						for (int32 ShaderIdx : PreloadedShaders)
 						{
 							ShaderUsageBitmap[ShaderIdx] = PreloadedShaderColor;
 						}
 
-						for (int32 ShaderIdx : CreatedShadersForVis)
+						// explicitly preloaded shaders
+						for (int32 ShaderIdx : ExplicitlyPreloadedShaders)
+						{
+							ShaderUsageBitmap[ShaderIdx] = ExplicitlyPreloadedShaderColor;
+						}
+
+						// fill decompressed ones, but mark up those that we didn't ask to preload differently
+						for (int32 ShaderIdx : DecompressedShaders)
+						{
+							bool bShaderWasRequestedToBePreloaded = ExplicitlyPreloadedShaders.Contains(ShaderIdx);
+							ShaderUsageBitmap[ShaderIdx] = bShaderWasRequestedToBePreloaded ? PreloadedAndDecompressedShaderColor : NotPreloadedButDecompressedShaderColor;
+						}
+
+						for (int32 ShaderIdx : CreatedShaders)
 						{
 							ShaderUsageBitmap[ShaderIdx] = CreatedShaderColor;
 						}
@@ -214,44 +243,70 @@ void FShaderUsageVisualizer::SaveShaderUsageBitmap(const FString& Name, EShaderP
 						}
 						else
 						{
-							// technically not every created might have been preloaded, but we ignore this here
-							UE_LOG(LogShaderLibrary, Log, TEXT("Saved shader usage bitmap %s. Legend: unused shaders: %d (dark grey), preloaded (compressed) to RAM: %d (white), uncompressed and created: %d (blue)"), 
-								*FilePath, NumShaders - PreloadedShadersForVis.Num(), PreloadedShadersForVis.Num(), CreatedShadersForVis.Num());
+							UE_LOG(LogShaderLibrary, Log, TEXT("Saved shader usage bitmap %s. Legend: shaders not loaded from disk - dark grey, loaded to RAM explicitly - green, loaded as part of their compressed group - bright grey, decompressed: blue if they were loaded explicitly, red if just because they were a part of the group. Actually created shaders - white"), 
+								*FilePath);
 						}
 					}
 					else
 					{
-						UE_LOG(LogShaderLibrary, Warning, TEXT("Error creating shader usage bitmap for archive %s (NumShaders: %d, preloaded %d, created %d) - cannot create a PNG image"), *Name, NumShaders, PreloadedShadersForVis.Num(), CreatedShadersForVis.Num());
+						UE_LOG(LogShaderLibrary, Warning, TEXT("Error creating shader usage bitmap for archive %s (NumShaders: %d, preloaded %d, created %d) - cannot create a PNG image"), *Name, NumShaders, PreloadedShaders.Num(), CreatedShaders.Num());
 					}
 				}
 				else
 				{
-					UE_LOG(LogShaderLibrary, Warning, TEXT("Couldn't create shader usage bitmap for archive %s (NumShaders: %d, preloaded %d, created %d) - cannot create ImageWrapper for PNG format"), *Name, NumShaders, PreloadedShadersForVis.Num(), CreatedShadersForVis.Num());
+					UE_LOG(LogShaderLibrary, Warning, TEXT("Couldn't create shader usage bitmap for archive %s (NumShaders: %d, preloaded %d, created %d) - cannot create ImageWrapper for PNG format"), *Name, NumShaders, PreloadedShaders.Num(), CreatedShaders.Num());
 				}
 			}
 			else
 			{
-				UE_LOG(LogShaderLibrary, Warning, TEXT("Couldn't create shader usage bitmap for archive %s (NumShaders: %d, preloaded %d, created %d) - no ImageWrapper module"), *Name, NumShaders, PreloadedShadersForVis.Num(), CreatedShadersForVis.Num());
+				UE_LOG(LogShaderLibrary, Warning, TEXT("Couldn't create shader usage bitmap for archive %s (NumShaders: %d, preloaded %d, created %d) - no ImageWrapper module"), *Name, NumShaders, PreloadedShaders.Num(), CreatedShaders.Num());
 			}
 		}
 	}
 }
 #endif
 
-namespace ShaderCodeArchive
+void ShaderCodeArchive::DecompressShader(uint8* OutDecompressedShader, int64 UncompressedSize, const uint8* CompressedShaderCode, int64 CompressedSize)
 {
-	/** Decompresses the shader into caller-provided memory. Caller is assumed to allocate at least ShaderEntry uncompressed size value.
-	 * The engine will crash (LogFatal) if this function fails.
-	 */
-	void DecompressShader(uint8* OutDecompressedShader, int64 UncompressedSize, const uint8* CompressedShaderCode, int64 CompressedSize)
+	bool bSucceed = FCompression::UncompressMemory(GetShaderCompressionFormat(), OutDecompressedShader, UncompressedSize, CompressedShaderCode, CompressedSize);
+	if (!bSucceed)
 	{
-		bool bSucceed = FCompression::UncompressMemory(GetShaderCompressionFormat(), OutDecompressedShader, UncompressedSize, CompressedShaderCode, CompressedSize);
-		check(bSucceed);
-		if (!bSucceed)
+		UE_LOG(LogShaderLibrary, Fatal, TEXT("ShaderCodeArchive::DecompressShader(): Could not decompress shader (GetShaderCompressionFormat=%s)"), *GetShaderCompressionFormat().ToString());
+	}
+}
+
+bool ShaderCodeArchive::CompressShaderUsingCurrentSettings(uint8* OutCompressedShader, int64& OutCompressedSize, const uint8* UncompressedShaderCode, int64 UncompressedSize)
+{
+	// see FShaderCode::Compress - while this doesn't have to match exactly, it should match at least the parameters used there
+	const FName ShaderCompressionFormat = GetShaderCompressionFormat();
+
+	bool bCompressed = false;
+	if (ShaderCompressionFormat != NAME_Oodle)
+	{
+		int32 CompressedSize32 = static_cast<int32>(OutCompressedSize);
+		checkf(static_cast<int64>(CompressedSize32) == OutCompressedSize, TEXT("CompressedSize is too large (%lld) for an old API that takes int32"), OutCompressedSize);
+		bCompressed = FCompression::CompressMemory(ShaderCompressionFormat, OutCompressedShader, CompressedSize32, UncompressedShaderCode, UncompressedSize, COMPRESS_BiasSize);
+		OutCompressedSize = static_cast<int64>(CompressedSize32);
+	}
+	else
+	{
+		FOodleDataCompression::ECompressor OodleCompressor;
+		FOodleDataCompression::ECompressionLevel OodleLevel;
+		GetShaderCompressionOodleSettings(OodleCompressor, OodleLevel);
+
+		// don't pass a nullptr to Oodle if we're only requesting an estimate
+		OutCompressedSize = OutCompressedShader ? FOodleDataCompression::Compress(OutCompressedShader, OutCompressedSize, UncompressedShaderCode, UncompressedSize, OodleCompressor, OodleLevel) : 0;
+		bCompressed = OutCompressedSize != 0;
+
+		// Oodle needs to return an estimate
+		if (!bCompressed)
 		{
-			UE_LOG(LogShaderLibrary, Fatal, TEXT("ShaderCodeArchive::DecompressShader(): Could not decompress shader (GetShaderCompressionFormat=%s)"), *GetShaderCompressionFormat().ToString());
-		}
+			// for Oodle, there is a separate estimation functon
+			OutCompressedSize = FOodleDataCompression::CompressedBufferSizeNeeded(UncompressedSize);
+		}	
 	}
+
+	return bCompressed;
 }
 
 void FSerializedShaderArchive::DecompressShader(int32 Index, const TArray<TArray<uint8>>& ShaderCode, TArray<uint8>& OutDecompressedShader) const
@@ -927,7 +982,7 @@ bool FShaderCodeArchive::PreloadShader(int32 ShaderIndex, FGraphEventArray& OutC
 		const FShaderCodeEntry& ShaderEntry = SerializedShaders.ShaderEntries[ShaderIndex];
 		ShaderPreloadEntry.Code = FMemory::Malloc(ShaderEntry.Size);
 		ShaderPreloadEntry.FramePreloadStarted = GFrameNumber;
-		DebugVisualizer.MarkPreloadedForVisualization(ShaderIndex);
+		DebugVisualizer.MarkExplicitlyPreloadedForVisualization(ShaderIndex);
 
 		const EAsyncIOPriorityAndFlags IOPriority = (EAsyncIOPriorityAndFlags)GShaderCodeLibraryAsyncLoadingPriority;
 
@@ -980,7 +1035,7 @@ bool FShaderCodeArchive::PreloadShaderMap(int32 ShaderMapIndex, FGraphEventArray
 			check(!ShaderPreloadEntry.PreloadEvent);
 			ShaderPreloadEntry.Code = FMemory::Malloc(ShaderEntry.Size);
 			ShaderPreloadEntry.FramePreloadStarted = FrameNumber;
-			DebugVisualizer.MarkPreloadedForVisualization(ShaderIndex);
+			DebugVisualizer.MarkExplicitlyPreloadedForVisualization(ShaderIndex);
 			PreloadMemory += ShaderEntry.Size;
 
 			FGraphEventArray ReadCompletionEvents;
@@ -1179,32 +1234,402 @@ FIoChunkId FIoStoreShaderCodeArchive::GetShaderCodeChunkId(const FSHAHash& Shade
 	return ChunkId;
 }
 
-void FIoStoreShaderCodeArchive::SaveIoStoreShaderCodeArchive(const FSerializedShaderArchive& SerializedShaders, FArchive& OutLibraryAr)
+void FIoStoreShaderCodeArchive::CreateIoStoreShaderCodeArchiveHeader(const FName& Format, const FSerializedShaderArchive& SerializedShaders, FIoStoreShaderCodeArchiveHeader& OutHeader)
 {
-	uint32 Version = CurrentVersion;
-	OutLibraryAr << Version;
-	OutLibraryAr << const_cast<FSerializedShaderArchive&>(SerializedShaders).ShaderMapHashes;
-	OutLibraryAr << const_cast<FSerializedShaderArchive&>(SerializedShaders).ShaderHashes;
-	TArray<FIoStoreShaderMapEntry> ShaderMapEntries;
-	ShaderMapEntries.Reserve(SerializedShaders.ShaderMapEntries.Num());
+	OutHeader.ShaderMapHashes = SerializedShaders.ShaderMapHashes;
+	OutHeader.ShaderHashes = SerializedShaders.ShaderHashes;
+	// shader group hashes will be populated later
+
+	OutHeader.ShaderMapEntries.Empty(SerializedShaders.ShaderMapEntries.Num());
 	for (const FShaderMapEntry& ShaderMapEntry : SerializedShaders.ShaderMapEntries)
 	{
-		FIoStoreShaderMapEntry& IoStoreShaderMapEntry = ShaderMapEntries.AddDefaulted_GetRef();
+		FIoStoreShaderMapEntry& IoStoreShaderMapEntry = OutHeader.ShaderMapEntries.AddDefaulted_GetRef();
 		IoStoreShaderMapEntry.ShaderIndicesOffset = ShaderMapEntry.ShaderIndicesOffset;
 		IoStoreShaderMapEntry.NumShaders = ShaderMapEntry.NumShaders;
 	}
-	OutLibraryAr << ShaderMapEntries;
-	TArray<FIoStoreShaderCodeEntry> ShaderEntries;
-	ShaderEntries.Reserve(SerializedShaders.ShaderEntries.Num());
+
+	// indices should be copied before grouping as the groups will append to the array
+	OutHeader.ShaderIndices = SerializedShaders.ShaderIndices;
+
+	// shader entries are copied, the remainder of the field will be assigned when splitting into groups
+	OutHeader.ShaderEntries.Empty(SerializedShaders.ShaderEntries.Num());
 	for (const FShaderCodeEntry& ShaderEntry : SerializedShaders.ShaderEntries)
 	{
-		FIoStoreShaderCodeEntry& IoStoreShaderEntry = ShaderEntries.AddDefaulted_GetRef();
-		IoStoreShaderEntry.UncompressedSize = ShaderEntry.UncompressedSize;
-		IoStoreShaderEntry.CompressedSize = ShaderEntry.Size;
+		FIoStoreShaderCodeEntry& IoStoreShaderEntry = OutHeader.ShaderEntries.AddDefaulted_GetRef();
 		IoStoreShaderEntry.Frequency = ShaderEntry.Frequency;
 	}
-	OutLibraryAr << ShaderEntries;
-	OutLibraryAr << const_cast<FSerializedShaderArchive&>(SerializedShaders).ShaderIndices;
+
+	// Higher level description of the group splitting algo that follows:
+	// We compress together shaders that are loaded together (all other strategies, like grouping similar shaders, were found to compress better but not reduce RAM usage).
+	// For that, we find for each shader which shadermaps are referencing it (often times it will be just one, but for some simple and shared shaders it can be thousands).
+	// We group the shaders by those sets of shadermaps - all shaders referenced by the same shadermap(s) are a candidate for being a single group. Then we potentially split this candidate group
+	// into raytracing and non-raytracing groups (so we can avoid preloading RTX shaders run-time if RTX is disabled), and then each of those is potentially split further by size
+	// (to avoid too large groups that will take too much time to decompress - this is regulated by r.ShaderCodeLibrary.MaxShaderGroupSize). The results of that process (note, it can still be
+	// a single group) is added to the header.
+	// Each group's indices, like in case of shadermaps, are stored in ShaderIndices array. Before we append a new group's indices however, we look if we can find an existing range that we can reuse.
+
+	TMap<uint32, TArray<int32>> ShaderToShadermaps;
+	{
+		TArray<TArray<int32>> ShaderToShadermapsArray;
+		ShaderToShadermapsArray.AddDefaulted(OutHeader.ShaderEntries.Num());
+
+		{
+			FCriticalSection ShaderLocks[1024];
+			// for each shader, find all the shadermaps it belongs to
+			ParallelFor(OutHeader.ShaderMapEntries.Num(),
+				[&ShaderToShadermapsArray, &OutHeader, &ShaderLocks](int32 ShaderMapIndex)
+				{
+					FIoStoreShaderMapEntry& ShaderMapEntry = OutHeader.ShaderMapEntries[ShaderMapIndex];
+					for (int32 ShaderIdxIdx = ShaderMapEntry.ShaderIndicesOffset, StopBeforeIdxIdx = ShaderMapEntry.ShaderIndicesOffset + ShaderMapEntry.NumShaders; ShaderIdxIdx < StopBeforeIdxIdx; ++ShaderIdxIdx)
+					{
+						int32 ShaderIndex = OutHeader.ShaderIndices[ShaderIdxIdx];
+						// add this shadermap as a dependency.
+						int32 ShaderLockNumber = ShaderIndex % UE_ARRAY_COUNT(ShaderLocks);
+						FScopeLock Locker(&ShaderLocks[ShaderLockNumber]);
+						ShaderToShadermapsArray[ShaderIndex].Add(ShaderMapIndex);
+					}
+				}
+			);
+		}
+
+		// sort shadermaps entries in shaders
+		{
+			const int32 kShaderSortedPerThread = 1024;
+			const int32 NumThreads = (ShaderToShadermapsArray.Num() / kShaderSortedPerThread) + 1;
+			ParallelFor(NumThreads,
+				[&ShaderToShadermapsArray, kShaderSortedPerThread](int ThreadIndex)
+				{
+					int32 StartingShader = ThreadIndex * kShaderSortedPerThread;
+					int32 EndingShader = FMath::Min(StartingShader + kShaderSortedPerThread, ShaderToShadermapsArray.Num());
+					for (int32 Idx = StartingShader; Idx < EndingShader; ++Idx)
+					{
+						ShaderToShadermapsArray[Idx].Sort();
+					}
+				},
+				EParallelForFlags::Unbalanced
+			);
+		}
+
+		// now add all this to map so we can sort it
+		for (int32 Idx = 0, Num = ShaderToShadermapsArray.Num(); Idx < Num; ++Idx)
+		{
+			// check that no shader is unreferenced
+			checkf(!ShaderToShadermapsArray[Idx].IsEmpty(), TEXT("Error converting to IoStore archive: shader (index=%d) is not referenced by any of the shadermaps!"), Idx);
+			ShaderToShadermaps.Add(Idx, ShaderToShadermapsArray[Idx]);
+		}
+	}
+
+	// sort the mapping so the first are shaders that are referenced by a smaller number of shadermaps.
+	ShaderToShadermaps.ValueSort(
+		[](const TArray<int32>& A, const TArray<int32>& B)
+		{
+			// if the number of shadermaps is the same, we need to sort "alphabetically"
+			if (A.Num() == B.Num())
+			{
+				for (int32 Idx = 0, Num = A.Num(); Idx < Num; ++Idx)
+				{
+					if (A[Idx] < B[Idx])
+					{
+						return true;
+					}
+				}
+
+				return false;
+			}
+
+			return A.Num() < B.Num();
+		}
+	);
+
+	// get the effective maximum uncompressed group size
+	FShaderPlatformCachedIniValue<int32> ConfiguredMaxShaderGroupSize(TEXT("r.ShaderCodeLibrary.MaxShaderGroupSize"));
+	EShaderPlatform ShaderPlatform = ShaderFormatNameToShaderPlatform(Format);
+	uint32 MaxUncompressedShaderGroupSize = FMath::Max(ConfiguredMaxShaderGroupSize.Get(ShaderPlatform), 1);	// cannot be lower than 1
+
+	// for statistics
+	uint64 Stats_TotalUncompressedMemory = 0;
+	uint32 Stats_MinGroupSize = MAX_uint32;
+	uint32 Stats_MaxGroupSize = 0;
+
+	// We want to avoid adding group indices to ShaderIndices, however looking them up sequentially is to slow. Store them here for a future ParallelFor lookup.
+	TArray<TArray<uint32>> StoredGroupShaderIndices;
+
+	/** Third and last stage of processing the shader group. We actually add it here, and do the book-keeping. */
+	auto ProcessShaderGroup_AddNewGroup = [&OutHeader, &SerializedShaders, &StoredGroupShaderIndices, &Stats_TotalUncompressedMemory, &Stats_MinGroupSize, &Stats_MaxGroupSize](TArray<uint32>& ShaderIndicesInGroup)
+	{
+		// first, sort the shaders by uncompressed size, as this was found to compress better
+		ShaderIndicesInGroup.Sort(
+			[&SerializedShaders](const int32 ShaderIndexA, const int32 ShaderIndexB)
+			{
+				return SerializedShaders.ShaderEntries[ShaderIndexA].UncompressedSize < SerializedShaders.ShaderEntries[ShaderIndexB].UncompressedSize;
+			}
+		);
+
+		// add a new group entry
+		const int32 CurrentGroupIdx = OutHeader.ShaderGroupEntries.Num();
+		FIoStoreShaderGroupEntry& GroupEntry = OutHeader.ShaderGroupEntries.AddDefaulted_GetRef();
+		StoredGroupShaderIndices.Add(ShaderIndicesInGroup);
+		GroupEntry.NumShaders = ShaderIndicesInGroup.Num();
+		// ShaderIndicesOffset will be filled later, once we know all the groups (see comment about StoredGroupShaderIndices above).
+
+		// update shader entries both with the group number and their uncompressed offset in the group
+		FSHA1 GroupHasher;
+		uint32 CurrentGroupSize = 0;
+		for (int32 ShaderIdxIdx = 0, NumIdxIdx = ShaderIndicesInGroup.Num(); ShaderIdxIdx < NumIdxIdx; ++ShaderIdxIdx)
+		{
+			int32 ShaderIndex = ShaderIndicesInGroup[ShaderIdxIdx];
+			FIoStoreShaderCodeEntry& IoStoreShaderEntry = OutHeader.ShaderEntries[ShaderIndex];
+			IoStoreShaderEntry.ShaderGroupIndex = CurrentGroupIdx;
+			IoStoreShaderEntry.UncompressedOffsetInGroup = CurrentGroupSize;
+
+			// group hash is constructed from hashing the shaders in the group
+			GroupHasher.Update(SerializedShaders.ShaderHashes[ShaderIndex].Hash, sizeof(FSHAHash));
+
+			CurrentGroupSize += SerializedShaders.ShaderEntries[ShaderIndex].UncompressedSize;
+		}
+		GroupEntry.UncompressedSize = CurrentGroupSize;
+		OutHeader.ShaderGroupIoHashes.Add(FIoStoreShaderCodeArchive::GetShaderCodeChunkId(GroupHasher.Finalize()));
+
+		Stats_TotalUncompressedMemory += CurrentGroupSize;
+		Stats_MinGroupSize = FMath::Min(Stats_MinGroupSize, CurrentGroupSize);
+		Stats_MaxGroupSize = FMath::Max(Stats_MaxGroupSize, CurrentGroupSize);
+	};
+
+	/** Second stage of processing shader group. Here we potentially split the group into smaller ones (as equally as possible), striving to meet limit imposed by r.ShaderCodeLibrary.MaxShaderGroupSize */
+	auto ProcessShaderGroup_SplitBySize = [&OutHeader, &SerializedShaders, &ProcessShaderGroup_AddNewGroup, &MaxUncompressedShaderGroupSize](TArray<uint32>& CurrentShaderGroup)
+	{
+		// calculate current group size
+		uint32 GroupSize = 0;
+		for (uint32 ShaderIdx : CurrentShaderGroup)
+		{
+			GroupSize += SerializedShaders.ShaderEntries[ShaderIdx].UncompressedSize;
+		}
+		
+		if (LIKELY(GroupSize <= MaxUncompressedShaderGroupSize))
+		{
+			ProcessShaderGroup_AddNewGroup(CurrentShaderGroup);
+		}
+		else
+		{
+			// split the shaders evenly into N new groups (don't allow more new groups than there are shaders)
+			int32 NumNewGroups = FMath::Min(static_cast<int32>(GroupSize / MaxUncompressedShaderGroupSize + 1), CurrentShaderGroup.Num());
+			checkf(NumNewGroups > 1, TEXT("Off by one error in group count calculation?"));
+
+			TArray<TArray<uint32>> NewGroups;
+			TArray<uint32> NewGroupSizes;
+			NewGroups.AddDefaulted(NumNewGroups);
+			NewGroupSizes.AddZeroed(NumNewGroups);
+
+			// sort the shaders descending, as this is easier to split (greedy algorithm)
+			CurrentShaderGroup.Sort(
+				[&SerializedShaders](const int32 ShaderIndexA, const int32 ShaderIndexB)
+				{
+					return SerializedShaders.ShaderEntries[ShaderIndexA].UncompressedSize > SerializedShaders.ShaderEntries[ShaderIndexB].UncompressedSize;
+				}
+			);
+
+			for (int32 ShaderIdx : CurrentShaderGroup)
+			{
+				// add the shader to the group of smallest size
+				int32 SmallestNewGroupIdx = 0;
+				for (int32 IdxNewGroup = 1; IdxNewGroup < NumNewGroups; ++IdxNewGroup)
+				{
+					if (NewGroupSizes[IdxNewGroup] < NewGroupSizes[SmallestNewGroupIdx])
+					{
+						SmallestNewGroupIdx = IdxNewGroup;
+					}
+				}
+
+				NewGroups[SmallestNewGroupIdx].Add(ShaderIdx);
+				NewGroupSizes[SmallestNewGroupIdx] += SerializedShaders.ShaderEntries[ShaderIdx].UncompressedSize;
+			}
+			
+#if DO_CHECK // sanity checks
+			uint32 NewGroupTotalSize = 0;
+			for (uint32 NewGroupSize : NewGroupSizes)
+			{
+				NewGroupTotalSize += NewGroupSize;
+			}
+			checkf(NewGroupTotalSize == GroupSize, TEXT("Original shader group was of size %u bytes, which was larger than limit %u, but it was split into %d group of total size %u, which is not %u - sizes must agree"),
+				GroupSize,
+				MaxUncompressedShaderGroupSize,
+				NumNewGroups,
+				NewGroupTotalSize, GroupSize
+				);
+#endif
+
+			for (TArray<uint32>& NewGroup : NewGroups)
+			{
+				// note there can be empty groups (take a very edge case of MaxUncompressedShaderGroupSize = 2 bytes and a shader group of 1 shader)
+				if (!NewGroup.IsEmpty())
+				{
+					ProcessShaderGroup_AddNewGroup(NewGroup);
+				}
+			}
+		}
+	};
+
+	/** First stage of processing a streak of shaders all referenced by the same set of shadermaps. We begin with separating raytracing and non-raytracing shaders, so we can avoid preloading RTX in non-RT runs. */
+	auto ProcessShaderGroup_SplitRaytracing = [&OutHeader, &ProcessShaderGroup_SplitBySize](TArray<uint32>& CurrentShaderGroup)
+	{
+		// The streak changed. Create the group, but first, determine if the group needs to be split in two because of the raytracing shaders.
+		// We want to isolate them into separate groups so their preload can be skipped if raytracing is off.
+		TArray<uint32> RaytracingShaders;
+		TArray<uint32> NonraytracingShaders;
+		for (int32 ShaderIndex : CurrentShaderGroup)
+		{
+			if (LIKELY(!IsRayTracingShaderFrequency(static_cast<EShaderFrequency>(OutHeader.ShaderEntries[ShaderIndex].Frequency))))
+			{
+				NonraytracingShaders.Add(ShaderIndex);
+			}
+			else
+			{
+				RaytracingShaders.Add(ShaderIndex);
+			}
+		}
+		check(CurrentShaderGroup.Num() == NonraytracingShaders.Num() + RaytracingShaders.Num());
+
+		if (LIKELY(!NonraytracingShaders.IsEmpty()))
+		{
+			ProcessShaderGroup_SplitBySize(NonraytracingShaders);
+		}
+		if (UNLIKELY(!RaytracingShaders.IsEmpty()))
+		{
+			ProcessShaderGroup_SplitBySize(RaytracingShaders);
+		}
+	};
+
+	// now split this into streaks of shaders that are referenced by the same set of shadermaps and compress
+	TArray<uint32> CurrentShaderGroup;
+	TArray<int32> LastShadermapSetSeen;
+	for (TMap<uint32, TArray<int32>>::TConstIterator Iter(ShaderToShadermaps); Iter; ++Iter)
+	{
+		// if we have have just started the group, we don't check the last seen
+		if (UNLIKELY(CurrentShaderGroup.IsEmpty()))
+		{
+			CurrentShaderGroup.Add(Iter.Key());
+			LastShadermapSetSeen = Iter.Value();
+		}
+		else if (UNLIKELY(LastShadermapSetSeen != Iter.Value()))
+		{
+			ProcessShaderGroup_SplitRaytracing(CurrentShaderGroup);
+
+			// reset the collection, but don't forget to add to it the current element
+			CurrentShaderGroup.SetNum(1);
+			CurrentShaderGroup[0] = Iter.Key();
+			LastShadermapSetSeen = Iter.Value();
+		}
+		else
+		{
+			// keep adding to the same group
+			CurrentShaderGroup.Add(Iter.Key());
+		}
+	}
+	// add the last group
+	if (!CurrentShaderGroup.IsEmpty())
+	{
+		ProcessShaderGroup_SplitRaytracing(CurrentShaderGroup);
+	}
+
+	/** Tries to find whether NewIndices exist as a subsequence in ExistingIndices.Returns - 1 if not found. */
+	auto FindSequenceInArray = [](const TArray<uint32>& ExistingIndices, const TArray<uint32>& NewIndices) -> int32
+	{
+		check(NewIndices.Num() > 0);
+
+		uint32 FirstNew = NewIndices[0];
+		int32 NumNew = NewIndices.Num();
+		for (int32 IdxExisting = 0, NumExisting = ExistingIndices.Num(); IdxExisting < NumExisting - NumNew + 1; ++IdxExisting)
+		{
+			if (LIKELY(ExistingIndices[IdxExisting] != FirstNew))
+			{
+				continue;
+			}
+
+			// check the rest
+			bool bFoundSequence = true;
+			for (int32 IdxNew = 1; IdxNew < NumNew; ++IdxNew)
+			{
+				checkSlow(IdxExisting + IdxNew < NumExisting);
+				if (LIKELY(ExistingIndices[IdxExisting + IdxNew] != NewIndices[IdxNew]))
+				{
+					bFoundSequence = false;
+					break;
+				}
+			}
+
+			if (UNLIKELY(bFoundSequence))
+			{
+				return IdxExisting;
+			}
+		}
+
+		return INDEX_NONE;
+	};
+
+	// now, try to see if we can look up group's indices in the existing ShaderIndicesArray to avoid storing them there.
+	checkf(StoredGroupShaderIndices.Num() == OutHeader.ShaderGroupEntries.Num(), TEXT("We should have stored shader indices for all groups."));
+	ParallelFor(StoredGroupShaderIndices.Num(),
+		[&OutHeader, &SerializedShaders, &StoredGroupShaderIndices, &FindSequenceInArray](int ShaderGroupIndex)
+		{
+			const TArray<uint32>& ShaderIndicesInGroup = StoredGroupShaderIndices[ShaderGroupIndex];
+			FIoStoreShaderGroupEntry& GroupEntry = OutHeader.ShaderGroupEntries[ShaderGroupIndex];
+			// See if we can find indices in that order somewhere in the ShaderIndices array already, to avoid adding new indices.
+			// We are looking in the read-only original array, because there's no sense to look in OutHeader.ShaderIndices - groups don't overlap,
+			// so we know that newly added (by some previous group) indices aren't useful for us.
+			int32 ExistingOffset = FindSequenceInArray(SerializedShaders.ShaderIndices, ShaderIndicesInGroup);
+			if (ExistingOffset != INDEX_NONE)
+			{
+				GroupEntry.ShaderIndicesOffset = static_cast<uint32>(ExistingOffset);
+			}
+			else
+			{
+				GroupEntry.ShaderIndicesOffset = MAX_uint32;
+			}
+		}
+	);
+
+	// Now append all the groups that weren't found to the end of ShaderIndices, slow (we could have done that in above ParallelFor, with a lock), but good for determinism
+	int32 Stats_GroupsThatAppendedToShaderIndices = 0;
+	for (int32 ShaderGroupIndex = 0, NumGroups = OutHeader.ShaderGroupEntries.Num(); ShaderGroupIndex < NumGroups; ++ShaderGroupIndex)
+	{
+		FIoStoreShaderGroupEntry& GroupEntry = OutHeader.ShaderGroupEntries[ShaderGroupIndex];
+
+		if (GroupEntry.ShaderIndicesOffset == MAX_uint32)
+		{
+			const TArray<uint32>& ShaderIndicesInGroup = StoredGroupShaderIndices[ShaderGroupIndex];
+			GroupEntry.ShaderIndicesOffset = OutHeader.ShaderIndices.Num();
+			OutHeader.ShaderIndices.Append(ShaderIndicesInGroup);
+			++Stats_GroupsThatAppendedToShaderIndices;
+		}
+	}
+
+	checkf(OutHeader.ShaderEntries.Num() == SerializedShaders.ShaderEntries.Num(), TEXT("Error creating IoStoreShaderArchive header - shader entries differ (%d in IoStore, %d original). Bug in grouping logic?"),
+		OutHeader.ShaderEntries.Num(), SerializedShaders.ShaderEntries.Num());
+	checkf(OutHeader.ShaderGroupIoHashes.Num() == OutHeader.ShaderGroupEntries.Num(), TEXT("Error creating IoStoreShaderArchive header - mismatch between shader group hashes and descriptors (%d descriptors, %d hashes). Bug in grouping logic?"),
+		OutHeader.ShaderGroupEntries.Num(), OutHeader.ShaderGroupIoHashes.Num());
+	checkf(OutHeader.ShaderGroupEntries.Num() != 0, TEXT("At least one group must have been created"));
+
+	UE_LOG(LogShaderLibrary, Display, TEXT("Created IoStoreShaderArchive header: shaders grouped in %d groups (%d of them didn't need new indices), average uncompressed size %llu bytes, min %u bytes, max %u bytes (r.ShaderCodeLibrary.MaxShaderGroupSize=%u)"),
+		OutHeader.ShaderGroupEntries.Num(), OutHeader.ShaderGroupEntries.Num() - Stats_GroupsThatAppendedToShaderIndices, Stats_TotalUncompressedMemory / static_cast<uint64>(OutHeader.ShaderGroupEntries.Num()), Stats_MinGroupSize, Stats_MaxGroupSize, MaxUncompressedShaderGroupSize);
+}
+
+FArchive& operator <<(FArchive& Ar, FIoStoreShaderCodeArchiveHeader& Ref)
+{
+	Ar << Ref.ShaderMapHashes;
+	Ar << Ref.ShaderHashes;
+	Ar << Ref.ShaderGroupIoHashes;
+	Ar << Ref.ShaderMapEntries;
+	Ar << Ref.ShaderEntries;
+	Ar << Ref.ShaderGroupEntries;
+	Ar << Ref.ShaderIndices;
+	return Ar;
+}
+
+void FIoStoreShaderCodeArchive::SaveIoStoreShaderCodeArchive(const FIoStoreShaderCodeArchiveHeader& Header, FArchive& OutLibraryAr)
+{
+	uint32 Version = CurrentVersion;
+	OutLibraryAr << Version;
+	OutLibraryAr << const_cast<FIoStoreShaderCodeArchiveHeader &>(Header);
 }
 
 FIoStoreShaderCodeArchive* FIoStoreShaderCodeArchive::Create(EShaderPlatform InPlatform, const FString& InLibraryName, FIoDispatcher& InIoDispatcher)
@@ -1226,34 +1651,29 @@ FIoStoreShaderCodeArchive* FIoStoreShaderCodeArchive::Create(EShaderPlatform InP
 		if (Version == CurrentVersion)
 		{
 			FIoStoreShaderCodeArchive* Library = new FIoStoreShaderCodeArchive(InPlatform, InLibraryName, InIoDispatcher);
-			Ar << Library->ShaderMapHashes;
-			Ar << Library->ShaderHashes;
-			Ar << Library->ShaderMapEntries;
-			Ar << Library->ShaderEntries;
-			Ar << Library->ShaderIndices;
-			Library->ShaderPreloads.SetNum(Library->GetNumShaders());
+			Ar << Library->Header;
 			{
-				const uint32 HashSize = FMath::Min<uint32>(0x10000, 1u << FMath::CeilLogTwo(Library->ShaderMapHashes.Num()));
-				Library->ShaderMapHashTable.Clear(HashSize, Library->ShaderMapHashes.Num());
-				for (int32 Index = 0; Index < Library->ShaderMapHashes.Num(); ++Index)
+				const uint32 HashSize = FMath::Min<uint32>(0x10000, 1u << FMath::CeilLogTwo(Library->Header.ShaderMapHashes.Num()));
+				Library->ShaderMapHashTable.Clear(HashSize, Library->Header.ShaderMapHashes.Num());
+				for (int32 Index = 0, Num = Library->Header.ShaderMapHashes.Num(); Index < Num; ++Index)
 				{
-					const uint32 Key = GetTypeHash(Library->ShaderMapHashes[Index]);
+					const uint32 Key = GetTypeHash(Library->Header.ShaderMapHashes[Index]);
 					Library->ShaderMapHashTable.Add(Key, Index);
 				}
 			}
 			{
-				const uint32 HashSize = FMath::Min<uint32>(0x10000, 1u << FMath::CeilLogTwo(Library->ShaderHashes.Num()));
-				Library->ShaderHashTable.Clear(HashSize, Library->ShaderHashes.Num());
-				for (int32 Index = 0; Index < Library->ShaderHashes.Num(); ++Index)
+				const uint32 HashSize = FMath::Min<uint32>(0x10000, 1u << FMath::CeilLogTwo(Library->Header.ShaderHashes.Num()));
+				Library->ShaderHashTable.Clear(HashSize, Library->Header.ShaderHashes.Num());
+				for (int32 Index = 0, Num = Library->Header.ShaderHashes.Num(); Index < Num; ++Index)
 				{
-					const uint32 Key = GetTypeHash(Library->ShaderHashes[Index]);
+					const uint32 Key = GetTypeHash(Library->Header.ShaderHashes[Index]);
 					Library->ShaderHashTable.Add(Key, Index);
 				}
 			}
 
-			Library->DebugVisualizer.Initialize(Library->ShaderEntries.Num());
+			Library->DebugVisualizer.Initialize(Library->Header.ShaderEntries.Num());
 
-			UE_LOG(LogShaderLibrary, Display, TEXT("Using IoDispatcher for shader code library %s. Total %d unique shaders."), *InLibraryName, Library->ShaderEntries.Num());
+			UE_LOG(LogShaderLibrary, Display, TEXT("Using IoDispatcher for shader code library %s. Total %d unique shaders."), *InLibraryName, Library->Header.ShaderEntries.Num());
 			INC_DWORD_STAT_BY(STAT_Shaders_ShaderResourceMemory, Library->GetSizeBytes());
 			return Library;
 		}
@@ -1275,157 +1695,195 @@ FIoStoreShaderCodeArchive::~FIoStoreShaderCodeArchive()
 
 void FIoStoreShaderCodeArchive::Teardown()
 {
-	for (int32 ShaderIndex = 0; ShaderIndex < ShaderPreloads.Num(); ++ShaderIndex)
+	DebugVisualizer.SaveShaderUsageBitmap(GetName(), GetPlatform());
+
+	for (TMap<int32, FShaderGroupPreloadEntry*>::TIterator Iter(PreloadedShaderGroups); Iter; ++Iter)
 	{
-		FShaderPreloadEntry& ShaderPreloadEntry = ShaderPreloads[ShaderIndex];
-		check(!ShaderPreloadEntry.NumRefs);
-	}
+		FShaderGroupPreloadEntry* PreloadEntry = Iter.Value();
 
-	DebugVisualizer.SaveShaderUsageBitmap(GetName(), GetPlatform());
+		checkf(PreloadEntry->NumRefs == 0, TEXT("Group %d has still %d references on deletion"), Iter.Key(), PreloadEntry->NumRefs);
+
+		const FIoStoreShaderGroupEntry& GroupEntry = Header.ShaderGroupEntries[Iter.Key()];
+		DEC_DWORD_STAT_BY(STAT_Shaders_ShaderPreloadMemory, (GroupEntry.CompressedSize + sizeof(FShaderGroupPreloadEntry)));
+
+		delete PreloadEntry;
+	}
 }
 
-bool FIoStoreShaderCodeArchive::PreloadShader(int32 ShaderIndex, FGraphEventArray& OutCompletionEvents)
+void FIoStoreShaderCodeArchive::SetupPreloadEntryForLoading(int32 ShaderGroupIndex, FShaderGroupPreloadEntry& PreloadEntry)
 {
-	LLM_SCOPE(ELLMTag::Shaders);
+	PreloadEntry.FramePreloadStarted = GFrameNumber;
+	check(!PreloadEntry.PreloadEvent);
+	PreloadEntry.PreloadEvent = FGraphEvent::CreateGraphEvent();
 
-	FWriteScopeLock Lock(ShaderPreloadLock);
+#if UE_SCA_VISUALIZE_SHADER_USAGE
+	const FIoStoreShaderGroupEntry& GroupEntry = Header.ShaderGroupEntries[ShaderGroupIndex];
+	for (uint32 ShaderIdxIdx = GroupEntry.ShaderIndicesOffset, StopBeforeIdxIdx = GroupEntry.ShaderIndicesOffset + GroupEntry.NumShaders; ShaderIdxIdx < StopBeforeIdxIdx; ++ShaderIdxIdx)
+	{
+		DebugVisualizer.MarkPreloadedForVisualization(Header.ShaderIndices[ShaderIdxIdx]);
+	}
+#endif // UE_SCA_VISUALIZE_SHADER_USAGE
+}
 
-	FShaderPreloadEntry& ShaderPreloadEntry = ShaderPreloads[ShaderIndex];
-	const uint32 ShaderNumRefs = ShaderPreloadEntry.NumRefs++;
-	if (ShaderNumRefs == 0u)
+bool FIoStoreShaderCodeArchive::PreloadShaderGroup(int32 ShaderGroupIndex, FGraphEventArray& OutCompletionEvents, FCoreDelegates::FAttachShaderReadRequestFunc* AttachShaderReadRequestFuncPtr)
+{
+	// should be called within LLMTag::Shaders scope
+	FWriteScopeLock Lock(PreloadedShaderGroupsLock);
+	FShaderGroupPreloadEntry& PreloadEntry = *FindOrAddPreloadEntry(ShaderGroupIndex);
+	checkf(!PreloadEntry.bNeverToBePreloaded, TEXT("We are preloading a shader group (index=%d) that shouldn't be preloaded in this run (e.g. raytracing shaders on D3D11)."), ShaderGroupIndex);
+
+	const uint32 NumRefs = PreloadEntry.NumRefs++;
+	if (NumRefs == 0u)
 	{
-		check(!ShaderPreloadEntry.PreloadEvent);
+		SetupPreloadEntryForLoading(ShaderGroupIndex, PreloadEntry);
 
-		const FIoStoreShaderCodeEntry& ShaderEntry = ShaderEntries[ShaderIndex];
-		ShaderPreloadEntry.FramePreloadStarted = GFrameNumber;
-		DebugVisualizer.MarkPreloadedForVisualization(ShaderIndex);
+		// only global shaders are going to hit this path, all other shaders will be preloaded with the package
+		if (UNLIKELY(AttachShaderReadRequestFuncPtr == nullptr))
+		{
+			FIoBatch IoBatch = IoDispatcher.NewBatch();
+			PreloadEntry.IoRequest = IoBatch.Read(Header.ShaderGroupIoHashes[ShaderGroupIndex], FIoReadOptions(), IoDispatcherPriority_Medium);
+			IoBatch.IssueAndDispatchSubsequents(PreloadEntry.PreloadEvent);
+		}
+		else
+		{
+			PreloadEntry.IoRequest = (*AttachShaderReadRequestFuncPtr)(Header.ShaderGroupIoHashes[ShaderGroupIndex], PreloadEntry.PreloadEvent);
+		}
 
-		ShaderPreloadEntry.PreloadEvent = FGraphEvent::CreateGraphEvent();
-		FIoBatch IoBatch = IoDispatcher.NewBatch();
-		ShaderPreloadEntry.IoRequest = IoBatch.Read(GetShaderCodeChunkId(ShaderHashes[ShaderIndex]), FIoReadOptions(), IoDispatcherPriority_Medium);
-		IoBatch.IssueAndDispatchSubsequents(ShaderPreloadEntry.PreloadEvent);
-		INC_DWORD_STAT_BY(STAT_Shaders_ShaderPreloadMemory, ShaderEntry.CompressedSize);
+		INC_DWORD_STAT_BY(STAT_Shaders_ShaderPreloadMemory, (Header.ShaderGroupEntries[ShaderGroupIndex].CompressedSize + sizeof(FShaderGroupPreloadEntry)));
 	}
 
-	if (ShaderPreloadEntry.PreloadEvent && !ShaderPreloadEntry.PreloadEvent->IsComplete())
+	if (AttachShaderReadRequestFuncPtr == nullptr && PreloadEntry.PreloadEvent && !PreloadEntry.PreloadEvent->IsComplete())
 	{
-		OutCompletionEvents.Add(ShaderPreloadEntry.PreloadEvent);
+		OutCompletionEvents.Add(PreloadEntry.PreloadEvent);
 	}
 	return true;
 }
 
-bool FIoStoreShaderCodeArchive::PreloadShaderMap(int32 ShaderMapIndex, FGraphEventArray& OutCompletionEvents)
+void FIoStoreShaderCodeArchive::MarkPreloadEntrySkipped(int32 ShaderGroupIndex)
+{
+	// should be called within LLMTag::Shaders scope
+	FWriteScopeLock Lock(PreloadedShaderGroupsLock);
+	FShaderGroupPreloadEntry& PreloadEntry = *FindOrAddPreloadEntry(ShaderGroupIndex);
+	const uint32 NumRefs = PreloadEntry.NumRefs++;
+	if (NumRefs == 0u)
+	{
+		PreloadEntry.bNeverToBePreloaded = 1;
+		INC_DWORD_STAT_BY(STAT_Shaders_ShaderPreloadMemory, sizeof(FShaderGroupPreloadEntry));
+	}
+}
+
+bool FIoStoreShaderCodeArchive::PreloadShader(int32 ShaderIndex, FGraphEventArray& OutCompletionEvents)
 {
+	// [RCL] 2022-01-12 because PreloadShader API is unbalanced and there's code (eg. FShaderPipelineCache::PreparePipelineBatch) that calls PreloadShader but never ReleasePreloadedShader (which isn't
+	// even exposed to client code), this API is disabled to prevent leaking preload entries. Something to address separately - UE-139060.
+#if 0
 	LLM_SCOPE(ELLMTag::Shaders);
+	DebugVisualizer.MarkExplicitlyPreloadedForVisualization(ShaderIndex);
+	return PreloadShaderGroup(GetGroupIndexForShader(ShaderIndex), OutCompletionEvents);
+#else
+	return true;
+#endif
+}
 
-	const FIoStoreShaderMapEntry& ShaderMapEntry = ShaderMapEntries[ShaderMapIndex];
-	const uint32 FrameNumber = GFrameNumber;
-	uint32 PreloadMemory = 0u;
-	uint32 PreloadCount = 0u;
+bool FIoStoreShaderCodeArchive::GroupOnlyContainsRaytracingShaders(int32 ShaderGroupIndex)
+{
+	const FIoStoreShaderGroupEntry& GroupEntry = Header.ShaderGroupEntries[ShaderGroupIndex];
+	for (uint32 ShaderIdxIdx = GroupEntry.ShaderIndicesOffset, StopIdxIdx = GroupEntry.ShaderIndicesOffset + GroupEntry.NumShaders; ShaderIdxIdx < StopIdxIdx; ++ShaderIdxIdx)
+	{
+		int32 ShaderIndex = Header.ShaderIndices[ShaderIdxIdx];
+		if (!IsRayTracingShaderFrequency(static_cast<EShaderFrequency>(Header.ShaderEntries[ShaderIndex].Frequency)))
+		{
+			return false;
+		}
+	}
 
-	FWriteScopeLock Lock(ShaderPreloadLock);
+	return true;
+}
+
+bool FIoStoreShaderCodeArchive::PreloadShaderMap(int32 ShaderMapIndex, FGraphEventArray& OutCompletionEvents)
+{
+	LLM_SCOPE(ELLMTag::Shaders);
+	const FIoStoreShaderMapEntry& ShaderMapEntry = Header.ShaderMapEntries[ShaderMapIndex];
 
 	for (uint32 i = 0u; i < ShaderMapEntry.NumShaders; ++i)
 	{
-		const int32 ShaderIndex = ShaderIndices[ShaderMapEntry.ShaderIndicesOffset + i];
-		FShaderPreloadEntry& ShaderPreloadEntry = ShaderPreloads[ShaderIndex];
-		const FIoStoreShaderCodeEntry& ShaderEntry = ShaderEntries[ShaderIndex];
+		const int32 ShaderIndex = Header.ShaderIndices[ShaderMapEntry.ShaderIndicesOffset + i];
+		const int32 ShaderGroupIndex = GetGroupIndexForShader(ShaderIndex);
 
 #if RHI_RAYTRACING
-		if (!IsRayTracingEnabled() && IsRayTracingShaderFrequency(static_cast<EShaderFrequency>(ShaderEntry.Frequency)))
+		if (!IsRayTracingEnabled() && GroupOnlyContainsRaytracingShaders(ShaderGroupIndex))
 		{
-			ShaderPreloadEntry.bNeverToBePreloaded = 1;
+			MarkPreloadEntrySkipped(ShaderGroupIndex);
 			continue;
 		}
 #endif
 
-		const uint32 ShaderNumRefs = ShaderPreloadEntry.NumRefs++;
-		if (ShaderNumRefs == 0u)
-		{
-			check(!ShaderPreloadEntry.PreloadEvent);
-			ShaderPreloadEntry.FramePreloadStarted = FrameNumber;
-			DebugVisualizer.MarkPreloadedForVisualization(ShaderIndex);
-			PreloadMemory += ShaderEntry.CompressedSize;
-			++PreloadCount;
-			ShaderPreloadEntry.PreloadEvent = FGraphEvent::CreateGraphEvent();
-			FIoBatch IoBatch = IoDispatcher.NewBatch();
-			ShaderPreloadEntry.IoRequest = IoBatch.Read(GetShaderCodeChunkId(ShaderHashes[ShaderIndex]), FIoReadOptions(), IoDispatcherPriority_Medium);
-			IoBatch.IssueAndDispatchSubsequents(ShaderPreloadEntry.PreloadEvent);
-		}
-		if (ShaderPreloadEntry.PreloadEvent && !ShaderPreloadEntry.PreloadEvent->IsComplete())
-		{
-			OutCompletionEvents.Add(ShaderPreloadEntry.PreloadEvent);
-		}
+		// only shaders we actually want to preload should be marked as such, not just everything in the group
+		DebugVisualizer.MarkExplicitlyPreloadedForVisualization(ShaderIndex);
+		PreloadShaderGroup(ShaderGroupIndex, OutCompletionEvents);
 	}
 
-	INC_DWORD_STAT_BY(STAT_Shaders_ShaderPreloadMemory, PreloadMemory);
-
 	return true;
 }
 
 bool FIoStoreShaderCodeArchive::PreloadShaderMap(int32 ShaderMapIndex, FCoreDelegates::FAttachShaderReadRequestFunc AttachShaderReadRequestFunc)
 {
 	LLM_SCOPE(ELLMTag::Shaders);
+	const FIoStoreShaderMapEntry& ShaderMapEntry = Header.ShaderMapEntries[ShaderMapIndex];
 
-	const FIoStoreShaderMapEntry& ShaderMapEntry = ShaderMapEntries[ShaderMapIndex];
-	const uint32 FrameNumber = GFrameNumber;
-	uint32 PreloadMemory = 0u;
-
-	FWriteScopeLock Lock(ShaderPreloadLock);
-
+	FGraphEventArray Dummy;
 	for (uint32 i = 0u; i < ShaderMapEntry.NumShaders; ++i)
 	{
-		const int32 ShaderIndex = ShaderIndices[ShaderMapEntry.ShaderIndicesOffset + i];
-		FShaderPreloadEntry& ShaderPreloadEntry = ShaderPreloads[ShaderIndex];
-		const FIoStoreShaderCodeEntry& ShaderEntry = ShaderEntries[ShaderIndex];
+		const int32 ShaderIndex = Header.ShaderIndices[ShaderMapEntry.ShaderIndicesOffset + i];
+		const int32 ShaderGroupIndex = GetGroupIndexForShader(ShaderIndex);
 
 #if RHI_RAYTRACING
-		if (!IsRayTracingEnabled() && IsRayTracingShaderFrequency(static_cast<EShaderFrequency>(ShaderEntry.Frequency)))
+		if (!IsRayTracingEnabled() && GroupOnlyContainsRaytracingShaders(ShaderGroupIndex))
 		{
-			ShaderPreloadEntry.bNeverToBePreloaded = 1;
+			MarkPreloadEntrySkipped(ShaderGroupIndex);
 			continue;
 		}
 #endif
 
-		const uint32 ShaderNumRefs = ShaderPreloadEntry.NumRefs++;
-		if (ShaderNumRefs == 0u)
-		{
-			check(!ShaderPreloadEntry.PreloadEvent);
-			ShaderPreloadEntry.FramePreloadStarted = FrameNumber;
-			DebugVisualizer.MarkPreloadedForVisualization(ShaderIndex);
-			PreloadMemory += ShaderEntry.CompressedSize;
-			ShaderPreloadEntry.PreloadEvent = FGraphEvent::CreateGraphEvent();
-			ShaderPreloadEntry.IoRequest = AttachShaderReadRequestFunc(GetShaderCodeChunkId(ShaderHashes[ShaderIndex]), ShaderPreloadEntry.PreloadEvent);
-		}
+		// only shaders we actually want to preload should be marked as such, not just everything in the group
+		DebugVisualizer.MarkExplicitlyPreloadedForVisualization(ShaderIndex);
+		PreloadShaderGroup(ShaderGroupIndex, Dummy, &AttachShaderReadRequestFunc);
 	}
-	INC_DWORD_STAT_BY(STAT_Shaders_ShaderPreloadMemory, PreloadMemory);
+
 	return true;
 }
 
-bool FIoStoreShaderCodeArchive::ReleaseRef(int32 ShaderIndex)
+void FIoStoreShaderCodeArchive::ReleasePreloadEntry(int32 ShaderGroupIndex)
 {
-	FShaderPreloadEntry& ShaderPreloadEntry = ShaderPreloads[ShaderIndex];
-	if (!ShaderPreloadEntry.bNeverToBePreloaded)
+	FWriteScopeLock Lock(PreloadedShaderGroupsLock);
+	FShaderGroupPreloadEntry** ExistingEntry = PreloadedShaderGroups.Find(ShaderGroupIndex);
+	FShaderGroupPreloadEntry* PreloadEntry = ExistingEntry ? *ExistingEntry : nullptr;
+	checkf(PreloadEntry, TEXT("Preload entry for shader group %d should exist if we're asked to release it"), ShaderGroupIndex);
+
+	const uint32 ShaderNumRefs = PreloadEntry->NumRefs--;
+	check(ShaderNumRefs > 0u);
+	if (ShaderNumRefs == 1u)
 	{
-		const uint32 ShaderNumRefs = ShaderPreloadEntry.NumRefs--;
-		check(ShaderNumRefs > 0u);
-		if (ShaderNumRefs == 1u)
+		if (!PreloadEntry->bNeverToBePreloaded)
 		{
-			ShaderPreloadEntry.IoRequest = FIoRequest();
-			ShaderPreloadEntry.PreloadEvent.SafeRelease();
-			const FIoStoreShaderCodeEntry& ShaderEntry = ShaderEntries[ShaderIndex];
-			DEC_DWORD_STAT_BY(STAT_Shaders_ShaderPreloadMemory, ShaderEntry.CompressedSize);
-			return true;
+			PreloadEntry->IoRequest = FIoRequest();
+			PreloadEntry->PreloadEvent.SafeRelease();
+			const FIoStoreShaderGroupEntry& GroupEntry = Header.ShaderGroupEntries[ShaderGroupIndex];
+			DEC_DWORD_STAT_BY(STAT_Shaders_ShaderPreloadMemory, (GroupEntry.CompressedSize + sizeof(FShaderGroupPreloadEntry)));
 		}
-		return false;
+		else
+		{
+			DEC_DWORD_STAT_BY(STAT_Shaders_ShaderPreloadMemory, sizeof(FShaderGroupPreloadEntry));
+		}
+
+		delete PreloadEntry;
+		PreloadedShaderGroups.Remove(ShaderGroupIndex);
 	}
-	return true;
 }
 
 void FIoStoreShaderCodeArchive::ReleasePreloadedShader(int32 ShaderIndex)
 {
-	FWriteScopeLock Lock(ShaderPreloadLock);
-	ReleaseRef(ShaderIndex);
+	ReleasePreloadEntry(GetGroupIndexForShader(ShaderIndex));
 }
 
 int32 FIoStoreShaderCodeArchive::FindShaderMapIndex(const FSHAHash& Hash)
@@ -1433,7 +1891,7 @@ int32 FIoStoreShaderCodeArchive::FindShaderMapIndex(const FSHAHash& Hash)
 	const uint32 Key = GetTypeHash(Hash);
 	for (uint32 Index = ShaderMapHashTable.First(Key); ShaderMapHashTable.IsValid(Index); Index = ShaderMapHashTable.Next(Index))
 	{
-		if (ShaderMapHashes[Index] == Hash)
+		if (Header.ShaderMapHashes[Index] == Hash)
 		{
 			return Index;
 		}
@@ -1446,7 +1904,7 @@ int32 FIoStoreShaderCodeArchive::FindShaderIndex(const FSHAHash& Hash)
 	const uint32 Key = GetTypeHash(Hash);
 	for (uint32 Index = ShaderHashTable.First(Key); ShaderHashTable.IsValid(Index); Index = ShaderHashTable.Next(Index))
 	{
-		if (ShaderHashes[Index] == Hash)
+		if (Header.ShaderHashes[Index] == Hash)
 		{
 			return Index;
 		}
@@ -1454,54 +1912,97 @@ int32 FIoStoreShaderCodeArchive::FindShaderIndex(const FSHAHash& Hash)
 	return INDEX_NONE;
 }
 
-TRefCountPtr<FRHIShader> FIoStoreShaderCodeArchive::CreateShader(int32 Index)
+TRefCountPtr<FRHIShader> FIoStoreShaderCodeArchive::CreateShader(int32 ShaderIndex)
 {
 	LLM_SCOPE(ELLMTag::Shaders);
+#if STATS
+	double TimeFunctionEntered = FPlatformTime::Seconds();
+	ON_SCOPE_EXIT
+	{
+		if (IsInRenderingThread())
+		{
+			double ShaderCreationTime = FPlatformTime::Seconds() - TimeFunctionEntered;
+			INC_FLOAT_STAT_BY(STAT_Shaders_TotalRTShaderInitForRenderingTime, ShaderCreationTime);
+		}
+	};
+#endif
+
 	TRefCountPtr<FRHIShader> Shader;
 
-	const FIoStoreShaderCodeEntry& ShaderEntry = ShaderEntries[Index];
-	FShaderPreloadEntry& ShaderPreloadEntry = ShaderPreloads[Index];
+	const FIoStoreShaderCodeEntry& ShaderEntry = Header.ShaderEntries[ShaderIndex];
+	int32 GroupIndex = GetGroupIndexForShader(ShaderIndex);
+	FShaderGroupPreloadEntry* PreloadEntryPtr = nullptr;
+	bool bShaderWasNotPreloaded = false;
 
 	FGraphEventRef Event;
 	{
-		FWriteScopeLock Lock(ShaderPreloadLock);
-		uint32 ShaderNumRefs = ShaderPreloadEntry.NumRefs++;
-		if (ShaderNumRefs == 0u)
 		{
-			check(!ShaderPreloadEntry.PreloadEvent);
-			ShaderPreloadEntry.PreloadEvent = FGraphEvent::CreateGraphEvent();
-			FIoBatch IoBatch = IoDispatcher.NewBatch();
-			ShaderPreloadEntry.IoRequest = IoBatch.Read(GetShaderCodeChunkId(ShaderHashes[Index]), FIoReadOptions(), IoDispatcherPriority_Max);
-			IoBatch.IssueAndDispatchSubsequents(ShaderPreloadEntry.PreloadEvent);
+			// preloading a shader group 
+			FReadScopeLock Lock(PreloadedShaderGroupsLock);
+			FShaderGroupPreloadEntry** ExistingEntry = PreloadedShaderGroups.Find(GroupIndex);
+			PreloadEntryPtr = ExistingEntry ? *ExistingEntry : PreloadEntryPtr;
+			// if we have an entry, addref it so it doesn't go away
+			if (PreloadEntryPtr)
+			{
+				checkf(PreloadEntryPtr->NumRefs > 0, TEXT("Existing preload entry is not referenced! (NumRefs=%d)"), PreloadEntryPtr->NumRefs);
+				++PreloadEntryPtr->NumRefs;
+			}
+		}
+
+		// the shader we're attempting to create could have not been preloaded
+		if (PreloadEntryPtr == nullptr)
+		{
+			bShaderWasNotPreloaded = true;
+			FGraphEventArray Dummy;
+			PreloadShaderGroup(GroupIndex, Dummy);
+
+			// preloading a shader group should have created (and AddRef'd) the entry, so we can get and keep the pointer
+			FReadScopeLock Lock(PreloadedShaderGroupsLock);
+			PreloadEntryPtr = FindOrAddPreloadEntry(GroupIndex);
 		}
-		else if (!ShaderPreloadEntry.IoRequest.Status().IsCompleted())
+		check(PreloadEntryPtr);
+		
+		// raise the prio if still ongoing
+		if (!PreloadEntryPtr->IoRequest.Status().IsCompleted())
 		{
-			ShaderPreloadEntry.IoRequest.UpdatePriority(IoDispatcherPriority_Max);
+			PreloadEntryPtr->IoRequest.UpdatePriority(IoDispatcherPriority_Max);
 		}
-		Event = ShaderPreloadEntry.PreloadEvent;
+		Event = PreloadEntryPtr->PreloadEvent;
 	}
 
-	const bool bNeededToWait = !Event->IsComplete();
+	const bool bNeededToWait = Event.IsValid() && !Event->IsComplete();
 	if (bNeededToWait)
 	{
 		TRACE_CPUPROFILER_EVENT_SCOPE(BlockingShaderLoad);
-		UE_LOG(LogShaderLibrary, Warning, TEXT("Blocking wait for shader preload, NumRefs: %d, FramePreloadStarted: %d"), ShaderPreloadEntry.NumRefs, ShaderPreloadEntry.FramePreloadStarted);
+		UE_LOG(LogShaderLibrary, Warning, TEXT("Blocking wait for shader preload, NumRefs: %d, FramePreloadStarted: %d"), PreloadEntryPtr->NumRefs, PreloadEntryPtr->FramePreloadStarted);
 		FTaskGraphInterface::Get().WaitUntilTaskCompletes(Event);
 	}
 
-	const uint8* ShaderCode = ShaderPreloadEntry.IoRequest.GetResultOrDie().Data();
+	const uint8* ShaderCode = PreloadEntryPtr->IoRequest.GetResultOrDie().Data();
 
 	FMemStackBase& MemStack = FMemStack::Get();
 	FMemMark Mark(MemStack);
-	if (ShaderEntry.UncompressedSize != ShaderEntry.CompressedSize)
+	FIoStoreShaderGroupEntry& GroupEntry = Header.ShaderGroupEntries[GroupIndex];
+	if (GroupEntry.IsGroupCompressed())
 	{
-		uint8* UncompressedCode = reinterpret_cast<uint8*>(MemStack.Alloc(ShaderEntry.UncompressedSize, 16));
-		ShaderCodeArchive::DecompressShader(UncompressedCode, ShaderEntry.UncompressedSize, ShaderCode, ShaderEntry.CompressedSize);
-		ShaderCode = (uint8*)UncompressedCode;
+		uint8* UncompressedCode = reinterpret_cast<uint8*>(MemStack.Alloc(GroupEntry.UncompressedSize, 16));
+		ShaderCodeArchive::DecompressShader(UncompressedCode, GroupEntry.UncompressedSize, ShaderCode, GroupEntry.CompressedSize);
+		ShaderCode = reinterpret_cast<uint8*>(UncompressedCode) + ShaderEntry.UncompressedOffsetInGroup;
+
+#if UE_SCA_VISUALIZE_SHADER_USAGE
+		for (uint32 ShaderIdxIdx = GroupEntry.ShaderIndicesOffset, StopBeforeIdxIdx = GroupEntry.ShaderIndicesOffset + GroupEntry.NumShaders; ShaderIdxIdx < StopBeforeIdxIdx; ++ShaderIdxIdx)
+		{
+			DebugVisualizer.MarkDecompressedForVisualization(Header.ShaderIndices[ShaderIdxIdx]);
+		}
+#endif // UE_SCA_VISUALIZE_SHADER_USAGE
+	}
+	else
+	{
+		ShaderCode += ShaderEntry.UncompressedOffsetInGroup;
 	}
 
-	const auto ShaderCodeView = MakeArrayView(ShaderCode, ShaderEntry.UncompressedSize);
-	const FSHAHash& ShaderHash = ShaderHashes[Index];
+	const auto ShaderCodeView = MakeArrayView(ShaderCode, Header.GetShaderUncompressedSize(ShaderIndex));
+	const FSHAHash& ShaderHash = Header.ShaderHashes[ShaderIndex];
 	switch (ShaderEntry.Frequency)
 	{
 	case SF_Vertex: Shader = RHICreateVertexShader(ShaderCodeView, ShaderHash); break;
@@ -1520,21 +2021,23 @@ TRefCountPtr<FRHIShader> FIoStoreShaderCodeArchive::CreateShader(int32 Index)
 		break;
 	default: checkNoEntry(); break;
 	}
-	DebugVisualizer.MarkCreatedForVisualization(Index);
+	DebugVisualizer.MarkCreatedForVisualization(ShaderIndex);
 
+	if (bShaderWasNotPreloaded)
 	{
-		FWriteScopeLock Lock(ShaderPreloadLock);
-		const uint32 ShaderNumRefs = ShaderPreloadEntry.NumRefs--;
-		check(ShaderNumRefs > 0u);
-		if (ShaderNumRefs == 1u)
-		{
-			ShaderPreloadEntry.PreloadEvent.SafeRelease();
-			ShaderPreloadEntry.IoRequest = FIoRequest();
-		}
+		ReleasePreloadEntry(GroupIndex);
 	}
+	else
+	{
+		// if the entry existed before Create was called, it will be released by the higher level resource after the creation.
+		checkf(PreloadEntryPtr->NumRefs > 1, TEXT("We shouldn't be releasing the last preload entry here (NumRefs=%d)"), PreloadEntryPtr->NumRefs);
+		--PreloadEntryPtr->NumRefs;
+	}
+	PreloadEntryPtr = nullptr;
 
 	if (Shader)
 	{
+		INC_DWORD_STAT(STAT_Shaders_NumShadersUsedForRendering);
 		Shader->SetHash(ShaderHash);
 	}
 
diff --git a/Engine/Source/Runtime/RenderCore/Public/ShaderCodeArchive.h b/Engine/Source/Runtime/RenderCore/Public/ShaderCodeArchive.h
--- a/Engine/Source/Runtime/RenderCore/Public/ShaderCodeArchive.h
+++ b/Engine/Source/Runtime/RenderCore/Public/ShaderCodeArchive.h
@@ -12,26 +12,16 @@
 // enable visualization in the desktop Development builds only as it has a memory hit and writes files
 #define UE_SCA_VISUALIZE_SHADER_USAGE			(!WITH_EDITOR && UE_BUILD_DEVELOPMENT && PLATFORM_DESKTOP)
 
-struct FIoStoreShaderMapEntry
+struct FShaderMapEntry
 {
 	uint32 ShaderIndicesOffset = 0u;
 	uint32 NumShaders = 0u;
-
-	friend FArchive& operator <<(FArchive& Ar, FIoStoreShaderMapEntry& Ref)
-	{
-		return Ar << Ref.ShaderIndicesOffset << Ref.NumShaders;
-	}
-};
-
-struct FShaderMapEntry : public FIoStoreShaderMapEntry
-{
 	uint32 FirstPreloadIndex = 0u;
 	uint32 NumPreloadEntries = 0u;
 
 	friend FArchive& operator <<(FArchive& Ar, FShaderMapEntry& Ref)
 	{
-		operator<<(Ar, static_cast<FIoStoreShaderMapEntry&>(Ref));
-		return Ar << Ref.FirstPreloadIndex << Ref.NumPreloadEntries;
+		return Ar << Ref.ShaderIndicesOffset << Ref.NumShaders << Ref.FirstPreloadIndex << Ref.NumPreloadEntries;
 	}
 };
 
@@ -204,31 +194,57 @@ struct FShaderUsageVisualizer
 	/** Total number of shaders. */
 	int32 NumShaders;
 
-	/** Shader indices that were preloaded. */
-	TSet<int32> PreloadedShadersForVis;
+	/** Shader indices that we explicitly preloaded (does not include shaders preloaded as part of a compressed group). */
+	TSet<int32> ExplicitlyPreloadedShaders;
+
+	/** Shader indices that we preloaded (either explicitly or because they were a part of a compressed group). */
+	TSet<int32> PreloadedShaders;
+
+	/** Shader indices that we decompressed. */
+	TSet<int32> DecompressedShaders;
 
 	/** Shader indices that were created. */
-	TSet<int32> CreatedShadersForVis;
+	TSet<int32> CreatedShaders;
 
 	void Initialize(const int32 InNumShaders);
 
-	void MarkPreloadedForVisualization(int32 ShaderIndex)
+	inline void MarkExplicitlyPreloadedForVisualization(int32 ShaderIndex)
+	{
+		extern int32 GShaderCodeLibraryVisualizeShaderUsage;
+		if (LIKELY(GShaderCodeLibraryVisualizeShaderUsage))
+		{
+			FScopeLock Lock(&VisualizeLock);
+			ExplicitlyPreloadedShaders.Add(ShaderIndex);
+		}
+	}
+
+	inline void MarkPreloadedForVisualization(int32 ShaderIndex)
 	{
 		extern int32 GShaderCodeLibraryVisualizeShaderUsage;
 		if (LIKELY(GShaderCodeLibraryVisualizeShaderUsage))
 		{
 			FScopeLock Lock(&VisualizeLock);
-			PreloadedShadersForVis.Add(ShaderIndex);
+			PreloadedShaders.Add(ShaderIndex);
 		}
 	}
 
-	void MarkCreatedForVisualization(int32 ShaderIndex)
+	inline void MarkDecompressedForVisualization(int32 ShaderIndex)
 	{
 		extern int32 GShaderCodeLibraryVisualizeShaderUsage;
 		if (LIKELY(GShaderCodeLibraryVisualizeShaderUsage))
 		{
 			FScopeLock Lock(&VisualizeLock);
-			CreatedShadersForVis.Add(ShaderIndex);
+			DecompressedShaders.Add(ShaderIndex);
+		}
+	}
+
+	inline void MarkCreatedForVisualization(int32 ShaderIndex)
+	{
+		extern int32 GShaderCodeLibraryVisualizeShaderUsage;
+		if (LIKELY(GShaderCodeLibraryVisualizeShaderUsage))
+		{
+			FScopeLock Lock(&VisualizeLock);
+			CreatedShaders.Add(ShaderIndex);
 		}
 	}
 
@@ -236,6 +252,8 @@ struct FShaderUsageVisualizer
 #else
 	inline void Initialize(const int32 InNumShaders) {}
 	inline void MarkPreloadedForVisualization(int32 ShaderIndex) {}
+	inline void MarkExplicitlyPreloadedForVisualization(int32 ShaderIndex) {}
+	inline void MarkDecompressedForVisualization(int32 ShaderIndex) {}
 	inline void MarkCreatedForVisualization(int32 ShaderIndex) {}
 	inline void SaveShaderUsageBitmap(const FString& Name, EShaderPlatform ShaderPlatform) {}
 #endif // UE_SCA_VISUALIZE_SHADER_USAGE
@@ -331,11 +349,59 @@ protected:
 	FShaderUsageVisualizer DebugVisualizer;
 };
 
+
+
+namespace ShaderCodeArchive
+{
+	/** Decompresses the shader into caller-provided memory. Caller is assumed to allocate at least ShaderEntry uncompressed size value.
+	 * The engine will crash (LogFatal) if this function fails.
+	 */
+	RENDERCORE_API void DecompressShader(uint8* OutDecompressedShader, int64 UncompressedSize, const uint8* CompressedShaderCode, int64 CompressedSize);
+
+	/** Compresses the shader into caller-provided memory using the current CVar settings (e.g. on Windows), which is as opposed to potentially different settings on 
+	 * the platform where the shaders are going to be decompressed (e.g. on console). It is the caller's responsibility to make sure those two match.
+	 * This function can also be used for the estimation as it will return false if the provided CompressedBufferLength isn't sufficient (it is safe to pass nullptr in that case).
+	 */
+	RENDERCORE_API bool CompressShaderUsingCurrentSettings(uint8* OutCompressedShader, int64& OutCompressedSize, const uint8* UncompressedShaderCode, int64 UncompressedSize);
+}
+
+/** Descriptor of a shader map. This concept exists in run time, so this class describes the information stored in the library for a particular FShaderMap */
+struct FIoStoreShaderMapEntry
+{
+	/** Offset to an the first shader index referenced by this shader map in the array of shader indices. */
+	uint32 ShaderIndicesOffset = 0u;
+	/** Number of shaders in this shader map. */
+	uint32 NumShaders = 0u;
+
+	friend FArchive& operator <<(FArchive& Ar, FIoStoreShaderMapEntry& Ref)
+	{
+		return Ar << Ref.ShaderIndicesOffset << Ref.NumShaders;
+	}
+};
+
+/** Descriptor of an individual shader. */
 struct FIoStoreShaderCodeEntry
 {
-	uint32 UncompressedSize = 0;
-	uint32 CompressedSize = 0;
-	uint8 Frequency;
+	union
+	{
+		uint64 Packed;
+		struct
+		{
+			/** Shader type aka frequency (vertex, pixel, etc) */
+			uint64 Frequency : SF_NumBits;	// 4 as of now
+
+			/** Each shader belongs to a (one and only) shader group (even if it is the only one shader in that group) that is compressed and decompressed together. */
+			uint64 ShaderGroupIndex : 30;
+
+			/** Offset of the uncompressed shader in a group of shaders that are compressed / decompressed together. */
+			uint64 UncompressedOffsetInGroup : 30;
+		};
+	};
+
+	FIoStoreShaderCodeEntry()
+		: Packed(0)
+	{
+	}
 
 	EShaderFrequency GetFrequency() const
 	{
@@ -344,7 +410,104 @@ struct FIoStoreShaderCodeEntry
 
 	friend FArchive& operator <<(FArchive& Ar, FIoStoreShaderCodeEntry& Ref)
 	{
-		return Ar << Ref.UncompressedSize << Ref.CompressedSize << Ref.Frequency;
+		return Ar << Ref.Packed;
+	}
+};
+
+static_assert(sizeof(FIoStoreShaderCodeEntry) == sizeof(uint64), TEXT("To reduce memory footprint, shader code entries should be as small as possible"));
+
+/** Descriptor of a group of shaders compressed together. This groups already deduplicated, and possibly unrelated, shaders, so this is a distinct concept from a shader map. */
+struct FIoStoreShaderGroupEntry
+{
+	/** Offset to an the first shader index referenced by this group in the array of shader indices. This extra level of indirection allows arbitrary grouping. */
+	uint32 ShaderIndicesOffset = 0u;
+	/** Number of shaders in this group. */
+	uint32 NumShaders = 0u;
+
+	/** Uncompressed size of the whole group. */
+	uint32 UncompressedSize = 0;
+	/** Compressed size of the whole group. */
+	uint32 CompressedSize = 0;
+
+	friend FArchive& operator <<(FArchive& Ar, FIoStoreShaderGroupEntry& Ref)
+	{
+		return Ar << Ref.ShaderIndicesOffset << Ref.NumShaders << Ref.UncompressedSize << Ref.CompressedSize;
+	}
+
+	/** Some groups can be stored uncompressed if their compression wasn't beneficial (this is very vell possible, for groups that contain only one small shader. */
+	inline bool IsGroupCompressed() const
+	{
+		return CompressedSize != UncompressedSize;
+	}
+};
+
+struct FIoStoreShaderCodeArchiveHeader
+{
+public:
+
+	/** Hashes of all shadermaps in the library */
+	TArray<FSHAHash> ShaderMapHashes;
+
+	/** Output hashes of all shaders in the library */
+	TArray<FSHAHash> ShaderHashes;
+
+	/** Chunk Ids (essentially hashes) of the shader groups - needed to be serialized as they are used for preloading. */
+	TArray<FIoChunkId> ShaderGroupIoHashes;
+
+	/** An array of a shadermap descriptors. Each shadermap can reference an arbitrary number of shaders */
+	TArray<FIoStoreShaderMapEntry> ShaderMapEntries;
+
+	/** An array of all shaders descriptors, deduplicated */
+	TArray<FIoStoreShaderCodeEntry> ShaderEntries;
+
+	/** An array of shader group descriptors */
+	TArray<FIoStoreShaderGroupEntry> ShaderGroupEntries;
+
+	/** Flat array of shaders referenced by all shadermaps. Each shadermap has a range in this array, beginning of which is
+	  * stored as ShaderIndicesOffset in the shadermap's descriptor (FIoStoreShaderMapEntry).
+	  * This is also used by the shader groups.
+	  */
+	TArray<uint32> ShaderIndices;
+
+	friend RENDERCORE_API FArchive& operator <<(FArchive& Ar, FIoStoreShaderCodeArchiveHeader& Ref);
+
+	inline uint64 GetShaderUncompressedSize(int ShaderIndex)
+	{
+		const FIoStoreShaderCodeEntry& ThisShaderEntry = ShaderEntries[ShaderIndex];
+		const FIoStoreShaderGroupEntry& GroupEntry = ShaderGroupEntries[ThisShaderEntry.ShaderGroupIndex];
+
+		for (uint32 ShaderIdxIdx = GroupEntry.ShaderIndicesOffset, StopBeforeIdxIdx = GroupEntry.ShaderIndicesOffset + GroupEntry.NumShaders; ShaderIdxIdx < StopBeforeIdxIdx; ++ShaderIdxIdx)
+		{
+			int32 GroupMemberShaderIndex = ShaderIndices[ShaderIdxIdx];
+			if (ShaderIndex == GroupMemberShaderIndex)
+			{
+				// found ourselves, now find our size by subtracting from the next neighbor or the group size
+				if (LIKELY(ShaderIdxIdx < StopBeforeIdxIdx - 1))
+				{
+					const FIoStoreShaderCodeEntry& NextShaderEntry = ShaderEntries[ShaderIndices[ShaderIdxIdx + 1]];
+					return NextShaderEntry.UncompressedOffsetInGroup - ThisShaderEntry.UncompressedOffsetInGroup;
+				}
+				else
+				{
+					return GroupEntry.UncompressedSize - ThisShaderEntry.UncompressedOffsetInGroup;
+				}
+			}
+		}
+
+		checkf(false, TEXT("Could not find shader index %d in its own group %d - library is corrupted."), ShaderIndex, ThisShaderEntry.ShaderGroupIndex);
+		return 0;
+	}
+
+	uint64 GetAllocatedSize() const
+	{
+		return sizeof(*this) +
+			ShaderMapHashes.GetAllocatedSize() +
+			ShaderHashes.GetAllocatedSize() +
+			ShaderGroupIoHashes.GetAllocatedSize() +
+			ShaderMapEntries.GetAllocatedSize() +
+			ShaderEntries.GetAllocatedSize() +
+			ShaderGroupEntries.GetAllocatedSize() +
+			ShaderIndices.GetAllocatedSize();
 	}
 };
 
@@ -353,32 +516,30 @@ class FIoStoreShaderCodeArchive : public FRHIShaderLibrary
 public:
 	RENDERCORE_API static FIoChunkId GetShaderCodeArchiveChunkId(const FString& LibraryName, FName FormatName);
 	RENDERCORE_API static FIoChunkId GetShaderCodeChunkId(const FSHAHash& ShaderHash);
-	RENDERCORE_API static void SaveIoStoreShaderCodeArchive(const FSerializedShaderArchive& SerializedShaders, FArchive& OutLibraryAr);
+	/** This function creates the archive header, including splitting shaders into groups. */
+	RENDERCORE_API static void CreateIoStoreShaderCodeArchiveHeader(const FName& Format, const FSerializedShaderArchive& SerializedShaders, FIoStoreShaderCodeArchiveHeader& OutHeader);
+	RENDERCORE_API static void SaveIoStoreShaderCodeArchive(const FIoStoreShaderCodeArchiveHeader& Header, FArchive& OutLibraryAr);
 	static FIoStoreShaderCodeArchive* Create(EShaderPlatform InPlatform, const FString& InLibraryName, FIoDispatcher& InIoDispatcher);
 
 	virtual ~FIoStoreShaderCodeArchive();
 
 	virtual bool IsNativeLibrary() const override { return false; }
 
-	uint32 GetSizeBytes() const
+	uint64 GetSizeBytes() const
 	{
 		return sizeof(*this) +
-			ShaderMapHashes.GetAllocatedSize() +
-			ShaderMapEntries.GetAllocatedSize() +
-			ShaderHashes.GetAllocatedSize() +
-			ShaderEntries.GetAllocatedSize() +
-			ShaderIndices.GetAllocatedSize() +
-			ShaderPreloads.GetAllocatedSize();
+			Header.GetAllocatedSize() +
+			PreloadedShaderGroups.GetAllocatedSize();
 	}
 
-	virtual int32 GetNumShaders() const override { return ShaderEntries.Num(); }
-	virtual int32 GetNumShaderMaps() const override { return ShaderMapEntries.Num(); }
-	virtual int32 GetNumShadersForShaderMap(int32 ShaderMapIndex) const override { return ShaderMapEntries[ShaderMapIndex].NumShaders; }
+	virtual int32 GetNumShaders() const override { return Header.ShaderEntries.Num(); }
+	virtual int32 GetNumShaderMaps() const override { return Header.ShaderMapEntries.Num(); }
+	virtual int32 GetNumShadersForShaderMap(int32 ShaderMapIndex) const override { return Header.ShaderMapEntries[ShaderMapIndex].NumShaders; }
 
 	virtual int32 GetShaderIndex(int32 ShaderMapIndex, int32 i) const override
 	{
-		const FIoStoreShaderMapEntry& ShaderMapEntry = ShaderMapEntries[ShaderMapIndex];
-		return ShaderIndices[ShaderMapEntry.ShaderIndicesOffset + i];
+		const FIoStoreShaderMapEntry& ShaderMapEntry = Header.ShaderMapEntries[ShaderMapIndex];
+		return Header.ShaderIndices[ShaderMapEntry.ShaderIndicesOffset + i];
 	}
 
 	virtual int32 FindShaderMapIndex(const FSHAHash& Hash) override;
@@ -393,7 +554,7 @@ public:
 private:
 	static constexpr uint32 CurrentVersion = 1;
 
-	struct FShaderPreloadEntry
+	struct FShaderGroupPreloadEntry
 	{
 		FGraphEventRef PreloadEvent;
 		FIoRequest IoRequest;
@@ -401,7 +562,7 @@ private:
 		uint32 NumRefs : 31;
 		uint32 bNeverToBePreloaded : 1;
 
-		FShaderPreloadEntry()
+		FShaderGroupPreloadEntry()
 			: NumRefs(0)
 			, bNeverToBePreloaded(0)
 		{
@@ -409,19 +570,52 @@ private:
 	};
 
 	FIoStoreShaderCodeArchive(EShaderPlatform InPlatform, const FString& InLibraryName, FIoDispatcher& InIoDispatcher);
-	bool ReleaseRef(int32 ShaderIndex);
 
 	FIoDispatcher& IoDispatcher;
 
-	TArray<FSHAHash> ShaderMapHashes;
-	TArray<FSHAHash> ShaderHashes;
-	TArray<FIoStoreShaderMapEntry> ShaderMapEntries;
-	TArray<FIoStoreShaderCodeEntry> ShaderEntries;
-	TArray<uint32> ShaderIndices;
+	/** Preloads a given shader group. */
+	bool PreloadShaderGroup(int32 ShaderGroupIndex, FGraphEventArray& OutCompletionEvents, FCoreDelegates::FAttachShaderReadRequestFunc* AttachShaderReadRequestFuncPtr = nullptr);
+
+	/** Sets up a new preload entry for preload.*/
+	void SetupPreloadEntryForLoading(int32 ShaderGroupIndex, FShaderGroupPreloadEntry& PreloadEntry);
+
+	/** Sets up a preload entry for groups that shouldn't be preloaded.*/
+	void MarkPreloadEntrySkipped(int32 ShaderGroupIndex);
+
+	/** Releases a reference to a preloaded shader group, potentially deleting it. */
+	void ReleasePreloadEntry(int32 ShaderGroupIndex);
+
+	/** Returns the index of shader group that a given shader belongs to. */
+	inline int32 GetGroupIndexForShader(int32 ShaderIndex) const
+	{
+		return Header.ShaderEntries[ShaderIndex].ShaderGroupIndex;
+	}
+
+	/** Finds or adds preload info for a shader group. Assumes lock guarding access to the info taken, never returns nullptr (except when new failed and we're already broken beyond repair)*/
+	inline FShaderGroupPreloadEntry* FindOrAddPreloadEntry(int32 ShaderGroupIndex)
+	{
+		FShaderGroupPreloadEntry*& Ptr = PreloadedShaderGroups.FindOrAdd(ShaderGroupIndex);
+		if (UNLIKELY(Ptr == nullptr))
+		{
+			Ptr = new FShaderGroupPreloadEntry;
+		}
+		return Ptr;
+	}
+
+	/** Returns true if the group contains only RTX shaders. We can avoid preloading it when running with RTX off. */
+	bool GroupOnlyContainsRaytracingShaders(int32 ShaderGroupIndex);
+
+	/** Archive header with all the metadata */
+	FIoStoreShaderCodeArchiveHeader Header;
+
+	/** Hash tables for faster searching for shader and shadermap hashes. */
 	FHashTable ShaderMapHashTable;
 	FHashTable ShaderHashTable;
-	TArray<FShaderPreloadEntry> ShaderPreloads;
-	FRWLock ShaderPreloadLock;
+
+	/** Mapping between the group index and preloaded groups. Should be only modified when lock is taken. */
+	TMap<int32, FShaderGroupPreloadEntry*> PreloadedShaderGroups;
+	/** Lock guarding access to the book-keeping info above.*/
+	FRWLock PreloadedShaderGroupsLock;
 
 	/** debug visualizer - in Shipping compiles out to an empty struct with no-op functions */
 	FShaderUsageVisualizer DebugVisualizer;
