DisplayCluster:  crash in GPU profiler when toggling views.  Need to clear description string buffer offsets when reusing GPU profiler events, since constructor doesn't get called again.

#rnx

diff --git a/Engine/Source/Runtime/RenderCore/Private/ProfilingDebugging/RealtimeGPUProfiler.cpp b/Engine/Source/Runtime/RenderCore/Private/ProfilingDebugging/RealtimeGPUProfiler.cpp
--- a/Engine/Source/Runtime/RenderCore/Private/ProfilingDebugging/RealtimeGPUProfiler.cpp
+++ b/Engine/Source/Runtime/RenderCore/Private/ProfilingDebugging/RealtimeGPUProfiler.cpp
@@ -295,6 +295,12 @@ public:
 		}
 	}
 
+	void ClearDescription()
+	{
+		DescriptionLength = 0;
+		DescriptionOffset = 0;
+	}
+
 	FRHIGPUMask GetGPUMask() const
 	{
 		return GPUMask;
@@ -510,6 +516,10 @@ public:
 		{
 			GpuProfilerEvents[EventIdx].SetDescription(Description, DescriptionStringBuffer);
 		}
+		else
+		{
+			GpuProfilerEvents[EventIdx].ClearDescription();
+		}
 		return GpuProfilerEvents[EventIdx].Begin(GPUMask, Name, StatName);
 	}
 
