Fix crash on packaging (cannot use per-platform CVar in UnrealPak) (UE-139676).

diff --git a/Engine/Source/Runtime/RenderCore/Private/ShaderCodeArchive.cpp b/Engine/Source/Runtime/RenderCore/Private/ShaderCodeArchive.cpp
--- a/Engine/Source/Runtime/RenderCore/Private/ShaderCodeArchive.cpp
+++ b/Engine/Source/Runtime/RenderCore/Private/ShaderCodeArchive.cpp
@@ -50,11 +50,11 @@ static FAutoConsoleVariableRef CVarShaderCodeLibraryVisualizeShaderUsage(
 	ECVF_RenderThreadSafe | ECVF_ReadOnly
 );
 
-int32 GShaderCodeLibraryMaxShaderGroupSize = 1024 * 1024;	// decompressing 1MB of shaders takes about 0.1ms on PC (TR 3970x, Oodle Mermaid6). Can be overriden per platform.
+int32 GShaderCodeLibraryMaxShaderGroupSize = 1024 * 1024;	// decompressing 1MB of shaders takes about 0.1ms on PC (TR 3970x, Oodle Mermaid6).
 static FAutoConsoleVariableRef CVarShaderCodeLibraryMaxShaderGroupSize(
 	TEXT("r.ShaderCodeLibrary.MaxShaderGroupSize"),
 	GShaderCodeLibraryMaxShaderGroupSize,
-	TEXT("Max (uncompressed) size of a group of shaders to be compressed/decompressed together. Can be set differently for each platform in their own *Engine.ini.")
+	TEXT("Max (uncompressed) size of a group of shaders to be compressed/decompressed together.")
 	TEXT("If a group exceeds it, it will be evenly split into subgroups which strive to not exceed it. However, if a shader group is down to one shader and still exceeds the limit, the limit will be ignored."),
 	ECVF_RenderThreadSafe | ECVF_ReadOnly
 );
@@ -1342,9 +1342,7 @@ void FIoStoreShaderCodeArchive::CreateIoStoreShaderCodeArchiveHeader(const FName
 	);
 
 	// get the effective maximum uncompressed group size
-	FShaderPlatformCachedIniValue<int32> ConfiguredMaxShaderGroupSize(TEXT("r.ShaderCodeLibrary.MaxShaderGroupSize"));
-	EShaderPlatform ShaderPlatform = ShaderFormatNameToShaderPlatform(Format);
-	uint32 MaxUncompressedShaderGroupSize = FMath::Max(ConfiguredMaxShaderGroupSize.Get(ShaderPlatform), 1);	// cannot be lower than 1
+	uint32 MaxUncompressedShaderGroupSize = FMath::Max(static_cast<uint32>(GShaderCodeLibraryMaxShaderGroupSize), 1U);	// cannot be lower than 1
 
 	// for statistics
 	uint64 Stats_TotalUncompressedMemory = 0;
