Fix for assert when async compute pass creates a transient resource. The resource lifetime was being properly extended to the graphics fork pass, but the graphics fork pass was also within a merged render pass, and the resource lifetime was not properly extended to the start of the render pass. This is important because resource transitions cannot occur within a render pass. This fix adjusts async compute resource lifetimes to begin and end on render pass boundaries of the fork / join graphics pass.

diff --git a/Engine/Source/Runtime/RenderCore/Private/RenderGraphBuilder.cpp b/Engine/Source/Runtime/RenderCore/Private/RenderGraphBuilder.cpp
--- a/Engine/Source/Runtime/RenderCore/Private/RenderGraphBuilder.cpp
+++ b/Engine/Source/Runtime/RenderCore/Private/RenderGraphBuilder.cpp
@@ -1373,6 +1373,16 @@ void FRDGBuilder::Compile(EExecuteMode ExecuteMode)
 			GraphicsPass->GetPrologueBarriersToEnd(Allocator).AddDependency(&EpilogueBarriersToBeginForGraphics);
 		};
 
+		const auto AddResourcesToBegin = [this](FRDGPass* PassToBegin, FRDGPass* PassWithResources)
+		{
+			Passes[PassToBegin->PrologueBarrierPass]->ResourcesToBegin.Add(PassWithResources);
+		};
+
+		const auto AddResourcesToEnd = [this](FRDGPass* PassToEnd, FRDGPass* PassWithResources)
+		{
+			Passes[PassToEnd->EpilogueBarrierPass]->ResourcesToEnd.Add(PassWithResources);
+		};
+
 		FRDGPassHandle CurrentGraphicsForkPassHandle;
 
 		for (FRDGPassHandle PassHandle = ProloguePassHandle + 1; PassHandle < EpiloguePassHandle; ++PassHandle)
@@ -1394,7 +1404,7 @@ void FRDGBuilder::Compile(EExecuteMode ExecuteMode)
 			FRDGPass* GraphicsForkPass = Passes[GraphicsForkPassHandle];
 
 			AsyncComputePass->GraphicsForkPass = GraphicsForkPassHandle;
-			GraphicsForkPass->ResourcesToBegin.Add(AsyncComputePass);
+			AddResourcesToBegin(GraphicsForkPass, AsyncComputePass);
 
 			if (CurrentGraphicsForkPassHandle != GraphicsForkPassHandle)
 			{
@@ -1412,7 +1422,7 @@ void FRDGBuilder::Compile(EExecuteMode ExecuteMode)
 			FRDGPass* AsyncComputeProloguePass = ProloguePasses[ERHIPipeline::AsyncCompute];
 
 			AsyncComputePass->GraphicsJoinPass = GraphicsJoinPassHandle;
-			GraphicsJoinPass->ResourcesToEnd.Add(AsyncComputePass);
+			AddResourcesToEnd(GraphicsJoinPass, AsyncComputePass);
 
 			if (AsyncComputeProloguePass->GraphicsJoinPass == AsyncComputeProloguePass->Handle)
 			{
@@ -3120,7 +3130,8 @@ void FRDGBuilder::BeginResourceRHI(FRDGPassHandle PassHandle, FRDGTextureRef Tex
 		check(Pass->Pipeline == ERHIPipeline::Graphics);
 
 		// Cannot begin a resource within a merged render pass region.
-		check(GetPrologueBarrierPassHandle(PassHandle) == PassHandle);
+		checkf(GetPrologueBarrierPassHandle(PassHandle) == PassHandle,
+			TEXT("Cannot begin a resource within a merged render pass. Pass (Handle: %d, Name: %s), Resource %s"), PassHandle, Pass->GetName(), Texture->Name);
 	}
 #endif
 
@@ -3213,7 +3224,8 @@ void FRDGBuilder::BeginResourceRHI(FRDGPassHandle PassHandle, FRDGBufferRef Buff
 		check(Pass->Pipeline == ERHIPipeline::Graphics);
 
 		// Cannot begin a resource within a merged render pass region.
-		check(GetPrologueBarrierPassHandle(PassHandle) == PassHandle);
+		checkf(GetPrologueBarrierPassHandle(PassHandle) == PassHandle,
+			TEXT("Cannot begin a resource within a merged render pass. Pass (Handle: %d, Name: %s), Resource %s"), PassHandle, Pass->GetName(), Buffer->Name);
 	}
 #endif
 	Buffer->FinalizeDesc();
