Fixed bug in FRenderBounds& operator += (const FVector3f&)

diff --git a/Engine/Source/Runtime/RenderCore/Public/RenderTransform.h b/Engine/Source/Runtime/RenderCore/Public/RenderTransform.h
--- a/Engine/Source/Runtime/RenderCore/Public/RenderTransform.h
+++ b/Engine/Source/Runtime/RenderCore/Public/RenderTransform.h
@@ -281,7 +281,7 @@ public:
 	FORCEINLINE FRenderBounds& operator += (const FVector3f& Other)
 	{
 		const VectorRegister VecOther = VectorLoadFloat3(&Other.X);
-		VectorStoreFloat3(VectorMin(VectorLoadFloat3(&Min.Y), VecOther), &Min);
+		VectorStoreFloat3(VectorMin(VectorLoadFloat3(&Min.X), VecOther), &Min);
 		VectorStoreFloat3(VectorMax(VectorLoadFloat3(&Max.X), VecOther), &Max);
 		return *this;
 	}
