Make IoStore-based shader library the default everywhere.

- IoStore-based code archive has significant advantages over the file-based one as it implements shader group compression.
- Also fixed a PC/desktop-specific issue of multiple opens of the library.

diff --git a/Engine/Config/BaseEngine.ini b/Engine/Config/BaseEngine.ini
--- a/Engine/Config/BaseEngine.ini
+++ b/Engine/Config/BaseEngine.ini
@@ -1121,6 +1121,7 @@ AssetLogShowsAbsolutePath=True
 TaskThreadStackSize=393216
 PoolThreadStackSize=131072
 DefaultThreadStackSize=1048576
+AllowShadersInIoStore=True
 
 [Core.System.Experimental]
 EnablePackageSidecarSaving=False
diff --git a/Engine/Source/Runtime/RenderCore/Private/ShaderCodeLibrary.cpp b/Engine/Source/Runtime/RenderCore/Private/ShaderCodeLibrary.cpp
--- a/Engine/Source/Runtime/RenderCore/Private/ShaderCodeLibrary.cpp
+++ b/Engine/Source/Runtime/RenderCore/Private/ShaderCodeLibrary.cpp
@@ -730,14 +730,17 @@ public:
 	static FShaderLibraryInstance* Create(EShaderPlatform InShaderPlatform, const FString& ShaderCodeDir, FString const& InLibraryName)
 	{
 		FRHIShaderLibraryRef Library;
+		FString ShaderCodeDirectory;
 		if (RHISupportsNativeShaderLibraries(InShaderPlatform))
 		{
 			Library = RHICreateShaderLibrary(InShaderPlatform, ShaderCodeDir, InLibraryName);
+			ShaderCodeDirectory = ShaderCodeDir;
 		}
 
 		if (!Library && FIoDispatcher::IsInitialized())
 		{
 			Library = FIoStoreShaderCodeArchive::Create(InShaderPlatform, InLibraryName, FIoDispatcher::Get());
+			ShaderCodeDirectory.Empty();	// paths don't matter for IoStore-based libraries
 		}
 
 		if (!Library)
@@ -754,6 +757,8 @@ public:
 					Library = FShaderCodeArchive::Create(InShaderPlatform, *Ar, DestFilePath, ShaderCodeDir, InLibraryName);
 					if (Library)
 					{
+						ShaderCodeDirectory = ShaderCodeDir;
+
 						bool ShaderCodeLibrarySeparateLoadingCacheCommandLineOverride = FParse::Param(FCommandLine::Get(), TEXT("ShaderCodeLibrarySeparateLoadingCache"));;
 						if (GShaderCodeLibrarySeparateLoadingCache || ShaderCodeLibrarySeparateLoadingCacheCommandLineOverride)
 						{
@@ -774,6 +779,7 @@ public:
 
 		FShaderLibraryInstance* Instance = new FShaderLibraryInstance();
 		Instance->Library = Library;
+		Instance->ShaderCodeDirectory = ShaderCodeDirectory;
 
 		const int32 NumResources = Library->GetNumShaderMaps();
 		Instance->Resources.AddZeroed(NumResources);
@@ -960,6 +966,17 @@ public:
 		}
 	}
 
+	bool HasContentFrom(const FString& ShaderCodeDir, FString const& InLibraryName)
+	{
+		if (Library->GetName() == InLibraryName)
+		{
+			// IoStore-based libraries don't care about the directory, so name collision alone is enough to say yes.
+			return ShaderCodeDirectory.IsEmpty() ? true : (ShaderCodeDirectory == ShaderCodeDir);
+		}
+
+		return false;
+	}
+
 public:
 	FRHIShaderLibraryRef Library;
 
@@ -994,6 +1011,9 @@ private:
 	/** Locks that guard access to particular shader buckets. */
 	FRWLock ShaderLocks[NumShaderLocks];
 	FRWLock ResourceLock;
+
+	/** Folder the library was created from (doesn't matter for IoStore-based libraries) */
+	FString ShaderCodeDirectory;
 };
 
 FShaderMapResource_SharedCode::FShaderMapResource_SharedCode(FShaderLibraryInstance* InLibraryInstance, int32 InShaderMapIndex)
@@ -2975,9 +2995,22 @@ void FShaderCodeLibrary::UnregisterSharedShaderCodeRequestDelegate_Handle(FDeleg
 
 // FNamedShaderLibrary methods
 
-// At runtime, open shader code collection for specified shader platform
+// At runtime, open shader code collection for specified shader platform. Returns true if new code was opened
 bool UE::ShaderLibrary::Private::FNamedShaderLibrary::OpenShaderCode(const FString& ShaderCodeDir, FString const& Library)
 {
+	// check if any of the components has this content
+	{
+		FRWScopeLock ReadLock(ComponentsMutex, SLT_Write);
+		for (TUniquePtr<FShaderLibraryInstance>& Component : Components)
+		{
+			if (Component->HasContentFrom(ShaderCodeDir, Library))
+			{
+				// in this context, "false" means "no new library was opened"
+				return false;
+			}
+		}
+	}
+
 	FShaderLibraryInstance* LibraryInstance = FShaderLibraryInstance::Create(ShaderPlatform, ShaderCodeDir, Library);
 	if (LibraryInstance == nullptr)
 	{
@@ -2986,6 +3019,19 @@ bool UE::ShaderLibrary::Private::FNamedShaderLibrary::OpenShaderCode(const FStri
 		return false; //-V773
 	}
 
+	FRWScopeLock WriteLock(ComponentsMutex, SLT_Write);
+
+	// re-check that no one has added the same library while we were creating it. If so, delete ours
+	for (TUniquePtr<FShaderLibraryInstance>& Component : Components)
+	{
+		if (Component->HasContentFrom(ShaderCodeDir, Library))
+		{
+			// in this context, "false" means "no new library was opened"
+			delete LibraryInstance;
+			return false;
+		}
+	}
+
 	if (LibraryInstance->Library->IsNativeLibrary())
 	{
 		UE_LOG(LogShaderLibrary, Display, TEXT("Cooked Context: Loaded Native Shared Shader Library %s"), *Library);
@@ -2995,7 +3041,6 @@ bool UE::ShaderLibrary::Private::FNamedShaderLibrary::OpenShaderCode(const FStri
 		UE_LOG(LogShaderLibrary, Display, TEXT("Cooked Context: Using Shared Shader Library %s"), *Library);
 	}
 
-	FRWScopeLock WriteLock(ComponentsMutex, SLT_Write);
 	Components.Emplace(LibraryInstance);
 	return true;
 }
