Added "ultimate" version of AddClearRenderTargetPass that lets the user clear multiple mips, multiple array slices, with/without a custom clear color or viewport and does so by using a clear action if possible.
Note: the existing versions have not been deprecated since they are more straightforward (they clear only a single mip/slice).

diff --git a/Engine/Source/Runtime/RenderCore/Private/RenderGraphUtils.cpp b/Engine/Source/Runtime/RenderCore/Private/RenderGraphUtils.cpp
--- a/Engine/Source/Runtime/RenderCore/Private/RenderGraphUtils.cpp
+++ b/Engine/Source/Runtime/RenderCore/Private/RenderGraphUtils.cpp
@@ -542,46 +542,103 @@ void AddClearUAVPass(FRDGBuilder& GraphBuilder, FRDGTextureUAVRef TextureUAV, co
 
 void AddClearRenderTargetPass(FRDGBuilder& GraphBuilder, FRDGTextureRef Texture)
 {
-	check(Texture);
-
-	FRenderTargetParameters* Parameters = GraphBuilder.AllocParameters<FRenderTargetParameters>();
-	Parameters->RenderTargets[0] = FRenderTargetBinding(Texture, ERenderTargetLoadAction::EClear);
-
-	GraphBuilder.AddPass(
-		RDG_EVENT_NAME("ClearRenderTarget(%s) %dx%d ClearAction", Texture->Name, Texture->Desc.Extent.X, Texture->Desc.Extent.Y),
-		Parameters,
-		ERDGPassFlags::Raster,
-		[](FRHICommandList& RHICmdList) {});
+	// Single mip, single slice, same clear color as what is specified in the render target :
+	AddClearRenderTargetPass(GraphBuilder, Texture, FRDGTextureClearInfo());
 }
 
 void AddClearRenderTargetPass(FRDGBuilder& GraphBuilder, FRDGTextureRef Texture, const FLinearColor& ClearColor)
 {
-	if (Texture->Desc.ClearValue.ColorBinding == EClearBinding::EColorBound && Texture->Desc.ClearValue.GetClearColor() == ClearColor)
-	{
-		AddClearRenderTargetPass(GraphBuilder, Texture);
-	}
-	else
-	{
-		AddClearRenderTargetPass(GraphBuilder, Texture, ClearColor, FIntRect(FIntPoint::ZeroValue, Texture->Desc.Extent));
-	}
+	// Single mip, single slice, custom clear color :
+	FRDGTextureClearInfo TextureClearInfo;
+	TextureClearInfo.ClearColor = ClearColor;
+	AddClearRenderTargetPass(GraphBuilder, Texture, TextureClearInfo);
 }
 
 void AddClearRenderTargetPass(FRDGBuilder& GraphBuilder, FRDGTextureRef Texture, const FLinearColor& ClearColor, FIntRect Viewport)
+{
+	// Single mip, single slice, custom viewport, custom clear color :
+	FRDGTextureClearInfo TextureClearInfo;
+	TextureClearInfo.ClearColor = ClearColor;
+	TextureClearInfo.Viewport = Viewport;
+	AddClearRenderTargetPass(GraphBuilder, Texture, TextureClearInfo);
+}
+
+void AddClearRenderTargetPass(FRDGBuilder& GraphBuilder, FRDGTextureRef Texture, const FRDGTextureClearInfo& TextureClearInfo)
 {
 	check(Texture);
 
-	FRenderTargetParameters* Parameters = GraphBuilder.AllocParameters<FRenderTargetParameters>();
-	Parameters->RenderTargets[0] = FRenderTargetBinding(Texture, ERenderTargetLoadAction::ENoAction);
+	bool bUseCustomViewport = (TextureClearInfo.Viewport.Area() > 0);
+	FLinearColor ClearColor = TextureClearInfo.ClearColor.IsSet() ? TextureClearInfo.ClearColor.GetValue() : Texture->Desc.ClearValue.GetClearColor();
+	uint16 TextureNumSlicesOrDepth = Texture->Desc.IsTexture3D() ? Texture->Desc.Depth : Texture->Desc.ArraySize;
 
-	GraphBuilder.AddPass(
-		RDG_EVENT_NAME("ClearRenderTarget(%s) [(%d, %d), (%d, %d)] ClearQuad", Texture->Name, Viewport.Min.X, Viewport.Min.Y, Viewport.Max.X, Viewport.Max.Y),
-		Parameters,
-		ERDGPassFlags::Raster,
-		[Parameters, ClearColor, Viewport](FRHICommandList& RHICmdList)
+	checkf((TextureClearInfo.FirstMipIndex < Texture->Desc.NumMips) && ((TextureClearInfo.FirstMipIndex + TextureClearInfo.NumMips) <= Texture->Desc.NumMips),
+		TEXT("Invalid mip range [%d, %d] for texture %s (%d mips)"), TextureClearInfo.FirstMipIndex, TextureClearInfo.FirstMipIndex + TextureClearInfo.NumMips - 1, Texture->Name, Texture->Desc.NumMips);
+	checkf(((TextureClearInfo.FirstSliceIndex == 0) && (TextureClearInfo.NumSlices == 1)) || (Texture->Desc.IsTextureArray() && EnumHasAnyFlags(Texture->Desc.Flags, ETextureCreateFlags::TargetArraySlicesIndependently)),
+		TEXT("Per-slice clear (outside of slice 0, i.e. clearing any other slice than the first one) is only supported on 2DArray at the moment and ETextureCreateFlags::TargetArraySlicesIndependently must be passed when creating the texture (texture %s)."), Texture->Name);
+	checkf((TextureClearInfo.FirstSliceIndex < TextureNumSlicesOrDepth) && ((TextureClearInfo.FirstSliceIndex + TextureClearInfo.NumSlices) <= TextureNumSlicesOrDepth),
+		TEXT("Invalid slice range [%d, %d] for texture %s (%d slices)"), TextureClearInfo.FirstSliceIndex, TextureClearInfo.FirstSliceIndex + TextureClearInfo.NumSlices - 1, Texture->Name, TextureNumSlicesOrDepth);
+
+	// Use clear action if no viewport specified and clear color is not passed or matches the fast clear color :
+	if (!bUseCustomViewport
+		&& (Texture->Desc.ClearValue.ColorBinding == EClearBinding::EColorBound)
+		&& (Texture->Desc.ClearValue.GetClearColor() == ClearColor))
 	{
-		RHICmdList.SetViewport((float)Viewport.Min.X, (float)Viewport.Min.Y, 0.0f, (float)Viewport.Max.X, (float)Viewport.Max.Y, 1.0f);
-		DrawClearQuad(RHICmdList, ClearColor);
-	});
+		for (uint32 SliceIndex = 0; SliceIndex < TextureClearInfo.NumSlices; ++SliceIndex)
+		{
+			uint16 CurrentSliceIndex = IntCastChecked<uint16>(TextureClearInfo.FirstSliceIndex + SliceIndex);
+			for (uint32 MipIndex = 0; MipIndex < TextureClearInfo.NumMips; ++MipIndex)
+			{
+				uint8 CurrentMipIndex = IntCastChecked<uint8>(TextureClearInfo.FirstMipIndex + MipIndex);
+
+				FRenderTargetParameters* Parameters = GraphBuilder.AllocParameters<FRenderTargetParameters>();
+				Parameters->RenderTargets[0] = FRenderTargetBinding(Texture, ERenderTargetLoadAction::EClear, CurrentMipIndex, CurrentSliceIndex);
+
+				GraphBuilder.AddPass(
+					RDG_EVENT_NAME("ClearRenderTarget(%s, slice %d, mip %d) %dx%d ClearAction", Texture->Name, CurrentSliceIndex, CurrentMipIndex, Texture->Desc.Extent.X, Texture->Desc.Extent.Y),
+					Parameters,
+					ERDGPassFlags::Raster,
+					[](FRHICommandList& RHICmdList) {});
+			}
+		}
+	}
+	else
+	{
+		FIntRect OriginalViewport = bUseCustomViewport ? TextureClearInfo.Viewport : FIntRect(FIntPoint::ZeroValue, Texture->Desc.Extent);
+		checkf((OriginalViewport.Max.X <= Texture->Desc.Extent.X) && (OriginalViewport.Max.Y <= Texture->Desc.Extent.Y), TEXT("Invalid custom viewport ((%d, %d) - (%d, %d)) for texture %s (size (%d, %d))"),
+			OriginalViewport.Min.X, OriginalViewport.Min.Y, OriginalViewport.Max.X, OriginalViewport.Max.Y, Texture->Name, Texture->Desc.Extent.X, Texture->Desc.Extent.Y);
+
+		for (uint32 SliceIndex = 0; SliceIndex < TextureClearInfo.NumSlices; ++SliceIndex)
+		{
+			uint16 CurrentSliceIndex = IntCastChecked<uint16>(TextureClearInfo.FirstSliceIndex + SliceIndex);
+			for (uint32 MipIndex = 0; MipIndex < TextureClearInfo.NumMips; ++MipIndex)
+			{
+				FIntRect CurrentViewport(
+					FMath::Max(1u, (uint32)OriginalViewport.Min.X >> MipIndex),
+					FMath::Max(1u, (uint32)OriginalViewport.Min.Y >> MipIndex),
+					FMath::Max(1u, (uint32)OriginalViewport.Max.X >> MipIndex),
+					FMath::Max(1u, (uint32)OriginalViewport.Max.Y >> MipIndex));
+
+				uint8 CurrentMipIndex = IntCastChecked<uint8>(TextureClearInfo.FirstMipIndex + MipIndex);
+				if (CurrentViewport.Area() > 0)
+				{
+					FRenderTargetParameters* Parameters = GraphBuilder.AllocParameters<FRenderTargetParameters>();
+					Parameters->RenderTargets[0] = FRenderTargetBinding(Texture, ERenderTargetLoadAction::ENoAction, CurrentMipIndex, CurrentSliceIndex);
+
+					GraphBuilder.AddPass(
+						RDG_EVENT_NAME("ClearRenderTarget(%s, slice %d, mip %d) [(%d, %d), (%d, %d)] ClearQuad", Texture->Name, CurrentSliceIndex, CurrentMipIndex, CurrentViewport.Min.X, CurrentViewport.Min.Y, CurrentViewport.Max.X, CurrentViewport.Max.Y),
+						Parameters,
+						ERDGPassFlags::Raster,
+						[Parameters, ClearColor, CurrentViewport](FRHICommandList& RHICmdList)
+					{
+						RHICmdList.SetViewport((float)CurrentViewport.Min.X, (float)CurrentViewport.Min.Y, 0.0f, (float)CurrentViewport.Max.X, (float)CurrentViewport.Max.Y, 1.0f);
+						DrawClearQuad(RHICmdList, ClearColor);
+					});
+
+					CurrentViewport = CurrentViewport / 2;
+				}
+			}
+		}
+	}
 }
 
 void AddClearDepthStencilPass(
diff --git a/Engine/Source/Runtime/RenderCore/Public/RenderGraphUtils.h b/Engine/Source/Runtime/RenderCore/Public/RenderGraphUtils.h
--- a/Engine/Source/Runtime/RenderCore/Public/RenderGraphUtils.h
+++ b/Engine/Source/Runtime/RenderCore/Public/RenderGraphUtils.h
@@ -692,15 +692,37 @@ RENDERCORE_API void AddClearUAVPass(FRDGBuilder& GraphBuilder, FRDGTextureUAVRef
 /** Clears parts of UAV specified by an array of screen rects. If no rects are specific, then it falls back to a standard UAV clear. */
 RENDERCORE_API void AddClearUAVPass(FRDGBuilder& GraphBuilder, FRDGTextureUAVRef TextureUAV, const uint32(&ClearValues)[4], FRDGBufferSRVRef RectMinMaxBufferSRV, uint32 NumRects);
 
-/** Adds a render graph pass to clear a render target to its clear value. */
+/** Adds a render graph pass to clear a render target to its clear value (single mip, single slice) */
 RENDERCORE_API void AddClearRenderTargetPass(FRDGBuilder& GraphBuilder, FRDGTextureRef Texture);
 
-/** Adds a render graph pass to clear a render target. Uses render pass clear actions if the clear color matches the fast clear color. */
+/** Adds a render graph pass to clear a render target (single mip, single slice). Uses render pass clear actions if the clear color matches the fast clear color. */
 RENDERCORE_API void AddClearRenderTargetPass(FRDGBuilder& GraphBuilder, FRDGTextureRef Texture, const FLinearColor& ClearColor);
 
-/** Adds a render graph pass to clear a render target. Draws a quad to the requested viewport. */
+/** Adds a render graph pass to clear a part of a render target (single mip, single slice). Draws a quad to the requested viewport. */
 RENDERCORE_API void AddClearRenderTargetPass(FRDGBuilder& GraphBuilder, FRDGTextureRef Texture, const FLinearColor& ClearColor, FIntRect Viewport);
 
+struct FRDGTextureClearInfo
+{
+	FRDGTextureClearInfo() = default;
+
+	// Optional : specifies the area of the texture to clear (draws a quad to the requested viewport). If zero (and clear color is not passed or matches the fast clear color), the clear will be done with render pass clear actions :
+	FIntRect Viewport;
+
+	// Optional : specifies the clear color. If not specified or if the clear color matches the fast clear color (and Viewport is zero), the clear will be done with render pass clear actions :
+	TOptional<FLinearColor> ClearColor;
+
+	// For texture arrays, specifies which slice(s) to clear
+	uint32 FirstSliceIndex = 0;
+	uint32 NumSlices = 1;
+
+	// Specifies which mip(s) to clear
+	uint32 FirstMipIndex = 0;
+	uint32 NumMips = 1;
+};
+
+/** Adds a render graph pass to clear a render target to its clear value. */
+RENDERCORE_API void AddClearRenderTargetPass(FRDGBuilder& GraphBuilder, FRDGTextureRef Texture, const FRDGTextureClearInfo& TextureClearInfo);
+
 /** Adds a render graph pass to clear a depth stencil target. Prefer to use clear actions if possible. */
 RENDERCORE_API void AddClearDepthStencilPass(
 	FRDGBuilder& GraphBuilder,
