Do not skip loading RT shaders if we force all shaders to be created (UE-135175).

diff --git a/Engine/Source/Runtime/RenderCore/Private/ShaderCodeArchive.cpp b/Engine/Source/Runtime/RenderCore/Private/ShaderCodeArchive.cpp
--- a/Engine/Source/Runtime/RenderCore/Private/ShaderCodeArchive.cpp
+++ b/Engine/Source/Runtime/RenderCore/Private/ShaderCodeArchive.cpp
@@ -59,6 +59,16 @@ static FAutoConsoleVariableRef CVarShaderCodeLibraryMaxShaderGroupSize(
 	ECVF_RenderThreadSafe | ECVF_ReadOnly
 );
 
+#if RHI_RAYTRACING	// this function is only needed to check if we need to avoid excluding raytracing shaders
+namespace
+{
+	bool IsCreateShadersOnLoadEnabled()
+	{
+		static IConsoleVariable* CVar = IConsoleManager::Get().FindConsoleVariable(TEXT("r.CreateShadersOnLoad"));
+		return CVar && CVar->GetInt() != 0;
+	}
+}
+#endif // RHI_RAYTRACING
 
 int32 FSerializedShaderArchive::FindShaderMapWithKey(const FSHAHash& Hash, uint32 Key) const
 {
@@ -1022,7 +1032,7 @@ bool FShaderCodeArchive::PreloadShaderMap(int32 ShaderMapIndex, FGraphEventArray
 		const FShaderCodeEntry& ShaderEntry = SerializedShaders.ShaderEntries[ShaderIndex];
 
 #if RHI_RAYTRACING
-		if (!IsRayTracingEnabled() && IsRayTracingShaderFrequency(static_cast<EShaderFrequency>(ShaderEntry.Frequency)))
+		if (!IsRayTracingEnabled() && !IsCreateShadersOnLoadEnabled() && IsRayTracingShaderFrequency(static_cast<EShaderFrequency>(ShaderEntry.Frequency)))
 		{
 			ShaderPreloadEntry.bNeverToBePreloaded = 1;
 			continue;
@@ -1809,7 +1819,7 @@ bool FIoStoreShaderCodeArchive::PreloadShaderMap(int32 ShaderMapIndex, FGraphEve
 		const int32 ShaderGroupIndex = GetGroupIndexForShader(ShaderIndex);
 
 #if RHI_RAYTRACING
-		if (!IsRayTracingEnabled() && GroupOnlyContainsRaytracingShaders(ShaderGroupIndex))
+		if (!IsRayTracingEnabled() && !IsCreateShadersOnLoadEnabled() && GroupOnlyContainsRaytracingShaders(ShaderGroupIndex))
 		{
 			MarkPreloadEntrySkipped(ShaderGroupIndex);
 			continue;
@@ -1836,7 +1846,7 @@ bool FIoStoreShaderCodeArchive::PreloadShaderMap(int32 ShaderMapIndex, FCoreDele
 		const int32 ShaderGroupIndex = GetGroupIndexForShader(ShaderIndex);
 
 #if RHI_RAYTRACING
-		if (!IsRayTracingEnabled() && GroupOnlyContainsRaytracingShaders(ShaderGroupIndex))
+		if (!IsRayTracingEnabled() && !IsCreateShadersOnLoadEnabled() && GroupOnlyContainsRaytracingShaders(ShaderGroupIndex))
 		{
 			MarkPreloadEntrySkipped(ShaderGroupIndex);
 			continue;
