Silence spammy warning about uncompressed shaders.

diff --git a/Engine/Source/Runtime/RenderCore/Private/ShaderResource.cpp b/Engine/Source/Runtime/RenderCore/Private/ShaderResource.cpp
--- a/Engine/Source/Runtime/RenderCore/Private/ShaderResource.cpp
+++ b/Engine/Source/Runtime/RenderCore/Private/ShaderResource.cpp
@@ -302,9 +302,10 @@ void FShaderMapResourceCode::AddShaderCode(EShaderFrequency InFrequency, const F
 				
 				// assume uncompressed due to worse ratio than the compression
 				Entry.UncompressedSize = ShaderCode.Num();
-				UE_LOG(LogShaders, Warning, TEXT("Shader %s is expected to be compressed with %s, but it arrived uncompressed. Assuming compressing made it longer and storing uncompressed."),
+				UE_LOG(LogShaders, Verbose, TEXT("Shader %s is expected to be compressed with %s, but it arrived uncompressed (size=%d). Assuming compressing made it longer and storing uncompressed."),
 					*InHash.ToString(),
-					*ShaderCompressionFormat.ToString()
+					*ShaderCompressionFormat.ToString(),
+					ShaderCode.Num()
 				);
 			}
 			else if (ShaderCompressionFormat == NAME_Oodle)
