Also clear AllocatedBufferHashes in FRDGBufferPool::ReleaseDynamicRHI so that both arrays always match.  Keeps the FRDGBufferPool in a reusable state for experiments.

diff --git a/Engine/Source/Runtime/RenderCore/Private/RenderGraphResourcePool.cpp b/Engine/Source/Runtime/RenderCore/Private/RenderGraphResourcePool.cpp
--- a/Engine/Source/Runtime/RenderCore/Private/RenderGraphResourcePool.cpp
+++ b/Engine/Source/Runtime/RenderCore/Private/RenderGraphResourcePool.cpp
@@ -107,6 +107,7 @@ TRefCountPtr<FRDGPooledBuffer> FRDGBufferPool::FindFreeBufferInternal(
 void FRDGBufferPool::ReleaseDynamicRHI()
 {
 	AllocatedBuffers.Empty();
+	AllocatedBufferHashes.Empty();
 }
 
 void FRDGBufferPool::TickPoolElements()
