Backing out changes to remove unknown states in pooled resources. Unported code is still transitioning pooled resources while not updating the tracked state. This will result in incorrect before states.

diff --git a/Engine/Source/Runtime/RenderCore/Private/RenderGraphBuilder.cpp b/Engine/Source/Runtime/RenderCore/Private/RenderGraphBuilder.cpp
--- a/Engine/Source/Runtime/RenderCore/Private/RenderGraphBuilder.cpp
+++ b/Engine/Source/Runtime/RenderCore/Private/RenderGraphBuilder.cpp
@@ -654,8 +654,7 @@ void FRDGBuilder::FinalizeResourceAccess(FRDGTextureAccessArray&& InTextures, FR
 FRDGTextureRef FRDGBuilder::RegisterExternalTexture(
 	const TRefCountPtr<IPooledRenderTarget>& ExternalPooledTexture,
 	ERenderTargetTexture RenderTargetTexture,
-	ERDGTextureFlags Flags,
-	ERHIAccess AccessInitial)
+	ERDGTextureFlags Flags)
 {
 #if RDG_ENABLE_DEBUG
 	checkf(ExternalPooledTexture.IsValid(), TEXT("Attempted to register NULL external texture."));
@@ -666,15 +665,14 @@ FRDGTextureRef FRDGBuilder::RegisterExternalTexture(
 	{
 		Name = TEXT("External");
 	}
-	return RegisterExternalTexture(ExternalPooledTexture, Name, RenderTargetTexture, Flags, AccessInitial);
+	return RegisterExternalTexture(ExternalPooledTexture, Name, RenderTargetTexture, Flags);
 }
 
 FRDGTextureRef FRDGBuilder::RegisterExternalTexture(
 	const TRefCountPtr<IPooledRenderTarget>& ExternalPooledTexture,
 	const TCHAR* Name,
 	ERenderTargetTexture RenderTargetTexture,
-	ERDGTextureFlags Flags,
-	ERHIAccess AccessInitial)
+	ERDGTextureFlags Flags)
 {
 	IF_RDG_ENABLE_DEBUG(UserValidation.ValidateRegisterExternalTexture(ExternalPooledTexture, Name, RenderTargetTexture, Flags));
 	FRHITexture* ExternalTextureRHI = ExternalPooledTexture->GetRenderTargetItem().GetRHI(RenderTargetTexture);
@@ -722,13 +720,11 @@ FRDGTextureRef FRDGBuilder::RegisterExternalTexture(
 		}
 	}
 
-	// The user provided a known initial state. Otherwise, we rely on the tracked state held by the pooled texture.
-	if (AccessInitial != ERHIAccess::Unknown)
-	{
-		FRDGSubresourceState State;
-		State.Access = AccessInitial;
-		InitAsWholeResource(Texture->GetState(), State);
-	}
+	FRDGTextureSubresourceState& TextureState = Texture->GetState();
+
+	checkf(IsWholeResource(TextureState) && GetWholeResource(TextureState).Access == ERHIAccess::Unknown,
+		TEXT("Externally registered texture '%s' has known RDG state. This means the graph did not sanitize it correctly, or ")
+		TEXT("an IPooledRenderTarget reference was improperly held within a pass."), Texture->Name);
 
 	ExternalTextures.Add(Texture->GetRHIUnchecked(), Texture);
 
@@ -737,7 +733,7 @@ FRDGTextureRef FRDGBuilder::RegisterExternalTexture(
 	return Texture;
 }
 
-FRDGBufferRef FRDGBuilder::RegisterExternalBuffer(const TRefCountPtr<FRDGPooledBuffer>& ExternalPooledBuffer, ERDGBufferFlags Flags, ERHIAccess AccessInitial)
+FRDGBufferRef FRDGBuilder::RegisterExternalBuffer(const TRefCountPtr<FRDGPooledBuffer>& ExternalPooledBuffer, ERDGBufferFlags Flags)
 {
 #if RDG_ENABLE_DEBUG
 	checkf(ExternalPooledBuffer.IsValid(), TEXT("Attempted to register NULL external buffer."));
@@ -748,14 +744,13 @@ FRDGBufferRef FRDGBuilder::RegisterExternalBuffer(const TRefCountPtr<FRDGPooledB
 	{
 		Name = TEXT("External");
 	}
-	return RegisterExternalBuffer(ExternalPooledBuffer, Name, Flags, AccessInitial);
+	return RegisterExternalBuffer(ExternalPooledBuffer, Name, Flags);
 }
 
 FRDGBufferRef FRDGBuilder::RegisterExternalBuffer(
 	const TRefCountPtr<FRDGPooledBuffer>& ExternalPooledBuffer,
 	const TCHAR* Name,
-	ERDGBufferFlags Flags,
-	ERHIAccess AccessInitial)
+	ERDGBufferFlags Flags)
 {
 	IF_RDG_ENABLE_DEBUG(UserValidation.ValidateRegisterExternalBuffer(ExternalPooledBuffer, Name, Flags));
 
@@ -787,13 +782,10 @@ FRDGBufferRef FRDGBuilder::RegisterExternalBuffer(
 		Buffer->AccessFinal = ERHIAccess::ReadOnlyExclusiveMask;
 	}
 
-	// The user provided a known initial state. Otherwise, we rely on the tracked state held by the pooled buffer.
-	if (AccessInitial != ERHIAccess::Unknown)
-	{
-		FRDGSubresourceState State;
-		State.Access = AccessInitial;
-		Buffer->GetState() = State;
-	}
+	FRDGSubresourceState& BufferState = Buffer->GetState();
+	checkf(BufferState.Access == ERHIAccess::Unknown,
+		TEXT("Externally registered buffer '%s' has known RDG state. This means the graph did not sanitize it correctly, or ")
+		TEXT("an FRDGPooledBuffer reference was improperly held within a pass."), Buffer->Name);
 
 	ExternalBuffers.Add(ExternalPooledBuffer, Buffer);
 
diff --git a/Engine/Source/Runtime/RenderCore/Private/RenderGraphResourcePool.cpp b/Engine/Source/Runtime/RenderCore/Private/RenderGraphResourcePool.cpp
--- a/Engine/Source/Runtime/RenderCore/Private/RenderGraphResourcePool.cpp
+++ b/Engine/Source/Runtime/RenderCore/Private/RenderGraphResourcePool.cpp
@@ -17,7 +17,9 @@ TRefCountPtr<FRDGPooledBuffer> FRDGBufferPool::FindFreeBuffer(
 	const FRDGBufferDesc& Desc,
 	const TCHAR* InDebugName)
 {
-	return FindFreeBufferInternal(RHICmdList, Desc, InDebugName);
+	TRefCountPtr<FRDGPooledBuffer> Result = FindFreeBufferInternal(RHICmdList, Desc, InDebugName);
+	Result->Reset();
+	return Result;
 }
 
 TRefCountPtr<FRDGPooledBuffer> FRDGBufferPool::FindFreeBufferInternal(
diff --git a/Engine/Source/Runtime/RenderCore/Private/RenderGraphResources.cpp b/Engine/Source/Runtime/RenderCore/Private/RenderGraphResources.cpp
--- a/Engine/Source/Runtime/RenderCore/Private/RenderGraphResources.cpp
+++ b/Engine/Source/Runtime/RenderCore/Private/RenderGraphResources.cpp
@@ -287,7 +287,15 @@ void FRDGTexture::Finalize(FRDGPooledTextureArray& PooledTextureArray)
 		{
 			if (PooledTexture)
 			{
-				PooledTexture->Finalize();
+				// External and extracted resources are user controlled, so we cannot assume the texture stays in its final state.
+				if (bExternal || bExtracted)
+				{
+					PooledTexture->Reset();
+				}
+				else
+				{
+					PooledTexture->Finalize();
+				}
 			}
 
 			if (PooledRenderTarget)
@@ -352,7 +360,15 @@ void FRDGBuffer::Finalize(FRDGPooledBufferArray& PooledBufferArray)
 		}
 		else
 		{
-			PooledBuffer->Finalize();
+			if (bExternal || bExtracted)
+			{
+				PooledBuffer->Reset();
+			}
+			else
+			{
+				PooledBuffer->Finalize();
+			}
+
 			PooledBufferArray.Emplace(PooledBuffer);
 		}
 	}
diff --git a/Engine/Source/Runtime/RenderCore/Private/RenderTargetPool.cpp b/Engine/Source/Runtime/RenderCore/Private/RenderTargetPool.cpp
--- a/Engine/Source/Runtime/RenderCore/Private/RenderTargetPool.cpp
+++ b/Engine/Source/Runtime/RenderCore/Private/RenderTargetPool.cpp
@@ -435,7 +435,27 @@ bool FRenderTargetPool::FindFreeElement(
 		}
 	}
 
-	Out = FindFreeElementInternal(RHICmdList, Desc, InDebugName);
+	TRefCountPtr<FPooledRenderTarget> Result = FindFreeElementInternal(RHICmdList, Desc, InDebugName);
+
+	// Reset RDG state back to an unknown default. The resource is being handed off to a user outside of RDG, so the state is no longer valid.
+	{
+		FRDGPooledTexture* TargetableTexture = Result->TargetableTexture;
+		FRDGPooledTexture* ShaderResourceTexture = Result->ShaderResourceTexture;
+
+		if (TargetableTexture)
+		{
+			checkf(!TargetableTexture->GetOwner(), TEXT("Allocated a pooled render target that is currently owned by RDG texture %s."), TargetableTexture->GetOwner()->Name);
+			TargetableTexture->Reset();
+		}
+
+		if (ShaderResourceTexture && ShaderResourceTexture != TargetableTexture)
+		{
+			checkf(!ShaderResourceTexture->GetOwner(), TEXT("Allocated a pooled render target that is currently owned by RDG texture %s."), ShaderResourceTexture->GetOwner()->Name);
+			ShaderResourceTexture->Reset();
+		}
+	}
+
+	Out = Result;
 	return false;
 }
 
diff --git a/Engine/Source/Runtime/RenderCore/Public/RenderGraphBuilder.h b/Engine/Source/Runtime/RenderCore/Public/RenderGraphBuilder.h
--- a/Engine/Source/Runtime/RenderCore/Public/RenderGraphBuilder.h
+++ b/Engine/Source/Runtime/RenderCore/Public/RenderGraphBuilder.h
@@ -39,29 +39,24 @@ public:
 	FRDGTextureRef RegisterExternalTexture(
 		const TRefCountPtr<IPooledRenderTarget>& ExternalPooledTexture,
 		ERenderTargetTexture Texture = ERenderTargetTexture::ShaderResource,
-		ERDGTextureFlags Flags = ERDGTextureFlags::None,
-		ERHIAccess AccessInitial = ERHIAccess::Unknown);
+		ERDGTextureFlags Flags = ERDGTextureFlags::None);
 
 	/** Register an external texture with a custom name. The name is only used if the texture has not already been registered. */
 	FRDGTextureRef RegisterExternalTexture(
 		const TRefCountPtr<IPooledRenderTarget>& ExternalPooledTexture,
 		const TCHAR* NameIfNotRegistered,
 		ERenderTargetTexture RenderTargetTexture = ERenderTargetTexture::ShaderResource,
-		ERDGTextureFlags Flags = ERDGTextureFlags::None,
-		ERHIAccess AccessInitial = ERHIAccess::Unknown);
+		ERDGTextureFlags Flags = ERDGTextureFlags::None);
 
 	/** Register a external buffer to be tracked by the render graph. */
-	FRDGBufferRef RegisterExternalBuffer(
-		const TRefCountPtr<FRDGPooledBuffer>& ExternalPooledBuffer,
-		ERDGBufferFlags Flags = ERDGBufferFlags::None,
-		ERHIAccess AccessInitial = ERHIAccess::Unknown);
+	FRDGBufferRef RegisterExternalBuffer(const TRefCountPtr<FRDGPooledBuffer>& ExternalPooledBuffer, ERDGBufferFlags Flags = ERDGBufferFlags::None);
+	FRDGBufferRef RegisterExternalBuffer(const TRefCountPtr<FRDGPooledBuffer>& ExternalPooledBuffer, ERDGBufferFlags Flags, ERHIAccess AccessFinal);
 
 	/** Register an external buffer with a custom name. The name is only used if the buffer has not already been registered. */
 	FRDGBufferRef RegisterExternalBuffer(
 		const TRefCountPtr<FRDGPooledBuffer>& ExternalPooledBuffer,
 		const TCHAR* NameIfNotRegistered,
-		ERDGBufferFlags Flags = ERDGBufferFlags::None,
-		ERHIAccess AccessInitial = ERHIAccess::Unknown);
+		ERDGBufferFlags Flags = ERDGBufferFlags::None);
 
 	/** Create graph tracked texture from a descriptor. The CPU memory is guaranteed to be valid through execution of
 	 *  the graph, at which point it is released. The underlying RHI texture lifetime is only guaranteed for passes which
