RealtimeGPUProfiler: Fix crash when toggling views.

Needed to clear description string buffer offsets when the description is empty (or invalid).

#rnx

diff --git a/Engine/Source/Runtime/RenderCore/Private/ProfilingDebugging/RealtimeGPUProfiler.cpp b/Engine/Source/Runtime/RenderCore/Private/ProfilingDebugging/RealtimeGPUProfiler.cpp
--- a/Engine/Source/Runtime/RenderCore/Private/ProfilingDebugging/RealtimeGPUProfiler.cpp
+++ b/Engine/Source/Runtime/RenderCore/Private/ProfilingDebugging/RealtimeGPUProfiler.cpp
@@ -293,6 +293,10 @@ public:
 			DescriptionStringBuffer.AddUninitialized(DescriptionLength);
 			FMemory::Memcpy(&DescriptionStringBuffer[DescriptionOffset], Description, DescriptionLength * sizeof(TCHAR));
 		}
+		else
+		{
+			ClearDescription();
+		}
 	}
 
 	void ClearDescription()
