Removal of RDG drain experiment for UE5 release. It is currently unused and hasn't proven itself necessary yet. It would be relatively straightforward to re-add if the need arises in the future.

diff --git a/Engine/Source/Runtime/RenderCore/Private/RenderGraphBuilder.cpp b/Engine/Source/Runtime/RenderCore/Private/RenderGraphBuilder.cpp
--- a/Engine/Source/Runtime/RenderCore/Private/RenderGraphBuilder.cpp
+++ b/Engine/Source/Runtime/RenderCore/Private/RenderGraphBuilder.cpp
@@ -340,7 +340,7 @@ void FRDGBuilder::SetFlushResourcesRHI()
 
 void FRDGBuilder::BeginFlushResourcesRHI()
 {
-	if (!bFlushResourcesRHI || ExecuteGeneration > 0)
+	if (!bFlushResourcesRHI)
 	{
 		return;
 	}
@@ -352,7 +352,7 @@ void FRDGBuilder::BeginFlushResourcesRHI()
 
 void FRDGBuilder::EndFlushResourcesRHI()
 {
-	if (!bFlushResourcesRHI || ExecuteGeneration > 0)
+	if (!bFlushResourcesRHI)
 	{
 		return;
 	}
@@ -425,11 +425,6 @@ bool FRDGBuilder::IsImmediateMode()
 	return ::IsImmediateMode();
 }
 
-bool FRDGBuilder::IsDrainEnabled()
-{
-	return GRDGDrain != 0;
-}
-
 ERDGPassFlags FRDGBuilder::OverridePassFlags(const TCHAR* PassName, ERDGPassFlags PassFlags, bool bAsyncComputeSupported)
 {
 	const bool bDebugAllowedForPass =
@@ -526,7 +521,6 @@ FRDGBuilder::FRDGBuilder(FRHICommandListImmediate& InRHICmdList, FRDGEventName I
 	, Blackboard(Allocator)
 	, RHICmdListAsyncCompute(FRHICommandListExecutor::GetImmediateAsyncComputeCommandList())
 	, BuilderName(InName)
-	, ProloguePasses(InPlace, nullptr)
 #if RDG_CPU_SCOPES
 	, CPUScopeStacks(Allocator)
 #endif
@@ -809,22 +803,17 @@ void FRDGBuilder::AddPassDependency(FRDGPassHandle ProducerHandle, FRDGPassHandl
 #endif
 }
 
-void FRDGBuilder::Compile(EExecuteMode ExecuteMode)
+void FRDGBuilder::Compile()
 {
 	SCOPE_CYCLE_COUNTER(STAT_RDG_CompileTime);
 	CSV_SCOPED_TIMING_STAT_EXCLUSIVE_CONDITIONAL(RDG_Compile, GRDGVerboseCSVStats != 0);
 	SCOPED_NAMED_EVENT(Compile, FColor::Emerald);
 
-	ExtendPassBitArray(PassesOnAsyncCompute, false);
-
-	const FRDGPassHandle ProloguePassHandle = GetProloguePassHandle();
 	const FRDGPassHandle EpiloguePassHandle = GetEpiloguePassHandle();
 
-	// This will differ from Passes.Num() when drains are performed.
-	const uint32 CompilePassCount = EpiloguePassHandle.GetIndex() - ProloguePassHandle.GetIndex() + 1;
+	const uint32 CompilePassCount = Passes.Num();
 
-	// Culling is only performed in the final execution mode.
-	const bool bCullPasses = GRDGCullPasses && ExecuteMode == EExecuteMode::Final;
+	const bool bCullPasses = GRDGCullPasses > 0;
 
 	TArray<FRDGPassHandle, FRDGArrayAllocator> PassStack;
 
@@ -835,6 +824,8 @@ void FRDGBuilder::Compile(EExecuteMode ExecuteMode)
 
 	TransitionCreateQueue.Reserve(CompilePassCount);
 
+	FRDGPassBitArray PassesOnAsyncCompute(false, CompilePassCount);
+
 	// Build producer / consumer dependencies across the graph and construct packed bit-arrays of metadata
 	// for better cache coherency when searching for passes meeting specific criteria. Search roots are also
 	// identified for culling. Passes with untracked RHI output (e.g. SHADER_PARAMETER_{BUFFER, TEXTURE}_UAV)
@@ -924,31 +915,28 @@ void FRDGBuilder::Compile(EExecuteMode ExecuteMode)
 			}
 		}
 
-		if (ExecuteMode == EExecuteMode::Final)
+		for (const FExtractedTexture& ExtractedTexture : ExtractedTextures)
 		{
-			for (const FExtractedTexture& ExtractedTexture : ExtractedTextures)
+			FRDGTextureRef Texture = ExtractedTexture.Texture;
+			for (auto& LastProducer : Texture->LastProducers)
 			{
-				FRDGTextureRef Texture = ExtractedTexture.Texture;
-				for (auto& LastProducer : Texture->LastProducers)
-				{
-					FRDGProducerState StateFinal;
-					StateFinal.Access = Texture->AccessFinal;
-					StateFinal.PassHandle = EpiloguePassHandle;
+				FRDGProducerState StateFinal;
+				StateFinal.Access = Texture->AccessFinal;
+				StateFinal.PassHandle = EpiloguePassHandle;
 
-					AddCullingDependency(LastProducer, StateFinal, ERHIPipeline::Graphics);
-				}
+				AddCullingDependency(LastProducer, StateFinal, ERHIPipeline::Graphics);
 			}
+		}
 
-			for (const FExtractedBuffer& ExtractedBuffer : ExtractedBuffers)
-			{
-				FRDGBufferRef Buffer = ExtractedBuffer.Buffer;
+		for (const FExtractedBuffer& ExtractedBuffer : ExtractedBuffers)
+		{
+			FRDGBufferRef Buffer = ExtractedBuffer.Buffer;
 
-				FRDGProducerState StateFinal;
-				StateFinal.Access = Buffer->AccessFinal;
-				StateFinal.PassHandle = EpiloguePassHandle;
+			FRDGProducerState StateFinal;
+			StateFinal.Access = Buffer->AccessFinal;
+			StateFinal.PassHandle = EpiloguePassHandle;
 
-				AddCullingDependency(Buffer->LastProducer, StateFinal, ERHIPipeline::Graphics);
-			}
+			AddCullingDependency(Buffer->LastProducer, StateFinal, ERHIPipeline::Graphics);
 		}
 	}
 
@@ -966,12 +954,7 @@ void FRDGBuilder::Compile(EExecuteMode ExecuteMode)
 		EpiloguePass->bCulled = 1;
 
 		// Manually mark the prologue passes as not culled.
-		ProloguePasses[ERHIPipeline::Graphics]->bCulled = 0;
-
-		if (FRDGPass* Pass = ProloguePasses[ERHIPipeline::AsyncCompute])
-		{
-			Pass->bCulled = 0;
-		}
+		ProloguePass->bCulled = 0;
 
 		while (PassStack.Num())
 		{
@@ -1009,7 +992,7 @@ void FRDGBuilder::Compile(EExecuteMode ExecuteMode)
 
 			const auto MergeSubresourceStates = [&](ERDGParentResourceType ResourceType, FRDGSubresourceState*& PassMergeState, FRDGSubresourceState*& ResourceMergeState, const FRDGSubresourceState& PassState)
 			{
-				if (!ResourceMergeState || !FRDGSubresourceState::IsMergeAllowed(ResourceType, *ResourceMergeState, PassState) || ResourceMergeState->ExecuteGeneration != ExecuteGeneration)
+				if (!ResourceMergeState || !FRDGSubresourceState::IsMergeAllowed(ResourceType, *ResourceMergeState, PassState))
 				{
 					// Cross-pipeline, non-mergable state changes require a new pass dependency for fencing purposes.
 					if (ResourceMergeState)
@@ -1405,39 +1388,6 @@ void FRDGBuilder::Compile(EExecuteMode ExecuteMode)
 			}
 		}
 
-		for (FRDGPassHandle PassHandle : PassesOnAsyncComputeToJoin)
-		{
-			const FRDGPassHandle GraphicsJoinPassHandle = FindCrossPipelineConsumer(PassHandle);
-
-			FRDGPass* AsyncComputePass = Passes[PassHandle];
-			FRDGPass* GraphicsJoinPass = Passes[GraphicsJoinPassHandle];
-			FRDGPass* AsyncComputeProloguePass = ProloguePasses[ERHIPipeline::AsyncCompute];
-
-			AsyncComputePass->GraphicsJoinPass = GraphicsJoinPassHandle;
-			AddResourcesToEnd(GraphicsJoinPass, AsyncComputePass);
-
-			if (AsyncComputeProloguePass->GraphicsJoinPass == AsyncComputeProloguePass->Handle)
-			{
-				AsyncComputeProloguePass->GraphicsJoinPass = GraphicsJoinPassHandle;
-			}
-			else
-			{
-				AsyncComputeProloguePass->GraphicsJoinPass = FMath::Min(AsyncComputeProloguePass->GraphicsJoinPass, GraphicsJoinPassHandle);
-			}
-		}
-
-		// Flush the prologue pass fence if it exists.
-		if (!PassesOnAsyncComputeToJoin.IsEmpty())
-		{
-			FRDGPass* AsyncComputePass = ProloguePasses[ERHIPipeline::AsyncCompute];
-			FRDGPass* GraphicsJoinPass = Passes[AsyncComputePass->GraphicsJoinPass];
-
-			check(AsyncComputePass && AsyncComputePass->GraphicsJoinPass && AsyncComputePass->Pipeline == ERHIPipeline::AsyncCompute);
-			InsertAsyncComputeToGraphicsJoin(AsyncComputePass, GraphicsJoinPass);
-		}
-
-		PassesOnAsyncComputeToJoin.Reset();
-
 		FRDGPassHandle CurrentGraphicsJoinPassHandle;
 
 		for (FRDGPassHandle PassHandle = EpiloguePassHandle - 1; PassHandle > ProloguePassHandle; --PassHandle)
@@ -1456,13 +1406,6 @@ void FRDGBuilder::Compile(EExecuteMode ExecuteMode)
 
 			const FRDGPassHandle GraphicsJoinPassHandle = FindCrossPipelineConsumer(PassHandle);
 
-			// When draining and the epilogue pass is found, mark it for the next cycle and continue.
-			if (GraphicsJoinPassHandle == EpiloguePassHandle && ExecuteMode == EExecuteMode::Drain)
-			{
-				PassesOnAsyncComputeToJoin.Add(PassHandle);
-				continue;
-			}
-
 			FRDGPass* GraphicsJoinPass = Passes[GraphicsJoinPassHandle];
 
 			AsyncComputePass->GraphicsJoinPass = GraphicsJoinPassHandle;
@@ -1477,22 +1420,7 @@ void FRDGBuilder::Compile(EExecuteMode ExecuteMode)
 	}
 }
 
-void FRDGBuilder::Drain()
-{
-	if (IsImmediateMode() || !GRDGDrain)
-	{
-		return;
-	}
-
-	Execute(EExecuteMode::Drain);
-}
-
 void FRDGBuilder::Execute()
-{
-	Execute(EExecuteMode::Final);
-}
-
-void FRDGBuilder::Execute(EExecuteMode ExecuteMode)
 {
 	CSV_SCOPED_TIMING_STAT_EXCLUSIVE(RDG);
 	SCOPED_NAMED_EVENT_TEXT("FRDGBuilder::Execute", FColor::Magenta);
@@ -1502,23 +1430,15 @@ void FRDGBuilder::Execute(EExecuteMode ExecuteMode)
 	// Create the epilogue pass at the end of the graph just prior to compilation.
 	{
 		bInDebugPassScope = true;
-		SetupEmptyPass(EpiloguePass = Passes.Allocate<FRDGSentinelPass>(Allocator, RDG_EVENT_NAME("Graph Epilogue%s", (ExecuteMode == EExecuteMode::Drain) ? TEXT(" (Drain)") : TEXT(""))));
+		SetupEmptyPass(EpiloguePass = Passes.Allocate<FRDGSentinelPass>(Allocator, RDG_EVENT_NAME("Graph Epilogue")));
 		bInDebugPassScope = false;
 	}
 
-	const FRDGPassHandle ProloguePassHandle = GetProloguePassHandle();
 	const FRDGPassHandle EpiloguePassHandle = GetEpiloguePassHandle();
 
-	const bool bFirstExecute = ExecuteGeneration == 0;
-	const bool bFinalExecute = ExecuteMode == EExecuteMode::Final;
-
 	FGraphEventArray AsyncCompileEvents;
 
-	if (bFirstExecute)
-	{
-		IF_RDG_ENABLE_DEBUG(UserValidation.ValidateExecuteBegin());
-	}
-
+	IF_RDG_ENABLE_DEBUG(UserValidation.ValidateExecuteBegin());
 	IF_RDG_ENABLE_DEBUG(GRDGAllowRHIAccess = true);
 
 	if (!IsImmediateMode())
@@ -1527,7 +1447,7 @@ void FRDGBuilder::Execute(EExecuteMode ExecuteMode)
 
 		SetupBufferUploads();
 
-		Compile(ExecuteMode);
+		Compile();
 
 		IF_RDG_GPU_SCOPES(GPUScopeStacks.ReserveOps(Passes.Num()));
 		IF_RDG_CPU_SCOPES(CPUScopeStacks.ReserveOps());
@@ -1552,69 +1472,54 @@ void FRDGBuilder::Execute(EExecuteMode ExecuteMode)
 			CSV_SCOPED_TIMING_STAT_EXCLUSIVE(RDG_CollectResources);
 			SCOPED_NAMED_EVENT_TEXT("FRDGBuilder::CollectResources", FColor::Magenta);
 
-			if (bFinalExecute)
+			EnumerateExtendedLifetimeResources(Textures, [](FRDGTexture* Texture)
 			{
-				EnumerateExtendedLifetimeResources(Textures, [](FRDGTexture* Texture)
-				{
-					++Texture->ReferenceCount;
-				});
-
-				EnumerateExtendedLifetimeResources(Buffers, [](FRDGBuffer* Buffer)
-				{
-					++Buffer->ReferenceCount;
-				});
+				++Texture->ReferenceCount;
+			});
 
-				for (FRDGPassHandle PassHandle = Passes.Begin(); PassHandle < ProloguePassHandle; ++PassHandle)
-				{
-					FRDGPass* Pass = Passes[PassHandle];
+			EnumerateExtendedLifetimeResources(Buffers, [](FRDGBuffer* Buffer)
+			{
+				++Buffer->ReferenceCount;
+			});
 
-					if (!Pass->bCulled)
-					{
-						EndResourcesRHI(Pass, ProloguePassHandle);
-					}
-				}
+			for (FRDGPassHandle PassHandle = Passes.Begin(); PassHandle < ProloguePassHandle; ++PassHandle)
+			{
+				FRDGPass* Pass = Passes[PassHandle];
 
-				for (FRDGPassHandle PassHandle = ProloguePassHandle; PassHandle <= EpiloguePassHandle; ++PassHandle)
+				if (!Pass->bCulled)
 				{
-					FRDGPass* Pass = Passes[PassHandle];
-
-					if (!Pass->bCulled)
-					{
-						BeginResourcesRHI(Pass, PassHandle);
-						EndResourcesRHI(Pass, PassHandle);
-					}
+					EndResourcesRHI(Pass, ProloguePassHandle);
 				}
+			}
 
-				EnumerateExtendedLifetimeResources(Textures, [&](FRDGTextureRef Texture)
-				{
-					EndResourceRHI(EpiloguePassHandle, Texture, 1);
-				});
-
-				EnumerateExtendedLifetimeResources(Buffers, [&](FRDGBufferRef Buffer)
-				{
-					EndResourceRHI(EpiloguePassHandle, Buffer, 1);
-				});
+			for (FRDGPassHandle PassHandle = ProloguePassHandle; PassHandle <= EpiloguePassHandle; ++PassHandle)
+			{
+				FRDGPass* Pass = Passes[PassHandle];
 
-				if (TransientResourceAllocator)
+				if (!Pass->bCulled)
 				{
-				#if RDG_ENABLE_TRACE
-					TransientResourceAllocator->Flush(RHICmdList, Trace.IsEnabled() ? &Trace.TransientAllocationStats : nullptr);
-				#else
-					TransientResourceAllocator->Flush(RHICmdList);
-				#endif
+					BeginResourcesRHI(Pass, PassHandle);
+					EndResourcesRHI(Pass, PassHandle);
 				}
 			}
-			else
+
+			EnumerateExtendedLifetimeResources(Textures, [&](FRDGTextureRef Texture)
 			{
-				for (FRDGPassHandle PassHandle = ProloguePassHandle; PassHandle < EpiloguePassHandle; ++PassHandle)
-				{
-					BeginResourcesRHI(Passes[PassHandle], PassHandle);
-				}
+				EndResourceRHI(EpiloguePassHandle, Texture, 1);
+			});
 
-				if (TransientResourceAllocator)
-				{
-					TransientResourceAllocator->Flush(RHICmdList);
-				}
+			EnumerateExtendedLifetimeResources(Buffers, [&](FRDGBufferRef Buffer)
+			{
+				EndResourceRHI(EpiloguePassHandle, Buffer, 1);
+			});
+
+			if (TransientResourceAllocator)
+			{
+			#if RDG_ENABLE_TRACE
+				TransientResourceAllocator->Flush(RHICmdList, Trace.IsEnabled() ? &Trace.TransientAllocationStats : nullptr);
+			#else
+				TransientResourceAllocator->Flush(RHICmdList);
+			#endif
 			}
 		}
 
@@ -1635,7 +1540,6 @@ void FRDGBuilder::Execute(EExecuteMode ExecuteMode)
 		}
 	}
 
-	if (bFinalExecute)
 	{
 		SCOPED_NAMED_EVENT_TEXT("FRDGBuilder::Finalize", FColor::Magenta);
 
@@ -1700,10 +1604,7 @@ void FRDGBuilder::Execute(EExecuteMode ExecuteMode)
 
 	EndFlushResourcesRHI();
 
-	if (bFirstExecute)
-	{
-		IF_RDG_ENABLE_TRACE(Trace.OutputGraphBegin());
-	}
+	IF_RDG_ENABLE_TRACE(Trace.OutputGraphBegin());
 
 	IF_RDG_ENABLE_DEBUG(GRDGAllowRHIAccess = bParallelExecuteEnabled);
 
@@ -1733,11 +1634,6 @@ void FRDGBuilder::Execute(EExecuteMode ExecuteMode)
 			}, TStatId(), nullptr, ENamedThreads::AnyHiPriThreadHiPriTask);
 		}
 
-		if (FRDGPass* AsyncComputeProloguePass = GetProloguePass(ERHIPipeline::AsyncCompute))
-		{
-			ExecutePass(AsyncComputeProloguePass, RHICmdListAsyncCompute);
-		}
-
 		for (FRDGPassHandle PassHandle = ProloguePassHandle; PassHandle <= EpiloguePassHandle; ++PassHandle)
 		{
 			FRDGPass* Pass = Passes[PassHandle];
@@ -1803,74 +1699,66 @@ void FRDGBuilder::Execute(EExecuteMode ExecuteMode)
 		DispatchParallelExecuteEvent->Wait(RenderThread);
 	}
 
-	if (bFinalExecute)
-	{
-		RHICmdList.SetStaticUniformBuffers({});
+	RHICmdList.SetStaticUniformBuffers({});
 
-	#if WITH_MGPU
-		if (NameForTemporalEffect != NAME_None)
+#if WITH_MGPU
+	if (NameForTemporalEffect != NAME_None)
+	{
+		TArray<FRHITexture*> BroadcastTexturesForTemporalEffect;
+		for (const FExtractedTexture& ExtractedTexture : ExtractedTextures)
 		{
-			TArray<FRHITexture*> BroadcastTexturesForTemporalEffect;
-			for (const FExtractedTexture& ExtractedTexture : ExtractedTextures)
+			if (EnumHasAnyFlags(ExtractedTexture.Texture->Flags, ERDGTextureFlags::MultiFrame))
 			{
-				if (EnumHasAnyFlags(ExtractedTexture.Texture->Flags, ERDGTextureFlags::MultiFrame))
-				{
-					BroadcastTexturesForTemporalEffect.Add(ExtractedTexture.Texture->GetRHIUnchecked());
-				}
+				BroadcastTexturesForTemporalEffect.Add(ExtractedTexture.Texture->GetRHIUnchecked());
 			}
-			RHICmdList.BroadcastTemporalEffect(NameForTemporalEffect, BroadcastTexturesForTemporalEffect);
 		}
+		RHICmdList.BroadcastTemporalEffect(NameForTemporalEffect, BroadcastTexturesForTemporalEffect);
+	}
 
-		if (bForceCopyCrossGPU)
-		{
-			ForceCopyCrossGPU();
-		}
-	#endif
+	if (bForceCopyCrossGPU)
+	{
+		ForceCopyCrossGPU();
+	}
+#endif
 
-		// Wait on the actual parallel execute tasks in the Execute call. When draining is okay to let them overlap with other graph setup.
-		// This also needs to be done before extraction of external resources to be consistent with non-parallel rendering.
-		if (!ParallelExecuteEvents.IsEmpty())
-		{
-			FTaskGraphInterface::Get().WaitUntilTasksComplete(ParallelExecuteEvents, RenderThread);
-		}
+	// Wait on the actual parallel execute tasks in the Execute call. When draining is okay to let them overlap with other graph setup.
+	// This also needs to be done before extraction of external resources to be consistent with non-parallel rendering.
+	if (!ParallelExecuteEvents.IsEmpty())
+	{
+		FTaskGraphInterface::Get().WaitUntilTasksComplete(ParallelExecuteEvents, RenderThread);
+	}
 
-		for (const FExtractedTexture& ExtractedTexture : ExtractedTextures)
-		{
-			check(ExtractedTexture.Texture->PooledRenderTarget);
-			*ExtractedTexture.PooledTexture = ExtractedTexture.Texture->PooledRenderTarget;
-		}
+	for (const FExtractedTexture& ExtractedTexture : ExtractedTextures)
+	{
+		check(ExtractedTexture.Texture->PooledRenderTarget);
+		*ExtractedTexture.PooledTexture = ExtractedTexture.Texture->PooledRenderTarget;
+	}
 
-		for (const FExtractedBuffer& ExtractedBuffer : ExtractedBuffers)
-		{
-			check(ExtractedBuffer.Buffer->PooledBuffer);
-			*ExtractedBuffer.PooledBuffer = ExtractedBuffer.Buffer->PooledBuffer;
-		}
+	for (const FExtractedBuffer& ExtractedBuffer : ExtractedBuffers)
+	{
+		check(ExtractedBuffer.Buffer->PooledBuffer);
+		*ExtractedBuffer.PooledBuffer = ExtractedBuffer.Buffer->PooledBuffer;
+	}
 
-		IF_RDG_ENABLE_TRACE(Trace.OutputGraphEnd(*this));
+	IF_RDG_ENABLE_TRACE(Trace.OutputGraphEnd(*this));
 
-		IF_RDG_GPU_SCOPES(GPUScopeStacks.Graphics.EndExecute(RHICmdList));
-		IF_RDG_GPU_SCOPES(GPUScopeStacks.AsyncCompute.EndExecute(RHICmdListAsyncCompute));
-		IF_RDG_CPU_SCOPES(CPUScopeStacks.EndExecute());
+	IF_RDG_GPU_SCOPES(GPUScopeStacks.Graphics.EndExecute(RHICmdList));
+	IF_RDG_GPU_SCOPES(GPUScopeStacks.AsyncCompute.EndExecute(RHICmdListAsyncCompute));
+	IF_RDG_CPU_SCOPES(CPUScopeStacks.EndExecute());
 
-		IF_RDG_ENABLE_DEBUG(UserValidation.ValidateExecuteEnd());
-		IF_RDG_ENABLE_DEBUG(LogFile.End());
-		IF_RDG_ENABLE_DEBUG(GRDGAllowRHIAccess = false);
+	IF_RDG_ENABLE_DEBUG(UserValidation.ValidateExecuteEnd());
+	IF_RDG_ENABLE_DEBUG(LogFile.End());
+	IF_RDG_ENABLE_DEBUG(GRDGAllowRHIAccess = false);
 
-	#if STATS
-		GRDGStatBufferCount += Buffers.Num();
-		GRDGStatTextureCount += Textures.Num();
-		GRDGStatViewCount += Views.Num();
-		GRDGStatMemoryWatermark = FMath::Max(GRDGStatMemoryWatermark, Allocator.GetByteCount());
-	#endif
-	}
-	else
-	{
-		AddProloguePass();
-	}
+#if STATS
+	GRDGStatBufferCount += Buffers.Num();
+	GRDGStatTextureCount += Textures.Num();
+	GRDGStatViewCount += Views.Num();
+	GRDGStatMemoryWatermark = FMath::Max(GRDGStatMemoryWatermark, Allocator.GetByteCount());
+#endif
 
 	RasterPassCount = 0;
 	AsyncComputePassCount = 0;
-	ExecuteGeneration++;
 
 	// Flush any outstanding async compute commands at the end to get things moving down the pipe.
 	if (RHICmdListAsyncCompute.HasCommands())
@@ -1937,7 +1825,6 @@ FRDGPass* FRDGBuilder::SetupPass(FRDGPass* Pass)
 
 		const auto AddSubresourceAccess = [&](FRDGSubresourceState& State)
 		{
-			State.ExecuteGeneration = ExecuteGeneration;
 			State.Access = MakeValidAccess(State.Access | Access);
 			State.Flags |= TransitionFlags;
 			State.NoUAVBarrierFilter.AddHandle(NoUAVBarrierHandle);
@@ -1990,7 +1877,6 @@ FRDGPass* FRDGBuilder::SetupPass(FRDGPass* Pass)
 		}
 
 		PassState->ReferenceCount++;
-		PassState->State.ExecuteGeneration = ExecuteGeneration;
 		PassState->State.Access = MakeValidAccess(PassState->State.Access | Access);
 		PassState->State.NoUAVBarrierFilter.AddHandle(NoUAVBarrierHandle);
 		PassState->State.SetPass(PassPipeline, PassHandle);
@@ -2287,7 +2173,7 @@ void FRDGBuilder::SetupParallelExecute()
 	};
 
 	ParallelPassSets.Reserve(32);
-	ParallelPassCandidates.Emplace(GetProloguePass(ERHIPipeline::Graphics));
+	ParallelPassCandidates.Emplace(ProloguePass);
 
 	for (FRDGPassHandle PassHandle = GetProloguePassHandle() + 1; PassHandle < GetEpiloguePassHandle(); ++PassHandle)
 	{
@@ -2373,9 +2259,8 @@ void FRDGBuilder::DispatchParallelExecute(IRHICommandContext* RHICmdContext)
 	SCOPED_NAMED_EVENT(DispatchParallelExecute, FColor::Emerald);
 	ParallelExecuteEvents.Reserve(ParallelExecuteEvents.Num() + ParallelPassSets.Num());
 
-	for (int32 Index = ParallelPassSetsOffset; Index < ParallelPassSets.Num(); ++Index)
+	for (FParallelPassSet& ParallelPassSet : ParallelPassSets)
 	{
-		FParallelPassSet& ParallelPassSet = ParallelPassSets[Index];
 		ParallelPassSet.RHICmdList = new FRHICommandList(FRHIGPUMask::All());
 		ParallelPassSet.RHICmdList->SetContext(RHICmdContext);
 
@@ -2406,9 +2291,6 @@ void FRDGBuilder::DispatchParallelExecute(IRHICommandContext* RHICmdContext)
 		// Enqueue the event to be synced at the end of RDG execution.
 		ParallelExecuteEvents.Emplace(ParallelPassSet.Event);
 	}
-
-	// Update the offset to avoid reprocessing the same parallel pass sets when drain is enabled.
-	ParallelPassSetsOffset = ParallelPassSets.Num();
 }
 
 void FRDGBuilder::CreateUniformBuffers()
@@ -2425,24 +2307,8 @@ void FRDGBuilder::CreateUniformBuffers()
 void FRDGBuilder::AddProloguePass()
 {
 	bInDebugPassScope = true;
-
-	if (AsyncComputePassCount > 0)
-	{
-		FRDGPass* AsyncComputePass = SetupEmptyPass(Passes.Allocate<FRDGSentinelPass>(Allocator, RDG_EVENT_NAME("Graph Prologue (AsyncCompute)"), ERDGPassFlags::AsyncCompute));
-		ProloguePasses[ERHIPipeline::AsyncCompute] = AsyncComputePass;
-		ProloguePassHandles[ERHIPipeline::AsyncCompute] = AsyncComputePass->Handle;
-	}
-
-	FRDGPass* GraphicsPass = SetupEmptyPass(Passes.Allocate<FRDGSentinelPass>(Allocator, RDG_EVENT_NAME("Graph Prologue (Graphics)")));
-	ProloguePasses[ERHIPipeline::Graphics] = GraphicsPass;
-	ProloguePassHandles[ERHIPipeline::Graphics] = GraphicsPass->Handle;
-
-	// If we don't need an async compute prologue pass, we still want a valid handle to make ClampToPrologue not always return null.
-	if (!AsyncComputePassCount)
-	{
-		ProloguePassHandles[ERHIPipeline::AsyncCompute] = GraphicsPass->Handle;
-	}
-
+	ProloguePass = SetupEmptyPass(Passes.Allocate<FRDGSentinelPass>(Allocator, RDG_EVENT_NAME("Graph Prologue (Graphics)")));
+	ProloguePassHandle = ProloguePass->Handle;
 	bInDebugPassScope = false;
 }
 
diff --git a/Engine/Source/Runtime/RenderCore/Private/RenderGraphPrivate.cpp b/Engine/Source/Runtime/RenderCore/Private/RenderGraphPrivate.cpp
--- a/Engine/Source/Runtime/RenderCore/Private/RenderGraphPrivate.cpp
+++ b/Engine/Source/Runtime/RenderCore/Private/RenderGraphPrivate.cpp
@@ -370,14 +370,6 @@ FAutoConsoleVariableRef CVarRDGIndirectArgBufferTransientAllocated(
 	TEXT("Whether indirect argument buffers should use transient resource allocator. Default: 0"),
 	ECVF_RenderThreadSafe);
 
-int32 GRDGDrain = 0;
-FAutoConsoleVariableRef CVarRDGDrain(
-	TEXT("r.RDG.Drain"), GRDGDrain,
-	TEXT("RDG will perform drain operations when requested.")
-	TEXT(" 0: disables draining, all work is deferred until FRDGBuilder::Execute; (default)")
-	TEXT(" 1: enables execution of the RDG subgraph when FRDGBuilder::Drain is called."),
-	ECVF_RenderThreadSafe);
-
 int32 GRDGParallelExecuteStress = 0;
 FAutoConsoleVariableRef CVarRDGDebugParallelExecute(
 	TEXT("r.RDG.ParallelExecuteStress"),
@@ -522,12 +514,6 @@ void InitRenderGraph()
 		GRDGCullPasses = CullPassesValue;
 	}
 
-	int32 DrainValue = 0;
-	if (FParse::Value(FCommandLine::Get(), TEXT("rdgdrain="), DrainValue))
-	{
-		GRDGDrain = DrainValue;
-	}
-
 	int32 ParallelExecuteValue = 0;
 	if (FParse::Value(FCommandLine::Get(), TEXT("rdgparallelxecute="), ParallelExecuteValue))
 	{
diff --git a/Engine/Source/Runtime/RenderCore/Private/RenderGraphPrivate.h b/Engine/Source/Runtime/RenderCore/Private/RenderGraphPrivate.h
--- a/Engine/Source/Runtime/RenderCore/Private/RenderGraphPrivate.h
+++ b/Engine/Source/Runtime/RenderCore/Private/RenderGraphPrivate.h
@@ -109,7 +109,6 @@ extern int32 GRDGMergeRenderPasses;
 extern int32 GRDGTransientAllocator;
 extern int32 GRDGTransientExtractedResources;
 extern int32 GRDGTransientIndirectArgBuffers;
-extern int32 GRDGDrain;
 extern int32 GRDGParallelExecute;
 extern int32 GRDGParallelExecutePassMin;
 extern int32 GRDGParallelExecutePassMax;
diff --git a/Engine/Source/Runtime/RenderCore/Public/RenderGraphBuilder.h b/Engine/Source/Runtime/RenderCore/Public/RenderGraphBuilder.h
--- a/Engine/Source/Runtime/RenderCore/Public/RenderGraphBuilder.h
+++ b/Engine/Source/Runtime/RenderCore/Public/RenderGraphBuilder.h
@@ -278,9 +278,6 @@ public:
 	/** Manually ends the current GPU event scope. */
 	void EndEventScope();
 
-	/** Drains the RDG queue early. All resource lifetimes are extended to the end of the Drain and culling is disabled for any drained passes. */
-	void Drain();
-
 	/** Executes the queued passes, managing setting of render targets (RHI RenderPasses), resource transitions and queued texture extraction. */
 	void Execute();
 
@@ -290,9 +287,6 @@ public:
 	/** Whether RDG is running in immediate mode. */
 	static bool IsImmediateMode();
 
-	/** Whether RDG will actually perform a drain when Drain() is called. */
-	static bool IsDrainEnabled();
-
 	/** The RHI command list used for the render graph. */
 	FRHICommandListImmediate& RHICmdList;
 
@@ -354,28 +348,28 @@ private:
 
 	void AddProloguePass();
 
-	FORCEINLINE FRDGPassHandle ClampToPrologue(FRDGPassHandle PassHandle, ERHIPipeline Pipeline = ERHIPipeline::Graphics) const
+	FORCEINLINE FRDGPassHandle ClampToPrologue(FRDGPassHandle PassHandle) const
 	{
 		// Preserve null inputs as outputs. Null is the highest value.
-		return ProloguePassHandles[Pipeline].GetIndexUnchecked() > PassHandle.GetIndexUnchecked() ? ProloguePassHandles[Pipeline] : PassHandle;
+		return ProloguePassHandle.GetIndexUnchecked() > PassHandle.GetIndexUnchecked() ? ProloguePassHandle : PassHandle;
 	}
 
 	FORCEINLINE FRDGPassHandlesByPipeline ClampToPrologue(FRDGPassHandlesByPipeline PassHandles) const
 	{
-		PassHandles[ERHIPipeline::Graphics]     = ClampToPrologue(PassHandles[ERHIPipeline::Graphics],     ERHIPipeline::Graphics);
-		PassHandles[ERHIPipeline::AsyncCompute] = ClampToPrologue(PassHandles[ERHIPipeline::AsyncCompute], ERHIPipeline::AsyncCompute);
+		PassHandles[ERHIPipeline::Graphics]     = ClampToPrologue(PassHandles[ERHIPipeline::Graphics]);
+		PassHandles[ERHIPipeline::AsyncCompute] = ClampToPrologue(PassHandles[ERHIPipeline::AsyncCompute]);
 		return PassHandles;
 	}
 
-	FORCEINLINE FRDGPass* GetProloguePass(ERHIPipeline Pipeline = ERHIPipeline::Graphics) const
+	FORCEINLINE FRDGPass* GetProloguePass() const
 	{
-		return ProloguePasses[Pipeline];
+		return ProloguePass;
 	}
 
 	/** Returns the graph prologue pass handle. */
-	FORCEINLINE FRDGPassHandle GetProloguePassHandle(ERHIPipeline Pipeline = ERHIPipeline::Graphics) const
+	FORCEINLINE FRDGPassHandle GetProloguePassHandle() const
 	{
-		return ProloguePassHandles[Pipeline];
+		return ProloguePassHandle;
 	}
 
 	/** Returns the graph epilogue pass handle. */
@@ -456,9 +450,6 @@ private:
 	FRDGViewRegistry Views;
 	FRDGUniformBufferRegistry UniformBuffers;
 
-	FRDGPassBitArray PassesOnAsyncCompute;
-	TArray<FRDGPassHandle, TInlineAllocator<8, FRDGArrayAllocator>> PassesOnAsyncComputeToJoin;
-
 	/** Uniform buffers which were used in a pass. */
 	TArray<FRDGUniformBufferHandle, TInlineAllocator<32, FRDGArrayAllocator>> UniformBuffersToCreate;
 
@@ -478,8 +469,8 @@ private:
 	 *  the graph for culling purposes. The epilogue pass is added to the very end of the pass array for traversal
 	 *  purposes. The prologue does not need to participate in any graph traversal behavior.
 	 */
-	FRDGPassHandlesByPipeline ProloguePassHandles;
-	FRDGPassesByPipeline ProloguePasses;
+	FRDGPassHandle ProloguePassHandle;
+	FRDGPass* ProloguePass = nullptr;
 	FRDGPass* EpiloguePass = nullptr;
 
 	struct FExtractedTexture
@@ -576,7 +567,6 @@ private:
 	};
 
 	TArray<FParallelPassSet, FRDGArrayAllocator> ParallelPassSets;
-	uint32 ParallelPassSetsOffset = 0;
 
 	/** Array of all active parallel execute tasks. */
 	FGraphEventArray ParallelExecuteEvents;
@@ -624,21 +614,13 @@ private:
 
 	uint32 AsyncComputePassCount = 0;
 	uint32 RasterPassCount = 0;
-	uint32 ExecuteGeneration = 0;
 
 	IF_RDG_CMDLIST_STATS(TStatId CommandListStatScope);
 	IF_RDG_CMDLIST_STATS(TStatId CommandListStatState);
 
 	IRHITransientResourceAllocator* TransientResourceAllocator = nullptr;
 
-	enum class EExecuteMode
-	{
-		Drain,
-		Final
-	};
-
-	void Execute(EExecuteMode ExecuteMode);
-	void Compile(EExecuteMode ExecuteMode);
+	void Compile();
 	void Clear();
 
 	void BeginResourcesRHI(FRDGPass* ResourcePass, FRDGPassHandle ExecutePassHandle);
@@ -712,11 +694,6 @@ private:
 
 	FRDGSubresourceState* AllocSubresource(const FRDGSubresourceState& Other);
 
-	inline void ExtendPassBitArray(FRDGPassBitArray& BitArray, bool bValue) const
-	{
-		BitArray.Insert(bValue, BitArray.Num(), Passes.Num() - BitArray.Num());
-	}
-
 #if RDG_DUMP_RESOURCES
 	void DumpResourcePassOutputs(const FRDGPass* Pass);
 
diff --git a/Engine/Source/Runtime/RenderCore/Public/RenderGraphResources.h b/Engine/Source/Runtime/RenderCore/Public/RenderGraphResources.h
--- a/Engine/Source/Runtime/RenderCore/Public/RenderGraphResources.h
+++ b/Engine/Source/Runtime/RenderCore/Public/RenderGraphResources.h
@@ -73,9 +73,6 @@ struct FRDGSubresourceState
 	/** The last no-UAV barrier to be used by this subresource. */
 	FRDGViewUniqueFilter NoUAVBarrierFilter;
 
-	/** The graph execution generation to submit this subresource state. */
-	uint8 ExecuteGeneration = 0;
-
 	/** The last used pass for log file debugging. */
 	IF_RDG_ENABLE_DEBUG(mutable FRDGPassHandle LogFilePass);
 };
@@ -272,13 +269,6 @@ public:
 		return bProduced;
 	}
 
-	/** Marks a resource as excluded from the transient allocator, if it would otherwise be included. */
-	void SetNonTransient()
-	{
-		check(!bTransient);
-		bForceNonTransient = 1;
-	}
-
 protected:
 	FRDGParentResource(const TCHAR* InName, ERDGParentResourceType InType);
 
diff --git a/Engine/Source/Runtime/Renderer/Private/DeferredShadingRenderer.cpp b/Engine/Source/Runtime/Renderer/Private/DeferredShadingRenderer.cpp
--- a/Engine/Source/Runtime/Renderer/Private/DeferredShadingRenderer.cpp
+++ b/Engine/Source/Runtime/Renderer/Private/DeferredShadingRenderer.cpp
@@ -2341,23 +2341,7 @@ void FDeferredShadingSceneRenderer::Render(FRDGBuilder& GraphBuilder)
 		}
 	}
 
-	/**
-	 * Helper to perform RDG drain (essentially Execute() for the queued passes so far), and do any tasks that need doing at the start of a new 
-	 * section of RDG setup (currently, restarting deferred culling).
-	 */
-	auto ExecuteDrainPoint = [this, &GraphBuilder, &InstanceCullingManager]()
 	{
-		GraphBuilder.Drain();
-		if (GraphBuilder.IsDrainEnabled())
-		{
-			// This sets up a fresh deferred context that is used by all subsequent culling passes since the previous one is 'used up' by Drain.
-			InstanceCullingManager.BeginDeferredCulling(GraphBuilder, Scene->GPUScene);
-		}
-	};
-
-	{
-		ExecuteDrainPoint();
-
 		{
 			RDG_RHI_GPU_STAT_SCOPE(GraphBuilder, VisibilityCommands);
 			InitViewsAfterPrepass(GraphBuilder, ILCTaskData, InstanceCullingManager);
@@ -2536,8 +2520,6 @@ void FDeferredShadingSceneRenderer::Render(FRDGBuilder& GraphBuilder)
 		RenderOcclusion(GraphBuilder, SceneTextures, bIsOcclusionTesting);
 
 		CompositionLighting.ProcessAfterOcclusion(GraphBuilder);
-
-		ExecuteDrainPoint();
 	};
 
 	// Early occlusion queries
