Dumps all cvars in DumpGPU captures

diff --git a/Engine/Source/Runtime/RenderCore/Private/RenderGraphDumpResources.cpp b/Engine/Source/Runtime/RenderCore/Private/RenderGraphDumpResources.cpp
--- a/Engine/Source/Runtime/RenderCore/Private/RenderGraphDumpResources.cpp
+++ b/Engine/Source/Runtime/RenderCore/Private/RenderGraphDumpResources.cpp
@@ -405,9 +405,7 @@ struct FRDGResourceDumpContext
 
 			if (GDumpRenderingConsoleVariablesCVar.GetValueOnGameThread() != 0)
 			{
-				IConsoleManager::Get().ForEachConsoleObjectThatStartsWith(FConsoleObjectVisitor::CreateLambda(OnConsoleVariable), TEXT("r."));
-				IConsoleManager::Get().ForEachConsoleObjectThatStartsWith(FConsoleObjectVisitor::CreateLambda(OnConsoleVariable), TEXT("sg."));
-				IConsoleManager::Get().ForEachConsoleObjectThatContains(FConsoleObjectVisitor::CreateLambda(OnConsoleVariable), TEXT("rhi"));
+				IConsoleManager::Get().ForEachConsoleObjectThatStartsWith(FConsoleObjectVisitor::CreateLambda(OnConsoleVariable), TEXT(""));
 			}
 			else
 			{
