Add ERDGPassFlags::NeverParallel to allow RDG passes to stay on render thread.

Both Immediate/AsyncCompute passes are always executed in order on the render thread but other RDG passes can run in parallel.
With my previous change, RayTracingScene pass was changed to Compute/AsyncCompute to have the correct order on the GPU.
But when it was set to Compute (when AsyncCompute was disabled) it was running in parallel with passes it depends on which led to random crashes.
Because the dependency between passes was implicit, we needed a new flag to express the intent and get the desired order between execution lambdas.

On behalf of Zach Bethel.

diff --git a/Engine/Plugins/RenderGraphInsights/Source/RenderGraphInsights/Private/RenderGraphProvider.cpp b/Engine/Plugins/RenderGraphInsights/Source/RenderGraphInsights/Private/RenderGraphProvider.cpp
--- a/Engine/Plugins/RenderGraphInsights/Source/RenderGraphInsights/Private/RenderGraphProvider.cpp
+++ b/Engine/Plugins/RenderGraphInsights/Source/RenderGraphInsights/Private/RenderGraphProvider.cpp
@@ -138,7 +138,7 @@ FPassPacket::FPassPacket(const UE::Trace::IAnalyzer::FOnEventContext& Context)
 	, bParallelExecuteBegin(Context.EventData.GetValue<bool>("IsParallelExecuteBegin"))
 	, bParallelExecuteEnd(Context.EventData.GetValue<bool>("IsParallelExecuteEnd"))
 	, bParallelExecute(Context.EventData.GetValue<bool>("IsParallelExecute"))
-	, bImmediateCommandList(Context.EventData.GetValue<bool>("UsesImmediateCommandList"))
+	, bParallelExecuteAllowed(Context.EventData.GetValue<bool>("IsParallelExecuteAllowed"))
 {}
 
 static const uint64 PageSize = 1024;
diff --git a/Engine/Plugins/RenderGraphInsights/Source/RenderGraphInsights/Private/RenderGraphProvider.h b/Engine/Plugins/RenderGraphInsights/Source/RenderGraphInsights/Private/RenderGraphProvider.h
--- a/Engine/Plugins/RenderGraphInsights/Source/RenderGraphInsights/Private/RenderGraphProvider.h
+++ b/Engine/Plugins/RenderGraphInsights/Source/RenderGraphInsights/Private/RenderGraphProvider.h
@@ -139,7 +139,7 @@ public:
 	bool bParallelExecuteBegin{};
 	bool bParallelExecuteEnd{};
 	bool bParallelExecute{};
-	bool bImmediateCommandList{};
+	bool bParallelExecuteAllowed{};
 
 	FPassPacket(const UE::Trace::IAnalyzer::FOnEventContext& Context);
 };
diff --git a/Engine/Plugins/RenderGraphInsights/Source/RenderGraphInsights/Private/RenderGraphTrack.cpp b/Engine/Plugins/RenderGraphInsights/Source/RenderGraphInsights/Private/RenderGraphTrack.cpp
--- a/Engine/Plugins/RenderGraphInsights/Source/RenderGraphInsights/Private/RenderGraphTrack.cpp
+++ b/Engine/Plugins/RenderGraphInsights/Source/RenderGraphInsights/Private/RenderGraphTrack.cpp
@@ -2271,19 +2271,18 @@ void FRenderGraphTrack::InitTooltip(FTooltipDrawState& Tooltip, const ITimingEve
 				Tooltip.AddNameValueTextLine(TEXT("Used Buffers:"), FString::Printf(TEXT("%d"), Pass.Buffers.Num()));
 			}
 
-			if (Pass.bParallelExecute)
+			if (Pass.bParallelExecuteAllowed)
 			{
-				Tooltip.AddTextLine(TEXT("Parallel Execute"), FLinearColor::Red);
+				Tooltip.AddNameValueTextLine(TEXT("Parallel Execute:"), Pass.bParallelExecute ? TEXT("Yes") : TEXT("No"));
 			}
-
-			if (Pass.bSkipRenderPassBegin || Pass.bSkipRenderPassEnd)
+			else
 			{
-				Tooltip.AddTextLine(TEXT("Merged RenderPass"), FLinearColor::Red);
+				Tooltip.AddNameValueTextLine(TEXT("Parallel Execute:"), TEXT("Not Allowed"));
 			}
 
-			if (Pass.bImmediateCommandList)
+			if (Pass.bSkipRenderPassBegin || Pass.bSkipRenderPassEnd)
 			{
-				Tooltip.AddTextLine(TEXT("Immediate Command List"), FLinearColor::Red);
+				Tooltip.AddTextLine(TEXT("Merged RenderPass"), FLinearColor::Red);
 			}
 		}
 		else if (InTooltipEvent.Is<FVisibleTextureEvent>())
diff --git a/Engine/Source/Runtime/RenderCore/Private/RenderGraphBuilder.cpp b/Engine/Source/Runtime/RenderCore/Private/RenderGraphBuilder.cpp
--- a/Engine/Source/Runtime/RenderCore/Private/RenderGraphBuilder.cpp
+++ b/Engine/Source/Runtime/RenderCore/Private/RenderGraphBuilder.cpp
@@ -2194,7 +2194,7 @@ void FRDGBuilder::SetupParallelExecute()
 			continue;
 		}
 
-		if (Pass->bImmediateCommandList)
+		if (!Pass->bParallelExecuteAllowed)
 		{
 			FlushParallelPassCandidates();
 			continue;
diff --git a/Engine/Source/Runtime/RenderCore/Private/RenderGraphEvent.cpp b/Engine/Source/Runtime/RenderCore/Private/RenderGraphEvent.cpp
--- a/Engine/Source/Runtime/RenderCore/Private/RenderGraphEvent.cpp
+++ b/Engine/Source/Runtime/RenderCore/Private/RenderGraphEvent.cpp
@@ -269,7 +269,7 @@ FRDGGPUStatScopeOpArray FRDGGPUStatScopeStack::CompilePassPrologue(const FRDGPas
 	if (IsEnabled() && Pass->GetPipeline() == ERHIPipeline::Graphics)
 	{
 		FRDGGPUStatScopeOpArray Ops(ScopeStack.CompilePassPrologue(Pass->GetGPUScopes().Stat), GPUMask);
-		if (Pass->IsImmediateCommandList())
+		if (!Pass->IsParallelExecuteAllowed())
 		{
 			OverrideEventIndex = FRealtimeGPUProfiler::Get()->GetCurrentEventIndex();
 			Ops.OverrideEventIndex = OverrideEventIndex;
diff --git a/Engine/Source/Runtime/RenderCore/Private/RenderGraphTrace.cpp b/Engine/Source/Runtime/RenderCore/Private/RenderGraphTrace.cpp
--- a/Engine/Source/Runtime/RenderCore/Private/RenderGraphTrace.cpp
+++ b/Engine/Source/Runtime/RenderCore/Private/RenderGraphTrace.cpp
@@ -41,7 +41,7 @@ UE_TRACE_EVENT_BEGIN(RDGTrace, PassMessage)
 	UE_TRACE_EVENT_FIELD(bool, IsParallelExecuteBegin)
 	UE_TRACE_EVENT_FIELD(bool, IsParallelExecuteEnd)
 	UE_TRACE_EVENT_FIELD(bool, IsParallelExecute)
-	UE_TRACE_EVENT_FIELD(bool, UsesImmediateCommandList)
+	UE_TRACE_EVENT_FIELD(bool, IsParallelExecuteAllowed)
 UE_TRACE_EVENT_END()
 
 UE_TRACE_EVENT_BEGIN(RDGTrace, BufferMessage)
@@ -188,7 +188,7 @@ void FRDGTrace::OutputGraphEnd(const FRDGBuilder& GraphBuilder)
 			<< PassMessage.IsParallelExecuteBegin(Pass->bParallelExecuteBegin != 0)
 			<< PassMessage.IsParallelExecuteEnd(Pass->bParallelExecuteEnd != 0)
 			<< PassMessage.IsParallelExecute(Pass->bParallelExecute != 0)
-			<< PassMessage.UsesImmediateCommandList(Pass->bImmediateCommandList != 0);
+			<< PassMessage.IsParallelExecuteAllowed(Pass->bParallelExecuteAllowed != 0);
 	}
 
 #if RDG_EVENTS
diff --git a/Engine/Source/Runtime/RenderCore/Public/RenderGraphDefinitions.h b/Engine/Source/Runtime/RenderCore/Public/RenderGraphDefinitions.h
--- a/Engine/Source/Runtime/RenderCore/Public/RenderGraphDefinitions.h
+++ b/Engine/Source/Runtime/RenderCore/Public/RenderGraphDefinitions.h
@@ -115,6 +115,9 @@ enum class ERDGPassFlags : uint8
 	/** Pass will never have its render pass merged with other passes. */
 	NeverMerge = 1 << 6,
 
+	/** Pass will never run off the render thread. */
+	NeverParallel = 1 << 7,
+
 	/** Pass uses copy commands but writes to a staging resource. */
 	Readback = Copy | NeverCull
 };
diff --git a/Engine/Source/Runtime/RenderCore/Public/RenderGraphPass.h b/Engine/Source/Runtime/RenderCore/Public/RenderGraphPass.h
--- a/Engine/Source/Runtime/RenderCore/Public/RenderGraphPass.h
+++ b/Engine/Source/Runtime/RenderCore/Public/RenderGraphPass.h
@@ -194,9 +194,9 @@ public:
 		return Handle;
 	}
 
-	bool IsImmediateCommandList() const
+	bool IsParallelExecuteAllowed() const
 	{
-		return bImmediateCommandList;
+		return bParallelExecuteAllowed;
 	}
 
 	bool IsMergedRenderPassBegin() const
@@ -362,8 +362,8 @@ protected:
 			/** Whether the pass only writes to resources in its render pass. */
 			uint32 bRenderPassOnlyWrites : 1;
 
-			/** Whether the pass uses the immediate command list. */
-			uint32 bImmediateCommandList : 1;
+			/** Whether the pass is allowed to execute in parallel. */
+			uint32 bParallelExecuteAllowed : 1;
 
 			/** Whether this pass has non-RDG UAV outputs. */
 			uint32 bHasExternalOutputs : 1;
@@ -552,7 +552,7 @@ public:
 		checkf(kSupportsAsyncCompute || !EnumHasAnyFlags(InPassFlags, ERDGPassFlags::AsyncCompute),
 			TEXT("Pass %s is set to use 'AsyncCompute', but the pass lambda's first argument is not FRHIComputeCommandList&."), GetName());
 
-		bImmediateCommandList = TIsSame<TRHICommandList, FRHICommandListImmediate>::Value;
+		bParallelExecuteAllowed = !TIsSame<TRHICommandList, FRHICommandListImmediate>::Value && !EnumHasAnyFlags(InPassFlags, ERDGPassFlags::NeverParallel);
 	}
 
 private:
diff --git a/Engine/Source/Runtime/Renderer/Private/DeferredShadingRenderer.cpp b/Engine/Source/Runtime/Renderer/Private/DeferredShadingRenderer.cpp
--- a/Engine/Source/Runtime/Renderer/Private/DeferredShadingRenderer.cpp
+++ b/Engine/Source/Runtime/Renderer/Private/DeferredShadingRenderer.cpp
@@ -1539,7 +1539,9 @@ bool FDeferredShadingSceneRenderer::DispatchRayTracingWorldUpdates(FRDGBuilder&
 		PassParams->RayTracingSceneInstanceBuffer = Scene->RayTracingScene.InstanceBuffer;
 		PassParams->DynamicGeometryScratchBuffer = OutDynamicGeometryScratchBuffer;
 
-		GraphBuilder.AddPass(RDG_EVENT_NAME("RayTracingScene"), PassParams, ComputePassFlags | ERDGPassFlags::NeverCull,
+		// Use ERDGPassFlags::NeverParallel so the pass never runs off the render thread and we always get the following order of execution on the CPU:
+		// BuildTLASInstanceBuffer, RayTracingScene, EndUpdate, ..., ReleaseRayTracingResources		
+		GraphBuilder.AddPass(RDG_EVENT_NAME("RayTracingScene"), PassParams, ComputePassFlags | ERDGPassFlags::NeverCull | ERDGPassFlags::NeverParallel,
 			[this, PassParams, bRayTracingAsyncBuild](FRHIComputeCommandList& RHICmdList)
 		{
 			FRHIBuffer* DynamicGeometryScratchBuffer = PassParams->DynamicGeometryScratchBuffer ? PassParams->DynamicGeometryScratchBuffer->GetRHI() : nullptr;
