Fix non-deterministic sorting in iostore shader group compression logic.
Replace the map with an array for efficiency.

#rnx

diff --git a/Engine/Source/Runtime/RenderCore/Private/ShaderCodeArchive.cpp b/Engine/Source/Runtime/RenderCore/Private/ShaderCodeArchive.cpp
--- a/Engine/Source/Runtime/RenderCore/Private/ShaderCodeArchive.cpp
+++ b/Engine/Source/Runtime/RenderCore/Private/ShaderCodeArchive.cpp
@@ -1278,9 +1278,8 @@ void FIoStoreShaderCodeArchive::CreateIoStoreShaderCodeArchiveHeader(const FName
 	// a single group) is added to the header.
 	// Each group's indices, like in case of shadermaps, are stored in ShaderIndices array. Before we append a new group's indices however, we look if we can find an existing range that we can reuse.
 
-	TMap<uint32, TArray<int32>> ShaderToShadermaps;
+	TArray<TPair<uint32, TArray<int32>>> ShaderToShadermapsArray;
 	{
-		TArray<TArray<int32>> ShaderToShadermapsArray;
 		ShaderToShadermapsArray.AddDefaulted(OutHeader.ShaderEntries.Num());
 
 		{
@@ -1296,7 +1295,7 @@ void FIoStoreShaderCodeArchive::CreateIoStoreShaderCodeArchiveHeader(const FName
 						// add this shadermap as a dependency.
 						int32 ShaderLockNumber = ShaderIndex % UE_ARRAY_COUNT(ShaderLocks);
 						FScopeLock Locker(&ShaderLocks[ShaderLockNumber]);
-						ShaderToShadermapsArray[ShaderIndex].Add(ShaderMapIndex);
+						ShaderToShadermapsArray[ShaderIndex].Value.Add(ShaderMapIndex);
 					}
 				}
 			);
@@ -1313,38 +1312,40 @@ void FIoStoreShaderCodeArchive::CreateIoStoreShaderCodeArchiveHeader(const FName
 					int32 EndingShader = FMath::Min(StartingShader + kShaderSortedPerThread, ShaderToShadermapsArray.Num());
 					for (int32 Idx = StartingShader; Idx < EndingShader; ++Idx)
 					{
-						ShaderToShadermapsArray[Idx].Sort();
+						ShaderToShadermapsArray[Idx].Value.Sort();
 					}
 				},
 				EParallelForFlags::Unbalanced
 			);
 		}
 
-		// now add all this to map so we can sort it
+		// now assigning the indices in the array so we can sort it
 		for (int32 Idx = 0, Num = ShaderToShadermapsArray.Num(); Idx < Num; ++Idx)
 		{
 			// check that no shader is unreferenced
-			checkf(!ShaderToShadermapsArray[Idx].IsEmpty(), TEXT("Error converting to IoStore archive: shader (index=%d) is not referenced by any of the shadermaps!"), Idx);
-			ShaderToShadermaps.Add(Idx, ShaderToShadermapsArray[Idx]);
+			checkf(!ShaderToShadermapsArray[Idx].Value.IsEmpty(), TEXT("Error converting to IoStore archive: shader (index=%d) is not referenced by any of the shadermaps!"), Idx);
+			ShaderToShadermapsArray[Idx].Key = Idx;
 		}
 	}
 
-	// sort the mapping so the first are shaders that are referenced by a smaller number of shadermaps.
-	ShaderToShadermaps.ValueSort(
-		[](const TArray<int32>& A, const TArray<int32>& B)
+	// sort the mapping so the first are shaders that are referenced by a smaller number of shadermaps, then by index for determinism
+	Algo::Sort(ShaderToShadermapsArray,
+		[](const TPair<uint32, TArray<int32>>& EntryA, const TPair<uint32, TArray<int32>>& EntryB)
 		{
+			const TArray<int32>& A = EntryA.Value;
+			const TArray<int32>& B = EntryB.Value;
 			// if the number of shadermaps is the same, we need to sort "alphabetically"
 			if (A.Num() == B.Num())
 			{
 				for (int32 Idx = 0, Num = A.Num(); Idx < Num; ++Idx)
 				{
-					if (A[Idx] < B[Idx])
+					if (A[Idx] != B[Idx])
 					{
-						return true;
+						return A[Idx] < B[Idx];
 					}
 				}
 
-				return false;
+				return EntryA.Key < EntryB.Key;
 			}
 
 			return A.Num() < B.Num();
@@ -1510,27 +1511,27 @@ void FIoStoreShaderCodeArchive::CreateIoStoreShaderCodeArchiveHeader(const FName
 	// now split this into streaks of shaders that are referenced by the same set of shadermaps and compress
 	TArray<uint32> CurrentShaderGroup;
 	TArray<int32> LastShadermapSetSeen;
-	for (TMap<uint32, TArray<int32>>::TConstIterator Iter(ShaderToShadermaps); Iter; ++Iter)
+	for (TPair<uint32, TArray<int32>>& Iter : ShaderToShadermapsArray)
 	{
 		// if we have have just started the group, we don't check the last seen
 		if (UNLIKELY(CurrentShaderGroup.IsEmpty()))
 		{
-			CurrentShaderGroup.Add(Iter.Key());
-			LastShadermapSetSeen = Iter.Value();
+			CurrentShaderGroup.Add(Iter.Key);
+			LastShadermapSetSeen = Iter.Value;
 		}
-		else if (UNLIKELY(LastShadermapSetSeen != Iter.Value()))
+		else if (UNLIKELY(LastShadermapSetSeen != Iter.Value))
 		{
 			ProcessShaderGroup_SplitRaytracing(CurrentShaderGroup);
 
 			// reset the collection, but don't forget to add to it the current element
 			CurrentShaderGroup.SetNum(1);
-			CurrentShaderGroup[0] = Iter.Key();
-			LastShadermapSetSeen = Iter.Value();
+			CurrentShaderGroup[0] = Iter.Key;
+			LastShadermapSetSeen = Iter.Value;
 		}
 		else
 		{
 			// keep adding to the same group
-			CurrentShaderGroup.Add(Iter.Key());
+			CurrentShaderGroup.Add(Iter.Key);
 		}
 	}
 	// add the last group
