Do not attempt to split 1 shader groups even if they are larger than the threshold (UE-141903).

diff --git a/Engine/Source/Runtime/RenderCore/Private/ShaderCodeArchive.cpp b/Engine/Source/Runtime/RenderCore/Private/ShaderCodeArchive.cpp
--- a/Engine/Source/Runtime/RenderCore/Private/ShaderCodeArchive.cpp
+++ b/Engine/Source/Runtime/RenderCore/Private/ShaderCodeArchive.cpp
@@ -1414,7 +1414,7 @@ void FIoStoreShaderCodeArchive::CreateIoStoreShaderCodeArchiveHeader(const FName
 			GroupSize += SerializedShaders.ShaderEntries[ShaderIdx].UncompressedSize;
 		}
 		
-		if (LIKELY(GroupSize <= MaxUncompressedShaderGroupSize))
+		if (LIKELY(GroupSize <= MaxUncompressedShaderGroupSize || CurrentShaderGroup.Num() == 1))
 		{
 			ProcessShaderGroup_AddNewGroup(CurrentShaderGroup);
 		}
@@ -1422,7 +1422,7 @@ void FIoStoreShaderCodeArchive::CreateIoStoreShaderCodeArchiveHeader(const FName
 		{
 			// split the shaders evenly into N new groups (don't allow more new groups than there are shaders)
 			int32 NumNewGroups = FMath::Min(static_cast<int32>(GroupSize / MaxUncompressedShaderGroupSize + 1), CurrentShaderGroup.Num());
-			checkf(NumNewGroups > 1, TEXT("Off by one error in group count calculation?"));
+			checkf(NumNewGroups > 1, TEXT("Off by one error in group count calculation? NumNewGroups=%d, GroupSize=%u, MaxUncompressedShaderGroupSize=%u, CurrentShaderGroup.Num()=%d"), NumNewGroups, GroupSize, MaxUncompressedShaderGroupSize, CurrentShaderGroup.Num());
 
 			TArray<TArray<uint32>> NewGroups;
 			TArray<uint32> NewGroupSizes;
