Adds a UE-DumpGPUPath environment variables to dump across perforce workspaces at the same location by default on a dev machine

diff --git a/Engine/Extras/GPUDumpViewer/GPUDumpViewer.html b/Engine/Extras/GPUDumpViewer/GPUDumpViewer.html
--- a/Engine/Extras/GPUDumpViewer/GPUDumpViewer.html
+++ b/Engine/Extras/GPUDumpViewer/GPUDumpViewer.html
@@ -1039,7 +1039,7 @@ const k_tips = [
 	// Dumping process
 	'Can use the CTRL+SHIFT+/ keyboard shortcut to summon the DumpGPU command.',
 	'Speed up your frame dump by only selecting the passes you need with r.DumpGPU.Root. For instance r.DumpGPU.Root="*PostProcessing*".',
-	'GPU dumps can be large and accumulate on your hard drive in your various projects\' Saved/ directories. Set r.DumpGPU.Directory="D:/tmp/DumpGPU/" in your console variables to dump them all at the same location.',
+	'GPU dumps can be large and accumulate on your hard drive in your various projects\' Saved/ directories. Set r.DumpGPU.Directory="D:/tmp/DumpGPU/" in your console variables or UE-DumpGPUPath environment variable to dump them all at the same location on your machine.',
 	'Uses -cvarsini to override console variables with your own ConsoleVariables.ini file on a cooked build.',
 
 	// General navigations
diff --git a/Engine/Source/Runtime/RenderCore/Private/RenderGraphDumpResources.cpp b/Engine/Source/Runtime/RenderCore/Private/RenderGraphDumpResources.cpp
--- a/Engine/Source/Runtime/RenderCore/Private/RenderGraphDumpResources.cpp
+++ b/Engine/Source/Runtime/RenderCore/Private/RenderGraphDumpResources.cpp
@@ -1226,12 +1226,23 @@ FString FRDGBuilder::BeginResourceDump(const TArray<FString>& Args)
 
 	NewResourceDumpContext.Time = FDateTime::Now();
 	{
-		FString MasterDirectoryPath = GDumpGPUDirectoryCVar.GetValueOnGameThread();
-		if (MasterDirectoryPath.IsEmpty())
+		FString CVarDirectoryPath = GDumpGPUDirectoryCVar.GetValueOnGameThread();
+		FString EnvDirectoryPath = FPlatformMisc::GetEnvironmentVariable(TEXT("UE-DumpGPUPath"));
+
+		FString DirectoryPath;
+		if (!CVarDirectoryPath.IsEmpty())
+		{
+			DirectoryPath = CVarDirectoryPath;
+		}
+		else if (!EnvDirectoryPath.IsEmpty())
+		{
+			DirectoryPath = EnvDirectoryPath;
+		}
+		else
 		{
-			MasterDirectoryPath = FPaths::ProjectSavedDir() / TEXT("GPUDumps/");
+			DirectoryPath = FPaths::ProjectSavedDir() / TEXT("GPUDumps/");
 		}
-		NewResourceDumpContext.DumpingDirectoryPath = MasterDirectoryPath / FApp::GetProjectName() + TEXT("-") + FPlatformProperties::PlatformName() + TEXT("-") + NewResourceDumpContext.Time.ToString() + TEXT("/");
+		NewResourceDumpContext.DumpingDirectoryPath = DirectoryPath / FApp::GetProjectName() + TEXT("-") + FPlatformProperties::PlatformName() + TEXT("-") + NewResourceDumpContext.Time.ToString() + TEXT("/");
 	}
 	NewResourceDumpContext.bEnableDiskWrite = GDumpTestEnableDiskWrite.GetValueOnGameThread() != 0;
 	NewResourceDumpContext.bShowInExplore = NewResourceDumpContext.bEnableDiskWrite && GDumpExploreCVar.GetValueOnGameThread() != 0;
