Fix some integration fallout
#rnx

diff --git a/Engine/Source/Runtime/Core/Private/Misc/ConfigCacheIni.cpp b/Engine/Source/Runtime/Core/Private/Misc/ConfigCacheIni.cpp
--- a/Engine/Source/Runtime/Core/Private/Misc/ConfigCacheIni.cpp
+++ b/Engine/Source/Runtime/Core/Private/Misc/ConfigCacheIni.cpp
@@ -51,7 +51,7 @@ namespace Private
 	static FCriticalSection ExpansionsCriticalSection;
 }
 }
-}?
+}
 const TSet<FString>* FConfigCacheIni::IniCacheSet = nullptr;
 
 namespace
diff --git a/Engine/Source/Runtime/RenderCore/Private/RenderGraphPrivate.h b/Engine/Source/Runtime/RenderCore/Private/RenderGraphPrivate.h
--- a/Engine/Source/Runtime/RenderCore/Private/RenderGraphPrivate.h
+++ b/Engine/Source/Runtime/RenderCore/Private/RenderGraphPrivate.h
@@ -205,6 +205,7 @@ FORCEINLINE bool IsParallelExecuteEnabled()
 		&& !IsImmediateMode()
 		&& !GRDGDebug
 		&& !GRDGTransitionLog
+		&& !IsMobilePlatform(GMaxRHIShaderPlatform)
 		&& GRHISupportsMultithreadedShaderCreation
 #if WITH_DUMPGPU
 		&& !IsDumpingRDGResources()
