Fix for content worker
#rnx

diff --git a/Engine/Source/Runtime/RenderCore/Private/ProfilingDebugging/RealtimeGPUProfiler.cpp b/Engine/Source/Runtime/RenderCore/Private/ProfilingDebugging/RealtimeGPUProfiler.cpp
--- a/Engine/Source/Runtime/RenderCore/Private/ProfilingDebugging/RealtimeGPUProfiler.cpp
+++ b/Engine/Source/Runtime/RenderCore/Private/ProfilingDebugging/RealtimeGPUProfiler.cpp
@@ -1196,7 +1196,7 @@ void FScopedGPUStatEvent::End()
 }
 #endif // HAS_GPU_STATS
 
-#if GPUPROFILERTRACE_ENABLED
+#if GPUPROFILERTRACE_ENABLED && HAS_GPU_STATS
 FRealtimeGPUProfilerHistoryItem::FRealtimeGPUProfilerHistoryItem()
 {
 	FMemory::Memset(*this, 0);
