Fixed double delete when destroying shader code libraries
#rnx

diff --git a/Engine/Source/Runtime/RenderCore/Private/ShaderCodeArchive.cpp b/Engine/Source/Runtime/RenderCore/Private/ShaderCodeArchive.cpp
--- a/Engine/Source/Runtime/RenderCore/Private/ShaderCodeArchive.cpp
+++ b/Engine/Source/Runtime/RenderCore/Private/ShaderCodeArchive.cpp
@@ -1706,6 +1706,7 @@ void FIoStoreShaderCodeArchive::Teardown()
 {
 	DebugVisualizer.SaveShaderUsageBitmap(GetName(), GetPlatform());
 
+	FWriteScopeLock Lock(PreloadedShaderGroupsLock);
 	for (TMap<int32, FShaderGroupPreloadEntry*>::TIterator Iter(PreloadedShaderGroups); Iter; ++Iter)
 	{
 		FShaderGroupPreloadEntry* PreloadEntry = Iter.Value();
@@ -1717,6 +1718,7 @@ void FIoStoreShaderCodeArchive::Teardown()
 
 		delete PreloadEntry;
 	}
+	PreloadedShaderGroups.Empty();
 }
 
 void FIoStoreShaderCodeArchive::SetupPreloadEntryForLoading(int32 ShaderGroupIndex, FShaderGroupPreloadEntry& PreloadEntry)
