Allow r.ForceLOD in shipping editor builds.

#rnx

diff --git a/Engine/Source/Runtime/RenderCore/Private/RenderCore.cpp b/Engine/Source/Runtime/RenderCore/Private/RenderCore.cpp
--- a/Engine/Source/Runtime/RenderCore/Private/RenderCore.cpp
+++ b/Engine/Source/Runtime/RenderCore/Private/RenderCore.cpp
@@ -200,7 +200,7 @@ DEFINE_STAT(STAT_RenderTargetPoolSize);
 DEFINE_STAT(STAT_RenderTargetPoolUsed);
 DEFINE_STAT(STAT_RenderTargetPoolCount);
 
-#define EXPOSE_FORCE_LOD !(UE_BUILD_SHIPPING || UE_BUILD_TEST)
+#define EXPOSE_FORCE_LOD !(UE_BUILD_SHIPPING || UE_BUILD_TEST) || WITH_EDITOR
 
 #if EXPOSE_FORCE_LOD
 
