Fixes black mip 0 in DumpGPU when a texture have a mip chain on D3D11 by using the capture by compute shader

diff --git a/Engine/Source/Runtime/RenderCore/Private/RenderGraphDumpResources.cpp b/Engine/Source/Runtime/RenderCore/Private/RenderGraphDumpResources.cpp
--- a/Engine/Source/Runtime/RenderCore/Private/RenderGraphDumpResources.cpp
+++ b/Engine/Source/Runtime/RenderCore/Private/RenderGraphDumpResources.cpp
@@ -685,7 +685,7 @@ struct FRDGResourceDumpContext
 		SubresourceDumpDesc.ByteSize = SIZE_T(SubresourceDumpDesc.SubResourceExtent.X) * SIZE_T(SubresourceDumpDesc.SubResourceExtent.Y) * SIZE_T(GPixelFormats[SubresourceDumpDesc.PreprocessedPixelFormat].BlockBytes);
 
 		// Whether the subresource need preprocessing pass before copy into staging.
-		SubresourceDumpDesc.bPreprocessForStaging = SubresourceDumpDesc.PreprocessedPixelFormat != Desc.Format || SubresourceDesc.MipLevel != 0;
+		SubresourceDumpDesc.bPreprocessForStaging = SubresourceDumpDesc.PreprocessedPixelFormat != Desc.Format || SubresourceDesc.Texture->Desc.NumMips > 1;
 
 		return SubresourceDumpDesc;
 	}
