Add a prebuilt nested struct to a ShaderParametersMetadataBuilder.
Will be useful for Optimus programatic metadata building.
#rnx

diff --git a/Engine/Source/Runtime/RenderCore/Private/ShaderParameterMetadata.cpp b/Engine/Source/Runtime/RenderCore/Private/ShaderParameterMetadata.cpp
--- a/Engine/Source/Runtime/RenderCore/Private/ShaderParameterMetadata.cpp
+++ b/Engine/Source/Runtime/RenderCore/Private/ShaderParameterMetadata.cpp
@@ -926,6 +926,30 @@ FString FShaderParametersMetadata::GetFullMemberCodeName(uint16 MemberOffset) co
 	return MemberName;
 }
 
+void FShaderParametersMetadataBuilder::AddNestedStruct(
+	const TCHAR* Name,
+	FShaderParametersMetadata* StructMetadata,
+	EShaderPrecisionModifier::Type Precision /* = EShaderPrecisionModifier::Float */
+)
+{
+	NextMemberOffset = Align(NextMemberOffset, SHADER_PARAMETER_STRUCT_ALIGNMENT);
+
+	new(Members) FShaderParametersMetadata::FMember(
+		Name,
+		TEXT(""),
+		__LINE__,
+		NextMemberOffset,
+		UBMT_NESTED_STRUCT,
+		Precision,
+		1,
+		1,
+		0,
+		StructMetadata
+	);
+
+	NextMemberOffset += Align(StructMetadata->GetSize(), SHADER_PARAMETER_STRUCT_ALIGNMENT);
+}
+
 void FShaderParametersMetadataBuilder::AddBufferSRV(
 	const TCHAR* Name,
 	const TCHAR* ShaderType,
diff --git a/Engine/Source/Runtime/RenderCore/Public/ShaderParameterMetadataBuilder.h b/Engine/Source/Runtime/RenderCore/Public/ShaderParameterMetadataBuilder.h
--- a/Engine/Source/Runtime/RenderCore/Public/ShaderParameterMetadataBuilder.h
+++ b/Engine/Source/Runtime/RenderCore/Public/ShaderParameterMetadataBuilder.h
@@ -64,6 +64,12 @@ public:
 		NextMemberOffset += sizeof(typename TParamTypeInfo::TAlignedType);
 	}
 
+	void AddNestedStruct(
+		const TCHAR* Name,
+		FShaderParametersMetadata* StructMetadata,
+		EShaderPrecisionModifier::Type Precision = EShaderPrecisionModifier::Float
+		);
+
 	void AddBufferSRV(
 		const TCHAR* Name,
 		const TCHAR* ShaderType,
