Fix a bug with GetMaxNumInstructionsForShader that report the incorrect instruction count for the given shader type.

- The function was not checking to see if the shader type in the pipeline was the same as the shader type we are asking for.

diff --git a/Engine/Source/Runtime/RenderCore/Private/ShaderMap.cpp b/Engine/Source/Runtime/RenderCore/Private/ShaderMap.cpp
--- a/Engine/Source/Runtime/RenderCore/Private/ShaderMap.cpp
+++ b/Engine/Source/Runtime/RenderCore/Private/ShaderMap.cpp
@@ -616,7 +616,12 @@ uint32 FShaderMapContent::GetMaxNumInstructionsForShader(const FShaderMapBase& I
 		FShader* PipelineShader = Pipeline->GetShader(ShaderType->GetFrequency());
 		if (PipelineShader)
 		{
-			MaxNumInstructions = FMath::Max(MaxNumInstructions, PipelineShader->GetNumInstructions());
+			const FShaderType* PipelineShaderType = PipelineShader->GetType(InShaderMap.GetPointerTable());
+			if (PipelineShaderType &&
+				(PipelineShaderType == ShaderType))
+			{
+				MaxNumInstructions = FMath::Max(MaxNumInstructions, PipelineShader->GetNumInstructions());
+			}
 		}
 	}
 
