Fix assert when the render thread crashes which was obscuring the real location of the crash.

The ERenderingThread named thread tag is not removed from ActiveNamedThreads when the render thread crashes, because the stack isn't unwound; this causes the assert in FTaskTagScope to fire when the Exit() method tries to set the tag. This resulted in the assert showing up as the cause of the crash, obscuring logs, crash reports etc.

#rnx

diff --git a/Engine/Source/Runtime/RenderCore/Private/RenderingThread.cpp b/Engine/Source/Runtime/RenderCore/Private/RenderingThread.cpp
--- a/Engine/Source/Runtime/RenderCore/Private/RenderingThread.cpp
+++ b/Engine/Source/Runtime/RenderCore/Private/RenderingThread.cpp
@@ -479,12 +479,18 @@ PRAGMA_ENABLE_DEPRECATION_WARNINGS
 
 	virtual void Exit(void) override
 	{
-		FTaskTagScope Scope(ETaskTag::ERenderingThread);
-		// Release rendering context ownership on the current thread if we had acquired it
-		if (bAcquiredThreadOwnership)
+		// Don't bother with cleanup if we crashed. This prevents an assert in FTaskTagScope caused by the fact that the stack isn't unwound when the
+		// render thread crashes, so the ERenderingThread bit added to ActiveNamedThreads by RenderingThreadMain remains set, and this FTaskTagScope
+		// // thinks we're trying to tag the thread twice.
+		if (GIsRenderingThreadHealthy)
 		{
-			bAcquiredThreadOwnership = false;
-			RHIReleaseThreadOwnership();
+			FTaskTagScope Scope(ETaskTag::ERenderingThread);
+			// Release rendering context ownership on the current thread if we had acquired it
+			if (bAcquiredThreadOwnership)
+			{
+				bAcquiredThreadOwnership = false;
+				RHIReleaseThreadOwnership();
+			}
 		}
 
 PRAGMA_DISABLE_DEPRECATION_WARNINGS
@@ -921,8 +927,6 @@ PRAGMA_ENABLE_DEPRECATION_WARNINGS
 
 void CheckRenderingThreadHealth()
 {
-	
-
 	if(!GIsRenderingThreadHealthy)
 	{
 		GErrorHist[0] = 0;
