Fix includes from generated files not being hashed during shader compilation (UE-134273).

- Fixes lack of recompilation when touching a vertex factory's .ush.
- Makes job input hash generation more robust generally by following the includes.
- Unfortunately, introduces a 5-10% perf hit when submitting jobs.
- Also contains changes to make -noshaderddc apply to global shaders

diff --git a/Engine/Source/Runtime/Engine/Private/ShaderCompiler/ShaderCompiler.cpp b/Engine/Source/Runtime/Engine/Private/ShaderCompiler/ShaderCompiler.cpp
--- a/Engine/Source/Runtime/Engine/Private/ShaderCompiler/ShaderCompiler.cpp
+++ b/Engine/Source/Runtime/Engine/Private/ShaderCompiler/ShaderCompiler.cpp
@@ -525,7 +525,7 @@ void FShaderCompileJobCollection::SubmitJobs(const TArray<FShaderCommonCompileJo
 			if (ShaderCompiler::IsJobCacheEnabled())
 			{
 				TRACE_CPUPROFILER_EVENT_SCOPE(ShaderCompiler.GetInputHash);
-				ParallelFor( TEXT("ShaderCompiler.GetInputHash.PF"), InJobs.Num(),1, [&InJobs](int32 Index) { InJobs[Index]->GetInputHash(); });
+				ParallelFor( TEXT("ShaderCompiler.GetInputHash.PF"), InJobs.Num(),1, [&InJobs](int32 Index) { InJobs[Index]->GetInputHash(); }, EParallelForFlags::Unbalanced);
 			}
 
 			FWriteScopeLock Locker(Lock);
@@ -7528,33 +7528,82 @@ FShaderCommonCompileJob::FInputHash FShaderCompileJob::GetInputHash()
 		checkf(Archive.IsSaving() && !Archive.IsLoading(), TEXT("A loading archive is passed to FShaderCompileJob::GetInputHash(), this is not supported as it may corrupt its data"));
 
 		Archive << Input;
-		Archive << Input.Environment;
+		Input.Environment.SerializeEverythingButFiles(Archive);
 
 		// hash the source file so changes to files during the development are picked up
 		const FSHAHash& SourceHash = GetShaderFileHash(*Input.VirtualSourceFilePath, Input.Target.GetPlatform());
 		Archive << const_cast<FSHAHash&>(SourceHash);
 
+		// unroll the included files for the parallel processing.
+		// These are temporary arrays that only exist for the ParallelFor
+		TArray<const TCHAR*> IncludeVirtualPaths;
+		TArray<const FString*> Contents;
+		TArray<bool> OnlyHashIncludes;
+		TArray<FBlake3Hash> Hashes;
+
+		// while the contents of this is already hashed (included in Environment's operator<<()), we still need to account for includes in the generated files and hash them, too
+		for (TMap<FString, FString>::TConstIterator It(Input.Environment.IncludeVirtualPathToContentsMap); It; ++It)
+		{
+			const FString& VirtualPath = It.Key();
+			IncludeVirtualPaths.Add(*VirtualPath);
+			Contents.Add(&It.Value());
+			OnlyHashIncludes.Add(true);	// not hashing contents of the file itself, as it was included in Environment's operator<<()
+			Hashes.AddDefaulted();
+		}
+
 		for (TMap<FString, FThreadSafeSharedStringPtr>::TConstIterator It(Input.Environment.IncludeVirtualPathToExternalContentsMap); It; ++It)
 		{
 			const FString& VirtualPath = It.Key();
-			Archive << const_cast<FString&>(VirtualPath);
+			IncludeVirtualPaths.Add(*VirtualPath);
 			check(It.Value());
-			FString& Contents = *It.Value();
-			Archive << Contents;
+			Contents.Add(&(*It.Value()));
+			OnlyHashIncludes.Add(false);
+			Hashes.AddDefaulted();
 		}
 
 		if (Input.SharedEnvironment)
 		{
-			Archive << *Input.SharedEnvironment;
+			Input.SharedEnvironment->SerializeEverythingButFiles(Archive);
+
+			for (TMap<FString, FString>::TConstIterator It(Input.SharedEnvironment->IncludeVirtualPathToContentsMap); It; ++It)
+			{
+				const FString& VirtualPath = It.Key();
+				IncludeVirtualPaths.Add(*VirtualPath);
+				Contents.Add(&It.Value());
+				OnlyHashIncludes.Add(true);	// not hashing contents of the file itself, as it was included in Environment's operator<<()
+				Hashes.AddDefaulted();
+			}
+
 			for (TMap<FString, FThreadSafeSharedStringPtr>::TConstIterator It(Input.SharedEnvironment->IncludeVirtualPathToExternalContentsMap); It; ++It)
 			{
 				const FString& VirtualPath = It.Key();
-				Archive << const_cast<FString&>(VirtualPath);
+				IncludeVirtualPaths.Add(*VirtualPath);
 				check(It.Value());
-				FString& Contents = *It.Value();
-				Archive << Contents;
+				Contents.Add(&(*It.Value()));
+				OnlyHashIncludes.Add(false);
+				Hashes.AddDefaulted();
 			}
 		}
+
+		check(IncludeVirtualPaths.Num() == Contents.Num());
+		check(Contents.Num() == OnlyHashIncludes.Num());
+		check(OnlyHashIncludes.Num() == Hashes.Num());
+
+		EShaderPlatform Platform = Input.Target.GetPlatform();
+		ParallelFor(Contents.Num(), [&IncludeVirtualPaths, &Contents, &OnlyHashIncludes, &Hashes, &Platform](int32 FileIndex)
+			{ 
+				FMemoryHasherBlake3 MemHasher;
+				HashShaderFileWithIncludes(MemHasher, IncludeVirtualPaths[FileIndex], *Contents[FileIndex], Platform, OnlyHashIncludes[FileIndex]);
+				Hashes[FileIndex] = MemHasher.Finalize();
+			},
+			EParallelForFlags::Unbalanced
+		);
+
+		// include the hashes in the main hash (consider sorting them if includes are found to have a random order)
+		for (int32 HashIndex = 0, NumHashes = Hashes.Num(); HashIndex < NumHashes; ++HashIndex)
+		{
+			Archive << Hashes[HashIndex];
+		}
 	};
 
 	// use faster hasher that doesn't allocate memory
diff --git a/Engine/Source/Runtime/RenderCore/Private/ShaderCore.cpp b/Engine/Source/Runtime/RenderCore/Private/ShaderCore.cpp
--- a/Engine/Source/Runtime/RenderCore/Private/ShaderCore.cpp
+++ b/Engine/Source/Runtime/RenderCore/Private/ShaderCore.cpp
@@ -1031,11 +1031,8 @@ static const TCHAR* FindFirstInclude(const TCHAR* Text)
 /**
  * Recursively populates IncludeFilenames with the unique include filenames found in the shader file named Filename.
  */
-static void GetShaderIncludes(const TCHAR* EntryPointVirtualFilePath, const TCHAR* VirtualFilePath, TArray<FString>& IncludeVirtualFilePaths, EShaderPlatform ShaderPlatform, uint32 DepthLimit, bool AddToIncludeFile)
+static void GetShaderIncludes(const TCHAR* EntryPointVirtualFilePath, const TCHAR* VirtualFilePath, const FString& FileContents, TArray<FString>& IncludeVirtualFilePaths, EShaderPlatform ShaderPlatform, uint32 DepthLimit, bool AddToIncludeFile)
 {
-	FString FileContents;
-	LoadShaderSourceFile(VirtualFilePath, ShaderPlatform, &FileContents, nullptr);
-
 	//avoid an infinite loop with a 0 length string
 	if (FileContents.Len() > 0)
 	{
@@ -1103,7 +1100,9 @@ static void GetShaderIncludes(const TCHAR* EntryPointVirtualFilePath, const TCHA
 					{
 						if (!IncludeVirtualFilePaths.Contains(ExtractedIncludeFilename))
 						{
-							GetShaderIncludes(EntryPointVirtualFilePath, *ExtractedIncludeFilename, IncludeVirtualFilePaths, ShaderPlatform, DepthLimit - 1, true);
+							FString IncludedFileContents;
+							LoadShaderSourceFile(*ExtractedIncludeFilename, ShaderPlatform, &IncludedFileContents, nullptr);
+							GetShaderIncludes(EntryPointVirtualFilePath, *ExtractedIncludeFilename, IncludedFileContents, IncludeVirtualFilePaths, ShaderPlatform, DepthLimit - 1, true);
 						}
 					}
 				}
@@ -1131,11 +1130,78 @@ static void GetShaderIncludes(const TCHAR* EntryPointVirtualFilePath, const TCHA
 	}
 }
 
+/**
+ * Recursively populates IncludeFilenames with the unique include filenames found in the shader file named Filename.
+ */
+static void GetShaderIncludes(const TCHAR* EntryPointVirtualFilePath, const TCHAR* VirtualFilePath, TArray<FString>& IncludeVirtualFilePaths, EShaderPlatform ShaderPlatform, uint32 DepthLimit, bool AddToIncludeFile)
+{
+	FString FileContents;
+	LoadShaderSourceFile(VirtualFilePath, ShaderPlatform, &FileContents, nullptr);
+
+	GetShaderIncludes(EntryPointVirtualFilePath, VirtualFilePath, FileContents, IncludeVirtualFilePaths, ShaderPlatform, DepthLimit, AddToIncludeFile);
+}
+
 void GetShaderIncludes(const TCHAR* EntryPointVirtualFilePath, const TCHAR* VirtualFilePath, TArray<FString>& IncludeVirtualFilePaths, EShaderPlatform ShaderPlatform, uint32 DepthLimit)
 {
 	GetShaderIncludes(EntryPointVirtualFilePath, VirtualFilePath, IncludeVirtualFilePaths, ShaderPlatform, DepthLimit, false);
 }
 
+void GetShaderIncludes(const TCHAR* EntryPointVirtualFilePath, const TCHAR* VirtualFilePath, const FString& FileContents, TArray<FString>& IncludeVirtualFilePaths, EShaderPlatform ShaderPlatform, uint32 DepthLimit)
+{
+	GetShaderIncludes(EntryPointVirtualFilePath, VirtualFilePath, FileContents, IncludeVirtualFilePaths, ShaderPlatform, DepthLimit, false);
+}
+
+void HashShaderFileWithIncludes(FArchive& HashingArchive, const TCHAR* VirtualFilePath, const FString& FileContents, EShaderPlatform ShaderPlatform, bool bOnlyHashIncludedFiles)
+{
+	auto HashSingleFile = [](FArchive& HashingArchive, const TCHAR* VirtualFilePath, EShaderPlatform ShaderPlatform, const FString& FileContents)
+	{
+		// first, a "soft" check
+		bool bFoundInCache = false;
+		{
+			FScopeLock ShaderHashAccessLock(&GShaderHashAccessGuard);
+			FSHAHash* CachedHash = GShaderHashCache.FindHash(ShaderPlatform, VirtualFilePath);
+			// If a hash for this filename has been cached, use that
+			if (CachedHash)
+			{
+				bFoundInCache = true;
+				HashingArchive << *CachedHash;
+			}
+		}
+
+		// outside of the lock scope because we don't need the lock and hashing can take time
+		if (!bFoundInCache)
+		{			
+			// if the file isn't generated, add it to the cache now
+			bool bGenerated = FCString::Strstr(VirtualFilePath, TEXT("Generated")) != nullptr;
+			if (!bGenerated)
+			{
+				// this function fails hard if it cannot load
+				const FSHAHash& FileHash = GetShaderFileHash(VirtualFilePath, ShaderPlatform);
+				HashingArchive << const_cast<FSHAHash&>(FileHash);
+			}
+			else
+			{
+				// note, it is legal for some generated files to have empty contents, so hash both the name and their contents
+				HashingArchive.Serialize(reinterpret_cast<void*>(const_cast<TCHAR*>(VirtualFilePath)), FCString::Strlen(VirtualFilePath));
+				HashingArchive << const_cast<FString&>(FileContents);
+			}
+		}
+	};
+
+	// First, always hash the file itself
+	HashSingleFile(HashingArchive, VirtualFilePath, ShaderPlatform, FileContents);
+
+	// Get the list of includes this file contains
+	TArray<FString> IncludeVirtualFilePaths;
+	GetShaderIncludes(VirtualFilePath, VirtualFilePath, FileContents, IncludeVirtualFilePaths, ShaderPlatform);
+
+	for (int32 IncludeIndex = 0; IncludeIndex < IncludeVirtualFilePaths.Num(); IncludeIndex++)
+	{
+		// Here, we assume that all includes can be found in cache. This also means that generated files won't include other generated files.
+		HashSingleFile(HashingArchive, *IncludeVirtualFilePaths[IncludeIndex], ShaderPlatform, FString());
+	}
+}
+
 static void UpdateSingleShaderFilehash(FSHA1& InOutHashState, const TCHAR* VirtualFilePath, EShaderPlatform ShaderPlatform)
 {
 	// Get the list of includes this file contains
diff --git a/Engine/Source/Runtime/RenderCore/Public/ShaderCore.h b/Engine/Source/Runtime/RenderCore/Public/ShaderCore.h
--- a/Engine/Source/Runtime/RenderCore/Public/ShaderCore.h
+++ b/Engine/Source/Runtime/RenderCore/Public/ShaderCore.h
@@ -529,6 +529,19 @@ struct FShaderCompilerEnvironment
 		RenderTargetOutputFormatsMap.Add(RenderTargetIndex, UE_PIXELFORMAT_TO_UINT8(PixelFormat));
 	}
 
+	/** This "core" serialization is also used for the hashing the compiler job (where files are handled differently). Should stay in sync with the ShaderCompileWorker. */
+	inline void SerializeEverythingButFiles(FArchive& Ar)
+	{
+		Ar << Definitions;
+		Ar << CompilerFlags;
+		Ar << RenderTargetOutputFormatsMap;
+		Ar << ResourceTableMap;
+		Ar << UniformBufferMap;
+		Ar << RemoteServerData;
+		Ar << ShaderFormatCVars;
+		Ar << FullPrecisionInPS;
+	}
+
 	friend FArchive& operator<<(FArchive& Ar,FShaderCompilerEnvironment& Environment)
 	{
 		// Note: this serialize is used to pass between UE and the shader compile worker, recompile both when modifying
@@ -536,15 +549,7 @@ struct FShaderCompilerEnvironment
 
 		// Note: skipping Environment.IncludeVirtualPathToExternalContentsMap, which is handled by FShaderCompileUtilities::DoWriteTasks in order to maintain sharing
 
-		Ar << Environment.Definitions;
-		Ar << Environment.CompilerFlags;
-		Ar << Environment.RenderTargetOutputFormatsMap;
-		Ar << Environment.ResourceTableMap;
-		Ar << Environment.UniformBufferMap;
-		Ar << Environment.RemoteServerData;
-		Ar << Environment.ShaderFormatCVars;
-		Ar << Environment.FullPrecisionInPS;
-
+		Environment.SerializeEverythingButFiles(Ar);
 		return Ar;
 	}
 	
@@ -998,6 +1003,7 @@ extern RENDERCORE_API void LoadShaderSourceFileChecked(const TCHAR* VirtualFileP
  * Recursively populates IncludeFilenames with the include filenames from Filename
  */
 extern RENDERCORE_API void GetShaderIncludes(const TCHAR* EntryPointVirtualFilePath, const TCHAR* VirtualFilePath, TArray<FString>& IncludeVirtualFilePaths, EShaderPlatform ShaderPlatform, uint32 DepthLimit=100);
+extern RENDERCORE_API void GetShaderIncludes(const TCHAR* EntryPointVirtualFilePath, const TCHAR* VirtualFilePath, const FString& FileContents, TArray<FString>& IncludeVirtualFilePaths, EShaderPlatform ShaderPlatform, uint32 DepthLimit = 100);
 
 /**
  * Calculates a Hash for the given filename if it does not already exist in the Hash cache.
@@ -1006,6 +1012,16 @@ extern RENDERCORE_API void GetShaderIncludes(const TCHAR* EntryPointVirtualFileP
  */
 extern RENDERCORE_API const class FSHAHash& GetShaderFileHash(const TCHAR* VirtualFilePath, EShaderPlatform ShaderPlatform);
 
+/**
+ * Calculates a hash for the given source file and all files included from it.
+ * @param HashingArchive - hash to update
+ * @param VirtualFilePath - name of this source code path (won't be loaded, as it is expected to be generated)
+ * @param FileContents - shader source code to Hash (included files will be hashed, too)
+ * @param ShaderPlatform - shader platform to Hash
+ * @param bOnlyHashIncludedFiles - skip hashing contents of the file itself (useful if it was already hashed outside of this function)
+ */
+extern RENDERCORE_API void HashShaderFileWithIncludes(FArchive& HashingArchive, const TCHAR* VirtualFilePath, const FString& FileContents, EShaderPlatform ShaderPlatform, bool bOnlyHashIncludedFiles);
+
 /**
  * Calculates a Hash for the list of filenames if it does not already exist in the Hash cache.
  */
