Enable CFLAG_ForceDXC for platforms that require DXC for HLSL2021.

#rnx

diff --git a/Engine/Source/Runtime/Engine/Private/ShaderCompiler/ShaderCompiler.cpp b/Engine/Source/Runtime/Engine/Private/ShaderCompiler/ShaderCompiler.cpp
--- a/Engine/Source/Runtime/Engine/Private/ShaderCompiler/ShaderCompiler.cpp
+++ b/Engine/Source/Runtime/Engine/Private/ShaderCompiler/ShaderCompiler.cpp
@@ -5707,7 +5707,8 @@ void GlobalBeginCompileShader(
 	}
 
 	// Add compiler flag CFLAG_ForceDXC if DXC is enabled
-	const bool bIsDxcEnabled = IsDxcEnabledForPlatform((EShaderPlatform)Target.Platform);
+	const bool bHlslVersion2021 = Input.Environment.CompilerFlags.Contains(CFLAG_HLSL2021);
+	const bool bIsDxcEnabled = IsDxcEnabledForPlatform((EShaderPlatform)Target.Platform, bHlslVersion2021);
 	Input.Environment.SetDefine(TEXT("COMPILER_DXC"), bIsDxcEnabled);
 	if (bIsDxcEnabled)
 	{
diff --git a/Engine/Source/Runtime/RenderCore/Private/Shader.cpp b/Engine/Source/Runtime/RenderCore/Private/Shader.cpp
--- a/Engine/Source/Runtime/RenderCore/Private/Shader.cpp
+++ b/Engine/Source/Runtime/RenderCore/Private/Shader.cpp
@@ -1297,13 +1297,13 @@ void DispatchIndirectComputeShader(
 	RHICmdList.DispatchIndirectComputeShader(ArgumentBuffer, ArgumentOffset);
 }
 
-bool IsDxcEnabledForPlatform(EShaderPlatform Platform)
+bool IsDxcEnabledForPlatform(EShaderPlatform Platform, bool bHlslVersion2021)
 {
 	// Check the generic console variable first (if DXC is supported)
 	if (FDataDrivenShaderPlatformInfo::GetSupportsDxc(Platform))
 	{
 		static const IConsoleVariable* CVar = IConsoleManager::Get().FindConsoleVariable(TEXT("r.Shaders.ForceDXC"));
-		if (CVar && CVar->GetInt() != 0)
+		if (bHlslVersion2021 || (CVar && CVar->GetInt() != 0))
 		{
 			return true;
 		}
@@ -1311,13 +1311,14 @@ bool IsDxcEnabledForPlatform(EShaderPlatform Platform)
 	// Check backend specific console variables next
 	if (IsD3DPlatform(Platform) && IsPCPlatform(Platform))
 	{
+		// D3D backend supports a precompile step for HLSL2021 which is separate from ForceDXC option
 		static const IConsoleVariable* CVar = IConsoleManager::Get().FindConsoleVariable(TEXT("r.D3D.ForceDXC"));
 		return (CVar && CVar->GetInt() != 0);
 	}
 	if (IsOpenGLPlatform(Platform))
 	{
 		static const IConsoleVariable* CVar = IConsoleManager::Get().FindConsoleVariable(TEXT("r.OpenGL.ForceDXC"));
-		return (CVar && CVar->GetInt() != 0);
+		return (bHlslVersion2021 || (CVar && CVar->GetInt() != 0));
 	}
 	// Hlslcc has been removed for Metal and Vulkan. There is only DXC now.
 	if (IsMetalPlatform(Platform) || IsVulkanPlatform(Platform))
diff --git a/Engine/Source/Runtime/RenderCore/Public/Shader.h b/Engine/Source/Runtime/RenderCore/Public/Shader.h
--- a/Engine/Source/Runtime/RenderCore/Public/Shader.h
+++ b/Engine/Source/Runtime/RenderCore/Public/Shader.h
@@ -2357,7 +2357,7 @@ inline void DispatchComputeShader(
 extern RENDERCORE_API bool IsUsingEmulatedUniformBuffers(EShaderPlatform Platform);
 
 /** Returns whether DirectXShaderCompiler (DXC) is enabled for the specified shader platform. See console variables "r.OpenGL.ForceDXC", "r.D3D.ForceDXC". */
-extern RENDERCORE_API bool IsDxcEnabledForPlatform(EShaderPlatform Platform);
+extern RENDERCORE_API bool IsDxcEnabledForPlatform(EShaderPlatform Platform, bool bHlslVersion2021 = false);
 
 /** Appends to KeyString for all shaders. */
 extern RENDERCORE_API void ShaderMapAppendKeyString(EShaderPlatform Platform, FString& KeyString);
