MGPU:  Fix for light map artifacts when reloading lighting following a bake.  Need to unbind SRVs used for scatter upload buffers, as we sometimes reuse the scatter upload buffers multiple times in a frame, and modifying them when they were bound can cause them to fail to update in the driver.  Strangely, this was only occurring on the second GPU...  I also added GPU mask blocks in a couple GPU scene update functions, for good measure.

The bound SRV issue could cause artifacts elsewhere in similar situations, but I didn't want to add perf overhead or complexity checking for this case everywhere, and decided to do a spot fix.  I didn't notice any other obvious cases where we have persistent buffers written to multiple times in a frame, like is done for the scatter upload buffers.  Still, it remains something to be alert for in case there are future rendering bugs.

#rnx

diff --git a/Engine/Source/Runtime/RenderCore/Private/UnifiedBuffer.cpp b/Engine/Source/Runtime/RenderCore/Private/UnifiedBuffer.cpp
--- a/Engine/Source/Runtime/RenderCore/Private/UnifiedBuffer.cpp
+++ b/Engine/Source/Runtime/RenderCore/Private/UnifiedBuffer.cpp
@@ -823,6 +823,10 @@ void FScatterUploadBuffer::ResourceUploadTo(FRHICommandList& RHICmdList, const R
 
 	RHICmdList.EndUAVOverlap(DstBuffer.UAV);
 
+	// We need to unbind shader SRVs in this case, because scatter upload buffers are sometimes used more than once in a
+	// frame, and this can cause rendering bugs on D3D12, where the driver fails to update the bound SRV with new data.
+	UnsetShaderSRVs(RHICmdList, ComputeShader, ComputeShader.GetComputeShader());
+
 	if (bFlush)
 	{
 		FRHICommandListExecutor::GetImmediateCommandList().ImmediateFlush(EImmediateFlushType::DispatchToRHIThread);
diff --git a/Engine/Source/Runtime/RenderCore/Public/ShaderParameterStruct.h b/Engine/Source/Runtime/RenderCore/Public/ShaderParameterStruct.h
--- a/Engine/Source/Runtime/RenderCore/Public/ShaderParameterStruct.h
+++ b/Engine/Source/Runtime/RenderCore/Public/ShaderParameterStruct.h
@@ -180,6 +180,26 @@ inline void UnsetShaderUAVs(TRHICmdList& RHICmdList, const TShaderRef<TShaderCla
 }
 
 
+/** Unset compute shader SRVs. */
+template<typename TRHICmdList, typename TShaderClass>
+inline void UnsetShaderSRVs(TRHICmdList& RHICmdList, const TShaderRef<TShaderClass>& Shader, FRHIComputeShader* ShadeRHI)
+{
+	const FShaderParameterBindings& Bindings = Shader->Bindings;
+
+	checkf(Bindings.RootParameterBufferIndex == FShaderParameterBindings::kInvalidBufferIndex, TEXT("Can't use UnsetShaderSRVs() for root parameter buffer index."));
+
+	for (const FShaderParameterBindings::FResourceParameter& ParameterBinding : Bindings.ResourceParameters)
+	{
+		if (ParameterBinding.BaseType == UBMT_SRV ||
+			ParameterBinding.BaseType == UBMT_RDG_TEXTURE_SRV ||
+			ParameterBinding.BaseType == UBMT_RDG_BUFFER_SRV)
+		{
+			RHICmdList.SetShaderResourceViewParameter(ShadeRHI, ParameterBinding.BaseIndex, nullptr);
+		}
+	}
+}
+
+
 /** Set shader's parameters from its parameters struct. */
 template<typename TRHICmdList, typename TShaderClass, typename TShaderRHI>
 inline void SetShaderParameters(
diff --git a/Engine/Source/Runtime/Renderer/Private/GPUScene.cpp b/Engine/Source/Runtime/Renderer/Private/GPUScene.cpp
--- a/Engine/Source/Runtime/Renderer/Private/GPUScene.cpp
+++ b/Engine/Source/Runtime/Renderer/Private/GPUScene.cpp
@@ -1826,6 +1826,8 @@ void FGPUScene::UploadDynamicPrimitiveShaderDataForView(FRDGBuilder& GraphBuilde
 {
 	if (bIsEnabled)
 	{
+		RDG_GPU_MASK_SCOPE(GraphBuilder, FRHIGPUMask::All());
+
 		UploadDynamicPrimitiveShaderDataForViewInternal(GraphBuilder, Scene, View, bIsShadowView);
 	}
 }
@@ -2113,6 +2115,8 @@ bool FGPUScene::ExecuteDeferredGPUWritePass(FRDGBuilder& GraphBuilder, const TAr
 		return false;
 	}
 
+	RDG_GPU_MASK_SCOPE(GraphBuilder, FRHIGPUMask::All());
+
 	// Mark this pass as having executed for the frame
 	LastDeferredGPUWritePass = GPUWritePass;
 
