Fix crash on exit due to RHI commands enqueued by the transient allocator being executed after the allocator is destroyed.

#rnx

diff --git a/Engine/Source/Runtime/RenderCore/Private/RenderGraphResourcePool.cpp b/Engine/Source/Runtime/RenderCore/Private/RenderGraphResourcePool.cpp
--- a/Engine/Source/Runtime/RenderCore/Private/RenderGraphResourcePool.cpp
+++ b/Engine/Source/Runtime/RenderCore/Private/RenderGraphResourcePool.cpp
@@ -197,6 +197,11 @@ void FRDGTransientResourceAllocator::ReleaseDynamicRHI()
 		DeallocatedList.Empty();
 
 		Allocator->Flush(RHICmdList);
+		
+		// Allocator->Flush() enqueues some lambdas on the command list, so make sure they are executed
+		// before the allocator is deleted.
+		RHICmdList.ImmediateFlush(EImmediateFlushType::FlushRHIThread);
+
 		Allocator->Release(RHICmdList);
 		Allocator = nullptr;
 	}
