Adds dump timing measurement details to log when using DumpGPU

diff --git a/Engine/Source/Runtime/RenderCore/Private/RenderGraphDumpResources.cpp b/Engine/Source/Runtime/RenderCore/Private/RenderGraphDumpResources.cpp
--- a/Engine/Source/Runtime/RenderCore/Private/RenderGraphDumpResources.cpp
+++ b/Engine/Source/Runtime/RenderCore/Private/RenderGraphDumpResources.cpp
@@ -155,8 +155,79 @@ struct FRDGResourceDumpContext
 	const FRDGPass* DrawDumpingPass = nullptr;
 	int32 DrawDumpCount = 0;
 
+	int32 MetadataFilesOpened = 0;
+	int64 MetadataFilesWriteBytes = 0;
+	int32 ResourceBinaryFilesOpened = 0;
+	int64 ResourceBinaryWriteBytes = 0;
+
+	enum class ETimingBucket : uint8
+	{
+		RHIReadbackCommands,
+		RHIReleaseResources,
+		GPUWait,
+		CPUProcessing,
+		MetadataFileWrite,
+		ResourceBinaryFileWrite,
+		MAX
+	};
+	mutable double TimingBucket[int32(ETimingBucket::MAX)];
+
+	class FTimeBucketMeasure
+	{
+	public:
+		FTimeBucketMeasure(FRDGResourceDumpContext* InDumpCtx, ETimingBucket InBucket)
+			: DumpCtx(InDumpCtx)
+			, Start(FPlatformTime::Cycles64())
+			, Bucket(InBucket)
+		{ }
+
+		~FTimeBucketMeasure()
+		{
+			uint64 End = FPlatformTime::Cycles64();
+			DumpCtx->TimingBucket[int32(Bucket)] += FPlatformTime::ToSeconds64(FMath::Max(End - Start, uint64(0)));
+		}
+
+	protected:
+		FRDGResourceDumpContext* const DumpCtx;
+
+	private:
+		const uint64 Start;
+		const ETimingBucket Bucket;
+	};
+
+	class FFileWriteCtx : public FTimeBucketMeasure
+	{
+	public:
+		FFileWriteCtx(FRDGResourceDumpContext* InDumpCtx, ETimingBucket InBucket, int64 WriteSizeBytes)
+			: FTimeBucketMeasure(InDumpCtx, InBucket)
+		{
+			if (InBucket == ETimingBucket::MetadataFileWrite)
+			{
+				DumpCtx->MetadataFilesOpened += 1;
+				DumpCtx->MetadataFilesWriteBytes += WriteSizeBytes;
+			}
+			else if (InBucket == ETimingBucket::ResourceBinaryFileWrite)
+			{
+				DumpCtx->ResourceBinaryFilesOpened += 1;
+				DumpCtx->ResourceBinaryWriteBytes += WriteSizeBytes;
+			}
+			else
+			{
+				unimplemented();
+			}
+		}
+	};
+
 	bool bShowInExplore = false;
 
+	FRDGResourceDumpContext()
+	{
+		for (int32 i = 0; i < int32(ETimingBucket::MAX); i++)
+		{
+			TimingBucket[i] = 0.0f;
+		}
+	}
+
 	bool IsDumpingFrame() const
 	{
 		check(IsInRenderingThread() || IsInParallelRenderingThread());
@@ -178,6 +249,8 @@ struct FRDGResourceDumpContext
 		}
 
 		FString FullPath = GetDumpFullPath(FileName);
+
+		FFileWriteCtx WriteCtx(this, ETimingBucket::MetadataFileWrite, OutputString.Len());
 		return FFileHelper::SaveStringToFile(
 			OutputString, *FullPath,
 			FFileHelper::EEncodingOptions::AutoDetect, &IFileManager::Get(), WriteFlags);
@@ -201,6 +274,7 @@ struct FRDGResourceDumpContext
 		}
 
 		FString FullPath = GetDumpFullPath(FileName);
+		FFileWriteCtx WriteCtx(this, ETimingBucket::ResourceBinaryFileWrite, ArrayView.Num());
 		return FFileHelper::SaveArrayToFile(ArrayView, *FullPath);
 	}
 
@@ -213,6 +287,8 @@ struct FRDGResourceDumpContext
 		}
 
 		FString FullPath = GetDumpFullPath(FileName);
+
+		FFileWriteCtx WriteCtx(this, ETimingBucket::ResourceBinaryFileWrite, Array.Num());
 		return FFileHelper::SaveArrayToFile(Array, *FullPath);
 	}
 
@@ -264,6 +340,8 @@ struct FRDGResourceDumpContext
 
 	void ReleaseRHIResources(FRHICommandListImmediate& RHICmdList)
 	{
+		FTimeBucketMeasure TimeBucketMeasure(this, ETimingBucket::RHIReleaseResources);
+
 		// Flush the RHI resource memory so the readback memory can be fully reused in the next resource dump.
 		{
 			RHICmdList.SubmitCommandsAndFlushGPU();
@@ -704,83 +782,89 @@ struct FRDGResourceDumpContext
 		FTextureRHIRef StagingSrcTexture;
 		EPixelFormat PreprocessedPixelFormat = SubresourceDumpDesc.PreprocessedPixelFormat;
 		SIZE_T SubresourceByteSize = SubresourceDumpDesc.ByteSize;
-		if (SubresourceDumpDesc.bPreprocessForStaging)
+		FTextureRHIRef StagingTexture;
+
 		{
-			// Some RHIs (GL) only support 32Bit single channel images as CS output
-			if (IsOpenGLPlatform(GMaxRHIShaderPlatform) &&
-				GPixelFormats[PreprocessedPixelFormat].NumComponents == 1 && 
-				GPixelFormats[PreprocessedPixelFormat].BlockBytes < 4)
+			FTimeBucketMeasure TimeBucketMeasure(this, ETimingBucket::RHIReadbackCommands);
+			if (SubresourceDumpDesc.bPreprocessForStaging)
 			{
-				SubresourceByteSize*= (4 / GPixelFormats[PreprocessedPixelFormat].BlockBytes);
-				PreprocessedPixelFormat = PF_R32_UINT;
-			}
+				// Some RHIs (GL) only support 32Bit single channel images as CS output
+				if (IsOpenGLPlatform(GMaxRHIShaderPlatform) &&
+					GPixelFormats[PreprocessedPixelFormat].NumComponents == 1 && 
+					GPixelFormats[PreprocessedPixelFormat].BlockBytes < 4)
+				{
+					SubresourceByteSize*= (4 / GPixelFormats[PreprocessedPixelFormat].BlockBytes);
+					PreprocessedPixelFormat = PF_R32_UINT;
+				}
 						
+				{
+					FRHIResourceCreateInfo CreateInfo(TEXT("DumpGPU.PreprocessTexture"));
+					StagingSrcTexture = RHICreateTexture2D(
+						SubresourceDumpDesc.SubResourceExtent.X,
+						SubresourceDumpDesc.SubResourceExtent.Y,
+						(uint8)PreprocessedPixelFormat,
+						/* NumMips = */ 1,
+						/* NumSamples = */ 1,
+						TexCreate_UAV | TexCreate_ShaderResource | TexCreate_HideInVisualizeTexture,
+						CreateInfo);
+				}
+
+				FUnorderedAccessViewRHIRef StagingOutput = RHICreateUnorderedAccessView(StagingSrcTexture, /* MipLevel = */ 0);
+
+				RHICmdList.Transition(FRHITransitionInfo(StagingSrcTexture, ERHIAccess::Unknown, ERHIAccess::UAVCompute));
+
+				FDumpTextureCS::FPermutationDomain PermutationVector;
+				PermutationVector.Set<FDumpTextureCS::FTextureTypeDim>(SubresourceDumpDesc.DumpTextureType);
+				TShaderMapRef<FDumpTextureCS> ComputeShader(GetGlobalShaderMap(GMaxRHIShaderPlatform), PermutationVector);
+
+				FDumpTextureCS::FParameters ShaderParameters;
+				ShaderParameters.Texture = SubResourceSRV;
+				ShaderParameters.StagingOutput = StagingOutput;
+				FComputeShaderUtils::Dispatch(
+					RHICmdList,
+					ComputeShader,
+					ShaderParameters,
+					FComputeShaderUtils::GetGroupCount(SubresourceDumpDesc.SubResourceExtent, 8));
+
+				RHICmdList.Transition(FRHITransitionInfo(StagingSrcTexture, ERHIAccess::UAVCompute, ERHIAccess::CopySrc));
+			}
+			else
 			{
-				FRHIResourceCreateInfo CreateInfo(TEXT("DumpGPU.PreprocessTexture"));
-				StagingSrcTexture = RHICreateTexture2D(
+				StagingSrcTexture = Texture;
+			}
+
+			// Copy the texture for CPU readback
+			{
+				FRHIResourceCreateInfo CreateInfo(TEXT("DumpGPU.StagingTexture"));
+				StagingTexture = RHICreateTexture2D(
 					SubresourceDumpDesc.SubResourceExtent.X,
 					SubresourceDumpDesc.SubResourceExtent.Y,
 					(uint8)PreprocessedPixelFormat,
 					/* NumMips = */ 1,
 					/* NumSamples = */ 1,
-					TexCreate_UAV | TexCreate_ShaderResource | TexCreate_HideInVisualizeTexture,
+					TexCreate_CPUReadback | TexCreate_HideInVisualizeTexture,
 					CreateInfo);
-			}
-
-			FUnorderedAccessViewRHIRef StagingOutput = RHICreateUnorderedAccessView(StagingSrcTexture, /* MipLevel = */ 0);
-
-			RHICmdList.Transition(FRHITransitionInfo(StagingSrcTexture, ERHIAccess::Unknown, ERHIAccess::UAVCompute));
-
-			FDumpTextureCS::FPermutationDomain PermutationVector;
-			PermutationVector.Set<FDumpTextureCS::FTextureTypeDim>(SubresourceDumpDesc.DumpTextureType);
-			TShaderMapRef<FDumpTextureCS> ComputeShader(GetGlobalShaderMap(GMaxRHIShaderPlatform), PermutationVector);
-
-			FDumpTextureCS::FParameters ShaderParameters;
-			ShaderParameters.Texture = SubResourceSRV;
-			ShaderParameters.StagingOutput = StagingOutput;
-			FComputeShaderUtils::Dispatch(
-				RHICmdList,
-				ComputeShader,
-				ShaderParameters,
-				FComputeShaderUtils::GetGroupCount(SubresourceDumpDesc.SubResourceExtent, 8));
-
-			RHICmdList.Transition(FRHITransitionInfo(StagingSrcTexture, ERHIAccess::UAVCompute, ERHIAccess::CopySrc));
-		}
-		else
-		{
-			StagingSrcTexture = Texture;
-		}
-
-		// Copy the texture for CPU readback
-		FTextureRHIRef StagingTexture;
-		{
-			FRHIResourceCreateInfo CreateInfo(TEXT("DumpGPU.StagingTexture"));
-			StagingTexture = RHICreateTexture2D(
-				SubresourceDumpDesc.SubResourceExtent.X,
-				SubresourceDumpDesc.SubResourceExtent.Y,
-				(uint8)PreprocessedPixelFormat,
-				/* NumMips = */ 1,
-				/* NumSamples = */ 1,
-				TexCreate_CPUReadback | TexCreate_HideInVisualizeTexture,
-				CreateInfo);
 
-			RHICmdList.Transition(FRHITransitionInfo(StagingTexture, ERHIAccess::Unknown, ERHIAccess::CopyDest));
+				RHICmdList.Transition(FRHITransitionInfo(StagingTexture, ERHIAccess::Unknown, ERHIAccess::CopyDest));
 
-			// Ensure this copy call does not perform any transitions. We're handling them manually.
-			FResolveParams ResolveParams;
-			ResolveParams.SourceAccessFinal = ERHIAccess::Unknown;
-			ResolveParams.DestAccessFinal = ERHIAccess::Unknown;
+				// Ensure this copy call does not perform any transitions. We're handling them manually.
+				FResolveParams ResolveParams;
+				ResolveParams.SourceAccessFinal = ERHIAccess::Unknown;
+				ResolveParams.DestAccessFinal = ERHIAccess::Unknown;
 
-			// Transfer memory GPU -> CPU
-			RHICmdList.CopyToResolveTarget(StagingSrcTexture, StagingTexture, ResolveParams);
+				// Transfer memory GPU -> CPU
+				RHICmdList.CopyToResolveTarget(StagingSrcTexture, StagingTexture, ResolveParams);
 
-			RHICmdList.Transition(FRHITransitionInfo(StagingTexture, ERHIAccess::CopyDest, ERHIAccess::CPURead));
+				RHICmdList.Transition(FRHITransitionInfo(StagingTexture, ERHIAccess::CopyDest, ERHIAccess::CPURead));
+			}
 		}
 
 		// Submit to GPU and wait for completion.
 		static const FName FenceName(TEXT("DumpGPU.TextureFence"));
 		FGPUFenceRHIRef Fence = RHICreateGPUFence(FenceName);
 		{
+			FTimeBucketMeasure TimeBucketMeasure(this, ETimingBucket::GPUWait);
+
 			Fence->Clear();
 			RHICmdList.WriteGPUFence(Fence);
 			RHICmdList.SubmitCommandsAndFlushGPU();
@@ -803,45 +887,50 @@ struct FRDGResourceDumpContext
 		if (Content)
 		{
 			TArray64<uint8> Array;
-			Array.SetNumUninitialized(SubresourceByteSize);
 
-			SIZE_T BytePerPixel = SIZE_T(GPixelFormats[PreprocessedPixelFormat].BlockBytes);
+			{
+				FTimeBucketMeasure TimeBucketMeasure(this, ETimingBucket::CPUProcessing);
 
-			const uint8* SrcData = static_cast<const uint8*>(Content);
+				Array.SetNumUninitialized(SubresourceByteSize);
 
-			for (int32 y = 0; y < SubresourceDumpDesc.SubResourceExtent.Y; y++)
-			{
-				// Flip the data to be bottom left corner for the WebGL viewer.
-				const uint8* SrcPos = SrcData + SIZE_T(SubresourceDumpDesc.SubResourceExtent.Y - 1 - y) * SIZE_T(RowPitchInPixels) * BytePerPixel;
-				uint8* DstPos = (&Array[0]) + SIZE_T(y) * SIZE_T(SubresourceDumpDesc.SubResourceExtent.X) * BytePerPixel;
+				SIZE_T BytePerPixel = SIZE_T(GPixelFormats[PreprocessedPixelFormat].BlockBytes);
 
-				FPlatformMemory::Memmove(DstPos, SrcPos, SIZE_T(SubresourceDumpDesc.SubResourceExtent.X) * BytePerPixel);
-			}
+				const uint8* SrcData = static_cast<const uint8*>(Content);
 
-			RHICmdList.UnmapStagingSurface(StagingTexture, GPUIndex);
+				for (int32 y = 0; y < SubresourceDumpDesc.SubResourceExtent.Y; y++)
+				{
+					// Flip the data to be bottom left corner for the WebGL viewer.
+					const uint8* SrcPos = SrcData + SIZE_T(SubresourceDumpDesc.SubResourceExtent.Y - 1 - y) * SIZE_T(RowPitchInPixels) * BytePerPixel;
+					uint8* DstPos = (&Array[0]) + SIZE_T(y) * SIZE_T(SubresourceDumpDesc.SubResourceExtent.X) * BytePerPixel;
 
-			if (PreprocessedPixelFormat != SubresourceDumpDesc.PreprocessedPixelFormat)
-			{
-				// Convert 32Bit values back to 16 or 8bit
-				const int32 DstPixelNumBytes = GPixelFormats[SubresourceDumpDesc.PreprocessedPixelFormat].BlockBytes;
-				const uint32* SrcData32 = (const uint32*)Array.GetData();
-				uint8* DstData8 = Array.GetData();
-				uint16* DstData16 = (uint16*)Array.GetData();
-											
-				for (int32 Index = 0; Index < Array.Num()/4; Index++)
+					FPlatformMemory::Memmove(DstPos, SrcPos, SIZE_T(SubresourceDumpDesc.SubResourceExtent.X) * BytePerPixel);
+				}
+
+				RHICmdList.UnmapStagingSurface(StagingTexture, GPUIndex);
+
+				if (PreprocessedPixelFormat != SubresourceDumpDesc.PreprocessedPixelFormat)
 				{
-					uint32 Value32 = SrcData32[Index];
-					if (DstPixelNumBytes == 2)
-					{
-						DstData16[Index] = (uint16)Value32;
-					}
-					else
+					// Convert 32Bit values back to 16 or 8bit
+					const int32 DstPixelNumBytes = GPixelFormats[SubresourceDumpDesc.PreprocessedPixelFormat].BlockBytes;
+					const uint32* SrcData32 = (const uint32*)Array.GetData();
+					uint8* DstData8 = Array.GetData();
+					uint16* DstData16 = (uint16*)Array.GetData();
+											
+					for (int32 Index = 0; Index < Array.Num()/4; Index++)
 					{
-						DstData8[Index] = (uint8)Value32;
+						uint32 Value32 = SrcData32[Index];
+						if (DstPixelNumBytes == 2)
+						{
+							DstData16[Index] = (uint16)Value32;
+						}
+						else
+						{
+							DstData8[Index] = (uint8)Value32;
+						}
 					}
-				}
 
-				Array.SetNum(Array.Num() / (4 / DstPixelNumBytes));
+					Array.SetNum(Array.Num() / (4 / DstPixelNumBytes));
+				}
 			}
 
 			DumpBinaryToFile(Array, DumpFilePath);
@@ -1126,15 +1215,20 @@ struct FRDGResourceDumpContext
 				[this, DumpFilePath, Buffer, ByteSize](FRHICommandListImmediate& RHICmdList)
 			{
 				check(IsInRenderingThread());
-				FStagingBufferRHIRef StagingBuffer = RHICreateStagingBuffer();
+				FStagingBufferRHIRef StagingBuffer;
+				{
+					FTimeBucketMeasure TimeBucketMeasure(this, ETimingBucket::RHIReadbackCommands);
+					StagingBuffer = RHICreateStagingBuffer();
 
-				// Transfer memory GPU -> CPU
-				RHICmdList.CopyToStagingBuffer(Buffer->GetRHI(), StagingBuffer, 0, ByteSize);
+					// Transfer memory GPU -> CPU
+					RHICmdList.CopyToStagingBuffer(Buffer->GetRHI(), StagingBuffer, 0, ByteSize);
+				}
 
 				// Submit to GPU and wait for completion.
 				static const FName FenceName(TEXT("DumpGPU.BufferFence"));
 				FGPUFenceRHIRef Fence = RHICreateGPUFence(FenceName);
 				{
+					FTimeBucketMeasure TimeBucketMeasure(this, ETimingBucket::GPUWait);
 					Fence->Clear();
 					RHICmdList.WriteGPUFence(Fence);
 					RHICmdList.SubmitCommandsAndFlushGPU();
@@ -1381,7 +1475,30 @@ void FRDGBuilder::EndResourceDump()
 	// Log information about the dump.
 	FString AbsDumpingDirectoryPath = IFileManager::Get().ConvertToAbsolutePathForExternalAppForRead(*GRDGResourceDumpContext.DumpingDirectoryPath);
 	{
-		UE_LOG(LogRendererCore, Display, TEXT("Dumped %d resources to %s"), GRDGResourceDumpContext.ResourcesDumpPasses, *AbsDumpingDirectoryPath);
+		FDateTime Now = FDateTime::Now();
+		double TotalDumpSeconds = (Now - GRDGResourceDumpContext.Time).GetTotalSeconds();
+		double RHIReadbackCommandsSeconds = GRDGResourceDumpContext.TimingBucket[int32(FRDGResourceDumpContext::ETimingBucket::RHIReadbackCommands)];
+		double GPUWaitSeconds = GRDGResourceDumpContext.TimingBucket[int32(FRDGResourceDumpContext::ETimingBucket::GPUWait)];
+		double CPUProcessingSeconds = GRDGResourceDumpContext.TimingBucket[int32(FRDGResourceDumpContext::ETimingBucket::CPUProcessing)];
+		double MetadataFileWriteSeconds = GRDGResourceDumpContext.TimingBucket[int32(FRDGResourceDumpContext::ETimingBucket::MetadataFileWrite)];
+		double ResourceBinaryFileWriteSeconds = GRDGResourceDumpContext.TimingBucket[int32(FRDGResourceDumpContext::ETimingBucket::ResourceBinaryFileWrite)];
+		double RHIReleaseResourcesTimeSeconds = GRDGResourceDumpContext.TimingBucket[int32(FRDGResourceDumpContext::ETimingBucket::RHIReleaseResources)];
+
+		UE_LOG(LogRendererCore, Display, TEXT("Dumped %d resources in %.3f s to %s"), GRDGResourceDumpContext.ResourcesDumpPasses, float(TotalDumpSeconds), *AbsDumpingDirectoryPath);
+		UE_LOG(LogRendererCore, Display, TEXT("Dumped GPU readback commands: %.3f s"), float(RHIReadbackCommandsSeconds));
+		UE_LOG(LogRendererCore, Display, TEXT("Dumped GPU wait: %.3f s"), float(GPUWaitSeconds));
+		UE_LOG(LogRendererCore, Display, TEXT("Dumped CPU resource binary post processing: %.3f s"), float(CPUProcessingSeconds));
+		UE_LOG(LogRendererCore, Display, TEXT("Dumped metadata: %.3f MB in %d files under %.3f s at %.3f MB/s"),
+			float(GRDGResourceDumpContext.MetadataFilesWriteBytes) / float(1024 * 1024),
+			GRDGResourceDumpContext.MetadataFilesOpened,
+			float(MetadataFileWriteSeconds),
+			float(GRDGResourceDumpContext.MetadataFilesWriteBytes) / (float(1024 * 1024) * float(MetadataFileWriteSeconds)));
+		UE_LOG(LogRendererCore, Display, TEXT("Dumped resource binary: %.3f MB in %d files under %.3f s at %.3f MB/s"),
+			float(GRDGResourceDumpContext.ResourceBinaryWriteBytes) / float(1024 * 1024),
+			GRDGResourceDumpContext.ResourceBinaryFilesOpened,
+			float(ResourceBinaryFileWriteSeconds),
+			float(GRDGResourceDumpContext.ResourceBinaryWriteBytes) / (float(1024 * 1024) * float(ResourceBinaryFileWriteSeconds)));
+		UE_LOG(LogRendererCore, Display, TEXT("Dumped GPU readback resource release: %.3f s"), float(RHIReleaseResourcesTimeSeconds));
 	}
 
 	// Dump the log into the dump directory.
