Implements r.DumpGPU.CompressResources

diff --git a/Engine/Extras/GPUDumpViewer/GPUDumpViewer.html b/Engine/Extras/GPUDumpViewer/GPUDumpViewer.html
--- a/Engine/Extras/GPUDumpViewer/GPUDumpViewer.html
+++ b/Engine/Extras/GPUDumpViewer/GPUDumpViewer.html
@@ -141,6 +141,31 @@ function load_binary_file(relative_path, callback)
 	}
 }
 
+function load_resource_binary_file(relative_path, callback)
+{
+	return load_binary_file(relative_path, function(raw_texture_data)
+	{
+		var compression_type = '';
+		if (g_dump_cvars['r.DumpGPU.CompressResources'] == '1')
+		{
+			compression_type = 'deflate';
+		}
+		else if (g_dump_cvars['r.DumpGPU.CompressResources'] == '2')
+		{
+			compression_type = 'gzip';
+		}
+
+		if (raw_texture_data === null || compression_type == '')
+		{
+			return callback(raw_texture_data);
+		}
+
+		var decompressor = new DecompressionStream(compression_type);
+		var decompressed_stream = new Blob([raw_texture_data]).stream().pipeThrough(decompressor);
+		new Response(decompressed_stream).arrayBuffer().then(callback, function() { callback(null); });
+	});
+}
+
 function load_json(relative_path)
 {
 	var txt = load_text_file(relative_path);
@@ -1474,7 +1499,7 @@ class TextureView extends ResourceView
 
 		var texture_view = this;
 		var subresource_version_name = get_subresource_unique_version_name(this.subresource_version_info);
-		load_binary_file(`Resources/${subresource_version_name}.bin`, function(raw_texture_data)
+		load_resource_binary_file(`Resources/${subresource_version_name}.bin`, function(raw_texture_data)
 		{
 			if (raw_texture_data)
 			{
@@ -3119,7 +3144,7 @@ class GenericBufferView extends ResourceView
 
 		var buffer_view = this;
 		console.assert(this.raw_buffer_data_path);
-		load_binary_file(this.raw_buffer_data_path, function(raw_buffer_data)
+		load_resource_binary_file(this.raw_buffer_data_path, function(raw_buffer_data)
 		{
 			if (raw_buffer_data)
 			{
diff --git a/Engine/Source/Runtime/RenderCore/Private/RenderGraphDumpResources.cpp b/Engine/Source/Runtime/RenderCore/Private/RenderGraphDumpResources.cpp
--- a/Engine/Source/Runtime/RenderCore/Private/RenderGraphDumpResources.cpp
+++ b/Engine/Source/Runtime/RenderCore/Private/RenderGraphDumpResources.cpp
@@ -69,6 +69,14 @@ static TAutoConsoleVariable<int32> GDumpRenderingConsoleVariablesCVar(
 	TEXT("Whether to dump rendering console variables (enabled by default)."),
 	ECVF_Default);
 
+static TAutoConsoleVariable<int32> GDumpGPUCompressResources(
+	TEXT("r.DumpGPU.CompressResources"), 0,
+	TEXT("Whether to compress resource binary.\n")
+	TEXT(" 0: Disabled (default)\n")
+	TEXT(" 1: Zlib\n")
+	TEXT(" 2: GZip"),
+	ECVF_Default);
+
 static TAutoConsoleVariable<int32> GDumpTestEnableDiskWrite(
 	TEXT("r.DumpGPU.Test.EnableDiskWrite"), 1,
 	TEXT("Master switch whether any files should be written to disk, used for r.DumpGPU automation tests to not fill up workers' hard drive."),
@@ -141,6 +149,7 @@ struct FRDGResourceDumpContext
 	static constexpr const TCHAR* kStructuresMetadataDir = TEXT("StructuresMetadata/");
 
 	bool bEnableDiskWrite = false;
+	FName ResourceCompressionName;
 	FString DumpingDirectoryPath;
 	FDateTime Time;
 	FGenericPlatformMemoryConstants MemoryConstants;
@@ -159,13 +168,15 @@ struct FRDGResourceDumpContext
 	int64 MetadataFilesWriteBytes = 0;
 	int32 ResourceBinaryFilesOpened = 0;
 	int64 ResourceBinaryWriteBytes = 0;
+	int64 ResourceBinaryUncompressedBytes = 0;
 
 	enum class ETimingBucket : uint8
 	{
 		RHIReadbackCommands,
 		RHIReleaseResources,
 		GPUWait,
-		CPUProcessing,
+		CPUPostProcessing,
+		CPUCompression,
 		MetadataFileWrite,
 		ResourceBinaryFileWrite,
 		MAX
@@ -278,7 +289,7 @@ struct FRDGResourceDumpContext
 		return FFileHelper::SaveArrayToFile(ArrayView, *FullPath);
 	}
 
-	bool DumpBinaryToFile(const TArray64<uint8>& Array, const FString& FileName)
+	bool DumpBinaryToFile(const uint8* Data, int64 DataByteSize, const FString& FileName)
 	{
 		// Make it has if the write happened and was successful.
 		if (!bEnableDiskWrite)
@@ -287,14 +298,62 @@ struct FRDGResourceDumpContext
 		}
 
 		FString FullPath = GetDumpFullPath(FileName);
+		FFileWriteCtx WriteCtx(this, ETimingBucket::ResourceBinaryFileWrite, DataByteSize);
+
+		TUniquePtr<FArchive> Ar = TUniquePtr<FArchive>(IFileManager::Get().CreateFileWriter(*FullPath, /* WriteFlags = */ 0));
+		if (!Ar)
+		{
+			return false;
+		}
+		Ar->Serialize((void*)Data, DataByteSize);
+
+		// Always explicitly close to catch errors from flush/close
+		Ar->Close();
+
+		return !Ar->IsError() && !Ar->IsCriticalError();
+	}
+
+	bool IsResourceBinaryCompressable(int64 UncompressedSize) const
+	{
+		return !ResourceCompressionName.IsNone() && (UncompressedSize < (int64(1) << int64(32)));
+	}
+
+	bool DumpResourceBinaryToFile(const uint8* UncompressedData, int64 UncompressedSize, const FString& FileName)
+	{
+		ResourceBinaryUncompressedBytes += UncompressedSize;
+
+		if (!IsResourceBinaryCompressable(UncompressedSize))
+		{
+			return DumpBinaryToFile(UncompressedData, UncompressedSize, FileName);
+		}
+
+		int32 CompressedSize;
+		TArray64<uint8> CompressedArray;
+		{
+			FTimeBucketMeasure TimeBucketMeasure(this, ETimingBucket::CPUCompression);
+
+			int32 MaxCompressedSize = FCompression::CompressMemoryBound(ResourceCompressionName, UncompressedSize, COMPRESS_BiasSize);
+			CompressedArray.SetNumUninitialized(MaxCompressedSize);
 
-		FFileWriteCtx WriteCtx(this, ETimingBucket::ResourceBinaryFileWrite, Array.Num());
-		return FFileHelper::SaveArrayToFile(Array, *FullPath);
+			CompressedSize = MaxCompressedSize;
+			bool bSuccess = FCompression::CompressMemory(
+				ResourceCompressionName,
+				/* out */ CompressedArray.GetData(), /* out */ CompressedSize,
+				UncompressedData, UncompressedSize,
+				COMPRESS_BiasSize);
+		}
+
+		return DumpBinaryToFile(CompressedArray.GetData(), CompressedSize, FileName);
 	}
 
 	bool IsUnsafeToDumpResource(SIZE_T ResourceByteSize, float DumpMemoryMultiplier) const
 	{
 		uint64 AproximatedStagingMemoryRequired = uint64(double(ResourceByteSize) * DumpMemoryMultiplier);
+		if (IsResourceBinaryCompressable(ResourceByteSize))
+		{
+			AproximatedStagingMemoryRequired += FCompression::CompressMemoryBound(ResourceCompressionName, ResourceByteSize, COMPRESS_BiasSize);
+		}
+
 		uint64 MaxMemoryAvailable = FMath::Min(MemoryStats.AvailablePhysical, MemoryStats.AvailableVirtual);
 
 		return AproximatedStagingMemoryRequired > MaxMemoryAvailable;
@@ -887,9 +946,8 @@ struct FRDGResourceDumpContext
 		if (Content)
 		{
 			TArray64<uint8> Array;
-
 			{
-				FTimeBucketMeasure TimeBucketMeasure(this, ETimingBucket::CPUProcessing);
+				FTimeBucketMeasure TimeBucketMeasure(this, ETimingBucket::CPUPostProcessing);
 
 				Array.SetNumUninitialized(SubresourceByteSize);
 
@@ -933,7 +991,7 @@ struct FRDGResourceDumpContext
 				}
 			}
 
-			DumpBinaryToFile(Array, DumpFilePath);
+			DumpResourceBinaryToFile(Array.GetData(), Array.Num(), DumpFilePath);
 		}
 		else
 		{
@@ -1238,8 +1296,7 @@ struct FRDGResourceDumpContext
 				void* Content = RHICmdList.LockStagingBuffer(StagingBuffer, Fence.GetReference(), 0, ByteSize);
 				if (Content)
 				{
-					TArrayView<const uint8> ArrayView(reinterpret_cast<const uint8*>(Content), ByteSize);
-					DumpBinaryToFile(ArrayView, DumpFilePath);
+					DumpResourceBinaryToFile(reinterpret_cast<const uint8*>(Content), ByteSize, DumpFilePath);
 
 					RHICmdList.UnlockStagingBuffer(StagingBuffer);
 				}
@@ -1339,6 +1396,16 @@ FString FRDGBuilder::BeginResourceDump(const TArray<FString>& Args)
 		NewResourceDumpContext.DumpingDirectoryPath = DirectoryPath / FApp::GetProjectName() + TEXT("-") + FPlatformProperties::PlatformName() + TEXT("-") + NewResourceDumpContext.Time.ToString() + TEXT("/");
 	}
 	NewResourceDumpContext.bEnableDiskWrite = GDumpTestEnableDiskWrite.GetValueOnGameThread() != 0;
+	{
+		if (GDumpGPUCompressResources.GetValueOnGameThread() == 1)
+		{
+			NewResourceDumpContext.ResourceCompressionName = NAME_Zlib;
+		}
+		else if (GDumpGPUCompressResources.GetValueOnGameThread() == 2)
+		{
+			NewResourceDumpContext.ResourceCompressionName = NAME_Gzip;
+		}
+	}
 	NewResourceDumpContext.bShowInExplore = NewResourceDumpContext.bEnableDiskWrite && GDumpExploreCVar.GetValueOnGameThread() != 0;
 	NewResourceDumpContext.MemoryConstants = FPlatformMemory::GetConstants();
 	NewResourceDumpContext.MemoryStats = FPlatformMemory::GetStats();
@@ -1479,7 +1546,7 @@ void FRDGBuilder::EndResourceDump()
 		double TotalDumpSeconds = (Now - GRDGResourceDumpContext.Time).GetTotalSeconds();
 		double RHIReadbackCommandsSeconds = GRDGResourceDumpContext.TimingBucket[int32(FRDGResourceDumpContext::ETimingBucket::RHIReadbackCommands)];
 		double GPUWaitSeconds = GRDGResourceDumpContext.TimingBucket[int32(FRDGResourceDumpContext::ETimingBucket::GPUWait)];
-		double CPUProcessingSeconds = GRDGResourceDumpContext.TimingBucket[int32(FRDGResourceDumpContext::ETimingBucket::CPUProcessing)];
+		double CPUPostProcessingSeconds = GRDGResourceDumpContext.TimingBucket[int32(FRDGResourceDumpContext::ETimingBucket::CPUPostProcessing)];
 		double MetadataFileWriteSeconds = GRDGResourceDumpContext.TimingBucket[int32(FRDGResourceDumpContext::ETimingBucket::MetadataFileWrite)];
 		double ResourceBinaryFileWriteSeconds = GRDGResourceDumpContext.TimingBucket[int32(FRDGResourceDumpContext::ETimingBucket::ResourceBinaryFileWrite)];
 		double RHIReleaseResourcesTimeSeconds = GRDGResourceDumpContext.TimingBucket[int32(FRDGResourceDumpContext::ETimingBucket::RHIReleaseResources)];
@@ -1487,7 +1554,20 @@ void FRDGBuilder::EndResourceDump()
 		UE_LOG(LogRendererCore, Display, TEXT("Dumped %d resources in %.3f s to %s"), GRDGResourceDumpContext.ResourcesDumpPasses, float(TotalDumpSeconds), *AbsDumpingDirectoryPath);
 		UE_LOG(LogRendererCore, Display, TEXT("Dumped GPU readback commands: %.3f s"), float(RHIReadbackCommandsSeconds));
 		UE_LOG(LogRendererCore, Display, TEXT("Dumped GPU wait: %.3f s"), float(GPUWaitSeconds));
-		UE_LOG(LogRendererCore, Display, TEXT("Dumped CPU resource binary post processing: %.3f s"), float(CPUProcessingSeconds));
+		UE_LOG(LogRendererCore, Display, TEXT("Dumped CPU resource binary post processing: %.3f s"), float(CPUPostProcessingSeconds));
+		if (!GRDGResourceDumpContext.ResourceCompressionName.IsNone())
+		{
+			double CPUCompressionSeconds = GRDGResourceDumpContext.TimingBucket[int32(FRDGResourceDumpContext::ETimingBucket::CPUCompression)];
+
+			UE_LOG(LogRendererCore, Display, TEXT("Dumped CPU resource binary compression from %.3f MB to %.3f MB (%.1f %%) using %s"),
+				float(GRDGResourceDumpContext.ResourceBinaryUncompressedBytes) / float(1024 * 1024),
+				float(GRDGResourceDumpContext.ResourceBinaryWriteBytes) / float(1024 * 1024),
+				100.0f * float(GRDGResourceDumpContext.ResourceBinaryWriteBytes) / float(GRDGResourceDumpContext.ResourceBinaryUncompressedBytes),
+				*GRDGResourceDumpContext.ResourceCompressionName.ToString());
+			UE_LOG(LogRendererCore, Display, TEXT("Dumped CPU resource binary compression took %.3f s at %.3f MB/s"),
+				float(CPUCompressionSeconds),
+				float(GRDGResourceDumpContext.ResourceBinaryUncompressedBytes) / (float(1024 * 1024) * float(CPUCompressionSeconds)));
+		}
 		UE_LOG(LogRendererCore, Display, TEXT("Dumped metadata: %.3f MB in %d files under %.3f s at %.3f MB/s"),
 			float(GRDGResourceDumpContext.MetadataFilesWriteBytes) / float(1024 * 1024),
 			GRDGResourceDumpContext.MetadataFilesOpened,
