Fixed ray tracing could be accidentally enabled when disabled in render settings

diff --git a/Engine/Source/Runtime/RenderCore/Private/RenderUtils.cpp b/Engine/Source/Runtime/RenderCore/Private/RenderUtils.cpp
--- a/Engine/Source/Runtime/RenderCore/Private/RenderUtils.cpp
+++ b/Engine/Source/Runtime/RenderCore/Private/RenderUtils.cpp
@@ -1485,13 +1485,13 @@ RENDERCORE_API void RenderUtilsInit()
 				}
 				else
 				{
-					GUseRayTracing = true;
+					//GUseRayTracing = true;
 
-					UE_LOG(LogRendererCore, Log, TEXT("Ray tracing is enabled for the game. Reason: r.RayTracing=1, and r.RayTracing.EnableInGame is not present (default true)."));
+					//UE_LOG(LogRendererCore, Log, TEXT("Ray tracing is enabled for the game. Reason: r.RayTracing=1, and r.RayTracing.EnableInGame is not present (default true)."));
 
-					//GUseRayTracing = false;
+					GUseRayTracing = false;
 
-					//UE_LOG(LogRendererCore, Log, TEXT("Ray tracing is disabled for the game. Reason: r.RayTracing=1, and r.RayTracing.EnableInGame is not present (default false)."));
+					UE_LOG(LogRendererCore, Log, TEXT("Ray tracing is disabled for the game. Reason: r.RayTracing=1, and r.RayTracing.EnableInGame is not present (default false)."));
 				}
 			}
 
