Add ROV type support for shader binding.

diff --git a/Engine/Config/VulkanPC/DataDrivenPlatformInfo.ini b/Engine/Config/VulkanPC/DataDrivenPlatformInfo.ini
--- a/Engine/Config/VulkanPC/DataDrivenPlatformInfo.ini
+++ b/Engine/Config/VulkanPC/DataDrivenPlatformInfo.ini
@@ -129,6 +129,10 @@ bSupportsVertexShaderLayer=false
 
 bSupportsAnisotropicMaterials=true
 
+; disable until DXC -> SPRIV ROV support is implemented
+bSupportsROV=false
+bSupportsOIT=false
+
 
 [ShaderPlatform VULKAN_PCES3_1]
 Language=Vulkan
diff --git a/Engine/Config/Windows/DataDrivenPlatformInfo.ini b/Engine/Config/Windows/DataDrivenPlatformInfo.ini
--- a/Engine/Config/Windows/DataDrivenPlatformInfo.ini
+++ b/Engine/Config/Windows/DataDrivenPlatformInfo.ini
@@ -88,6 +88,7 @@ bSupportsCapsuleShadows=true
 bSupportsPercentageCloserShadows=true
 bSupportsVolumetricFog=true
 bSupportsIndexBufferUAVs=true
+bSupportsROV=true
 bSupportsOIT=true
 
 [PreviewPlatform PCD3D_SM5]
@@ -136,6 +137,7 @@ bSupportsPercentageCloserShadows=true
 bSupportsVolumetricFog=true
 bSupportsIndexBufferUAVs=true
 bSupportsBindless=true
+bSupportsROV=true
 bSupportsOIT=true
 
 [PreviewPlatform PCD3D_SM6]
diff --git a/Engine/Shaders/Public/Platform.ush b/Engine/Shaders/Public/Platform.ush
--- a/Engine/Shaders/Public/Platform.ush
+++ b/Engine/Shaders/Public/Platform.ush
@@ -92,6 +92,10 @@
 #define PLATFORM_SUPPORTS_SRV_UB 0
 #endif
 
+#ifndef PLATFORM_SUPPORTS_ROV
+#define PLATFORM_SUPPORTS_ROV 0
+#endif
+
 #ifndef SM6_PROFILE
 #define SM6_PROFILE 0
 #endif
@@ -1032,6 +1036,16 @@ uint MaskedBitCount( uint2 Bits, uint Index )
 	return Offset;
 }
 
+// Lock a critical region of code within a pixel shader and guarantees no concurrent execution for the same pixel
+#ifndef RASTER_ORDERED_VIEW_LOCK
+#define RASTER_ORDERED_VIEW_LOCK()
+#endif
+
+// Unlock a critical region of code within a pixel shader.
+#ifndef RASTER_ORDERED_VIEW_UNLOCK
+#define RASTER_ORDERED_VIEW_UNLOCK()
+#endif
+
 #if PLATFORM_SUPPORTS_SM6_0_WAVE_OPERATIONS
 
 #ifndef COMPILER_SUPPORTS_TO_SCALAR_MEMORY
diff --git a/Engine/Source/Developer/ShaderCompilerCommon/Private/HlslExpressionParser.inl b/Engine/Source/Developer/ShaderCompilerCommon/Private/HlslExpressionParser.inl
--- a/Engine/Source/Developer/ShaderCompilerCommon/Private/HlslExpressionParser.inl
+++ b/Engine/Source/Developer/ShaderCompilerCommon/Private/HlslExpressionParser.inl
@@ -360,6 +360,7 @@ namespace CrossCompiler
 		case EHlslToken::RWTexture1DArray:
 		case EHlslToken::RWTexture2D:
 		case EHlslToken::RWTexture2DArray:
+		case EHlslToken::RasterizerOrderedTexture2D:
 		case EHlslToken::RWTexture3D:
 		case EHlslToken::ConstantBuffer:
 		case EHlslToken::StructuredBuffer:
diff --git a/Engine/Source/Developer/ShaderCompilerCommon/Private/HlslLexer.cpp b/Engine/Source/Developer/ShaderCompilerCommon/Private/HlslLexer.cpp
--- a/Engine/Source/Developer/ShaderCompilerCommon/Private/HlslLexer.cpp
+++ b/Engine/Source/Developer/ShaderCompilerCommon/Private/HlslLexer.cpp
@@ -401,6 +401,7 @@ namespace CrossCompiler
 				InsertToken(TEXT("RW_Texture2D"), EHlslToken::RWTexture2D);	// PSSL
 				InsertToken(TEXT("RWTexture2DArray"), EHlslToken::RWTexture2DArray);
 				InsertToken(TEXT("RW_Texture2D_Array"), EHlslToken::RWTexture2DArray);	// PSSL
+				InsertToken(TEXT("RasterizerOrderedTexture2D"), EHlslToken::RasterizerOrderedTexture2D);
 				InsertToken(TEXT("RWTexture3D"), EHlslToken::RWTexture3D);
 				InsertToken(TEXT("RW_Texture3D"), EHlslToken::RWTexture3D);	// PSSL
 				InsertToken(TEXT("StructuredBuffer"), EHlslToken::StructuredBuffer);
diff --git a/Engine/Source/Developer/ShaderCompilerCommon/Private/HlslLexer.h b/Engine/Source/Developer/ShaderCompilerCommon/Private/HlslLexer.h
--- a/Engine/Source/Developer/ShaderCompilerCommon/Private/HlslLexer.h
+++ b/Engine/Source/Developer/ShaderCompilerCommon/Private/HlslLexer.h
@@ -247,6 +247,7 @@ namespace CrossCompiler
 		RWTexture1DArray,
 		RWTexture2D,
 		RWTexture2DArray,
+		RasterizerOrderedTexture2D,
 		RWTexture3D,
 		StructuredBuffer,
 		ConstantBuffer,
diff --git a/Engine/Source/Developer/VulkanShaderFormat/Private/VulkanShaderCompiler.cpp b/Engine/Source/Developer/VulkanShaderFormat/Private/VulkanShaderCompiler.cpp
--- a/Engine/Source/Developer/VulkanShaderFormat/Private/VulkanShaderCompiler.cpp
+++ b/Engine/Source/Developer/VulkanShaderFormat/Private/VulkanShaderCompiler.cpp
@@ -2272,6 +2272,7 @@ void DoCompileVulkanShader(const FShaderCompilerInput& Input, FShaderCompilerOut
 
 	AdditionalDefines.SetDefine(TEXT("COMPILER_SUPPORTS_ATTRIBUTES"), (uint32)1);
 	AdditionalDefines.SetDefine(TEXT("COMPILER_SUPPORTS_DUAL_SOURCE_BLENDING_SLOT_DECORATION"), (uint32)1);
+	AdditionalDefines.SetDefine(TEXT("PLATFORM_SUPPORTS_ROV"), 0); // Disabled until DXC->SPRIV ROV support is implemented
 
 	if (Input.Environment.FullPrecisionInPS)
 	{
diff --git a/Engine/Source/Developer/Windows/ShaderFormatD3D/Private/D3DShaderCompiler.cpp b/Engine/Source/Developer/Windows/ShaderFormatD3D/Private/D3DShaderCompiler.cpp
--- a/Engine/Source/Developer/Windows/ShaderFormatD3D/Private/D3DShaderCompiler.cpp
+++ b/Engine/Source/Developer/Windows/ShaderFormatD3D/Private/D3DShaderCompiler.cpp
@@ -1163,7 +1163,8 @@ void CompileD3DShader(const FShaderCompilerInput& Input, FShaderCompilerOutput&
 
 	// Set additional defines.
 	AdditionalDefines.SetDefine(TEXT("COMPILER_HLSL"), 1);
-
+	AdditionalDefines.SetDefine(TEXT("PLATFORM_SUPPORTS_ROV"), 1); // Assume min. spec HW supports with DX12/SM5
+	
 	if (bUseDXC)
 	{
 		AdditionalDefines.SetDefine(TEXT("PLATFORM_SUPPORTS_SM6_0_WAVE_OPERATIONS"), 1);
diff --git a/Engine/Source/Runtime/RHI/Private/RHI.cpp b/Engine/Source/Runtime/RHI/Private/RHI.cpp
--- a/Engine/Source/Runtime/RHI/Private/RHI.cpp
+++ b/Engine/Source/Runtime/RHI/Private/RHI.cpp
@@ -2381,6 +2381,7 @@ void FGenericDataDrivenShaderPlatformInfo::ParseDataDrivenShaderInfo(const FConf
 	GET_SECTION_BOOL_HELPER(bMobileSupportFetchBindedCustomStencilBuffer);
 	GET_SECTION_BOOL_HELPER(bSupportsBindless);
 	GET_SECTION_BOOL_HELPER(bSupportsVolumeTextureAtomics);
+	GET_SECTION_BOOL_HELPER(bSupportsROV);
 	GET_SECTION_BOOL_HELPER(bSupportsOIT);
 #undef GET_SECTION_BOOL_HELPER
 #undef GET_SECTION_INT_HELPER
diff --git a/Engine/Source/Runtime/RHI/Public/RHIDefinitions.h b/Engine/Source/Runtime/RHI/Public/RHIDefinitions.h
--- a/Engine/Source/Runtime/RHI/Public/RHIDefinitions.h
+++ b/Engine/Source/Runtime/RHI/Public/RHIDefinitions.h
@@ -397,6 +397,7 @@ class RHI_API FGenericDataDrivenShaderPlatformInfo
 	uint32 bMobileSupportFetchBindedCustomStencilBuffer : 1;
 	uint32 bSupportsBindless : 1;
 	uint32 bSupportsVolumeTextureAtomics : 1;
+	uint32 bSupportsROV : 1;
 	uint32 bSupportsOIT : 1;
 		
 #if WITH_EDITOR
@@ -921,6 +922,11 @@ public:
 		return Infos[Platform].bSupportsVolumeTextureAtomics;
 	}
 
+	static FORCEINLINE_DEBUGGABLE const bool GetSupportsROV(const FStaticShaderPlatform Platform)
+	{
+		return Infos[Platform].bSupportsROV;
+	}
+
 	static FORCEINLINE_DEBUGGABLE const bool GetSupportsOIT(const FStaticShaderPlatform Platform)
 	{
 		return Infos[Platform].bSupportsOIT;
@@ -1295,6 +1301,8 @@ enum class EShaderCodeResourceBindingType : uint8
 	RWStructuredBuffer,
 	RWByteAddressBuffer,
 
+	RasterizerOrderedTexture2D,
+
 	MAX
 };
 
diff --git a/Engine/Source/Runtime/RenderCore/Private/ShaderParameterMetadata.cpp b/Engine/Source/Runtime/RenderCore/Private/ShaderParameterMetadata.cpp
--- a/Engine/Source/Runtime/RenderCore/Private/ShaderParameterMetadata.cpp
+++ b/Engine/Source/Runtime/RenderCore/Private/ShaderParameterMetadata.cpp
@@ -236,6 +236,10 @@ EShaderCodeResourceBindingType ParseShaderResourceBindingType(const TCHAR* Shade
 	{
 		BindingType = EShaderCodeResourceBindingType::RaytracingAccelerationStructure;
 	}
+	else if (bIsRasterizerOrderedResource)
+	{
+		BindingType = EShaderCodeResourceBindingType::RasterizerOrderedTexture2D;
+	}
 	return BindingType;
 }
 
@@ -261,6 +265,7 @@ const TCHAR* const kShaderCodeResourceBindingTypeNames[] = {
 	TEXT("RWBuffer"),
 	TEXT("RWStructuredBuffer"),
 	TEXT("RWByteAddressBuffer"),
+	TEXT("RasterizerOrderedTexture2D"),
 };
 
 static_assert(UE_ARRAY_COUNT(kShaderCodeResourceBindingTypeNames) == int32(EShaderCodeResourceBindingType::MAX), "TODO.");
@@ -593,7 +598,8 @@ void FShaderParametersMetadata::InitializeLayout(FRHIUniformBufferLayoutInitiali
 					BindingType == EShaderCodeResourceBindingType::RWTexture2DArray ||
 					BindingType == EShaderCodeResourceBindingType::RWTexture3D ||
 					BindingType == EShaderCodeResourceBindingType::RWTextureCube ||
-					BindingType == EShaderCodeResourceBindingType::RWTextureMetadata)
+					BindingType == EShaderCodeResourceBindingType::RWTextureMetadata ||
+					BindingType == EShaderCodeResourceBindingType::RasterizerOrderedTexture2D)
 				{
 					bIsValidBindingType = (
 						BaseType == UBMT_UAV ||
