Renames RenderGraphDumpResources.cpp to DumpGPU.cpp

diff --git a/Engine/Source/Runtime/RenderCore/Private/DumpGPU.cpp b/Engine/Source/Runtime/RenderCore/Private/DumpGPU.cpp
new file mode 100644
--- /dev/null
+++ b/Engine/Source/Runtime/RenderCore/Private/DumpGPU.cpp
@@ -0,0 +1,2066 @@
+// Copyright Epic Games, Inc. All Rights Reserved.
+
+#include "RenderGraph.h"
+#include "HAL/PlatformFileManager.h"
+#include "HAL/PlatformMisc.h"
+#include "Misc/FileHelper.h"
+#include "Misc/WildcardString.h"
+#include "Misc/OutputDeviceRedirector.h"
+#include "Runtime/Launch/Resources/Version.h"
+#include "BuildSettings.h"
+#include "Serialization/JsonTypes.h"
+#include "Serialization/JsonReader.h"
+#include "Serialization/JsonSerializer.h"
+#include "GenericPlatform/GenericPlatformCrashContext.h"
+#include "GenericPlatform/GenericPlatformDriver.h"
+#include "RenderGraphUtils.h"
+#include "GlobalShader.h"
+#include "RHIValidation.h"
+
+#if RDG_DUMP_RESOURCES
+
+namespace
+{
+
+static TAutoConsoleVariable<FString> GDumpGPURootCVar(
+	TEXT("r.DumpGPU.Root"),
+	TEXT("*"),
+	TEXT("Allows to filter the tree when using r.DumpGPU command, the pattern match is case sensitive."),
+	ECVF_Default);
+
+static TAutoConsoleVariable<int32> GDumpTextureCVar(
+	TEXT("r.DumpGPU.Texture"), 2,
+	TEXT("Whether to dump textures.\n")
+	TEXT(" 0: Ignores all textures\n")
+	TEXT(" 1: Dump only textures' descriptors\n")
+	TEXT(" 2: Dump textures' descriptors and binaries (default)"),
+	ECVF_RenderThreadSafe);
+
+static TAutoConsoleVariable<int32> GDumpBufferCVar(
+	TEXT("r.DumpGPU.Buffer"), 2,
+	TEXT("Whether to dump buffer.\n")
+	TEXT(" 0: Ignores all buffers\n")
+	TEXT(" 1: Dump only buffers' descriptors\n")
+	TEXT(" 2: Dump buffers' descriptors and binaries (default)"),
+	ECVF_RenderThreadSafe);
+
+static TAutoConsoleVariable<int32> GDumpGPUPassParameters(
+	TEXT("r.DumpGPU.PassParameters"), 1,
+	TEXT("Whether to dump the pass parameters."),
+	ECVF_RenderThreadSafe);
+
+static TAutoConsoleVariable<int32> GDumpGPUDraws(
+	TEXT("r.DumpGPU.Draws"), 0,
+	TEXT("Whether to dump resource after each individual draw call (disabled by default)."),
+	ECVF_Default);
+
+static TAutoConsoleVariable<int32> GDumpGPUMask(
+	TEXT("r.DumpGPU.Mask"), 1,
+	TEXT("Whether to include GPU mask in the name of each Pass (has no effect unless system has multiple GPUs)."),
+	ECVF_Default);
+
+static TAutoConsoleVariable<int32> GDumpExploreCVar(
+	TEXT("r.DumpGPU.Explore"), 1,
+	TEXT("Whether to open file explorer to where the GPU dump on completion (enabled by default)."),
+	ECVF_Default);
+
+static TAutoConsoleVariable<int32> GDumpRenderingConsoleVariablesCVar(
+	TEXT("r.DumpGPU.ConsoleVariables"), 1,
+	TEXT("Whether to dump rendering console variables (enabled by default)."),
+	ECVF_Default);
+
+static TAutoConsoleVariable<int32> GDumpGPUCompressResources(
+	TEXT("r.DumpGPU.CompressResources"), 0,
+	TEXT("Whether to compress resource binary.\n")
+	TEXT(" 0: Disabled (default)\n")
+	TEXT(" 1: Zlib\n")
+	TEXT(" 2: GZip"),
+	ECVF_Default);
+
+static TAutoConsoleVariable<int32> GDumpTestEnableDiskWrite(
+	TEXT("r.DumpGPU.Test.EnableDiskWrite"), 1,
+	TEXT("Master switch whether any files should be written to disk, used for r.DumpGPU automation tests to not fill up workers' hard drive."),
+	ECVF_Default);
+
+static TAutoConsoleVariable<int32> GDumpTestPrettifyResourceFileNames(
+	TEXT("r.DumpGPU.Test.PrettifyResourceFileNames"), 0,
+	TEXT("Whether the resource file names should include resource name. May increase the likelyness of running into Windows' filepath limit."),
+	ECVF_RenderThreadSafe);
+
+static TAutoConsoleVariable<FString> GDumpGPUDirectoryCVar(
+	TEXT("r.DumpGPU.Directory"), TEXT(""),
+	TEXT("Directory to dump to."),
+	ECVF_Default);
+
+// Although this cvar does not seams used in the C++ code base, it is dumped by DumpRenderingCVarsToCSV() and used by GPUDumpViewer.html.
+static TAutoConsoleVariable<FString> GDumpGPUVisualizeResource(
+	TEXT("r.DumpGPU.Viewer.Visualize"), TEXT(""),
+	TEXT("Name of RDG output resource to automatically open in the dump viewer."),
+	ECVF_Default);
+
+class FDumpTextureCS : public FGlobalShader
+{
+	DECLARE_GLOBAL_SHADER(FDumpTextureCS);
+	SHADER_USE_PARAMETER_STRUCT(FDumpTextureCS, FGlobalShader);
+
+	static inline bool IsSupported(const FStaticShaderPlatform ShaderPlatform)
+	{
+		return RHISupportsComputeShaders(ShaderPlatform);
+	}
+
+	enum class ETextureType
+	{
+		Texture2DFloatNoMSAA,
+		Texture2DUintNoMSAA,
+		Texture2DDepthStencilNoMSAA,
+		MAX
+	};
+	class FTextureTypeDim : SHADER_PERMUTATION_ENUM_CLASS("DIM_TEXTURE_TYPE", ETextureType);
+	using FPermutationDomain = TShaderPermutationDomain<FTextureTypeDim>;
+
+	static bool ShouldCompilePermutation(const FGlobalShaderPermutationParameters& Parameters)
+	{
+		return IsSupported(Parameters.Platform);
+	}
+
+	BEGIN_SHADER_PARAMETER_STRUCT(FParameters, )
+		SHADER_PARAMETER_SRV(Texture2D, Texture)
+		SHADER_PARAMETER_UAV(RWTexture2D, StagingOutput)
+	END_SHADER_PARAMETER_STRUCT()
+};
+
+IMPLEMENT_GLOBAL_SHADER(FDumpTextureCS, "/Engine/Private/Tools/DumpTexture.usf", "MainCS", SF_Compute);
+
+BEGIN_SHADER_PARAMETER_STRUCT(FDumpTexturePass, )
+	SHADER_PARAMETER_RDG_TEXTURE_SRV(Texture2D, Texture)
+	RDG_TEXTURE_ACCESS_DYNAMIC(TextureAccess)
+END_SHADER_PARAMETER_STRUCT()
+
+BEGIN_SHADER_PARAMETER_STRUCT(FDumpBufferPass, )
+	RDG_BUFFER_ACCESS(Buffer, ERHIAccess::CopySrc)
+END_SHADER_PARAMETER_STRUCT()
+
+struct FRDGResourceDumpContext
+{
+	static constexpr const TCHAR* kBaseDir = TEXT("Base/");
+	static constexpr const TCHAR* kPassesDir = TEXT("Passes/");
+	static constexpr const TCHAR* kResourcesDir = TEXT("Resources/");
+	static constexpr const TCHAR* kStructuresDir = TEXT("Structures/");
+	static constexpr const TCHAR* kStructuresMetadataDir = TEXT("StructuresMetadata/");
+
+	bool bEnableDiskWrite = false;
+	FName ResourceCompressionName;
+	FString DumpingDirectoryPath;
+	FDateTime Time;
+	FGenericPlatformMemoryConstants MemoryConstants;
+	FGenericPlatformMemoryStats MemoryStats;
+	int32 ResourcesDumpPasses = 0;
+	int32 ResourcesDumpExecutedPasses = 0;
+	int32 PassesCount = 0;
+	TMap<const FRDGResource*, const FRDGPass*> LastResourceVersion;
+	TSet<const void*> IsDumpedToDisk;
+
+	// Pass being dumping individual draws
+	const FRDGPass* DrawDumpingPass = nullptr;
+	int32 DrawDumpCount = 0;
+
+	int32 MetadataFilesOpened = 0;
+	int64 MetadataFilesWriteBytes = 0;
+	int32 ResourceBinaryFilesOpened = 0;
+	int64 ResourceBinaryWriteBytes = 0;
+	int64 ResourceBinaryUncompressedBytes = 0;
+
+	enum class ETimingBucket : uint8
+	{
+		RHIReadbackCommands,
+		RHIReleaseResources,
+		GPUWait,
+		CPUPostProcessing,
+		CPUCompression,
+		MetadataFileWrite,
+		ResourceBinaryFileWrite,
+		MAX
+	};
+	mutable double TimingBucket[int32(ETimingBucket::MAX)];
+
+	class FTimeBucketMeasure
+	{
+	public:
+		FTimeBucketMeasure(FRDGResourceDumpContext* InDumpCtx, ETimingBucket InBucket)
+			: DumpCtx(InDumpCtx)
+			, Start(FPlatformTime::Cycles64())
+			, Bucket(InBucket)
+		{ }
+
+		~FTimeBucketMeasure()
+		{
+			uint64 End = FPlatformTime::Cycles64();
+			DumpCtx->TimingBucket[int32(Bucket)] += FPlatformTime::ToSeconds64(FMath::Max(End - Start, uint64(0)));
+		}
+
+	protected:
+		FRDGResourceDumpContext* const DumpCtx;
+
+	private:
+		const uint64 Start;
+		const ETimingBucket Bucket;
+	};
+
+	class FFileWriteCtx : public FTimeBucketMeasure
+	{
+	public:
+		FFileWriteCtx(FRDGResourceDumpContext* InDumpCtx, ETimingBucket InBucket, int64 WriteSizeBytes)
+			: FTimeBucketMeasure(InDumpCtx, InBucket)
+		{
+			if (InBucket == ETimingBucket::MetadataFileWrite)
+			{
+				DumpCtx->MetadataFilesOpened += 1;
+				DumpCtx->MetadataFilesWriteBytes += WriteSizeBytes;
+			}
+			else if (InBucket == ETimingBucket::ResourceBinaryFileWrite)
+			{
+				DumpCtx->ResourceBinaryFilesOpened += 1;
+				DumpCtx->ResourceBinaryWriteBytes += WriteSizeBytes;
+			}
+			else
+			{
+				unimplemented();
+			}
+		}
+	};
+
+	bool bShowInExplore = false;
+
+	FRDGResourceDumpContext()
+	{
+		for (int32 i = 0; i < int32(ETimingBucket::MAX); i++)
+		{
+			TimingBucket[i] = 0.0f;
+		}
+	}
+
+	bool IsDumpingFrame() const
+	{
+		check(IsInRenderingThread() || IsInParallelRenderingThread());
+		return !DumpingDirectoryPath.IsEmpty();
+	}
+
+	FString GetDumpFullPath(const FString& DumpRelativeFileName) const
+	{
+		check(bEnableDiskWrite);
+		return DumpingDirectoryPath / DumpRelativeFileName;
+	}
+
+	bool DumpStringToFile(FStringView OutputString, const FString& FileName, uint32 WriteFlags = FILEWRITE_None)
+	{
+		// Make it has if the write happened and was successful.
+		if (!bEnableDiskWrite)
+		{
+			return true;
+		}
+
+		FString FullPath = GetDumpFullPath(FileName);
+
+		FFileWriteCtx WriteCtx(this, ETimingBucket::MetadataFileWrite, OutputString.Len());
+		return FFileHelper::SaveStringToFile(
+			OutputString, *FullPath,
+			FFileHelper::EEncodingOptions::AutoDetect, &IFileManager::Get(), WriteFlags);
+	}
+
+	bool DumpJsonToFile(const TSharedPtr<FJsonObject>& JsonObject, const FString& FileName, uint32 WriteFlags = FILEWRITE_None)
+	{
+		FString OutputString;
+		TSharedRef<TJsonWriter<TCHAR, TPrettyJsonPrintPolicy<TCHAR>>> Writer = TJsonWriterFactory<TCHAR, TPrettyJsonPrintPolicy<TCHAR>>::Create(&OutputString);
+		FJsonSerializer::Serialize(JsonObject.ToSharedRef(), Writer);
+
+		return DumpStringToFile(OutputString, FileName, WriteFlags);
+	}
+
+	bool DumpBinaryToFile(TArrayView<const uint8> ArrayView, const FString& FileName)
+	{
+		// Make it has if the write happened and was successful.
+		if (!bEnableDiskWrite)
+		{
+			return true;
+		}
+
+		FString FullPath = GetDumpFullPath(FileName);
+		FFileWriteCtx WriteCtx(this, ETimingBucket::ResourceBinaryFileWrite, ArrayView.Num());
+		return FFileHelper::SaveArrayToFile(ArrayView, *FullPath);
+	}
+
+	bool DumpBinaryToFile(const uint8* Data, int64 DataByteSize, const FString& FileName)
+	{
+		// Make it has if the write happened and was successful.
+		if (!bEnableDiskWrite)
+		{
+			return true;
+		}
+
+		FString FullPath = GetDumpFullPath(FileName);
+		FFileWriteCtx WriteCtx(this, ETimingBucket::ResourceBinaryFileWrite, DataByteSize);
+
+		TUniquePtr<FArchive> Ar = TUniquePtr<FArchive>(IFileManager::Get().CreateFileWriter(*FullPath, /* WriteFlags = */ 0));
+		if (!Ar)
+		{
+			return false;
+		}
+		Ar->Serialize((void*)Data, DataByteSize);
+
+		// Always explicitly close to catch errors from flush/close
+		Ar->Close();
+
+		return !Ar->IsError() && !Ar->IsCriticalError();
+	}
+
+	bool IsResourceBinaryCompressable(int64 UncompressedSize) const
+	{
+		return !ResourceCompressionName.IsNone() && (UncompressedSize < (int64(1) << int64(32)));
+	}
+
+	bool DumpResourceBinaryToFile(const uint8* UncompressedData, int64 UncompressedSize, const FString& FileName)
+	{
+		ResourceBinaryUncompressedBytes += UncompressedSize;
+
+		if (!IsResourceBinaryCompressable(UncompressedSize))
+		{
+			return DumpBinaryToFile(UncompressedData, UncompressedSize, FileName);
+		}
+
+		int32 CompressedSize;
+		TArray64<uint8> CompressedArray;
+		{
+			FTimeBucketMeasure TimeBucketMeasure(this, ETimingBucket::CPUCompression);
+
+			int32 MaxCompressedSize = FCompression::CompressMemoryBound(ResourceCompressionName, UncompressedSize, COMPRESS_BiasSize);
+			CompressedArray.SetNumUninitialized(MaxCompressedSize);
+
+			CompressedSize = MaxCompressedSize;
+			bool bSuccess = FCompression::CompressMemory(
+				ResourceCompressionName,
+				/* out */ CompressedArray.GetData(), /* out */ CompressedSize,
+				UncompressedData, UncompressedSize,
+				COMPRESS_BiasSize);
+		}
+
+		return DumpBinaryToFile(CompressedArray.GetData(), CompressedSize, FileName);
+	}
+
+	bool IsUnsafeToDumpResource(SIZE_T ResourceByteSize, float DumpMemoryMultiplier) const
+	{
+		uint64 AproximatedStagingMemoryRequired = uint64(double(ResourceByteSize) * DumpMemoryMultiplier);
+		if (IsResourceBinaryCompressable(ResourceByteSize))
+		{
+			AproximatedStagingMemoryRequired += FCompression::CompressMemoryBound(ResourceCompressionName, ResourceByteSize, COMPRESS_BiasSize);
+		}
+
+		uint64 MaxMemoryAvailable = FMath::Min(MemoryStats.AvailablePhysical, MemoryStats.AvailableVirtual);
+
+		return AproximatedStagingMemoryRequired > MaxMemoryAvailable;
+	}
+
+	template<typename T>
+	static uint64 PtrToUint(const T* Ptr)
+	{
+		return static_cast<uint64>(reinterpret_cast<size_t>(Ptr));
+	}
+
+	template<typename T>
+	static FString PtrToString(const T* Ptr)
+	{
+		return FString::Printf(TEXT("%016x"), PtrToUint(Ptr));
+	}
+
+	static FString GetUniqueResourceName(const FRDGResource* Resource)
+	{
+		if (GDumpTestPrettifyResourceFileNames.GetValueOnRenderThread())
+		{
+			FString UniqueResourceName = FString::Printf(TEXT("%s.%016x"), Resource->Name, PtrToUint(Resource));
+			UniqueResourceName.ReplaceInline(TEXT("/"), TEXT(""));
+			UniqueResourceName.ReplaceInline(TEXT("\\"), TEXT(""));
+			return UniqueResourceName;
+		}
+		return PtrToString(Resource);
+	}
+
+	static FString GetUniqueSubResourceName(FRDGTextureSRVDesc SubResourceDesc)
+	{
+		check(SubResourceDesc.NumMipLevels == 1);
+
+		FString UniqueResourceName = GetUniqueResourceName(SubResourceDesc.Texture);
+
+		if (SubResourceDesc.Format == PF_X24_G8)
+		{
+			return FString::Printf(TEXT("%s.stencil"), *UniqueResourceName);
+		}
+
+		return FString::Printf(TEXT("%s.mip%d"), *UniqueResourceName, SubResourceDesc.MipLevel);
+	}
+
+	void ReleaseRHIResources(FRHICommandListImmediate& RHICmdList)
+	{
+		FTimeBucketMeasure TimeBucketMeasure(this, ETimingBucket::RHIReleaseResources);
+
+		// Flush the RHI resource memory so the readback memory can be fully reused in the next resource dump.
+		{
+			RHICmdList.SubmitCommandsAndFlushGPU();
+			RHICmdList.BlockUntilGPUIdle();
+			FRHIResource::FlushPendingDeletes(RHICmdList);
+			RHICmdList.FlushResources();
+			RHICmdList.ImmediateFlush(EImmediateFlushType::FlushRHIThreadFlushResources);
+		}
+	}
+
+	void UpdatePassProgress()
+	{
+		ResourcesDumpExecutedPasses++;
+
+		if (ResourcesDumpExecutedPasses % 10 == 0)
+		{
+			UE_LOG(LogRendererCore, Display, TEXT("Dumped %d / %d resources"), ResourcesDumpExecutedPasses, ResourcesDumpPasses);
+		}
+	}
+
+	void DumpRenderingCVarsToCSV() const
+	{
+		FString FileName = GetDumpFullPath(FString(FRDGResourceDumpContext::kBaseDir) / TEXT("ConsoleVariables.csv"));
+
+		TUniquePtr<FArchive> Ar = TUniquePtr<FArchive>(IFileManager::Get().CreateFileWriter(*FileName));
+		auto OnConsoleVariable = [&Ar](const TCHAR* CVarName, IConsoleObject* ConsoleObj)
+		{
+			if (ConsoleObj->TestFlags(ECVF_Unregistered))
+			{
+				return;
+			}
+
+			const IConsoleVariable* CVar = ConsoleObj->AsVariable();
+			if (!CVar)
+			{
+				return;
+			}
+
+			EConsoleVariableFlags CVarFlags = CVar->GetFlags();
+
+			const TCHAR* Type = nullptr;
+			if (CVar->IsVariableBool())
+			{
+				Type = TEXT("bool");
+			}
+			else if (CVar->IsVariableInt())
+			{
+				Type = TEXT("int32");
+			}
+			else if (CVar->IsVariableFloat())
+			{
+				Type = TEXT("float");
+			}
+			else if (CVar->IsVariableString())
+			{
+				Type = TEXT("FString");
+			}
+			else
+			{
+				return;
+			}
+
+			const TCHAR* SetByName = TEXT("");
+			switch (CVarFlags & ECVF_SetByMask)
+			{
+			case ECVF_SetByConstructor:
+			{
+				SetByName = TEXT("Constructor");
+				break;
+			}
+			case ECVF_SetByScalability:
+			{
+				SetByName = TEXT("Scalability");
+				break;
+			}
+			case ECVF_SetByGameSetting:
+			{
+				SetByName = TEXT("GameSetting");
+				break;
+			}
+			case ECVF_SetByProjectSetting:
+			{
+				SetByName = TEXT("ProjectSetting");
+				break;
+			}
+			case ECVF_SetBySystemSettingsIni:
+			{
+				SetByName = TEXT("SystemSettingsIni");
+				break;
+			}
+			case ECVF_SetByDeviceProfile:
+			{
+				SetByName = TEXT("DeviceProfile");
+				break;
+			}
+			case ECVF_SetByConsoleVariablesIni:
+			{
+				SetByName = TEXT("ConsoleVariablesIni");
+				break;
+			}
+			case ECVF_SetByCommandline:
+			{
+				SetByName = TEXT("Commandline");
+				break;
+			}
+			case ECVF_SetByCode:
+			{
+				SetByName = TEXT("Code");
+				break;
+			}
+			case ECVF_SetByConsole:
+			{
+				SetByName = TEXT("Console");
+				break;
+			}
+			default:
+				unimplemented();
+			}
+
+			FString Value = CVar->GetString();
+
+			FString CSVLine = FString::Printf(TEXT("%s,%s,%s,%s\n"), CVarName, Type, SetByName, *Value);
+			FStringView CSVLineView(CSVLine);
+
+			auto Src = StringCast<ANSICHAR>(CSVLineView.GetData(), CSVLineView.Len());
+			Ar->Serialize((ANSICHAR*)Src.Get(), Src.Length() * sizeof(ANSICHAR));
+		};
+
+		bool bSuccess = false;
+		if (Ar)
+		{
+			{
+				FString CSVLine = TEXT("CVar,Type,SetBy,Value\n");
+				FStringView CSVLineView(CSVLine);
+
+				auto Src = StringCast<ANSICHAR>(CSVLineView.GetData(), CSVLineView.Len());
+				Ar->Serialize((ANSICHAR*)Src.Get(), Src.Length() * sizeof(ANSICHAR));
+			}
+
+			if (GDumpRenderingConsoleVariablesCVar.GetValueOnGameThread() != 0)
+			{
+				IConsoleManager::Get().ForEachConsoleObjectThatStartsWith(FConsoleObjectVisitor::CreateLambda(OnConsoleVariable), TEXT(""));
+			}
+			else
+			{
+				IConsoleManager::Get().ForEachConsoleObjectThatStartsWith(FConsoleObjectVisitor::CreateLambda(OnConsoleVariable), TEXT("r.DumpGPU."));
+			}
+			// Always explicitly close to catch errors from flush/close
+			Ar->Close();
+
+			bSuccess = !Ar->IsError() && !Ar->IsCriticalError();
+		}
+
+		if (bSuccess)
+		{
+			UE_LOG(LogRendererCore, Display, TEXT("DumpGPU dumped rendering cvars to %s."), *FileName);
+		}
+		else
+		{
+			UE_LOG(LogRendererCore, Error, TEXT("DumpGPU had a file error when dumping rendering cvars to %s."), *FileName);
+		}
+	}
+
+	template<typename T>
+	bool IsDumped(const T* Ptr) const
+	{
+		return IsDumpedToDisk.Contains(static_cast<const void*>(Ptr));
+	}
+
+	template<typename T>
+	void SetDumped(const T* Ptr)
+	{
+		check(!IsDumped(Ptr));
+		if (IsDumpedToDisk.Num() % 1024 == 0)
+		{
+			IsDumpedToDisk.Reserve(IsDumpedToDisk.Num() + 1024);
+		}
+		IsDumpedToDisk.Add(static_cast<const void*>(Ptr));
+	}
+
+	void GetResourceDumpInfo(
+		const FRDGPass* Pass,
+		const FRDGResource* Resource,
+		bool bIsOutputResource,
+		bool* bOutDumpResourceInfos,
+		bool* bOutDumpResourceBinary)
+	{
+		*bOutDumpResourceInfos = false;
+		*bOutDumpResourceBinary = bIsOutputResource;
+
+		//if (bIsOutputResource)
+		//{
+		//	*OutResourceVersionDumpName = GetResourceVersionDumpName(Pass, Resource);
+		//}
+
+		if (!LastResourceVersion.Contains(Resource))
+		{
+			// First time we ever see this resource, so dump it's info to disk
+			*bOutDumpResourceInfos = true;
+
+			// If not an output, it might be a resource undumped by r.DumpGPU.Root or external texture so still dump it as v0.
+			if (!bIsOutputResource)
+			{
+				*bOutDumpResourceBinary = true;
+			}
+
+			if (LastResourceVersion.Num() % 1024 == 0)
+			{
+				LastResourceVersion.Reserve(LastResourceVersion.Num() + 1024);
+			}
+			LastResourceVersion.Add(Resource, Pass);
+		}
+		else
+		{
+			LastResourceVersion[Resource] = Pass;
+		}
+	}
+
+	FString ToJson(EPixelFormat Format)
+	{
+		FString PixelFormat = GPixelFormats[Format].Name;
+		if (!PixelFormat.StartsWith(TEXT("PF_")))
+		{
+			PixelFormat = FString::Printf(TEXT("PF_%s"), GPixelFormats[Format].Name);
+		}
+		return PixelFormat;
+	}
+
+	TSharedPtr<FJsonObject> ToJson(const FShaderParametersMetadata::FMember& Member)
+	{
+		TSharedPtr<FJsonObject> JsonObject = MakeShareable(new FJsonObject);
+		JsonObject->SetStringField(TEXT("Name"), Member.GetName());
+		JsonObject->SetStringField(TEXT("ShaderType"), Member.GetShaderType());
+		JsonObject->SetNumberField(TEXT("FileLine"), Member.GetFileLine());
+		JsonObject->SetNumberField(TEXT("Offset"), Member.GetOffset());
+		JsonObject->SetStringField(TEXT("BaseType"), GetUniformBufferBaseTypeString(Member.GetBaseType()));
+		JsonObject->SetNumberField(TEXT("Precision"), Member.GetPrecision());
+		JsonObject->SetNumberField(TEXT("NumRows"), Member.GetNumRows());
+		JsonObject->SetNumberField(TEXT("NumColumns"), Member.GetNumColumns());
+		JsonObject->SetNumberField(TEXT("NumElements"), Member.GetNumElements());
+		JsonObject->SetStringField(TEXT("StructMetadata"), *PtrToString(Member.GetStructMetadata()));
+		return JsonObject;
+	}
+
+	TSharedPtr<FJsonObject> ToJson(const FShaderParametersMetadata* Metadata)
+	{
+		TSharedPtr<FJsonObject> JsonObject = MakeShareable(new FJsonObject);
+		JsonObject->SetStringField(TEXT("StructTypeName"), Metadata->GetStructTypeName());
+		JsonObject->SetStringField(TEXT("ShaderVariableName"), Metadata->GetShaderVariableName());
+		JsonObject->SetStringField(TEXT("FileName"), Metadata->GetFileName());
+		JsonObject->SetNumberField(TEXT("FileLine"), Metadata->GetFileLine());
+		JsonObject->SetNumberField(TEXT("Size"), Metadata->GetSize());
+		//JsonObject->SetNumberField(TEXT("UseCase"), Metadata->GetUseCase());
+
+		{
+			TArray<TSharedPtr<FJsonValue>> Members;
+			for (const FShaderParametersMetadata::FMember& Member : Metadata->GetMembers())
+			{
+				Members.Add(MakeShared<FJsonValueObject>(ToJson(Member)));
+			}
+			JsonObject->SetArrayField(TEXT("Members"), Members);
+		}
+
+		return JsonObject;
+	}
+
+	TSharedPtr<FJsonObject> ToJson(const FString& UniqueResourceName, const TCHAR* Name, const FRDGTextureDesc& Desc)
+	{
+		FString PixelFormat = ToJson(Desc.Format);
+
+		int32 ResourceByteSize = Desc.Extent.X * Desc.Extent.Y * Desc.Depth * Desc.ArraySize * Desc.NumSamples * GPixelFormats[Desc.Format].BlockBytes;
+
+		TSharedPtr<FJsonObject> JsonObject = MakeShareable(new FJsonObject);
+		JsonObject->SetStringField(TEXT("Name"), Name);
+		JsonObject->SetStringField(TEXT("UniqueResourceName"), UniqueResourceName);
+		JsonObject->SetNumberField(TEXT("ByteSize"), ResourceByteSize);
+		JsonObject->SetStringField(TEXT("Desc"), TEXT("FRDGTextureDesc"));
+		JsonObject->SetStringField(TEXT("Type"), GetTextureDimensionString(Desc.Dimension));
+		JsonObject->SetStringField(TEXT("Format"), *PixelFormat);
+		JsonObject->SetNumberField(TEXT("ExtentX"), Desc.Extent.X);
+		JsonObject->SetNumberField(TEXT("ExtentY"), Desc.Extent.Y);
+		JsonObject->SetNumberField(TEXT("Depth"), Desc.Depth);
+		JsonObject->SetNumberField(TEXT("ArraySize"), Desc.ArraySize);
+		JsonObject->SetNumberField(TEXT("NumMips"), Desc.NumMips);
+		JsonObject->SetNumberField(TEXT("NumSamples"), Desc.NumSamples);
+
+		{
+			TArray<TSharedPtr<FJsonValue>> FlagsNames;
+			for (uint64 BitId = 0; BitId < uint64(8 * sizeof(ETextureCreateFlags)); BitId++)
+			{
+				ETextureCreateFlags Flag = ETextureCreateFlags(uint64(1) << BitId);
+				if (EnumHasAnyFlags(Desc.Flags, Flag))
+				{
+					FlagsNames.Add(MakeShareable(new FJsonValueString(GetTextureCreateFlagString(Flag))));
+				}
+			}
+			JsonObject->SetArrayField(TEXT("Flags"), FlagsNames);
+		}
+
+		return JsonObject;
+	}
+
+	TSharedPtr<FJsonObject> ToJson(const FString& UniqueResourceName, const TCHAR* Name, const FRDGBufferDesc& Desc, int32 ResourceByteSize)
+	{
+		TSharedPtr<FJsonObject> JsonObject = MakeShareable(new FJsonObject);
+		JsonObject->SetStringField(TEXT("Name"), Name);
+		JsonObject->SetStringField(TEXT("UniqueResourceName"), UniqueResourceName);
+		JsonObject->SetNumberField(TEXT("ByteSize"), ResourceByteSize);
+		JsonObject->SetStringField(TEXT("Desc"), TEXT("FRDGBufferDesc"));
+		JsonObject->SetStringField(TEXT("Type"), GetBufferUnderlyingTypeName(Desc.UnderlyingType));
+		JsonObject->SetNumberField(TEXT("BytesPerElement"), Desc.BytesPerElement);
+		JsonObject->SetNumberField(TEXT("NumElements"), Desc.NumElements);
+		JsonObject->SetStringField(TEXT("Metadata"), *PtrToString(Desc.Metadata));
+
+		{
+			TArray<TSharedPtr<FJsonValue>> UsageNames;
+			for (uint64 BitId = 0; BitId < uint64(8 * sizeof(EBufferUsageFlags)); BitId++)
+			{
+				EBufferUsageFlags Flag = EBufferUsageFlags(uint64(1) << BitId);
+				if (EnumHasAnyFlags(Desc.Usage, Flag))
+				{
+					UsageNames.Add(MakeShareable(new FJsonValueString(GetBufferUsageFlagString(Flag))));
+				}
+			}
+			JsonObject->SetArrayField(TEXT("Usage"), UsageNames);
+		}
+		return JsonObject;
+	}
+
+	void Dump(const FShaderParametersMetadata* Metadata)
+	{
+		if (IsDumped(Metadata))
+		{
+			return;
+		}
+
+		TSharedPtr<FJsonObject> JsonObject = ToJson(Metadata);
+		FString JsonPath = kStructuresMetadataDir / PtrToString(Metadata) + TEXT(".json");
+		DumpJsonToFile(JsonObject, JsonPath);
+		SetDumped(Metadata);
+
+		// Dump dependencies
+		Metadata->IterateStructureMetadataDependencies(
+		[&](const FShaderParametersMetadata* Struct)
+		{
+			if (Struct)
+			{
+				Dump(Struct);
+			}
+		});
+	}
+
+	struct FTextureSubresourceDumpDesc
+	{
+		FIntPoint SubResourceExtent = FIntPoint(0, 0);
+		SIZE_T ByteSize = 0;
+		bool bPreprocessForStaging = false;
+		FDumpTextureCS::ETextureType DumpTextureType = FDumpTextureCS::ETextureType::MAX;
+		EPixelFormat PreprocessedPixelFormat;
+
+		bool IsDumpSupported() const
+		{
+			return ByteSize != SIZE_T(0);
+		}
+	};
+	
+	FTextureSubresourceDumpDesc TranslateSubresourceDumpDesc(FRDGTextureSRVDesc SubresourceDesc)
+	{
+		const FRDGTextureDesc& Desc = SubresourceDesc.Texture->Desc;
+
+		FTextureSubresourceDumpDesc SubresourceDumpDesc;
+		SubresourceDumpDesc.PreprocessedPixelFormat = Desc.Format;
+
+		bool bIsUnsupported = false;
+
+		if (GPixelFormats[Desc.Format].BlockSizeX != 1 ||
+			GPixelFormats[Desc.Format].BlockSizeY != 1 ||
+			GPixelFormats[Desc.Format].BlockSizeZ != 1)
+		{
+			bIsUnsupported = true;
+		}
+
+		if (!bIsUnsupported && Desc.IsTexture2D() && !Desc.IsMultisample() && !Desc.IsTextureArray())
+		{
+			SubresourceDumpDesc.SubResourceExtent.X = Desc.Extent.X >> SubresourceDesc.MipLevel;
+			SubresourceDumpDesc.SubResourceExtent.Y = Desc.Extent.Y >> SubresourceDesc.MipLevel;
+
+			if (IsUintFormat(Desc.Format) || IsSintFormat(Desc.Format))
+			{
+				SubresourceDumpDesc.DumpTextureType = FDumpTextureCS::ETextureType::Texture2DUintNoMSAA;
+			}
+			else
+			{
+				SubresourceDumpDesc.DumpTextureType = FDumpTextureCS::ETextureType::Texture2DFloatNoMSAA;
+			}
+
+			if (SubresourceDesc.Format == PF_X24_G8)
+			{
+				SubresourceDumpDesc.PreprocessedPixelFormat = PF_R8_UINT;
+				SubresourceDumpDesc.DumpTextureType = FDumpTextureCS::ETextureType::Texture2DDepthStencilNoMSAA;
+			}
+			else if (Desc.Format == PF_DepthStencil)
+			{
+				SubresourceDumpDesc.PreprocessedPixelFormat = PF_R32_FLOAT;
+				SubresourceDumpDesc.DumpTextureType = FDumpTextureCS::ETextureType::Texture2DFloatNoMSAA;
+			}
+			else if (Desc.Format == PF_ShadowDepth)
+			{
+				SubresourceDumpDesc.PreprocessedPixelFormat = PF_R32_FLOAT;
+				SubresourceDumpDesc.DumpTextureType = FDumpTextureCS::ETextureType::Texture2DFloatNoMSAA;
+			}
+			else if (Desc.Format == PF_D24)
+			{
+				SubresourceDumpDesc.PreprocessedPixelFormat = PF_R32_FLOAT;
+				SubresourceDumpDesc.DumpTextureType = FDumpTextureCS::ETextureType::Texture2DFloatNoMSAA;
+			}
+		}
+
+		SubresourceDumpDesc.ByteSize = SIZE_T(SubresourceDumpDesc.SubResourceExtent.X) * SIZE_T(SubresourceDumpDesc.SubResourceExtent.Y) * SIZE_T(GPixelFormats[SubresourceDumpDesc.PreprocessedPixelFormat].BlockBytes);
+
+		// Whether the subresource need preprocessing pass before copy into staging.
+		SubresourceDumpDesc.bPreprocessForStaging = SubresourceDumpDesc.PreprocessedPixelFormat != Desc.Format || SubresourceDesc.Texture->Desc.NumMips > 1;
+
+		return SubresourceDumpDesc;
+	}
+
+	void DumpTextureSubResource(
+		FRHICommandListImmediate& RHICmdList,
+		const TCHAR* TextureDebugName,
+		FRHITexture* Texture,
+		FRHIShaderResourceView* SubResourceSRV,
+		const FTextureSubresourceDumpDesc& SubresourceDumpDesc,
+		const FString& DumpFilePath)
+	{
+		check(IsInRenderingThread());
+
+		// Preprocess
+		FTextureRHIRef StagingSrcTexture;
+		EPixelFormat PreprocessedPixelFormat = SubresourceDumpDesc.PreprocessedPixelFormat;
+		SIZE_T SubresourceByteSize = SubresourceDumpDesc.ByteSize;
+		FTextureRHIRef StagingTexture;
+
+		{
+			FTimeBucketMeasure TimeBucketMeasure(this, ETimingBucket::RHIReadbackCommands);
+			if (SubresourceDumpDesc.bPreprocessForStaging)
+			{
+				// Some RHIs (GL) only support 32Bit single channel images as CS output
+				if (IsOpenGLPlatform(GMaxRHIShaderPlatform) &&
+					GPixelFormats[PreprocessedPixelFormat].NumComponents == 1 && 
+					GPixelFormats[PreprocessedPixelFormat].BlockBytes < 4)
+				{
+					SubresourceByteSize*= (4 / GPixelFormats[PreprocessedPixelFormat].BlockBytes);
+					PreprocessedPixelFormat = PF_R32_UINT;
+				}
+						
+				{
+					FRHIResourceCreateInfo CreateInfo(TEXT("DumpGPU.PreprocessTexture"));
+					StagingSrcTexture = RHICreateTexture2D(
+						SubresourceDumpDesc.SubResourceExtent.X,
+						SubresourceDumpDesc.SubResourceExtent.Y,
+						(uint8)PreprocessedPixelFormat,
+						/* NumMips = */ 1,
+						/* NumSamples = */ 1,
+						TexCreate_UAV | TexCreate_ShaderResource | TexCreate_HideInVisualizeTexture,
+						CreateInfo);
+				}
+
+				FUnorderedAccessViewRHIRef StagingOutput = RHICreateUnorderedAccessView(StagingSrcTexture, /* MipLevel = */ 0);
+
+				RHICmdList.Transition(FRHITransitionInfo(StagingSrcTexture, ERHIAccess::Unknown, ERHIAccess::UAVCompute));
+
+				FDumpTextureCS::FPermutationDomain PermutationVector;
+				PermutationVector.Set<FDumpTextureCS::FTextureTypeDim>(SubresourceDumpDesc.DumpTextureType);
+				TShaderMapRef<FDumpTextureCS> ComputeShader(GetGlobalShaderMap(GMaxRHIShaderPlatform), PermutationVector);
+
+				FDumpTextureCS::FParameters ShaderParameters;
+				ShaderParameters.Texture = SubResourceSRV;
+				ShaderParameters.StagingOutput = StagingOutput;
+				FComputeShaderUtils::Dispatch(
+					RHICmdList,
+					ComputeShader,
+					ShaderParameters,
+					FComputeShaderUtils::GetGroupCount(SubresourceDumpDesc.SubResourceExtent, 8));
+
+				RHICmdList.Transition(FRHITransitionInfo(StagingSrcTexture, ERHIAccess::UAVCompute, ERHIAccess::CopySrc));
+			}
+			else
+			{
+				StagingSrcTexture = Texture;
+			}
+
+			// Copy the texture for CPU readback
+			{
+				FRHIResourceCreateInfo CreateInfo(TEXT("DumpGPU.StagingTexture"));
+				StagingTexture = RHICreateTexture2D(
+					SubresourceDumpDesc.SubResourceExtent.X,
+					SubresourceDumpDesc.SubResourceExtent.Y,
+					(uint8)PreprocessedPixelFormat,
+					/* NumMips = */ 1,
+					/* NumSamples = */ 1,
+					TexCreate_CPUReadback | TexCreate_HideInVisualizeTexture,
+					CreateInfo);
+
+				RHICmdList.Transition(FRHITransitionInfo(StagingTexture, ERHIAccess::Unknown, ERHIAccess::CopyDest));
+
+				// Ensure this copy call does not perform any transitions. We're handling them manually.
+				FResolveParams ResolveParams;
+				ResolveParams.SourceAccessFinal = ERHIAccess::Unknown;
+				ResolveParams.DestAccessFinal = ERHIAccess::Unknown;
+
+				// Transfer memory GPU -> CPU
+				RHICmdList.CopyToResolveTarget(StagingSrcTexture, StagingTexture, ResolveParams);
+
+				RHICmdList.Transition(FRHITransitionInfo(StagingTexture, ERHIAccess::CopyDest, ERHIAccess::CPURead));
+			}
+		}
+
+		// Submit to GPU and wait for completion.
+		static const FName FenceName(TEXT("DumpGPU.TextureFence"));
+		FGPUFenceRHIRef Fence = RHICreateGPUFence(FenceName);
+		{
+			FTimeBucketMeasure TimeBucketMeasure(this, ETimingBucket::GPUWait);
+
+			Fence->Clear();
+			RHICmdList.WriteGPUFence(Fence);
+			RHICmdList.SubmitCommandsAndFlushGPU();
+			RHICmdList.BlockUntilGPUIdle();
+		}
+
+		void* Content = nullptr;
+		int32 RowPitchInPixels = 0;
+		int32 ColumnPitchInPixels = 0;
+
+		// jhoerner_todo 12/9/2021:  pick arbitrary GPU out of mask to avoid assert.  Eventually want to dump results for all GPUs, but
+		// I need to understand how to modify the dumping logic, and this works for now (usually when debugging, the bugs happen on
+		// secondary GPUs, so I figure the last index is most useful if we need to pick one).  I also would like the dump to include
+		// information about the GPUMask for each pass, and perhaps have the dump include the final state of all external resources
+		// modified by the graph (especially useful for MGPU, where we are concerned about cross-view or cross-frame state).
+		uint32 GPUIndex = RHICmdList.GetGPUMask().GetLastIndex();
+
+		RHICmdList.MapStagingSurface(StagingTexture, Fence.GetReference(), Content, RowPitchInPixels, ColumnPitchInPixels, GPUIndex);
+
+		if (Content)
+		{
+			TArray64<uint8> Array;
+			{
+				FTimeBucketMeasure TimeBucketMeasure(this, ETimingBucket::CPUPostProcessing);
+
+				Array.SetNumUninitialized(SubresourceByteSize);
+
+				SIZE_T BytePerPixel = SIZE_T(GPixelFormats[PreprocessedPixelFormat].BlockBytes);
+
+				const uint8* SrcData = static_cast<const uint8*>(Content);
+
+				for (int32 y = 0; y < SubresourceDumpDesc.SubResourceExtent.Y; y++)
+				{
+					// Flip the data to be bottom left corner for the WebGL viewer.
+					const uint8* SrcPos = SrcData + SIZE_T(SubresourceDumpDesc.SubResourceExtent.Y - 1 - y) * SIZE_T(RowPitchInPixels) * BytePerPixel;
+					uint8* DstPos = (&Array[0]) + SIZE_T(y) * SIZE_T(SubresourceDumpDesc.SubResourceExtent.X) * BytePerPixel;
+
+					FPlatformMemory::Memmove(DstPos, SrcPos, SIZE_T(SubresourceDumpDesc.SubResourceExtent.X) * BytePerPixel);
+				}
+
+				RHICmdList.UnmapStagingSurface(StagingTexture, GPUIndex);
+
+				if (PreprocessedPixelFormat != SubresourceDumpDesc.PreprocessedPixelFormat)
+				{
+					// Convert 32Bit values back to 16 or 8bit
+					const int32 DstPixelNumBytes = GPixelFormats[SubresourceDumpDesc.PreprocessedPixelFormat].BlockBytes;
+					const uint32* SrcData32 = (const uint32*)Array.GetData();
+					uint8* DstData8 = Array.GetData();
+					uint16* DstData16 = (uint16*)Array.GetData();
+											
+					for (int32 Index = 0; Index < Array.Num()/4; Index++)
+					{
+						uint32 Value32 = SrcData32[Index];
+						if (DstPixelNumBytes == 2)
+						{
+							DstData16[Index] = (uint16)Value32;
+						}
+						else
+						{
+							DstData8[Index] = (uint8)Value32;
+						}
+					}
+
+					Array.SetNum(Array.Num() / (4 / DstPixelNumBytes));
+				}
+			}
+
+			DumpResourceBinaryToFile(Array.GetData(), Array.Num(), DumpFilePath);
+		}
+		else
+		{
+			UE_LOG(LogRendererCore, Warning, TEXT("RHICmdList.MapStagingSurface() to dump texture %s failed."), TextureDebugName);
+		}
+	}
+
+	void DumpDrawTextureSubResource(
+		FRHICommandList& RHICmdList,
+		FRDGTextureSRVDesc SubresourceDesc,
+		ERHIAccess RHIAccessState)
+	{
+		check(IsInRenderingThread());
+
+		FRHICommandListImmediate& RHICmdListImmediate = FRHICommandListExecutor::GetImmediateCommandList();
+		check(&RHICmdListImmediate == &RHICmdList);
+
+		const FString UniqueResourceSubResourceName = GetUniqueSubResourceName(SubresourceDesc);
+		const FTextureSubresourceDumpDesc SubresourceDumpDesc = TranslateSubresourceDumpDesc(SubresourceDesc);
+
+		if (!SubresourceDumpDesc.IsDumpSupported())
+		{
+			return;
+		}
+
+		FRHITexture* RHITexture = SubresourceDesc.Texture->GetRHI();
+
+		FShaderResourceViewRHIRef SubResourceSRV;
+		if (SubresourceDumpDesc.bPreprocessForStaging)
+		{
+			SubResourceSRV = RHICreateShaderResourceView(RHITexture, FRHITextureSRVCreateInfo(SubresourceDesc));
+			RHICmdListImmediate.Transition(FRHITransitionInfo(RHITexture, ERHIAccess::Unknown, ERHIAccess::SRVCompute));
+		}
+		else
+		{
+			RHICmdListImmediate.Transition(FRHITransitionInfo(RHITexture, RHIAccessState, ERHIAccess::CopySrc));
+		}
+
+		FString DumpFilePath = kResourcesDir / FString::Printf(
+			TEXT("%s.v%016x.d%d.bin"),
+			*UniqueResourceSubResourceName,
+			PtrToUint(DrawDumpingPass),
+			DrawDumpCount);
+
+		DumpTextureSubResource(
+			RHICmdListImmediate,
+			SubresourceDesc.Texture->Name,
+			RHITexture,
+			SubResourceSRV,
+			SubresourceDumpDesc,
+			DumpFilePath);
+
+		if (SubresourceDumpDesc.bPreprocessForStaging)
+		{
+			RHICmdListImmediate.Transition(FRHITransitionInfo(RHITexture, ERHIAccess::SRVCompute, RHIAccessState));
+		}
+		else
+		{
+			RHICmdListImmediate.Transition(FRHITransitionInfo(RHITexture, ERHIAccess::CopySrc, RHIAccessState));
+		}
+
+		SubResourceSRV = nullptr;
+		ReleaseRHIResources(RHICmdListImmediate);
+	}
+
+	void AddDumpTextureSubResourcePass(
+		FRDGBuilder& GraphBuilder,
+		TArray<TSharedPtr<FJsonValue>>& InputResourceNames,
+		TArray<TSharedPtr<FJsonValue>>& OutputResourceNames,
+		const FRDGPass* Pass,
+		FRDGTextureSRVDesc SubresourceDesc,
+		bool bIsOutputResource)
+	{
+		int32 DumpTextureMode = GDumpTextureCVar.GetValueOnRenderThread();
+
+		if (DumpTextureMode == 0)
+		{
+			return;
+		}
+
+		const FRDGTextureDesc& Desc = SubresourceDesc.Texture->Desc;
+		const FString UniqueResourceName = GetUniqueResourceName(SubresourceDesc.Texture);
+		const FString UniqueResourceSubResourceName = GetUniqueSubResourceName(SubresourceDesc);
+		const FTextureSubresourceDumpDesc SubresourceDumpDesc = TranslateSubresourceDumpDesc(SubresourceDesc);
+
+		if (bIsOutputResource)
+		{
+			OutputResourceNames.AddUnique(MakeShareable(new FJsonValueString(UniqueResourceSubResourceName)));
+		}
+		else
+		{
+			InputResourceNames.AddUnique(MakeShareable(new FJsonValueString(UniqueResourceSubResourceName)));
+		}
+
+		bool bDumpResourceInfos;
+		bool bDumpResourceBinary;
+		GetResourceDumpInfo(Pass, SubresourceDesc.Texture, bIsOutputResource, &bDumpResourceInfos, &bDumpResourceBinary);
+
+		// Dump the information of the texture to json file.
+		if (bDumpResourceInfos)
+		{
+			TSharedPtr<FJsonObject> JsonObject = ToJson(UniqueResourceName, SubresourceDesc.Texture->Name, Desc);
+			DumpJsonToFile(JsonObject, FString(FRDGResourceDumpContext::kBaseDir) / TEXT("ResourceDescs.json"), FILEWRITE_Append);
+		}
+
+		if (!SubresourceDumpDesc.IsDumpSupported())
+		{
+			return;
+		}
+
+		// Early return if this resource shouldn't be dumped.
+		if (!bDumpResourceBinary || DumpTextureMode != 2)
+		{
+			return;
+		}
+
+		FString DumpFilePath = kResourcesDir / FString::Printf(TEXT("%s.v%016x.bin"), *UniqueResourceSubResourceName, PtrToUint(bIsOutputResource ? Pass : nullptr));
+
+		// Verify there is enough available memory to dump the resource.
+		if (IsUnsafeToDumpResource(SubresourceDumpDesc.ByteSize, 2.2f + (SubresourceDumpDesc.bPreprocessForStaging ? 1.0f : 0.0f)))
+		{
+			UE_LOG(LogRendererCore, Warning, TEXT("Not dumping %s because of insuficient memory available for staging texture."), *DumpFilePath);
+			return;
+		}
+
+		// Dump the resource's binary to a .bin file.
+		{
+			FDumpTexturePass* PassParameters = GraphBuilder.AllocParameters<FDumpTexturePass>();
+			if (SubresourceDumpDesc.bPreprocessForStaging)
+			{
+				if (!(SubresourceDesc.Texture->Desc.Flags & TexCreate_ShaderResource))
+				{
+					UE_LOG(LogRendererCore, Warning, TEXT("Not dumping %s because requires copy to staging texture using compute, but is missing TexCreate_ShaderResource."), *UniqueResourceSubResourceName);
+					return;
+				}
+
+				if (!FDumpTextureCS::IsSupported(GMaxRHIShaderPlatform))
+				{
+					UE_LOG(LogRendererCore, Warning, TEXT("Not dumping %s because FDumpTextureCS compute shader is not supported."), *UniqueResourceSubResourceName);
+					return;
+				}
+
+				PassParameters->Texture = GraphBuilder.CreateSRV(SubresourceDesc);
+			}
+			else
+			{
+				PassParameters->TextureAccess = FRDGTextureAccess(SubresourceDesc.Texture, ERHIAccess::CopySrc);
+			}
+
+			GraphBuilder.AddPass(
+				RDG_EVENT_NAME("RDG DumpTexture(%s -> %s) %dx%d",
+					SubresourceDesc.Texture->Name, *DumpFilePath,
+					SubresourceDumpDesc.SubResourceExtent.X, SubresourceDumpDesc.SubResourceExtent.Y),
+				PassParameters,
+				(SubresourceDumpDesc.bPreprocessForStaging ? ERDGPassFlags::Compute : ERDGPassFlags::Copy) | ERDGPassFlags::NeverCull,
+				[
+					PassParameters, this, DumpFilePath,
+					SubresourceDesc, SubresourceDumpDesc]
+				(FRHICommandListImmediate& RHICmdList)
+			{
+
+				this->DumpTextureSubResource(
+					RHICmdList,
+					SubresourceDesc.Texture->Name,
+					SubresourceDumpDesc.bPreprocessForStaging ? nullptr : SubresourceDesc.Texture->GetRHI(),
+					SubresourceDumpDesc.bPreprocessForStaging ? PassParameters->Texture->GetRHI() : nullptr,
+					SubresourceDumpDesc,
+					DumpFilePath);
+				this->ReleaseRHIResources(RHICmdList);
+				this->UpdatePassProgress();
+			});
+
+			ResourcesDumpPasses++;
+		}
+	}
+
+	void AddDumpTexturePasses(
+		FRDGBuilder& GraphBuilder,
+		TArray<TSharedPtr<FJsonValue>>& InputResourceNames,
+		TArray<TSharedPtr<FJsonValue>>& OutputResourceNames,
+		const FRDGPass* Pass,
+		FRDGTextureSRVDesc SubresourceRangeDesc,
+		bool bIsOutputResource)
+	{
+		if (SubresourceRangeDesc.Format == PF_X24_G8)
+		{
+			AddDumpTextureSubResourcePass(
+				GraphBuilder,
+				InputResourceNames,
+				OutputResourceNames,
+				Pass,
+				SubresourceRangeDesc,
+				bIsOutputResource);
+		}
+		else
+		{
+			for (int32 MipLevel = SubresourceRangeDesc.MipLevel; MipLevel < (SubresourceRangeDesc.MipLevel + SubresourceRangeDesc.NumMipLevels); MipLevel++)
+			{
+				FRDGTextureSRVDesc SubresourceDesc = FRDGTextureSRVDesc::CreateForMipLevel(SubresourceRangeDesc.Texture, MipLevel);
+				AddDumpTextureSubResourcePass(
+					GraphBuilder,
+					InputResourceNames,
+					OutputResourceNames,
+					Pass,
+					SubresourceDesc,
+					bIsOutputResource);
+			}
+		}
+	}
+
+	void AddDumpBufferPass(
+		FRDGBuilder& GraphBuilder,
+		TArray<TSharedPtr<FJsonValue>>& InputResourceNames,
+		TArray<TSharedPtr<FJsonValue>>& OutputResourceNames,
+		const FRDGPass* Pass,
+		FRDGBuffer* Buffer,
+		bool bIsOutputResource)
+	{
+		int32 DumpBufferMode = GDumpTextureCVar.GetValueOnRenderThread();
+
+		if (DumpBufferMode == 0)
+		{
+			return;
+		}
+
+		FString UniqueResourceName = GetUniqueResourceName(Buffer);
+
+		if (bIsOutputResource)
+		{
+			OutputResourceNames.AddUnique(MakeShareable(new FJsonValueString(UniqueResourceName)));
+		}
+		else
+		{
+			InputResourceNames.AddUnique(MakeShareable(new FJsonValueString(UniqueResourceName)));
+		}
+
+		const FRDGBufferDesc& Desc = Buffer->Desc;
+		const int32 ByteSize = Desc.GetTotalNumBytes();
+
+		bool bDumpResourceInfos;
+		bool bDumpResourceBinary;
+		GetResourceDumpInfo(Pass, Buffer, bIsOutputResource, &bDumpResourceInfos, &bDumpResourceBinary);
+
+		// Dump the information of the buffer to json file.
+		if (bDumpResourceInfos)
+		{
+			TSharedPtr<FJsonObject> JsonObject = ToJson(UniqueResourceName, Buffer->Name, Desc, ByteSize);
+			DumpJsonToFile(JsonObject, FString(FRDGResourceDumpContext::kBaseDir) / TEXT("ResourceDescs.json"), FILEWRITE_Append);
+
+			if (Desc.Metadata && DumpBufferMode == 2)
+			{
+				Dump(Desc.Metadata);
+			}
+		}
+
+		if (Desc.UnderlyingType == FRDGBufferDesc::EUnderlyingType::AccelerationStructure)
+		{
+			return;
+		}
+
+		// Dump the resource's binary to a .bin file.
+		if (bDumpResourceBinary && DumpBufferMode == 2)
+		{
+			FString DumpFilePath = kResourcesDir / FString::Printf(TEXT("%s.v%016x.bin"), *UniqueResourceName, PtrToUint(bIsOutputResource ? Pass : nullptr));
+
+			if (IsUnsafeToDumpResource(ByteSize, 1.2f))
+			{
+				UE_LOG(LogRendererCore, Warning, TEXT("Not dumping %s because of insuficient memory available for staging buffer."), *DumpFilePath);
+				return;
+			}
+
+			FDumpBufferPass* PassParameters = GraphBuilder.AllocParameters<FDumpBufferPass>();
+			PassParameters->Buffer = Buffer;
+
+			GraphBuilder.AddPass(
+				RDG_EVENT_NAME("RDG DumpBuffer(%s -> %s)", Buffer->Name, *DumpFilePath),
+				PassParameters,
+				ERDGPassFlags::Readback,
+				[this, DumpFilePath, Buffer, ByteSize](FRHICommandListImmediate& RHICmdList)
+			{
+				check(IsInRenderingThread());
+				FStagingBufferRHIRef StagingBuffer;
+				{
+					FTimeBucketMeasure TimeBucketMeasure(this, ETimingBucket::RHIReadbackCommands);
+					StagingBuffer = RHICreateStagingBuffer();
+
+					// Transfer memory GPU -> CPU
+					RHICmdList.CopyToStagingBuffer(Buffer->GetRHI(), StagingBuffer, 0, ByteSize);
+				}
+
+				// Submit to GPU and wait for completion.
+				static const FName FenceName(TEXT("DumpGPU.BufferFence"));
+				FGPUFenceRHIRef Fence = RHICreateGPUFence(FenceName);
+				{
+					FTimeBucketMeasure TimeBucketMeasure(this, ETimingBucket::GPUWait);
+					Fence->Clear();
+					RHICmdList.WriteGPUFence(Fence);
+					RHICmdList.SubmitCommandsAndFlushGPU();
+					RHICmdList.BlockUntilGPUIdle();
+				}
+
+				void* Content = RHICmdList.LockStagingBuffer(StagingBuffer, Fence.GetReference(), 0, ByteSize);
+				if (Content)
+				{
+					DumpResourceBinaryToFile(reinterpret_cast<const uint8*>(Content), ByteSize, DumpFilePath);
+
+					RHICmdList.UnlockStagingBuffer(StagingBuffer);
+				}
+				else
+				{
+					UE_LOG(LogRendererCore, Warning, TEXT("RHICmdList.LockStagingBuffer() to dump buffer %s failed."), Buffer->Name);
+				}
+
+				StagingBuffer = nullptr;
+				Fence = nullptr;
+				this->ReleaseRHIResources(RHICmdList);
+				this->UpdatePassProgress();
+			});
+
+			ResourcesDumpPasses++;
+		}
+	}
+
+	// Look whether the pass matches matches r.DumpGPU.Root
+	bool IsDumpingPass(const FRDGPass* Pass)
+	{
+		FString RootWildcardString = GDumpGPURootCVar.GetValueOnRenderThread();
+		FWildcardString WildcardFilter(RootWildcardString);
+
+		bool bDumpPass = (RootWildcardString == TEXT("*"));
+
+		if (!bDumpPass)
+		{
+			bDumpPass = WildcardFilter.IsMatch(Pass->GetEventName().GetTCHAR());
+		}
+
+		#if RDG_GPU_SCOPES
+		if (!bDumpPass)
+		{
+			const FRDGEventScope* ParentScope = Pass->GetGPUScopes().Event;
+
+			while (ParentScope)
+			{
+				bDumpPass = bDumpPass || WildcardFilter.IsMatch(ParentScope->Name.GetTCHAR());
+				ParentScope = ParentScope->ParentScope;
+			}
+		}
+		#endif
+
+		return bDumpPass;
+	}
+};
+
+// 0 = not dumping, MAX_uint64 dump request for next frame, otherwise dump frame counter
+static uint64 DumpingFrameCounter_GameThread = 0;
+FRDGResourceDumpContext GRDGResourceDumpContext;
+
+}
+
+bool IsDumpingRDGResources()
+{
+	return GRDGResourceDumpContext.IsDumpingFrame();
+}
+
+void FRDGBuilder::InitResourceDump()
+{
+	if(DumpingFrameCounter_GameThread == MAX_uint64)
+	{
+		DumpingFrameCounter_GameThread = GFrameCounter;
+	}
+}
+
+FString FRDGBuilder::BeginResourceDump(const TArray<FString>& Args)
+{
+	check(IsInGameThread());
+
+	if (DumpingFrameCounter_GameThread != 0)
+	{
+		return FString();
+	}
+
+	FRDGResourceDumpContext NewResourceDumpContext;
+
+	NewResourceDumpContext.Time = FDateTime::Now();
+	{
+		FString CVarDirectoryPath = GDumpGPUDirectoryCVar.GetValueOnGameThread();
+		FString EnvDirectoryPath = FPlatformMisc::GetEnvironmentVariable(TEXT("UE-DumpGPUPath"));
+
+		FString DirectoryPath;
+		if (!CVarDirectoryPath.IsEmpty())
+		{
+			DirectoryPath = CVarDirectoryPath;
+		}
+		else if (!EnvDirectoryPath.IsEmpty())
+		{
+			DirectoryPath = EnvDirectoryPath;
+		}
+		else
+		{
+			DirectoryPath = FPaths::ProjectSavedDir() / TEXT("GPUDumps/");
+		}
+		NewResourceDumpContext.DumpingDirectoryPath = DirectoryPath / FApp::GetProjectName() + TEXT("-") + FPlatformProperties::PlatformName() + TEXT("-") + NewResourceDumpContext.Time.ToString() + TEXT("/");
+	}
+	NewResourceDumpContext.bEnableDiskWrite = GDumpTestEnableDiskWrite.GetValueOnGameThread() != 0;
+	{
+		if (GDumpGPUCompressResources.GetValueOnGameThread() == 1)
+		{
+			NewResourceDumpContext.ResourceCompressionName = NAME_Zlib;
+		}
+		else if (GDumpGPUCompressResources.GetValueOnGameThread() == 2)
+		{
+			NewResourceDumpContext.ResourceCompressionName = NAME_Gzip;
+		}
+	}
+	NewResourceDumpContext.bShowInExplore = NewResourceDumpContext.bEnableDiskWrite && GDumpExploreCVar.GetValueOnGameThread() != 0;
+	NewResourceDumpContext.MemoryConstants = FPlatformMemory::GetConstants();
+	NewResourceDumpContext.MemoryStats = FPlatformMemory::GetStats();
+
+	IPlatformFile& PlatformFile = FPlatformFileManager::Get().GetPlatformFile();
+	if (NewResourceDumpContext.bEnableDiskWrite)
+	{
+		if (!PlatformFile.DirectoryExists(*NewResourceDumpContext.DumpingDirectoryPath))
+		{
+			PlatformFile.CreateDirectoryTree(*NewResourceDumpContext.DumpingDirectoryPath);
+		}
+		PlatformFile.CreateDirectoryTree(*(NewResourceDumpContext.DumpingDirectoryPath / FRDGResourceDumpContext::kBaseDir));
+		PlatformFile.CreateDirectoryTree(*(NewResourceDumpContext.DumpingDirectoryPath / FRDGResourceDumpContext::kResourcesDir));
+
+		NewResourceDumpContext.DumpStringToFile(TEXT(""), FString(FRDGResourceDumpContext::kBaseDir) / TEXT("Passes.json"));
+		NewResourceDumpContext.DumpStringToFile(TEXT(""), FString(FRDGResourceDumpContext::kBaseDir) / TEXT("ResourceDescs.json"));
+		NewResourceDumpContext.DumpStringToFile(TEXT(""), FString(FRDGResourceDumpContext::kBaseDir) / TEXT("PassDrawCounts.json"));
+	}
+
+	// Output informations
+	{
+		const TCHAR* BranchName = BuildSettings::GetBranchName();
+		const TCHAR* BuildDate = BuildSettings::GetBuildDate();
+		const TCHAR* BuildVersion = BuildSettings::GetBuildVersion();
+
+		FString BuildConfiguration = LexToString(FApp::GetBuildConfiguration());
+		FString BuildTarget = LexToString(FApp::GetBuildTargetType());
+
+		FGPUDriverInfo GPUDriverInfo = FPlatformMisc::GetGPUDriverInfo(GRHIAdapterName);
+
+		TSharedPtr<FJsonObject> JsonObject = MakeShareable(new FJsonObject);
+		JsonObject->SetStringField(TEXT("Project"), FApp::GetProjectName());
+		JsonObject->SetNumberField(TEXT("EngineMajorVersion"), ENGINE_MAJOR_VERSION);
+		JsonObject->SetNumberField(TEXT("EngineMinorVersion"), ENGINE_MINOR_VERSION);
+		JsonObject->SetNumberField(TEXT("EnginePatchVersion"), ENGINE_PATCH_VERSION);
+		JsonObject->SetStringField(TEXT("BuildBranch"), BranchName ? BranchName : TEXT(""));
+		JsonObject->SetStringField(TEXT("BuildDate"), BuildDate ? BuildDate : TEXT(""));
+		JsonObject->SetStringField(TEXT("BuildVersion"), BuildVersion ? BuildVersion : TEXT(""));
+		JsonObject->SetStringField(TEXT("BuildTarget"), BuildTarget);
+		JsonObject->SetStringField(TEXT("BuildConfiguration"), BuildConfiguration);
+		JsonObject->SetNumberField(TEXT("Build64Bits"), (PLATFORM_64BITS ? 1 : 0));
+		JsonObject->SetStringField(TEXT("Platform"), FPlatformProperties::IniPlatformName());
+		JsonObject->SetStringField(TEXT("DeviceName"), FPlatformProcess::ComputerName());
+		JsonObject->SetStringField(TEXT("CPUVendor"), FPlatformMisc::GetCPUVendor());
+		JsonObject->SetStringField(TEXT("CPUBrand"), FPlatformMisc::GetCPUBrand());
+		JsonObject->SetNumberField(TEXT("CPUNumberOfCores"), FPlatformMisc::NumberOfCores());
+		JsonObject->SetNumberField(TEXT("CPUNumberOfCoresIncludingHyperthreads"), FPlatformMisc::NumberOfCoresIncludingHyperthreads());
+		JsonObject->SetStringField(TEXT("GPUVendor"), RHIVendorIdToString());
+		JsonObject->SetStringField(TEXT("GPUDeviceDescription"), GPUDriverInfo.DeviceDescription);
+		JsonObject->SetStringField(TEXT("GPUDriverUserVersion"), GPUDriverInfo.UserDriverVersion);
+		JsonObject->SetStringField(TEXT("GPUDriverInternalVersion"), GPUDriverInfo.GetUnifiedDriverVersion());
+		JsonObject->SetStringField(TEXT("GPUDriverDate"), GPUDriverInfo.DriverDate);
+		JsonObject->SetNumberField(TEXT("MemoryTotalPhysical"), NewResourceDumpContext.MemoryConstants.TotalPhysical);
+		JsonObject->SetNumberField(TEXT("MemoryPageSize"), NewResourceDumpContext.MemoryConstants.PageSize);
+		JsonObject->SetStringField(TEXT("RHI"), GDynamicRHI->GetName());
+		JsonObject->SetStringField(TEXT("RHIMaxFeatureLevel"), LexToString(GMaxRHIFeatureLevel));
+		JsonObject->SetStringField(TEXT("DumpTime"), NewResourceDumpContext.Time.ToString());
+
+		NewResourceDumpContext.DumpJsonToFile(JsonObject, FString(FRDGResourceDumpContext::kBaseDir) / TEXT("Infos.json"));
+	}
+
+	// Dump the rendering cvars
+	if (NewResourceDumpContext.bEnableDiskWrite)
+	{
+		NewResourceDumpContext.DumpRenderingCVarsToCSV();
+	}
+
+	// Copy the viewer
+	if (NewResourceDumpContext.bEnableDiskWrite)
+	{
+		const TCHAR* OpenGPUDumpViewerBatName = TEXT("OpenGPUDumpViewer.bat");
+		const TCHAR* OpenGPUDumpViewerShName = TEXT("OpenGPUDumpViewer.sh");
+
+		const TCHAR* ViewerHTML = TEXT("GPUDumpViewer.html");
+		FString DumpGPUViewerSourcePath = FPaths::EngineDir() + FString(TEXT("Extras")) / TEXT("GPUDumpViewer");
+
+		PlatformFile.CopyFile(*(NewResourceDumpContext.DumpingDirectoryPath / ViewerHTML), *(DumpGPUViewerSourcePath / ViewerHTML));
+		PlatformFile.CopyFile(*(NewResourceDumpContext.DumpingDirectoryPath / OpenGPUDumpViewerBatName), *(DumpGPUViewerSourcePath / OpenGPUDumpViewerBatName));
+		PlatformFile.CopyFile(*(NewResourceDumpContext.DumpingDirectoryPath / OpenGPUDumpViewerShName), *(DumpGPUViewerSourcePath / OpenGPUDumpViewerShName));
+	}
+
+	ENQUEUE_RENDER_COMMAND(FStartGPUDump)(
+		[NewResourceDumpContext](FRHICommandListImmediate& ImmediateRHICmdList)
+	{
+		check(IsInRenderingThread());
+		GRDGResourceDumpContext = NewResourceDumpContext;
+
+		ImmediateRHICmdList.SubmitCommandsAndFlushGPU();
+
+		// Disable the validation for BUF_SourceCopy so that all buffers can be copied into staging buffer for CPU readback.
+		#if ENABLE_RHI_VALIDATION
+			GRHIValidateBufferSourceCopy = false;
+		#endif
+	});
+
+	// Mark ready for dump on next available frame
+	DumpingFrameCounter_GameThread = MAX_uint64;
+
+	if (NewResourceDumpContext.bEnableDiskWrite)
+	{
+		return NewResourceDumpContext.DumpingDirectoryPath;
+	}
+	return FString();
+}
+
+void FRDGBuilder::EndResourceDump()
+{
+	check(IsInGameThread());
+
+	 // make sure at least one frame has passed since we start a resource dump and we are not waiting on the dump to begin
+	if (DumpingFrameCounter_GameThread == 0 ||
+		DumpingFrameCounter_GameThread == MAX_uint64 ||
+		DumpingFrameCounter_GameThread >= GFrameCounter)
+	{
+		return;
+	}
+
+	// Wait all rendering commands are completed to finish with GRDGResourceDumpContext.
+	{
+		UE_LOG(LogRendererCore, Display, TEXT("Stalling game thread until render thread finishes to dump resources"));
+
+		ENQUEUE_RENDER_COMMAND(FEndGPUDump)(
+			[](FRHICommandListImmediate& ImmediateRHICmdList)
+		{
+			ImmediateRHICmdList.SubmitCommandsAndFlushGPU();
+			#if ENABLE_RHI_VALIDATION
+				GRHIValidateBufferSourceCopy = true;
+			#endif
+		});
+
+		FlushRenderingCommands();
+	}
+
+	// Log information about the dump.
+	FString AbsDumpingDirectoryPath = IFileManager::Get().ConvertToAbsolutePathForExternalAppForRead(*GRDGResourceDumpContext.DumpingDirectoryPath);
+	{
+		FDateTime Now = FDateTime::Now();
+		double TotalDumpSeconds = (Now - GRDGResourceDumpContext.Time).GetTotalSeconds();
+		double RHIReadbackCommandsSeconds = GRDGResourceDumpContext.TimingBucket[int32(FRDGResourceDumpContext::ETimingBucket::RHIReadbackCommands)];
+		double GPUWaitSeconds = GRDGResourceDumpContext.TimingBucket[int32(FRDGResourceDumpContext::ETimingBucket::GPUWait)];
+		double CPUPostProcessingSeconds = GRDGResourceDumpContext.TimingBucket[int32(FRDGResourceDumpContext::ETimingBucket::CPUPostProcessing)];
+		double MetadataFileWriteSeconds = GRDGResourceDumpContext.TimingBucket[int32(FRDGResourceDumpContext::ETimingBucket::MetadataFileWrite)];
+		double ResourceBinaryFileWriteSeconds = GRDGResourceDumpContext.TimingBucket[int32(FRDGResourceDumpContext::ETimingBucket::ResourceBinaryFileWrite)];
+		double RHIReleaseResourcesTimeSeconds = GRDGResourceDumpContext.TimingBucket[int32(FRDGResourceDumpContext::ETimingBucket::RHIReleaseResources)];
+
+		UE_LOG(LogRendererCore, Display, TEXT("Dumped %d resources in %.3f s to %s"), GRDGResourceDumpContext.ResourcesDumpPasses, float(TotalDumpSeconds), *AbsDumpingDirectoryPath);
+		UE_LOG(LogRendererCore, Display, TEXT("Dumped GPU readback commands: %.3f s"), float(RHIReadbackCommandsSeconds));
+		UE_LOG(LogRendererCore, Display, TEXT("Dumped GPU wait: %.3f s"), float(GPUWaitSeconds));
+		UE_LOG(LogRendererCore, Display, TEXT("Dumped CPU resource binary post processing: %.3f s"), float(CPUPostProcessingSeconds));
+		if (!GRDGResourceDumpContext.ResourceCompressionName.IsNone())
+		{
+			double CPUCompressionSeconds = GRDGResourceDumpContext.TimingBucket[int32(FRDGResourceDumpContext::ETimingBucket::CPUCompression)];
+
+			UE_LOG(LogRendererCore, Display, TEXT("Dumped CPU resource binary compression from %.3f MB to %.3f MB (%.1f %%) using %s"),
+				float(GRDGResourceDumpContext.ResourceBinaryUncompressedBytes) / float(1024 * 1024),
+				float(GRDGResourceDumpContext.ResourceBinaryWriteBytes) / float(1024 * 1024),
+				100.0f * float(GRDGResourceDumpContext.ResourceBinaryWriteBytes) / float(GRDGResourceDumpContext.ResourceBinaryUncompressedBytes),
+				*GRDGResourceDumpContext.ResourceCompressionName.ToString());
+			UE_LOG(LogRendererCore, Display, TEXT("Dumped CPU resource binary compression took %.3f s at %.3f MB/s"),
+				float(CPUCompressionSeconds),
+				float(GRDGResourceDumpContext.ResourceBinaryUncompressedBytes) / (float(1024 * 1024) * float(CPUCompressionSeconds)));
+		}
+		UE_LOG(LogRendererCore, Display, TEXT("Dumped metadata: %.3f MB in %d files under %.3f s at %.3f MB/s"),
+			float(GRDGResourceDumpContext.MetadataFilesWriteBytes) / float(1024 * 1024),
+			GRDGResourceDumpContext.MetadataFilesOpened,
+			float(MetadataFileWriteSeconds),
+			float(GRDGResourceDumpContext.MetadataFilesWriteBytes) / (float(1024 * 1024) * float(MetadataFileWriteSeconds)));
+		UE_LOG(LogRendererCore, Display, TEXT("Dumped resource binary: %.3f MB in %d files under %.3f s at %.3f MB/s"),
+			float(GRDGResourceDumpContext.ResourceBinaryWriteBytes) / float(1024 * 1024),
+			GRDGResourceDumpContext.ResourceBinaryFilesOpened,
+			float(ResourceBinaryFileWriteSeconds),
+			float(GRDGResourceDumpContext.ResourceBinaryWriteBytes) / (float(1024 * 1024) * float(ResourceBinaryFileWriteSeconds)));
+		UE_LOG(LogRendererCore, Display, TEXT("Dumped GPU readback resource release: %.3f s"), float(RHIReleaseResourcesTimeSeconds));
+	}
+
+	// Dump the log into the dump directory.
+	if (GRDGResourceDumpContext.bEnableDiskWrite)
+	{
+		if (GLog)
+		{
+			GLog->FlushThreadedLogs();
+			GLog->Flush();
+		}
+		FGenericCrashContext::DumpLog(GRDGResourceDumpContext.DumpingDirectoryPath / FRDGResourceDumpContext::kBaseDir);
+	}
+
+	#if PLATFORM_DESKTOP
+	if (GRDGResourceDumpContext.bShowInExplore)
+	{
+		FPlatformProcess::ExploreFolder(*AbsDumpingDirectoryPath);
+	}
+	#endif
+
+	GRDGResourceDumpContext = FRDGResourceDumpContext();
+	DumpingFrameCounter_GameThread = 0;
+}
+
+static const TCHAR* GetPassEventNameWithGPUMask(const FRDGPass* Pass, FString& OutNameStorage)
+{
+#if WITH_MGPU
+	if ((GNumExplicitGPUsForRendering > 1) && GDumpGPUMask.GetValueOnRenderThread())
+	{
+		// Prepend GPU mask on the event name of each pass, so you can see which GPUs the pass ran on.  Putting the mask at the
+		// front rather than the back makes all the masks line up, and easier to read (or ignore if you don't care about them).
+		// Also, it's easy to globally search for passes with a particular GPU mask using name search in the dump browser.
+		OutNameStorage = FString::Printf(TEXT("[%x] %s"), Pass->GetGPUMask().GetNative(), Pass->GetEventName().GetTCHAR());
+		return *OutNameStorage;
+	}
+	else
+#endif  // WITH_MGPU
+	{
+		return Pass->GetEventName().GetTCHAR();
+	}
+}
+
+void FRDGBuilder::DumpResourcePassOutputs(const FRDGPass* Pass)
+{
+	if (bInDebugPassScope)
+	{
+		return;
+	}
+
+	if (!GRDGResourceDumpContext.IsDumpingFrame())
+	{
+		return;
+	}
+
+	check(IsInRenderingThread());
+	if (!GRDGResourceDumpContext.IsDumpingPass(Pass))
+	{
+		return;
+	}
+
+	bInDebugPassScope = true;
+
+	TArray<TSharedPtr<FJsonValue>> InputResourceNames;
+	TArray<TSharedPtr<FJsonValue>> OutputResourceNames;
+	Pass->GetParameters().Enumerate([&](FRDGParameter Parameter)
+	{
+		switch (Parameter.GetType())
+		{
+		case UBMT_RDG_TEXTURE:
+		{
+			if (FRDGTextureRef Texture = Parameter.GetAsTexture())
+			{
+				FRDGTextureSRVDesc TextureSubResource = FRDGTextureSRVDesc::Create(Texture);
+				GRDGResourceDumpContext.AddDumpTexturePasses(*this, InputResourceNames, OutputResourceNames, Pass, TextureSubResource, /* bIsOutputResource = */ false);
+			}
+		}
+		break;
+		case UBMT_RDG_TEXTURE_SRV:
+		{
+			if (FRDGTextureSRVRef SRV = Parameter.GetAsTextureSRV())
+			{
+				if (SRV->Desc.MetaData == ERHITextureMetaDataAccess::None)
+				{
+					GRDGResourceDumpContext.AddDumpTexturePasses(*this, InputResourceNames, OutputResourceNames, Pass, SRV->Desc, /* bIsOutputResource = */ false);
+				}
+				else
+				{
+					UE_LOG(LogRendererCore, Warning, TEXT("Dumping texture %s's meta data unsupported"), SRV->Desc.Texture->Name);
+				}
+			}
+		}
+		break;
+		case UBMT_RDG_TEXTURE_UAV:
+		{
+			if (FRDGTextureUAVRef UAV = Parameter.GetAsTextureUAV())
+			{
+				if (UAV->Desc.MetaData == ERHITextureMetaDataAccess::None)
+				{
+					FRDGTextureSRVDesc TextureSubResource = FRDGTextureSRVDesc::CreateForMipLevel(UAV->Desc.Texture, UAV->Desc.MipLevel);
+					GRDGResourceDumpContext.AddDumpTextureSubResourcePass(*this, InputResourceNames, OutputResourceNames, Pass, TextureSubResource, /* bIsOutputResource = */ true);
+				}
+				else
+				{
+					UE_LOG(LogRendererCore, Warning, TEXT("Dumping texture %s's meta data unsupported"), UAV->Desc.Texture->Name);
+				}
+			}
+		}
+		break;
+		case UBMT_RDG_TEXTURE_ACCESS:
+		{
+			if (FRDGTextureAccess TextureAccess = Parameter.GetAsTextureAccess())
+			{
+				bool bIsOutputResource = (
+					TextureAccess.GetAccess() == ERHIAccess::UAVCompute ||
+					TextureAccess.GetAccess() == ERHIAccess::UAVGraphics ||
+					TextureAccess.GetAccess() == ERHIAccess::RTV);
+
+				FRDGTextureSRVDesc TextureSubResource = FRDGTextureSRVDesc::Create(TextureAccess);
+				GRDGResourceDumpContext.AddDumpTexturePasses(*this, InputResourceNames, OutputResourceNames, Pass, TextureSubResource, bIsOutputResource);
+			}
+		}
+		break;
+		case UBMT_RDG_TEXTURE_ACCESS_ARRAY:
+		{
+			const FRDGTextureAccessArray& TextureAccessArray = Parameter.GetAsTextureAccessArray();
+
+			for (FRDGTextureAccess TextureAccess : TextureAccessArray)
+			{
+				bool bIsOutputResource = (
+					TextureAccess.GetAccess() == ERHIAccess::UAVCompute ||
+					TextureAccess.GetAccess() == ERHIAccess::UAVGraphics ||
+					TextureAccess.GetAccess() == ERHIAccess::RTV);
+
+				FRDGTextureSRVDesc TextureSubResource = FRDGTextureSRVDesc::Create(TextureAccess);
+				GRDGResourceDumpContext.AddDumpTexturePasses(*this, InputResourceNames, OutputResourceNames, Pass, TextureSubResource, bIsOutputResource);
+			}
+		}
+		break;
+
+		case UBMT_RDG_BUFFER_SRV:
+		{
+			if (FRDGBufferSRVRef SRV = Parameter.GetAsBufferSRV())
+			{
+				FRDGBufferRef Buffer = SRV->Desc.Buffer;
+				GRDGResourceDumpContext.AddDumpBufferPass(*this, InputResourceNames, OutputResourceNames, Pass, Buffer, /* bIsOutputResource = */ false);
+			}
+		}
+		break;
+		case UBMT_RDG_BUFFER_UAV:
+		{
+			if (FRDGBufferUAVRef UAV = Parameter.GetAsBufferUAV())
+			{
+				FRDGBufferRef Buffer = UAV->Desc.Buffer;
+				GRDGResourceDumpContext.AddDumpBufferPass(*this, InputResourceNames, OutputResourceNames, Pass, Buffer, /* bIsOutputResource = */ true);
+			}
+		}
+		break;
+		case UBMT_RDG_BUFFER_ACCESS:
+		{
+			if (FRDGBufferAccess BufferAccess = Parameter.GetAsBufferAccess())
+			{
+				bool bIsOutputResource = (
+					BufferAccess.GetAccess() == ERHIAccess::UAVCompute ||
+					BufferAccess.GetAccess() == ERHIAccess::UAVGraphics);
+
+				GRDGResourceDumpContext.AddDumpBufferPass(*this, InputResourceNames, OutputResourceNames, Pass, BufferAccess, bIsOutputResource);
+			}
+		}
+		break;
+		case UBMT_RDG_BUFFER_ACCESS_ARRAY:
+		{
+			const FRDGBufferAccessArray& BufferAccessArray = Parameter.GetAsBufferAccessArray();
+
+			for (FRDGBufferAccess BufferAccess : BufferAccessArray)
+			{
+				bool bIsOutputResource = (
+					BufferAccess.GetAccess() == ERHIAccess::UAVCompute ||
+					BufferAccess.GetAccess() == ERHIAccess::UAVGraphics);
+
+				GRDGResourceDumpContext.AddDumpBufferPass(*this, InputResourceNames, OutputResourceNames, Pass, BufferAccess, bIsOutputResource);
+			}
+		}
+		break;
+
+		case UBMT_RENDER_TARGET_BINDING_SLOTS:
+		{
+			const FRenderTargetBindingSlots& RenderTargets = Parameter.GetAsRenderTargetBindingSlots();
+
+			RenderTargets.Enumerate([&](FRenderTargetBinding RenderTarget)
+			{
+				FRDGTextureRef Texture = RenderTarget.GetTexture();
+				FRDGTextureSRVDesc TextureSubResource = FRDGTextureSRVDesc::CreateForMipLevel(Texture, RenderTarget.GetMipIndex());
+				GRDGResourceDumpContext.AddDumpTextureSubResourcePass(*this, InputResourceNames, OutputResourceNames, Pass, TextureSubResource, /* bIsOutputResource = */ true);
+			});
+
+			const FDepthStencilBinding& DepthStencil = RenderTargets.DepthStencil;
+
+			if (FRDGTextureRef Texture = DepthStencil.GetTexture())
+			{
+				FExclusiveDepthStencil DepthStencilAccess = DepthStencil.GetDepthStencilAccess();
+
+				if (DepthStencilAccess.IsUsingDepth())
+				{
+					FRDGTextureSRVDesc TextureSubResource = FRDGTextureSRVDesc::CreateForMipLevel(Texture, 0);
+					GRDGResourceDumpContext.AddDumpTextureSubResourcePass(
+						*this, InputResourceNames, OutputResourceNames, Pass, TextureSubResource,
+						/* bIsOutputResource = */ DepthStencilAccess.IsDepthWrite());
+				}
+
+				if (DepthStencilAccess.IsUsingStencil())
+				{
+					FRDGTextureSRVDesc TextureSubResource = FRDGTextureSRVDesc::CreateWithPixelFormat(Texture, PF_X24_G8);
+					GRDGResourceDumpContext.AddDumpTextureSubResourcePass(
+						*this, InputResourceNames, OutputResourceNames, Pass, TextureSubResource,
+						/* bIsOutputResource = */ DepthStencilAccess.IsStencilWrite());
+				}
+			}
+		}
+		break;
+		}
+	});
+
+	// Dump the pass informations
+	{
+		TArray<TSharedPtr<FJsonValue>> ParentEventScopeNames;
+		#if RDG_GPU_SCOPES
+		{
+			const FRDGEventScope* ParentScope = Pass->GetGPUScopes().Event;
+
+			while (ParentScope)
+			{
+				ParentEventScopeNames.Add(MakeShareable(new FJsonValueString(ParentScope->Name.GetTCHAR())));
+				ParentScope = ParentScope->ParentScope;
+			}
+		}
+		#endif
+		{
+			ParentEventScopeNames.Add(MakeShareable(new FJsonValueString(FString::Printf(TEXT("Frame %llu"), GFrameCounterRenderThread))));
+		}
+
+		FString EventNameStorage;
+
+		TSharedPtr<FJsonObject> JsonObject = MakeShareable(new FJsonObject);
+		JsonObject->SetStringField(TEXT("EventName"), GetPassEventNameWithGPUMask(Pass, EventNameStorage));
+		JsonObject->SetStringField(TEXT("ParametersName"), Pass->GetParameters().GetLayout().GetDebugName());
+		JsonObject->SetStringField(TEXT("Parameters"), FRDGResourceDumpContext::PtrToString(Pass->GetParameters().GetContents()));
+		JsonObject->SetStringField(TEXT("ParametersMetadata"), FRDGResourceDumpContext::PtrToString(Pass->GetParameters().GetMetadata()));
+		JsonObject->SetStringField(TEXT("Pointer"), FString::Printf(TEXT("%016x"), FRDGResourceDumpContext::PtrToUint(Pass)));
+		JsonObject->SetNumberField(TEXT("Id"), GRDGResourceDumpContext.PassesCount);
+		JsonObject->SetArrayField(TEXT("ParentEventScopes"), ParentEventScopeNames);
+		JsonObject->SetArrayField(TEXT("InputResources"), InputResourceNames);
+		JsonObject->SetArrayField(TEXT("OutputResources"), OutputResourceNames);
+
+		GRDGResourceDumpContext.DumpJsonToFile(JsonObject, FString(FRDGResourceDumpContext::kBaseDir) / TEXT("Passes.json"), FILEWRITE_Append);
+	}
+
+	// Dump the pass' parameters
+	if (GDumpGPUPassParameters.GetValueOnRenderThread() != 0)
+	{
+		int32 PassParametersByteSize = 0;
+		{
+			const FShaderParametersMetadata* Metadata = Pass->GetParameters().GetMetadata();
+			if (Metadata)
+			{
+				GRDGResourceDumpContext.Dump(Metadata);
+				PassParametersByteSize = Metadata->GetSize();
+			}
+		}
+
+		if (PassParametersByteSize == 0 && Pass->GetParameters().GetLayoutPtr())
+		{
+			PassParametersByteSize = Pass->GetParameters().GetLayout().ConstantBufferSize;
+		}
+
+		const uint8* PassParametersContent = Pass->GetParameters().GetContents();
+		if (PassParametersContent && !GRDGResourceDumpContext.IsDumped(PassParametersContent))
+		{
+			TArrayView<const uint8> ArrayView(PassParametersContent, PassParametersByteSize);
+			FString DumpFilePath = FRDGResourceDumpContext::kStructuresDir / FRDGResourceDumpContext::PtrToString(PassParametersContent) + TEXT(".bin");
+			GRDGResourceDumpContext.DumpBinaryToFile(ArrayView, DumpFilePath);
+			GRDGResourceDumpContext.SetDumped(PassParametersContent);
+		}
+	}
+
+	GRDGResourceDumpContext.PassesCount++;
+
+	bInDebugPassScope = false;
+}
+
+#if RDG_DUMP_RESOURCES_AT_EACH_DRAW
+
+void FRDGBuilder::BeginPassDump(const FRDGPass* Pass)
+{
+	if (!GRDGResourceDumpContext.IsDumpingFrame())
+	{
+		return;
+	}
+
+	if (!GDumpGPUDraws.GetValueOnRenderThread())
+	{
+		return;
+	}
+
+	if (!EnumHasAnyFlags(Pass->GetFlags(), ERDGPassFlags::Raster))
+	{
+		return;
+	}
+
+	if (!IsInRenderingThread())
+	{
+		UE_LOG(LogRendererCore, Warning, TEXT("Couldn't start dumping draw's resources for pass %s because not in the rendering thread"), Pass->GetEventName().GetTCHAR());
+		return;
+	}
+
+	check(GRDGResourceDumpContext.DrawDumpingPass == nullptr);
+
+	if (GRDGResourceDumpContext.IsDumpingPass(Pass))
+	{
+		GRDGResourceDumpContext.DrawDumpingPass = Pass;
+		GRDGResourceDumpContext.DrawDumpCount = 0;
+	}
+}
+
+// static
+void FRDGBuilder::DumpDraw(const FRDGEventName& DrawEventName)
+{
+	if (!GRDGResourceDumpContext.IsDumpingFrame())
+	{
+		return;
+	}
+
+	if (!IsInRenderingThread())
+	{
+		UE_LOG(LogRendererCore, Warning, TEXT("Couldn't dump draw because not in the rendering thread"));
+		return;
+	}
+
+	if (!GRDGResourceDumpContext.DrawDumpingPass)
+	{
+		return;
+	}
+
+	const FRDGPass* Pass = GRDGResourceDumpContext.DrawDumpingPass;
+
+	FRHICommandListImmediate& RHICmdList = FRHICommandListExecutor::GetImmediateCommandList();
+
+	if (EnumHasAnyFlags(Pass->GetFlags(), ERDGPassFlags::Raster))
+	{
+		RHICmdList.EndRenderPass();
+	}
+
+	Pass->GetParameters().Enumerate([&](FRDGParameter Parameter)
+	{
+		switch (Parameter.GetType())
+		{
+		case UBMT_RENDER_TARGET_BINDING_SLOTS:
+		{
+			const FRenderTargetBindingSlots& RenderTargets = Parameter.GetAsRenderTargetBindingSlots();
+
+			RenderTargets.Enumerate([&](FRenderTargetBinding RenderTarget)
+			{
+				FRDGTextureRef Texture = RenderTarget.GetTexture();
+				FRDGTextureSRVDesc TextureSubResource = FRDGTextureSRVDesc::CreateForMipLevel(Texture, RenderTarget.GetMipIndex());
+				GRDGResourceDumpContext.DumpDrawTextureSubResource(
+					RHICmdList,
+					TextureSubResource,
+					ERHIAccess::RTV);
+			});
+
+			const FDepthStencilBinding& DepthStencil = RenderTargets.DepthStencil;
+
+			if (FRDGTextureRef Texture = DepthStencil.GetTexture())
+			{
+				FExclusiveDepthStencil DepthStencilAccess = DepthStencil.GetDepthStencilAccess();
+
+				if (DepthStencilAccess.IsDepthWrite())
+				{
+					FRDGTextureSRVDesc TextureSubResource = FRDGTextureSRVDesc::CreateForMipLevel(Texture, 0);
+					GRDGResourceDumpContext.DumpDrawTextureSubResource(
+						RHICmdList,
+						TextureSubResource,
+						ERHIAccess::RTV);
+				}
+
+				if (DepthStencilAccess.IsStencilWrite())
+				{
+					FRDGTextureSRVDesc TextureSubResource = FRDGTextureSRVDesc::CreateWithPixelFormat(Texture, PF_X24_G8);
+					GRDGResourceDumpContext.DumpDrawTextureSubResource(
+						RHICmdList,
+						TextureSubResource,
+						ERHIAccess::RTV);
+				}
+			}
+		}
+		break;
+		}
+	});
+
+	if (EnumHasAnyFlags(Pass->GetFlags(), ERDGPassFlags::Raster))
+	{
+		RHICmdList.BeginRenderPass(Pass->GetParameters().GetRenderPassInfo(), Pass->GetName());
+	}
+
+	// Dump the draw even name
+	{
+		TSharedPtr<FJsonObject> JsonObject = MakeShareable(new FJsonObject);
+		JsonObject->SetStringField(TEXT("DrawName"), DrawEventName.GetTCHAR());
+
+		FString DumpFilePath = FRDGResourceDumpContext::kPassesDir / FString::Printf(TEXT("Pass.%016x.Draws.json"), FRDGResourceDumpContext::PtrToUint(Pass));
+		GRDGResourceDumpContext.DumpJsonToFile(JsonObject, DumpFilePath, FILEWRITE_Append);
+	}
+
+	GRDGResourceDumpContext.DrawDumpCount++;
+
+	if (GRDGResourceDumpContext.DrawDumpCount % 10 == 0)
+	{
+		UE_LOG(LogRendererCore, Display, TEXT("Dumped %d draws' resources"), GRDGResourceDumpContext.DrawDumpCount);
+		return;
+	}
+}
+
+void FRDGBuilder::EndPassDump(const FRDGPass* Pass)
+{
+	if (!GRDGResourceDumpContext.IsDumpingFrame())
+	{
+		return;
+	}
+
+	if (!IsInRenderingThread())
+	{
+		return;
+	}
+
+	if (!GRDGResourceDumpContext.DrawDumpingPass)
+	{
+		return;
+	}
+
+	check(Pass == GRDGResourceDumpContext.DrawDumpingPass);
+
+	// Output how many draw has been dump for this pass.
+	if (GRDGResourceDumpContext.DrawDumpCount > 0)
+	{
+		FString EventNameStorage;
+
+		TSharedPtr<FJsonObject> JsonObject = MakeShareable(new FJsonObject);
+		JsonObject->SetStringField(TEXT("EventName"), GetPassEventNameWithGPUMask(Pass, EventNameStorage));
+		JsonObject->SetStringField(TEXT("Pointer"), FString::Printf(TEXT("%016x"), FRDGResourceDumpContext::PtrToUint(Pass)));
+		JsonObject->SetNumberField(TEXT("DrawCount"), GRDGResourceDumpContext.DrawDumpCount);
+
+		GRDGResourceDumpContext.DumpJsonToFile(JsonObject, FString(FRDGResourceDumpContext::kBaseDir) / TEXT("PassDrawCounts.json"), FILEWRITE_Append);
+
+		UE_LOG(LogRendererCore, Display, TEXT("Completed dump of %d draws for pass: %s"), GRDGResourceDumpContext.DrawDumpCount, Pass->GetEventName().GetTCHAR());
+	}
+
+	GRDGResourceDumpContext.DrawDumpingPass = nullptr;
+	GRDGResourceDumpContext.DrawDumpCount = 0;
+}
+
+// static
+bool FRDGBuilder::IsDumpingFrame()
+{
+	return GRDGResourceDumpContext.IsDumpingFrame();
+}
+
+bool FRDGBuilder::IsDumpingDraws()
+{
+	if (!GRDGResourceDumpContext.IsDumpingFrame())
+	{
+		return false;
+	}
+
+	return GDumpGPUDraws.GetValueOnRenderThread() != 0;
+}
+
+#endif // RDG_DUMP_RESOURCES_AT_EACH_DRAW
+
+#else //! RDG_DUMP_RESOURCES
+
+bool IsDumpingRDGResources()
+{
+	return false;
+}
+
+#endif //! RDG_DUMP_RESOURCES
diff --git a/Engine/Source/Runtime/RenderCore/Private/RenderGraphDumpResources.cpp b/Engine/Source/Runtime/RenderCore/Private/RenderGraphDumpResources.cpp
deleted file mode 100644
