Implements DumpGPUServices plugin to have DumpGPU -upload

diff --git a/Engine/Config/BaseInput.ini b/Engine/Config/BaseInput.ini
--- a/Engine/Config/BaseInput.ini
+++ b/Engine/Config/BaseInput.ini
@@ -70,6 +70,7 @@ MaxScrollbackSize=1024
 +ManualAutoCompleteList=(Command="DumpShaderStats",Desc="Dump shader information")
 +ManualAutoCompleteList=(Command="DumpShaderPipelineStats",Desc="Dump shader pipeline information")
 +ManualAutoCompleteList=(Command="DumpShaderCompileStats",Desc="Dump shader compilation information")
++ManualAutoCompleteList=(Command="DumpGPU -upload",Desc="Dump the GPU's intermediary resources and upload to network")
 +ManualAutoCompleteList=(Command="StartFPSChart",Desc="after that use StopFPSChart")
 +ManualAutoCompleteList=(Command="StopFPSChart",Desc="after that look for the output in Saved/Profiling/FPSChartStats")
 +ManualAutoCompleteList=(Command="FreezeAt",Desc="Locks the player view and rendering time.")
diff --git a/Engine/Extras/GPUDumpViewer/GPUDumpViewer.html b/Engine/Extras/GPUDumpViewer/GPUDumpViewer.html
--- a/Engine/Extras/GPUDumpViewer/GPUDumpViewer.html
+++ b/Engine/Extras/GPUDumpViewer/GPUDumpViewer.html
@@ -15,6 +15,7 @@ const k_null_json_ptr = "0000000000000000";
 
 // -------------------------------------------------------------------- GLOBALS
 
+var g_dump_service = {};
 var g_infos = {};
 var g_dump_cvars = {};
 var g_passes = [];
@@ -150,11 +151,11 @@ function load_resource_binary_file(relative_path, callback)
 	return load_binary_file(relative_path, function(raw_texture_data)
 	{
 		var compression_type = '';
-		if (g_dump_cvars['r.DumpGPU.CompressResources'] == '1')
+		if (g_dump_service['CompressionName'] == 'Zlib')
 		{
 			compression_type = 'deflate';
 		}
-		else if (g_dump_cvars['r.DumpGPU.CompressResources'] == '2')
+		else if (g_dump_service['CompressionName'] == 'GZip')
 		{
 			compression_type = 'gzip';
 		}
@@ -3910,9 +3911,13 @@ function init_top_buttons()
 
 function init_page()
 {
+	// Load the dump service information to know how to read files. This is independent of Base/Infos.json because may depends whether it was uploaded through the DumpGPUServices plugin.
+	g_dump_service = load_json('Base/DumpService.json');
+
 	// Load the base information of the page.
 	g_infos = load_json('Base/Infos.json');
-	if (!g_infos)
+
+	if (!g_dump_service || !g_infos)
 	{
 		document.body.innerHTML = 'Please use OpenGPUDumpViewer script for the viewer to work correctly.';
 		return;
diff --git a/Engine/Plugins/Developer/DumpGPUServices/DumpGPUServices.uplugin b/Engine/Plugins/Developer/DumpGPUServices/DumpGPUServices.uplugin
new file mode 100644
--- /dev/null
+++ b/Engine/Plugins/Developer/DumpGPUServices/DumpGPUServices.uplugin
@@ -0,0 +1,27 @@
+{
+	"FileVersion": 3,
+	"Version": 1,
+	"VersionName": "1.0",
+	"FriendlyName": "Dump GPU Services",
+	"Description": "Implements automatic upload services for the DumpGPU command.",
+	"Category": "Rendering",
+	"CreatedBy": "Epic Games, Inc.",
+	"CreatedByURL": "http://epicgames.com",
+	"DocsURL": "",
+	"MarketplaceURL": "",
+	"SupportURL": "",
+	"EnabledByDefault": true,
+	"CanContainContent": false,
+	"IsBetaVersion": false,
+	"IsExperimentalVersion": true,
+	"Installed": false,
+	"Modules": [
+		{
+			"Name": "DumpGPUServices",
+			"Type": "Runtime",
+			"LoadingPhase": "Default",
+			"TargetDenyList": [ "Server" ],
+			"TargetConfigurationDenyList": [ "Shipping" ]
+		}
+	]
+}
\ No newline at end of file
diff --git a/Engine/Plugins/Developer/DumpGPUServices/Source/DumpGPUServices/DumpGPUServices.Build.cs b/Engine/Plugins/Developer/DumpGPUServices/Source/DumpGPUServices/DumpGPUServices.Build.cs
new file mode 100644
--- /dev/null
+++ b/Engine/Plugins/Developer/DumpGPUServices/Source/DumpGPUServices/DumpGPUServices.Build.cs
@@ -0,0 +1,42 @@
+// Copyright Epic Games, Inc. All Rights Reserved.
+
+namespace UnrealBuildTool.Rules
+{
+	public class DumpGPUServices : ModuleRules
+	{
+		public DumpGPUServices(ReadOnlyTargetRules Target) : base(Target)
+		{
+			PrivateIncludePaths.AddRange(
+				new string[] {
+					EngineDirectory + "/Source/Runtime/RenderCore/Private",
+				}
+				);
+
+			PublicDependencyModuleNames.AddRange(
+				new string[]
+				{
+					"Core",
+				}
+				);
+
+			PrivateDependencyModuleNames.AddRange(
+				new string[]
+				{
+					"RenderCore",
+					"Json",
+					"HTTP",
+				}
+				);
+
+			if (Target.Type == TargetType.Editor)
+			{
+				PrivateDependencyModuleNames.AddRange(
+					new string[]
+					{
+					"Slate",
+					}
+					);
+			}
+		}
+	}
+}
diff --git a/Engine/Plugins/Developer/DumpGPUServices/Source/DumpGPUServices/Private/DumpGPUHTTPUploadServices.cpp b/Engine/Plugins/Developer/DumpGPUServices/Source/DumpGPUServices/Private/DumpGPUHTTPUploadServices.cpp
new file mode 100644
--- /dev/null
+++ b/Engine/Plugins/Developer/DumpGPUServices/Source/DumpGPUServices/Private/DumpGPUHTTPUploadServices.cpp
@@ -0,0 +1,361 @@
+// Copyright Epic Games, Inc. All Rights Reserved.
+
+#include "DumpGPU.h"
+#include "DumpGPUServices.h"
+#include "Misc/Paths.h"
+#include "Misc/FileHelper.h"
+#include "Misc/ScopedSlowTask.h"
+#include "Serialization/JsonTypes.h"
+#include "Serialization/JsonReader.h"
+#include "Serialization/JsonSerializer.h"
+#include "HttpModule.h"
+#include "HttpManager.h"
+#include "Interfaces/IHttpResponse.h"
+
+#if WITH_EDITOR
+	#include "Framework/Notifications/NotificationManager.h"
+#endif
+
+class FDumpGPUHTTPUpload;
+
+
+TSharedPtr<FJsonObject> LoadJsonFile(const FString& JsonPath)
+{
+	TSharedPtr<FJsonObject> JsonObject;
+
+	FString JsonContent;
+	if (!FFileHelper::LoadFileToString(JsonContent, *JsonPath))
+	{
+		return JsonObject;
+	}
+
+	auto JsonReader = TJsonReaderFactory<>::Create(JsonContent);
+	FJsonSerializer::Deserialize(JsonReader, JsonObject);
+	return JsonObject;
+}
+
+class FDumpGPUHTTPUploadProvider : public IDumpGPUUploadServiceProvider
+{
+public:
+	FDumpGPUHTTPUploadProvider(const FString& InUploadURLPattern)
+		: UploadURLPattern(InUploadURLPattern)
+	{ }
+
+	virtual void UploadDump(const FDumpParameters& Parameters) override;
+
+	~FDumpGPUHTTPUploadProvider();
+
+private:
+	const FString UploadURLPattern;
+	TArray<FDumpGPUHTTPUpload*, TInlineAllocator<1>> PendingUploads;
+
+	friend class FDumpGPUHTTPUpload;
+};
+
+IDumpGPUUploadServiceProvider* CreateHTTPUploadProvider(const FString& UploadURLPattern)
+{
+	return new FDumpGPUHTTPUploadProvider(UploadURLPattern);
+}
+
+class FDumpGPUHTTPUpload
+{
+public:
+	static constexpr int32 kOverlappedRequests = 4;
+
+	FDumpGPUHTTPUpload(FDumpGPUHTTPUploadProvider* InProvider, const IDumpGPUUploadServiceProvider::FDumpParameters& InDumpParamters, const FString& InDumpUploadURL)
+		: DumpParamters(InDumpParamters)
+		, DumpUploadURL(InDumpUploadURL)
+		, Provider(InProvider)
+	{ }
+
+	void Start()
+	{
+		check(IsInGameThread());
+
+		IPlatformFile::GetPlatformPhysical().FindFilesRecursively(Files, *DumpParamters.LocalPath, nullptr);
+
+#if WITH_EDITOR
+		ProgressHandle = FSlateNotificationManager::Get().StartProgressNotification(NSLOCTEXT("DumpGPUSerices", "UploadingGPUDump", "Uploading GPU dump"), Files.Num());
+#endif
+
+		for (int32 FileId = 0; FileId < kOverlappedRequests; FileId++)
+		{
+			StartNextFileUpload();
+		}
+	}
+
+	void AbortDueToError()
+	{
+		check(IsInGameThread());
+
+		if (!bAbort)
+		{
+			#if WITH_EDITOR
+			if (ProgressHandle.IsValid())
+			{
+				FSlateNotificationManager::Get().CancelProgressNotification(ProgressHandle);
+			}
+			#endif
+			bAbort = true;
+		}
+
+		if (PendingHttpRequests.Num() == 0)
+		{
+			return Finish();
+		}
+	}
+
+private:
+	void StartUploadFile(int32 FileId)
+	{
+		check(IsInGameThread());
+		check(FileId < Files.Num());
+
+		FString RelativeFilePath = Files[FileId];
+		FPaths::MakePathRelativeTo(RelativeFilePath, *DumpParamters.LocalPath);
+
+		FString FilePath = DumpParamters.LocalPath / RelativeFilePath;
+
+		TSharedPtr<IHttpRequest, ESPMode::ThreadSafe> HttpRequest = FHttpModule::Get().CreateRequest();
+		PendingHttpRequests.Add(HttpRequest);
+
+		HttpRequest->SetURL(DumpUploadURL / RelativeFilePath);
+		HttpRequest->SetVerb(TEXT("PUT"));
+		HttpRequest->SetHeader(TEXT("Content-Type"), TEXT("application/octet-stream"));
+
+		// If the dump service parameters, upload the dump service parameters used for the upload that might have different compression settings.
+		if (RelativeFilePath == IDumpGPUUploadServiceProvider::FDumpParameters::kServiceFileName)
+		{
+			FString ServiceParameters = DumpParamters.DumpServiceParametersFileContent();
+			HttpRequest->SetContentAsString(ServiceParameters);
+			return ProcessHttpRequest(HttpRequest, RelativeFilePath);
+		}
+
+		// If no compression is needed on that file, just process the HTTP request immediately.
+		if (DumpParamters.CompressionName.IsNone() || !DumpParamters.CompressionFiles.IsMatch(RelativeFilePath))
+		{
+			if (!HttpRequest->SetContentAsStreamedFile(FilePath))
+			{
+				return AbortDueToError();
+			}
+
+			return ProcessHttpRequest(HttpRequest, RelativeFilePath);
+		}
+
+		// Load and compress in a background task
+		FFunctionGraphTask::CreateAndDispatchWhenReady([this, RelativeFilePath, HttpRequest]()
+		{
+			FString FilePath = this->DumpParamters.LocalPath / RelativeFilePath;
+
+			bool bSuccess = false;
+			TArray<uint8> CompressedArray;
+
+			TArray<uint8> UncompressedData;
+			if (FFileHelper::LoadFileToArray(UncompressedData, *FilePath))
+			{
+				const int32 UncompressedSize = UncompressedData.Num();
+				const int32 MaxCompressedSize = FCompression::CompressMemoryBound(DumpParamters.CompressionName, UncompressedSize, COMPRESS_BiasSize);
+
+				CompressedArray.SetNumUninitialized(MaxCompressedSize);
+				int32 CompressedSize = MaxCompressedSize;
+				bSuccess = FCompression::CompressMemory(
+					DumpParamters.CompressionName,
+					/* out */ CompressedArray.GetData(), /* out */ CompressedSize,
+					UncompressedData.GetData(), UncompressedSize,
+					COMPRESS_BiasSize);
+
+				UncompressedData.Empty();
+
+				if (bSuccess)
+				{
+					check(CompressedSize <= MaxCompressedSize);
+					CompressedArray.SetNum(CompressedSize, /* bAllowShrinking = */ false);
+					UE_LOG(LogDumpGPUServices, Verbose, TEXT("Compress %s completed"), *RelativeFilePath);
+				}
+				else
+				{
+					UE_LOG(LogDumpGPUServices, Warning, TEXT("Failed to compress %s"), *RelativeFilePath);
+				}
+			}
+			else
+			{
+				UE_LOG(LogDumpGPUServices, Warning, TEXT("Failed to load %s"), *RelativeFilePath);
+			}
+
+			if (bSuccess)
+			{
+				HttpRequest->SetContent(MoveTemp(CompressedArray));
+
+				// Queue game thread task to issue HTTP request
+				FFunctionGraphTask::CreateAndDispatchWhenReady([this, RelativeFilePath, HttpRequest]()
+				{
+					this->ProcessHttpRequest(HttpRequest, RelativeFilePath);
+				}, TStatId(), nullptr, ENamedThreads::GameThread);
+			}
+			else
+			{
+				// Queue game thread task to issue error
+				FFunctionGraphTask::CreateAndDispatchWhenReady([this]()
+				{
+					this->AbortDueToError();
+				}, TStatId(), nullptr, ENamedThreads::GameThread);
+			}
+		}, TStatId(), nullptr, ENamedThreads::AnyBackgroundThreadNormalTask);
+	}
+
+	void ProcessHttpRequest(TSharedPtr<IHttpRequest, ESPMode::ThreadSafe> HttpRequest, const FString & RelativeFilePath)
+	{
+		check(IsInGameThread());
+
+		HttpRequest->OnProcessRequestComplete().BindLambda([this, RelativeFilePath](FHttpRequestPtr HttpRequestPtr, FHttpResponsePtr HttpResponse, bool bSuccess)
+		{
+			check(IsInGameThread());
+
+			EHttpRequestStatus::Type RequestStatus = HttpRequestPtr->GetStatus();
+			int32 ResponseCode = HttpResponse ? HttpResponse->GetResponseCode() : 0;
+
+			// Remove the request from in-flight requests
+			PendingHttpRequests.Remove(HttpRequestPtr);
+
+			if (RequestStatus == EHttpRequestStatus::Succeeded && ResponseCode == 201)
+			{
+				UE_LOG(LogDumpGPUServices, Verbose, TEXT("Upload %s completed"), *RelativeFilePath);
+
+				FilesCompleted++;
+
+				#if WITH_EDITOR
+				if (!bAbort && ProgressHandle.IsValid())
+				{
+					FSlateNotificationManager::Get().UpdateProgressNotification(ProgressHandle, FilesCompleted);
+				}
+				#endif
+
+				if ((FilesCompleted % 20) == 0)
+				{
+					UE_LOG(LogDumpGPUServices, Display, TEXT("Uploaded %d out of %d files"), FilesCompleted, Files.Num());
+				}
+
+				StartNextFileUpload();
+			}
+			else
+			{
+				UE_LOG(LogDumpGPUServices, Warning, TEXT("Upload %s failed"), *RelativeFilePath);
+				AbortDueToError();
+			}
+		});
+
+		HttpRequest->ProcessRequest();
+	}
+
+	void StartNextFileUpload()
+	{
+		check(IsInGameThread());
+
+		if (bAbort)
+		{
+			return AbortDueToError();
+		}
+
+		if (NextFileToUpload < Files.Num())
+		{
+			StartUploadFile(NextFileToUpload);
+			NextFileToUpload++;
+		}
+
+		if (PendingHttpRequests.Num() == 0)
+		{
+			return Finish();
+		}
+	}
+	
+	void Finish()
+	{
+		check(IsInGameThread());
+		check(PendingHttpRequests.Num() == 0);
+
+		FString AbsDumpingDirectoryPath = IFileManager::Get().ConvertToAbsolutePathForExternalAppForRead(*DumpParamters.LocalPath);
+		if (bAbort == false)
+		{
+			IPlatformFile::GetPlatformPhysical().DeleteDirectoryRecursively(*DumpParamters.LocalPath);
+			UE_LOG(LogDumpGPUServices, Display, TEXT("DumpGPU HTTP Upload to %s completed"), *DumpUploadURL);
+			UE_LOG(LogDumpGPUServices, Display, TEXT("%s has been deleted"), *AbsDumpingDirectoryPath);
+
+			#if WITH_EDITOR && PLATFORM_DESKTOP
+				FPlatformProcess::LaunchURL(*DumpUploadURL, nullptr, nullptr);
+			#endif
+		}
+		else
+		{
+			UE_LOG(LogDumpGPUServices, Warning, TEXT("DumpGPU HTTP Upload to %s failed"), *DumpUploadURL);
+			UE_LOG(LogDumpGPUServices, Display, TEXT("%s is still available on disk"), *AbsDumpingDirectoryPath);
+		}
+
+		if (Provider->PendingUploads.Contains(this))
+		{
+			Provider->PendingUploads.Remove(this);
+			if (Provider->PendingUploads.Num() > 0)
+			{
+				Provider->PendingUploads[0]->Start();
+			}
+		}
+
+		delete this;
+	}
+
+	const IDumpGPUUploadServiceProvider::FDumpParameters DumpParamters;
+	const FString DumpUploadURL;
+	FDumpGPUHTTPUploadProvider* const Provider;
+	int32 FilesCompleted = 0;
+	#if WITH_EDITOR
+		FProgressNotificationHandle ProgressHandle;
+	#endif
+
+	TArray<FString> Files;
+
+	int32 NextFileToUpload = 0;
+	bool bAbort = false;
+
+	TArray<TSharedPtr<IHttpRequest, ESPMode::ThreadSafe>, TInlineAllocator<kOverlappedRequests>> PendingHttpRequests;
+}; // FDumpGPUHTTPUpload
+
+
+void FDumpGPUHTTPUploadProvider::UploadDump(const FDumpParameters& DumpParameters)
+{
+	check(IsInGameThread());
+
+	FString UploadURL;
+	{
+		FString Project = FApp::GetProjectName();
+		FString Platform = FPlatformProperties::IniPlatformName();
+
+		UploadURL = UploadURLPattern;
+		UploadURL.ReplaceInline(TEXT("[Project]"), *Project);
+		UploadURL.ReplaceInline(TEXT("[Platform]"), *Platform);
+		UploadURL.ReplaceInline(TEXT("[DumpTime]"), *DumpParameters.Time);
+		UploadURL.ReplaceInline(TEXT("[DumpType]"), *DumpParameters.Type);
+	}
+
+	FDumpGPUHTTPUpload* PendingUpload = new FDumpGPUHTTPUpload(this, DumpParameters, UploadURL);
+	PendingUploads.Add(PendingUpload);
+	if (PendingUploads.Num() == 1)
+	{
+		PendingUploads[0]->Start();
+	}
+}
+
+FDumpGPUHTTPUploadProvider::~FDumpGPUHTTPUploadProvider()
+{
+	check(IsInGameThread());
+
+	if (PendingUploads.Num() > 0)
+	{
+		for (int32 i = 1; i < PendingUploads.Num(); i++)
+		{
+			delete PendingUploads[i];
+		}
+		PendingUploads.SetNum(1);
+
+		PendingUploads[0]->AbortDueToError();
+		ensure(PendingUploads.Num() == 0);
+	}
+}
diff --git a/Engine/Plugins/Developer/DumpGPUServices/Source/DumpGPUServices/Private/DumpGPUServices.cpp b/Engine/Plugins/Developer/DumpGPUServices/Source/DumpGPUServices/Private/DumpGPUServices.cpp
new file mode 100644
--- /dev/null
+++ b/Engine/Plugins/Developer/DumpGPUServices/Source/DumpGPUServices/Private/DumpGPUServices.cpp
@@ -0,0 +1,61 @@
+// Copyright Epic Games, Inc. All Rights Reserved.
+
+#include "DumpGPUServices.h"
+#include "CoreMinimal.h"
+#include "CoreGlobals.h"
+#include "Modules/ModuleManager.h"
+#include "IDumpGPUServices.h"
+
+
+class FDumpGPUServices : public IDumpGPUServices
+{
+	/** IModuleInterface implementation */
+	virtual void StartupModule() override;
+	virtual void ShutdownModule() override;
+
+	IDumpGPUUploadServiceProvider* UploadProvider = nullptr;
+};
+
+IMPLEMENT_MODULE( FDumpGPUServices, DumpGPUServices )
+
+DEFINE_LOG_CATEGORY(LogDumpGPUServices);
+
+IDumpGPUUploadServiceProvider* CreateHTTPUploadProvider(const FString& UploadURLPattern);
+
+
+void FDumpGPUServices::StartupModule()
+{
+	if (IDumpGPUUploadServiceProvider::GProvider)
+	{
+		return;
+	}
+
+	FString UploadURLPattern;
+	GConfig->GetString(TEXT("Rendering.DumpGPUServices"), TEXT("UploadURLPattern"), UploadURLPattern, GEngineIni);
+
+	if (UploadURLPattern.StartsWith(TEXT("http://")))
+	{
+		UploadProvider = CreateHTTPUploadProvider(UploadURLPattern);
+	}
+
+	if (UploadProvider)
+	{
+		IDumpGPUUploadServiceProvider::GProvider = UploadProvider;
+		UE_LOG(LogDumpGPUServices, Display, TEXT("DumpGPU upload service set up with %s"), *UploadURLPattern);
+	}
+}
+
+void FDumpGPUServices::ShutdownModule()
+{
+	if (!UploadProvider)
+	{
+		return;
+	}
+
+	if (IDumpGPUUploadServiceProvider::GProvider == UploadProvider)
+	{
+		IDumpGPUUploadServiceProvider::GProvider = nullptr;
+	}
+
+	delete UploadProvider;
+}
diff --git a/Engine/Plugins/Developer/DumpGPUServices/Source/DumpGPUServices/Private/DumpGPUServices.h b/Engine/Plugins/Developer/DumpGPUServices/Source/DumpGPUServices/Private/DumpGPUServices.h
new file mode 100644
--- /dev/null
+++ b/Engine/Plugins/Developer/DumpGPUServices/Source/DumpGPUServices/Private/DumpGPUServices.h
@@ -0,0 +1,7 @@
+// Copyright Epic Games, Inc. All Rights Reserved.
+
+#pragma once
+
+#include "CoreMinimal.h"
+
+DECLARE_LOG_CATEGORY_EXTERN(LogDumpGPUServices, Log, All);
diff --git a/Engine/Plugins/Developer/DumpGPUServices/Source/DumpGPUServices/Public/IDumpGPUServices.h b/Engine/Plugins/Developer/DumpGPUServices/Source/DumpGPUServices/Public/IDumpGPUServices.h
new file mode 100644
--- /dev/null
+++ b/Engine/Plugins/Developer/DumpGPUServices/Source/DumpGPUServices/Public/IDumpGPUServices.h
@@ -0,0 +1,40 @@
+// Copyright Epic Games, Inc. All Rights Reserved.
+
+#pragma once
+
+#include "CoreMinimal.h"
+#include "Modules/ModuleInterface.h"
+#include "Modules/ModuleManager.h"
+
+
+/**
+ * The public interface to this module.  In most cases, this interface is only public to sibling modules 
+ * within this plugin.
+ */
+class IDumpGPUServices : public IModuleInterface
+{
+
+public:
+
+	/**
+	 * Singleton-like access to this module's interface.  This is just for convenience!
+	 * Beware of calling this during the shutdown phase, though.  Your module might have been unloaded already.
+	 *
+	 * @return Returns singleton instance, loading the module on demand if needed
+	 */
+	static inline IDumpGPUServices& Get()
+	{
+		return FModuleManager::LoadModuleChecked< IDumpGPUServices >( "DumpGPUServices" );
+	}
+
+	/**
+	 * Checks to see if this module is loaded and ready.  It is only valid to call Get() if IsAvailable() returns true.
+	 *
+	 * @return True if the module is loaded and ready to use
+	 */
+	static inline bool IsAvailable()
+	{
+		return FModuleManager::Get().IsModuleLoaded( "DumpGPUServices" );
+	}
+};
+
diff --git a/Engine/Source/Runtime/Engine/Private/UnrealEngine.cpp b/Engine/Source/Runtime/Engine/Private/UnrealEngine.cpp
--- a/Engine/Source/Runtime/Engine/Private/UnrealEngine.cpp
+++ b/Engine/Source/Runtime/Engine/Private/UnrealEngine.cpp
@@ -5592,7 +5592,7 @@ bool UEngine::HandleDumpGPUCommand(const TCHAR* Cmd, FOutputDevice& Ar)
 
 	check(IsInGameThread());
 	TArray<FString> Args;
-	FString ResourceDumpDirectory = FRDGBuilder::BeginResourceDump(Args);
+	FString ResourceDumpDirectory = FRDGBuilder::BeginResourceDump(Cmd);
 
 	if (!ResourceDumpDirectory.IsEmpty())
 	{
diff --git a/Engine/Source/Runtime/RenderCore/Private/DumpGPU.cpp b/Engine/Source/Runtime/RenderCore/Private/DumpGPU.cpp
--- a/Engine/Source/Runtime/RenderCore/Private/DumpGPU.cpp
+++ b/Engine/Source/Runtime/RenderCore/Private/DumpGPU.cpp
@@ -1,5 +1,6 @@
 // Copyright Epic Games, Inc. All Rights Reserved.
 
+#include "DumpGPU.h"
 #include "RenderGraph.h"
 #include "HAL/PlatformFileManager.h"
 #include "HAL/PlatformMisc.h"
@@ -17,6 +18,32 @@
 #include "GlobalShader.h"
 #include "RHIValidation.h"
 
+
+IDumpGPUUploadServiceProvider* IDumpGPUUploadServiceProvider::GProvider = nullptr;
+
+FString IDumpGPUUploadServiceProvider::FDumpParameters::DumpServiceParametersFileContent() const
+{
+	TSharedPtr<FJsonObject> JsonObject = MakeShareable(new FJsonObject);
+	JsonObject->SetStringField(TEXT("Type"), *Type);
+	JsonObject->SetStringField(TEXT("Time"), *Time);
+	JsonObject->SetStringField(TEXT("CompressionName"), *CompressionName.ToString());
+	JsonObject->SetStringField(TEXT("CompressionFiles"), *CompressionFiles);
+
+	FString OutputString;
+	TSharedRef<TJsonWriter<TCHAR, TPrettyJsonPrintPolicy<TCHAR>>> Writer = TJsonWriterFactory<TCHAR, TPrettyJsonPrintPolicy<TCHAR>>::Create(&OutputString);
+	FJsonSerializer::Serialize(JsonObject.ToSharedRef(), Writer);
+
+	return OutputString;
+}
+
+bool IDumpGPUUploadServiceProvider::FDumpParameters::DumpServiceParametersFile() const
+{
+	FString FileName = LocalPath / kServiceFileName;
+	FString OutputString = DumpServiceParametersFileContent();
+
+	return FFileHelper::SaveStringToFile(OutputString, *FileName);
+}
+
 #if RDG_DUMP_RESOURCES
 
 namespace
@@ -69,14 +96,6 @@ static TAutoConsoleVariable<int32> GDumpRenderingConsoleVariablesCVar(
 	TEXT("Whether to dump rendering console variables (enabled by default)."),
 	ECVF_Default);
 
-static TAutoConsoleVariable<int32> GDumpGPUCompressResources(
-	TEXT("r.DumpGPU.CompressResources"), 0,
-	TEXT("Whether to compress resource binary.\n")
-	TEXT(" 0: Disabled (default)\n")
-	TEXT(" 1: Zlib\n")
-	TEXT(" 2: GZip"),
-	ECVF_Default);
-
 static TAutoConsoleVariable<int32> GDumpTestEnableDiskWrite(
 	TEXT("r.DumpGPU.Test.EnableDiskWrite"), 1,
 	TEXT("Master switch whether any files should be written to disk, used for r.DumpGPU automation tests to not fill up workers' hard drive."),
@@ -92,6 +111,19 @@ static TAutoConsoleVariable<FString> GDumpGPUDirectoryCVar(
 	TEXT("Directory to dump to."),
 	ECVF_Default);
 
+static TAutoConsoleVariable<int32> GDumpGPUUploadCVar(
+	TEXT("r.DumpGPU.Upload"), 1,
+	TEXT("Allows to upload the GPU dump automatically if set-up."),
+	ECVF_Default);
+
+static TAutoConsoleVariable<int32> GDumpGPUUploadCompressResources(
+	TEXT("r.DumpGPU.Upload.CompressResources"), 1,
+	TEXT("Whether to compress resource binary.\n")
+	TEXT(" 0: Disabled (default)\n")
+	TEXT(" 1: Zlib\n")
+	TEXT(" 2: GZip"),
+	ECVF_Default);
+
 // Although this cvar does not seams used in the C++ code base, it is dumped by DumpRenderingCVarsToCSV() and used by GPUDumpViewer.html.
 static TAutoConsoleVariable<FString> GDumpGPUVisualizeResource(
 	TEXT("r.DumpGPU.Viewer.Visualize"), TEXT(""),
@@ -149,7 +181,8 @@ struct FRDGResourceDumpContext
 	static constexpr const TCHAR* kStructuresMetadataDir = TEXT("StructuresMetadata/");
 
 	bool bEnableDiskWrite = false;
-	FName ResourceCompressionName;
+	bool bUpload = false;
+	FName UploadResourceCompressionName;
 	FString DumpingDirectoryPath;
 	FDateTime Time;
 	FGenericPlatformMemoryConstants MemoryConstants;
@@ -168,7 +201,6 @@ struct FRDGResourceDumpContext
 	int64 MetadataFilesWriteBytes = 0;
 	int32 ResourceBinaryFilesOpened = 0;
 	int64 ResourceBinaryWriteBytes = 0;
-	int64 ResourceBinaryUncompressedBytes = 0;
 
 	enum class ETimingBucket : uint8
 	{
@@ -176,7 +208,6 @@ struct FRDGResourceDumpContext
 		RHIReleaseResources,
 		GPUWait,
 		CPUPostProcessing,
-		CPUCompression,
 		MetadataFileWrite,
 		ResourceBinaryFileWrite,
 		MAX
@@ -239,6 +270,15 @@ struct FRDGResourceDumpContext
 		}
 	}
 
+	IDumpGPUUploadServiceProvider::FDumpParameters GetDumpParameters() const
+	{
+		IDumpGPUUploadServiceProvider::FDumpParameters DumpServiceParameters;
+		DumpServiceParameters.Type = TEXT("DumpGPU");
+		DumpServiceParameters.LocalPath = DumpingDirectoryPath;
+		DumpServiceParameters.Time = Time.ToString();
+		return DumpServiceParameters;
+	}
+
 	bool IsDumpingFrame() const
 	{
 		check(IsInRenderingThread() || IsInParallelRenderingThread());
@@ -313,47 +353,14 @@ struct FRDGResourceDumpContext
 		return !Ar->IsError() && !Ar->IsCriticalError();
 	}
 
-	bool IsResourceBinaryCompressable(int64 UncompressedSize) const
-	{
-		return !ResourceCompressionName.IsNone() && (UncompressedSize < (int64(1) << int64(32)));
-	}
-
 	bool DumpResourceBinaryToFile(const uint8* UncompressedData, int64 UncompressedSize, const FString& FileName)
 	{
-		ResourceBinaryUncompressedBytes += UncompressedSize;
-
-		if (!IsResourceBinaryCompressable(UncompressedSize))
-		{
-			return DumpBinaryToFile(UncompressedData, UncompressedSize, FileName);
-		}
-
-		int32 CompressedSize;
-		TArray64<uint8> CompressedArray;
-		{
-			FTimeBucketMeasure TimeBucketMeasure(this, ETimingBucket::CPUCompression);
-
-			int32 MaxCompressedSize = FCompression::CompressMemoryBound(ResourceCompressionName, UncompressedSize, COMPRESS_BiasSize);
-			CompressedArray.SetNumUninitialized(MaxCompressedSize);
-
-			CompressedSize = MaxCompressedSize;
-			bool bSuccess = FCompression::CompressMemory(
-				ResourceCompressionName,
-				/* out */ CompressedArray.GetData(), /* out */ CompressedSize,
-				UncompressedData, UncompressedSize,
-				COMPRESS_BiasSize);
-		}
-
-		return DumpBinaryToFile(CompressedArray.GetData(), CompressedSize, FileName);
+		return DumpBinaryToFile(UncompressedData, UncompressedSize, FileName);
 	}
 
 	bool IsUnsafeToDumpResource(SIZE_T ResourceByteSize, float DumpMemoryMultiplier) const
 	{
 		uint64 AproximatedStagingMemoryRequired = uint64(double(ResourceByteSize) * DumpMemoryMultiplier);
-		if (IsResourceBinaryCompressable(ResourceByteSize))
-		{
-			AproximatedStagingMemoryRequired += FCompression::CompressMemoryBound(ResourceCompressionName, ResourceByteSize, COMPRESS_BiasSize);
-		}
-
 		uint64 MaxMemoryAvailable = FMath::Min(MemoryStats.AvailablePhysical, MemoryStats.AvailableVirtual);
 
 		return AproximatedStagingMemoryRequired > MaxMemoryAvailable;
@@ -1364,7 +1371,7 @@ void FRDGBuilder::InitResourceDump()
 	}
 }
 
-FString FRDGBuilder::BeginResourceDump(const TArray<FString>& Args)
+FString FRDGBuilder::BeginResourceDump(const TCHAR* Cmd)
 {
 	check(IsInGameThread());
 
@@ -1373,6 +1380,41 @@ FString FRDGBuilder::BeginResourceDump(const TArray<FString>& Args)
 		return FString();
 	}
 
+	TArray<FString> Tokens;
+	TArray<FString> Switches;
+	TMap<FString, FString> Parameters;
+	{
+		const TCHAR* LocalCmd = Cmd;
+
+		FString NextToken;
+		while (FParse::Token(LocalCmd, NextToken, false))
+		{
+			if (**NextToken == TCHAR('-'))
+			{
+				new(Switches) FString(NextToken.Mid(1));
+			}
+			else
+			{
+				new(Tokens) FString(NextToken);
+			}
+		}
+
+		for (int32 SwitchIdx = Switches.Num() - 1; SwitchIdx >= 0; --SwitchIdx)
+		{
+			FString& Switch = Switches[SwitchIdx];
+			TArray<FString> SplitSwitch;
+
+			// Remove Bar=1 from the switch list and put it in params as {Bar,1}.
+			// Note: Handle nested equality such as Bar="Key=Value"
+			int32 AssignmentIndex = 0;
+			if (Switch.FindChar(TEXT('='), AssignmentIndex))
+			{
+				Parameters.Add(Switch.Left(AssignmentIndex), Switch.RightChop(AssignmentIndex + 1).TrimQuotes());
+				Switches.RemoveAt(SwitchIdx);
+			}
+		}
+	}
+
 	FRDGResourceDumpContext NewResourceDumpContext;
 
 	NewResourceDumpContext.Time = FDateTime::Now();
@@ -1396,17 +1438,36 @@ FString FRDGBuilder::BeginResourceDump(const TArray<FString>& Args)
 		NewResourceDumpContext.DumpingDirectoryPath = DirectoryPath / FApp::GetProjectName() + TEXT("-") + FPlatformProperties::PlatformName() + TEXT("-") + NewResourceDumpContext.Time.ToString() + TEXT("/");
 	}
 	NewResourceDumpContext.bEnableDiskWrite = GDumpTestEnableDiskWrite.GetValueOnGameThread() != 0;
+
+	if (Switches.Contains(TEXT("upload")))
 	{
-		if (GDumpGPUCompressResources.GetValueOnGameThread() == 1)
+		if (!IDumpGPUUploadServiceProvider::GProvider || !NewResourceDumpContext.bEnableDiskWrite)
 		{
-			NewResourceDumpContext.ResourceCompressionName = NAME_Zlib;
+			UE_LOG(LogRendererCore, Warning, TEXT("DumpGPU upload services are not set up."));
 		}
-		else if (GDumpGPUCompressResources.GetValueOnGameThread() == 2)
+		else if (GDumpGPUUploadCVar.GetValueOnGameThread() == 0)
 		{
-			NewResourceDumpContext.ResourceCompressionName = NAME_Gzip;
+			UE_LOG(LogRendererCore, Warning, TEXT("DumpGPU upload services are not available because r.DumpGPU.Upload=0."));
+		}
+		else
+		{
+			NewResourceDumpContext.bUpload = true;
 		}
 	}
-	NewResourceDumpContext.bShowInExplore = NewResourceDumpContext.bEnableDiskWrite && GDumpExploreCVar.GetValueOnGameThread() != 0;
+
+	if (NewResourceDumpContext.bUpload)
+	{
+		if (GDumpGPUUploadCompressResources.GetValueOnGameThread() == 1)
+		{
+			NewResourceDumpContext.UploadResourceCompressionName = NAME_Zlib;
+		}
+		else if (GDumpGPUUploadCompressResources.GetValueOnGameThread() == 2)
+		{
+			NewResourceDumpContext.UploadResourceCompressionName = NAME_Gzip;
+		}
+	}
+
+	NewResourceDumpContext.bShowInExplore = NewResourceDumpContext.bEnableDiskWrite && GDumpExploreCVar.GetValueOnGameThread() != 0 && !NewResourceDumpContext.bUpload;
 	NewResourceDumpContext.MemoryConstants = FPlatformMemory::GetConstants();
 	NewResourceDumpContext.MemoryStats = FPlatformMemory::GetStats();
 
@@ -1425,6 +1486,12 @@ FString FRDGBuilder::BeginResourceDump(const TArray<FString>& Args)
 		NewResourceDumpContext.DumpStringToFile(TEXT(""), FString(FRDGResourceDumpContext::kBaseDir) / TEXT("PassDrawCounts.json"));
 	}
 
+	// Dump service parameters so GPUDumpViewer.html remain compatible when not using upload provider.
+	if (NewResourceDumpContext.bEnableDiskWrite)
+	{
+		NewResourceDumpContext.GetDumpParameters().DumpServiceParametersFile();
+	}
+
 	// Output informations
 	{
 		const TCHAR* BranchName = BuildSettings::GetBranchName();
@@ -1555,19 +1622,6 @@ void FRDGBuilder::EndResourceDump()
 		UE_LOG(LogRendererCore, Display, TEXT("Dumped GPU readback commands: %.3f s"), float(RHIReadbackCommandsSeconds));
 		UE_LOG(LogRendererCore, Display, TEXT("Dumped GPU wait: %.3f s"), float(GPUWaitSeconds));
 		UE_LOG(LogRendererCore, Display, TEXT("Dumped CPU resource binary post processing: %.3f s"), float(CPUPostProcessingSeconds));
-		if (!GRDGResourceDumpContext.ResourceCompressionName.IsNone())
-		{
-			double CPUCompressionSeconds = GRDGResourceDumpContext.TimingBucket[int32(FRDGResourceDumpContext::ETimingBucket::CPUCompression)];
-
-			UE_LOG(LogRendererCore, Display, TEXT("Dumped CPU resource binary compression from %.3f MB to %.3f MB (%.1f %%) using %s"),
-				float(GRDGResourceDumpContext.ResourceBinaryUncompressedBytes) / float(1024 * 1024),
-				float(GRDGResourceDumpContext.ResourceBinaryWriteBytes) / float(1024 * 1024),
-				100.0f * float(GRDGResourceDumpContext.ResourceBinaryWriteBytes) / float(GRDGResourceDumpContext.ResourceBinaryUncompressedBytes),
-				*GRDGResourceDumpContext.ResourceCompressionName.ToString());
-			UE_LOG(LogRendererCore, Display, TEXT("Dumped CPU resource binary compression took %.3f s at %.3f MB/s"),
-				float(CPUCompressionSeconds),
-				float(GRDGResourceDumpContext.ResourceBinaryUncompressedBytes) / (float(1024 * 1024) * float(CPUCompressionSeconds)));
-		}
 		UE_LOG(LogRendererCore, Display, TEXT("Dumped metadata: %.3f MB in %d files under %.3f s at %.3f MB/s"),
 			float(GRDGResourceDumpContext.MetadataFilesWriteBytes) / float(1024 * 1024),
 			GRDGResourceDumpContext.MetadataFilesOpened,
@@ -1592,6 +1646,20 @@ void FRDGBuilder::EndResourceDump()
 		FGenericCrashContext::DumpLog(GRDGResourceDumpContext.DumpingDirectoryPath / FRDGResourceDumpContext::kBaseDir);
 	}
 
+	if (GRDGResourceDumpContext.bUpload && IDumpGPUUploadServiceProvider::GProvider)
+	{
+		IDumpGPUUploadServiceProvider::FDumpParameters DumpCompletedParameters = GRDGResourceDumpContext.GetDumpParameters();
+
+		// Compress the resource binary in background before uploading.
+		if (!GRDGResourceDumpContext.UploadResourceCompressionName.IsNone())
+		{
+			DumpCompletedParameters.CompressionName = GRDGResourceDumpContext.UploadResourceCompressionName;
+			DumpCompletedParameters.CompressionFiles = FWildcardString(TEXT("*.bin"));
+		}
+
+		IDumpGPUUploadServiceProvider::GProvider->UploadDump(DumpCompletedParameters);
+	}
+
 	#if PLATFORM_DESKTOP
 	if (GRDGResourceDumpContext.bShowInExplore)
 	{
diff --git a/Engine/Source/Runtime/RenderCore/Private/DumpGPU.h b/Engine/Source/Runtime/RenderCore/Private/DumpGPU.h
new file mode 100644
--- /dev/null
+++ b/Engine/Source/Runtime/RenderCore/Private/DumpGPU.h
@@ -0,0 +1,32 @@
+// Copyright Epic Games, Inc. All Rights Reserved.
+
+#pragma once
+
+#include "CoreMinimal.h"
+#include "Misc/WildcardString.h"
+
+
+class RENDERCORE_API IDumpGPUUploadServiceProvider
+{
+public:
+
+	struct RENDERCORE_API FDumpParameters
+	{
+		static constexpr const TCHAR* kServiceFileName = TEXT("Base/DumpService.json");
+
+		FString Type;
+		FString LocalPath;
+		FString Time;
+
+		FName CompressionName;
+		FWildcardString CompressionFiles;
+
+		FString DumpServiceParametersFileContent() const;
+		bool DumpServiceParametersFile() const;
+	};
+
+	virtual void UploadDump(const FDumpParameters& Parameters) = 0;
+	virtual ~IDumpGPUUploadServiceProvider() = default;
+
+	static IDumpGPUUploadServiceProvider* GProvider;
+};
diff --git a/Engine/Source/Runtime/RenderCore/Public/RenderGraphBuilder.h b/Engine/Source/Runtime/RenderCore/Public/RenderGraphBuilder.h
--- a/Engine/Source/Runtime/RenderCore/Public/RenderGraphBuilder.h
+++ b/Engine/Source/Runtime/RenderCore/Public/RenderGraphBuilder.h
@@ -296,7 +296,7 @@ public:
 	FRDGBlackboard Blackboard;
 
 #if RDG_DUMP_RESOURCES
-	static FString BeginResourceDump(const TArray<FString>& Args);
+	static FString BeginResourceDump(const TCHAR* Cmd);
 	static void InitResourceDump();
 	static void EndResourceDump();
 	static bool IsDumpingFrame();
