Adding missing DDC key modifications for DDSPI controlled defines.

diff --git a/Engine/Source/Runtime/RenderCore/Private/Shader.cpp b/Engine/Source/Runtime/RenderCore/Private/Shader.cpp
--- a/Engine/Source/Runtime/RenderCore/Private/Shader.cpp
+++ b/Engine/Source/Runtime/RenderCore/Private/Shader.cpp
@@ -1887,6 +1887,11 @@ void ShaderMapAppendKeyString(EShaderPlatform Platform, FString& KeyString)
 		KeyString += TEXT("_PPDBM");
 	}
 
+	if (FDataDrivenShaderPlatformInfo::GetSupportsDistanceFields(Platform))
+	{
+		KeyString += TEXT("_DF");
+	}
+
 	if (RHISupportsMeshShadersTier0(Platform))
 	{
 		KeyString += TEXT("_MS_T0");
@@ -1897,6 +1902,11 @@ void ShaderMapAppendKeyString(EShaderPlatform Platform, FString& KeyString)
 		KeyString += TEXT("_MS_T1");
 	}
 
+	if (RHISupportsBindless(Platform))
+	{
+		KeyString += TEXT("_BNDLS");
+	}
+
 	if (ShouldCompileRayTracingShadersForProject(Platform))
 	{
 		static const auto CVarCompileCHS = IConsoleManager::Get().FindConsoleVariable(TEXT("r.RayTracing.CompileMaterialCHS"));
