Added RHI tracked access API to remove Unknown transitions.
 - New RHI command list SetTrackedAccess method for the user to supply a current whole-resource state.
 - New RHI command context GetTrackedAccess method for querying the tracked access in RHIBeginTransitions / RHIEndTransitions on the RHI thread.
 - Hooked RHICmdList.Transition and FRHICommandListExecutor::Transition to assign tracked state automatically.
 - Refactored RDG and resource pools to use new RHI tracking.
      - FRDGPooledBuffer / FRDGPooledTexture no longer contain tracked state. RDG temp-allocates state through the graph allocator instead.
      - All prologue transitions are 'Unknown', and all epilogue transitions coalesce into a whole resource state.
 - Implemented platform support for patching the 'before' state with the tracked state.
 - Implemented various RHI validation checks:
      - Asserts that the user assigned tracked state matches RHI validation tracked state, for all subresources.
      - Asserts that tracked state is not assigned or queried from a parallel translation context.
 - Added FRHIViewableResource and FRHIView base classes to RHI. FRHIView contains a pointer to an FRHIViewableResource. This is currently a raw pointer, but should be extended to a full reference in a later CL.

NOTE on RHI thread constraint:

Transition evaluation is now restricted to the RHI thread (i.e. no parallel translation contexts). Transitions aren't performed in parallel translate contexts anyway, so this is not a problem. If, however, we decide to refactor parallel translation to be more general, this implementation could be extended to track the state per context and update from the 'dispatch' thread.

diff --git a/Engine/Plugins/RenderGraphInsights/Source/RenderGraphInsights/Private/RenderGraphTrack.cpp b/Engine/Plugins/RenderGraphInsights/Source/RenderGraphInsights/Private/RenderGraphTrack.cpp
--- a/Engine/Plugins/RenderGraphInsights/Source/RenderGraphInsights/Private/RenderGraphTrack.cpp
+++ b/Engine/Plugins/RenderGraphInsights/Source/RenderGraphInsights/Private/RenderGraphTrack.cpp
@@ -1199,7 +1199,7 @@ void FRenderGraphTrack::BuildDrawState(FRenderGraphTrackDrawStateBuilder& Builde
 			uint64 SizeInBytes{};
 			uint32 Index{};
 			uint32 Order{};
-			ERDGParentResourceType Type = ERDGParentResourceType::Texture;
+			ERDGViewableResourceType Type = ERDGViewableResourceType::Texture;
 			bool bHasPreviousOwner{};
 		};
 
@@ -1297,7 +1297,7 @@ void FRenderGraphTrack::BuildDrawState(FRenderGraphTrackDrawStateBuilder& Builde
 					Entry.SizeInBytes = Texture.SizeInBytes;
 					Entry.Index = TextureIndex;
 					Entry.Order = Texture.Order;
-					Entry.Type = ERDGParentResourceType::Texture;
+					Entry.Type = ERDGViewableResourceType::Texture;
 					Entry.bHasPreviousOwner = Texture.PrevousOwnerHandle.IsValid();
 				}
 			}
@@ -1324,7 +1324,7 @@ void FRenderGraphTrack::BuildDrawState(FRenderGraphTrackDrawStateBuilder& Builde
 					Entry.SizeInBytes = Buffer.SizeInBytes;
 					Entry.Index = BufferIndex;
 					Entry.Order = Buffer.Order;
-					Entry.Type = ERDGParentResourceType::Buffer;
+					Entry.Type = ERDGViewableResourceType::Buffer;
 					Entry.bHasPreviousOwner = Buffer.PrevousOwnerHandle.IsValid();
 				}
 			}
@@ -1369,7 +1369,7 @@ void FRenderGraphTrack::BuildDrawState(FRenderGraphTrackDrawStateBuilder& Builde
 			{
 				if (!Entry.bHasPreviousOwner)
 				{
-					auto& Array = Entry.Type == ERDGParentResourceType::Texture ? TextureIndexToDepth : BufferIndexToDepth;
+					auto& Array = Entry.Type == ERDGViewableResourceType::Texture ? TextureIndexToDepth : BufferIndexToDepth;
 					Array[Entry.Index] = DepthOffset++;
 				}
 
diff --git a/Engine/Source/Runtime/Apple/AGXRHI/Public/AGXResources.h b/Engine/Source/Runtime/Apple/AGXRHI/Public/AGXResources.h
--- a/Engine/Source/Runtime/Apple/AGXRHI/Public/AGXResources.h
+++ b/Engine/Source/Runtime/Apple/AGXRHI/Public/AGXResources.h
@@ -546,7 +546,8 @@ class FAGXShaderResourceView final : public FRHIShaderResourceView, public FAGXR
 {
 public:
 	explicit FAGXShaderResourceView(const FShaderResourceViewInitializer& Initializer)
-		: FAGXResourceViewBase(
+		: FRHIShaderResourceView(Initializer.AsBufferSRV().Buffer)
+		, FAGXResourceViewBase(
 			  Initializer.AsBufferSRV().Buffer
 			, Initializer.AsBufferSRV().StartOffsetBytes
 			, Initializer.AsBufferSRV().NumElements
@@ -555,7 +556,8 @@ public:
 	{}
 
 	explicit FAGXShaderResourceView(FRHITexture* Texture, const FRHITextureSRVCreateInfo& CreateInfo)
-		: FAGXResourceViewBase(
+		: FRHIShaderResourceView(Texture)
+		, FAGXResourceViewBase(
 			  Texture
 			, CreateInfo.Format
 			, CreateInfo.MipLevel
@@ -574,18 +576,21 @@ class FAGXUnorderedAccessView final : public FRHIUnorderedAccessView, public FAG
 {
 public:
 	explicit FAGXUnorderedAccessView(FRHIBuffer* Buffer, EPixelFormat Format)
-		: FAGXResourceViewBase(Buffer, 0, UINT_MAX, Format)
+		: FRHIUnorderedAccessView(Buffer)
+		, FAGXResourceViewBase(Buffer, 0, UINT_MAX, Format)
 	{}
 
 	explicit FAGXUnorderedAccessView(FRHIBuffer* Buffer, bool bUseUAVCounter, bool bAppendBuffer)
-		: FAGXResourceViewBase(Buffer, 0, UINT_MAX, PF_Unknown)
+		: FRHIUnorderedAccessView(Buffer)
+		, FAGXResourceViewBase(Buffer, 0, UINT_MAX, PF_Unknown)
 	{
 		checkf(!bUseUAVCounter, TEXT("UAV counters not implemented."));
 		checkf(!bAppendBuffer, TEXT("UAV append buffers not implemented."));
 	}
 
 	explicit FAGXUnorderedAccessView(FRHITexture* Texture, uint32 MipLevel, uint16 FirstArraySlice, uint16 NumArraySlices)
-		: FAGXResourceViewBase(
+		: FRHIUnorderedAccessView(Texture)
+		, FAGXResourceViewBase(
 			  Texture
 			, PF_Unknown
 			, MipLevel
diff --git a/Engine/Source/Runtime/Apple/MetalRHI/Public/MetalResources.h b/Engine/Source/Runtime/Apple/MetalRHI/Public/MetalResources.h
--- a/Engine/Source/Runtime/Apple/MetalRHI/Public/MetalResources.h
+++ b/Engine/Source/Runtime/Apple/MetalRHI/Public/MetalResources.h
@@ -546,7 +546,8 @@ class FMetalShaderResourceView final : public FRHIShaderResourceView, public FMe
 {
 public:
 	explicit FMetalShaderResourceView(const FShaderResourceViewInitializer& Initializer)
-		: FMetalResourceViewBase(
+		: FRHIShaderResourceView(Initializer.AsBufferSRV().Buffer)
+		, FMetalResourceViewBase(
 			  Initializer.AsBufferSRV().Buffer
 			, Initializer.AsBufferSRV().StartOffsetBytes
 			, Initializer.AsBufferSRV().NumElements
@@ -555,7 +556,8 @@ public:
 	{}
 
 	explicit FMetalShaderResourceView(FRHITexture* Texture, const FRHITextureSRVCreateInfo& CreateInfo)
-		: FMetalResourceViewBase(
+		: FRHIShaderResourceView(Texture)
+		, FMetalResourceViewBase(
 			  Texture
 			, CreateInfo.Format
 			, CreateInfo.MipLevel
@@ -574,18 +576,21 @@ class FMetalUnorderedAccessView final : public FRHIUnorderedAccessView, public F
 {
 public:
 	explicit FMetalUnorderedAccessView(FRHIBuffer* Buffer, EPixelFormat Format)
-		: FMetalResourceViewBase(Buffer, 0, UINT_MAX, Format)
+		: FRHIUnorderedAccessView(Buffer)
+		, FMetalResourceViewBase(Buffer, 0, UINT_MAX, Format)
 	{}
 
 	explicit FMetalUnorderedAccessView(FRHIBuffer* Buffer, bool bUseUAVCounter, bool bAppendBuffer)
-		: FMetalResourceViewBase(Buffer, 0, UINT_MAX, PF_Unknown)
+		: FRHIUnorderedAccessView(Buffer)
+		, FMetalResourceViewBase(Buffer, 0, UINT_MAX, PF_Unknown)
 	{
 		checkf(!bUseUAVCounter, TEXT("UAV counters not implemented."));
 		checkf(!bAppendBuffer, TEXT("UAV append buffers not implemented."));
 	}
 
 	explicit FMetalUnorderedAccessView(FRHITexture* Texture, uint32 MipLevel, uint16 FirstArraySlice, uint16 NumArraySlices)
-		: FMetalResourceViewBase(
+		: FRHIUnorderedAccessView(Buffer)
+		, FMetalResourceViewBase(
 			  Texture
 			, PF_Unknown
 			, MipLevel
diff --git a/Engine/Source/Runtime/D3D12RHI/Private/D3D12Commands.cpp b/Engine/Source/Runtime/D3D12RHI/Private/D3D12Commands.cpp
--- a/Engine/Source/Runtime/D3D12RHI/Private/D3D12Commands.cpp
+++ b/Engine/Source/Runtime/D3D12RHI/Private/D3D12Commands.cpp
@@ -459,12 +459,6 @@ static void HandleResourceTransitions(FD3D12CommandContext& Context, const FD3D1
 
 	for (const FRHITransitionInfo& Info : TransitionData->TransitionInfos)
 	{
-		// A transition back to discard is a no-op.
-		if (Info.AccessAfter == ERHIAccess::Discard)
-		{
-			continue;
-		}
-
 		const bool bUAVAccessAfter = EnumHasAnyFlags(Info.AccessAfter, ERHIAccess::UAVMask);
 
 		// If targeting all pipelines and UAV access then add UAV barrier even with RHI based transitions
@@ -498,23 +492,31 @@ static void HandleResourceTransitions(FD3D12CommandContext& Context, const FD3D1
 						return;
 					}
 
+					ERHIAccess AccessBefore = Info.AccessBefore;
+
+					if (AccessBefore == ERHIAccess::Unknown)
+					{
+						if (FRHIViewableResource* ViewableResource = GetViewableResource(Info))
+						{
+							AccessBefore = Context.GetTrackedAccess(ViewableResource);
+						}
+					}
+
 					const bool bIsAsyncCompute = EnumHasAnyFlags(TransitionData->DstPipelines, ERHIPipeline::AsyncCompute);
-					const bool bKnownBeforeState = (Info.AccessBefore != ERHIAccess::Unknown && Info.AccessBefore != ERHIAccess::Discard);
 
-					// Derive or get the correct before & after state
-					D3D12_RESOURCE_STATES BeforeState = bKnownBeforeState ? GetD3D12ResourceState(Info.AccessBefore, bIsAsyncCompute) : D3D12_RESOURCE_STATE_TBD;
-					D3D12_RESOURCE_STATES AfterState = GetD3D12ResourceState(Info.AccessAfter, bIsAsyncCompute);
+					D3D12_RESOURCE_STATES StateBefore =      AccessBefore == ERHIAccess::Discard ? GetInitialResourceState(Resource->GetDesc()) : GetD3D12ResourceState(     AccessBefore, bIsAsyncCompute);
+					D3D12_RESOURCE_STATES StateAfter  = Info.AccessAfter  == ERHIAccess::Discard ? GetInitialResourceState(Resource->GetDesc()) : GetD3D12ResourceState(Info.AccessAfter,  bIsAsyncCompute);
 					
 					// Add the compression flags if needed
 					if (EnumHasAnyFlags(Info.Flags, EResourceTransitionFlags::MaintainCompression))
 					{
-						AfterState |= SkipFastClearEliminateState;
+						StateAfter |= SkipFastClearEliminateState;
 					}
 
 					// enqueue the correct transitions
 					if (Info.IsWholeResource() || Resource->GetSubresourceCount() == 1)
 					{
-						if (FD3D12DynamicRHI::TransitionResource(Context.CommandListHandle, Resource, BeforeState, AfterState, D3D12_RESOURCE_BARRIER_ALL_SUBRESOURCES, FD3D12DynamicRHI::ETransitionMode::Apply))
+						if (FD3D12DynamicRHI::TransitionResource(Context.CommandListHandle, Resource, StateBefore, StateAfter, D3D12_RESOURCE_BARRIER_ALL_SUBRESOURCES, FD3D12DynamicRHI::ETransitionMode::Apply))
 						{
 							bUAVBarrier = true;
 						}
@@ -523,7 +525,7 @@ static void HandleResourceTransitions(FD3D12CommandContext& Context, const FD3D1
 					{
 						EnumerateSubresources(Resource, Info, [&](uint32 Subresource)
 						{
-							if (FD3D12DynamicRHI::TransitionResource(Context.CommandListHandle, Resource, BeforeState, AfterState, Subresource, FD3D12DynamicRHI::ETransitionMode::Apply))
+							if (FD3D12DynamicRHI::TransitionResource(Context.CommandListHandle, Resource, StateBefore, StateAfter, Subresource, FD3D12DynamicRHI::ETransitionMode::Apply))
 							{
 								bUAVBarrier = true;
 							}
diff --git a/Engine/Source/Runtime/D3D12RHI/Private/D3D12RenderTarget.cpp b/Engine/Source/Runtime/D3D12RHI/Private/D3D12RenderTarget.cpp
--- a/Engine/Source/Runtime/D3D12RHI/Private/D3D12RenderTarget.cpp
+++ b/Engine/Source/Runtime/D3D12RHI/Private/D3D12RenderTarget.cpp
@@ -207,10 +207,10 @@ void FD3D12CommandContext::RHICopyToResolveTarget(FRHITexture* SourceTextureRHI,
 	FD3D12Texture* SourceTexture = GetD3D12TextureFromRHITexture(SourceTextureRHI);
 	FD3D12Texture* DestTexture = GetD3D12TextureFromRHITexture(DestTextureRHI);
 
-	ETextureDimension SourceDimension = SourceTexture->GetDesc().Dimension;
-	ETextureDimension DestDimension = DestTexture->GetDesc().Dimension;
+	const FRHITextureDesc& SourceDesc = SourceTexture->GetDesc();
+	const FRHITextureDesc& DestDesc = DestTexture->GetDesc();
 
-	if (SourceDimension == ETextureDimension::Texture2D && DestDimension == ETextureDimension::Texture2D)
+	if (SourceDesc.IsTexture2D() && DestDesc.IsTexture2D())
 	{
 		const D3D_FEATURE_LEVEL FeatureLevel = GetParentDevice()->GetParentAdapter()->GetFeatureLevel();
 		if (SourceTexture != DestTexture)
@@ -376,7 +376,7 @@ void FD3D12CommandContext::RHICopyToResolveTarget(FRHITexture* SourceTextureRHI,
 			}
 		}
 	}
-	else if (SourceDimension == ETextureDimension::TextureCube && DestDimension == ETextureDimension::TextureCube)
+	else if (SourceDesc.IsTextureCube() && DestDesc.IsTextureCube())
 	{
 		if (SourceTexture != DestTexture)
 		{
@@ -430,7 +430,7 @@ void FD3D12CommandContext::RHICopyToResolveTarget(FRHITexture* SourceTextureRHI,
 			}
 		}
 	}
-	else if (SourceDimension == ETextureDimension::Texture2D && DestDimension == ETextureDimension::TextureCube)
+	else if (SourceDesc.IsTexture2D() && DestDesc.IsTextureCube())
 	{
 		// If source is 2D and Dest is a cube then copy the 2D texture to the specified cube face.
 		// Determine the cubemap face being resolved.
@@ -455,23 +455,13 @@ void FD3D12CommandContext::RHICopyToResolveTarget(FRHITexture* SourceTextureRHI,
 		CommandListHandle.UpdateResidency(DestTexture->GetResource());
 	}
 
-	else if (SourceDimension == ETextureDimension::Texture3D && DestDimension == ETextureDimension::Texture3D)
+	else if (SourceDesc.IsTexture3D() && DestDesc.IsTexture3D())
 	{
 		// bit of a hack.  no one resolves slice by slice and 0 is the default value.  assume for the moment they are resolving the whole texture.
 		check(ResolveParams.SourceArrayIndex == 0);
 		check(SourceTexture == DestTexture);
 	}
 
-	// Force transtion to readable since RHI user expects both source and dest to be in readable state after calling RHICopyToResolveTarget
-	if (SourceTextureRHI && ResolveParams.SourceAccessFinal != ERHIAccess::Unknown)
-	{
-		RHICmdList.TransitionResource(ResolveParams.SourceAccessFinal, SourceTextureRHI);
-	}
-	if (DestTextureRHI && SourceTextureRHI != DestTextureRHI && ResolveParams.DestAccessFinal != ERHIAccess::Unknown)
-	{
-		RHICmdList.TransitionResource(ResolveParams.DestAccessFinal, DestTextureRHI);
-	}
-
 	ConditionalFlushCommandList();
 
 	DEBUG_EXECUTE_COMMAND_LIST(this);
diff --git a/Engine/Source/Runtime/D3D12RHI/Private/D3D12Texture.cpp b/Engine/Source/Runtime/D3D12RHI/Private/D3D12Texture.cpp
--- a/Engine/Source/Runtime/D3D12RHI/Private/D3D12Texture.cpp
+++ b/Engine/Source/Runtime/D3D12RHI/Private/D3D12Texture.cpp
@@ -860,9 +860,8 @@ FD3D12Texture* FD3D12DynamicRHI::CreateD3D12Texture(const FRHITextureCreateDesc&
 		ClearValuePtr = &ClearValue;
 	}
 
-	// The state this resource will be in when it leaves this function
 	const FD3D12Resource::FD3D12ResourceTypeHelper Type(ResourceDesc, D3D12_HEAP_TYPE_DEFAULT);
-	const D3D12_RESOURCE_STATES InitialState = Type.GetOptimalInitialState(CreateDesc.InitialState, false);
+	const D3D12_RESOURCE_STATES InitialState = Type.GetOptimalInitialState(CreateDesc.InitialState, true);
 	const D3D12_RESOURCE_STATES CreateState = (CreateDesc.BulkData != nullptr) ? D3D12_RESOURCE_STATE_COPY_DEST : InitialState;
 
 	FD3D12Adapter* Adapter = &GetAdapter();
diff --git a/Engine/Source/Runtime/D3D12RHI/Private/D3D12TransientResourceAllocator.cpp b/Engine/Source/Runtime/D3D12RHI/Private/D3D12TransientResourceAllocator.cpp
--- a/Engine/Source/Runtime/D3D12RHI/Private/D3D12TransientResourceAllocator.cpp
+++ b/Engine/Source/Runtime/D3D12RHI/Private/D3D12TransientResourceAllocator.cpp
@@ -144,20 +144,9 @@ FRHITransientTexture* FD3D12TransientResourceHeapAllocator::CreateTexture(const
 	return CreateTextureInternal(InCreateInfo, InDebugName, InPassIndex, Info.SizeInBytes, Info.Alignment,
 		[&](const FRHITransientHeap::FResourceInitializer& Initializer)
 	{
-		ERHIAccess InitialState = ERHIAccess::UAVMask;
-
-		if (EnumHasAnyFlags(InCreateInfo.Flags, TexCreate_RenderTargetable | TexCreate_ResolveTargetable))
-		{
-			InitialState = ERHIAccess::RTV;
-		}
-		else if (EnumHasAnyFlags(InCreateInfo.Flags, TexCreate_DepthStencilTargetable))
-		{
-			InitialState = ERHIAccess::DSVWrite;
-		}
-
 		FResourceAllocatorAdapter ResourceAllocatorAdapter(GetParentAdapter(), static_cast<FD3D12TransientHeap&>(Initializer.Heap), Initializer.Allocation, Desc);
 
-		FRHITextureCreateDesc CreateDesc(InCreateInfo, InitialState, InDebugName);
+		FRHITextureCreateDesc CreateDesc(InCreateInfo, ERHIAccess::Discard, InDebugName);
 		FRHITexture* Texture = DynamicRHI->CreateD3D12Texture(CreateDesc, nullptr, &ResourceAllocatorAdapter);
 		return new FRHITransientTexture(Texture, ResourceAllocatorAdapter.GpuVirtualAddress, Initializer.Hash, Info.SizeInBytes, ERHITransientAllocationType::Heap, InCreateInfo);
 	});
@@ -211,6 +200,7 @@ FRHITransientBuffer* FD3D12TransientResourceHeapAllocator::CreateBuffer(const FR
 	{
 		FResourceAllocatorAdapter ResourceAllocatorAdapter(GetParentAdapter(), static_cast<FD3D12TransientHeap&>(Initializer.Heap), Initializer.Allocation, Desc);
 		FRHIBuffer* Buffer = FD3D12DynamicRHI::GetD3DRHI()->CreateBuffer(InCreateInfo, InDebugName, ERHIAccess::Discard, &ResourceAllocatorAdapter);
+		Buffer->SetTrackedAccess_Unsafe(ERHIAccess::Discard);
 		return new FRHITransientBuffer(Buffer, ResourceAllocatorAdapter.GpuVirtualAddress, Initializer.Hash, Size, ERHITransientAllocationType::Heap, InCreateInfo);
 	});
 }
\ No newline at end of file
diff --git a/Engine/Source/Runtime/D3D12RHI/Private/D3D12UAV.cpp b/Engine/Source/Runtime/D3D12RHI/Private/D3D12UAV.cpp
--- a/Engine/Source/Runtime/D3D12RHI/Private/D3D12UAV.cpp
+++ b/Engine/Source/Runtime/D3D12RHI/Private/D3D12UAV.cpp
@@ -3,13 +3,15 @@
 #include "D3D12RHIPrivate.h"
 #include "ClearReplacementShaders.h"
 
-FD3D12UnorderedAccessView::FD3D12UnorderedAccessView(FD3D12Device* InParent)
-	: FD3D12View(InParent, ERHIDescriptorHeapType::Standard, ViewSubresourceSubsetFlags_None)
+FD3D12UnorderedAccessView::FD3D12UnorderedAccessView(FD3D12Device* InParent, FRHIViewableResource* InParentResource)
+	: FRHIUnorderedAccessView(InParentResource)
+	, FD3D12View(InParent, ERHIDescriptorHeapType::Standard, ViewSubresourceSubsetFlags_None)
 {
 }
 
-FD3D12UnorderedAccessView::FD3D12UnorderedAccessView(FD3D12Device* InParent, const D3D12_UNORDERED_ACCESS_VIEW_DESC& InDesc, FD3D12BaseShaderResource* InBaseShaderResource, FD3D12Resource* InCounterResource)
-	: FD3D12View(InParent, ERHIDescriptorHeapType::Standard, ViewSubresourceSubsetFlags_None)
+FD3D12UnorderedAccessView::FD3D12UnorderedAccessView(FD3D12Device* InParent, FRHIViewableResource* InParentResource, const D3D12_UNORDERED_ACCESS_VIEW_DESC& InDesc, FD3D12BaseShaderResource* InBaseShaderResource, FD3D12Resource* InCounterResource)
+	: FRHIUnorderedAccessView(InParentResource)
+	, FD3D12View(InParent, ERHIDescriptorHeapType::Standard, ViewSubresourceSubsetFlags_None)
 	, CounterResource(InCounterResource)
 {
 	SetDesc(InDesc);
@@ -58,7 +60,7 @@ inline FD3D12UnorderedAccessView* CreateUAV(D3D12_UNORDERED_ACCESS_VIEW_DESC& De
 			Device->GetParentAdapter()->CreateBuffer(HeapProps, Node, InitialState, ED3D12ResourceStateMode::MultiState, D3D12_RESOURCE_STATE_UNORDERED_ACCESS, 4, &CounterResource,  TEXT("Counter"), D3D12_RESOURCE_FLAG_ALLOW_UNORDERED_ACCESS);
 		}
 
-		return new FD3D12UnorderedAccessView(Device, Desc, Resource, CounterResource);
+		return new FD3D12UnorderedAccessView(Device, Resource, Desc, Resource, CounterResource);
 	});
 }
 
diff --git a/Engine/Source/Runtime/D3D12RHI/Private/D3D12Util.cpp b/Engine/Source/Runtime/D3D12RHI/Private/D3D12Util.cpp
--- a/Engine/Source/Runtime/D3D12RHI/Private/D3D12Util.cpp
+++ b/Engine/Source/Runtime/D3D12RHI/Private/D3D12Util.cpp
@@ -1273,6 +1273,7 @@ D3D12_RESOURCE_STATES GetD3D12ResourceState(ERHIAccess InRHIAccess, bool InIsAsy
 	case ERHIAccess::DSVWrite:				return D3D12_RESOURCE_STATE_DEPTH_WRITE;
 	case ERHIAccess::CopyDest:				return D3D12_RESOURCE_STATE_COPY_DEST;
 	case ERHIAccess::ResolveDst:			return D3D12_RESOURCE_STATE_RESOLVE_DEST;
+	case ERHIAccess::Present:				return D3D12_RESOURCE_STATE_PRESENT;
 
 	// Generic read for mask read states
 	case ERHIAccess::ReadOnlyMask:	
diff --git a/Engine/Source/Runtime/D3D12RHI/Private/D3D12View.cpp b/Engine/Source/Runtime/D3D12RHI/Private/D3D12View.cpp
--- a/Engine/Source/Runtime/D3D12RHI/Private/D3D12View.cpp
+++ b/Engine/Source/Runtime/D3D12RHI/Private/D3D12View.cpp
@@ -125,7 +125,8 @@ void FD3D12ViewDescriptorHandle::UpdateBindlessSlot(ED3D12DescriptorCreateReason
 ////////////////////////////////////////////////////////////////////////////////////////////////////////////////////////
 
 FD3D12ShaderResourceView::FD3D12ShaderResourceView(FD3D12Device* InParent)
-	: FD3D12View(InParent, ERHIDescriptorHeapType::Standard, ViewSubresourceSubsetFlags_None)
+	: FRHIShaderResourceView(nullptr)
+	, FD3D12View(InParent, ERHIDescriptorHeapType::Standard, ViewSubresourceSubsetFlags_None)
 	, bContainsDepthPlane(false)
 	, bContainsStencilPlane(false)
 	, bSkipFastClearFinalize(false)
diff --git a/Engine/Source/Runtime/D3D12RHI/Private/D3D12View.h b/Engine/Source/Runtime/D3D12RHI/Private/D3D12View.h
--- a/Engine/Source/Runtime/D3D12RHI/Private/D3D12View.h
+++ b/Engine/Source/Runtime/D3D12RHI/Private/D3D12View.h
@@ -835,9 +835,9 @@ protected:
 class FD3D12UnorderedAccessView : public FRHIUnorderedAccessView, public FD3D12View < D3D12_UNORDERED_ACCESS_VIEW_DESC >, public FD3D12LinkedAdapterObject<FD3D12UnorderedAccessView>
 {
 protected:
-	FD3D12UnorderedAccessView(FD3D12Device* InParent);
+	FD3D12UnorderedAccessView(FD3D12Device* InParent, FRHIViewableResource* InParentResource);
 public:
-	FD3D12UnorderedAccessView(FD3D12Device* InParent, const D3D12_UNORDERED_ACCESS_VIEW_DESC& InDesc, FD3D12BaseShaderResource* InBaseShaderResource, FD3D12Resource* InCounterResource);
+	FD3D12UnorderedAccessView(FD3D12Device* InParent, FRHIViewableResource* InParentResource, const D3D12_UNORDERED_ACCESS_VIEW_DESC& InDesc, FD3D12BaseShaderResource* InBaseShaderResource, FD3D12Resource* InCounterResource);
 
 	virtual void RecreateView() override;
 
diff --git a/Engine/Source/Runtime/NullDrv/Public/NullRHI.h b/Engine/Source/Runtime/NullDrv/Public/NullRHI.h
--- a/Engine/Source/Runtime/NullDrv/Public/NullRHI.h
+++ b/Engine/Source/Runtime/NullDrv/Public/NullRHI.h
@@ -146,35 +146,35 @@ public:
 
 	virtual FUnorderedAccessViewRHIRef RHICreateUnorderedAccessView(FRHIBuffer* StructuredBuffer, bool bUseUAVCounter, bool bAppendBuffer) final override
 	{ 
-		return new FRHIUnorderedAccessView(); 
+		return new FRHIUnorderedAccessView(StructuredBuffer); 
 	}
 
 
 	virtual FUnorderedAccessViewRHIRef RHICreateUnorderedAccessView(FRHITexture* Texture, uint32 MipLevel, uint16 FirstArraySlice, uint16 NumArraySlices) final override
 	{ 
-		return new FRHIUnorderedAccessView(); 
+		return new FRHIUnorderedAccessView(Texture);
 	}
 
 
 	virtual FUnorderedAccessViewRHIRef RHICreateUnorderedAccessView(FRHIBuffer* Buffer, uint8 Format) final override
 	{ 
-		return new FRHIUnorderedAccessView(); 
+		return new FRHIUnorderedAccessView(Buffer);
 	}
 
 	virtual FShaderResourceViewRHIRef RHICreateShaderResourceView(FRHIBuffer* Buffer) final override
 	{ 
-		return new FRHIShaderResourceView(); 
+		return new FRHIShaderResourceView(Buffer); 
 	}
 
 
 	virtual FShaderResourceViewRHIRef RHICreateShaderResourceView(FRHIBuffer* Buffer, uint32 Stride, uint8 Format) final override
 	{ 
-		return new FRHIShaderResourceView(); 
+		return new FRHIShaderResourceView(Buffer);
 	}
 
 	virtual FShaderResourceViewRHIRef RHICreateShaderResourceView(const FShaderResourceViewInitializer& Initializer) final override
 	{
-		return new FRHIShaderResourceView();
+		return new FRHIShaderResourceView(Initializer.AsBufferSRV().Buffer);
 	}
 
 	virtual void RHIClearUAVFloat(FRHIUnorderedAccessView* UnorderedAccessViewRHI, const FVector4f& Values) final override
@@ -241,7 +241,7 @@ public:
 
 	virtual FShaderResourceViewRHIRef RHICreateShaderResourceView(FRHITexture* Texture, const FRHITextureSRVCreateInfo& CreateInfo) final override
 	{ 
-		return new FRHIShaderResourceView(); 
+		return new FRHIShaderResourceView(Texture); 
 	}
 
 	virtual void RHIGenerateMips(FRHITexture* Texture) final override
diff --git a/Engine/Source/Runtime/OpenGLDrv/Private/OpenGLTexture.cpp b/Engine/Source/Runtime/OpenGLDrv/Private/OpenGLTexture.cpp
--- a/Engine/Source/Runtime/OpenGLDrv/Private/OpenGLTexture.cpp
+++ b/Engine/Source/Runtime/OpenGLDrv/Private/OpenGLTexture.cpp
@@ -1656,7 +1656,8 @@ void FOpenGLDynamicRHI::RHICopySharedMips(FRHITexture2D* DestTexture2D, FRHIText
 }
 
 FOpenGLShaderResourceView::FOpenGLShaderResourceView(FOpenGLTexture* InTexture, const FRHITextureSRVCreateInfo& CreateInfo)
-	: Target      (InTexture->Target)
+	: FRHIShaderResourceView(InTexture)
+	, Target      (InTexture->Target)
 	, Texture     (InTexture)
 	, LimitMip    (CreateInfo.MipLevel)
 	, OwnsResource(FOpenGL::SupportsTextureView())
diff --git a/Engine/Source/Runtime/OpenGLDrv/Private/OpenGLUAV.cpp b/Engine/Source/Runtime/OpenGLDrv/Private/OpenGLUAV.cpp
--- a/Engine/Source/Runtime/OpenGLDrv/Private/OpenGLUAV.cpp
+++ b/Engine/Source/Runtime/OpenGLDrv/Private/OpenGLUAV.cpp
@@ -30,6 +30,7 @@ static void BindGLTexBufferRange(GLenum Target, GLenum InternalFormat, GLuint Bu
 }
 
 FOpenGLShaderResourceView::FOpenGLShaderResourceView(const FShaderResourceViewInitializer& Initializer)
+	: FRHIShaderResourceView(Initializer.AsBufferSRV().Buffer)
 {
 	FShaderResourceViewInitializer::FBufferShaderResourceViewInitializer Desc = Initializer.AsBufferSRV();
 	Buffer = FOpenGLDynamicRHI::ResourceCast(Desc.Buffer);
@@ -190,6 +191,7 @@ FUnorderedAccessViewRHIRef FOpenGLDynamicRHI::RHICreateUnorderedAccessView(FRHIT
 }
 
 FOpenGLTextureUnorderedAccessView::FOpenGLTextureUnorderedAccessView(FRHITexture* InTextureRHI):
+	FOpenGLUnorderedAccessView(InTextureRHI),
 	TextureRHI(InTextureRHI)
 {
 	VERIFY_GL_SCOPE();
@@ -211,6 +213,7 @@ FUnorderedAccessViewRHIRef FOpenGLDynamicRHI::RHICreateUnorderedAccessView(FRHIB
 }
 
 FOpenGLTexBufferUnorderedAccessView::FOpenGLTexBufferUnorderedAccessView(FOpenGLDynamicRHI* InOpenGLRHI, FRHIBuffer* InBufferRHI, uint8 Format):
+	FOpenGLUnorderedAccessView(InBufferRHI),
 	BufferRHI(InBufferRHI),
 	OpenGLRHI(InOpenGLRHI)
 {
@@ -261,7 +264,8 @@ FUnorderedAccessViewRHIRef FOpenGLDynamicRHI::RHICreateUnorderedAccessView(FRHIB
 }
 
 FOpenGLBufferUnorderedAccessView::FOpenGLBufferUnorderedAccessView(FOpenGLDynamicRHI* InOpenGLRHI, FRHIBuffer* InBufferRHI)
-	: BufferRHI(InBufferRHI)
+	: FOpenGLUnorderedAccessView(InBufferRHI)
+	, BufferRHI(InBufferRHI)
 	, OpenGLRHI(InOpenGLRHI)
 {
 	VERIFY_GL_SCOPE();
diff --git a/Engine/Source/Runtime/OpenGLDrv/Public/OpenGLResources.h b/Engine/Source/Runtime/OpenGLDrv/Public/OpenGLResources.h
--- a/Engine/Source/Runtime/OpenGLDrv/Public/OpenGLResources.h
+++ b/Engine/Source/Runtime/OpenGLDrv/Public/OpenGLResources.h
@@ -1539,7 +1539,8 @@ class FOpenGLUnorderedAccessView : public FRHIUnorderedAccessView
 {
 
 public:
-	FOpenGLUnorderedAccessView():
+	FOpenGLUnorderedAccessView(FRHIViewableResource* InParentResource):
+		FRHIUnorderedAccessView(InParentResource),
 		Resource(0),
 		BufferResource(0),
 		Format(0),
@@ -1621,7 +1622,6 @@ public:
 class FOpenGLShaderResourceView : public FRHIShaderResourceView
 {
 public:
-	explicit FOpenGLShaderResourceView() = default;
 	explicit FOpenGLShaderResourceView(const FShaderResourceViewInitializer& Initializer);
 	explicit FOpenGLShaderResourceView(FOpenGLTexture* Texture, const FRHITextureSRVCreateInfo& CreateInfo);
 
diff --git a/Engine/Source/Runtime/RHI/Private/RHI.cpp b/Engine/Source/Runtime/RHI/Private/RHI.cpp
--- a/Engine/Source/Runtime/RHI/Private/RHI.cpp
+++ b/Engine/Source/Runtime/RHI/Private/RHI.cpp
@@ -2277,6 +2277,21 @@ static inline uint32 GetSectionUint(const FConfigSection& Section, FName Key, ui
 	}
 }
 
+FRHIViewableResource* GetViewableResource(const FRHITransitionInfo& Info)
+{
+	switch (Info.Type)
+	{
+	case FRHITransitionInfo::EType::Buffer:
+	case FRHITransitionInfo::EType::Texture:
+		return Info.ViewableResource;
+
+	case FRHITransitionInfo::EType::UAV:
+		return Info.UAV ? Info.UAV->GetParentResource() : nullptr;
+	}
+
+	return nullptr;
+}
+
 void FGenericDataDrivenShaderPlatformInfo::SetDefaultValues()
 {
 	MaxFeatureLevel = ERHIFeatureLevel::Num;
diff --git a/Engine/Source/Runtime/RHI/Private/RHICommandList.cpp b/Engine/Source/Runtime/RHI/Private/RHICommandList.cpp
--- a/Engine/Source/Runtime/RHI/Private/RHICommandList.cpp
+++ b/Engine/Source/Runtime/RHI/Private/RHICommandList.cpp
@@ -1049,6 +1049,13 @@ void FRHICommandListExecutor::Transition(TArrayView<const FRHITransitionInfo> In
 
 	FRHIAsyncComputeCommandListImmediate& RHICmdListAsyncCompute = GetImmediateAsyncComputeCommandList();
 
+#if DO_CHECK
+	for (const FRHITransitionInfo& Info : Infos)
+	{
+		checkf(Info.IsWholeResource(), TEXT("Only whole resource transitions are allowed in FRHICommandListExecutor::Transition."));
+	}
+#endif
+
 	if (!GSupportsEfficientAsyncCompute)
 	{
 		checkf(SrcPipelines != ERHIPipeline::AsyncCompute, TEXT("Async compute is disabled. Cannot transition from it."));
@@ -1078,6 +1085,11 @@ void FRHICommandListExecutor::Transition(TArrayView<const FRHITransitionInfo> In
 	{
 		FRHIAsyncComputeCommandListImmediate::ImmediateDispatch(RHICmdListAsyncCompute);
 	}
+
+	if (EnumHasAnyFlags(SrcPipelines | DstPipelines, ERHIPipeline::Graphics))
+	{
+		CommandLists[ERHIPipeline::Graphics]->SetTrackedAccess(Infos);
+	}
 }
 
 FRHICommandListBase::FRHICommandListBase(FRHIGPUMask InGPUMask)
@@ -1598,7 +1610,7 @@ void FRHICommandListBase::QueueParallelAsyncCommandListSubmit(FGraphEventRef* An
 	}
 }
 
-void FRHICommandListBase::QueueAsyncCommandListSubmit(FGraphEventRef& AnyThreadCompletionEvent, class FRHICommandList* CmdList)
+void FRHICommandListBase::QueueAsyncCommandListSubmit(FGraphEventRef& AnyThreadCompletionEvent, FRHICommandList* CmdList)
 {
 	check(IsInRenderingThread() && IsImmediate());
 
@@ -1639,7 +1651,7 @@ FRHICOMMAND_MACRO(FRHICommandWaitForAndSubmitRTSubList)
 	}
 };
 
-void FRHICommandListBase::QueueRenderThreadCommandListSubmit(FGraphEventRef& RenderThreadCompletionEvent, class FRHICommandList* CmdList)
+void FRHICommandListBase::QueueRenderThreadCommandListSubmit(FGraphEventRef& RenderThreadCompletionEvent, FRHICommandList* CmdList)
 {
 	check(IsInRenderingThread() && IsImmediate());
 
@@ -1678,7 +1690,7 @@ FRHICOMMAND_MACRO(FRHICommandSubmitSubList)
 	}
 };
 
-void FRHICommandListBase::QueueCommandListSubmit(class FRHICommandList* CmdList)
+void FRHICommandListBase::QueueCommandListSubmit(FRHICommandList* CmdList)
 {
 	ALLOC_COMMAND(FRHICommandSubmitSubList)(CmdList);
 
@@ -2148,6 +2160,45 @@ void FRHIComputeCommandList::operator delete(void *RawMemory)
 	FMemory::Free(RawMemory);
 }
 
+void FRHIComputeCommandList::Transition(TArrayView<const FRHITransitionInfo> Infos)
+{
+	const ERHIPipeline Pipeline = GetPipeline();
+
+	if (Bypass())
+	{
+		// Stack allocate the transition
+		FMemStack& MemStack = FMemStack::Get();
+		FMemMark Mark(MemStack);
+		FRHITransition* Transition = new (MemStack.Alloc(FRHITransition::GetTotalAllocationSize(), FRHITransition::GetAlignment())) FRHITransition(Pipeline, Pipeline);
+		GDynamicRHI->RHICreateTransition(Transition, FRHITransitionCreateInfo(Pipeline, Pipeline, ERHITransitionCreateFlags::NoSplit, Infos));
+
+		GetComputeContext().RHIBeginTransitions(MakeArrayView((const FRHITransition**)&Transition, 1));
+		GetComputeContext().RHIEndTransitions(MakeArrayView((const FRHITransition**)&Transition, 1));
+
+		// Manual release
+		GDynamicRHI->RHIReleaseTransition(Transition);
+		Transition->~FRHITransition();
+	}
+	else
+	{
+		// Allocate the transition in the command list
+		FRHITransition* Transition = new (Alloc(FRHITransition::GetTotalAllocationSize(), FRHITransition::GetAlignment())) FRHITransition(Pipeline, Pipeline);
+		GDynamicRHI->RHICreateTransition(Transition, FRHITransitionCreateInfo(Pipeline, Pipeline, ERHITransitionCreateFlags::NoSplit, Infos));
+
+		ALLOC_COMMAND(FRHICommandResourceTransition)(Transition);
+	}
+
+	for (const FRHITransitionInfo& Info : Infos)
+	{
+		ensureMsgf(Info.IsWholeResource(), TEXT("The Transition method only supports whole resource transitions."));
+
+		if (FRHIViewableResource* Resource = GetViewableResource(Info))
+		{
+			SetTrackedAccess({ FRHITrackedAccessInfo(Resource, Info.AccessAfter) });
+		}
+	}
+}
+
 #if RHI_RAYTRACING
 void FRHIComputeCommandList::BuildAccelerationStructure(FRHIRayTracingGeometry* Geometry)
 {
@@ -2196,7 +2247,8 @@ FBufferRHIRef FDynamicRHI::CreateBuffer_RenderThread(class FRHICommandListImmedi
 {
 	CSV_SCOPED_TIMING_STAT(RHITStalls, CreateBuffer_RenderThread);
 	FScopedRHIThreadStaller StallRHIThread(RHICmdList);
-	return GDynamicRHI->RHICreateBuffer(Size, Usage, Stride, ResourceState, CreateInfo);
+	FBufferRHIRef Buffer = GDynamicRHI->RHICreateBuffer(Size, Usage, Stride, ResourceState, CreateInfo);
+	return Buffer;
 }
 
 FShaderResourceViewRHIRef FDynamicRHI::CreateShaderResourceView_RenderThread(class FRHICommandListImmediate& RHICmdList, FRHIBuffer* Buffer, uint32 Stride, uint8 Format)
diff --git a/Engine/Source/Runtime/RHI/Private/RHIGPUReadback.cpp b/Engine/Source/Runtime/RHI/Private/RHIGPUReadback.cpp
--- a/Engine/Source/Runtime/RHI/Private/RHIGPUReadback.cpp
+++ b/Engine/Source/Runtime/RHI/Private/RHIGPUReadback.cpp
@@ -146,11 +146,11 @@ void FRHIGPUTextureReadback::EnqueueCopyInternal(FRHICommandList& RHICmdList, FR
 
 			FString FenceName = Fence->GetFName().ToString();
 			FRHIResourceCreateInfo CreateInfo(*FenceName);
-			DestinationStagingTexture = RHICreateTexture2D(TextureSize.X, TextureSize.Y, SourceTexture->GetFormat(), 1, 1, TexCreate_CPUReadback | TexCreate_HideInVisualizeTexture, CreateInfo);
+			DestinationStagingTexture = RHICreateTexture2D(TextureSize.X, TextureSize.Y, SourceTexture->GetFormat(), 1, 1, TexCreate_CPUReadback | TexCreate_HideInVisualizeTexture, ERHIAccess::CPURead, CreateInfo);
 		}
 
 		// We need the destination texture to be writable from a copy operation
-		RHICmdList.Transition(FRHITransitionInfo(DestinationStagingTexture, ERHIAccess::Unknown, ERHIAccess::CopyDest));
+		RHICmdList.Transition(FRHITransitionInfo(DestinationStagingTexture, ERHIAccess::CPURead, ERHIAccess::CopyDest));
 
 		// Ensure this copy call does not perform any transitions. We're handling them manually.
 		ResolveParams.SourceAccessFinal = ERHIAccess::Unknown;
diff --git a/Engine/Source/Runtime/RHI/Private/RHIValidation.cpp b/Engine/Source/Runtime/RHI/Private/RHIValidation.cpp
--- a/Engine/Source/Runtime/RHI/Private/RHIValidation.cpp
+++ b/Engine/Source/Runtime/RHI/Private/RHIValidation.cpp
@@ -310,6 +310,7 @@ void FValidationRHI::RHICreateTransition(FRHITransition* Transition, const FRHIT
 		if (!Info.Resource)
 			continue;
 
+		checkf(Info.AccessAfter != ERHIAccess::Unknown, TEXT("FRHITransitionInfo::AccessAfter cannot be Unknown when creating a resource transition."));
 		checkf(Info.Type != FRHITransitionInfo::EType::Unknown, TEXT("FRHITransitionInfo::Type cannot be Unknown when creating a resource transition."));
 
 		FResourceIdentity Identity;
@@ -468,8 +469,9 @@ void FValidationComputeContext::FState::Reset()
 	GlobalUniformBuffers.Reset();
 }
 
-FValidationContext::FValidationContext()
+FValidationContext::FValidationContext(EType InType)
 	: RHIContext(nullptr)
+	, Type(InType)
 {
 	State.Reset();
 	Tracker = &State.TrackerInstance;
@@ -697,6 +699,24 @@ namespace RHIValidation
 			*GetRHIPipelineName(RequiredState.Pipelines));
 	}
 
+	static inline FString GetReasonString_IncorrectTrackedAccess(
+		FResource* Resource, FSubresourceIndex const& SubresourceIndex,
+		const FState& CurrentState,
+		const FState& TrackedState)
+	{
+		FString DebugName = GetResourceDebugName(Resource, SubresourceIndex);
+		return FString::Printf(
+			BARRIER_TRACKER_LOG_PREFIX_RESNAME
+			TEXT("Attempted to assign resource %s a tracked access that does not match its validation tracked access.\n\n")
+			TEXT("    --- Actual access states:                    %s\n")
+			TEXT("    --- Assigned access states:                  %s\n")
+			BARRIER_TRACKER_LOG_SUFFIX,
+			*DebugName,
+			*DebugName,
+			*GetRHIAccessName(CurrentState.Access),
+			*GetRHIAccessName(TrackedState.Access));
+	}
+
 	static inline FString GetReasonString_BeginBacktrace(void* CreateTrace, void* BeginTrace)
 	{
 		if (CreateTrace || BeginTrace)
@@ -917,7 +937,7 @@ namespace RHIValidation
 			*GetRHIAccessName(CurrentStateFromRHI.Access));
 	}
 
-	static inline FString GetReasonString_IncorrectPreviousState(
+	static inline FString GetReasonString_IncorrectPreviousExplicitState(
 		FResource* Resource, FSubresourceIndex const& SubresourceIndex,
 		const FState& CurrentState,
 		const FState& CurrentStateFromRHI)
@@ -939,6 +959,29 @@ namespace RHIValidation
 			*GetRHIPipelineName(CurrentStateFromRHI.Pipelines));
 	}
 
+	static inline FString GetReasonString_IncorrectPreviousTrackedState(
+		FResource* Resource, FSubresourceIndex const& SubresourceIndex,
+		const FState& CurrentState,
+		ERHIPipeline PipelineFromRHI)
+	{
+		FString DebugName = GetResourceDebugName(Resource, SubresourceIndex);
+		return FString::Printf(
+			BARRIER_TRACKER_LOG_PREFIX_RESNAME
+			TEXT("The tracked previous state \"%s\" does not match the tracked current state \"%s\" for the resource %s.\n")
+			TEXT("    --- Allowed pipelines for this resource are:                           %s\n")
+			TEXT("    --- Previous pipelines passed as part of the resource transition were: %s\n\n")
+			TEXT("    --- The previous state was pulled from the last call to RHICmdList.SetTrackedAccess due to the use of ERHIAccess::Unknown. If this doesn't match the expected state, be sure to update the \n")
+			TEXT("    --- tracked state after using manual low - level transitions. It is highly recommended to coalesce all subresources into the same state before relying on tracked previous states with \n")
+			TEXT("    --- ERHIAccess::Unknown. RHICmdList.SetTrackedAccess applies to whole resources.\n")
+			BARRIER_TRACKER_LOG_SUFFIX,
+			*DebugName,
+			*GetRHIAccessName(Resource->GetTrackedAccess()),
+			*GetRHIAccessName(CurrentState.Access),
+			*DebugName,
+			*GetRHIPipelineName(CurrentState.Pipelines),
+			*GetRHIPipelineName(PipelineFromRHI));
+	}
+
 	static inline FString GetReasonString_MismatchedEndTransition(
 		FResource* Resource, FSubresourceIndex const& SubresourceIndex,
 		const FState& TargetState,
@@ -1162,9 +1205,17 @@ namespace RHIValidation
 			// Check for the correct pipeline
 			RHI_VALIDATION_CHECK(EnumHasAllFlags(CurrentStateFromRHI.Pipelines, ExecutingPipeline), *GetReasonString_WrongPipeline(Resource, SubresourceIndex, State.Current, TargetState));
 
-			// Check the current RHI state passed in matches the tracked state for the resource, or is "unknown".
-			RHI_VALIDATION_CHECK(CurrentStateFromRHI.Access == ERHIAccess::Unknown || (CurrentStateFromRHI.Access == State.Previous.Access && CurrentStateFromRHI.Pipelines == State.Previous.Pipelines),
-				*GetReasonString_IncorrectPreviousState(Resource, SubresourceIndex, State.Previous, CurrentStateFromRHI));
+			if (CurrentStateFromRHI.Access == ERHIAccess::Unknown)
+			{
+				RHI_VALIDATION_CHECK(Resource->TrackedAccess == State.Previous.Access && CurrentStateFromRHI.Pipelines == State.Previous.Pipelines,
+					*GetReasonString_IncorrectPreviousTrackedState(Resource, SubresourceIndex, State.Previous, CurrentStateFromRHI.Pipelines));
+			}
+			else
+			{
+				// Check the current RHI state passed in matches the tracked state for the resource.
+				RHI_VALIDATION_CHECK(CurrentStateFromRHI.Access == State.Previous.Access && CurrentStateFromRHI.Pipelines == State.Previous.Pipelines,
+					*GetReasonString_IncorrectPreviousExplicitState(Resource, SubresourceIndex, State.Previous, CurrentStateFromRHI));
+			}
 		}
 
 		// Check for unnecessary transitions
@@ -1267,6 +1318,26 @@ namespace RHIValidation
 		State.Current.Access = DecayResourceAccess(State.Current.Access, RequiredState.Access, bAllowAllUAVsOverlap || State.bExplicitAllowUAVOverlap);
 	}
 
+	void FSubresourceState::AssertTracked(FResource* Resource, FSubresourceIndex const& SubresourceIndex, const FState& RequiredState)
+	{
+		if (Resource->LoggingMode != ELoggingMode::None
+#if LOG_UNNAMED_RESOURCES
+			|| Resource->GetDebugName() == nullptr
+#endif
+			)
+		{
+			Log(Resource, SubresourceIndex, nullptr, nullptr, TEXT("AssertTracked"), *FString::Printf(TEXT("Access: %s, Pipeline %s"), *GetRHIAccessName(RequiredState.Access), *GetRHIPipelineName(RequiredState.Pipelines)));
+		}
+
+		FPipelineState& State = States[RequiredState.Pipelines];
+
+		// Check we're not trying to access the resource whilst a pending resource transition is in progress.
+		RHI_VALIDATION_CHECK(!State.bTransitioning, *GetReasonString_AccessDuringTransition(Resource, SubresourceIndex, State.Current, RequiredState, State.CreateTransitionBacktrace, State.BeginTransitionBacktrace));
+
+		// Ensure the resource is in the required state for this operation
+		RHI_VALIDATION_CHECK(State.Current.Access == RequiredState.Access, *GetReasonString_IncorrectTrackedAccess(Resource, SubresourceIndex, State.Current, RequiredState));
+	}
+
 	void FSubresourceState::SpecificUAVOverlap(FResource* Resource, FSubresourceIndex const& SubresourceIndex, ERHIPipeline Pipeline, bool bAllow)
 	{
 		if (Resource->LoggingMode != ELoggingMode::None
@@ -1375,6 +1446,18 @@ namespace RHIValidation
 			Data_AliasingOverlap.ResourceAfter->ReleaseOpRef();
 			break;
 
+		case EOpType::SetTrackedAccess:
+			Data_Assert.Identity.Resource->EnumerateSubresources(Data_SetTrackedAccess.Resource->GetWholeResourceRange(), [this, Pipeline](FSubresourceState& State, FSubresourceIndex const& SubresourceIndex)
+			{
+				State.AssertTracked(
+					Data_SetTrackedAccess.Resource,
+					SubresourceIndex,
+					FState(Data_SetTrackedAccess.Access, Pipeline));
+			});
+			Data_SetTrackedAccess.Resource->TrackedAccess = Data_SetTrackedAccess.Access;
+			Data_SetTrackedAccess.Resource->ReleaseOpRef();
+			break;
+
 		case EOpType::AcquireTransient:
 			Data_AcquireTransient.Resource->TransientState.Acquire(Data_AcquireTransient.Resource, Data_AcquireTransient.CreateBacktrace);
 			Data_AcquireTransient.Resource->ReleaseOpRef();
diff --git a/Engine/Source/Runtime/RHI/Public/RHI.h b/Engine/Source/Runtime/RHI/Public/RHI.h
--- a/Engine/Source/Runtime/RHI/Public/RHI.h
+++ b/Engine/Source/Runtime/RHI/Public/RHI.h
@@ -1972,18 +1972,6 @@ inline constexpr bool IsValidAccess(ERHIAccess Access)
 	return !IsInvalidAccess(Access);
 }
 
-enum class ERHITransitionAccessMode
-{
-	// All upcoming rhi transitions can still have Unkown before state - rhi might need to do internal tracking if needed
-	AllowUnknown = 0,
-
-	// All upcoming rhi transitions are fully defined - no internal RHI tracking required anymore
-	DisallowUnknown,
-
-	Default = AllowUnknown
-};
-ENUM_CLASS_FLAGS(ERHITransitionAccessMode);
-
 enum class ERHITransitionCreateFlags
 {
 	None = 0,
@@ -1993,7 +1981,9 @@ enum class ERHITransitionCreateFlags
 
 	// Indicates the transition will have no useful work between the Begin/End calls,
 	// so should use a partial flush rather than a fence as this is more optimal.
-	NoSplit = 1 << 1
+	NoSplit = 1 << 1,
+
+	BeginSimpleMode
 };
 ENUM_CLASS_FLAGS(ERHITransitionCreateFlags);
 
@@ -2087,6 +2077,7 @@ struct FRHITransitionInfo : public FRHISubresourceRange
 	union
 	{
 		class FRHIResource* Resource = nullptr;
+		class FRHIViewableResource* ViewableResource;
 		class FRHITexture* Texture;
 		class FRHIBuffer* Buffer;
 		class FRHIUnorderedAccessView* UAV;
@@ -2148,6 +2139,24 @@ struct FRHITransitionInfo : public FRHISubresourceRange
 		, Flags(InFlags)
 	{}
 
+	FRHITransitionInfo(class FRHITexture* InTexture, ERHIAccess InNewState)
+		: Texture(InTexture)
+		, Type(EType::Texture)
+		, AccessAfter(InNewState)
+	{}
+
+	FRHITransitionInfo(class FRHIUnorderedAccessView* InUAV, ERHIAccess InNewState)
+		: UAV(InUAV)
+		, Type(EType::UAV)
+		, AccessAfter(InNewState)
+	{}
+
+	FRHITransitionInfo(class FRHIBuffer* InRHIBuffer, ERHIAccess InNewState)
+		: Buffer(InRHIBuffer)
+		, Type(EType::Buffer)
+		, AccessAfter(InNewState)
+	{}
+
 	inline bool operator == (FRHITransitionInfo const& RHS) const
 	{
 		return Resource == RHS.Resource
@@ -2164,6 +2173,8 @@ struct FRHITransitionInfo : public FRHISubresourceRange
 	}
 };
 
+RHI_API FRHIViewableResource* GetViewableResource(const FRHITransitionInfo& Info);
+
 struct FRHITransientAliasingOverlap
 {
 	union
@@ -2338,6 +2349,19 @@ struct FRHITransitionCreateInfo
 	TArrayView<const FRHITransientAliasingInfo> AliasingInfos;
 };
 
+struct FRHITrackedAccessInfo
+{
+	FRHITrackedAccessInfo() = default;
+
+	FRHITrackedAccessInfo(FRHIViewableResource* InResource, ERHIAccess InAccess)
+		: Resource(InResource)
+		, Access(InAccess)
+	{}
+
+	FRHIViewableResource* Resource = nullptr;
+	ERHIAccess Access = ERHIAccess::Unknown;
+};
+
 #include "RHIValidationCommon.h"
 
 // Opaque data structure used to represent a pending resource transition in the RHI.
diff --git a/Engine/Source/Runtime/RHI/Public/RHICommandList.h b/Engine/Source/Runtime/RHI/Public/RHICommandList.h
--- a/Engine/Source/Runtime/RHI/Public/RHICommandList.h
+++ b/Engine/Source/Runtime/RHI/Public/RHICommandList.h
@@ -448,6 +448,19 @@ public:
 		return TArrayView<T>((T*) NewArray, InArray.Num());
 	}
 
+	template <typename T>
+	FORCEINLINE_DEBUGGABLE const TArrayView<T> AllocArrayIfDeferred(const TArrayView<T> InArray)
+	{
+		if (Bypass())
+		{
+			return InArray;
+		}
+		else
+		{
+			return AllocArray(InArray);
+		}
+	}
+
 	FORCEINLINE_DEBUGGABLE TCHAR* AllocString(const TCHAR* Name)
 	{
 		int32 Len = FCString::Strlen(Name) + 1;
@@ -1585,6 +1598,18 @@ FRHICOMMAND_MACRO(FRHICommandResourceTransition)
 	RHI_API void Execute(FRHICommandListBase& CmdList);
 };
 
+FRHICOMMAND_MACRO(FRHICommandSetTrackedAccess)
+{
+	TArrayView<const FRHITrackedAccessInfo> Infos;
+
+	FRHICommandSetTrackedAccess(TArrayView<const FRHITrackedAccessInfo> InInfos)
+		: Infos(InInfos)
+	{
+	}
+
+	RHI_API void Execute(FRHICommandListBase& CmdList);
+};
+
 struct FRHICommandSetAsyncComputeBudgetString
 {
 	static const TCHAR* TStr() { return TEXT("FRHICommandSetAsyncComputeBudget"); }
@@ -2365,7 +2390,7 @@ public:
 
 	/** Custom new/delete with recycling */
 	void* operator new(size_t Size);
-	void operator delete(void *RawMemory);
+	void operator delete(void* RawMemory);
 
 	template <typename LAMBDA>
 	FORCEINLINE_DEBUGGABLE void EnqueueLambda(LAMBDA&& Lambda)
@@ -2383,7 +2408,7 @@ public:
 	inline FRHIComputeShader* GetBoundComputeShader() const { return BoundComputeShaderRHI; }
 
 	UE_DEPRECATED(5.0, "Please rename to SetStaticUniformBuffers")
-	FORCEINLINE_DEBUGGABLE void SetGlobalUniformBuffers(const FUniformBufferStaticBindings& UniformBuffers)
+		FORCEINLINE_DEBUGGABLE void SetGlobalUniformBuffers(const FUniformBufferStaticBindings& UniformBuffers)
 	{
 		SetStaticUniformBuffers(UniformBuffers);
 	}
@@ -2431,7 +2456,7 @@ public:
 	}
 
 	UE_DEPRECATED(5.0, "Use Layout pointers instead")
-	FORCEINLINE_DEBUGGABLE FLocalUniformBuffer BuildLocalUniformBuffer(const void* Contents, uint32 ContentsSize, const FRHIUniformBufferLayout& Layout)
+		FORCEINLINE_DEBUGGABLE FLocalUniformBuffer BuildLocalUniformBuffer(const void* Contents, uint32 ContentsSize, const FRHIUniformBufferLayout& Layout)
 	{
 		return BuildLocalUniformBuffer(Contents, ContentsSize, &Layout);
 	}
@@ -2534,16 +2559,16 @@ public:
 	}
 
 	UE_DEPRECATED(5.1, "ComputePipelineStates should be used instead of direct ComputeShaders. You can use SetComputePipelineState(RHICmdList, ComputeShader).")
-	FORCEINLINE_DEBUGGABLE void SetComputeShader(FRHIComputeShader* ComputeShader)
+		FORCEINLINE_DEBUGGABLE void SetComputeShader(FRHIComputeShader* ComputeShader)
 	{
 		BoundComputeShaderRHI = ComputeShader;
 		ComputeShader->UpdateStats();
 		if (Bypass())
 		{
-PRAGMA_DISABLE_DEPRECATION_WARNINGS
-			GetComputeContext().RHISetComputeShader(ComputeShader);
-PRAGMA_ENABLE_DEPRECATION_WARNINGS
-			return;
+			PRAGMA_DISABLE_DEPRECATION_WARNINGS
+				GetComputeContext().RHISetComputeShader(ComputeShader);
+			PRAGMA_ENABLE_DEPRECATION_WARNINGS
+				return;
 		}
 		ALLOC_COMMAND(FRHICommandSetComputeShader)(ComputeShader);
 	}
@@ -2553,7 +2578,7 @@ PRAGMA_ENABLE_DEPRECATION_WARNINGS
 		BoundComputeShaderRHI = ComputeShader;
 		if (Bypass())
 		{
-			extern RHI_API FRHIComputePipelineState* ExecuteSetComputePipelineState(FComputePipelineState* ComputePipelineState);
+			extern RHI_API FRHIComputePipelineState* ExecuteSetComputePipelineState(FComputePipelineState * ComputePipelineState);
 			FRHIComputePipelineState* RHIComputePipelineState = ExecuteSetComputePipelineState(ComputePipelineState);
 			GetComputeContext().RHISetComputePipelineState(RHIComputePipelineState);
 			return;
@@ -2653,34 +2678,7 @@ PRAGMA_ENABLE_DEPRECATION_WARNINGS
 		}
 	}
 
-	inline void Transition(TArrayView<const FRHITransitionInfo> Infos)
-	{
-		ERHIPipeline Pipeline = GetPipeline();
-
-		if (Bypass())
-		{
-			// Stack allocate the transition
-			FMemStack& MemStack = FMemStack::Get();
-			FMemMark Mark(MemStack);
-			FRHITransition* Transition = new (MemStack.Alloc(FRHITransition::GetTotalAllocationSize(), FRHITransition::GetAlignment())) FRHITransition(Pipeline, Pipeline);
-			GDynamicRHI->RHICreateTransition(Transition, FRHITransitionCreateInfo(Pipeline, Pipeline, ERHITransitionCreateFlags::NoSplit, Infos));
-
-			GetComputeContext().RHIBeginTransitions(MakeArrayView((const FRHITransition**)&Transition, 1));
-			GetComputeContext().RHIEndTransitions(MakeArrayView((const FRHITransition**)&Transition, 1));
-
-			// Manual release
-			GDynamicRHI->RHIReleaseTransition(Transition);
-			Transition->~FRHITransition();
-		}
-		else
-		{
-			// Allocate the transition in the command list
-			FRHITransition* Transition = new (Alloc(FRHITransition::GetTotalAllocationSize(), FRHITransition::GetAlignment())) FRHITransition(Pipeline, Pipeline);
-			GDynamicRHI->RHICreateTransition(Transition, FRHITransitionCreateInfo(Pipeline, Pipeline, ERHITransitionCreateFlags::NoSplit, Infos));
-
-			ALLOC_COMMAND(FRHICommandResourceTransition)(Transition);
-		}
-	}
+	void Transition(TArrayView<const FRHITransitionInfo> Infos);
 
 	FORCEINLINE_DEBUGGABLE void BeginTransition(const FRHITransition* Transition)
 	{
@@ -2697,18 +2695,47 @@ PRAGMA_ENABLE_DEPRECATION_WARNINGS
 		Transition(MakeArrayView(&Info, 1));
 	}
 
+	FORCEINLINE_DEBUGGABLE void SetTrackedAccess(TArrayView<const FRHITrackedAccessInfo> Infos)
+	{
+		if (Bypass())
+		{
+			for (const FRHITrackedAccessInfo& Info : Infos)
+			{
+				GetComputeContext().SetTrackedAccess(Info);
+			}
+		}
+		else
+		{
+			ALLOC_COMMAND(FRHICommandSetTrackedAccess)(AllocArray(Infos));
+		}
+	}
+
+	FORCEINLINE_DEBUGGABLE void SetTrackedAccess(TArrayView<const FRHITransitionInfo> Infos)
+	{
+		for (const FRHITransitionInfo& Info : Infos)
+		{
+			if (FRHIViewableResource* Resource = GetViewableResource(Info))
+			{
+				SetTrackedAccess({ FRHITrackedAccessInfo(Resource, Info.AccessAfter) });
+			}
+		}
+	}
+
 	/* LEGACY API */
 
+	UE_DEPRECATED(5.1, "TransitionResource is deprecated. Use Transition instead.")
 	FORCEINLINE_DEBUGGABLE void TransitionResource(ERHIAccess TransitionType, const FTextureRHIRef& InTexture)
 	{
 		Transition(FRHITransitionInfo(InTexture.GetReference(), ERHIAccess::Unknown, TransitionType));
 	}
 
+	UE_DEPRECATED(5.1, "TransitionResource is deprecated. Use Transition instead.")
 	FORCEINLINE_DEBUGGABLE void TransitionResource(ERHIAccess TransitionType, FRHITexture* InTexture)
 	{
 		Transition(FRHITransitionInfo(InTexture, ERHIAccess::Unknown, TransitionType));
 	}
 
+	UE_DEPRECATED(5.1, "TransitionResources is deprecated. Use Transition instead.")
 	inline void TransitionResources(ERHIAccess TransitionType, FRHITexture* const* InTextures, int32 NumTextures)
 	{
 		// Stack allocate the transition descriptors. These will get memcpy()ed onto the RHI command list if required.
@@ -2724,11 +2751,15 @@ PRAGMA_ENABLE_DEPRECATION_WARNINGS
 		Transition(Infos);
 	}
 
+	UE_DEPRECATED(5.1, "TransitionResourceArrayNoCopy is deprecated. Use Transition instead.")
 	FORCEINLINE_DEBUGGABLE void TransitionResourceArrayNoCopy(ERHIAccess TransitionType, TArray<FRHITexture*>& InTextures)
 	{
+	PRAGMA_DISABLE_DEPRECATION_WARNINGS
 		TransitionResources(TransitionType, &InTextures[0], InTextures.Num());
+	PRAGMA_ENABLE_DEPRECATION_WARNINGS
 	}
 
+	UE_DEPRECATED(5.1, "WaitComputeFence is deprecated. Use RHI transitions instead.")
 	FORCEINLINE_DEBUGGABLE void WaitComputeFence(FRHIComputeFence* WaitFence)
 	{
 		check(WaitFence->Transition);
@@ -3519,9 +3550,28 @@ public:
 		if (Bypass())
 		{
 			GetContext().RHICopyToResolveTarget(SourceTextureRHI, DestTextureRHI, ResolveParams);
-			return;
 		}
-		ALLOC_COMMAND(FRHICommandCopyToResolveTarget)(SourceTextureRHI, DestTextureRHI, ResolveParams);
+		else
+		{
+			ALLOC_COMMAND(FRHICommandCopyToResolveTarget)(SourceTextureRHI, DestTextureRHI, ResolveParams);
+		}
+
+		TArray<FRHITransitionInfo, TInlineAllocator<2>> Infos;
+
+		if (ResolveParams.SourceAccessFinal != ERHIAccess::Unknown && SourceTextureRHI != DestTextureRHI)
+		{
+			Infos.Emplace(SourceTextureRHI, ERHIAccess::Unknown, ResolveParams.SourceAccessFinal);
+		}
+
+		if (ResolveParams.DestAccessFinal != ERHIAccess::Unknown)
+		{
+			Infos.Emplace(DestTextureRHI, ERHIAccess::Unknown, ResolveParams.DestAccessFinal);
+		}
+
+		if (!Infos.IsEmpty())
+		{
+			Transition(Infos);
+		}
 	}
 
 	FORCEINLINE_DEBUGGABLE void CopyTexture(FRHITexture* SourceTextureRHI, FRHITexture* DestTextureRHI, const FRHICopyTextureInfo& CopyInfo)
@@ -4159,6 +4209,7 @@ public:
 	FORCEINLINE FBufferRHIRef CreateAndLockIndexBuffer(uint32 Stride, uint32 Size, EBufferUsageFlags InUsage, ERHIAccess InResourceState, FRHIResourceCreateInfo& CreateInfo, void*& OutDataBuffer)
 	{
 		FBufferRHIRef IndexBuffer = GDynamicRHI->CreateBuffer_RenderThread(*this, Size, InUsage | BUF_IndexBuffer, Stride, InResourceState, CreateInfo);
+		IndexBuffer->SetTrackedAccess_Unsafe(InResourceState);
 		OutDataBuffer = GDynamicRHI->RHILockBuffer(*this, IndexBuffer, 0, Size, RLM_WriteOnly);
 		return IndexBuffer;
 	}
@@ -4195,6 +4246,7 @@ public:
 	FORCEINLINE FBufferRHIRef CreateAndLockVertexBuffer(uint32 Size, EBufferUsageFlags InUsage, ERHIAccess InResourceState, FRHIResourceCreateInfo& CreateInfo, void*& OutDataBuffer)
 	{
 		FBufferRHIRef VertexBuffer = GDynamicRHI->CreateBuffer_RenderThread(*this, Size, InUsage | BUF_VertexBuffer, 0, InResourceState, CreateInfo);
+		VertexBuffer->SetTrackedAccess_Unsafe(InResourceState);
 		OutDataBuffer = GDynamicRHI->RHILockBuffer(*this, VertexBuffer, 0, Size, RLM_WriteOnly);
 		return VertexBuffer;
 	}
@@ -5181,17 +5233,23 @@ FORCEINLINE FBufferRHIRef RHICreateAndLockIndexBuffer(uint32 Stride, uint32 Size
 
 FORCEINLINE FBufferRHIRef RHICreateBuffer(uint32 Size, EBufferUsageFlags Usage, uint32 Stride, ERHIAccess ResourceState, FRHIResourceCreateInfo& CreateInfo)
 {
-	return GDynamicRHI->CreateBuffer_RenderThread(FRHICommandListExecutor::GetImmediateCommandList(), Size, Usage, Stride, ResourceState, CreateInfo);
+	FBufferRHIRef Buffer = GDynamicRHI->CreateBuffer_RenderThread(FRHICommandListExecutor::GetImmediateCommandList(), Size, Usage, Stride, ResourceState, CreateInfo);
+	Buffer->SetTrackedAccess_Unsafe(ResourceState);
+	return Buffer;
 }
 
 FORCEINLINE FBufferRHIRef RHICreateIndexBuffer(uint32 Stride, uint32 Size, EBufferUsageFlags InUsage, ERHIAccess InResourceState, FRHIResourceCreateInfo& CreateInfo)
 {
-	return GDynamicRHI->CreateBuffer_RenderThread(FRHICommandListExecutor::GetImmediateCommandList(), Size, InUsage | BUF_IndexBuffer, Stride, InResourceState, CreateInfo);
+	FBufferRHIRef Buffer = GDynamicRHI->CreateBuffer_RenderThread(FRHICommandListExecutor::GetImmediateCommandList(), Size, InUsage | BUF_IndexBuffer, Stride, InResourceState, CreateInfo);
+	Buffer->SetTrackedAccess_Unsafe(InResourceState);
+	return Buffer;
 }
 
 FORCEINLINE FBufferRHIRef RHIAsyncCreateIndexBuffer(uint32 Stride, uint32 Size, EBufferUsageFlags InUsage, ERHIAccess InResourceState, FRHIResourceCreateInfo& CreateInfo)
 {
-	return GDynamicRHI->RHICreateBuffer(Size, InUsage, Stride, InResourceState, CreateInfo);
+	FBufferRHIRef Buffer = GDynamicRHI->RHICreateBuffer(Size, InUsage, Stride, InResourceState, CreateInfo);
+	Buffer->SetTrackedAccess_Unsafe(InResourceState);
+	return Buffer;
 }
 
 FORCEINLINE FBufferRHIRef RHICreateIndexBuffer(uint32 Stride, uint32 Size, EBufferUsageFlags InUsage, FRHIResourceCreateInfo& CreateInfo)
@@ -5230,12 +5288,16 @@ FORCEINLINE FBufferRHIRef RHICreateAndLockVertexBuffer(uint32 Size, EBufferUsage
 
 FORCEINLINE FBufferRHIRef RHICreateVertexBuffer(uint32 Size, EBufferUsageFlags InUsage, ERHIAccess InResourceState, FRHIResourceCreateInfo& CreateInfo)
 {
-	return GDynamicRHI->CreateBuffer_RenderThread(FRHICommandListExecutor::GetImmediateCommandList(), Size, InUsage | BUF_VertexBuffer, 0, InResourceState, CreateInfo);
+	FBufferRHIRef Buffer = GDynamicRHI->CreateBuffer_RenderThread(FRHICommandListExecutor::GetImmediateCommandList(), Size, InUsage | BUF_VertexBuffer, 0, InResourceState, CreateInfo);
+	Buffer->SetTrackedAccess_Unsafe(InResourceState);
+	return Buffer;
 }
 
 FORCEINLINE FBufferRHIRef RHIAsyncCreateVertexBuffer(uint32 Size, EBufferUsageFlags InUsage, ERHIAccess InResourceState, FRHIResourceCreateInfo& CreateInfo)
 {
-	return GDynamicRHI->RHICreateBuffer(Size, InUsage, 0, InResourceState, CreateInfo);
+	FBufferRHIRef Buffer = GDynamicRHI->RHICreateBuffer(Size, InUsage, 0, InResourceState, CreateInfo);
+	Buffer->SetTrackedAccess_Unsafe(InResourceState);
+	return Buffer;
 }
 
 FORCEINLINE FBufferRHIRef RHICreateVertexBuffer(uint32 Size, EBufferUsageFlags InUsage, FRHIResourceCreateInfo& CreateInfo)
@@ -5268,7 +5330,9 @@ FORCEINLINE void RHIUnlockVertexBuffer(FRHIBuffer* VertexBuffer)
 
 FORCEINLINE FBufferRHIRef RHICreateStructuredBuffer(uint32 Stride, uint32 Size, EBufferUsageFlags InUsage, ERHIAccess InResourceState, FRHIResourceCreateInfo& CreateInfo)
 {
-	return GDynamicRHI->CreateBuffer_RenderThread(FRHICommandListExecutor::GetImmediateCommandList(), Size, InUsage | BUF_StructuredBuffer, Stride, InResourceState, CreateInfo);
+	FBufferRHIRef Buffer = GDynamicRHI->CreateBuffer_RenderThread(FRHICommandListExecutor::GetImmediateCommandList(), Size, InUsage | BUF_StructuredBuffer, Stride, InResourceState, CreateInfo);
+	Buffer->SetTrackedAccess_Unsafe(InResourceState);
+	return Buffer;
 }
 
 FORCEINLINE FBufferRHIRef RHICreateStructuredBuffer(uint32 Stride, uint32 Size, EBufferUsageFlags InUsage, FRHIResourceCreateInfo& CreateInfo)
diff --git a/Engine/Source/Runtime/RHI/Public/RHICommandListCommandExecutes.inl b/Engine/Source/Runtime/RHI/Public/RHICommandListCommandExecutes.inl
--- a/Engine/Source/Runtime/RHI/Public/RHICommandListCommandExecutes.inl
+++ b/Engine/Source/Runtime/RHI/Public/RHICommandListCommandExecutes.inl
@@ -435,6 +435,15 @@ void FRHICommandResourceTransition::Execute(FRHICommandListBase& CmdList)
 	Transition->~FRHITransition();
 }
 
+void FRHICommandSetTrackedAccess::Execute(FRHICommandListBase& CmdList)
+{
+	RHISTAT(SetTrackedAccess);
+	for (const FRHITrackedAccessInfo& Info : Infos)
+	{
+		INTERNAL_DECORATOR_COMPUTE(SetTrackedAccess)(Info);
+	}
+}
+
 void FRHICommandSetAsyncComputeBudget::Execute(FRHICommandListBase& CmdList)
 {
 	RHISTAT(SetAsyncComputeBudget);
diff --git a/Engine/Source/Runtime/RHI/Public/RHIContext.h b/Engine/Source/Runtime/RHI/Public/RHIContext.h
--- a/Engine/Source/Runtime/RHI/Public/RHIContext.h
+++ b/Engine/Source/Runtime/RHI/Public/RHIContext.h
@@ -389,6 +389,23 @@ public:
 
 #endif
 
+#if ENABLE_RHI_VALIDATION
+	virtual void SetTrackedAccess(const FRHITrackedAccessInfo& Info)
+#else
+	inline  void SetTrackedAccess(const FRHITrackedAccessInfo& Info)
+#endif
+	{
+		check(Info.Resource != nullptr);
+		check(Info.Access != ERHIAccess::Unknown);
+		Info.Resource->TrackedAccess = Info.Access;
+	}
+
+	inline ERHIAccess GetTrackedAccess(const FRHIViewableResource* Resource) const
+	{
+		check(Resource);
+		return Resource->TrackedAccess;
+	}
+
 	virtual void* RHIGetNativeCommandBuffer() { return nullptr; }
 	virtual void RHIPostExternalCommandsReset() { }
 };
diff --git a/Engine/Source/Runtime/RHI/Public/RHIResources.h b/Engine/Source/Runtime/RHI/Public/RHIResources.h
--- a/Engine/Source/Runtime/RHI/Public/RHIResources.h
+++ b/Engine/Source/Runtime/RHI/Public/RHIResources.h
@@ -39,6 +39,7 @@
 
 #define RHI_ENABLE_RESOURCE_INFO (RHI_WANT_RESOURCE_INFO && !RHI_FORCE_DISABLE_RESOURCE_INFO)
 
+class FRHIComputeCommandList;
 class FRHICommandListImmediate;
 struct FClearValueBinding;
 struct FRHIResourceInfo;
@@ -1067,17 +1068,55 @@ private:
 	uint32 LayoutConstantBufferSize;
 };
 
-class FRHIBuffer : public FRHIResource
+class FRHIViewableResource : public FRHIResource
+{
+public:
+	// TODO (RemoveUnknowns) remove once FRHIBufferCreateDesc contains initial access.
+	void SetTrackedAccess_Unsafe(ERHIAccess Access)
+	{
+		TrackedAccess = Access;
+	}
+
+	FName GetName() const
+	{
+		return Name;
+	}
+
+#if ENABLE_RHI_VALIDATION
+	virtual RHIValidation::FResource* GetTrackerResource() = 0;
+#endif
+
+protected:
+	FRHIViewableResource(ERHIResourceType InResourceType, ERHIAccess InAccess)
+		: FRHIResource(InResourceType)
+		, TrackedAccess(InAccess)
+	{}
+
+	void Swap(FRHIViewableResource& Other)
+	{
+		::Swap(TrackedAccess, Other.TrackedAccess);
+	}
+
+	FName Name;
+
+private:
+	ERHIAccess TrackedAccess;
+
+	friend class FRHIComputeCommandList;
+	friend class IRHIComputeContext;
+};
+
+class FRHIBuffer : public FRHIViewableResource
 #if ENABLE_RHI_VALIDATION
 	, public RHIValidation::FBufferResource
 #endif
 {
 public:
-	FRHIBuffer() : FRHIResource(RRT_Buffer) {}
+	FRHIBuffer() : FRHIViewableResource(RRT_Buffer, ERHIAccess::Unknown) {}
 
 	/** Initialization constructor. */
 	FRHIBuffer(uint32 InSize, EBufferUsageFlags InUsage, uint32 InStride)
-		: FRHIResource(RRT_Buffer)
+		: FRHIViewableResource(RRT_Buffer, ERHIAccess::Unknown /* TODO (RemoveUnknowns): Use InitialAccess from descriptor after refactor. */)
 		, Size(InSize)
 		, Stride(InStride)
 		, Usage(InUsage)
@@ -1093,15 +1132,21 @@ public:
 	/** @return The usage flags used to create the buffer. */
 	EBufferUsageFlags GetUsage() const { return Usage; }
 
-	void SetName(const FName& InName) { BufferName = InName; }
-
-	FName GetName() const { return BufferName; }
+	void SetName(const FName& InName) { Name = InName; }
 
 	virtual uint32 GetParentGPUIndex() const { return 0; }
 
+#if ENABLE_RHI_VALIDATION
+	virtual RHIValidation::FResource* GetTrackerResource() final override
+	{
+		return this;
+	}
+#endif
+
 protected:
 	void Swap(FRHIBuffer& Other)
 	{
+		FRHIViewableResource::Swap(Other);
 		::Swap(Stride, Other.Stride);
 		::Swap(Size, Other.Size);
 		::Swap(Usage, Other.Usage);
@@ -1123,7 +1168,6 @@ private:
 	uint32 Size{};
 	uint32 Stride{};
 	EBufferUsageFlags Usage{};
-	FName BufferName;
 };
 
 UE_DEPRECATED(5.0, "FRHIIndexBuffer is deprecated, please use FRHIBuffer.")      typedef class FRHIBuffer FRHIIndexBuffer;
@@ -1632,7 +1676,7 @@ class FRHITexture;
 /*UE_DEPRECATED(5.1, "FRHITextureCube is deprecated, please use FRHITexture.")    */ typedef class FRHITexture FRHITextureCube;
 
 
-class RHI_API FRHITexture : public FRHIResource
+class RHI_API FRHITexture : public FRHIViewableResource
 #if ENABLE_RHI_VALIDATION
 	, public RHIValidation::FTextureResource
 #endif
@@ -1640,7 +1684,7 @@ class RHI_API FRHITexture : public FRHIResource
 protected:
 	/** Initialization constructor. Should only be called by platform RHI implementations. */
 	FRHITexture(const FRHITextureCreateDesc& InDesc)
-		: FRHIResource(RRT_Texture)
+		: FRHIViewableResource(RRT_Texture, InDesc.InitialState)
 #if ENABLE_RHI_VALIDATION
 		, RHIValidation::FTextureResource(InDesc)
 #endif
@@ -1810,18 +1854,13 @@ public:
 
 	void SetName(const FName& InName)
 	{
-		TextureName = InName;
+		Name = InName;
 
 #if TEXTURE_PROFILER_ENABLED
 		FTextureProfiler::Get()->UpdateTextureName(this);
 #endif
 	}
 
-	FName GetName() const
-	{
-		return TextureName;
-	}
-	
 	///
 	/// Deprecated functions
 	/// 
@@ -1833,7 +1872,7 @@ public:
 	//UE_DEPRECATED(5.1, "FRHITexture3D is deprecated, please use FRHITexture directly")
 	inline FRHITexture3D* GetTexture3D() { return TextureDesc.Dimension == ETextureDimension::Texture3D ? this : nullptr; }
 	//UE_DEPRECATED(5.1, "FRHITextureCube is deprecated, please use FRHITexture directly")
-	inline FRHITextureCube* GetTextureCube() { return TextureDesc.Dimension == ETextureDimension::TextureCube ? this : nullptr; }
+	inline FRHITextureCube* GetTextureCube() { return TextureDesc.IsTextureCube() ? this : nullptr; }
 
 	//UE_DEPRECATED(5.1, "GetSizeX() is deprecated, please use GetDesc().Extent.X instead")
 	uint32 GetSizeX() const { return GetDesc().Extent.X; }
@@ -1865,12 +1904,19 @@ public:
 	//UE_DEPRECATED(5.1, "GetSize() is deprecated, please use GetDesc().Extent.X instead")
 	uint32 GetSize() const { check(GetDesc().IsTextureCube()); return GetDesc().Extent.X; }
 
+#if ENABLE_RHI_VALIDATION
+	virtual RHIValidation::FResource* GetTrackerResource() override
+	{
+		return RHIValidation::FTextureResource::GetTrackerResource();
+	}
+#endif
+
 private:
 
 	friend class FRHITextureReference;
 	/** Constructor for texture references */
 	FRHITexture(ERHIResourceType InResourceType)
-		: FRHIResource(InResourceType)
+		: FRHIViewableResource(InResourceType, ERHIAccess::Unknown)
 	{
 		check(InResourceType == RRT_TextureReference);
 	}
@@ -1878,8 +1924,6 @@ private:
 	FRHITextureDesc TextureDesc;
 
 	FLastRenderTimeContainer LastRenderTime;
-
-	FName TextureName;
 };
 
 class RHI_API FRHITextureReference final : public FRHITexture
@@ -2185,27 +2229,51 @@ public:
 // Views
 //
 
-class FRHIUnorderedAccessView : public FRHIResource
+class FRHIView : public FRHIResource
+{
+public:
+	FRHIView(ERHIResourceType InResourceType, FRHIViewableResource* InParentResource)
+		: FRHIResource(InResourceType)
+		, ParentResource(InParentResource)
+	{}
+
+	virtual FRHIDescriptorHandle GetBindlessHandle() const { return FRHIDescriptorHandle(); }
+
+	FRHIViewableResource* GetParentResource() const { return ParentResource; }
+
+protected:
+	void SetParentResource(FRHIViewableResource* InParentResource)
+	{
+		check(InParentResource);
+		ParentResource = InParentResource;
+	}
+
+private:
+	FRHIViewableResource* ParentResource;
+};
+
+class FRHIUnorderedAccessView : public FRHIView
 #if ENABLE_RHI_VALIDATION
 	, public RHIValidation::FUnorderedAccessView
 #endif
 {
 public:
-	FRHIUnorderedAccessView() : FRHIResource(RRT_UnorderedAccessView) {}
-	virtual FRHIDescriptorHandle GetBindlessHandle() const { return FRHIDescriptorHandle(); }
+	explicit FRHIUnorderedAccessView(FRHIViewableResource* InParentResource)
+		: FRHIView(RRT_UnorderedAccessView, InParentResource)
+	{}
 };
 
-class FRHIShaderResourceView : public FRHIResource 
+class FRHIShaderResourceView : public FRHIView
 #if ENABLE_RHI_VALIDATION
 	, public RHIValidation::FShaderResourceView
 #endif
 {
 public:
-	FRHIShaderResourceView() : FRHIResource(RRT_ShaderResourceView) {}
-	virtual FRHIDescriptorHandle GetBindlessHandle() const { return FRHIDescriptorHandle(); }
+	explicit FRHIShaderResourceView(FRHIViewableResource* InParentResource)
+		: FRHIView(RRT_ShaderResourceView, InParentResource)
+	{}
 };
 
-
 typedef TRefCountPtr<FRHISamplerState> FSamplerStateRHIRef;
 typedef TRefCountPtr<FRHIRasterizerState> FRasterizerStateRHIRef;
 typedef TRefCountPtr<FRHIDepthStencilState> FDepthStencilStateRHIRef;
diff --git a/Engine/Source/Runtime/RHI/Public/RHIValidationCommon.h b/Engine/Source/Runtime/RHI/Public/RHIValidationCommon.h
--- a/Engine/Source/Runtime/RHI/Public/RHIValidationCommon.h
+++ b/Engine/Source/Runtime/RHI/Public/RHIValidationCommon.h
@@ -176,6 +176,7 @@ namespace RHIValidation
 		void BeginTransition   (FResource* Resource, FSubresourceIndex const& SubresourceIndex, const FState& CurrentStateFromRHI, const FState& TargetState, EResourceTransitionFlags NewFlags, ERHIPipeline Pipeline, void* CreateTrace);
 		void EndTransition     (FResource* Resource, FSubresourceIndex const& SubresourceIndex, const FState& CurrentStateFromRHI, const FState& TargetState, ERHIPipeline Pipeline, void* CreateTrace);
 		void Assert            (FResource* Resource, FSubresourceIndex const& SubresourceIndex, const FState& RequiredState, bool bAllowAllUAVsOverlap);
+		void AssertTracked     (FResource* Resource, FSubresourceIndex const& SubresourceIndex, const FState& TrackedState);
 		void SpecificUAVOverlap(FResource* Resource, FSubresourceIndex const& SubresourceIndex, ERHIPipeline Pipeline, bool bAllow);
 	};
 
@@ -278,6 +279,7 @@ namespace RHIValidation
 		uint32 NumArraySlices = 0;
 		uint32 NumPlanes = 0;
 		FTransientState TransientState;
+		ERHIAccess TrackedAccess = ERHIAccess::Unknown;
 
 	private:
 		FString DebugName;
@@ -313,6 +315,33 @@ namespace RHIValidation
 			check(RefCount >= 0);
 		}
 
+		inline ERHIAccess GetTrackedAccess() const
+		{
+			return TrackedAccess;
+		}
+
+		inline FSubresourceRange GetWholeResourceRange()
+		{
+			checkSlow(NumMips > 0 && NumArraySlices > 0 && NumPlanes > 0);
+
+			FSubresourceRange SubresourceRange;
+			SubresourceRange.MipIndex = 0;
+			SubresourceRange.ArraySlice = 0;
+			SubresourceRange.PlaneIndex = 0;
+			SubresourceRange.NumMips = NumMips;
+			SubresourceRange.NumArraySlices = NumArraySlices;
+			SubresourceRange.NumPlanes = NumPlanes;
+			return SubresourceRange;
+		}
+
+		inline FResourceIdentity GetWholeResourceIdentity()
+		{
+			FResourceIdentity Identity;
+			Identity.Resource = this;
+			Identity.SubresourceRange = GetWholeResourceRange();
+			return Identity;
+		}
+
 	protected:
 		inline void InitBarrierTracking(int32 InNumMips, int32 InNumArraySlices, int32 InNumPlanes, ERHIAccess InResourceState, const TCHAR* InDebugName)
 		{
@@ -323,6 +352,7 @@ namespace RHIValidation
 			NumArraySlices = InNumArraySlices;
 			NumPlanes = InNumPlanes;
 			TransientState = FTransientState(InResourceState);
+			TrackedAccess = InResourceState;
 
 			for (auto& State : WholeResourceState.States)
 			{
@@ -354,21 +384,6 @@ namespace RHIValidation
 		{
 			FResource::InitBarrierTracking(1, 1, 1, InResourceState, InDebugName);
 		}
-
-		inline FResourceIdentity GetWholeResourceIdentity()
-		{
-			checkSlow(NumMips == 1 && NumArraySlices == 1 && NumPlanes == 1);
-
-			FResourceIdentity Identity;
-			Identity.Resource = this;
-			Identity.SubresourceRange.MipIndex = 0;
-			Identity.SubresourceRange.ArraySlice = 0;
-			Identity.SubresourceRange.PlaneIndex = 0;
-			Identity.SubresourceRange.NumMips = 1;
-			Identity.SubresourceRange.NumArraySlices = 1;
-			Identity.SubresourceRange.NumPlanes = 1;
-			return Identity;
-		}
 	};
 
 	class FAccelerationStructureResource : public FBufferResource
@@ -389,7 +404,7 @@ namespace RHIValidation
 
 		virtual ~FTextureResource() {}
 
-		virtual FResource* GetTrackerResource() { return &PRIVATE_TrackerResource; }
+		FResource* GetTrackerResource() { return &PRIVATE_TrackerResource; }
 
 		void InitBarrierTracking(FRHITextureCreateDesc const& CreateDesc);
 
@@ -501,19 +516,7 @@ namespace RHIValidation
 
 		inline FResourceIdentity GetWholeResourceIdentity()
 		{
-			FResource* Resource = GetTrackerResource();
-
-			checkSlow(Resource->NumMips > 0 && Resource->NumArraySlices > 0 && Resource->NumPlanes > 0);
-
-			FResourceIdentity Identity;
-			Identity.Resource = Resource;
-			Identity.SubresourceRange.MipIndex = 0;
-			Identity.SubresourceRange.ArraySlice = 0;
-			Identity.SubresourceRange.PlaneIndex = 0;
-			Identity.SubresourceRange.NumMips = Resource->NumMips;
-			Identity.SubresourceRange.NumArraySlices = Resource->NumArraySlices;
-			Identity.SubresourceRange.NumPlanes = Resource->NumPlanes;
-			return Identity;
+			return GetTrackerResource()->GetWholeResourceIdentity();
 		}
 
 		inline FResourceIdentity GetWholeResourceIdentitySRV()
@@ -562,6 +565,7 @@ namespace RHIValidation
 	{
 		BeginTransition,
 		EndTransition,
+		SetTrackedAccess,
 		AliasingOverlap,
 		AcquireTransient,
 		DiscardTransient,
@@ -607,6 +611,12 @@ namespace RHIValidation
 				void* CreateBacktrace;
 			} Data_EndTransition;
 
+			struct
+			{
+				FResource* Resource;
+				ERHIAccess Access;
+			} Data_SetTrackedAccess;
+
 			struct
 			{
 				FResource* ResourceBefore;
@@ -704,6 +714,17 @@ namespace RHIValidation
 			return MoveTemp(Op);
 		}
 
+		static inline FOperation SetTrackedAccess(FResource* Resource, ERHIAccess Access)
+		{
+			Resource->AddOpRef();
+
+			FOperation Op;
+			Op.Type = EOpType::SetTrackedAccess;
+			Op.Data_SetTrackedAccess.Resource = Resource;
+			Op.Data_SetTrackedAccess.Access = Access;
+			return MoveTemp(Op);
+		}
+
 		static inline FOperation AliasingOverlap(FResource* ResourceBefore, FResource* ResourceAfter, void* CreateBacktrace)
 		{
 			ResourceBefore->AddOpRef();
@@ -939,6 +960,11 @@ namespace RHIValidation
 			return MoveTemp(CurrentList);
 		}
 
+		inline void SetTrackedAccess(FResource* Resource, ERHIAccess Access)
+		{
+			AddOp(FOperation::SetTrackedAccess(Resource, Access));
+		}
+
 		inline void Rename(FResource* Resource, const TCHAR* NewName, const TCHAR* Suffix = nullptr)
 		{
 			AddOp(FOperation::Rename(Resource, NewName, Suffix));
diff --git a/Engine/Source/Runtime/RHI/Public/RHIValidationContext.h b/Engine/Source/Runtime/RHI/Public/RHIValidationContext.h
--- a/Engine/Source/Runtime/RHI/Public/RHIValidationContext.h
+++ b/Engine/Source/Runtime/RHI/Public/RHIValidationContext.h
@@ -252,7 +252,17 @@ protected:
 class FValidationContext : public IRHICommandContext
 {
 public:
-	FValidationContext();
+	enum EType
+	{
+		Default,
+		Parallel
+	};
+
+	FValidationContext(EType InType);
+
+	FValidationContext()
+		: FValidationContext(EType::Default)
+	{}
 
 	virtual IRHICommandContext& GetLowestLevelContext() override final
 	{
@@ -352,28 +362,6 @@ public:
 	virtual void RHICopyToResolveTarget(FRHITexture* SourceTexture, FRHITexture* DestTexture, const FResolveParams& ResolveParams) override final
 	{
 		RHIContext->RHICopyToResolveTarget(SourceTexture, DestTexture, ResolveParams);
-
-		auto SetAccess = [&](FRHITexture* Texture, ERHIAccess Access)
-		{
-			// If this is a depth texture, only transition the depth plane.
-			RHIValidation::FResourceIdentity Identity = EnumHasAnyFlags(Texture->GetFlags(), TexCreate_DepthStencilTargetable | TexCreate_DepthStencilResolveTarget)
-				? Texture->GetViewIdentity(0, 0, 0, 0, uint32(RHIValidation::EResourcePlane::Common), 1)
-				: Texture->GetWholeResourceIdentity();
-
-			Tracker->TransitionResource(Identity,
-				RHIValidation::FState(ERHIAccess::Unknown, ERHIPipeline::Graphics),
-				RHIValidation::FState(Access, ERHIPipeline::Graphics),
-				EResourceTransitionFlags::None);
-		};
-
-		if (SourceTexture && ResolveParams.SourceAccessFinal != ERHIAccess::Unknown)
-		{
-			SetAccess(SourceTexture, ResolveParams.SourceAccessFinal);
-		}
-		if (DestTexture && SourceTexture != DestTexture && ResolveParams.DestAccessFinal != ERHIAccess::Unknown)
-		{
-			SetAccess(DestTexture, ResolveParams.DestAccessFinal);
-		}
 	}
 
 	virtual void RHIResummarizeHTile(FRHITexture2D* DepthTexture) override final
@@ -426,6 +414,20 @@ public:
 		RHIContext->RHIEndTransitions(Transitions);
 	}
 
+	virtual void SetTrackedAccess(const FRHITrackedAccessInfo& Info) override final
+	{
+		check(Info.Resource != nullptr);
+		check(Info.Access != ERHIAccess::Unknown);
+		checkf(Type != EType::Parallel,
+			TEXT("SetTrackedAccess(%s, %s) was called from a parallel translate context. This is not allowed. This is most likely a call to RHICmdList.Transition on a command list queued for parallel dispatch."),
+			*Info.Resource->GetName().ToString(),
+			*GetRHIAccessName(Info.Access));
+
+		Tracker->SetTrackedAccess(Info.Resource->GetTrackerResource(), Info.Access);
+
+		RHIContext->SetTrackedAccess(Info);
+	}
+
 	virtual void RHIBeginRenderQuery(FRHIRenderQuery* RenderQuery) override final
 	{
 		RHIContext->RHIBeginRenderQuery(RenderQuery);
@@ -1112,6 +1114,8 @@ protected:
 	};
 	FState State;
 
+	EType Type = EType::Default;
+
 	friend class FValidationRHICommandContextContainer;
 	friend class FValidationRHI;
 
@@ -1158,7 +1162,7 @@ public:
 
 		if (!CurrentContext)
 		{
-			CurrentContext = new FValidationContext();
+			CurrentContext = new FValidationContext(FValidationContext::EType::Parallel);
 		}
 
 		CurrentContext->State.Reset();
diff --git a/Engine/Source/Runtime/RenderCore/Private/RenderGraphBuilder.cpp b/Engine/Source/Runtime/RenderCore/Private/RenderGraphBuilder.cpp
--- a/Engine/Source/Runtime/RenderCore/Private/RenderGraphBuilder.cpp
+++ b/Engine/Source/Runtime/RenderCore/Private/RenderGraphBuilder.cpp
@@ -16,14 +16,14 @@ inline void GatherPassUAVsForOverlapValidation(const FRDGPass* Pass, TArray<FRHI
 	// RHI validation tracking of Begin/EndUAVOverlaps happens on the underlying resource, so we need to be careful about not
 	// passing multiple UAVs that refer to the same resource, otherwise we get double-Begin and double-End validation errors.
 	// Filter UAVs to only those with unique parent resources.
-	TArray<FRDGParentResource*, TInlineAllocator<MaxSimultaneousUAVs, FRDGArrayAllocator>> UniqueParents;
+	TArray<FRDGViewableResource*, TInlineAllocator<MaxSimultaneousUAVs, FRDGArrayAllocator>> UniqueParents;
 	Pass->GetParameters().Enumerate([&](FRDGParameter Parameter)
 	{
 		if (Parameter.IsUAV())
 		{
 			if (FRDGUnorderedAccessViewRef UAV = Parameter.GetAsUAV())
 			{
-				FRDGParentResource* Parent = UAV->GetParent();
+				FRDGViewableResource* Parent = UAV->GetParent();
 
 				// Check if we've already seen this parent.
 				bool bFound = false;
@@ -471,7 +471,7 @@ bool FRDGBuilder::IsTransient(FRDGTextureRef Texture) const
 	return IsTransientInternal(Texture, EnumHasAnyFlags(Texture->Desc.Flags, ETextureCreateFlags::FastVRAM));
 }
 
-bool FRDGBuilder::IsTransientInternal(FRDGParentResourceRef Resource, bool bFastVRAM) const
+bool FRDGBuilder::IsTransientInternal(FRDGViewableResource* Resource, bool bFastVRAM) const
 {
 	// Immediate mode can't use the transient allocator because we don't know if the user will extract the resource.
 	if (!GRDGTransientAllocator || IsImmediateMode())
@@ -499,7 +499,7 @@ bool FRDGBuilder::IsTransientInternal(FRDGParentResourceRef Resource, bool bFast
 				return false;
 			}
 
-			if (GRDGTransientExtractedResources == 1 && Resource->TransientExtractionHint == FRDGParentResource::ETransientExtractionHint::Disable)
+			if (GRDGTransientExtractedResources == 1 && Resource->TransientExtractionHint == FRDGViewableResource::ETransientExtractionHint::Disable)
 			{
 				return false;
 			}
@@ -531,7 +531,6 @@ FRDGBuilder::FRDGBuilder(FRHICommandListImmediate& InRHICmdList, FRDGEventName I
 #if RDG_ENABLE_DEBUG
 	, UserValidation(Allocator, bParallelExecuteEnabled)
 	, BarrierValidation(&Passes, BuilderName)
-	, LogFile(Passes)
 #endif
 	, TransientResourceAllocator(GRDGTransientResourceAllocator.Get())
 {
@@ -548,8 +547,6 @@ FRDGBuilder::FRDGBuilder(FRHICommandListImmediate& InRHICmdList, FRDGEventName I
 		BreadcrumbState = FRDGBreadcrumbState::Create(Allocator);
 	}
 #endif
-
-	IF_RDG_ENABLE_DEBUG(LogFile.Begin(BuilderName));
 }
 
 FRDGBuilder::~FRDGBuilder()
@@ -564,6 +561,8 @@ FRDGBuilder::~FRDGBuilder()
 	ActivePooledBuffers.Empty();
 }
 
+///////////////////////////////////////////////////////////////////////////////////////////////////
+
 const TRefCountPtr<FRDGPooledBuffer>& FRDGBuilder::ConvertToExternalBuffer(FRDGBufferRef Buffer)
 {
 	IF_RDG_ENABLE_DEBUG(UserValidation.ValidateConvertToExternalResource(Buffer));
@@ -571,7 +570,6 @@ const TRefCountPtr<FRDGPooledBuffer>& FRDGBuilder::ConvertToExternalBuffer(FRDGB
 	{
 		Buffer->bExternal = 1;
 		Buffer->bForceNonTransient = 1;
-		Buffer->AccessFinal = kDefaultAccessFinal;
 		BeginResourceRHI(GetProloguePassHandle(), Buffer);
 		ExternalBuffers.Add(Buffer->PooledBuffer, Buffer);
 	}
@@ -585,13 +583,14 @@ const TRefCountPtr<IPooledRenderTarget>& FRDGBuilder::ConvertToExternalTexture(F
 	{
 		Texture->bExternal = 1;
 		Texture->bForceNonTransient = 1;
-		Texture->AccessFinal = kDefaultAccessFinal;
 		BeginResourceRHI(GetProloguePassHandle(), Texture);
 		ExternalTextures.Add(Texture->GetRHIUnchecked(), Texture);
 	}
 	return GetPooledTexture(Texture);
 }
 
+///////////////////////////////////////////////////////////////////////////////////////////////////
+
 BEGIN_SHADER_PARAMETER_STRUCT(FFinalizePassParameters, )
 	RDG_TEXTURE_ACCESS_ARRAY(Textures)
 	RDG_BUFFER_ACCESS_ARRAY(Buffers)
@@ -636,15 +635,17 @@ void FRDGBuilder::FinalizeResourceAccess(FRDGTextureAccessArray&& InTextures, FR
 
 	for (FRDGTextureAccess TextureAccess : LocalTextures)
 	{
-		TextureAccess->bFinalizedAccess = 1;
+		TextureAccess->SetFinalizedAccess(TextureAccess.GetAccess());
 	}
 
 	for (FRDGBufferAccess BufferAccess : LocalBuffers)
 	{
-		BufferAccess->bFinalizedAccess = 1;
+		BufferAccess->SetFinalizedAccess(BufferAccess.GetAccess());
 	}
 }
 
+///////////////////////////////////////////////////////////////////////////////////////////////////
+
 FRDGTextureRef FRDGBuilder::RegisterExternalTexture(
 	const TRefCountPtr<IPooledRenderTarget>& ExternalPooledTexture,
 	ERDGTextureFlags Flags)
@@ -686,37 +687,33 @@ FRDGTextureRef FRDGBuilder::RegisterExternalTexture(
 	}
 
 	FRDGTextureRef Texture = Textures.Allocate(Allocator, Name, Desc, Flags);
-	Texture->SetRHI(ExternalPooledTexture.GetReference());
+	SetRHI(Texture, ExternalPooledTexture.GetReference(), GetProloguePassHandle());
 
 	Texture->bExternal = true;
-	Texture->AccessFinal = EnumHasAnyFlags(ExternalTextureRHI->GetFlags(), TexCreate_Foveation) ? ERHIAccess::ShadingRateSource : kDefaultAccessFinal;
-	Texture->FirstPass = GetProloguePassHandle();
+
+
+	if (EnumHasAnyFlags(ExternalTextureRHI->GetFlags(), TexCreate_Foveation))
+	{
+		Texture->EpilogueAccess = ERHIAccess::ShadingRateSource;
+	}
 
 	// Textures that are created read-only are not transitioned by RDG.
 	if (bFinalizedAccess)
 	{
-		// When in 'finalized access' mode, the access represents the valid set of states to touch the resource for
-		// validation, not its final state after the graph executes. That's why it's okay to have a write state mixed
-		// with read states.
-		Texture->bFinalizedAccess = 1;
-		Texture->AccessFinal = ERHIAccess::ReadOnlyExclusiveMask;
+		ERHIAccess ValidAccessMask = ERHIAccess::ReadOnlyExclusiveMask;
 
 		if (EnumHasAnyFlags(Desc.Flags, TexCreate_CPUReadback))
 		{
-			Texture->AccessFinal |= ERHIAccess::CopyDest;
+			ValidAccessMask |= ERHIAccess::CopyDest;
 		}
 
 		if (EnumHasAnyFlags(Desc.Flags, TexCreate_Foveation))
 		{
-			Texture->AccessFinal |= ERHIAccess::ShadingRateSource;
+			ValidAccessMask |= ERHIAccess::ShadingRateSource;
 		}
-	}
 
-	FRDGTextureSubresourceState& TextureState = Texture->GetState();
-
-	checkf(IsWholeResource(TextureState) && GetWholeResource(TextureState).Access == ERHIAccess::Unknown,
-		TEXT("Externally registered texture '%s' has known RDG state. This means the graph did not sanitize it correctly, or ")
-		TEXT("an IPooledRenderTarget reference was improperly held within a pass."), Texture->Name);
+		Texture->SetFinalizedAccess(ValidAccessMask);
+	}
 
 	ExternalTextures.Add(Texture->GetRHIUnchecked(), Texture);
 
@@ -761,24 +758,16 @@ FRDGBufferRef FRDGBuilder::RegisterExternalBuffer(
 	}
 
 	FRDGBufferRef Buffer = Buffers.Allocate(Allocator, Name, ExternalPooledBuffer->Desc, Flags);
-	Buffer->SetRHI(ExternalPooledBuffer);
+	SetRHI(Buffer, ExternalPooledBuffer, GetProloguePassHandle());
 
 	Buffer->bExternal = true;
-	Buffer->AccessFinal = kDefaultAccessFinal;
-	Buffer->FirstPass = GetProloguePassHandle();
 
 	// Buffers that are created read-only are not transitioned by RDG.
 	if (bFinalizedAccess)
 	{
-		Buffer->bFinalizedAccess = 1;
-		Buffer->AccessFinal = ERHIAccess::ReadOnlyExclusiveMask;
+		Buffer->SetFinalizedAccess(ERHIAccess::ReadOnlyExclusiveMask);
 	}
 
-	FRDGSubresourceState& BufferState = Buffer->GetState();
-	checkf(BufferState.Access == ERHIAccess::Unknown,
-		TEXT("Externally registered buffer '%s' has known RDG state. This means the graph did not sanitize it correctly, or ")
-		TEXT("an FRDGPooledBuffer reference was improperly held within a pass."), Buffer->Name);
-
 	ExternalBuffers.Add(ExternalPooledBuffer, Buffer);
 
 	IF_RDG_ENABLE_DEBUG(UserValidation.ValidateRegisterExternalBuffer(Buffer));
@@ -786,6 +775,8 @@ FRDGBufferRef FRDGBuilder::RegisterExternalBuffer(
 	return Buffer;
 }
 
+///////////////////////////////////////////////////////////////////////////////////////////////////
+
 void FRDGBuilder::AddPassDependency(FRDGPassHandle ProducerHandle, FRDGPassHandle ConsumerHandle)
 {
 	FRDGPass* Consumer = Passes[ConsumerHandle];
@@ -807,6 +798,7 @@ void FRDGBuilder::Compile()
 	CSV_SCOPED_TIMING_STAT_EXCLUSIVE_CONDITIONAL(RDG_Compile, GRDGVerboseCSVStats != 0);
 	SCOPED_NAMED_EVENT(Compile, FColor::Emerald);
 
+	const FRDGPassHandle ProloguePassHandle = GetProloguePassHandle();
 	const FRDGPassHandle EpiloguePassHandle = GetEpiloguePassHandle();
 
 	const uint32 CompilePassCount = Passes.Num();
@@ -919,7 +911,7 @@ void FRDGBuilder::Compile()
 			for (auto& LastProducer : Texture->LastProducers)
 			{
 				FRDGProducerState StateFinal;
-				StateFinal.Access = Texture->AccessFinal;
+				StateFinal.Access = Texture->EpilogueAccess;
 				StateFinal.PassHandle = EpiloguePassHandle;
 
 				AddCullingDependency(LastProducer, StateFinal, ERHIPipeline::Graphics);
@@ -931,7 +923,7 @@ void FRDGBuilder::Compile()
 			FRDGBufferRef Buffer = ExtractedBuffer.Buffer;
 
 			FRDGProducerState StateFinal;
-			StateFinal.Access = Buffer->AccessFinal;
+			StateFinal.Access = Buffer->EpilogueAccess;
 			StateFinal.PassHandle = EpiloguePassHandle;
 
 			AddCullingDependency(Buffer->LastProducer, StateFinal, ERHIPipeline::Graphics);
@@ -988,7 +980,7 @@ void FRDGBuilder::Compile()
 
 			const ERHIPipeline PassPipeline = Pass->Pipeline;
 
-			const auto MergeSubresourceStates = [&](ERDGParentResourceType ResourceType, FRDGSubresourceState*& PassMergeState, FRDGSubresourceState*& ResourceMergeState, const FRDGSubresourceState& PassState)
+			const auto MergeSubresourceStates = [&](ERDGViewableResourceType ResourceType, FRDGSubresourceState*& PassMergeState, FRDGSubresourceState*& ResourceMergeState, const FRDGSubresourceState& PassState)
 			{
 				if (!ResourceMergeState || !FRDGSubresourceState::IsMergeAllowed(ResourceType, *ResourceMergeState, PassState))
 				{
@@ -1051,7 +1043,7 @@ void FRDGBuilder::Compile()
 						continue;
 					}
 
-					MergeSubresourceStates(ERDGParentResourceType::Texture, PassState.MergeState[Index], Texture->MergeState[Index], PassState.State[Index]);
+					MergeSubresourceStates(ERDGViewableResourceType::Texture, PassState.MergeState[Index], Texture->MergeState[Index], PassState.State[Index]);
 				}
 			}
 
@@ -1066,7 +1058,7 @@ void FRDGBuilder::Compile()
 				GRDGStatBufferReferenceCount += PassState.ReferenceCount;
 			#endif
 
-				MergeSubresourceStates(ERDGParentResourceType::Buffer, PassState.MergeState, Buffer->MergeState, PassState.State);
+				MergeSubresourceStates(ERDGViewableResourceType::Buffer, PassState.MergeState, Buffer->MergeState, PassState.State);
 			}
 		}
 	}
@@ -1432,6 +1424,7 @@ void FRDGBuilder::Execute()
 		bInDebugPassScope = false;
 	}
 
+	const FRDGPassHandle ProloguePassHandle = GetProloguePassHandle();
 	const FRDGPassHandle EpiloguePassHandle = GetEpiloguePassHandle();
 
 	FGraphEventArray AsyncCompileEvents;
@@ -1480,6 +1473,30 @@ void FRDGBuilder::Execute()
 				++Buffer->ReferenceCount;
 			});
 
+			// Null out any culled external resources so that the reference is freed up.
+	
+			for (const auto& Pair : ExternalTextures)
+			{
+				FRDGTexture* Texture = Pair.Value;
+
+				if (Texture->bCulled)
+				{
+					check(Texture->ReferenceCount == 0);
+					EndResourceRHI(ProloguePassHandle, Texture, 0);
+				}
+			}
+
+			for (const auto& Pair : ExternalBuffers)
+			{
+				FRDGBuffer* Buffer = Pair.Value;
+
+				if (Buffer->bCulled)
+				{
+					check(Buffer->ReferenceCount == 0);
+					EndResourceRHI(ProloguePassHandle, Buffer, 0);
+				}
+			}
+
 			for (FRDGPassHandle PassHandle = Passes.Begin(); PassHandle < ProloguePassHandle; ++PassHandle)
 			{
 				FRDGPass* Pass = Passes[PassHandle];
@@ -1541,43 +1558,16 @@ void FRDGBuilder::Execute()
 	{
 		SCOPED_NAMED_EVENT_TEXT("FRDGBuilder::Finalize", FColor::Magenta);
 
-#if RDG_ENABLE_DEBUG
-		const auto LogResource = [&](auto* Resource, auto& Registry)
-		{
-			if (!Resource->bCulled)
-			{
-				if (!Resource->bLastOwner)
-				{
-					auto* NextOwner = Registry[Resource->NextOwner];
-					LogFile.AddAliasEdge(Resource, Resource->LastPass, NextOwner, NextOwner->FirstPass);
-				}
-				LogFile.AddFirstEdge(Resource, Resource->FirstPass);
-			}
-		};
-#endif
+		EpilogueResourceAccesses.Reserve(Textures.Num() + Buffers.Num());
 
-		ActivePooledTextures.Reserve(Textures.Num());
 		Textures.Enumerate([&](FRDGTextureRef Texture)
 		{
-			if (Texture->HasRHI())
-			{
-				AddEpilogueTransition(Texture);
-				Texture->Finalize(ActivePooledTextures);
-
-				IF_RDG_ENABLE_DEBUG(LogResource(Texture, Textures));
-			}
+			AddEpilogueTransition(Texture);
 		});
 
-		ActivePooledBuffers.Reserve(Buffers.Num());
 		Buffers.Enumerate([&](FRDGBufferRef Buffer)
 		{
-			if (Buffer->HasRHI())
-			{
-				AddEpilogueTransition(Buffer);
-				Buffer->Finalize(ActivePooledBuffers);
-
-				IF_RDG_ENABLE_DEBUG(LogResource(Buffer, Buffers));
-			}
+			AddEpilogueTransition(Buffer);
 		});
 	}
 
@@ -1719,6 +1709,8 @@ void FRDGBuilder::Execute()
 	}
 #endif
 
+	RHICmdList.SetTrackedAccess(EpilogueResourceAccesses);
+
 	// Wait on the actual parallel execute tasks in the Execute call. When draining is okay to let them overlap with other graph setup.
 	// This also needs to be done before extraction of external resources to be consistent with non-parallel rendering.
 	if (!ParallelExecuteEvents.IsEmpty())
@@ -1728,8 +1720,8 @@ void FRDGBuilder::Execute()
 
 	for (const FExtractedTexture& ExtractedTexture : ExtractedTextures)
 	{
-		check(ExtractedTexture.Texture->PooledRenderTarget);
-		*ExtractedTexture.PooledTexture = ExtractedTexture.Texture->PooledRenderTarget;
+		check(ExtractedTexture.Texture->RenderTarget);
+		*ExtractedTexture.PooledTexture = ExtractedTexture.Texture->RenderTarget;
 	}
 
 	for (const FExtractedBuffer& ExtractedBuffer : ExtractedBuffers)
@@ -1745,7 +1737,6 @@ void FRDGBuilder::Execute()
 	IF_RDG_CPU_SCOPES(CPUScopeStacks.EndExecute());
 
 	IF_RDG_ENABLE_DEBUG(UserValidation.ValidateExecuteEnd());
-	IF_RDG_ENABLE_DEBUG(LogFile.End());
 	IF_RDG_ENABLE_DEBUG(GRDGAllowRHIAccess = false);
 
 #if STATS
@@ -1765,6 +1756,8 @@ void FRDGBuilder::Execute()
 	}
 }
 
+///////////////////////////////////////////////////////////////////////////////////////////////////
+
 FRDGPass* FRDGBuilder::SetupPass(FRDGPass* Pass)
 {
 	IF_RDG_ENABLE_DEBUG(UserValidation.ValidateAddPass(Pass, bInDebugPassScope));
@@ -2032,6 +2025,8 @@ void FRDGBuilder::SetupPassInternal(FRDGPass* Pass, FRDGPassHandle PassHandle, E
 	#endif
 }
 
+///////////////////////////////////////////////////////////////////////////////////////////////////
+
 void FRDGBuilder::SetupBufferUploads()
 {
 	SCOPED_NAMED_EVENT_TEXT("FRDGBuilder::PrepareBufferUploads", FColor::Magenta);
@@ -2092,6 +2087,8 @@ void FRDGBuilder::SubmitBufferUploads()
 	UploadedBuffers.Reset();
 }
 
+///////////////////////////////////////////////////////////////////////////////////////////////////
+
 void FRDGBuilder::SetupParallelExecute()
 {
 	SCOPED_NAMED_EVENT(SetupParallelExecute, FColor::Emerald);
@@ -2291,6 +2288,8 @@ void FRDGBuilder::DispatchParallelExecute(IRHICommandContext* RHICmdContext)
 	}
 }
 
+///////////////////////////////////////////////////////////////////////////////////////////////////
+
 void FRDGBuilder::CreateUniformBuffers()
 {
 	SCOPED_NAMED_EVENT_TEXT("FRDGBuilder::CreateUniformBuffers", FColor::Magenta);
@@ -2302,14 +2301,17 @@ void FRDGBuilder::CreateUniformBuffers()
 	UniformBuffersToCreate.Reset();
 }
 
+///////////////////////////////////////////////////////////////////////////////////////////////////
+
 void FRDGBuilder::AddProloguePass()
 {
 	bInDebugPassScope = true;
 	ProloguePass = SetupEmptyPass(Passes.Allocate<FRDGSentinelPass>(Allocator, RDG_EVENT_NAME("Graph Prologue (Graphics)")));
-	ProloguePassHandle = ProloguePass->Handle;
 	bInDebugPassScope = false;
 }
 
+///////////////////////////////////////////////////////////////////////////////////////////////////
+
 void FRDGBuilder::ExecutePassPrologue(FRHIComputeCommandList& RHICmdListPass, FRDGPass* Pass)
 {
 	CSV_SCOPED_TIMING_STAT_EXCLUSIVE_CONDITIONAL(RDGBuilder_ExecutePassPrologue, GRDGVerboseCSVStats != 0);
@@ -2474,6 +2476,8 @@ void FRDGBuilder::ExecutePass(FRDGPass* Pass, FRHIComputeCommandList& RHICmdList
 	}
 }
 
+///////////////////////////////////////////////////////////////////////////////////////////////////
+
 void FRDGBuilder::BeginResourcesRHI(FRDGPass* ResourcePass, FRDGPassHandle ExecutePassHandle)
 {
 	for (FRDGPass* PassToBegin : ResourcePass->ResourcesToBegin)
@@ -2520,6 +2524,8 @@ void FRDGBuilder::EndResourcesRHI(FRDGPass* ResourcePass, FRDGPassHandle Execute
 	}
 }
 
+///////////////////////////////////////////////////////////////////////////////////////////////////
+
 void FRDGBuilder::CollectPassBarriers(FRDGPass* Pass, FRDGPassHandle PassHandle)
 {
 	IF_RDG_ENABLE_DEBUG(ConditionalDebugBreak(RDG_BREAKPOINT_PASS_COMPILE, BuilderName.GetTCHAR(), Pass->GetName()));
@@ -2552,119 +2558,115 @@ void FRDGBuilder::CreatePassBarriers()
 	TransitionCreateQueue.Reset();
 }
 
+///////////////////////////////////////////////////////////////////////////////////////////////////
+
 void FRDGBuilder::AddEpilogueTransition(FRDGTextureRef Texture)
 {
-	if (!Texture->bLastOwner || Texture->bCulled || Texture->bFinalizedAccess)
+	check(Texture->HasRHI() || (Texture->bCulled && !Texture->bExternal));
+
+	if (Texture->bCulled || !Texture->bLastOwner)
 	{
 		return;
 	}
 
-	const FRDGPassHandle EpiloguePassHandle = GetEpiloguePassHandle();
-
-	FRDGSubresourceState ScratchSubresourceState;
-
-	// Texture is using the RHI transient allocator. Transition it back to Discard in the final pass it is used.
-	if (Texture->bTransient && !Texture->TransientTexture->IsAcquired())
+	if (!Texture->bFinalizedAccess)
 	{
-		const TInterval<uint32> DiscardPasses = Texture->TransientTexture->GetDiscardPasses();
-		const FRDGPassHandle MinDiscardPassHandle(DiscardPasses.Min);
-		const FRDGPassHandle MaxDiscardPassHandle(FMath::Min<uint32>(DiscardPasses.Max, EpiloguePassHandle.GetIndex()));
-
-		AddAliasingTransition(MinDiscardPassHandle, MaxDiscardPassHandle, Texture, FRHITransientAliasingInfo::Discard(Texture->GetRHIUnchecked()));
+		const FRDGPassHandle EpiloguePassHandle = GetEpiloguePassHandle();
 
-		ScratchSubresourceState.SetPass(ERHIPipeline::Graphics, MaxDiscardPassHandle);
-		ScratchSubresourceState.Access = ERHIAccess::Discard;
-		InitAsWholeResource(ScratchTextureState, &ScratchSubresourceState);
-	}
-	// A known final state means extraction from the graph (or an external texture).
-	else if(Texture->AccessFinal != ERHIAccess::Unknown)
-	{
-		ScratchSubresourceState.SetPass(ERHIPipeline::Graphics, EpiloguePassHandle);
-		ScratchSubresourceState.Access = Texture->AccessFinal;
-		InitAsWholeResource(ScratchTextureState, &ScratchSubresourceState);
-	}
-	// Lifetime is within the graph, but a pass may have left the resource in an async compute state. We cannot
-	// release the pooled texture back to the pool until we transition back to the graphics pipe.
-	else if (Texture->bUsedByAsyncComputePass)
-	{
-		FRDGTextureSubresourceState& TextureState = Texture->GetState();
-		ScratchTextureState.SetNumUninitialized(TextureState.Num(), false);
+		FRDGSubresourceState SubresourceState;
+		SubresourceState.SetPass(ERHIPipeline::Graphics, EpiloguePassHandle);
 
-		for (uint32 Index = 0, Count = ScratchTextureState.Num(); Index < Count; ++Index)
+		// Texture is using the RHI transient allocator. Transition it back to Discard in the final pass it is used.
+		if (Texture->bTransient && !Texture->TransientTexture->IsAcquired())
 		{
-			FRDGSubresourceState SubresourceState = TextureState[Index];
+			const TInterval<uint32> DiscardPasses = Texture->TransientTexture->GetDiscardPasses();
+			const FRDGPassHandle MinDiscardPassHandle(DiscardPasses.Min);
+			const FRDGPassHandle MaxDiscardPassHandle(FMath::Min<uint32>(DiscardPasses.Max, EpiloguePassHandle.GetIndex()));
 
-			// Transition async compute back to the graphics pipe.
-			if (SubresourceState.IsUsedBy(ERHIPipeline::AsyncCompute))
-			{
-				SubresourceState.SetPass(ERHIPipeline::Graphics, EpiloguePassHandle);
+			AddAliasingTransition(MinDiscardPassHandle, MaxDiscardPassHandle, Texture, FRHITransientAliasingInfo::Discard(Texture->GetRHIUnchecked()));
 
-				ScratchTextureState[Index] = AllocSubresource(SubresourceState);
-			}
-			else
-			{
-				ScratchTextureState[Index] = nullptr;
-			}
+			SubresourceState.SetPass(ERHIPipeline::Graphics, MaxDiscardPassHandle);
+			Texture->EpilogueAccess = ERHIAccess::Discard;
 		}
+
+		SubresourceState.Access = Texture->EpilogueAccess;
+
+		InitAsWholeResource(ScratchTextureState, &SubresourceState);
+		AddTransition(EpiloguePassHandle, Texture, ScratchTextureState);
+		ScratchTextureState.Reset();
 	}
-	// No need to transition; texture stayed on the graphics pipe and its lifetime stayed within the graph.
 	else
 	{
-		return;
+		Texture->EpilogueAccess = Texture->GetState()[0].Access;
+	}
+
+	// Resources that were finalized on registration could have an unknown epilogue state.
+	if (Texture->EpilogueAccess != ERHIAccess::Unknown)
+	{
+		EpilogueResourceAccesses.Emplace(Texture->GetRHI(), Texture->EpilogueAccess);
 	}
 
-	AddTransition(EpiloguePassHandle, Texture, ScratchTextureState);
-	ScratchTextureState.Reset();
+	if (Texture->RenderTarget)
+	{
+		ActivePooledTextures.Emplace(Texture->RenderTarget);
+	}
+
+	Texture->Allocation = nullptr;
 }
 
 void FRDGBuilder::AddEpilogueTransition(FRDGBufferRef Buffer)
 {
-	if (!Buffer->bLastOwner || Buffer->bCulled || Buffer->bFinalizedAccess)
+	check(Buffer->HasRHI() || (Buffer->bCulled && !Buffer->bExternal));
+
+	if (Buffer->bCulled || !Buffer->bLastOwner)
 	{
 		return;
 	}
 
-	const FRDGPassHandle EpiloguePassHandle = GetEpiloguePassHandle();
-
-	if (Buffer->bTransient)
+	if (!Buffer->bFinalizedAccess)
 	{
-		const TInterval<uint32> DiscardPasses = Buffer->TransientBuffer->GetDiscardPasses();
-		const FRDGPassHandle MinDiscardPassHandle(DiscardPasses.Min);
-		const FRDGPassHandle MaxDiscardPassHandle(FMath::Min<uint32>(DiscardPasses.Max, EpiloguePassHandle.GetIndex()));
+		const FRDGPassHandle EpiloguePassHandle = GetEpiloguePassHandle();
+
+		FRDGSubresourceState SubresourceState;
+		SubresourceState.SetPass(ERHIPipeline::Graphics, EpiloguePassHandle);
 
-		AddAliasingTransition(MinDiscardPassHandle, MaxDiscardPassHandle, Buffer, FRHITransientAliasingInfo::Discard(Buffer->GetRHIUnchecked()));
+		// Texture is using the RHI transient allocator. Transition it back to Discard in the final pass it is used.
+		if (Buffer->bTransient)
+		{
+			const TInterval<uint32> DiscardPasses = Buffer->TransientBuffer->GetDiscardPasses();
+			const FRDGPassHandle MinDiscardPassHandle(DiscardPasses.Min);
+			const FRDGPassHandle MaxDiscardPassHandle(FMath::Min<uint32>(DiscardPasses.Max, EpiloguePassHandle.GetIndex()));
+
+			AddAliasingTransition(MinDiscardPassHandle, MaxDiscardPassHandle, Buffer, FRHITransientAliasingInfo::Discard(Buffer->GetRHIUnchecked()));
+
+			SubresourceState.SetPass(ERHIPipeline::Graphics, MaxDiscardPassHandle);
+			Buffer->EpilogueAccess = ERHIAccess::Discard;
+		}
+
+		SubresourceState.Access = Buffer->EpilogueAccess;
 
-		FRDGSubresourceState StateFinal;
-		StateFinal.SetPass(ERHIPipeline::Graphics, MaxDiscardPassHandle);
-		StateFinal.Access = ERHIAccess::Discard;
-		AddTransition(Buffer->LastPass, Buffer, StateFinal);
+		AddTransition(Buffer->LastPass, Buffer, SubresourceState);
 	}
 	else
 	{
-		ERHIAccess AccessFinal = Buffer->AccessFinal;
-
-		// Transition async compute back to the graphics pipe.
-		if (AccessFinal == ERHIAccess::Unknown)
-		{
-			const FRDGSubresourceState State = Buffer->GetState();
+		Buffer->EpilogueAccess = Buffer->GetState().Access;
+	}
 
-			if (State.IsUsedBy(ERHIPipeline::AsyncCompute))
-			{
-				AccessFinal = State.Access;
-			}
-		}
+	// Resources that were finalized on registration could have an unknown epilogue state.
+	if (Buffer->EpilogueAccess != ERHIAccess::Unknown)
+	{
+		EpilogueResourceAccesses.Emplace(Buffer->GetRHI(), Buffer->EpilogueAccess);
+	}
 
-		if (AccessFinal != ERHIAccess::Unknown)
-		{
-			FRDGSubresourceState StateFinal;
-			StateFinal.SetPass(ERHIPipeline::Graphics, EpiloguePassHandle);
-			StateFinal.Access = AccessFinal;
-			AddTransition(EpiloguePassHandle, Buffer, StateFinal);
-		}
+	if (!Buffer->bTransient)
+	{
+		ActivePooledBuffers.Emplace(Buffer->PooledBuffer);
 	}
+
+	Buffer->Allocation = nullptr;
 }
 
-void FRDGBuilder::AddTransition(FRDGPassHandle PassHandle, FRDGTextureRef Texture, const FRDGTextureTransientSubresourceStateIndirect& StateAfter)
+void FRDGBuilder::AddTransition(FRDGPassHandle PassHandle, FRDGTextureRef Texture, const FRDGTextureSubresourceStateIndirect& StateAfter)
 {
 	const FRDGTextureSubresourceRange WholeRange = Texture->GetSubresourceRange();
 	const FRDGTextureSubresourceLayout Layout = Texture->Layout;
@@ -2700,15 +2702,6 @@ void FRDGBuilder::AddTransition(FRDGPassHandle PassHandle, FRDGTextureRef Textur
 
 			AddTransition(Texture, SubresourceStateBefore, SubresourceStateAfter, Info);
 		}
-
-		if (Subresource)
-		{
-			IF_RDG_ENABLE_DEBUG(LogFile.AddTransitionEdge(PassHandle, SubresourceStateBefore, SubresourceStateAfter, Texture, *Subresource));
-		}
-		else
-		{
-			IF_RDG_ENABLE_DEBUG(LogFile.AddTransitionEdge(PassHandle, SubresourceStateBefore, SubresourceStateAfter, Texture));
-		}
 	};
 
 	if (IsWholeResource(StateBefore))
@@ -2789,12 +2782,11 @@ void FRDGBuilder::AddTransition(FRDGPassHandle PassHandle, FRDGBufferRef Buffer,
 		AddTransition(Buffer, StateBefore, StateAfter, Info);
 	}
 
-	IF_RDG_ENABLE_DEBUG(LogFile.AddTransitionEdge(PassHandle, StateBefore, StateAfter, Buffer));
 	StateBefore = StateAfter;
 }
 
 void FRDGBuilder::AddTransition(
-	FRDGParentResource* Resource,
+	FRDGViewableResource* Resource,
 	FRDGSubresourceState StateBefore,
 	FRDGSubresourceState StateAfter,
 	const FRHITransitionInfo& TransitionInfo)
@@ -2817,8 +2809,6 @@ void FRDGBuilder::AddTransition(
 		return;
 	}
 
-	StateBefore.LastPass = ClampToPrologue(StateBefore.LastPass);
-
 	ERHIPipeline PipelinesBefore = StateBefore.GetPipelines();
 	ERHIPipeline PipelinesAfter = StateAfter.GetPipelines();
 
@@ -2934,7 +2924,7 @@ void FRDGBuilder::AddTransition(
 	}
 }
 
-void FRDGBuilder::AddAliasingTransition(FRDGPassHandle BeginPassHandle, FRDGPassHandle EndPassHandle, FRDGParentResourceRef Resource, const FRHITransientAliasingInfo& Info)
+void FRDGBuilder::AddAliasingTransition(FRDGPassHandle BeginPassHandle, FRDGPassHandle EndPassHandle, FRDGViewableResource* Resource, const FRHITransientAliasingInfo& Info)
 {
 	check(BeginPassHandle <= EndPassHandle);
 
@@ -2967,6 +2957,128 @@ void FRDGBuilder::AddAliasingTransition(FRDGPassHandle BeginPassHandle, FRDGPass
 	EndPass->GetPrologueBarriersToEnd(Allocator).AddDependency(BarriersToBegin);
 }
 
+///////////////////////////////////////////////////////////////////////////////////////////////////
+
+void FRDGBuilder::SetRHI(FRDGTexture* Texture, IPooledRenderTarget* RenderTarget, FRDGPassHandle PassHandle)
+{
+	Texture->RenderTarget = RenderTarget;
+
+	if (FRHITransientTexture* TransientTexture = RenderTarget->GetTransientTexture())
+	{
+		FRDGTransientRenderTarget* TransientRenderTarget = static_cast<FRDGTransientRenderTarget*>(RenderTarget);
+		Texture->Allocation = TRefCountPtr<FRDGTransientRenderTarget>(TransientRenderTarget);
+
+		SetRHI(Texture, TransientTexture, PassHandle);
+	}
+	else
+	{
+		FPooledRenderTarget* PooledRenderTarget = static_cast<FPooledRenderTarget*>(RenderTarget);
+		Texture->Allocation = TRefCountPtr<FPooledRenderTarget>(PooledRenderTarget);
+
+		SetRHI(Texture, &PooledRenderTarget->PooledTexture, PassHandle);
+	}
+}
+
+void FRDGBuilder::SetRHI(FRDGTexture* Texture, FRDGPooledTexture* PooledTexture, FRDGPassHandle PassHandle)
+{
+	FRHITexture* TextureRHI = PooledTexture->GetRHI();
+
+	Texture->ResourceRHI = TextureRHI;
+	Texture->PooledTexture = PooledTexture;
+	Texture->ViewCache = &PooledTexture->ViewCache;
+	Texture->FirstPass = PassHandle;
+
+	FRDGTexture*& Owner = PooledTextureOwnershipMap.FindOrAdd(PooledTexture);
+
+	// Link the previous alias to this one.
+	if (Owner)
+	{
+		Owner->NextOwner = Texture->Handle;
+		Owner->bLastOwner = false;
+
+		// Chain the state allocation between all RDG textures which share this pooled texture.
+		Texture->State = Owner->State;
+	}
+	else
+	{
+		Texture->State = Allocator.AllocNoDestruct<FRDGTextureSubresourceState>();
+		InitAsWholeResource(*Texture->State, {});
+	}
+
+	Owner = Texture;
+}
+
+void FRDGBuilder::SetRHI(FRDGTexture* Texture, FRHITransientTexture* TransientTexture, FRDGPassHandle PassHandle)
+{
+	Texture->ResourceRHI = TransientTexture->GetRHI();
+	Texture->TransientTexture = TransientTexture;
+	Texture->ViewCache = &TransientTexture->ViewCache;
+	Texture->FirstPass = PassHandle;
+	Texture->bTransient = true;
+
+	const FRDGPassHandle MinAcquirePassHandle(TransientTexture->GetAcquirePasses().Min);
+
+	AddAliasingTransition(MinAcquirePassHandle, PassHandle, Texture, FRHITransientAliasingInfo::Acquire(TransientTexture->GetRHI(), TransientTexture->GetAliasingOverlaps()));
+
+	FRDGSubresourceState InitialState;
+	InitialState.SetPass(ERHIPipeline::Graphics, MinAcquirePassHandle);
+	InitialState.Access = ERHIAccess::Discard;
+
+	Texture->State = Allocator.AllocNoDestruct<FRDGTextureSubresourceState>();
+	InitAsWholeResource(*Texture->State, InitialState);
+}
+
+void FRDGBuilder::SetRHI(FRDGBuffer* Buffer, FRDGPooledBuffer* PooledBuffer, FRDGPassHandle PassHandle)
+{
+	FRHIBuffer* BufferRHI = PooledBuffer->GetRHI();
+
+	Buffer->ResourceRHI = BufferRHI;
+	Buffer->PooledBuffer = PooledBuffer;
+	Buffer->ViewCache = &PooledBuffer->ViewCache;
+	Buffer->Allocation = PooledBuffer;
+	Buffer->FirstPass = PassHandle;
+
+	FRDGBuffer*& Owner = PooledBufferOwnershipMap.FindOrAdd(PooledBuffer);
+
+	// Link the previous owner to this one.
+	if (Owner)
+	{
+		Owner->NextOwner = Buffer->Handle;
+		Owner->bLastOwner = false;
+
+		// Chain the state allocation between all RDG buffers which share this pooled buffer.
+		Buffer->State = Owner->State;
+	}
+	else
+	{
+		Buffer->State = Allocator.AllocNoDestruct<FRDGSubresourceState>();
+	}
+
+	Owner = Buffer;
+}
+
+void FRDGBuilder::SetRHI(FRDGBuffer* Buffer, FRHITransientBuffer* TransientBuffer, FRDGPassHandle PassHandle)
+{
+	check(!Buffer->ResourceRHI);
+
+	Buffer->ResourceRHI = TransientBuffer->GetRHI();
+	Buffer->TransientBuffer = TransientBuffer;
+	Buffer->ViewCache = &TransientBuffer->ViewCache;
+	Buffer->FirstPass = PassHandle;
+	Buffer->bTransient = true;
+
+	const FRDGPassHandle MinAcquirePassHandle(TransientBuffer->GetAcquirePasses().Min);
+
+	AddAliasingTransition(MinAcquirePassHandle, PassHandle, Buffer, FRHITransientAliasingInfo::Acquire(TransientBuffer->GetRHI(), TransientBuffer->GetAliasingOverlaps()));
+
+	FRDGSubresourceState* InitialState = Allocator.AllocNoDestruct<FRDGSubresourceState>();
+	InitialState->SetPass(ERHIPipeline::Graphics, MinAcquirePassHandle);
+	InitialState->Access = ERHIAccess::Discard;
+	Buffer->State = InitialState;
+}
+
+///////////////////////////////////////////////////////////////////////////////////////////////////
+
 void FRDGBuilder::BeginResourceRHI(FRDGPassHandle PassHandle, FRDGTextureRef Texture)
 {
 	check(Texture);
@@ -2997,22 +3109,13 @@ void FRDGBuilder::BeginResourceRHI(FRDGPassHandle PassHandle, FRDGTextureRef Tex
 		{
 			if (Texture->bExternal || Texture->bExtracted)
 			{
-				Texture->SetRHI(GRDGTransientResourceAllocator.AllocateRenderTarget(TransientTexture));
+				SetRHI(Texture, GRDGTransientResourceAllocator.AllocateRenderTarget(TransientTexture), PassHandle);
 			}
 			else
 			{
-				Texture->SetRHI(TransientTexture, Allocator.AllocNoDestruct<FRDGTextureSubresourceState>());
+				SetRHI(Texture, TransientTexture, PassHandle);
 			}
 
-			const FRDGPassHandle MinAcquirePassHandle = ClampToPrologue(FRDGPassHandle(TransientTexture->GetAcquirePasses().Min));
-
-			AddAliasingTransition(MinAcquirePassHandle, PassHandle, Texture, FRHITransientAliasingInfo::Acquire(TransientTexture->GetRHI(), TransientTexture->GetAliasingOverlaps()));
-
-			FRDGSubresourceState InitialState;
-			InitialState.SetPass(ERHIPipeline::Graphics, MinAcquirePassHandle);
-			InitialState.Access = ERHIAccess::Discard;
-			InitAsWholeResource(Texture->GetState(), InitialState);
-
 		#if STATS
 			GRDGStatTransientTextureCount++;
 		#endif
@@ -3021,13 +3124,8 @@ void FRDGBuilder::BeginResourceRHI(FRDGPassHandle PassHandle, FRDGTextureRef Tex
 
 	if (!Texture->ResourceRHI)
 	{
-		const bool bResetToUnknownState = false;
-		Texture->SetRHI(GRenderTargetPool.FindFreeElementInternal(Texture->Desc, Texture->Name, bResetToUnknownState));
+		SetRHI(Texture, GRenderTargetPool.FindFreeElement(Texture->Desc, Texture->Name), PassHandle);
 	}
-
-	Texture->FirstPass = PassHandle;
-
-	check(Texture->HasRHI());
 }
 
 void FRDGBuilder::BeginResourceRHI(FRDGPassHandle PassHandle, FRDGTextureSRVRef SRV)
@@ -3092,15 +3190,7 @@ void FRDGBuilder::BeginResourceRHI(FRDGPassHandle PassHandle, FRDGBufferRef Buff
 	{
 		if (FRHITransientBuffer* TransientBuffer = TransientResourceAllocator->CreateBuffer(Translate(Buffer->Desc), Buffer->Name, PassHandle.GetIndex()))
 		{
-			Buffer->SetRHI(TransientBuffer, Allocator);
-
-			const FRDGPassHandle MinAcquirePassHandle = ClampToPrologue(FRDGPassHandle(TransientBuffer->GetAcquirePasses().Min));
-
-			AddAliasingTransition(MinAcquirePassHandle, PassHandle, Buffer, FRHITransientAliasingInfo::Acquire(TransientBuffer->GetRHI(), TransientBuffer->GetAliasingOverlaps()));
-
-			FRDGSubresourceState& InitialState = Buffer->GetState();
-			InitialState.SetPass(ERHIPipeline::Graphics, MinAcquirePassHandle);
-			InitialState.Access = ERHIAccess::Discard;
+			SetRHI(Buffer, TransientBuffer, PassHandle);
 
 		#if STATS
 			GRDGStatTransientBufferCount++;
@@ -3110,12 +3200,8 @@ void FRDGBuilder::BeginResourceRHI(FRDGPassHandle PassHandle, FRDGBufferRef Buff
 
 	if (!Buffer->bTransient)
 	{
-		Buffer->SetRHI(GRenderGraphResourcePool.FindFreeBufferInternal(RHICmdList, Buffer->Desc, Buffer->Name));
+		SetRHI(Buffer, GRenderGraphResourcePool.FindFreeBuffer(RHICmdList, Buffer->Desc, Buffer->Name), PassHandle);
 	}
-
-	Buffer->FirstPass = PassHandle;
-
-	check(Buffer->HasRHI());
 }
 
 void FRDGBuilder::BeginResourceRHI(FRDGPassHandle PassHandle, FRDGBufferSRVRef SRV)
@@ -3193,6 +3279,8 @@ void FRDGBuilder::BeginResourceRHI(FRDGPassHandle PassHandle, FRDGView* View)
 	}
 }
 
+///////////////////////////////////////////////////////////////////////////////////////////////////
+
 void FRDGBuilder::EndResourceRHI(FRDGPassHandle PassHandle, FRDGTextureRef Texture, uint32 ReferenceCount)
 {
 	check(Texture);
@@ -3204,7 +3292,7 @@ void FRDGBuilder::EndResourceRHI(FRDGPassHandle PassHandle, FRDGTextureRef Textu
 		if (Texture->bTransient)
 		{
 			// Texture is using a transient external render target.
-			if (Texture->PooledRenderTarget)
+			if (Texture->RenderTarget)
 			{
 				// This releases the reference without invoking a virtual function call.
 				GRDGTransientResourceAllocator.Release(TRefCountPtr<FRDGTransientRenderTarget>(MoveTemp(Texture->Allocation)), PassHandle);
@@ -3217,12 +3305,8 @@ void FRDGBuilder::EndResourceRHI(FRDGPassHandle PassHandle, FRDGTextureRef Textu
 		}
 		else
 		{
-			// If this is a non-transient texture, it must be backed by a pooled render target.
-			FPooledRenderTarget* RenderTarget = static_cast<FPooledRenderTarget*>(Texture->PooledRenderTarget);
-			check(RenderTarget);
-
 			// Only tracked render targets are released. Untracked ones persist until the end of the frame.
-			if (RenderTarget->IsTracked())
+			if (static_cast<FPooledRenderTarget*>(Texture->RenderTarget)->IsTracked())
 			{
 				// This releases the reference without invoking a virtual function call.
 				TRefCountPtr<FPooledRenderTarget>(MoveTemp(Texture->Allocation));
@@ -3254,6 +3338,8 @@ void FRDGBuilder::EndResourceRHI(FRDGPassHandle PassHandle, FRDGBufferRef Buffer
 	}
 }
 
+///////////////////////////////////////////////////////////////////////////////////////////////////
+
 #if RDG_ENABLE_DEBUG
 
 void FRDGBuilder::VisualizePassOutputs(const FRDGPass* Pass)
diff --git a/Engine/Source/Runtime/RenderCore/Private/RenderGraphPass.cpp b/Engine/Source/Runtime/RenderCore/Private/RenderGraphPass.cpp
--- a/Engine/Source/Runtime/RenderCore/Private/RenderGraphPass.cpp
+++ b/Engine/Source/Runtime/RenderCore/Private/RenderGraphPass.cpp
@@ -128,7 +128,7 @@ FRDGBarrierBatchBegin::FRDGBarrierBatchBegin(ERHIPipeline InPipelinesToBegin, ER
 #endif
 {}
 
-void FRDGBarrierBatchBegin::AddTransition(FRDGParentResourceRef Resource, const FRHITransitionInfo& Info)
+void FRDGBarrierBatchBegin::AddTransition(FRDGViewableResource* Resource, const FRHITransitionInfo& Info)
 {
 	Transitions.Add(Info);
 	bTransitionNeeded = true;
@@ -142,7 +142,7 @@ void FRDGBarrierBatchBegin::AddTransition(FRDGParentResourceRef Resource, const
 #endif
 }
 
-void FRDGBarrierBatchBegin::AddAlias(FRDGParentResourceRef Resource, const FRHITransientAliasingInfo& Info)
+void FRDGBarrierBatchBegin::AddAlias(FRDGViewableResource* Resource, const FRHITransientAliasingInfo& Info)
 {
 	Aliases.Add(Info);
 	bTransitionNeeded = true;
diff --git a/Engine/Source/Runtime/RenderCore/Private/RenderGraphResourcePool.cpp b/Engine/Source/Runtime/RenderCore/Private/RenderGraphResourcePool.cpp
--- a/Engine/Source/Runtime/RenderCore/Private/RenderGraphResourcePool.cpp
+++ b/Engine/Source/Runtime/RenderCore/Private/RenderGraphResourcePool.cpp
@@ -16,16 +16,6 @@ TRefCountPtr<FRDGPooledBuffer> FRDGBufferPool::FindFreeBuffer(
 	FRHICommandList& RHICmdList,
 	const FRDGBufferDesc& Desc,
 	const TCHAR* InDebugName)
-{
-	TRefCountPtr<FRDGPooledBuffer> Result = FindFreeBufferInternal(RHICmdList, Desc, InDebugName);
-	Result->Reset();
-	return Result;
-}
-
-TRefCountPtr<FRDGPooledBuffer> FRDGBufferPool::FindFreeBufferInternal(
-	FRHICommandList& RHICmdList,
-	const FRDGBufferDesc& Desc,
-	const TCHAR* InDebugName)
 {
 	const uint64 BufferPageSize = 64 * 1024;
 
@@ -108,7 +98,6 @@ TRefCountPtr<FRDGPooledBuffer> FRDGBufferPool::FindFreeBufferInternal(
 		check(PooledBuffer->GetRefCount() == 2);
 
 		PooledBuffer->LastUsedFrame = FrameCounter;
-		PooledBuffer->State.Access = InitialAccess;
 
 		return PooledBuffer;
 	}
@@ -228,7 +217,6 @@ TRefCountPtr<FRDGTransientRenderTarget> FRDGTransientResourceAllocator::Allocate
 	RenderTarget->LifetimeState = ERDGTransientResourceLifetimeState::Allocated;
 	RenderTarget->GetRenderTargetItem().TargetableTexture = Texture->GetRHI();
 	RenderTarget->GetRenderTargetItem().ShaderResourceTexture = Texture->GetRHI();
-	InitAsWholeResource(RenderTarget->State, {});
 	return RenderTarget;
 }
 
@@ -236,6 +224,8 @@ void FRDGTransientResourceAllocator::Release(TRefCountPtr<FRDGTransientRenderTar
 {
 	check(RenderTarget);
 
+	// If this is true, we hold the final reference in the RenderTarget argument. We want to zero out its
+	// members before dereferencing to zero so that it gets marked as deallocated rather than pending.
 	if (RenderTarget->GetRefCount() == 1)
 	{
 		Allocator->DeallocateMemory(RenderTarget->Texture, PassHandle.GetIndex());
diff --git a/Engine/Source/Runtime/RenderCore/Private/RenderGraphResourcePool.h b/Engine/Source/Runtime/RenderCore/Private/RenderGraphResourcePool.h
--- a/Engine/Source/Runtime/RenderCore/Private/RenderGraphResourcePool.h
+++ b/Engine/Source/Runtime/RenderCore/Private/RenderGraphResourcePool.h
@@ -26,9 +26,6 @@ public:
 private:
 	void ReleaseDynamicRHI() override;
 
-	/** Allocate a buffer from a given descriptor. */
-	TRefCountPtr<FRDGPooledBuffer> FindFreeBufferInternal(FRHICommandList& RHICmdList, const FRDGBufferDesc& Desc, const TCHAR* InDebugName);
-
 	/** Elements can be 0, we compact the buffer later. */
 	TArray<TRefCountPtr<FRDGPooledBuffer>> AllocatedBuffers;
 	TArray<uint64> AllocatedBufferHashes;
@@ -74,8 +71,6 @@ public:
 		RenderTargetItem.TargetableTexture = nullptr;
 	}
 
-	FRDGTextureSubresourceState State;
-
 private:
 	FRDGTransientRenderTarget() = default;
 
diff --git a/Engine/Source/Runtime/RenderCore/Private/RenderGraphResources.cpp b/Engine/Source/Runtime/RenderCore/Private/RenderGraphResources.cpp
--- a/Engine/Source/Runtime/RenderCore/Private/RenderGraphResources.cpp
+++ b/Engine/Source/Runtime/RenderCore/Private/RenderGraphResources.cpp
@@ -21,7 +21,7 @@ inline bool NeedsUAVBarrier(const FRDGSubresourceState& Previous, const FRDGSubr
 	return NeedsUAVBarrier(Previous.NoUAVBarrierFilter.GetUniqueHandle(), Next.NoUAVBarrierFilter.GetUniqueHandle());
 }
 
-FRDGParentResource::FRDGParentResource(const TCHAR* InName, const ERDGParentResourceType InType)
+FRDGViewableResource::FRDGViewableResource(const TCHAR* InName, const ERDGViewableResourceType InType)
 	: FRDGResource(InName)
 	, Type(InType)
 	, bExternal(0)
@@ -74,7 +74,7 @@ bool FRDGProducerState::IsDependencyRequired(FRDGProducerState LastProducer, ERH
 	return true;
 }
 
-bool FRDGSubresourceState::IsMergeAllowed(ERDGParentResourceType ResourceType, const FRDGSubresourceState& Previous, const FRDGSubresourceState& Next)
+bool FRDGSubresourceState::IsMergeAllowed(ERDGViewableResourceType ResourceType, const FRDGSubresourceState& Previous, const FRDGSubresourceState& Next)
 {
 	/** State merging occurs during compilation and before resource transitions are collected. It serves to remove the bulk
 	 *  of unnecessary transitions by looking ahead in the resource usage chain. A resource transition cannot occur within
@@ -175,21 +175,6 @@ void FRDGUniformBuffer::InitRHI()
 	ResourceRHI = UniformBufferRHI;
 }
 
-void FRDGPooledTexture::Finalize()
-{
-	for (FRDGSubresourceState& SubresourceState : State)
-	{
-		SubresourceState.Finalize();
-	}
-	Owner = nullptr;
-}
-
-void FRDGPooledTexture::Reset()
-{
-	InitAsWholeResource(State);
-	Owner = nullptr;
-}
-
 FRDGTextureSubresourceRange FRDGTexture::GetSubresourceRangeSRV() const
 {
 	FRDGTextureSubresourceRange Range = GetSubresourceRange();
@@ -205,168 +190,5 @@ FRDGTextureSubresourceRange FRDGTexture::GetSubresourceRangeSRV() const
 IPooledRenderTarget* FRDGTexture::GetPooledRenderTarget() const
 {
 	IF_RDG_ENABLE_DEBUG(ValidateRHIAccess());
-	return PooledRenderTarget;
-}
-
-void FRDGTexture::SetRHI(IPooledRenderTarget* InPooledRenderTarget)
-{
-	if (FRHITransientTexture* LocalTransientTexture = InPooledRenderTarget->GetTransientTexture())
-	{
-		FRDGTransientRenderTarget* LocalRenderTarget = static_cast<FRDGTransientRenderTarget*>(InPooledRenderTarget);
-		Allocation = TRefCountPtr<FRDGTransientRenderTarget>(LocalRenderTarget);
-
-		SetRHI(LocalTransientTexture, &LocalRenderTarget->State);
-	}
-	else
-	{
-		FPooledRenderTarget* LocalRenderTarget = static_cast<FPooledRenderTarget*>(InPooledRenderTarget);
-		Allocation = TRefCountPtr<FPooledRenderTarget>(LocalRenderTarget);
-
-		SetRHI(&LocalRenderTarget->PooledTexture);
-	}
-
-	PooledRenderTarget = InPooledRenderTarget;
-}
-
-void FRDGTexture::SetRHI(FRDGPooledTexture* InPooledTexture)
-{
-	PooledTexture = InPooledTexture;
-	State = &PooledTexture->State;
-	ViewCache = &PooledTexture->ViewCache;
-
-	// Return the previous owner and assign this texture as the new one.
-	FRDGTextureRef PreviousOwner = PooledTexture->Owner;
-	PooledTexture->Owner = this;
-
-	// Link the previous alias to this one.
-	if (PreviousOwner)
-	{
-		PreviousOwner->NextOwner = Handle;
-		PreviousOwner->bLastOwner = false;
-	}
-
-	ResourceRHI = PooledTexture->GetRHI();
-}
-
-void FRDGTexture::SetRHI(FRHITransientTexture* InTransientTexture, FRDGTextureSubresourceState* InTransientTextureState)
-{
-	TransientTexture = InTransientTexture;
-	State = InTransientTextureState;
-	ViewCache = &InTransientTexture->ViewCache;
-	ResourceRHI = InTransientTexture->GetRHI();
-	bTransient = true;
-}
-
-void FRDGTexture::Finalize(FRDGPooledTextureArray& PooledTextureArray)
-{
-	checkf(NextOwner.IsNull() == !!bLastOwner, TEXT("NextOwner must match bLastOwner."));
-	checkf(!bExtracted || bLastOwner, TEXT("Extracted resources must be the last owner of a resource."));
-
-	if (bLastOwner)
-	{
-		if (bTransient)
-		{
-			if (PooledRenderTarget)
-			{
-				InitAsWholeResource(*State, {});
-				PooledTextureArray.Emplace(MoveTemp(Allocation));
-			}
-			else
-			{
-				// Manually deconstruct the allocated state so as not to invoke overhead from the allocators destructor tracking.
-				State->~FRDGTextureSubresourceState();
-				State = nullptr;
-			}
-		}
-		else
-		{
-			if (PooledTexture)
-			{
-				// External and extracted resources are user controlled, so we cannot assume the texture stays in its final state.
-				if (bExternal || bExtracted)
-				{
-					PooledTexture->Reset();
-				}
-				else
-				{
-					PooledTexture->Finalize();
-				}
-			}
-
-			if (PooledRenderTarget)
-			{
-				PooledTextureArray.Emplace(PooledRenderTarget);
-			}
-		}
-	}
-
-	Allocation = nullptr;
-}
-
-void FRDGBuffer::SetRHI(FRDGPooledBuffer* InPooledBuffer)
-{
-	// Return the previous owner and assign this buffer as the new one.
-	FRDGBuffer* PreviousOwner = InPooledBuffer->Owner;
-	InPooledBuffer->Owner = this;
-
-	// Link the previous owner to this one.
-	if (PreviousOwner)
-	{
-		PreviousOwner->NextOwner = Handle;
-		PreviousOwner->bLastOwner = false;
-	}
-
-	PooledBuffer = InPooledBuffer;
-	Allocation = InPooledBuffer;
-	State = &PooledBuffer->State;
-	ViewCache = &PooledBuffer->ViewCache;
-	ResourceRHI = InPooledBuffer->GetRHI();
-
-	// The upload with UAV workaround performs its own transitions outside of RDG, so fall back to Unknown for simplicity.
-#if PLATFORM_NEEDS_GPU_UAV_RESOURCE_INIT_WORKAROUND
-	if (bUAVAccessed && bQueuedForUpload)
-	{
-		*State = FRDGSubresourceState();
-	}
-#endif
-}
-
-void FRDGBuffer::SetRHI(FRHITransientBuffer* InTransientBuffer, FRDGAllocator& Allocator)
-{
-	TransientBuffer = InTransientBuffer;
-	State = Allocator.AllocNoDestruct<FRDGSubresourceState>();
-	ViewCache = &InTransientBuffer->ViewCache;
-	ResourceRHI = InTransientBuffer->GetRHI();
-
-	bTransient = true;
-}
-
-void FRDGBuffer::Finalize(FRDGPooledBufferArray& PooledBufferArray)
-{
-	// If these fire, the graph is not tracking state properly.
-	checkf(NextOwner.IsNull() == !!bLastOwner, TEXT("NextOwner must match bLastOwner."));
-	checkf(!bExtracted || bLastOwner, TEXT("Extracted resources must be the last owner of a resource."));
-
-	if (bLastOwner)
-	{
-		if (bTransient)
-		{
-			State = nullptr;
-		}
-		else
-		{
-			if (bExternal || bExtracted)
-			{
-				PooledBuffer->Reset();
-			}
-			else
-			{
-				PooledBuffer->Finalize();
-			}
-
-			PooledBufferArray.Emplace(PooledBuffer);
-		}
-	}
-
-	Allocation = nullptr;
+	return RenderTarget;
 }
\ No newline at end of file
diff --git a/Engine/Source/Runtime/RenderCore/Private/RenderGraphTrace.cpp b/Engine/Source/Runtime/RenderCore/Private/RenderGraphTrace.cpp
--- a/Engine/Source/Runtime/RenderCore/Private/RenderGraphTrace.cpp
+++ b/Engine/Source/Runtime/RenderCore/Private/RenderGraphTrace.cpp
@@ -446,7 +446,7 @@ void FRDGTrace::OutputGraphEnd(const FRDGBuilder& GraphBuilder)
 	UE_TRACE_LOG(RDGTrace, GraphEndMessage, RDGChannel);
 }
 
-void FRDGTrace::AddResource(FRDGParentResource* Resource)
+void FRDGTrace::AddResource(FRDGViewableResource* Resource)
 {
 	Resource->TraceOrder = ResourceOrder++;
 }
diff --git a/Engine/Source/Runtime/RenderCore/Private/RenderGraphValidation.cpp b/Engine/Source/Runtime/RenderCore/Private/RenderGraphValidation.cpp
--- a/Engine/Source/Runtime/RenderCore/Private/RenderGraphValidation.cpp
+++ b/Engine/Source/Runtime/RenderCore/Private/RenderGraphValidation.cpp
@@ -91,7 +91,7 @@ FRDGResourceDebugData& FRDGResource::GetDebugData() const
 	return *DebugData;
 }
 
-struct FRDGParentResourceDebugData
+struct FRDGViewableResourceDebugData
 {
 	/** Pointer towards the pass that is the first to produce it, for even more convenient error message. */
 	const FRDGPass* FirstProducer = nullptr;
@@ -106,10 +106,10 @@ struct FRDGParentResourceDebugData
 	FRDGPassHandle FinalizePass;
 };
 
-FRDGParentResourceDebugData& FRDGParentResource::GetParentDebugData() const
+FRDGViewableResourceDebugData& FRDGViewableResource::GetViewableDebugData() const
 {
-	check(ParentDebugData);
-	return *ParentDebugData;
+	check(ViewableDebugData);
+	return *ViewableDebugData;
 }
 
 struct FRDGTextureDebugData
@@ -181,15 +181,15 @@ void FRDGUserValidation::ValidateCreateResource(FRDGResourceRef Resource)
 	check(!bAlreadyInSet);
 }
 
-void FRDGUserValidation::ValidateCreateParentResource(FRDGParentResourceRef Resource)
+void FRDGUserValidation::ValidateCreateViewableResource(FRDGViewableResource* Resource)
 {
 	ValidateCreateResource(Resource);
-	Resource->ParentDebugData = Allocator.Alloc<FRDGParentResourceDebugData>();
+	Resource->ViewableDebugData = Allocator.Alloc<FRDGViewableResourceDebugData>();
 }
 
 void FRDGUserValidation::ValidateCreateTexture(FRDGTextureRef Texture)
 {
-	ValidateCreateParentResource(Texture);
+	ValidateCreateViewableResource(Texture);
 	Texture->TextureDebugData = Allocator.Alloc<FRDGTextureDebugData>();
 	if (GRDGDebug)
 	{
@@ -199,7 +199,7 @@ void FRDGUserValidation::ValidateCreateTexture(FRDGTextureRef Texture)
 
 void FRDGUserValidation::ValidateCreateBuffer(FRDGBufferRef Buffer)
 {
-	ValidateCreateParentResource(Buffer);
+	ValidateCreateViewableResource(Buffer);
 	Buffer->BufferDebugData = Allocator.Alloc<FRDGBufferDebugData>();
 	if (GRDGDebug)
 	{
@@ -413,7 +413,7 @@ void FRDGUserValidation::ValidateExtractBuffer(FRDGBufferRef Buffer, TRefCountPt
 	checkf(OutBufferPtr, TEXT("Texture %s was extracted, but the output texture pointer is null."), Buffer->Name);
 }
 
-void FRDGUserValidation::ValidateExtractResource(FRDGParentResourceRef Resource)
+void FRDGUserValidation::ValidateExtractResource(FRDGViewableResource* Resource)
 {
 	check(Resource);
 
@@ -425,10 +425,10 @@ void FRDGUserValidation::ValidateExtractResource(FRDGParentResourceRef Resource)
 	 *  emitting a 'produced but never used' warning. We don't have the history of registered
 	 *  resources to be able to emit a proper warning.
 	 */
-	Resource->GetParentDebugData().PassAccessCount++;
+	Resource->GetViewableDebugData().PassAccessCount++;
 }
 
-void FRDGUserValidation::ValidateConvertToExternalResource(FRDGParentResourceRef Resource)
+void FRDGUserValidation::ValidateConvertToExternalResource(FRDGViewableResource* Resource)
 {
 	check(Resource);
 	checkf(!bHasExecuteBegun || !Resource->bTransient,
@@ -436,22 +436,22 @@ void FRDGUserValidation::ValidateConvertToExternalResource(FRDGParentResourceRef
 		Resource->Name);
 }
 
-void FRDGUserValidation::RemoveUnusedWarning(FRDGParentResourceRef Resource)
+void FRDGUserValidation::RemoveUnusedWarning(FRDGViewableResource* Resource)
 {
 	check(Resource);
 	ExecuteGuard(TEXT("RemoveUnusedResourceWarning"), Resource->Name);
 
 	// Removes 'produced but not used' warning.
-	Resource->GetParentDebugData().PassAccessCount++;
+	Resource->GetViewableDebugData().PassAccessCount++;
 
 	// Removes 'not used' warning.
 	Resource->GetDebugData().bIsActuallyUsedByPass = true;
 }
 
-bool FRDGUserValidation::TryMarkForClobber(FRDGParentResourceRef Resource) const
+bool FRDGUserValidation::TryMarkForClobber(FRDGViewableResource* Resource) const
 {
 	check(Resource);
-	FRDGParentResourceDebugData& DebugData = Resource->GetParentDebugData();
+	FRDGViewableResourceDebugData& DebugData = Resource->GetViewableDebugData();
 
 	const bool bClobber = !DebugData.bHasBeenClobbered && !Resource->bExternal && IsDebugAllowedForResource(Resource->Name);
 
@@ -475,7 +475,7 @@ void FRDGUserValidation::ValidateGetPooledBuffer(FRDGBufferRef Buffer) const
 	checkf(Buffer->bExternal, TEXT("GetPooledBuffer called on buffer %s, but it is not external. Call PreallocateBuffer or register as an external buffer instead."), Buffer->Name);
 }
 
-void FRDGUserValidation::ValidateSetAccessFinal(FRDGParentResourceRef Resource, ERHIAccess AccessFinal)
+void FRDGUserValidation::ValidateSetAccessFinal(FRDGViewableResource* Resource, ERHIAccess AccessFinal)
 {
 	check(Resource);
 	check(AccessFinal != ERHIAccess::Unknown && IsValidAccess(AccessFinal));
@@ -483,20 +483,20 @@ void FRDGUserValidation::ValidateSetAccessFinal(FRDGParentResourceRef Resource,
 	checkf(!Resource->bFinalizedAccess, TEXT("Cannot set final access on finalized resource %s."), Resource->Name);
 }
 
-void FRDGUserValidation::ValidateFinalize(FRDGParentResourceRef Resource, ERHIAccess AccessFinal, FRDGPassHandle FinalizePass)
+void FRDGUserValidation::ValidateFinalize(FRDGViewableResource* Resource, ERHIAccess AccessFinal, FRDGPassHandle FinalizePass)
 {
 	check(Resource);
 	check(AccessFinal != ERHIAccess::Unknown && IsValidAccess(AccessFinal));
 	checkf(IsReadOnlyAccess(AccessFinal), TEXT("Cannot convert resource %s to untracked with access %s. Access must be read-only."), Resource->Name, *GetRHIAccessName(AccessFinal));
 	checkf(Resource->bExternal || Resource->bExtracted, TEXT("Cannot convert resource %s to untracked unless it is first extracted or made external."), Resource->Name);
-	Resource->GetParentDebugData().FinalizePass = FinalizePass;
+	Resource->GetViewableDebugData().FinalizePass = FinalizePass;
 }
 
-void FRDGUserValidation::ValidateFinalizedAccess(FRDGParentResourceRef Resource, ERHIAccess Access, const FRDGPass* Pass)
+void FRDGUserValidation::ValidateFinalizedAccess(FRDGViewableResource* Resource, ERHIAccess Access, const FRDGPass* Pass)
 {
-	ensureMsgf(EnumHasAnyFlags(Resource->AccessFinal, Access),
+	ensureMsgf(EnumHasAnyFlags(Resource->FinalizedAccessMask, Access),
 		TEXT("Resource %s was finalized with access %s, but is being used in pass %s with access %s. Any future pass must use a subset of the finalized access state."),
-		Resource->Name, *GetRHIAccessName(Resource->AccessFinal), Pass->GetName(), *GetRHIAccessName(Access), *GetRHIPipelineName(Pass->GetPipeline()));
+		Resource->Name, *GetRHIAccessName(Resource->FinalizedAccessMask), Pass->GetName(), *GetRHIAccessName(Access), *GetRHIPipelineName(Pass->GetPipeline()));
 
 #if 0 // TODO: Need to account for read-only resources. 
 	ensureMsgf(Pass->GetPipeline() == ERHIPipeline::Graphics,
@@ -548,11 +548,11 @@ void FRDGUserValidation::ValidateAddPass(const FRDGPass* Pass, bool bSkipPassAcc
 	const bool bIsAnyCompute = EnumHasAnyFlags(PassFlags, ERDGPassFlags::Compute | ERDGPassFlags::AsyncCompute);
 	const bool bSkipRenderPass = EnumHasAnyFlags(PassFlags, ERDGPassFlags::SkipRenderPass);
 
-	const auto MarkAsProduced = [&](FRDGParentResourceRef Resource)
+	const auto MarkAsProduced = [&](FRDGViewableResource* Resource)
 	{
 		if (!bSkipPassAccessMarking)
 		{
-			auto& Debug = Resource->GetParentDebugData();
+			auto& Debug = Resource->GetViewableDebugData();
 			if (!Debug.FirstProducer)
 			{
 				Debug.FirstProducer = Pass;
@@ -573,7 +573,7 @@ void FRDGUserValidation::ValidateAddPass(const FRDGPass* Pass, bool bSkipPassAcc
 		MarkAsProduced(Buffer);
 	};
 
-	const auto MarkAsConsumed = [&] (FRDGParentResourceRef Resource)
+	const auto MarkAsConsumed = [&] (FRDGViewableResource* Resource)
 	{
 		ensureMsgf(Resource->bProduced || Resource->bExternal || Resource->bQueuedForUpload,
 			TEXT("Pass %s has a read dependency on %s, but it was never written to."),
@@ -581,7 +581,7 @@ void FRDGUserValidation::ValidateAddPass(const FRDGPass* Pass, bool bSkipPassAcc
 
 		if (!bSkipPassAccessMarking)
 		{
-			Resource->GetParentDebugData().PassAccessCount++;
+			Resource->GetViewableDebugData().PassAccessCount++;
 		}
 	};
 
@@ -597,7 +597,7 @@ void FRDGUserValidation::ValidateAddPass(const FRDGPass* Pass, bool bSkipPassAcc
 
 	bool bCanProduce = false;
 
-	const auto CheckResourceAccess = [&](FRDGParentResourceRef Resource, ERHIAccess Access)
+	const auto CheckResourceAccess = [&](FRDGViewableResource* Resource, ERHIAccess Access)
 	{
 		checkf(bIsCopy || !EnumHasAnyFlags(Access, AccessMaskCopy), TEXT("Pass '%s' uses resource '%s' with access '%s' containing states which require the 'ERDGPass::Copy' flag."), Pass->GetName(), Resource->Name, *GetRHIAccessName(Access));
 		checkf(bIsAnyCompute || !EnumHasAnyFlags(Access, AccessMaskCompute), TEXT("Pass '%s' uses resource '%s' with access '%s' containing states which require the 'ERDGPass::Compute' or 'ERDGPassFlags::AsyncCompute' flag."), Pass->GetName(), Resource->Name, *GetRHIAccessName(Access));
@@ -865,12 +865,12 @@ void FRDGUserValidation::ValidateExecuteEnd()
 
 	if (GRDGDebug)
 	{
-		auto ValidateResourceAtExecuteEnd = [](const FRDGParentResourceRef Resource)
+		auto ValidateResourceAtExecuteEnd = [](const FRDGViewableResource* Resource)
 		{
 			check(Resource->ReferenceCount == Resource->bExtracted ? 1 : 0);
 
-			const auto& ParentDebugData = Resource->GetParentDebugData();
-			const bool bProducedButNeverUsed = ParentDebugData.PassAccessCount == 1 && ParentDebugData.FirstProducer;
+			const auto& ViewableDebugData = Resource->GetViewableDebugData();
+			const bool bProducedButNeverUsed = ViewableDebugData.PassAccessCount == 1 && ViewableDebugData.FirstProducer;
 
 			if (bProducedButNeverUsed)
 			{
@@ -878,7 +878,7 @@ void FRDGUserValidation::ValidateExecuteEnd()
 
 				EmitRDGWarningf(
 					TEXT("Resource %s has been produced by the pass %s, but never used by another pass."),
-					Resource->Name, ParentDebugData.FirstProducer->GetName());
+					Resource->Name, ViewableDebugData.FirstProducer->GetName());
 			}
 		};
 
@@ -886,23 +886,23 @@ void FRDGUserValidation::ValidateExecuteEnd()
 		{
 			ValidateResourceAtExecuteEnd(Texture);
 
-			const auto& ParentDebugData = Texture->GetParentDebugData();
+			const auto& ViewableDebugData = Texture->GetViewableDebugData();
 			const auto& TextureDebugData = Texture->GetTextureDebugData();
 
-			const bool bHasBeenProducedByGraph = !Texture->bExternal && ParentDebugData.PassAccessCount > 0;
+			const bool bHasBeenProducedByGraph = !Texture->bExternal && ViewableDebugData.PassAccessCount > 0;
 
 			if (bHasBeenProducedByGraph && !TextureDebugData.bHasNeededUAV && EnumHasAnyFlags(Texture->Desc.Flags, TexCreate_UAV))
 			{
 				EmitRDGWarningf(
 					TEXT("Resource %s first produced by the pass %s had the TexCreate_UAV flag, but no UAV has been used."),
-					Texture->Name, ParentDebugData.FirstProducer->GetName());
+					Texture->Name, ViewableDebugData.FirstProducer->GetName());
 			}
 
 			if (bHasBeenProducedByGraph && !TextureDebugData.bHasBeenBoundAsRenderTarget && EnumHasAnyFlags(Texture->Desc.Flags, TexCreate_RenderTargetable))
 			{
 				EmitRDGWarningf(
 					TEXT("Resource %s first produced by the pass %s had the TexCreate_RenderTargetable flag, but has never been bound as a render target of a pass."),
-					Texture->Name, ParentDebugData.FirstProducer->GetName());
+					Texture->Name, ViewableDebugData.FirstProducer->GetName());
 			}
 		}
 
@@ -1154,16 +1154,16 @@ void FRDGBarrierValidation::ValidateBarrierBatchBegin(const FRDGPass* Pass, cons
 
 		for (int32 Index = 0; Index < Batch.Transitions.Num(); ++Index)
 		{
-			FRDGParentResourceRef Resource = Batch.DebugTransitionResources[Index];
+			FRDGViewableResource* Resource = Batch.DebugTransitionResources[Index];
 			const FRHITransitionInfo& Transition = Batch.Transitions[Index];
 
-			if (Resource->Type == ERDGParentResourceType::Texture)
+			if (Resource->Type == ERDGViewableResourceType::Texture)
 			{
 				ResourceMap->Textures.FindOrAdd(static_cast<FRDGTextureRef>(Resource)).Add(Transition);
 			}
 			else
 			{
-				check(Resource->Type == ERDGParentResourceType::Buffer);
+				check(Resource->Type == ERDGViewableResourceType::Buffer);
 				ResourceMap->Buffers.Emplace(static_cast<FRDGBufferRef>(Resource), Transition);
 			}
 		}
@@ -1195,7 +1195,7 @@ void FRDGBarrierValidation::ValidateBarrierBatchBegin(const FRDGPass* Pass, cons
 		const FRHITransientAliasingInfo& Info = KeyValue.Value;
 		if (Info.IsAcquire())
 		{
-			FRDGParentResourceRef Resource = KeyValue.Key;
+			FRDGViewableResource* Resource = KeyValue.Key;
 
 			if (IsDebugAllowedForResource(Resource->Name))
 			{
@@ -1329,7 +1329,7 @@ void FRDGBarrierValidation::ValidateBarrierBatchEnd(const FRDGPass* Pass, const
 			const FRHITransientAliasingInfo& Info = KeyValue.Value;
 			if (Info.IsDiscard())
 			{
-				FRDGParentResourceRef Resource = KeyValue.Key;
+				FRDGViewableResource* Resource = KeyValue.Key;
 
 				if (IsDebugAllowedForResource(Resource->Name))
 				{
@@ -1341,491 +1341,4 @@ void FRDGBarrierValidation::ValidateBarrierBatchEnd(const FRDGPass* Pass, const
 	}
 }
 
-namespace
-{
-	const TCHAR* RasterColorName = TEXT("#ff7070");
-	const TCHAR* ComputeColorName = TEXT("#70b8ff");
-	const TCHAR* AsyncComputeColorName = TEXT("#70ff99");
-	const TCHAR* CopyColorName = TEXT("#ffdb70");
-	const TCHAR* TextureColorAttributes = TEXT("color=\"#5800a1\", fontcolor=\"#5800a1\"");
-	const TCHAR* BufferColorAttributes = TEXT("color=\"#007309\", fontcolor=\"#007309\"");
-	const TCHAR* AliasColorAttributes = TEXT("color=\"#00ff00\", fontcolor=\"#00ff00\"");
-	const TCHAR* AllPipelinesColorName = TEXT("#f170ff");
-
-	const TCHAR* GetPassColorName(ERDGPassFlags Flags)
-	{
-		if (EnumHasAnyFlags(Flags, ERDGPassFlags::Raster))
-		{
-			return RasterColorName;
-		}
-		if (EnumHasAnyFlags(Flags, ERDGPassFlags::Compute))
-		{
-			return ComputeColorName;
-		}
-		if (EnumHasAnyFlags(Flags, ERDGPassFlags::AsyncCompute))
-		{
-			return AsyncComputeColorName;
-		}
-		if (EnumHasAnyFlags(Flags, ERDGPassFlags::Copy))
-		{
-			return CopyColorName;
-		}
-		return TEXT("#ffffff");
-	}
-
-	FString GetSubresourceStateLabel(FRDGSubresourceState State)
-	{
-		const ERHIPipeline Pipelines = State.GetPipelines();
-		const TCHAR* FontColor = nullptr;
-		switch (Pipelines)
-		{
-		default:
-			checkNoEntry();
-		case ERHIPipeline::Graphics:
-			FontColor = RasterColorName;
-			break;
-		case ERHIPipeline::AsyncCompute:
-			FontColor = AsyncComputeColorName;
-			break;
-		case ERHIPipeline::All:
-			FontColor = AllPipelinesColorName;
-			break;
-		}
-		return FString::Printf(TEXT("<font color=\"%s\">%s</font>"), FontColor, *GetRHIAccessName(State.Access));
-	}
-}
-
-FString FRDGLogFile::GetProducerName(FRDGPassHandle PassHandle)
-{
-	check(PassHandle.IsValid());
-	return GetNodeName(PassHandle);
-}
-
-FString FRDGLogFile::GetConsumerName(FRDGPassHandle PassHandle)
-{
-	check(PassHandle.IsValid());
-	return GetNodeName(PassHandle);
-}
-
-FString FRDGLogFile::GetNodeName(FRDGPassHandle PassHandle)
-{
-	PassesReferenced.Add(PassHandle);
-	return FString::Printf(TEXT("P%d"), PassHandle.GetIndex());
-}
-
-FString FRDGLogFile::GetNodeName(const FRDGTexture* Texture)
-{
-	return FString::Printf(TEXT("T%d"), Textures.AddUnique(Texture));
-}
-
-FString FRDGLogFile::GetNodeName(const FRDGBuffer* Buffer)
-{
-	return FString::Printf(TEXT("B%d"), Buffers.AddUnique(Buffer));
-}
-
-void FRDGLogFile::AddLine(const FString& Line)
-{
-	File += Indentation + Line + TEXT("\n");
-}
-
-void FRDGLogFile::AddBraceBegin()
-{
-	AddLine(TEXT("{"));
-	Indentation += TEXT("\t");
-}
-
-void FRDGLogFile::AddBraceEnd()
-{
-	const bool bSuccess = Indentation.RemoveFromEnd(TEXT("\t"));
-	check(bSuccess);
-
-	AddLine(TEXT("}"));
-}
-
-void FRDGLogFile::Begin(const FRDGEventName& InGraphName)
-{
-	if (GRDGDumpGraph)
-	{
-		if (IsImmediateMode())
-		{
-			UE_LOG(LogRDG, Warning, TEXT("Dump graph (%d) requested, but immediate mode is enabled. Skipping."), GRDGDumpGraph);
-			return;
-		}
-
-		check(File.IsEmpty());
-
-		GraphName = InGraphName.GetTCHAR();
-
-		if (GraphName.IsEmpty())
-		{
-			const int32 UnknownGraphIndex = GRDGDumpGraphUnknownCount++;
-			GraphName = FString::Printf(TEXT("Unknown%d"), UnknownGraphIndex);
-		}
-
-		AddLine(TEXT("digraph RDG"));
-		AddBraceBegin();
-		AddLine(TEXT("rankdir=LR; labelloc=\"t\""));
-
-		bOpen = true;
-	}
-}
-
-void FRDGLogFile::End()
-{
-	if (!GRDGDumpGraph || !bOpen)
-	{
-		return;
-	}
-
-	TArray<FRDGPassHandle> PassesGraphics;
-	TArray<FRDGPassHandle> PassesAsyncCompute;
-
-	for (FRDGPassHandle PassHandle = Passes.Begin(); PassHandle != Passes.End(); ++PassHandle)
-	{
-		const FRDGPass* Pass = Passes[PassHandle];
-		const ERHIPipeline Pipeline = Pass->GetPipeline();
-
-		switch (Pipeline)
-		{
-		case ERHIPipeline::Graphics:
-			PassesGraphics.Add(PassHandle);
-			break;
-		case ERHIPipeline::AsyncCompute:
-			PassesAsyncCompute.Add(PassHandle);
-			break;
-		default:
-			checkNoEntry();
-		}
-	}
-
-	if (GRDGDumpGraph == RDG_DUMP_GRAPH_TRACKS)
-	{
-		FRDGPassHandle PrevPassesByPipeline[uint32(ERHIPipeline::Num)];
-
-		for (FRDGPassHandle PassHandle = Passes.Begin(); PassHandle != Passes.End(); ++PassHandle)
-		{
-			const FRDGPass* Pass = Passes[PassHandle];
-
-			if (!EnumHasAnyFlags(Pass->GetFlags(), ERDGPassFlags::Copy | ERDGPassFlags::Raster | ERDGPassFlags::Compute | ERDGPassFlags::AsyncCompute))
-			{
-				continue;
-			}
-
-			ERHIPipeline PassPipeline = Pass->GetPipeline();
-			checkf(FMath::IsPowerOfTwo(uint32(PassPipeline)), TEXT("This logic doesn't handle multi-pipe passes."));
-			uint32 PipeIndex = FMath::FloorLog2(uint32(PassPipeline));
-
-			FRDGPassHandle& PrevPassInPipelineHandle = PrevPassesByPipeline[PipeIndex];
-
-			if (PrevPassInPipelineHandle.IsValid())
-			{
-				AddLine(FString::Printf(TEXT("\"%s\" -> \"%s\" [style=\"filled\", penwidth=2, color=\"%s\"]"),
-					*GetNodeName(PrevPassInPipelineHandle), *GetNodeName(PassHandle), GetPassColorName(Pass->GetFlags())));
-			}
-
-			if (Pass->GetPipeline() == ERHIPipeline::AsyncCompute)
-			{
-				const auto AddCrossPipelineEdge = [&](FRDGPassHandle PassBefore, FRDGPassHandle PassAfter)
-				{
-					AddLine(FString::Printf(TEXT("\"%s\" -> \"%s\" [penwidth=5, style=\"dashed\" color=\"#f003fc\"]"),
-						*GetNodeName(PassBefore), *GetNodeName(PassAfter)));
-				};
-
-				if (Pass->IsAsyncComputeBegin())
-				{
-					AddCrossPipelineEdge(Pass->GetGraphicsForkPass(), PassHandle);
-				}
-
-				if (Pass->IsAsyncComputeEnd())
-				{
-					AddCrossPipelineEdge(PassHandle, Pass->GetGraphicsJoinPass());
-				}
-			}
-
-			PrevPassInPipelineHandle = PassHandle;
-		}
-	}
-	else if (GRDGDumpGraph == RDG_DUMP_GRAPH_PRODUCERS)
-	{
-		for (FRDGPassHandle PassHandle = Passes.Begin(); PassHandle < Passes.Last(); ++PassHandle)
-		{
-			const FRDGPass* Pass = Passes[PassHandle];
-
-			for (FRDGPassHandle ProducerHandle : Pass->GetProducers())
-			{
-				const FRDGPass* Producer = Passes[ProducerHandle];
-
-				File += FString::Printf(TEXT("\t\"%s\" -> \"%s\" [penwidth=2, color=\"%s:%s\"]\n"),
-					*GetNodeName(ProducerHandle), *GetNodeName(PassHandle), GetPassColorName(Pass->GetFlags()), GetPassColorName(Producer->GetFlags()));
-			}
-		}
-	}
-
-	AddLine(TEXT("subgraph Passes"));
-	AddBraceBegin();
-
-	const auto AddPass = [&](FRDGPassHandle PassHandle)
-	{
-		if (!PassesReferenced.Contains(PassHandle))
-		{
-			return;
-		}
-
-		const FRDGPass* Pass = Passes[PassHandle];
-		const TCHAR* Style = Pass->IsCulled() ? TEXT("dashed") : TEXT("filled");
-		FString PassName = FString::Printf(TEXT("[%d]: %s"), PassHandle.GetIndex(), Pass->GetName());
-
-		if (Pass->GetParameters().HasExternalOutputs())
-		{
-			PassName += TEXT("\n(Has External UAVs)");
-		}
-
-		AddLine(FString::Printf(TEXT("\"%s\" [shape=box, style=%s, label=\"%s\", color=\"%s\"]"), *GetNodeName(PassHandle), Style, *PassName, GetPassColorName(Pass->GetFlags())));
-	};
-
-	{
-		uint32 RenderTargetClusterCount = 0;
-
-		for (FRDGPassHandle PassHandle : PassesGraphics)
-		{
-			const FRDGPass* Pass = Passes[PassHandle];
-
-			if (Pass->IsMergedRenderPassBegin())
-			{
-				const uint32 RenderTargetClusterIndex = RenderTargetClusterCount++;
-
-				AddLine(FString::Printf(TEXT("subgraph cluster_%d"), RenderTargetClusterIndex));
-				AddBraceBegin();
-				AddLine(TEXT("style=filled;color=\"#ffe0e0\";fontcolor=\"#aa0000\";label=\"Render Pass Merge\";fontsize=10"));
-			}
-
-			AddPass(PassHandle);
-
-			if (Pass->IsMergedRenderPassEnd())
-			{
-				AddBraceEnd();
-			}
-		}
-	}
-
-	for (FRDGPassHandle PassHandle : PassesAsyncCompute)
-	{
-		AddPass(PassHandle);
-	}
-
-	AddBraceEnd();
-
-	AddLine(TEXT("subgraph Textures"));
-	AddBraceBegin();
-	for (const FRDGTexture* Texture : Textures)
-	{
-		FString Line = FString::Printf(TEXT("\"%s\" [shape=oval, %s, label=\"%s"), *GetNodeName(Texture), TextureColorAttributes, Texture->Name);
-		if (Texture->IsExternal())
-		{
-			Line += TEXT("\n(External)");
-		}
-		Line += TEXT("\"]");
-		AddLine(Line);
-	}
-	AddBraceEnd();
-
-	AddLine(TEXT("subgraph Buffers"));
-	AddBraceBegin();
-	for (const FRDGBuffer* Buffer : Buffers)
-	{
-		FString Line = FString::Printf(TEXT("\"%s\" [shape=oval, %s, label=\"%s"), *GetNodeName(Buffer), BufferColorAttributes, Buffer->Name);
-		if (Buffer->IsExternal())
-		{
-			Line += TEXT("\n(External)");
-		}
-		Line += TEXT("\"]");
-		AddLine(Line);
-	}
-	AddBraceEnd();
-
-	uint32 NumPassesActive = 0;
-	uint32 NumPassesCulled = 0;
-
-	Passes.Enumerate([&](const FRDGPass* Pass)
-	{
-		if (Pass->IsCulled())
-		{
-			NumPassesCulled++;
-		}
-		else
-		{
-			NumPassesActive++;
-		}
-	});
-
-	AddLine(FString::Printf(TEXT("label=\"%s [Active Passes: %d, Culled Passes: %d, Textures: %d, Buffers: %d]\""), *GraphName, NumPassesActive, NumPassesCulled, Textures.Num(), Buffers.Num()));
-
-	AddBraceEnd();
-	check(Indentation.IsEmpty());
-
-	const TCHAR* DumpType = TEXT("");
-
-	switch (GRDGDumpGraph)
-	{
-	case RDG_DUMP_GRAPH_RESOURCES:
-		DumpType = TEXT("_resources");
-		break;
-	case RDG_DUMP_GRAPH_PRODUCERS:
-		DumpType = TEXT("_producers");
-		break;
-	case RDG_DUMP_GRAPH_TRACKS:
-		DumpType = TEXT("_tracks");
-		break;
-	}
-
-	FFileHelper::SaveStringToFile(File, *(FPaths::ProjectLogDir() / FString::Printf(TEXT("RDG_%s%s.gv"), *GraphName, DumpType)));
-
-	bOpen = false;
-}
-
-bool FRDGLogFile::IncludeTransitionEdgeInGraph(FRDGPassHandle Pass) const
-{
-	return Pass.IsValid() && !Passes[Pass]->IsSentinel();
-}
-
-bool FRDGLogFile::IncludeTransitionEdgeInGraph(FRDGPassHandle PassBefore, FRDGPassHandle PassAfter) const
-{
-	return IncludeTransitionEdgeInGraph(PassBefore) && IncludeTransitionEdgeInGraph(PassAfter) && PassBefore < PassAfter;
-}
-
-void FRDGLogFile::AddFirstEdge(const FRDGTextureRef Texture, FRDGPassHandle FirstPass)
-{
-	if (GRDGDumpGraph == RDG_DUMP_GRAPH_RESOURCES && bOpen && IncludeTransitionEdgeInGraph(FirstPass) && IsDebugAllowedForResource(Texture->Name))
-	{
-		AddLine(FString::Printf(TEXT("\"%s\" -> \"%s\" [%s]"),
-			*GetNodeName(Texture),
-			*GetNodeName(FirstPass),
-			TextureColorAttributes));
-	}
-}
-
-void FRDGLogFile::AddFirstEdge(const FRDGBufferRef Buffer, FRDGPassHandle FirstPass)
-{
-	if (GRDGDumpGraph == RDG_DUMP_GRAPH_RESOURCES && bOpen && IncludeTransitionEdgeInGraph(FirstPass) && IsDebugAllowedForResource(Buffer->Name))
-	{
-		AddLine(FString::Printf(TEXT("\"%s\" -> \"%s\" [%s]"),
-			*GetNodeName(Buffer),
-			*GetNodeName(FirstPass),
-			BufferColorAttributes));
-	}
-}
-
-
-void FRDGLogFile::AddAliasEdge(const FRDGTextureRef TextureBefore, FRDGPassHandle BeforePass, const FRDGTextureRef TextureAfter, FRDGPassHandle AfterPass)
-{
-	if (GRDGDumpGraph == RDG_DUMP_GRAPH_RESOURCES && bOpen && IncludeTransitionEdgeInGraph(BeforePass, AfterPass) && IsDebugAllowedForResource(TextureBefore->Name) && IsDebugAllowedForResource(TextureAfter->Name))
-	{
-		AddLine(FString::Printf(TEXT("\"%s\" -> \"%s\" [%s, label=<Alias: <b>%s -&gt; %s</b>>]"),
-			*GetProducerName(BeforePass),
-			*GetConsumerName(AfterPass),
-			AliasColorAttributes,
-			TextureBefore->Name,
-			TextureAfter->Name));
-	}
-}
-
-void FRDGLogFile::AddAliasEdge(const FRDGBufferRef BufferBefore, FRDGPassHandle BeforePass, const FRDGBufferRef BufferAfter, FRDGPassHandle AfterPass)
-{
-	if (GRDGDumpGraph == RDG_DUMP_GRAPH_RESOURCES && bOpen && IncludeTransitionEdgeInGraph(BeforePass, AfterPass) && IsDebugAllowedForResource(BufferBefore->Name) && IsDebugAllowedForResource(BufferAfter->Name))
-	{
-		AddLine(FString::Printf(TEXT("\"%s\" -> \"%s\" [%s, label=<Alias: <b>%s -&gt; %s</b>>]"),
-			*GetProducerName(BeforePass),
-			*GetConsumerName(AfterPass),
-			AliasColorAttributes,
-			BufferBefore->Name,
-			BufferAfter->Name));
-	}
-}
-
-void FRDGLogFile::AddTransitionEdge(FRDGPassHandle PassHandle, const FRDGSubresourceState& StateBefore, const FRDGSubresourceState& StateAfter, const FRDGTextureRef Texture)
-{
-	if (GRDGDumpGraph == RDG_DUMP_GRAPH_RESOURCES && bOpen && IsDebugAllowedForResource(Texture->Name))
-	{
-		if (IncludeTransitionEdgeInGraph(StateBefore.GetLastPass(), StateAfter.GetFirstPass()) && FRDGSubresourceState::IsTransitionRequired(StateBefore, StateAfter))
-		{
-			AddLine(FString::Printf(TEXT("\"%s\" -> \"%s\" [%s, label=<%s: <b>%s -&gt; %s</b>>]"),
-				*GetProducerName(StateBefore.GetLastPass()),
-				*GetConsumerName(StateAfter.GetFirstPass()),
-				TextureColorAttributes,
-				Texture->Name,
-				*GetSubresourceStateLabel(StateBefore),
-				*GetSubresourceStateLabel(StateAfter)));
-		}
-		else if (IncludeTransitionEdgeInGraph(StateBefore.LogFilePass, PassHandle))
-		{
-			AddLine(FString::Printf(TEXT("\"%s\" -> \"%s\" [%s, label=<%s: <b>%s</b>>]"),
-				*GetProducerName(StateBefore.LogFilePass),
-				*GetConsumerName(PassHandle),
-				TextureColorAttributes,
-				Texture->Name,
-				*GetSubresourceStateLabel(StateBefore)));
-		}
-
-		StateAfter.LogFilePass = PassHandle;
-	}
-}
-
-void FRDGLogFile::AddTransitionEdge(FRDGPassHandle PassHandle, const FRDGSubresourceState& StateBefore, const FRDGSubresourceState& StateAfter, const FRDGTextureRef Texture, FRDGTextureSubresource Subresource)
-{
-	if (GRDGDumpGraph == RDG_DUMP_GRAPH_RESOURCES && bOpen && IsDebugAllowedForResource(Texture->Name))
-	{
-		if (IncludeTransitionEdgeInGraph(StateBefore.GetLastPass(), StateAfter.GetFirstPass()) && FRDGSubresourceState::IsTransitionRequired(StateBefore, StateAfter))
-		{
-			AddLine(FString::Printf(TEXT("\"%s\" -> \"%s\" [%s, label=<%s[%d][%d][%d]: <b>%s -&gt; %s</b>>]"),
-				*GetProducerName(StateBefore.GetLastPass()),
-				*GetConsumerName(StateAfter.GetFirstPass()),
-				TextureColorAttributes,
-				Texture->Name,
-				Subresource.MipIndex, Subresource.ArraySlice, Subresource.PlaneSlice,
-				*GetSubresourceStateLabel(StateBefore),
-				*GetSubresourceStateLabel(StateAfter)));
-		}
-		else if (IncludeTransitionEdgeInGraph(StateBefore.LogFilePass, PassHandle))
-		{
-			AddLine(FString::Printf(TEXT("\"%s\" -> \"%s\" [%s, label=<%s[%d][%d][%d]: <b>%s</b>>]"),
-				*GetProducerName(StateBefore.LogFilePass),
-				*GetConsumerName(PassHandle),
-				TextureColorAttributes,
-				Texture->Name,
-				Subresource.MipIndex, Subresource.ArraySlice, Subresource.PlaneSlice,
-				*GetSubresourceStateLabel(StateBefore)));
-		}
-
-		StateAfter.LogFilePass = PassHandle;
-	}
-}
-
-void FRDGLogFile::AddTransitionEdge(FRDGPassHandle PassHandle, const FRDGSubresourceState& StateBefore, const FRDGSubresourceState& StateAfter, const FRDGBufferRef Buffer)
-{
-	if (GRDGDumpGraph == RDG_DUMP_GRAPH_RESOURCES && bOpen && IsDebugAllowedForResource(Buffer->Name))
-	{
-		if (IncludeTransitionEdgeInGraph(StateBefore.GetLastPass(), StateAfter.GetFirstPass()) && FRDGSubresourceState::IsTransitionRequired(StateBefore, StateAfter))
-		{
-			AddLine(FString::Printf(TEXT("\"%s\" -> \"%s\" [%s, label=<%s: <b>%s -&gt; %s</b>>]"),
-				*GetProducerName(StateBefore.GetLastPass()),
-				*GetConsumerName(StateAfter.GetFirstPass()),
-				BufferColorAttributes,
-				Buffer->Name,
-				*GetSubresourceStateLabel(StateBefore),
-				*GetSubresourceStateLabel(StateAfter)));
-		}
-		else if (IncludeTransitionEdgeInGraph(StateBefore.LogFilePass, PassHandle))
-		{
-			AddLine(FString::Printf(TEXT("\"%s\" -> \"%s\" [%s, label=<%s: <b>%s</b>>]"),
-				*GetProducerName(StateBefore.LogFilePass),
-				*GetConsumerName(PassHandle),
-				BufferColorAttributes,
-				Buffer->Name,
-				*GetSubresourceStateLabel(StateBefore)));
-		}
-
-		StateAfter.LogFilePass = PassHandle;
-	}
-}
-
 #endif
\ No newline at end of file
diff --git a/Engine/Source/Runtime/RenderCore/Private/RenderTargetPool.cpp b/Engine/Source/Runtime/RenderCore/Private/RenderTargetPool.cpp
--- a/Engine/Source/Runtime/RenderCore/Private/RenderTargetPool.cpp
+++ b/Engine/Source/Runtime/RenderCore/Private/RenderTargetPool.cpp
@@ -55,7 +55,7 @@ static uint32 ComputeSizeInKB(FPooledRenderTarget& Element)
 	return (Element.ComputeMemorySize() + 1023) / 1024;
 }
 
-TRefCountPtr<IPooledRenderTarget> FRenderTargetPool::FindFreeElementInternal(FRHITextureCreateInfo Desc, const TCHAR* Name, bool bResetStateToUnknown)
+TRefCountPtr<IPooledRenderTarget> FRenderTargetPool::FindFreeElement(FRHITextureCreateInfo Desc, const TCHAR* Name)
 {
 	FPooledRenderTarget* Found = 0;
 	uint32 FoundIndex = -1;
@@ -101,7 +101,6 @@ TRefCountPtr<IPooledRenderTarget> FRenderTargetPool::FindFreeElementInternal(FRH
 
 		Found = new FPooledRenderTarget(
 			RHICreateTexture(CreateDesc),
-			AccessInitial,
 			Translate(CreateDesc),
 			this);
 
@@ -137,11 +136,6 @@ TRefCountPtr<IPooledRenderTarget> FRenderTargetPool::FindFreeElementInternal(FRH
 	Found->Desc.DebugName = Name;
 	Found->UnusedForNFrames = 0;
 
-	if (bResetStateToUnknown)
-	{
-		Found->PooledTexture.Reset();
-	}
-
 #if !(UE_BUILD_SHIPPING || UE_BUILD_TEST)
 	RHIBindDebugLabelName(Found->GetRHI(), Name);
 #endif
@@ -195,20 +189,13 @@ bool FRenderTargetPool::FindFreeElement(const FRHITextureCreateInfo& Desc, TRefC
 		}
 	}
 
-	const bool bResetStateToUnknown = true;
-	Out = FindFreeElementInternal(Desc, Name, bResetStateToUnknown);
+	Out = FindFreeElement(Desc, Name);
 	return false;
 }
 
-TRefCountPtr<IPooledRenderTarget> FRenderTargetPool::FindFreeElement(const FRHITextureCreateInfo& Desc, const TCHAR* Name)
-{
-	const bool bResetStateToUnknown = true;
-	return FindFreeElementInternal(Desc, Name, bResetStateToUnknown);
-}
-
 void FRenderTargetPool::CreateUntrackedElement(const FPooledRenderTargetDesc& Desc, TRefCountPtr<IPooledRenderTarget>& Out, const FSceneRenderTargetItem& Item)
 {
-	FPooledRenderTarget* Result = new FPooledRenderTarget(Item.GetRHI(), ERHIAccess::Unknown, Desc, nullptr);
+	FPooledRenderTarget* Result = new FPooledRenderTarget(Item.GetRHI(), Desc, nullptr);
 	Result->RenderTargetItem = Item;
 	Out = Result;
 }
diff --git a/Engine/Source/Runtime/RenderCore/Public/RenderGraphBuilder.h b/Engine/Source/Runtime/RenderCore/Public/RenderGraphBuilder.h
--- a/Engine/Source/Runtime/RenderCore/Public/RenderGraphBuilder.h
+++ b/Engine/Source/Runtime/RenderCore/Public/RenderGraphBuilder.h
@@ -12,8 +12,6 @@
 #include "ShaderParameterMacros.h"
 #include "ProfilingDebugging/CsvProfiler.h"
 
-class FRDGLogFile;
-
 /** Use the render graph builder to build up a graph of passes and then call Execute() to process them. Resource barriers
  *  and lifetimes are derived from _RDG_ parameters in the pass parameter struct provided to each AddPass call. The resulting
  *  graph is compiled, culled, and executed in Execute(). The builder should be created on the stack and executed prior to
@@ -363,7 +361,6 @@ PRAGMA_ENABLE_DEPRECATION_WARNINGS
 	//////////////////////////////////////////////////////////////////////////
 
 private:
-	static const ERHIAccess kDefaultAccessFinal = ERHIAccess::SRVMask;
 	static const char* const kDefaultUnaccountedCSVStat;
 
 	FRHIAsyncComputeCommandListImmediate& RHICmdListAsyncCompute;
@@ -382,19 +379,6 @@ private:
 
 	void AddProloguePass();
 
-	FORCEINLINE FRDGPassHandle ClampToPrologue(FRDGPassHandle PassHandle) const
-	{
-		// Preserve null inputs as outputs. Null is the highest value.
-		return ProloguePassHandle.GetIndexUnchecked() > PassHandle.GetIndexUnchecked() ? ProloguePassHandle : PassHandle;
-	}
-
-	FORCEINLINE FRDGPassHandlesByPipeline ClampToPrologue(FRDGPassHandlesByPipeline PassHandles) const
-	{
-		PassHandles[ERHIPipeline::Graphics]     = ClampToPrologue(PassHandles[ERHIPipeline::Graphics]);
-		PassHandles[ERHIPipeline::AsyncCompute] = ClampToPrologue(PassHandles[ERHIPipeline::AsyncCompute]);
-		return PassHandles;
-	}
-
 	FORCEINLINE FRDGPass* GetProloguePass() const
 	{
 		return ProloguePass;
@@ -403,7 +387,7 @@ private:
 	/** Returns the graph prologue pass handle. */
 	FORCEINLINE FRDGPassHandle GetProloguePassHandle() const
 	{
-		return ProloguePassHandle;
+		return FRDGPassHandle(0);
 	}
 
 	/** Returns the graph epilogue pass handle. */
@@ -489,7 +473,15 @@ private:
 
 	/** Tracks external resources to their registered render graph counterparts for de-duplication. */
 	TSortedMap<FRHITexture*, FRDGTexture*, FRDGArrayAllocator> ExternalTextures;
-	TSortedMap<const FRDGPooledBuffer*, FRDGBuffer*, FRDGArrayAllocator> ExternalBuffers;
+	TSortedMap<FRDGPooledBuffer*, FRDGBuffer*, FRDGArrayAllocator> ExternalBuffers;
+
+	/** Tracks the latest RDG resource to own an alias of a pooled resource (multiple RDG resources can reference the same pooled resource). */
+	TMap<FRDGPooledTexture*, FRDGTexture*, FRDGSetAllocator> PooledTextureOwnershipMap;
+	TMap<FRDGPooledBuffer*, FRDGBuffer*, FRDGSetAllocator> PooledBufferOwnershipMap;
+
+	/** Array of all pooled references held during execution. */
+	TArray<TRefCountPtr<IPooledRenderTarget>, FRDGArrayAllocator> ActivePooledTextures;
+	TArray<TRefCountPtr<FRDGPooledBuffer>, FRDGArrayAllocator> ActivePooledBuffers;
 
 	/** Map of barrier batches begun from more than one pipe. */
 	TMap<FRDGBarrierBatchBeginId, FRDGBarrierBatchBegin*, FRDGSetAllocator> BarrierBatchMap;
@@ -503,7 +495,6 @@ private:
 	 *  the graph for culling purposes. The epilogue pass is added to the very end of the pass array for traversal
 	 *  purposes. The prologue does not need to participate in any graph traversal behavior.
 	 */
-	FRDGPassHandle ProloguePassHandle;
 	FRDGPass* ProloguePass = nullptr;
 	FRDGPass* EpiloguePass = nullptr;
 
@@ -605,12 +596,10 @@ private:
 	/** Array of all active parallel execute tasks. */
 	FGraphEventArray ParallelExecuteEvents;
 
-	/** Array of all pooled references held during execution. */
-	FRDGPooledTextureArray ActivePooledTextures;
-	FRDGPooledBufferArray ActivePooledBuffers;
+	TArray<FRHITrackedAccessInfo, FRDGArrayAllocator> EpilogueResourceAccesses;
 
 	/** Texture state used for intermediate operations. Held here to avoid re-allocating. */
-	FRDGTextureTransientSubresourceStateIndirect ScratchTextureState;
+	FRDGTextureSubresourceStateIndirect ScratchTextureState;
 
 	/** Current scope's async compute budget. This is passed on to every pass created. */
 	EAsyncComputeBudget AsyncComputeBudgetScope = EAsyncComputeBudget::EAll_4;
@@ -629,7 +618,6 @@ private:
 #if RDG_ENABLE_DEBUG
 	FRDGUserValidation UserValidation;
 	FRDGBarrierValidation BarrierValidation;
-	FRDGLogFile LogFile;
 #endif
 
 	/** Tracks whether we are in a scope of adding passes to the builder. Used to avoid recursion. */
@@ -657,6 +645,12 @@ private:
 	void Compile();
 	void Clear();
 
+	void SetRHI(FRDGTexture* Texture, IPooledRenderTarget* RenderTarget, FRDGPassHandle PassHandle);
+	void SetRHI(FRDGTexture* Texture, FRDGPooledTexture* PooledTexture, FRDGPassHandle PassHandle);
+	void SetRHI(FRDGTexture* Texture, FRHITransientTexture* TransientTexture, FRDGPassHandle PassHandle);
+	void SetRHI(FRDGBuffer* Buffer, FRDGPooledBuffer* PooledBuffer, FRDGPassHandle PassHandle);
+	void SetRHI(FRDGBuffer* Buffer, FRHITransientBuffer* TransientBuffer, FRDGPassHandle PassHandle);
+
 	void BeginResourcesRHI(FRDGPass* ResourcePass, FRDGPassHandle ExecutePassHandle);
 	void BeginResourceRHI(FRDGPassHandle, FRDGTexture* Texture);
 	void BeginResourceRHI(FRDGPassHandle, FRDGTextureSRV* SRV);
@@ -701,7 +695,7 @@ private:
 	void AddTransition(
 		FRDGPassHandle PassHandle,
 		FRDGTextureRef Texture,
-		const FRDGTextureTransientSubresourceStateIndirect& StateAfter);
+		const FRDGTextureSubresourceStateIndirect& StateAfter);
 
 	void AddTransition(
 		FRDGPassHandle PassHandle,
@@ -709,7 +703,7 @@ private:
 		FRDGSubresourceState StateAfter);
 
 	void AddTransition(
-		FRDGParentResource* Resource,
+		FRDGViewableResource* Resource,
 		FRDGSubresourceState StateBefore,
 		FRDGSubresourceState StateAfter,
 		const FRHITransitionInfo& TransitionInfo);
@@ -717,12 +711,12 @@ private:
 	void AddAliasingTransition(
 		FRDGPassHandle BeginPassHandle,
 		FRDGPassHandle EndPassHandle,
-		FRDGParentResourceRef Resource,
+		FRDGViewableResource* Resource,
 		const FRHITransientAliasingInfo& Info);
 
 	bool IsTransient(FRDGTextureRef Texture) const;
 	bool IsTransient(FRDGBufferRef Buffer) const;
-	bool IsTransientInternal(FRDGParentResourceRef Resource, bool bFastVRAM) const;
+	bool IsTransientInternal(FRDGViewableResource* Resource, bool bFastVRAM) const;
 
 	FRHIRenderPassInfo GetRenderPassInfo(const FRDGPass* Pass) const;
 
diff --git a/Engine/Source/Runtime/RenderCore/Public/RenderGraphBuilder.inl b/Engine/Source/Runtime/RenderCore/Public/RenderGraphBuilder.inl
--- a/Engine/Source/Runtime/RenderCore/Public/RenderGraphBuilder.inl
+++ b/Engine/Source/Runtime/RenderCore/Public/RenderGraphBuilder.inl
@@ -308,11 +308,6 @@ inline void FRDGBuilder::QueueTextureExtraction(FRDGTextureRef Texture, TRefCoun
 	}
 
 	ExtractedTextures.Emplace(Texture, OutTexturePtr);
-
-	if (Texture->AccessFinal == ERHIAccess::Unknown)
-	{
-		Texture->AccessFinal = kDefaultAccessFinal;
-	}
 }
 
 inline void FRDGBuilder::QueueBufferExtraction(FRDGBufferRef Buffer, TRefCountPtr<FRDGPooledBuffer>* OutBufferPtr)
@@ -326,11 +321,6 @@ inline void FRDGBuilder::QueueBufferExtraction(FRDGBufferRef Buffer, TRefCountPt
 	Buffer->bCulled = false;
 	Buffer->bForceNonTransient = true;
 	ExtractedBuffers.Emplace(Buffer, OutBufferPtr);
-
-	if (Buffer->AccessFinal == ERHIAccess::Unknown)
-	{
-		Buffer->AccessFinal = kDefaultAccessFinal;
-	}
 }
 
 inline void FRDGBuilder::QueueBufferExtraction(FRDGBufferRef Buffer, TRefCountPtr<FRDGPooledBuffer>* OutBufferPtr, ERHIAccess AccessFinal)
@@ -367,13 +357,13 @@ inline const TRefCountPtr<FRDGPooledBuffer>& FRDGBuilder::GetPooledBuffer(FRDGBu
 inline void FRDGBuilder::SetTextureAccessFinal(FRDGTextureRef Texture, ERHIAccess AccessFinal)
 {
 	IF_RDG_ENABLE_DEBUG(UserValidation.ValidateSetAccessFinal(Texture, AccessFinal));
-	Texture->AccessFinal = AccessFinal;
+	Texture->EpilogueAccess = AccessFinal;
 }
 
 inline void FRDGBuilder::SetBufferAccessFinal(FRDGBufferRef Buffer, ERHIAccess AccessFinal)
 {
 	IF_RDG_ENABLE_DEBUG(UserValidation.ValidateSetAccessFinal(Buffer, AccessFinal));
-	Buffer->AccessFinal = AccessFinal;
+	Buffer->EpilogueAccess = AccessFinal;
 }
 
 inline void FRDGBuilder::RemoveUnusedTextureWarning(FRDGTextureRef Texture)
diff --git a/Engine/Source/Runtime/RenderCore/Public/RenderGraphDefinitions.h b/Engine/Source/Runtime/RenderCore/Public/RenderGraphDefinitions.h
--- a/Engine/Source/Runtime/RenderCore/Public/RenderGraphDefinitions.h
+++ b/Engine/Source/Runtime/RenderCore/Public/RenderGraphDefinitions.h
@@ -171,13 +171,16 @@ enum class ERDGUnorderedAccessViewFlags : uint8
 ENUM_CLASS_FLAGS(ERDGUnorderedAccessViewFlags);
 
 /** The set of concrete parent resource types. */
-enum class ERDGParentResourceType : uint8
+enum class ERDGViewableResourceType : uint8
 {
 	Texture,
 	Buffer,
 	MAX
 };
 
+UE_DEPRECATED(5.1, "ERDGParentResourceType has been renamed to ERDGViewableResourceType.")
+typedef ERDGViewableResourceType ERDGParentResourceType;
+
 /** The set of concrete view types. */
 enum class ERDGViewType : uint8
 {
@@ -216,22 +219,28 @@ enum class ERDGInitialDataFlags : uint8
 ENUM_CLASS_FLAGS(ERDGInitialDataFlags)
 
 /** Returns the equivalent parent resource type for a view type. */
-inline ERDGParentResourceType GetParentResourceType(ERDGViewType ViewType)
+inline ERDGViewableResourceType GetViewableResourceType(ERDGViewType ViewType)
 {
 	switch (ViewType)
 	{
 	case ERDGViewType::TextureUAV:
 	case ERDGViewType::TextureSRV:
-		return ERDGParentResourceType::Texture;
+		return ERDGViewableResourceType::Texture;
 	case ERDGViewType::BufferUAV:
 	case ERDGViewType::BufferSRV:
-		return ERDGParentResourceType::Buffer;
+		return ERDGViewableResourceType::Buffer;
 	default:
 		checkNoEntry();
-		return ERDGParentResourceType::MAX;
+		return ERDGViewableResourceType::MAX;
 	}
 }
 
+UE_DEPRECATED(5.1, "GetParentResourceType has been renamed to GetViewableResourceType.")
+inline ERDGViewableResourceType GetParentResourceType(ERDGViewType ViewType)
+{
+	return GetViewableResourceType(ViewType);
+}
+
 using ERDGTextureMetaDataAccess = ERHITextureMetaDataAccess;
 
 /** Returns the associated FRHITransitionInfo plane index. */
@@ -548,8 +557,13 @@ class FRDGUserValidation;
 class FRDGResource;
 using FRDGResourceRef = FRDGResource*;
 
-class FRDGParentResource;
-using FRDGParentResourceRef = FRDGParentResource*;
+class FRDGViewableResource;
+
+UE_DEPRECATED(5.1, "FRDGParentResource has been renamed to FRDGViewableResource.")
+typedef FRDGViewableResource FRDGParentResource;
+
+UE_DEPRECATED(5.1, "FRDGParentResourceRef has been renamed to FRDGViewableResource*.")
+typedef FRDGViewableResource* FRDGParentResourceRef;
 
 class FRDGShaderResourceView;
 using FRDGShaderResourceViewRef = FRDGShaderResourceView*;
diff --git a/Engine/Source/Runtime/RenderCore/Public/RenderGraphParameter.h b/Engine/Source/Runtime/RenderCore/Public/RenderGraphParameter.h
--- a/Engine/Source/Runtime/RenderCore/Public/RenderGraphParameter.h
+++ b/Engine/Source/Runtime/RenderCore/Public/RenderGraphParameter.h
@@ -73,11 +73,17 @@ public:
 		return MemberType == UBMT_RDG_UNIFORM_BUFFER;
 	}
 
+	UE_DEPRECATED(5.1, "Use IsViewableResource instead.")
 	bool IsParentResource() const
 	{
 		return IsTexture() || IsBuffer();
 	}
 
+	bool IsViewableResource() const
+	{
+		return IsTexture() || IsBuffer();
+	}
+
 	bool IsRenderTargetBindingSlots() const
 	{
 		return MemberType == UBMT_RENDER_TARGET_BINDING_SLOTS;
@@ -100,10 +106,18 @@ public:
 		return *GetAs<FRDGUniformBufferBinding>();
 	}
 
-	FRDGParentResourceRef GetAsParentResource() const
+	UE_DEPRECATED(5.1, "Use IsViewableResource instead.")
+	FRDGViewableResource* GetAsParentResource() const
+	{
+		check(IsViewableResource());
+		return *GetAs<FRDGViewableResource*>();
+	}
+
+
+	FRDGViewableResource* GetAsViewableResource() const
 	{
-		check(IsParentResource());
-		return *GetAs<FRDGParentResourceRef>();
+		check(IsViewableResource());
+		return *GetAs<FRDGViewableResource*>();
 	}
 
 	FRDGViewRef GetAsView() const
diff --git a/Engine/Source/Runtime/RenderCore/Public/RenderGraphPass.h b/Engine/Source/Runtime/RenderCore/Public/RenderGraphPass.h
--- a/Engine/Source/Runtime/RenderCore/Public/RenderGraphPass.h
+++ b/Engine/Source/Runtime/RenderCore/Public/RenderGraphPass.h
@@ -75,9 +75,9 @@ public:
 	FRDGBarrierBatchBegin(ERHIPipeline PipelinesToBegin, ERHIPipeline PipelinesToEnd, const TCHAR* DebugName, FRDGPass* DebugPass);
 	FRDGBarrierBatchBegin(ERHIPipeline PipelinesToBegin, ERHIPipeline PipelinesToEnd, const TCHAR* DebugName, FRDGPassesByPipeline DebugPasses);
 
-	void AddTransition(FRDGParentResourceRef Resource, const FRHITransitionInfo& Info);
+	void AddTransition(FRDGViewableResource* Resource, const FRHITransitionInfo& Info);
 
-	void AddAlias(FRDGParentResourceRef Resource, const FRHITransientAliasingInfo& Info);
+	void AddAlias(FRDGViewableResource* Resource, const FRHITransientAliasingInfo& Info);
 
 	void SetUseCrossPipelineFence()
 	{
@@ -112,8 +112,8 @@ private:
 
 #if RDG_ENABLE_DEBUG
 	FRDGPassesByPipeline DebugPasses;
-	TArray<FRDGParentResource*, FRDGArrayAllocator> DebugTransitionResources;
-	TArray<FRDGParentResource*, FRDGArrayAllocator> DebugAliasingResources;
+	TArray<FRDGViewableResource*, FRDGArrayAllocator> DebugTransitionResources;
+	TArray<FRDGViewableResource*, FRDGArrayAllocator> DebugAliasingResources;
 	const TCHAR* DebugName;
 	ERHIPipeline DebugPipelinesToBegin;
 	ERHIPipeline DebugPipelinesToEnd;
@@ -434,8 +434,8 @@ protected:
 		}
 
 		FRDGTextureRef Texture = nullptr;
-		FRDGTextureTransientSubresourceState State;
-		FRDGTextureTransientSubresourceStateIndirect MergeState;
+		FRDGTextureSubresourceState State;
+		FRDGTextureSubresourceStateIndirect MergeState;
 		uint16 ReferenceCount = 0;
 	};
 
diff --git a/Engine/Source/Runtime/RenderCore/Public/RenderGraphResources.h b/Engine/Source/Runtime/RenderCore/Public/RenderGraphResources.h
--- a/Engine/Source/Runtime/RenderCore/Public/RenderGraphResources.h
+++ b/Engine/Source/Runtime/RenderCore/Public/RenderGraphResources.h
@@ -33,7 +33,7 @@ struct FRDGSubresourceState
 	static bool IsTransitionRequired(const FRDGSubresourceState& Previous, const FRDGSubresourceState& Next);
 
 	/** Given a before and after state, returns whether they can be merged into a single state. */
-	static bool IsMergeAllowed(ERDGParentResourceType ResourceType, const FRDGSubresourceState& Previous, const FRDGSubresourceState& Next);
+	static bool IsMergeAllowed(ERDGViewableResourceType ResourceType, const FRDGSubresourceState& Previous, const FRDGSubresourceState& Next);
 
 	FRDGSubresourceState() = default;
 
@@ -76,17 +76,10 @@ struct FRDGSubresourceState
 
 	/** The last no-UAV barrier to be used by this subresource. */
 	FRDGViewUniqueFilter NoUAVBarrierFilter;
-
-	/** The last used pass for log file debugging. */
-	IF_RDG_ENABLE_DEBUG(mutable FRDGPassHandle LogFilePass);
 };
 
-using FRDGTextureSubresourceState = TRDGTextureSubresourceArray<FRDGSubresourceState, FDefaultAllocator>;
-using FRDGTextureTransientSubresourceState = TRDGTextureSubresourceArray<FRDGSubresourceState, FRDGArrayAllocator>;
-using FRDGTextureTransientSubresourceStateIndirect = TRDGTextureSubresourceArray<FRDGSubresourceState*, FRDGArrayAllocator>;
-
-using FRDGPooledTextureArray = TArray<TRefCountPtr<IPooledRenderTarget>, FRDGArrayAllocator>;
-using FRDGPooledBufferArray = TArray<TRefCountPtr<FRDGPooledBuffer>, FRDGArrayAllocator>;
+using FRDGTextureSubresourceState = TRDGTextureSubresourceArray<FRDGSubresourceState, FRDGArrayAllocator>;
+using FRDGTextureSubresourceStateIndirect = TRDGTextureSubresourceArray<FRDGSubresourceState*, FRDGArrayAllocator>;
 
 /** Generic graph resource. */
 class RENDERCORE_API FRDGResource
@@ -241,12 +234,12 @@ private:
 };
 
 /** A render graph resource with an allocation lifetime tracked by the graph. May have child resources which reference it (e.g. views). */
-class RENDERCORE_API FRDGParentResource
+class RENDERCORE_API FRDGViewableResource
 	: public FRDGResource
 {
 public:
 	/** The type of this resource; useful for casting between types. */
-	const ERDGParentResourceType Type;
+	const ERDGViewableResourceType Type;
 
 	/** Whether this resource is externally registered with the graph (i.e. the user holds a reference to the underlying resource outside the graph). */
 	bool IsExternal() const
@@ -274,7 +267,7 @@ public:
 	}
 
 protected:
-	FRDGParentResource(const TCHAR* InName, ERDGParentResourceType InType);
+	FRDGViewableResource(const TCHAR* InName, ERDGViewableResourceType InType);
 
 	enum class ETransientExtractionHint
 	{
@@ -335,8 +328,21 @@ private:
 	/** Scratch index allocated for the resource in the pass being setup. */
 	uint16 PassStateIndex = 0;
 
-	/** The final state of the resource at the end of graph execution, if known. */
-	ERHIAccess AccessFinal = ERHIAccess::Unknown;
+	/** The state of the resource at the graph epilogue. */
+	ERHIAccess EpilogueAccess = ERHIAccess::SRVMask;
+
+	/** The valid set of access states for finalized access. */
+	IF_RDG_ENABLE_DEBUG(ERHIAccess FinalizedAccessMask = ERHIAccess::Unknown);
+
+	void SetFinalizedAccess(ERHIAccess InValidAccessMask)
+	{
+		bFinalizedAccess = 1;
+
+		// Finalized resources are not always added to the pass states (unless marked as such within the graph), so marked as not culled here.
+		bCulled = 0;
+
+		IF_RDG_ENABLE_DEBUG(FinalizedAccessMask = InValidAccessMask);
+	}
 
 #if RDG_ENABLE_TRACE
 	uint16 TraceOrder = 0;
@@ -344,8 +350,8 @@ private:
 #endif
 
 #if RDG_ENABLE_DEBUG
-	struct FRDGParentResourceDebugData* ParentDebugData = nullptr;
-	FRDGParentResourceDebugData& GetParentDebugData() const;
+	struct FRDGViewableResourceDebugData* ViewableDebugData = nullptr;
+	FRDGViewableResourceDebugData& GetViewableDebugData() const;
 #endif
 
 	friend FRDGBuilder;
@@ -354,7 +360,7 @@ private:
 	friend FRDGTrace;
 };
 
-/** A render graph resource (e.g. a view) which references a single parent resource (e.g. a texture / buffer). Provides an abstract way to access the parent resource. */
+/** A render graph resource (e.g. a view) which references a single viewable resource (e.g. a texture / buffer). Provides an abstract way to access the viewable resource. */
 class FRDGView
 	: public FRDGResource
 {
@@ -363,7 +369,7 @@ public:
 	const ERDGViewType Type;
 
 	/** Returns the referenced parent render graph resource. */
-	virtual FRDGParentResourceRef GetParent() const = 0;
+	virtual FRDGViewableResource* GetParent() const = 0;
 
 	FRDGViewHandle GetHandle() const
 	{
@@ -403,11 +409,9 @@ class RENDERCORE_API FRDGPooledTexture final
 	: public FRefCountedObject
 {
 public:
-	FRDGPooledTexture(FRHITexture* InTexture, const FRDGTextureSubresourceLayout& InLayout, ERHIAccess AccessInitial = ERHIAccess::Unknown)
+	FRDGPooledTexture(FRHITexture* InTexture)
 		: Texture(InTexture)
-	{
-		InitAsWholeResource(State, FRDGSubresourceState(AccessInitial));
-	}
+	{}
 
 	/** Finds a UAV matching the descriptor in the cache or creates a new one and updates the cache. */
 	FORCEINLINE FRHIUnorderedAccessView* GetOrCreateUAV(const FRHITextureUAVCreateInfo& UAVDesc) { return ViewCache.GetOrCreateUAV(Texture, UAVDesc); }
@@ -417,29 +421,16 @@ public:
 
 	FORCEINLINE FRHITexture* GetRHI() const { return Texture; }
 
-	FORCEINLINE FRDGTexture* GetOwner() const { return Owner; }
-
 private:
-	/** Prepares the pooled texture state for re-use across RDG builder instances. */
-	void Finalize();
-
-	/** Resets the pooled texture state back an unknown value. */
-	void Reset();
-
 	TRefCountPtr<FRHITexture> Texture;
-	FRDGTexture* Owner = nullptr;
-	FRDGTextureSubresourceState State;
 	FRHITextureViewCache ViewCache;
 
-	friend FRDGTexture;
 	friend FRDGBuilder;
-	friend FRenderTargetPool;
-	friend FRDGAllocator;
 };
 
 /** Render graph tracked Texture. */
 class RENDERCORE_API FRDGTexture final
-	: public FRDGParentResource
+	: public FRDGViewableResource
 {
 public:
 	const FRDGTextureDesc Desc;
@@ -489,7 +480,7 @@ public:
 
 private:
 	FRDGTexture(const TCHAR* InName, const FRDGTextureDesc& InDesc, ERDGTextureFlags InFlags)
-		: FRDGParentResource(InName, ERDGParentResourceType::Texture)
+		: FRDGViewableResource(InName, ERDGViewableResourceType::Texture)
 		, Desc(InDesc)
 		, Flags(InFlags)
 		, Layout(InDesc)
@@ -503,18 +494,6 @@ private:
 		bSwapChain = EnumHasAnyFlags(Desc.Flags, ETextureCreateFlags::Presentable);
 	}
 
-	/** Assigns a render target as the backing RHI resource. */
-	void SetRHI(IPooledRenderTarget* PooledRenderTarget);
-
-	/** Assigns a pooled texture as the backing RHI resource. */
-	void SetRHI(FRDGPooledTexture* PooledTexture);
-
-	/** Assigns a transient texture as the backing RHI resource. */
-	void SetRHI(FRHITransientTexture* TransientTexture, FRDGTextureSubresourceState* TransientTextureState);
-
-	/** Finalizes the texture for execution; no other transitions are allowed after calling this. */
-	void Finalize(FRDGPooledTextureArray& PooledTextureArray);
-
 	/** Returns RHI texture without access checks. */
 	FRHITexture* GetRHIUnchecked() const
 	{
@@ -539,8 +518,8 @@ private:
 	FRDGTextureSubresourceRange  WholeRange;
 	const uint32 SubresourceCount;
 
-	/** The assigned pooled render target to use during execution. Never reset. */
-	IPooledRenderTarget* PooledRenderTarget = nullptr;
+	/** The assigned render target to use during execution. Never reset. */
+	IPooledRenderTarget* RenderTarget = nullptr;
 
 	union
 	{
@@ -557,11 +536,11 @@ private:
 	/** Valid strictly when holding a strong reference; use PooledRenderTarget instead. */
 	TRefCountPtr<IPooledRenderTarget> Allocation;
 
-	/** Cached state pointer from the pooled texture. */
+	/** Tracks subresource states as the graph is built. */
 	FRDGTextureSubresourceState* State = nullptr;
 
 	/** Tracks merged subresource states as the graph is built. */
-	FRDGTextureTransientSubresourceStateIndirect MergeState;
+	FRDGTextureSubresourceStateIndirect MergeState;
 
 	/** Tracks pass producers for each subresource as the graph is built. */
 	TRDGTextureSubresourceArray<FRDGProducerStatesByPipeline, FRDGArrayAllocator> LastProducers;
@@ -636,7 +615,7 @@ class FRDGTextureSRVDesc final
 {
 public:
 	FRDGTextureSRVDesc() = default;
-	
+
 	FRDGTextureRef Texture = nullptr;
 
 	/** Create SRV that access all sub resources of texture. */
@@ -987,18 +966,6 @@ private:
 	TRefCountPtr<FRHIBuffer> Buffer;
 	FRHIBufferViewCache ViewCache;
 
-	void Reset()
-	{
-		Owner = nullptr;
-		State = {};
-	}
-
-	void Finalize()
-	{
-		Owner = nullptr;
-		State.Finalize();
-	}
-
 	FRDGBufferDesc GetAlignedDesc() const
 	{
 		FRDGBufferDesc AlignedDesc = Desc;
@@ -1007,21 +974,17 @@ private:
 	}
 
 	const TCHAR* Name = nullptr;
-	FRDGBufferRef Owner = nullptr;
-	FRDGSubresourceState State;
 
 	const uint32 NumAllocatedElements;
 	uint32 LastUsedFrame = 0;
 
-	friend FRDGBufferPool;
 	friend FRDGBuilder;
-	friend FRDGBuffer;
-	friend FRDGAllocator;
+	friend FRDGBufferPool;
 };
 
 /** A render graph tracked buffer. */
 class RENDERCORE_API FRDGBuffer final
-	: public FRDGParentResource
+	: public FRDGViewableResource
 {
 public:
 	FRDGBufferDesc Desc;
@@ -1033,7 +996,7 @@ public:
 	/** Returns the underlying RHI buffer resource */
 	FRHIBuffer* GetRHI() const
 	{
-		return static_cast<FRHIBuffer*>(FRDGParentResource::GetRHI());
+		return static_cast<FRHIBuffer*>(FRDGViewableResource::GetRHI());
 	}
 
 	/** Returns the buffer to use for indirect RHI calls. */
@@ -1066,7 +1029,7 @@ public:
 
 private:
 	FRDGBuffer(const TCHAR* InName, const FRDGBufferDesc& InDesc, ERDGBufferFlags InFlags)
-		: FRDGParentResource(InName, ERDGParentResourceType::Buffer)
+		: FRDGViewableResource(InName, ERDGViewableResourceType::Buffer)
 		, Desc(InDesc)
 		, Flags(InFlags)
 	{}
@@ -1077,15 +1040,6 @@ private:
 		NumElementsCallback = MoveTemp(InNumElementsCallback);
 	}
 
-	/** Assigns a pooled buffer as the backing RHI resource. */
-	void SetRHI(FRDGPooledBuffer* InPooledBuffer);
-
-	/** Assigns a transient buffer as the backing RHI resource. */
-	void SetRHI(FRHITransientBuffer* InTransientBuffer, FRDGAllocator& Allocator);
-
-	/** Finalizes the buffer for execution; no other transitions are allowed after calling this. */
-	void Finalize(FRDGPooledBufferArray& PooledBufferArray);
-
 	/** Finalizes any pending field of the buffer descriptor. */
 	void FinalizeDesc()
 	{
diff --git a/Engine/Source/Runtime/RenderCore/Public/RenderGraphTrace.h b/Engine/Source/Runtime/RenderCore/Public/RenderGraphTrace.h
--- a/Engine/Source/Runtime/RenderCore/Public/RenderGraphTrace.h
+++ b/Engine/Source/Runtime/RenderCore/Public/RenderGraphTrace.h
@@ -17,7 +17,7 @@ public:
 	void OutputGraphBegin();
 	void OutputGraphEnd(const FRDGBuilder& GraphBuilder);
 
-	void AddResource(FRDGParentResource* Resource);
+	void AddResource(FRDGViewableResource* Resource);
 	void AddTexturePassDependency(FRDGTexture* Texture, FRDGPass* Pass);
 	void AddBufferPassDependency(FRDGBuffer* Buffer, FRDGPass* Pass);
 
diff --git a/Engine/Source/Runtime/RenderCore/Public/RenderGraphUtils.h b/Engine/Source/Runtime/RenderCore/Public/RenderGraphUtils.h
--- a/Engine/Source/Runtime/RenderCore/Public/RenderGraphUtils.h
+++ b/Engine/Source/Runtime/RenderCore/Public/RenderGraphUtils.h
@@ -13,7 +13,7 @@
 class FGlobalShaderMap;
 
 /** Returns whether the resource was produced by a prior pass. */
-inline bool HasBeenProduced(FRDGParentResourceRef Resource)
+inline bool HasBeenProduced(FRDGViewableResource* Resource)
 {
 	return Resource && Resource->HasBeenProduced();
 }
diff --git a/Engine/Source/Runtime/RenderCore/Public/RenderGraphValidation.h b/Engine/Source/Runtime/RenderCore/Public/RenderGraphValidation.h
--- a/Engine/Source/Runtime/RenderCore/Public/RenderGraphValidation.h
+++ b/Engine/Source/Runtime/RenderCore/Public/RenderGraphValidation.h
@@ -67,7 +67,7 @@ public:
 	void ValidateExtractTexture(FRDGTextureRef Texture, TRefCountPtr<IPooledRenderTarget>* OutTexturePtr);
 	void ValidateExtractBuffer(FRDGBufferRef Buffer, TRefCountPtr<FRDGPooledBuffer>* OutBufferPtr);
 
-	void ValidateConvertToExternalResource(FRDGParentResourceRef Resource);
+	void ValidateConvertToExternalResource(FRDGViewableResource* Resource);
 
 	/** Tracks and validates the addition of a new pass to the graph.
 	 *  @param bSkipPassAccessMarking Skips marking the pass as a producer or incrementing the pass access. Useful when
@@ -87,26 +87,26 @@ public:
 	void ValidateExecuteEnd();
 
 	/** Removes the 'produced but not used' warning from the requested resource. */
-	void RemoveUnusedWarning(FRDGParentResourceRef Resource);
+	void RemoveUnusedWarning(FRDGViewableResource* Resource);
 
 	/** Attempts to mark a resource for clobbering. If already marked, returns false.  */
-	bool TryMarkForClobber(FRDGParentResourceRef Resource) const;
+	bool TryMarkForClobber(FRDGViewableResource* Resource) const;
 
 	void ValidateGetPooledTexture(FRDGTextureRef Texture) const;
 	void ValidateGetPooledBuffer(FRDGBufferRef Buffer) const;
 
-	void ValidateSetAccessFinal(FRDGParentResourceRef Resource, ERHIAccess AccessFinal);
+	void ValidateSetAccessFinal(FRDGViewableResource* Resource, ERHIAccess AccessFinal);
 
-	void ValidateFinalize(FRDGParentResourceRef Resource, ERHIAccess Access, FRDGPassHandle ConvertToUntrackedPassHandle);
-	void ValidateFinalizedAccess(FRDGParentResourceRef Resource, ERHIAccess Access, const FRDGPass* Pass);
+	void ValidateFinalize(FRDGViewableResource* Resource, ERHIAccess Access, FRDGPassHandle ConvertToUntrackedPassHandle);
+	void ValidateFinalizedAccess(FRDGViewableResource* Resource, ERHIAccess Access, const FRDGPass* Pass);
 
 	/** Traverses all resources in the pass and marks whether they are externally accessible by user pass implementations. */
 	static void SetAllowRHIAccess(const FRDGPass* Pass, bool bAllowAccess);
 
 private:
-	void ValidateCreateParentResource(FRDGParentResourceRef Resource);
+	void ValidateCreateViewableResource(FRDGViewableResource* Resource);
 	void ValidateCreateResource(FRDGResourceRef Resource);
-	void ValidateExtractResource(FRDGParentResourceRef Resource);
+	void ValidateExtractResource(FRDGViewableResource* Resource);
 
 	FRDGAllocator& Allocator;
 
@@ -144,7 +144,7 @@ private:
 	{
 		TMap<FRDGTextureRef, TArray<FRHITransitionInfo>> Textures;
 		TMap<FRDGBufferRef, FRHITransitionInfo> Buffers;
-		TMap<FRDGParentResourceRef, FRHITransientAliasingInfo> Aliases;
+		TMap<FRDGViewableResource*, FRHITransientAliasingInfo> Aliases;
 	};
 
 	using FBarrierBatchMap = TMap<const FRDGBarrierBatchBegin*, FResourceMap>;
@@ -155,57 +155,4 @@ private:
 	const TCHAR* GraphName = nullptr;
 };
 
-class FRDGLogFile
-{
-public:
-	FRDGLogFile(const FRDGPassRegistry& InPasses)
-		: Passes(InPasses)
-	{}
-
-	void Begin(const FRDGEventName& GraphName);
-
-	void AddFirstEdge(const FRDGTextureRef Texture, FRDGPassHandle FirstPass);
-
-	void AddFirstEdge(const FRDGBufferRef Buffer, FRDGPassHandle FirstPass);
-
-	void AddAliasEdge(const FRDGTextureRef TextureBefore, FRDGPassHandle BeforePass, const FRDGTextureRef TextureAfter, FRDGPassHandle PassAfter);
-
-	void AddAliasEdge(const FRDGBufferRef BufferBefore, FRDGPassHandle BeforePass, const FRDGBufferRef BufferAfter, FRDGPassHandle PassAfter);
-
-	void AddTransitionEdge(FRDGPassHandle PassHandle, const FRDGSubresourceState& StateBefore, const FRDGSubresourceState& StateAfter, const FRDGTextureRef Texture);
-
-	void AddTransitionEdge(FRDGPassHandle PassHandle, const FRDGSubresourceState& StateBefore, const FRDGSubresourceState& StateAfter, const FRDGTextureRef Texture, FRDGTextureSubresource Subresource);
-
-	void AddTransitionEdge(FRDGPassHandle PassHandle, const FRDGSubresourceState& StateBefore, const FRDGSubresourceState& StateAfter, const FRDGBufferRef Buffer);
-
-	void End();
-
-private:
-	void AddLine(const FString& Line);
-	void AddBraceBegin();
-	void AddBraceEnd();
-
-	FString GetProducerName(FRDGPassHandle PassHandle);
-	FString GetConsumerName(FRDGPassHandle PassHandle);
-
-	FString GetNodeName(FRDGPassHandle Pass);
-	FString GetNodeName(const FRDGTexture* Texture);
-	FString GetNodeName(const FRDGBuffer* Buffer);
-
-	bool IncludeTransitionEdgeInGraph(FRDGPassHandle PassBefore, FRDGPassHandle PassAfter) const;
-	bool IncludeTransitionEdgeInGraph(FRDGPassHandle Pass) const;
-
-	bool bOpen = false;
-
-	const FRDGPassRegistry& Passes;
-
-	TSet<FRDGPassHandle> PassesReferenced;
-	TArray<const FRDGTexture*> Textures;
-	TArray<const FRDGBuffer*> Buffers;
-
-	FString Indentation;
-	FString File;
-	FString GraphName;
-};
-
 #endif
\ No newline at end of file
diff --git a/Engine/Source/Runtime/RenderCore/Public/RenderTargetPool.h b/Engine/Source/Runtime/RenderCore/Public/RenderTargetPool.h
--- a/Engine/Source/Runtime/RenderCore/Public/RenderTargetPool.h
+++ b/Engine/Source/Runtime/RenderCore/Public/RenderTargetPool.h
@@ -17,10 +17,10 @@ class FRenderTargetPool;
 /** The reference to a pooled render target, use like this: TRefCountPtr<IPooledRenderTarget> */
 struct RENDERCORE_API FPooledRenderTarget final : public IPooledRenderTarget
 {
-	FPooledRenderTarget(FRHITexture* Texture, ERHIAccess AccessInitial, const FPooledRenderTargetDesc& InDesc, FRenderTargetPool* InRenderTargetPool) 
+	FPooledRenderTarget(FRHITexture* Texture, const FPooledRenderTargetDesc& InDesc, FRenderTargetPool* InRenderTargetPool) 
 		: RenderTargetPool(InRenderTargetPool)
 		, Desc(InDesc)
-		, PooledTexture(Texture, Translate(InDesc), AccessInitial)
+		, PooledTexture(Texture)
 	{
 		RenderTargetItem.TargetableTexture = RenderTargetItem.ShaderResourceTexture = Texture;
 	}
@@ -98,7 +98,7 @@ class RENDERCORE_API FRenderTargetPool : public FRenderResource
 public:
 	FRenderTargetPool() = default;
 
-	TRefCountPtr<IPooledRenderTarget> FindFreeElement(const FRHITextureCreateInfo& Desc, const TCHAR* Name);
+	TRefCountPtr<IPooledRenderTarget> FindFreeElement(FRHITextureCreateInfo Desc, const TCHAR* Name);
 
 	bool FindFreeElement(const FRHITextureCreateInfo& Desc, TRefCountPtr<IPooledRenderTarget>& Out, const TCHAR* Name);
 
@@ -148,7 +148,6 @@ public:
 	void DumpMemoryUsage(FOutputDevice& OutputDevice);
 
 private:
-	TRefCountPtr<IPooledRenderTarget> FindFreeElementInternal(FRHITextureCreateInfo Desc, const TCHAR* Name, bool bResetStateToUnknown);
 	void FreeElementAtIndex(int32 Index);
 
 	/** Elements can be 0, we compact the buffer later. */
diff --git a/Engine/Source/Runtime/VulkanRHI/Private/VulkanResources.h b/Engine/Source/Runtime/VulkanRHI/Private/VulkanResources.h
--- a/Engine/Source/Runtime/VulkanRHI/Private/VulkanResources.h
+++ b/Engine/Source/Runtime/VulkanRHI/Private/VulkanResources.h
@@ -1058,7 +1058,7 @@ protected:
 class FVulkanShaderResourceView : public FRHIShaderResourceView, public VulkanRHI::FVulkanViewBase
 {
 public:
-	FVulkanShaderResourceView(FVulkanDevice* Device, FRHIResource* InRHIBuffer, FVulkanResourceMultiBuffer* InSourceBuffer, uint32 InSize, EPixelFormat InFormat, uint32 InOffset = 0);
+	FVulkanShaderResourceView(FVulkanDevice* Device, FRHIViewableResource* InRHIBuffer, FVulkanResourceMultiBuffer* InSourceBuffer, uint32 InSize, EPixelFormat InFormat, uint32 InOffset = 0);
 	FVulkanShaderResourceView(FVulkanDevice* Device, FRHITexture* InSourceTexture, const FRHITextureSRVCreateInfo& InCreateInfo);
 	FVulkanShaderResourceView(FVulkanDevice* Device, FVulkanResourceMultiBuffer* InSourceBuffer, uint32 InOffset = 0);
 
diff --git a/Engine/Source/Runtime/VulkanRHI/Private/VulkanTransientResourceAllocator.cpp b/Engine/Source/Runtime/VulkanRHI/Private/VulkanTransientResourceAllocator.cpp
--- a/Engine/Source/Runtime/VulkanRHI/Private/VulkanTransientResourceAllocator.cpp
+++ b/Engine/Source/Runtime/VulkanRHI/Private/VulkanTransientResourceAllocator.cpp
@@ -114,17 +114,7 @@ FRHITransientTexture* FVulkanTransientResourceAllocator::CreateTexture(const FRH
 	return CreateTextureInternal(InCreateInfo, InDebugName, InPassIndex, MemReq.Size, MemReq.Align,
 		[&](const FRHITransientHeap::FResourceInitializer& Initializer)
 	{
-		ERHIAccess InitialState = ERHIAccess::UAVMask;
-		if (EnumHasAnyFlags(InCreateInfo.Flags, TexCreate_RenderTargetable | TexCreate_ResolveTargetable))
-		{
-			InitialState = ERHIAccess::RTV;
-		}
-		else if (EnumHasAnyFlags(InCreateInfo.Flags, TexCreate_DepthStencilTargetable))
-		{
-			InitialState = ERHIAccess::DSVWrite;
-		}
-
-		FRHITextureCreateDesc CreateDesc(InCreateInfo, InitialState, InDebugName);
+		FRHITextureCreateDesc CreateDesc(InCreateInfo, ERHIAccess::Discard, InDebugName);
 		FRHITexture* Texture = new FVulkanTexture(*Device, CreateDesc, &Initializer.Allocation);
 		return new FRHITransientTexture(Texture, 0/*GpuVirtualAddress*/, Initializer.Hash, MemReq.Size, ERHITransientAllocationType::Heap, InCreateInfo);
 	});
@@ -141,6 +131,7 @@ FRHITransientBuffer* FVulkanTransientResourceAllocator::CreateBuffer(const FRHIB
 	{
 		FRHIResourceCreateInfo ResourceCreateInfo(InDebugName);
 		FRHIBuffer* Buffer = GVulkanRHI->CreateBuffer(InCreateInfo, ResourceCreateInfo, &Initializer.Allocation);
+		Buffer->SetTrackedAccess_Unsafe(ERHIAccess::Discard);
 		return new FRHITransientBuffer(Buffer, 0/*GpuVirtualAddress*/, Initializer.Hash, Size, ERHITransientAllocationType::Heap, InCreateInfo);
 	});
 }
diff --git a/Engine/Source/Runtime/VulkanRHI/Private/VulkanUAV.cpp b/Engine/Source/Runtime/VulkanRHI/Private/VulkanUAV.cpp
--- a/Engine/Source/Runtime/VulkanRHI/Private/VulkanUAV.cpp
+++ b/Engine/Source/Runtime/VulkanRHI/Private/VulkanUAV.cpp
@@ -8,8 +8,9 @@
 #include "VulkanRayTracing.h"
 #endif // VULKAN_RHI_RAYTRACING
 
-FVulkanShaderResourceView::FVulkanShaderResourceView(FVulkanDevice* Device, FRHIResource* InRHIBuffer, FVulkanResourceMultiBuffer* InSourceBuffer, uint32 InSize, EPixelFormat InFormat, uint32 InOffset)
-	: VulkanRHI::FVulkanViewBase(Device)
+FVulkanShaderResourceView::FVulkanShaderResourceView(FVulkanDevice* Device, FRHIViewableResource* InRHIBuffer, FVulkanResourceMultiBuffer* InSourceBuffer, uint32 InSize, EPixelFormat InFormat, uint32 InOffset)
+	: FRHIShaderResourceView(InRHIBuffer)
+	, VulkanRHI::FVulkanViewBase(Device)
 	, BufferViewFormat(InFormat)
 	, SourceTexture(nullptr)
 	, SourceStructuredBuffer(nullptr)
@@ -29,7 +30,8 @@ FVulkanShaderResourceView::FVulkanShaderResourceView(FVulkanDevice* Device, FRHI
 
 
 FVulkanShaderResourceView::FVulkanShaderResourceView(FVulkanDevice* Device, FRHITexture* InSourceTexture, const FRHITextureSRVCreateInfo& InCreateInfo)
-	: VulkanRHI::FVulkanViewBase(Device)
+	: FRHIShaderResourceView(InSourceTexture)
+	, VulkanRHI::FVulkanViewBase(Device)
 	, BufferViewFormat((EPixelFormat)InCreateInfo.Format)
 	, SRGBOverride(InCreateInfo.SRGBOverride)
 	, SourceTexture(InSourceTexture)
@@ -47,7 +49,8 @@ FVulkanShaderResourceView::FVulkanShaderResourceView(FVulkanDevice* Device, FRHI
 }
 
 FVulkanShaderResourceView::FVulkanShaderResourceView(FVulkanDevice* InDevice, FVulkanResourceMultiBuffer* InSourceBuffer, uint32 InOffset)
-	: VulkanRHI::FVulkanViewBase(InDevice)
+	: FRHIShaderResourceView(InSourceBuffer)
+	, VulkanRHI::FVulkanViewBase(InDevice)
 {
 	check(InDevice && InSourceBuffer);
 
@@ -231,7 +234,8 @@ void FVulkanShaderResourceView::UpdateView()
 	}
 }
 FVulkanUnorderedAccessView::FVulkanUnorderedAccessView(FVulkanDevice* Device, FVulkanResourceMultiBuffer* Buffer, bool bUseUAVCounter, bool bAppendBuffer)
-	: VulkanRHI::FVulkanViewBase(Device)
+	: FRHIUnorderedAccessView(Buffer)
+	, VulkanRHI::FVulkanViewBase(Device)
 	, MipLevel(0)
 	, FirstArraySlice(0)
 	, NumArraySlices(0)
@@ -242,7 +246,8 @@ FVulkanUnorderedAccessView::FVulkanUnorderedAccessView(FVulkanDevice* Device, FV
 }
 
 FVulkanUnorderedAccessView::FVulkanUnorderedAccessView(FVulkanDevice* Device, FRHITexture* TextureRHI, uint32 MipLevel, uint16 InFirstArraySlice, uint16 InNumArraySlices)
-	: VulkanRHI::FVulkanViewBase(Device)
+	: FRHIUnorderedAccessView(TextureRHI)
+	, VulkanRHI::FVulkanViewBase(Device)
 	, SourceTexture(TextureRHI)
 	, MipLevel(MipLevel)
 	, FirstArraySlice(InFirstArraySlice)
@@ -256,7 +261,8 @@ FVulkanUnorderedAccessView::FVulkanUnorderedAccessView(FVulkanDevice* Device, FR
 
 
 FVulkanUnorderedAccessView::FVulkanUnorderedAccessView(FVulkanDevice* Device, FVulkanResourceMultiBuffer* Buffer, EPixelFormat Format)
-	: VulkanRHI::FVulkanViewBase(Device)
+	: FRHIUnorderedAccessView(Buffer)
+	, VulkanRHI::FVulkanViewBase(Device)
 	, MipLevel(0)
 	, FirstArraySlice(0)
 	, NumArraySlices(0)
diff --git a/Engine/Source/Runtime/Windows/D3D11RHI/Private/D3D11Texture.cpp b/Engine/Source/Runtime/Windows/D3D11RHI/Private/D3D11Texture.cpp
--- a/Engine/Source/Runtime/Windows/D3D11RHI/Private/D3D11Texture.cpp
+++ b/Engine/Source/Runtime/Windows/D3D11RHI/Private/D3D11Texture.cpp
@@ -1102,7 +1102,7 @@ FShaderResourceViewRHIRef FD3D11DynamicRHI::RHICreateShaderResourceView(FRHIText
 	TRefCountPtr<ID3D11ShaderResourceView> ShaderResourceView;
 	VERIFYD3D11RESULT_EX(Direct3DDevice->CreateShaderResourceView(Texture->GetResource(), &SRVDesc, (ID3D11ShaderResourceView**)ShaderResourceView.GetInitReference()), Direct3DDevice);
 
-	return new FD3D11ShaderResourceView(ShaderResourceView, Texture);
+	return new FD3D11ShaderResourceView(ShaderResourceView, Texture, Texture);
 }
 
 FShaderResourceViewRHIRef FD3D11DynamicRHI::RHICreateShaderResourceView_RenderThread(class FRHICommandListImmediate& RHICmdList, FRHITexture* Texture, const FRHITextureSRVCreateInfo& CreateInfo)
diff --git a/Engine/Source/Runtime/Windows/D3D11RHI/Private/D3D11UAV.cpp b/Engine/Source/Runtime/Windows/D3D11RHI/Private/D3D11UAV.cpp
--- a/Engine/Source/Runtime/Windows/D3D11RHI/Private/D3D11UAV.cpp
+++ b/Engine/Source/Runtime/Windows/D3D11RHI/Private/D3D11UAV.cpp
@@ -62,7 +62,7 @@ FUnorderedAccessViewRHIRef FD3D11DynamicRHI::RHICreateUnorderedAccessView(FRHIBu
 	TRefCountPtr<ID3D11UnorderedAccessView> UnorderedAccessView;
 	VERIFYD3D11CREATEVIEWRESULT(Direct3DDevice->CreateUnorderedAccessView(Buffer->Resource, &UAVDesc, (ID3D11UnorderedAccessView**)UnorderedAccessView.GetInitReference()), Direct3DDevice, BufferRHI, UAVDesc);
 
-	return new FD3D11UnorderedAccessView(UnorderedAccessView, Buffer);
+	return new FD3D11UnorderedAccessView(UnorderedAccessView, Buffer, Buffer);
 }
 
 FUnorderedAccessViewRHIRef FD3D11DynamicRHI::RHICreateUnorderedAccessView_RenderThread(
@@ -112,7 +112,7 @@ FUnorderedAccessViewRHIRef FD3D11DynamicRHI::RHICreateUnorderedAccessView(FRHITe
 	TRefCountPtr<ID3D11UnorderedAccessView> UnorderedAccessView;
 	VERIFYD3D11CREATEVIEWRESULT(Direct3DDevice->CreateUnorderedAccessView(Texture->GetResource(),&UAVDesc,(ID3D11UnorderedAccessView**)UnorderedAccessView.GetInitReference()), Direct3DDevice, TextureRHI, UAVDesc);
 
-	return new FD3D11UnorderedAccessView(UnorderedAccessView,Texture);
+	return new FD3D11UnorderedAccessView(UnorderedAccessView, Texture, Texture);
 }
 
 FUnorderedAccessViewRHIRef FD3D11DynamicRHI::RHICreateUnorderedAccessView_RenderThread(
@@ -137,7 +137,7 @@ FUnorderedAccessViewRHIRef FD3D11DynamicRHI::RHICreateUnorderedAccessView(FRHIBu
 	TRefCountPtr<ID3D11UnorderedAccessView> UnorderedAccessView;
 	VERIFYD3D11CREATEVIEWRESULT(Direct3DDevice->CreateUnorderedAccessView(Buffer->Resource, &UAVDesc, (ID3D11UnorderedAccessView**)UnorderedAccessView.GetInitReference()), Direct3DDevice, BufferRHI, UAVDesc);
 
-	return new FD3D11UnorderedAccessView(UnorderedAccessView, Buffer);
+	return new FD3D11UnorderedAccessView(UnorderedAccessView, Buffer, Buffer);
 }
 
 FUnorderedAccessViewRHIRef FD3D11DynamicRHI::RHICreateUnorderedAccessView_RenderThread(class FRHICommandListImmediate& RHICmdList, FRHIBuffer* Buffer, uint8 Format)
@@ -223,7 +223,7 @@ FShaderResourceViewRHIRef FD3D11DynamicRHI::RHICreateShaderResourceView(const FS
 
 	if (!Desc.Buffer || !Buffer->Resource)
 	{
-		return new FD3D11ShaderResourceView(nullptr, nullptr);
+		return new FD3D11ShaderResourceView(nullptr, nullptr, nullptr);
 	}
 
 	TRefCountPtr<ID3D11ShaderResourceView> ShaderResourceView;
@@ -234,7 +234,7 @@ FShaderResourceViewRHIRef FD3D11DynamicRHI::RHICreateShaderResourceView(const FS
 		{
 			CreateD3D11ShaderResourceViewOnBuffer(Direct3DDevice, Buffer->Resource, Desc.StartOffsetBytes, Desc.NumElements, Desc.Format, ShaderResourceView.GetInitReference());
 
-			return new FD3D11ShaderResourceView(ShaderResourceView, Buffer);
+			return new FD3D11ShaderResourceView(ShaderResourceView, Buffer, Buffer);
 		}
 
 		case FShaderResourceViewInitializer::EType::StructuredBufferSRV:
@@ -273,7 +273,7 @@ FShaderResourceViewRHIRef FD3D11DynamicRHI::RHICreateShaderResourceView(const FS
 
 			VERIFYD3D11RESULT_EX(Direct3DDevice->CreateShaderResourceView(Buffer->Resource, &SRVDesc, (ID3D11ShaderResourceView**)ShaderResourceView.GetInitReference()), Direct3DDevice);
 
-			return new FD3D11ShaderResourceView(ShaderResourceView, Buffer);
+			return new FD3D11ShaderResourceView(ShaderResourceView, Buffer, Buffer);
 		}
 
 		case FShaderResourceViewInitializer::EType::IndexBufferSRV:
@@ -285,7 +285,7 @@ FShaderResourceViewRHIRef FD3D11DynamicRHI::RHICreateShaderResourceView(const FS
 
 			CreateD3D11ShaderResourceViewOnBuffer(Direct3DDevice, Buffer->Resource, Desc.StartOffsetBytes, Desc.NumElements, Format, ShaderResourceView.GetInitReference());
 
-			return new FD3D11ShaderResourceView(ShaderResourceView, Buffer);
+			return new FD3D11ShaderResourceView(ShaderResourceView, Buffer, Buffer);
 		}
 	}
 
diff --git a/Engine/Source/Runtime/Windows/D3D11RHI/Public/D3D11Resources.h b/Engine/Source/Runtime/Windows/D3D11RHI/Public/D3D11Resources.h
--- a/Engine/Source/Runtime/Windows/D3D11RHI/Public/D3D11Resources.h
+++ b/Engine/Source/Runtime/Windows/D3D11RHI/Public/D3D11Resources.h
@@ -477,8 +477,9 @@ public:
 	TRefCountPtr<ID3D11ShaderResourceView> View;
 	TRefCountPtr<FD3D11BaseShaderResource> Resource;
 
-	FD3D11ShaderResourceView(ID3D11ShaderResourceView* InView,FD3D11BaseShaderResource* InResource)
-	: View(InView)
+	FD3D11ShaderResourceView(ID3D11ShaderResourceView* InView, FRHIViewableResource* InParentResource, FD3D11BaseShaderResource* InResource)
+	: FRHIShaderResourceView(InParentResource)
+	, View(InView)
 	, Resource(InResource)
 	{}
 
@@ -498,8 +499,9 @@ public:
 	TRefCountPtr<FD3D11BaseShaderResource> Resource;	
 	void* IHVResourceHandle;
 	
-	FD3D11UnorderedAccessView(ID3D11UnorderedAccessView* InView,FD3D11BaseShaderResource* InResource)
-	: View(InView)
+	FD3D11UnorderedAccessView(ID3D11UnorderedAccessView* InView, FRHIViewableResource* InParentResource, FD3D11BaseShaderResource* InResource)
+	: FRHIUnorderedAccessView(InParentResource)
+	, View(InView)
 	, Resource(InResource)
 	, IHVResourceHandle(nullptr)
 	{}
