Do not DumpGPU RDG resources that are read-only

diff --git a/Engine/Source/Runtime/RenderCore/Private/DumpGPU.cpp b/Engine/Source/Runtime/RenderCore/Private/DumpGPU.cpp
--- a/Engine/Source/Runtime/RenderCore/Private/DumpGPU.cpp
+++ b/Engine/Source/Runtime/RenderCore/Private/DumpGPU.cpp
@@ -1115,17 +1115,24 @@ struct FRDGResourceDumpContext
 			return;
 		}
 
-		FString DumpFilePath = kResourcesDir / FString::Printf(TEXT("%s.v%016x.bin"), *UniqueResourceSubResourceName, PtrToUint(bIsOutputResource ? Pass : nullptr));
-
 		// Verify there is enough available memory to dump the resource.
 		if (IsUnsafeToDumpResource(SubresourceDumpDesc.ByteSize, 2.2f + (SubresourceDumpDesc.bPreprocessForStaging ? 1.0f : 0.0f)))
 		{
-			UE_LOG(LogRendererCore, Warning, TEXT("Not dumping %s because of insuficient memory available for staging texture."), *DumpFilePath);
+			UE_LOG(LogRendererCore, Warning, TEXT("Not dumping %s because of insuficient memory available for staging texture."), SubresourceDesc.Texture->Name);
+			return;
+		}
+
+		// Verify the texture is able to do resource transitions.
+		if (EnumHasAnyFlags(SubresourceDesc.Texture->Flags, ERDGTextureFlags::ReadOnly))
+		{
+			UE_LOG(LogRendererCore, Warning, TEXT("Not dumping %s because has ERDGTextureFlags::ReadOnly."), SubresourceDesc.Texture->Name);
 			return;
 		}
 
 		// Dump the resource's binary to a .bin file.
 		{
+			FString DumpFilePath = kResourcesDir / FString::Printf(TEXT("%s.v%016x.bin"), *UniqueResourceSubResourceName, PtrToUint(bIsOutputResource ? Pass : nullptr));
+
 			FDumpTexturePass* PassParameters = GraphBuilder.AllocParameters<FDumpTexturePass>();
 			if (SubresourceDumpDesc.bPreprocessForStaging)
 			{
@@ -1270,6 +1277,13 @@ struct FRDGResourceDumpContext
 				return;
 			}
 
+			// Verify the texture is able to do resource transitions.
+			if (EnumHasAnyFlags(Buffer->Flags, ERDGBufferFlags::ReadOnly))
+			{
+				UE_LOG(LogRendererCore, Warning, TEXT("Not dumping %s because has ERDGBufferFlags::ReadOnly."), Buffer->Name);
+				return;
+			}
+
 			FDumpBufferPass* PassParameters = GraphBuilder.AllocParameters<FDumpBufferPass>();
 			PassParameters->Buffer = Buffer;
 
