Don't clobber RDG resources that were uploaded from host memory.

diff --git a/Engine/Source/Runtime/RenderCore/Private/RenderGraphValidation.cpp b/Engine/Source/Runtime/RenderCore/Private/RenderGraphValidation.cpp
--- a/Engine/Source/Runtime/RenderCore/Private/RenderGraphValidation.cpp
+++ b/Engine/Source/Runtime/RenderCore/Private/RenderGraphValidation.cpp
@@ -453,7 +453,7 @@ bool FRDGUserValidation::TryMarkForClobber(FRDGViewableResource* Resource) const
 	check(Resource);
 	FRDGViewableResourceDebugData& DebugData = Resource->GetViewableDebugData();
 
-	const bool bClobber = !DebugData.bHasBeenClobbered && !Resource->bExternal && IsDebugAllowedForResource(Resource->Name);
+	const bool bClobber = !DebugData.bHasBeenClobbered && !Resource->bExternal && !Resource->bQueuedForUpload && IsDebugAllowedForResource(Resource->Name);
 
 	if (bClobber)
 	{
