Have DumpGPU to dump a log on crash, as well as a Base/Status.txt to explicitly detect in the viewer if something went really wrong

diff --git a/Engine/Extras/GPUDumpViewer/GPUDumpViewer.html b/Engine/Extras/GPUDumpViewer/GPUDumpViewer.html
--- a/Engine/Extras/GPUDumpViewer/GPUDumpViewer.html
+++ b/Engine/Extras/GPUDumpViewer/GPUDumpViewer.html
@@ -3923,6 +3923,27 @@ function init_page()
 		return;
 	}
 
+	// Verify the status of the dump, to check whether a crash happened during the dump.
+	{
+		var dump_status = load_text_file('Base/Status.txt');
+		if (dump_status == 'ok')
+		{
+			add_console_event('log', `The dump completed gracefully.`);
+		}
+		else if (dump_status == 'crash')
+		{
+			add_console_event('error', `A crash happened while the frame was being dumped.`);
+		}
+		else if (dump_status == 'dumping')
+		{
+			add_console_event('error', `The dump has not completed. This may be due to opening the viewer before completion, or serious problem that has not been handled with FCoreDelegates::OnShutdownAfterError`);
+		}
+		else
+		{
+			add_console_event('error', `The dump completed with status: ${dump_status}`);
+		}
+	}
+
 	// Load the cvars used for dumping.
 	{
 		g_dump_cvars = {};
diff --git a/Engine/Source/Runtime/RenderCore/Private/DumpGPU.cpp b/Engine/Source/Runtime/RenderCore/Private/DumpGPU.cpp
--- a/Engine/Source/Runtime/RenderCore/Private/DumpGPU.cpp
+++ b/Engine/Source/Runtime/RenderCore/Private/DumpGPU.cpp
@@ -7,6 +7,7 @@
 #include "Misc/FileHelper.h"
 #include "Misc/WildcardString.h"
 #include "Misc/OutputDeviceRedirector.h"
+#include "Misc/CoreDelegates.h"
 #include "Runtime/Launch/Resources/Version.h"
 #include "BuildSettings.h"
 #include "Serialization/JsonTypes.h"
@@ -279,12 +280,6 @@ struct FRDGResourceDumpContext
 		return DumpServiceParameters;
 	}
 
-	bool IsDumpingFrame() const
-	{
-		check(IsInRenderingThread() || IsInParallelRenderingThread());
-		return !DumpingDirectoryPath.IsEmpty();
-	}
-
 	FString GetDumpFullPath(const FString& DumpRelativeFileName) const
 	{
 		check(bEnableDiskWrite);
@@ -307,6 +302,11 @@ struct FRDGResourceDumpContext
 			FFileHelper::EEncodingOptions::AutoDetect, &IFileManager::Get(), WriteFlags);
 	}
 
+	bool DumpStatusToFile(FStringView StatusString)
+	{
+		return DumpStringToFile(StatusString, FString(FRDGResourceDumpContext::kBaseDir) / TEXT("Status.txt"));
+	}
+
 	bool DumpJsonToFile(const TSharedPtr<FJsonObject>& JsonObject, const FString& FileName, uint32 WriteFlags = FILEWRITE_None)
 	{
 		FString OutputString;
@@ -1364,17 +1364,27 @@ struct FRDGResourceDumpContext
 
 		return bDumpPass;
 	}
+
+	void OnCrash()
+	{
+		if (GLog)
+		{
+			GLog->PanicFlush();
+		}
+		DumpStatusToFile(TEXT("crash"));
+		FGenericCrashContext::DumpLog(DumpingDirectoryPath / FRDGResourceDumpContext::kBaseDir);
+	}
 };
 
 // 0 = not dumping, MAX_uint64 dump request for next frame, otherwise dump frame counter
 static uint64 DumpingFrameCounter_GameThread = 0;
-FRDGResourceDumpContext GRDGResourceDumpContext;
+FRDGResourceDumpContext* GRDGResourceDumpContext = nullptr;
 
 }
 
 bool IsDumpingRDGResources()
 {
-	return GRDGResourceDumpContext.IsDumpingFrame();
+	return GRDGResourceDumpContext != nullptr;
 }
 
 void FRDGBuilder::InitResourceDump()
@@ -1429,9 +1439,9 @@ FString FRDGBuilder::BeginResourceDump(const TCHAR* Cmd)
 		}
 	}
 
-	FRDGResourceDumpContext NewResourceDumpContext;
+	FRDGResourceDumpContext* NewResourceDumpContext = new FRDGResourceDumpContext;
 
-	NewResourceDumpContext.Time = FDateTime::Now();
+	NewResourceDumpContext->Time = FDateTime::Now();
 	{
 		FString CVarDirectoryPath = GDumpGPUDirectoryCVar.GetValueOnGameThread();
 		FString EnvDirectoryPath = FPlatformMisc::GetEnvironmentVariable(TEXT("UE-DumpGPUPath"));
@@ -1449,13 +1459,13 @@ FString FRDGBuilder::BeginResourceDump(const TCHAR* Cmd)
 		{
 			DirectoryPath = FPaths::ProjectSavedDir() / TEXT("GPUDumps/");
 		}
-		NewResourceDumpContext.DumpingDirectoryPath = DirectoryPath / FApp::GetProjectName() + TEXT("-") + FPlatformProperties::PlatformName() + TEXT("-") + NewResourceDumpContext.Time.ToString() + TEXT("/");
+		NewResourceDumpContext->DumpingDirectoryPath = DirectoryPath / FApp::GetProjectName() + TEXT("-") + FPlatformProperties::PlatformName() + TEXT("-") + NewResourceDumpContext->Time.ToString() + TEXT("/");
 	}
-	NewResourceDumpContext.bEnableDiskWrite = GDumpTestEnableDiskWrite.GetValueOnGameThread() != 0;
+	NewResourceDumpContext->bEnableDiskWrite = GDumpTestEnableDiskWrite.GetValueOnGameThread() != 0;
 
 	if (Switches.Contains(TEXT("upload")))
 	{
-		if (!IDumpGPUUploadServiceProvider::GProvider || !NewResourceDumpContext.bEnableDiskWrite)
+		if (!IDumpGPUUploadServiceProvider::GProvider || !NewResourceDumpContext->bEnableDiskWrite)
 		{
 			UE_LOG(LogRendererCore, Warning, TEXT("DumpGPU upload services are not set up."));
 		}
@@ -1465,45 +1475,51 @@ FString FRDGBuilder::BeginResourceDump(const TCHAR* Cmd)
 		}
 		else
 		{
-			NewResourceDumpContext.bUpload = true;
+			NewResourceDumpContext->bUpload = true;
 		}
 	}
 
-	if (NewResourceDumpContext.bUpload)
+	if (NewResourceDumpContext->bUpload)
 	{
 		if (GDumpGPUUploadCompressResources.GetValueOnGameThread() == 1)
 		{
-			NewResourceDumpContext.UploadResourceCompressionName = NAME_Zlib;
+			NewResourceDumpContext->UploadResourceCompressionName = NAME_Zlib;
 		}
 		else if (GDumpGPUUploadCompressResources.GetValueOnGameThread() == 2)
 		{
-			NewResourceDumpContext.UploadResourceCompressionName = NAME_Gzip;
+			NewResourceDumpContext->UploadResourceCompressionName = NAME_Gzip;
 		}
 	}
 
-	NewResourceDumpContext.bShowInExplore = NewResourceDumpContext.bEnableDiskWrite && GDumpExploreCVar.GetValueOnGameThread() != 0 && !NewResourceDumpContext.bUpload;
-	NewResourceDumpContext.MemoryConstants = FPlatformMemory::GetConstants();
-	NewResourceDumpContext.MemoryStats = FPlatformMemory::GetStats();
+	NewResourceDumpContext->bShowInExplore = NewResourceDumpContext->bEnableDiskWrite && GDumpExploreCVar.GetValueOnGameThread() != 0 && !NewResourceDumpContext->bUpload;
+	NewResourceDumpContext->MemoryConstants = FPlatformMemory::GetConstants();
+	NewResourceDumpContext->MemoryStats = FPlatformMemory::GetStats();
 
 	IPlatformFile& PlatformFile = FPlatformFileManager::Get().GetPlatformFile();
-	if (NewResourceDumpContext.bEnableDiskWrite)
+	if (NewResourceDumpContext->bEnableDiskWrite)
 	{
-		if (!PlatformFile.DirectoryExists(*NewResourceDumpContext.DumpingDirectoryPath))
+		if (!PlatformFile.DirectoryExists(*NewResourceDumpContext->DumpingDirectoryPath))
 		{
-			PlatformFile.CreateDirectoryTree(*NewResourceDumpContext.DumpingDirectoryPath);
+			PlatformFile.CreateDirectoryTree(*NewResourceDumpContext->DumpingDirectoryPath);
 		}
-		PlatformFile.CreateDirectoryTree(*(NewResourceDumpContext.DumpingDirectoryPath / FRDGResourceDumpContext::kBaseDir));
-		PlatformFile.CreateDirectoryTree(*(NewResourceDumpContext.DumpingDirectoryPath / FRDGResourceDumpContext::kResourcesDir));
+		PlatformFile.CreateDirectoryTree(*(NewResourceDumpContext->DumpingDirectoryPath / FRDGResourceDumpContext::kBaseDir));
+		PlatformFile.CreateDirectoryTree(*(NewResourceDumpContext->DumpingDirectoryPath / FRDGResourceDumpContext::kResourcesDir));
 
-		NewResourceDumpContext.DumpStringToFile(TEXT(""), FString(FRDGResourceDumpContext::kBaseDir) / TEXT("Passes.json"));
-		NewResourceDumpContext.DumpStringToFile(TEXT(""), FString(FRDGResourceDumpContext::kBaseDir) / TEXT("ResourceDescs.json"));
-		NewResourceDumpContext.DumpStringToFile(TEXT(""), FString(FRDGResourceDumpContext::kBaseDir) / TEXT("PassDrawCounts.json"));
+		NewResourceDumpContext->DumpStringToFile(TEXT(""), FString(FRDGResourceDumpContext::kBaseDir) / TEXT("Passes.json"));
+		NewResourceDumpContext->DumpStringToFile(TEXT(""), FString(FRDGResourceDumpContext::kBaseDir) / TEXT("ResourceDescs.json"));
+		NewResourceDumpContext->DumpStringToFile(TEXT(""), FString(FRDGResourceDumpContext::kBaseDir) / TEXT("PassDrawCounts.json"));
+	}
+
+	// Dump status file and register NewResourceDumpContext to listen for OnShutdownAfterError to get a log of what happened with callstack.
+	{
+		NewResourceDumpContext->DumpStatusToFile(TEXT("dumping"));
+		FCoreDelegates::OnShutdownAfterError.AddRaw(NewResourceDumpContext, &FRDGResourceDumpContext::OnCrash);
 	}
 
 	// Dump service parameters so GPUDumpViewer.html remain compatible when not using upload provider.
-	if (NewResourceDumpContext.bEnableDiskWrite)
+	if (NewResourceDumpContext->bEnableDiskWrite)
 	{
-		NewResourceDumpContext.GetDumpParameters().DumpServiceParametersFile();
+		NewResourceDumpContext->GetDumpParameters().DumpServiceParametersFile();
 	}
 
 	// Output informations
@@ -1539,23 +1555,23 @@ FString FRDGBuilder::BeginResourceDump(const TCHAR* Cmd)
 		JsonObject->SetStringField(TEXT("GPUDriverUserVersion"), GPUDriverInfo.UserDriverVersion);
 		JsonObject->SetStringField(TEXT("GPUDriverInternalVersion"), GPUDriverInfo.GetUnifiedDriverVersion());
 		JsonObject->SetStringField(TEXT("GPUDriverDate"), GPUDriverInfo.DriverDate);
-		JsonObject->SetNumberField(TEXT("MemoryTotalPhysical"), NewResourceDumpContext.MemoryConstants.TotalPhysical);
-		JsonObject->SetNumberField(TEXT("MemoryPageSize"), NewResourceDumpContext.MemoryConstants.PageSize);
+		JsonObject->SetNumberField(TEXT("MemoryTotalPhysical"), NewResourceDumpContext->MemoryConstants.TotalPhysical);
+		JsonObject->SetNumberField(TEXT("MemoryPageSize"), NewResourceDumpContext->MemoryConstants.PageSize);
 		JsonObject->SetStringField(TEXT("RHI"), GDynamicRHI->GetName());
 		JsonObject->SetStringField(TEXT("RHIMaxFeatureLevel"), LexToString(GMaxRHIFeatureLevel));
-		JsonObject->SetStringField(TEXT("DumpTime"), NewResourceDumpContext.Time.ToString());
+		JsonObject->SetStringField(TEXT("DumpTime"), NewResourceDumpContext->Time.ToString());
 
-		NewResourceDumpContext.DumpJsonToFile(JsonObject, FString(FRDGResourceDumpContext::kBaseDir) / TEXT("Infos.json"));
+		NewResourceDumpContext->DumpJsonToFile(JsonObject, FString(FRDGResourceDumpContext::kBaseDir) / TEXT("Infos.json"));
 	}
 
 	// Dump the rendering cvars
-	if (NewResourceDumpContext.bEnableDiskWrite)
+	if (NewResourceDumpContext->bEnableDiskWrite)
 	{
-		NewResourceDumpContext.DumpRenderingCVarsToCSV();
+		NewResourceDumpContext->DumpRenderingCVarsToCSV();
 	}
 
 	// Copy the viewer
-	if (NewResourceDumpContext.bEnableDiskWrite)
+	if (NewResourceDumpContext->bEnableDiskWrite)
 	{
 		const TCHAR* OpenGPUDumpViewerBatName = TEXT("OpenGPUDumpViewer.bat");
 		const TCHAR* OpenGPUDumpViewerShName = TEXT("OpenGPUDumpViewer.sh");
@@ -1563,9 +1579,9 @@ FString FRDGBuilder::BeginResourceDump(const TCHAR* Cmd)
 		const TCHAR* ViewerHTML = TEXT("GPUDumpViewer.html");
 		FString DumpGPUViewerSourcePath = FPaths::EngineDir() + FString(TEXT("Extras")) / TEXT("GPUDumpViewer");
 
-		PlatformFile.CopyFile(*(NewResourceDumpContext.DumpingDirectoryPath / ViewerHTML), *(DumpGPUViewerSourcePath / ViewerHTML));
-		PlatformFile.CopyFile(*(NewResourceDumpContext.DumpingDirectoryPath / OpenGPUDumpViewerBatName), *(DumpGPUViewerSourcePath / OpenGPUDumpViewerBatName));
-		PlatformFile.CopyFile(*(NewResourceDumpContext.DumpingDirectoryPath / OpenGPUDumpViewerShName), *(DumpGPUViewerSourcePath / OpenGPUDumpViewerShName));
+		PlatformFile.CopyFile(*(NewResourceDumpContext->DumpingDirectoryPath / ViewerHTML), *(DumpGPUViewerSourcePath / ViewerHTML));
+		PlatformFile.CopyFile(*(NewResourceDumpContext->DumpingDirectoryPath / OpenGPUDumpViewerBatName), *(DumpGPUViewerSourcePath / OpenGPUDumpViewerBatName));
+		PlatformFile.CopyFile(*(NewResourceDumpContext->DumpingDirectoryPath / OpenGPUDumpViewerShName), *(DumpGPUViewerSourcePath / OpenGPUDumpViewerShName));
 	}
 
 	ENQUEUE_RENDER_COMMAND(FStartGPUDump)(
@@ -1585,9 +1601,9 @@ FString FRDGBuilder::BeginResourceDump(const TCHAR* Cmd)
 	// Mark ready for dump on next available frame
 	DumpingFrameCounter_GameThread = MAX_uint64;
 
-	if (NewResourceDumpContext.bEnableDiskWrite)
+	if (NewResourceDumpContext->bEnableDiskWrite)
 	{
-		return NewResourceDumpContext.DumpingDirectoryPath;
+		return NewResourceDumpContext->DumpingDirectoryPath;
 	}
 	return FString();
 }
@@ -1621,53 +1637,59 @@ void FRDGBuilder::EndResourceDump()
 	}
 
 	// Log information about the dump.
-	FString AbsDumpingDirectoryPath = IFileManager::Get().ConvertToAbsolutePathForExternalAppForRead(*GRDGResourceDumpContext.DumpingDirectoryPath);
+	FString AbsDumpingDirectoryPath = IFileManager::Get().ConvertToAbsolutePathForExternalAppForRead(*GRDGResourceDumpContext->DumpingDirectoryPath);
 	{
 		FDateTime Now = FDateTime::Now();
-		double TotalDumpSeconds = (Now - GRDGResourceDumpContext.Time).GetTotalSeconds();
-		double RHIReadbackCommandsSeconds = GRDGResourceDumpContext.TimingBucket[int32(FRDGResourceDumpContext::ETimingBucket::RHIReadbackCommands)];
-		double GPUWaitSeconds = GRDGResourceDumpContext.TimingBucket[int32(FRDGResourceDumpContext::ETimingBucket::GPUWait)];
-		double CPUPostProcessingSeconds = GRDGResourceDumpContext.TimingBucket[int32(FRDGResourceDumpContext::ETimingBucket::CPUPostProcessing)];
-		double MetadataFileWriteSeconds = GRDGResourceDumpContext.TimingBucket[int32(FRDGResourceDumpContext::ETimingBucket::MetadataFileWrite)];
-		double ResourceBinaryFileWriteSeconds = GRDGResourceDumpContext.TimingBucket[int32(FRDGResourceDumpContext::ETimingBucket::ResourceBinaryFileWrite)];
-		double RHIReleaseResourcesTimeSeconds = GRDGResourceDumpContext.TimingBucket[int32(FRDGResourceDumpContext::ETimingBucket::RHIReleaseResources)];
-
-		UE_LOG(LogRendererCore, Display, TEXT("Dumped %d resources in %.3f s to %s"), GRDGResourceDumpContext.ResourcesDumpPasses, float(TotalDumpSeconds), *AbsDumpingDirectoryPath);
+		double TotalDumpSeconds = (Now - GRDGResourceDumpContext->Time).GetTotalSeconds();
+		double RHIReadbackCommandsSeconds = GRDGResourceDumpContext->TimingBucket[int32(FRDGResourceDumpContext::ETimingBucket::RHIReadbackCommands)];
+		double GPUWaitSeconds = GRDGResourceDumpContext->TimingBucket[int32(FRDGResourceDumpContext::ETimingBucket::GPUWait)];
+		double CPUPostProcessingSeconds = GRDGResourceDumpContext->TimingBucket[int32(FRDGResourceDumpContext::ETimingBucket::CPUPostProcessing)];
+		double MetadataFileWriteSeconds = GRDGResourceDumpContext->TimingBucket[int32(FRDGResourceDumpContext::ETimingBucket::MetadataFileWrite)];
+		double ResourceBinaryFileWriteSeconds = GRDGResourceDumpContext->TimingBucket[int32(FRDGResourceDumpContext::ETimingBucket::ResourceBinaryFileWrite)];
+		double RHIReleaseResourcesTimeSeconds = GRDGResourceDumpContext->TimingBucket[int32(FRDGResourceDumpContext::ETimingBucket::RHIReleaseResources)];
+
+		UE_LOG(LogRendererCore, Display, TEXT("Dumped %d resources in %.3f s to %s"), GRDGResourceDumpContext->ResourcesDumpPasses, float(TotalDumpSeconds), *AbsDumpingDirectoryPath);
 		UE_LOG(LogRendererCore, Display, TEXT("Dumped GPU readback commands: %.3f s"), float(RHIReadbackCommandsSeconds));
 		UE_LOG(LogRendererCore, Display, TEXT("Dumped GPU wait: %.3f s"), float(GPUWaitSeconds));
 		UE_LOG(LogRendererCore, Display, TEXT("Dumped CPU resource binary post processing: %.3f s"), float(CPUPostProcessingSeconds));
 		UE_LOG(LogRendererCore, Display, TEXT("Dumped metadata: %.3f MB in %d files under %.3f s at %.3f MB/s"),
-			float(GRDGResourceDumpContext.MetadataFilesWriteBytes) / float(1024 * 1024),
-			GRDGResourceDumpContext.MetadataFilesOpened,
+			float(GRDGResourceDumpContext->MetadataFilesWriteBytes) / float(1024 * 1024),
+			GRDGResourceDumpContext->MetadataFilesOpened,
 			float(MetadataFileWriteSeconds),
-			float(GRDGResourceDumpContext.MetadataFilesWriteBytes) / (float(1024 * 1024) * float(MetadataFileWriteSeconds)));
+			float(GRDGResourceDumpContext->MetadataFilesWriteBytes) / (float(1024 * 1024) * float(MetadataFileWriteSeconds)));
 		UE_LOG(LogRendererCore, Display, TEXT("Dumped resource binary: %.3f MB in %d files under %.3f s at %.3f MB/s"),
-			float(GRDGResourceDumpContext.ResourceBinaryWriteBytes) / float(1024 * 1024),
-			GRDGResourceDumpContext.ResourceBinaryFilesOpened,
+			float(GRDGResourceDumpContext->ResourceBinaryWriteBytes) / float(1024 * 1024),
+			GRDGResourceDumpContext->ResourceBinaryFilesOpened,
 			float(ResourceBinaryFileWriteSeconds),
-			float(GRDGResourceDumpContext.ResourceBinaryWriteBytes) / (float(1024 * 1024) * float(ResourceBinaryFileWriteSeconds)));
+			float(GRDGResourceDumpContext->ResourceBinaryWriteBytes) / (float(1024 * 1024) * float(ResourceBinaryFileWriteSeconds)));
 		UE_LOG(LogRendererCore, Display, TEXT("Dumped GPU readback resource release: %.3f s"), float(RHIReleaseResourcesTimeSeconds));
 	}
 
 	// Dump the log into the dump directory.
-	if (GRDGResourceDumpContext.bEnableDiskWrite)
+	if (GRDGResourceDumpContext->bEnableDiskWrite)
 	{
 		if (GLog)
 		{
 			GLog->FlushThreadedLogs();
 			GLog->Flush();
 		}
-		FGenericCrashContext::DumpLog(GRDGResourceDumpContext.DumpingDirectoryPath / FRDGResourceDumpContext::kBaseDir);
+		FGenericCrashContext::DumpLog(GRDGResourceDumpContext->DumpingDirectoryPath / FRDGResourceDumpContext::kBaseDir);
+	}
+
+	// Update the dump status to OK and removes subscription to crashes
+	{
+		GRDGResourceDumpContext->DumpStatusToFile(TEXT("ok"));
+		FCoreDelegates::OnShutdownAfterError.RemoveAll(GRDGResourceDumpContext);
 	}
 
-	if (GRDGResourceDumpContext.bUpload && IDumpGPUUploadServiceProvider::GProvider)
+	if (GRDGResourceDumpContext->bUpload && IDumpGPUUploadServiceProvider::GProvider)
 	{
-		IDumpGPUUploadServiceProvider::FDumpParameters DumpCompletedParameters = GRDGResourceDumpContext.GetDumpParameters();
+		IDumpGPUUploadServiceProvider::FDumpParameters DumpCompletedParameters = GRDGResourceDumpContext->GetDumpParameters();
 
 		// Compress the resource binary in background before uploading.
-		if (!GRDGResourceDumpContext.UploadResourceCompressionName.IsNone())
+		if (!GRDGResourceDumpContext->UploadResourceCompressionName.IsNone())
 		{
-			DumpCompletedParameters.CompressionName = GRDGResourceDumpContext.UploadResourceCompressionName;
+			DumpCompletedParameters.CompressionName = GRDGResourceDumpContext->UploadResourceCompressionName;
 			DumpCompletedParameters.CompressionFiles = FWildcardString(TEXT("*.bin"));
 		}
 
@@ -1675,13 +1697,14 @@ void FRDGBuilder::EndResourceDump()
 	}
 
 	#if PLATFORM_DESKTOP
-	if (GRDGResourceDumpContext.bShowInExplore)
+	if (GRDGResourceDumpContext->bShowInExplore)
 	{
 		FPlatformProcess::ExploreFolder(*AbsDumpingDirectoryPath);
 	}
 	#endif
 
-	GRDGResourceDumpContext = FRDGResourceDumpContext();
+	delete GRDGResourceDumpContext;
+	GRDGResourceDumpContext = nullptr;
 	DumpingFrameCounter_GameThread = 0;
 }
 
@@ -1710,13 +1733,13 @@ void FRDGBuilder::DumpResourcePassOutputs(const FRDGPass* Pass)
 		return;
 	}
 
-	if (!GRDGResourceDumpContext.IsDumpingFrame())
+	if (!GRDGResourceDumpContext)
 	{
 		return;
 	}
 
 	check(IsInRenderingThread());
-	if (!GRDGResourceDumpContext.IsDumpingPass(Pass))
+	if (!GRDGResourceDumpContext->IsDumpingPass(Pass))
 	{
 		return;
 	}
@@ -1734,7 +1757,7 @@ void FRDGBuilder::DumpResourcePassOutputs(const FRDGPass* Pass)
 			if (FRDGTextureRef Texture = Parameter.GetAsTexture())
 			{
 				FRDGTextureSRVDesc TextureSubResource = FRDGTextureSRVDesc::Create(Texture);
-				GRDGResourceDumpContext.AddDumpTexturePasses(*this, InputResourceNames, OutputResourceNames, Pass, TextureSubResource, /* bIsOutputResource = */ false);
+				GRDGResourceDumpContext->AddDumpTexturePasses(*this, InputResourceNames, OutputResourceNames, Pass, TextureSubResource, /* bIsOutputResource = */ false);
 			}
 		}
 		break;
@@ -1744,7 +1767,7 @@ void FRDGBuilder::DumpResourcePassOutputs(const FRDGPass* Pass)
 			{
 				if (SRV->Desc.MetaData == ERHITextureMetaDataAccess::None)
 				{
-					GRDGResourceDumpContext.AddDumpTexturePasses(*this, InputResourceNames, OutputResourceNames, Pass, SRV->Desc, /* bIsOutputResource = */ false);
+					GRDGResourceDumpContext->AddDumpTexturePasses(*this, InputResourceNames, OutputResourceNames, Pass, SRV->Desc, /* bIsOutputResource = */ false);
 				}
 				else
 				{
@@ -1760,7 +1783,7 @@ void FRDGBuilder::DumpResourcePassOutputs(const FRDGPass* Pass)
 				if (UAV->Desc.MetaData == ERHITextureMetaDataAccess::None)
 				{
 					FRDGTextureSRVDesc TextureSubResource = FRDGTextureSRVDesc::CreateForMipLevel(UAV->Desc.Texture, UAV->Desc.MipLevel);
-					GRDGResourceDumpContext.AddDumpTextureSubResourcePass(*this, InputResourceNames, OutputResourceNames, Pass, TextureSubResource, /* bIsOutputResource = */ true);
+					GRDGResourceDumpContext->AddDumpTextureSubResourcePass(*this, InputResourceNames, OutputResourceNames, Pass, TextureSubResource, /* bIsOutputResource = */ true);
 				}
 				else
 				{
@@ -1779,7 +1802,7 @@ void FRDGBuilder::DumpResourcePassOutputs(const FRDGPass* Pass)
 					TextureAccess.GetAccess() == ERHIAccess::RTV);
 
 				FRDGTextureSRVDesc TextureSubResource = FRDGTextureSRVDesc::Create(TextureAccess);
-				GRDGResourceDumpContext.AddDumpTexturePasses(*this, InputResourceNames, OutputResourceNames, Pass, TextureSubResource, bIsOutputResource);
+				GRDGResourceDumpContext->AddDumpTexturePasses(*this, InputResourceNames, OutputResourceNames, Pass, TextureSubResource, bIsOutputResource);
 			}
 		}
 		break;
@@ -1795,7 +1818,7 @@ void FRDGBuilder::DumpResourcePassOutputs(const FRDGPass* Pass)
 					TextureAccess.GetAccess() == ERHIAccess::RTV);
 
 				FRDGTextureSRVDesc TextureSubResource = FRDGTextureSRVDesc::Create(TextureAccess);
-				GRDGResourceDumpContext.AddDumpTexturePasses(*this, InputResourceNames, OutputResourceNames, Pass, TextureSubResource, bIsOutputResource);
+				GRDGResourceDumpContext->AddDumpTexturePasses(*this, InputResourceNames, OutputResourceNames, Pass, TextureSubResource, bIsOutputResource);
 			}
 		}
 		break;
@@ -1805,7 +1828,7 @@ void FRDGBuilder::DumpResourcePassOutputs(const FRDGPass* Pass)
 			if (FRDGBufferSRVRef SRV = Parameter.GetAsBufferSRV())
 			{
 				FRDGBufferRef Buffer = SRV->Desc.Buffer;
-				GRDGResourceDumpContext.AddDumpBufferPass(*this, InputResourceNames, OutputResourceNames, Pass, Buffer, /* bIsOutputResource = */ false);
+				GRDGResourceDumpContext->AddDumpBufferPass(*this, InputResourceNames, OutputResourceNames, Pass, Buffer, /* bIsOutputResource = */ false);
 			}
 		}
 		break;
@@ -1814,7 +1837,7 @@ void FRDGBuilder::DumpResourcePassOutputs(const FRDGPass* Pass)
 			if (FRDGBufferUAVRef UAV = Parameter.GetAsBufferUAV())
 			{
 				FRDGBufferRef Buffer = UAV->Desc.Buffer;
-				GRDGResourceDumpContext.AddDumpBufferPass(*this, InputResourceNames, OutputResourceNames, Pass, Buffer, /* bIsOutputResource = */ true);
+				GRDGResourceDumpContext->AddDumpBufferPass(*this, InputResourceNames, OutputResourceNames, Pass, Buffer, /* bIsOutputResource = */ true);
 			}
 		}
 		break;
@@ -1826,7 +1849,7 @@ void FRDGBuilder::DumpResourcePassOutputs(const FRDGPass* Pass)
 					BufferAccess.GetAccess() == ERHIAccess::UAVCompute ||
 					BufferAccess.GetAccess() == ERHIAccess::UAVGraphics);
 
-				GRDGResourceDumpContext.AddDumpBufferPass(*this, InputResourceNames, OutputResourceNames, Pass, BufferAccess, bIsOutputResource);
+				GRDGResourceDumpContext->AddDumpBufferPass(*this, InputResourceNames, OutputResourceNames, Pass, BufferAccess, bIsOutputResource);
 			}
 		}
 		break;
@@ -1840,7 +1863,7 @@ void FRDGBuilder::DumpResourcePassOutputs(const FRDGPass* Pass)
 					BufferAccess.GetAccess() == ERHIAccess::UAVCompute ||
 					BufferAccess.GetAccess() == ERHIAccess::UAVGraphics);
 
-				GRDGResourceDumpContext.AddDumpBufferPass(*this, InputResourceNames, OutputResourceNames, Pass, BufferAccess, bIsOutputResource);
+				GRDGResourceDumpContext->AddDumpBufferPass(*this, InputResourceNames, OutputResourceNames, Pass, BufferAccess, bIsOutputResource);
 			}
 		}
 		break;
@@ -1853,7 +1876,7 @@ void FRDGBuilder::DumpResourcePassOutputs(const FRDGPass* Pass)
 			{
 				FRDGTextureRef Texture = RenderTarget.GetTexture();
 				FRDGTextureSRVDesc TextureSubResource = FRDGTextureSRVDesc::CreateForMipLevel(Texture, RenderTarget.GetMipIndex());
-				GRDGResourceDumpContext.AddDumpTextureSubResourcePass(*this, InputResourceNames, OutputResourceNames, Pass, TextureSubResource, /* bIsOutputResource = */ true);
+				GRDGResourceDumpContext->AddDumpTextureSubResourcePass(*this, InputResourceNames, OutputResourceNames, Pass, TextureSubResource, /* bIsOutputResource = */ true);
 			});
 
 			const FDepthStencilBinding& DepthStencil = RenderTargets.DepthStencil;
@@ -1865,7 +1888,7 @@ void FRDGBuilder::DumpResourcePassOutputs(const FRDGPass* Pass)
 				if (DepthStencilAccess.IsUsingDepth())
 				{
 					FRDGTextureSRVDesc TextureSubResource = FRDGTextureSRVDesc::CreateForMipLevel(Texture, 0);
-					GRDGResourceDumpContext.AddDumpTextureSubResourcePass(
+					GRDGResourceDumpContext->AddDumpTextureSubResourcePass(
 						*this, InputResourceNames, OutputResourceNames, Pass, TextureSubResource,
 						/* bIsOutputResource = */ DepthStencilAccess.IsDepthWrite());
 				}
@@ -1873,7 +1896,7 @@ void FRDGBuilder::DumpResourcePassOutputs(const FRDGPass* Pass)
 				if (DepthStencilAccess.IsUsingStencil())
 				{
 					FRDGTextureSRVDesc TextureSubResource = FRDGTextureSRVDesc::CreateWithPixelFormat(Texture, PF_X24_G8);
-					GRDGResourceDumpContext.AddDumpTextureSubResourcePass(
+					GRDGResourceDumpContext->AddDumpTextureSubResourcePass(
 						*this, InputResourceNames, OutputResourceNames, Pass, TextureSubResource,
 						/* bIsOutputResource = */ DepthStencilAccess.IsStencilWrite());
 				}
@@ -1909,12 +1932,12 @@ void FRDGBuilder::DumpResourcePassOutputs(const FRDGPass* Pass)
 		JsonObject->SetStringField(TEXT("Parameters"), FRDGResourceDumpContext::PtrToString(Pass->GetParameters().GetContents()));
 		JsonObject->SetStringField(TEXT("ParametersMetadata"), FRDGResourceDumpContext::PtrToString(Pass->GetParameters().GetMetadata()));
 		JsonObject->SetStringField(TEXT("Pointer"), FString::Printf(TEXT("%016x"), FRDGResourceDumpContext::PtrToUint(Pass)));
-		JsonObject->SetNumberField(TEXT("Id"), GRDGResourceDumpContext.PassesCount);
+		JsonObject->SetNumberField(TEXT("Id"), GRDGResourceDumpContext->PassesCount);
 		JsonObject->SetArrayField(TEXT("ParentEventScopes"), ParentEventScopeNames);
 		JsonObject->SetArrayField(TEXT("InputResources"), InputResourceNames);
 		JsonObject->SetArrayField(TEXT("OutputResources"), OutputResourceNames);
 
-		GRDGResourceDumpContext.DumpJsonToFile(JsonObject, FString(FRDGResourceDumpContext::kBaseDir) / TEXT("Passes.json"), FILEWRITE_Append);
+		GRDGResourceDumpContext->DumpJsonToFile(JsonObject, FString(FRDGResourceDumpContext::kBaseDir) / TEXT("Passes.json"), FILEWRITE_Append);
 	}
 
 	// Dump the pass' parameters
@@ -1925,7 +1948,7 @@ void FRDGBuilder::DumpResourcePassOutputs(const FRDGPass* Pass)
 			const FShaderParametersMetadata* Metadata = Pass->GetParameters().GetMetadata();
 			if (Metadata)
 			{
-				GRDGResourceDumpContext.Dump(Metadata);
+				GRDGResourceDumpContext->Dump(Metadata);
 				PassParametersByteSize = Metadata->GetSize();
 			}
 		}
@@ -1936,16 +1959,16 @@ void FRDGBuilder::DumpResourcePassOutputs(const FRDGPass* Pass)
 		}
 
 		const uint8* PassParametersContent = Pass->GetParameters().GetContents();
-		if (PassParametersContent && !GRDGResourceDumpContext.IsDumped(PassParametersContent))
+		if (PassParametersContent && !GRDGResourceDumpContext->IsDumped(PassParametersContent))
 		{
 			TArrayView<const uint8> ArrayView(PassParametersContent, PassParametersByteSize);
 			FString DumpFilePath = FRDGResourceDumpContext::kStructuresDir / FRDGResourceDumpContext::PtrToString(PassParametersContent) + TEXT(".bin");
-			GRDGResourceDumpContext.DumpBinaryToFile(ArrayView, DumpFilePath);
-			GRDGResourceDumpContext.SetDumped(PassParametersContent);
+			GRDGResourceDumpContext->DumpBinaryToFile(ArrayView, DumpFilePath);
+			GRDGResourceDumpContext->SetDumped(PassParametersContent);
 		}
 	}
 
-	GRDGResourceDumpContext.PassesCount++;
+	GRDGResourceDumpContext->PassesCount++;
 
 	bInDebugPassScope = false;
 }
@@ -1954,7 +1977,7 @@ void FRDGBuilder::DumpResourcePassOutputs(const FRDGPass* Pass)
 
 void FRDGBuilder::BeginPassDump(const FRDGPass* Pass)
 {
-	if (!GRDGResourceDumpContext.IsDumpingFrame())
+	if (!GRDGResourceDumpContext)
 	{
 		return;
 	}
@@ -1975,19 +1998,19 @@ void FRDGBuilder::BeginPassDump(const FRDGPass* Pass)
 		return;
 	}
 
-	check(GRDGResourceDumpContext.DrawDumpingPass == nullptr);
+	check(GRDGResourceDumpContext->DrawDumpingPass == nullptr);
 
-	if (GRDGResourceDumpContext.IsDumpingPass(Pass))
+	if (GRDGResourceDumpContext->IsDumpingPass(Pass))
 	{
-		GRDGResourceDumpContext.DrawDumpingPass = Pass;
-		GRDGResourceDumpContext.DrawDumpCount = 0;
+		GRDGResourceDumpContext->DrawDumpingPass = Pass;
+		GRDGResourceDumpContext->DrawDumpCount = 0;
 	}
 }
 
 // static
 void FRDGBuilder::DumpDraw(const FRDGEventName& DrawEventName)
 {
-	if (!GRDGResourceDumpContext.IsDumpingFrame())
+	if (!GRDGResourceDumpContext)
 	{
 		return;
 	}
@@ -1998,12 +2021,12 @@ void FRDGBuilder::DumpDraw(const FRDGEventName& DrawEventName)
 		return;
 	}
 
-	if (!GRDGResourceDumpContext.DrawDumpingPass)
+	if (!GRDGResourceDumpContext->DrawDumpingPass)
 	{
 		return;
 	}
 
-	const FRDGPass* Pass = GRDGResourceDumpContext.DrawDumpingPass;
+	const FRDGPass* Pass = GRDGResourceDumpContext->DrawDumpingPass;
 
 	FRHICommandListImmediate& RHICmdList = FRHICommandListExecutor::GetImmediateCommandList();
 
@@ -2024,7 +2047,7 @@ void FRDGBuilder::DumpDraw(const FRDGEventName& DrawEventName)
 			{
 				FRDGTextureRef Texture = RenderTarget.GetTexture();
 				FRDGTextureSRVDesc TextureSubResource = FRDGTextureSRVDesc::CreateForMipLevel(Texture, RenderTarget.GetMipIndex());
-				GRDGResourceDumpContext.DumpDrawTextureSubResource(
+				GRDGResourceDumpContext->DumpDrawTextureSubResource(
 					RHICmdList,
 					TextureSubResource,
 					ERHIAccess::RTV);
@@ -2039,7 +2062,7 @@ void FRDGBuilder::DumpDraw(const FRDGEventName& DrawEventName)
 				if (DepthStencilAccess.IsDepthWrite())
 				{
 					FRDGTextureSRVDesc TextureSubResource = FRDGTextureSRVDesc::CreateForMipLevel(Texture, 0);
-					GRDGResourceDumpContext.DumpDrawTextureSubResource(
+					GRDGResourceDumpContext->DumpDrawTextureSubResource(
 						RHICmdList,
 						TextureSubResource,
 						ERHIAccess::RTV);
@@ -2048,7 +2071,7 @@ void FRDGBuilder::DumpDraw(const FRDGEventName& DrawEventName)
 				if (DepthStencilAccess.IsStencilWrite())
 				{
 					FRDGTextureSRVDesc TextureSubResource = FRDGTextureSRVDesc::CreateWithPixelFormat(Texture, PF_X24_G8);
-					GRDGResourceDumpContext.DumpDrawTextureSubResource(
+					GRDGResourceDumpContext->DumpDrawTextureSubResource(
 						RHICmdList,
 						TextureSubResource,
 						ERHIAccess::RTV);
@@ -2070,21 +2093,21 @@ void FRDGBuilder::DumpDraw(const FRDGEventName& DrawEventName)
 		JsonObject->SetStringField(TEXT("DrawName"), DrawEventName.GetTCHAR());
 
 		FString DumpFilePath = FRDGResourceDumpContext::kPassesDir / FString::Printf(TEXT("Pass.%016x.Draws.json"), FRDGResourceDumpContext::PtrToUint(Pass));
-		GRDGResourceDumpContext.DumpJsonToFile(JsonObject, DumpFilePath, FILEWRITE_Append);
+		GRDGResourceDumpContext->DumpJsonToFile(JsonObject, DumpFilePath, FILEWRITE_Append);
 	}
 
-	GRDGResourceDumpContext.DrawDumpCount++;
+	GRDGResourceDumpContext->DrawDumpCount++;
 
-	if (GRDGResourceDumpContext.DrawDumpCount % 10 == 0)
+	if (GRDGResourceDumpContext->DrawDumpCount % 10 == 0)
 	{
-		UE_LOG(LogRendererCore, Display, TEXT("Dumped %d draws' resources"), GRDGResourceDumpContext.DrawDumpCount);
+		UE_LOG(LogRendererCore, Display, TEXT("Dumped %d draws' resources"), GRDGResourceDumpContext->DrawDumpCount);
 		return;
 	}
 }
 
 void FRDGBuilder::EndPassDump(const FRDGPass* Pass)
 {
-	if (!GRDGResourceDumpContext.IsDumpingFrame())
+	if (!GRDGResourceDumpContext)
 	{
 		return;
 	}
@@ -2094,41 +2117,41 @@ void FRDGBuilder::EndPassDump(const FRDGPass* Pass)
 		return;
 	}
 
-	if (!GRDGResourceDumpContext.DrawDumpingPass)
+	if (!GRDGResourceDumpContext->DrawDumpingPass)
 	{
 		return;
 	}
 
-	check(Pass == GRDGResourceDumpContext.DrawDumpingPass);
+	check(Pass == GRDGResourceDumpContext->DrawDumpingPass);
 
 	// Output how many draw has been dump for this pass.
-	if (GRDGResourceDumpContext.DrawDumpCount > 0)
+	if (GRDGResourceDumpContext->DrawDumpCount > 0)
 	{
 		FString EventNameStorage;
 
 		TSharedPtr<FJsonObject> JsonObject = MakeShareable(new FJsonObject);
 		JsonObject->SetStringField(TEXT("EventName"), GetPassEventNameWithGPUMask(Pass, EventNameStorage));
 		JsonObject->SetStringField(TEXT("Pointer"), FString::Printf(TEXT("%016x"), FRDGResourceDumpContext::PtrToUint(Pass)));
-		JsonObject->SetNumberField(TEXT("DrawCount"), GRDGResourceDumpContext.DrawDumpCount);
+		JsonObject->SetNumberField(TEXT("DrawCount"), GRDGResourceDumpContext->DrawDumpCount);
 
-		GRDGResourceDumpContext.DumpJsonToFile(JsonObject, FString(FRDGResourceDumpContext::kBaseDir) / TEXT("PassDrawCounts.json"), FILEWRITE_Append);
+		GRDGResourceDumpContext->DumpJsonToFile(JsonObject, FString(FRDGResourceDumpContext::kBaseDir) / TEXT("PassDrawCounts.json"), FILEWRITE_Append);
 
-		UE_LOG(LogRendererCore, Display, TEXT("Completed dump of %d draws for pass: %s"), GRDGResourceDumpContext.DrawDumpCount, Pass->GetEventName().GetTCHAR());
+		UE_LOG(LogRendererCore, Display, TEXT("Completed dump of %d draws for pass: %s"), GRDGResourceDumpContext->DrawDumpCount, Pass->GetEventName().GetTCHAR());
 	}
 
-	GRDGResourceDumpContext.DrawDumpingPass = nullptr;
-	GRDGResourceDumpContext.DrawDumpCount = 0;
+	GRDGResourceDumpContext->DrawDumpingPass = nullptr;
+	GRDGResourceDumpContext->DrawDumpCount = 0;
 }
 
 // static
 bool FRDGBuilder::IsDumpingFrame()
 {
-	return GRDGResourceDumpContext.IsDumpingFrame();
+	return GRDGResourceDumpContext != nullptr;
 }
 
 bool FRDGBuilder::IsDumpingDraws()
 {
-	if (!GRDGResourceDumpContext.IsDumpingFrame())
+	if (!GRDGResourceDumpContext)
 	{
 		return false;
 	}
