Fixed compilation error introduced by 19547074 colliding with 19546929

diff --git a/Engine/Source/Runtime/RenderCore/Private/DumpGPU.cpp b/Engine/Source/Runtime/RenderCore/Private/DumpGPU.cpp
--- a/Engine/Source/Runtime/RenderCore/Private/DumpGPU.cpp
+++ b/Engine/Source/Runtime/RenderCore/Private/DumpGPU.cpp
@@ -1369,7 +1369,7 @@ struct FRDGResourceDumpContext
 	{
 		if (GLog)
 		{
-			GLog->PanicFlush();
+			GLog->Panic();
 		}
 		DumpStatusToFile(TEXT("crash"));
 		FGenericCrashContext::DumpLog(DumpingDirectoryPath / FRDGResourceDumpContext::kBaseDir);
