fix GlobalShader crash in stompmalloc due to pointing at a member after a resize

diff --git a/Engine/Source/Runtime/RenderCore/Private/GlobalShader.cpp b/Engine/Source/Runtime/RenderCore/Private/GlobalShader.cpp
--- a/Engine/Source/Runtime/RenderCore/Private/GlobalShader.cpp
+++ b/Engine/Source/Runtime/RenderCore/Private/GlobalShader.cpp
@@ -189,9 +189,19 @@ private:
 			// No found platforms that support this shader format
 			return nullptr;
 		}
-
-		const FPlatformInfo& Platform = GetOrLoadPlatformInfo(IniPlatforms[0]);
 		
+		// first add all platforms to populate the map :
+		for (int32 PlatformIndex = 0; PlatformIndex < IniPlatforms.Num(); ++PlatformIndex)
+		{
+			FPlatformInfo& Platform = ConfigDefines.FindOrAdd(IniPlatforms[PlatformIndex]);
+
+			if (!Platform.bInitializedFromConfig)
+			{
+				InitializePlatform(Platform, IniPlatforms[PlatformIndex]);
+				check( Platform.bInitializedFromConfig );
+			}
+		}
+
 		if (!ErrorCheckedPlatforms[ShaderPlatform])
 		{
 			ErrorCheckedPlatforms[ShaderPlatform] = true;
@@ -200,25 +210,24 @@ private:
 			{
 				// This shader platform is shared by multiple target platforms that can be configured independently. We need to make sure all config defines
 				// match up, and that no platform-specific defines exist that might introduce shader compiler output that diverges between target platforms
+								
+				// pointer to Platform0 is safe to hold now because all are added first :
+				const FPlatformInfo * Platform0 = ConfigDefines.Find(IniPlatforms[0]);
+				check( Platform0 != nullptr );
+
 				for (int32 PlatformIndex = 1; PlatformIndex < IniPlatforms.Num(); ++PlatformIndex)
 				{
-					const FPlatformInfo& OtherPlatform = GetOrLoadPlatformInfo(IniPlatforms[PlatformIndex]);
-					ErrorCheckPlatformsForShaderFormat(Platform, OtherPlatform, OutShaderFormat);
+					const FPlatformInfo * OtherPlatform = ConfigDefines.Find(IniPlatforms[PlatformIndex]);
+					check( OtherPlatform != nullptr );
+
+					ErrorCheckPlatformsForShaderFormat(*Platform0, *OtherPlatform, OutShaderFormat);
 				}
 			}
 		}
-
-		return &Platform;
-	}
-
-	static FPlatformInfo& GetOrLoadPlatformInfo(FName PlatformName)
-	{
-		FPlatformInfo& Platform = ConfigDefines.FindOrAdd(PlatformName);
-
-		if (!Platform.bInitializedFromConfig)
-		{
-			InitializePlatform(Platform, PlatformName);
-		}
+		
+		// reget the pointer to IniPlatforms[0] after all Map adds are done, to return out :
+		// note returned pointer is not safe if any more adds are done
+		const FPlatformInfo * Platform = ConfigDefines.Find(IniPlatforms[0]);
 
 		return Platform;
 	}
