Fixed TSan warning in Shader Resource

diff --git a/Engine/Source/Runtime/RenderCore/Private/ShaderResource.cpp b/Engine/Source/Runtime/RenderCore/Private/ShaderResource.cpp
--- a/Engine/Source/Runtime/RenderCore/Private/ShaderResource.cpp
+++ b/Engine/Source/Runtime/RenderCore/Private/ShaderResource.cpp
@@ -441,19 +441,21 @@ FShaderMapResource::FShaderMapResource(EShaderPlatform InPlatform, int32 NumShad
 FShaderMapResource::~FShaderMapResource()
 {
 	ReleaseShaders();
-	check(NumRefs == 0);
+	check(NumRefs.load(std::memory_order_relaxed) == 0);
 }
 
 void FShaderMapResource::AddRef()
 {
-	FPlatformAtomics::InterlockedIncrement((volatile int32*)&NumRefs);
+	NumRefs.fetch_add(1, std::memory_order_relaxed);
 }
 
 void FShaderMapResource::Release()
 {
-	check(NumRefs > 0);
-	if (FPlatformAtomics::InterlockedDecrement((volatile int32*)&NumRefs) == 0 && TryRelease())
+	check(NumRefs.load(std::memory_order_relaxed) > 0);
+	if (NumRefs.fetch_sub(1, std::memory_order_release) - 1 == 0 && TryRelease())
 	{
+		//check https://www.boost.org/doc/libs/1_55_0/doc/html/atomic/usage_examples.html for explanation
+		std::atomic_thread_fence(std::memory_order_acquire);
 		// Send a release message to the rendering thread when the shader loses its last reference.
 		BeginReleaseResource(this);
 		BeginCleanup(this);
diff --git a/Engine/Source/Runtime/RenderCore/Public/Shader.h b/Engine/Source/Runtime/RenderCore/Public/Shader.h
--- a/Engine/Source/Runtime/RenderCore/Public/Shader.h
+++ b/Engine/Source/Runtime/RenderCore/Public/Shader.h
@@ -212,7 +212,7 @@ public:
 
 	void AddRef();
 	void Release();
-	inline int32 GetNumRefs() const { return NumRefs; }
+	inline int32 GetNumRefs() const { return NumRefs.load(std::memory_order_relaxed); }
 
 	// FRenderResource interface.
 	virtual void ReleaseRHI();
@@ -310,7 +310,7 @@ private:
 	EShaderPlatform Platform;
 
 	/** The number of references to this shader. */
-	int32 NumRefs;
+	std::atomic<int32> NumRefs;
 };
 
 class FShaderMapResourceCode : public FThreadSafeRefCountedObject
