Fix for blocked cooks due to per-shader pre-processor defines targetting non-existant shader

- shader got renamed so retarget the associated pre-processor defines to the new shader
- also convert the errors incurred to warnings as they should not be fatal to cooks

diff --git a/Engine/Source/Runtime/RenderCore/Private/GlobalShader.cpp b/Engine/Source/Runtime/RenderCore/Private/GlobalShader.cpp
--- a/Engine/Source/Runtime/RenderCore/Private/GlobalShader.cpp
+++ b/Engine/Source/Runtime/RenderCore/Private/GlobalShader.cpp
@@ -263,7 +263,7 @@ private:
 					if (ShaderType == nullptr || ShaderType->GetGlobalShaderType() == nullptr)
 					{
 						// This global shader doesn't actually exist
-						UE_LOG(LogShaders, Error, TEXT("Global shader definition '%s' found in engine config for global shader '%s', which does not exist"), *DefineName, *ShaderName);
+						UE_LOG(LogShaders, Warning, TEXT("Global shader definition '%s' found in engine config for global shader '%s', which does not exist"), *DefineName, *ShaderName);
 						continue;
 					}
 
