Disable parallel RDG if not running in a dedicated render thread - fix from zach.bethel

diff --git a/Engine/Source/Runtime/RenderCore/Private/RenderGraphPrivate.h b/Engine/Source/Runtime/RenderCore/Private/RenderGraphPrivate.h
--- a/Engine/Source/Runtime/RenderCore/Private/RenderGraphPrivate.h
+++ b/Engine/Source/Runtime/RenderCore/Private/RenderGraphPrivate.h
@@ -210,6 +210,8 @@ FORCEINLINE bool IsParallelExecuteEnabled()
 #if WITH_DUMPGPU
 		&& !IsDumpingRDGResources()
 #endif
+		// Only run parallel RDG if we have a rendering thread.
+		&& IsInActualRenderingThread()
 		;
 }
 
