Fix for crash in RDG due to bad aliasing transition.

diff --git a/Engine/Source/Runtime/RenderCore/Private/RenderGraphBuilder.cpp b/Engine/Source/Runtime/RenderCore/Private/RenderGraphBuilder.cpp
--- a/Engine/Source/Runtime/RenderCore/Private/RenderGraphBuilder.cpp
+++ b/Engine/Source/Runtime/RenderCore/Private/RenderGraphBuilder.cpp
@@ -3015,17 +3015,8 @@ void FRDGBuilder::SetRHI(FRDGTexture* Texture, FRHITransientTexture* TransientTe
 	Texture->ViewCache = &TransientTexture->ViewCache;
 	Texture->FirstPass = PassHandle;
 	Texture->bTransient = true;
-
-	const FRDGPassHandle MinAcquirePassHandle(TransientTexture->GetAcquirePasses().Min);
-
-	AddAliasingTransition(MinAcquirePassHandle, PassHandle, Texture, FRHITransientAliasingInfo::Acquire(TransientTexture->GetRHI(), TransientTexture->GetAliasingOverlaps()));
-
-	FRDGSubresourceState InitialState;
-	InitialState.SetPass(ERHIPipeline::Graphics, MinAcquirePassHandle);
-	InitialState.Access = ERHIAccess::Discard;
-
 	Texture->State = Allocator.AllocNoDestruct<FRDGTextureSubresourceState>();
-	InitAsWholeResource(*Texture->State, InitialState);
+	InitAsWholeResource(*Texture->State, {});
 }
 
 void FRDGBuilder::SetRHI(FRDGBuffer* Buffer, FRDGPooledBuffer* PooledBuffer, FRDGPassHandle PassHandle)
@@ -3066,15 +3057,7 @@ void FRDGBuilder::SetRHI(FRDGBuffer* Buffer, FRHITransientBuffer* TransientBuffe
 	Buffer->ViewCache = &TransientBuffer->ViewCache;
 	Buffer->FirstPass = PassHandle;
 	Buffer->bTransient = true;
-
-	const FRDGPassHandle MinAcquirePassHandle(TransientBuffer->GetAcquirePasses().Min);
-
-	AddAliasingTransition(MinAcquirePassHandle, PassHandle, Buffer, FRHITransientAliasingInfo::Acquire(TransientBuffer->GetRHI(), TransientBuffer->GetAliasingOverlaps()));
-
-	FRDGSubresourceState* InitialState = Allocator.AllocNoDestruct<FRDGSubresourceState>();
-	InitialState->SetPass(ERHIPipeline::Graphics, MinAcquirePassHandle);
-	InitialState->Access = ERHIAccess::Discard;
-	Buffer->State = InitialState;
+	Buffer->State = Allocator.AllocNoDestruct<FRDGSubresourceState>();
 }
 
 ///////////////////////////////////////////////////////////////////////////////////////////////////
@@ -3116,6 +3099,15 @@ void FRDGBuilder::BeginResourceRHI(FRDGPassHandle PassHandle, FRDGTextureRef Tex
 				SetRHI(Texture, TransientTexture, PassHandle);
 			}
 
+			const FRDGPassHandle MinAcquirePassHandle(TransientTexture->GetAcquirePasses().Min);
+
+			AddAliasingTransition(MinAcquirePassHandle, PassHandle, Texture, FRHITransientAliasingInfo::Acquire(TransientTexture->GetRHI(), TransientTexture->GetAliasingOverlaps()));
+
+			FRDGSubresourceState InitialState;
+			InitialState.SetPass(ERHIPipeline::Graphics, MinAcquirePassHandle);
+			InitialState.Access = ERHIAccess::Discard;
+			InitAsWholeResource(*Texture->State, InitialState);
+
 		#if STATS
 			GRDGStatTransientTextureCount++;
 		#endif
@@ -3192,6 +3184,14 @@ void FRDGBuilder::BeginResourceRHI(FRDGPassHandle PassHandle, FRDGBufferRef Buff
 		{
 			SetRHI(Buffer, TransientBuffer, PassHandle);
 
+			const FRDGPassHandle MinAcquirePassHandle(TransientBuffer->GetAcquirePasses().Min);
+
+			AddAliasingTransition(MinAcquirePassHandle, PassHandle, Buffer, FRHITransientAliasingInfo::Acquire(TransientBuffer->GetRHI(), TransientBuffer->GetAliasingOverlaps()));
+
+			FRDGSubresourceState* InitialState = Buffer->State;
+			InitialState->SetPass(ERHIPipeline::Graphics, MinAcquirePassHandle);
+			InitialState->Access = ERHIAccess::Discard;
+
 		#if STATS
 			GRDGStatTransientBufferCount++;
 		#endif
