Change ViewRect rounding to be 8 pixels instead of 4 when strata is enabled.

This allows to align view rect onto Strata tile size, which easy multi-view handling with tile rendering.

diff --git a/Engine/Source/Runtime/RenderCore/Private/RenderUtils.cpp b/Engine/Source/Runtime/RenderCore/Private/RenderUtils.cpp
--- a/Engine/Source/Runtime/RenderCore/Private/RenderUtils.cpp
+++ b/Engine/Source/Runtime/RenderCore/Private/RenderUtils.cpp
@@ -1646,12 +1646,19 @@ RENDERCORE_API FBufferRHIRef& GetUnitCubeAABBVertexBuffer()
 }
 #endif // RHI_RAYTRACING
 
+bool RenderCore_IsStrataEnabled();
+
 RENDERCORE_API void QuantizeSceneBufferSize(const FIntPoint& InBufferSize, FIntPoint& OutBufferSize)
 {
+	// Ensure sizes are dividable by STRATA_TILE_SIZE (==8) 2d tiles to make it more convenient.
+	const uint32 StrataDividableBy = 8;
+	static_assert(StrataDividableBy % 8 == 0, "A lot of graphic algorithms where previously assuming DividableBy >= 4");
+
 	// Ensure sizes are dividable by the ideal group size for 2d tiles to make it more convenient.
-	const uint32 DividableBy = 4;
+	const uint32 LegacyDividableBy = 4;
+	static_assert(LegacyDividableBy % 4 == 0, "A lot of graphic algorithms where previously assuming DividableBy == 4");
 
-	static_assert(DividableBy % 4 == 0, "A lot of graphic algorithms where previously assuming DividableBy == 4");
+	const uint32 DividableBy = RenderCore_IsStrataEnabled() ? StrataDividableBy : LegacyDividableBy;
 
 	const uint32 Mask = ~(DividableBy - 1);
 	OutBufferSize.X = (InBufferSize.X + DividableBy - 1) & Mask;
