Fixed: Use SceneDepthAux texture for depth capture on mobile
Fixed: Correct viewport computations in AddClearRenderTargetPass

diff --git a/Engine/Shaders/Private/SceneCapturePixelShader.usf b/Engine/Shaders/Private/SceneCapturePixelShader.usf
--- a/Engine/Shaders/Private/SceneCapturePixelShader.usf
+++ b/Engine/Shaders/Private/SceneCapturePixelShader.usf
@@ -4,8 +4,8 @@
 	SceneCapturePixelShader.usf: 
 =============================================================================*/
 
-// do not use framebuffer fetch
-#define FORCE_DEPTH_TEXTURE_READS 1
+// Sample SceneDepthAux texture on mobile
+#define POST_PROCESS_MATERIAL_MOBILE 1
 
 #include "Common.ush"
 #include "DeferredShadingCommon.ush"
diff --git a/Engine/Source/Runtime/RenderCore/Private/RenderGraphUtils.cpp b/Engine/Source/Runtime/RenderCore/Private/RenderGraphUtils.cpp
--- a/Engine/Source/Runtime/RenderCore/Private/RenderGraphUtils.cpp
+++ b/Engine/Source/Runtime/RenderCore/Private/RenderGraphUtils.cpp
@@ -595,8 +595,8 @@ void AddClearRenderTargetPass(FRDGBuilder& GraphBuilder, FRDGTextureRef Texture,
 			for (uint32 MipIndex = 0; MipIndex < TextureClearInfo.NumMips; ++MipIndex)
 			{
 				FIntRect CurrentViewport(
-					FMath::Max(1u, (uint32)OriginalViewport.Min.X >> MipIndex),
-					FMath::Max(1u, (uint32)OriginalViewport.Min.Y >> MipIndex),
+					(uint32)OriginalViewport.Min.X >> MipIndex, 
+					(uint32)OriginalViewport.Min.Y >> MipIndex,
 					FMath::Max(1u, (uint32)OriginalViewport.Max.X >> MipIndex),
 					FMath::Max(1u, (uint32)OriginalViewport.Max.Y >> MipIndex));
 
