Fixed crash due to invalid extracted transient render target allocation lifetime.

diff --git a/Engine/Source/Runtime/RenderCore/Private/RenderGraphBuilder.cpp b/Engine/Source/Runtime/RenderCore/Private/RenderGraphBuilder.cpp
--- a/Engine/Source/Runtime/RenderCore/Private/RenderGraphBuilder.cpp
+++ b/Engine/Source/Runtime/RenderCore/Private/RenderGraphBuilder.cpp
@@ -2606,12 +2606,16 @@ void FRDGBuilder::AddEpilogueTransition(FRDGTextureRef Texture)
 		EpilogueResourceAccesses.Emplace(Texture->GetRHI(), Texture->EpilogueAccess);
 	}
 
-	if (Texture->RenderTarget)
+	if (Texture->Allocation)
 	{
+		ActivePooledTextures.Emplace(MoveTemp(Texture->Allocation));
+	}
+	else if (!Texture->bTransient)
+	{
+		// Non-transient render targets need to be 'resurrected' to hold the last reference in the chain.
+		// Transient render targets don't have that restriction since there's no actual pooling happening.
 		ActivePooledTextures.Emplace(Texture->RenderTarget);
 	}
-
-	Texture->Allocation = nullptr;
 }
 
 void FRDGBuilder::AddEpilogueTransition(FRDGBufferRef Buffer)
@@ -3294,8 +3298,11 @@ void FRDGBuilder::EndResourceRHI(FRDGPassHandle PassHandle, FRDGTextureRef Textu
 			// Texture is using a transient external render target.
 			if (Texture->RenderTarget)
 			{
-				// This releases the reference without invoking a virtual function call.
-				GRDGTransientResourceAllocator.Release(TRefCountPtr<FRDGTransientRenderTarget>(MoveTemp(Texture->Allocation)), PassHandle);
+				if (!Texture->bExtracted)
+				{
+					// This releases the reference without invoking a virtual function call.
+					GRDGTransientResourceAllocator.Release(TRefCountPtr<FRDGTransientRenderTarget>(MoveTemp(Texture->Allocation)), PassHandle);
+				}
 			}
 			// Texture is using an internal transient texture.
 			else
