Fix for SRV array slice not referencing the correct subresource in RDG.

diff --git a/Engine/Source/Runtime/RenderCore/Public/RenderGraphResources.inl b/Engine/Source/Runtime/RenderCore/Public/RenderGraphResources.inl
--- a/Engine/Source/Runtime/RenderCore/Public/RenderGraphResources.inl
+++ b/Engine/Source/Runtime/RenderCore/Public/RenderGraphResources.inl
@@ -134,6 +134,7 @@ inline FRDGTextureSubresourceRange FRDGTextureSRV::GetSubresourceRange() const
 {
 	FRDGTextureSubresourceRange Range = GetParent()->GetSubresourceRange();
 	Range.MipIndex = Desc.MipLevel;
+	Range.ArraySlice = Desc.FirstArraySlice;
 	Range.PlaneSlice = GetResourceTransitionPlaneForMetadataAccess(Desc.MetaData);
 
 	if (Desc.MetaData == ERDGTextureMetaDataAccess::None && Desc.Texture && Desc.Texture->Desc.Format == PF_DepthStencil)
