Added draw texture utility to RenderGraphUtils (used for copying with format conversions).

diff --git a/Engine/Shaders/Private/Tools/DrawTexture.usf b/Engine/Shaders/Private/Tools/DrawTexture.usf
new file mode 100644
--- /dev/null
+++ b/Engine/Shaders/Private/Tools/DrawTexture.usf
@@ -0,0 +1,11 @@
+// Copyright Epic Games, Inc. All Rights Reserved.
+
+#include "../Common.ush"
+
+Texture2D InputTexture;
+uint2 InputOffset;
+
+float4 DrawTexturePS(float4 SvPosition : SV_POSITION) : SV_Target0
+{
+	return InputTexture.Load(uint3((uint2)SvPosition.xy + InputOffset, 0));
+}
\ No newline at end of file
diff --git a/Engine/Source/Runtime/RenderCore/Private/RenderGraphUtils.cpp b/Engine/Source/Runtime/RenderCore/Private/RenderGraphUtils.cpp
--- a/Engine/Source/Runtime/RenderCore/Private/RenderGraphUtils.cpp
+++ b/Engine/Source/Runtime/RenderCore/Private/RenderGraphUtils.cpp
@@ -205,6 +205,20 @@ BEGIN_SHADER_PARAMETER_STRUCT(FCopyTextureParameters, )
 	RDG_TEXTURE_ACCESS(Output, ERHIAccess::CopyDest)
 END_SHADER_PARAMETER_STRUCT()
 
+class RENDERCORE_API FDrawTexturePS : public FGlobalShader
+{
+	DECLARE_GLOBAL_SHADER(FDrawTexturePS);
+	SHADER_USE_PARAMETER_STRUCT(FDrawTexturePS, FGlobalShader);
+
+	BEGIN_SHADER_PARAMETER_STRUCT(FParameters, )
+		SHADER_PARAMETER_RDG_TEXTURE_SRV(Texture2D, InputTexture)
+		SHADER_PARAMETER(FIntPoint, InputOffset)
+		RENDER_TARGET_BINDING_SLOTS()
+	END_SHADER_PARAMETER_STRUCT()
+};
+
+IMPLEMENT_GLOBAL_SHADER(FDrawTexturePS, "/Engine/Private/Tools/DrawTexture.usf", "DrawTexturePS", SF_Pixel);
+
 void AddCopyTexturePass(
 	FRDGBuilder& GraphBuilder,
 	FRDGTextureRef InputTexture,
@@ -229,6 +243,81 @@ void AddCopyTexturePass(
 	});
 }
 
+RENDERCORE_API void AddDrawTexturePass(
+	FRDGBuilder& GraphBuilder,
+	const FGlobalShaderMap* ShaderMap,
+	FRDGTextureRef InputTexture,
+	FRDGTextureRef OutputTexture,
+	const FRDGDrawTextureInfo& DrawInfo)
+{
+	const FRDGTextureDesc& InputDesc = InputTexture->Desc;
+	const FRDGTextureDesc& OutputDesc = OutputTexture->Desc;
+
+	// Use a hardware copy if formats match.
+	if (InputDesc.Format == OutputDesc.Format)
+	{
+		FRHICopyTextureInfo CopyInfo;
+
+		// Translate the draw texture info into a copy info.
+		CopyInfo.Size = FIntVector(DrawInfo.Size.X, DrawInfo.Size.Y, 0);
+		CopyInfo.SourcePosition = FIntVector(DrawInfo.SourcePosition.X, DrawInfo.SourcePosition.Y, 0);
+		CopyInfo.DestPosition = FIntVector(DrawInfo.DestPosition.X, DrawInfo.DestPosition.Y, 0);
+		CopyInfo.SourceSliceIndex = DrawInfo.SourceSliceIndex;
+		CopyInfo.DestSliceIndex = DrawInfo.DestSliceIndex;
+		CopyInfo.NumSlices = DrawInfo.NumSlices;
+		CopyInfo.SourceMipIndex = DrawInfo.SourceMipIndex;
+		CopyInfo.DestMipIndex = DrawInfo.DestMipIndex;
+		CopyInfo.NumMips = DrawInfo.NumMips;
+
+		AddCopyTexturePass(GraphBuilder, InputTexture, OutputTexture, CopyInfo);
+	}
+	else
+	{
+		const FIntPoint DrawSize = DrawInfo.Size == FIntPoint::ZeroValue ? OutputDesc.Extent : DrawInfo.Size;
+
+		// Don't load color data if the whole texture is being overwritten.
+		const ERenderTargetLoadAction LoadAction = (DrawInfo.DestPosition == FIntPoint::ZeroValue && DrawSize == OutputDesc.Extent)
+			? ERenderTargetLoadAction::ENoAction
+			: ERenderTargetLoadAction::ELoad;
+
+		TShaderMapRef<FDrawTexturePS> PixelShader(ShaderMap);
+
+		for (uint32 MipIndex = 0; MipIndex < DrawInfo.NumMips; ++MipIndex)
+		{
+			const int32 SourceMipIndex = MipIndex + DrawInfo.SourceMipIndex;
+			const int32 DestMipIndex = MipIndex + DrawInfo.DestMipIndex;
+
+			for (uint32 SliceIndex = 0; SliceIndex < DrawInfo.NumSlices; ++SliceIndex)
+			{
+				const int32 SourceSliceIndex = SliceIndex + DrawInfo.SourceSliceIndex;
+				const int32 DestSliceIndex = SliceIndex + DrawInfo.DestSliceIndex;
+
+				FRDGTextureSRVDesc SRVDesc = FRDGTextureSRVDesc::Create(InputTexture);
+				SRVDesc.FirstArraySlice = SourceSliceIndex;
+				SRVDesc.NumArraySlices = 1;
+				SRVDesc.MipLevel = SourceMipIndex;
+				SRVDesc.NumMipLevels = 1;
+
+				auto* PassParameters = GraphBuilder.AllocParameters<FDrawTexturePS::FParameters>();
+				PassParameters->InputTexture = GraphBuilder.CreateSRV(SRVDesc);
+				PassParameters->InputOffset = DrawInfo.SourcePosition;
+				PassParameters->RenderTargets[0] = FRenderTargetBinding(OutputTexture, LoadAction, DestMipIndex, DestSliceIndex);
+
+				const FIntRect ViewRect(DrawInfo.DestPosition, DrawInfo.DestPosition + DrawSize);
+
+				FPixelShaderUtils::AddFullscreenPass(
+					GraphBuilder,
+					ShaderMap,
+					RDG_EVENT_NAME("DrawTexture ([%s, Mip: %d, Slice: %d] -> [%s, Mip: %d, Slice: %d])", InputTexture->Name, SourceMipIndex, SourceSliceIndex, OutputTexture->Name, DestMipIndex, DestSliceIndex),
+					PixelShader,
+					PassParameters,
+					ViewRect
+				);
+			}
+		}
+	}
+}
+
 BEGIN_SHADER_PARAMETER_STRUCT(FCopyToResolveTargetParameters, )
 	RDG_TEXTURE_ACCESS_DYNAMIC(Input)
 	RDG_TEXTURE_ACCESS_DYNAMIC(Output)
diff --git a/Engine/Source/Runtime/RenderCore/Public/RenderGraphUtils.h b/Engine/Source/Runtime/RenderCore/Public/RenderGraphUtils.h
--- a/Engine/Source/Runtime/RenderCore/Public/RenderGraphUtils.h
+++ b/Engine/Source/Runtime/RenderCore/Public/RenderGraphUtils.h
@@ -630,6 +630,32 @@ inline void AddCopyTexturePass(
 	AddCopyTexturePass(GraphBuilder, InputTexture, OutputTexture, CopyInfo);
 }
 
+struct FRDGDrawTextureInfo
+{
+	// Number of texels to copy. By default it will copy the whole resource if no size is specified.
+	FIntPoint Size = FIntPoint::ZeroValue;
+
+	// Position of the copy from the source texture/to destination texture
+	FIntPoint SourcePosition = FIntPoint::ZeroValue;
+	FIntPoint DestPosition = FIntPoint::ZeroValue;
+
+	uint32 SourceSliceIndex = 0;
+	uint32 DestSliceIndex = 0;
+	uint32 NumSlices = 1;
+
+	// Mips to copy and destination mips
+	uint32 SourceMipIndex = 0;
+	uint32 DestMipIndex = 0;
+	uint32 NumMips = 1;
+};
+
+RENDERCORE_API void AddDrawTexturePass(
+	FRDGBuilder& GraphBuilder,
+	const FGlobalShaderMap* ShaderMap,
+	FRDGTextureRef InputTexture,
+	FRDGTextureRef OutputTexture,
+	const FRDGDrawTextureInfo& DrawInfo);
+
 /** Adds a render graph pass to resolve from one texture to another. Uses RHICopyToResolveTarget under the hood.
  *  The formats of the two textures don't need to match.
  */
