FVertexType and FShaderType AddReferencedUniformBufferIncludes: Fix unguarded reading from ReferencedUniformBufferStructsCache by putting it under a read lock.
#rnx

diff --git a/Engine/Source/Runtime/RenderCore/Private/ShaderParameters.cpp b/Engine/Source/Runtime/RenderCore/Private/ShaderParameters.cpp
--- a/Engine/Source/Runtime/RenderCore/Private/ShaderParameters.cpp
+++ b/Engine/Source/Runtime/RenderCore/Private/ShaderParameters.cpp
@@ -11,6 +11,7 @@
 #include "Shader.h"
 #include "VertexFactory.h"
 #include "ShaderCodeLibrary.h"
+#include "Misc/ScopeRWLock.h"
 #include "Misc/ScopeLock.h"
 
 IMPLEMENT_TYPE_LAYOUT(FShaderParameter);
@@ -328,14 +329,14 @@ RENDERCORE_API void CacheUniformBufferIncludes(TMap<const TCHAR*,FCachedUniformB
 }
 
 static const uint32 NumUniformBufferLocks = 16u;
-static FCriticalSection UniformBufferLocks[NumUniformBufferLocks];
+static FRWLock UniformBufferLocks[NumUniformBufferLocks];
 
 void FShaderType::FlushShaderFileCache(const TMap<FString, TArray<const TCHAR*> >& ShaderFileToUniformBufferVariables)
 {
 	if (CachedUniformBufferPlatform != SP_NumPlatforms)
 	{
 		const uint32 LockIndex = HashedName.GetHash() % NumUniformBufferLocks;
-		FScopeLock Lock(&UniformBufferLocks[LockIndex]);
+		FWriteScopeLock Lock(UniformBufferLocks[LockIndex]);
 		if (CachedUniformBufferPlatform != SP_NumPlatforms)
 		{
 			ReferencedUniformBufferStructsCache.Empty();
@@ -349,24 +350,38 @@ void FShaderType::AddReferencedUniformBufferIncludes(FShaderCompilerEnvironment&
 {
 	TRACE_CPUPROFILER_EVENT_SCOPE(FShaderType::AddReferencedUniformBufferIncludes);
 
-	// Cache uniform buffer struct declarations referenced by this shader type's files
-	if (CachedUniformBufferPlatform != Platform)
+	const uint32 LockIndex = HashedName.GetHash() % NumUniformBufferLocks;
+	TOptional<FReadScopeLock> ReadScopeLock; // TOptional is so we can unlock and lock manually; the FReadScopeLock api does not support that
+	ReadScopeLock.Emplace(UniformBufferLocks[LockIndex]);
+	for (;;)
 	{
-		const uint32 LockIndex = HashedName.GetHash() % NumUniformBufferLocks;
-		FScopeLock Lock(&UniformBufferLocks[LockIndex]);
-		if (CachedUniformBufferPlatform != Platform)
+		// Cache uniform buffer struct declarations referenced by this shader type's files
+		if (CachedUniformBufferPlatform == Platform)
 		{
-			// if there is already a cache but for another platform, keep the keys but reset the values
-			if (CachedUniformBufferPlatform != SP_NumPlatforms)
+			break;
+		}
+
+		// Drop the ReadLock, acquire a write lock
+		//     Write the ReferencedUniformBufferStructsCache for the new Platform
+		// Drop the WriteLock, reacquire the ReadLock
+		ReadScopeLock.Reset();
+		{
+			FWriteScopeLock WriteScopeLock(UniformBufferLocks[LockIndex]);
+			if (CachedUniformBufferPlatform != Platform)
 			{
-				for (TMap<const TCHAR*, FCachedUniformBufferDeclaration>::TIterator It(ReferencedUniformBufferStructsCache); It; ++It)
+				// If there is already a cache but for another platform, keep the keys but reset the values
+				if (CachedUniformBufferPlatform != SP_NumPlatforms)
 				{
-					It.Value() = FCachedUniformBufferDeclaration();
+					for (TMap<const TCHAR*, FCachedUniformBufferDeclaration>::TIterator It(ReferencedUniformBufferStructsCache); It; ++It)
+					{
+						It.Value() = FCachedUniformBufferDeclaration();
+					}
 				}
+				CacheUniformBufferIncludes(ReferencedUniformBufferStructsCache, Platform);
+				CachedUniformBufferPlatform = Platform;
 			}
-			CacheUniformBufferIncludes(ReferencedUniformBufferStructsCache, Platform);
-			CachedUniformBufferPlatform = Platform;
 		}
+		ReadScopeLock.Emplace(UniformBufferLocks[LockIndex]);
 	}
 
 	FString UniformBufferIncludes;
@@ -463,7 +478,7 @@ void FVertexFactoryType::FlushShaderFileCache(const TMap<FString, TArray<const T
 	if (CachedUniformBufferPlatform != SP_NumPlatforms)
 	{
 		const uint32 LockIndex = HashedName.GetHash() % NumUniformBufferLocks;
-		FScopeLock Lock(&UniformBufferLocks[LockIndex]);
+		FWriteScopeLock Lock(UniformBufferLocks[LockIndex]);
 		if (CachedUniformBufferPlatform != SP_NumPlatforms)
 		{
 			ReferencedUniformBufferStructsCache.Empty();
@@ -475,24 +490,38 @@ void FVertexFactoryType::FlushShaderFileCache(const TMap<FString, TArray<const T
 
 void FVertexFactoryType::AddReferencedUniformBufferIncludes(FShaderCompilerEnvironment& OutEnvironment, FString& OutSourceFilePrefix, EShaderPlatform Platform) const
 {
-	// Cache uniform buffer struct declarations referenced by this shader type's files
-	if (CachedUniformBufferPlatform != Platform)
+	const uint32 LockIndex = HashedName.GetHash() % NumUniformBufferLocks;
+	TOptional<FReadScopeLock> ReadScopeLock; // TOptional is so we can unlock and lock manually; the FReadScopeLock api does not support that
+	ReadScopeLock.Emplace(UniformBufferLocks[LockIndex]);
+	for (;;)
 	{
-		const uint32 LockIndex = HashedName.GetHash() % NumUniformBufferLocks;
-		FScopeLock Lock(&UniformBufferLocks[LockIndex]);
-		if (CachedUniformBufferPlatform != Platform)
+		// Cache uniform buffer struct declarations referenced by this shader type's files
+		if (CachedUniformBufferPlatform == Platform)
 		{
-			// if there is already a cache but for another platform, keep the keys but reset the values
-			if (CachedUniformBufferPlatform != SP_NumPlatforms)
+			break;
+		}
+
+		// Drop the ReadLock, acquire a write lock
+		//     Write the ReferencedUniformBufferStructsCache for the new Platform
+		// Drop the WriteLock, reacquire the ReadLock
+		ReadScopeLock.Reset();
+		{
+			FWriteScopeLock WriteScopeLock(UniformBufferLocks[LockIndex]);
+			if (CachedUniformBufferPlatform != Platform)
 			{
-				for (TMap<const TCHAR*, FCachedUniformBufferDeclaration>::TIterator It(ReferencedUniformBufferStructsCache); It; ++It)
+				// If there is already a cache but for another platform, keep the keys but reset the values
+				if (CachedUniformBufferPlatform != SP_NumPlatforms)
 				{
-					It.Value() = FCachedUniformBufferDeclaration();
+					for (TMap<const TCHAR*, FCachedUniformBufferDeclaration>::TIterator It(ReferencedUniformBufferStructsCache); It; ++It)
+					{
+						It.Value() = FCachedUniformBufferDeclaration();
+					}
 				}
+				CacheUniformBufferIncludes(ReferencedUniformBufferStructsCache, Platform);
+				CachedUniformBufferPlatform = Platform;
 			}
-			CacheUniformBufferIncludes(ReferencedUniformBufferStructsCache, Platform);
-			CachedUniformBufferPlatform = Platform;
-		} 
+		}
+		ReadScopeLock.Emplace(UniformBufferLocks[LockIndex]);
 	}
 
 	FString UniformBufferIncludes;
