Moving the bulk of SetShaderParameters to a cpp file. The shader type template arguments are only used to get data that's common for all binding paths and the command list types are minimal and boiled down to just 3 sets of parameters.

diff --git a/Engine/Source/Runtime/RenderCore/Private/ShaderParameterStruct.cpp b/Engine/Source/Runtime/RenderCore/Private/ShaderParameterStruct.cpp
--- a/Engine/Source/Runtime/RenderCore/Private/ShaderParameterStruct.cpp
+++ b/Engine/Source/Runtime/RenderCore/Private/ShaderParameterStruct.cpp
@@ -560,3 +560,282 @@ void ValidateShaderParameterResourcesRHI(const void* Contents, const FRHIUniform
 }
 
 #endif // DO_CHECK
+
+template<typename TRHICmdList>
+inline void SetShaderUAV(TRHICmdList& RHICmdList, FRHIGraphicsShader* ShaderRHI, const uint8* Base, const FShaderParameterBindings::FResourceParameter& ParameterBinding)
+{
+	if (ShaderRHI->GetType() == RRT_PixelShader)
+	{
+		FRHIPixelShader* PixelShaderRHI = static_cast<FRHIPixelShader*>(ShaderRHI);
+
+		if (ParameterBinding.BaseType == UBMT_UAV)
+		{
+			FRHIUnorderedAccessView* ShaderParameterRef = *(FRHIUnorderedAccessView**)(Base + ParameterBinding.ByteOffset);
+			RHICmdList.SetUAVParameter(PixelShaderRHI, ParameterBinding.BaseIndex, ShaderParameterRef);
+		}
+		else if (ParameterBinding.BaseType == UBMT_RDG_TEXTURE_UAV || ParameterBinding.BaseType == UBMT_RDG_BUFFER_UAV)
+		{
+			auto GraphUAV = *reinterpret_cast<FRDGUnorderedAccessView* const*>(Base + ParameterBinding.ByteOffset);
+
+			checkSlow(GraphUAV);
+			GraphUAV->MarkResourceAsUsed();
+			RHICmdList.SetUAVParameter(PixelShaderRHI, ParameterBinding.BaseIndex, GraphUAV->GetRHI());
+		}
+	}
+	else
+	{
+		checkf(false, TEXT("TShaderRHI Can't have compute shader to be set. UAVs are not supported on vertex, tessellation and geometry shaders."));
+	}
+}
+
+template<typename TRHICmdList>
+inline void SetShaderUAV(TRHICmdList& RHICmdList, FRHIComputeShader* ShaderRHI, const uint8* Base, const FShaderParameterBindings::FResourceParameter& ParameterBinding)
+{
+	if (ParameterBinding.BaseType == UBMT_UAV)
+	{
+		FRHIUnorderedAccessView* ShaderParameterRef = *(FRHIUnorderedAccessView**)(Base + ParameterBinding.ByteOffset);
+		RHICmdList.SetUAVParameter(ShaderRHI, ParameterBinding.BaseIndex, ShaderParameterRef);
+	}
+	else if (ParameterBinding.BaseType == UBMT_RDG_TEXTURE_UAV || ParameterBinding.BaseType == UBMT_RDG_BUFFER_UAV)
+	{
+		auto GraphUAV = *reinterpret_cast<FRDGUnorderedAccessView* const*>(Base + ParameterBinding.ByteOffset);
+
+		checkSlow(GraphUAV);
+		GraphUAV->MarkResourceAsUsed();
+		RHICmdList.SetUAVParameter(ShaderRHI, ParameterBinding.BaseIndex, GraphUAV->GetRHI());
+	}
+}
+
+/** Set shader's parameters from its parameters struct. */
+template<typename TRHICmdList, typename TShaderRHI>
+inline void SetShaderParametersInternal(
+	TRHICmdList& RHICmdList,
+	TShaderRHI* ShaderRHI,
+	const FShaderParameterBindings& Bindings,
+	const FShaderParametersMetadata* ParametersMetadata,
+	const uint8* Base)
+{
+	checkf(Bindings.RootParameterBufferIndex == FShaderParameterBindings::kInvalidBufferIndex, TEXT("Can't use SetShaderParameters() for root parameter buffer index."));
+
+	// Parameters
+	for (const FShaderParameterBindings::FParameter& ParameterBinding : Bindings.Parameters)
+	{
+		const void* DataPtr = Base + ParameterBinding.ByteOffset;
+		RHICmdList.SetShaderParameter(ShaderRHI, ParameterBinding.BufferIndex, ParameterBinding.BaseIndex, ParameterBinding.ByteSize, DataPtr);
+	}
+
+	TArray<FShaderParameterBindings::FResourceParameter, TInlineAllocator<16>> GraphSRVs;
+
+	for (const FShaderParameterBindings::FResourceParameter& ParameterBinding : Bindings.ResourceParameters)
+	{
+		EUniformBufferBaseType BaseType = (EUniformBufferBaseType)ParameterBinding.BaseType;
+		switch (BaseType)
+		{
+		case UBMT_TEXTURE:
+		{
+			FRHITexture* TexRef = *(FRHITexture**)(Base + ParameterBinding.ByteOffset);
+			RHICmdList.SetShaderTexture(ShaderRHI, ParameterBinding.BaseIndex, TexRef);
+		}
+		break;
+		case UBMT_SRV:
+		{
+			FRHIShaderResourceView* SRVRef = *(FRHIShaderResourceView**)(Base + ParameterBinding.ByteOffset);
+			RHICmdList.SetShaderResourceViewParameter(ShaderRHI, ParameterBinding.BaseIndex, SRVRef);
+		}
+		break;
+		case UBMT_SAMPLER:
+		{
+			FRHISamplerState* SamplerRef = *(FRHISamplerState**)(Base + ParameterBinding.ByteOffset);
+			RHICmdList.SetShaderSampler(ShaderRHI, ParameterBinding.BaseIndex, SamplerRef);
+		}
+		break;
+		case UBMT_RDG_TEXTURE:
+		{
+			auto GraphTexture = *reinterpret_cast<FRDGTexture* const*>(Base + ParameterBinding.ByteOffset);
+			checkSlow(GraphTexture);
+			GraphTexture->MarkResourceAsUsed();
+			RHICmdList.SetShaderTexture(ShaderRHI, ParameterBinding.BaseIndex, GraphTexture->GetRHI());
+		}
+		break;
+		case UBMT_RDG_TEXTURE_SRV:
+		case UBMT_RDG_BUFFER_SRV:
+		{
+			//HACKHACK: defer SRVs binding after UAVs 
+			GraphSRVs.Add(ParameterBinding);
+		}
+		break;
+		case UBMT_UAV:
+		case UBMT_RDG_TEXTURE_UAV:
+		case UBMT_RDG_BUFFER_UAV:
+		{
+			SetShaderUAV(RHICmdList, ShaderRHI, Base, ParameterBinding);
+		}
+		break;
+		default:
+			checkf(false, TEXT("Unhandled resource type?"));
+			break;
+		}
+	}
+
+	//HACKHACK: Bind SRVs after UAVs as a workaround for D3D11 RHI unbinding SRVs when binding a UAV on the same resource even when the views don't overlap.
+	for (const FShaderParameterBindings::FResourceParameter& ParameterBinding : GraphSRVs)
+	{
+		auto GraphSRV = *reinterpret_cast<FRDGShaderResourceView* const*>(Base + ParameterBinding.ByteOffset);
+
+		checkSlow(GraphSRV);
+		GraphSRV->MarkResourceAsUsed();
+		RHICmdList.SetShaderResourceViewParameter(ShaderRHI, ParameterBinding.BaseIndex, GraphSRV->GetRHI());
+	}
+
+	// Graph Uniform Buffers
+	for (const FShaderParameterBindings::FParameterStructReference& ParameterBinding : Bindings.GraphUniformBuffers)
+	{
+		const FRDGUniformBufferBinding& UniformBufferBinding = *reinterpret_cast<const FRDGUniformBufferBinding*>(Base + ParameterBinding.ByteOffset);
+
+		if (UniformBufferBinding.IsShader())
+		{
+			UniformBufferBinding->MarkResourceAsUsed();
+			RHICmdList.SetShaderUniformBuffer(ShaderRHI, ParameterBinding.BufferIndex, UniformBufferBinding->GetRHI());
+		}
+	}
+
+	// Reference structures
+	for (const FShaderParameterBindings::FParameterStructReference& ParameterBinding : Bindings.ParameterReferences)
+	{
+		const FUniformBufferBinding& UniformBufferBinding = *reinterpret_cast<const FUniformBufferBinding*>(Base + ParameterBinding.ByteOffset);
+
+		if (UniformBufferBinding.IsShader())
+		{
+			RHICmdList.SetShaderUniformBuffer(ShaderRHI, ParameterBinding.BufferIndex, UniformBufferBinding.GetUniformBuffer());
+		}
+	}
+}
+
+void SetShaderParameters(
+	FRHIComputeCommandList& RHICmdList,
+	FRHIComputeShader* ShaderRHI,
+	const FShaderParameterBindings& Bindings,
+	const FShaderParametersMetadata* ParametersMetadata,
+	const uint8* Base)
+{
+	SetShaderParametersInternal(RHICmdList, ShaderRHI, Bindings, ParametersMetadata, Base);
+}
+
+void SetShaderParameters(
+	FRHICommandList& RHICmdList,
+	FRHIGraphicsShader* ShaderRHI,
+	const FShaderParameterBindings& Bindings,
+	const FShaderParametersMetadata* ParametersMetadata,
+	const uint8* Base)
+{
+	SetShaderParametersInternal(RHICmdList, ShaderRHI, Bindings, ParametersMetadata, Base);
+}
+
+void SetShaderParameters(
+	FRHICommandList& RHICmdList,
+	FRHIComputeShader* ShaderRHI,
+	const FShaderParameterBindings& Bindings,
+	const FShaderParametersMetadata* ParametersMetadata,
+	const uint8* Base)
+{
+	SetShaderParametersInternal(RHICmdList, ShaderRHI, Bindings, ParametersMetadata, Base);
+}
+
+#if RHI_RAYTRACING
+void SetShaderParameters(
+	FRayTracingShaderBindingsWriter& RTBindingsWriter,
+	const FShaderParameterBindings& Bindings,
+	const FRHIUniformBufferLayout* RootUniformBufferLayout,
+	const uint8* Base)
+{
+	for (const FShaderParameterBindings::FResourceParameter& ParameterBinding : Bindings.ResourceParameters)
+	{
+		EUniformBufferBaseType BaseType = (EUniformBufferBaseType)ParameterBinding.BaseType;
+		switch (BaseType)
+		{
+		case UBMT_TEXTURE:
+		{
+			auto ShaderParameterRef = *(FRHITexture**)(Base + ParameterBinding.ByteOffset);
+			RTBindingsWriter.SetTexture(ParameterBinding.BaseIndex, ShaderParameterRef);
+		}
+		break;
+		case UBMT_SRV:
+		{
+			FRHIShaderResourceView* ShaderParameterRef = *(FRHIShaderResourceView**)(Base + ParameterBinding.ByteOffset);
+			RTBindingsWriter.SetSRV(ParameterBinding.BaseIndex, ShaderParameterRef);
+		}
+		break;
+		case UBMT_UAV:
+		{
+			FRHIUnorderedAccessView* ShaderParameterRef = *(FRHIUnorderedAccessView**)(Base + ParameterBinding.ByteOffset);
+			RTBindingsWriter.SetUAV(ParameterBinding.BaseIndex, ShaderParameterRef);
+		}
+		break;
+		case UBMT_SAMPLER:
+		{
+			FRHISamplerState* ShaderParameterRef = *(FRHISamplerState**)(Base + ParameterBinding.ByteOffset);
+			RTBindingsWriter.SetSampler(ParameterBinding.BaseIndex, ShaderParameterRef);
+		}
+		break;
+		case UBMT_RDG_TEXTURE:
+		{
+			auto GraphTexture = *reinterpret_cast<FRDGTexture* const*>(Base + ParameterBinding.ByteOffset);
+			checkSlow(GraphTexture);
+			GraphTexture->MarkResourceAsUsed();
+			RTBindingsWriter.SetTexture(ParameterBinding.BaseIndex, GraphTexture->GetRHI());
+		}
+		break;
+		case UBMT_RDG_TEXTURE_SRV:
+		case UBMT_RDG_BUFFER_SRV:
+		{
+			auto GraphSRV = *reinterpret_cast<FRDGShaderResourceView* const*>(Base + ParameterBinding.ByteOffset);
+
+			checkSlow(GraphSRV);
+			GraphSRV->MarkResourceAsUsed();
+			RTBindingsWriter.SetSRV(ParameterBinding.BaseIndex, GraphSRV->GetRHI());
+		}
+		break;
+		case UBMT_RDG_TEXTURE_UAV:
+		case UBMT_RDG_BUFFER_UAV:
+		{
+			auto UAV = *reinterpret_cast<FRDGUnorderedAccessView* const*>(Base + ParameterBinding.ByteOffset);
+
+			checkSlow(UAV);
+			UAV->MarkResourceAsUsed();
+			RTBindingsWriter.SetUAV(ParameterBinding.BaseIndex, UAV->GetRHI());
+		}
+		break;
+		default:
+			checkf(false, TEXT("Unhandled resource type?"));
+			break;
+		}
+	}
+
+	// Graph Uniform Buffers
+	for (const FShaderParameterBindings::FParameterStructReference& ParameterBinding : Bindings.GraphUniformBuffers)
+	{
+		const FRDGUniformBufferBinding& UniformBufferBinding = *reinterpret_cast<const FRDGUniformBufferBinding*>(Base + ParameterBinding.ByteOffset);
+
+		checkSlow(UniformBufferBinding);
+		UniformBufferBinding->MarkResourceAsUsed();
+		RTBindingsWriter.SetUniformBuffer(ParameterBinding.BufferIndex, UniformBufferBinding->GetRHI());
+	}
+
+	// Referenced uniform buffers
+	for (const FShaderParameterBindings::FParameterStructReference& ParameterBinding : Bindings.ParameterReferences)
+	{
+		const FUniformBufferBinding& UniformBufferBinding = *reinterpret_cast<const FUniformBufferBinding*>(Base + ParameterBinding.ByteOffset);
+		RTBindingsWriter.SetUniformBuffer(ParameterBinding.BufferIndex, UniformBufferBinding.GetUniformBuffer());
+	}
+
+	// Root uniform buffer.
+	if (Bindings.RootParameterBufferIndex != FShaderParameterBindings::kInvalidBufferIndex)
+	{
+		// Do not do any validation at some resources may have been removed from the structure because known to not be used by the shader.
+		EUniformBufferValidation Validation = EUniformBufferValidation::None;
+
+		RTBindingsWriter.RootUniformBuffer = RHICreateUniformBuffer(Base, RootUniformBufferLayout, UniformBuffer_SingleDraw, Validation);
+		RTBindingsWriter.SetUniformBuffer(Bindings.RootParameterBufferIndex, RTBindingsWriter.RootUniformBuffer);
+	}
+}
+#endif // RHI_RAYTRACING
diff --git a/Engine/Source/Runtime/RenderCore/Public/ShaderParameterStruct.h b/Engine/Source/Runtime/RenderCore/Public/ShaderParameterStruct.h
--- a/Engine/Source/Runtime/RenderCore/Public/ShaderParameterStruct.h
+++ b/Engine/Source/Runtime/RenderCore/Public/ShaderParameterStruct.h
@@ -115,50 +115,6 @@ FORCEINLINE void ValidateShaderParameters(const TShaderRef<TShaderClass>& Shader
 	return ValidateShaderParameters(Shader, TShaderClass::FParameters::FTypeInfo::GetStructMetadata(), ParameterPtr);
 }
 
-/** Set compute shader UAVs. */
-template<typename TRHICmdList, typename TShaderClass, typename TShaderRHI>
-inline void SetShaderUAV(TRHICmdList& RHICmdList, const TShaderRef<TShaderClass>& Shader, TShaderRHI* ShadeRHI, const uint8* Base, const FShaderParameterBindings::FResourceParameter& ParameterBinding)
-{
-	checkf(false, TEXT("TShaderRHI Can't have compute shader to be set. UAVs are not supported on vertex, tessellation and geometry shaders."));
-}
-
-template<typename TRHICmdList, typename TShaderClass>
-inline void SetShaderUAV(TRHICmdList& RHICmdList, const TShaderRef<TShaderClass>& Shader, FRHIPixelShader* ShadeRHI, const uint8* Base, const FShaderParameterBindings::FResourceParameter& ParameterBinding)
-{
-	if (ParameterBinding.BaseType == UBMT_UAV)
-	{
-		FRHIUnorderedAccessView* ShaderParameterRef = *(FRHIUnorderedAccessView**)(Base + ParameterBinding.ByteOffset);
-		RHICmdList.SetUAVParameter(ShadeRHI, ParameterBinding.BaseIndex, ShaderParameterRef);
-	}
-	else if (ParameterBinding.BaseType == UBMT_RDG_TEXTURE_UAV || ParameterBinding.BaseType == UBMT_RDG_BUFFER_UAV)
-	{
-		auto GraphUAV = *reinterpret_cast<FRDGUnorderedAccessView* const*>(Base + ParameterBinding.ByteOffset);
-
-		checkSlow(GraphUAV);
-		GraphUAV->MarkResourceAsUsed();
-		RHICmdList.SetUAVParameter(ShadeRHI, ParameterBinding.BaseIndex, GraphUAV->GetRHI());
-	}
-}
-
-template<typename TRHICmdList, typename TShaderClass>
-inline void SetShaderUAV(TRHICmdList& RHICmdList, const TShaderRef<TShaderClass>& Shader, FRHIComputeShader* ShadeRHI, const uint8* Base, const FShaderParameterBindings::FResourceParameter& ParameterBinding)
-{
-	if (ParameterBinding.BaseType == UBMT_UAV)
-	{
-		FRHIUnorderedAccessView* ShaderParameterRef = *(FRHIUnorderedAccessView**)(Base + ParameterBinding.ByteOffset);
-		RHICmdList.SetUAVParameter(ShadeRHI, ParameterBinding.BaseIndex, ShaderParameterRef);
-	}
-	else if (ParameterBinding.BaseType == UBMT_RDG_TEXTURE_UAV || ParameterBinding.BaseType == UBMT_RDG_BUFFER_UAV)
-	{
-		auto GraphUAV = *reinterpret_cast<FRDGUnorderedAccessView* const*>(Base + ParameterBinding.ByteOffset);
-
-		checkSlow(GraphUAV);
-		GraphUAV->MarkResourceAsUsed();
-		RHICmdList.SetUAVParameter(ShadeRHI, ParameterBinding.BaseIndex, GraphUAV->GetRHI());
-	}
-}
-
-
 /** Unset compute shader UAVs. */
 template<typename TRHICmdList, typename TShaderClass>
 inline void UnsetShaderUAVs(TRHICmdList& RHICmdList, const TShaderRef<TShaderClass>& Shader, FRHIComputeShader* ShadeRHI)
@@ -179,7 +135,6 @@ inline void UnsetShaderUAVs(TRHICmdList& RHICmdList, const TShaderRef<TShaderCla
 	}
 }
 
-
 /** Unset compute shader SRVs. */
 template<typename TRHICmdList, typename TShaderClass>
 inline void UnsetShaderSRVs(TRHICmdList& RHICmdList, const TShaderRef<TShaderClass>& Shader, FRHIComputeShader* ShadeRHI)
@@ -199,13 +154,32 @@ inline void UnsetShaderSRVs(TRHICmdList& RHICmdList, const TShaderRef<TShaderCla
 	}
 }
 
+RENDERCORE_API void SetShaderParameters(
+	FRHIComputeCommandList& RHICmdList,
+	FRHIComputeShader* ShaderRHI,
+	const FShaderParameterBindings& Bindings,
+	const FShaderParametersMetadata* ParametersMetadata,
+	const uint8* Base);
+
+RENDERCORE_API void SetShaderParameters(
+	FRHICommandList& RHICmdList,
+	FRHIGraphicsShader* ShaderRHI,
+	const FShaderParameterBindings& Bindings,
+	const FShaderParametersMetadata* ParametersMetadata,
+	const uint8* Base);
+
+RENDERCORE_API void SetShaderParameters(
+	FRHICommandList& RHICmdList,
+	FRHIComputeShader* ShaderRHI,
+	const FShaderParameterBindings& Bindings,
+	const FShaderParametersMetadata* ParametersMetadata,
+	const uint8* Base);
 
-/** Set shader's parameters from its parameters struct. */
 template<typename TRHICmdList, typename TShaderClass, typename TShaderRHI>
 inline void SetShaderParameters(
-	TRHICmdList& RHICmdList, 
-	const TShaderRef<TShaderClass>& Shader, 
-	TShaderRHI* ShadeRHI, 
+	TRHICmdList& RHICmdList,
+	const TShaderRef<TShaderClass>& Shader,
+	TShaderRHI* ShaderRHI,
 	const FShaderParametersMetadata* ParametersMetadata,
 	const typename TShaderClass::FParameters& Parameters)
 {
@@ -217,100 +191,7 @@ inline void SetShaderParameters(
 	const typename TShaderClass::FParameters* ParametersPtr = &Parameters;
 	const uint8* Base = reinterpret_cast<const uint8*>(ParametersPtr);
 
-	checkf(Bindings.RootParameterBufferIndex == FShaderParameterBindings::kInvalidBufferIndex, TEXT("Can't use SetShaderParameters() for root parameter buffer index."));
-
-	// Parameters
-	for (const FShaderParameterBindings::FParameter& ParameterBinding : Bindings.Parameters)
-	{
-		const void* DataPtr = reinterpret_cast<const char*>(&Parameters) + ParameterBinding.ByteOffset;
-		RHICmdList.SetShaderParameter(ShadeRHI, ParameterBinding.BufferIndex, ParameterBinding.BaseIndex, ParameterBinding.ByteSize, DataPtr);
-	}
-
-	TArray<FShaderParameterBindings::FResourceParameter, TInlineAllocator<16>> GraphSRVs;
-
-	for (const FShaderParameterBindings::FResourceParameter& ParameterBinding : Bindings.ResourceParameters)
-	{
-		EUniformBufferBaseType BaseType = (EUniformBufferBaseType)ParameterBinding.BaseType;
-		switch (BaseType)
-		{
-			case UBMT_TEXTURE:
-			{
-				FRHITexture* TexRef = *(FRHITexture**)(Base + ParameterBinding.ByteOffset);
-				RHICmdList.SetShaderTexture(ShadeRHI, ParameterBinding.BaseIndex, TexRef);
-			}
-			break;
-			case UBMT_SRV:
-			{
-				FRHIShaderResourceView* SRVRef = *(FRHIShaderResourceView**)(Base + ParameterBinding.ByteOffset);
-				RHICmdList.SetShaderResourceViewParameter(ShadeRHI, ParameterBinding.BaseIndex, SRVRef);
-			}
-			break;
-			case UBMT_SAMPLER:
-			{
-				FRHISamplerState* SamplerRef = *(FRHISamplerState**)(Base + ParameterBinding.ByteOffset);
-				RHICmdList.SetShaderSampler(ShadeRHI, ParameterBinding.BaseIndex, SamplerRef);
-			}
-			break;
-			case UBMT_RDG_TEXTURE:
-			{
-				auto GraphTexture = *reinterpret_cast<FRDGTexture* const*>(Base + ParameterBinding.ByteOffset);
-				checkSlow(GraphTexture);
-				GraphTexture->MarkResourceAsUsed();
-				RHICmdList.SetShaderTexture(ShadeRHI, ParameterBinding.BaseIndex, GraphTexture->GetRHI());
-			}
-			break;
-			case UBMT_RDG_TEXTURE_SRV:
-			case UBMT_RDG_BUFFER_SRV:
-			{
-				//HACKHACK: defer SRVs binding after UAVs 
-				GraphSRVs.Add(ParameterBinding);
-			}
-			break;
-			case UBMT_UAV:
-			case UBMT_RDG_TEXTURE_UAV:
-			case UBMT_RDG_BUFFER_UAV:
-			{
-				SetShaderUAV(RHICmdList, Shader, ShadeRHI, Base, ParameterBinding);
-			}
-			break;
-			default:
-				checkf(false, TEXT("Unhandled resource type?"));
-				break;
-		}
-	}
-
-	//HACKHACK: Bind SRVs after UAVs as a workaround for D3D11 RHI unbinding SRVs when binding a UAV on the same resource even when the views don't overlap.
-	for (const FShaderParameterBindings::FResourceParameter& ParameterBinding : GraphSRVs)
-	{
-		auto GraphSRV = *reinterpret_cast<FRDGShaderResourceView* const*>(Base + ParameterBinding.ByteOffset);
-
-		checkSlow(GraphSRV);
-		GraphSRV->MarkResourceAsUsed();
-		RHICmdList.SetShaderResourceViewParameter(ShadeRHI, ParameterBinding.BaseIndex, GraphSRV->GetRHI());
-	}
-
-	// Graph Uniform Buffers
-	for (const FShaderParameterBindings::FParameterStructReference& ParameterBinding : Bindings.GraphUniformBuffers)
-	{
-		const FRDGUniformBufferBinding& UniformBufferBinding = *reinterpret_cast<const FRDGUniformBufferBinding*>(Base + ParameterBinding.ByteOffset);
-
-		if (UniformBufferBinding.IsShader())
-		{
-			UniformBufferBinding->MarkResourceAsUsed();
-			RHICmdList.SetShaderUniformBuffer(ShadeRHI, ParameterBinding.BufferIndex, UniformBufferBinding->GetRHI());
-		}
-	}
-
-	// Reference structures
-	for (const FShaderParameterBindings::FParameterStructReference& ParameterBinding : Bindings.ParameterReferences)
-	{
-		const FUniformBufferBinding& UniformBufferBinding = *reinterpret_cast<const FUniformBufferBinding*>(Base + ParameterBinding.ByteOffset);
-
-		if (UniformBufferBinding.IsShader())
-		{
-			RHICmdList.SetShaderUniformBuffer(ShadeRHI, ParameterBinding.BufferIndex, UniformBufferBinding.GetUniformBuffer());
-		}
-	}
+	SetShaderParameters(RHICmdList, ShaderRHI, Bindings, ParametersMetadata, Base);
 }
 
 template<typename TRHICmdList, typename TShaderClass, typename TShaderRHI>
@@ -322,6 +203,12 @@ inline void SetShaderParameters(TRHICmdList& RHICmdList, const TShaderRef<TShade
 
 #if RHI_RAYTRACING
 
+RENDERCORE_API void SetShaderParameters(
+	FRayTracingShaderBindingsWriter& RTBindingsWriter,
+	const FShaderParameterBindings& Bindings,
+	const FRHIUniformBufferLayout* RootUniformBufferLayout,
+	const uint8* Base);
+
 /** Set shader's parameters from its parameters struct. */
 template<typename TShaderClass>
 void SetShaderParameters(FRayTracingShaderBindingsWriter& RTBindingsWriter, const TShaderRef<TShaderClass>& Shader, const typename TShaderClass::FParameters& Parameters)
@@ -335,95 +222,7 @@ void SetShaderParameters(FRayTracingShaderBindingsWriter& RTBindingsWriter, cons
 	const typename TShaderClass::FParameters* ParametersPtr = &Parameters;
 	const uint8* Base = reinterpret_cast<const uint8*>(ParametersPtr);
 
-	for (const FShaderParameterBindings::FResourceParameter& ParameterBinding : Bindings.ResourceParameters)
-	{
-		EUniformBufferBaseType BaseType = (EUniformBufferBaseType)ParameterBinding.BaseType;
-		switch (BaseType)
-		{
-			case UBMT_TEXTURE:
-			{
-				auto ShaderParameterRef = *(FRHITexture**)(Base + ParameterBinding.ByteOffset);
-				RTBindingsWriter.SetTexture(ParameterBinding.BaseIndex, ShaderParameterRef);
-			}
-			break;
-			case UBMT_SRV:
-			{
-				FRHIShaderResourceView* ShaderParameterRef = *(FRHIShaderResourceView**)(Base + ParameterBinding.ByteOffset);
-				RTBindingsWriter.SetSRV(ParameterBinding.BaseIndex, ShaderParameterRef);
-			}
-			break;
-			case UBMT_UAV:
-			{
-				FRHIUnorderedAccessView* ShaderParameterRef = *(FRHIUnorderedAccessView**)(Base + ParameterBinding.ByteOffset);
-				RTBindingsWriter.SetUAV(ParameterBinding.BaseIndex, ShaderParameterRef);
-			}
-			break;
-			case UBMT_SAMPLER:
-			{
-				FRHISamplerState* ShaderParameterRef = *(FRHISamplerState**)(Base + ParameterBinding.ByteOffset);
-				RTBindingsWriter.SetSampler(ParameterBinding.BaseIndex, ShaderParameterRef);
-			}
-			break;
-			case UBMT_RDG_TEXTURE:
-			{
-				auto GraphTexture = *reinterpret_cast<FRDGTexture* const*>(Base + ParameterBinding.ByteOffset);
-				checkSlow(GraphTexture);
-				GraphTexture->MarkResourceAsUsed();
-				RTBindingsWriter.SetTexture(ParameterBinding.BaseIndex, GraphTexture->GetRHI());
-			}
-			break;
-			case UBMT_RDG_TEXTURE_SRV:
-			case UBMT_RDG_BUFFER_SRV:
-			{
-				auto GraphSRV = *reinterpret_cast<FRDGShaderResourceView* const*>(Base + ParameterBinding.ByteOffset);
-
-				checkSlow(GraphSRV);
-				GraphSRV->MarkResourceAsUsed();
-				RTBindingsWriter.SetSRV(ParameterBinding.BaseIndex, GraphSRV->GetRHI());
-			}
-			break;
-			case UBMT_RDG_TEXTURE_UAV:
-			case UBMT_RDG_BUFFER_UAV:
-			{
-				auto UAV = *reinterpret_cast<FRDGUnorderedAccessView* const*>(Base + ParameterBinding.ByteOffset);
-
-				checkSlow(UAV);
-				UAV->MarkResourceAsUsed();
-				RTBindingsWriter.SetUAV(ParameterBinding.BaseIndex, UAV->GetRHI());
-			}
-			break;
-			default:
-				checkf(false, TEXT("Unhandled resource type?"));
-				break;
-		}
-	}
-
-	// Graph Uniform Buffers
-	for (const FShaderParameterBindings::FParameterStructReference& ParameterBinding : Bindings.GraphUniformBuffers)
-	{
-		const FRDGUniformBufferBinding& UniformBufferBinding = *reinterpret_cast<const FRDGUniformBufferBinding*>(Base + ParameterBinding.ByteOffset);
-
-		checkSlow(UniformBufferBinding);
-		UniformBufferBinding->MarkResourceAsUsed();
-		RTBindingsWriter.SetUniformBuffer(ParameterBinding.BufferIndex, UniformBufferBinding->GetRHI());
-	}
-
-	// Referenced uniform buffers
-	for (const FShaderParameterBindings::FParameterStructReference& ParameterBinding : Bindings.ParameterReferences)
-	{
-		const FUniformBufferBinding& UniformBufferBinding = *reinterpret_cast<const FUniformBufferBinding*>(Base + ParameterBinding.ByteOffset);
-		RTBindingsWriter.SetUniformBuffer(ParameterBinding.BufferIndex, UniformBufferBinding.GetUniformBuffer());
-	}
-
-	// Root uniform buffer.
-	if (Bindings.RootParameterBufferIndex != FShaderParameterBindings::kInvalidBufferIndex)
-	{
-		// Do not do any validation at some resources may have been removed from the structure because known to not be used by the shader.
-		EUniformBufferValidation Validation = EUniformBufferValidation::None;
-
-		RTBindingsWriter.RootUniformBuffer = CreateUniformBufferImmediate(Parameters, UniformBuffer_SingleDraw, Validation);
-		RTBindingsWriter.SetUniformBuffer(Bindings.RootParameterBufferIndex, RTBindingsWriter.RootUniformBuffer);
-	}
+	SetShaderParameters(RTBindingsWriter, Bindings, TShaderClass::FParameters::FTypeInfo::GetStructMetadata()->GetLayoutPtr(), Base);
 }
 
 #endif // RHI_RAYTRACING
