Fix debug event names in FRayTracingValidateGeometryBuildParamsCS::Dispatch that were not null-terminated

diff --git a/Engine/Source/Runtime/RenderCore/Private/RayTracingValidationShaders.cpp b/Engine/Source/Runtime/RenderCore/Private/RayTracingValidationShaders.cpp
--- a/Engine/Source/Runtime/RenderCore/Private/RayTracingValidationShaders.cpp
+++ b/Engine/Source/Runtime/RenderCore/Private/RayTracingValidationShaders.cpp
@@ -30,8 +30,8 @@ void FRayTracingValidateGeometryBuildParamsCS::Dispatch(FRHICommandList& RHICmdL
 		EventName.Append(TEXT(" - "));
 		EventName.Append(*DebugNameString);
 	}
-	
-	RHICmdList.PushEvent(EventName.GetData(), FColor::Black);
+
+	RHICmdList.PushEvent(EventName.ToString(), FColor::Black);
 
 	const uint32 IndexStride = Initializer.IndexBuffer->GetStride();
 
