VulkanRHI: Allow MOBILE_EMULATION to be defined in VULKAN_PROFILE shaders.  Fixes gizmo not appearing when using VulkanRHI for Android emulation (alpha channel wasn't being set correctly for composite pass).

diff --git a/Engine/Shaders/Public/Platform.ush b/Engine/Shaders/Public/Platform.ush
--- a/Engine/Shaders/Public/Platform.ush
+++ b/Engine/Shaders/Public/Platform.ush
@@ -355,7 +355,7 @@
 #endif
 
 #ifndef MOBILE_EMULATION
-#define MOBILE_EMULATION ((FEATURE_LEVEL == FEATURE_LEVEL_ES3_1) && (!(COMPILER_GLSL_ES3_1) && USE_DEVELOPMENT_SHADERS && (!(METAL_PROFILE) || MAC) && !VULKAN_PROFILE))
+#define MOBILE_EMULATION 0
 #endif
 
 // Whether the platform supports independent texture and samplers
diff --git a/Engine/Source/Runtime/Engine/Private/ShaderCompiler/ShaderCompiler.cpp b/Engine/Source/Runtime/Engine/Private/ShaderCompiler/ShaderCompiler.cpp
--- a/Engine/Source/Runtime/Engine/Private/ShaderCompiler/ShaderCompiler.cpp
+++ b/Engine/Source/Runtime/Engine/Private/ShaderCompiler/ShaderCompiler.cpp
@@ -5733,6 +5733,18 @@ void GlobalBeginCompileShader(
 		}
 	}
 
+	if (IsAndroidPlatform(EShaderPlatform(Target.Platform)))
+	{
+		// Force no development shaders on Android platforms
+		bAllowDevelopmentShaderCompile = false;
+	}
+
+	// Mobile emulation should be defined when a PC platform is using a mobile renderer (limited to feature level ES3_1)...  eg SP_PCD3D_ES3_1,SP_VULKAN_PCES3_1,SP_METAL_MACES3_1
+	if (IsSimulatedPlatform(EShaderPlatform(Target.Platform)) && bAllowDevelopmentShaderCompile)
+	{
+		Input.Environment.SetDefine(TEXT("MOBILE_EMULATION"), 1);
+	}
+
 	// Add compiler flag CFLAG_ForceDXC if DXC is enabled
 	const bool bHlslVersion2021 = Input.Environment.CompilerFlags.Contains(CFLAG_HLSL2021);
 	const bool bIsDxcEnabled = IsDxcEnabledForPlatform((EShaderPlatform)Target.Platform, bHlslVersion2021);
diff --git a/Engine/Source/Runtime/RHI/Public/RHIDefinitions.h b/Engine/Source/Runtime/RHI/Public/RHIDefinitions.h
--- a/Engine/Source/Runtime/RHI/Public/RHIDefinitions.h
+++ b/Engine/Source/Runtime/RHI/Public/RHIDefinitions.h
@@ -2063,6 +2063,12 @@ inline bool IsConsolePlatform(const FStaticShaderPlatform Platform)
 	return FDataDrivenShaderPlatformInfo::GetIsConsole(Platform);
 }
 
+// @todo: data drive uses of this function
+inline bool IsAndroidPlatform(const FStaticShaderPlatform Platform)
+{
+	return (Platform == SP_VULKAN_ES3_1_ANDROID) || (Platform == SP_VULKAN_SM5_ANDROID) || (Platform == SP_OPENGL_ES3_1_ANDROID);
+}
+
 inline bool IsVulkanPlatform(const FStaticShaderPlatform Platform)
 {
 	return FDataDrivenShaderPlatformInfo::GetIsLanguageVulkan(Platform);
diff --git a/Engine/Source/Runtime/RenderCore/Private/ShaderMaterialDerivedHelpers.cpp b/Engine/Source/Runtime/RenderCore/Private/ShaderMaterialDerivedHelpers.cpp
--- a/Engine/Source/Runtime/RenderCore/Private/ShaderMaterialDerivedHelpers.cpp
+++ b/Engine/Source/Runtime/RenderCore/Private/ShaderMaterialDerivedHelpers.cpp
@@ -83,9 +83,7 @@ FShaderMaterialDerivedDefines RENDERCORE_API CalculateDerivedMaterialParameters(
 
 	Dst.IS_MESHPARTICLE_FACTORY = (Lightmap.PARTICLE_MESH_FACTORY || Lightmap.NIAGARA_MESH_FACTORY);
 
-	Dst.MOBILE_EMULATION = ((FEATURE_LEVEL == ERHIFeatureLevel::ES3_1) && (!(Compiler.COMPILER_GLSL_ES3_1) && Dst.USE_DEVELOPMENT_SHADERS && (!(Compiler.METAL_PROFILE) || Compiler.MAC) && !Compiler.VULKAN_PROFILE));
-
-	Dst.SUPPORTS_PIXEL_COVERAGE = (FEATURE_LEVEL >= ERHIFeatureLevel::SM5 && !Compiler.COMPILER_GLSL && !Dst.MOBILE_EMULATION);
+	Dst.SUPPORTS_PIXEL_COVERAGE = (FEATURE_LEVEL >= ERHIFeatureLevel::SM5 && !Compiler.COMPILER_GLSL);
 
 	Dst.FORCE_FULLY_ROUGH = (Lightmap.SIMPLE_FORWARD_SHADING || Mat.MATERIAL_FULLY_ROUGH);
 	Dst.EDITOR_ALPHA2COVERAGE = (Dst.USE_EDITOR_COMPOSITING && Dst.SUPPORTS_PIXEL_COVERAGE);
diff --git a/Engine/Source/Runtime/RenderCore/Public/ShaderMaterial.h b/Engine/Source/Runtime/RenderCore/Public/ShaderMaterial.h
--- a/Engine/Source/Runtime/RenderCore/Public/ShaderMaterial.h
+++ b/Engine/Source/Runtime/RenderCore/Public/ShaderMaterial.h
@@ -255,7 +255,6 @@ struct FShaderMaterialDerivedDefines
 	bool PIXELSHADEROUTPUT_COVERAGE;
 
 	bool SUPPORTS_PIXEL_COVERAGE;
-	bool MOBILE_EMULATION;
 
 	bool COMPILE_SHADERS_FOR_DEVELOPMENT;
 	bool USE_DEVELOPMENT_SHADERS;
