Fix PSO cache splitting.

- OutChunkFilenames is allowed to have existing entries.

diff --git a/Engine/Source/Runtime/RenderCore/Private/PipelineCacheUtilities.cpp b/Engine/Source/Runtime/RenderCore/Private/PipelineCacheUtilities.cpp
--- a/Engine/Source/Runtime/RenderCore/Private/PipelineCacheUtilities.cpp
+++ b/Engine/Source/Runtime/RenderCore/Private/PipelineCacheUtilities.cpp
@@ -806,8 +806,6 @@ bool UE::PipelineCacheUtilities::LoadStablePipelineCacheFile(const FString& File
 
 bool UE::PipelineCacheUtilities::SaveChunkInfo(const FString& ShaderLibraryName, const int32 InChunkId, const TSet<FName>& InPackagesInChunk, const ITargetPlatform* TargetPlatform, const FString& PathToSaveTo, TArray<FString>& OutChunkFilenames)
 {
-	checkf(OutChunkFilenames.IsEmpty(), TEXT("SaveChunkInfo() expects the array of chunk filenames to be empty."));
-
 	const FString InfoFile = FString::Printf(TEXT("%s_%s_Chunk%d.cacheinfo"), *TargetPlatform->PlatformName(), *ShaderLibraryName, InChunkId);
 	const FString FullPath = PathToSaveTo / InfoFile;
 
@@ -815,10 +813,12 @@ bool UE::PipelineCacheUtilities::SaveChunkInfo(const FString& ShaderLibraryName,
 	// Note that at this point we don't know which formats will actually have the PSO cache, so add all targeted ones
 	TArray<FName> ShaderFormats;
 	TargetPlatform->GetAllTargetedShaderFormats(ShaderFormats);
+	TArray<FString> BundledFiles;
 	for (FName ShaderFormat : ShaderFormats)
 	{
 		FString BundledFile = FString::Printf(TEXT("%s_%s_Chunk%d.stable.upipelinecache"), *ShaderLibraryName, *ShaderFormat.ToString(), InChunkId);
 		OutChunkFilenames.Add(BundledFile);
+		BundledFiles.Add(BundledFile);
 	}
 
 	TUniquePtr<FArchive> AssetInfoWriter(IFileManager::Get().CreateFileWriter(*FullPath, FILEWRITE_NoFail));
@@ -838,7 +838,7 @@ bool UE::PipelineCacheUtilities::SaveChunkInfo(const FString& ShaderLibraryName,
 			{
 				Writer->WriteObjectStart();
 				Writer->WriteValue(TEXT("ShaderFormat"), ShaderFormats[Idx].ToString());
-				Writer->WriteValue(TEXT("OutputFile"), OutChunkFilenames[Idx]);
+				Writer->WriteValue(TEXT("OutputFile"), BundledFiles[Idx]);
 				Writer->WriteObjectEnd();
 			}
 			Writer->WriteArrayEnd();
