- FIxed loading shader caches from plugins' IOStore files

diff --git a/Engine/Source/Runtime/RenderCore/Private/ShaderCodeLibrary.cpp b/Engine/Source/Runtime/RenderCore/Private/ShaderCodeLibrary.cpp
--- a/Engine/Source/Runtime/RenderCore/Private/ShaderCodeLibrary.cpp
+++ b/Engine/Source/Runtime/RenderCore/Private/ShaderCodeLibrary.cpp
@@ -2546,15 +2546,8 @@ static void FShaderCodeLibraryPluginMountedCallback(IPlugin& Plugin)
 {
 	if (Plugin.CanContainContent() && Plugin.IsEnabled())
 	{
-		// make sure we have some libraries before bothering with this plugin
-		TArray<FString> Found;
-		IFileManager::Get().FindFiles(Found, *Plugin.GetBaseDir(), TEXT(".ushaderbytecode"));
-		IFileManager::Get().FindFiles(Found, *Plugin.GetContentDir(), TEXT(".ushaderbytecode"));
-		if (Found.Num() > 0)
-		{
-			FShaderCodeLibrary::OpenLibrary(Plugin.GetName(), Plugin.GetBaseDir());
-			FShaderCodeLibrary::OpenLibrary(Plugin.GetName(), Plugin.GetContentDir());
-		}
+		// load any shader libraries that may exist in this plugin
+		FShaderCodeLibrary::OpenLibrary(Plugin.GetName(), Plugin.GetContentDir());
 	}
 }
 
