Making a FShaderParameterMap::FindParameterAllocation method that returns the parameter struct instead of using reference parameters.
Cleaning up SetShaderParameters implementations using a new helper class to load parameters from the binary blobs passed in.

diff --git a/Engine/Source/Runtime/RenderCore/Private/ShaderCore.cpp b/Engine/Source/Runtime/RenderCore/Private/ShaderCore.cpp
--- a/Engine/Source/Runtime/RenderCore/Private/ShaderCore.cpp
+++ b/Engine/Source/Runtime/RenderCore/Private/ShaderCore.cpp
@@ -491,15 +491,10 @@ bool AllowGlobalShaderLoad()
 
 }
 
-bool FShaderParameterMap::FindParameterAllocation(const TCHAR* ParameterName,uint16& OutBufferIndex,uint16& OutBaseIndex,uint16& OutSize) const
+TOptional<FParameterAllocation> FShaderParameterMap::FindParameterAllocation(const FString& ParameterName) const
 {
-	const FParameterAllocation* Allocation = ParameterMap.Find(ParameterName);
-	if(Allocation)
+	if (const FParameterAllocation* Allocation = ParameterMap.Find(ParameterName))
 	{
-		OutBufferIndex = Allocation->BufferIndex;
-		OutBaseIndex = Allocation->BaseIndex;
-		OutSize = Allocation->Size;
-
 		if (Allocation->bBound)
 		{
 			// Can detect copy-paste errors in binding parameters.  Need to fix all the false positives before enabling.
@@ -507,12 +502,25 @@ bool FShaderParameterMap::FindParameterAllocation(const TCHAR* ParameterName,uin
 		}
 
 		Allocation->bBound = true;
-		return true;
+
+		return TOptional<FParameterAllocation>(*Allocation);
 	}
-	else
+
+	return TOptional<FParameterAllocation>();
+}
+
+bool FShaderParameterMap::FindParameterAllocation(const TCHAR* ParameterName, uint16& OutBufferIndex, uint16& OutBaseIndex, uint16& OutSize) const
+{
+	if (TOptional<FParameterAllocation> Allocation = FindParameterAllocation(ParameterName))
 	{
-		return false;
+		OutBufferIndex = Allocation->BufferIndex;
+		OutBaseIndex = Allocation->BaseIndex;
+		OutSize = Allocation->Size;
+
+		return true;
 	}
+
+	return false;
 }
 
 bool FShaderParameterMap::ContainsParameterAllocation(const TCHAR* ParameterName) const
@@ -523,12 +531,7 @@ bool FShaderParameterMap::ContainsParameterAllocation(const TCHAR* ParameterName
 void FShaderParameterMap::AddParameterAllocation(const TCHAR* ParameterName,uint16 BufferIndex,uint16 BaseIndex,uint16 Size,EShaderParameterType ParameterType)
 {
 	check(ParameterType < EShaderParameterType::Num);
-	FParameterAllocation Allocation;
-	Allocation.BufferIndex = BufferIndex;
-	Allocation.BaseIndex = BaseIndex;
-	Allocation.Size = Size;
-	Allocation.Type = ParameterType;
-	ParameterMap.Add(ParameterName,Allocation);
+	ParameterMap.Add(ParameterName, FParameterAllocation(BufferIndex, BaseIndex, Size, ParameterType));
 }
 
 void FShaderParameterMap::RemoveParameterAllocation(const TCHAR* ParameterName)
diff --git a/Engine/Source/Runtime/RenderCore/Private/ShaderParameterStruct.cpp b/Engine/Source/Runtime/RenderCore/Private/ShaderParameterStruct.cpp
--- a/Engine/Source/Runtime/RenderCore/Private/ShaderParameterStruct.cpp
+++ b/Engine/Source/Runtime/RenderCore/Private/ShaderParameterStruct.cpp
@@ -153,11 +153,12 @@ struct FShaderParameterStructBindingContext
 					UE_LOG(LogShaders, Fatal, TEXT("%s can't bind shader parameter %s, because it has already be bound by %s."), *CppName, *ElementShaderBindingName, **ShaderGlobalScopeBindings.Find(ShaderBindingName));
 				}
 
-				uint16 BufferIndex, BaseIndex, BoundSize;
-				if (!ParametersMap->FindParameterAllocation(*ElementShaderBindingName, BufferIndex, BaseIndex, BoundSize))
+				TOptional<FParameterAllocation> ParameterAllocation = ParametersMap->FindParameterAllocation(ElementShaderBindingName);
+				if (!ParameterAllocation.IsSet())
 				{
 					continue;
 				}
+
 				ShaderGlobalScopeBindings.Add(ElementShaderBindingName, CppName);
 
 				if (bIsVariableNativeType)
@@ -166,16 +167,16 @@ struct FShaderParameterStructBindingContext
 					uint32 ByteSize = Member.GetMemberSize();
 
 					FShaderParameterBindings::FParameter Parameter;
-					Parameter.BufferIndex = BufferIndex;
-					Parameter.BaseIndex = BaseIndex;
+					Parameter.BufferIndex = ParameterAllocation->BufferIndex;
+					Parameter.BaseIndex = ParameterAllocation->BaseIndex;
 					Parameter.ByteOffset = ByteOffset;
-					Parameter.ByteSize = BoundSize;
+					Parameter.ByteSize = ParameterAllocation->Size;
 
-					if (uint32(BoundSize) > ByteSize)
+					if (uint32(ParameterAllocation->Size) > ByteSize)
 					{
 						UE_LOG(LogShaders, Fatal, TEXT("The size required to bind shader %s's (Permutation Id %d) struct %s parameter %s is %i bytes, smaller than %s's %i bytes."),
 							Shader->GetTypeUnfrozen()->GetName(), PermutationId, StructMetaData.GetStructTypeName(),
-							*ElementShaderBindingName, BoundSize, *CppName, ByteSize);
+							*ElementShaderBindingName, ParameterAllocation->Size, *CppName, ByteSize);
 					}
 
 					Bindings->Parameters.Add(Parameter);
@@ -184,7 +185,7 @@ struct FShaderParameterStructBindingContext
 				{
 					check(!bIsArray);
 					FShaderParameterBindings::FParameterStructReference Parameter;
-					Parameter.BufferIndex = BufferIndex;
+					Parameter.BufferIndex = ParameterAllocation->BufferIndex;
 					Parameter.ByteOffset = ByteOffset;
 
 					if (BaseType == UBMT_REFERENCED_STRUCT)
@@ -198,21 +199,21 @@ struct FShaderParameterStructBindingContext
 				}
 				else if (bIsRHIResource || bIsRDGResource)
 				{
-					checkf(BaseIndex < 256, TEXT("BaseIndex does not fit into uint8. Change FResourceParameter::BaseIndex type to uint16"));
-										
+					checkf(ParameterAllocation->BaseIndex < 256, TEXT("BaseIndex does not fit into uint8. Change FResourceParameter::BaseIndex type to uint16"));
+
 					FShaderParameterBindings::FResourceParameter Parameter;
-					Parameter.BaseIndex = (uint8)BaseIndex;
+					Parameter.BaseIndex = (uint8)ParameterAllocation->BaseIndex;
 					Parameter.ByteOffset = ByteOffset + ArrayElementId * SHADER_PARAMETER_POINTER_ALIGNMENT;
 					Parameter.BaseType = BaseType;
 
-					if (BoundSize != 1)
+					if (ParameterAllocation->Size != 1)
 					{
 						UE_LOG(LogShaders, Fatal,
 							TEXT("Error with shader %s's (Permutation Id %d) parameter %s is %i bytes, cpp name = %s.")
 							TEXT("The shader compiler should give precisely which elements of an array did not get compiled out, ")
 							TEXT("for optimal automatic render graph pass dependency with ClearUnusedGraphResources()."),
 							Shader->GetTypeUnfrozen()->GetName(), PermutationId,
-							*ElementShaderBindingName, BoundSize, *CppName);
+							*ElementShaderBindingName, ParameterAllocation->Size, *CppName);
 					}
 
 					Bindings->ResourceParameters.Add(Parameter);
@@ -454,6 +455,27 @@ void EmitNullShaderParameterFatalError(const TShaderRef<FShader>& Shader, const
 		*MemberName);
 }
 
+/// Utility class for reading shader parameters out of the data blob passed in.
+struct FShaderParameterReader
+{
+	FShaderParameterReader() = delete;
+	FShaderParameterReader(const void* InData) : Data(reinterpret_cast<const uint8*>(InData)) { }
+
+	template<typename TParameterIn>
+	const void* GetRawPointer(const TParameterIn& InParameter) const
+	{
+		return (Data + InParameter.ByteOffset);
+	}
+
+	template<typename TParameterOut, typename TParameterIn>
+	const TParameterOut& Read(const TParameterIn& InParameter) const
+	{
+		return *reinterpret_cast<const TParameterOut*>(Data + InParameter.ByteOffset);
+	}
+
+	const uint8* Data;
+};
+
 #if DO_CHECK
 
 void ValidateShaderParameters(const TShaderRef<FShader>& Shader, const FShaderParametersMetadata* ParametersMetadata, const void* Parameters)
@@ -465,14 +487,14 @@ void ValidateShaderParameters(const TShaderRef<FShader>& Shader, const FShaderPa
 		TEXT("Shader %s's parameter structure has changed without recompilation of the shader"),
 		Shader.GetType()->GetName());
 
-	const uint8* Base = reinterpret_cast<const uint8*>(Parameters);
+	const FShaderParameterReader Reader(Parameters);
 
 	const TCHAR* ShaderClassName = Shader.GetType()->GetName();
 	const TCHAR* ShaderParameterStructName = ParametersMetadata->GetStructTypeName();
 
-	for (const FShaderParameterBindings::FResourceParameter& ParameterBinding : Bindings.ResourceParameters)
+	for (const FShaderParameterBindings::FResourceParameter& Parameter : Bindings.ResourceParameters)
 	{
-		EUniformBufferBaseType BaseType = (EUniformBufferBaseType)ParameterBinding.BaseType;
+		const EUniformBufferBaseType BaseType = static_cast<EUniformBufferBaseType>(Parameter.BaseType);
 		switch (BaseType)
 		{
 			case UBMT_TEXTURE:
@@ -480,19 +502,19 @@ void ValidateShaderParameters(const TShaderRef<FShader>& Shader, const FShaderPa
 			case UBMT_UAV:
 			case UBMT_SAMPLER:
 			{
-				FRHIResource* ShaderParameterRef = *(FRHIResource**)(Base + ParameterBinding.ByteOffset);
-				if (!ShaderParameterRef)
+				const FRHIResource* Resource = Reader.Read<const FRHIResource*>(Parameter);
+				if (!Resource)
 				{
-					EmitNullShaderParameterFatalError(Shader, ParametersMetadata, ParameterBinding.ByteOffset);
+					EmitNullShaderParameterFatalError(Shader, ParametersMetadata, Parameter.ByteOffset);
 				}
 			}
 			break;
 			case UBMT_RDG_TEXTURE:
 			{
-				const FRDGTexture* GraphTexture = *reinterpret_cast<const FRDGTexture* const*>(Base + ParameterBinding.ByteOffset);
+				const FRDGTexture* GraphTexture = Reader.Read<const FRDGTexture*>(Parameter);
 				if (!GraphTexture)
 				{
-					EmitNullShaderParameterFatalError(Shader, ParametersMetadata, ParameterBinding.ByteOffset);
+					EmitNullShaderParameterFatalError(Shader, ParametersMetadata, Parameter.ByteOffset);
 				}
 			}
 			break;
@@ -501,10 +523,10 @@ void ValidateShaderParameters(const TShaderRef<FShader>& Shader, const FShaderPa
 			case UBMT_RDG_BUFFER_SRV:
 			case UBMT_RDG_BUFFER_UAV:
 			{
-				const FRDGResource* GraphResource = *reinterpret_cast<const FRDGResource* const*>(Base + ParameterBinding.ByteOffset);
+				const FRDGResource* GraphResource = Reader.Read<const FRDGResource*>(Parameter);
 				if (!GraphResource)
 				{
-					EmitNullShaderParameterFatalError(Shader, ParametersMetadata, ParameterBinding.ByteOffset);
+					EmitNullShaderParameterFatalError(Shader, ParametersMetadata, Parameter.ByteOffset);
 				}
 			}
 			break;
@@ -514,24 +536,22 @@ void ValidateShaderParameters(const TShaderRef<FShader>& Shader, const FShaderPa
 	}
 
 	// Graph Uniform Buffers
-	for (const FShaderParameterBindings::FParameterStructReference& ParameterBinding : Bindings.GraphUniformBuffers)
+	for (const FShaderParameterBindings::FParameterStructReference& Parameter : Bindings.GraphUniformBuffers)
 	{
-		const FRDGUniformBufferBinding& UniformBufferBinding = *reinterpret_cast<const FRDGUniformBufferBinding*>(Base + ParameterBinding.ByteOffset);
-
+		const FRDGUniformBufferBinding& UniformBufferBinding = Reader.Read<FRDGUniformBufferBinding>(Parameter);
 		if (!UniformBufferBinding)
 		{
-			EmitNullShaderParameterFatalError(Shader, ParametersMetadata, ParameterBinding.ByteOffset);
+			EmitNullShaderParameterFatalError(Shader, ParametersMetadata, Parameter.ByteOffset);
 		}
 	}
 
 	// Reference structures
-	for (const FShaderParameterBindings::FParameterStructReference& ParameterBinding : Bindings.ParameterReferences)
+	for (const FShaderParameterBindings::FParameterStructReference& Parameter : Bindings.ParameterReferences)
 	{
-		const FUniformBufferBinding& UniformBufferBinding = *reinterpret_cast<const FUniformBufferBinding*>(Base + ParameterBinding.ByteOffset);
-
+		const FUniformBufferBinding& UniformBufferBinding = Reader.Read<FUniformBufferBinding>(Parameter);
 		if (!UniformBufferBinding)
 		{
-			EmitNullShaderParameterFatalError(Shader, ParametersMetadata, ParameterBinding.ByteOffset);
+			EmitNullShaderParameterFatalError(Shader, ParametersMetadata, Parameter.ByteOffset);
 		}
 	}
 }
@@ -562,7 +582,7 @@ void ValidateShaderParameterResourcesRHI(const void* Contents, const FRHIUniform
 #endif // DO_CHECK
 
 template<typename TRHICmdList>
-inline void SetShaderUAV(TRHICmdList& RHICmdList, FRHIGraphicsShader* ShaderRHI, const uint8* Base, const FShaderParameterBindings::FResourceParameter& ParameterBinding)
+inline void SetShaderUAV(TRHICmdList& RHICmdList, FRHIGraphicsShader* ShaderRHI, FShaderParameterReader Reader, const FShaderParameterBindings::FResourceParameter& ParameterBinding)
 {
 	if (ShaderRHI->GetType() == RRT_PixelShader)
 	{
@@ -570,14 +590,14 @@ inline void SetShaderUAV(TRHICmdList& RHICmdList, FRHIGraphicsShader* ShaderRHI,
 
 		if (ParameterBinding.BaseType == UBMT_UAV)
 		{
-			FRHIUnorderedAccessView* ShaderParameterRef = *(FRHIUnorderedAccessView**)(Base + ParameterBinding.ByteOffset);
+			FRHIUnorderedAccessView* ShaderParameterRef = Reader.Read<FRHIUnorderedAccessView*>(ParameterBinding);
 			RHICmdList.SetUAVParameter(PixelShaderRHI, ParameterBinding.BaseIndex, ShaderParameterRef);
 		}
 		else if (ParameterBinding.BaseType == UBMT_RDG_TEXTURE_UAV || ParameterBinding.BaseType == UBMT_RDG_BUFFER_UAV)
 		{
-			auto GraphUAV = *reinterpret_cast<FRDGUnorderedAccessView* const*>(Base + ParameterBinding.ByteOffset);
-
+			FRDGUnorderedAccessView* GraphUAV = Reader.Read<FRDGUnorderedAccessView*>(ParameterBinding);
 			checkSlow(GraphUAV);
+
 			GraphUAV->MarkResourceAsUsed();
 			RHICmdList.SetUAVParameter(PixelShaderRHI, ParameterBinding.BaseIndex, GraphUAV->GetRHI());
 		}
@@ -589,18 +609,18 @@ inline void SetShaderUAV(TRHICmdList& RHICmdList, FRHIGraphicsShader* ShaderRHI,
 }
 
 template<typename TRHICmdList>
-inline void SetShaderUAV(TRHICmdList& RHICmdList, FRHIComputeShader* ShaderRHI, const uint8* Base, const FShaderParameterBindings::FResourceParameter& ParameterBinding)
+inline void SetShaderUAV(TRHICmdList& RHICmdList, FRHIComputeShader* ShaderRHI, FShaderParameterReader Reader, const FShaderParameterBindings::FResourceParameter& ParameterBinding)
 {
 	if (ParameterBinding.BaseType == UBMT_UAV)
 	{
-		FRHIUnorderedAccessView* ShaderParameterRef = *(FRHIUnorderedAccessView**)(Base + ParameterBinding.ByteOffset);
+		FRHIUnorderedAccessView* ShaderParameterRef = Reader.Read<FRHIUnorderedAccessView*>(ParameterBinding);
 		RHICmdList.SetUAVParameter(ShaderRHI, ParameterBinding.BaseIndex, ShaderParameterRef);
 	}
 	else if (ParameterBinding.BaseType == UBMT_RDG_TEXTURE_UAV || ParameterBinding.BaseType == UBMT_RDG_BUFFER_UAV)
 	{
-		auto GraphUAV = *reinterpret_cast<FRDGUnorderedAccessView* const*>(Base + ParameterBinding.ByteOffset);
-
+		FRDGUnorderedAccessView* GraphUAV = Reader.Read<FRDGUnorderedAccessView*>(ParameterBinding);
 		checkSlow(GraphUAV);
+
 		GraphUAV->MarkResourceAsUsed();
 		RHICmdList.SetUAVParameter(ShaderRHI, ParameterBinding.BaseIndex, GraphUAV->GetRHI());
 	}
@@ -617,58 +637,62 @@ inline void SetShaderParametersInternal(
 {
 	checkf(Bindings.RootParameterBufferIndex == FShaderParameterBindings::kInvalidBufferIndex, TEXT("Can't use SetShaderParameters() for root parameter buffer index."));
 
+	FShaderParameterReader Reader(Base);
+
 	// Parameters
-	for (const FShaderParameterBindings::FParameter& ParameterBinding : Bindings.Parameters)
+	for (const FShaderParameterBindings::FParameter& Parameter : Bindings.Parameters)
 	{
-		const void* DataPtr = Base + ParameterBinding.ByteOffset;
-		RHICmdList.SetShaderParameter(ShaderRHI, ParameterBinding.BufferIndex, ParameterBinding.BaseIndex, ParameterBinding.ByteSize, DataPtr);
+		const void* DataPtr = Reader.GetRawPointer(Parameter);
+		RHICmdList.SetShaderParameter(ShaderRHI, Parameter.BufferIndex, Parameter.BaseIndex, Parameter.ByteSize, DataPtr);
 	}
 
 	TArray<FShaderParameterBindings::FResourceParameter, TInlineAllocator<16>> GraphSRVs;
 
-	for (const FShaderParameterBindings::FResourceParameter& ParameterBinding : Bindings.ResourceParameters)
+	for (const FShaderParameterBindings::FResourceParameter& Parameter : Bindings.ResourceParameters)
 	{
-		EUniformBufferBaseType BaseType = (EUniformBufferBaseType)ParameterBinding.BaseType;
+		const EUniformBufferBaseType BaseType = static_cast<EUniformBufferBaseType>(Parameter.BaseType);
 		switch (BaseType)
 		{
 		case UBMT_TEXTURE:
 		{
-			FRHITexture* TexRef = *(FRHITexture**)(Base + ParameterBinding.ByteOffset);
-			RHICmdList.SetShaderTexture(ShaderRHI, ParameterBinding.BaseIndex, TexRef);
+			FRHITexture* RHITexture = Reader.Read<FRHITexture*>(Parameter);
+
+			RHICmdList.SetShaderTexture(ShaderRHI, Parameter.BaseIndex, RHITexture);
 		}
 		break;
 		case UBMT_SRV:
 		{
-			FRHIShaderResourceView* SRVRef = *(FRHIShaderResourceView**)(Base + ParameterBinding.ByteOffset);
-			RHICmdList.SetShaderResourceViewParameter(ShaderRHI, ParameterBinding.BaseIndex, SRVRef);
+			FRHIShaderResourceView* ShaderResourceView = Reader.Read<FRHIShaderResourceView*>(Parameter);
+			RHICmdList.SetShaderResourceViewParameter(ShaderRHI, Parameter.BaseIndex, ShaderResourceView);
 		}
 		break;
 		case UBMT_SAMPLER:
 		{
-			FRHISamplerState* SamplerRef = *(FRHISamplerState**)(Base + ParameterBinding.ByteOffset);
-			RHICmdList.SetShaderSampler(ShaderRHI, ParameterBinding.BaseIndex, SamplerRef);
+			FRHISamplerState* SamplerState = Reader.Read<FRHISamplerState*>(Parameter);
+			RHICmdList.SetShaderSampler(ShaderRHI, Parameter.BaseIndex, SamplerState);
 		}
 		break;
 		case UBMT_RDG_TEXTURE:
 		{
-			auto GraphTexture = *reinterpret_cast<FRDGTexture* const*>(Base + ParameterBinding.ByteOffset);
-			checkSlow(GraphTexture);
-			GraphTexture->MarkResourceAsUsed();
-			RHICmdList.SetShaderTexture(ShaderRHI, ParameterBinding.BaseIndex, GraphTexture->GetRHI());
+			FRDGTexture* RDGTexture = Reader.Read<FRDGTexture*>(Parameter);
+			checkSlow(RDGTexture);
+			RDGTexture->MarkResourceAsUsed();
+
+			RHICmdList.SetShaderTexture(ShaderRHI, Parameter.BaseIndex, RDGTexture->GetRHI());
 		}
 		break;
 		case UBMT_RDG_TEXTURE_SRV:
 		case UBMT_RDG_BUFFER_SRV:
 		{
 			//HACKHACK: defer SRVs binding after UAVs 
-			GraphSRVs.Add(ParameterBinding);
+			GraphSRVs.Add(Parameter);
 		}
 		break;
 		case UBMT_UAV:
 		case UBMT_RDG_TEXTURE_UAV:
 		case UBMT_RDG_BUFFER_UAV:
 		{
-			SetShaderUAV(RHICmdList, ShaderRHI, Base, ParameterBinding);
+			SetShaderUAV(RHICmdList, ShaderRHI, Reader, Parameter);
 		}
 		break;
 		default:
@@ -678,35 +702,33 @@ inline void SetShaderParametersInternal(
 	}
 
 	//HACKHACK: Bind SRVs after UAVs as a workaround for D3D11 RHI unbinding SRVs when binding a UAV on the same resource even when the views don't overlap.
-	for (const FShaderParameterBindings::FResourceParameter& ParameterBinding : GraphSRVs)
+	for (const FShaderParameterBindings::FResourceParameter& Parameter : GraphSRVs)
 	{
-		auto GraphSRV = *reinterpret_cast<FRDGShaderResourceView* const*>(Base + ParameterBinding.ByteOffset);
+		FRDGShaderResourceView* RDGShaderResourceView = Reader.Read<FRDGShaderResourceView*>(Parameter);
+		checkSlow(RDGShaderResourceView);
+		RDGShaderResourceView->MarkResourceAsUsed();
 
-		checkSlow(GraphSRV);
-		GraphSRV->MarkResourceAsUsed();
-		RHICmdList.SetShaderResourceViewParameter(ShaderRHI, ParameterBinding.BaseIndex, GraphSRV->GetRHI());
+		RHICmdList.SetShaderResourceViewParameter(ShaderRHI, Parameter.BaseIndex, RDGShaderResourceView->GetRHI());
 	}
 
 	// Graph Uniform Buffers
-	for (const FShaderParameterBindings::FParameterStructReference& ParameterBinding : Bindings.GraphUniformBuffers)
+	for (const FShaderParameterBindings::FParameterStructReference& Parameter : Bindings.GraphUniformBuffers)
 	{
-		const FRDGUniformBufferBinding& UniformBufferBinding = *reinterpret_cast<const FRDGUniformBufferBinding*>(Base + ParameterBinding.ByteOffset);
-
+		const FRDGUniformBufferBinding& UniformBufferBinding = Reader.Read<FRDGUniformBufferBinding>(Parameter);
 		if (UniformBufferBinding.IsShader())
 		{
 			UniformBufferBinding->MarkResourceAsUsed();
-			RHICmdList.SetShaderUniformBuffer(ShaderRHI, ParameterBinding.BufferIndex, UniformBufferBinding->GetRHI());
+			RHICmdList.SetShaderUniformBuffer(ShaderRHI, Parameter.BufferIndex, UniformBufferBinding->GetRHI());
 		}
 	}
 
 	// Reference structures
-	for (const FShaderParameterBindings::FParameterStructReference& ParameterBinding : Bindings.ParameterReferences)
+	for (const FShaderParameterBindings::FParameterStructReference& Parameter : Bindings.ParameterReferences)
 	{
-		const FUniformBufferBinding& UniformBufferBinding = *reinterpret_cast<const FUniformBufferBinding*>(Base + ParameterBinding.ByteOffset);
-
+		const FUniformBufferBinding& UniformBufferBinding = Reader.Read<FUniformBufferBinding>(Parameter);
 		if (UniformBufferBinding.IsShader())
 		{
-			RHICmdList.SetShaderUniformBuffer(ShaderRHI, ParameterBinding.BufferIndex, UniformBufferBinding.GetUniformBuffer());
+			RHICmdList.SetShaderUniformBuffer(ShaderRHI, Parameter.BufferIndex, UniformBufferBinding.GetUniformBuffer());
 		}
 	}
 }
@@ -748,61 +770,61 @@ void SetShaderParameters(
 	const FRHIUniformBufferLayout* RootUniformBufferLayout,
 	const uint8* Base)
 {
-	for (const FShaderParameterBindings::FResourceParameter& ParameterBinding : Bindings.ResourceParameters)
+	const FShaderParameterReader Reader(Base);
+
+	for (const FShaderParameterBindings::FResourceParameter& Parameter : Bindings.ResourceParameters)
 	{
-		EUniformBufferBaseType BaseType = (EUniformBufferBaseType)ParameterBinding.BaseType;
+		const EUniformBufferBaseType BaseType = static_cast<EUniformBufferBaseType>(Parameter.BaseType);
 		switch (BaseType)
 		{
 		case UBMT_TEXTURE:
 		{
-			auto ShaderParameterRef = *(FRHITexture**)(Base + ParameterBinding.ByteOffset);
-			RTBindingsWriter.SetTexture(ParameterBinding.BaseIndex, ShaderParameterRef);
+			FRHITexture* Texture = Reader.Read<FRHITexture*>(Parameter);
+			RTBindingsWriter.SetTexture(Parameter.BaseIndex, Texture);
 		}
 		break;
 		case UBMT_SRV:
 		{
-			FRHIShaderResourceView* ShaderParameterRef = *(FRHIShaderResourceView**)(Base + ParameterBinding.ByteOffset);
-			RTBindingsWriter.SetSRV(ParameterBinding.BaseIndex, ShaderParameterRef);
+			FRHIShaderResourceView* ShaderResourceView = Reader.Read<FRHIShaderResourceView*>(Parameter);
+			RTBindingsWriter.SetSRV(Parameter.BaseIndex, ShaderResourceView);
 		}
 		break;
 		case UBMT_UAV:
 		{
-			FRHIUnorderedAccessView* ShaderParameterRef = *(FRHIUnorderedAccessView**)(Base + ParameterBinding.ByteOffset);
-			RTBindingsWriter.SetUAV(ParameterBinding.BaseIndex, ShaderParameterRef);
+			FRHIUnorderedAccessView* UnorderedAccessView = Reader.Read<FRHIUnorderedAccessView*>(Parameter);
+			RTBindingsWriter.SetUAV(Parameter.BaseIndex, UnorderedAccessView);
 		}
 		break;
 		case UBMT_SAMPLER:
 		{
-			FRHISamplerState* ShaderParameterRef = *(FRHISamplerState**)(Base + ParameterBinding.ByteOffset);
-			RTBindingsWriter.SetSampler(ParameterBinding.BaseIndex, ShaderParameterRef);
+			FRHISamplerState* SamplerState = Reader.Read<FRHISamplerState*>(Parameter);
+			RTBindingsWriter.SetSampler(Parameter.BaseIndex, SamplerState);
 		}
 		break;
 		case UBMT_RDG_TEXTURE:
 		{
-			auto GraphTexture = *reinterpret_cast<FRDGTexture* const*>(Base + ParameterBinding.ByteOffset);
-			checkSlow(GraphTexture);
-			GraphTexture->MarkResourceAsUsed();
-			RTBindingsWriter.SetTexture(ParameterBinding.BaseIndex, GraphTexture->GetRHI());
+			FRDGTexture* RDGTexture = Reader.Read<FRDGTexture*>(Parameter);
+			checkSlow(RDGTexture);
+			RDGTexture->MarkResourceAsUsed();
+			RTBindingsWriter.SetTexture(Parameter.BaseIndex, RDGTexture->GetRHI());
 		}
 		break;
 		case UBMT_RDG_TEXTURE_SRV:
 		case UBMT_RDG_BUFFER_SRV:
 		{
-			auto GraphSRV = *reinterpret_cast<FRDGShaderResourceView* const*>(Base + ParameterBinding.ByteOffset);
-
-			checkSlow(GraphSRV);
-			GraphSRV->MarkResourceAsUsed();
-			RTBindingsWriter.SetSRV(ParameterBinding.BaseIndex, GraphSRV->GetRHI());
+			FRDGShaderResourceView* RDGShaderResourceView = Reader.Read<FRDGShaderResourceView*>(Parameter);
+			checkSlow(RDGShaderResourceView);
+			RDGShaderResourceView->MarkResourceAsUsed();
+			RTBindingsWriter.SetSRV(Parameter.BaseIndex, RDGShaderResourceView->GetRHI());
 		}
 		break;
 		case UBMT_RDG_TEXTURE_UAV:
 		case UBMT_RDG_BUFFER_UAV:
 		{
-			auto UAV = *reinterpret_cast<FRDGUnorderedAccessView* const*>(Base + ParameterBinding.ByteOffset);
-
-			checkSlow(UAV);
-			UAV->MarkResourceAsUsed();
-			RTBindingsWriter.SetUAV(ParameterBinding.BaseIndex, UAV->GetRHI());
+			FRDGUnorderedAccessView* RDGUnorderedAccessView = Reader.Read<FRDGUnorderedAccessView*>(Parameter);
+			checkSlow(RDGUnorderedAccessView);
+			RDGUnorderedAccessView->MarkResourceAsUsed();
+			RTBindingsWriter.SetUAV(Parameter.BaseIndex, RDGUnorderedAccessView->GetRHI());
 		}
 		break;
 		default:
@@ -812,20 +834,20 @@ void SetShaderParameters(
 	}
 
 	// Graph Uniform Buffers
-	for (const FShaderParameterBindings::FParameterStructReference& ParameterBinding : Bindings.GraphUniformBuffers)
+	for (const FShaderParameterBindings::FParameterStructReference& Parameter : Bindings.GraphUniformBuffers)
 	{
-		const FRDGUniformBufferBinding& UniformBufferBinding = *reinterpret_cast<const FRDGUniformBufferBinding*>(Base + ParameterBinding.ByteOffset);
+		const FRDGUniformBufferBinding& UniformBufferBinding = Reader.Read<FRDGUniformBufferBinding>(Parameter);
 
 		checkSlow(UniformBufferBinding);
 		UniformBufferBinding->MarkResourceAsUsed();
-		RTBindingsWriter.SetUniformBuffer(ParameterBinding.BufferIndex, UniformBufferBinding->GetRHI());
+		RTBindingsWriter.SetUniformBuffer(Parameter.BufferIndex, UniformBufferBinding->GetRHI());
 	}
 
 	// Referenced uniform buffers
-	for (const FShaderParameterBindings::FParameterStructReference& ParameterBinding : Bindings.ParameterReferences)
+	for (const FShaderParameterBindings::FParameterStructReference& Parameter : Bindings.ParameterReferences)
 	{
-		const FUniformBufferBinding& UniformBufferBinding = *reinterpret_cast<const FUniformBufferBinding*>(Base + ParameterBinding.ByteOffset);
-		RTBindingsWriter.SetUniformBuffer(ParameterBinding.BufferIndex, UniformBufferBinding.GetUniformBuffer());
+		const FUniformBufferBinding& UniformBufferBinding = Reader.Read<FUniformBufferBinding>(Parameter);
+		RTBindingsWriter.SetUniformBuffer(Parameter.BufferIndex, UniformBufferBinding.GetUniformBuffer());
 	}
 
 	// Root uniform buffer.
diff --git a/Engine/Source/Runtime/RenderCore/Public/ShaderCore.h b/Engine/Source/Runtime/RenderCore/Public/ShaderCore.h
--- a/Engine/Source/Runtime/RenderCore/Public/ShaderCore.h
+++ b/Engine/Source/Runtime/RenderCore/Public/ShaderCore.h
@@ -232,6 +232,7 @@ public:
 	FShaderParameterMap()
 	{}
 
+	RENDERCORE_API TOptional<FParameterAllocation> FindParameterAllocation(const FString& ParameterName) const;
 	RENDERCORE_API bool FindParameterAllocation(const TCHAR* ParameterName,uint16& OutBufferIndex,uint16& OutBaseIndex,uint16& OutSize) const;
 	RENDERCORE_API bool ContainsParameterAllocation(const TCHAR* ParameterName) const;
 	RENDERCORE_API void AddParameterAllocation(const TCHAR* ParameterName,uint16 BufferIndex,uint16 BaseIndex,uint16 Size,EShaderParameterType ParameterType);
