Fix up VSM compile permutation logic to not depend on the runtime variable

diff --git a/Engine/Source/Runtime/RenderCore/Public/RenderUtils.h b/Engine/Source/Runtime/RenderCore/Public/RenderUtils.h
--- a/Engine/Source/Runtime/RenderCore/Public/RenderUtils.h
+++ b/Engine/Source/Runtime/RenderCore/Public/RenderUtils.h
@@ -766,12 +766,21 @@ inline bool UseVirtualShadowMaps(EShaderPlatform ShaderPlatform, const FStaticFe
 {
 	static const auto EnableVirtualSMCVar = IConsoleManager::Get().FindConsoleVariable(TEXT("r.Shadow.Virtual.Enable"));
 	const bool bVirtualShadowMapsEnabled = EnableVirtualSMCVar ? (EnableVirtualSMCVar->GetInt() != 0) : false;
-	return bVirtualShadowMapsEnabled && DoesRuntimeSupportNanite(ShaderPlatform, true /* check for atomics */, true /* check project setting */);
+	return bVirtualShadowMapsEnabled && DoesRuntimeSupportNanite(ShaderPlatform, true /* check for atomics */, false /* check project setting */);
 }
 
 /**
-* Returns true if non-Nanite virtual shadow maps are enabled by CVar r.Shadow.Virtual.NonNaniteVSM 
-* and UseVirtualShadowMaps is true for the given platform and feature level.
+* Returns true if non-Nanite virtual shadow maps are enabled by CVar r.Shadow.Virtual.NonNaniteVSM
+* and the runtime supports Nanite/virtual shadow maps.
+*/
+inline bool DoesPlatformSupportNonNaniteVirtualShadowMaps(EShaderPlatform ShaderPlatform)
+{
+	static const auto EnableCVar = IConsoleManager::Get().FindConsoleVariable(TEXT("r.Shadow.Virtual.NonNaniteVSM"));	
+	return EnableCVar->GetInt() != 0 && DoesPlatformSupportNanite(ShaderPlatform, false /* check project setting */);
+}
+
+/**
+* Similar to DoesPlatformSupportNonNaniteVirtualShadowMaps, but checks if nanite and virtual shadow maps are enabled (at runtime).
 */
 inline bool UseNonNaniteVirtualShadowMaps(EShaderPlatform ShaderPlatform, const FStaticFeatureLevel FeatureLevel)
 {
diff --git a/Engine/Source/Runtime/Renderer/Private/ShadowDepthRendering.cpp b/Engine/Source/Runtime/Renderer/Private/ShadowDepthRendering.cpp
--- a/Engine/Source/Runtime/Renderer/Private/ShadowDepthRendering.cpp
+++ b/Engine/Source/Runtime/Renderer/Private/ShadowDepthRendering.cpp
@@ -327,7 +327,7 @@ public:
 			return false;
 		}
 
-		if (ShaderMode == VertexShadowDepth_VirtualShadowMap && !UseNonNaniteVirtualShadowMaps(Platform, GetMaxSupportedFeatureLevel(Platform)))
+		if (ShaderMode == VertexShadowDepth_VirtualShadowMap && !DoesPlatformSupportNonNaniteVirtualShadowMaps(Platform))
 		{
 			return false;
 		}
@@ -444,7 +444,7 @@ public:
 			return false;
 		}
 
-		if (ShaderMode == PixelShadowDepth_VirtualShadowMap && !UseNonNaniteVirtualShadowMaps(Platform, GetMaxSupportedFeatureLevel(Platform)))
+		if (ShaderMode == PixelShadowDepth_VirtualShadowMap && !DoesPlatformSupportNonNaniteVirtualShadowMaps(Platform))
 		{
 			return false;
 		}
