- When FShaderPlatformCachedIniValue looks up a value for the current running ShaderPlatform, use the platform's own value, even if ShaderPlatformToPlatformName returned a different platform (like Vulkan will return Windows, even when running on Linux)

diff --git a/Engine/Source/Runtime/RenderCore/Public/RenderUtils.h b/Engine/Source/Runtime/RenderCore/Public/RenderUtils.h
--- a/Engine/Source/Runtime/RenderCore/Public/RenderUtils.h
+++ b/Engine/Source/Runtime/RenderCore/Public/RenderUtils.h
@@ -546,8 +546,10 @@ struct RENDERCORE_API FShaderPlatformCachedIniValue
 			CVar = IConsoleManager::Get().FindConsoleVariable(*CVarName);
 		}
 
-		// if we are looking up our own platform, just use the current value
-		if (IniPlatformName == FPlatformProperties::IniPlatformName())
+		// if we are looking up our own platform, just use the current value, however
+		// ShaderPlatformToPlatformName can return the wrong platform than expected - for instance, Linux Vulkan will return Windows
+		// so instead of hitting an asser below, we detect that the request SP is the current SP, and use the CVar value that is set currently
+		if (IniPlatformName == FPlatformProperties::IniPlatformName() || ShaderPlatform == GMaxRHIShaderPlatform)
 		{
 			checkf(CVar != nullptr, TEXT("Failed to find CVar %s when getting current value for FShaderPlatformCachedIniValue"));
 
