Fixed: Enabling auto-instancing crashes on Mali devices with more than 2000 mesh instances in the scene

diff --git a/Engine/Source/Runtime/RenderCore/Private/UnifiedBuffer.cpp b/Engine/Source/Runtime/RenderCore/Private/UnifiedBuffer.cpp
--- a/Engine/Source/Runtime/RenderCore/Private/UnifiedBuffer.cpp
+++ b/Engine/Source/Runtime/RenderCore/Private/UnifiedBuffer.cpp
@@ -618,12 +618,10 @@ void FScatterUploadBuffer::Init( uint32 NumElements, uint32 InNumBytesPerElement
 	const uint32 TypeSize = bInFloat4Buffer ? 16 : 4;
 
 	uint32 ScatterBytes = NumElements * sizeof( uint32 );
-	uint32 ScatterBufferSize = (uint32)FMath::Min( (uint64)FMath::RoundUpToPowerOfTwo( ScatterBytes ), GetMaxBufferDimension() * sizeof( uint32 ) );
-	check( ScatterBufferSize >= ScatterBytes );
-
+	uint32 ScatterBufferSize = FMath::RoundUpToPowerOfTwo( ScatterBytes );
+	
 	uint32 UploadBytes = NumElements * NumBytesPerElement;
-	uint32 UploadBufferSize = (uint32)FMath::Min( (uint64)FMath::RoundUpToPowerOfTwo( UploadBytes ), GetMaxBufferDimension() * TypeSize );
-	check( UploadBufferSize >= UploadBytes );
+	uint32 UploadBufferSize = FMath::RoundUpToPowerOfTwo( UploadBytes );
 
 	if (bUploadViaCreate)
 	{
