Fix hitching while precompiling PSO cache.

- Add a new API to release preloaded shader to provide a way to manage preloaded data (which was previously statically allocated).

diff --git a/Engine/Source/Runtime/RenderCore/Private/ShaderCodeArchive.cpp b/Engine/Source/Runtime/RenderCore/Private/ShaderCodeArchive.cpp
--- a/Engine/Source/Runtime/RenderCore/Private/ShaderCodeArchive.cpp
+++ b/Engine/Source/Runtime/RenderCore/Private/ShaderCodeArchive.cpp
@@ -1785,15 +1785,9 @@ void FIoStoreShaderCodeArchive::MarkPreloadEntrySkipped(int32 ShaderGroupIndex)
 
 bool FIoStoreShaderCodeArchive::PreloadShader(int32 ShaderIndex, FGraphEventArray& OutCompletionEvents)
 {
-	// [RCL] 2022-01-12 because PreloadShader API is unbalanced and there's code (eg. FShaderPipelineCache::PreparePipelineBatch) that calls PreloadShader but never ReleasePreloadedShader (which isn't
-	// even exposed to client code), this API is disabled to prevent leaking preload entries. Something to address separately - UE-139060.
-#if 0
 	LLM_SCOPE(ELLMTag::Shaders);
 	DebugVisualizer.MarkExplicitlyPreloadedForVisualization(ShaderIndex);
 	return PreloadShaderGroup(GetGroupIndexForShader(ShaderIndex), OutCompletionEvents);
-#else
-	return true;
-#endif
 }
 
 bool FIoStoreShaderCodeArchive::GroupOnlyContainsRaytracingShaders(int32 ShaderGroupIndex)
diff --git a/Engine/Source/Runtime/RenderCore/Private/ShaderCodeLibrary.cpp b/Engine/Source/Runtime/RenderCore/Private/ShaderCodeLibrary.cpp
--- a/Engine/Source/Runtime/RenderCore/Private/ShaderCodeLibrary.cpp
+++ b/Engine/Source/Runtime/RenderCore/Private/ShaderCodeLibrary.cpp
@@ -860,6 +860,12 @@ public:
 		}
 	}
 
+	void ReleasePreloadedShader(int32 ShaderIndex)
+	{
+		SCOPED_LOADTIMER(FShaderLibraryInstance_PreloadShader);
+		Library->ReleasePreloadedShader(ShaderIndex);
+	}
+
 	TRefCountPtr<FShaderMapResource_SharedCode> GetResource(int32 ShaderMapIndex)
 	{
 		FRWScopeLock Locker(ResourceLock, SLT_ReadOnly);
@@ -2259,6 +2265,18 @@ public:
 		return false;
 	}
 
+	bool ReleasePreloadedShader(const FSHAHash& Hash)
+	{
+		int32 ShaderIndex = INDEX_NONE;
+		FShaderLibraryInstance* LibraryInstance = FindShaderLibraryForShader(Hash, ShaderIndex);
+		if (LibraryInstance)
+		{
+			LibraryInstance->ReleasePreloadedShader(ShaderIndex);
+			return true;
+		}
+		return false;
+	}
+
 	bool ContainsShaderCode(const FSHAHash& Hash)
 	{
 		int32 ShaderIndex = INDEX_NONE;
@@ -2695,6 +2713,15 @@ bool FShaderCodeLibrary::PreloadShader(const FSHAHash& Hash, FArchive* Ar)
 	return false;
 }
 
+bool FShaderCodeLibrary::ReleasePreloadedShader(const FSHAHash& Hash)
+{
+	if (FShaderLibrariesCollection::Impl)
+	{
+		return FShaderLibrariesCollection::Impl->ReleasePreloadedShader(Hash);
+	}
+	return false;
+}
+
 FVertexShaderRHIRef FShaderCodeLibrary::CreateVertexShader(EShaderPlatform Platform, const FSHAHash& Hash)
 {
 	if (FShaderLibrariesCollection::Impl)
diff --git a/Engine/Source/Runtime/RenderCore/Private/ShaderPipelineCache.cpp b/Engine/Source/Runtime/RenderCore/Private/ShaderPipelineCache.cpp
--- a/Engine/Source/Runtime/RenderCore/Private/ShaderPipelineCache.cpp
+++ b/Engine/Source/Runtime/RenderCore/Private/ShaderPipelineCache.cpp
@@ -792,6 +792,166 @@ bool FShaderPipelineCache::Precompile(FRHICommandListImmediate& RHICmdList, ESha
 	return bOk;
 }
 
+bool FShaderPipelineCache::CompileJob::PreloadShaders(TSet<FSHAHash>& OutRequiredShaders, bool &bOutCompatible)
+{
+	static FSHAHash EmptySHA;
+	bool bOK = true;
+	TArray<FSHAHash> ShadersToUnpreloadInCaseOfFailure;
+
+	auto PreloadShader = [&bOK, &ShadersToUnpreloadInCaseOfFailure](const FSHAHash& ShaderHash, FShaderPipelineCacheArchive* ReadReqs)
+	{
+		if (bOK && ShaderHash != EmptySHA)
+		{
+			bOK &= FShaderCodeLibrary::PreloadShader(ShaderHash, ReadReqs);
+			if (bOK)
+			{
+				ShadersToUnpreloadInCaseOfFailure.Add(ShaderHash);
+			}
+			UE_CLOG(!bOK, LogRHI, Verbose, TEXT("Failed to read shader: %s"), *(ShaderHash.ToString()));
+		}
+	};
+
+	if (PSO.Type == FPipelineCacheFileFormatPSO::DescriptorType::Graphics)
+	{
+		// See if the shaders exist in the current code libraries, before trying to load the shader data
+		if (PSO.GraphicsDesc.MeshShader != EmptySHA)
+		{
+			OutRequiredShaders.Add(PSO.GraphicsDesc.MeshShader);
+			bOK &= FShaderCodeLibrary::ContainsShaderCode(PSO.GraphicsDesc.MeshShader);
+			UE_CLOG(!bOK, LogRHI, Verbose, TEXT("Failed to find MeshShader shader: %s"), *(PSO.GraphicsDesc.MeshShader.ToString()));
+
+			if (PSO.GraphicsDesc.AmplificationShader != EmptySHA)
+			{
+				OutRequiredShaders.Add(PSO.GraphicsDesc.AmplificationShader);
+				bOK &= FShaderCodeLibrary::ContainsShaderCode(PSO.GraphicsDesc.AmplificationShader);
+				UE_CLOG(!bOK, LogRHI, Verbose, TEXT("Failed to find AmplificationShader shader: %s"), *(PSO.GraphicsDesc.AmplificationShader.ToString()));
+			}
+		}
+		else if (PSO.GraphicsDesc.VertexShader != EmptySHA)
+		{
+			OutRequiredShaders.Add(PSO.GraphicsDesc.VertexShader);
+			bOK &= FShaderCodeLibrary::ContainsShaderCode(PSO.GraphicsDesc.VertexShader);
+			UE_CLOG(!bOK, LogRHI, Verbose, TEXT("Failed to find VertexShader shader: %s"), *(PSO.GraphicsDesc.VertexShader.ToString()));
+
+			if (PSO.GraphicsDesc.GeometryShader != EmptySHA)
+			{
+				OutRequiredShaders.Add(PSO.GraphicsDesc.GeometryShader);
+				bOK &= FShaderCodeLibrary::ContainsShaderCode(PSO.GraphicsDesc.GeometryShader);
+				UE_CLOG(!bOK, LogRHI, Verbose, TEXT("Failed to find GeometryShader shader: %s"), *(PSO.GraphicsDesc.GeometryShader.ToString()));
+			}
+		}
+		else
+		{
+			// if we don't have a vertex shader then we won't bother to add any shaders to the list of outstanding shaders to load
+			// Later on this PSO will be killed forever because it is truly bogus.
+			UE_LOG(LogRHI, Error, TEXT("PSO Entry has no vertex shader!"));
+			bOK = false;
+		}
+
+		if (PSO.GraphicsDesc.FragmentShader != EmptySHA)
+		{
+			OutRequiredShaders.Add(PSO.GraphicsDesc.FragmentShader);
+			bOK &= FShaderCodeLibrary::ContainsShaderCode(PSO.GraphicsDesc.FragmentShader);
+			UE_CLOG(!bOK, LogRHI, Verbose, TEXT("Failed to find FragmentShader shader: %s"), *(PSO.GraphicsDesc.FragmentShader.ToString()));
+		}
+
+		// If everything is OK then we can issue reads of the actual shader code
+		PreloadShader(PSO.GraphicsDesc.VertexShader, ReadRequests);
+		PreloadShader(PSO.GraphicsDesc.MeshShader, ReadRequests);
+		PreloadShader(PSO.GraphicsDesc.AmplificationShader, ReadRequests);
+		PreloadShader(PSO.GraphicsDesc.FragmentShader, ReadRequests);
+		PreloadShader(PSO.GraphicsDesc.GeometryShader, ReadRequests);
+	}
+	else if (PSO.Type == FPipelineCacheFileFormatPSO::DescriptorType::Compute)
+	{
+		if (PSO.ComputeDesc.ComputeShader != EmptySHA)
+		{
+			OutRequiredShaders.Add(PSO.ComputeDesc.ComputeShader);
+			PreloadShader(PSO.ComputeDesc.ComputeShader, ReadRequests);
+		}
+		else
+		{
+			bOK = false;
+			UE_LOG(LogRHI, Error, TEXT("Invalid PSO entry in pipeline cache!"));
+		}
+	}
+	else if (PSO.Type == FPipelineCacheFileFormatPSO::DescriptorType::RayTracing)
+	{
+		if (IsRayTracingEnabled())
+		{
+			if (PSO.RayTracingDesc.ShaderHash != EmptySHA)
+			{
+				OutRequiredShaders.Add(PSO.RayTracingDesc.ShaderHash);
+				PreloadShader(PSO.RayTracingDesc.ShaderHash, ReadRequests);
+			}
+			else
+			{
+				bOK = false;
+				UE_LOG(LogRHI, Error, TEXT("Invalid PSO entry in pipeline cache!"));
+			}
+		}
+		else
+		{
+			bOutCompatible = false;
+		}
+	}
+	else
+	{
+		bOK = false;
+		UE_LOG(LogRHI, Error, TEXT("Invalid PSO entry in pipeline cache!"));
+	}
+
+	if (!bOK)
+	{
+		for (const FSHAHash& PreloadedShader : ShadersToUnpreloadInCaseOfFailure)
+		{
+			FShaderCodeLibrary::ReleasePreloadedShader(PreloadedShader);
+		}
+		ShadersToUnpreloadInCaseOfFailure.Empty();
+	}
+	
+	bShadersPreloaded = bOK;
+	return bShadersPreloaded;
+}
+
+void FShaderPipelineCache::CompileJob::ReleasePreloadedShaders()
+{
+	if (bShadersPreloaded)
+	{
+		static FSHAHash EmptySHA;
+
+		auto ReleasePreloadedShader = [](const FSHAHash& ShaderHash)
+		{
+			if (ShaderHash != EmptySHA)
+			{
+				FShaderCodeLibrary::ReleasePreloadedShader(ShaderHash);
+			}
+		};
+
+		if (PSO.Type == FPipelineCacheFileFormatPSO::DescriptorType::Graphics)
+		{
+			ReleasePreloadedShader(PSO.GraphicsDesc.VertexShader);
+			ReleasePreloadedShader(PSO.GraphicsDesc.MeshShader);
+			ReleasePreloadedShader(PSO.GraphicsDesc.AmplificationShader);
+			ReleasePreloadedShader(PSO.GraphicsDesc.FragmentShader);
+			ReleasePreloadedShader(PSO.GraphicsDesc.GeometryShader);
+		}
+		else if (PSO.Type == FPipelineCacheFileFormatPSO::DescriptorType::Compute)
+		{
+			ReleasePreloadedShader(PSO.ComputeDesc.ComputeShader);
+		}
+		else if (PSO.Type == FPipelineCacheFileFormatPSO::DescriptorType::RayTracing)
+		{
+			if (IsRayTracingEnabled())
+			{
+				ReleasePreloadedShader(PSO.RayTracingDesc.ShaderHash);
+			}
+		}
+
+		bShadersPreloaded = true;
+	}
+}
+
 void FShaderPipelineCache::PreparePipelineBatch(TDoubleLinkedList<FPipelineCacheFileFormatPSORead*>& PipelineBatch)
 {
 	TDoubleLinkedList<FPipelineCacheFileFormatPSORead*>::TDoubleLinkedListNode* CurrentNode = PipelineBatch.GetHead();
@@ -825,126 +985,8 @@ void FShaderPipelineCache::PreparePipelineBatch(TDoubleLinkedList<FPipelineCache
 			CompileJob AsyncJob;
 			AsyncJob.PSO = PSO;
 			AsyncJob.ReadRequests = new FShaderPipelineCacheArchive;
-			
-            static FSHAHash EmptySHA;
-            
-            if (PSO.Type == FPipelineCacheFileFormatPSO::DescriptorType::Graphics)
-			{
-                // See if the shaders exist in the current code libraries, before trying to load the shader data
-				if (PSO.GraphicsDesc.MeshShader != EmptySHA)
-				{
-					RequiredShaders.Add(PSO.GraphicsDesc.MeshShader);
-					bOK &= FShaderCodeLibrary::ContainsShaderCode(PSO.GraphicsDesc.MeshShader);
-					UE_CLOG(!bOK, LogRHI, Verbose, TEXT("Failed to find MeshShader shader: %s"), *(PSO.GraphicsDesc.MeshShader.ToString()));
-
-					if (PSO.GraphicsDesc.AmplificationShader != EmptySHA)
-					{
-						RequiredShaders.Add(PSO.GraphicsDesc.AmplificationShader);
-						bOK &= FShaderCodeLibrary::ContainsShaderCode(PSO.GraphicsDesc.AmplificationShader);
-						UE_CLOG(!bOK, LogRHI, Verbose, TEXT("Failed to find AmplificationShader shader: %s"), *(PSO.GraphicsDesc.AmplificationShader.ToString()));
-					}
-				}
-                else if (PSO.GraphicsDesc.VertexShader != EmptySHA)
-                {
-                    RequiredShaders.Add(PSO.GraphicsDesc.VertexShader);
-                    bOK &= FShaderCodeLibrary::ContainsShaderCode(PSO.GraphicsDesc.VertexShader);
-                    UE_CLOG(!bOK, LogRHI, Verbose, TEXT("Failed to find VertexShader shader: %s"), *(PSO.GraphicsDesc.VertexShader.ToString()));
-
-					if (PSO.GraphicsDesc.GeometryShader != EmptySHA)
-					{
-						RequiredShaders.Add(PSO.GraphicsDesc.GeometryShader);
-						bOK &= FShaderCodeLibrary::ContainsShaderCode(PSO.GraphicsDesc.GeometryShader);
-						UE_CLOG(!bOK, LogRHI, Verbose, TEXT("Failed to find GeometryShader shader: %s"), *(PSO.GraphicsDesc.GeometryShader.ToString()));
-					}
-				}
-				else
-				{
-					// if we don't have a vertex shader then we won't bother to add any shaders to the list of outstanding shaders to load
-					// Later on this PSO will be killed forever because it is truly bogus.
-					UE_LOG(LogRHI, Error, TEXT("PSO Entry has no vertex shader: %u this is an invalid entry!"), PSORead->Hash);
-					bOK = false;
-				}
-
-				if (PSO.GraphicsDesc.FragmentShader != EmptySHA)
-				{
-					RequiredShaders.Add(PSO.GraphicsDesc.FragmentShader);
-					bOK &= FShaderCodeLibrary::ContainsShaderCode(PSO.GraphicsDesc.FragmentShader);
-					UE_CLOG(!bOK, LogRHI, Verbose, TEXT("Failed to find FragmentShader shader: %s"), *(PSO.GraphicsDesc.FragmentShader.ToString()));
-				}
-
-                // If everything is OK then we can issue reads of the actual shader code
-				if (bOK && PSO.GraphicsDesc.VertexShader != FSHAHash())
-				{
-					bOK &= FShaderCodeLibrary::PreloadShader(PSO.GraphicsDesc.VertexShader, AsyncJob.ReadRequests);
-                    UE_CLOG(!bOK, LogRHI, Verbose, TEXT("Failed to read VertexShader shader: %s"), *(PSO.GraphicsDesc.VertexShader.ToString()));
-				}
-
-				// If everything is OK then we can issue reads of the actual shader code
-				if (bOK && PSO.GraphicsDesc.MeshShader != FSHAHash())
-				{
-					bOK &= FShaderCodeLibrary::PreloadShader(PSO.GraphicsDesc.MeshShader, AsyncJob.ReadRequests);
-					UE_CLOG(!bOK, LogRHI, Verbose, TEXT("Failed to read MeshShader shader: %s"), *(PSO.GraphicsDesc.MeshShader.ToString()));
-				}
-
-				// If everything is OK then we can issue reads of the actual shader code
-				if (bOK && PSO.GraphicsDesc.AmplificationShader != FSHAHash())
-				{
-					bOK &= FShaderCodeLibrary::PreloadShader(PSO.GraphicsDesc.AmplificationShader, AsyncJob.ReadRequests);
-					UE_CLOG(!bOK, LogRHI, Verbose, TEXT("Failed to read AmplificationShader shader: %s"), *(PSO.GraphicsDesc.AmplificationShader.ToString()));
-				}
-
-				if (bOK && PSO.GraphicsDesc.FragmentShader != EmptySHA)
-				{
-					bOK &= FShaderCodeLibrary::PreloadShader(PSO.GraphicsDesc.FragmentShader, AsyncJob.ReadRequests);
-                    UE_CLOG(!bOK, LogRHI, Verbose, TEXT("Failed to read FragmentShader shader: %s"), *(PSO.GraphicsDesc.FragmentShader.ToString()));
-				}
 		
-				if (bOK && PSO.GraphicsDesc.GeometryShader != EmptySHA)
-				{
-					bOK &= FShaderCodeLibrary::PreloadShader(PSO.GraphicsDesc.GeometryShader, AsyncJob.ReadRequests);
-                    UE_CLOG(!bOK, LogRHI, Verbose, TEXT("Failed to read GeometryShader shader: %s"), *(PSO.GraphicsDesc.GeometryShader.ToString()));
-				}
-			}
-			else if (PSO.Type == FPipelineCacheFileFormatPSO::DescriptorType::Compute)
-			{
-				if (PSO.ComputeDesc.ComputeShader != EmptySHA)
-				{
-                    RequiredShaders.Add(PSO.ComputeDesc.ComputeShader);
-					bOK &= FShaderCodeLibrary::PreloadShader(PSO.ComputeDesc.ComputeShader, AsyncJob.ReadRequests);
-                    UE_CLOG(!bOK, LogRHI, Verbose, TEXT("Failed to find ComputeShader shader: %s"), *(PSO.ComputeDesc.ComputeShader.ToString()));
-				}
-				else
-				{
-					bOK = false;
-					UE_LOG(LogRHI, Error, TEXT("Invalid PSO entry in pipeline cache!"));
-				}
-			}
-			else if (PSO.Type == FPipelineCacheFileFormatPSO::DescriptorType::RayTracing)
-			{
-				if (IsRayTracingEnabled())
-				{
-					if (PSO.RayTracingDesc.ShaderHash != EmptySHA)
-					{
-						RequiredShaders.Add(PSO.RayTracingDesc.ShaderHash);
-						bOK &= FShaderCodeLibrary::PreloadShader(PSO.RayTracingDesc.ShaderHash, AsyncJob.ReadRequests);
-						UE_CLOG(!bOK, LogRHI, Verbose, TEXT("Failed to find RayTracing shader: %s"), *(PSO.RayTracingDesc.ShaderHash.ToString()));
-					}
-					else
-					{
-						bOK = false;
-						UE_LOG(LogRHI, Error, TEXT("Invalid PSO entry in pipeline cache!"));
-					}
-				}
-				else
-				{
-					bCompatible = false;
-				}
-			}
-			else
-			{
-				bOK = false;
-				UE_LOG(LogRHI, Error, TEXT("Invalid PSO entry in pipeline cache!"));
-			}
+			bOK = AsyncJob.PreloadShaders(RequiredShaders, bCompatible);
 			
 			// Then if and only if all shaders can be found do we schedule a compile job
 			// Otherwise this job needs to be put in the shutdown list to correctly release shader code
@@ -1052,6 +1094,7 @@ void FShaderPipelineCache::PrecompilePipelineBatch()
 		Precompile(RHICmdList, GMaxRHIShaderPlatform, CompileTask.PSO);
 		CompiledHashes.Add(PSOHash);
 		
+		CompileTask.ReleasePreloadedShaders();
 		delete CompileTask.ReadRequests;
 		CompileTask.ReadRequests = nullptr;
 		
@@ -1129,7 +1172,8 @@ void FShaderPipelineCache::PollShutdownItems()
 			check(ShutdownReadCompileTasks[i].ReadRequests);
 			if (ShutdownReadCompileTasks[i].ReadRequests->PollExternalReadDependencies())
 			{
-                delete ShutdownReadCompileTasks[i].ReadRequests;
+				ShutdownReadCompileTasks[i].ReleasePreloadedShaders();
+				delete ShutdownReadCompileTasks[i].ReadRequests;
 				ShutdownReadCompileTasks[i].ReadRequests = nullptr;
 				
 				ShutdownReadCompileTasks.RemoveAt(i, 1, false);
@@ -1274,6 +1318,7 @@ FShaderPipelineCache::~FShaderPipelineCache()
 		if(Entry.ReadRequests != nullptr)
 		{
 			Entry.ReadRequests->BlockingWaitComplete();
+			Entry.ReleasePreloadedShaders();
 			delete Entry.ReadRequests;
 			Entry.ReadRequests = nullptr;
 		}
diff --git a/Engine/Source/Runtime/RenderCore/Public/ShaderCodeLibrary.h b/Engine/Source/Runtime/RenderCore/Public/ShaderCodeLibrary.h
--- a/Engine/Source/Runtime/RenderCore/Public/ShaderCodeLibrary.h
+++ b/Engine/Source/Runtime/RenderCore/Public/ShaderCodeLibrary.h
@@ -195,6 +195,7 @@ struct RENDERCORE_API FShaderCodeLibrary
 	static TRefCountPtr<FShaderMapResource> LoadResource(const FSHAHash& Hash, FArchive* Ar);
 
 	static bool PreloadShader(const FSHAHash& Hash, FArchive* Ar);
+	static bool ReleasePreloadedShader(const FSHAHash& Hash);
 
 	static FVertexShaderRHIRef CreateVertexShader(EShaderPlatform Platform, const FSHAHash& Hash);
 	static FPixelShaderRHIRef CreatePixelShader(EShaderPlatform Platform, const FSHAHash& Hash);
diff --git a/Engine/Source/Runtime/RenderCore/Public/ShaderPipelineCache.h b/Engine/Source/Runtime/RenderCore/Public/ShaderPipelineCache.h
--- a/Engine/Source/Runtime/RenderCore/Public/ShaderPipelineCache.h
+++ b/Engine/Source/Runtime/RenderCore/Public/ShaderPipelineCache.h
@@ -72,6 +72,24 @@ class RENDERCORE_API FShaderPipelineCache : public FTickableObjectRenderThread
 	{
 		FPipelineCacheFileFormatPSO PSO;
 		FShaderPipelineCacheArchive* ReadRequests;
+
+		/** Tracks whether the shaders were preloaded. */
+		bool bShadersPreloaded = false;
+
+		/** 
+		 * Schedules preload for all the shaders in the PSO. No shaders are preloaded if one of them isn't in the library.
+		 * Multiple jobs can request to preload the same shaders and then release them independently.
+		 * 
+		 * @param OutRequiredShaders shader hashes of shaders that are actually required by this PSO
+		 * @param bOutCompatible if set to false, then this PSO type isn't supported (RTX as of now)
+		 */
+		bool PreloadShaders(TSet<FSHAHash>& OutRequiredShaders, bool& bOutCompatible);
+
+		/**
+		 * Releases preloaded entries for the shaders in the PSO. Symmetric to to PreloadShaders.
+		 * Multiple jobs can request to preload the same shaders and then release them independently.
+		 */
+		void ReleasePreloadedShaders();
 	};
 
 public:
