Minor RDG improvements in preparation for UnifiedBuffer conversion.
 - Added resource pool counters and events.
 - Added AllocatePooledBuffer method and refactored pool to no longer take a command list.
 - Refactored swap chain barrier logic to be a bit cleaner.
 - Added helper methods to cast between views.
 - Added power of two alignment option to buffer pool.
 - Added GetTypeHash implementations for RDG SRV | UAV descriptors.

diff --git a/Engine/Source/Runtime/RenderCore/Private/DumpGPU.cpp b/Engine/Source/Runtime/RenderCore/Private/DumpGPU.cpp
--- a/Engine/Source/Runtime/RenderCore/Private/DumpGPU.cpp
+++ b/Engine/Source/Runtime/RenderCore/Private/DumpGPU.cpp
@@ -1284,10 +1284,12 @@ struct FRDGResourceDumpContext
 			}
 		}
 
+	PRAGMA_DISABLE_DEPRECATION_WARNINGS
 		if (Desc.UnderlyingType == FRDGBufferDesc::EUnderlyingType::AccelerationStructure)
 		{
 			return;
 		}
+	PRAGMA_ENABLE_DEPRECATION_WARNINGS
 
 		// Dump the resource's binary to a .bin file.
 		if (bDumpResourceBinary && DumpBufferMode == 2)
@@ -1751,7 +1753,7 @@ static const TCHAR* GetPassEventNameWithGPUMask(const FRDGPass* Pass, FString& O
 
 void FRDGBuilder::DumpResourcePassOutputs(const FRDGPass* Pass)
 {
-	if (bInDebugPassScope)
+	if (bSkipAuxiliaryPasses)
 	{
 		return;
 	}
@@ -1767,7 +1769,7 @@ void FRDGBuilder::DumpResourcePassOutputs(const FRDGPass* Pass)
 		return;
 	}
 
-	bInDebugPassScope = true;
+	bSkipAuxiliaryPasses = true;
 
 	TArray<TSharedPtr<FJsonValue>> InputResourceNames;
 	TArray<TSharedPtr<FJsonValue>> OutputResourceNames;
@@ -1993,7 +1995,7 @@ void FRDGBuilder::DumpResourcePassOutputs(const FRDGPass* Pass)
 
 	GRDGResourceDumpContext->PassesCount++;
 
-	bInDebugPassScope = false;
+	bSkipAuxiliaryPasses = false;
 }
 
 #if RDG_DUMP_RESOURCES_AT_EACH_DRAW
diff --git a/Engine/Source/Runtime/RenderCore/Private/RenderGraphBuilder.cpp b/Engine/Source/Runtime/RenderCore/Private/RenderGraphBuilder.cpp
--- a/Engine/Source/Runtime/RenderCore/Private/RenderGraphBuilder.cpp
+++ b/Engine/Source/Runtime/RenderCore/Private/RenderGraphBuilder.cpp
@@ -571,7 +571,7 @@ const TRefCountPtr<FRDGPooledBuffer>& FRDGBuilder::ConvertToExternalBuffer(FRDGB
 		Buffer->bExternal = 1;
 		Buffer->bForceNonTransient = 1;
 		BeginResourceRHI(GetProloguePassHandle(), Buffer);
-		ExternalBuffers.Add(Buffer->PooledBuffer, Buffer);
+		ExternalBuffers.Add(Buffer->GetRHIUnchecked(), Buffer);
 	}
 	return GetPooledBuffer(Buffer);
 }
@@ -662,7 +662,7 @@ FRDGTextureRef FRDGBuilder::RegisterExternalTexture(
 	return RegisterExternalTexture(ExternalPooledTexture, Name, Flags);
 }
 
-FRDGTextureRef FRDGBuilder::RegisterExternalTexture(
+FRDGTexture* FRDGBuilder::RegisterExternalTexture(
 	const TRefCountPtr<IPooledRenderTarget>& ExternalPooledTexture,
 	const TCHAR* Name,
 	ERDGTextureFlags Flags)
@@ -671,7 +671,7 @@ FRDGTextureRef FRDGBuilder::RegisterExternalTexture(
 	FRHITexture* ExternalTextureRHI = ExternalPooledTexture->GetRHI();
 	IF_RDG_ENABLE_DEBUG(checkf(ExternalTextureRHI, TEXT("Attempted to register texture %s, but its RHI texture is null."), Name));
 
-	if (FRDGTextureRef FoundTexture = FindExternalTexture(ExternalTextureRHI))
+	if (FRDGTexture* FoundTexture = FindExternalTexture(ExternalTextureRHI))
 	{
 		return FoundTexture;
 	}
@@ -743,9 +743,9 @@ FRDGBufferRef FRDGBuilder::RegisterExternalBuffer(
 {
 	IF_RDG_ENABLE_DEBUG(UserValidation.ValidateRegisterExternalBuffer(ExternalPooledBuffer, Name, Flags));
 
-	if (FRDGBufferRef* FoundBufferPtr = ExternalBuffers.Find(ExternalPooledBuffer.GetReference()))
+	if (FRDGBuffer* FoundBuffer = FindExternalBuffer(ExternalPooledBuffer))
 	{
-		return *FoundBufferPtr;
+		return FoundBuffer;
 	}
 
 	const FRDGBufferDesc& Desc = ExternalPooledBuffer->Desc;
@@ -768,7 +768,7 @@ FRDGBufferRef FRDGBuilder::RegisterExternalBuffer(
 		Buffer->SetFinalizedAccess(ERHIAccess::ReadOnlyExclusiveMask);
 	}
 
-	ExternalBuffers.Add(ExternalPooledBuffer, Buffer);
+	ExternalBuffers.Add(Buffer->GetRHIUnchecked(), Buffer);
 
 	IF_RDG_ENABLE_DEBUG(UserValidation.ValidateRegisterExternalBuffer(Buffer));
 	IF_RDG_ENABLE_TRACE(Trace.AddResource(Buffer));
@@ -1025,9 +1025,9 @@ void FRDGBuilder::Compile()
 				Texture->bUsedByAsyncComputePass |= bAsyncComputePass;
 				Texture->bCulled = false;
 
-				if (Texture->bSwapChain && !Texture->bSwapChainAlreadyMoved)
+				if (Texture->FirstBarrier == FRDGTexture::EFirstBarrier::ImmediateRequested)
 				{
-					Texture->bSwapChainAlreadyMoved = 1;
+					Texture->FirstBarrier = FRDGTexture::EFirstBarrier::ImmediateConfirmed;
 					Texture->FirstPass = PassHandle;
 					GetWholeResource(Texture->GetState()).SetPass(ERHIPipeline::Graphics, PassHandle);
 				}
@@ -1419,9 +1419,9 @@ void FRDGBuilder::Execute()
 
 	// Create the epilogue pass at the end of the graph just prior to compilation.
 	{
-		bInDebugPassScope = true;
+		bSkipAuxiliaryPasses = true;
 		SetupEmptyPass(EpiloguePass = Passes.Allocate<FRDGSentinelPass>(Allocator, RDG_EVENT_NAME("Graph Epilogue")));
-		bInDebugPassScope = false;
+		bSkipAuxiliaryPasses = false;
 	}
 
 	const FRDGPassHandle ProloguePassHandle = GetProloguePassHandle();
@@ -1436,10 +1436,10 @@ void FRDGBuilder::Execute()
 	{
 		BeginFlushResourcesRHI();
 
-		SetupBufferUploads();
-
 		Compile();
 
+		SetupBufferUploads();
+
 		IF_RDG_GPU_SCOPES(GPUScopeStacks.ReserveOps(Passes.Num()));
 		IF_RDG_CPU_SCOPES(CPUScopeStacks.ReserveOps());
 
@@ -1463,6 +1463,8 @@ void FRDGBuilder::Execute()
 			CSV_SCOPED_TIMING_STAT_EXCLUSIVE(RDG_CollectResources);
 			SCOPED_NAMED_EVENT_TEXT("FRDGBuilder::CollectResources", FColor::Magenta);
 
+			UniformBuffersToCreate.Reserve(UniformBuffers.Num());
+
 			EnumerateExtendedLifetimeResources(Textures, [](FRDGTexture* Texture)
 			{
 				++Texture->ReferenceCount;
@@ -1474,7 +1476,7 @@ void FRDGBuilder::Execute()
 			});
 
 			// Null out any culled external resources so that the reference is freed up.
-	
+
 			for (const auto& Pair : ExternalTextures)
 			{
 				FRDGTexture* Texture = Pair.Value;
@@ -1760,7 +1762,7 @@ void FRDGBuilder::Execute()
 
 FRDGPass* FRDGBuilder::SetupPass(FRDGPass* Pass)
 {
-	IF_RDG_ENABLE_DEBUG(UserValidation.ValidateAddPass(Pass, bInDebugPassScope));
+	IF_RDG_ENABLE_DEBUG(UserValidation.ValidateAddPass(Pass, bSkipAuxiliaryPasses));
 	CSV_SCOPED_TIMING_STAT_EXCLUSIVE_CONDITIONAL(RDGBuilder_SetupPass, GRDGVerboseCSVStats != 0);
 
 	const FRDGParameterStruct PassParameters = Pass->GetParameters();
@@ -1995,15 +1997,11 @@ void FRDGBuilder::SetupPassInternal(FRDGPass* Pass, FRDGPassHandle PassHandle, E
 					PassState.MergeState[Index] = &PassState.State[Index];
 				}
 			}
-
-			PassState.Texture->bCulled = false;
 		}
 
 		for (auto& PassState : Pass->BufferStates)
 		{
 			PassState.MergeState = &PassState.State;
-
-			PassState.Buffer->bCulled = false;
 		}
 
 		check(!EnumHasAnyFlags(PassPipeline, ERHIPipeline::AsyncCompute));
@@ -2020,9 +2018,9 @@ void FRDGBuilder::SetupPassInternal(FRDGPass* Pass, FRDGPassHandle PassHandle, E
 
 	IF_RDG_ENABLE_DEBUG(VisualizePassOutputs(Pass));
 
-	#if RDG_DUMP_RESOURCES
-		DumpResourcePassOutputs( Pass);
-	#endif
+#if RDG_DUMP_RESOURCES
+	DumpResourcePassOutputs(Pass);
+#endif
 }
 
 ///////////////////////////////////////////////////////////////////////////////////////////////////
@@ -2033,6 +2031,11 @@ void FRDGBuilder::SetupBufferUploads()
 
 	for (FUploadedBuffer& UploadedBuffer : UploadedBuffers)
 	{
+		if (UploadedBuffer.Buffer->bCulled)
+		{
+			continue;
+		}
+
 		if (UploadedBuffer.bUseDataCallbacks)
 		{
 			UploadedBuffer.Data = UploadedBuffer.DataCallback();
@@ -2042,7 +2045,7 @@ void FRDGBuilder::SetupBufferUploads()
 		if (UploadedBuffer.Data && UploadedBuffer.DataSize)
 		{
 			ConvertToExternalBuffer(UploadedBuffer.Buffer);
-			check(UploadedBuffer.DataSize <= UploadedBuffer.Buffer->Desc.GetTotalNumBytes());
+			check(UploadedBuffer.DataSize <= UploadedBuffer.Buffer->Desc.GetSize());
 		}
 	}
 }
@@ -2053,6 +2056,11 @@ void FRDGBuilder::SubmitBufferUploads()
 
 	for (const FUploadedBuffer& UploadedBuffer : UploadedBuffers)
 	{
+		if (UploadedBuffer.Buffer->bCulled)
+		{
+			continue;
+		}
+
 		if (UploadedBuffer.Data && UploadedBuffer.DataSize)
 		{
 #if PLATFORM_NEEDS_GPU_UAV_RESOURCE_INIT_WORKAROUND
@@ -2305,9 +2313,9 @@ void FRDGBuilder::CreateUniformBuffers()
 
 void FRDGBuilder::AddProloguePass()
 {
-	bInDebugPassScope = true;
+	bSkipAuxiliaryPasses = true;
 	ProloguePass = SetupEmptyPass(Passes.Allocate<FRDGSentinelPass>(Allocator, RDG_EVENT_NAME("Graph Prologue (Graphics)")));
-	bInDebugPassScope = false;
+	bSkipAuxiliaryPasses = false;
 }
 
 ///////////////////////////////////////////////////////////////////////////////////////////////////
@@ -2503,7 +2511,7 @@ void FRDGBuilder::BeginResourcesRHI(FRDGPass* ResourcePass, FRDGPassHandle Execu
 
 		for (FRDGViewHandle ViewHandle : PassToBegin->Views)
 		{
-			BeginResourceRHI(ExecutePassHandle, Views[ViewHandle]);
+			InitRHI(Views[ViewHandle]);
 		}
 	}
 }
@@ -2662,12 +2670,16 @@ void FRDGBuilder::AddEpilogueTransition(FRDGBufferRef Buffer)
 		EpilogueResourceAccesses.Emplace(Buffer->GetRHI(), Buffer->EpilogueAccess);
 	}
 
-	if (!Buffer->bTransient)
+	if (Buffer->Allocation)
+	{
+		ActivePooledBuffers.Emplace(MoveTemp(Buffer->Allocation));
+	}
+	else if (!Buffer->bTransient)
 	{
+		// Non-transient buffers need to be 'resurrected' to hold the last reference in the chain.
+		// Transient buffers don't have that restriction since there's no actual pooling happening.
 		ActivePooledBuffers.Emplace(Buffer->PooledBuffer);
 	}
-
-	Buffer->Allocation = nullptr;
 }
 
 void FRDGBuilder::AddTransition(FRDGPassHandle PassHandle, FRDGTextureRef Texture, const FRDGTextureSubresourceStateIndirect& StateAfter)
@@ -3124,7 +3136,7 @@ void FRDGBuilder::BeginResourceRHI(FRDGPassHandle PassHandle, FRDGTextureRef Tex
 	}
 }
 
-void FRDGBuilder::BeginResourceRHI(FRDGPassHandle PassHandle, FRDGTextureSRVRef SRV)
+void FRDGBuilder::InitRHI(FRDGTextureSRVRef SRV)
 {
 	check(SRV);
 
@@ -3140,7 +3152,7 @@ void FRDGBuilder::BeginResourceRHI(FRDGPassHandle PassHandle, FRDGTextureSRVRef
 	SRV->ResourceRHI = Texture->ViewCache->GetOrCreateSRV(TextureRHI, SRV->Desc);
 }
 
-void FRDGBuilder::BeginResourceRHI(FRDGPassHandle PassHandle, FRDGTextureUAVRef UAV)
+void FRDGBuilder::InitRHI(FRDGTextureUAVRef UAV)
 {
 	check(UAV);
 
@@ -3204,11 +3216,13 @@ void FRDGBuilder::BeginResourceRHI(FRDGPassHandle PassHandle, FRDGBufferRef Buff
 
 	if (!Buffer->bTransient)
 	{
-		SetRHI(Buffer, GRenderGraphResourcePool.FindFreeBuffer(RHICmdList, Buffer->Desc, Buffer->Name), PassHandle);
+		const ERDGPooledBufferAlignment Alignment = Buffer->bQueuedForUpload ? ERDGPooledBufferAlignment::PowerOfTwo : ERDGPooledBufferAlignment::Page;
+
+		SetRHI(Buffer, GRenderGraphResourcePool.FindFreeBuffer(Buffer->Desc, Buffer->Name, Alignment), PassHandle);
 	}
 }
 
-void FRDGBuilder::BeginResourceRHI(FRDGPassHandle PassHandle, FRDGBufferSRVRef SRV)
+void FRDGBuilder::InitRHI(FRDGBufferSRVRef SRV)
 {
 	check(SRV);
 
@@ -3234,7 +3248,7 @@ void FRDGBuilder::BeginResourceRHI(FRDGPassHandle PassHandle, FRDGBufferSRVRef S
 	SRV->ResourceRHI = Buffer->ViewCache->GetOrCreateSRV(BufferRHI, SRVCreateInfo);
 }
 
-void FRDGBuilder::BeginResourceRHI(FRDGPassHandle PassHandle, FRDGBufferUAV* UAV)
+void FRDGBuilder::InitRHI(FRDGBufferUAV* UAV)
 {
 	check(UAV);
 
@@ -3259,7 +3273,7 @@ void FRDGBuilder::BeginResourceRHI(FRDGPassHandle PassHandle, FRDGBufferUAV* UAV
 	UAV->ResourceRHI = Buffer->ViewCache->GetOrCreateUAV(Buffer->GetRHIUnchecked(), UAVCreateInfo);
 }
 
-void FRDGBuilder::BeginResourceRHI(FRDGPassHandle PassHandle, FRDGView* View)
+void FRDGBuilder::InitRHI(FRDGView* View)
 {
 	if (View->HasRHI())
 	{
@@ -3269,16 +3283,16 @@ void FRDGBuilder::BeginResourceRHI(FRDGPassHandle PassHandle, FRDGView* View)
 	switch (View->Type)
 	{
 	case ERDGViewType::TextureUAV:
-		BeginResourceRHI(PassHandle, static_cast<FRDGTextureUAV*>(View));
+		InitRHI(static_cast<FRDGTextureUAV*>(View));
 		break;
 	case ERDGViewType::TextureSRV:
-		BeginResourceRHI(PassHandle, static_cast<FRDGTextureSRV*>(View));
+		InitRHI(static_cast<FRDGTextureSRV*>(View));
 		break;
 	case ERDGViewType::BufferUAV:
-		BeginResourceRHI(PassHandle, static_cast<FRDGBufferUAV*>(View));
+		InitRHI(static_cast<FRDGBufferUAV*>(View));
 		break;
 	case ERDGViewType::BufferSRV:
-		BeginResourceRHI(PassHandle, static_cast<FRDGBufferSRV*>(View));
+		InitRHI(static_cast<FRDGBufferSRV*>(View));
 		break;
 	}
 }
@@ -3352,12 +3366,12 @@ void FRDGBuilder::EndResourceRHI(FRDGPassHandle PassHandle, FRDGBufferRef Buffer
 void FRDGBuilder::VisualizePassOutputs(const FRDGPass* Pass)
 {
 #if SUPPORTS_VISUALIZE_TEXTURE
-	if (bInDebugPassScope)
+	if (bSkipAuxiliaryPasses)
 	{
 		return;
 	}
 
-	bInDebugPassScope = true;
+	bSkipAuxiliaryPasses = true;
 
 
 	Pass->GetParameters().EnumerateTextures([&](FRDGParameter Parameter)
@@ -3425,7 +3439,7 @@ void FRDGBuilder::VisualizePassOutputs(const FRDGPass* Pass)
 		}
 	});
 
-	bInDebugPassScope = false;
+	bSkipAuxiliaryPasses = false;
 #endif
 }
 
@@ -3436,11 +3450,11 @@ void FRDGBuilder::ClobberPassOutputs(const FRDGPass* Pass)
 		return;
 	}
 
-	if (bInDebugPassScope)
+	if (bSkipAuxiliaryPasses)
 	{
 		return;
 	}
-	bInDebugPassScope = true;
+	bSkipAuxiliaryPasses = true;
 
 	RDG_EVENT_SCOPE(*this, "RDG ClobberResources");
 
@@ -3535,7 +3549,7 @@ void FRDGBuilder::ClobberPassOutputs(const FRDGPass* Pass)
 		}
 	});
 
-	bInDebugPassScope = false;
+	bSkipAuxiliaryPasses = false;
 }
 
 #endif //! RDG_ENABLE_DEBUG
diff --git a/Engine/Source/Runtime/RenderCore/Private/RenderGraphResourcePool.cpp b/Engine/Source/Runtime/RenderCore/Private/RenderGraphResourcePool.cpp
--- a/Engine/Source/Runtime/RenderCore/Private/RenderGraphResourcePool.cpp
+++ b/Engine/Source/Runtime/RenderCore/Private/RenderGraphResourcePool.cpp
@@ -1,27 +1,37 @@
 // Copyright Epic Games, Inc. All Rights Reserved.
 
-/*=============================================================================
-	RenderTargetPool.cpp: Scene render target pool manager.
-=============================================================================*/
-
 #include "RenderGraphResourcePool.h"
 #include "RenderGraphResources.h"
+#include "Trace/Trace.inl"
+#include "ProfilingDebugging/CountersTrace.h"
 
-uint64 ComputeHash(const FRDGBufferDesc& Desc)
-{
-	return CityHash64((const char*)&Desc, sizeof(FRDGBufferDesc));
-}
+TRACE_DECLARE_INT_COUNTER(BufferPoolCount, TEXT("BufferPool/BufferCount"));
+TRACE_DECLARE_INT_COUNTER(BufferPoolCreateCount, TEXT("BufferPool/BufferCreateCount"));
+TRACE_DECLARE_INT_COUNTER(BufferPoolReleaseCount, TEXT("BufferPool/BufferReleaseCount"));
+TRACE_DECLARE_MEMORY_COUNTER(BufferPoolSize, TEXT("BufferPool/Size"));
+
+UE_TRACE_EVENT_BEGIN(Cpu, FRDGBufferPool_CreateBuffer, NoSync)
+	UE_TRACE_EVENT_FIELD(UE::Trace::WideString, Name)
+	UE_TRACE_EVENT_FIELD(uint32, SizeInBytes)
+UE_TRACE_EVENT_END()
 
-TRefCountPtr<FRDGPooledBuffer> FRDGBufferPool::FindFreeBuffer(
-	FRHICommandList& RHICmdList,
-	const FRDGBufferDesc& Desc,
-	const TCHAR* InDebugName)
+TRefCountPtr<FRDGPooledBuffer> FRDGBufferPool::FindFreeBuffer(const FRDGBufferDesc& Desc, const TCHAR* InDebugName, ERDGPooledBufferAlignment Alignment)
 {
 	const uint64 BufferPageSize = 64 * 1024;
 
 	FRDGBufferDesc AlignedDesc = Desc;
-	AlignedDesc.NumElements = Align(AlignedDesc.BytesPerElement * AlignedDesc.NumElements, BufferPageSize) / AlignedDesc.BytesPerElement;
-	const uint64 BufferHash = ComputeHash(AlignedDesc);
+
+	switch (Alignment)
+	{
+	case ERDGPooledBufferAlignment::PowerOfTwo:
+		AlignedDesc.NumElements = FMath::RoundUpToPowerOfTwo(AlignedDesc.BytesPerElement * AlignedDesc.NumElements) / AlignedDesc.BytesPerElement;
+		// Fall through to align up to page size for small buffers; helps with reuse.
+
+	case ERDGPooledBufferAlignment::Page:
+		AlignedDesc.NumElements = Align(AlignedDesc.BytesPerElement * AlignedDesc.NumElements, BufferPageSize) / AlignedDesc.BytesPerElement;
+	}
+
+	const uint32 BufferHash = GetTypeHash(AlignedDesc);
 	
 	// First find if available.
 	for (int32 Index = 0; Index < AllocatedBufferHashes.Num(); ++Index)
@@ -57,9 +67,17 @@ TRefCountPtr<FRDGPooledBuffer> FRDGBufferPool::FindFreeBuffer(
 
 	// Allocate new one
 	{
-		TRACE_CPUPROFILER_EVENT_SCOPE(FRDGBufferPool::CreateBuffer);
+		const uint32 NumBytes = AlignedDesc.GetSize();
+
+#if CPUPROFILERTRACE_ENABLED
+		UE_TRACE_LOG_SCOPED_T(Cpu, FRDGBufferPool_CreateBuffer, CpuChannel)
+			<< FRDGBufferPool_CreateBuffer.Name(InDebugName)
+			<< FRDGBufferPool_CreateBuffer.SizeInBytes(NumBytes);
+#endif
 
-		const uint32 NumBytes = AlignedDesc.GetTotalNumBytes();
+		TRACE_COUNTER_ADD(BufferPoolCount, 1);
+		TRACE_COUNTER_ADD(BufferPoolCreateCount, 1);
+		TRACE_COUNTER_ADD(BufferPoolSize, NumBytes);
 
 		FRHIResourceCreateInfo CreateInfo(InDebugName);
 		TRefCountPtr<FRHIBuffer> BufferRHI;
@@ -78,11 +96,13 @@ TRefCountPtr<FRDGPooledBuffer> FRDGBufferPool::FindFreeBuffer(
 			InitialAccess = RHIGetDefaultResourceState(Usage, false);
 			BufferRHI = RHICreateStructuredBuffer(Desc.BytesPerElement, NumBytes, Usage, InitialAccess, CreateInfo);
 		}
+	PRAGMA_DISABLE_DEPRECATION_WARNINGS
 		else if (Desc.UnderlyingType == FRDGBufferDesc::EUnderlyingType::AccelerationStructure)
 		{
 			InitialAccess = ERHIAccess::BVHWrite;
 			BufferRHI = RHICreateBuffer(NumBytes, Desc.Usage, 0, InitialAccess, CreateInfo);
 		}
+	PRAGMA_ENABLE_DEPRECATION_WARNINGS
 		else
 		{
 			check(0);
@@ -114,6 +134,8 @@ void FRDGBufferPool::TickPoolElements()
 	const uint32 kFramesUntilRelease = 30;
 
 	int32 BufferIndex = 0;
+	int32 NumReleasedBuffers = 0;
+	int64 NumReleasedBufferBytes = 0;
 
 	while (BufferIndex < AllocatedBuffers.Num())
 	{
@@ -125,8 +147,12 @@ void FRDGBufferPool::TickPoolElements()
 
 		if (bIsUnused && bNotRequestedRecently)
 		{
+			NumReleasedBufferBytes += Buffer->GetAlignedDesc().GetSize();
+
 			AllocatedBuffers.RemoveAtSwap(BufferIndex);
 			AllocatedBufferHashes.RemoveAtSwap(BufferIndex);
+
+			++NumReleasedBuffers;
 		}
 		else
 		{
@@ -134,6 +160,11 @@ void FRDGBufferPool::TickPoolElements()
 		}
 	}
 
+	TRACE_COUNTER_SUBTRACT(BufferPoolSize, NumReleasedBufferBytes);
+	TRACE_COUNTER_SUBTRACT(BufferPoolCount, NumReleasedBuffers);
+	TRACE_COUNTER_SET(BufferPoolReleaseCount, NumReleasedBuffers);
+	TRACE_COUNTER_SET(BufferPoolCreateCount, 0);
+
 	++FrameCounter;
 }
 
diff --git a/Engine/Source/Runtime/RenderCore/Private/RenderGraphResourcePool.h b/Engine/Source/Runtime/RenderCore/Private/RenderGraphResourcePool.h
--- a/Engine/Source/Runtime/RenderCore/Private/RenderGraphResourcePool.h
+++ b/Engine/Source/Runtime/RenderCore/Private/RenderGraphResourcePool.h
@@ -21,14 +21,20 @@ public:
 	void TickPoolElements();
 
 	/** Allocate a buffer from a given descriptor. */
-	TRefCountPtr<FRDGPooledBuffer> FindFreeBuffer(FRHICommandList& RHICmdList, const FRDGBufferDesc& Desc, const TCHAR* InDebugName);
+	UE_DEPRECATED(5.1, "Use FindFreeBuffer without a command list instead.")
+	TRefCountPtr<FRDGPooledBuffer> FindFreeBuffer(FRHICommandList& RHICmdList, const FRDGBufferDesc& Desc, const TCHAR* InDebugName)
+	{
+		return FindFreeBuffer(Desc, InDebugName);
+	}
+
+	TRefCountPtr<FRDGPooledBuffer> FindFreeBuffer(const FRDGBufferDesc& Desc, const TCHAR* InDebugName, ERDGPooledBufferAlignment Alignment = ERDGPooledBufferAlignment::Page);
 
 private:
 	void ReleaseDynamicRHI() override;
 
 	/** Elements can be 0, we compact the buffer later. */
 	TArray<TRefCountPtr<FRDGPooledBuffer>> AllocatedBuffers;
-	TArray<uint64> AllocatedBufferHashes;
+	TArray<uint32> AllocatedBufferHashes;
 
 	uint32 FrameCounter = 0;
 
diff --git a/Engine/Source/Runtime/RenderCore/Private/RenderGraphResources.cpp b/Engine/Source/Runtime/RenderCore/Private/RenderGraphResources.cpp
--- a/Engine/Source/Runtime/RenderCore/Private/RenderGraphResources.cpp
+++ b/Engine/Source/Runtime/RenderCore/Private/RenderGraphResources.cpp
@@ -32,11 +32,10 @@ FRDGViewableResource::FRDGViewableResource(const TCHAR* InName, const ERDGViewab
 	, TransientExtractionHint(ETransientExtractionHint::None)
 	, bFinalizedAccess(0)
 	, bLastOwner(1)
-	, bCulled(1)
+	, bCulled(!IsImmediateMode())
 	, bUsedByAsyncComputePass(0)
 	, bQueuedForUpload(0)
-	, bSwapChain(0)
-	, bSwapChainAlreadyMoved(0)
+	, FirstBarrier(EFirstBarrier::Split)
 	, bUAVAccessed(0)
 {}
 
diff --git a/Engine/Source/Runtime/RenderCore/Private/RenderGraphUtils.cpp b/Engine/Source/Runtime/RenderCore/Private/RenderGraphUtils.cpp
--- a/Engine/Source/Runtime/RenderCore/Private/RenderGraphUtils.cpp
+++ b/Engine/Source/Runtime/RenderCore/Private/RenderGraphUtils.cpp
@@ -381,13 +381,13 @@ void AddCopyBufferPass(FRDGBuilder& GraphBuilder, FRDGBufferRef DstBuffer, FRDGB
 	const uint64 NumBytes = Parameters->SrcBuffer->Desc.NumElements * Parameters->SrcBuffer->Desc.BytesPerElement;
 
 	GraphBuilder.AddPass(
-		RDG_EVENT_NAME("CopyBuffer(%s Size=%ubytes)", SrcBuffer->Name, SrcBuffer->Desc.GetTotalNumBytes()),
+		RDG_EVENT_NAME("CopyBuffer(%s Size=%ubytes)", SrcBuffer->Name, SrcBuffer->Desc.GetSize()),
 		Parameters,
 		ERDGPassFlags::Copy,
 		[&Parameters, SrcBuffer, DstBuffer, NumBytes](FRHICommandList& RHICmdList)
-		{
-			RHICmdList.CopyBufferRegion(DstBuffer->GetRHI(), 0, SrcBuffer->GetRHI(), 0, NumBytes);
-		});
+	{
+		RHICmdList.CopyBufferRegion(DstBuffer->GetRHI(), 0, SrcBuffer->GetRHI(), 0, NumBytes);
+	});
 }
 
 BEGIN_SHADER_PARAMETER_STRUCT(FClearBufferUAVParameters, )
@@ -402,7 +402,7 @@ void AddClearUAVPass(FRDGBuilder& GraphBuilder, FRDGBufferUAVRef BufferUAV, uint
 	Parameters->BufferUAV = BufferUAV;
 
 	GraphBuilder.AddPass(
-		RDG_EVENT_NAME("ClearBuffer(%s Size=%ubytes)", BufferUAV->GetParent()->Name, BufferUAV->GetParent()->Desc.GetTotalNumBytes()),
+		RDG_EVENT_NAME("ClearBuffer(%s Size=%ubytes)", BufferUAV->GetParent()->Name, BufferUAV->GetParent()->Desc.GetSize()),
 		Parameters,
 		ERDGPassFlags::Compute,
 		[&Parameters, BufferUAV, Value](FRHIComputeCommandList& RHICmdList)
@@ -418,14 +418,14 @@ void AddClearUAVFloatPass(FRDGBuilder& GraphBuilder, FRDGBufferUAVRef BufferUAV,
 	Parameters->BufferUAV = BufferUAV;
 
 	GraphBuilder.AddPass(
-		RDG_EVENT_NAME("ClearBuffer(%s Size=%ubytes)", BufferUAV->GetParent()->Name, BufferUAV->GetParent()->Desc.GetTotalNumBytes()),
+		RDG_EVENT_NAME("ClearBuffer(%s Size=%ubytes)", BufferUAV->GetParent()->Name, BufferUAV->GetParent()->Desc.GetSize()),
 		Parameters,
 		ERDGPassFlags::Compute,
 		[&Parameters, BufferUAV, Value](FRHIComputeCommandList& RHICmdList)
-		{
-			RHICmdList.ClearUAVFloat(BufferUAV->GetRHI(), FVector4f(Value, Value, Value, Value));
-			BufferUAV->MarkResourceAsUsed();
-		});
+	{
+		RHICmdList.ClearUAVFloat(BufferUAV->GetRHI(), FVector4f(Value, Value, Value, Value));
+		BufferUAV->MarkResourceAsUsed();
+	});
 }
 
 BEGIN_SHADER_PARAMETER_STRUCT(FClearTextureUAVParameters, )
@@ -907,14 +907,12 @@ void FComputeShaderUtils::ClearUAV(FRDGBuilder& GraphBuilder, FGlobalShaderMap*
 		FIntVector(FMath::DivideAndRoundUp<int32>(PassParameters->NumEntries, 64), 1, 1));
 }
 
-
-
 class FInitIndirectArgs1DCS : public FGlobalShader
 {
 	DECLARE_GLOBAL_SHADER(FInitIndirectArgs1DCS);
 	SHADER_USE_PARAMETER_STRUCT(FInitIndirectArgs1DCS, FGlobalShader)
 
-		BEGIN_SHADER_PARAMETER_STRUCT(FParameters, )
+	BEGIN_SHADER_PARAMETER_STRUCT(FParameters, )
 		SHADER_PARAMETER_RDG_BUFFER_SRV(StructuredBuffer< uint >, InputCountBuffer)
 		SHADER_PARAMETER(uint32, Multiplier)
 		SHADER_PARAMETER(uint32, Divisor)
@@ -923,8 +921,8 @@ class FInitIndirectArgs1DCS : public FGlobalShader
 		SHADER_PARAMETER_RDG_BUFFER_UAV(RWBuffer< uint >, IndirectDispatchArgsOut)
 	END_SHADER_PARAMETER_STRUCT()
 };
-IMPLEMENT_GLOBAL_SHADER(FInitIndirectArgs1DCS, "/Engine/Private/Tools/SetupIndirectArgs.usf", "InitIndirectArgs1DCS", SF_Compute);
 
+IMPLEMENT_GLOBAL_SHADER(FInitIndirectArgs1DCS, "/Engine/Private/Tools/SetupIndirectArgs.usf", "InitIndirectArgs1DCS", SF_Compute);
 
 FRDGBufferRef FComputeShaderUtils::AddIndirectArgsSetupCsPass1D(FRDGBuilder& GraphBuilder, FRDGBufferRef& InputCountBuffer, const TCHAR* OutputBufferName, uint32 Divisor, uint32 InputCountOffset, uint32 Multiplier)
 {
@@ -951,7 +949,6 @@ FRDGBufferRef FComputeShaderUtils::AddIndirectArgsSetupCsPass1D(FRDGBuilder& Gra
 	return IndirectArgsBuffer;
 }
 
-
 FRDGBufferRef CreateStructuredBuffer(
 	FRDGBuilder& GraphBuilder,
 	const TCHAR* Name,
@@ -1033,11 +1030,11 @@ FRDGWaitForTasksScope::~FRDGWaitForTasksScope()
 	}
 }
 
-bool GetPooledFreeBuffer(
-	FRHICommandList& RHICmdList,
+bool AllocatePooledBuffer(
 	const FRDGBufferDesc& Desc,
 	TRefCountPtr<FRDGPooledBuffer>& Out,
-	const TCHAR* InDebugName)
+	const TCHAR* Name,
+	ERDGPooledBufferAlignment Alignment)
 {
 	if (Out && Out->Desc == Desc)
 	{
@@ -1046,10 +1043,15 @@ bool GetPooledFreeBuffer(
 	}
 
 	// New allocation.
-	Out = GRenderGraphResourcePool.FindFreeBuffer(RHICmdList, Desc, InDebugName);
+	Out = GRenderGraphResourcePool.FindFreeBuffer(Desc, Name, Alignment);
 	return true;
 }
 
+TRefCountPtr<FRDGPooledBuffer> AllocatePooledBuffer(const FRDGBufferDesc& Desc, const TCHAR* Name, ERDGPooledBufferAlignment Alignment)
+{
+	return GRenderGraphResourcePool.FindFreeBuffer(Desc, Name, Alignment);
+}
+
 bool AllocatePooledTexture(const FRDGTextureDesc& Desc, TRefCountPtr<IPooledRenderTarget>& Out, const TCHAR* Name)
 {
 	return GRenderTargetPool.FindFreeElement(Desc, Out, Name);
diff --git a/Engine/Source/Runtime/RenderCore/Private/RenderTargetPool.cpp b/Engine/Source/Runtime/RenderCore/Private/RenderTargetPool.cpp
--- a/Engine/Source/Runtime/RenderCore/Private/RenderTargetPool.cpp
+++ b/Engine/Source/Runtime/RenderCore/Private/RenderTargetPool.cpp
@@ -4,6 +4,8 @@
 #include "RHIStaticStates.h"
 #include "Misc/OutputDeviceRedirector.h"
 #include "Hash/CityHash.h"
+#include "Trace/Trace.inl"
+#include "ProfilingDebugging/CountersTrace.h"
 
 /** The global render targets pool. */
 TGlobalResource<FRenderTargetPool> GRenderTargetPool;
@@ -12,6 +14,12 @@ DEFINE_LOG_CATEGORY_STATIC(LogRenderTargetPool, Warning, All);
 
 CSV_DEFINE_CATEGORY(RenderTargetPool, !UE_SERVER);
 
+TRACE_DECLARE_INT_COUNTER(RenderTargetPoolCount, TEXT("RenderTargetPool/Count"));
+TRACE_DECLARE_MEMORY_COUNTER(RenderTargetPoolSize, TEXT("RenderTargetPool/Size"));
+
+UE_TRACE_EVENT_BEGIN(Cpu, FRenderTargetPool_CreateTexture, NoSync)
+	UE_TRACE_EVENT_FIELD(UE::Trace::WideString, Name)
+UE_TRACE_EVENT_END()
 
 TRefCountPtr<IPooledRenderTarget> CreateRenderTarget(FRHITexture* Texture, const TCHAR* Name)
 {
@@ -94,7 +102,10 @@ TRefCountPtr<IPooledRenderTarget> FRenderTargetPool::FindFreeElement(FRHITexture
 
 	if (!Found)
 	{
-		TRACE_CPUPROFILER_EVENT_SCOPE(FRenderTargetPool::CreateTexture);
+#if CPUPROFILERTRACE_ENABLED
+		UE_TRACE_LOG_SCOPED_T(Cpu, FRenderTargetPool_CreateTexture, CpuChannel)
+			<< FRenderTargetPool_CreateTexture.Name(Name);
+#endif
 
 		const ERHIAccess AccessInitial = ERHIAccess::SRVMask;
 		FRHITextureCreateDesc CreateDesc(Desc, AccessInitial, Name);
@@ -128,6 +139,8 @@ TRefCountPtr<IPooledRenderTarget> FRenderTargetPool::FindFreeElement(FRHITexture
 		}
 
 		AllocationLevelInKB += ComputeSizeInKB(*Found);
+		TRACE_COUNTER_ADD(RenderTargetPoolCount, 1);
+		TRACE_COUNTER_SET(RenderTargetPoolSize, (int64)AllocationLevelInKB * 1024);
 
 		FoundIndex = PooledRenderTargets.Num() - 1;
 		Found->Desc.DebugName = Name;
@@ -182,6 +195,8 @@ bool FRenderTargetPool::FindFreeElement(const FRHITextureCreateInfo& Desc, TRefC
 			if (Current->IsFree())
 			{
 				AllocationLevelInKB -= ComputeSizeInKB(*Current);
+				TRACE_COUNTER_SUBTRACT(RenderTargetPoolCount, 1);
+				TRACE_COUNTER_SET(RenderTargetPoolSize, (int64)AllocationLevelInKB * 1024);
 				int32 Index = FindIndex(Current);
 				check(Index >= 0);
 				FreeElementAtIndex(Index);
@@ -297,6 +312,8 @@ void FRenderTargetPool::TickPoolElements()
 		if (OldestElementIndex != -1)
 		{
 			AllocationLevelInKB -= ComputeSizeInKB(*PooledRenderTargets[OldestElementIndex]);
+			TRACE_COUNTER_SUBTRACT(RenderTargetPoolCount, 1);
+			TRACE_COUNTER_SET(RenderTargetPoolSize, (int64)AllocationLevelInKB * 1024);
 
 			// we assume because of reference counting the resource gets released when not needed any more
 			// we don't use Remove() to not shuffle around the elements for better transparency on RenderTargetPoolEvents
@@ -386,6 +403,9 @@ void FRenderTargetPool::FreeUnusedResource(TRefCountPtr<IPooledRenderTarget>& In
 		if (Element->IsFree())
 		{
 			AllocationLevelInKB -= ComputeSizeInKB(*Element);
+			TRACE_COUNTER_SUBTRACT(RenderTargetPoolCount, 1);
+			TRACE_COUNTER_SET(RenderTargetPoolSize, (int64)AllocationLevelInKB * 1024);
+
 			// we assume because of reference counting the resource gets released when not needed any more
 			DeferredDeleteArray.Add(PooledRenderTargets[Index]);
 			FreeElementAtIndex(Index);
@@ -404,6 +424,9 @@ void FRenderTargetPool::FreeUnusedResources()
 		if (Element && Element->IsFree())
 		{
 			AllocationLevelInKB -= ComputeSizeInKB(*Element);
+			TRACE_COUNTER_SUBTRACT(RenderTargetPoolCount, 1);
+			TRACE_COUNTER_SET(RenderTargetPoolSize, (int64)AllocationLevelInKB * 1024);
+
 			// we assume because of reference counting the resource gets released when not needed any more
 			// we don't use Remove() to not shuffle around the elements for better transparency on RenderTargetPoolEvents
 			DeferredDeleteArray.Add(PooledRenderTargets[i]);
diff --git a/Engine/Source/Runtime/RenderCore/Private/RenderUtils.cpp b/Engine/Source/Runtime/RenderCore/Private/RenderUtils.cpp
--- a/Engine/Source/Runtime/RenderCore/Private/RenderUtils.cpp
+++ b/Engine/Source/Runtime/RenderCore/Private/RenderUtils.cpp
@@ -247,14 +247,11 @@ public:
 	{
 		if (!Buffer.IsValid())
 		{
-			FRHICommandList* UnusedCmdList = new FRHICommandList(FRHIGPUMask::All());
-			GetPooledFreeBuffer(*UnusedCmdList, FRDGBufferDesc::CreateBufferDesc(sizeof(FVector4f), 1), Buffer, TEXT("WhiteVertexBufferWithRDG"));
+			Buffer = AllocatePooledBuffer(FRDGBufferDesc::CreateBufferDesc(sizeof(FVector4f), 1), TEXT("WhiteVertexBufferWithRDG"));
 
 			FVector4f* BufferData = (FVector4f*)RHILockBuffer(Buffer->GetRHI(), 0, sizeof(FVector4f), RLM_WriteOnly);
 			*BufferData = FVector4f(1.0f, 1.0f, 1.0f, 1.0f);
 			RHIUnlockBuffer(Buffer->GetRHI());
-			delete UnusedCmdList;
-			UnusedCmdList = nullptr;
 		}
 	}
 };
diff --git a/Engine/Source/Runtime/RenderCore/Public/RenderGraphBuilder.h b/Engine/Source/Runtime/RenderCore/Public/RenderGraphBuilder.h
--- a/Engine/Source/Runtime/RenderCore/Public/RenderGraphBuilder.h
+++ b/Engine/Source/Runtime/RenderCore/Public/RenderGraphBuilder.h
@@ -26,8 +26,12 @@ public:
 	~FRDGBuilder();
 
 	/** Finds an RDG texture associated with the external texture, or returns null if none is found. */
-	FRDGTextureRef FindExternalTexture(FRHITexture* Texture) const;
-	FRDGTextureRef FindExternalTexture(IPooledRenderTarget* ExternalPooledTexture) const;
+	FRDGTexture* FindExternalTexture(FRHITexture* Texture) const;
+	FRDGTexture* FindExternalTexture(IPooledRenderTarget* ExternalPooledTexture) const;
+
+	/** Finds an RDG buffer associated with the external buffer, or returns null if none is found. */
+	FRDGBuffer* FindExternalBuffer(FRHIBuffer* Buffer) const;
+	FRDGBuffer* FindExternalBuffer(FRDGPooledBuffer* ExternalPooledBuffer) const;
 
 	/** Registers a external pooled render target texture to be tracked by the render graph. The name of the registered RDG texture is pulled from the pooled render target. */
 	FRDGTextureRef RegisterExternalTexture(
@@ -107,7 +111,7 @@ public:
 	// Allocation Methods
 
 	/** Allocates raw memory using an allocator tied to the lifetime of the graph. */
-	void* Alloc(uint64 SizeInBytes, uint32 AlignInBytes);
+	void* Alloc(uint64 SizeInBytes, uint32 AlignInBytes = 16);
 
 	/** Allocates POD memory using an allocator tied to the lifetime of the graph. Does not construct / destruct. */
 	template <typename PODType>
@@ -469,11 +473,11 @@ private:
 	FRDGUniformBufferRegistry UniformBuffers;
 
 	/** Uniform buffers which were used in a pass. */
-	TArray<FRDGUniformBufferHandle, TInlineAllocator<32, FRDGArrayAllocator>> UniformBuffersToCreate;
+	TArray<FRDGUniformBufferHandle, FRDGArrayAllocator> UniformBuffersToCreate;
 
 	/** Tracks external resources to their registered render graph counterparts for de-duplication. */
 	TSortedMap<FRHITexture*, FRDGTexture*, FRDGArrayAllocator> ExternalTextures;
-	TSortedMap<FRDGPooledBuffer*, FRDGBuffer*, FRDGArrayAllocator> ExternalBuffers;
+	TSortedMap<FRHIBuffer*, FRDGBuffer*, FRDGArrayAllocator> ExternalBuffers;
 
 	/** Tracks the latest RDG resource to own an alias of a pooled resource (multiple RDG resources can reference the same pooled resource). */
 	TMap<FRDGPooledTexture*, FRDGTexture*, FRDGSetAllocator> PooledTextureOwnershipMap;
@@ -596,6 +600,7 @@ private:
 	/** Array of all active parallel execute tasks. */
 	FGraphEventArray ParallelExecuteEvents;
 
+	/** Tracks the final access used on resources in order to call SetTrackedAccess. */
 	TArray<FRHITrackedAccessInfo, FRDGArrayAllocator> EpilogueResourceAccesses;
 
 	/** Texture state used for intermediate operations. Held here to avoid re-allocating. */
@@ -621,7 +626,7 @@ private:
 #endif
 
 	/** Tracks whether we are in a scope of adding passes to the builder. Used to avoid recursion. */
-	bool bInDebugPassScope = false;
+	bool bSkipAuxiliaryPasses = false;
 
 #if WITH_MGPU
 	/** Name for the temporal effect used to synchronize multi-frame resources. */
@@ -653,17 +658,18 @@ private:
 
 	void BeginResourcesRHI(FRDGPass* ResourcePass, FRDGPassHandle ExecutePassHandle);
 	void BeginResourceRHI(FRDGPassHandle, FRDGTexture* Texture);
-	void BeginResourceRHI(FRDGPassHandle, FRDGTextureSRV* SRV);
-	void BeginResourceRHI(FRDGPassHandle, FRDGTextureUAV* UAV);
 	void BeginResourceRHI(FRDGPassHandle, FRDGBuffer* Buffer);
-	void BeginResourceRHI(FRDGPassHandle, FRDGBufferSRV* SRV);
-	void BeginResourceRHI(FRDGPassHandle, FRDGBufferUAV* UAV);
-	void BeginResourceRHI(FRDGPassHandle, FRDGView* View);
 
 	void EndResourcesRHI(FRDGPass* ResourcePass, FRDGPassHandle ExecutePassHandle);
 	void EndResourceRHI(FRDGPassHandle, FRDGTexture* Texture, uint32 ReferenceCount);
 	void EndResourceRHI(FRDGPassHandle, FRDGBuffer* Buffer, uint32 ReferenceCount);
 
+	void InitRHI(FRDGView* View);
+	void InitRHI(FRDGBufferSRV* SRV);
+	void InitRHI(FRDGBufferUAV* UAV);
+	void InitRHI(FRDGTextureSRV* SRV);
+	void InitRHI(FRDGTextureUAV* UAV);
+
 	void SetupParallelExecute();
 	void DispatchParallelExecute(IRHICommandContext* RHICmdContext);
 
diff --git a/Engine/Source/Runtime/RenderCore/Public/RenderGraphBuilder.inl b/Engine/Source/Runtime/RenderCore/Public/RenderGraphBuilder.inl
--- a/Engine/Source/Runtime/RenderCore/Public/RenderGraphBuilder.inl
+++ b/Engine/Source/Runtime/RenderCore/Public/RenderGraphBuilder.inl
@@ -2,16 +2,16 @@
 
 #pragma once
 
-inline FRDGTextureRef FRDGBuilder::FindExternalTexture(FRHITexture* ExternalTexture) const
+inline FRDGTexture* FRDGBuilder::FindExternalTexture(FRHITexture* ExternalTexture) const
 {
-	if (const FRDGTextureRef* FoundTexturePtr = ExternalTextures.Find(ExternalTexture))
+	if (FRDGTexture* const* FoundTexturePtr = ExternalTextures.Find(ExternalTexture))
 	{
 		return *FoundTexturePtr;
 	}
 	return nullptr;
 }
 
-inline FRDGTextureRef FRDGBuilder::FindExternalTexture(IPooledRenderTarget* ExternalTexture) const
+inline FRDGTexture* FRDGBuilder::FindExternalTexture(IPooledRenderTarget* ExternalTexture) const
 {
 	if (ExternalTexture)
 	{
@@ -20,6 +20,24 @@ inline FRDGTextureRef FRDGBuilder::FindExternalTexture(IPooledRenderTarget* Exte
 	return nullptr;
 }
 
+inline FRDGBuffer* FRDGBuilder::FindExternalBuffer(FRHIBuffer* ExternalBuffer) const
+{
+	if (FRDGBuffer* const* FoundBufferPtr = ExternalBuffers.Find(ExternalBuffer))
+	{
+		return *FoundBufferPtr;
+	}
+	return nullptr;
+}
+
+inline FRDGBuffer* FRDGBuilder::FindExternalBuffer(FRDGPooledBuffer* ExternalBuffer) const
+{
+	if (ExternalBuffer)
+	{
+		return FindExternalBuffer(ExternalBuffer->GetRHI());
+	}
+	return nullptr;
+}
+
 inline FRDGTextureRef FRDGBuilder::CreateTexture(
 	const FRDGTextureDesc& Desc,
 	const TCHAR* Name,
@@ -234,12 +252,7 @@ inline void FRDGBuilder::QueueBufferUpload(FRDGBufferRef Buffer, const void* Ini
 {
 	IF_RDG_ENABLE_DEBUG(UserValidation.ValidateUploadBuffer(Buffer, InitialData, InitialDataSize));
 
-	if (InitialDataSize == 0)
-	{
-		return;
-	}
-
-	if (!EnumHasAnyFlags(InitialDataFlags, ERDGInitialDataFlags::NoCopy))
+	if (InitialDataSize > 0 && !EnumHasAnyFlags(InitialDataFlags, ERDGInitialDataFlags::NoCopy))
 	{
 		void* InitialDataCopy = Alloc(InitialDataSize, 16);
 		FMemory::Memcpy(InitialDataCopy, InitialData, InitialDataSize);
diff --git a/Engine/Source/Runtime/RenderCore/Public/RenderGraphDefinitions.h b/Engine/Source/Runtime/RenderCore/Public/RenderGraphDefinitions.h
--- a/Engine/Source/Runtime/RenderCore/Public/RenderGraphDefinitions.h
+++ b/Engine/Source/Runtime/RenderCore/Public/RenderGraphDefinitions.h
@@ -191,6 +191,21 @@ enum class ERDGViewType : uint8
 	MAX
 };
 
+inline ERDGViewableResourceType GetParentType(ERDGViewType ViewType)
+{
+	switch (ViewType)
+	{
+	case ERDGViewType::TextureUAV:
+	case ERDGViewType::TextureSRV:
+		return ERDGViewableResourceType::Texture;
+	case ERDGViewType::BufferUAV:
+	case ERDGViewType::BufferSRV:
+		return ERDGViewableResourceType::Buffer;
+	}
+	checkNoEntry();
+	return ERDGViewableResourceType::MAX;
+}
+
 enum class ERDGResourceExtractionFlags : uint8
 {
 	None = 0,
@@ -218,6 +233,18 @@ enum class ERDGInitialDataFlags : uint8
 };
 ENUM_CLASS_FLAGS(ERDGInitialDataFlags)
 
+enum class ERDGPooledBufferAlignment : uint8
+{
+	// The buffer size is not aligned.
+	None,
+
+	// The buffer size is aligned up to the next page size.
+	Page,
+
+	// The buffer size is aligned up to the next power of two.
+	PowerOfTwo
+};
+
 /** Returns the equivalent parent resource type for a view type. */
 inline ERDGViewableResourceType GetViewableResourceType(ERDGViewType ViewType)
 {
@@ -707,6 +734,8 @@ using FRDGTextureHandle = TRDGHandle<FRDGTexture, uint16>;
 using FRDGTextureRegistry = TRDGHandleRegistry<FRDGTextureHandle, ERDGHandleRegistryDestructPolicy::Never>;
 using FRDGTextureBitArray = TRDGHandleBitArray<FRDGTextureHandle>;
 
+struct FRDGBufferDesc;
+
 class FRDGBuffer;
 using FRDGBufferRef = FRDGBuffer*;
 using FRDGBufferHandle = TRDGHandle<FRDGBuffer, uint16>;
@@ -736,4 +765,4 @@ template <typename ArrayType,
 	typename ArrayTypeNoRef = std::remove_reference_t<ArrayType>,
 	typename = typename TEnableIf<TIsTArray<ArrayTypeNoRef>::Value>::Type> using TRDGBufferArrayCallback = TFunction<const ArrayType&()>;
 using FRDGBufferInitialDataFreeCallback = TFunction<void(const void* InData)>;
-using FRDGDispatchGroupCountCallback = TFunction<FIntVector()>;
\ No newline at end of file
+using FRDGDispatchGroupCountCallback = TFunction<FIntVector()>;
diff --git a/Engine/Source/Runtime/RenderCore/Public/RenderGraphResources.h b/Engine/Source/Runtime/RenderCore/Public/RenderGraphResources.h
--- a/Engine/Source/Runtime/RenderCore/Public/RenderGraphResources.h
+++ b/Engine/Source/Runtime/RenderCore/Public/RenderGraphResources.h
@@ -269,13 +269,20 @@ public:
 protected:
 	FRDGViewableResource(const TCHAR* InName, ERDGViewableResourceType InType);
 
-	enum class ETransientExtractionHint
+	enum class ETransientExtractionHint : uint8
 	{
 		None,
 		Disable,
 		Enable
 	};
 
+	enum class EFirstBarrier : uint8
+	{
+		Split,
+		ImmediateRequested,
+		ImmediateConfirmed
+	};
+
 	/** Whether this is an externally registered resource. */
 	uint8 bExternal : 1;
 
@@ -309,11 +316,8 @@ protected:
 	/** If true, the resource has been queued for an upload operation. */
 	uint8 bQueuedForUpload : 1;
 
-	/** If true, this resource is a swap chain texture. */
-	uint8 bSwapChain : 1;
-
-	/** If true, the swap chain transition has been moved to occur at the latest possible time. */
-	uint8 bSwapChainAlreadyMoved : 1;
+	/** If true, this resource should skip the prologue split barrier and perform transition right away. */
+	EFirstBarrier FirstBarrier : 2;
 
 	/** If true, the resource is access through at least one UAV. */
 	uint8 bUAVAccessed : 1;
@@ -371,6 +375,11 @@ public:
 	/** Returns the referenced parent render graph resource. */
 	virtual FRDGViewableResource* GetParent() const = 0;
 
+	ERDGViewableResourceType GetParentType() const
+	{
+		return ::GetParentType(Type);
+	}
+
 	FRDGViewHandle GetHandle() const
 	{
 		return Handle;
@@ -433,6 +442,8 @@ class RENDERCORE_API FRDGTexture final
 	: public FRDGViewableResource
 {
 public:
+	static const ERDGViewableResourceType StaticType = ERDGViewableResourceType::Texture;
+
 	const FRDGTextureDesc Desc;
 	const ERDGTextureFlags Flags;
 
@@ -491,7 +502,11 @@ private:
 		MergeState.SetNum(SubresourceCount);
 		LastProducers.Reserve(SubresourceCount);
 		LastProducers.SetNum(SubresourceCount);
-		bSwapChain = EnumHasAnyFlags(Desc.Flags, ETextureCreateFlags::Presentable);
+
+		if (EnumHasAnyFlags(Desc.Flags, ETextureCreateFlags::Presentable))
+		{
+			FirstBarrier = EFirstBarrier::ImmediateRequested;
+		}
 	}
 
 	/** Returns RHI texture without access checks. */
@@ -616,15 +631,16 @@ class FRDGTextureSRVDesc final
 public:
 	FRDGTextureSRVDesc() = default;
 
-	FRDGTextureRef Texture = nullptr;
+	FRDGTextureSRVDesc(FRDGTexture* InTexture)
+	{
+		Texture = InTexture;
+		NumMipLevels = InTexture->Desc.NumMips;
+	}
 
 	/** Create SRV that access all sub resources of texture. */
 	static FRDGTextureSRVDesc Create(FRDGTextureRef Texture)
 	{
-		FRDGTextureSRVDesc Desc;
-		Desc.Texture = Texture;
-		Desc.NumMipLevels = Texture->Desc.NumMips;
-		return Desc;
+		return FRDGTextureSRVDesc(Texture);
 	}
 
 	/** Create SRV that access one specific mip level. */
@@ -653,6 +669,23 @@ public:
 		Desc.MetaData = MetaData;
 		return Desc;
 	}
+
+	bool operator == (const FRDGTextureSRVDesc& Other) const
+	{
+		return Texture == Other.Texture && FRHITextureSRVCreateInfo::operator==(Other);
+	}
+
+	bool operator != (const FRDGTextureSRVDesc& Other) const
+	{
+		return !(*this == Other);
+	}
+
+	friend uint32 GetTypeHash(const FRDGTextureSRVDesc& Desc)
+	{
+		return HashCombine(GetTypeHash(static_cast<const FRHITextureSRVCreateInfo&>(Desc)), GetTypeHash(Desc.Texture));
+	}
+
+	FRDGTextureRef Texture = nullptr;
 };
 
 /** Render graph tracked SRV. */
@@ -660,6 +693,8 @@ class FRDGTextureSRV final
 	: public FRDGShaderResourceView
 {
 public:
+	static const ERDGViewType StaticType = ERDGViewType::TextureSRV;
+
 	/** Descriptor of the graph tracked SRV. */
 	const FRDGTextureSRVDesc Desc;
 
@@ -701,6 +736,21 @@ public:
 		return Desc;
 	}
 
+	bool operator == (const FRDGTextureUAVDesc& Other) const
+	{
+		return Texture == Other.Texture && FRHITextureUAVCreateInfo::operator==(Other);
+	}
+
+	bool operator != (const FRDGTextureUAVDesc& Other) const
+	{
+		return !(*this == Other);
+	}
+
+	friend uint32 GetTypeHash(const FRDGTextureUAVDesc& Desc)
+	{
+		return HashCombine(GetTypeHash(static_cast<const FRHITextureUAVCreateInfo&>(Desc)), GetTypeHash(Desc.Texture));
+	}
+
 	FRDGTextureRef Texture = nullptr;
 };
 
@@ -709,6 +759,8 @@ class FRDGTextureUAV final
 	: public FRDGUnorderedAccessView
 {
 public:
+	static const ERDGViewType StaticType = ERDGViewType::TextureUAV;
+
 	/** Descriptor of the graph tracked UAV. */
 	const FRDGTextureUAVDesc Desc;
 
@@ -745,7 +797,7 @@ struct FRDGBufferDesc
 	{
 		VertexBuffer,
 		StructuredBuffer,
-		AccelerationStructure
+		AccelerationStructure UE_DEPRECATED(5.1, "AccelerationStructure is deprecated")
 	};
 
 	/** Create the descriptor for an indirect RHI call.
@@ -760,7 +812,7 @@ struct FRDGBufferDesc
 	{
 		FRDGBufferDesc Desc;
 		Desc.UnderlyingType = EUnderlyingType::VertexBuffer;
-		Desc.Usage = (EBufferUsageFlags)(BUF_Static | BUF_DrawIndirect | BUF_UnorderedAccess | BUF_ShaderResource);
+		Desc.Usage = EBufferUsageFlags::Static | EBufferUsageFlags::DrawIndirect | EBufferUsageFlags::UnorderedAccess | EBufferUsageFlags::ShaderResource;
 		Desc.BytesPerElement = sizeof(IndirectParameterStruct);
 		Desc.NumElements = NumElements;
 		return Desc;
@@ -770,7 +822,7 @@ struct FRDGBufferDesc
 	{
 		FRDGBufferDesc Desc;
 		Desc.UnderlyingType = EUnderlyingType::VertexBuffer;
-		Desc.Usage = (EBufferUsageFlags)(BUF_Static | BUF_DrawIndirect | BUF_UnorderedAccess | BUF_ShaderResource);
+		Desc.Usage = EBufferUsageFlags::Static | EBufferUsageFlags::DrawIndirect | EBufferUsageFlags::UnorderedAccess | EBufferUsageFlags::ShaderResource;
 		Desc.BytesPerElement = 4;
 		Desc.NumElements = NumElements;
 		return Desc;
@@ -780,7 +832,7 @@ struct FRDGBufferDesc
 	{
 		FRDGBufferDesc Desc;
 		Desc.UnderlyingType = EUnderlyingType::StructuredBuffer;
-		Desc.Usage = (EBufferUsageFlags)(BUF_Static | BUF_UnorderedAccess | BUF_ShaderResource);
+		Desc.Usage = EBufferUsageFlags::Static | EBufferUsageFlags::UnorderedAccess | EBufferUsageFlags::ShaderResource;
 		Desc.BytesPerElement = BytesPerElement;
 		Desc.NumElements = NumElements;
 		return Desc;
@@ -798,7 +850,7 @@ struct FRDGBufferDesc
 	{
 		FRDGBufferDesc Desc;
 		Desc.UnderlyingType = EUnderlyingType::VertexBuffer;
-		Desc.Usage = (EBufferUsageFlags)(BUF_Static | BUF_UnorderedAccess | BUF_ShaderResource);
+		Desc.Usage = EBufferUsageFlags::Static | EBufferUsageFlags::UnorderedAccess | EBufferUsageFlags::ShaderResource;
 		Desc.BytesPerElement = BytesPerElement;
 		Desc.NumElements = NumElements;
 		return Desc;
@@ -817,7 +869,7 @@ struct FRDGBufferDesc
 		check(NumBytes % 4 == 0);
 		FRDGBufferDesc Desc;
 		Desc.UnderlyingType = EUnderlyingType::StructuredBuffer;
-		Desc.Usage = (EBufferUsageFlags)(BUF_UnorderedAccess | BUF_ShaderResource | BUF_ByteAddressBuffer);
+		Desc.Usage = EBufferUsageFlags::Static | EBufferUsageFlags::UnorderedAccess | EBufferUsageFlags::ShaderResource | EBufferUsageFlags::ByteAddressBuffer;
 		Desc.BytesPerElement = 4;
 		Desc.NumElements = NumBytes / 4;
 		return Desc;
@@ -835,7 +887,7 @@ struct FRDGBufferDesc
 	{
 		FRDGBufferDesc Desc;
 		Desc.UnderlyingType = EUnderlyingType::VertexBuffer;
-		Desc.Usage = (EBufferUsageFlags)(BUF_Static | BUF_ShaderResource);
+		Desc.Usage = EBufferUsageFlags::ShaderResource;
 		Desc.BytesPerElement = BytesPerElement;
 		Desc.NumElements = NumElements;
 		return Desc;
@@ -849,12 +901,65 @@ struct FRDGBufferDesc
 		return Desc;
 	}
 
+	static FRDGBufferDesc CreateStructuredUploadDesc(uint32 BytesPerElement, uint32 NumElements)
+	{
+		FRDGBufferDesc Desc;
+		Desc.UnderlyingType = EUnderlyingType::StructuredBuffer;
+		Desc.Usage = EBufferUsageFlags::ShaderResource;
+		Desc.BytesPerElement = BytesPerElement;
+		Desc.NumElements = NumElements;
+		return Desc;
+	}
+
+	template<typename ParameterStruct>
+	static FRDGBufferDesc CreateStructuredUploadDesc(uint32 NumElements)
+	{
+		FRDGBufferDesc Desc = CreateStructuredUploadDesc(sizeof(ParameterStruct), NumElements);
+		Desc.Metadata = ParameterStruct::FTypeInfo::GetStructMetadata();
+		return Desc;
+	}
+
+	static FRDGBufferDesc CreateByteAddressUploadDesc(uint32 NumBytes)
+	{
+		check(NumBytes % 4 == 0);
+		FRDGBufferDesc Desc;
+		Desc.UnderlyingType = EUnderlyingType::StructuredBuffer;
+		Desc.Usage = EBufferUsageFlags::ShaderResource | EBufferUsageFlags::ByteAddressBuffer;
+		Desc.BytesPerElement = 4;
+		Desc.NumElements = NumBytes / 4;
+		return Desc;
+	}
+
+	template<typename ParameterStruct>
+	static FRDGBufferDesc CreateByteAddressUploadDesc(uint32 NumElements)
+	{
+		FRDGBufferDesc Desc = CreateByteAddressUploadDesc(sizeof(ParameterStruct) * NumElements);
+		Desc.Metadata = ParameterStruct::FTypeInfo::GetStructMetadata();
+		return Desc;
+	}
+
 	/** Returns the total number of bytes allocated for a such buffer. */
+	uint32 GetSize() const
+	{
+		return BytesPerElement * NumElements;
+	}
+
+	//UE_DEPRECATED(5.1, "GetTotalNumBytes is deprecated, use GetSize instead.")
 	uint32 GetTotalNumBytes() const
 	{
 		return BytesPerElement * NumElements;
 	}
 
+	friend uint32 GetTypeHash(const FRDGBufferDesc& Desc)
+	{
+		uint32 Hash = GetTypeHash(Desc.BytesPerElement);
+		Hash = HashCombine(Hash, GetTypeHash(Desc.NumElements));
+		Hash = HashCombine(Hash, GetTypeHash(Desc.Usage));
+		Hash = HashCombine(Hash, GetTypeHash(Desc.UnderlyingType));
+		Hash = HashCombine(Hash, GetTypeHash(Desc.Metadata));
+		return Hash;
+	}
+
 	bool operator == (const FRDGBufferDesc& Other) const
 	{
 		return (
@@ -876,7 +981,7 @@ struct FRDGBufferDesc
 	uint32 NumElements = 1;
 
 	/** Bitfields describing the uses of that buffer. */
-	EBufferUsageFlags Usage = BUF_None;
+	EBufferUsageFlags Usage = EBufferUsageFlags::None;
 
 	/** The underlying RHI type to use. */
 	EUnderlyingType UnderlyingType = EUnderlyingType::VertexBuffer;
@@ -893,8 +998,10 @@ inline const TCHAR* GetBufferUnderlyingTypeName(FRDGBufferDesc::EUnderlyingType
 		return TEXT("VertexBuffer");
 	case FRDGBufferDesc::EUnderlyingType::StructuredBuffer:
 		return TEXT("StructuredBuffer");
+PRAGMA_DISABLE_DEPRECATION_WARNINGS
 	case FRDGBufferDesc::EUnderlyingType::AccelerationStructure:
 		return TEXT("AccelerationStructure");
+PRAGMA_ENABLE_DEPRECATION_WARNINGS
 	}
 	return TEXT("");
 }
@@ -913,6 +1020,21 @@ struct FRDGBufferSRVDesc final
 		BytesPerElement = GPixelFormats[Format].BlockBytes;
 	}
 
+	bool operator == (const FRDGBufferSRVDesc& Other) const
+	{
+		return Buffer == Other.Buffer && FRHIBufferSRVCreateInfo::operator==(Other);
+	}
+
+	bool operator != (const FRDGBufferSRVDesc& Other) const
+	{
+		return !(*this == Other);
+	}
+
+	friend uint32 GetTypeHash(const FRDGBufferSRVDesc& Desc)
+	{
+		return HashCombine(GetTypeHash(static_cast<const FRHIBufferSRVCreateInfo&>(Desc)), GetTypeHash(Desc.Buffer));
+	}
+
 	FRDGBufferRef Buffer = nullptr;
 };
 
@@ -928,6 +1050,21 @@ struct FRDGBufferUAVDesc final
 		, Buffer(InBuffer)
 	{}
 
+	bool operator == (const FRDGBufferUAVDesc& Other) const
+	{
+		return Buffer == Other.Buffer && FRHIBufferUAVCreateInfo::operator==(Other);
+	}
+
+	bool operator != (const FRDGBufferUAVDesc& Other) const
+	{
+		return !(*this == Other);
+	}
+
+	friend uint32 GetTypeHash(const FRDGBufferUAVDesc& Desc)
+	{
+		return HashCombine(GetTypeHash(static_cast<const FRHIBufferUAVCreateInfo&>(Desc)), GetTypeHash(Desc.Buffer));
+	}
+
 	FRDGBufferRef Buffer = nullptr;
 };
 
@@ -956,14 +1093,29 @@ public:
 	/** Returns the RHI buffer. */
 	FORCEINLINE FRHIBuffer* GetRHI() const { return Buffer; }
 
-	UE_DEPRECATED(5.0, "Buffers types have been consolidated; use GetRHI() instead.")
-	FORCEINLINE FRHIBuffer* GetVertexBufferRHI() const { return Buffer; }
+	/** Returns the default SRV. */
+	FORCEINLINE FRHIShaderResourceView* GetSRV()
+	{
+		if (!CachedSRV)
+		{
+			CachedSRV = GetOrCreateSRV(FRHIBufferSRVCreateInfo());
+		}
+		return CachedSRV;
+	}
 
-	UE_DEPRECATED(5.0, "Buffers types have been consolidated; use GetRHI() instead.")
-	FORCEINLINE FRHIBuffer* GetStructuredBufferRHI() const { return Buffer; }
+	FORCEINLINE uint32 GetSize() const
+	{
+		return Desc.GetSize();
+	}
+
+	FORCEINLINE uint32 GetAlignedSize() const
+	{
+		return Desc.BytesPerElement * NumAllocatedElements;
+	}
 
 private:
 	TRefCountPtr<FRHIBuffer> Buffer;
+	FRHIShaderResourceView* CachedSRV = nullptr;
 	FRHIBufferViewCache ViewCache;
 
 	FRDGBufferDesc GetAlignedDesc() const
@@ -987,6 +1139,8 @@ class RENDERCORE_API FRDGBuffer final
 	: public FRDGViewableResource
 {
 public:
+	static const ERDGViewableResourceType StaticType = ERDGViewableResourceType::Buffer;
+
 	FRDGBufferDesc Desc;
 	const ERDGBufferFlags Flags;
 
@@ -1027,6 +1181,11 @@ public:
 		return Handle;
 	}
 
+	FORCEINLINE uint32 GetSize() const
+	{
+		return Desc.GetSize();
+	}
+
 private:
 	FRDGBuffer(const TCHAR* InName, const FRDGBufferDesc& InDesc, ERDGBufferFlags InFlags)
 		: FRDGViewableResource(InName, ERDGViewableResourceType::Buffer)
@@ -1112,6 +1271,8 @@ class FRDGBufferSRV final
 	: public FRDGShaderResourceView
 {
 public:
+	static const ERDGViewType StaticType = ERDGViewType::BufferSRV;
+
 	/** Descriptor of the graph tracked SRV. */
 	const FRDGBufferSRVDesc Desc;
 
@@ -1136,6 +1297,8 @@ class FRDGBufferUAV final
 	: public FRDGUnorderedAccessView
 {
 public:
+	static const ERDGViewType StaticType = ERDGViewType::BufferUAV;
+
 	/** Descriptor of the graph tracked UAV. */
 	const FRDGBufferUAVDesc Desc;
 
@@ -1155,6 +1318,50 @@ private:
 	friend FRDGAllocator;
 };
 
+template <typename ViewableResourceType>
+inline ViewableResourceType* GetAs(FRDGViewableResource* Resource)
+{
+	check(ViewableResourceType::StaticType == Resource->Type);
+	return static_cast<ViewableResourceType*>(Resource);
+}
+
+template <typename ViewType>
+inline ViewType* GetAs(FRDGView* View)
+{
+	check(ViewType::StaticType == View->Type);
+	return static_cast<ViewType*>(View);
+}
+
+inline FRDGBuffer* GetAsBuffer(FRDGViewableResource* Resource)
+{
+	return GetAs<FRDGBuffer>(Resource);
+}
+
+inline FRDGTexture* GetAsTexture(FRDGViewableResource* Resource)
+{
+	return GetAs<FRDGTexture>(Resource);
+}
+
+inline FRDGBufferUAV* GetAsBufferUAV(FRDGView* View)
+{
+	return GetAs<FRDGBufferUAV>(View);
+}
+
+inline FRDGBufferSRV* GetAsBufferSRV(FRDGView* View)
+{
+	return GetAs<FRDGBufferSRV>(View);
+}
+
+inline FRDGTextureUAV* GetAsTextureUAV(FRDGView* View)
+{
+	return GetAs<FRDGTextureUAV>(View);
+}
+
+inline FRDGTextureSRV* GetAsTextureSRV(FRDGView* View)
+{
+	return GetAs<FRDGTextureSRV>(View);
+}
+
 inline FGraphicsPipelineRenderTargetsInfo ExtractRenderTargetsInfo(const FRDGParameterStruct& ParameterStruct);
 
 #include "RenderGraphResources.inl"
\ No newline at end of file
diff --git a/Engine/Source/Runtime/RenderCore/Public/RenderGraphResources.inl b/Engine/Source/Runtime/RenderCore/Public/RenderGraphResources.inl
--- a/Engine/Source/Runtime/RenderCore/Public/RenderGraphResources.inl
+++ b/Engine/Source/Runtime/RenderCore/Public/RenderGraphResources.inl
@@ -78,7 +78,7 @@ inline FPooledRenderTargetDesc Translate(const FRHITextureDesc& InDesc)
 inline FRHIBufferCreateInfo Translate(const FRDGBufferDesc& InDesc)
 {
 	FRHIBufferCreateInfo CreateInfo;
-	CreateInfo.Size = InDesc.GetTotalNumBytes();
+	CreateInfo.Size = InDesc.GetSize();
 	if (InDesc.UnderlyingType == FRDGBufferDesc::EUnderlyingType::VertexBuffer)
 	{
 		CreateInfo.Stride = 0;
diff --git a/Engine/Source/Runtime/RenderCore/Public/RenderGraphUtils.h b/Engine/Source/Runtime/RenderCore/Public/RenderGraphUtils.h
--- a/Engine/Source/Runtime/RenderCore/Public/RenderGraphUtils.h
+++ b/Engine/Source/Runtime/RenderCore/Public/RenderGraphUtils.h
@@ -52,6 +52,41 @@ inline FRHITexture* TryGetRHI(FRDGTextureRef Texture)
 	return Texture ? Texture->GetRHI() : nullptr;
 }
 
+inline FRHIBuffer* TryGetRHI(FRDGBuffer* Buffer)
+{
+	return Buffer ? Buffer->GetRHI() : nullptr;
+}
+
+inline FRHIBuffer* TryGetRHI(FRDGPooledBuffer* Buffer)
+{
+	return Buffer ? Buffer->GetRHI() : nullptr;
+}
+
+inline FRHIShaderResourceView* TryGetSRV(FRDGPooledBuffer* Buffer)
+{
+	return Buffer ? Buffer->GetSRV() : 0;
+}
+
+inline uint64 TryGetSize(const FRDGBuffer* Buffer)
+{
+	return Buffer ? Buffer->GetSize() : 0;
+}
+
+inline uint64 TryGetSize(const FRDGPooledBuffer* Buffer)
+{
+	return Buffer ? Buffer->GetSize() : 0;
+}
+
+inline bool IsRegistered(FRDGBuilder& GraphBuilder, const TRefCountPtr<IPooledRenderTarget>& RenderTarget)
+{
+	return GraphBuilder.FindExternalTexture(RenderTarget) != nullptr;
+}
+
+inline bool IsRegistered(FRDGBuilder& GraphBuilder, const TRefCountPtr<FRDGPooledBuffer>& Buffer)
+{
+	return GraphBuilder.FindExternalBuffer(Buffer) != nullptr;
+}
+
 /** Returns the pooled render target from an RDG texture if it exists, or null otherwise. */
 UE_DEPRECATED(5.0, "Accessing the underlying pooled render target has been deprecated. Use TryGetRHI() instead.")
 inline IPooledRenderTarget* TryGetPooledRenderTarget(FRDGTextureRef Texture)
@@ -588,7 +623,7 @@ struct RENDERCORE_API FComputeShaderUtils
 	{
 		checkf(IndirectArgsBuffer->Desc.UnderlyingType == FRDGBufferDesc::EUnderlyingType::VertexBuffer, TEXT("The buffer %s needs to be a vertex buffer to be used as an indirect dispatch parameters"), IndirectArgsBuffer->Name);
 		checkf(IndirectArgsBuffer->Desc.Usage & BUF_DrawIndirect, TEXT("The buffer %s for indirect dispatch parameters was not flagged with BUF_DrawIndirect"), IndirectArgsBuffer->Name);
-		ValidateIndirectArgsBuffer(IndirectArgsBuffer->Desc.GetTotalNumBytes(), IndirectArgOffset);
+		ValidateIndirectArgsBuffer(IndirectArgsBuffer->GetSize(), IndirectArgOffset);
 	}
 
 	/**
@@ -1180,11 +1215,22 @@ private:
 #define RDG_WAIT_FOR_TASKS(GraphBuilder) RDG_WAIT_FOR_TASKS_CONDITIONAL(GraphBuilder, true)
 
 // Allocates an RDG pooled buffer instance. Attempts to reuse allocation if Out has a value. Returns true a new instance was allocated, or false if the existing allocation was reused.
-RENDERCORE_API bool GetPooledFreeBuffer(
-	FRHICommandList& RHICmdList,
+RENDERCORE_API bool AllocatePooledBuffer(
 	const FRDGBufferDesc& Desc,
 	TRefCountPtr<FRDGPooledBuffer>& Out,
-	const TCHAR* InDebugName);
+	const TCHAR* Name,
+	ERDGPooledBufferAlignment Alignment = ERDGPooledBufferAlignment::Page);
+
+RENDERCORE_API TRefCountPtr<FRDGPooledBuffer> AllocatePooledBuffer(
+	const FRDGBufferDesc& Desc,
+	const TCHAR* Name,
+	ERDGPooledBufferAlignment Alignment = ERDGPooledBufferAlignment::Page);
+
+//UE_DEPRECATED(5.1, "GetPooledFreeBuffer is deprecated. Use AllocatePooledBuffer instead.")
+inline bool GetPooledFreeBuffer(FRHICommandList& RHICmdList, const FRDGBufferDesc& Desc, TRefCountPtr<FRDGPooledBuffer>& Out, const TCHAR* Name)
+{
+	return AllocatePooledBuffer(Desc, Out, Name);
+}
 
 RENDERCORE_API bool AllocatePooledTexture(
 	const FRDGTextureDesc& Desc,
diff --git a/Engine/Source/Runtime/Renderer/Private/ShaderPrint.cpp b/Engine/Source/Runtime/Renderer/Private/ShaderPrint.cpp
--- a/Engine/Source/Runtime/Renderer/Private/ShaderPrint.cpp
+++ b/Engine/Source/Runtime/Renderer/Private/ShaderPrint.cpp
@@ -94,13 +94,7 @@ namespace ShaderPrint
 	public:
 		void InitRHI() override
 		{
-			if (!Buffer.IsValid())
-			{
-				FRHICommandList* UnusedCmdList = new FRHICommandList(FRHIGPUMask::All());
-				GetPooledFreeBuffer(*UnusedCmdList, FRDGBufferDesc::CreateStructuredDesc(sizeof(FPackedShaderPrintItem), 1), Buffer, TEXT("ShaderPrint.EmptyValueBuffer"));
-				delete UnusedCmdList;
-				UnusedCmdList = nullptr;
-			}
+			Buffer = AllocatePooledBuffer(FRDGBufferDesc::CreateStructuredDesc(sizeof(FPackedShaderPrintItem), 1), TEXT("ShaderPrint.EmptyValueBuffer"));
 		}
 	};
 
