Added buffer pool dump command for debugging memory usage.

diff --git a/Engine/Source/Runtime/RenderCore/Private/RenderGraphResourcePool.cpp b/Engine/Source/Runtime/RenderCore/Private/RenderGraphResourcePool.cpp
--- a/Engine/Source/Runtime/RenderCore/Private/RenderGraphResourcePool.cpp
+++ b/Engine/Source/Runtime/RenderCore/Private/RenderGraphResourcePool.cpp
@@ -15,6 +15,44 @@ UE_TRACE_EVENT_BEGIN(Cpu, FRDGBufferPool_CreateBuffer, NoSync)
 	UE_TRACE_EVENT_FIELD(uint32, SizeInBytes)
 UE_TRACE_EVENT_END()
 
+RENDERCORE_API void DumpBufferPoolMemory(FOutputDevice& OutputDevice)
+{
+	GRenderGraphResourcePool.DumpMemoryUsage(OutputDevice);
+}
+
+static FAutoConsoleCommandWithOutputDevice GDumpBufferPoolMemoryCmd(
+	TEXT("r.DumpBufferPoolMemory"),
+	TEXT("Dump allocation information for the buffer pool."),
+	FConsoleCommandWithOutputDeviceDelegate::CreateStatic(DumpBufferPoolMemory)
+);
+
+void FRDGBufferPool::DumpMemoryUsage(FOutputDevice& OutputDevice)
+{
+	OutputDevice.Logf(TEXT("Pooled Buffers:"));
+
+	TArray<TRefCountPtr<FRDGPooledBuffer>> BuffersBySize = AllocatedBuffers;
+
+	Algo::Sort(BuffersBySize, [](const TRefCountPtr<FRDGPooledBuffer>& LHS, const TRefCountPtr<FRDGPooledBuffer>& RHS)
+	{
+		return LHS->GetAlignedSize() > RHS->GetAlignedSize();
+	});
+
+	for (const TRefCountPtr<FRDGPooledBuffer>& Buffer : BuffersBySize)
+	{
+		const uint32 BufferSize = Buffer->GetAlignedSize();
+		const uint32 UnusedForNFrames = FrameCounter - Buffer->LastUsedFrame;
+
+		OutputDevice.Logf(
+			TEXT("  %6.3fMB Name: %s, NumElements: %u, BytesPerElement: %u, UAV: %s, Frames Since Requested: %u"),
+			(float)BufferSize / (1024.0f * 1024.0f),
+			Buffer->Name,
+			Buffer->NumAllocatedElements,
+			Buffer->Desc.BytesPerElement,
+			EnumHasAnyFlags(Buffer->Desc.Usage, EBufferUsageFlags::UnorderedAccess) ? TEXT("Yes") : TEXT("No"),
+			UnusedForNFrames);
+	}
+}
+
 TRefCountPtr<FRDGPooledBuffer> FRDGBufferPool::FindFreeBuffer(const FRDGBufferDesc& Desc, const TCHAR* InDebugName, ERDGPooledBufferAlignment Alignment)
 {
 	const uint64 BufferPageSize = 64 * 1024;
diff --git a/Engine/Source/Runtime/RenderCore/Private/RenderGraphResourcePool.h b/Engine/Source/Runtime/RenderCore/Private/RenderGraphResourcePool.h
--- a/Engine/Source/Runtime/RenderCore/Private/RenderGraphResourcePool.h
+++ b/Engine/Source/Runtime/RenderCore/Private/RenderGraphResourcePool.h
@@ -29,6 +29,8 @@ public:
 
 	TRefCountPtr<FRDGPooledBuffer> FindFreeBuffer(const FRDGBufferDesc& Desc, const TCHAR* InDebugName, ERDGPooledBufferAlignment Alignment = ERDGPooledBufferAlignment::Page);
 
+	void DumpMemoryUsage(FOutputDevice& OutputDevice);
+
 private:
 	void ReleaseDynamicRHI() override;
 
