Added RDG versions of UnifiedBuffer utilities.
 - TRefCountPtr<FRDGPooledBuffer> replaces FRW{Structured, ByteAddress}Buffer.
 - Resize{Structured, ByteAddress}BufferIfNeeded replaces ResizeResourceIfNeeded.
 - FRDGScatterUploadBuffer replaces FScatterUploadBuffer.

diff --git a/Engine/Source/Runtime/RenderCore/Private/UnifiedBuffer.cpp b/Engine/Source/Runtime/RenderCore/Private/UnifiedBuffer.cpp
--- a/Engine/Source/Runtime/RenderCore/Private/UnifiedBuffer.cpp
+++ b/Engine/Source/Runtime/RenderCore/Private/UnifiedBuffer.cpp
@@ -141,6 +141,439 @@ static uint32 CalculateFloat4sPerLine()
 	return PrimitivesPerTextureLine * FScatterUploadBuffer::PrimitiveDataStrideInFloat4s;
 }
 
+///////////////////////////////////////////////////////////////////////////////////////////////////
+
+class FRDGByteBufferShader : public FGlobalShader
+{
+	DECLARE_INLINE_TYPE_LAYOUT(FRDGByteBufferShader, NonVirtual);
+
+	FRDGByteBufferShader() {}
+	FRDGByteBufferShader(const ShaderMetaType::CompiledShaderInitializerType& Initializer)
+		: FGlobalShader(Initializer)
+	{}
+
+	class ResourceTypeDim : SHADER_PERMUTATION_INT("RESOURCE_TYPE", (int)EByteBufferResourceType::Count);
+	using FPermutationDomain = TShaderPermutationDomain<ResourceTypeDim>;
+
+	static bool ShouldCompilePermutation(const FGlobalShaderPermutationParameters& Parameters)
+	{
+		return FByteBufferShader::ShouldCompilePermutation(Parameters);
+	}
+
+	BEGIN_SHADER_PARAMETER_STRUCT(FParameters, )
+		SHADER_PARAMETER(uint32, Value)
+		SHADER_PARAMETER(uint32, Size)
+		SHADER_PARAMETER(uint32, SrcOffset)
+		SHADER_PARAMETER(uint32, DstOffset)
+		SHADER_PARAMETER(uint32, Float4sPerLine)
+		SHADER_PARAMETER_RDG_BUFFER_UAV(RWBuffer<float4>, DstBuffer)
+		SHADER_PARAMETER_RDG_BUFFER_UAV(RWStructuredBuffer<float4>, DstStructuredBuffer)
+		SHADER_PARAMETER_RDG_BUFFER_UAV(RWByteAddressBuffer, DstByteAddressBuffer)
+		SHADER_PARAMETER_RDG_TEXTURE_UAV(RWTexture2D<float4>, DstTexture)
+	END_SHADER_PARAMETER_STRUCT()
+};
+
+class FRDGMemsetBufferCS : public FRDGByteBufferShader
+{
+	DECLARE_GLOBAL_SHADER(FRDGMemsetBufferCS);
+	SHADER_USE_PARAMETER_STRUCT(FRDGMemsetBufferCS, FRDGByteBufferShader);
+};
+IMPLEMENT_GLOBAL_SHADER(FRDGMemsetBufferCS, "/Engine/Private/ByteBuffer.usf", "MemsetBufferCS", SF_Compute);
+
+class FRDGMemcpyCS : public FRDGByteBufferShader
+{
+	DECLARE_GLOBAL_SHADER(FRDGMemcpyCS);
+	SHADER_USE_PARAMETER_STRUCT(FRDGMemcpyCS, FRDGByteBufferShader);
+
+	BEGIN_SHADER_PARAMETER_STRUCT(FParameters, )
+		SHADER_PARAMETER_STRUCT_INCLUDE(FRDGByteBufferShader::FParameters, Common)
+		SHADER_PARAMETER_RDG_BUFFER_SRV(Buffer<float4>, SrcBuffer)
+		SHADER_PARAMETER_RDG_BUFFER_SRV(StructuredBuffer<float4>, SrcStructuredBuffer)
+		SHADER_PARAMETER_RDG_BUFFER_SRV(ByteAddressBuffer, SrcByteAddressBuffer)
+		SHADER_PARAMETER_RDG_TEXTURE_SRV(Texture2D<float4>, SrcTexture)
+	END_SHADER_PARAMETER_STRUCT()
+};
+IMPLEMENT_GLOBAL_SHADER(FRDGMemcpyCS, "/Engine/Private/ByteBuffer.usf", "MemcpyCS", SF_Compute);
+
+class FRDGScatterCopyCS : public FRDGByteBufferShader
+{
+	DECLARE_GLOBAL_SHADER(FRDGScatterCopyCS);
+	SHADER_USE_PARAMETER_STRUCT(FRDGScatterCopyCS, FRDGByteBufferShader);
+
+	BEGIN_SHADER_PARAMETER_STRUCT(FParameters, )
+		SHADER_PARAMETER_STRUCT_INCLUDE(FRDGByteBufferShader::FParameters, Common)
+		SHADER_PARAMETER(uint32, NumScatters)
+		SHADER_PARAMETER_RDG_BUFFER_SRV(ByteAddressBuffer, UploadByteAddressBuffer)
+		SHADER_PARAMETER_RDG_BUFFER_SRV(StructuredBuffer<float4>, UploadStructuredBuffer)
+		SHADER_PARAMETER_RDG_BUFFER_SRV(ByteAddressBuffer, ScatterByteAddressBuffer)
+		SHADER_PARAMETER_RDG_BUFFER_SRV(StructuredBuffer<uint>, ScatterStructuredBuffer)
+	END_SHADER_PARAMETER_STRUCT()
+};
+IMPLEMENT_GLOBAL_SHADER(FRDGScatterCopyCS, "/Engine/Private/ByteBuffer.usf", "ScatterCopyCS", SF_Compute);
+
+EResourceType GetBufferType(FRDGBuffer* Buffer)
+{
+	const FRDGBufferDesc& Desc = Buffer->Desc;
+
+	if (Desc.UnderlyingType == FRDGBufferDesc::EUnderlyingType::StructuredBuffer)
+	{
+		if (EnumHasAnyFlags(Desc.Usage, BUF_ByteAddressBuffer))
+		{
+			return EResourceType::BYTEBUFFER;
+		}
+		else
+		{
+			return EResourceType::STRUCTURED_BUFFER;
+		}
+	}
+	else
+	{
+		return EResourceType::BUFFER;
+	}
+}
+
+EResourceType GetResourceType(FRDGViewableResource* Resource)
+{
+	check(Resource);
+	switch (Resource->Type)
+	{
+	case ERDGViewableResourceType::Texture:
+		return EResourceType::TEXTURE;
+	case ERDGViewableResourceType::Buffer:
+		return GetBufferType(GetAsBuffer(Resource));
+	}
+	checkNoEntry();
+	return EResourceType::BUFFER;
+}
+
+void MemsetResource(FRDGBuilder& GraphBuilder, FRDGViewableResource* Resource, const FMemsetResourceParams& Params)
+{
+	check(Resource);
+
+	EByteBufferResourceType ResourceTypeEnum = EByteBufferResourceType::Count;
+
+	auto* PassParameters = GraphBuilder.AllocParameters<FRDGMemsetBufferCS::FParameters>();
+	PassParameters->Value = Params.Value;
+	PassParameters->Size = Params.Count;
+	PassParameters->DstOffset = Params.DstOffset;
+
+	// each thread will set 4 floats / uints
+	uint32 Divisor = 1;
+
+	switch (GetResourceType(Resource))
+	{
+	case EResourceType::BYTEBUFFER:
+		ResourceTypeEnum = EByteBufferResourceType::Uint_Buffer;
+		PassParameters->DstByteAddressBuffer = GraphBuilder.CreateUAV(GetAsBuffer(Resource));
+		Divisor = 4;
+		break;
+	case EResourceType::BUFFER:
+		ResourceTypeEnum = EByteBufferResourceType::Float4_Buffer;
+		PassParameters->DstBuffer = GraphBuilder.CreateUAV(GetAsBuffer(Resource));
+		break;
+	case EResourceType::STRUCTURED_BUFFER:
+		ResourceTypeEnum = EByteBufferResourceType::Float4_StructuredBuffer;
+		PassParameters->DstStructuredBuffer = GraphBuilder.CreateUAV(GetAsBuffer(Resource));
+		break;
+	case EResourceType::TEXTURE:
+		ResourceTypeEnum = EByteBufferResourceType::Float4_Texture;
+		PassParameters->DstTexture = GraphBuilder.CreateUAV(GetAsTexture(Resource));
+		PassParameters->Float4sPerLine = CalculateFloat4sPerLine();
+		break;
+	default:
+		checkNoEntry();
+	}
+
+	FRDGMemsetBufferCS::FPermutationDomain PermutationVector;
+	PermutationVector.Set<FRDGMemsetBufferCS::ResourceTypeDim >((int)ResourceTypeEnum);
+	TShaderMapRef<FRDGMemsetBufferCS> ComputeShader(GetGlobalShaderMap(GMaxRHIFeatureLevel), PermutationVector);
+
+	FComputeShaderUtils::AddPass(
+		GraphBuilder,
+		RDG_EVENT_NAME("MemsetResource (%s)", Resource->Name),
+		ComputeShader,
+		PassParameters,
+		FIntVector(FMath::DivideAndRoundUp(Params.Count / Divisor, 64u), 1, 1));
+}
+
+void MemcpyResource(FRDGBuilder& GraphBuilder, FRDGViewableResource* DstResource, FRDGViewableResource* SrcResource, const FMemcpyResourceParams& Params)
+{
+	const EResourceType ResourceType = GetResourceType(DstResource);
+	checkf(ResourceType == GetResourceType(SrcResource), TEXT("Unable to MemcpyResource because the source and destination view types don't match."));
+
+	RDG_EVENT_SCOPE(GraphBuilder, "Memcpy %s -> %s", DstResource->Name, SrcResource->Name);
+
+	// each thread will copy 4 floats / uints
+	const uint32 Divisor = ResourceType == EResourceType::BYTEBUFFER ? 4 : 1;
+
+	uint32 NumElementsProcessed = 0;
+
+	while (NumElementsProcessed < Params.Count)
+	{
+		const uint32 NumWaves = FMath::Max(FMath::Min<uint32>(GRHIMaxDispatchThreadGroupsPerDimension.X, FMath::DivideAndRoundUp(Params.Count / Divisor, 64u)), 1u);
+		const uint32 NumElementsPerDispatch = FMath::Min(FMath::Max(NumWaves, 1u) * Divisor * 64, Params.Count - NumElementsProcessed);
+
+		EByteBufferResourceType ResourceTypeEnum = EByteBufferResourceType::Count;
+
+		auto* PassParameters = GraphBuilder.AllocParameters<FRDGMemcpyCS::FParameters>();
+		PassParameters->Common.Size = NumElementsPerDispatch;
+		PassParameters->Common.SrcOffset = (Params.SrcOffset + NumElementsProcessed);
+		PassParameters->Common.DstOffset = (Params.DstOffset + NumElementsProcessed);
+
+		switch (ResourceType)
+		{
+		case EResourceType::BYTEBUFFER:
+			ResourceTypeEnum = EByteBufferResourceType::Uint_Buffer;
+			PassParameters->SrcByteAddressBuffer           = GraphBuilder.CreateSRV(GetAsBuffer(SrcResource));
+			PassParameters->Common.DstByteAddressBuffer    = GraphBuilder.CreateUAV(GetAsBuffer(DstResource), ERDGUnorderedAccessViewFlags::SkipBarrier);
+			break;
+		case EResourceType::STRUCTURED_BUFFER:
+			ResourceTypeEnum = EByteBufferResourceType::Float4_StructuredBuffer;
+			PassParameters->SrcStructuredBuffer            = GraphBuilder.CreateSRV(GetAsBuffer(SrcResource));
+			PassParameters->Common.DstStructuredBuffer = GraphBuilder.CreateUAV(GetAsBuffer(DstResource), ERDGUnorderedAccessViewFlags::SkipBarrier);
+			break;
+		case EResourceType::BUFFER:
+			ResourceTypeEnum = EByteBufferResourceType::Float4_Buffer;
+			PassParameters->SrcBuffer                      = GraphBuilder.CreateSRV(GetAsBuffer(SrcResource));
+			PassParameters->Common.DstBuffer               = GraphBuilder.CreateUAV(GetAsBuffer(DstResource), ERDGUnorderedAccessViewFlags::SkipBarrier);
+			break;
+		case EResourceType::TEXTURE:
+			ResourceTypeEnum = EByteBufferResourceType::Float4_Texture;
+			PassParameters->SrcTexture                     = GraphBuilder.CreateSRV(GetAsTexture(SrcResource));
+			PassParameters->Common.DstTexture              = GraphBuilder.CreateUAV(GetAsTexture(DstResource), ERDGUnorderedAccessViewFlags::SkipBarrier);
+			PassParameters->Common.Float4sPerLine          = CalculateFloat4sPerLine();
+			break;
+		default:
+			checkNoEntry();
+		}
+
+		FRDGMemcpyCS::FPermutationDomain PermutationVector;
+		PermutationVector.Set<FRDGMemcpyCS::ResourceTypeDim >((int)ResourceTypeEnum);
+		TShaderMapRef<FRDGMemcpyCS> ComputeShader(GetGlobalShaderMap(GMaxRHIFeatureLevel), PermutationVector);
+
+		FComputeShaderUtils::AddPass(
+			GraphBuilder,
+			RDG_EVENT_NAME("Offset[%d] Count[%d]", NumElementsProcessed, NumElementsPerDispatch),
+			ComputeShader,
+			PassParameters,
+			FIntVector(NumWaves, 1, 1));
+
+		NumElementsProcessed += NumElementsPerDispatch;
+	}
+}
+
+FRDGBuffer* ResizeBufferIfNeeded(FRDGBuilder& GraphBuilder, TRefCountPtr<FRDGPooledBuffer>& ExternalBuffer, const FRDGBufferDesc& BufferDesc, const TCHAR* Name)
+{
+	FRDGBuffer* InternalBufferNew = nullptr;
+
+	if (!ExternalBuffer)
+	{
+		InternalBufferNew = GraphBuilder.CreateBuffer(BufferDesc, Name);
+		ExternalBuffer = GraphBuilder.ConvertToExternalBuffer(InternalBufferNew);
+		return InternalBufferNew;
+	}
+
+	FRDGBuffer* InternalBufferOld = GraphBuilder.RegisterExternalBuffer(ExternalBuffer);
+
+	const uint32 BufferSize = BufferDesc.GetSize();
+	const uint32 BufferSizeOld = InternalBufferOld->GetSize();
+
+	if (BufferSize != BufferSizeOld)
+	{
+		InternalBufferNew = GraphBuilder.CreateBuffer(BufferDesc, Name);
+		ExternalBuffer = GraphBuilder.ConvertToExternalBuffer(InternalBufferNew);
+
+		// Copy data to new buffer
+		FMemcpyResourceParams Params;
+		Params.Count = FMath::Min(BufferSize, BufferSizeOld) / BufferDesc.BytesPerElement;
+		Params.SrcOffset = 0;
+		Params.DstOffset = 0;
+		MemcpyResource(GraphBuilder, InternalBufferNew, InternalBufferOld, Params);
+
+		return InternalBufferNew;
+	}
+
+	return InternalBufferOld;
+}
+
+FRDGBuffer* ResizeBufferIfNeeded(FRDGBuilder& GraphBuilder, TRefCountPtr<FRDGPooledBuffer>& ExternalBuffer, EPixelFormat Format, uint32 NumElements, const TCHAR* Name)
+{
+	const uint32 BytesPerElement = GPixelFormats[Format].BlockBytes;
+
+	return ResizeBufferIfNeeded(GraphBuilder, ExternalBuffer, FRDGBufferDesc::CreateBufferDesc(BytesPerElement, NumElements), Name);
+}
+
+FRDGBuffer* ResizeStructuredBufferIfNeeded(FRDGBuilder& GraphBuilder, TRefCountPtr<FRDGPooledBuffer>& ExternalBuffer, uint32 NumBytes, const TCHAR* Name)
+{
+	const uint32 BytesPerElement = 16;
+
+	check((NumBytes & (BytesPerElement - 1)) == 0);
+
+	const uint32 NumElements = NumBytes / BytesPerElement;
+
+	return ResizeBufferIfNeeded(GraphBuilder, ExternalBuffer, FRDGBufferDesc::CreateStructuredDesc(BytesPerElement, NumElements), Name);
+}
+
+FRDGBuffer* ResizeByteAddressBufferIfNeeded(FRDGBuilder& GraphBuilder, TRefCountPtr<FRDGPooledBuffer>& ExternalBuffer, uint32 NumBytes, const TCHAR* Name)
+{
+	// Needs to be aligned to 16 bytes to MemcpyResource to work correctly (otherwise it skips last unaligned elements of the buffer during resize)
+	check((NumBytes & 15) == 0);
+
+	return ResizeBufferIfNeeded(GraphBuilder, ExternalBuffer, FRDGBufferDesc::CreateByteAddressDesc(NumBytes), Name);
+}
+
+void FRDGScatterUploadBuffer::Init(FRDGBuilder& GraphBuilder, uint32 NumElements, uint32 InNumBytesPerElement, bool bInFloat4Buffer, const TCHAR* Name)
+{
+	TRACE_CPUPROFILER_EVENT_SCOPE(FRDGScatterUploadBuffer::Init);
+
+	NumScatters = 0;
+	MaxScatters = NumElements;
+	NumBytesPerElement = InNumBytesPerElement;
+	bFloat4Buffer = bInFloat4Buffer;
+
+	const EBufferUsageFlags Usage = bInFloat4Buffer ? BUF_None : BUF_ByteAddressBuffer;
+	const uint32 TypeSize = bInFloat4Buffer ? 16 : 4;
+
+	const uint32 ScatterNumBytesPerElement = sizeof(uint32);
+	const uint32 ScatterBytes = NumElements * ScatterNumBytesPerElement;
+	const uint32 ScatterBufferSize = (uint32)FMath::Min((uint64)FMath::RoundUpToPowerOfTwo(ScatterBytes), GetMaxBufferDimension() * sizeof(uint32));
+	check(ScatterBufferSize >= ScatterBytes);
+
+	const uint32 UploadNumBytesPerElement = TypeSize;
+	const uint32 UploadBytes = NumElements * NumBytesPerElement;
+	const uint32 UploadBufferSize = (uint32)FMath::Min((uint64)FMath::RoundUpToPowerOfTwo(UploadBytes), GetMaxBufferDimension() * TypeSize);
+	check(UploadBufferSize >= UploadBytes);
+
+	// Recreate buffers is they are already queued into RDG from a previous call.
+	if (IsRegistered(GraphBuilder, ScatterBuffer))
+	{
+		ScatterBuffer = nullptr;
+		UploadBuffer = nullptr;
+	}
+
+	if (!ScatterBuffer || ScatterBytes > ScatterBuffer->GetSize() || ScatterBufferSize < ScatterBuffer->GetSize() / 2)
+	{
+		FRDGBufferDesc Desc = FRDGBufferDesc::CreateStructuredUploadDesc(ScatterNumBytesPerElement, ScatterBufferSize / ScatterNumBytesPerElement);
+		Desc.Usage |= Usage;
+
+		AllocatePooledBuffer(Desc, ScatterBuffer, Name, ERDGPooledBufferAlignment::None);
+	}
+
+	if (!UploadBuffer || UploadBytes > UploadBuffer->GetSize() || UploadBufferSize < UploadBuffer->GetSize() / 2)
+	{
+		FRDGBufferDesc Desc = FRDGBufferDesc::CreateStructuredUploadDesc(TypeSize, UploadBufferSize / UploadNumBytesPerElement);
+		Desc.Usage |= Usage;
+
+		AllocatePooledBuffer(Desc, UploadBuffer, Name, ERDGPooledBufferAlignment::None);
+	}
+
+	ScatterData = (uint32*)RHILockBuffer(ScatterBuffer->GetRHI(), 0, ScatterBytes, RLM_WriteOnly);
+	UploadData = (uint8*)RHILockBuffer(UploadBuffer->GetRHI(), 0, UploadBytes, RLM_WriteOnly);
+}
+
+void FRDGScatterUploadBuffer::ResourceUploadTo(FRDGBuilder& GraphBuilder, FRDGViewableResource* DstResource)
+{
+	TRACE_CPUPROFILER_EVENT_SCOPE(FRDGScatterUploadBuffer::ResourceUploadTo);
+
+	RHIUnlockBuffer(ScatterBuffer->GetRHI());
+	RHIUnlockBuffer(UploadBuffer->GetRHI());
+
+	ScatterData = nullptr;
+	UploadData = nullptr;
+
+	if (NumScatters == 0)
+	{
+		return;
+	}
+
+	constexpr uint32 ThreadGroupSize = 64u;
+	uint32 NumBytesPerThread = (NumBytesPerElement & 15) == 0 ? 16 : 4;
+	uint32 NumThreadsPerScatter = NumBytesPerElement / NumBytesPerThread;
+	uint32 NumThreads = NumScatters * NumThreadsPerScatter;
+	uint32 NumDispatches = FMath::DivideAndRoundUp(NumThreads, ThreadGroupSize);
+	uint32 NumLoops = FMath::DivideAndRoundUp(NumDispatches, (uint32)GMaxComputeDispatchDimension);
+
+	const EResourceType DstResourceType = GetResourceType(DstResource);
+	check(bFloat4Buffer != (DstResourceType == EResourceType::BYTEBUFFER));
+
+	EByteBufferResourceType ResourceTypeEnum = EByteBufferResourceType::Count;
+
+	FRDGScatterCopyCS::FParameters Parameters;
+	Parameters.Common.Size = NumThreadsPerScatter;
+	Parameters.NumScatters = NumScatters;
+
+	FRDGBufferSRV* ScatterBufferSRV = GraphBuilder.CreateSRV(GraphBuilder.RegisterExternalBuffer(ScatterBuffer));
+	FRDGBufferSRV* UploadBufferSRV  = GraphBuilder.CreateSRV(GraphBuilder.RegisterExternalBuffer(UploadBuffer));
+
+	if (DstResourceType == EResourceType::BYTEBUFFER)
+	{
+		if (NumBytesPerThread == 16)
+		{
+			ResourceTypeEnum = EByteBufferResourceType::Uint4Aligned_Buffer;
+		}
+		else
+		{
+			ResourceTypeEnum = EByteBufferResourceType::Uint_Buffer;
+		}
+		Parameters.UploadByteAddressBuffer = UploadBufferSRV;
+		Parameters.ScatterByteAddressBuffer = ScatterBufferSRV;
+		Parameters.Common.DstByteAddressBuffer = GraphBuilder.CreateUAV(GetAsBuffer(DstResource), ERDGUnorderedAccessViewFlags::SkipBarrier);
+	}
+	else if (DstResourceType == EResourceType::STRUCTURED_BUFFER)
+	{
+		ResourceTypeEnum = EByteBufferResourceType::Float4_StructuredBuffer;
+
+		Parameters.UploadStructuredBuffer = UploadBufferSRV;
+		Parameters.ScatterStructuredBuffer = ScatterBufferSRV;
+		Parameters.Common.DstStructuredBuffer = GraphBuilder.CreateUAV(GetAsBuffer(DstResource), ERDGUnorderedAccessViewFlags::SkipBarrier);
+	}
+	else if (DstResourceType == EResourceType::BUFFER)
+	{
+		ResourceTypeEnum = EByteBufferResourceType::Float4_Buffer;
+
+		Parameters.UploadStructuredBuffer = UploadBufferSRV;
+		Parameters.ScatterStructuredBuffer = ScatterBufferSRV;
+		Parameters.Common.DstBuffer = GraphBuilder.CreateUAV(GetAsBuffer(DstResource), ERDGUnorderedAccessViewFlags::SkipBarrier);
+	}
+	else if (DstResourceType == EResourceType::TEXTURE)
+	{
+		ResourceTypeEnum = EByteBufferResourceType::Float4_Texture;
+
+		Parameters.UploadStructuredBuffer = UploadBufferSRV;
+		Parameters.ScatterStructuredBuffer = ScatterBufferSRV;
+		Parameters.Common.DstTexture = GraphBuilder.CreateUAV(GetAsTexture(DstResource), ERDGUnorderedAccessViewFlags::SkipBarrier);
+
+		Parameters.Common.Float4sPerLine = CalculateFloat4sPerLine();
+	}
+
+	FRDGByteBufferShader::FPermutationDomain PermutationVector;
+	PermutationVector.Set<FRDGByteBufferShader::ResourceTypeDim>((int)ResourceTypeEnum);
+	TShaderMapRef<FRDGScatterCopyCS> ComputeShader(GetGlobalShaderMap(GMaxRHIFeatureLevel), PermutationVector);
+
+	for (uint32 LoopIdx = 0; LoopIdx < NumLoops; ++LoopIdx)
+	{
+		Parameters.Common.SrcOffset = LoopIdx * (uint32)GMaxComputeDispatchDimension * ThreadGroupSize;
+
+		uint32 LoopNumDispatch = FMath::Min(NumDispatches - LoopIdx * (uint32)GMaxComputeDispatchDimension, (uint32)GMaxComputeDispatchDimension);
+
+		FComputeShaderUtils::AddPass(
+			GraphBuilder,
+			RDG_EVENT_NAME("ScatterUpload[%d] (Offset: %u, GroupSize: %u)", Parameters.Common.SrcOffset, LoopNumDispatch),
+			ComputeShader,
+			GraphBuilder.AllocParameters(&Parameters),
+			FIntVector(LoopNumDispatch, 1, 1));
+	}
+
+	Reset();
+}
+
+void FRDGScatterUploadBuffer::Reset()
+{
+	NumScatters = 0;
+	MaxScatters = 0;
+	NumBytesPerElement = 0;
+}
+
 template<typename ResourceType>
 void MemsetResource(FRHICommandList& RHICmdList, const ResourceType& DstBuffer, const FMemsetResourceParams& Params)
 {
diff --git a/Engine/Source/Runtime/RenderCore/Public/UnifiedBuffer.h b/Engine/Source/Runtime/RenderCore/Public/UnifiedBuffer.h
--- a/Engine/Source/Runtime/RenderCore/Public/UnifiedBuffer.h
+++ b/Engine/Source/Runtime/RenderCore/Public/UnifiedBuffer.h
@@ -5,8 +5,7 @@
 #include "CoreMinimal.h"
 #include "RHI.h"
 #include "Containers/DynamicRHIResourceArray.h"
-
-class FRDGBuilder;
+#include "RenderGraphDefinitions.h"
 
 /*
  * Can store arbitrary data so long as it follows alignment restrictions. Intended mostly for read only data uploaded from CPU.
@@ -177,4 +176,85 @@ public:
 			bUploadViaCreate = bInUploadViaCreate;
 		}
 	}
+};
+
+RENDERCORE_API void MemsetResource(FRDGBuilder& GraphBuilder, FRDGViewableResource* DstResource, const FMemsetResourceParams& Params);
+RENDERCORE_API void MemcpyResource(FRDGBuilder& GraphBuilder, FRDGViewableResource* DstResource, FRDGViewableResource* SrcResource, const FMemcpyResourceParams& Params);
+
+RENDERCORE_API FRDGBuffer* ResizeBufferIfNeeded(FRDGBuilder& GraphBuilder, TRefCountPtr<FRDGPooledBuffer>& ExternalBuffer, const FRDGBufferDesc& BufferDesc, const TCHAR* Name);
+RENDERCORE_API FRDGBuffer* ResizeBufferIfNeeded(FRDGBuilder& GraphBuilder, TRefCountPtr<FRDGPooledBuffer>& ExternalBuffer, EPixelFormat Format, uint32 NumElements, const TCHAR* Name);
+RENDERCORE_API FRDGBuffer* ResizeStructuredBufferIfNeeded(FRDGBuilder& GraphBuilder, TRefCountPtr<FRDGPooledBuffer>& ExternalBuffer, uint32 NumBytes, const TCHAR* Name);
+RENDERCORE_API FRDGBuffer* ResizeByteAddressBufferIfNeeded(FRDGBuilder& GraphBuilder, TRefCountPtr<FRDGPooledBuffer>& ExternalBuffer, uint32 NumBytes, const TCHAR* Name);
+
+class FRDGScatterUploadBuffer
+{
+public:
+	enum { PrimitiveDataStrideInFloat4s = FScatterUploadBuffer::PrimitiveDataStrideInFloat4s };
+
+	RENDERCORE_API void Init(FRDGBuilder& GraphBuilder, uint32 NumElements, uint32 NumBytesPerElement, bool bInFloat4Buffer, const TCHAR* Name);
+
+	RENDERCORE_API void ResourceUploadTo(FRDGBuilder& GraphBuilder, FRDGViewableResource* DstResource);
+
+	void Add(uint32 Index, const void* Data, uint32 Num = 1)
+	{
+		void* Dst = Add_GetRef(Index, Num);
+		FMemory::ParallelMemcpy(Dst, Data, Num * NumBytesPerElement, EMemcpyCachePolicy::StoreUncached);
+	}
+
+	void* Add_GetRef(uint32 Index, uint32 Num = 1)
+	{
+		checkSlow(NumScatters + Num <= MaxScatters);
+		checkSlow(ScatterData != nullptr);
+		checkSlow(UploadData != nullptr);
+
+		uint32* ScatterWriteData = ScatterData + NumScatters;
+
+		for (uint32 i = 0; i < Num; i++)
+		{
+			ScatterWriteData[i] = Index + i;
+		}
+
+		void* Result = UploadData + NumScatters * NumBytesPerElement;
+		NumScatters += Num;
+		return Result;
+	}
+
+	void* Set_GetRef(uint32 ElementIndex, uint32 ElementScatterOffset, uint32 Num = 1)
+	{
+		checkSlow(ElementIndex + Num <= MaxScatters);
+		checkSlow(ScatterData != nullptr);
+		checkSlow(UploadData != nullptr);
+
+		for (uint32 i = 0; i < Num; i++)
+		{
+			ScatterData[ElementIndex + i] = ElementScatterOffset + i;
+		}
+		return UploadData + ElementIndex * NumBytesPerElement;
+	}
+
+	/**
+	 * Get pointer to an element data area, given the index of the element (not the destination scatter offset).
+	 */
+	FORCEINLINE void* GetRef(uint32 ElementIndex)
+	{
+		checkSlow(ScatterData != nullptr);
+		checkSlow(UploadData != nullptr);
+
+		return UploadData + ElementIndex * NumBytesPerElement;
+	}
+
+private:
+	RENDERCORE_API void Reset();
+
+	TRefCountPtr<FRDGPooledBuffer> ScatterBuffer;
+	TRefCountPtr<FRDGPooledBuffer> UploadBuffer;
+
+	uint32* ScatterData = nullptr;
+	uint8* UploadData = nullptr;
+
+	uint32 	NumScatters = 0;
+	uint32 	MaxScatters = 0;
+	uint32	NumBytesPerElement = 0;
+
+	bool bFloat4Buffer = false;
 };
\ No newline at end of file
