Refactored RDG in preparation for UnifiedBuffer conversions.

 - Refactord 'Finalized Access' feature into a more flexible 'External' vs. 'Internal' access mode per resource toggle.
      - Resources can transition between modes multiple times within the graph.
      - Supports async compute pipeline.
      - Supports queueing of requests to avoid back-to-back helper passes.
      - This feature is needed to support conversion of GPU scene buffers.

 - Deprecated the ReadOnly and ForceTracking resource flags and added a 'SkipTracking' flag instead.
      - Previous semantics were confusing and error prone.
      - New model requires a manual flag to tell RDG never to transition a resource.
      - This flag is used for read-only dummy resources as an optimization.

 - Renamed some of the auxiliary 'FinalizedResource' utilities since the name no longer matches the semantics.

diff --git a/Engine/Plugins/RenderGraphInsights/Source/RenderGraphInsights/Private/RenderGraphProvider.cpp b/Engine/Plugins/RenderGraphInsights/Source/RenderGraphInsights/Private/RenderGraphProvider.cpp
--- a/Engine/Plugins/RenderGraphInsights/Source/RenderGraphInsights/Private/RenderGraphProvider.cpp
+++ b/Engine/Plugins/RenderGraphInsights/Source/RenderGraphInsights/Private/RenderGraphProvider.cpp
@@ -65,6 +65,7 @@ FResourcePacket::FResourcePacket(const UE::Trace::IAnalyzer::FOnEventContext& Co
 	, bExternal(Context.EventData.GetValue<bool>("IsExternal"))
 	, bExtracted(Context.EventData.GetValue<bool>("IsExtracted"))
 	, bCulled(Context.EventData.GetValue<bool>("IsCulled"))
+	, bTrackingSkipped(Context.EventData.GetValue<bool>("IsTrackingSkipped"))
 	, bTransient(Context.EventData.GetValue<bool>("IsTransient"))
 	, bTransientUntracked(Context.EventData.GetValue<bool>("IsTransientUntracked"))
 	, bTransientCacheHit(Context.EventData.GetValue<bool>("IsTransientCacheHit"))
diff --git a/Engine/Plugins/RenderGraphInsights/Source/RenderGraphInsights/Private/RenderGraphProvider.h b/Engine/Plugins/RenderGraphInsights/Source/RenderGraphInsights/Private/RenderGraphProvider.h
--- a/Engine/Plugins/RenderGraphInsights/Source/RenderGraphInsights/Private/RenderGraphProvider.h
+++ b/Engine/Plugins/RenderGraphInsights/Source/RenderGraphInsights/Private/RenderGraphProvider.h
@@ -86,6 +86,7 @@ public:
 	bool bExternal{};
 	bool bExtracted{};
 	bool bCulled{};
+	bool bTrackingSkipped{};
 	bool bTransient{};
 	bool bTransientUntracked{};
 	bool bTransientCacheHit{};
diff --git a/Engine/Plugins/RenderGraphInsights/Source/RenderGraphInsights/Private/RenderGraphTrack.cpp b/Engine/Plugins/RenderGraphInsights/Source/RenderGraphInsights/Private/RenderGraphTrack.cpp
--- a/Engine/Plugins/RenderGraphInsights/Source/RenderGraphInsights/Private/RenderGraphTrack.cpp
+++ b/Engine/Plugins/RenderGraphInsights/Source/RenderGraphInsights/Private/RenderGraphTrack.cpp
@@ -1234,7 +1234,7 @@ void FRenderGraphTrack::BuildDrawState(FRenderGraphTrackDrawStateBuilder& Builde
 
 			const auto IsPacketCulled = [&](const FResourcePacket& Packet)
 			{
-				if (Packet.bCulled || Packet.bTransientUntracked)
+				if (Packet.bCulled || Packet.bTrackingSkipped || Packet.bTransientUntracked)
 				{
 					return true;
 				}
diff --git a/Engine/Plugins/Runtime/HairStrands/Source/HairStrandsCore/Private/GroomResources.cpp b/Engine/Plugins/Runtime/HairStrands/Source/HairStrandsCore/Private/GroomResources.cpp
--- a/Engine/Plugins/Runtime/HairStrands/Source/HairStrandsCore/Private/GroomResources.cpp
+++ b/Engine/Plugins/Runtime/HairStrands/Source/HairStrandsCore/Private/GroomResources.cpp
@@ -1713,8 +1713,6 @@ void CreateHairStrandsDebugResources(FRDGBuilder& GraphBuilder, const FHairStran
 		In->VoxelData.GetData(),
 		sizeof(FHairStrandsDebugDatas::FVoxel) * In->VoxelData.Num());
 
-	FRDGResourceAccessFinalizer ResourceAccessFinalizer;
-	Out->VoxelOffsetAndCount = ConvertToFinalizedExternalBuffer(GraphBuilder, ResourceAccessFinalizer, VoxelOffsetAndCount);
-	Out->VoxelData = ConvertToFinalizedExternalBuffer(GraphBuilder, ResourceAccessFinalizer, VoxelData);
-	ResourceAccessFinalizer.Finalize(GraphBuilder);
+	Out->VoxelOffsetAndCount = ConvertToExternalAccessBuffer(GraphBuilder, VoxelOffsetAndCount);
+	Out->VoxelData = ConvertToExternalAccessBuffer(GraphBuilder, VoxelData);
 }
diff --git a/Engine/Plugins/Runtime/HairStrands/Source/HairStrandsCore/Private/HairStrandsRendering.cpp b/Engine/Plugins/Runtime/HairStrands/Source/HairStrandsCore/Private/HairStrandsRendering.cpp
--- a/Engine/Plugins/Runtime/HairStrands/Source/HairStrandsCore/Private/HairStrandsRendering.cpp
+++ b/Engine/Plugins/Runtime/HairStrands/Source/HairStrandsCore/Private/HairStrandsRendering.cpp
@@ -1701,7 +1701,7 @@ void ComputeHairStrandsInterpolation(
 	Instance->HairGroupPublicData->VFInput.Meshes	= FHairGroupPublicData::FVertexFactoryInput::FMeshes();
 	const EHairGeometryType InstanceGeometryType	= Instance->GeometryType;
 
-	FRDGResourceAccessFinalizer ResourceAccessFinalizer;
+	FRDGExternalAccessQueue ExternalAccessQueue;
 
 	if (InstanceGeometryType == EHairGeometryType::Strands)
 	{
@@ -2118,8 +2118,8 @@ void ComputeHairStrandsInterpolation(
 						Instance->Guides.DeformedRootResource);
 				}
 
-				ResourceAccessFinalizer.AddBuffer(Register(GraphBuilder, LOD.DeformedResource->GetBuffer(FHairCardsDeformedResource::Current), ERDGImportedBufferFlags::None).Buffer, ERHIAccess::SRVMask);
-				ResourceAccessFinalizer.AddBuffer(Register(GraphBuilder, LOD.DeformedResource->GetBuffer(FHairCardsDeformedResource::Previous), ERDGImportedBufferFlags::None).Buffer, ERHIAccess::SRVMask);
+				ExternalAccessQueue.Add(Register(GraphBuilder, LOD.DeformedResource->GetBuffer(FHairCardsDeformedResource::Current), ERDGImportedBufferFlags::None).Buffer, ERHIAccess::SRVMask);
+				ExternalAccessQueue.Add(Register(GraphBuilder, LOD.DeformedResource->GetBuffer(FHairCardsDeformedResource::Previous), ERDGImportedBufferFlags::None).Buffer, ERHIAccess::SRVMask);
 			}
 
 			#if RHI_RAYTRACING
@@ -2187,7 +2187,7 @@ void ComputeHairStrandsInterpolation(
 					Instance->Guides.RestRootResource,
 					Instance->Guides.DeformedRootResource);
 
-				ResourceAccessFinalizer.AddBuffer(Register(GraphBuilder, MeshesInstance.DeformedResource->GetBuffer(FHairMeshesDeformedResource::Current), ERDGImportedBufferFlags::None).Buffer, ERHIAccess::SRVMask);
+				ExternalAccessQueue.Add(Register(GraphBuilder, MeshesInstance.DeformedResource->GetBuffer(FHairMeshesDeformedResource::Current), ERDGImportedBufferFlags::None).Buffer, ERHIAccess::SRVMask);
 			}
 
 			#if RHI_RAYTRACING
@@ -2228,7 +2228,7 @@ void ComputeHairStrandsInterpolation(
 	Instance->HairGroupPublicData->VFInput.LocalToWorldTransform = Instance->GetCurrentLocalToWorld();
 	Instance->HairGroupPublicData->bSupportVoxelization = Instance->Strands.Modifier.bSupportVoxelization && Instance->bCastShadow;
 
-	ResourceAccessFinalizer.Finalize(GraphBuilder);
+	ExternalAccessQueue.Submit(GraphBuilder);
 }
 
 void ResetHairStrandsInterpolation(
diff --git a/Engine/Plugins/Runtime/HairStrands/Source/HairStrandsCore/Public/GroomResources.h b/Engine/Plugins/Runtime/HairStrands/Source/HairStrandsCore/Public/GroomResources.h
--- a/Engine/Plugins/Runtime/HairStrands/Source/HairStrandsCore/Public/GroomResources.h
+++ b/Engine/Plugins/Runtime/HairStrands/Source/HairStrandsCore/Public/GroomResources.h
@@ -15,7 +15,7 @@
 
 inline uint32 GetBufferTotalNumBytes(const FRDGExternalBuffer& In) 
 {
-	return In.Buffer ? In.Buffer->Desc.GetTotalNumBytes() : 0;
+	return In.Buffer ? In.Buffer->GetSize() : 0;
 }
 
 enum class EHairStrandsResourcesType : uint8
diff --git a/Engine/Source/Runtime/Engine/Private/Rendering/SubsurfaceProfile.cpp b/Engine/Source/Runtime/Engine/Private/Rendering/SubsurfaceProfile.cpp
--- a/Engine/Source/Runtime/Engine/Private/Rendering/SubsurfaceProfile.cpp
+++ b/Engine/Source/Runtime/Engine/Private/Rendering/SubsurfaceProfile.cpp
@@ -294,7 +294,7 @@ IPooledRenderTarget* FSubsurfaceProfileTexture::GetSSProfilesPreIntegratedTextur
 			FPixelShaderUtils::AddFullscreenPass(GraphBuilder, GlobalShaderMap, RDG_EVENT_NAME("SSS::SSProfilePreIntegrated"), PixelShader, PassParameters, ViewRect);
 		}
 
-		GSSProfilesPreIntegratedTexture = ConvertToFinalizedExternalTexture(GraphBuilder, ProfileTexture);
+		GSSProfilesPreIntegratedTexture = ConvertToExternalAccessTexture(GraphBuilder, ProfileTexture);
 	}
 
 	return GSSProfilesPreIntegratedTexture;
diff --git a/Engine/Source/Runtime/Landscape/Private/LandscapeEditLayers.cpp b/Engine/Source/Runtime/Landscape/Private/LandscapeEditLayers.cpp
--- a/Engine/Source/Runtime/Landscape/Private/LandscapeEditLayers.cpp
+++ b/Engine/Source/Runtime/Landscape/Private/LandscapeEditLayers.cpp
@@ -4059,10 +4059,6 @@ namespace EditLayersHeightmapLocalMerge_RenderThread
 				if (LayerTrackedTexture == nullptr)
 				{
 					LayerTrackedTexture = &(OutTrackedTextures.Add(LayerHeightmap.Texture, FLandscapeRDGTrackedTexture(LayerHeightmap.Texture)));
-					// Layer heightmaps were not created with TexCreate_RenderTargetable and are therefore considered as ReadOnly by the RDG, which prevents their resource access from being automatically managed
-					//  (i.e. transfered to CopySrc when we copy them and back to SRVMask at the end) so we need to force the RDG to track them:
-					// TODO [jonathan.bard] : remove this if we don't use the texture array :
-					LayerTrackedTexture->bNeedsForceTracking = true;
 				}
 			}
 		}
@@ -4301,7 +4297,7 @@ namespace EditLayersHeightmapLocalMerge_RenderThread
 			TRefCountPtr<IPooledRenderTarget> RenderTarget = CreateRenderTarget(TextureResolveInfo.Texture->TextureRHI, **DebugName);
 
 			// Force tracking on the external texture, so that it can be copied to via CopyTexture within the graph : 
-			FRDGTextureRef DestinationTexture = GraphBuilder.RegisterExternalTexture(RenderTarget, ERDGTextureFlags::ForceTracking);
+			FRDGTextureRef DestinationTexture = GraphBuilder.RegisterExternalTexture(RenderTarget);
 
 			FRHICopyTextureInfo CopyTextureInfo;
 			// We want to copy all mips : 
@@ -6017,9 +6013,6 @@ namespace EditLayersWeightmapLocalMerge_RenderThread
 				{
 					TrackedTexture = &(OutTrackedTextures.Add(VisibleEditLayersWeightmapTexture, FLandscapeRDGTrackedTexture(VisibleEditLayersWeightmapTexture)));
 				}
-				// Source weightmaps were not created with TexCreate_RenderTargetable and are therefore considered as ReadOnly by the RDG, which prevents their resource access from being automatically managed
-				//  (i.e. transfered to CopySrc when we copy them) so we need to force the RDG to track them: 
-				TrackedTexture->bNeedsForceTracking = true;
 				TrackedTexture->bNeedsSRV = true;
 			}
 		}
@@ -6033,9 +6026,6 @@ namespace EditLayersWeightmapLocalMerge_RenderThread
 				TrackedTexture = &(OutTrackedTextures.Add(TextureResolveInfo.Texture, FLandscapeRDGTrackedTexture(TextureResolveInfo.Texture)));
 			}
 
-			// Destination weightmaps were not created with TexCreate_RenderTargetable and are therefore considered as ReadOnly by the RDG, which prevents their resource access from being automatically managed
-			//  (i.e. transfered to CopyDst when we copy to them and back to SRVMask at the end) so we need to force the RDG to track them: 
-			TrackedTexture->bNeedsForceTracking = true;
 			TrackedTexture->bNeedsSRV = true;
 		}
 	}
diff --git a/Engine/Source/Runtime/Landscape/Private/LandscapeEditResources.cpp b/Engine/Source/Runtime/Landscape/Private/LandscapeEditResources.cpp
--- a/Engine/Source/Runtime/Landscape/Private/LandscapeEditResources.cpp
+++ b/Engine/Source/Runtime/Landscape/Private/LandscapeEditResources.cpp
@@ -176,12 +176,8 @@ void TrackLandscapeRDGTextures(FRDGBuilder& GraphBuilder, TMap<FTexture2DResourc
 		// We need a debug name that will last until the graph is executed: 
 		TrackedTexture.DebugName = GraphBuilder.AllocObject<FString>(TrackedTexture.TextureResource->GetTextureName().ToString());
 
-		// Register the texture to the GraphBuilder : we'll read from / write to it like a render target : 
-		TRefCountPtr<IPooledRenderTarget> RenderTarget = CreateRenderTarget(TrackedTexture.TextureResource->TextureRHI, **TrackedTexture.DebugName);
-
-		// Force tracking on the external texture if necessary, so that it can be copied from/to via CopyTexture within the graph : 
-		ERDGTextureFlags TextureFlags = (TrackedTexture.bNeedsForceTracking || TrackedTexture.bNeedsScratch) ? ERDGTextureFlags::ForceTracking : ERDGTextureFlags::ReadOnly;
-		TrackedTexture.ExternalTextureRef = GraphBuilder.RegisterExternalTexture(RenderTarget, **TrackedTexture.DebugName, TextureFlags);
+		// Register the texture to the GraphBuilder : we'll read from / write to it like a render target :
+		TrackedTexture.ExternalTextureRef = GraphBuilder.RegisterExternalTexture(CreateRenderTarget(TrackedTexture.TextureResource->TextureRHI, **TrackedTexture.DebugName));
 
 		if (TrackedTexture.bNeedsScratch)
 		{
diff --git a/Engine/Source/Runtime/Landscape/Private/LandscapeEditResources.h b/Engine/Source/Runtime/Landscape/Private/LandscapeEditResources.h
--- a/Engine/Source/Runtime/Landscape/Private/LandscapeEditResources.h
+++ b/Engine/Source/Runtime/Landscape/Private/LandscapeEditResources.h
@@ -99,8 +99,6 @@ struct FLandscapeRDGTrackedTexture
 	bool bNeedsScratch = false;
 	/** If bNeedsSRV is set to true, the registration of the texture in the render graph will also produce a FRDGTextureSRVRef (ExternalTextureSRVRef) in order to access this texture as a SRV */
 	bool bNeedsSRV = false;
-	/** If bNeedsForceTracking is set to true, the texture will be declared in the render graph with the ERDGTextureFlags::ForceTracking flag, so that it can be copied from/to via CopyTexture within the graph */
-	bool bNeedsForceTracking = false;
 	/** Debug name for this texture (transient memory : will last for the lifetime of the GraphBuilder) */
 	FString* DebugName = nullptr;
 	/** If bNeedsScratch is true, contains the RDG ref of the allocated intermediate scratch texture */
diff --git a/Engine/Source/Runtime/RenderCore/Private/DumpGPU.cpp b/Engine/Source/Runtime/RenderCore/Private/DumpGPU.cpp
--- a/Engine/Source/Runtime/RenderCore/Private/DumpGPU.cpp
+++ b/Engine/Source/Runtime/RenderCore/Private/DumpGPU.cpp
@@ -1146,9 +1146,9 @@ struct FRDGResourceDumpContext
 		}
 
 		// Verify the texture is able to do resource transitions.
-		if (EnumHasAnyFlags(SubresourceDesc.Texture->Flags, ERDGTextureFlags::ReadOnly))
+		if (EnumHasAnyFlags(SubresourceDesc.Texture->Flags, ERDGTextureFlags::SkipTracking))
 		{
-			UE_LOG(LogRendererCore, Warning, TEXT("Not dumping %s because has ERDGTextureFlags::ReadOnly."), SubresourceDesc.Texture->Name);
+			UE_LOG(LogRendererCore, Warning, TEXT("Not dumping %s because has ERDGTextureFlags::SkipTracking."), SubresourceDesc.Texture->Name);
 			return;
 		}
 
@@ -1266,7 +1266,7 @@ struct FRDGResourceDumpContext
 		}
 
 		const FRDGBufferDesc& Desc = Buffer->Desc;
-		const int32 ByteSize = Desc.GetTotalNumBytes();
+		const int32 ByteSize = Desc.GetSize();
 
 		bool bDumpResourceInfos;
 		bool bDumpResourceBinary;
@@ -1303,9 +1303,9 @@ struct FRDGResourceDumpContext
 			}
 
 			// Verify the texture is able to do resource transitions.
-			if (EnumHasAnyFlags(Buffer->Flags, ERDGBufferFlags::ReadOnly))
+			if (EnumHasAnyFlags(Buffer->Flags, ERDGBufferFlags::SkipTracking))
 			{
-				UE_LOG(LogRendererCore, Warning, TEXT("Not dumping %s because has ERDGBufferFlags::ReadOnly."), Buffer->Name);
+				UE_LOG(LogRendererCore, Warning, TEXT("Not dumping %s because has ERDGBufferFlags::SkipTracking."), Buffer->Name);
 				return;
 			}
 
diff --git a/Engine/Source/Runtime/RenderCore/Private/RenderGraphBuilder.cpp b/Engine/Source/Runtime/RenderCore/Private/RenderGraphBuilder.cpp
--- a/Engine/Source/Runtime/RenderCore/Private/RenderGraphBuilder.cpp
+++ b/Engine/Source/Runtime/RenderCore/Private/RenderGraphBuilder.cpp
@@ -591,57 +591,149 @@ const TRefCountPtr<IPooledRenderTarget>& FRDGBuilder::ConvertToExternalTexture(F
 
 ///////////////////////////////////////////////////////////////////////////////////////////////////
 
-BEGIN_SHADER_PARAMETER_STRUCT(FFinalizePassParameters, )
+BEGIN_SHADER_PARAMETER_STRUCT(FAccessModePassParameters, )
 	RDG_TEXTURE_ACCESS_ARRAY(Textures)
 	RDG_BUFFER_ACCESS_ARRAY(Buffers)
 END_SHADER_PARAMETER_STRUCT()
 
-void FRDGBuilder::FinalizeResourceAccess(FRDGTextureAccessArray&& InTextures, FRDGBufferAccessArray&& InBuffers)
+void FRDGBuilder::UseExternalAccessMode(FRDGViewableResource* Resource, ERHIAccess ReadOnlyAccess, ERHIPipeline Pipelines)
 {
-	auto* PassParameters = AllocParameters<FFinalizePassParameters>();
-	PassParameters->Textures = Forward<FRDGTextureAccessArray&&>(InTextures);
-	PassParameters->Buffers = Forward<FRDGBufferAccessArray&&>(InBuffers);
+	IF_RDG_ENABLE_DEBUG(UserValidation.ValidateUseExternalAccessMode(Resource, ReadOnlyAccess, Pipelines));
 
-	// Take reference to pass parameters version since we've moved the memory.
-	const auto& LocalTextures = PassParameters->Textures;
-	const auto& LocalBuffers = PassParameters->Buffers;
+	if (!GRDGAsyncCompute || !GSupportsEfficientAsyncCompute)
+	{
+		Pipelines = ERHIPipeline::Graphics;
+	}
 
-#if RDG_ENABLE_DEBUG
+	auto& AccessModeState = Resource->AccessModeState;
+
+	// We already validated that back-to-back calls to UseExternalAccessMode are valid only if the parameters match,
+	// so we can safely no-op this call.
+	if (AccessModeState.Mode == FRDGViewableResource::EAccessMode::External)
+	{
+		return;
+	}
+
+	// We have to flush the queue when going from QueuedInternal -> External. A queued internal state
+	// implies that the resource was in an external access mode before, so it needs an 'end' pass to 
+	// contain any passes which might have used the resource in its external state.
+	if (AccessModeState.QueueIndex.IsValid())
+	{
+		FlushAccessModeQueue();
+	}
+
+	check(AccessModeState.QueueIndex.IsNull());
+	AccessModeState.QueueIndex = FRDGViewableResource::FAccessModeState::FQueueIndex(AccessModeQueue.Num());
+	AccessModeQueue.Emplace(Resource);
+
+	Resource->SetExternalAccessMode(ReadOnlyAccess, Pipelines);
+}
+
+void FRDGBuilder::UseInternalAccessMode(FRDGViewableResource* Resource)
+{
+	IF_RDG_ENABLE_DEBUG(UserValidation.ValidateUseInternalAccessMode(Resource));
+
+	auto& AccessModeState = Resource->AccessModeState;
+
+	// Just no-op if the resource is already in (or queued for) the 'Internal state.
+	if (AccessModeState.Mode == FRDGViewableResource::EAccessMode::Internal)
 	{
-		const FRDGPassHandle FinalizePassHandle(Passes.Num());
+		return;
+	}
 
-		for (FRDGTextureAccess TextureAccess : LocalTextures)
+	// If the resource has a queued transition to the external access state, then we can safely back it out.
+	if (AccessModeState.QueueIndex.IsValid())
+	{
+		AccessModeQueue[AccessModeState.QueueIndex.GetIndex()] = nullptr;
+		AccessModeState.QueueIndex = {};
+	}
+	else
+	{
+		AccessModeState.QueueIndex = FRDGViewableResource::FAccessModeState::FQueueIndex(AccessModeQueue.Num());
+		AccessModeQueue.Emplace(Resource);
+	}
+
+	AccessModeState.Mode = FRDGViewableResource::EAccessMode::Internal;
+}
+
+void FRDGBuilder::FlushAccessModeQueue()
+{
+	if (AccessModeQueue.IsEmpty() || bSkipAuxiliaryPasses)
+	{
+		return;
+	}
+
+	bSkipAuxiliaryPasses = true;
+
+	FAccessModePassParameters* ParametersByPipeline[] =
+	{
+		AllocParameters<FAccessModePassParameters>(),
+		AllocParameters<FAccessModePassParameters>()
+	};
+
+	ERHIPipeline ParameterPipelines = ERHIPipeline::None;
+
+	for (FRDGViewableResource* Resource : AccessModeQueue)
+	{
+		if (!Resource)
 		{
-			UserValidation.ValidateFinalize(TextureAccess.GetTexture(), TextureAccess.GetAccess(), FinalizePassHandle);
+			continue;
 		}
 
-		for (FRDGBufferAccess BufferAccess : LocalBuffers)
+		const auto& AccessModeState = Resource->AccessModeState;
+
+		ParameterPipelines |= AccessModeState.Pipelines;
+
+		for (uint32 PipelineIndex = 0; PipelineIndex < GetRHIPipelineCount(); ++PipelineIndex)
 		{
-			UserValidation.ValidateFinalize(BufferAccess.GetBuffer(), BufferAccess.GetAccess(), FinalizePassHandle);
+			const ERHIPipeline Pipeline = static_cast<ERHIPipeline>(1 << PipelineIndex);
+
+			if (EnumHasAnyFlags(AccessModeState.Pipelines, Pipeline))
+			{
+				switch (Resource->Type)
+				{
+				case ERDGViewableResourceType::Texture:
+					ParametersByPipeline[PipelineIndex]->Textures.Emplace(GetAsTexture(Resource), AccessModeState.Access);
+					break;
+				case ERDGViewableResourceType::Buffer:
+					ParametersByPipeline[PipelineIndex]->Buffers.Emplace(GetAsBuffer(Resource), AccessModeState.Access);
+					break;
+				}
+			}
 		}
 	}
-#endif
 
-	AddPass(
-		RDG_EVENT_NAME("FinalizeResourceAccess(Textures: %d, Buffers: %d)", LocalTextures.Num(), LocalBuffers.Num()),
-		PassParameters,
-		// Use all of the work flags so that any access is valid.
-		ERDGPassFlags::Copy | ERDGPassFlags::Compute | ERDGPassFlags::Raster | ERDGPassFlags::SkipRenderPass |
-		// We're not writing to anything, so we have to tell the pass not to cull.
-		ERDGPassFlags::NeverCull,
-		[](FRHICommandList&) {});
+	if (EnumHasAnyFlags(ParameterPipelines, ERHIPipeline::Graphics))
+	{
+		FAccessModePassParameters* Parameters = ParametersByPipeline[GetRHIPipelineIndex(ERHIPipeline::Graphics)];
 
-	// bFinalized must be set after adding the finalize pass, as future declarations of the resource will be ignored.
+		AddPass(
+			RDG_EVENT_NAME("AccessModePass[Graphics] (Textures: %d, Buffers: %d)", Parameters->Textures.Num(), Parameters->Buffers.Num()),
+			Parameters,
+			// Use all of the work flags so that any access is valid.
+			ERDGPassFlags::Copy | ERDGPassFlags::Compute | ERDGPassFlags::Raster | ERDGPassFlags::SkipRenderPass | ERDGPassFlags::NeverCull,
+			[](FRHICommandList&) {});
+	}
 
-	for (FRDGTextureAccess TextureAccess : LocalTextures)
+	if (EnumHasAnyFlags(ParameterPipelines, ERHIPipeline::AsyncCompute))
 	{
-		TextureAccess->SetFinalizedAccess(TextureAccess.GetAccess());
+		FAccessModePassParameters* Parameters = ParametersByPipeline[GetRHIPipelineIndex(ERHIPipeline::AsyncCompute)];
+
+		AddPass(
+			RDG_EVENT_NAME("AccessModePass[AsyncCompute] (Textures: %d, Buffers: %d)", Parameters->Textures.Num(), Parameters->Buffers.Num()),
+			Parameters,
+			ERDGPassFlags::AsyncCompute | ERDGPassFlags::NeverCull,
+			[](FRHIComputeCommandList&) {});
 	}
 
-	for (FRDGBufferAccess BufferAccess : LocalBuffers)
+	for (FRDGViewableResource* Resource : AccessModeQueue)
 	{
-		BufferAccess->SetFinalizedAccess(BufferAccess.GetAccess());
+		Resource->AccessModeState.QueueIndex = {};
 	}
+
+	AccessModeQueue.Reset();
+
+	bSkipAuxiliaryPasses = false;
 }
 
 ///////////////////////////////////////////////////////////////////////////////////////////////////
@@ -677,44 +769,9 @@ FRDGTexture* FRDGBuilder::RegisterExternalTexture(
 	}
 
 	const FRDGTextureDesc Desc = Translate(ExternalPooledTexture->GetDesc());
-	bool bFinalizedAccess = false;
-
-	if (!EnumHasAnyFlags(Flags, ERDGTextureFlags::ForceTracking) &&
-		!EnumHasAnyFlags(Desc.Flags, TexCreate_RenderTargetable | TexCreate_ResolveTargetable | TexCreate_DepthStencilTargetable | TexCreate_UAV | TexCreate_DepthStencilResolveTarget))
-	{
-		Flags |= ERDGTextureFlags::ReadOnly;
-		bFinalizedAccess = true;
-	}
-
-	FRDGTextureRef Texture = Textures.Allocate(Allocator, Name, Desc, Flags);
+	FRDGTexture* Texture = Textures.Allocate(Allocator, Name, Desc, Flags);
 	SetRHI(Texture, ExternalPooledTexture.GetReference(), GetProloguePassHandle());
-
 	Texture->bExternal = true;
-
-
-	if (EnumHasAnyFlags(ExternalTextureRHI->GetFlags(), TexCreate_Foveation))
-	{
-		Texture->EpilogueAccess = ERHIAccess::ShadingRateSource;
-	}
-
-	// Textures that are created read-only are not transitioned by RDG.
-	if (bFinalizedAccess)
-	{
-		ERHIAccess ValidAccessMask = ERHIAccess::ReadOnlyExclusiveMask;
-
-		if (EnumHasAnyFlags(Desc.Flags, TexCreate_CPUReadback))
-		{
-			ValidAccessMask |= ERHIAccess::CopyDest;
-		}
-
-		if (EnumHasAnyFlags(Desc.Flags, TexCreate_Foveation))
-		{
-			ValidAccessMask |= ERHIAccess::ShadingRateSource;
-		}
-
-		Texture->SetFinalizedAccess(ValidAccessMask);
-	}
-
 	ExternalTextures.Add(Texture->GetRHIUnchecked(), Texture);
 
 	IF_RDG_ENABLE_DEBUG(UserValidation.ValidateRegisterExternalTexture(Texture));
@@ -748,26 +805,10 @@ FRDGBufferRef FRDGBuilder::RegisterExternalBuffer(
 		return FoundBuffer;
 	}
 
-	const FRDGBufferDesc& Desc = ExternalPooledBuffer->Desc;
-	bool bFinalizedAccess = false;
-
-	if (!EnumHasAnyFlags(Flags, ERDGBufferFlags::ForceTracking) && !EnumHasAnyFlags(Desc.Usage, BUF_UnorderedAccess))
-	{
-		Flags |= ERDGBufferFlags::ReadOnly;
-		bFinalizedAccess = true;
-	}
-
-	FRDGBufferRef Buffer = Buffers.Allocate(Allocator, Name, ExternalPooledBuffer->Desc, Flags);
+	FRDGBuffer* Buffer = Buffers.Allocate(Allocator, Name, ExternalPooledBuffer->Desc, Flags);
 	SetRHI(Buffer, ExternalPooledBuffer, GetProloguePassHandle());
-
 	Buffer->bExternal = true;
 
-	// Buffers that are created read-only are not transitioned by RDG.
-	if (bFinalizedAccess)
-	{
-		Buffer->SetFinalizedAccess(ERHIAccess::ReadOnlyExclusiveMask);
-	}
-
 	ExternalBuffers.Add(Buffer->GetRHIUnchecked(), Buffer);
 
 	IF_RDG_ENABLE_DEBUG(UserValidation.ValidateRegisterExternalBuffer(Buffer));
@@ -1416,6 +1457,7 @@ void FRDGBuilder::Execute()
 	SCOPED_NAMED_EVENT_TEXT("FRDGBuilder::Execute", FColor::Magenta);
 
 	GRDGTransientResourceAllocator.ReleasePendingDeallocations();
+	FlushAccessModeQueue();
 
 	// Create the epilogue pass at the end of the graph just prior to compilation.
 	{
@@ -1787,13 +1829,12 @@ FRDGPass* FRDGBuilder::SetupPass(FRDGPass* Pass)
 	{
 		TryAddView(TextureView);
 
-		if (Texture->bFinalizedAccess)
+		if (Texture->AccessModeState.IsExternalAccess())
 		{
-			// Finalized resources expected to remain in the same state, so are ignored by the graph.
-			// As only External | Extracted resources can be finalized by the user, the graph doesn't
-			// need to track them any more for culling / transition purposes. Validation checks that these
-			// invariants are true.
-			IF_RDG_ENABLE_DEBUG(UserValidation.ValidateFinalizedAccess(Texture, Access, Pass));
+			// Resources in external access mode are expected to remain in the same state and are ignored by the graph.
+			// As only External | Extracted resources can be set as external by the user, the graph doesn't need to track
+			// them any more for culling / transition purposes. Validation checks that these invariants are true.
+			IF_RDG_ENABLE_DEBUG(UserValidation.ValidateExternalAccess(Texture, Access, Pass));
 			return;
 		}
 
@@ -1843,13 +1884,12 @@ FRDGPass* FRDGBuilder::SetupPass(FRDGPass* Pass)
 	{
 		TryAddView(BufferView);
 
-		if (Buffer->bFinalizedAccess)
+		if (Buffer->AccessModeState.IsExternalAccess())
 		{
-			// Finalized resources expected to remain in the same state, so are ignored by the graph.
-			// As only External | Extracted resources can be finalized by the user, the graph doesn't
-			// need to track them any more for culling / transition purposes. Validation checks that these
-			// invariants are true.
-			IF_RDG_ENABLE_DEBUG(UserValidation.ValidateFinalizedAccess(Buffer, Access, Pass));
+			// Resources in external access mode are expected to remain in the same state and are ignored by the graph.
+			// As only External | Extracted resources can be set as external by the user, the graph doesn't need to track
+			// them any more for culling / transition purposes. Validation checks that these invariants are true.
+			IF_RDG_ENABLE_DEBUG(UserValidation.ValidateExternalAccess(Buffer, Access, Pass));
 			return;
 		}
 
@@ -2570,14 +2610,13 @@ void FRDGBuilder::CreatePassBarriers()
 
 void FRDGBuilder::AddEpilogueTransition(FRDGTextureRef Texture)
 {
-	check(Texture->HasRHI() || (Texture->bCulled && !Texture->bExternal));
-
-	if (Texture->bCulled || !Texture->bLastOwner)
+	if (!Texture->HasRHI() || Texture->bCulled || !Texture->bLastOwner)
 	{
 		return;
 	}
 
-	if (!Texture->bFinalizedAccess)
+	// Only textures with internal access mode are transitioned to their epilogue state.
+	if (Texture->AccessModeState.Mode == FRDGViewableResource::EAccessMode::Internal)
 	{
 		const FRDGPassHandle EpiloguePassHandle = GetEpiloguePassHandle();
 
@@ -2603,12 +2642,13 @@ void FRDGBuilder::AddEpilogueTransition(FRDGTextureRef Texture)
 		AddTransition(EpiloguePassHandle, Texture, ScratchTextureState);
 		ScratchTextureState.Reset();
 	}
+	// External access mode textures are left in their last external state.
 	else
 	{
 		Texture->EpilogueAccess = Texture->GetState()[0].Access;
 	}
 
-	// Resources that were finalized on registration could have an unknown epilogue state.
+	// Resources that were marked for external access on registration (via SkipTracking) could have an unknown epilogue state.
 	if (Texture->EpilogueAccess != ERHIAccess::Unknown)
 	{
 		EpilogueResourceAccesses.Emplace(Texture->GetRHI(), Texture->EpilogueAccess);
@@ -2620,22 +2660,21 @@ void FRDGBuilder::AddEpilogueTransition(FRDGTextureRef Texture)
 	}
 	else if (!Texture->bTransient)
 	{
-		// Non-transient render targets need to be 'resurrected' to hold the last reference in the chain.
-		// Transient render targets don't have that restriction since there's no actual pooling happening.
+		// Non-transient textures need to be 'resurrected' to hold the last reference in the chain.
+		// Transient textures don't have that restriction since there's no actual pooling happening.
 		ActivePooledTextures.Emplace(Texture->RenderTarget);
 	}
 }
 
 void FRDGBuilder::AddEpilogueTransition(FRDGBufferRef Buffer)
 {
-	check(Buffer->HasRHI() || (Buffer->bCulled && !Buffer->bExternal));
-
-	if (Buffer->bCulled || !Buffer->bLastOwner)
+	if (!Buffer->HasRHI() || Buffer->bCulled || !Buffer->bLastOwner)
 	{
 		return;
 	}
 
-	if (!Buffer->bFinalizedAccess)
+	// Only buffers with internal access mode are transitioned to their epilogue state.
+	if (Buffer->AccessModeState.Mode == FRDGViewableResource::EAccessMode::Internal)
 	{
 		const FRDGPassHandle EpiloguePassHandle = GetEpiloguePassHandle();
 
@@ -2659,12 +2698,13 @@ void FRDGBuilder::AddEpilogueTransition(FRDGBufferRef Buffer)
 
 		AddTransition(Buffer->LastPass, Buffer, SubresourceState);
 	}
+	// External access mode buffers are left in their last external state.
 	else
 	{
 		Buffer->EpilogueAccess = Buffer->GetState().Access;
 	}
 
-	// Resources that were finalized on registration could have an unknown epilogue state.
+	// Resources that were marked for external access on registration (via SkipTracking) could have an unknown epilogue state.
 	if (Buffer->EpilogueAccess != ERHIAccess::Unknown)
 	{
 		EpilogueResourceAccesses.Emplace(Buffer->GetRHI(), Buffer->EpilogueAccess);
diff --git a/Engine/Source/Runtime/RenderCore/Private/RenderGraphResources.cpp b/Engine/Source/Runtime/RenderCore/Private/RenderGraphResources.cpp
--- a/Engine/Source/Runtime/RenderCore/Private/RenderGraphResources.cpp
+++ b/Engine/Source/Runtime/RenderCore/Private/RenderGraphResources.cpp
@@ -21,7 +21,7 @@ inline bool NeedsUAVBarrier(const FRDGSubresourceState& Previous, const FRDGSubr
 	return NeedsUAVBarrier(Previous.NoUAVBarrierFilter.GetUniqueHandle(), Next.NoUAVBarrierFilter.GetUniqueHandle());
 }
 
-FRDGViewableResource::FRDGViewableResource(const TCHAR* InName, const ERDGViewableResourceType InType)
+FRDGViewableResource::FRDGViewableResource(const TCHAR* InName, const ERDGViewableResourceType InType, bool bSkipTracking)
 	: FRDGResource(InName)
 	, Type(InType)
 	, bExternal(0)
@@ -30,14 +30,19 @@ FRDGViewableResource::FRDGViewableResource(const TCHAR* InName, const ERDGViewab
 	, bTransient(0)
 	, bForceNonTransient(0)
 	, TransientExtractionHint(ETransientExtractionHint::None)
-	, bFinalizedAccess(0)
 	, bLastOwner(1)
 	, bCulled(!IsImmediateMode())
 	, bUsedByAsyncComputePass(0)
 	, bQueuedForUpload(0)
 	, FirstBarrier(EFirstBarrier::Split)
 	, bUAVAccessed(0)
-{}
+{
+	if (bSkipTracking)
+	{
+		SetExternalAccessMode(ERHIAccess::ReadOnlyExclusiveMask, ERHIPipeline::All);
+		AccessModeState.bLocked = 1;
+	}
+}
 
 bool FRDGProducerState::IsDependencyRequired(FRDGProducerState LastProducer, ERHIPipeline LastPipeline, FRDGProducerState NextState, ERHIPipeline NextPipeline)
 {
diff --git a/Engine/Source/Runtime/RenderCore/Private/RenderGraphTrace.cpp b/Engine/Source/Runtime/RenderCore/Private/RenderGraphTrace.cpp
--- a/Engine/Source/Runtime/RenderCore/Private/RenderGraphTrace.cpp
+++ b/Engine/Source/Runtime/RenderCore/Private/RenderGraphTrace.cpp
@@ -59,6 +59,7 @@ UE_TRACE_EVENT_BEGIN(RDGTrace, BufferMessage)
 	UE_TRACE_EVENT_FIELD(bool, IsExternal)
 	UE_TRACE_EVENT_FIELD(bool, IsExtracted)
 	UE_TRACE_EVENT_FIELD(bool, IsCulled)
+	UE_TRACE_EVENT_FIELD(bool, IsTrackingSkipped)
 	UE_TRACE_EVENT_FIELD(bool, IsTransient)
 	UE_TRACE_EVENT_FIELD(bool, IsTransientUntracked)
 	UE_TRACE_EVENT_FIELD(bool, IsTransientCacheHit)
@@ -88,6 +89,7 @@ UE_TRACE_EVENT_BEGIN(RDGTrace, TextureMessage)
 	UE_TRACE_EVENT_FIELD(bool, IsExternal)
 	UE_TRACE_EVENT_FIELD(bool, IsExtracted)
 	UE_TRACE_EVENT_FIELD(bool, IsCulled)
+	UE_TRACE_EVENT_FIELD(bool, IsTrackingSkipped)
 	UE_TRACE_EVENT_FIELD(bool, IsTransient)
 	UE_TRACE_EVENT_FIELD(bool, IsTransientUntracked)
 	UE_TRACE_EVENT_FIELD(bool, IsTransientCacheHit)
@@ -337,6 +339,7 @@ void FRDGTrace::OutputGraphEnd(const FRDGBuilder& GraphBuilder)
 			<< TextureMessage.IsExternal(bool(Texture->bExternal))
 			<< TextureMessage.IsExtracted(bool(Texture->bExtracted))
 			<< TextureMessage.IsCulled(bool(Texture->bCulled))
+			<< TextureMessage.IsTrackingSkipped(EnumHasAnyFlags(Texture->Flags, ERDGTextureFlags::SkipTracking))
 			<< TextureMessage.IsTransient(bool(Texture->bTransient))
 			<< TextureMessage.IsTransientUntracked(false)
 			<< TextureMessage.IsTransientCacheHit(TransientAllocation.bCacheHit);
@@ -370,6 +373,7 @@ void FRDGTrace::OutputGraphEnd(const FRDGBuilder& GraphBuilder)
 			<< BufferMessage.IsExternal(bool(Buffer->bExternal))
 			<< BufferMessage.IsExtracted(bool(Buffer->bExtracted))
 			<< BufferMessage.IsCulled(bool(Buffer->bCulled))
+			<< BufferMessage.IsTrackingSkipped(EnumHasAnyFlags(Buffer->Flags, ERDGBufferFlags::SkipTracking))
 			<< BufferMessage.IsTransient(bool(Buffer->bTransient))
 			<< BufferMessage.IsTransientUntracked(false)
 			<< BufferMessage.IsTransientCacheHit(TransientAllocation.bCacheHit);
@@ -413,6 +417,7 @@ void FRDGTrace::OutputGraphEnd(const FRDGBuilder& GraphBuilder)
 				<< TextureMessage.IsExternal(false)
 				<< TextureMessage.IsExtracted(false)
 				<< TextureMessage.IsCulled(false)
+				<< TextureMessage.IsTrackingSkipped(false)
 				<< TextureMessage.IsTransient(true)
 				<< TextureMessage.IsTransientUntracked(true)
 				<< TextureMessage.IsTransientCacheHit(TransientAllocation.bCacheHit);
@@ -435,6 +440,7 @@ void FRDGTrace::OutputGraphEnd(const FRDGBuilder& GraphBuilder)
 				<< BufferMessage.IsExternal(false)
 				<< BufferMessage.IsExtracted(false)
 				<< BufferMessage.IsCulled(false)
+				<< BufferMessage.IsTrackingSkipped(false)
 				<< BufferMessage.IsTransient(true)
 				<< BufferMessage.IsTransientUntracked(true)
 				<< BufferMessage.IsTransientCacheHit(TransientAllocation.bCacheHit);
diff --git a/Engine/Source/Runtime/RenderCore/Private/RenderGraphValidation.cpp b/Engine/Source/Runtime/RenderCore/Private/RenderGraphValidation.cpp
--- a/Engine/Source/Runtime/RenderCore/Private/RenderGraphValidation.cpp
+++ b/Engine/Source/Runtime/RenderCore/Private/RenderGraphValidation.cpp
@@ -101,9 +101,6 @@ struct FRDGViewableResourceDebugData
 
 	/** Tracks whether this resource was clobbered by the builder prior to use. */
 	bool bHasBeenClobbered = false;
-
-	/** Tracks which pass performed a finalize operation on the resource. */
-	FRDGPassHandle FinalizePass;
 };
 
 FRDGViewableResourceDebugData& FRDGViewableResource::GetViewableDebugData() const
@@ -239,7 +236,6 @@ void FRDGUserValidation::ValidateRegisterExternalTexture(
 {
 	checkf(Name, TEXT("Attempted to register external texture with NULL name."));
 	checkf(ExternalPooledTexture.IsValid(), TEXT("Attempted to register NULL external texture."));
-	checkf(!EnumHasAnyFlags(Flags, ERDGTextureFlags::ReadOnly) || !EnumHasAnyFlags(Flags, ERDGTextureFlags::ForceTracking), TEXT("External texture %s cannot be ReadOnly and ForceTracking (flags are mutually exclusive)"), Name);
 	ExecuteGuard(TEXT("RegisterExternalTexture"), Name);
 }
 
@@ -247,7 +243,6 @@ void FRDGUserValidation::ValidateRegisterExternalBuffer(const TRefCountPtr<FRDGP
 {
 	checkf(Name, TEXT("Attempted to register external buffer with NULL name."));
 	checkf(ExternalPooledBuffer.IsValid(), TEXT("Attempted to register NULL external buffer."));
-	checkf(!EnumHasAnyFlags(Flags, ERDGBufferFlags::ReadOnly) || !EnumHasAnyFlags(Flags, ERDGBufferFlags::ForceTracking), TEXT("External buffer %s cannot be ReadOnly and ForceTracking (flags are mutually exclusive)"), Name);
 	ExecuteGuard(TEXT("RegisterExternalBuffer"), Name);
 }
 
@@ -279,7 +274,7 @@ void FRDGUserValidation::ValidateCreateTexture(const FRDGTextureDesc& Desc, cons
 	const bool bIsUAVForMSAATexture = bIsMSAA && bCanHaveUAV;
 	checkf(!bIsUAVForMSAATexture, TEXT("TexCreate_UAV is not allowed on MSAA texture %s."), Name);
 
-	checkf(!EnumHasAnyFlags(Flags, ERDGTextureFlags::ReadOnly), TEXT("Cannot create texture %s with the ReadOnly flag. Only registered textures can use this flag."), Name);
+	checkf(!EnumHasAnyFlags(Flags, ERDGTextureFlags::SkipTracking), TEXT("Cannot create texture %s with the SkipTracking flag. Only registered textures can use this flag."), Name);
 }
 
 void FRDGUserValidation::ValidateCreateBuffer(const FRDGBufferDesc& Desc, const TCHAR* Name, ERDGBufferFlags Flags)
@@ -287,7 +282,7 @@ void FRDGUserValidation::ValidateCreateBuffer(const FRDGBufferDesc& Desc, const
 	checkf(Name, TEXT("Creating a buffer requires a valid debug name."));
 	ExecuteGuard(TEXT("CreateBuffer"), Name);
 
-	checkf(Desc.GetTotalNumBytes() > 0, TEXT("Creating buffer '%s' is zero bytes in size."), Name);
+	checkf(Desc.GetSize() > 0, TEXT("Creating buffer '%s' is zero bytes in size."), Name);
 
 	const bool bIsByteAddress = (Desc.Usage & BUF_ByteAddressBuffer) == BUF_ByteAddressBuffer;
 
@@ -296,7 +291,7 @@ void FRDGUserValidation::ValidateCreateBuffer(const FRDGBufferDesc& Desc, const
 		checkf(Desc.BytesPerElement == 4, TEXT("Creating buffer '%s' as a structured buffer that is also byte addressable, BytesPerElement must be 4! Instead it is %d"), Name, Desc.BytesPerElement);
 	}
 
-	checkf(!EnumHasAnyFlags(Flags, ERDGBufferFlags::ReadOnly), TEXT("Cannot create buffer %s with the ReadOnly flag. Only registered buffers can use this flag."), Name);
+	checkf(!EnumHasAnyFlags(Flags, ERDGBufferFlags::SkipTracking), TEXT("Cannot create buffer %s with the SkipTracking flag. Only registered buffers can use this flag."), Name);
 }
 
 void FRDGUserValidation::ValidateCreateSRV(const FRDGTextureSRVDesc& Desc)
@@ -380,6 +375,7 @@ void FRDGUserValidation::ValidateUploadBuffer(FRDGBufferRef Buffer, const void*
 	checkf(!Buffer->bQueuedForUpload, TEXT("Buffer %s already has an upload queued. Only one upload can be done for each graph."), Buffer->Name);
 	check(InitialData || InitialDataSize == 0);
 }
+
 void FRDGUserValidation::ValidateUploadBuffer(FRDGBufferRef Buffer, const void* InitialData, uint64 InitialDataSize, const FRDGBufferInitialDataFreeCallback& InitialDataFreeCallback)
 {
 	check(Buffer);
@@ -480,29 +476,46 @@ void FRDGUserValidation::ValidateSetAccessFinal(FRDGViewableResource* Resource,
 	check(Resource);
 	check(AccessFinal != ERHIAccess::Unknown && IsValidAccess(AccessFinal));
 	checkf(Resource->bExternal || Resource->bExtracted, TEXT("Cannot set final access on non-external resource '%s' unless it is first extracted or preallocated."), Resource->Name);
-	checkf(!Resource->bFinalizedAccess, TEXT("Cannot set final access on finalized resource %s."), Resource->Name);
+	checkf(Resource->AccessModeState.Mode == FRDGViewableResource::EAccessMode::Internal, TEXT("Cannot set final access on a resource in external access mode: %s."), Resource->Name);
 }
 
-void FRDGUserValidation::ValidateFinalize(FRDGViewableResource* Resource, ERHIAccess AccessFinal, FRDGPassHandle FinalizePass)
+void FRDGUserValidation::ValidateUseExternalAccessMode(FRDGViewableResource* Resource, ERHIAccess ReadOnlyAccess, ERHIPipeline Pipelines)
 {
 	check(Resource);
-	check(AccessFinal != ERHIAccess::Unknown && IsValidAccess(AccessFinal));
-	checkf(IsReadOnlyAccess(AccessFinal), TEXT("Cannot convert resource %s to untracked with access %s. Access must be read-only."), Resource->Name, *GetRHIAccessName(AccessFinal));
-	checkf(Resource->bExternal || Resource->bExtracted, TEXT("Cannot convert resource %s to untracked unless it is first extracted or made external."), Resource->Name);
-	Resource->GetViewableDebugData().FinalizePass = FinalizePass;
+	check(Pipelines != ERHIPipeline::None);
+
+	const auto& AccessModeState = Resource->AccessModeState;
+	checkf(!AccessModeState.bLocked, TEXT("Resource is locked in external access mode by the SkipTracking resource flag: %s."), Resource->Name);
+	checkf(IsReadOnlyAccess(ReadOnlyAccess), TEXT("A read only access is required when use external access mode. (Resource: %s, Access: %s, Pipelines: %s)"), Resource->Name, *GetRHIAccessName(ReadOnlyAccess), *GetRHIPipelineName(Pipelines));
+	checkf(
+		AccessModeState.Mode == FRDGViewableResource::EAccessMode::Internal ||
+		(AccessModeState.Access == ReadOnlyAccess && AccessModeState.Pipelines == Pipelines),
+		TEXT("UseExternalAccessMode called on a resource that is already in external access mode, but different parameters were used.\n")
+		TEXT("Resource: %s\n")
+		TEXT("\tExisting Access: %s, Requested Access: %s\n")
+		TEXT("\tExisting Pipelines: %s, Requested Pipelines: %s\n"),
+		Resource->Name, *GetRHIAccessName(AccessModeState.Access), *GetRHIAccessName(ReadOnlyAccess), *GetRHIPipelineName(AccessModeState.Pipelines), *GetRHIPipelineName(Pipelines));
+}
+
+void FRDGUserValidation::ValidateUseInternalAccessMode(FRDGViewableResource* Resource)
+{
+	check(Resource);
+
+	const auto& AccessModeState = Resource->AccessModeState;
+	checkf(!AccessModeState.bLocked, TEXT("Resource is locked in external access mode by the SkipTracking resource flag: %s."), Resource->Name);
 }
 
-void FRDGUserValidation::ValidateFinalizedAccess(FRDGViewableResource* Resource, ERHIAccess Access, const FRDGPass* Pass)
+void FRDGUserValidation::ValidateExternalAccess(FRDGViewableResource* Resource, ERHIAccess Access, const FRDGPass* Pass)
 {
-	ensureMsgf(EnumHasAnyFlags(Resource->FinalizedAccessMask, Access),
-		TEXT("Resource %s was finalized with access %s, but is being used in pass %s with access %s. Any future pass must use a subset of the finalized access state."),
-		Resource->Name, *GetRHIAccessName(Resource->FinalizedAccessMask), Pass->GetName(), *GetRHIAccessName(Access), *GetRHIPipelineName(Pass->GetPipeline()));
-
-#if 0 // TODO: Need to account for read-only resources. 
-	ensureMsgf(Pass->GetPipeline() == ERHIPipeline::Graphics,
-		TEXT("Resource %s was finalized but is being used on the async compute pass %s. Only graphics pipe access is allowed for finalized resources."),
-		Resource->Name, Pass->GetName());
-#endif
+	const auto& AccessModeState = Resource->AccessModeState;
+
+	ensureMsgf(EnumHasAnyFlags(AccessModeState.Access, Access),
+		TEXT("Resource %s is in external access mode and is valid for access with the following states: %s, but is being used in pass %s with access %s."),
+		Resource->Name, *GetRHIAccessName(AccessModeState.Access), Pass->GetName(), *GetRHIAccessName(Access));
+ 
+	ensureMsgf(EnumHasAnyFlags(AccessModeState.Pipelines, Pass->GetPipeline()),
+		TEXT("Resource %s is in external access mode and is valid for access on the following pipelines: %s, but is being used on the %s pipe in pass %s."),
+		Resource->Name, *GetRHIPipelineName(AccessModeState.Pipelines), *GetRHIPipelineName(Pass->GetPipeline()), Pass->GetName());
 }
 
 void FRDGUserValidation::ValidateAddPass(const FRDGEventName& Name, ERDGPassFlags Flags)
@@ -561,18 +574,6 @@ void FRDGUserValidation::ValidateAddPass(const FRDGPass* Pass, bool bSkipPassAcc
 		}
 	};
 
-	const auto MarkTextureAsProduced = [&](FRDGTextureRef Texture)
-	{
-		checkf(!EnumHasAnyFlags(Texture->Flags, ERDGTextureFlags::ReadOnly), TEXT("Pass %s is attempting to write to texture %s which is marked as ReadOnly."), Pass->GetName(), Texture->Name);
-		MarkAsProduced(Texture);
-	};
-
-	const auto MarkBufferAsProduced = [&](FRDGBufferRef Buffer)
-	{
-		checkf(!EnumHasAnyFlags(Buffer->Flags, ERDGBufferFlags::ReadOnly), TEXT("Pass %s is attempting to write to buffer %s which is marked as ReadOnly."), Pass->GetName(), Buffer->Name);
-		MarkAsProduced(Buffer);
-	};
-
 	const auto MarkAsConsumed = [&] (FRDGViewableResource* Resource)
 	{
 		ensureMsgf(Resource->bProduced || Resource->bExternal || Resource->bQueuedForUpload,
@@ -611,7 +612,7 @@ void FRDGUserValidation::ValidateAddPass(const FRDGPass* Pass, bool bSkipPassAcc
 
 		if (IsWritableAccess(Access))
 		{
-			MarkBufferAsProduced(Buffer);
+			MarkAsProduced(Buffer);
 			bCanProduce = true;
 		}
 	};
@@ -622,7 +623,7 @@ void FRDGUserValidation::ValidateAddPass(const FRDGPass* Pass, bool bSkipPassAcc
 
 		if (IsWritableAccess(Access))
 		{
-			MarkTextureAsProduced(Texture);
+			MarkAsProduced(Texture);
 			bCanProduce = true;
 		}
 	};
@@ -664,7 +665,7 @@ void FRDGUserValidation::ValidateAddPass(const FRDGPass* Pass, bool bSkipPassAcc
 			{
 				FRDGTextureRef Texture = UAV->GetParent();
 				CheckNotCopy(Texture);
-				MarkTextureAsProduced(Texture);
+				MarkAsProduced(Texture);
 			}
 		}
 		break;
@@ -685,7 +686,7 @@ void FRDGUserValidation::ValidateAddPass(const FRDGPass* Pass, bool bSkipPassAcc
 			{
 				FRDGBufferRef Buffer = UAV->GetParent();
 				CheckNotCopy(Buffer);
-				MarkBufferAsProduced(Buffer);
+				MarkAsProduced(Buffer);
 			}
 		}
 		break;
@@ -772,7 +773,7 @@ void FRDGUserValidation::ValidateAddPass(const FRDGPass* Pass, bool bSkipPassAcc
 				CheckValidResource(Texture);
 				if (DepthStencil.GetDepthStencilAccess().IsAnyWrite())
 				{
-					MarkTextureAsProduced(Texture);
+					MarkAsProduced(Texture);
 				}
 				else
 				{
@@ -816,7 +817,7 @@ void FRDGUserValidation::ValidateAddPass(const FRDGPass* Pass, bool bSkipPassAcc
 						PassName, ResolveTexture->Name);
 
 					CheckValidResource(Texture);
-					MarkTextureAsProduced(ResolveTexture);
+					MarkAsProduced(ResolveTexture);
 				}
 
 				if (Texture)
@@ -829,7 +830,7 @@ void FRDGUserValidation::ValidateAddPass(const FRDGPass* Pass, bool bSkipPassAcc
 					CheckValidResource(Texture);
 
 					/** Mark the pass as a producer for render targets with a store action. */
-					MarkTextureAsProduced(Texture);
+					MarkAsProduced(Texture);
 				}
 				else
 				{
diff --git a/Engine/Source/Runtime/RenderCore/Public/RenderGraphBuilder.h b/Engine/Source/Runtime/RenderCore/Public/RenderGraphBuilder.h
--- a/Engine/Source/Runtime/RenderCore/Public/RenderGraphBuilder.h
+++ b/Engine/Source/Runtime/RenderCore/Public/RenderGraphBuilder.h
@@ -245,32 +245,19 @@ public:
 	/** (External | Extracted only) Sets the access to transition to after execution at the end of the graph. Overwrites any previously set final access. */
 	void SetBufferAccessFinal(FRDGBufferRef Buffer, ERHIAccess Access);
 
-	/** Finalizes the access of multiple resources so that they are immutable for the remainder of the graph. This immediately forces all resources into
-	 *  their finalized states. The resources cannot be used in any other state within the graph and must be used on the graphics pipe. This is designed
-	 *  for complex cases where resources are produced early in the graph and transitioned to a read-only state. Finalized access resources can be used
-	 *  outside of RDG pass parameters on any future pass without invoking RHI validation failures.
+	/** Configures the resource for external access for all subsequent passes, or until UseInternalAccessMode is called.
+	 *  Only read-only access states are allowed. When in external access mode, it is safe to access the underlying RHI
+	 *  resource directly in later RDG passes. This method is only allowed for registered or externally converted resources.
+	 *  The method effectively guarantees that RDG will transition the resource into the desired state for all subsequent
+	 *  passes so long as the resource remains externally accessible.
 	 */
-	void FinalizeResourceAccess(FRDGTextureAccessArray&& Textures, FRDGBufferAccessArray&& Buffers);
+	void UseExternalAccessMode(FRDGViewableResource* Resource, ERHIAccess ReadOnlyAccess, ERHIPipeline Pipelines = ERHIPipeline::Graphics);
 
-	inline void FinalizeTextureAccess(FRDGTextureAccessArray&& InTextures)
-	{
-		FinalizeResourceAccess(Forward<FRDGTextureAccessArray&&>(InTextures), {});
-	}
-
-	inline void FinalizeBufferAccess(FRDGBufferAccessArray&& InBuffers)
-	{
-		FinalizeResourceAccess({}, Forward<FRDGBufferAccessArray&&>(InBuffers));
-	}
-
-	inline void FinalizeTextureAccess(FRDGTextureRef Texture, ERHIAccess Access)
-	{
-		FinalizeResourceAccess({ FRDGTextureAccess(Texture, Access) }, {});
-	}
-
-	inline void FinalizeBufferAccess(FRDGBufferRef Buffer, ERHIAccess Access)
-	{
-		FinalizeResourceAccess({}, { FRDGBufferAccess(Buffer, Access) });
-	}
+	/** Use this method to resume tracking of a resource after calling UseExternalAccessMode. It is safe to call this method
+	 *  even if external access mode was not enabled (it will simply no-op). It is not valid to access the underlying RHI
+	 *  resource in any pass added after calling this method.
+	 */
+	void UseInternalAccessMode(FRDGViewableResource* Resource);
 
 	/** Flag a resource that is produced by a pass but never used or extracted to not emit an 'unused' warning. */
 	void RemoveUnusedTextureWarning(FRDGTextureRef Texture);
@@ -362,6 +349,53 @@ PRAGMA_ENABLE_DEPRECATION_WARNINGS
 	{
 		return FindExternalTexture(ExternalPooledTexture);
 	}
+
+	UE_DEPRECATED(5.1, "FinalizeResourceAccess has been replaced by UseExternalAccessMode")
+	inline void FinalizeResourceAccess(FRDGTextureAccessArray&& InTextures, FRDGBufferAccessArray&& InBuffers)
+	{
+		for (FRDGTextureAccess Texture : InTextures)
+		{
+			UseExternalAccessMode(Texture.GetTexture(), Texture.GetAccess());
+		}
+
+		for (FRDGBufferAccess Buffer : InBuffers)
+		{
+			UseExternalAccessMode(Buffer.GetBuffer(), Buffer.GetAccess());
+		}
+	}
+
+	UE_DEPRECATED(5.1, "FinalizeResourceAccess has been replaced by UseExternalAccessMode")
+	inline void FinalizeTextureAccess(FRDGTextureAccessArray&& InTextures)
+	{
+	PRAGMA_DISABLE_DEPRECATION_WARNINGS
+		FinalizeResourceAccess(Forward<FRDGTextureAccessArray&&>(InTextures), {});
+	PRAGMA_ENABLE_DEPRECATION_WARNINGS
+	}
+
+	UE_DEPRECATED(5.1, "FinalizeResourceAccess has been replaced by UseExternalAccessMode")
+	inline void FinalizeBufferAccess(FRDGBufferAccessArray&& InBuffers)
+	{
+	PRAGMA_DISABLE_DEPRECATION_WARNINGS
+		FinalizeResourceAccess({}, Forward<FRDGBufferAccessArray&&>(InBuffers));
+	PRAGMA_ENABLE_DEPRECATION_WARNINGS
+	}
+
+	UE_DEPRECATED(5.1, "FinalizeResourceAccess has been replaced by UseExternalAccessMode")
+	inline void FinalizeTextureAccess(FRDGTextureRef Texture, ERHIAccess Access)
+	{
+	PRAGMA_DISABLE_DEPRECATION_WARNINGS
+		FinalizeResourceAccess({ FRDGTextureAccess(Texture, Access) }, {});
+	PRAGMA_ENABLE_DEPRECATION_WARNINGS
+	}
+
+	UE_DEPRECATED(5.1, "FinalizeResourceAccess has been replaced by UseExternalAccessMode")
+	inline void FinalizeBufferAccess(FRDGBufferRef Buffer, ERHIAccess Access)
+	{
+	PRAGMA_DISABLE_DEPRECATION_WARNINGS
+		FinalizeResourceAccess({}, { FRDGBufferAccess(Buffer, Access) });
+	PRAGMA_ENABLE_DEPRECATION_WARNINGS
+	}
+
 	//////////////////////////////////////////////////////////////////////////
 
 private:
@@ -603,6 +637,9 @@ private:
 	/** Tracks the final access used on resources in order to call SetTrackedAccess. */
 	TArray<FRHITrackedAccessInfo, FRDGArrayAllocator> EpilogueResourceAccesses;
 
+	/** Contains resources queued for either access mode change passes. */
+	TArray<FRDGViewableResource*, FRDGArrayAllocator> AccessModeQueue;
+
 	/** Texture state used for intermediate operations. Held here to avoid re-allocating. */
 	FRDGTextureSubresourceStateIndirect ScratchTextureState;
 
@@ -619,6 +656,7 @@ private:
 	bool bDispatchHint = false;
 	bool bFlushResourcesRHI = false;
 	bool bParallelExecuteEnabled = false;
+	bool bInAccessModeQueueFlush = false;
 
 #if RDG_ENABLE_DEBUG
 	FRDGUserValidation UserValidation;
@@ -678,6 +716,8 @@ private:
 	void BeginFlushResourcesRHI();
 	void EndFlushResourcesRHI();
 
+	void FlushAccessModeQueue();
+
 	void SetupPassInternal(FRDGPass* Pass, FRDGPassHandle PassHandle, ERHIPipeline PassPipeline, bool bEmptyParameters);
 	FRDGPass* SetupPass(FRDGPass* Pass);
 	FRDGPass* SetupEmptyPass(FRDGPass* Pass);
diff --git a/Engine/Source/Runtime/RenderCore/Public/RenderGraphBuilder.inl b/Engine/Source/Runtime/RenderCore/Public/RenderGraphBuilder.inl
--- a/Engine/Source/Runtime/RenderCore/Public/RenderGraphBuilder.inl
+++ b/Engine/Source/Runtime/RenderCore/Public/RenderGraphBuilder.inl
@@ -197,6 +197,8 @@ FRDGPassRef FRDGBuilder::AddPass(
 
 	Flags |= ERDGPassFlags::NeverCull;
 
+	FlushAccessModeQueue();
+
 	LambdaPassType* Pass = Passes.Allocate<LambdaPassType>(Allocator, MoveTemp(Name), Flags, MoveTemp(ExecuteLambda));
 	SetupEmptyPass(Pass);
 	return Pass;
@@ -214,6 +216,8 @@ FRDGPassRef FRDGBuilder::AddPassInternal(
 
 	IF_RDG_ENABLE_DEBUG(UserValidation.ValidateAddPass(ParameterStruct, ParametersMetadata, Name, Flags));
 
+	FlushAccessModeQueue();
+
 	FRDGPass* Pass = Allocator.AllocNoDestruct<LambdaPassType>(
 		MoveTemp(Name),
 		ParametersMetadata,
diff --git a/Engine/Source/Runtime/RenderCore/Public/RenderGraphDefinitions.h b/Engine/Source/Runtime/RenderCore/Public/RenderGraphDefinitions.h
--- a/Engine/Source/Runtime/RenderCore/Public/RenderGraphDefinitions.h
+++ b/Engine/Source/Runtime/RenderCore/Public/RenderGraphDefinitions.h
@@ -131,12 +131,16 @@ enum class ERDGBufferFlags : uint8
 	/** Tag the buffer to survive through frame, that is important for multi GPU alternate frame rendering. */
 	MultiFrame = 1 << 0,
 
-	/** The buffer may only be used for read-only access within the graph. This flag is only allowed for registered buffers. */
-	ReadOnly = 1 << 1, 
+	/** The buffer is ignored by RDG tracking and will never be transitioned. Use the flag when registering a buffer with no writable GPU flags.
+	 *  Write access is not allowed for the duration of the graph. This flag is intended as an optimization to cull out tracking of read-only
+	 *  buffers that are used frequently throughout the graph. Note that it's the user's responsibility to ensure the resource is in the correct
+	 *  readable state for use with RDG passes, as RDG does not know the exact state of the resource.
+	 */
+	SkipTracking = 1 << 1,
 
-	/** Force the graph to track this resource even if it can be considered as readonly (no UAV, no RTV, etc.) This allows the graph copying from and to textures, and handling the corresponding transitions, for example.
-	 This flag is only allowed for registered buffers. Mutually exclusive with ReadOnly. */
-	ForceTracking = 1 << 2,
+	// Deprecated Enums
+	ReadOnly      UE_DEPRECATED(5.1, "ReadOnly is deprecated. Use SkipTracking instead.")                                            = 0,
+	ForceTracking UE_DEPRECATED(5.1, "ForceTracking is deprecated. Resources now opt-out of tracking explicitly with SkipTracking.") = 0
 };
 ENUM_CLASS_FLAGS(ERDGBufferFlags);
 
@@ -148,15 +152,19 @@ enum class ERDGTextureFlags : uint8
 	/** Tag the texture to survive through frame, that is important for multi GPU alternate frame rendering. */
 	MultiFrame = 1 << 0,
 
-	/** The texture may only be used for read-only access within the graph. This flag is only allowed for registered textures. */
-	ReadOnly = 1 << 1,
+	/** The buffer is ignored by RDG tracking and will never be transitioned. Use the flag when registering a buffer with no writable GPU flags.
+	 *  Write access is not allowed for the duration of the graph. This flag is intended as an optimization to cull out tracking of read-only
+	 *  buffers that are used frequently throughout the graph. Note that it's the user's responsibility to ensure the resource is in the correct
+	 *  readable state for use with RDG passes, as RDG does not know the exact state of the resource.
+	 */
+	SkipTracking = 1 << 1,
 
 	/** Prevents metadata decompression on this texture. */
 	MaintainCompression = 1 << 2,
 
-	/** Force the graph to track this resource even if it can be considered as readonly (no UAV, no RTV, etc.) This allows the graph copying from and to textures, and handling the corresponding transitions, for example.
-	 This flag is only allowed for registered buffers. Mutually exclusive with ReadOnly. */
-	ForceTracking = 1 << 3,
+	// Deprecated Enums
+	ReadOnly      UE_DEPRECATED(5.1, "ReadOnly is deprecated. Use SkipTracking instead.")                                            = 0,
+	ForceTracking UE_DEPRECATED(5.1, "ForceTracking is deprecated. Resources now opt-out of tracking explicitly with SkipTracking.") = 0
 };
 ENUM_CLASS_FLAGS(ERDGTextureFlags);
 
diff --git a/Engine/Source/Runtime/RenderCore/Public/RenderGraphResources.h b/Engine/Source/Runtime/RenderCore/Public/RenderGraphResources.h
--- a/Engine/Source/Runtime/RenderCore/Public/RenderGraphResources.h
+++ b/Engine/Source/Runtime/RenderCore/Public/RenderGraphResources.h
@@ -267,7 +267,7 @@ public:
 	}
 
 protected:
-	FRDGViewableResource(const TCHAR* InName, ERDGViewableResourceType InType);
+	FRDGViewableResource(const TCHAR* InName, ERDGViewableResourceType InType, bool bSkipTracking);
 
 	enum class ETransientExtractionHint : uint8
 	{
@@ -283,6 +283,33 @@ protected:
 		ImmediateConfirmed
 	};
 
+	enum class EAccessMode : uint8
+	{
+		Internal,
+		External
+	};
+
+	struct FAccessModeState
+	{
+		using FQueueIndex = TRDGHandle<FAccessModeState, uint16>;
+
+		bool IsQueued() const { return QueueIndex.IsValid(); }
+		bool IsExternalAccess() const { return Mode == EAccessMode::External && !IsQueued(); }
+
+		FAccessModeState()
+			: Pipelines(ERHIPipeline::None)
+			, Mode(EAccessMode::Internal)
+			, bLocked(0)
+		{}
+
+		ERHIAccess			Access			= ERHIAccess::None;
+		ERHIPipeline		Pipelines		: 2;
+		EAccessMode			Mode			: 1;
+		uint8				bLocked			: 1;
+		FQueueIndex QueueIndex;
+
+	} AccessModeState;
+
 	/** Whether this is an externally registered resource. */
 	uint8 bExternal : 1;
 
@@ -301,9 +328,6 @@ protected:
 	/** Whether this resource is allowed to be both transient and extracted. */
 	ETransientExtractionHint TransientExtractionHint : 2;
 
-	/** (External | Extracted only) If true, the resource is locked in its current state and will not be transitioned any more. */
-	uint8 bFinalizedAccess : 1;
-
 	/** Whether this resource is the last owner of its allocation (i.e. nothing aliases the allocation later in the execution timeline). */
 	uint8 bLastOwner : 1;
 
@@ -325,6 +349,9 @@ protected:
 	FRDGPassHandle FirstPass;
 	FRDGPassHandle LastPass;
 
+	/** The state of the resource at the graph epilogue. */
+	ERHIAccess EpilogueAccess = ERHIAccess::SRVMask;
+
 private:
 	/** Number of references in passes and deferred queries. */
 	uint16 ReferenceCount = 0;
@@ -332,20 +359,16 @@ private:
 	/** Scratch index allocated for the resource in the pass being setup. */
 	uint16 PassStateIndex = 0;
 
-	/** The state of the resource at the graph epilogue. */
-	ERHIAccess EpilogueAccess = ERHIAccess::SRVMask;
-
-	/** The valid set of access states for finalized access. */
-	IF_RDG_ENABLE_DEBUG(ERHIAccess FinalizedAccessMask = ERHIAccess::Unknown);
-
-	void SetFinalizedAccess(ERHIAccess InValidAccessMask)
+	void SetExternalAccessMode(ERHIAccess InReadOnlyAccess, ERHIPipeline InPipelines)
 	{
-		bFinalizedAccess = 1;
+		check(!AccessModeState.bLocked);
 
-		// Finalized resources are not always added to the pass states (unless marked as such within the graph), so marked as not culled here.
-		bCulled = 0;
+		AccessModeState.Mode = EAccessMode::External;
+		AccessModeState.Access = InReadOnlyAccess;
+		AccessModeState.Pipelines = InPipelines;
 
-		IF_RDG_ENABLE_DEBUG(FinalizedAccessMask = InValidAccessMask);
+		// External access resources are not always added to the pass states (unless marked as such within the graph), so marked as not culled here.
+		bCulled = 0;
 	}
 
 #if RDG_ENABLE_TRACE
@@ -491,7 +514,7 @@ public:
 
 private:
 	FRDGTexture(const TCHAR* InName, const FRDGTextureDesc& InDesc, ERDGTextureFlags InFlags)
-		: FRDGViewableResource(InName, ERDGViewableResourceType::Texture)
+		: FRDGViewableResource(InName, ERDGViewableResourceType::Texture, EnumHasAnyFlags(InFlags, ERDGTextureFlags::SkipTracking))
 		, Desc(InDesc)
 		, Flags(InFlags)
 		, Layout(InDesc)
@@ -507,6 +530,11 @@ private:
 		{
 			FirstBarrier = EFirstBarrier::ImmediateRequested;
 		}
+
+		if (EnumHasAnyFlags(Desc.Flags, ETextureCreateFlags::Foveation))
+		{
+			EpilogueAccess = ERHIAccess::ShadingRateSource;
+		}
 	}
 
 	/** Returns RHI texture without access checks. */
@@ -944,7 +972,7 @@ struct FRDGBufferDesc
 		return BytesPerElement * NumElements;
 	}
 
-	//UE_DEPRECATED(5.1, "GetTotalNumBytes is deprecated, use GetSize instead.")
+	UE_DEPRECATED(5.1, "GetTotalNumBytes is deprecated, use GetSize instead.")
 	uint32 GetTotalNumBytes() const
 	{
 		return BytesPerElement * NumElements;
@@ -1188,7 +1216,7 @@ public:
 
 private:
 	FRDGBuffer(const TCHAR* InName, const FRDGBufferDesc& InDesc, ERDGBufferFlags InFlags)
-		: FRDGViewableResource(InName, ERDGViewableResourceType::Buffer)
+		: FRDGViewableResource(InName, ERDGViewableResourceType::Buffer, EnumHasAnyFlags(InFlags, ERDGBufferFlags::SkipTracking))
 		, Desc(InDesc)
 		, Flags(InFlags)
 	{}
diff --git a/Engine/Source/Runtime/RenderCore/Public/RenderGraphUtils.h b/Engine/Source/Runtime/RenderCore/Public/RenderGraphUtils.h
--- a/Engine/Source/Runtime/RenderCore/Public/RenderGraphUtils.h
+++ b/Engine/Source/Runtime/RenderCore/Public/RenderGraphUtils.h
@@ -1103,96 +1103,133 @@
 	GraphBuilder.AddPass(MoveTemp(Name), PassParameters, ERDGPassFlags::Readback, MoveTemp(ExecuteLambda));
 }
 
-/** Batches up RDG resource access finalizations and submits them all at once to RDG. */
-class FRDGResourceAccessFinalizer
+/** Batches up RDG external resource access mode requests and submits them all at once to RDG. */
+class FRDGExternalAccessQueue
 {
 public:
-	FRDGResourceAccessFinalizer() = default;
-	~FRDGResourceAccessFinalizer()
+	FRDGExternalAccessQueue() = default;
+
+	~FRDGExternalAccessQueue()
 	{
-		checkf(IsEmpty(), TEXT("Finalize must be called before destruction."));
+		checkf(IsEmpty(), TEXT("Submit must be called before destruction."));
 	}
 
+	void Reserve(uint32 ResourceCount)
+	{
+		Resources.Reserve(ResourceCount);
+	}
+
+	void Add(FRDGViewableResource* Resource, ERHIAccess Access = ERHIAccess::SRVMask, ERHIPipeline Pipelines = ERHIPipeline::Graphics)
+	{
+		if (Resource)
+		{
+			checkf(IsValidAccess(Access) && Access != ERHIAccess::Unknown, TEXT("Attempted to finalize texture %s with an invalid access %s."), Resource->Name, *GetRHIAccessName(Access));
+			Resources.Emplace(Resource, Access, Pipelines);
+		}
+	}
+
+	void Submit(FRDGBuilder& GraphBuilder)
+	{
+		for (FResource Resource : Resources)
+		{
+			GraphBuilder.UseExternalAccessMode(Resource.Resource, Resource.Access, Resource.Pipelines);
+		}
+		Resources.Reset();
+	}
+
+	bool IsEmpty() const
+	{
+		return Resources.IsEmpty();
+	}
+
+	UE_DEPRECATED(5.1, "Reserve with two parameters is deprecated.")
 	void Reserve(uint32 TextureCount, uint32 BufferCount)
 	{
-		Textures.Reserve(TextureCount);
-		Buffers.Reserve(BufferCount);
+		Resources.Reserve(TextureCount + BufferCount);
 	}
 
+	UE_DEPRECATED(5.1, "Finalize is deprecated. Call Submit instead.")
+	void Finalize(FRDGBuilder& GraphBuilder)
+	{
+		Submit(GraphBuilder);
+	}
+
+	UE_DEPRECATED(5.1, "AddTexture is deprecated. Call Add instead.")
 	void AddTexture(FRDGTextureRef Texture, ERHIAccess Access)
 	{
-		if (Texture)
-		{
-			checkf(IsValidAccess(Access) && Access != ERHIAccess::Unknown, TEXT("Attempted to finalize texture %s with an invalid access %s."), Texture->Name, *GetRHIAccessName(Access));
-			Textures.Emplace(Texture, Access);
-		}
+		Add(Texture, Access);
 	}
 
+	UE_DEPRECATED(5.1, "AddBuffer is deprecated. Call Add instead.")
 	void AddBuffer(FRDGBufferRef Buffer, ERHIAccess Access)
 	{
-		if (Buffer)
-		{
-			checkf(IsValidAccess(Access) && Access != ERHIAccess::Unknown, TEXT("Attempted to finalize buffer %s with an invalid access %s."), Buffer->Name, *GetRHIAccessName(Access));
-			Buffers.Emplace(Buffer, Access);
-		}
+		Add(Buffer, Access);
 	}
 
-	void Finalize(FRDGBuilder& GraphBuilder)
-	{
-		if (!IsEmpty())
-		{
-			GraphBuilder.FinalizeResourceAccess(MoveTemp(Textures), MoveTemp(Buffers));
-		}
-	}
-
-	bool IsEmpty() const
-	{
-		return Textures.IsEmpty() && Buffers.IsEmpty();
-	}
-
 private:
-	FRDGTextureAccessArray Textures;
-	FRDGBufferAccessArray Buffers;
+	struct FResource
+	{
+		FResource() = default;
+
+		FResource(FRDGViewableResource* InResource, ERHIAccess InAccess, ERHIPipeline InPipelines)
+			: Resource(InResource)
+			, Access(InAccess)
+			, Pipelines(InPipelines)
+		{}
+
+		FRDGViewableResource* Resource;
+		ERHIAccess Access;
+		ERHIPipeline Pipelines;
+	};
+
+	TArray<FResource, FRDGArrayAllocator> Resources;
 };
 
-inline const TRefCountPtr<IPooledRenderTarget>& ConvertToFinalizedExternalTexture(
+UE_DEPRECATED(5.1, "FRDGResourceAccessFinalizer is deprecated. Use FRDGExternalAccessQueue instead.")
+typedef FRDGExternalAccessQueue FRDGResourceAccessFinalizer;
+
+inline const TRefCountPtr<IPooledRenderTarget>& ConvertToExternalAccessTexture(
 	FRDGBuilder& GraphBuilder,
-	FRDGResourceAccessFinalizer& ResourceAccessFinalizer,
-	FRDGTextureRef Texture,
-	ERHIAccess AccessFinal = ERHIAccess::SRVMask)
+	FRDGTexture* Texture,
+	ERHIAccess Access = ERHIAccess::SRVMask,
+	ERHIPipeline Pipelines = ERHIPipeline::Graphics)
 {
-	ResourceAccessFinalizer.AddTexture(Texture, AccessFinal);
-	return GraphBuilder.ConvertToExternalTexture(Texture);
-}
-
-inline const TRefCountPtr<FRDGPooledBuffer>& ConvertToFinalizedExternalBuffer(
-	FRDGBuilder& GraphBuilder,
-	FRDGResourceAccessFinalizer& ResourceAccessFinalizer,
-	FRDGBufferRef Buffer,
-	ERHIAccess AccessFinal = ERHIAccess::SRVMask)
-{
-	ResourceAccessFinalizer.AddBuffer(Buffer, AccessFinal);
-	return GraphBuilder.ConvertToExternalBuffer(Buffer);
-}
-
-inline const TRefCountPtr<IPooledRenderTarget>& ConvertToFinalizedExternalTexture(
-	FRDGBuilder& GraphBuilder,
-	FRDGTextureRef Texture,
-	ERHIAccess AccessFinal = ERHIAccess::SRVMask)
-{
 	const TRefCountPtr<IPooledRenderTarget>& PooledTexture = GraphBuilder.ConvertToExternalTexture(Texture);
-	GraphBuilder.FinalizeTextureAccess(Texture, AccessFinal);
+	GraphBuilder.UseExternalAccessMode(Texture, Access, Pipelines);
 	return PooledTexture;
 }
 
-inline const TRefCountPtr<FRDGPooledBuffer>& ConvertToFinalizedExternalBuffer(
+inline const TRefCountPtr<FRDGPooledBuffer>& ConvertToExternalAccessBuffer(
 	FRDGBuilder& GraphBuilder,
-	FRDGBufferRef Buffer,
-	ERHIAccess AccessFinal)
+	FRDGBuffer* Buffer,
+	ERHIAccess Access = ERHIAccess::SRVMask,
+	ERHIPipeline Pipelines = ERHIPipeline::Graphics)
 {
 	const TRefCountPtr<FRDGPooledBuffer>& PooledBuffer = GraphBuilder.ConvertToExternalBuffer(Buffer);
-	GraphBuilder.FinalizeBufferAccess(Buffer, AccessFinal);
+	GraphBuilder.UseExternalAccessMode(Buffer, Access, Pipelines);
 	return PooledBuffer;
 }
 
+inline const TRefCountPtr<IPooledRenderTarget>& ConvertToExternalAccessTexture(
+	FRDGBuilder& GraphBuilder,
+	FRDGExternalAccessQueue& ExternalAccessQueue,
+	FRDGTexture* Texture,
+	ERHIAccess Access = ERHIAccess::SRVMask,
+	ERHIPipeline Pipelines = ERHIPipeline::Graphics)
+{
+	ExternalAccessQueue.Add(Texture, Access, Pipelines);
+	return GraphBuilder.ConvertToExternalTexture(Texture);
+}
+
+inline const TRefCountPtr<FRDGPooledBuffer>& ConvertToExternalAccessBuffer(
+	FRDGBuilder& GraphBuilder,
+	FRDGExternalAccessQueue& ExternalAccessQueue,
+	FRDGBuffer* Buffer,
+	ERHIAccess Access = ERHIAccess::SRVMask,
+	ERHIPipeline Pipelines = ERHIPipeline::Graphics)
+{
+	ExternalAccessQueue.Add(Buffer, Access, Pipelines);
+	return GraphBuilder.ConvertToExternalBuffer(Buffer);
+}
+
 /** Scope used to wait for outstanding tasks when the scope destructor is called. Used for command list recording tasks. */
@@ -1226,12 +1263,6 @@
 	const TCHAR* Name,
 	ERDGPooledBufferAlignment Alignment = ERDGPooledBufferAlignment::Page);
 
-//UE_DEPRECATED(5.1, "GetPooledFreeBuffer is deprecated. Use AllocatePooledBuffer instead.")
-inline bool GetPooledFreeBuffer(FRHICommandList& RHICmdList, const FRDGBufferDesc& Desc, TRefCountPtr<FRDGPooledBuffer>& Out, const TCHAR* Name)
-{
-	return AllocatePooledBuffer(Desc, Out, Name);
-}
-
 RENDERCORE_API bool AllocatePooledTexture(
 	const FRDGTextureDesc& Desc,
 	TRefCountPtr<IPooledRenderTarget>& Out,
@@ -1242,36 +1273,48 @@
 //////////////////////////////////////////////////////////////////////////
 //! Deprecated Functions
 
-UE_DEPRECATED(5.0, "ConvertToExternalBuffer has been refactored to FRDGBuilder::ConvertToExternalBuffer.")
-inline void ConvertToExternalBuffer(FRDGBuilder& GraphBuilder, FRDGBufferRef Buffer, TRefCountPtr<FRDGPooledBuffer>& OutPooledBuffer)
+UE_DEPRECATED(5.1, "GetPooledFreeBuffer is deprecated. Use AllocatePooledBuffer instead.")
+inline bool GetPooledFreeBuffer(FRHICommandList& RHICmdList, const FRDGBufferDesc& Desc, TRefCountPtr<FRDGPooledBuffer>& Out, const TCHAR* Name)
 {
-	OutPooledBuffer = GraphBuilder.ConvertToExternalBuffer(Buffer);
+	return AllocatePooledBuffer(Desc, Out, Name);
 }
 
-UE_DEPRECATED(5.0, "ConvertToExternalTexture has been refactored to FRDGBuilder::ConvertToExternalTexture.")
-inline void ConvertToExternalTexture(FRDGBuilder& GraphBuilder, FRDGTextureRef Texture, TRefCountPtr<IPooledRenderTarget>& OutPooledRenderTarget)
+UE_DEPRECATED(5.1, "ConvertToFinalizedExternalTexture has been refactored to ConvertToExternalAccessTexture.")
+inline const TRefCountPtr<IPooledRenderTarget>& ConvertToFinalizedExternalTexture(
+	FRDGBuilder& GraphBuilder,
+	FRDGExternalAccessQueue& ExternalAccessQueue,
+	FRDGTextureRef Texture,
+	ERHIAccess AccessFinal = ERHIAccess::SRVMask)
 {
-	OutPooledRenderTarget = GraphBuilder.ConvertToExternalTexture(Texture);
+	return ConvertToExternalAccessTexture(GraphBuilder, ExternalAccessQueue, Texture, AccessFinal);
 }
 
-UE_DEPRECATED(5.0, "ConvertToUntrackedExternalTexture has been refactored to ConvertToFinalizedExternalTexture.")
-inline void ConvertToUntrackedExternalTexture(
+UE_DEPRECATED(5.1, "ConvertToFinalizedExternalBuffer has been refactored to ConvertToExternalAccessBuffer.")
+inline const TRefCountPtr<FRDGPooledBuffer>& ConvertToFinalizedExternalBuffer(
+	FRDGBuilder& GraphBuilder,
+	FRDGExternalAccessQueue& ExternalAccessQueue,
+	FRDGBufferRef Buffer,
+	ERHIAccess AccessFinal = ERHIAccess::SRVMask)
+{
+	return ConvertToExternalAccessBuffer(GraphBuilder, ExternalAccessQueue, Buffer, AccessFinal);
+}
+
+UE_DEPRECATED(5.1, "ConvertToFinalizedExternalTexture has been refactored to ConvertToExternalAccessTexture.")
+inline const TRefCountPtr<IPooledRenderTarget>& ConvertToFinalizedExternalTexture(
 	FRDGBuilder& GraphBuilder,
 	FRDGTextureRef Texture,
-	TRefCountPtr<IPooledRenderTarget>& OutPooledRenderTarget,
+	ERHIAccess AccessFinal = ERHIAccess::SRVMask)
+{
+	return ConvertToExternalAccessTexture(GraphBuilder, Texture, AccessFinal);
+}
+
+UE_DEPRECATED(5.1, "ConvertToFinalizedExternalBuffer has been refactored to ConvertToExternalAccessBuffer.")
+inline const TRefCountPtr<FRDGPooledBuffer>& ConvertToFinalizedExternalBuffer(
+	FRDGBuilder& GraphBuilder,
+	FRDGBufferRef Buffer,
 	ERHIAccess AccessFinal)
 {
-	OutPooledRenderTarget = ConvertToFinalizedExternalTexture(GraphBuilder, Texture, AccessFinal);
+	return ConvertToExternalAccessBuffer(GraphBuilder, Buffer, AccessFinal);
 }
 
-UE_DEPRECATED(5.0, "ConvertToUntrackedExternalTexture has been refactored to ConvertToFinalizedExternalBuffer.")
-inline void ConvertToUntrackedExternalBuffer(
-	FRDGBuilder& GraphBuilder,
-	FRDGBufferRef Buffer,
-	TRefCountPtr<FRDGPooledBuffer>& OutPooledBuffer,
-	ERHIAccess AccessFinal)
-{
-	OutPooledBuffer = ConvertToFinalizedExternalBuffer(GraphBuilder, Buffer, AccessFinal);
-}
-
 UE_DEPRECATED(5.0, "RegisterExternalTextureWithFallback no longer requires ERenderTargetTexture")
diff --git a/Engine/Source/Runtime/RenderCore/Public/RenderGraphValidation.h b/Engine/Source/Runtime/RenderCore/Public/RenderGraphValidation.h
--- a/Engine/Source/Runtime/RenderCore/Public/RenderGraphValidation.h
+++ b/Engine/Source/Runtime/RenderCore/Public/RenderGraphValidation.h
@@ -97,8 +97,10 @@ public:
 
 	void ValidateSetAccessFinal(FRDGViewableResource* Resource, ERHIAccess AccessFinal);
 
-	void ValidateFinalize(FRDGViewableResource* Resource, ERHIAccess Access, FRDGPassHandle ConvertToUntrackedPassHandle);
-	void ValidateFinalizedAccess(FRDGViewableResource* Resource, ERHIAccess Access, const FRDGPass* Pass);
+	void ValidateUseExternalAccessMode(FRDGViewableResource* Resource, ERHIAccess ReadOnlyAccess, ERHIPipeline Pipelines);
+	void ValidateUseInternalAccessMode(FRDGViewableResource* Resaource);
+
+	void ValidateExternalAccess(FRDGViewableResource* Resource, ERHIAccess Access, const FRDGPass* Pass);
 
 	/** Traverses all resources in the pass and marks whether they are externally accessible by user pass implementations. */
 	static void SetAllowRHIAccess(const FRDGPass* Pass, bool bAllowAccess);
diff --git a/Engine/Source/Runtime/Renderer/Private/DistanceFieldStreaming.cpp b/Engine/Source/Runtime/Renderer/Private/DistanceFieldStreaming.cpp
--- a/Engine/Source/Runtime/Renderer/Private/DistanceFieldStreaming.cpp
+++ b/Engine/Source/Runtime/Renderer/Private/DistanceFieldStreaming.cpp
@@ -1324,6 +1324,8 @@ void FDistanceFieldSceneData::UpdateDistanceFieldAtlas(
 
 	TArray<FDistanceFieldAssetMipId> AssetDataUploads;
 
+	FRDGExternalAccessQueue ExternalAccessQueue;
+
 	for (FSetElementId AssetSetId : DistanceFieldAssetRemoves)
 	{
 		const FDistanceFieldAssetState& AssetState = AssetStateArray[AssetSetId];
@@ -1566,10 +1568,10 @@ void FDistanceFieldSceneData::UpdateDistanceFieldAtlas(
 			});
 		}
 
-		FRDGTextureRef DistanceFieldBrickVolumeTextureRDG = GraphBuilder.RegisterExternalTexture(DistanceFieldBrickVolumeTexture, TEXT("DistanceFields.DistanceFieldBrickVolumeTexture"));
-
 		if (NumBrickUploads > 0)
 		{
+			FRDGTextureRef DistanceFieldBrickVolumeTextureRDG = GraphBuilder.RegisterExternalTexture(DistanceFieldBrickVolumeTexture, TEXT("DistanceFields.DistanceFieldBrickVolumeTexture"));
+
 			// GRHIMaxDispatchThreadGroupsPerDimension can be MAX_int32 so we need to do this math in 64-bit.
 			const int32 MaxBrickUploadsPerPass = (int32)FMath::Min<int64>((int64)GRHIMaxDispatchThreadGroupsPerDimension.Z * FScatterUploadDistanceFieldAtlasCS::GetGroupSize() / DistanceField::BrickSize, MAX_int32);
 
@@ -1595,10 +1597,8 @@ void FDistanceFieldSceneData::UpdateDistanceFieldAtlas(
 					FComputeShaderUtils::GetGroupCount(FIntVector(DistanceField::BrickSize, DistanceField::BrickSize, NumBrickUploadsThisPass * DistanceField::BrickSize), FScatterUploadDistanceFieldAtlasCS::GetGroupSize()));
 			}
 
-			DistanceFieldBrickVolumeTexture = GraphBuilder.ConvertToExternalTexture(DistanceFieldBrickVolumeTextureRDG);
+			ExternalAccessQueue.Add(DistanceFieldBrickVolumeTextureRDG, ERHIAccess::SRVMask);
 		}
-
-		GraphBuilder.FinalizeTextureAccess(DistanceFieldBrickVolumeTextureRDG, ERHIAccess::SRVMask);
 	}
 	
 	if (bIndirectionAtlasResized)
@@ -1618,6 +1618,8 @@ void FDistanceFieldSceneData::UpdateDistanceFieldAtlas(
 		ListMeshDistanceFields(bDumpAssetStats);
 		GDistanceFieldAtlasLogStats = 0;
 	}
+
+	ExternalAccessQueue.Submit(GraphBuilder);
 }
 
 void FDistanceFieldSceneData::ListMeshDistanceFields(bool bDumpAssetStats) const
diff --git a/Engine/Source/Runtime/Renderer/Private/GlobalDistanceField.cpp b/Engine/Source/Runtime/Renderer/Private/GlobalDistanceField.cpp
--- a/Engine/Source/Runtime/Renderer/Private/GlobalDistanceField.cpp
+++ b/Engine/Source/Runtime/Renderer/Private/GlobalDistanceField.cpp
@@ -763,13 +763,13 @@ static void ComputeUpdateRegionsAndUpdateViewState(
 
 			if (!ViewState.GlobalDistanceFieldPageFreeListAllocatorBuffer)
 			{
-				GetPooledFreeBuffer(RHICmdList, FRDGBufferDesc::CreateStructuredDesc(sizeof(uint32), 1), ViewState.GlobalDistanceFieldPageFreeListAllocatorBuffer, TEXT("PageFreeListAllocator"));
+				ViewState.GlobalDistanceFieldPageFreeListAllocatorBuffer = AllocatePooledBuffer(FRDGBufferDesc::CreateStructuredDesc(sizeof(uint32), 1), TEXT("PageFreeListAllocator"));
 			}
 
 			if (!ViewState.GlobalDistanceFieldPageFreeListBuffer
 				|| ViewState.GlobalDistanceFieldPageFreeListBuffer->Desc.NumElements != MaxPageNum)
 			{
-				GetPooledFreeBuffer(RHICmdList, FRDGBufferDesc::CreateStructuredDesc(sizeof(uint32), MaxPageNum), ViewState.GlobalDistanceFieldPageFreeListBuffer, TEXT("PageFreeList"));
+				ViewState.GlobalDistanceFieldPageFreeListBuffer = AllocatePooledBuffer(FRDGBufferDesc::CreateStructuredDesc(sizeof(uint32), MaxPageNum), TEXT("PageFreeList"));
 			}
 
 			if (!ViewState.GlobalDistanceFieldPageAtlasTexture
@@ -2488,47 +2488,47 @@ void UpdateGlobalDistanceFieldVolume(
 				}
 			}
 
-			FRDGResourceAccessFinalizer ResourceAccessFinalizer;
+			FRDGExternalAccessQueue ExternalAccessQueue;
 
 			for (int32 CacheType = StartCacheType; CacheType < GDF_Num; CacheType++)
 			{
 				if (PageTableLayerTextures[CacheType])
 				{
-					GlobalDistanceFieldInfo.PageTableLayerTextures[CacheType] = ConvertToFinalizedExternalTexture(GraphBuilder, ResourceAccessFinalizer, PageTableLayerTextures[CacheType]);
+					GlobalDistanceFieldInfo.PageTableLayerTextures[CacheType] = ConvertToExternalAccessTexture(GraphBuilder, ExternalAccessQueue, PageTableLayerTextures[CacheType]);
 				}
 			}
 
 			if (PageFreeListAllocatorBuffer)
 			{
-				GlobalDistanceFieldInfo.PageFreeListAllocatorBuffer = ConvertToFinalizedExternalBuffer(GraphBuilder, ResourceAccessFinalizer, PageFreeListAllocatorBuffer);
+				GlobalDistanceFieldInfo.PageFreeListAllocatorBuffer = ConvertToExternalAccessBuffer(GraphBuilder, ExternalAccessQueue, PageFreeListAllocatorBuffer);
 			}
 
 			if (PageFreeListBuffer)
 			{
-				GlobalDistanceFieldInfo.PageFreeListBuffer = ConvertToFinalizedExternalBuffer(GraphBuilder, ResourceAccessFinalizer, PageFreeListBuffer);
+				GlobalDistanceFieldInfo.PageFreeListBuffer = ConvertToExternalAccessBuffer(GraphBuilder, ExternalAccessQueue, PageFreeListBuffer);
 			}
 
 			if (PageAtlasTexture)
 			{
-				GlobalDistanceFieldInfo.PageAtlasTexture = ConvertToFinalizedExternalTexture(GraphBuilder, ResourceAccessFinalizer, PageAtlasTexture);
+				GlobalDistanceFieldInfo.PageAtlasTexture = ConvertToExternalAccessTexture(GraphBuilder, ExternalAccessQueue, PageAtlasTexture);
 			}
 
 			if (CoverageAtlasTexture)
 			{
-				GlobalDistanceFieldInfo.CoverageAtlasTexture = ConvertToFinalizedExternalTexture(GraphBuilder, ResourceAccessFinalizer, CoverageAtlasTexture);
+				GlobalDistanceFieldInfo.CoverageAtlasTexture = ConvertToExternalAccessTexture(GraphBuilder, ExternalAccessQueue, CoverageAtlasTexture);
 			}
 
 			if (PageTableCombinedTexture)
 			{
-				GlobalDistanceFieldInfo.PageTableCombinedTexture = ConvertToFinalizedExternalTexture(GraphBuilder, ResourceAccessFinalizer, PageTableCombinedTexture);
+				GlobalDistanceFieldInfo.PageTableCombinedTexture = ConvertToExternalAccessTexture(GraphBuilder, ExternalAccessQueue, PageTableCombinedTexture);
 			}
 
 			if (MipTexture)
 			{
-				GlobalDistanceFieldInfo.MipTexture = ConvertToFinalizedExternalTexture(GraphBuilder, ResourceAccessFinalizer, MipTexture);
+				GlobalDistanceFieldInfo.MipTexture = ConvertToExternalAccessTexture(GraphBuilder, ExternalAccessQueue, MipTexture);
 			}
 
-			ResourceAccessFinalizer.Finalize(GraphBuilder);
+			ExternalAccessQueue.Submit(GraphBuilder);
 		}
 
 		if (CVarGlobalDistanceFieldDebug.GetValueOnRenderThread() != 0 && View.GlobalDistanceFieldInfo.PageFreeListAllocatorBuffer)
diff --git a/Engine/Source/Runtime/Renderer/Private/HairStrands/HairStrandsLUT.cpp b/Engine/Source/Runtime/Renderer/Private/HairStrands/HairStrandsLUT.cpp
--- a/Engine/Source/Runtime/Renderer/Private/HairStrands/HairStrandsLUT.cpp
+++ b/Engine/Source/Runtime/Renderer/Private/HairStrands/HairStrandsLUT.cpp
@@ -345,11 +345,11 @@ static FRDGTextureRef InternalGetHairLUT(FRDGBuilder& GraphBuilder, const FViewI
 		if (bNeedUpdate)
 		{
 			Out = AddHairLUTPass(GraphBuilder, View, HairLUTType_DualScattering);
-			GSystemTextures.HairLUT0 = ConvertToFinalizedExternalTexture(GraphBuilder, Out);
+			GSystemTextures.HairLUT0 = ConvertToExternalAccessTexture(GraphBuilder, Out);
 		}
 		else if (bRegister)
 		{
-			Out = GraphBuilder.RegisterExternalTexture(GSystemTextures.HairLUT0, TEXT("Hair.LUT(DualScattering)"));
+			Out = GraphBuilder.RegisterExternalTexture(GSystemTextures.HairLUT0, TEXT("Hair.LUT(DualScattering)"), ERDGTextureFlags::SkipTracking);
 		}
 	}
 	else if (Type == FHairLUTType::HairLUTType_MeanEnergy)
@@ -358,11 +358,11 @@ static FRDGTextureRef InternalGetHairLUT(FRDGBuilder& GraphBuilder, const FViewI
 		if (bNeedUpdate)
 		{
 			Out = AddHairLUTPass(GraphBuilder, View, HairLUTType_MeanEnergy);
-			GSystemTextures.HairLUT1 = ConvertToFinalizedExternalTexture(GraphBuilder, Out);
+			GSystemTextures.HairLUT1 = ConvertToExternalAccessTexture(GraphBuilder, Out);
 		}
 		else if (bRegister)
 		{
-			Out = GraphBuilder.RegisterExternalTexture(GSystemTextures.HairLUT1, TEXT("Hair.LUT(MeanEnergy)"));
+			Out = GraphBuilder.RegisterExternalTexture(GSystemTextures.HairLUT1, TEXT("Hair.LUT(MeanEnergy)"), ERDGTextureFlags::SkipTracking);
 		}
 	}
 	else if (Type == FHairLUTType::HairLUTType_Coverage)
@@ -371,11 +371,11 @@ static FRDGTextureRef InternalGetHairLUT(FRDGBuilder& GraphBuilder, const FViewI
 		if (bNeedUpdate)
 		{
 			Out = AddHairCoverageLUTPass(GraphBuilder, View);
-			GSystemTextures.HairLUT2 = ConvertToFinalizedExternalTexture(GraphBuilder, Out);
+			GSystemTextures.HairLUT2 = ConvertToExternalAccessTexture(GraphBuilder, Out);
 		}
 		else if (bRegister)
 		{
-			Out = GraphBuilder.RegisterExternalTexture(GSystemTextures.HairLUT2, TEXT("Hair.LUT(Coverage)"));
+			Out = GraphBuilder.RegisterExternalTexture(GSystemTextures.HairLUT2, TEXT("Hair.LUT(Coverage)"), ERDGTextureFlags::SkipTracking);
 		}
 	}
 
diff --git a/Engine/Source/Runtime/Renderer/Private/HairStrands/HairStrandsVisibility.cpp b/Engine/Source/Runtime/Renderer/Private/HairStrands/HairStrandsVisibility.cpp
--- a/Engine/Source/Runtime/Renderer/Private/HairStrands/HairStrandsVisibility.cpp
+++ b/Engine/Source/Runtime/Renderer/Private/HairStrands/HairStrandsVisibility.cpp
@@ -2814,6 +2814,8 @@ void AddMeshDrawTransitionPass(
 	const FViewInfo& ViewInfo,
 	const FHairStrandsMacroGroupDatas& MacroGroupDatas)
 {
+	FRDGExternalAccessQueue ExternalAccessQueue;
+
 	for (const FHairStrandsMacroGroupData& MacroGroup : MacroGroupDatas)
 	{
 		for (const FHairStrandsMacroGroupData::PrimitiveInfo& PrimitiveInfo : MacroGroup.PrimitivesInfos)
@@ -2821,26 +2823,22 @@ void AddMeshDrawTransitionPass(
 			FHairGroupPublicData* HairGroupPublicData = PrimitiveInfo.PublicDataPtr;
 			check(HairGroupPublicData);
 
-			FRDGResourceAccessFinalizer ResourceAccessFinalizer;
-
 			FHairGroupPublicData::FVertexFactoryInput& VFInput = HairGroupPublicData->VFInput;
-			ResourceAccessFinalizer.AddBuffer(VFInput.Strands.PositionBuffer.Buffer,			ERHIAccess::SRVMask);
-			ResourceAccessFinalizer.AddBuffer(VFInput.Strands.PrevPositionBuffer.Buffer,		ERHIAccess::SRVMask);
-			ResourceAccessFinalizer.AddBuffer(VFInput.Strands.TangentBuffer.Buffer,				ERHIAccess::SRVMask);
-			ResourceAccessFinalizer.AddBuffer(VFInput.Strands.Attribute0Buffer.Buffer,			ERHIAccess::SRVMask);
-			ResourceAccessFinalizer.AddBuffer(VFInput.Strands.Attribute1Buffer.Buffer,			ERHIAccess::SRVMask);
-			ResourceAccessFinalizer.AddBuffer(VFInput.Strands.MaterialBuffer.Buffer,			ERHIAccess::SRVMask);
-			ResourceAccessFinalizer.AddBuffer(VFInput.Strands.PositionOffsetBuffer.Buffer,		ERHIAccess::SRVMask);
-			ResourceAccessFinalizer.AddBuffer(VFInput.Strands.PrevPositionOffsetBuffer.Buffer,	ERHIAccess::SRVMask);
+			ExternalAccessQueue.Add(VFInput.Strands.PositionBuffer.Buffer);
+			ExternalAccessQueue.Add(VFInput.Strands.PrevPositionBuffer.Buffer);
+			ExternalAccessQueue.Add(VFInput.Strands.TangentBuffer.Buffer);
+			ExternalAccessQueue.Add(VFInput.Strands.Attribute0Buffer.Buffer);
+			ExternalAccessQueue.Add(VFInput.Strands.Attribute1Buffer.Buffer);
+			ExternalAccessQueue.Add(VFInput.Strands.MaterialBuffer.Buffer);
+			ExternalAccessQueue.Add(VFInput.Strands.PositionOffsetBuffer.Buffer);
+			ExternalAccessQueue.Add(VFInput.Strands.PrevPositionOffsetBuffer.Buffer);
 
 			FRDGBufferRef CulledVertexIdBuffer = Register(GraphBuilder, HairGroupPublicData->CulledVertexIdBuffer, ERDGImportedBufferFlags::None).Buffer;
 			FRDGBufferRef CulledVertexRadiusScaleBuffer = Register(GraphBuilder, HairGroupPublicData->CulledVertexRadiusScaleBuffer, ERDGImportedBufferFlags::None).Buffer;
 			FRDGBufferRef DrawIndirectBuffer = Register(GraphBuilder, HairGroupPublicData->DrawIndirectBuffer, ERDGImportedBufferFlags::None).Buffer;
-			ResourceAccessFinalizer.AddBuffer(CulledVertexIdBuffer,								ERHIAccess::SRVMask);
-			ResourceAccessFinalizer.AddBuffer(CulledVertexRadiusScaleBuffer,					ERHIAccess::SRVMask);
-			ResourceAccessFinalizer.AddBuffer(DrawIndirectBuffer,								ERHIAccess::IndirectArgs);
-
-			ResourceAccessFinalizer.Finalize(GraphBuilder);
+			ExternalAccessQueue.Add(CulledVertexIdBuffer);
+			ExternalAccessQueue.Add(CulledVertexRadiusScaleBuffer);
+			ExternalAccessQueue.Add(DrawIndirectBuffer, ERHIAccess::IndirectArgs);
 
 			if (GetHairVisibilityRenderMode() != HairVisibilityRenderMode_ComputeRaster)
 			{
@@ -2855,6 +2853,8 @@ void AddMeshDrawTransitionPass(
 			}
 		}
 	}
+
+	ExternalAccessQueue.Submit(GraphBuilder);
 }
 ///////////////////////////////////////////////////////////////////////////////////////////////////
 
diff --git a/Engine/Source/Runtime/Renderer/Private/ReflectionEnvironmentCapture.cpp b/Engine/Source/Runtime/Renderer/Private/ReflectionEnvironmentCapture.cpp
--- a/Engine/Source/Runtime/Renderer/Private/ReflectionEnvironmentCapture.cpp
+++ b/Engine/Source/Runtime/Renderer/Private/ReflectionEnvironmentCapture.cpp
@@ -1229,7 +1229,7 @@ void CopyToSceneArray(FRDGBuilder& GraphBuilder, FScene* Scene, FRDGTexture* Fil
 	const int32 NumMips = GetNumMips(Scene->ReflectionSceneData.CubemapArray.GetCubemapSize());
 	const int32 CaptureIndex = FindOrAllocateCubemapIndex(Scene, ReflectionProxy->Component);
 
-	FRDGTexture* DestCubeTexture = GraphBuilder.RegisterExternalTexture(Scene->ReflectionSceneData.CubemapArray.GetRenderTarget(), ERDGTextureFlags::ForceTracking);
+	FRDGTexture* DestCubeTexture = GraphBuilder.RegisterExternalTexture(Scene->ReflectionSceneData.CubemapArray.GetRenderTarget());
 
 	// GPU copy back to the scene's texture array, which is not a render target
 	for (int32 MipIndex = 0; MipIndex < NumMips; MipIndex++)
@@ -1413,7 +1413,7 @@ void CopyToSkyTexture(FRDGBuilder& GraphBuilder, FScene* Scene, FRDGTexture* Inp
 	{
 		RDG_EVENT_SCOPE(GraphBuilder, "CopyToSkyTexture");
 
-		FRDGTexture* OutputTexture = GraphBuilder.RegisterExternalTexture(CreateRenderTarget(ProcessedTexture->TextureRHI, TEXT("SkyTexture")), ERDGTextureFlags::ForceTracking);
+		FRDGTexture* OutputTexture = GraphBuilder.RegisterExternalTexture(CreateRenderTarget(ProcessedTexture->TextureRHI, TEXT("SkyTexture")));
 
 		// GPU copy back to the skylight's texture, which is not a render target
 		FRHICopyTextureInfo CopyInfo;
diff --git a/Engine/Source/Runtime/Renderer/Private/ReflectionEnvironmentRealTimeCapture.cpp b/Engine/Source/Runtime/Renderer/Private/ReflectionEnvironmentRealTimeCapture.cpp
--- a/Engine/Source/Runtime/Renderer/Private/ReflectionEnvironmentRealTimeCapture.cpp
+++ b/Engine/Source/Runtime/Renderer/Private/ReflectionEnvironmentRealTimeCapture.cpp
@@ -1068,7 +1068,7 @@ void FScene::AllocateAndCaptureFrameSkyEnvMap(
 
 	if (ConvolvedSkyRenderTarget[ConvolvedSkyRenderTargetReadyIndex])
 	{
-		GraphBuilder.FinalizeTextureAccess(GraphBuilder.RegisterExternalTexture(ConvolvedSkyRenderTarget[ConvolvedSkyRenderTargetReadyIndex]), ERHIAccess::SRVMask);
+		GraphBuilder.UseExternalAccessMode(GraphBuilder.RegisterExternalTexture(ConvolvedSkyRenderTarget[ConvolvedSkyRenderTargetReadyIndex]), ERHIAccess::SRVMask);
 	}
 }
 
diff --git a/Engine/Source/Runtime/Renderer/Private/ShadowDepthRendering.cpp b/Engine/Source/Runtime/Renderer/Private/ShadowDepthRendering.cpp
--- a/Engine/Source/Runtime/Renderer/Private/ShadowDepthRendering.cpp
+++ b/Engine/Source/Runtime/Renderer/Private/ShadowDepthRendering.cpp
@@ -1517,7 +1517,7 @@ void FSceneRenderer::RenderShadowDepthMapAtlases(FRDGBuilder& GraphBuilder)
 
 	Scene->PrevAtlasHZBs.SetNum(SortedShadowsForShadowDepthPass.ShadowMapAtlases.Num());
 
-	FRDGResourceAccessFinalizer ResourceAccessFinalizer;
+	FRDGExternalAccessQueue ExternalAccessQueue;
 
 	for (int32 AtlasIndex = 0; AtlasIndex < SortedShadowsForShadowDepthPass.ShadowMapAtlases.Num(); AtlasIndex++)
 	{
@@ -1624,10 +1624,10 @@ void FSceneRenderer::RenderShadowDepthMapAtlases(FRDGBuilder& GraphBuilder)
 		// Make readable because AtlasDepthTexture is not tracked via RDG yet
 		// On mobile CSM atlas sampled only in pixel shaders
 		ERHIAccess AtlasDepthTextureAccessFinal = (FeatureLevel == ERHIFeatureLevel::ES3_1 ? ERHIAccess::SRVGraphics : ERHIAccess::SRVMask);
-		ShadowMapAtlas.RenderTargets.DepthTarget = ConvertToFinalizedExternalTexture(GraphBuilder, ResourceAccessFinalizer, AtlasDepthTexture, AtlasDepthTextureAccessFinal);
+		ShadowMapAtlas.RenderTargets.DepthTarget = ConvertToExternalAccessTexture(GraphBuilder, ExternalAccessQueue, AtlasDepthTexture, AtlasDepthTextureAccessFinal);
 	}
 
-	ResourceAccessFinalizer.Finalize(GraphBuilder);
+	ExternalAccessQueue.Submit(GraphBuilder);
 }
 
 void FSceneRenderer::RenderVirtualShadowMaps(FRDGBuilder& GraphBuilder, bool bNaniteEnabled)
@@ -1866,7 +1866,7 @@ void FSceneRenderer::RenderShadowDepthMaps(FRDGBuilder& GraphBuilder, FInstanceC
 
 	const bool bUseGeometryShader = !GRHISupportsArrayIndexFromAnyShader;
 
-	FRDGResourceAccessFinalizer ResourceAccessFinalizer;
+	FRDGExternalAccessQueue ExternalAccessQueue;
 
 	for (int32 CubemapIndex = 0; CubemapIndex < SortedShadowsForShadowDepthPass.ShadowMapCubemaps.Num(); CubemapIndex++)
 	{
@@ -2031,10 +2031,10 @@ void FSceneRenderer::RenderShadowDepthMaps(FRDGBuilder& GraphBuilder, FInstanceC
 		}
 
 		// Make readable because ShadowDepthTexture is not tracked via RDG yet
-		ShadowMap.RenderTargets.DepthTarget = ConvertToFinalizedExternalTexture(GraphBuilder, ResourceAccessFinalizer, ShadowDepthTexture);
+		ShadowMap.RenderTargets.DepthTarget = ConvertToExternalAccessTexture(GraphBuilder, ExternalAccessQueue, ShadowDepthTexture);
 	}
 
-	ResourceAccessFinalizer.Finalize(GraphBuilder);
+	ExternalAccessQueue.Submit(GraphBuilder);
 
 	if (SortedShadowsForShadowDepthPass.PreshadowCache.Shadows.Num() > 0)
 	{
diff --git a/Engine/Source/Runtime/Renderer/Private/SkyAtmosphereRendering.cpp b/Engine/Source/Runtime/Renderer/Private/SkyAtmosphereRendering.cpp
--- a/Engine/Source/Runtime/Renderer/Private/SkyAtmosphereRendering.cpp
+++ b/Engine/Source/Runtime/Renderer/Private/SkyAtmosphereRendering.cpp
@@ -1333,7 +1333,7 @@ void FSceneRenderer::RenderSkyAtmosphereLookUpTables(FRDGBuilder& GraphBuilder)
 	FRDGTextureRef MultiScatteredLuminanceLut = GraphBuilder.RegisterExternalTexture(SkyInfo.GetMultiScatteredLuminanceLutTexture());
 	FRDGTextureUAVRef MultiScatteredLuminanceLutUAV = GraphBuilder.CreateUAV(FRDGTextureUAVDesc(MultiScatteredLuminanceLut, 0));
 
-	FRDGResourceAccessFinalizer ResourceAccessFinalizer;
+	FRDGExternalAccessQueue ExternalAccessQueue;
 
 	// Transmittance LUT
 	FGlobalShaderMap* GlobalShaderMap = GetGlobalShaderMap(FeatureLevel);
@@ -1421,7 +1421,7 @@ void FSceneRenderer::RenderSkyAtmosphereLookUpTables(FRDGBuilder& GraphBuilder)
 		const FIntVector NumGroups = FIntVector::DivideAndRoundUp(TextureSize, FRenderDistantSkyLightLutCS::GroupSize);
 		FComputeShaderUtils::AddPass(GraphBuilder, RDG_EVENT_NAME("DistantSkyLightLut"), ComputeShader, PassParameters, NumGroups);
 
-		SkyInfo.GetDistantSkyLightLutTexture() = ConvertToFinalizedExternalTexture(GraphBuilder, ResourceAccessFinalizer, DistantSkyLightLut);
+		SkyInfo.GetDistantSkyLightLutTexture() = ConvertToExternalAccessTexture(GraphBuilder, ExternalAccessQueue, DistantSkyLightLut);
 	}
 
 	SkyAtmosphereLightShadowData LightShadowData;
@@ -1521,7 +1521,7 @@ void FSceneRenderer::RenderSkyAtmosphereLookUpTables(FRDGBuilder& GraphBuilder)
 			const FIntVector NumGroups = FIntVector::DivideAndRoundUp(TextureSize, FRenderSkyViewLutCS::GroupSize);
 			FComputeShaderUtils::AddPass(GraphBuilder, RDG_EVENT_NAME("RealTimeCaptureSkyViewLut"), ComputeShader, PassParameters, NumGroups);
 
-			Scene->RealTimeReflectionCaptureSkyAtmosphereViewLutTexture = ConvertToFinalizedExternalTexture(GraphBuilder, ResourceAccessFinalizer, RealTimeReflectionCaptureSkyAtmosphereViewLutTexture);
+			Scene->RealTimeReflectionCaptureSkyAtmosphereViewLutTexture = ConvertToExternalAccessTexture(GraphBuilder, ExternalAccessQueue, RealTimeReflectionCaptureSkyAtmosphereViewLutTexture);
 		}
 		else
 		{
@@ -1574,7 +1574,7 @@ void FSceneRenderer::RenderSkyAtmosphereLookUpTables(FRDGBuilder& GraphBuilder)
 			const FIntVector NumGroups = FIntVector::DivideAndRoundUp(TextureSize, FRenderCameraAerialPerspectiveVolumeCS::GroupSize);
 			FComputeShaderUtils::AddPass(GraphBuilder, RDG_EVENT_NAME("RealTimeCaptureCamera360VolumeLut"), ComputeShader, PassParameters, NumGroups);
 
-			Scene->RealTimeReflectionCaptureCamera360APLutTexture = ConvertToFinalizedExternalTexture(GraphBuilder, ResourceAccessFinalizer, RealTimeReflectionCaptureCamera360APLutTexture);
+			Scene->RealTimeReflectionCaptureCamera360APLutTexture = ConvertToExternalAccessTexture(GraphBuilder, ExternalAccessQueue, RealTimeReflectionCaptureCamera360APLutTexture);
 		}
 		else
 		{
@@ -1694,13 +1694,13 @@ void FSceneRenderer::RenderSkyAtmosphereLookUpTables(FRDGBuilder& GraphBuilder)
 			FComputeShaderUtils::AddPass(GraphBuilder, RDG_EVENT_NAME("CameraVolumeLut"), ComputeShader, PassParameters, NumGroups);
 		}
 
-		View.SkyAtmosphereViewLutTexture = ConvertToFinalizedExternalTexture(GraphBuilder, ResourceAccessFinalizer, SkyAtmosphereViewLutTexture);
-		View.SkyAtmosphereCameraAerialPerspectiveVolume = ConvertToFinalizedExternalTexture(GraphBuilder, ResourceAccessFinalizer, SkyAtmosphereCameraAerialPerspectiveVolume);
+		View.SkyAtmosphereViewLutTexture = ConvertToExternalAccessTexture(GraphBuilder, ExternalAccessQueue, SkyAtmosphereViewLutTexture);
+		View.SkyAtmosphereCameraAerialPerspectiveVolume = ConvertToExternalAccessTexture(GraphBuilder, ExternalAccessQueue, SkyAtmosphereCameraAerialPerspectiveVolume);
 	}
 
-	SkyInfo.GetTransmittanceLutTexture() = ConvertToFinalizedExternalTexture(GraphBuilder, ResourceAccessFinalizer, TransmittanceLut);
+	SkyInfo.GetTransmittanceLutTexture() = ConvertToExternalAccessTexture(GraphBuilder, ExternalAccessQueue, TransmittanceLut);
 
-	ResourceAccessFinalizer.Finalize(GraphBuilder);
+	ExternalAccessQueue.Submit(GraphBuilder);
 
 	// TODO have RDG execute those above passes with compute overlap similarly to using AutomaticCacheFlushAfterComputeShader(true);
 }
diff --git a/Engine/Source/Runtime/Renderer/Private/SystemTextures.cpp b/Engine/Source/Runtime/Renderer/Private/SystemTextures.cpp
--- a/Engine/Source/Runtime/Renderer/Private/SystemTextures.cpp
+++ b/Engine/Source/Runtime/Renderer/Private/SystemTextures.cpp
@@ -22,7 +22,7 @@ const FRDGSystemTextures& FRDGSystemTextures::Create(FRDGBuilder& GraphBuilder)
 {
 	const auto Register = [&](const TRefCountPtr<IPooledRenderTarget>& RenderTarget)
 	{
-		return TryRegisterExternalTexture(GraphBuilder, RenderTarget, ERDGTextureFlags::ReadOnly);
+		return TryRegisterExternalTexture(GraphBuilder, RenderTarget, ERDGTextureFlags::SkipTracking);
 	};
 
 	auto& SystemTextures = GraphBuilder.Blackboard.Create<FRDGSystemTextures>();
@@ -913,72 +913,72 @@ void FSystemTextures::ReleaseDynamicRHI()
 
 FRDGTextureRef FSystemTextures::GetBlackDummy(FRDGBuilder& GraphBuilder) const
 {
-	return GraphBuilder.RegisterExternalTexture(BlackDummy, TEXT("BlackDummy"));
+	return GraphBuilder.RegisterExternalTexture(BlackDummy, TEXT("BlackDummy"), ERDGTextureFlags::SkipTracking);
 }
 
 FRDGTextureRef FSystemTextures::GetBlackAlphaOneDummy(FRDGBuilder& GraphBuilder) const
 {
-	return GraphBuilder.RegisterExternalTexture(BlackAlphaOneDummy, TEXT("BlackAlphaOneDummy"));
+	return GraphBuilder.RegisterExternalTexture(BlackAlphaOneDummy, TEXT("BlackAlphaOneDummy"), ERDGTextureFlags::SkipTracking);
 }
 
 FRDGTextureRef FSystemTextures::GetBlackArrayDummy(FRDGBuilder& GraphBuilder) const
 {
-	return GraphBuilder.RegisterExternalTexture(BlackArrayDummy, TEXT("BlackArrayDummy"));
+	return GraphBuilder.RegisterExternalTexture(BlackArrayDummy, TEXT("BlackArrayDummy"), ERDGTextureFlags::SkipTracking);
 }
 
 FRDGTextureRef FSystemTextures::GetWhiteDummy(FRDGBuilder& GraphBuilder) const
 {
-	return GraphBuilder.RegisterExternalTexture(WhiteDummy, TEXT("WhiteDummy"));
+	return GraphBuilder.RegisterExternalTexture(WhiteDummy, TEXT("WhiteDummy"), ERDGTextureFlags::SkipTracking);
 }
 
 FRDGTextureRef FSystemTextures::GetMaxFP16Depth(FRDGBuilder& GraphBuilder) const
 {
-	return GraphBuilder.RegisterExternalTexture(MaxFP16Depth, TEXT("MaxFP16Depth"));
+	return GraphBuilder.RegisterExternalTexture(MaxFP16Depth, TEXT("MaxFP16Depth"), ERDGTextureFlags::SkipTracking);
 }
 
 FRDGTextureRef FSystemTextures::GetDepthDummy(FRDGBuilder& GraphBuilder) const
 {
-	return GraphBuilder.RegisterExternalTexture(DepthDummy, TEXT("DepthDummy"));
+	return GraphBuilder.RegisterExternalTexture(DepthDummy, TEXT("DepthDummy"), ERDGTextureFlags::SkipTracking);
 }
 
 FRDGTextureRef FSystemTextures::GetStencilDummy(FRDGBuilder& GraphBuilder) const
 {
-	return GraphBuilder.RegisterExternalTexture(StencilDummy, TEXT("StencilDummy"));
+	return GraphBuilder.RegisterExternalTexture(StencilDummy, TEXT("StencilDummy"), ERDGTextureFlags::SkipTracking);
 }
 
 FRDGTextureRef FSystemTextures::GetGreenDummy(FRDGBuilder& GraphBuilder) const
 {
-	return GraphBuilder.RegisterExternalTexture(GreenDummy, TEXT("GreenDummy"));
+	return GraphBuilder.RegisterExternalTexture(GreenDummy, TEXT("GreenDummy"), ERDGTextureFlags::SkipTracking);
 }
 
 FRDGTextureRef FSystemTextures::GetDefaultNormal8Bit(FRDGBuilder& GraphBuilder) const
 {
-	return GraphBuilder.RegisterExternalTexture(DefaultNormal8Bit, TEXT("DefaultNormal8Bit"));
+	return GraphBuilder.RegisterExternalTexture(DefaultNormal8Bit, TEXT("DefaultNormal8Bit"), ERDGTextureFlags::SkipTracking);
 }
 
 FRDGTextureRef FSystemTextures::GetMidGreyDummy(FRDGBuilder& GraphBuilder) const
 {
-	return GraphBuilder.RegisterExternalTexture(MidGreyDummy, TEXT("MidGreyDummy"));
+	return GraphBuilder.RegisterExternalTexture(MidGreyDummy, TEXT("MidGreyDummy"), ERDGTextureFlags::SkipTracking);
 }
 
 FRDGTextureRef FSystemTextures::GetVolumetricBlackDummy(FRDGBuilder& GraphBuilder) const
 {
-	return GraphBuilder.RegisterExternalTexture(VolumetricBlackDummy, TEXT("VolumetricBlackDummy"));
+	return GraphBuilder.RegisterExternalTexture(VolumetricBlackDummy, TEXT("VolumetricBlackDummy"), ERDGTextureFlags::SkipTracking);
 }
 
 FRDGTextureRef FSystemTextures::GetVolumetricBlackUintDummy(FRDGBuilder& GraphBuilder) const
 {
-	return GraphBuilder.RegisterExternalTexture(VolumetricBlackUintDummy, TEXT("VolumetricBlackUintDummy"));
+	return GraphBuilder.RegisterExternalTexture(VolumetricBlackUintDummy, TEXT("VolumetricBlackUintDummy"), ERDGTextureFlags::SkipTracking);
 }
 
 FRDGTextureRef FSystemTextures::GetZeroUIntDummy(FRDGBuilder& GraphBuilder) const
 {
-	return GraphBuilder.RegisterExternalTexture(ZeroUIntDummy, TEXT("ZeroUIntDummy"));
+	return GraphBuilder.RegisterExternalTexture(ZeroUIntDummy, TEXT("ZeroUIntDummy"), ERDGTextureFlags::SkipTracking);
 }
 
 FRDGTextureRef FSystemTextures::GetZeroUIntArrayDummy(FRDGBuilder& GraphBuilder) const
 {
-	return GraphBuilder.RegisterExternalTexture(ZeroUIntArrayDummy, TEXT("ZeroUIntArrayDummy"));
+	return GraphBuilder.RegisterExternalTexture(ZeroUIntArrayDummy, TEXT("ZeroUIntArrayDummy"), ERDGTextureFlags::SkipTracking);
 }
 
 //////////////////////////////////////////////////////////////////////////////////////////////////////////
@@ -1350,7 +1350,7 @@ FRDGTextureRef GetInternalDefaultTexture(
 
 	if (HashDefaultTextures.IsValid(Index) && DefaultTextures[Index].Texture != nullptr)
 	{
-		return GraphBuilder.RegisterExternalTexture(DefaultTextures[Index].Texture);
+		return GraphBuilder.RegisterExternalTexture(DefaultTextures[Index].Texture, ERDGTextureFlags::SkipTracking);
 	}
 
 	FDefaultTexture Entry;
@@ -1415,7 +1415,7 @@ FRDGTextureRef GetInternalDefaultTexture(
 
 	Index = DefaultTextures.Add(Entry);
 	HashDefaultTextures.Add(Hash, Index);
-	return GraphBuilder.RegisterExternalTexture(Entry.Texture);
+	return GraphBuilder.RegisterExternalTexture(Entry.Texture, ERDGTextureFlags::SkipTracking);
 }
 
 //////////////////////////////////////////////////////////////////////////////////////////////////////////
@@ -1474,7 +1474,7 @@ FRDGBufferRef GetInternalDefaultBuffer(
 	if (HashDefaultBuffers.IsValid(Index) && DefaultBuffers[Index].Buffer != nullptr)
 	{
 		check(DefaultBuffers[Index].Hash == Hash); //Sanitycheck
-		return GraphBuilder.RegisterExternalBuffer(DefaultBuffers[Index].Buffer);
+		return GraphBuilder.RegisterExternalBuffer(DefaultBuffers[Index].Buffer, ERDGBufferFlags::SkipTracking);
 	}
 
 	const uint32 BufferSize = NumBytePerElement * NumElements;
@@ -1530,7 +1530,7 @@ FRDGBufferRef GetInternalDefaultBuffer(
 
 	Index = DefaultBuffers.Add(Entry);
 	HashDefaultBuffers.Add(Hash, Index);
-	return GraphBuilder.RegisterExternalBuffer(Entry.Buffer);
+	return GraphBuilder.RegisterExternalBuffer(Entry.Buffer, ERDGBufferFlags::SkipTracking);
 }
 
 FVector4f GetClearBindingValue(EPixelFormat Format, FClearValueBinding Value)
diff --git a/Engine/Source/Runtime/Renderer/Private/VT/VirtualTextureFeedbackBuffer.cpp b/Engine/Source/Runtime/Renderer/Private/VT/VirtualTextureFeedbackBuffer.cpp
--- a/Engine/Source/Runtime/Renderer/Private/VT/VirtualTextureFeedbackBuffer.cpp
+++ b/Engine/Source/Runtime/Renderer/Private/VT/VirtualTextureFeedbackBuffer.cpp
@@ -86,21 +86,21 @@ void FVirtualTextureFeedbackBuffer::Begin(FRDGBuilder& GraphBuilder, const FVirt
 	FRDGBufferDesc BufferDesc(FRDGBufferDesc::CreateBufferDesc(sizeof(uint32), Desc.BufferSize.X * Desc.BufferSize.Y));
 	BufferDesc.Usage |= BUF_SourceCopy;
 
-	if (GetPooledFreeBuffer(GraphBuilder.RHICmdList, BufferDesc, PooledBuffer, TEXT("VirtualTextureFeedbackGPU")))
+	if (AllocatePooledBuffer(BufferDesc, PooledBuffer, TEXT("VirtualTextureFeedbackGPU")))
 	{
 		FRDGBufferUAVDesc UAVDesc;
 		UAVDesc.Format = PF_R32_UINT;
 		UAV = PooledBuffer->GetOrCreateUAV(UAVDesc);
 	}
 
-	AddPass(GraphBuilder, RDG_EVENT_NAME("VirtualTextureFeedbackClear"), [this](FRHICommandList& RHICmdList)
-	{
-		// Clear virtual texture feedback to default value
-		RHICmdList.Transition(FRHITransitionInfo(UAV, ERHIAccess::Unknown, ERHIAccess::UAVCompute));
-		RHICmdList.ClearUAVUint(UAV, FUintVector4(~0u, ~0u, ~0u, ~0u));
-		RHICmdList.Transition(FRHITransitionInfo(UAV, ERHIAccess::UAVCompute, ERHIAccess::UAVGraphics));
-		RHICmdList.BeginUAVOverlap(UAV);
-	});
+	// We can go ahead and just clear this now. No need to wait on the RDG timeline.
+	FRHICommandListImmediate& RHICmdList = GraphBuilder.RHICmdList;
+
+	// Clear virtual texture feedback to default value
+	RHICmdList.Transition(FRHITransitionInfo(UAV, ERHIAccess::Unknown, ERHIAccess::UAVCompute));
+	RHICmdList.ClearUAVUint(UAV, FUintVector4(~0u, ~0u, ~0u, ~0u));
+	RHICmdList.Transition(FRHITransitionInfo(UAV, ERHIAccess::UAVCompute, ERHIAccess::UAVGraphics));
+	RHICmdList.BeginUAVOverlap(UAV);
 }
 
 void FVirtualTextureFeedbackBuffer::End(FRDGBuilder& GraphBuilder)
diff --git a/Engine/Source/Runtime/Renderer/Private/VT/VirtualTextureSpace.cpp b/Engine/Source/Runtime/Renderer/Private/VT/VirtualTextureSpace.cpp
--- a/Engine/Source/Runtime/Renderer/Private/VT/VirtualTextureSpace.cpp
+++ b/Engine/Source/Runtime/Renderer/Private/VT/VirtualTextureSpace.cpp
@@ -540,13 +540,15 @@ void FVirtualTextureSpace::FinalizeTextures(FRDGBuilder& GraphBuilder)
 		FTextureEntry& PageTableEntry = PageTable[TextureIndex];
 		if (PageTableEntry.RenderTarget)
 		{
-			FRDGTextureRef PageTableTexture = GraphBuilder.RegisterExternalTexture(PageTableEntry.RenderTarget);
-			GraphBuilder.FinalizeTextureAccess(PageTableTexture, ERHIAccess::SRVMask);
+			// It's only necessary to enable external access mode on textures modified by RDG this frame.
+			if (FRDGTexture* Texture = GraphBuilder.FindExternalTexture(PageTableEntry.RenderTarget))
+			{
+				GraphBuilder.UseExternalAccessMode(Texture, ERHIAccess::SRVMask);
+			}
 		}
 	}
 }
 
-
 void FVirtualTextureSpace::DumpToConsole(bool verbose)
 {
 	UE_LOG(LogConsoleResponse, Display, TEXT("-= Space ID %i =-"), ID);
diff --git a/Engine/Source/Runtime/Renderer/Private/VolumetricCloudRendering.cpp b/Engine/Source/Runtime/Renderer/Private/VolumetricCloudRendering.cpp
--- a/Engine/Source/Runtime/Renderer/Private/VolumetricCloudRendering.cpp
+++ b/Engine/Source/Runtime/Renderer/Private/VolumetricCloudRendering.cpp
@@ -1453,7 +1453,7 @@ void FSceneRenderer::InitVolumetricCloudsForViews(FRDGBuilder& GraphBuilder, boo
 
 	if (Scene)
 	{
-		FRDGResourceAccessFinalizer ResourceAccessFinalizer;
+		FRDGExternalAccessQueue ExternalAccessQueue;
 
 		check(ShouldRenderVolumetricCloud(Scene, ActiveViewFamily->EngineShowFlags)); // This should not be called if we should not render SkyAtmosphere
 
@@ -1866,12 +1866,12 @@ void FSceneRenderer::InitVolumetricCloudsForViews(FRDGBuilder& GraphBuilder, boo
 										FilterTracedCloudTexture(&SpatiallyFilteredShadowTexture, VolumetricCloudParams.VolumetricCloud.CloudShadowmapSizeInvSize[LightIndex], CloudShadowTextureTexelWorldSizeInvSize, false);
 									}
 									ViewInfo.VolumetricCloudShadowRenderTarget[LightIndex] = SpatiallyFilteredShadowTexture;
-									ViewInfo.VolumetricCloudShadowExtractedRenderTarget[LightIndex] = ConvertToFinalizedExternalTexture(GraphBuilder, ResourceAccessFinalizer, ViewInfo.VolumetricCloudShadowRenderTarget[LightIndex]);
+									ViewInfo.VolumetricCloudShadowExtractedRenderTarget[LightIndex] = ConvertToExternalAccessTexture(GraphBuilder, ExternalAccessQueue, ViewInfo.VolumetricCloudShadowRenderTarget[LightIndex]);
 								}
 								else
 								{
 									ViewInfo.VolumetricCloudShadowRenderTarget[LightIndex] = GraphBuilder.RegisterExternalTexture(CloudShadowTemporalRT.CurrentRenderTarget());
-									ViewInfo.VolumetricCloudShadowExtractedRenderTarget[LightIndex] = ConvertToFinalizedExternalTexture(GraphBuilder, ResourceAccessFinalizer, ViewInfo.VolumetricCloudShadowRenderTarget[LightIndex]);
+									ViewInfo.VolumetricCloudShadowExtractedRenderTarget[LightIndex] = ConvertToExternalAccessTexture(GraphBuilder, ExternalAccessQueue, ViewInfo.VolumetricCloudShadowRenderTarget[LightIndex]);
 								}
 
 							}
@@ -1887,12 +1887,12 @@ void FSceneRenderer::InitVolumetricCloudsForViews(FRDGBuilder& GraphBuilder, boo
 										FilterTracedCloudTexture(&SpatiallyFilteredShadowTexture, VolumetricCloudParams.VolumetricCloud.CloudShadowmapSizeInvSize[LightIndex], CloudShadowTextureTexelWorldSizeInvSize, false);
 									}
 									ViewInfo.VolumetricCloudShadowRenderTarget[LightIndex] = SpatiallyFilteredShadowTexture;
-									ViewInfo.VolumetricCloudShadowExtractedRenderTarget[LightIndex] = ConvertToFinalizedExternalTexture(GraphBuilder, ResourceAccessFinalizer, ViewInfo.VolumetricCloudShadowRenderTarget[LightIndex]);
+									ViewInfo.VolumetricCloudShadowExtractedRenderTarget[LightIndex] = ConvertToExternalAccessTexture(GraphBuilder, ExternalAccessQueue, ViewInfo.VolumetricCloudShadowRenderTarget[LightIndex]);
 								}
 								else
 								{
 									ViewInfo.VolumetricCloudShadowRenderTarget[LightIndex] = NewCloudShadowTexture;
-									ViewInfo.VolumetricCloudShadowExtractedRenderTarget[LightIndex] = ConvertToFinalizedExternalTexture(GraphBuilder, ResourceAccessFinalizer, ViewInfo.VolumetricCloudShadowRenderTarget[LightIndex]);
+									ViewInfo.VolumetricCloudShadowExtractedRenderTarget[LightIndex] = ConvertToExternalAccessTexture(GraphBuilder, ExternalAccessQueue, ViewInfo.VolumetricCloudShadowRenderTarget[LightIndex]);
 								}
 							}
 						}
@@ -1928,7 +1928,7 @@ void FSceneRenderer::InitVolumetricCloudsForViews(FRDGBuilder& GraphBuilder, boo
 			CleanUpCloudDataPass(GraphBuilder);
 		}
 
-		ResourceAccessFinalizer.Finalize(GraphBuilder);
+		ExternalAccessQueue.Submit(GraphBuilder);
 	}
 	else
 	{
