Add basic GPU based validation of ray tracing scene build parameters

diff --git a/Engine/Shaders/Private/RayTracing/RayTracingCommon.ush b/Engine/Shaders/Private/RayTracing/RayTracingCommon.ush
--- a/Engine/Shaders/Private/RayTracing/RayTracingCommon.ush
+++ b/Engine/Shaders/Private/RayTracing/RayTracingCommon.ush
@@ -151,6 +151,26 @@ struct FRayTracingInstanceDescriptor
 	uint InstanceIdAndMask;
 	uint InstanceContributionToHitGroupIndexAndFlags;
 	uint2 AccelerationStructure; //uint64
+
+	uint GetInstanceId()
+	{
+		return InstanceIdAndMask & 0x00FFFFFF;
+	}
+
+	uint GetMask()
+	{
+		return (InstanceIdAndMask >> 24) & 0xFF;
+	}
+
+	uint GetInstanceContributionToHitGroupIndex()
+	{
+		return InstanceContributionToHitGroupIndexAndFlags & 0x00FFFFFF;
+	}
+
+	uint GetFlags()
+	{
+		return (InstanceContributionToHitGroupIndexAndFlags >> 24) & 0xFF;
+	}
 };
 
 void SetRayTracingInstanceTransform(inout FRayTracingInstanceDescriptor Desc, float3x4 LocalToWorld)
diff --git a/Engine/Shaders/Private/RayTracing/RayTracingValidation.usf b/Engine/Shaders/Private/RayTracing/RayTracingValidation.usf
--- a/Engine/Shaders/Private/RayTracing/RayTracingValidation.usf
+++ b/Engine/Shaders/Private/RayTracing/RayTracingValidation.usf
@@ -94,6 +94,59 @@ void RayTracingValidateGeometryBuildParamsCS(uint3 DispatchThreadId : SV_Dispatc
 	CheckTriangle(Tri, PrimitiveIndex);
 }
 
+uint NumInstances;
+uint NumHitGroups;
+ByteAddressBuffer InstanceBuffer;
+uint InstanceBufferOffsetInBytes;
+uint InstanceBufferStrideInBytes;
+
+[numthreads(NUM_THREADS_X, 1, 1)]
+void RayTracingValidateSceneBuildParamsCS(uint3 DispatchThreadId : SV_DispatchThreadID)
+{
+	uint InstanceIndex = DispatchThreadId.x;
+	if (InstanceIndex >= NumInstances)
+	{
+		return;
+	}
+
+	uint InstanceAddress = InstanceBufferOffsetInBytes + InstanceBufferStrideInBytes * InstanceIndex;
+	FRayTracingInstanceDescriptor InstanceDesc = InstanceBuffer.Load<FRayTracingInstanceDescriptor>(InstanceAddress);
+
+	// Check that the instance matrix is finite (if instance culling is required, it should be done through null BLAS address)
+	{
+		for (uint i = 0; i < 3; ++i)
+		{
+			float4 Row = InstanceDesc.Transform[i];
+			for (uint j = 0; j < 4; ++j)
+			{
+				float Value = Row[j];
+				if (!IsFinite(Value))
+				{
+					UEReportAssertWithPayload(0x1C46B0FD, uint4(InstanceIndex, asuint(Value), 0, 0));
+				}
+			}
+		}
+	}
+
+	// Check that flag bits belong to D3D12_RAYTRACING_INSTANCE_FLAGS entries (only first 4 bits can be set)
+	{
+		uint Flags = InstanceDesc.GetFlags();
+		if ((Flags & 0xF0) != 0)
+		{
+			UEReportAssertWithPayload(0x07085473, uint4(InstanceIndex, Flags, 0, 0));
+		}
+	}
+
+	// Check that instance contribution to SBT index is sane (will not cause OOB access)
+	{
+		uint Contribution = InstanceDesc.GetInstanceContributionToHitGroupIndex();
+		if (Contribution > NumHitGroups)
+		{
+			UEReportAssertWithPayload(0xC5A7424A, uint4(InstanceIndex, Contribution, NumHitGroups, 0));
+		}
+	}
+}
+
 #else // PLATFORM_SUPPORTS_DIAGNOSTIC_BUFFER
 
 [numthreads(NUM_THREADS_X, 1, 1)]
@@ -101,5 +154,10 @@ void RayTracingValidateGeometryBuildParamsCS(uint3 DispatchThreadId : SV_Dispatc
 {
 }
 
+[numthreads(NUM_THREADS_X, 1, 1)]
+void RayTracingValidateSceneBuildParamsCS(uint3 DispatchThreadId : SV_DispatchThreadID)
+{
+}
+
 #endif // PLATFORM_SUPPORTS_DIAGNOSTIC_BUFFER
 
diff --git a/Engine/Source/Runtime/D3D12RHI/Private/D3D12RayTracing.cpp b/Engine/Source/Runtime/D3D12RHI/Private/D3D12RayTracing.cpp
--- a/Engine/Source/Runtime/D3D12RHI/Private/D3D12RayTracing.cpp
+++ b/Engine/Source/Runtime/D3D12RHI/Private/D3D12RayTracing.cpp
@@ -4180,6 +4180,16 @@ void FD3D12RayTracingScene::BuildAccelerationStructure(FD3D12CommandContext& Com
 
 	checkf(BuildDesc.ScratchAccelerationStructureData % 256 == 0, TEXT("TLAS build scratch buffer must have 256 byte alignment."));
 
+	if (GD3D12RayTracingGPUValidation)
+	{
+		TRHICommandList_RecursiveHazardous<FD3D12CommandContext> RHICmdList(&CommandContext, CommandContext.GetGPUMask());
+		uint32 InstanceBufferStride = uint32(sizeof(D3D12_RAYTRACING_INSTANCE_DESC));
+		uint32 TotalHitGroupSlots = Initializer.NumTotalSegments * Initializer.ShaderSlotsPerGeometrySegment;
+		FRayTracingValidateSceneBuildParamsCS::Dispatch(RHICmdList, 
+			TotalHitGroupSlots, BuildInputs.NumDescs,
+			InstanceBuffer, InstanceBufferOffset, InstanceBufferStride);
+	}
+
 	// UAV barrier is used here to ensure that all bottom level acceleration structures are built
 	CommandContext.CommandListHandle.AddUAVBarrier();
 	CommandContext.CommandListHandle.FlushResourceBarriers();
diff --git a/Engine/Source/Runtime/RenderCore/Private/RayTracingValidationShaders.cpp b/Engine/Source/Runtime/RenderCore/Private/RayTracingValidationShaders.cpp
--- a/Engine/Source/Runtime/RenderCore/Private/RayTracingValidationShaders.cpp
+++ b/Engine/Source/Runtime/RenderCore/Private/RayTracingValidationShaders.cpp
@@ -6,6 +6,8 @@
 
 #if RHI_RAYTRACING
 
+// FRayTracingValidateGeometryBuildParamsCS
+
 IMPLEMENT_GLOBAL_SHADER(FRayTracingValidateGeometryBuildParamsCS, "/Engine/Private/RayTracing/RayTracingValidation.usf", "RayTracingValidateGeometryBuildParamsCS", SF_Compute);
 
 void FRayTracingValidateGeometryBuildParamsCS::Dispatch(FRHICommandList& RHICmdList, const FRayTracingGeometryBuildParams& Params)
@@ -23,7 +25,7 @@ void FRayTracingValidateGeometryBuildParamsCS::Dispatch(FRHICommandList& RHICmdL
 	}
 
 	TWideStringBuilder<256> EventName;
-	EventName.Append(TEXT("RTValidation"));
+	EventName.Append(TEXT("RTGeometryValidation"));
 	if (!Initializer.DebugName.IsNone())
 	{
 		FString DebugNameString = Initializer.DebugName.ToString();
@@ -74,4 +76,37 @@ void FRayTracingValidateGeometryBuildParamsCS::Dispatch(FRHICommandList& RHICmdL
 	RHICmdList.PopEvent();
 }
 
+// FRayTracingValidateSceneBuildParamsCS
+
+IMPLEMENT_GLOBAL_SHADER(FRayTracingValidateSceneBuildParamsCS, "/Engine/Private/RayTracing/RayTracingValidation.usf", "RayTracingValidateSceneBuildParamsCS", SF_Compute);
+
+void FRayTracingValidateSceneBuildParamsCS::Dispatch(FRHICommandList& RHICmdList, 
+	uint32 NumHitGroups, uint32 NumInstances, 
+	FRHIBuffer* InstanceBuffer, uint32 InstanceBufferOffset, uint32 InstanceBufferStride)
+{
+	TShaderMapRef<FRayTracingValidateSceneBuildParamsCS> ComputeShader(GetGlobalShaderMap(GMaxRHIFeatureLevel));
+	FRHIComputeShader* ShaderRHI = ComputeShader.GetComputeShader();
+	SetComputePipelineState(RHICmdList, ShaderRHI);
+
+	RHICmdList.PushEvent(TEXT("RTSceneValidation"), FColor::Black);
+
+	const FRawBufferShaderResourceViewInitializer InstanceBufferViewInitializer(InstanceBuffer);
+	FShaderResourceViewRHIRef InstanceBufferSRV = RHICreateShaderResourceView(InstanceBufferViewInitializer);
+
+	SetShaderValue(RHICmdList, ShaderRHI, ComputeShader->NumInstancesParam, NumInstances);
+	SetShaderValue(RHICmdList, ShaderRHI, ComputeShader->NumHitGroupsParam, NumHitGroups);
+	SetShaderValue(RHICmdList, ShaderRHI, ComputeShader->InstanceBufferOffsetInBytesParam, InstanceBufferOffset);
+	SetShaderValue(RHICmdList, ShaderRHI, ComputeShader->InstanceBufferStrideInBytesParam, InstanceBufferStride);
+	SetSRVParameter(RHICmdList, ShaderRHI, ComputeShader->InstanceBufferParam, InstanceBufferSRV);
+
+	const uint32 MaxDispatchDimension = 65536; // D3D11_CS_DISPATCH_MAX_THREAD_GROUPS_PER_DIMENSION
+	const uint32 NumGroupsX = FMath::Min((NumInstances + NumThreadsX - 1) / NumThreadsX, MaxDispatchDimension);
+
+	RHICmdList.DispatchComputeShader(NumGroupsX, 1, 1);
+
+	SetSRVParameter(RHICmdList, ShaderRHI, ComputeShader->InstanceBufferParam, nullptr);
+
+	RHICmdList.PopEvent();
+}
+
 #endif // RHI_RAYTRACING
diff --git a/Engine/Source/Runtime/RenderCore/Public/RayTracingValidationShaders.h b/Engine/Source/Runtime/RenderCore/Public/RayTracingValidationShaders.h
--- a/Engine/Source/Runtime/RenderCore/Public/RayTracingValidationShaders.h
+++ b/Engine/Source/Runtime/RenderCore/Public/RayTracingValidationShaders.h
@@ -54,4 +54,47 @@ public:
 
 };
 
+class RENDERCORE_API FRayTracingValidateSceneBuildParamsCS : public FBuiltInRayTracingShader
+{
+	DECLARE_GLOBAL_SHADER(FRayTracingValidateSceneBuildParamsCS);
+
+public:
+	FRayTracingValidateSceneBuildParamsCS() = default;
+	FRayTracingValidateSceneBuildParamsCS(const ShaderMetaType::CompiledShaderInitializerType& Initializer)
+		: FBuiltInRayTracingShader(Initializer)
+	{
+		InstanceBufferParam.Bind(Initializer.ParameterMap, TEXT("InstanceBuffer"), SPF_Optional);
+		InstanceBufferOffsetInBytesParam.Bind(Initializer.ParameterMap, TEXT("InstanceBufferOffsetInBytes"), SPF_Optional);
+		InstanceBufferStrideInBytesParam.Bind(Initializer.ParameterMap, TEXT("InstanceBufferStrideInBytes"), SPF_Optional);
+		NumInstancesParam.Bind(Initializer.ParameterMap, TEXT("NumInstances"), SPF_Optional);
+		NumHitGroupsParam.Bind(Initializer.ParameterMap, TEXT("NumHitGroups"), SPF_Optional);
+	}
+
+	static bool ShouldCompilePermutation(const FGlobalShaderPermutationParameters& Parameters)
+	{
+		return FBuiltInRayTracingShader::ShouldCompilePermutation(Parameters);
+	}
+
+	// Large thread group to handle large meshes with a single 1D dispatch
+	static const uint32 NumThreadsX = 1024;
+
+	static inline void ModifyCompilationEnvironment(const FGlobalShaderPermutationParameters& Parameters, FShaderCompilerEnvironment& OutEnvironment)
+	{
+		OutEnvironment.CompilerFlags.Add(CFLAG_ForceDXC);
+		OutEnvironment.SetDefine(TEXT("NUM_THREADS_X"), NumThreadsX);
+		FBuiltInRayTracingShader::ModifyCompilationEnvironment(Parameters, OutEnvironment);
+	}
+
+	static void Dispatch(FRHICommandList& RHICmdList, 
+		uint32 NumHitGroups, uint32 NumInstances, 
+		FRHIBuffer* InstanceBuffer, uint32 InstanceBufferOffset, uint32 InstanceBufferStride);
+
+	LAYOUT_FIELD(FShaderResourceParameter, InstanceBufferParam);
+	LAYOUT_FIELD(FShaderParameter, InstanceBufferOffsetInBytesParam);
+	LAYOUT_FIELD(FShaderParameter, InstanceBufferStrideInBytesParam);
+	LAYOUT_FIELD(FShaderParameter, NumInstancesParam);
+	LAYOUT_FIELD(FShaderParameter, NumHitGroupsParam);
+};
+
+
 #endif // RHI_RAYTRACING
