- Add support for FUintVector3 & FQuat4f as shader parameters

diff --git a/Engine/Source/Runtime/RenderCore/Public/ShaderParameterMacros.h b/Engine/Source/Runtime/RenderCore/Public/ShaderParameterMacros.h
--- a/Engine/Source/Runtime/RenderCore/Public/ShaderParameterMacros.h
+++ b/Engine/Source/Runtime/RenderCore/Public/ShaderParameterMacros.h
@@ -987,6 +987,21 @@ struct TShaderParameterTypeInfo<FIntVector>
 	static const FShaderParametersMetadata* GetStructMetadata() { return nullptr; }
 };
 
+template<>
+struct TShaderParameterTypeInfo<FUintVector3>
+{
+	static constexpr EUniformBufferBaseType BaseType = UBMT_UINT32;
+	static constexpr int32 NumRows = 1;
+	static constexpr int32 NumColumns = 3;
+	static constexpr int32 NumElements = 0;
+	static constexpr int32 Alignment = 16;
+	static constexpr bool bIsStoredInConstantBuffer = true;
+
+	using TAlignedType = TAlignedTypedef<FUintVector3, Alignment>::Type;
+
+	static const FShaderParametersMetadata* GetStructMetadata() { return nullptr; }
+};
+
 template<>
 struct TShaderParameterTypeInfo<FIntVector4>
 {
@@ -1032,6 +1047,21 @@ struct TShaderParameterTypeInfo<FIntRect>
 	static const FShaderParametersMetadata* GetStructMetadata() { return nullptr; }
 };
 
+template<>
+struct TShaderParameterTypeInfo<FQuat4f>
+{
+	static constexpr EUniformBufferBaseType BaseType = UBMT_FLOAT32;
+	static constexpr int32 NumRows = 1;
+	static constexpr int32 NumColumns = 4;
+	static constexpr int32 NumElements = 0;
+	static constexpr int32 Alignment = 16;
+	static constexpr bool bIsStoredInConstantBuffer = true;
+
+	using TAlignedType = TAlignedTypedef<FQuat4f, Alignment>::Type;
+
+	static const FShaderParametersMetadata* GetStructMetadata() { return nullptr; }
+};
+
 template<>
 struct TShaderParameterTypeInfo<FMatrix44f>
 {
