Added RDG_USE_MALLOC debug macro to help find memory stomps with -stompmalloc.

diff --git a/Engine/Source/Runtime/RenderCore/Private/RenderGraphAllocator.cpp b/Engine/Source/Runtime/RenderCore/Private/RenderGraphAllocator.cpp
--- a/Engine/Source/Runtime/RenderCore/Private/RenderGraphAllocator.cpp
+++ b/Engine/Source/Runtime/RenderCore/Private/RenderGraphAllocator.cpp
@@ -11,17 +11,34 @@ FRDGAllocator& FRDGAllocator::Get()
 
 FRDGAllocator::~FRDGAllocator()
 {
+#if RDG_USE_MALLOC
+	check(Context.RawAllocs.IsEmpty() && ContextForTasks.RawAllocs.IsEmpty());
+#else
 	check(Context.MemStack.IsEmpty() && ContextForTasks.MemStack.IsEmpty());
+#endif
 }
 
 void FRDGAllocator::FContext::ReleaseAll()
 {
 	for (int32 Index = TrackedAllocs.Num() - 1; Index >= 0; --Index)
 	{
+#if RDG_USE_MALLOC
+		delete TrackedAllocs[Index];
+#else
 		TrackedAllocs[Index]->~FTrackedAlloc();
+#endif
 	}
 	TrackedAllocs.Reset();
+
+#if RDG_USE_MALLOC
+	for (void* RawAlloc : RawAllocs)
+	{
+		FMemory::Free(RawAlloc);
+	}
+	RawAllocs.Reset();
+#else
 	MemStack.Flush();
+#endif
 }
 
 void FRDGAllocator::ReleaseAll()
diff --git a/Engine/Source/Runtime/RenderCore/Public/RenderGraphAllocator.h b/Engine/Source/Runtime/RenderCore/Public/RenderGraphAllocator.h
--- a/Engine/Source/Runtime/RenderCore/Public/RenderGraphAllocator.h
+++ b/Engine/Source/Runtime/RenderCore/Public/RenderGraphAllocator.h
@@ -5,6 +5,8 @@
 #include "CoreMinimal.h"
 #include "Misc/MemStack.h"
 
+#define RDG_USE_MALLOC 0
+
 /** Private allocator used by RDG to track its internal memory. All memory is released after RDG builder execution. */
 class RENDERCORE_API FRDGAllocator final
 {
@@ -14,7 +16,13 @@ public:
 	/** Allocates raw memory. */
 	FORCEINLINE void* Alloc(uint64 SizeInBytes, uint32 AlignInBytes)
 	{
+#if RDG_USE_MALLOC
+		void* Memory = FMemory::Malloc(SizeInBytes, AlignInBytes);
+		GetContext().RawAllocs.Emplace(Memory);
+		return Memory;
+#else
 		return GetContext().MemStack.Alloc(SizeInBytes, AlignInBytes);
+#endif
 	}
 
 	/** Allocates an uninitialized type without destructor tracking. */
@@ -29,7 +37,11 @@ public:
 	FORCEINLINE T* Alloc(TArgs&&... Args)
 	{
 		FContext& LocalContext = GetContext();
+	#if RDG_USE_MALLOC
+		TTrackedAlloc<T>* TrackedAlloc = new TTrackedAlloc<T>(Forward<TArgs&&>(Args)...);
+	#else
 		TTrackedAlloc<T>* TrackedAlloc = new(LocalContext.MemStack) TTrackedAlloc<T>(Forward<TArgs&&>(Args)...);
+	#endif
 		check(TrackedAlloc);
 		LocalContext.TrackedAllocs.Add(TrackedAlloc);
 		return &TrackedAlloc->Alloc;
@@ -39,12 +51,20 @@ public:
 	template <typename T, typename... TArgs>
 	FORCEINLINE T* AllocNoDestruct(TArgs&&... Args)
 	{
+#if RDG_USE_MALLOC
+		return new (AllocUninitialized<T>(1)) T(Forward<TArgs&&>(Args)...);
+#else
 		return new (GetContext().MemStack) T(Forward<TArgs&&>(Args)...);
+#endif
 	}
 
 	FORCEINLINE int32 GetByteCount() const
 	{
+	#if RDG_USE_MALLOC
+		return 0;
+	#else
 		return Context.MemStack.GetByteCount() + ContextForTasks.MemStack.GetByteCount();
+	#endif
 	}
 
 private:
@@ -70,7 +90,11 @@ private:
 
 	struct FContext
 	{
+#if RDG_USE_MALLOC
+		TArray<void*> RawAllocs;
+#else
 		FMemStackBase MemStack;
+#endif
 		TArray<FTrackedAlloc*> TrackedAllocs;
 
 		void ReleaseAll();
