Fixes to external access RDG feature to handle async compute without validation failures.

diff --git a/Engine/Source/Runtime/RHI/Private/RHI.cpp b/Engine/Source/Runtime/RHI/Private/RHI.cpp
--- a/Engine/Source/Runtime/RHI/Private/RHI.cpp
+++ b/Engine/Source/Runtime/RHI/Private/RHI.cpp
@@ -1506,8 +1506,8 @@ RHI_API uint64 GRHITransitionPrivateData_AlignInBytes = 0;
 // By default, read only states and UAV states are allowed to participate in state merging.
 ERHIAccess GRHIMergeableAccessMask = ERHIAccess::ReadOnlyMask | ERHIAccess::UAVMask;
 
-// By default, no multi-pipe states are allowed. Each platform must opt-in.
-ERHIAccess GRHIMultiPipelineMergeableAccessMask = ERHIAccess::Unknown;
+// By default, only exclusively read only accesses are allowed.
+ERHIAccess GRHIMultiPipelineMergeableAccessMask = ERHIAccess::ReadOnlyExclusiveMask;
 
 /** Called once per frame only from within an RHI. */
 void RHIPrivateBeginFrame()
diff --git a/Engine/Source/Runtime/RHI/Public/RHI.h b/Engine/Source/Runtime/RHI/Public/RHI.h
--- a/Engine/Source/Runtime/RHI/Public/RHI.h
+++ b/Engine/Source/Runtime/RHI/Public/RHI.h
@@ -1068,6 +1068,9 @@ enum class ERHIAccess
 	// A mask of all bits representing read-only states which cannot be combined with other write states.
 	ReadOnlyExclusiveMask = CPURead | Present | IndirectArgs | VertexOrIndexBuffer | SRVGraphics | SRVCompute | CopySrc | ResolveSrc | BVHRead,
 
+	// A mask of all bits representing read-only states on the compute pipe which cannot be combined with other write states.
+	ReadOnlyExclusiveComputeMask = CPURead | IndirectArgs | SRVCompute | CopySrc | BVHRead,
+
 	// A mask of all bits representing read-only states which may be combined with other write states.
 	ReadOnlyMask = ReadOnlyExclusiveMask | DSVRead | ShadingRateSource,
 
@@ -1946,6 +1949,11 @@ struct FRHICopyTextureInfo
 	uint32 NumMips = 1;
 };
 
+inline constexpr bool IsReadOnlyExclusiveAccess(ERHIAccess Access)
+{
+	return EnumHasAnyFlags(Access, ERHIAccess::ReadOnlyExclusiveMask) && !EnumHasAnyFlags(Access, ~ERHIAccess::ReadOnlyExclusiveMask);
+}
+
 inline constexpr bool IsReadOnlyAccess(ERHIAccess Access)
 {
 	return EnumHasAnyFlags(Access, ERHIAccess::ReadOnlyMask) && !EnumHasAnyFlags(Access, ~ERHIAccess::ReadOnlyMask);
diff --git a/Engine/Source/Runtime/RenderCore/Private/RenderGraphBuilder.cpp b/Engine/Source/Runtime/RenderCore/Private/RenderGraphBuilder.cpp
--- a/Engine/Source/Runtime/RenderCore/Private/RenderGraphBuilder.cpp
+++ b/Engine/Source/Runtime/RenderCore/Private/RenderGraphBuilder.cpp
@@ -600,7 +600,7 @@ void FRDGBuilder::UseExternalAccessMode(FRDGViewableResource* Resource, ERHIAcce
 {
 	IF_RDG_ENABLE_DEBUG(UserValidation.ValidateUseExternalAccessMode(Resource, ReadOnlyAccess, Pipelines));
 
-	if (!GRDGAsyncCompute || !GSupportsEfficientAsyncCompute)
+	if (!GRDGAsyncCompute || !GSupportsEfficientAsyncCompute || IsImmediateMode())
 	{
 		Pipelines = ERHIPipeline::Graphics;
 	}
@@ -617,14 +617,14 @@ void FRDGBuilder::UseExternalAccessMode(FRDGViewableResource* Resource, ERHIAcce
 	// We have to flush the queue when going from QueuedInternal -> External. A queued internal state
 	// implies that the resource was in an external access mode before, so it needs an 'end' pass to 
 	// contain any passes which might have used the resource in its external state.
-	if (AccessModeState.QueueIndex.IsValid())
+	if (AccessModeState.bQueued)
 	{
 		FlushAccessModeQueue();
 	}
 
-	check(AccessModeState.QueueIndex.IsNull());
-	AccessModeState.QueueIndex = FRDGViewableResource::FAccessModeState::FQueueIndex(AccessModeQueue.Num());
+	check(!AccessModeState.bQueued);
 	AccessModeQueue.Emplace(Resource);
+	AccessModeState.bQueued = 1;
 
 	Resource->SetExternalAccessMode(ReadOnlyAccess, Pipelines);
 }
@@ -635,25 +635,28 @@ void FRDGBuilder::UseInternalAccessMode(FRDGViewableResource* Resource)
 
 	auto& AccessModeState = Resource->AccessModeState;
 
-	// Just no-op if the resource is already in (or queued for) the 'Internal state.
+	// Just no-op if the resource is already in (or queued for) the Internal state.
 	if (AccessModeState.Mode == FRDGViewableResource::EAccessMode::Internal)
 	{
 		return;
 	}
 
 	// If the resource has a queued transition to the external access state, then we can safely back it out.
-	if (AccessModeState.QueueIndex.IsValid())
+	if (AccessModeState.bQueued)
 	{
-		AccessModeQueue[AccessModeState.QueueIndex.GetIndex()] = nullptr;
-		AccessModeState.QueueIndex = {};
+		int32 Index = AccessModeQueue.IndexOfByKey(Resource);
+		check(Index < AccessModeQueue.Num());
+		AccessModeQueue.RemoveAtSwap(Index, 1, false);
+		AccessModeState.bQueued = 0;
 	}
 	else
 	{
-		AccessModeState.QueueIndex = FRDGViewableResource::FAccessModeState::FQueueIndex(AccessModeQueue.Num());
 		AccessModeQueue.Emplace(Resource);
+		AccessModeState.bQueued = 1;
 	}
 
 	AccessModeState.Mode = FRDGViewableResource::EAccessMode::Internal;
+	Resource->EpilogueAccess = FRDGViewableResource::DefaultEpilogueAccess;
 }
 
 void FRDGBuilder::FlushAccessModeQueue()
@@ -671,32 +674,45 @@ void FRDGBuilder::FlushAccessModeQueue()
 		AllocParameters<FAccessModePassParameters>()
 	};
 
+	const ERHIAccess AccessMaskByPipeline[] =
+	{
+		ERHIAccess::ReadOnlyExclusiveMask,
+		ERHIAccess::ReadOnlyExclusiveComputeMask
+	};
+
 	ERHIPipeline ParameterPipelines = ERHIPipeline::None;
 
 	for (FRDGViewableResource* Resource : AccessModeQueue)
 	{
-		if (!Resource)
-		{
-			continue;
-		}
-
 		const auto& AccessModeState = Resource->AccessModeState;
 
 		ParameterPipelines |= AccessModeState.Pipelines;
 
+		if (AccessModeState.Mode == FRDGViewableResource::EAccessMode::External)
+		{
+			ExternalAccessResources.Emplace(Resource);
+		}
+		else
+		{
+			ExternalAccessResources.Remove(Resource);
+		}
+
 		for (uint32 PipelineIndex = 0; PipelineIndex < GetRHIPipelineCount(); ++PipelineIndex)
 		{
 			const ERHIPipeline Pipeline = static_cast<ERHIPipeline>(1 << PipelineIndex);
 
 			if (EnumHasAnyFlags(AccessModeState.Pipelines, Pipeline))
 			{
+				const ERHIAccess Access = AccessModeState.Access & AccessMaskByPipeline[PipelineIndex];
+				check(Access != ERHIAccess::None);
+
 				switch (Resource->Type)
 				{
 				case ERDGViewableResourceType::Texture:
-					ParametersByPipeline[PipelineIndex]->Textures.Emplace(GetAsTexture(Resource), AccessModeState.Access);
+					ParametersByPipeline[PipelineIndex]->Textures.Emplace(GetAsTexture(Resource), Access);
 					break;
 				case ERDGViewableResourceType::Buffer:
-					ParametersByPipeline[PipelineIndex]->Buffers.Emplace(GetAsBuffer(Resource), AccessModeState.Access);
+					ParametersByPipeline[PipelineIndex]->Buffers.Emplace(GetAsBuffer(Resource), Access);
 					break;
 				}
 			}
@@ -726,9 +742,10 @@ void FRDGBuilder::FlushAccessModeQueue()
 			[](FRHIComputeCommandList&) {});
 	}
 
+	// This has to be done *after* adding the passes so they don't inherit the new access mode.
 	for (FRDGViewableResource* Resource : AccessModeQueue)
 	{
-		Resource->AccessModeState.QueueIndex = {};
+		Resource->AccessModeState.bQueued = false;
 	}
 
 	AccessModeQueue.Reset();
@@ -1457,7 +1474,15 @@ void FRDGBuilder::Execute()
 	SCOPED_NAMED_EVENT_TEXT("FRDGBuilder::Execute", FColor::Magenta);
 
 	GRDGTransientResourceAllocator.ReleasePendingDeallocations();
-	FlushAccessModeQueue();
+
+	{
+		SCOPED_NAMED_EVENT_TEXT("FRDGBuilder::FlushAccessModeQueue", FColor::Magenta);
+		for (FRDGViewableResource* Resource : ExternalAccessResources)
+		{
+			UseInternalAccessMode(Resource);
+		}
+		FlushAccessModeQueue();
+	}
 
 	// Create the epilogue pass at the end of the graph just prior to compilation.
 	{
@@ -2615,8 +2640,7 @@ void FRDGBuilder::AddEpilogueTransition(FRDGTextureRef Texture)
 		return;
 	}
 
-	// Only textures with internal access mode are transitioned to their epilogue state.
-	if (Texture->AccessModeState.Mode == FRDGViewableResource::EAccessMode::Internal)
+	if (!EnumHasAnyFlags(Texture->Flags, ERDGTextureFlags::SkipTracking))
 	{
 		const FRDGPassHandle EpiloguePassHandle = GetEpiloguePassHandle();
 
@@ -2641,16 +2665,7 @@ void FRDGBuilder::AddEpilogueTransition(FRDGTextureRef Texture)
 		InitAsWholeResource(ScratchTextureState, &SubresourceState);
 		AddTransition(EpiloguePassHandle, Texture, ScratchTextureState);
 		ScratchTextureState.Reset();
-	}
-	// External access mode textures are left in their last external state.
-	else
-	{
-		Texture->EpilogueAccess = Texture->GetState()[0].Access;
-	}
 
-	// Resources that were marked for external access on registration (via SkipTracking) could have an unknown epilogue state.
-	if (Texture->EpilogueAccess != ERHIAccess::Unknown)
-	{
 		EpilogueResourceAccesses.Emplace(Texture->GetRHI(), Texture->EpilogueAccess);
 	}
 
@@ -2673,8 +2688,7 @@ void FRDGBuilder::AddEpilogueTransition(FRDGBufferRef Buffer)
 		return;
 	}
 
-	// Only buffers with internal access mode are transitioned to their epilogue state.
-	if (Buffer->AccessModeState.Mode == FRDGViewableResource::EAccessMode::Internal)
+	if (!EnumHasAnyFlags(Buffer->Flags, ERDGBufferFlags::SkipTracking))
 	{
 		const FRDGPassHandle EpiloguePassHandle = GetEpiloguePassHandle();
 
@@ -2697,16 +2711,7 @@ void FRDGBuilder::AddEpilogueTransition(FRDGBufferRef Buffer)
 		SubresourceState.Access = Buffer->EpilogueAccess;
 
 		AddTransition(Buffer->LastPass, Buffer, SubresourceState);
-	}
-	// External access mode buffers are left in their last external state.
-	else
-	{
-		Buffer->EpilogueAccess = Buffer->GetState().Access;
-	}
 
-	// Resources that were marked for external access on registration (via SkipTracking) could have an unknown epilogue state.
-	if (Buffer->EpilogueAccess != ERHIAccess::Unknown)
-	{
 		EpilogueResourceAccesses.Emplace(Buffer->GetRHI(), Buffer->EpilogueAccess);
 	}
 
diff --git a/Engine/Source/Runtime/RenderCore/Private/RenderGraphPass.cpp b/Engine/Source/Runtime/RenderCore/Private/RenderGraphPass.cpp
--- a/Engine/Source/Runtime/RenderCore/Private/RenderGraphPass.cpp
+++ b/Engine/Source/Runtime/RenderCore/Private/RenderGraphPass.cpp
@@ -174,6 +174,16 @@ void FRDGBarrierBatchBegin::Submit(FRHIComputeCommandList& RHICmdList, ERHIPipel
 #endif
 }
 
+FRDGBarrierBatchEndId FRDGBarrierBatchEnd::GetId() const
+{
+	return FRDGBarrierBatchEndId(Pass->GetHandle(), BarrierLocation);
+}
+
+bool FRDGBarrierBatchEnd::IsPairedWith(const FRDGBarrierBatchBegin& BeginBatch) const
+{
+	return GetId() == BeginBatch.BarriersToEnd[Pass->GetPipeline()];
+}
+
 void FRDGBarrierBatchBegin::Submit(FRHIComputeCommandList& RHICmdList, ERHIPipeline Pipeline)
 {
 	FRDGTransitionQueue TransitionsToBegin;
@@ -201,7 +211,7 @@ void FRDGBarrierBatchEnd::AddDependency(FRDGBarrierBatchBegin* BeginBatch)
 #endif
 
 	{
-		const FRDGBarrierBatchEndId Id(Pass->GetHandle(), BarrierLocation);
+		const FRDGBarrierBatchEndId Id = GetId();
 
 		FRDGBarrierBatchEndId& EarliestEndId = BeginBatch->BarriersToEnd[Pass->GetPipeline()];
 
diff --git a/Engine/Source/Runtime/RenderCore/Private/RenderGraphValidation.cpp b/Engine/Source/Runtime/RenderCore/Private/RenderGraphValidation.cpp
--- a/Engine/Source/Runtime/RenderCore/Private/RenderGraphValidation.cpp
+++ b/Engine/Source/Runtime/RenderCore/Private/RenderGraphValidation.cpp
@@ -486,7 +486,8 @@ void FRDGUserValidation::ValidateUseExternalAccessMode(FRDGViewableResource* Res
 
 	const auto& AccessModeState = Resource->AccessModeState;
 	checkf(!AccessModeState.bLocked, TEXT("Resource is locked in external access mode by the SkipTracking resource flag: %s."), Resource->Name);
-	checkf(IsReadOnlyAccess(ReadOnlyAccess), TEXT("A read only access is required when use external access mode. (Resource: %s, Access: %s, Pipelines: %s)"), Resource->Name, *GetRHIAccessName(ReadOnlyAccess), *GetRHIPipelineName(Pipelines));
+	checkf(IsReadOnlyExclusiveAccess(ReadOnlyAccess), TEXT("A read only access is required when use external access mode. (Resource: %s, Access: %s, Pipelines: %s)"), Resource->Name, *GetRHIAccessName(ReadOnlyAccess), *GetRHIPipelineName(Pipelines));
+
 	checkf(
 		AccessModeState.Mode == FRDGViewableResource::EAccessMode::Internal ||
 		(AccessModeState.Access == ReadOnlyAccess && AccessModeState.Pipelines == Pipelines),
@@ -495,6 +496,13 @@ void FRDGUserValidation::ValidateUseExternalAccessMode(FRDGViewableResource* Res
 		TEXT("\tExisting Access: %s, Requested Access: %s\n")
 		TEXT("\tExisting Pipelines: %s, Requested Pipelines: %s\n"),
 		Resource->Name, *GetRHIAccessName(AccessModeState.Access), *GetRHIAccessName(ReadOnlyAccess), *GetRHIPipelineName(AccessModeState.Pipelines), *GetRHIPipelineName(Pipelines));
+
+	checkf(EnumHasAnyFlags(Pipelines, ERHIPipeline::Graphics) || !EnumHasAnyFlags(ReadOnlyAccess, AccessMaskRaster),
+		TEXT("Raster access flags were specified for external access but the graphics pipe was not specified.\n")
+		TEXT("Resource: %s\n")
+		TEXT("\tRequested Access: %s\n")
+		TEXT("\tRequested Pipelines: %s\n"),
+		Resource->Name, *GetRHIAccessName(ReadOnlyAccess), *GetRHIPipelineName(Pipelines));
 }
 
 void FRDGUserValidation::ValidateUseInternalAccessMode(FRDGViewableResource* Resource)
@@ -1187,7 +1195,7 @@ void FRDGBarrierValidation::ValidateBarrierBatchBegin(const FRDGPass* Pass, cons
 		if (!bFoundFirst)
 		{
 			bFoundFirst = true;
-			UE_LOG(LogRDG, Display, TEXT("[%s(Index: %d, Pipeline: %s): %s] (Begin):"), Pass->GetName(), Pass->GetHandle().GetIndex(), *GetRHIPipelineName(Pass->GetPipeline()), Batch.DebugName);
+			UE_LOG(LogRDG, Display, TEXT("[%s(Index: %d, Pipeline: %s): Batch: (%p) %s] (Begin):"), Pass->GetName(), Pass->GetHandle().GetIndex(), *GetRHIPipelineName(Pass->GetPipeline()), &Batch, Batch.DebugName);
 		}
 	};
 
@@ -1277,9 +1285,11 @@ void FRDGBarrierValidation::ValidateBarrierBatchEnd(const FRDGPass* Pass, const
 
 	bool bFoundFirst = false;
 
+	const FRDGBarrierBatchEndId Id = Batch.GetId();
+
 	for (const FRDGBarrierBatchBegin* Dependent : Batch.Dependencies)
 	{
-		if (Dependent->PipelinesToEnd == ERHIPipeline::None)
+		if (!Batch.IsPairedWith(*Dependent))
 		{
 			continue;
 		}
@@ -1303,7 +1313,7 @@ void FRDGBarrierValidation::ValidateBarrierBatchEnd(const FRDGPass* Pass, const
 			if (!bFoundFirst)
 			{
 				bFoundFirst = true;
-				UE_LOG(LogRDG, Display, TEXT("[%s(Index: %d, Pipeline: %s) %s] (End):"), Pass->GetName(), Pass->GetHandle().GetIndex(), Dependent->DebugName, *GetRHIPipelineName(Pass->GetPipeline()));
+				UE_LOG(LogRDG, Display, TEXT("[%s(Index: %d, Pipeline: %s) Batch: (%p) %s] (End):"), Pass->GetName(), Pass->GetHandle().GetIndex(), *GetRHIPipelineName(Pass->GetPipeline()), Dependent, Dependent->DebugName);
 			}
 		};
 
@@ -1339,6 +1349,8 @@ void FRDGBarrierValidation::ValidateBarrierBatchEnd(const FRDGPass* Pass, const
 				}
 			}
 		}
+
+		bFoundFirst = false;
 	}
 }
 
diff --git a/Engine/Source/Runtime/RenderCore/Public/RenderGraphBuilder.h b/Engine/Source/Runtime/RenderCore/Public/RenderGraphBuilder.h
--- a/Engine/Source/Runtime/RenderCore/Public/RenderGraphBuilder.h
+++ b/Engine/Source/Runtime/RenderCore/Public/RenderGraphBuilder.h
@@ -253,12 +253,28 @@ public:
 	 */
 	void UseExternalAccessMode(FRDGViewableResource* Resource, ERHIAccess ReadOnlyAccess, ERHIPipeline Pipelines = ERHIPipeline::Graphics);
 
+	void UseExternalAccessMode(TArrayView<FRDGViewableResource* const> Resources, ERHIAccess ReadOnlyAccess, ERHIPipeline Pipelines = ERHIPipeline::Graphics)
+	{
+		for (FRDGViewableResource* Resource : Resources)
+		{
+			UseExternalAccessMode(Resource, ReadOnlyAccess, Pipelines);
+		}
+	}
+
 	/** Use this method to resume tracking of a resource after calling UseExternalAccessMode. It is safe to call this method
 	 *  even if external access mode was not enabled (it will simply no-op). It is not valid to access the underlying RHI
 	 *  resource in any pass added after calling this method.
 	 */
 	void UseInternalAccessMode(FRDGViewableResource* Resource);
 
+	inline void UseInternalAccessMode(TArrayView<FRDGViewableResource* const> Resources)
+	{
+		for (FRDGViewableResource* Resource : Resources)
+		{
+			UseInternalAccessMode(Resource);
+		}
+	}
+
 	/** Flag a resource that is produced by a pass but never used or extracted to not emit an 'unused' warning. */
 	void RemoveUnusedTextureWarning(FRDGTextureRef Texture);
 	void RemoveUnusedBufferWarning(FRDGBufferRef Buffer);
@@ -639,6 +655,7 @@ private:
 
 	/** Contains resources queued for either access mode change passes. */
 	TArray<FRDGViewableResource*, FRDGArrayAllocator> AccessModeQueue;
+	TSet<FRDGViewableResource*, DefaultKeyFuncs<FRDGViewableResource*>, FRDGSetAllocator> ExternalAccessResources;
 
 	/** Texture state used for intermediate operations. Held here to avoid re-allocating. */
 	FRDGTextureSubresourceStateIndirect ScratchTextureState;
@@ -656,7 +673,6 @@ private:
 	bool bDispatchHint = false;
 	bool bFlushResourcesRHI = false;
 	bool bParallelExecuteEnabled = false;
-	bool bInAccessModeQueueFlush = false;
 
 #if RDG_ENABLE_DEBUG
 	FRDGUserValidation UserValidation;
diff --git a/Engine/Source/Runtime/RenderCore/Public/RenderGraphPass.h b/Engine/Source/Runtime/RenderCore/Public/RenderGraphPass.h
--- a/Engine/Source/Runtime/RenderCore/Public/RenderGraphPass.h
+++ b/Engine/Source/Runtime/RenderCore/Public/RenderGraphPass.h
@@ -143,6 +143,10 @@ public:
 		Dependencies.Reserve(TransitionBatchCount);
 	}
 
+	FRDGBarrierBatchEndId GetId() const;
+
+	bool IsPairedWith(const FRDGBarrierBatchBegin& BeginBatch) const;
+
 private:
 	TArray<FRDGBarrierBatchBegin*, TInlineAllocator<4, FRDGArrayAllocator>> Dependencies;
 	FRDGPass* Pass;
diff --git a/Engine/Source/Runtime/RenderCore/Public/RenderGraphResources.h b/Engine/Source/Runtime/RenderCore/Public/RenderGraphResources.h
--- a/Engine/Source/Runtime/RenderCore/Public/RenderGraphResources.h
+++ b/Engine/Source/Runtime/RenderCore/Public/RenderGraphResources.h
@@ -269,6 +269,8 @@ public:
 protected:
 	FRDGViewableResource(const TCHAR* InName, ERDGViewableResourceType InType, bool bSkipTracking);
 
+	static const ERHIAccess DefaultEpilogueAccess = ERHIAccess::SRVMask;
+
 	enum class ETransientExtractionHint : uint8
 	{
 		None,
@@ -291,22 +293,20 @@ protected:
 
 	struct FAccessModeState
 	{
-		using FQueueIndex = TRDGHandle<FAccessModeState, uint16>;
-
-		bool IsQueued() const { return QueueIndex.IsValid(); }
-		bool IsExternalAccess() const { return Mode == EAccessMode::External && !IsQueued(); }
+		bool IsExternalAccess() const { return Mode == EAccessMode::External && !bQueued; }
 
 		FAccessModeState()
 			: Pipelines(ERHIPipeline::None)
 			, Mode(EAccessMode::Internal)
 			, bLocked(0)
+			, bQueued(0)
 		{}
 
 		ERHIAccess			Access			= ERHIAccess::None;
 		ERHIPipeline		Pipelines		: 2;
 		EAccessMode			Mode			: 1;
 		uint8				bLocked			: 1;
-		FQueueIndex QueueIndex;
+		uint8				bQueued			: 1;
 
 	} AccessModeState;
 
@@ -350,7 +350,7 @@ protected:
 	FRDGPassHandle LastPass;
 
 	/** The state of the resource at the graph epilogue. */
-	ERHIAccess EpilogueAccess = ERHIAccess::SRVMask;
+	ERHIAccess EpilogueAccess = DefaultEpilogueAccess;
 
 private:
 	/** Number of references in passes and deferred queries. */
@@ -369,6 +369,8 @@ private:
 
 		// External access resources are not always added to the pass states (unless marked as such within the graph), so marked as not culled here.
 		bCulled = 0;
+
+		EpilogueAccess = InReadOnlyAccess;
 	}
 
 #if RDG_ENABLE_TRACE
@@ -1214,6 +1216,11 @@ public:
 		return Desc.GetSize();
 	}
 
+	FORCEINLINE uint32 GetStride() const
+	{
+		return Desc.BytesPerElement;
+	}
+
 private:
 	FRDGBuffer(const TCHAR* InName, const FRDGBufferDesc& InDesc, ERDGBufferFlags InFlags)
 		: FRDGViewableResource(InName, ERDGViewableResourceType::Buffer, EnumHasAnyFlags(InFlags, ERDGBufferFlags::SkipTracking))
diff --git a/Engine/Source/Runtime/RenderCore/Public/RenderGraphUtils.h b/Engine/Source/Runtime/RenderCore/Public/RenderGraphUtils.h
--- a/Engine/Source/Runtime/RenderCore/Public/RenderGraphUtils.h
+++ b/Engine/Source/Runtime/RenderCore/Public/RenderGraphUtils.h
@@ -1121,9 +1121,15 @@ public:
 
 	void Add(FRDGViewableResource* Resource, ERHIAccess Access = ERHIAccess::SRVMask, ERHIPipeline Pipelines = ERHIPipeline::Graphics)
 	{
-		if (Resource)
+		Validate(Resource, Access, Pipelines);
+		Resources.Emplace(Resource, Access, Pipelines);
+	}
+
+	void AddUnique(FRDGViewableResource* Resource, ERHIAccess Access = ERHIAccess::SRVMask, ERHIPipeline Pipelines = ERHIPipeline::Graphics)
+	{
+		Validate(Resource, Access, Pipelines);
+		if (!Contains(Resource))
 		{
-			checkf(IsValidAccess(Access) && Access != ERHIAccess::Unknown, TEXT("Attempted to finalize texture %s with an invalid access %s."), Resource->Name, *GetRHIAccessName(Access));
 			Resources.Emplace(Resource, Access, Pipelines);
 		}
 	}
@@ -1134,7 +1140,12 @@ public:
 		{
 			GraphBuilder.UseExternalAccessMode(Resource.Resource, Resource.Access, Resource.Pipelines);
 		}
-		Resources.Reset();
+		Resources.Empty();
+	}
+
+	bool Contains(FRDGViewableResource* Resource)
+	{
+		return Resources.ContainsByPredicate([&](const FResource& InResource) { return InResource.Resource == Resource; });
 	}
 
 	bool IsEmpty() const
@@ -1167,6 +1178,13 @@ public:
 	}
 
 private:
+	void Validate(FRDGViewableResource* Resource, ERHIAccess Access, ERHIPipeline Pipelines)
+	{
+		checkf(Resource, TEXT("Added a null resource to FRDGExternalAccessQueue"));
+		checkf(IsValidAccess(Access) && Access != ERHIAccess::Unknown, TEXT("Attempted to finalize texture %s with an invalid access %s."), Resource->Name, *GetRHIAccessName(Access));
+		check(Pipelines != ERHIPipeline::None);
+	}
+
 	struct FResource
 	{
 		FResource() = default;
