Added ResizeStructuredBufferSAOIfNeeded for GPU scene buffer refactor.

diff --git a/Engine/Source/Runtime/RenderCore/Private/UnifiedBuffer.cpp b/Engine/Source/Runtime/RenderCore/Private/UnifiedBuffer.cpp
--- a/Engine/Source/Runtime/RenderCore/Private/UnifiedBuffer.cpp
+++ b/Engine/Source/Runtime/RenderCore/Private/UnifiedBuffer.cpp
@@ -246,9 +246,20 @@ EResourceType GetResourceType(FRDGViewableResource* Resource)
 	return EResourceType::BUFFER;
 }
 
-void MemsetResource(FRDGBuilder& GraphBuilder, FRDGViewableResource* Resource, const FMemsetResourceParams& Params)
+void MemsetResource(FRDGBuilder& GraphBuilder, FRDGBuffer* DstResource, const FMemsetResourceParams& Params)
 {
-	check(Resource);
+	MemsetResource(GraphBuilder, GraphBuilder.CreateUAV(DstResource, ERDGUnorderedAccessViewFlags::SkipBarrier), Params);
+}
+
+void MemcpyResource(FRDGBuilder& GraphBuilder, FRDGBuffer* DstResource, FRDGBuffer* SrcResource, const FMemcpyResourceParams& Params)
+{
+	MemcpyResource(GraphBuilder, GraphBuilder.CreateUAV(DstResource, ERDGUnorderedAccessViewFlags::SkipBarrier), GraphBuilder.CreateSRV(SrcResource), Params);
+}
+
+void MemsetResource(FRDGBuilder& GraphBuilder, FRDGUnorderedAccessView* UAV, const FMemsetResourceParams& Params)
+{
+	check(UAV);
+	FRDGViewableResource* Resource = UAV->GetParent();
 
 	EByteBufferResourceType ResourceTypeEnum = EByteBufferResourceType::Count;
 
@@ -264,20 +275,20 @@ void MemsetResource(FRDGBuilder& GraphBuilder, FRDGViewableResource* Resource, c
 	{
 	case EResourceType::BYTEBUFFER:
 		ResourceTypeEnum = EByteBufferResourceType::Uint_Buffer;
-		PassParameters->DstByteAddressBuffer = GraphBuilder.CreateUAV(GetAsBuffer(Resource));
+		PassParameters->DstByteAddressBuffer = GetAs<FRDGBufferUAV>(UAV);
 		Divisor = 4;
 		break;
 	case EResourceType::BUFFER:
 		ResourceTypeEnum = EByteBufferResourceType::Float4_Buffer;
-		PassParameters->DstBuffer = GraphBuilder.CreateUAV(GetAsBuffer(Resource));
+		PassParameters->DstBuffer = GetAs<FRDGBufferUAV>(UAV);
 		break;
 	case EResourceType::STRUCTURED_BUFFER:
 		ResourceTypeEnum = EByteBufferResourceType::Float4_StructuredBuffer;
-		PassParameters->DstStructuredBuffer = GraphBuilder.CreateUAV(GetAsBuffer(Resource));
+		PassParameters->DstStructuredBuffer = GetAs<FRDGBufferUAV>(UAV);
 		break;
 	case EResourceType::TEXTURE:
 		ResourceTypeEnum = EByteBufferResourceType::Float4_Texture;
-		PassParameters->DstTexture = GraphBuilder.CreateUAV(GetAsTexture(Resource));
+		PassParameters->DstTexture = GetAs<FRDGTextureUAV>(UAV);
 		PassParameters->Float4sPerLine = CalculateFloat4sPerLine();
 		break;
 	default:
@@ -296,8 +307,12 @@ void MemsetResource(FRDGBuilder& GraphBuilder, FRDGViewableResource* Resource, c
 		FIntVector(FMath::DivideAndRoundUp(Params.Count / Divisor, 64u), 1, 1));
 }
 
-void MemcpyResource(FRDGBuilder& GraphBuilder, FRDGViewableResource* DstResource, FRDGViewableResource* SrcResource, const FMemcpyResourceParams& Params)
+void MemcpyResource(FRDGBuilder& GraphBuilder, FRDGUnorderedAccessView* UAV, FRDGShaderResourceView* SRV, const FMemcpyResourceParams& Params)
 {
+	check(UAV && SRV);
+	FRDGViewableResource* DstResource = UAV->GetParent();
+	FRDGViewableResource* SrcResource = SRV->GetParent();
+
 	const EResourceType ResourceType = GetResourceType(DstResource);
 	checkf(ResourceType == GetResourceType(SrcResource), TEXT("Unable to MemcpyResource because the source and destination view types don't match."));
 
@@ -324,23 +339,23 @@ void MemcpyResource(FRDGBuilder& GraphBuilder, FRDGViewableResource* DstResource
 		{
 		case EResourceType::BYTEBUFFER:
 			ResourceTypeEnum = EByteBufferResourceType::Uint_Buffer;
-			PassParameters->SrcByteAddressBuffer           = GraphBuilder.CreateSRV(GetAsBuffer(SrcResource));
-			PassParameters->Common.DstByteAddressBuffer    = GraphBuilder.CreateUAV(GetAsBuffer(DstResource), ERDGUnorderedAccessViewFlags::SkipBarrier);
+			PassParameters->SrcByteAddressBuffer           = GetAs<FRDGBufferSRV>(SRV);
+			PassParameters->Common.DstByteAddressBuffer    = GetAs<FRDGBufferUAV>(UAV);
 			break;
 		case EResourceType::STRUCTURED_BUFFER:
 			ResourceTypeEnum = EByteBufferResourceType::Float4_StructuredBuffer;
-			PassParameters->SrcStructuredBuffer            = GraphBuilder.CreateSRV(GetAsBuffer(SrcResource));
-			PassParameters->Common.DstStructuredBuffer = GraphBuilder.CreateUAV(GetAsBuffer(DstResource), ERDGUnorderedAccessViewFlags::SkipBarrier);
+			PassParameters->SrcStructuredBuffer            = GetAs<FRDGBufferSRV>(SRV);
+			PassParameters->Common.DstStructuredBuffer     = GetAs<FRDGBufferUAV>(UAV);
 			break;
 		case EResourceType::BUFFER:
 			ResourceTypeEnum = EByteBufferResourceType::Float4_Buffer;
-			PassParameters->SrcBuffer                      = GraphBuilder.CreateSRV(GetAsBuffer(SrcResource));
-			PassParameters->Common.DstBuffer               = GraphBuilder.CreateUAV(GetAsBuffer(DstResource), ERDGUnorderedAccessViewFlags::SkipBarrier);
+			PassParameters->SrcBuffer                      = GetAs<FRDGBufferSRV>(SRV);
+			PassParameters->Common.DstBuffer               = GetAs<FRDGBufferUAV>(UAV);
 			break;
 		case EResourceType::TEXTURE:
 			ResourceTypeEnum = EByteBufferResourceType::Float4_Texture;
-			PassParameters->SrcTexture                     = GraphBuilder.CreateSRV(GetAsTexture(SrcResource));
-			PassParameters->Common.DstTexture              = GraphBuilder.CreateUAV(GetAsTexture(DstResource), ERDGUnorderedAccessViewFlags::SkipBarrier);
+			PassParameters->SrcTexture                     = GetAs<FRDGTextureSRV>(SRV);
+			PassParameters->Common.DstTexture              = GetAs<FRDGTextureUAV>(UAV);
 			PassParameters->Common.Float4sPerLine          = CalculateFloat4sPerLine();
 			break;
 		default:
@@ -414,6 +429,64 @@ FRDGBuffer* ResizeStructuredBufferIfNeeded(FRDGBuilder& GraphBuilder, TRefCountP
 	return ResizeBufferIfNeeded(GraphBuilder, ExternalBuffer, FRDGBufferDesc::CreateStructuredDesc(BytesPerElement, NumElements), Name);
 }
 
+FRDGBuffer* ResizeStructuredBufferSOAIfNeeded(FRDGBuilder& GraphBuilder, TRefCountPtr<FRDGPooledBuffer>& ExternalBuffer, const FResizeResourceSOAParams& Params, const TCHAR* Name)
+{
+	const uint32 BytesPerElement = 16;
+	const uint32 ExternalBufferSize = TryGetSize(ExternalBuffer);
+
+	checkf(Params.NumBytes % BytesPerElement == 0, TEXT("NumBytes (%s) must be a multiple of BytesPerElement (%s)"), Params.NumBytes, BytesPerElement);
+	checkf(ExternalBufferSize % BytesPerElement == 0, TEXT("NumBytes (%s) must be a multiple of BytesPerElement (%s)"), ExternalBufferSize, BytesPerElement);
+
+	uint32 NumElements = Params.NumBytes / BytesPerElement;
+	uint32 NumElementsOld = ExternalBufferSize / BytesPerElement;
+
+	checkf(NumElements % Params.NumArrays == 0, TEXT("NumElements (%s) must be a multiple of NumArrays (%s)"), NumElements, Params.NumArrays);
+	checkf(NumElementsOld % Params.NumArrays == 0, TEXT("NumElements (%s) must be a multiple of NumArrays (%s)"), NumElementsOld, Params.NumArrays);
+
+	const FRDGBufferDesc BufferDesc = FRDGBufferDesc::CreateStructuredDesc(BytesPerElement, NumElements);
+
+	FRDGBuffer* InternalBufferNew = nullptr;
+
+	if (!ExternalBuffer)
+	{
+		InternalBufferNew = GraphBuilder.CreateBuffer(BufferDesc, Name);
+		ExternalBuffer = GraphBuilder.ConvertToExternalBuffer(InternalBufferNew);
+		return InternalBufferNew;
+	}
+
+	FRDGBuffer* InternalBufferOld = GraphBuilder.RegisterExternalBuffer(ExternalBuffer);
+
+	const uint32 BufferSize = BufferDesc.GetSize();
+	const uint32 BufferSizeOld = InternalBufferOld->GetSize();
+
+	if (BufferSize != BufferSizeOld)
+	{
+		InternalBufferNew = GraphBuilder.CreateBuffer(BufferDesc, Name);
+		ExternalBuffer = GraphBuilder.ConvertToExternalBuffer(InternalBufferNew);
+
+		FRDGBufferUAV* NewBufferUAV = GraphBuilder.CreateUAV(InternalBufferNew, ERDGUnorderedAccessViewFlags::SkipBarrier);
+		FRDGBufferSRV* OldBufferSRV = GraphBuilder.CreateSRV(InternalBufferOld);
+
+		// Copy data to new buffer
+		uint32 OldArraySize = NumElementsOld / Params.NumArrays;
+		uint32 NewArraySize = NumElements / Params.NumArrays;
+
+		FMemcpyResourceParams MemcpyParams;
+		MemcpyParams.Count = FMath::Min(NewArraySize, OldArraySize);
+
+		for (uint32 Index = 0; Index < Params.NumArrays; Index++)
+		{
+			MemcpyParams.SrcOffset = Index * OldArraySize;
+			MemcpyParams.DstOffset = Index * NewArraySize;
+			MemcpyResource(GraphBuilder, NewBufferUAV, OldBufferSRV, MemcpyParams);
+		}
+
+		return InternalBufferNew;
+	}
+
+	return InternalBufferOld;
+}
+
 FRDGBuffer* ResizeByteAddressBufferIfNeeded(FRDGBuilder& GraphBuilder, TRefCountPtr<FRDGPooledBuffer>& ExternalBuffer, uint32 NumBytes, const TCHAR* Name)
 {
 	// Needs to be aligned to 16 bytes to MemcpyResource to work correctly (otherwise it skips last unaligned elements of the buffer during resize)
@@ -422,6 +495,31 @@ FRDGBuffer* ResizeByteAddressBufferIfNeeded(FRDGBuilder& GraphBuilder, TRefCount
 	return ResizeBufferIfNeeded(GraphBuilder, ExternalBuffer, FRDGBufferDesc::CreateByteAddressDesc(NumBytes), Name);
 }
 
+void FRDGScatterUploadBuffer::Release()
+{
+	check(!ScatterData);
+	ScatterBuffer = nullptr;
+	UploadBuffer = nullptr;
+}
+
+uint32 FRDGScatterUploadBuffer::GetNumBytes() const
+{
+	return TryGetSize(ScatterBuffer) + TryGetSize(UploadBuffer);
+}
+
+void FRDGScatterUploadBuffer::Init(FRDGBuilder& GraphBuilder, TArrayView<const uint32> ElementScatterOffsets, uint32 InNumBytesPerElement, bool bInFloat4Buffer, const TCHAR* DebugName)
+{
+	Init(GraphBuilder, ElementScatterOffsets.Num(), InNumBytesPerElement, bInFloat4Buffer, DebugName);
+	FMemory::ParallelMemcpy(ScatterData, ElementScatterOffsets.GetData(), ElementScatterOffsets.Num() * ElementScatterOffsets.GetTypeSize(), EMemcpyCachePolicy::StoreUncached);
+	NumScatters = ElementScatterOffsets.Num();
+}
+
+void FRDGScatterUploadBuffer::InitPreSized(FRDGBuilder& GraphBuilder, uint32 NumElements, uint32 InNumBytesPerElement, bool bInFloat4Buffer, const TCHAR* DebugName)
+{
+	Init(GraphBuilder, NumElements, InNumBytesPerElement, bInFloat4Buffer, DebugName);
+	NumScatters = NumElements;
+}
+
 void FRDGScatterUploadBuffer::Init(FRDGBuilder& GraphBuilder, uint32 NumElements, uint32 InNumBytesPerElement, bool bInFloat4Buffer, const TCHAR* Name)
 {
 	TRACE_CPUPROFILER_EVENT_SCOPE(FRDGScatterUploadBuffer::Init);
@@ -558,7 +656,7 @@ void FRDGScatterUploadBuffer::ResourceUploadTo(FRDGBuilder& GraphBuilder, FRDGVi
 
 		FComputeShaderUtils::AddPass(
 			GraphBuilder,
-			RDG_EVENT_NAME("ScatterUpload[%d] (Offset: %u, GroupSize: %u)", Parameters.Common.SrcOffset, LoopNumDispatch),
+			RDG_EVENT_NAME("ScatterUpload[%d] (Resource: %s, Offset: %u, GroupSize: %u)", LoopIdx, DstResource->Name, Parameters.Common.SrcOffset, LoopNumDispatch),
 			ComputeShader,
 			GraphBuilder.AllocParameters(&Parameters),
 			FIntVector(LoopNumDispatch, 1, 1));
diff --git a/Engine/Source/Runtime/RenderCore/Public/UnifiedBuffer.h b/Engine/Source/Runtime/RenderCore/Public/UnifiedBuffer.h
--- a/Engine/Source/Runtime/RenderCore/Public/UnifiedBuffer.h
+++ b/Engine/Source/Runtime/RenderCore/Public/UnifiedBuffer.h
@@ -178,22 +178,38 @@ public:
 	}
 };
 
-RENDERCORE_API void MemsetResource(FRDGBuilder& GraphBuilder, FRDGViewableResource* DstResource, const FMemsetResourceParams& Params);
-RENDERCORE_API void MemcpyResource(FRDGBuilder& GraphBuilder, FRDGViewableResource* DstResource, FRDGViewableResource* SrcResource, const FMemcpyResourceParams& Params);
+extern RENDERCORE_API void MemsetResource(FRDGBuilder& GraphBuilder, FRDGBuffer* DstResource, const FMemsetResourceParams& Params);
+extern RENDERCORE_API void MemcpyResource(FRDGBuilder& GraphBuilder, FRDGBuffer* DstResource, FRDGBuffer* SrcResource, const FMemcpyResourceParams& Params);
 
-RENDERCORE_API FRDGBuffer* ResizeBufferIfNeeded(FRDGBuilder& GraphBuilder, TRefCountPtr<FRDGPooledBuffer>& ExternalBuffer, const FRDGBufferDesc& BufferDesc, const TCHAR* Name);
-RENDERCORE_API FRDGBuffer* ResizeBufferIfNeeded(FRDGBuilder& GraphBuilder, TRefCountPtr<FRDGPooledBuffer>& ExternalBuffer, EPixelFormat Format, uint32 NumElements, const TCHAR* Name);
-RENDERCORE_API FRDGBuffer* ResizeStructuredBufferIfNeeded(FRDGBuilder& GraphBuilder, TRefCountPtr<FRDGPooledBuffer>& ExternalBuffer, uint32 NumBytes, const TCHAR* Name);
-RENDERCORE_API FRDGBuffer* ResizeByteAddressBufferIfNeeded(FRDGBuilder& GraphBuilder, TRefCountPtr<FRDGPooledBuffer>& ExternalBuffer, uint32 NumBytes, const TCHAR* Name);
+extern RENDERCORE_API void MemsetResource(FRDGBuilder& GraphBuilder, FRDGUnorderedAccessView* DstResource, const FMemsetResourceParams& Params);
+extern RENDERCORE_API void MemcpyResource(FRDGBuilder& GraphBuilder, FRDGUnorderedAccessView* DstResource, FRDGShaderResourceView* SrcResource, const FMemcpyResourceParams& Params);
 
-class FRDGScatterUploadBuffer
+extern RENDERCORE_API FRDGBuffer* ResizeBufferIfNeeded(FRDGBuilder& GraphBuilder, TRefCountPtr<FRDGPooledBuffer>& ExternalBuffer, const FRDGBufferDesc& BufferDesc, const TCHAR* Name);
+extern RENDERCORE_API FRDGBuffer* ResizeBufferIfNeeded(FRDGBuilder& GraphBuilder, TRefCountPtr<FRDGPooledBuffer>& ExternalBuffer, EPixelFormat Format, uint32 NumElements, const TCHAR* Name);
+extern RENDERCORE_API FRDGBuffer* ResizeStructuredBufferIfNeeded(FRDGBuilder& GraphBuilder, TRefCountPtr<FRDGPooledBuffer>& ExternalBuffer, uint32 NumBytes, const TCHAR* Name);
+extern RENDERCORE_API FRDGBuffer* ResizeStructuredBufferSOAIfNeeded(FRDGBuilder& GraphBuilder, TRefCountPtr<FRDGPooledBuffer>& ExternalBuffer, const FResizeResourceSOAParams& Params, const TCHAR* DebugName);
+extern RENDERCORE_API FRDGBuffer* ResizeByteAddressBufferIfNeeded(FRDGBuilder& GraphBuilder, TRefCountPtr<FRDGPooledBuffer>& ExternalBuffer, uint32 NumBytes, const TCHAR* Name);
+
+class RENDERCORE_API FRDGScatterUploadBuffer
 {
 public:
 	enum { PrimitiveDataStrideInFloat4s = FScatterUploadBuffer::PrimitiveDataStrideInFloat4s };
 
-	RENDERCORE_API void Init(FRDGBuilder& GraphBuilder, uint32 NumElements, uint32 NumBytesPerElement, bool bInFloat4Buffer, const TCHAR* Name);
+	/**
+	 * Init with presized num scatters, expecting each to be set at a later point. Requires the user to keep track of the offsets to use.
+	 */
+	void InitPreSized(FRDGBuilder& GraphBuilder, uint32 NumElements, uint32 InNumBytesPerElement, bool bInFloat4Buffer, const TCHAR* DebugName);
 
-	RENDERCORE_API void ResourceUploadTo(FRDGBuilder& GraphBuilder, FRDGViewableResource* DstResource);
+	/**
+	 * Init with pre-existing destination index data, performs a bulk-copy.
+	 */
+	void Init(FRDGBuilder& GraphBuilder, TArrayView<const uint32> ElementScatterOffsets, uint32 InNumBytesPerElement, bool bInFloat4Buffer, const TCHAR* DebugName);
+
+	void Init(FRDGBuilder& GraphBuilder, uint32 NumElements, uint32 NumBytesPerElement, bool bInFloat4Buffer, const TCHAR* Name);
+
+	void ResourceUploadTo(FRDGBuilder& GraphBuilder, FRDGViewableResource* DstResource);
+
+	void Release();
 
 	void Add(uint32 Index, const void* Data, uint32 Num = 1)
 	{
@@ -243,8 +259,10 @@ public:
 		return UploadData + ElementIndex * NumBytesPerElement;
 	}
 
+	uint32 GetNumBytes() const;
+
 private:
-	RENDERCORE_API void Reset();
+	void Reset();
 
 	TRefCountPtr<FRDGPooledBuffer> ScatterBuffer;
 	TRefCountPtr<FRDGPooledBuffer> UploadBuffer;
