Adds support for Texture2DArray in DumpGPU

diff --git a/Engine/Extras/GPUDumpViewer/GPUDumpViewer.html b/Engine/Extras/GPUDumpViewer/GPUDumpViewer.html
--- a/Engine/Extras/GPUDumpViewer/GPUDumpViewer.html
+++ b/Engine/Extras/GPUDumpViewer/GPUDumpViewer.html
@@ -265,26 +265,59 @@ function parse_subresource_unique_name(subresource_unique_name)
 	var splitted_name = subresource_unique_name.split(".");
 
 	var subresource_info = {};
-	if (splitted_name[splitted_name.length - 1 ].startsWith('mip') || splitted_name[splitted_name.length - 1 ] == 'stencil')
+	subresource_info['subresource'] = null;
+	subresource_info['array_slice'] = null;
+
+	if (splitted_name[splitted_name.length - 1].startsWith('mip') || splitted_name[splitted_name.length - 1] == 'stencil')
 	{
 		subresource_info['subresource'] = splitted_name.pop();
-		subresource_info['resource'] = splitted_name.join('.');
 	}
-	else
+
+	if (splitted_name[splitted_name.length - 1].startsWith('[') && splitted_name[splitted_name.length - 1].endsWith(']'))
 	{
-		subresource_info['subresource'] = null;
-		subresource_info['resource'] = subresource_unique_name;
+		var array_slice_bracket = splitted_name.pop();
+		subresource_info['array_slice'] = parseInt(array_slice_bracket.substring(1, array_slice_bracket.length - 1));
 	}
+
+	subresource_info['resource'] = splitted_name.join('.');
+
 	return subresource_info;
 }
 
 function get_subresource_unique_name(subresource_info)
 {
-	if (subresource_info['subresource'] === null)
+	var subresource_unique_name = subresource_info['resource'];
+	if (subresource_info['array_slice'] !== null)
 	{
-		return subresource_info['resource'];
+		subresource_unique_name += `.[${subresource_info['array_slice']}]`;
 	}
-	return `${subresource_info['resource']}.${subresource_info['subresource']}`;
+	if (subresource_info['subresource'] !== null)
+	{
+		subresource_unique_name += `.${subresource_info['subresource']}`;
+	}
+	return subresource_unique_name;
+}
+
+function prettify_subresource_unique_name(subresource_info, resource_desc)
+{
+	if (!resource_desc)
+	{
+		return get_subresource_unique_name(subresource_info);
+	}
+
+	var name = resource_desc['Name'];
+
+	if (subresource_info['array_slice'] !== null)
+	{
+		name += ` slice[${subresource_info['array_slice']}]`;
+	}
+
+	if (subresource_info['subresource'] !== null && (subresource_info['subresource'] == 'stencil' || resource_desc['NumMips'] > 1))
+	{
+		name += ` ${subresource_info['subresource']}`;
+	}
+	
+	return name;
 }
 
 function parse_subresource_unique_version_name(subresource_unique_version_name)
@@ -649,18 +682,10 @@ class PassView extends IView
 			var subresource_info = parse_subresource_unique_name(subresource_unique_name);
 			var resource_desc = get_resource_desc(subresource_info['resource']);
 
-			var name = subresource_unique_name;
+			var name = prettify_subresource_unique_name(subresource_info, resource_desc);
 			var href = null;
 			if (resource_desc)
 			{
-				if (subresource_info['subresource'] == null)
-				{
-					name = resource_desc['Name'];
-				}
-				else
-				{
-					name = `${resource_desc['Name']} ${subresource_info['subresource']}`;
-				}
 				href = `#display_input_resource(${this.pass_id},'${subresource_unique_name}');`
 			}
 
@@ -691,19 +716,10 @@ class PassView extends IView
 			var subresource_info = parse_subresource_unique_name(subresource_unique_name);
 			var resource_desc = get_resource_desc(subresource_info['resource']);
 
-			var name = subresource_unique_name;
+			var name = prettify_subresource_unique_name(subresource_info, resource_desc);
 			var href = null;
 			if (resource_desc)
 			{
-				if (subresource_info['subresource'] == null)
-				{
-					name = resource_desc['Name'];
-				}
-				else
-				{
-					name = `${resource_desc['Name']} ${subresource_info['subresource']}`;
-				}
-
 				if (draw_id >= 0)
 				{
 					href = `#display_draw_output_resource(${this.pass_id},'${subresource_unique_name}',${draw_id});`;
@@ -922,7 +938,7 @@ function display_resource_internal(subresource_version_info)
 		return;
 	}
 
-	if (resource_desc['Type'] == 'Texture2D')
+	if (resource_desc['Type'] == 'Texture2D' || resource_desc['Type'] == 'Texture2DArray')
 	{
 		g_view.set_resource_view(new TextureView(subresource_version_info, resource_desc));
 	}
@@ -1486,7 +1502,7 @@ class TextureView extends ResourceView
 				</table>`;
 		}
 
-		var title = `${this.resource_desc['Name']} ${this.subresource_version_info['subresource']}`;
+		var title = prettify_subresource_unique_name(this.subresource_version_info, this.resource_desc);
 		if (this.subresource_version_info['draw'] >= 0)
 		{
 			title += ` draw:${this.subresource_version_info['draw']}`;
@@ -2281,12 +2297,7 @@ void main(void)
 			text += `Pass:\n\t${pass_data['EventName']}\n`;
 
 			{
-				var name = texture_view.resource_desc['Name'];
-				if (texture_view.subresource_version_info['subresource'])
-				{
-					name = `${texture_view.resource_desc['Name']} ${texture_view.subresource_version_info['subresource']}`;
-				}
-
+				var name = prettify_subresource_unique_name(texture_view.subresource_version_info, texture_view.resource_desc);
 				text += `${is_output_resource ? 'Output' : 'Input'} resource:\n\t${name}\n`;
 			}
 
diff --git a/Engine/Source/Runtime/RenderCore/Private/DumpGPU.cpp b/Engine/Source/Runtime/RenderCore/Private/DumpGPU.cpp
--- a/Engine/Source/Runtime/RenderCore/Private/DumpGPU.cpp
+++ b/Engine/Source/Runtime/RenderCore/Private/DumpGPU.cpp
@@ -390,12 +390,18 @@ struct FRDGResourceDumpContext
 		return PtrToString(Resource);
 	}
 
-	static FString GetUniqueSubResourceName(FRDGTextureSRVDesc SubResourceDesc)
+	static FString GetUniqueSubResourceName(const FRDGTextureSRVDesc& SubResourceDesc)
 	{
 		check(SubResourceDesc.NumMipLevels == 1);
+		check(!SubResourceDesc.Texture->Desc.IsTextureArray() || SubResourceDesc.NumArraySlices == 1);
 
 		FString UniqueResourceName = GetUniqueResourceName(SubResourceDesc.Texture);
 
+		if (SubResourceDesc.Texture->Desc.IsTextureArray())
+		{
+			UniqueResourceName += FString::Printf(TEXT(".[%d]"), SubResourceDesc.FirstArraySlice);
+		}
+
 		if (SubResourceDesc.Format == PF_X24_G8)
 		{
 			return FString::Printf(TEXT("%s.stencil"), *UniqueResourceName);
@@ -588,42 +594,30 @@ struct FRDGResourceDumpContext
 		IsDumpedToDisk.Add(static_cast<const void*>(Ptr));
 	}
 
-	void GetResourceDumpInfo(
-		const FRDGPass* Pass,
-		const FRDGResource* Resource,
-		bool bIsOutputResource,
-		bool* bOutDumpResourceInfos,
-		bool* bOutDumpResourceBinary)
+	bool HasResourceEverBeenDumped(const FRDGResource* Resource) const
 	{
-		*bOutDumpResourceInfos = false;
-		*bOutDumpResourceBinary = bIsOutputResource;
-
-		//if (bIsOutputResource)
-		//{
-		//	*OutResourceVersionDumpName = GetResourceVersionDumpName(Pass, Resource);
-		//}
-
-		if (!LastResourceVersion.Contains(Resource))
-		{
-			// First time we ever see this resource, so dump it's info to disk
-			*bOutDumpResourceInfos = true;
-
-			// If not an output, it might be a resource undumped by r.DumpGPU.Root or external texture so still dump it as v0.
-			if (!bIsOutputResource)
-			{
-				*bOutDumpResourceBinary = true;
-			}
+		return LastResourceVersion.Contains(Resource);
+	}
 
-			if (LastResourceVersion.Num() % 1024 == 0)
-			{
-				LastResourceVersion.Reserve(LastResourceVersion.Num() + 1024);
-			}
-			LastResourceVersion.Add(Resource, Pass);
-		}
-		else
+	void RegisterResourceAsDumped(
+		const FRDGResource* Resource,
+		const FRDGPass* LastModifyingPass)
+	{
+		check(!LastResourceVersion.Contains(Resource));
+		if (LastResourceVersion.Num() % 1024 == 0)
 		{
-			LastResourceVersion[Resource] = Pass;
+			LastResourceVersion.Reserve(LastResourceVersion.Num() + 1024);
 		}
+		LastResourceVersion.Add(Resource, LastModifyingPass);
+	}
+
+	void UpdateResourceLastModifyingPass(
+		const FRDGResource* Resource,
+		const FRDGPass* LastModifyingPass)
+	{
+		check(LastResourceVersion.Contains(Resource));
+		check(LastModifyingPass);
+		LastResourceVersion[Resource] = LastModifyingPass;
 	}
 
 	FString ToJson(EPixelFormat Format)
@@ -776,6 +770,9 @@ struct FRDGResourceDumpContext
 	
 	FTextureSubresourceDumpDesc TranslateSubresourceDumpDesc(FRDGTextureSRVDesc SubresourceDesc)
 	{
+		check(SubresourceDesc.NumMipLevels == 1);
+		check(!SubresourceDesc.Texture->Desc.IsTextureArray() || SubresourceDesc.NumArraySlices == 1);
+
 		const FRDGTextureDesc& Desc = SubresourceDesc.Texture->Desc;
 
 		FTextureSubresourceDumpDesc SubresourceDumpDesc;
@@ -798,7 +795,7 @@ struct FRDGResourceDumpContext
 			bIsUnsupported = true;
 		}
 
-		if (!bIsUnsupported && Desc.IsTexture2D() && !Desc.IsMultisample() && !Desc.IsTextureArray())
+		if (!bIsUnsupported && Desc.IsTexture2D() && !Desc.IsMultisample())
 		{
 			SubresourceDumpDesc.SubResourceExtent.X = Desc.Extent.X >> SubresourceDesc.MipLevel;
 			SubresourceDumpDesc.SubResourceExtent.Y = Desc.Extent.Y >> SubresourceDesc.MipLevel;
@@ -852,7 +849,16 @@ struct FRDGResourceDumpContext
 		SubresourceDumpDesc.ByteSize = SIZE_T(SubresourceDumpDesc.SubResourceExtent.X) * SIZE_T(SubresourceDumpDesc.SubResourceExtent.Y) * SIZE_T(GPixelFormats[SubresourceDumpDesc.PreprocessedPixelFormat].BlockBytes);
 
 		// Whether the subresource need preprocessing pass before copy into staging.
-		SubresourceDumpDesc.bPreprocessForStaging = SubresourceDumpDesc.PreprocessedPixelFormat != Desc.Format || SubresourceDesc.Texture->Desc.NumMips > 1;
+		{
+			// If need a pixel format conversion, use a pixel shader to do it.
+			SubresourceDumpDesc.bPreprocessForStaging |= SubresourceDumpDesc.PreprocessedPixelFormat != Desc.Format;
+
+			// If the texture has a mip chain, use pixel shader to correctly copy the right mip given RHI doesn't support copy from mip levels. Also on Mip 0 to avoid bugs on D3D11
+			SubresourceDumpDesc.bPreprocessForStaging |= SubresourceDesc.Texture->Desc.NumMips > 1;
+
+			// If the texture is an array, use pixel shader to correctly copy the right slice given RHI doesn't support copy from slices
+			SubresourceDumpDesc.bPreprocessForStaging |= SubresourceDesc.Texture->Desc.IsTextureArray();
+		}
 
 		return SubresourceDumpDesc;
 	}
@@ -1093,7 +1099,9 @@ struct FRDGResourceDumpContext
 		TArray<TSharedPtr<FJsonValue>>& OutputResourceNames,
 		const FRDGPass* Pass,
 		FRDGTextureSRVDesc SubresourceDesc,
-		bool bIsOutputResource)
+		bool bRegisterSubResourceToPass,
+		bool bIsOutputResource,
+		bool bAllowDumpBinary)
 	{
 		int32 DumpTextureMode = GDumpTextureCVar.GetValueOnRenderThread();
 
@@ -1103,37 +1111,28 @@ struct FRDGResourceDumpContext
 		}
 
 		const FRDGTextureDesc& Desc = SubresourceDesc.Texture->Desc;
-		const FString UniqueResourceName = GetUniqueResourceName(SubresourceDesc.Texture);
 		const FString UniqueResourceSubResourceName = GetUniqueSubResourceName(SubresourceDesc);
 		const FTextureSubresourceDumpDesc SubresourceDumpDesc = TranslateSubresourceDumpDesc(SubresourceDesc);
 
-		if (bIsOutputResource)
-		{
-			OutputResourceNames.AddUnique(MakeShareable(new FJsonValueString(UniqueResourceSubResourceName)));
-		}
-		else
-		{
-			InputResourceNames.AddUnique(MakeShareable(new FJsonValueString(UniqueResourceSubResourceName)));
-		}
-
-		bool bDumpResourceInfos;
-		bool bDumpResourceBinary;
-		GetResourceDumpInfo(Pass, SubresourceDesc.Texture, bIsOutputResource, &bDumpResourceInfos, &bDumpResourceBinary);
-
-		// Dump the information of the texture to json file.
-		if (bDumpResourceInfos)
+		if (bRegisterSubResourceToPass)
 		{
-			TSharedPtr<FJsonObject> JsonObject = ToJson(UniqueResourceName, SubresourceDesc.Texture->Name, Desc);
-			DumpJsonToFile(JsonObject, FString(FRDGResourceDumpContext::kBaseDir) / TEXT("ResourceDescs.json"), FILEWRITE_Append);
+			if (bIsOutputResource)
+			{
+				OutputResourceNames.AddUnique(MakeShareable(new FJsonValueString(UniqueResourceSubResourceName)));
+			}
+			else
+			{
+				InputResourceNames.AddUnique(MakeShareable(new FJsonValueString(UniqueResourceSubResourceName)));
+			}
 		}
 
-		if (!SubresourceDumpDesc.IsDumpSupported())
+		// Early return if this resource shouldn't be dumped.
+		if (!bAllowDumpBinary || DumpTextureMode != 2)
 		{
 			return;
 		}
 
-		// Early return if this resource shouldn't be dumped.
-		if (!bDumpResourceBinary || DumpTextureMode != 2)
+		if (!SubresourceDumpDesc.IsDumpSupported())
 		{
 			return;
 		}
@@ -1205,40 +1204,154 @@ struct FRDGResourceDumpContext
 		}
 	}
 
-	void AddDumpTexturePasses(
+	void AddDumpTextureAllSubResourcesPasses(
 		FRDGBuilder& GraphBuilder,
 		TArray<TSharedPtr<FJsonValue>>& InputResourceNames,
 		TArray<TSharedPtr<FJsonValue>>& OutputResourceNames,
 		const FRDGPass* Pass,
-		FRDGTextureSRVDesc SubresourceRangeDesc,
-		bool bIsOutputResource)
+		const FRDGTextureSRVDesc& SubresourceRangeDesc,
+		bool bRegisterSubResourceToPass,
+		bool bIsOutputResource,
+		bool bAllowDumpBinary)
 	{
-		if (SubresourceRangeDesc.Format == PF_X24_G8)
+		if (SubresourceRangeDesc.Texture->Desc.IsTextureArray() && SubresourceRangeDesc.NumArraySlices != 1)
 		{
+			// If this is an array, recursively dump all subresource of each individual array slice.
+			uint32 ArraySliceStart = SubresourceRangeDesc.FirstArraySlice;
+			uint32 ArraySliceCount = (SubresourceRangeDesc.NumArraySlices > 0) ? SubresourceRangeDesc.NumArraySlices : (SubresourceRangeDesc.Texture->Desc.ArraySize - SubresourceRangeDesc.FirstArraySlice);
+
+			for (uint32 ArraySlice = ArraySliceStart; ArraySlice < (ArraySliceStart + ArraySliceCount); ArraySlice++)
+			{
+				FRDGTextureSRVDesc SubresourceDesc = SubresourceRangeDesc;
+				SubresourceDesc.FirstArraySlice = ArraySlice;
+				SubresourceDesc.NumArraySlices = 1;
+
+				AddDumpTextureAllSubResourcesPasses(
+					GraphBuilder,
+					InputResourceNames,
+					OutputResourceNames,
+					Pass,
+					SubresourceDesc,
+					bRegisterSubResourceToPass,
+					bIsOutputResource,
+					bAllowDumpBinary);
+			}
+		}
+		else if (SubresourceRangeDesc.Format == PF_X24_G8)
+		{
+			// Dump the stencil.
 			AddDumpTextureSubResourcePass(
 				GraphBuilder,
 				InputResourceNames,
 				OutputResourceNames,
 				Pass,
 				SubresourceRangeDesc,
-				bIsOutputResource);
+				bRegisterSubResourceToPass,
+				bIsOutputResource,
+				bAllowDumpBinary);
 		}
 		else
 		{
 			for (int32 MipLevel = SubresourceRangeDesc.MipLevel; MipLevel < (SubresourceRangeDesc.MipLevel + SubresourceRangeDesc.NumMipLevels); MipLevel++)
 			{
-				FRDGTextureSRVDesc SubresourceDesc = FRDGTextureSRVDesc::CreateForMipLevel(SubresourceRangeDesc.Texture, MipLevel);
+				FRDGTextureSRVDesc SubresourceDesc = SubresourceRangeDesc;
+				SubresourceDesc.MipLevel = MipLevel;
+				SubresourceDesc.NumMipLevels = 1;
+
 				AddDumpTextureSubResourcePass(
 					GraphBuilder,
 					InputResourceNames,
 					OutputResourceNames,
 					Pass,
 					SubresourceDesc,
-					bIsOutputResource);
+					bRegisterSubResourceToPass,
+					bIsOutputResource,
+					bAllowDumpBinary);
 			}
 		}
 	}
 
+	void AddDumpTexturePasses(
+		FRDGBuilder& GraphBuilder,
+		TArray<TSharedPtr<FJsonValue>>& InputResourceNames,
+		TArray<TSharedPtr<FJsonValue>>& OutputResourceNames,
+		const FRDGPass* Pass,
+		const FRDGTextureSRVDesc& SubresourceRangeDesc,
+		bool bIsOutputResource)
+	{
+		bool bAllowDumpBinary = true;
+		if (!HasResourceEverBeenDumped(SubresourceRangeDesc.Texture))
+		{
+			RegisterResourceAsDumped(SubresourceRangeDesc.Texture, /* LastModifyingPass = */ bIsOutputResource ? Pass : nullptr);
+
+			// Dump the descriptor of the resource
+			{
+				const FString UniqueResourceName = GetUniqueResourceName(SubresourceRangeDesc.Texture);
+				TSharedPtr<FJsonObject> JsonObject = ToJson(UniqueResourceName, SubresourceRangeDesc.Texture->Name, SubresourceRangeDesc.Texture->Desc);
+				DumpJsonToFile(JsonObject, FString(FRDGResourceDumpContext::kBaseDir) / TEXT("ResourceDescs.json"), FILEWRITE_Append);
+			}
+
+			// If the resource is input but has not been dumped ever before, make sure to dump all of its sub resources instead of just the sub resource described by SubresourceRangeDesc
+			if (!bIsOutputResource)
+			{
+				// Dump the stencil if there is one
+				if (IsStencilFormat(SubresourceRangeDesc.Texture->Desc.Format))
+				{
+					AddDumpTextureAllSubResourcesPasses(
+						GraphBuilder,
+						InputResourceNames,
+						OutputResourceNames,
+						Pass,
+						FRDGTextureSRVDesc::CreateWithPixelFormat(SubresourceRangeDesc.Texture, PF_X24_G8),
+						/* bRegisterSubResourceToPass = */ false,
+						bIsOutputResource,
+						/* bAllowDumpBinary = */ true);
+				}
+
+				// Dump all the remaining sub resource (mip levels, array slices...)
+				AddDumpTextureAllSubResourcesPasses(
+					GraphBuilder,
+					InputResourceNames,
+					OutputResourceNames,
+					Pass,
+					FRDGTextureSRVDesc::Create(SubresourceRangeDesc.Texture),
+					/* bRegisterSubResourceToPass = */ false,
+					bIsOutputResource,
+					/* bAllowDumpBinary = */ true);
+
+				// Still register the input subresources used without dump any binary.
+				return AddDumpTextureAllSubResourcesPasses(
+					GraphBuilder,
+					InputResourceNames,
+					OutputResourceNames,
+					Pass,
+					SubresourceRangeDesc,
+					/* bRegisterSubResourceToPass = */ true,
+					bIsOutputResource,
+					/* bAllowDumpBinary = */ false);
+			}
+		}
+		else if (bIsOutputResource)
+		{
+			UpdateResourceLastModifyingPass(SubresourceRangeDesc.Texture, Pass);
+		}
+		else
+		{
+			// The texture has already been dumped, so can early return given it is not being modified by this Pass.
+			bAllowDumpBinary = false;
+		}
+
+		return AddDumpTextureAllSubResourcesPasses(
+			GraphBuilder,
+			InputResourceNames,
+			OutputResourceNames,
+			Pass,
+			SubresourceRangeDesc,
+			/* bRegisterSubResourceToPass = */ true,
+			bIsOutputResource,
+			bAllowDumpBinary);
+	}
+
 	void AddDumpBufferPass(
 		FRDGBuilder& GraphBuilder,
 		TArray<TSharedPtr<FJsonValue>>& InputResourceNames,
@@ -1268,9 +1381,27 @@ struct FRDGResourceDumpContext
 		const FRDGBufferDesc& Desc = Buffer->Desc;
 		const int32 ByteSize = Desc.GetSize();
 
-		bool bDumpResourceInfos;
-		bool bDumpResourceBinary;
-		GetResourceDumpInfo(Pass, Buffer, bIsOutputResource, &bDumpResourceInfos, &bDumpResourceBinary);
+		bool bDumpResourceInfos = false;
+		bool bDumpResourceBinary = bIsOutputResource;
+		{
+			if (!HasResourceEverBeenDumped(Buffer))
+			{
+				// First time we ever see this buffer, so dump it's info to disk
+				bDumpResourceInfos = true;
+
+				// If not an output, it might be a resource undumped by r.DumpGPU.Root or external texture so still dump it as v0.
+				if (!bIsOutputResource)
+				{
+					bDumpResourceBinary = true;
+				}
+
+				RegisterResourceAsDumped(Buffer, bIsOutputResource ? Pass : nullptr);
+			}
+			else if (bIsOutputResource)
+			{
+				UpdateResourceLastModifyingPass(Buffer, Pass);
+			}
+		}
 
 		// Dump the information of the buffer to json file.
 		if (bDumpResourceInfos)
@@ -1807,8 +1938,15 @@ void FRDGBuilder::DumpResourcePassOutputs(const FRDGPass* Pass)
 			{
 				if (UAV->Desc.MetaData == ERHITextureMetaDataAccess::None)
 				{
-					FRDGTextureSRVDesc TextureSubResource = FRDGTextureSRVDesc::CreateForMipLevel(UAV->Desc.Texture, UAV->Desc.MipLevel);
-					GRDGResourceDumpContext->AddDumpTextureSubResourcePass(*this, InputResourceNames, OutputResourceNames, Pass, TextureSubResource, /* bIsOutputResource = */ true);
+					FRDGTextureSRVDesc TextureSubResource = FRDGTextureSRVDesc::Create(UAV->Desc.Texture);
+					TextureSubResource.MipLevel = UAV->Desc.MipLevel;
+					TextureSubResource.NumMipLevels = 1;
+					if (UAV->Desc.Texture->Desc.IsTextureArray())
+					{
+						TextureSubResource.FirstArraySlice = UAV->Desc.FirstArraySlice;
+						TextureSubResource.NumArraySlices = UAV->Desc.NumArraySlices;
+					}
+					GRDGResourceDumpContext->AddDumpTexturePasses(*this, InputResourceNames, OutputResourceNames, Pass, TextureSubResource, /* bIsOutputResource = */ true);
 				}
 				else
 				{
@@ -1901,7 +2039,14 @@ void FRDGBuilder::DumpResourcePassOutputs(const FRDGPass* Pass)
 			{
 				FRDGTextureRef Texture = RenderTarget.GetTexture();
 				FRDGTextureSRVDesc TextureSubResource = FRDGTextureSRVDesc::CreateForMipLevel(Texture, RenderTarget.GetMipIndex());
-				GRDGResourceDumpContext->AddDumpTextureSubResourcePass(*this, InputResourceNames, OutputResourceNames, Pass, TextureSubResource, /* bIsOutputResource = */ true);
+
+				if (Texture->Desc.IsTextureArray())
+				{
+					TextureSubResource.FirstArraySlice = RenderTarget.GetArraySlice();
+					TextureSubResource.NumArraySlices = 1;
+				}
+
+				GRDGResourceDumpContext->AddDumpTexturePasses(*this, InputResourceNames, OutputResourceNames, Pass, TextureSubResource, /* bIsOutputResource = */ true);
 			});
 
 			const FDepthStencilBinding& DepthStencil = RenderTargets.DepthStencil;
@@ -1913,7 +2058,13 @@ void FRDGBuilder::DumpResourcePassOutputs(const FRDGPass* Pass)
 				if (DepthStencilAccess.IsUsingDepth())
 				{
 					FRDGTextureSRVDesc TextureSubResource = FRDGTextureSRVDesc::CreateForMipLevel(Texture, 0);
-					GRDGResourceDumpContext->AddDumpTextureSubResourcePass(
+					if (Texture->Desc.IsTextureArray())
+					{
+						TextureSubResource.FirstArraySlice = 0;
+						TextureSubResource.NumArraySlices = 1;
+					}
+
+					GRDGResourceDumpContext->AddDumpTexturePasses(
 						*this, InputResourceNames, OutputResourceNames, Pass, TextureSubResource,
 						/* bIsOutputResource = */ DepthStencilAccess.IsDepthWrite());
 				}
@@ -1921,7 +2072,13 @@ void FRDGBuilder::DumpResourcePassOutputs(const FRDGPass* Pass)
 				if (DepthStencilAccess.IsUsingStencil())
 				{
 					FRDGTextureSRVDesc TextureSubResource = FRDGTextureSRVDesc::CreateWithPixelFormat(Texture, PF_X24_G8);
-					GRDGResourceDumpContext->AddDumpTextureSubResourcePass(
+					if (Texture->Desc.IsTextureArray())
+					{
+						TextureSubResource.FirstArraySlice = 0;
+						TextureSubResource.NumArraySlices = 1;
+					}
+
+					GRDGResourceDumpContext->AddDumpTexturePasses(
 						*this, InputResourceNames, OutputResourceNames, Pass, TextureSubResource,
 						/* bIsOutputResource = */ DepthStencilAccess.IsStencilWrite());
 				}
