Adds OS name and version into DumpGPU's informations tab

diff --git a/Engine/Source/Runtime/RenderCore/Private/DumpGPU.cpp b/Engine/Source/Runtime/RenderCore/Private/DumpGPU.cpp
--- a/Engine/Source/Runtime/RenderCore/Private/DumpGPU.cpp
+++ b/Engine/Source/Runtime/RenderCore/Private/DumpGPU.cpp
@@ -1687,6 +1687,13 @@ FString FRDGBuilder::BeginResourceDump(const TCHAR* Cmd)
 		FString BuildConfiguration = LexToString(FApp::GetBuildConfiguration());
 		FString BuildTarget = LexToString(FApp::GetBuildTargetType());
 
+		FString OSLabel, OSVersion;
+		FPlatformMisc::GetOSVersions(OSLabel, OSVersion);
+		if (!OSVersion.IsEmpty())
+		{
+			OSLabel = FString::Printf(TEXT("%s %s"), *OSLabel, *OSVersion);
+		}
+
 		FGPUDriverInfo GPUDriverInfo = FPlatformMisc::GetGPUDriverInfo(GRHIAdapterName);
 
 		TSharedPtr<FJsonObject> JsonObject = MakeShareable(new FJsonObject);
@@ -1701,6 +1708,7 @@ FString FRDGBuilder::BeginResourceDump(const TCHAR* Cmd)
 		JsonObject->SetStringField(TEXT("BuildConfiguration"), BuildConfiguration);
 		JsonObject->SetNumberField(TEXT("Build64Bits"), (PLATFORM_64BITS ? 1 : 0));
 		JsonObject->SetStringField(TEXT("Platform"), FPlatformProperties::IniPlatformName());
+		JsonObject->SetStringField(TEXT("OS"), OSLabel);
 		JsonObject->SetStringField(TEXT("DeviceName"), FPlatformProcess::ComputerName());
 		JsonObject->SetStringField(TEXT("CPUVendor"), FPlatformMisc::GetCPUVendor());
 		JsonObject->SetStringField(TEXT("CPUBrand"), FPlatformMisc::GetCPUBrand());
