Fixed test and validation regressions due to RDG changes.

diff --git a/Engine/Source/Runtime/RenderCore/Private/DumpGPU.cpp b/Engine/Source/Runtime/RenderCore/Private/DumpGPU.cpp
--- a/Engine/Source/Runtime/RenderCore/Private/DumpGPU.cpp
+++ b/Engine/Source/Runtime/RenderCore/Private/DumpGPU.cpp
@@ -173,8 +173,11 @@ BEGIN_SHADER_PARAMETER_STRUCT(FDumpBufferPass, )
 	RDG_BUFFER_ACCESS(Buffer, ERHIAccess::CopySrc)
 END_SHADER_PARAMETER_STRUCT()
 
-struct FRDGResourceDumpContext
+}
+
+class FRDGResourceDumpContext
 {
+public:
 	static constexpr const TCHAR* kBaseDir = TEXT("Base/");
 	static constexpr const TCHAR* kPassesDir = TEXT("Passes/");
 	static constexpr const TCHAR* kResourcesDir = TEXT("Resources/");
@@ -1151,6 +1154,13 @@ struct FRDGResourceDumpContext
 			return;
 		}
 
+		const FRDGViewableResource::FAccessModeState AccessModeState = SubresourceDesc.Texture->AccessModeState;
+
+		if (AccessModeState.Mode == FRDGViewableResource::EAccessMode::External)
+		{
+			GraphBuilder.UseInternalAccessMode(SubresourceDesc.Texture);
+		}
+
 		// Dump the resource's binary to a .bin file.
 		{
 			FString DumpFilePath = kResourcesDir / FString::Printf(TEXT("%s.v%016x.bin"), *UniqueResourceSubResourceName, PtrToUint(bIsOutputResource ? Pass : nullptr));
@@ -1202,6 +1212,11 @@ struct FRDGResourceDumpContext
 
 			ResourcesDumpPasses++;
 		}
+
+		if (AccessModeState.Mode == FRDGViewableResource::EAccessMode::External)
+		{
+			GraphBuilder.UseExternalAccessMode(SubresourceDesc.Texture, AccessModeState.Access, AccessModeState.Pipelines);
+		}
 	}
 
 	void AddDumpTextureAllSubResourcesPasses(
@@ -1440,6 +1455,13 @@ struct FRDGResourceDumpContext
 				return;
 			}
 
+			const FRDGViewableResource::FAccessModeState AccessModeState = Buffer->AccessModeState;
+
+			if (AccessModeState.Mode == FRDGViewableResource::EAccessMode::External)
+			{
+				GraphBuilder.UseInternalAccessMode(Buffer);
+			}
+
 			FDumpBufferPass* PassParameters = GraphBuilder.AllocParameters<FDumpBufferPass>();
 			PassParameters->Buffer = Buffer;
 
@@ -1488,6 +1510,11 @@ struct FRDGResourceDumpContext
 				this->UpdatePassProgress();
 			});
 
+			if (AccessModeState.Mode == FRDGViewableResource::EAccessMode::External)
+			{
+				GraphBuilder.UseExternalAccessMode(Buffer, AccessModeState.Access, AccessModeState.Pipelines);
+			}
+
 			ResourcesDumpPasses++;
 		}
 	}
@@ -1534,9 +1561,7 @@ struct FRDGResourceDumpContext
 
 // 0 = not dumping, MAX_uint64 dump request for next frame, otherwise dump frame counter
 static uint64 DumpingFrameCounter_GameThread = 0;
-FRDGResourceDumpContext* GRDGResourceDumpContext = nullptr;
-
-}
+static FRDGResourceDumpContext* GRDGResourceDumpContext = nullptr;
 
 bool IsDumpingRDGResources()
 {
@@ -1892,7 +1917,7 @@ static const TCHAR* GetPassEventNameWithGPUMask(const FRDGPass* Pass, FString& O
 
 void FRDGBuilder::DumpResourcePassOutputs(const FRDGPass* Pass)
 {
-	if (bSkipAuxiliaryPasses)
+	if (bInDumpPassOutputs)
 	{
 		return;
 	}
@@ -1908,7 +1933,7 @@ void FRDGBuilder::DumpResourcePassOutputs(const FRDGPass* Pass)
 		return;
 	}
 
-	bSkipAuxiliaryPasses = true;
+	bInDumpPassOutputs = true;
 
 	TArray<TSharedPtr<FJsonValue>> InputResourceNames;
 	TArray<TSharedPtr<FJsonValue>> OutputResourceNames;
@@ -2160,7 +2185,7 @@ void FRDGBuilder::DumpResourcePassOutputs(const FRDGPass* Pass)
 
 	GRDGResourceDumpContext->PassesCount++;
 
-	bSkipAuxiliaryPasses = false;
+	bInDumpPassOutputs = false;
 }
 
 #if RDG_DUMP_RESOURCES_AT_EACH_DRAW
diff --git a/Engine/Source/Runtime/RenderCore/Private/RenderGraphBuilder.cpp b/Engine/Source/Runtime/RenderCore/Private/RenderGraphBuilder.cpp
--- a/Engine/Source/Runtime/RenderCore/Private/RenderGraphBuilder.cpp
+++ b/Engine/Source/Runtime/RenderCore/Private/RenderGraphBuilder.cpp
@@ -598,13 +598,13 @@ END_SHADER_PARAMETER_STRUCT()
 
 void FRDGBuilder::UseExternalAccessMode(FRDGViewableResource* Resource, ERHIAccess ReadOnlyAccess, ERHIPipeline Pipelines)
 {
-	IF_RDG_ENABLE_DEBUG(UserValidation.ValidateUseExternalAccessMode(Resource, ReadOnlyAccess, Pipelines));
-
 	if (!GRDGAsyncCompute || !GSupportsEfficientAsyncCompute || IsImmediateMode())
 	{
 		Pipelines = ERHIPipeline::Graphics;
 	}
 
+	IF_RDG_ENABLE_DEBUG(UserValidation.ValidateUseExternalAccessMode(Resource, ReadOnlyAccess, Pipelines));
+
 	auto& AccessModeState = Resource->AccessModeState;
 
 	// We already validated that back-to-back calls to UseExternalAccessMode are valid only if the parameters match,
@@ -661,12 +661,12 @@ void FRDGBuilder::UseInternalAccessMode(FRDGViewableResource* Resource)
 
 void FRDGBuilder::FlushAccessModeQueue()
 {
-	if (AccessModeQueue.IsEmpty() || bSkipAuxiliaryPasses)
+	if (AccessModeQueue.IsEmpty() || bInFlushExternalAccess)
 	{
 		return;
 	}
 
-	bSkipAuxiliaryPasses = true;
+	bInFlushExternalAccess = true;
 
 	FAccessModePassParameters* ParametersByPipeline[] =
 	{
@@ -750,7 +750,7 @@ void FRDGBuilder::FlushAccessModeQueue()
 
 	AccessModeQueue.Reset();
 
-	bSkipAuxiliaryPasses = false;
+	bInFlushExternalAccess = false;
 }
 
 ///////////////////////////////////////////////////////////////////////////////////////////////////
@@ -1485,11 +1485,7 @@ void FRDGBuilder::Execute()
 	}
 
 	// Create the epilogue pass at the end of the graph just prior to compilation.
-	{
-		bSkipAuxiliaryPasses = true;
-		SetupEmptyPass(EpiloguePass = Passes.Allocate<FRDGSentinelPass>(Allocator, RDG_EVENT_NAME("Graph Epilogue")));
-		bSkipAuxiliaryPasses = false;
-	}
+	SetupEmptyPass(EpiloguePass = Passes.Allocate<FRDGSentinelPass>(Allocator, RDG_EVENT_NAME("Graph Epilogue")));
 
 	const FRDGPassHandle ProloguePassHandle = GetProloguePassHandle();
 	const FRDGPassHandle EpiloguePassHandle = GetEpiloguePassHandle();
@@ -1505,8 +1501,6 @@ void FRDGBuilder::Execute()
 
 		Compile();
 
-		SetupBufferUploads();
-
 		IF_RDG_GPU_SCOPES(GPUScopeStacks.ReserveOps(Passes.Num()));
 		IF_RDG_CPU_SCOPES(CPUScopeStacks.ReserveOps());
 
@@ -1525,6 +1519,8 @@ void FRDGBuilder::Execute()
 			}, TStatId(), nullptr, ENamedThreads::AnyHiPriThreadHiPriTask));
 		}
 
+		SetupBufferUploads();
+
 		{
 			SCOPE_CYCLE_COUNTER(STAT_RDG_CollectResourcesTime);
 			CSV_SCOPED_TIMING_STAT_EXCLUSIVE(RDG_CollectResources);
@@ -1829,7 +1825,7 @@ void FRDGBuilder::Execute()
 
 FRDGPass* FRDGBuilder::SetupPass(FRDGPass* Pass)
 {
-	IF_RDG_ENABLE_DEBUG(UserValidation.ValidateAddPass(Pass, bSkipAuxiliaryPasses));
+	IF_RDG_ENABLE_DEBUG(UserValidation.ValidateAddPass(Pass, InAuxiliaryPass()));
 	CSV_SCOPED_TIMING_STAT_EXCLUSIVE_CONDITIONAL(RDGBuilder_SetupPass, GRDGVerboseCSVStats != 0);
 
 	const FRDGParameterStruct PassParameters = Pass->GetParameters();
@@ -2096,11 +2092,6 @@ void FRDGBuilder::SetupBufferUploads()
 
 	for (FUploadedBuffer& UploadedBuffer : UploadedBuffers)
 	{
-		if (UploadedBuffer.Buffer->bCulled)
-		{
-			continue;
-		}
-
 		if (UploadedBuffer.bUseDataCallbacks)
 		{
 			UploadedBuffer.Data = UploadedBuffer.DataCallback();
@@ -2121,11 +2112,6 @@ void FRDGBuilder::SubmitBufferUploads()
 
 	for (const FUploadedBuffer& UploadedBuffer : UploadedBuffers)
 	{
-		if (UploadedBuffer.Buffer->bCulled)
-		{
-			continue;
-		}
-
 		if (UploadedBuffer.Data && UploadedBuffer.DataSize)
 		{
 #if PLATFORM_NEEDS_GPU_UAV_RESOURCE_INIT_WORKAROUND
@@ -2378,9 +2364,7 @@ void FRDGBuilder::CreateUniformBuffers()
 
 void FRDGBuilder::AddProloguePass()
 {
-	bSkipAuxiliaryPasses = true;
 	ProloguePass = SetupEmptyPass(Passes.Allocate<FRDGSentinelPass>(Allocator, RDG_EVENT_NAME("Graph Prologue (Graphics)")));
-	bSkipAuxiliaryPasses = false;
 }
 
 ///////////////////////////////////////////////////////////////////////////////////////////////////
@@ -3411,13 +3395,12 @@ void FRDGBuilder::EndResourceRHI(FRDGPassHandle PassHandle, FRDGBufferRef Buffer
 void FRDGBuilder::VisualizePassOutputs(const FRDGPass* Pass)
 {
 #if SUPPORTS_VISUALIZE_TEXTURE
-	if (bSkipAuxiliaryPasses)
+	if (bInVisualizePassOutputs)
 	{
 		return;
 	}
 
-	bSkipAuxiliaryPasses = true;
-
+	bInVisualizePassOutputs = true;
 
 	Pass->GetParameters().EnumerateTextures([&](FRDGParameter Parameter)
 	{
@@ -3484,22 +3467,18 @@ void FRDGBuilder::VisualizePassOutputs(const FRDGPass* Pass)
 		}
 	});
 
-	bSkipAuxiliaryPasses = false;
+	bInVisualizePassOutputs = false;
 #endif
 }
 
 void FRDGBuilder::ClobberPassOutputs(const FRDGPass* Pass)
 {
-	if (!GRDGClobberResources)
+	if (!GRDGClobberResources || bInClobberPassOutputs)
 	{
 		return;
 	}
 
-	if (bSkipAuxiliaryPasses)
-	{
-		return;
-	}
-	bSkipAuxiliaryPasses = true;
+	bInClobberPassOutputs = true;
 
 	RDG_EVENT_SCOPE(*this, "RDG ClobberResources");
 
@@ -3594,7 +3573,7 @@ void FRDGBuilder::ClobberPassOutputs(const FRDGPass* Pass)
 		}
 	});
 
-	bSkipAuxiliaryPasses = false;
+	bInClobberPassOutputs = false;
 }
 
 #endif //! RDG_ENABLE_DEBUG
diff --git a/Engine/Source/Runtime/RenderCore/Public/RenderGraphBuilder.h b/Engine/Source/Runtime/RenderCore/Public/RenderGraphBuilder.h
--- a/Engine/Source/Runtime/RenderCore/Public/RenderGraphBuilder.h
+++ b/Engine/Source/Runtime/RenderCore/Public/RenderGraphBuilder.h
@@ -680,7 +680,12 @@ private:
 #endif
 
 	/** Tracks whether we are in a scope of adding passes to the builder. Used to avoid recursion. */
-	bool bSkipAuxiliaryPasses = false;
+	bool bInClobberPassOutputs = false;
+	bool bInVisualizePassOutputs = false;
+	bool bInDumpPassOutputs = false;
+	bool bInFlushExternalAccess = false;
+
+	bool InAuxiliaryPass() const { return bInClobberPassOutputs || bInVisualizePassOutputs || bInDumpPassOutputs || bInFlushExternalAccess; }
 
 #if WITH_MGPU
 	/** Name for the temporal effect used to synchronize multi-frame resources. */
diff --git a/Engine/Source/Runtime/RenderCore/Public/RenderGraphDefinitions.h b/Engine/Source/Runtime/RenderCore/Public/RenderGraphDefinitions.h
--- a/Engine/Source/Runtime/RenderCore/Public/RenderGraphDefinitions.h
+++ b/Engine/Source/Runtime/RenderCore/Public/RenderGraphDefinitions.h
@@ -765,6 +765,7 @@ using FRDGPassHandlesByPipeline = TRHIPipelineArray<FRDGPassHandle>;
 using FRDGPassesByPipeline = TRHIPipelineArray<FRDGPass*>;
 
 class FRDGTrace;
+class FRDGResourceDumpContext;
 
 using FRDGBufferNumElementsCallback = TFunction<uint32()>;
 using FRDGBufferInitialDataCallback = TFunction<const void*()>;
diff --git a/Engine/Source/Runtime/RenderCore/Public/RenderGraphResources.h b/Engine/Source/Runtime/RenderCore/Public/RenderGraphResources.h
--- a/Engine/Source/Runtime/RenderCore/Public/RenderGraphResources.h
+++ b/Engine/Source/Runtime/RenderCore/Public/RenderGraphResources.h
@@ -386,6 +386,7 @@ private:
 	friend FRDGBuilder;
 	friend FRDGUserValidation;
 	friend FRDGBarrierBatchBegin;
+	friend FRDGResourceDumpContext;
 	friend FRDGTrace;
 };
 
diff --git a/Engine/Source/Runtime/RenderCore/Public/RenderGraphUtils.h b/Engine/Source/Runtime/RenderCore/Public/RenderGraphUtils.h
--- a/Engine/Source/Runtime/RenderCore/Public/RenderGraphUtils.h
+++ b/Engine/Source/Runtime/RenderCore/Public/RenderGraphUtils.h
@@ -1121,17 +1121,30 @@ public:
 
 	void Add(FRDGViewableResource* Resource, ERHIAccess Access = ERHIAccess::SRVMask, ERHIPipeline Pipelines = ERHIPipeline::Graphics)
 	{
+		if (!Resource)
+		{
+			return;
+		}
+
 		Validate(Resource, Access, Pipelines);
 		Resources.Emplace(Resource, Access, Pipelines);
 	}
 
 	void AddUnique(FRDGViewableResource* Resource, ERHIAccess Access = ERHIAccess::SRVMask, ERHIPipeline Pipelines = ERHIPipeline::Graphics)
 	{
+		if (!Resource)
+		{
+			return;
+		}
+
 		Validate(Resource, Access, Pipelines);
-		if (!Contains(Resource))
+
+		if (Contains(Resource))
 		{
-			Resources.Emplace(Resource, Access, Pipelines);
+			return;
 		}
+
+		Resources.Emplace(Resource, Access, Pipelines);
 	}
 
 	void Submit(FRDGBuilder& GraphBuilder)
@@ -1180,7 +1193,6 @@ public:
 private:
 	void Validate(FRDGViewableResource* Resource, ERHIAccess Access, ERHIPipeline Pipelines)
 	{
-		checkf(Resource, TEXT("Added a null resource to FRDGExternalAccessQueue"));
 		checkf(IsValidAccess(Access) && Access != ERHIAccess::Unknown, TEXT("Attempted to finalize texture %s with an invalid access %s."), Resource->Name, *GetRHIAccessName(Access));
 		check(Pipelines != ERHIPipeline::None);
 	}
diff --git a/Engine/Source/Runtime/Renderer/Private/DistanceFieldStreaming.cpp b/Engine/Source/Runtime/Renderer/Private/DistanceFieldStreaming.cpp
--- a/Engine/Source/Runtime/Renderer/Private/DistanceFieldStreaming.cpp
+++ b/Engine/Source/Runtime/Renderer/Private/DistanceFieldStreaming.cpp
@@ -1547,7 +1547,7 @@ void FDistanceFieldSceneData::UpdateDistanceFieldAtlas(
 					FComputeShaderUtils::GetGroupCount(FIntVector(DistanceField::BrickSize, DistanceField::BrickSize, NumBrickUploadsThisPass * DistanceField::BrickSize), FScatterUploadDistanceFieldAtlasCS::GetGroupSize()));
 			}
 
-			DistanceFieldBrickVolumeTexture = GraphBuilder.ConvertToExternalTexture(DistanceFieldBrickVolumeTextureRDG);
+			DistanceFieldBrickVolumeTexture = ConvertToExternalAccessTexture(GraphBuilder, ExternalAccessQueue, DistanceFieldBrickVolumeTextureRDG);
 		}
 	}
 	
