Fix for RHI validation error in AncientGame.

diff --git a/Engine/Source/Runtime/RenderCore/Private/RenderGraphBuilder.cpp b/Engine/Source/Runtime/RenderCore/Private/RenderGraphBuilder.cpp
--- a/Engine/Source/Runtime/RenderCore/Private/RenderGraphBuilder.cpp
+++ b/Engine/Source/Runtime/RenderCore/Private/RenderGraphBuilder.cpp
@@ -656,7 +656,6 @@ void FRDGBuilder::UseInternalAccessMode(FRDGViewableResource* Resource)
 	}
 
 	AccessModeState.Mode = FRDGViewableResource::EAccessMode::Internal;
-	Resource->EpilogueAccess = FRDGViewableResource::DefaultEpilogueAccess;
 }
 
 void FRDGBuilder::FlushAccessModeQueue()
