VulkanRHI: Add possibility to force depth and stencil to always be treated together when considering transitions in RDG since underlying VulkanRHI cannot track them separately:
- when CopySrc or CopyDest is set on either depth or stencil, force the range to include the entire texture.
- force the transitions of both depth and stencil to occur on the same pass (do not split the copy transition into two different passes if StateBefore is not the same).

diff --git a/Engine/Source/Runtime/RHI/Private/RHI.cpp b/Engine/Source/Runtime/RHI/Private/RHI.cpp
--- a/Engine/Source/Runtime/RHI/Private/RHI.cpp
+++ b/Engine/Source/Runtime/RHI/Private/RHI.cpp
@@ -1478,6 +1478,8 @@ bool GRHISupportsAsyncPipelinePrecompile = true;
 
 bool GRHISupportsMapWriteNoOverwrite = false;
 
+bool GRHISupportsSeparateDepthStencilCopyAccess = true;
+
 bool GRHISupportsBindless = false;
 
 FVertexElementTypeSupportInfo GVertexElementTypeSupport;
diff --git a/Engine/Source/Runtime/RHI/Public/RHI.h b/Engine/Source/Runtime/RHI/Public/RHI.h
--- a/Engine/Source/Runtime/RHI/Public/RHI.h
+++ b/Engine/Source/Runtime/RHI/Public/RHI.h
@@ -500,6 +500,10 @@ extern RHI_API bool GRHISupportsRWTextureBuffers;
 /** Whether a raw (ByteAddress) buffer view can be created for any buffer, regardless of its EBufferUsageFlags::ByteAddressBuffer flag. */
 extern RHI_API bool GRHISupportsRawViewsForAnyBuffer;
 
+/** Whether depth or stencil can individually be set to CopySrc/Dest access. */
+extern RHI_API bool GRHISupportsSeparateDepthStencilCopyAccess;
+
+
 FORCEINLINE uint64 GetMaxBufferDimension()
 {
 	return GMaxBufferDimensions;
diff --git a/Engine/Source/Runtime/RenderCore/Private/RenderGraphBuilder.cpp b/Engine/Source/Runtime/RenderCore/Private/RenderGraphBuilder.cpp
--- a/Engine/Source/Runtime/RenderCore/Private/RenderGraphBuilder.cpp
+++ b/Engine/Source/Runtime/RenderCore/Private/RenderGraphBuilder.cpp
@@ -1858,6 +1858,16 @@ FRDGPass* FRDGBuilder::SetupPass(FRDGPass* Pass)
 			return;
 		}
 
+		// Use the whole texture range for copy states on RHI's that don't track them separately
+		if (!GRHISupportsSeparateDepthStencilCopyAccess)
+		{
+			const ERHIAccess AnyCopyAccess = ERHIAccess::CopySrc | ERHIAccess::CopyDest;
+			if ((Texture->Desc.Format == PF_DepthStencil) && EnumHasAnyFlags(Access, AnyCopyAccess))
+			{
+				Range = Texture->GetSubresourceRange();  // WholeRange
+			}
+		}
+
 		const FRDGViewHandle NoUAVBarrierHandle = GetHandleIfNoUAVBarrier(TextureView);
 		const EResourceTransitionFlags TransitionFlags = GetTextureViewTransitionFlags(TextureView, Texture);
 
@@ -2716,6 +2726,63 @@ void FRDGBuilder::AddTransition(FRDGPassHandle PassHandle, FRDGTextureRef Textur
 	const FRDGTextureSubresourceLayout Layout = Texture->Layout;
 	FRDGTextureSubresourceState& StateBefore = Texture->GetState();
 
+
+	// Use the entire texture for copy states on RHI's that don't track them separately
+	if (!GRHISupportsSeparateDepthStencilCopyAccess)
+	{
+		if ((Texture->Desc.Format == PF_DepthStencil) && !IsWholeResource(StateAfter))
+		{
+			const ERHIAccess AnyCopyAccess = ERHIAccess::CopySrc | ERHIAccess::CopyDest;
+
+			uint32 TotalTransitionCount = 0;
+			uint32 CopyTransitionCount = 0;
+			WholeRange.EnumerateSubresources([&](FRDGTextureSubresource Subresource)
+			{
+				if (FRDGSubresourceState* SubresourceStateAfter = GetSubresource(StateAfter, Layout, Subresource))
+				{
+					++TotalTransitionCount;
+					if (EnumHasAnyFlags(SubresourceStateAfter->Access, AnyCopyAccess))
+					{
+						++CopyTransitionCount;
+					}
+				}
+			});
+
+			// Use a common FRDGSubresourceState to keep the transitions of depth/stencil together in the same pass
+			if ((CopyTransitionCount > 0) && !IsWholeResource(StateBefore))
+			{
+				checkf((CopyTransitionCount == TotalTransitionCount), TEXT("Depth/Stencil Copy transitions on this platform are supposed to touch all subresources."));
+
+				const FRDGTextureSubresource MaxSubresource = WholeRange.GetMaxSubresource();
+				const FRDGTextureSubresource MinSubresource = WholeRange.GetMinSubresource();
+
+				for (uint32 LocalArraySlice = MinSubresource.ArraySlice; LocalArraySlice < MaxSubresource.ArraySlice; ++LocalArraySlice)
+				{
+					for (uint32 LocalMipIndex = MinSubresource.MipIndex; LocalMipIndex < MaxSubresource.MipIndex; ++LocalMipIndex)
+					{
+						FRDGSubresourceState* LatestSubresourceStateBefore = nullptr;
+						for (uint32 LocalPlaneSlice = MinSubresource.PlaneSlice; LocalPlaneSlice < MaxSubresource.PlaneSlice; ++LocalPlaneSlice)
+						{
+							FRDGSubresourceState& OtherStateBefore = GetSubresource(StateBefore, Layout, FRDGTextureSubresource(LocalMipIndex, LocalArraySlice, LocalPlaneSlice));
+							if (LatestSubresourceStateBefore == nullptr)
+							{
+								LatestSubresourceStateBefore = &OtherStateBefore;
+							}
+							else if (LatestSubresourceStateBefore->GetLastPass() < OtherStateBefore.GetLastPass())
+							{
+								*LatestSubresourceStateBefore = OtherStateBefore;
+							}
+							else if (LatestSubresourceStateBefore->GetLastPass() > OtherStateBefore.GetLastPass())
+							{
+								OtherStateBefore = *LatestSubresourceStateBefore;
+							}
+						}
+					}
+				}
+			}
+		}
+	}
+
 	const auto AddSubresourceTransition = [&] (
 		const FRDGSubresourceState& SubresourceStateBefore,
 		const FRDGSubresourceState& SubresourceStateAfter,
diff --git a/Engine/Source/Runtime/VulkanRHI/Private/VulkanBarriers.cpp b/Engine/Source/Runtime/VulkanRHI/Private/VulkanBarriers.cpp
--- a/Engine/Source/Runtime/VulkanRHI/Private/VulkanBarriers.cpp
+++ b/Engine/Source/Runtime/VulkanRHI/Private/VulkanBarriers.cpp
@@ -418,24 +418,6 @@ static void GetDepthStencilStageAndAccessFlags(ERHIAccess DepthAccess, ERHIAcces
 		return;
 	}
 
-	// If one of the source states has copy and the other doesn't, a cache lookup will be necessary
-	const ERHIAccess AnyCopy = ERHIAccess::CopySrc | ERHIAccess::CopyDest;
-	if (!EnumHasAnyFlags(DepthAccess, StencilAccess) &&
-		(EnumHasAnyFlags(DepthAccess, AnyCopy) || EnumHasAnyFlags(StencilAccess, AnyCopy)))
-	{
-		if (bIsSourceState)
-		{
-			StageFlags = VK_PIPELINE_STAGE_ALL_COMMANDS_BIT;
-			AccessFlags = VK_ACCESS_MEMORY_READ_BIT | VK_ACCESS_MEMORY_WRITE_BIT;
-			Layout = VK_IMAGE_LAYOUT_UNDEFINED;
-			return;
-		}
-		else
-		{
-			checkf(!EnumHasAnyFlags(StencilAccess, AnyCopy), TEXT("Only the Stencil being transitioned to Copy state, this is not supported."));
-		}
-	}
-
 	Layout = VK_IMAGE_LAYOUT_DEPTH_STENCIL_READ_ONLY_OPTIMAL;
 
 	if (EnumHasAnyFlags(DepthAccess, ERHIAccess::CopySrc))
@@ -887,6 +869,7 @@ void FVulkanDynamicRHI::RHICreateTransition(FRHITransition* Transition, const FR
 			return EnumHasAnyFlags(ExplicitAspectAccess, DepthStencilFlags) ? ERHIAccess::DSVRead : ExplicitAspectAccess;
 		};
 
+		const ERHIAccess CopyAccess = ERHIAccess::CopySrc | ERHIAccess::CopyDest;
 		ERHIAccess SrcDepthAccess, SrcStencilAccess, DstDepthAccess, DstStencilAccess;
 		VkImageAspectFlags AspectMask;
 		if (PendingTransition.bDepthAccessSet)
@@ -896,6 +879,9 @@ void FVulkanDynamicRHI::RHICreateTransition(FRHITransition* Transition, const FR
 			SrcStencilAccess = GetOtherAspectAccess(SrcDepthAccess);
 			DstStencilAccess = GetOtherAspectAccess(DstDepthAccess);
 			AspectMask = VK_IMAGE_ASPECT_DEPTH_BIT;
+
+			check(!EnumHasAnyFlags(PendingTransition.SrcDepthAccess, CopyAccess));
+			check(!EnumHasAnyFlags(PendingTransition.DestDepthAccess, CopyAccess));
 		}
 		else
 		{
@@ -904,6 +890,9 @@ void FVulkanDynamicRHI::RHICreateTransition(FRHITransition* Transition, const FR
 			SrcDepthAccess = GetOtherAspectAccess(SrcStencilAccess);
 			DstDepthAccess = GetOtherAspectAccess(DstStencilAccess);
 			AspectMask = VK_IMAGE_ASPECT_STENCIL_BIT;
+
+			check(!EnumHasAnyFlags(PendingTransition.SrcStencilAccess, CopyAccess));
+			check(!EnumHasAnyFlags(PendingTransition.DestStencilAccess, CopyAccess));
 		}
 
 		GetDepthStencilStageAndAccessFlags(SrcDepthAccess, SrcStencilAccess, SrcStageMask, SrcAccessFlags, SrcLayout, true);
diff --git a/Engine/Source/Runtime/VulkanRHI/Private/VulkanRHI.cpp b/Engine/Source/Runtime/VulkanRHI/Private/VulkanRHI.cpp
--- a/Engine/Source/Runtime/VulkanRHI/Private/VulkanRHI.cpp
+++ b/Engine/Source/Runtime/VulkanRHI/Private/VulkanRHI.cpp
@@ -836,6 +836,7 @@ void FVulkanDynamicRHI::InitInstance()
 		GMaxTextureArrayLayers = Props.limits.maxImageArrayLayers;
 		GRHISupportsBaseVertexIndex = true;
 		GSupportsSeparateRenderTargetBlendState = true;
+		GRHISupportsSeparateDepthStencilCopyAccess = false;
 
 		GRHIMaxDispatchThreadGroupsPerDimension.X = FMath::Min<uint32>(Limits.maxComputeWorkGroupCount[0], 0x7fffffff);
 		GRHIMaxDispatchThreadGroupsPerDimension.Y = FMath::Min<uint32>(Limits.maxComputeWorkGroupCount[1], 0x7fffffff);
