DumpGPU: Skip AvailableVirtual in IsUnsafeToDumpResource check if it's exactly 0 (value when it's not supported).

#rnx

diff --git a/Engine/Source/Runtime/RenderCore/Private/DumpGPU.cpp b/Engine/Source/Runtime/RenderCore/Private/DumpGPU.cpp
--- a/Engine/Source/Runtime/RenderCore/Private/DumpGPU.cpp
+++ b/Engine/Source/Runtime/RenderCore/Private/DumpGPU.cpp
@@ -363,8 +363,10 @@ public:
 
 	bool IsUnsafeToDumpResource(SIZE_T ResourceByteSize, float DumpMemoryMultiplier) const
 	{
-		uint64 AproximatedStagingMemoryRequired = uint64(double(ResourceByteSize) * DumpMemoryMultiplier);
-		uint64 MaxMemoryAvailable = FMath::Min(MemoryStats.AvailablePhysical, MemoryStats.AvailableVirtual);
+		const uint64 AproximatedStagingMemoryRequired = uint64(double(ResourceByteSize) * DumpMemoryMultiplier);
+
+		// Skip AvailableVirtual if it's not supported (reported as exactly 0)
+		const uint64 MaxMemoryAvailable = (MemoryStats.AvailableVirtual > 0) ? FMath::Min(MemoryStats.AvailablePhysical, MemoryStats.AvailableVirtual) : MemoryStats.AvailablePhysical;
 
 		return AproximatedStagingMemoryRequired > MaxMemoryAvailable;
 	}
