Deprecated FRDGBufferDesc::EUnderlyingType

diff --git a/Engine/Plugins/Experimental/NNI/Source/NeuralNetworkInference/Private/NeuralTensor.cpp b/Engine/Plugins/Experimental/NNI/Source/NeuralNetworkInference/Private/NeuralTensor.cpp
--- a/Engine/Plugins/Experimental/NNI/Source/NeuralNetworkInference/Private/NeuralTensor.cpp
+++ b/Engine/Plugins/Experimental/NNI/Source/NeuralNetworkInference/Private/NeuralTensor.cpp
@@ -666,8 +666,7 @@ void FNeuralTensor::InitPooledBufferForUEAndORTBackEnd_RenderThread(FRDGBuilder&
 	FRDGBufferDesc BufferDesc;
 	BufferDesc.BytesPerElement = FNeuralDataTypeUtils::GetByteSize(DataType);
 	BufferDesc.NumElements = Num();
-	BufferDesc.Usage = EBufferUsageFlags::UnorderedAccess | EBufferUsageFlags::ShaderResource;
-	BufferDesc.UnderlyingType = FRDGBufferDesc::EUnderlyingType::VertexBuffer;
+	BufferDesc.Usage = EBufferUsageFlags::UnorderedAccess | EBufferUsageFlags::ShaderResource | EBufferUsageFlags::VertexBuffer;
 	FRDGBufferRef BufferRef = GraphBuilder.CreateBuffer(BufferDesc, *GetName());
 
 	// Extract buffer immediately. Need this so that we can potentially bind to DirectML.
@@ -892,8 +891,7 @@ void FNeuralTensor::CPUToRDGBuilder_RenderThread(FRDGBuilder* InOutGraphBuilder,
 		FRDGBufferDesc BufferDesc;
 		BufferDesc.BytesPerElement = FNeuralDataTypeUtils::GetByteSize(DataType);
 		BufferDesc.NumElements = Num();
-		BufferDesc.Usage = InBufferUsageFlags;
-		BufferDesc.UnderlyingType = FRDGBufferDesc::EUnderlyingType::VertexBuffer;
+		BufferDesc.Usage = InBufferUsageFlags | EBufferUsageFlags::VertexBuffer;
 		
 		BufferRef = bInShouldCopyFromCPU 
 			? CreateVertexBuffer(*InOutGraphBuilder, *Name, BufferDesc, UnderlyingUInt8ArrayData.GetData(), NumInBytes(),
diff --git a/Engine/Plugins/Runtime/HairStrands/Source/HairStrandsCore/Private/GroomResources.cpp b/Engine/Plugins/Runtime/HairStrands/Source/HairStrandsCore/Private/GroomResources.cpp
--- a/Engine/Plugins/Runtime/HairStrands/Source/HairStrandsCore/Private/GroomResources.cpp
+++ b/Engine/Plugins/Runtime/HairStrands/Source/HairStrandsCore/Private/GroomResources.cpp
@@ -160,7 +160,7 @@ static FRDGBufferRef InternalCreateVertexBuffer(
 	uint64 InitialDataSize,
 	ERDGInitialDataFlags InitialDataFlags)
 {
-	checkf(Desc.UnderlyingType == FRDGBufferDesc::EUnderlyingType::VertexBuffer, TEXT("CreateVertexBuffer called with an FRDGBufferDesc underlying type that is not 'VertexBuffer'. Buffer: %s"), Name);
+	checkf(EnumHasAnyFlags(Desc.Usage, EBufferUsageFlags::VertexBuffer), TEXT("CreateVertexBuffer called with an FRDGBufferDesc underlying type that is not 'VertexBuffer'. Buffer: %s"), Name);
 	FRDGBufferRef Buffer = GraphBuilder.CreateBuffer(Desc, Name, ERDGBufferFlags::MultiFrame);
 	if (InitialData && InitialDataSize)
 	{
@@ -202,7 +202,7 @@ void InternalCreateVertexBufferRDG_FromBulkData(FRDGBuilder& GraphBuilder, FByte
 	#if !WITH_EDITORONLY_DATA
 	if (bReleaseCPUData)
 	{
-		checkf(Desc.UnderlyingType == FRDGBufferDesc::EUnderlyingType::VertexBuffer, TEXT("CreateVertexBuffer called with an FRDGBufferDesc underlying type that is not 'VertexBuffer'. Buffer: %s"), DebugName);
+		checkf(EnumHasAnyFlags(Desc.Usage, EBufferUsageFlags::VertexBuffer), TEXT("CreateVertexBuffer called with an FRDGBufferDesc underlying type that is not 'VertexBuffer'. Buffer: %s"), DebugName);
 		Buffer = GraphBuilder.CreateBuffer(Desc, DebugName, ERDGBufferFlags::MultiFrame);
 		if (BulkData && DataSizeInBytes)
 		{
diff --git a/Engine/Source/Runtime/Engine/Private/RayTracingSkinnedGeometry.cpp b/Engine/Source/Runtime/Engine/Private/RayTracingSkinnedGeometry.cpp
--- a/Engine/Source/Runtime/Engine/Private/RayTracingSkinnedGeometry.cpp
+++ b/Engine/Source/Runtime/Engine/Private/RayTracingSkinnedGeometry.cpp
@@ -214,8 +214,7 @@ void FRayTracingSkinnedGeometryUpdateQueue::Commit(FRDGBuilder& GraphBuilder)
 		const uint32 ScratchAlignment = GRHIRayTracingAccelerationStructureAlignment;
 		
 		FRDGBufferDesc ScratchBufferDesc;
-		ScratchBufferDesc.UnderlyingType = FRDGBufferDesc::EUnderlyingType::StructuredBuffer;
-		ScratchBufferDesc.Usage = BUF_RayTracingScratch;
+		ScratchBufferDesc.Usage = EBufferUsageFlags::RayTracingScratch | EBufferUsageFlags::StructuredBuffer;
 		ScratchBufferDesc.BytesPerElement = ScratchAlignment;
 		ScratchBufferDesc.NumElements = FMath::DivideAndRoundUp(BLASScratchSize, ScratchAlignment);
 
diff --git a/Engine/Source/Runtime/RenderCore/Private/DumpGPU.cpp b/Engine/Source/Runtime/RenderCore/Private/DumpGPU.cpp
--- a/Engine/Source/Runtime/RenderCore/Private/DumpGPU.cpp
+++ b/Engine/Source/Runtime/RenderCore/Private/DumpGPU.cpp
@@ -716,7 +716,6 @@ public:
 		JsonObject->SetStringField(TEXT("UniqueResourceName"), UniqueResourceName);
 		JsonObject->SetNumberField(TEXT("ByteSize"), ResourceByteSize);
 		JsonObject->SetStringField(TEXT("Desc"), TEXT("FRDGBufferDesc"));
-		JsonObject->SetStringField(TEXT("Type"), GetBufferUnderlyingTypeName(Desc.UnderlyingType));
 		JsonObject->SetNumberField(TEXT("BytesPerElement"), Desc.BytesPerElement);
 		JsonObject->SetNumberField(TEXT("NumElements"), Desc.NumElements);
 		JsonObject->SetStringField(TEXT("Metadata"), *PtrToString(Desc.Metadata));
@@ -1432,13 +1431,6 @@ public:
 			}
 		}
 
-	PRAGMA_DISABLE_DEPRECATION_WARNINGS
-		if (Desc.UnderlyingType == FRDGBufferDesc::EUnderlyingType::AccelerationStructure)
-		{
-			return;
-		}
-	PRAGMA_ENABLE_DEPRECATION_WARNINGS
-
 		// Dump the resource's binary to a .bin file.
 		if (bDumpResourceBinary && DumpBufferMode == 2)
 		{
diff --git a/Engine/Source/Runtime/RenderCore/Private/GPUMessaging.cpp b/Engine/Source/Runtime/RenderCore/Private/GPUMessaging.cpp
--- a/Engine/Source/Runtime/RenderCore/Private/GPUMessaging.cpp
+++ b/Engine/Source/Runtime/RenderCore/Private/GPUMessaging.cpp
@@ -84,7 +84,7 @@ public:
 		const uint32 MaxBufferSize = GetMaxBufferSize();
 
 		FRDGBufferDesc Desc = FRDGBufferDesc::CreateStructuredDesc(sizeof(uint32), MaxBufferSize);
-		Desc.Usage = EBufferUsageFlags(Desc.Usage | BUF_SourceCopy);
+		Desc.Usage |= EBufferUsageFlags::SourceCopy;
 		MessageBuffer = GraphBuilder.CreateBuffer(Desc, TEXT("GPUMessageManager.MessageBuffer"));
 
 		FParameters& Parameters = GraphBuilder.Blackboard.Create<FParameters>();
diff --git a/Engine/Source/Runtime/RenderCore/Private/RenderGraphBuilder.cpp b/Engine/Source/Runtime/RenderCore/Private/RenderGraphBuilder.cpp
--- a/Engine/Source/Runtime/RenderCore/Private/RenderGraphBuilder.cpp
+++ b/Engine/Source/Runtime/RenderCore/Private/RenderGraphBuilder.cpp
@@ -3332,7 +3332,7 @@ void FRDGBuilder::InitRHI(FRDGBufferSRVRef SRV)
 
 	FRHIBufferSRVCreateInfo SRVCreateInfo = SRV->Desc;
 
-	if (Buffer->Desc.UnderlyingType == FRDGBufferDesc::EUnderlyingType::StructuredBuffer)
+	if (EnumHasAnyFlags(Buffer->Desc.Usage, EBufferUsageFlags::StructuredBuffer))
 	{
 		// RDG allows structured buffer views to be typed, but the view creation logic requires that it
 		// be unknown (as do platform APIs -- structured buffers are not typed). This could be validated
@@ -3357,7 +3357,7 @@ void FRDGBuilder::InitRHI(FRDGBufferUAV* UAV)
 
 	FRHIBufferUAVCreateInfo UAVCreateInfo = UAV->Desc;
 
-	if (Buffer->Desc.UnderlyingType == FRDGBufferDesc::EUnderlyingType::StructuredBuffer)
+	if (EnumHasAnyFlags(Buffer->Desc.Usage, EBufferUsageFlags::StructuredBuffer))
 	{
 		// RDG allows structured buffer views to be typed, but the view creation logic requires that it
 		// be unknown (as do platform APIs -- structured buffers are not typed). This could be validated
diff --git a/Engine/Source/Runtime/RenderCore/Private/RenderGraphResourcePool.cpp b/Engine/Source/Runtime/RenderCore/Private/RenderGraphResourcePool.cpp
--- a/Engine/Source/Runtime/RenderCore/Private/RenderGraphResourcePool.cpp
+++ b/Engine/Source/Runtime/RenderCore/Private/RenderGraphResourcePool.cpp
@@ -117,34 +117,9 @@ TRefCountPtr<FRDGPooledBuffer> FRDGBufferPool::FindFreeBuffer(const FRDGBufferDe
 		TRACE_COUNTER_ADD(BufferPoolCreateCount, 1);
 		TRACE_COUNTER_ADD(BufferPoolSize, NumBytes);
 
+		const ERHIAccess InitialAccess = RHIGetDefaultResourceState(Desc.Usage, false);
 		FRHIResourceCreateInfo CreateInfo(InDebugName);
-		TRefCountPtr<FRHIBuffer> BufferRHI;
-
-		ERHIAccess InitialAccess = ERHIAccess::Unknown;
-
-		if (Desc.UnderlyingType == FRDGBufferDesc::EUnderlyingType::VertexBuffer)
-		{
-			const EBufferUsageFlags Usage = Desc.Usage | BUF_VertexBuffer;
-			InitialAccess = RHIGetDefaultResourceState(Usage, false);
-			BufferRHI = RHICreateVertexBuffer(NumBytes, Usage, InitialAccess, CreateInfo);
-		}
-		else if (Desc.UnderlyingType == FRDGBufferDesc::EUnderlyingType::StructuredBuffer)
-		{
-			const EBufferUsageFlags Usage = Desc.Usage | BUF_StructuredBuffer;
-			InitialAccess = RHIGetDefaultResourceState(Usage, false);
-			BufferRHI = RHICreateStructuredBuffer(Desc.BytesPerElement, NumBytes, Usage, InitialAccess, CreateInfo);
-		}
-	PRAGMA_DISABLE_DEPRECATION_WARNINGS
-		else if (Desc.UnderlyingType == FRDGBufferDesc::EUnderlyingType::AccelerationStructure)
-		{
-			InitialAccess = ERHIAccess::BVHWrite;
-			BufferRHI = RHICreateBuffer(NumBytes, Desc.Usage, 0, InitialAccess, CreateInfo);
-		}
-	PRAGMA_ENABLE_DEPRECATION_WARNINGS
-		else
-		{
-			check(0);
-		}
+		TRefCountPtr<FRHIBuffer> BufferRHI = RHICreateBuffer(NumBytes, Desc.Usage, Desc.BytesPerElement, InitialAccess, CreateInfo);
 
 	#if !(UE_BUILD_SHIPPING || UE_BUILD_TEST)
 		RHIBindDebugLabelName(BufferRHI, InDebugName);
diff --git a/Engine/Source/Runtime/RenderCore/Private/RenderGraphUtils.cpp b/Engine/Source/Runtime/RenderCore/Private/RenderGraphUtils.cpp
--- a/Engine/Source/Runtime/RenderCore/Private/RenderGraphUtils.cpp
+++ b/Engine/Source/Runtime/RenderCore/Private/RenderGraphUtils.cpp
@@ -226,6 +226,11 @@ void AddCopyTexturePass(
 	FRDGTextureRef OutputTexture,
 	const FRHICopyTextureInfo& CopyInfo)
 {
+	if (InputTexture == OutputTexture)
+	{
+		return;
+	}
+
 	const FRDGTextureDesc& InputDesc = InputTexture->Desc;
 	const FRDGTextureDesc& OutputDesc = OutputTexture->Desc;
 	checkf(InputDesc.Format == OutputDesc.Format, TEXT("This method does not support format conversion."));
@@ -251,6 +256,11 @@ RENDERCORE_API void AddDrawTexturePass(
 	FRDGTextureRef OutputTexture,
 	const FRDGDrawTextureInfo& DrawInfo)
 {
+	if (InputTexture == OutputTexture)
+	{
+		return;
+	}
+
 	const FRDGTextureDesc& InputDesc = InputTexture->Desc;
 	const FRDGTextureDesc& OutputDesc = OutputTexture->Desc;
 
@@ -1002,7 +1012,7 @@ FRDGBufferRef CreateVertexBuffer(
 	ERDGInitialDataFlags InitialDataFlags)
 {
 	checkf(Name, TEXT("Buffer must have a name."));
-	checkf(Desc.UnderlyingType == FRDGBufferDesc::EUnderlyingType::VertexBuffer, TEXT("CreateVertexBuffer called with an FRDGBufferDesc underlying type that is not 'VertexBuffer'. Buffer: %s"), Name);
+	checkf(EnumHasAnyFlags(Desc.Usage, EBufferUsageFlags::VertexBuffer), TEXT("CreateVertexBuffer called with an FRDGBufferDesc underlying type that is not 'VertexBuffer'. Buffer: %s"), Name);
 
 	FRDGBufferRef Buffer = GraphBuilder.CreateBuffer(Desc, Name);
 	GraphBuilder.QueueBufferUpload(Buffer, InitialData, InitialDataSize, InitialDataFlags);
diff --git a/Engine/Source/Runtime/RenderCore/Private/RenderGraphValidation.cpp b/Engine/Source/Runtime/RenderCore/Private/RenderGraphValidation.cpp
--- a/Engine/Source/Runtime/RenderCore/Private/RenderGraphValidation.cpp
+++ b/Engine/Source/Runtime/RenderCore/Private/RenderGraphValidation.cpp
@@ -284,9 +284,7 @@ void FRDGUserValidation::ValidateCreateBuffer(const FRDGBufferDesc& Desc, const
 
 	checkf(Desc.GetSize() > 0, TEXT("Creating buffer '%s' is zero bytes in size."), Name);
 
-	const bool bIsByteAddress = (Desc.Usage & BUF_ByteAddressBuffer) == BUF_ByteAddressBuffer;
-
-	if (bIsByteAddress && Desc.UnderlyingType == FRDGBufferDesc::EUnderlyingType::StructuredBuffer)
+	if (EnumHasAllFlags(Desc.Usage, EBufferUsageFlags::StructuredBuffer | EBufferUsageFlags::ByteAddressBuffer))
 	{
 		checkf(Desc.BytesPerElement == 4, TEXT("Creating buffer '%s' as a structured buffer that is also byte addressable, BytesPerElement must be 4! Instead it is %d"), Name, Desc.BytesPerElement);
 	}
diff --git a/Engine/Source/Runtime/RenderCore/Private/UnifiedBuffer.cpp b/Engine/Source/Runtime/RenderCore/Private/UnifiedBuffer.cpp
--- a/Engine/Source/Runtime/RenderCore/Private/UnifiedBuffer.cpp
+++ b/Engine/Source/Runtime/RenderCore/Private/UnifiedBuffer.cpp
@@ -215,16 +215,13 @@ EResourceType GetBufferType(FRDGBuffer* Buffer)
 {
 	const FRDGBufferDesc& Desc = Buffer->Desc;
 
-	if (Desc.UnderlyingType == FRDGBufferDesc::EUnderlyingType::StructuredBuffer)
+	if (EnumHasAnyFlags(Desc.Usage, EBufferUsageFlags::ByteAddressBuffer))
 	{
-		if (EnumHasAnyFlags(Desc.Usage, BUF_ByteAddressBuffer))
-		{
-			return EResourceType::BYTEBUFFER;
-		}
-		else
-		{
-			return EResourceType::STRUCTURED_BUFFER;
-		}
+		return EResourceType::BYTEBUFFER;
+	}
+	else if (EnumHasAnyFlags(Desc.Usage, EBufferUsageFlags::StructuredBuffer))
+	{
+		return EResourceType::STRUCTURED_BUFFER;
 	}
 	else
 	{
diff --git a/Engine/Source/Runtime/RenderCore/Public/RenderGraphResources.h b/Engine/Source/Runtime/RenderCore/Public/RenderGraphResources.h
--- a/Engine/Source/Runtime/RenderCore/Public/RenderGraphResources.h
+++ b/Engine/Source/Runtime/RenderCore/Public/RenderGraphResources.h
@@ -849,13 +849,21 @@ private:
 /** Descriptor for render graph tracked Buffer. */
 struct FRDGBufferDesc
 {
-	enum class EUnderlyingType
+	enum class UE_DEPRECATED(5.1, "EUnderlying type has been deprecated.") EUnderlyingType
 	{
 		VertexBuffer,
 		StructuredBuffer,
-		AccelerationStructure UE_DEPRECATED(5.1, "AccelerationStructure is deprecated")
+		AccelerationStructure
 	};
 
+PRAGMA_DISABLE_DEPRECATION_WARNINGS
+	FRDGBufferDesc() = default;
+	FRDGBufferDesc(FRDGBufferDesc&&) = default;
+	FRDGBufferDesc(const FRDGBufferDesc&) = default;
+	FRDGBufferDesc& operator=(FRDGBufferDesc&&) = default;
+	FRDGBufferDesc& operator=(const FRDGBufferDesc&) = default;
+PRAGMA_ENABLE_DEPRECATION_WARNINGS
+
 	/** Create the descriptor for an indirect RHI call.
 	 *
 	 * Note, IndirectParameterStruct should be one of the:
@@ -867,8 +875,7 @@ struct FRDGBufferDesc
 	static FRDGBufferDesc CreateIndirectDesc(uint32 NumElements = 1)
 	{
 		FRDGBufferDesc Desc;
-		Desc.UnderlyingType = EUnderlyingType::VertexBuffer;
-		Desc.Usage = EBufferUsageFlags::Static | EBufferUsageFlags::DrawIndirect | EBufferUsageFlags::UnorderedAccess | EBufferUsageFlags::ShaderResource;
+		Desc.Usage = EBufferUsageFlags::Static | EBufferUsageFlags::DrawIndirect | EBufferUsageFlags::UnorderedAccess | EBufferUsageFlags::ShaderResource | EBufferUsageFlags::VertexBuffer;
 		Desc.BytesPerElement = sizeof(IndirectParameterStruct);
 		Desc.NumElements = NumElements;
 		return Desc;
@@ -877,8 +884,7 @@ struct FRDGBufferDesc
 	static FRDGBufferDesc CreateIndirectDesc(uint32 NumElements = 1)
 	{
 		FRDGBufferDesc Desc;
-		Desc.UnderlyingType = EUnderlyingType::VertexBuffer;
-		Desc.Usage = EBufferUsageFlags::Static | EBufferUsageFlags::DrawIndirect | EBufferUsageFlags::UnorderedAccess | EBufferUsageFlags::ShaderResource;
+		Desc.Usage = EBufferUsageFlags::Static | EBufferUsageFlags::DrawIndirect | EBufferUsageFlags::UnorderedAccess | EBufferUsageFlags::ShaderResource | EBufferUsageFlags::VertexBuffer;
 		Desc.BytesPerElement = 4;
 		Desc.NumElements = NumElements;
 		return Desc;
@@ -887,8 +893,7 @@ struct FRDGBufferDesc
 	static FRDGBufferDesc CreateStructuredDesc(uint32 BytesPerElement, uint32 NumElements)
 	{
 		FRDGBufferDesc Desc;
-		Desc.UnderlyingType = EUnderlyingType::StructuredBuffer;
-		Desc.Usage = EBufferUsageFlags::Static | EBufferUsageFlags::UnorderedAccess | EBufferUsageFlags::ShaderResource;
+		Desc.Usage = EBufferUsageFlags::Static | EBufferUsageFlags::UnorderedAccess | EBufferUsageFlags::ShaderResource | EBufferUsageFlags::StructuredBuffer;
 		Desc.BytesPerElement = BytesPerElement;
 		Desc.NumElements = NumElements;
 		return Desc;
@@ -905,8 +910,7 @@ struct FRDGBufferDesc
 	static FRDGBufferDesc CreateBufferDesc(uint32 BytesPerElement, uint32 NumElements)
 	{
 		FRDGBufferDesc Desc;
-		Desc.UnderlyingType = EUnderlyingType::VertexBuffer;
-		Desc.Usage = EBufferUsageFlags::Static | EBufferUsageFlags::UnorderedAccess | EBufferUsageFlags::ShaderResource;
+		Desc.Usage = EBufferUsageFlags::Static | EBufferUsageFlags::UnorderedAccess | EBufferUsageFlags::ShaderResource | EBufferUsageFlags::VertexBuffer;
 		Desc.BytesPerElement = BytesPerElement;
 		Desc.NumElements = NumElements;
 		return Desc;
@@ -924,8 +928,7 @@ struct FRDGBufferDesc
 	{
 		check(NumBytes % 4 == 0);
 		FRDGBufferDesc Desc;
-		Desc.UnderlyingType = EUnderlyingType::StructuredBuffer;
-		Desc.Usage = EBufferUsageFlags::Static | EBufferUsageFlags::UnorderedAccess | EBufferUsageFlags::ShaderResource | EBufferUsageFlags::ByteAddressBuffer;
+		Desc.Usage = EBufferUsageFlags::Static | EBufferUsageFlags::UnorderedAccess | EBufferUsageFlags::ShaderResource | EBufferUsageFlags::StructuredBuffer | EBufferUsageFlags::ByteAddressBuffer;
 		Desc.BytesPerElement = 4;
 		Desc.NumElements = NumBytes / 4;
 		return Desc;
@@ -942,8 +945,7 @@ struct FRDGBufferDesc
 	static FRDGBufferDesc CreateUploadDesc(uint32 BytesPerElement, uint32 NumElements)
 	{
 		FRDGBufferDesc Desc;
-		Desc.UnderlyingType = EUnderlyingType::VertexBuffer;
-		Desc.Usage = EBufferUsageFlags::ShaderResource;
+		Desc.Usage = EBufferUsageFlags::ShaderResource | EBufferUsageFlags::VertexBuffer;
 		Desc.BytesPerElement = BytesPerElement;
 		Desc.NumElements = NumElements;
 		return Desc;
@@ -960,8 +962,7 @@ struct FRDGBufferDesc
 	static FRDGBufferDesc CreateStructuredUploadDesc(uint32 BytesPerElement, uint32 NumElements)
 	{
 		FRDGBufferDesc Desc;
-		Desc.UnderlyingType = EUnderlyingType::StructuredBuffer;
-		Desc.Usage = EBufferUsageFlags::ShaderResource;
+		Desc.Usage = EBufferUsageFlags::ShaderResource | EBufferUsageFlags::StructuredBuffer;
 		Desc.BytesPerElement = BytesPerElement;
 		Desc.NumElements = NumElements;
 		return Desc;
@@ -979,8 +980,7 @@ struct FRDGBufferDesc
 	{
 		check(NumBytes % 4 == 0);
 		FRDGBufferDesc Desc;
-		Desc.UnderlyingType = EUnderlyingType::StructuredBuffer;
-		Desc.Usage = EBufferUsageFlags::ShaderResource | EBufferUsageFlags::ByteAddressBuffer;
+		Desc.Usage = EBufferUsageFlags::ShaderResource | EBufferUsageFlags::ByteAddressBuffer | EBufferUsageFlags::StructuredBuffer;
 		Desc.BytesPerElement = 4;
 		Desc.NumElements = NumBytes / 4;
 		return Desc;
@@ -1011,18 +1011,16 @@ struct FRDGBufferDesc
 		uint32 Hash = GetTypeHash(Desc.BytesPerElement);
 		Hash = HashCombine(Hash, GetTypeHash(Desc.NumElements));
 		Hash = HashCombine(Hash, GetTypeHash(Desc.Usage));
-		Hash = HashCombine(Hash, GetTypeHash(Desc.UnderlyingType));
 		Hash = HashCombine(Hash, GetTypeHash(Desc.Metadata));
 		return Hash;
 	}
 
 	bool operator == (const FRDGBufferDesc& Other) const
 	{
-		return (
+		return
 			BytesPerElement == Other.BytesPerElement &&
 			NumElements == Other.NumElements &&
-			Usage == Other.Usage &&
-			UnderlyingType == Other.UnderlyingType);
+			Usage == Other.Usage;
 	}
 
 	bool operator != (const FRDGBufferDesc& Other) const
@@ -1040,25 +1038,20 @@ struct FRDGBufferDesc
 	EBufferUsageFlags Usage = EBufferUsageFlags::None;
 
 	/** The underlying RHI type to use. */
+PRAGMA_DISABLE_DEPRECATION_WARNINGS
+	UE_DEPRECATED(5.1, "EUnderlyingType has been deprecated. Use explicit EBufferUsageFlags instead.")
 	EUnderlyingType UnderlyingType = EUnderlyingType::VertexBuffer;
+PRAGMA_ENABLE_DEPRECATION_WARNINGS
 
 	/** Meta data of the layout of the buffer for debugging purposes. */
 	const FShaderParametersMetadata* Metadata = nullptr;
 };
 
-inline const TCHAR* GetBufferUnderlyingTypeName(FRDGBufferDesc::EUnderlyingType BufferType)
-{
-	switch (BufferType)
-	{
-	case FRDGBufferDesc::EUnderlyingType::VertexBuffer:
-		return TEXT("VertexBuffer");
-	case FRDGBufferDesc::EUnderlyingType::StructuredBuffer:
-		return TEXT("StructuredBuffer");
+UE_DEPRECATED(5.1, "FRDGBuffer::EUnderylingType has been deprecated. Use EBufferCreateFlags instead.")
 PRAGMA_DISABLE_DEPRECATION_WARNINGS
-	case FRDGBufferDesc::EUnderlyingType::AccelerationStructure:
-		return TEXT("AccelerationStructure");
+inline const TCHAR* GetBufferUnderlyingTypeName(FRDGBufferDesc::EUnderlyingType)
 PRAGMA_ENABLE_DEPRECATION_WARNINGS
-	}
+{
 	return TEXT("");
 }
 
diff --git a/Engine/Source/Runtime/RenderCore/Public/RenderGraphResources.inl b/Engine/Source/Runtime/RenderCore/Public/RenderGraphResources.inl
--- a/Engine/Source/Runtime/RenderCore/Public/RenderGraphResources.inl
+++ b/Engine/Source/Runtime/RenderCore/Public/RenderGraphResources.inl
@@ -79,12 +79,13 @@ inline FRHIBufferCreateInfo Translate(const FRDGBufferDesc& InDesc)
 {
 	FRHIBufferCreateInfo CreateInfo;
 	CreateInfo.Size = InDesc.GetSize();
-	if (InDesc.UnderlyingType == FRDGBufferDesc::EUnderlyingType::VertexBuffer)
+
+	if (EnumHasAnyFlags(InDesc.Usage, EBufferUsageFlags::VertexBuffer))
 	{
 		CreateInfo.Stride = 0;
 		CreateInfo.Usage = InDesc.Usage | BUF_VertexBuffer;
 	}
-	else if (InDesc.UnderlyingType == FRDGBufferDesc::EUnderlyingType::StructuredBuffer)
+	else if (EnumHasAnyFlags(InDesc.Usage, EBufferUsageFlags::StructuredBuffer))
 	{
 		CreateInfo.Stride = InDesc.BytesPerElement;
 		CreateInfo.Usage = InDesc.Usage | BUF_StructuredBuffer;
@@ -186,31 +187,31 @@ inline FRDGTextureSubresourceRange FRDGTextureUAV::GetSubresourceRange() const
 inline FRDGBufferSRVDesc::FRDGBufferSRVDesc(FRDGBufferRef InBuffer)
 	: Buffer(InBuffer)
 {
-	if (EnumHasAnyFlags(Buffer->Desc.Usage, BUF_DrawIndirect))
+	if (EnumHasAnyFlags(Buffer->Desc.Usage, EBufferUsageFlags::DrawIndirect))
 	{
 		BytesPerElement = 4;
 		Format = PF_R32_UINT;
 	}
-	else if (EnumHasAnyFlags(Buffer->Desc.Usage, BUF_AccelerationStructure))
+	else if (EnumHasAnyFlags(Buffer->Desc.Usage, EBufferUsageFlags::AccelerationStructure))
 	{
 		// nothing special here
 	}
 	else
 	{
-		checkf(Buffer->Desc.UnderlyingType != FRDGBufferDesc::EUnderlyingType::VertexBuffer, TEXT("VertexBuffer %s requires a type when creating a SRV."), Buffer->Name);
+		checkf(!EnumHasAnyFlags(Buffer->Desc.Usage, EBufferUsageFlags::VertexBuffer), TEXT("VertexBuffer %s requires a type when creating a SRV."), Buffer->Name);
 	}
 }
 
 inline FRDGBufferUAVDesc::FRDGBufferUAVDesc(FRDGBufferRef InBuffer)
 	: Buffer(InBuffer)
 {
-	if (EnumHasAnyFlags(Buffer->Desc.Usage, BUF_DrawIndirect))
+	if (EnumHasAnyFlags(Buffer->Desc.Usage, EBufferUsageFlags::DrawIndirect))
 	{
 		Format = PF_R32_UINT;
 	}
 	else
 	{
-		checkf(Buffer->Desc.UnderlyingType != FRDGBufferDesc::EUnderlyingType::VertexBuffer, TEXT("VertexBuffer %s requires a type when creating a UAV."), Buffer->Name);
+		checkf(!EnumHasAnyFlags(Buffer->Desc.Usage, EBufferUsageFlags::VertexBuffer), TEXT("VertexBuffer %s requires a type when creating a UAV."), Buffer->Name);
 	}
 }
 
diff --git a/Engine/Source/Runtime/RenderCore/Public/RenderGraphUtils.h b/Engine/Source/Runtime/RenderCore/Public/RenderGraphUtils.h
--- a/Engine/Source/Runtime/RenderCore/Public/RenderGraphUtils.h
+++ b/Engine/Source/Runtime/RenderCore/Public/RenderGraphUtils.h
@@ -621,8 +621,8 @@ struct RENDERCORE_API FComputeShaderUtils
 
 	static inline void ValidateIndirectArgsBuffer(FRDGBufferRef IndirectArgsBuffer, uint32 IndirectArgOffset)
 	{
-		checkf(IndirectArgsBuffer->Desc.UnderlyingType == FRDGBufferDesc::EUnderlyingType::VertexBuffer, TEXT("The buffer %s needs to be a vertex buffer to be used as an indirect dispatch parameters"), IndirectArgsBuffer->Name);
-		checkf(IndirectArgsBuffer->Desc.Usage & BUF_DrawIndirect, TEXT("The buffer %s for indirect dispatch parameters was not flagged with BUF_DrawIndirect"), IndirectArgsBuffer->Name);
+		checkf(EnumHasAnyFlags(IndirectArgsBuffer->Desc.Usage, EBufferUsageFlags::VertexBuffer), TEXT("The buffer %s needs to be a vertex buffer to be used as an indirect dispatch parameters"), IndirectArgsBuffer->Name);
+		checkf(EnumHasAnyFlags(IndirectArgsBuffer->Desc.Usage, EBufferUsageFlags::DrawIndirect), TEXT("The buffer %s for indirect dispatch parameters was not flagged with BUF_DrawIndirect"), IndirectArgsBuffer->Name);
 		ValidateIndirectArgsBuffer(IndirectArgsBuffer->GetSize(), IndirectArgOffset);
 	}
 
diff --git a/Engine/Source/Runtime/Renderer/Private/DeferredShadingRenderer.cpp b/Engine/Source/Runtime/Renderer/Private/DeferredShadingRenderer.cpp
--- a/Engine/Source/Runtime/Renderer/Private/DeferredShadingRenderer.cpp
+++ b/Engine/Source/Runtime/Renderer/Private/DeferredShadingRenderer.cpp
@@ -1584,8 +1584,7 @@ bool FDeferredShadingSceneRenderer::DispatchRayTracingWorldUpdates(FRDGBuilder&
 	{
 		const uint32 ScratchAlignment = GRHIRayTracingAccelerationStructureAlignment;
 		FRDGBufferDesc ScratchBufferDesc;
-		ScratchBufferDesc.UnderlyingType = FRDGBufferDesc::EUnderlyingType::StructuredBuffer;
-		ScratchBufferDesc.Usage = BUF_RayTracingScratch;
+		ScratchBufferDesc.Usage = EBufferUsageFlags::RayTracingScratch | EBufferUsageFlags::StructuredBuffer;
 		ScratchBufferDesc.BytesPerElement = ScratchAlignment;
 		ScratchBufferDesc.NumElements = FMath::DivideAndRoundUp(BLASScratchSize, ScratchAlignment);
 
diff --git a/Engine/Source/Runtime/Renderer/Private/RayTracing/RayTracingScene.cpp b/Engine/Source/Runtime/Renderer/Private/RayTracing/RayTracingScene.cpp
--- a/Engine/Source/Runtime/Renderer/Private/RayTracing/RayTracingScene.cpp
+++ b/Engine/Source/Runtime/Renderer/Private/RayTracing/RayTracingScene.cpp
@@ -59,7 +59,7 @@ void FRayTracingScene::Create(FRDGBuilder& GraphBuilder, const FGPUScene& GPUSce
 		|| SizeInfo.ResultSize < RayTracingSceneBuffer->GetSize() / 2)
 	{
 		FRHIResourceCreateInfo CreateInfo(TEXT("FRayTracingScene::SceneBuffer"));
-		RayTracingSceneBuffer = RHICreateBuffer(uint32(SizeInfo.ResultSize), BUF_AccelerationStructure, 0, ERHIAccess::BVHWrite, CreateInfo);
+		RayTracingSceneBuffer = RHICreateBuffer(uint32(SizeInfo.ResultSize), EBufferUsageFlags::AccelerationStructure, 0, ERHIAccess::BVHWrite, CreateInfo);
 	}
 
 	LayerSRVs.SetNum(NumLayers);
@@ -73,8 +73,7 @@ void FRayTracingScene::Create(FRDGBuilder& GraphBuilder, const FGPUScene& GPUSce
 	{
 		const uint64 ScratchAlignment = GRHIRayTracingScratchBufferAlignment;
 		FRDGBufferDesc ScratchBufferDesc;
-		ScratchBufferDesc.UnderlyingType = FRDGBufferDesc::EUnderlyingType::StructuredBuffer;
-		ScratchBufferDesc.Usage = BUF_UnorderedAccess;
+		ScratchBufferDesc.Usage = EBufferUsageFlags::UnorderedAccess | EBufferUsageFlags::StructuredBuffer;
 		ScratchBufferDesc.BytesPerElement = uint32(ScratchAlignment);
 		ScratchBufferDesc.NumElements = uint32(FMath::DivideAndRoundUp(SizeInfo.BuildScratchSize, ScratchAlignment));
 
@@ -83,8 +82,7 @@ void FRayTracingScene::Create(FRDGBuilder& GraphBuilder, const FGPUScene& GPUSce
 
 	{
 		FRDGBufferDesc InstanceBufferDesc;
-		InstanceBufferDesc.UnderlyingType = FRDGBufferDesc::EUnderlyingType::StructuredBuffer;
-		InstanceBufferDesc.Usage = BUF_UnorderedAccess | BUF_ShaderResource;
+		InstanceBufferDesc.Usage = EBufferUsageFlags::UnorderedAccess | EBufferUsageFlags::ShaderResource | EBufferUsageFlags::StructuredBuffer;
 		InstanceBufferDesc.BytesPerElement = GRHIRayTracingInstanceDescriptorSize;
 		InstanceBufferDesc.NumElements = NumNativeInstancesAligned;
 
diff --git a/Engine/Source/Runtime/Renderer/Private/VirtualShadowMaps/VirtualShadowMapArray.cpp b/Engine/Source/Runtime/Renderer/Private/VirtualShadowMaps/VirtualShadowMapArray.cpp
--- a/Engine/Source/Runtime/Renderer/Private/VirtualShadowMaps/VirtualShadowMapArray.cpp
+++ b/Engine/Source/Runtime/Renderer/Private/VirtualShadowMaps/VirtualShadowMapArray.cpp
@@ -337,8 +337,7 @@ static FRDGBufferRef CreateProjectionDataBuffer(
 	uint64 DataSize = InitialData.Num() * InitialData.GetTypeSize();
 
 	FRDGBufferDesc Desc;
-	Desc.UnderlyingType = FRDGBufferDesc::EUnderlyingType::StructuredBuffer;
-	Desc.Usage = (EBufferUsageFlags)(BUF_UnorderedAccess | BUF_ShaderResource | BUF_ByteAddressBuffer);
+	Desc.Usage = EBufferUsageFlags::UnorderedAccess | EBufferUsageFlags::ShaderResource | EBufferUsageFlags::ByteAddressBuffer | EBufferUsageFlags::StructuredBuffer;
 	Desc.BytesPerElement = 4;
 	Desc.NumElements = DataSize / 4;
 
