Submit async compute work prior to a fence from graphics -> async compute so that prior work isn't grouped with the submission to work around an issue in D3D12, which currently brute forces a fence from graphics to async compute any time graphics work is submitted.

diff --git a/Engine/Source/Runtime/RenderCore/Private/RenderGraphBuilder.cpp b/Engine/Source/Runtime/RenderCore/Private/RenderGraphBuilder.cpp
--- a/Engine/Source/Runtime/RenderCore/Private/RenderGraphBuilder.cpp
+++ b/Engine/Source/Runtime/RenderCore/Private/RenderGraphBuilder.cpp
@@ -1373,6 +1373,8 @@ void FRDGBuilder::Compile()
 			return EarliestConsumerHandle;
 		};
 
+		FRDGPass* AsyncComputePassBeforeFork = nullptr;
+
 		const auto InsertGraphicsToAsyncComputeFork = [&](FRDGPass* GraphicsPass, FRDGPass* AsyncComputePass)
 		{
 			FRDGBarrierBatchBegin& EpilogueBarriersToBeginForAsyncCompute = GraphicsPass->GetEpilogueBarriersToBeginForAsyncCompute(Allocator, TransitionCreateQueue);
@@ -1382,6 +1384,12 @@ void FRDGBuilder::Compile()
 
 			AsyncComputePass->bAsyncComputeBegin = 1;
 			AsyncComputePass->GetPrologueBarriersToEnd(Allocator).AddDependency(&EpilogueBarriersToBeginForAsyncCompute);
+
+			// Since we are fencing the graphics pipe to some new async compute work, make sure to flush any prior work.
+			if (AsyncComputePassBeforeFork)
+			{
+				AsyncComputePassBeforeFork->bDispatchAfterExecute = 1;
+			}
 		};
 
 		const auto InsertAsyncComputeToGraphicsJoin = [&](FRDGPass* AsyncComputePass, FRDGPass* GraphicsPass)
@@ -1389,6 +1397,7 @@ void FRDGBuilder::Compile()
 			FRDGBarrierBatchBegin& EpilogueBarriersToBeginForGraphics = AsyncComputePass->GetEpilogueBarriersToBeginForGraphics(Allocator, TransitionCreateQueue);
 
 			AsyncComputePass->bAsyncComputeEnd = 1;
+			AsyncComputePass->bDispatchAfterExecute = 1;
 			EpilogueBarriersToBeginForGraphics.SetUseCrossPipelineFence();
 
 			GraphicsPass->bGraphicsJoin = 1;
@@ -1433,6 +1442,8 @@ void FRDGBuilder::Compile()
 				CurrentGraphicsForkPassHandle = GraphicsForkPassHandle;
 				InsertGraphicsToAsyncComputeFork(GraphicsForkPass, AsyncComputePass);
 			}
+
+			AsyncComputePassBeforeFork = AsyncComputePass;
 		}
 
 		FRDGPassHandle CurrentGraphicsJoinPassHandle;
@@ -1719,9 +1730,17 @@ void FRDGBuilder::Execute()
 
 						IF_RHI_WANT_BREADCRUMB_EVENTS(RHICmdList.ImportBreadcrumbState(*ParallelPassSet.BreadcrumbStateEnd));
 
-						if (ParallelPassSet.bDispatchAfterExecute && IsRunningRHIInSeparateThread())
+						if (ParallelPassSet.DispatchAfterExecutePipelines != ERHIPipeline::None)
 						{
-							RHICmdList.ImmediateFlush(EImmediateFlushType::DispatchToRHIThread);
+							if (RHICmdListAsyncCompute.HasCommands())
+							{
+								FRHIAsyncComputeCommandListImmediate::ImmediateDispatch(RHICmdListAsyncCompute);
+							}
+
+							if (EnumHasAnyFlags(ParallelPassSet.DispatchAfterExecutePipelines, ERHIPipeline::Graphics))
+							{
+								RHICmdList.ImmediateFlush(EImmediateFlushType::DispatchToRHIThread);
+							}
 						}
 					}
 
@@ -2008,15 +2027,12 @@ void FRDGBuilder::SetupPassInternal(FRDGPass* Pass, FRDGPassHandle PassHandle, E
 	check(Pass->Pipeline == PassPipeline);
 
 	Pass->bEmptyParameters = bEmptyParameters;
-	Pass->bDispatchAfterExecute = bDispatchHint;
 	Pass->GraphicsForkPass = PassHandle;
 	Pass->GraphicsJoinPass = PassHandle;
 	Pass->PrologueBarrierPass = PassHandle;
 	Pass->EpilogueBarrierPass = PassHandle;
 
-	bDispatchHint = false;
-
-	if (!EnumHasAnyFlags(Pass->Flags, ERDGPassFlags::AsyncCompute))
+	if (Pass->Pipeline == ERHIPipeline::Graphics)
 	{
 		Pass->ResourcesToBegin.Add(Pass);
 		Pass->ResourcesToEnd.Add(Pass);
@@ -2164,7 +2180,7 @@ void FRDGBuilder::SetupParallelExecute()
 
 	TArray<FRDGPass*, TInlineAllocator<64, FRDGArrayAllocator>> ParallelPassCandidates;
 	int32 MergedRenderPassCandidates = 0;
-	bool bDispatchAfterExecute = false;
+	ERHIPipeline DispatchAfterExecutePipelines = ERHIPipeline::None;
 
 	const auto FlushParallelPassCandidates = [&]()
 	{
@@ -2227,12 +2243,12 @@ void FRDGBuilder::SetupParallelExecute()
 
 			FParallelPassSet& ParallelPassSet = ParallelPassSets.Emplace_GetRef();
 			ParallelPassSet.Passes.Append(ParallelPassCandidates.GetData() + PassBeginIndex, ParallelPassCandidateCount);
-			ParallelPassSet.bDispatchAfterExecute = bDispatchAfterExecute;
+			ParallelPassSet.DispatchAfterExecutePipelines = DispatchAfterExecutePipelines;
 		}
 
 		ParallelPassCandidates.Reset();
 		MergedRenderPassCandidates = 0;
-		bDispatchAfterExecute = false;
+		DispatchAfterExecutePipelines = ERHIPipeline::None;
 	};
 
 	ParallelPassSets.Reserve(32);
@@ -2266,7 +2282,7 @@ void FRDGBuilder::SetupParallelExecute()
 		}
 
 		ParallelPassCandidates.Emplace(Pass);
-		bDispatchAfterExecute |= Pass->bDispatchAfterExecute;
+		DispatchAfterExecutePipelines |= Pass->bDispatchAfterExecute ? Pass->Pipeline : ERHIPipeline::None;
 
 		// Don't count merged render passes for the maximum pass threshold. This avoids the case where
 		// a large merged render pass span could end up forcing it back onto the render thread, since
@@ -2521,15 +2537,17 @@ void FRDGBuilder::ExecutePass(FRDGPass* Pass, FRHIComputeCommandList& RHICmdList
 
 	ExecutePassEpilogue(RHICmdListPass, Pass);
 
-	if (Pass->bAsyncComputeEnd)
+	if (!Pass->bParallelExecute && Pass->bDispatchAfterExecute)
 	{
-		RHICmdListAsyncCompute.SetStaticUniformBuffers({});
-		FRHIAsyncComputeCommandListImmediate::ImmediateDispatch(RHICmdListAsyncCompute);
-	}
+		if (RHICmdListAsyncCompute.HasCommands())
+		{
+			FRHIAsyncComputeCommandListImmediate::ImmediateDispatch(RHICmdListAsyncCompute);
+		}
 
-	if (!Pass->bParallelExecute && Pass->bDispatchAfterExecute && IsRunningRHIInSeparateThread())
-	{
-		RHICmdList.ImmediateFlush(EImmediateFlushType::DispatchToRHIThread);
+		if (Pass->Pipeline == ERHIPipeline::Graphics)
+		{
+			RHICmdList.ImmediateFlush(EImmediateFlushType::DispatchToRHIThread);
+		}
 	}
 
 	if (!bParallelExecuteEnabled)
diff --git a/Engine/Source/Runtime/RenderCore/Public/RenderGraphBuilder.h b/Engine/Source/Runtime/RenderCore/Public/RenderGraphBuilder.h
--- a/Engine/Source/Runtime/RenderCore/Public/RenderGraphBuilder.h
+++ b/Engine/Source/Runtime/RenderCore/Public/RenderGraphBuilder.h
@@ -642,7 +642,7 @@ private:
 		IF_RHI_WANT_BREADCRUMB_EVENTS(FRDGBreadcrumbState* BreadcrumbStateBegin{});
 		IF_RHI_WANT_BREADCRUMB_EVENTS(FRDGBreadcrumbState* BreadcrumbStateEnd{});
 		int8 bInitialized = 0;
-		bool bDispatchAfterExecute = false;
+		ERHIPipeline DispatchAfterExecutePipelines = ERHIPipeline::None;
 	};
 
 	TArray<FParallelPassSet, FRDGArrayAllocator> ParallelPassSets;
@@ -670,7 +670,6 @@ private:
 
 	IF_RDG_ENABLE_TRACE(FRDGTrace Trace);
 
-	bool bDispatchHint = false;
 	bool bFlushResourcesRHI = false;
 	bool bParallelExecuteEnabled = false;
 
diff --git a/Engine/Source/Runtime/RenderCore/Public/RenderGraphBuilder.inl b/Engine/Source/Runtime/RenderCore/Public/RenderGraphBuilder.inl
--- a/Engine/Source/Runtime/RenderCore/Public/RenderGraphBuilder.inl
+++ b/Engine/Source/Runtime/RenderCore/Public/RenderGraphBuilder.inl
@@ -356,7 +356,10 @@ inline void FRDGBuilder::SetCommandListStat(TStatId StatId)
 
 inline void FRDGBuilder::AddDispatchHint()
 {
-	bDispatchHint = true;
+	if (Passes.Num() > 0)
+	{
+		Passes[Passes.Last()]->bDispatchAfterExecute = 1;
+	}
 }
 
 inline const TRefCountPtr<IPooledRenderTarget>& FRDGBuilder::GetPooledTexture(FRDGTextureRef Texture) const
diff --git a/Engine/Source/Runtime/RenderCore/Public/RenderGraphPass.h b/Engine/Source/Runtime/RenderCore/Public/RenderGraphPass.h
--- a/Engine/Source/Runtime/RenderCore/Public/RenderGraphPass.h
+++ b/Engine/Source/Runtime/RenderCore/Public/RenderGraphPass.h
@@ -381,7 +381,7 @@ protected:
 			/** Whether this pass does not contain parameters. */
 			uint32 bEmptyParameters : 1;
 
-			/** If set, dispatches to the RHI thread before executing this pass. */
+			/** If set, dispatches to the RHI thread after executing this pass. */
 			uint32 bDispatchAfterExecute : 1;
 
 			/** If set, the pass should set its command list stat. */
