Implemented support for fused CopySrc | CopyDest transitions for Vulkan in RDG.
 - Simplified RDG transition logic to only use subresource transitions.
 - Don't allow unmergeable states to be merged during pass setup.

diff --git a/Engine/Source/Runtime/RenderCore/Private/RenderGraphBuilder.cpp b/Engine/Source/Runtime/RenderCore/Private/RenderGraphBuilder.cpp
--- a/Engine/Source/Runtime/RenderCore/Private/RenderGraphBuilder.cpp
+++ b/Engine/Source/Runtime/RenderCore/Private/RenderGraphBuilder.cpp
@@ -70,13 +70,25 @@ inline void EndUAVOverlap(const FRDGPass* Pass, FRHIComputeCommandList& RHICmdLi
 #endif
 }
 
-inline ERHIAccess MakeValidAccess(ERHIAccess Access)
+inline ERHIAccess MakeValidAccess(ERHIAccess AccessOld, ERHIAccess AccessNew)
 {
-	// If we find any write states in the access mask, remove all read-only states. This mainly exists
-	// to allow RDG uniform buffers to contain read-only parameters which are also bound for write on the
-	// pass. Often times these uniform buffers are created and only relevant things are accessed. If an
-	// invalid access does occur, the RHI validation layer will catch it.
-	return IsWritableAccess(Access) ? (Access & ~ERHIAccess::ReadOnlyExclusiveMask) : Access;
+	const ERHIAccess AccessUnion = AccessOld | AccessNew;
+	const ERHIAccess NonMergeableAccessMask = ~GRHIMergeableAccessMask;
+
+	// Return the union of new and old if they are okay to merge.
+	if (!EnumHasAnyFlags(AccessUnion, NonMergeableAccessMask))
+	{
+		return IsWritableAccess(AccessUnion) ? (AccessUnion & ~ERHIAccess::ReadOnlyExclusiveMask) : AccessUnion;
+	}
+
+	// Keep the old one if it can't be merged.
+	if (EnumHasAnyFlags(AccessOld, NonMergeableAccessMask))
+	{
+		return AccessOld;
+	}
+
+	// Replace with the new one if it can't be merged.
+	return AccessNew;
 }
 
 inline void GetPassAccess(ERDGPassFlags PassFlags, ERHIAccess& SRVAccess, ERHIAccess& UAVAccess)
@@ -434,18 +446,21 @@ ERDGPassFlags FRDGBuilder::OverridePassFlags(const TCHAR* PassName, ERDGPassFlag
 		true;
 #endif
 
-	const bool bGlobalForceAsyncCompute = (GRDGAsyncCompute == RDG_ASYNC_COMPUTE_FORCE_ENABLED && !IsImmediateMode() && bDebugAllowedForPass);
-
-	if (EnumHasAnyFlags(PassFlags, ERDGPassFlags::Compute) && (bGlobalForceAsyncCompute))
+	if (IsAsyncComputeSupported() && bAsyncComputeSupported)
 	{
-		PassFlags &= ~ERDGPassFlags::Compute;
-		PassFlags |= ERDGPassFlags::AsyncCompute;
+		if (EnumHasAnyFlags(PassFlags, ERDGPassFlags::Compute) && GRDGAsyncCompute == RDG_ASYNC_COMPUTE_FORCE_ENABLED)
+		{
+			PassFlags &= ~ERDGPassFlags::Compute;
+			PassFlags |= ERDGPassFlags::AsyncCompute;
+		}
 	}
-
-	if (EnumHasAnyFlags(PassFlags, ERDGPassFlags::AsyncCompute) && (GRDGAsyncCompute == RDG_ASYNC_COMPUTE_DISABLED || IsImmediateMode() || !bAsyncComputeSupported))
+	else
 	{
-		PassFlags &= ~ERDGPassFlags::AsyncCompute;
-		PassFlags |= ERDGPassFlags::Compute;
+		if (EnumHasAnyFlags(PassFlags, ERDGPassFlags::AsyncCompute))
+		{
+			PassFlags &= ~ERDGPassFlags::AsyncCompute;
+			PassFlags |= ERDGPassFlags::Compute;
+		}
 	}
 
 	return PassFlags;
@@ -598,7 +613,7 @@ END_SHADER_PARAMETER_STRUCT()
 
 void FRDGBuilder::UseExternalAccessMode(FRDGViewableResource* Resource, ERHIAccess ReadOnlyAccess, ERHIPipeline Pipelines)
 {
-	if (!GRDGAsyncCompute || !GSupportsEfficientAsyncCompute || IsImmediateMode())
+	if (!IsAsyncComputeSupported())
 	{
 		Pipelines = ERHIPipeline::Graphics;
 	}
@@ -1086,7 +1101,11 @@ void FRDGBuilder::Compile()
 				{
 					Texture->FirstBarrier = FRDGTexture::EFirstBarrier::ImmediateConfirmed;
 					Texture->FirstPass = PassHandle;
-					GetWholeResource(Texture->GetState()).SetPass(ERHIPipeline::Graphics, PassHandle);
+
+					for (FRDGSubresourceState& SubresourceState : *Texture->State)
+					{
+						SubresourceState.SetPass(ERHIPipeline::Graphics, PassHandle);
+					}
 				}
 
 			#if STATS
@@ -1877,16 +1896,6 @@ FRDGPass* FRDGBuilder::SetupPass(FRDGPass* Pass)
 			return;
 		}
 
-		// Use the whole texture range for copy states on RHI's that don't track them separately
-		if (!GRHISupportsSeparateDepthStencilCopyAccess)
-		{
-			const ERHIAccess AnyCopyAccess = ERHIAccess::CopySrc | ERHIAccess::CopyDest;
-			if ((Texture->Desc.Format == PF_DepthStencil) && EnumHasAnyFlags(Access, AnyCopyAccess))
-			{
-				Range = Texture->GetSubresourceRange();  // WholeRange
-			}
-		}
-
 		const FRDGViewHandle NoUAVBarrierHandle = GetHandleIfNoUAVBarrier(TextureView);
 		const EResourceTransitionFlags TransitionFlags = GetTextureViewTransitionFlags(TextureView, Texture);
 
@@ -1906,22 +1915,15 @@ FRDGPass* FRDGBuilder::SetupPass(FRDGPass* Pass)
 
 		PassState->ReferenceCount++;
 
-		const auto AddSubresourceAccess = [&](FRDGSubresourceState& State)
+		EnumerateSubresourceRange(PassState->State, Texture->Layout, Range, [&](FRDGSubresourceState& State)
 		{
-			State.Access = MakeValidAccess(State.Access | Access);
+			IF_RDG_ENABLE_DEBUG(UserValidation.ValidateAddSubresourceAccess(Texture, State, Access));
+
+			State.Access = MakeValidAccess(State.Access, Access);
 			State.Flags |= TransitionFlags;
 			State.NoUAVBarrierFilter.AddHandle(NoUAVBarrierHandle);
 			State.SetPass(PassPipeline, PassHandle);
-		};
-
-		if (IsWholeResource(PassState->State))
-		{
-			AddSubresourceAccess(GetWholeResource(PassState->State));
-		}
-		else
-		{
-			EnumerateSubresourceRange(PassState->State, Texture->Layout, Range, AddSubresourceAccess);
-		}
+		});
 
 		const bool bWritableAccess = IsWritableAccess(Access);
 		bRenderPassOnlyWrites &= (!bWritableAccess || EnumHasAnyFlags(AccessFlags, ERDGTextureAccessFlags::RenderTarget));
@@ -1958,8 +1960,10 @@ FRDGPass* FRDGBuilder::SetupPass(FRDGPass* Pass)
 			PassState = &Pass->BufferStates[Buffer->PassStateIndex];
 		}
 
+		IF_RDG_ENABLE_DEBUG(UserValidation.ValidateAddSubresourceAccess(Buffer, PassState->State, Access));
+
 		PassState->ReferenceCount++;
-		PassState->State.Access = MakeValidAccess(PassState->State.Access | Access);
+		PassState->State.Access = MakeValidAccess(PassState->State.Access, Access);
 		PassState->State.NoUAVBarrierFilter.AddHandle(NoUAVBarrierHandle);
 		PassState->State.SetPass(PassPipeline, PassHandle);
 
@@ -2614,7 +2618,7 @@ void FRDGBuilder::CollectPassBarriers(FRDGPass* Pass, FRDGPassHandle PassHandle)
 {
 	IF_RDG_ENABLE_DEBUG(ConditionalDebugBreak(RDG_BREAKPOINT_PASS_COMPILE, BuilderName.GetTCHAR(), Pass->GetName()));
 
-	for (const auto& PassState : Pass->TextureStates)
+	for (auto& PassState : Pass->TextureStates)
 	{
 		FRDGTextureRef Texture = PassState.Texture;
 		AddTransition(PassHandle, Texture, PassState.MergeState);
@@ -2622,7 +2626,7 @@ void FRDGBuilder::CollectPassBarriers(FRDGPass* Pass, FRDGPassHandle PassHandle)
 		IF_RDG_ENABLE_TRACE(Trace.AddTexturePassDependency(Texture, Pass));
 	}
 
-	for (const auto& PassState : Pass->BufferStates)
+	for (auto& PassState : Pass->BufferStates)
 	{
 		FRDGBufferRef Buffer = PassState.Buffer;
 		AddTransition(PassHandle, Buffer, *PassState.MergeState);
@@ -2673,7 +2677,7 @@ void FRDGBuilder::AddEpilogueTransition(FRDGTextureRef Texture)
 
 		SubresourceState.Access = Texture->EpilogueAccess;
 
-		InitAsWholeResource(ScratchTextureState, &SubresourceState);
+		InitTextureSubresources(ScratchTextureState, Texture->Layout, &SubresourceState);
 		AddTransition(EpiloguePassHandle, Texture, ScratchTextureState);
 		ScratchTextureState.Reset();
 
@@ -2738,157 +2742,118 @@ void FRDGBuilder::AddEpilogueTransition(FRDGBufferRef Buffer)
 	}
 }
 
-void FRDGBuilder::AddTransition(FRDGPassHandle PassHandle, FRDGTextureRef Texture, const FRDGTextureSubresourceStateIndirect& StateAfter)
+void FRDGBuilder::AddTransition(FRDGPassHandle PassHandle, FRDGTextureRef Texture, FRDGTextureSubresourceStateIndirect& StateAfter)
 {
-	const FRDGTextureSubresourceRange WholeRange = Texture->GetSubresourceRange();
 	const FRDGTextureSubresourceLayout Layout = Texture->Layout;
 	FRDGTextureSubresourceState& StateBefore = Texture->GetState();
+	const uint32 SubresourceCount = StateBefore.Num();
 
+	check(StateBefore.Num() == StateAfter.Num());
+	check(StateBefore.Num() == Layout.GetSubresourceCount());
 
-	// Use the entire texture for copy states on RHI's that don't track them separately
-	if (!GRHISupportsSeparateDepthStencilCopyAccess)
+	if (!GRHISupportsSeparateDepthStencilCopyAccess && Texture->Desc.Format == PF_DepthStencil)
 	{
-		if ((Texture->Desc.Format == PF_DepthStencil) && !IsWholeResource(StateAfter))
+		// Certain RHIs require a fused depth / stencil copy state. For any mip / slice transition involving a copy state,
+		// adjust the split transitions so both subresources are transitioned using the same barrier batch (i.e. the RHI transition).
+		// Note that this is only possible when async compute is disabled, as it's not possible to merge transitions from different pipes.
+		// There are two cases to correct (D for depth, S for stencil, horizontal axis is time):
+		//
+		// Case 1: both states transitioning from previous states on passes A and B to a copy state at pass C.
+		//
+		// [Pass] A     B     C                         A     B     C
+		// [D]          X --> X      Corrected To:            X --> X
+		// [S]    X --------> X                               X --> X (S is pushed forward to transition with D on pass B)
+		//
+		// Case 2a|b: one plane transitioning out of a copy state on pass A to pass B (this pass), but the other is not transitioning yet.
+		//
+		// [Pass] A     B     ?                         A     B
+		// [D]    X --> X            Corrected To:      X --> X
+		// [S]    X --------> X                         X --> X (S's state is unknown, so it transitions with D and matches D's state).
+
+		const ERHIPipeline GraphicsPipe = ERHIPipeline::Graphics;
+		const uint32 NumSlicesAndMips = Layout.NumMips * Layout.NumArraySlices;
+
+		for (uint32 DepthIndex = 0, StencilIndex = NumSlicesAndMips; DepthIndex < NumSlicesAndMips; ++DepthIndex, ++StencilIndex)
 		{
-			const ERHIAccess AnyCopyAccess = ERHIAccess::CopySrc | ERHIAccess::CopyDest;
+			FRDGSubresourceState*& DepthStateAfter   = StateAfter[DepthIndex];
+			FRDGSubresourceState*& StencilStateAfter = StateAfter[StencilIndex];
 
-			uint32 TotalTransitionCount = 0;
-			uint32 CopyTransitionCount = 0;
-			WholeRange.EnumerateSubresources([&](FRDGTextureSubresource Subresource)
+			// Skip if neither depth nor stencil are being transitioned.
+			if (!DepthStateAfter && !StencilStateAfter)
 			{
-				if (FRDGSubresourceState* SubresourceStateAfter = GetSubresource(StateAfter, Layout, Subresource))
-				{
-					++TotalTransitionCount;
-					if (EnumHasAnyFlags(SubresourceStateAfter->Access, AnyCopyAccess))
-					{
-						++CopyTransitionCount;
-					}
-				}
-			});
+				continue;
+			}
+
+			FRDGSubresourceState& DepthStateBefore   = StateBefore[DepthIndex];
+			FRDGSubresourceState& StencilStateBefore = StateBefore[StencilIndex];
 
-			// Use a common FRDGSubresourceState to keep the transitions of depth/stencil together in the same pass
-			if ((CopyTransitionCount > 0) && !IsWholeResource(StateBefore))
+			// Case 1: transitioning into a fused copy state.
+			if (DepthStateAfter && EnumHasAnyFlags(DepthStateAfter->Access, ERHIAccess::CopySrc | ERHIAccess::CopyDest))
 			{
-				checkf((CopyTransitionCount == TotalTransitionCount), TEXT("Depth/Stencil Copy transitions on this platform are supposed to touch all subresources."));
+				check(StencilStateAfter && StencilStateAfter->Access == DepthStateAfter->Access);
 
-				const FRDGTextureSubresource MaxSubresource = WholeRange.GetMaxSubresource();
-				const FRDGTextureSubresource MinSubresource = WholeRange.GetMinSubresource();
+				const FRDGPassHandle MaxPassHandle = FRDGPassHandle::Max(DepthStateBefore.LastPass[GraphicsPipe], StencilStateBefore.LastPass[GraphicsPipe]);
+				DepthStateBefore.LastPass[GraphicsPipe]   = MaxPassHandle;
+				StencilStateBefore.LastPass[GraphicsPipe] = MaxPassHandle;
+			}
+			// Case 2: transitioning out of a fused copy state.
+			else if (EnumHasAnyFlags(DepthStateBefore.Access, ERHIAccess::CopySrc | ERHIAccess::CopyDest))
+			{
+				check(StencilStateBefore.Access        == DepthStateBefore.Access);
+				check(StencilStateBefore.GetLastPass() == DepthStateBefore.GetLastPass());
 
-				for (uint32 LocalArraySlice = MinSubresource.ArraySlice; LocalArraySlice < MaxSubresource.ArraySlice; ++LocalArraySlice)
+				// Case 2a: depth unknown, so transition to match stencil.
+				if (!DepthStateAfter)
 				{
-					for (uint32 LocalMipIndex = MinSubresource.MipIndex; LocalMipIndex < MaxSubresource.MipIndex; ++LocalMipIndex)
-					{
-						FRDGSubresourceState* LatestSubresourceStateBefore = nullptr;
-						for (uint32 LocalPlaneSlice = MinSubresource.PlaneSlice; LocalPlaneSlice < MaxSubresource.PlaneSlice; ++LocalPlaneSlice)
-						{
-							FRDGSubresourceState& OtherStateBefore = GetSubresource(StateBefore, Layout, FRDGTextureSubresource(LocalMipIndex, LocalArraySlice, LocalPlaneSlice));
-							if (LatestSubresourceStateBefore == nullptr)
-							{
-								LatestSubresourceStateBefore = &OtherStateBefore;
-							}
-							else if (LatestSubresourceStateBefore->GetLastPass() < OtherStateBefore.GetLastPass())
-							{
-								*LatestSubresourceStateBefore = OtherStateBefore;
-							}
-							else if (LatestSubresourceStateBefore->GetLastPass() > OtherStateBefore.GetLastPass())
-							{
-								OtherStateBefore = *LatestSubresourceStateBefore;
-							}
-						}
-					}
+					DepthStateAfter = AllocSubresource(*StencilStateAfter);
+					DepthStateAfter->SetPass(GraphicsPipe, PassHandle);
+				}
+				// Case 2b: stencil unknown, so transition to match depth.
+				else if (!StencilStateAfter)
+				{
+					StencilStateAfter = AllocSubresource(*DepthStateAfter);
+					StencilStateAfter->SetPass(GraphicsPipe, PassHandle);
+				}
+				// Two valid after states should be transitioning on this pass.
+				else
+				{
+					check(StencilStateAfter->GetFirstPass() == PassHandle && StencilStateAfter->GetFirstPass() == PassHandle);
 				}
 			}
 		}
 	}
 
-	const auto AddSubresourceTransition = [&] (
-		const FRDGSubresourceState& SubresourceStateBefore,
-		const FRDGSubresourceState& SubresourceStateAfter,
-		FRDGTextureSubresource* Subresource)
+	for (uint32 SubresourceIndex = 0; SubresourceIndex < SubresourceCount; ++SubresourceIndex)
 	{
-		check(SubresourceStateAfter.Access != ERHIAccess::Unknown);
-
-		if (FRDGSubresourceState::IsTransitionRequired(SubresourceStateBefore, SubresourceStateAfter))
+		if (const FRDGSubresourceState* SubresourceStateAfter = StateAfter[SubresourceIndex])
 		{
-			FRHITransitionInfo Info;
-			Info.Texture = Texture->GetRHIUnchecked();
-			Info.Type = FRHITransitionInfo::EType::Texture;
-			Info.Flags = SubresourceStateAfter.Flags;
-			Info.AccessBefore = SubresourceStateBefore.Access;
-			Info.AccessAfter = SubresourceStateAfter.Access;
+			check(SubresourceStateAfter->Access != ERHIAccess::Unknown);
 
-			if (Info.AccessBefore == ERHIAccess::Discard)
-			{
-				Info.Flags |= EResourceTransitionFlags::Discard;
-			}
+			FRDGSubresourceState& SubresourceStateBefore = StateBefore[SubresourceIndex];
 
-			if (Subresource)
+			if (FRDGSubresourceState::IsTransitionRequired(SubresourceStateBefore, *SubresourceStateAfter))
 			{
-				Info.MipIndex = Subresource->MipIndex;
-				Info.ArraySlice = Subresource->ArraySlice;
-				Info.PlaneSlice = Subresource->PlaneSlice;
-			}
+				const FRDGTextureSubresource Subresource = Layout.GetSubresource(SubresourceIndex);
 
-			AddTransition(Texture, SubresourceStateBefore, SubresourceStateAfter, Info);
-		}
-	};
+				FRHITransitionInfo Info;
+				Info.Texture      = Texture->GetRHIUnchecked();
+				Info.Type         = FRHITransitionInfo::EType::Texture;
+				Info.Flags        = SubresourceStateAfter->Flags;
+				Info.AccessBefore = SubresourceStateBefore.Access;
+				Info.AccessAfter  = SubresourceStateAfter->Access;
+				Info.MipIndex     = Subresource.MipIndex;
+				Info.ArraySlice   = Subresource.ArraySlice;
+				Info.PlaneSlice   = Subresource.PlaneSlice;
 
-	if (IsWholeResource(StateBefore))
-	{
-		// 1 -> 1
-		if (IsWholeResource(StateAfter))
-		{
-			if (const FRDGSubresourceState* SubresourceStateAfter = GetWholeResource(StateAfter))
-			{
-				FRDGSubresourceState& SubresourceStateBefore = GetWholeResource(StateBefore);
-				AddSubresourceTransition(SubresourceStateBefore, *SubresourceStateAfter, nullptr);
-				SubresourceStateBefore = *SubresourceStateAfter;
-			}
-		}
-		// 1 -> N
-		else
-		{
-			const FRDGSubresourceState SubresourceStateBeforeWhole = GetWholeResource(StateBefore);
-			InitAsSubresources(StateBefore, Layout, SubresourceStateBeforeWhole);
-			WholeRange.EnumerateSubresources([&](FRDGTextureSubresource Subresource)
-			{
-				if (FRDGSubresourceState* SubresourceStateAfter = GetSubresource(StateAfter, Layout, Subresource))
+				if (Info.AccessBefore == ERHIAccess::Discard)
 				{
-					AddSubresourceTransition(SubresourceStateBeforeWhole, *SubresourceStateAfter, &Subresource);
-					FRDGSubresourceState& SubresourceStateBefore = GetSubresource(StateBefore, Layout, Subresource);
-					SubresourceStateBefore = *SubresourceStateAfter;
+					Info.Flags |= EResourceTransitionFlags::Discard;
 				}
-			});
-		}
-	}
-	else
-	{
-		// N -> 1
-		if (IsWholeResource(StateAfter))
-		{
-			if (const FRDGSubresourceState* SubresourceStateAfter = GetWholeResource(StateAfter))
-			{
-				WholeRange.EnumerateSubresources([&](FRDGTextureSubresource Subresource)
-				{
-					AddSubresourceTransition(GetSubresource(StateBefore, Layout, Subresource), *SubresourceStateAfter, &Subresource);
-				});
-				InitAsWholeResource(StateBefore);
-				FRDGSubresourceState& SubresourceStateBefore = GetWholeResource(StateBefore);
-				SubresourceStateBefore = *SubresourceStateAfter;
+
+				AddTransition(Texture, SubresourceStateBefore, *SubresourceStateAfter, Info);
 			}
-		}
-		// N -> N
-		else
-		{
-			WholeRange.EnumerateSubresources([&](FRDGTextureSubresource Subresource)
-			{
-				if (FRDGSubresourceState* SubresourceStateAfter = GetSubresource(StateAfter, Layout, Subresource))
-				{
-					FRDGSubresourceState& SubresourceStateBefore = GetSubresource(StateBefore, Layout, Subresource);
-					AddSubresourceTransition(SubresourceStateBefore, *SubresourceStateAfter, &Subresource);
-					SubresourceStateBefore = *SubresourceStateAfter;
-				}
-			});
+
+			SubresourceStateBefore = *SubresourceStateAfter;
 		}
 	}
 }
@@ -2938,15 +2903,8 @@ void FRDGBuilder::AddTransition(
 		return;
 	}
 
-	ERHIPipeline PipelinesBefore = StateBefore.GetPipelines();
-	ERHIPipeline PipelinesAfter = StateAfter.GetPipelines();
-
-	// This may be the first use of the resource in the graph, so we assign the prologue as the previous pass.
-	if (PipelinesBefore == ERHIPipeline::None)
-	{
-		StateBefore.SetPass(Graphics, GetProloguePassHandle());
-		PipelinesBefore = Graphics;
-	}
+	const ERHIPipeline PipelinesBefore = StateBefore.GetPipelines();
+	const ERHIPipeline PipelinesAfter = StateAfter.GetPipelines();
 
 	check(PipelinesBefore != ERHIPipeline::None && PipelinesAfter != ERHIPipeline::None);
 	checkf(StateBefore.GetLastPass() <= StateAfter.GetFirstPass(), TEXT("Submitted a state for '%s' that begins before our previous state has ended."), Resource->Name);
@@ -3131,7 +3089,10 @@ void FRDGBuilder::SetRHI(FRDGTexture* Texture, FRDGPooledTexture* PooledTexture,
 	else
 	{
 		Texture->State = Allocator.AllocNoDestruct<FRDGTextureSubresourceState>();
-		InitAsWholeResource(*Texture->State, {});
+
+		FRDGSubresourceState State;
+		State.SetPass(ERHIPipeline::Graphics, GetProloguePassHandle());
+		InitTextureSubresources(*Texture->State, Texture->Layout, State);
 	}
 
 	Owner = Texture;
@@ -3145,7 +3106,6 @@ void FRDGBuilder::SetRHI(FRDGTexture* Texture, FRHITransientTexture* TransientTe
 	Texture->FirstPass = PassHandle;
 	Texture->bTransient = true;
 	Texture->State = Allocator.AllocNoDestruct<FRDGTextureSubresourceState>();
-	InitAsWholeResource(*Texture->State, {});
 }
 
 void FRDGBuilder::SetRHI(FRDGBuffer* Buffer, FRDGPooledBuffer* PooledBuffer, FRDGPassHandle PassHandle)
@@ -3171,7 +3131,9 @@ void FRDGBuilder::SetRHI(FRDGBuffer* Buffer, FRDGPooledBuffer* PooledBuffer, FRD
 	}
 	else
 	{
-		Buffer->State = Allocator.AllocNoDestruct<FRDGSubresourceState>();
+		FRDGSubresourceState State;
+		State.SetPass(ERHIPipeline::Graphics, GetProloguePassHandle());
+		Buffer->State = Allocator.AllocNoDestruct<FRDGSubresourceState>(State);
 	}
 
 	Owner = Buffer;
@@ -3235,7 +3197,7 @@ void FRDGBuilder::BeginResourceRHI(FRDGPassHandle PassHandle, FRDGTextureRef Tex
 			FRDGSubresourceState InitialState;
 			InitialState.SetPass(ERHIPipeline::Graphics, MinAcquirePassHandle);
 			InitialState.Access = ERHIAccess::Discard;
-			InitAsWholeResource(*Texture->State, InitialState);
+			InitTextureSubresources(*Texture->State, Texture->Layout, InitialState);
 
 		#if STATS
 			GRDGStatTransientTextureCount++;
diff --git a/Engine/Source/Runtime/RenderCore/Private/RenderGraphPrivate.cpp b/Engine/Source/Runtime/RenderCore/Private/RenderGraphPrivate.cpp
--- a/Engine/Source/Runtime/RenderCore/Private/RenderGraphPrivate.cpp
+++ b/Engine/Source/Runtime/RenderCore/Private/RenderGraphPrivate.cpp
@@ -286,7 +286,7 @@ FAutoConsoleVariableSink CVarRDGAsyncComputeSink(FConsoleCommandDelegate::Create
 {
 	GRDGAsyncCompute = CVarRDGAsyncCompute.GetValueOnGameThread();
 
-	if (GRDGAsyncCompute && !GSupportsEfficientAsyncCompute)
+	if (!IsAsyncComputeSupported())
 	{
 		GRDGAsyncCompute = 0;
 	}
diff --git a/Engine/Source/Runtime/RenderCore/Private/RenderGraphPrivate.h b/Engine/Source/Runtime/RenderCore/Private/RenderGraphPrivate.h
--- a/Engine/Source/Runtime/RenderCore/Private/RenderGraphPrivate.h
+++ b/Engine/Source/Runtime/RenderCore/Private/RenderGraphPrivate.h
@@ -196,6 +196,11 @@ FORCEINLINE bool IsRenderPassMergeEnabled()
 	return GRDGMergeRenderPasses != 0 && !IsImmediateMode();
 }
 
+FORCEINLINE bool IsAsyncComputeSupported()
+{
+	return GRDGAsyncCompute > 0 && !IsImmediateMode() && GSupportsEfficientAsyncCompute && GRHISupportsSeparateDepthStencilCopyAccess;
+}
+
 bool IsDumpingRDGResources();
 
 FORCEINLINE bool IsParallelExecuteEnabled()
diff --git a/Engine/Source/Runtime/RenderCore/Private/RenderGraphValidation.cpp b/Engine/Source/Runtime/RenderCore/Private/RenderGraphValidation.cpp
--- a/Engine/Source/Runtime/RenderCore/Private/RenderGraphValidation.cpp
+++ b/Engine/Source/Runtime/RenderCore/Private/RenderGraphValidation.cpp
@@ -490,6 +490,16 @@ void FRDGUserValidation::ValidateExternalAccess(FRDGViewableResource* Resource,
 		Resource->Name, *GetRHIPipelineName(AccessModeState.Pipelines), *GetRHIPipelineName(Pass->GetPipeline()), Pass->GetName());
 }
 
+void FRDGUserValidation::ValidateAddSubresourceAccess(FRDGViewableResource* Resource, const FRDGSubresourceState& Subresource, ERHIAccess Access)
+{
+	const bool bOldAccessMergeable = EnumHasAnyFlags(Subresource.Access, GRHIMergeableAccessMask) || Subresource.Access == ERHIAccess::Unknown;
+	const bool bNewAccessMergeable = EnumHasAnyFlags(Access, GRHIMergeableAccessMask);
+
+	checkf(bOldAccessMergeable || bNewAccessMergeable || Subresource.Access == Access,
+		TEXT("Resource %s has incompatible access states specified for the same subresource. AccessBefore: %s, AccessAfter: %s."),
+		Resource->Name, *GetRHIAccessName(Subresource.Access), *GetRHIAccessName(Access));
+}
+
 void FRDGUserValidation::ValidateAddPass(const FRDGEventName& Name, ERDGPassFlags Flags)
 {
 	ExecuteGuard(TEXT("AddPass"), Name.GetTCHAR());
diff --git a/Engine/Source/Runtime/RenderCore/Public/RenderGraphBuilder.h b/Engine/Source/Runtime/RenderCore/Public/RenderGraphBuilder.h
--- a/Engine/Source/Runtime/RenderCore/Public/RenderGraphBuilder.h
+++ b/Engine/Source/Runtime/RenderCore/Public/RenderGraphBuilder.h
@@ -761,7 +761,7 @@ private:
 	void AddTransition(
 		FRDGPassHandle PassHandle,
 		FRDGTextureRef Texture,
-		const FRDGTextureSubresourceStateIndirect& StateAfter);
+		FRDGTextureSubresourceStateIndirect& StateAfter);
 
 	void AddTransition(
 		FRDGPassHandle PassHandle,
diff --git a/Engine/Source/Runtime/RenderCore/Public/RenderGraphDefinitions.h b/Engine/Source/Runtime/RenderCore/Public/RenderGraphDefinitions.h
--- a/Engine/Source/Runtime/RenderCore/Public/RenderGraphDefinitions.h
+++ b/Engine/Source/Runtime/RenderCore/Public/RenderGraphDefinitions.h
@@ -758,9 +758,6 @@ class FRDGTransientRenderTarget;
 template <typename TUniformStruct> class TRDGUniformBuffer;
 template <typename TUniformStruct> using TRDGUniformBufferRef = TRDGUniformBuffer<TUniformStruct>*;
 
-template <typename InElementType, typename InAllocatorType = FDefaultAllocator>
-using TRDGTextureSubresourceArray = TArray<InElementType, TInlineAllocator<1, InAllocatorType>>;
-
 using FRDGPassHandlesByPipeline = TRHIPipelineArray<FRDGPassHandle>;
 using FRDGPassesByPipeline = TRHIPipelineArray<FRDGPass*>;
 
diff --git a/Engine/Source/Runtime/RenderCore/Public/RenderGraphTextureSubresource.h b/Engine/Source/Runtime/RenderCore/Public/RenderGraphTextureSubresource.h
--- a/Engine/Source/Runtime/RenderCore/Public/RenderGraphTextureSubresource.h
+++ b/Engine/Source/Runtime/RenderCore/Public/RenderGraphTextureSubresource.h
@@ -2,7 +2,10 @@
 
 #pragma once
 
-#include "RenderGraphDefinitions.h"
+#include "RHIResources.h"
+
+template <typename InElementType, typename InAllocatorType = FDefaultAllocator>
+using TRDGTextureSubresourceArray = TArray<InElementType, TInlineAllocator<1, InAllocatorType>>;
 
 struct FRDGTextureSubresource
 {
@@ -77,8 +80,8 @@ struct FRDGTextureSubresourceLayout
 		, NumArraySlices(InNumArraySlices)
 	{}
 
-	FRDGTextureSubresourceLayout(const FRHITextureCreateInfo& CreateInfo)
-		: FRDGTextureSubresourceLayout(CreateInfo.NumMips, CreateInfo.ArraySize * (CreateInfo.IsTextureCube() ? 6 : 1), IsStencilFormat(CreateInfo.Format) ? 2 : 1)
+	FRDGTextureSubresourceLayout(const FRHITextureDesc& Desc)
+		: FRDGTextureSubresourceLayout(Desc.NumMips, Desc.ArraySize * (Desc.IsTextureCube() ? 6 : 1), IsStencilFormat(Desc.Format) ? 2 : 1)
 	{}
 
 	inline uint32 GetSubresourceCount() const
@@ -158,6 +161,11 @@ struct FRDGTextureSubresourceRange
 		return !(*this == RHS);
 	}
 
+	inline uint32 GetSubresourceCount() const
+	{
+		return NumMips * NumArraySlices * NumPlaneSlices;
+	}
+
 	FRDGTextureSubresource GetMinSubresource() const
 	{
 		return FRDGTextureSubresource(MipIndex, ArraySlice, PlaneSlice);
@@ -196,6 +204,13 @@ struct FRDGTextureSubresourceRange
 			&& NumArraySlices == Layout.NumArraySlices;
 	}
 
+	bool IsValid(const FRDGTextureSubresourceLayout& Layout) const
+	{
+		return MipIndex + NumMips <= Layout.NumMips
+			&& PlaneSlice + NumPlaneSlices <= Layout.NumPlaneSlices
+			&& ArraySlice + NumArraySlices <= Layout.NumArraySlices;
+	}
+
 	uint32 MipIndex       : 8;
 	uint32 PlaneSlice     : 8;
 	uint32 ArraySlice     : 16;
@@ -208,22 +223,13 @@ template <typename ElementType, typename AllocatorType>
 FORCEINLINE void VerifyLayout(const TRDGTextureSubresourceArray<ElementType, AllocatorType>& SubresourceArray, const FRDGTextureSubresourceLayout& Layout)
 {
 	checkf(Layout.GetSubresourceCount() > 0, TEXT("Subresource layout has no subresources."));
-	checkf(SubresourceArray.Num() == 1 || SubresourceArray.Num() == Layout.GetSubresourceCount(), TEXT("Subresource array does not match the subresource layout."));
+	checkf(SubresourceArray.Num() == Layout.GetSubresourceCount(), TEXT("Subresource array does not match the subresource layout."));
 }
 
 template <typename ElementType, typename AllocatorType>
-FORCEINLINE void InitAsWholeResource(TRDGTextureSubresourceArray<ElementType, AllocatorType>& SubresourceArray, const ElementType& Element = {})
-{
-	SubresourceArray.SetNum(1, false);
-	SubresourceArray[0] = Element;
-}
-
-template <typename ElementType, typename AllocatorType>
-inline void InitAsSubresources(TRDGTextureSubresourceArray<ElementType, AllocatorType>& SubresourceArray, const FRDGTextureSubresourceLayout& Layout, const ElementType& Element = {})
+inline void InitTextureSubresources(TRDGTextureSubresourceArray<ElementType, AllocatorType>& SubresourceArray, const FRDGTextureSubresourceLayout& Layout, const ElementType& Element = {})
 {
 	const uint32 SubresourceCount = Layout.GetSubresourceCount();
-	checkf(SubresourceCount > 0, TEXT("Subresource layout has no subresources."));
-	checkf(SubresourceCount > 1, TEXT("Subresource layout has only 1 resource. Use InitAsWholeResource instead."));
 	SubresourceArray.Reserve(SubresourceCount);
 	SubresourceArray.SetNum(SubresourceCount);
 	for (uint32 SubresourceIndex = 0; SubresourceIndex < SubresourceCount; ++SubresourceIndex)
@@ -232,39 +238,10 @@ inline void InitAsSubresources(TRDGTextureSubresourceArray<ElementType, Allocato
 	}
 }
 
-template <typename ElementType, typename AllocatorType>
-FORCEINLINE bool IsWholeResource(const TRDGTextureSubresourceArray<ElementType, AllocatorType>& SubresourceArray)
-{
-	checkf(SubresourceArray.Num() > 0, TEXT("IsWholeResource is only valid on initialized arrays."));
-	return SubresourceArray.Num() == 1;
-}
-
-template <typename ElementType, typename AllocatorType>
-FORCEINLINE bool IsSubresources(const TRDGTextureSubresourceArray<ElementType, AllocatorType>& SubresourceArray)
-{
-	checkf(SubresourceArray.Num() > 0, TEXT("IsSubresources is only valid on initialized arrays."));
-	return SubresourceArray.Num() > 1;
-}
-
-template <typename ElementType, typename AllocatorType>
-FORCEINLINE const ElementType& GetWholeResource(const TRDGTextureSubresourceArray<ElementType, AllocatorType>& SubresourceArray)
-{
-	check(IsWholeResource(SubresourceArray));
-	return SubresourceArray[0];
-}
-
-template <typename ElementType, typename AllocatorType>
-FORCEINLINE ElementType& GetWholeResource(TRDGTextureSubresourceArray<ElementType, AllocatorType>& SubresourceArray)
-{
-	checkf(IsWholeResource(SubresourceArray), TEXT("GetWholeResource with a may only be called an array initialized with InitAsWholeResource."));
-	return SubresourceArray[0];
-}
-
 template <typename ElementType, typename AllocatorType>
 FORCEINLINE const ElementType& GetSubresource(const TRDGTextureSubresourceArray<ElementType, AllocatorType>& SubresourceArray, const FRDGTextureSubresourceLayout& Layout, FRDGTextureSubresource Subresource)
 {
 	VerifyLayout(SubresourceArray, Layout);
-	checkf(IsSubresources(SubresourceArray), TEXT("GetSubresource with a may only be called an array initialized InitAsSubresources."));
 	return SubresourceArray[Layout.GetSubresourceIndex(Subresource)];
 }
 
@@ -272,7 +249,6 @@ template <typename ElementType, typename AllocatorType>
 FORCEINLINE ElementType& GetSubresource(TRDGTextureSubresourceArray<ElementType, AllocatorType>& SubresourceArray, const FRDGTextureSubresourceLayout& Layout, FRDGTextureSubresource Subresource)
 {
 	VerifyLayout(SubresourceArray, Layout);
-	checkf(IsSubresources(SubresourceArray), TEXT("GetSubresource with a may only be called an array initialized InitAsSubresources."));
 	return SubresourceArray[Layout.GetSubresourceIndex(Subresource)];
 }
 
@@ -280,7 +256,6 @@ template <typename ElementType, typename AllocatorType, typename FunctionType>
 inline void EnumerateSubresourceRange(TRDGTextureSubresourceArray<ElementType, AllocatorType>& SubresourceArray, const FRDGTextureSubresourceLayout& Layout, const FRDGTextureSubresourceRange& Range, FunctionType Function)
 {
 	VerifyLayout(SubresourceArray, Layout);
-	checkf(IsSubresources(SubresourceArray), TEXT("EnumerateSubresources with a range may only be called an array initialized as subresources."));
 	Range.EnumerateSubresources([&](FRDGTextureSubresource Subresource)
 	{
 		Function(GetSubresource(SubresourceArray, Layout, Subresource));
@@ -291,7 +266,6 @@ template <typename ElementType, typename AllocatorType, typename FunctionType>
 inline void EnumerateSubresourceRange(const TRDGTextureSubresourceArray<ElementType, AllocatorType>& SubresourceArray, const FRDGTextureSubresourceLayout& Layout, const FRDGTextureSubresourceRange& Range, FunctionType Function)
 {
 	VerifyLayout(SubresourceArray, Layout);
-	checkf(IsSubresources(SubresourceArray), TEXT("EnumerateSubresources with a range may only be called an array initialized as subresources."));
 	Range.EnumerateSubresources([&](FRDGTextureSubresource Subresource)
 	{
 		Function(GetSubresource(SubresourceArray, Layout, Subresource));
diff --git a/Engine/Source/Runtime/RenderCore/Public/RenderGraphValidation.h b/Engine/Source/Runtime/RenderCore/Public/RenderGraphValidation.h
--- a/Engine/Source/Runtime/RenderCore/Public/RenderGraphValidation.h
+++ b/Engine/Source/Runtime/RenderCore/Public/RenderGraphValidation.h
@@ -97,6 +97,8 @@ public:
 
 	void ValidateSetAccessFinal(FRDGViewableResource* Resource, ERHIAccess AccessFinal);
 
+	void ValidateAddSubresourceAccess(FRDGViewableResource* Resource, const FRDGSubresourceState& Subresource, ERHIAccess Access);
+
 	void ValidateUseExternalAccessMode(FRDGViewableResource* Resource, ERHIAccess ReadOnlyAccess, ERHIPipeline Pipelines);
 	void ValidateUseInternalAccessMode(FRDGViewableResource* Resaource);
 
