Loosened RDG validation to allow unformatted SRV / UAV when the structured buffer flag is set, rather than assert if the vertex buffer flag is present, in cases where both are set.

diff --git a/Engine/Source/Runtime/RenderCore/Public/RenderGraphResources.inl b/Engine/Source/Runtime/RenderCore/Public/RenderGraphResources.inl
--- a/Engine/Source/Runtime/RenderCore/Public/RenderGraphResources.inl
+++ b/Engine/Source/Runtime/RenderCore/Public/RenderGraphResources.inl
@@ -198,7 +198,7 @@ inline FRDGBufferSRVDesc::FRDGBufferSRVDesc(FRDGBufferRef InBuffer)
 	}
 	else
 	{
-		checkf(!EnumHasAnyFlags(Buffer->Desc.Usage, EBufferUsageFlags::VertexBuffer), TEXT("VertexBuffer %s requires a type when creating a SRV."), Buffer->Name);
+		checkf(EnumHasAnyFlags(Buffer->Desc.Usage, EBufferUsageFlags::StructuredBuffer), TEXT("Non-structured buffer %s requires a type when creating a SRV."), Buffer->Name);
 	}
 }
 
@@ -211,7 +211,7 @@ inline FRDGBufferUAVDesc::FRDGBufferUAVDesc(FRDGBufferRef InBuffer)
 	}
 	else
 	{
-		checkf(!EnumHasAnyFlags(Buffer->Desc.Usage, EBufferUsageFlags::VertexBuffer), TEXT("VertexBuffer %s requires a type when creating a UAV."), Buffer->Name);
+		checkf(EnumHasAnyFlags(Buffer->Desc.Usage, EBufferUsageFlags::StructuredBuffer), TEXT("Non-structured buffer %s requires a type when creating a UAV."), Buffer->Name);
 	}
 }
 
