Removing unhelpful RDG assert as there are external buffers with incorrect flags causing it to fire, and the underlying issue is non-fatal.

diff --git a/Engine/Source/Runtime/RenderCore/Public/RenderGraphResources.inl b/Engine/Source/Runtime/RenderCore/Public/RenderGraphResources.inl
--- a/Engine/Source/Runtime/RenderCore/Public/RenderGraphResources.inl
+++ b/Engine/Source/Runtime/RenderCore/Public/RenderGraphResources.inl
@@ -196,10 +196,6 @@ inline FRDGBufferSRVDesc::FRDGBufferSRVDesc(FRDGBufferRef InBuffer)
 	{
 		// nothing special here
 	}
-	else
-	{
-		checkf(EnumHasAnyFlags(Buffer->Desc.Usage, EBufferUsageFlags::StructuredBuffer), TEXT("Non-structured buffer %s requires a type when creating a SRV."), Buffer->Name);
-	}
 }
 
 inline FRDGBufferUAVDesc::FRDGBufferUAVDesc(FRDGBufferRef InBuffer)
@@ -209,10 +205,6 @@ inline FRDGBufferUAVDesc::FRDGBufferUAVDesc(FRDGBufferRef InBuffer)
 	{
 		Format = PF_R32_UINT;
 	}
-	else
-	{
-		checkf(EnumHasAnyFlags(Buffer->Desc.Usage, EBufferUsageFlags::StructuredBuffer), TEXT("Non-structured buffer %s requires a type when creating a UAV."), Buffer->Name);
-	}
 }
 
 inline FGraphicsPipelineRenderTargetsInfo ExtractRenderTargetsInfo(const FRDGParameterStruct& ParameterStruct)
