Fixed static analysis error and added more assert information to catch regression.

diff --git a/Engine/Source/Runtime/RenderCore/Private/RenderGraphBuilder.cpp b/Engine/Source/Runtime/RenderCore/Private/RenderGraphBuilder.cpp
--- a/Engine/Source/Runtime/RenderCore/Private/RenderGraphBuilder.cpp
+++ b/Engine/Source/Runtime/RenderCore/Private/RenderGraphBuilder.cpp
@@ -2748,8 +2748,9 @@ void FRDGBuilder::AddTransition(FRDGPassHandle PassHandle, FRDGTextureRef Textur
 	FRDGTextureSubresourceState& StateBefore = Texture->GetState();
 	const uint32 SubresourceCount = StateBefore.Num();
 
-	check(StateBefore.Num() == StateAfter.Num());
-	check(StateBefore.Num() == Layout.GetSubresourceCount());
+	checkf(StateBefore.Num() == Layout.GetSubresourceCount() && StateBefore.Num() == StateAfter.Num(),
+		TEXT("Before state array (%d) does not match after state array (%d) for resource %s on pass %s."),
+		StateBefore.Num(), StateAfter.Num(), Texture->Name, Passes[PassHandle]->GetName());
 
 	if (!GRHISupportsSeparateDepthStencilCopyAccess && Texture->Desc.Format == PF_DepthStencil)
 	{
@@ -2817,7 +2818,7 @@ void FRDGBuilder::AddTransition(FRDGPassHandle PassHandle, FRDGTextureRef Textur
 				// Two valid after states should be transitioning on this pass.
 				else
 				{
-					check(StencilStateAfter->GetFirstPass() == PassHandle && StencilStateAfter->GetFirstPass() == PassHandle);
+					check(StencilStateAfter->GetFirstPass() == PassHandle && DepthStateAfter->GetFirstPass() == PassHandle);
 				}
 			}
 		}
