Fixed crash in RDG due to transient extracted texture.

diff --git a/Engine/Source/Runtime/RenderCore/Private/RenderGraphBuilder.cpp b/Engine/Source/Runtime/RenderCore/Private/RenderGraphBuilder.cpp
--- a/Engine/Source/Runtime/RenderCore/Private/RenderGraphBuilder.cpp
+++ b/Engine/Source/Runtime/RenderCore/Private/RenderGraphBuilder.cpp
@@ -805,6 +805,13 @@ FRDGTexture* FRDGBuilder::RegisterExternalTexture(
 	Texture->bExternal = true;
 	ExternalTextures.Add(Texture->GetRHIUnchecked(), Texture);
 
+	if (Texture->bTransient)
+	{
+		FRDGSubresourceState State;
+		State.SetPass(ERHIPipeline::Graphics, GetProloguePassHandle());
+		InitTextureSubresources(*Texture->State, Texture->Layout, State);
+	}
+
 	IF_RDG_ENABLE_DEBUG(UserValidation.ValidateRegisterExternalTexture(Texture));
 	IF_RDG_ENABLE_TRACE(Trace.AddResource(Texture));
 	return Texture;
