Fixing static analysis warning.

diff --git a/Engine/Source/Runtime/RenderCore/Private/ShaderParameterStruct.cpp b/Engine/Source/Runtime/RenderCore/Private/ShaderParameterStruct.cpp
--- a/Engine/Source/Runtime/RenderCore/Private/ShaderParameterStruct.cpp
+++ b/Engine/Source/Runtime/RenderCore/Private/ShaderParameterStruct.cpp
@@ -634,7 +634,7 @@ inline FRHIDescriptorHandle GetBindlessParameterHandle(FShaderParameterReader Re
 		checkSlow(RDGShaderResourceView);
 		RDGShaderResourceView->MarkResourceAsUsed();
 
-		FRHIShaderResourceView* ShaderResourceView = RDGShaderResourceView ? RDGShaderResourceView->GetRHI() : nullptr;
+		FRHIShaderResourceView* ShaderResourceView = RDGShaderResourceView->GetRHI();
 		return ShaderResourceView ? ShaderResourceView->GetBindlessHandle() : FRHIDescriptorHandle();
 	}
 	break;
