Deprecated CopyToResolveTarget method.

diff --git a/Engine/Source/Runtime/RHI/Public/RHICommandList.h b/Engine/Source/Runtime/RHI/Public/RHICommandList.h
--- a/Engine/Source/Runtime/RHI/Public/RHICommandList.h
+++ b/Engine/Source/Runtime/RHI/Public/RHICommandList.h
@@ -3540,6 +3540,7 @@ public:
 #endif
 	}
 
+	UE_DEPRECATED(5.1, "CopyToResolveTarget is deprecated. Use render passes for MSAA resolves, or CopyTexture for copies instead.")
 	FORCEINLINE_DEBUGGABLE void CopyToResolveTarget(FRHITexture* SourceTextureRHI, FRHITexture* DestTextureRHI, const FResolveParams& ResolveParams)
 	{
 		//check(IsOutsideRenderPass());
diff --git a/Engine/Source/Runtime/RenderCore/Private/RenderGraphUtils.cpp b/Engine/Source/Runtime/RenderCore/Private/RenderGraphUtils.cpp
--- a/Engine/Source/Runtime/RenderCore/Private/RenderGraphUtils.cpp
+++ b/Engine/Source/Runtime/RenderCore/Private/RenderGraphUtils.cpp
@@ -383,7 +383,9 @@ void AddCopyToResolveTargetPass(
 		ERDGPassFlags::Copy | ERDGPassFlags::Raster | ERDGPassFlags::SkipRenderPass,
 		[InputTexture, OutputTexture, LocalResolveParams](FRHICommandList& RHICmdList)
 	{
+	PRAGMA_DISABLE_DEPRECATION_WARNINGS
 		RHICmdList.CopyToResolveTarget(InputTexture->GetRHI(), OutputTexture->GetRHI(), LocalResolveParams);
+	PRAGMA_ENABLE_DEPRECATION_WARNINGS
 	});
 }
 
