Fix disabling instanced stereo crashing the project.

diff --git a/Engine/Source/Runtime/RenderCore/Private/Shader.cpp b/Engine/Source/Runtime/RenderCore/Private/Shader.cpp
--- a/Engine/Source/Runtime/RenderCore/Private/Shader.cpp
+++ b/Engine/Source/Runtime/RenderCore/Private/Shader.cpp
@@ -1397,10 +1397,10 @@ void ShaderMapAppendKeyString(EShaderPlatform Platform, FString& KeyString)
 		bool bIsInstancedStereo = (RHISupportsInstancedStereo(Platform) && (CVarInstancedStereo && CVarInstancedStereo->GetValueOnGameThread() != 0));
 		const bool bIsMultiView = (RHISupportsMultiView(Platform) && bIsInstancedStereo);
 
-		bool bIsMobileMultiView = (CVarMobileMultiView && CVarMobileMultiView->GetValueOnGameThread() != 0);
+		bool bIsMobileMultiView = IsMobilePlatform(Platform) && (CVarMobileMultiView && CVarMobileMultiView->GetValueOnGameThread() != 0);
 		if (bIsMobileMultiView && !RHISupportsMobileMultiView(Platform))
 		{
-			// Native mobile multi-view is not supported, fall back to instancing if available
+			// Native mobile multi-view is not supported, fall back to instancing if available (even if disabled in settings)
 			bIsMobileMultiView = bIsInstancedStereo = RHISupportsInstancedStereo(Platform);
 		}
 
