Implements DynamicRenderScaling API with GPU timing measurement integrated within RDG

This allows to define a new dynamic scaling in the renderer with low amount of boiler plate:

DynamicRenderScaling::FHeuristicSettings GetDynamicTranslucencyResolutionSettings()
{
	RenderingDynamicScaling::FHeuristicSettings BucketSetting;
	BucketSetting.Model = RenderingDynamicScaling::EHeuristicModel::Quadratic;
	BucketSetting.bModelScalesWithPrimaryScreenPercentage = true;
	BucketSetting.MinResolutionFraction = ...
	...
	return BucketSetting;
}

DynamicRenderScaling::FBudget GDynamicTranslucencyResolution(TEXT("DynamicTranslucencyResolution"), &GetDynamicTranslucencyResolutionSettings);

And then simply define a scope to measure the GPU timing as such:

{
	DynamicRenderScaling::FRDGScope DynamicTranslucencyResolutionScope(GraphBuilder, GDynamicTranslucencyResolution);

	// add passes to GraphBuilder
}

diff --git a/Engine/Source/Runtime/RenderCore/Private/DumpGPU.cpp b/Engine/Source/Runtime/RenderCore/Private/DumpGPU.cpp
--- a/Engine/Source/Runtime/RenderCore/Private/DumpGPU.cpp
+++ b/Engine/Source/Runtime/RenderCore/Private/DumpGPU.cpp
@@ -1513,7 +1513,7 @@
 			bDumpPass = WildcardFilter.IsMatch(Pass->GetEventName().GetTCHAR());
 		}
 
-		#if RDG_GPU_SCOPES
+		#if RDG_GPU_DEBUG_SCOPES
 		if (!bDumpPass)
 		{
 			const FRDGEventScope* ParentScope = Pass->GetGPUScopes().Event;
@@ -2105,7 +2105,7 @@
 	// Dump the pass informations
 	{
 		TArray<TSharedPtr<FJsonValue>> ParentEventScopeNames;
-		#if RDG_GPU_SCOPES
+		#if RDG_GPU_DEBUG_SCOPES
 		{
 			const FRDGEventScope* ParentScope = Pass->GetGPUScopes().Event;
 
diff --git a/Engine/Source/Runtime/RenderCore/Private/DynamicRenderScaling.cpp b/Engine/Source/Runtime/RenderCore/Private/DynamicRenderScaling.cpp
new file mode 100644
--- /dev/null
+++ b/Engine/Source/Runtime/RenderCore/Private/DynamicRenderScaling.cpp
@@ -0,0 +1,168 @@
+// Copyright Epic Games, Inc. All Rights Reserved.
+
+#include "DynamicRenderScaling.h"
+#include "RenderCore.h"
+
+namespace
+{
+
+int32 GRDSGlobalListLinkSize = 0;
+
+}
+
+namespace DynamicRenderScaling
+{
+
+bool FHeuristicSettings::IsEnabled() const
+{
+	return BudgetMs > 0.0f;
+}
+
+float FHeuristicSettings::GetTargetedMs() const
+{
+	check(IsEnabled());
+	return BudgetMs * (1.0f - TargetedHeadRoom);
+}
+
+float FHeuristicSettings::EstimateCostScale(float ResolutionFraction) const
+{
+	float Cost = ResolutionFraction;
+	if (Model == EHeuristicModel::Linear)
+	{
+		Cost = ResolutionFraction;
+	}
+	else if (Model == EHeuristicModel::Quadratic)
+	{
+		Cost = ResolutionFraction * ResolutionFraction;
+	}
+	else
+	{
+		unimplemented();
+	}
+	return Cost;
+}
+
+float FHeuristicSettings::EstimateResolutionFactor(float TimingMs) const
+{
+	float TargetedMS = GetTargetedMs();
+
+	float S = 1.0f;
+	if (Model == EHeuristicModel::Linear)
+	{
+		S = TargetedMS / TimingMs;
+	}
+	else if (Model == EHeuristicModel::Quadratic)
+	{
+		S = FMath::Sqrt(TargetedMS / TimingMs);
+	}
+	else
+	{
+		unimplemented();
+	}
+
+	return S;
+}
+
+float FHeuristicSettings::CorrectNewResolutionFraction(float CurrentResolutionFraction, float NewResolutionFraction, float ResolutionFractionScale) const
+{
+	// If scaling the resolution up, amortize to avoid oscillations.
+	if (NewResolutionFraction > CurrentResolutionFraction)
+	{
+		NewResolutionFraction = FMath::Lerp(
+			CurrentResolutionFraction,
+			NewResolutionFraction,
+			IncreaseAmortizationFactor);
+	}
+
+	if (FractionQuantization <= 0)
+	{
+		return FMath::Clamp(NewResolutionFraction, MinResolutionFraction * ResolutionFractionScale, MaxResolutionFraction * ResolutionFractionScale);
+	}
+
+	float Lerp = FMath::Clamp((NewResolutionFraction - MinResolutionFraction * ResolutionFractionScale) / (MaxResolutionFraction * ResolutionFractionScale - MinResolutionFraction * ResolutionFractionScale), 0.0f, 1.0f);
+	int32 Quantization = FMath::RoundToZero(Lerp * FractionQuantization);
+	float NewLerp = float(Quantization) / float(FractionQuantization);
+
+	return FMath::Lerp(MinResolutionFraction, MaxResolutionFraction, Lerp) * ResolutionFractionScale;
+}
+
+bool FHeuristicSettings::DoesResolutionChangeEnough(float CurrentResolutionFraction, float NewResolutionFraction, bool bCanChangeResolution) const
+{
+	float IncreaseResolutionFractionThreshold = CurrentResolutionFraction * (1.0f + ChangeThreshold);
+	float DecreaseResolutionFractionThreshold = CurrentResolutionFraction * (1.0f - ChangeThreshold);
+	return bCanChangeResolution && (
+		(NewResolutionFraction > IncreaseResolutionFractionThreshold) ||
+		(NewResolutionFraction < DecreaseResolutionFractionThreshold) ||
+		(NewResolutionFraction != CurrentResolutionFraction && NewResolutionFraction == MinResolutionFraction) ||
+		(NewResolutionFraction != CurrentResolutionFraction && NewResolutionFraction == MaxResolutionFraction));
+}
+
+float FHeuristicSettings::EstimateTimeFactor(float CurrentResolutionFraction, float NewResolutionFraction) const
+{
+	float R = NewResolutionFraction / CurrentResolutionFraction;
+
+	if (Model == EHeuristicModel::Linear)
+	{
+		return R;
+	}
+	else if (Model == EHeuristicModel::Quadratic)
+	{
+		return R * R;
+	}
+	else
+	{
+		unimplemented();
+	}
+
+	return 1.0f;
+}
+
+FBudget::FBudget(const TCHAR* InName, FHeuristicSettings (*InHeuristicSettingsGetter)(void))
+	: Name(InName)
+	, HeuristicSettingsGetter(InHeuristicSettingsGetter)
+	, GlobalListLink(this)
+{
+	check(InName);
+	GlobalListLink.LinkHead(FBudget::GetGlobalList());
+	BudgetId = GRDSGlobalListLinkSize;
+	GRDSGlobalListLinkSize++;
+
+#if !UE_BUILD_SHIPPING
+	if (GRDSGlobalListLinkSize > TMap<float>::kInlineAllocatedBudgets)
+	{
+		UE_LOG(LogRendererCore, Warning, TEXT("Consider bumping DynamicRenderScaling::TMap::kInlineAlloactedBudgets to %d."), GRDSGlobalListLinkSize);
+	}
+#endif
+}
+
+FBudget::~FBudget()
+{
+	GRDSGlobalListLinkSize--;
+	GlobalListLink.Unlink();
+}
+
+// static
+TLinkedList<FBudget*>*& FBudget::GetGlobalList()
+{
+	static TLinkedList<FBudget*>* GDynamicScalingBudgetsList = nullptr;
+	return GDynamicScalingBudgetsList;
+}
+
+// static
+int32 DynamicRenderScaling::FBudget::GetGlobalListSize()
+{
+	return GRDSGlobalListLinkSize;
+}
+
+void UpdateHeuristicsSettings()
+{
+	check(IsInRenderingThread());
+	for (TLinkedList<DynamicRenderScaling::FBudget*>::TIterator BudgetIt(DynamicRenderScaling::FBudget::GetGlobalList()); BudgetIt; BudgetIt.Next())
+	{
+		DynamicRenderScaling::FBudget& Budget = **BudgetIt;
+		Budget.CachedSettings = (*Budget.HeuristicSettingsGetter)();
+		check(Budget.CachedSettings.Model != EHeuristicModel::Unknown);
+	}
+}
+
+} // namespace DynamicRenderScaling
diff --git a/Engine/Source/Runtime/RenderCore/Private/RenderGraphBuilder.cpp b/Engine/Source/Runtime/RenderCore/Private/RenderGraphBuilder.cpp
--- a/Engine/Source/Runtime/RenderCore/Private/RenderGraphBuilder.cpp
+++ b/Engine/Source/Runtime/RenderCore/Private/RenderGraphBuilder.cpp
@@ -539,9 +539,7 @@
 #if RDG_CPU_SCOPES
 	, CPUScopeStacks(Allocator)
 #endif
-#if RDG_GPU_SCOPES
 	, GPUScopeStacks(Allocator)
-#endif
 	, bParallelExecuteEnabled(IsParallelExecuteEnabled() && EnumHasAnyFlags(InFlags, ERDGBuilderFlags::AllowParallelExecute))
 #if RDG_ENABLE_DEBUG
 	, UserValidation(Allocator, bParallelExecuteEnabled)
@@ -1537,7 +1535,7 @@
 
 		Compile();
 
-		IF_RDG_GPU_SCOPES(GPUScopeStacks.ReserveOps(Passes.Num()));
+		GPUScopeStacks.ReserveOps(Passes.Num());
 		IF_RDG_CPU_SCOPES(CPUScopeStacks.ReserveOps());
 
 		if (bParallelExecuteEnabled)
@@ -1841,8 +1839,8 @@
 
 	IF_RDG_ENABLE_TRACE(Trace.OutputGraphEnd(*this));
 
-	IF_RDG_GPU_SCOPES(GPUScopeStacks.Graphics.EndExecute(RHICmdList));
-	IF_RDG_GPU_SCOPES(GPUScopeStacks.AsyncCompute.EndExecute(RHICmdListAsyncCompute));
+	GPUScopeStacks.Graphics.EndExecute(RHICmdList);
+	GPUScopeStacks.AsyncCompute.EndExecute(RHICmdListAsyncCompute);
 	IF_RDG_CPU_SCOPES(CPUScopeStacks.EndExecute());
 
 	IF_RDG_ENABLE_DEBUG(UserValidation.ValidateExecuteEnd());
@@ -2026,10 +2024,8 @@
 	Pass->CPUScopeOps = CPUScopeStacks.CompilePassPrologue(Pass);
 #endif
 
-#if RDG_GPU_SCOPES
 	Pass->GPUScopeOpsPrologue = GPUScopeStacks.CompilePassPrologue(Pass, GPUMask);
 	Pass->GPUScopeOpsEpilogue = GPUScopeStacks.CompilePassEpilogue(Pass);
-#endif
 }
 
 void FRDGBuilder::SetupPassInternal(FRDGPass* Pass, FRDGPassHandle PassHandle, ERHIPipeline PassPipeline, bool bEmptyParameters)
@@ -2064,13 +2060,13 @@
 #endif
 
 	IF_RDG_CPU_SCOPES(Pass->CPUScopes = CPUScopeStacks.GetCurrentScopes());
-	IF_RDG_GPU_SCOPES(Pass->GPUScopes = GPUScopeStacks.GetCurrentScopes(PassPipeline));
+	Pass->GPUScopes = GPUScopeStacks.GetCurrentScopes(PassPipeline);
 
-#if RDG_GPU_SCOPES && RDG_ENABLE_TRACE
+#if RDG_GPU_DEBUG_SCOPES && RDG_ENABLE_TRACE
 	Pass->TraceEventScope = GPUScopeStacks.GetCurrentScopes(ERHIPipeline::Graphics).Event;
 #endif
 
-#if RDG_GPU_SCOPES && RDG_ENABLE_DEBUG
+#if RDG_GPU_DEBUG_SCOPES && RDG_ENABLE_DEBUG
 	if (const FRDGEventScope* Scope = Pass->GPUScopes.Event)
 	{
 		Pass->FullPathIfDebug = Scope->GetPath(Pass->Name);
@@ -2528,9 +2524,7 @@
 
 	ExecutePassPrologue(RHICmdListPass, Pass);
 
-#if RDG_GPU_SCOPES
 	Pass->GPUScopeOpsPrologue.Execute(RHICmdListPass);
-#endif
 
 #if RDG_DUMP_RESOURCES_AT_EACH_DRAW
 	BeginPassDump(Pass);
@@ -2542,9 +2536,7 @@
 	EndPassDump(Pass);
 #endif
 
-#if RDG_GPU_SCOPES
 	Pass->GPUScopeOpsEpilogue.Execute(RHICmdListPass);
-#endif
 
 	ExecutePassEpilogue(RHICmdListPass, Pass);
 
diff --git a/Engine/Source/Runtime/RenderCore/Private/RenderGraphEvent.cpp b/Engine/Source/Runtime/RenderCore/Private/RenderGraphEvent.cpp
--- a/Engine/Source/Runtime/RenderCore/Private/RenderGraphEvent.cpp
+++ b/Engine/Source/Runtime/RenderCore/Private/RenderGraphEvent.cpp
@@ -5,6 +5,499 @@
 #include "RenderGraphPrivate.h"
 #include "RenderGraphPass.h"
 
+class FRDGTimingPool : public FRenderResource
+{
+public:
+	FRenderQueryPoolRHIRef QueryPool;
+
+	// Destructor
+	virtual ~FRDGTimingPool() = default;
+
+	virtual void InitDynamicRHI() override
+	{
+		check(IsInRenderingThread());
+		bIsBudgetRecordingEnabled.SetAll(false);
+		LastTimings.SetAll(uint64(0));
+	}
+
+	virtual void ReleaseDynamicRHI() override
+	{
+		check(IsInRenderingThread());
+
+		if (QueryPool)
+		{
+			// Release all in-flight queries
+			LandInFlightFrames(/* bWait = */ true);
+			InFlightFrames.Reset();
+
+			// Release the pool
+			QueryPool.SafeRelease();
+		}
+	}
+
+	/* Scopes of a timing budget. */
+	struct FInFlightTimingScope
+	{
+		// Global index of the timing budget this scope belongs to.
+		int32 BudgetId = -1;
+
+		// Index of the query at the begining of the scope in FInFlightFrame::TimestampQueries
+		int32 BeginQueryId = -1;
+
+		// Index of the query at the ending of the scope in FInFlightFrame::TimestampQueries
+		int32 EndQueryId = -1;
+
+		bool IsValid() const
+		{
+			return BeginQueryId >= 0 && EndQueryId >= 0 && BudgetId >= 0;
+		}
+	};
+
+	/* Full frame of timestamp queries in flight. */
+	struct FInFlightFrame
+	{
+		static const int32 kTimingScopesPreallocation = 64;
+		static const int32 kTimestampQueriesPreallocation = kTimingScopesPreallocation * 2;
+
+		// Frame counter this frame belongs too, or -1 if this entry cane be reused for a new frame.
+		uint64 FrameCounter;
+
+		// Arrays of all timestamp queries issued in this frame.
+		TArray<FRHIPooledRenderQuery> TimestampQueries;
+
+		// Arrays of all scopes issued in this frame.
+		TArray<FInFlightTimingScope> TimingScopes;
+
+		// Indices of TimingScopes::FInFlightTimingScope that has been began but not ended of different pipes.
+		DynamicRenderScaling::TMap<int32> LastTimingScopesGraphics;
+		DynamicRenderScaling::TMap<int32> LastTimingScopesAsyncCompute;
+
+		// Last timestamp queries issues for all pipes to  know whether an in-flight frame has landed.
+		FRHIRenderQuery* LastQueryGraphics;
+		FRHIRenderQuery* LastQueryAsyncCompute;
+
+		// Fence for the RHI command to be completed before polling RHI queries.
+		mutable FGraphEventRef RHIEndFence;
+
+		FInFlightFrame()
+		{
+			 ResetValues();
+		}
+
+		// Begins a new timing scope for a specific budget on a pipeline.
+		void BeginTimingScope(ERHIPipeline Pipeline, int32 BudgetId, int32 BeginQueryId)
+		{
+			check(Pipeline == ERHIPipeline::Graphics || Pipeline == ERHIPipeline::AsyncCompute);
+			DynamicRenderScaling::TMap<int32>& LastTimingScopes = Pipeline == ERHIPipeline::Graphics ? LastTimingScopesGraphics : LastTimingScopesAsyncCompute;
+
+			// Make sure there isn't an ongoing scope on going for that same budget.
+			check(LastTimingScopes[BudgetId] == -1);
+
+			if (TimingScopes.Num() == 0)
+			{
+				TimingScopes.Reserve(kTimingScopesPreallocation);
+			}
+			else if (TimingScopes.Num() == TimingScopes.Max())
+			{
+				TimingScopes.Reserve(TimingScopes.Max() * 2);
+			}
+
+			const int32 TimingScopeId = TimingScopes.Num();
+
+			FInFlightTimingScope NewTimingScope;
+			NewTimingScope.BudgetId = BudgetId;
+			NewTimingScope.BeginQueryId = BeginQueryId;
+			TimingScopes.Add(NewTimingScope);
+
+			LastTimingScopes[BudgetId] = TimingScopeId;
+		}
+
+		// Ends a new timing scope for a specific budget on a pipeline.
+		void EndTimingScope(ERHIPipeline Pipeline, int32 BudgetId, int32 EndQueryId)
+		{
+			check(Pipeline == ERHIPipeline::Graphics || Pipeline == ERHIPipeline::AsyncCompute);
+			DynamicRenderScaling::TMap<int32>& LastTimingScopes = Pipeline == ERHIPipeline::Graphics ? LastTimingScopesGraphics : LastTimingScopesAsyncCompute;
+
+			// Make sure there is an ongoing scope on going for that same budget.
+			const int32 TimingScopeId = LastTimingScopes[BudgetId];
+			check(TimingScopeId != -1);
+
+			FInFlightTimingScope& NewTimingScope = TimingScopes[TimingScopeId];
+			check(NewTimingScope.BudgetId == BudgetId);
+			NewTimingScope.EndQueryId = EndQueryId;
+			check(NewTimingScope.IsValid());
+
+			LastTimingScopes[BudgetId] = -1;
+		}
+
+		// Whether this entry is currently in flight or can be reused as new in flight frame.
+		bool IsInFlight() const
+		{
+			return FrameCounter != uint64(-1);
+		}
+
+		// Returns whether this frame is landed.
+		bool IsLanded(bool bWait) const
+		{
+			check(IsInRenderingThread());
+			check(IsInFlight());
+
+			// If there is a RHI thread, make sure all the queries has been processed by RHI thread to avoid thread race
+			// between RHIEndRenderQuery() and RHIGetRenderQueryResult()
+			if (IsRunningRHIInSeparateThread())
+			{
+				if (RHIEndFence.GetReference() && !RHIEndFence->IsComplete())
+				{
+					if (bWait)
+					{
+						FRHICommandListExecutor::WaitOnRHIThreadFence(RHIEndFence);
+					}
+					else
+					{
+						// Not all RHIEndRenderQuery() has been processed by RHI thread, so return none of the queries have landed
+						return false;
+					}
+				}
+			}
+
+			uint64 Timestamp = 0;
+			bool bIsLanded = true;
+
+			// verify last graphic and async compute queries are all landed.
+			if (LastQueryGraphics)
+			{
+				bool bQueryIsLanded = RHIGetRenderQueryResult(LastQueryGraphics, /* out */ Timestamp, bWait);
+				bIsLanded = bIsLanded && bQueryIsLanded;
+			}
+			if (LastQueryAsyncCompute)
+			{
+				bool bQueryIsLanded = RHIGetRenderQueryResult(LastQueryAsyncCompute, /* out */ Timestamp, bWait);
+				bIsLanded = bIsLanded && bQueryIsLanded;
+			}
+
+			if (!bIsLanded)
+			{
+				return false;
+			}
+
+			// Verify all the queries are truly landed.
+			for (int32 i = 0; i < TimestampQueries.Num(); i++)
+			{
+				bool bQueryIsLanded = RHIGetRenderQueryResult(TimestampQueries[i].GetQuery(), /* out */ Timestamp, bWait);
+				bIsLanded = bIsLanded && bQueryIsLanded;
+			}
+
+			return bIsLanded;
+		}
+
+		// Aggregates all the timing from the different budget together.
+		DynamicRenderScaling::TMap<uint64> AggregateLandedTimings(bool bWait) const
+		{
+			check(IsInRenderingThread());
+			DynamicRenderScaling::TMap<uint64> Timings;
+			Timings.SetAll(uint64(0));
+
+			// Lands all tthe queries
+			TArray<uint64> TimestampQueryResults;
+			TimestampQueryResults.Reserve(TimestampQueries.Num());
+			for (int32 i = 0; i < TimestampQueries.Num(); i++)
+			{
+				uint64 Timestamp = 0;
+				FRHIRenderQuery* Query = TimestampQueries[i].GetQuery();
+				bool bQueryIsLanded = RHIGetRenderQueryResult(TimestampQueries[i].GetQuery(), /* out */ Timestamp, bWait);
+				check(bQueryIsLanded);
+				TimestampQueryResults.Add(Timestamp);
+			}
+
+			for (const FInFlightTimingScope& Scope : TimingScopes)
+			{
+				Timings[Scope.BudgetId] += TimestampQueryResults[Scope.EndQueryId] - TimestampQueryResults[Scope.BeginQueryId];
+			}
+			return Timings;
+		}
+
+		void ResetValues()
+		{
+			FrameCounter = uint64(-1);
+			TimestampQueries.Empty(TimestampQueries.Max());
+			TimingScopes.Empty(TimingScopes.Max());
+			LastTimingScopesGraphics.SetAll(-1);
+			LastTimingScopesAsyncCompute.SetAll(-1);
+			LastQueryGraphics = nullptr;
+			LastQueryAsyncCompute = nullptr;
+			RHIEndFence = FGraphEventRef();
+		}
+	};
+
+	// Whether should record timing this frame.
+	bool IsRecordingThisFrame() const
+	{
+		return CurrentFrameInFlightIndex >= 0;
+	}
+
+	// Whether should record timing this frame.
+	bool IsRecordingThisFrame(const DynamicRenderScaling::FBudget& BudgetId) const
+	{
+		check(IsInRenderingThread());
+		return IsRecordingThisFrame() && bIsBudgetRecordingEnabled[BudgetId];
+	}
+
+	void LandInFlightFrames(bool bWait)
+	{
+		check(!IsRecordingThisFrame());
+		for (int32 i = 0; i < InFlightFrames.Num(); i++)
+		{
+			FInFlightFrame& InFlightFrame = InFlightFrames[i];
+
+			if (!InFlightFrame.IsInFlight())
+			{
+				continue;
+			}
+
+			if (!InFlightFrame.IsLanded(bWait))
+			{
+				continue;
+			}
+
+			if (InFlightFrame.FrameCounter > LastTimingFrameCounter)
+			{
+				LastTimings = InFlightFrame.AggregateLandedTimings(bWait);
+			}
+
+			InFlightFrame.ResetValues();
+		}
+
+	}
+
+	void BeginFrame(const DynamicRenderScaling::TMap<bool>& bInIsBudgetEnabled)
+	{
+		check(IsInRenderingThread());
+		check(CurrentFrameInFlightIndex == -1);
+
+
+		// Land frames
+		{
+			LandInFlightFrames(/* bWait = */ false);
+		}
+
+		bool bRecordThisFrame = false;
+		for (TLinkedList<DynamicRenderScaling::FBudget*>::TIterator BudgetIt(DynamicRenderScaling::FBudget::GetGlobalList()); BudgetIt; BudgetIt.Next())
+		{
+			const DynamicRenderScaling::FBudget& Budget = **BudgetIt;
+			bRecordThisFrame = bRecordThisFrame || bInIsBudgetEnabled[Budget];
+		}
+
+		// Allocate new inflight frame.
+		if (bRecordThisFrame)
+		{
+			check(GRHISupportsGPUTimestampBubblesRemoval);
+
+			if (!QueryPool.IsValid())
+			{
+				QueryPool = RHICreateRenderQueryPool(RQT_AbsoluteTime);
+			}
+
+			for (int32 i = 0; i < InFlightFrames.Num(); i++)
+			{
+				FInFlightFrame& InFlightFrame = InFlightFrames[i];
+				if (!InFlightFrame.IsInFlight())
+				{
+					CurrentFrameInFlightIndex = i;
+					break;
+				}
+			}
+
+			// Allocate a new in-flight frame in the unlikely event.
+			if (CurrentFrameInFlightIndex == -1)
+			{
+				CurrentFrameInFlightIndex = InFlightFrames.Num();
+				InFlightFrames.Add(FInFlightFrame());
+
+				// make sure there is no memory leak, we Really shouldn't have more than 10 frame in flight.
+				ensure(InFlightFrames.Num() < 10);
+			}
+
+			InFlightFrames[CurrentFrameInFlightIndex].FrameCounter = GFrameCounterRenderThread;
+			check(InFlightFrames[CurrentFrameInFlightIndex].IsInFlight());
+
+			bIsBudgetRecordingEnabled = bInIsBudgetEnabled;
+		}
+	}
+
+	void EndFrame()
+	{
+		check(IsInRenderingThread());
+
+		if (IsRecordingThisFrame())
+		{
+			FInFlightFrame& InFlightFrame = InFlightFrames[CurrentFrameInFlightIndex];
+
+			if (IsRunningRHIInSeparateThread())
+			{
+				FRHICommandListImmediate& RHICmdList = FRHICommandListExecutor::GetImmediateCommandList();
+				InFlightFrame.RHIEndFence = RHICmdList.RHIThreadFence();
+				RHICmdList.ImmediateFlush(EImmediateFlushType::DispatchToRHIThread);
+			}
+
+			CurrentFrameInFlightIndex = -1;
+			bIsBudgetRecordingEnabled.SetAll(false);
+		}
+	}
+
+	const DynamicRenderScaling::TMap<uint64>& GetLatestTimings() const
+	{
+		check(IsInRenderingThread());
+		return LastTimings;
+	}
+
+	FInFlightFrame* GetCurrentInFlightFrame()
+	{
+		check(IsRecordingThisFrame());
+		check(InFlightFrames[CurrentFrameInFlightIndex].IsInFlight());
+		return &InFlightFrames[CurrentFrameInFlightIndex];
+	}
+
+	void CreateTimestampQuery(ERHIPipeline Pipeline, FInFlightFrame* InFlightFrame, FRHIRenderQuery*& OutQuery, int32& OutQueryIndex)
+	{
+		if (InFlightFrame->TimestampQueries.Num() == 0)
+		{
+			InFlightFrame->TimestampQueries.Reserve(FInFlightFrame::kTimestampQueriesPreallocation);
+		}
+		else if (InFlightFrame->TimestampQueries.Num() == InFlightFrame->TimestampQueries.Max())
+		{
+			InFlightFrame->TimestampQueries.Reserve(InFlightFrame->TimestampQueries.Max() * 2);
+		}
+
+		OutQueryIndex = InFlightFrame->TimestampQueries.Num();
+		InFlightFrame->TimestampQueries.Add(QueryPool->AllocateQuery());
+		OutQuery = InFlightFrame->TimestampQueries.Last().GetQuery();
+
+		if (Pipeline == ERHIPipeline::Graphics)
+		{
+			InFlightFrame->LastQueryGraphics = OutQuery;
+		}
+		else if (Pipeline == ERHIPipeline::AsyncCompute)
+		{
+			InFlightFrame->LastQueryAsyncCompute = OutQuery;
+		}
+		else
+		{
+			unimplemented();
+		}
+	}
+
+	// List of frame queries in flight.
+private:
+	TArray<FInFlightFrame> InFlightFrames;
+
+	DynamicRenderScaling::TMap<bool> bIsBudgetRecordingEnabled;
+	DynamicRenderScaling::TMap<uint64> LastTimings;
+	uint64 LastTimingFrameCounter = 0;
+
+	// Current frame's in flight index.
+	int32 CurrentFrameInFlightIndex = -1;
+};
+
+TGlobalResource<FRDGTimingPool> GRDGTimingPool;
+
+
+FRDGTimingScopeOpArray::FRDGTimingScopeOpArray(ERHIPipeline Pipeline, const TRDGScopeOpArray<FRDGTimingScopeOp>& Ops)
+{
+	if (Ops.Num() == 0 || Pipeline != ERHIPipeline::Graphics)
+	{
+		return;
+	}
+
+	FRDGTimingPool::FInFlightFrame* InFlightFrame = GRDGTimingPool.GetCurrentInFlightFrame();
+
+	// Issue a new timestamp query to share across the different FInFlightTimingScope
+	int32 TimestampQueryIndex = -1;
+	{
+		GRDGTimingPool.CreateTimestampQuery(Pipeline, InFlightFrame, /* out */ TimestampQuery, /* out */ TimestampQueryIndex);
+		check(TimestampQuery && TimestampQueryIndex >= 0);
+	}
+
+	for (int32 Index = 0; Index < Ops.Num(); ++Index)
+	{
+		FRDGTimingScopeOp Op = Ops[Index];
+		check(Op.IsScope());
+
+		int32 BudgetId = Op.Scope->GetBudgetId();
+
+		if (Op.IsPush())
+		{
+			InFlightFrame->BeginTimingScope(Pipeline, BudgetId, TimestampQueryIndex);
+		}
+		else
+		{
+			InFlightFrame->EndTimingScope(Pipeline, BudgetId, TimestampQueryIndex);
+		}
+	}
+}
+
+void FRDGTimingScopeOpArray::Execute(FRHIComputeCommandList& RHICmdList)
+{
+	ERHIPipeline Pipeline = RHICmdList.GetPipeline();
+	if (Pipeline != ERHIPipeline::Graphics)
+	{
+		check(TimestampQuery == nullptr);
+		return; // TODO: FRHIComputeCommandList::EndRenderQuery()
+	}
+
+	if (TimestampQuery != nullptr)
+	{
+		FRHICommandList& RHICmdListGraphics = static_cast<FRHICommandList&>(RHICmdList);
+		RHICmdListGraphics.EndRenderQuery(TimestampQuery);
+	}
+}
+
+FRDGTimingScopeOpArray FRDGTimingScopeStack::CompilePassPrologue(const FRDGPass* Pass)
+{
+	ERHIPipeline Pipeline = Pass->GetPipeline();
+	TRDGScopeOpArray<FRDGTimingScopeOp> Ops = ScopeStack.CompilePassPrologue(Pass->GetGPUScopes().Timing);
+
+	return FRDGTimingScopeOpArray(Pipeline, Ops);
+}
+
+namespace DynamicRenderScaling
+{
+
+FRDGScope::FRDGScope(FRDGBuilder& InGraphBuilder, const DynamicRenderScaling::FBudget& InBudget)
+	: GraphBuilder(InGraphBuilder)
+	, bIsEnabled(GRDGTimingPool.IsRecordingThisFrame(InBudget))
+{
+	if (bIsEnabled)
+	{
+		GraphBuilder.GPUScopeStacks.BeginTimingScope(InBudget);
+	}
+}
+
+FRDGScope::~FRDGScope()
+{
+	if (bIsEnabled)
+	{
+		GraphBuilder.GPUScopeStacks.EndTimingScope();
+	}
+}
+
+void BeginFrame(const DynamicRenderScaling::TMap<bool>& bIsBudgetEnabled)
+{
+	GRDGTimingPool.BeginFrame(bIsBudgetEnabled);
+}
+
+void EndFrame()
+{
+	GRDGTimingPool.EndFrame();
+}
+
+const TMap<uint64>& GetLastestTimings()
+{
+	return GRDGTimingPool.GetLatestTimings();
+}
+
+} // namespace DynamicRenderScaling
+
+
 RENDERCORE_API bool GetEmitRDGEvents()
 {
 #if RDG_EVENTS != RDG_EVENTS_NONE
@@ -39,7 +532,7 @@
 
 #endif
 
-#if RDG_GPU_SCOPES
+#if RDG_GPU_DEBUG_SCOPES
 
 static void GetEventScopePathRecursive(const FRDGEventScope* Root, FString& String)
 {
diff --git a/Engine/Source/Runtime/RenderCore/Public/DynamicRenderScaling.h b/Engine/Source/Runtime/RenderCore/Public/DynamicRenderScaling.h
new file mode 100644
--- /dev/null
+++ b/Engine/Source/Runtime/RenderCore/Public/DynamicRenderScaling.h
@@ -0,0 +1,217 @@
+// Copyright Epic Games, Inc. All Rights Reserved.
+
+#pragma once
+
+#include "CoreMinimal.h"
+#include "Containers/List.h"
+#include "HAL/IConsoleManager.h"
+
+namespace DynamicRenderScaling
+{
+
+template<typename Type>
+class TMap;
+
+/* Model how the GPU costs of a budget scales with a ResolutionFraction. */
+enum class EHeuristicModel
+{
+	Unknown,
+
+	// GPU cost scales linearly to the ResolutionFraction
+	Linear,
+
+	// GPU cost scales quadratically to the ResolutionFraction
+	Quadratic,
+};
+
+/* Represent an independent budget to dynamically scale by its own. */
+struct RENDERCORE_API FHeuristicSettings final
+{
+	static constexpr float kDefaultMinResolutionFraction = 0.5f;
+	static constexpr float kDefaultMaxResolutionFraction = 1.0f;
+	static constexpr float kBudgetMsDisabled = 0.0f;
+	static constexpr float kDefaultChangeThreshold = 0.02f;
+	static constexpr float kDefaultTargetedHeadRoom = 0.05f;
+	static constexpr float kDefaultIncreaseAmortizationFactor = 0.9f;
+	static constexpr int32 kDefaultFractionQuantization = 0;
+
+
+	EHeuristicModel Model = EHeuristicModel::Unknown;
+	bool bModelScalesWithPrimaryScreenPercentage = false;
+
+	float MinResolutionFraction      = kDefaultMinResolutionFraction;
+	float MaxResolutionFraction      = kDefaultMaxResolutionFraction;
+	float BudgetMs                   = kBudgetMsDisabled;
+	float ChangeThreshold            = kDefaultChangeThreshold;
+	float TargetedHeadRoom           = kDefaultTargetedHeadRoom;
+	float IncreaseAmortizationFactor = kDefaultIncreaseAmortizationFactor;
+	int32 FractionQuantization       = kDefaultFractionQuantization;
+	
+
+	/** Returns whether the heuristic is enabled or not. */
+	bool IsEnabled() const;
+
+	/** Returns the desired GPU cost to be targeted to have head room left to not go over budget. */
+	float GetTargetedMs() const;
+
+	/** Returns how much the GPU cost scales for a given ResolutionFraction. */
+	float EstimateCostScale(float ResolutionFraction) const;
+
+	/** Returns how much the ResolutionFraction should scale for a GPU timing to fit to target. */
+	float EstimateResolutionFactor(float TimingMs) const;
+
+	/** Returns how much the GPU time should scale between two different resolution fraction. */
+	float EstimateTimeFactor(float CurrentResolutionFraction, float NewResolutionFraction) const;
+
+	/** Corrects new resolution fraction to ensure it's within bounds, honor AmortizationFactor and FractionQuantization. */
+	float CorrectNewResolutionFraction(float CurrentResolutionFraction, float NewResolutionFraction, float ResolutionFractionScale) const;
+
+	/** Returns whether the ResolutionFraction is changing meaningfully enough. */
+	bool DoesResolutionChangeEnough(float CurrentResolutionFraction, float NewResolutionFraction, bool bCanChangeResolution) const;
+};
+
+
+/* Represent an independent budget to dynamically scale by its own. */
+class RENDERCORE_API FBudget final
+{
+public:
+	FBudget(const TCHAR* Name, FHeuristicSettings (*HeuristicSettingsGetter)(void));
+	~FBudget();
+
+	/* Returns the global list of all dynamic resolution budgets. */
+	static TLinkedList<FBudget*>*& GetGlobalList();
+	static int32 GetGlobalListSize();
+
+	inline const TCHAR* GetName() const
+	{
+		return Name;
+	}
+
+	/* Returns the settings of the heuristic for this budget. */
+	inline const FHeuristicSettings& GetSettings() const
+	{
+		return CachedSettings;
+	}
+
+	inline int32 GetBudgetId() const
+	{
+		return BudgetId;
+	}
+
+	inline bool operator == (const FBudget& Other) const
+	{
+		return BudgetId == Other.BudgetId;
+	}
+
+private:
+	const TCHAR* Name;
+	FHeuristicSettings(*HeuristicSettingsGetter)(void);
+	TLinkedList<FBudget*> GlobalListLink;
+	FHeuristicSettings CachedSettings;
+	int32 BudgetId = 0;
+
+	UE_NONCOPYABLE(FBudget);
+
+	template<typename Type>
+	friend class TMap;
+	friend RENDERCORE_API void UpdateHeuristicsSettings();
+};
+
+
+/* Map of FBudget -> <Type>. */
+template<typename Type>
+class TMap final
+{
+public:
+	static constexpr int32 kInlineAllocatedBudgets = 16;
+
+	TMap()
+	{
+		Array.SetNum(FBudget::GetGlobalListSize());
+		for (TLinkedList<FBudget*>::TIterator BudgetIt(FBudget::GetGlobalList()); BudgetIt; BudgetIt.Next())
+		{
+			Array[BudgetIt->BudgetId].Name = BudgetIt->GetName();
+		}
+	}
+
+	TMap(const Type& Value)
+		: TMap()
+	{
+		SetAll(Value);
+	}
+
+	TMap(const TMap<Type>& Map) = default;
+	~TMap() = default;
+
+	inline void SetAll(const Type& Value)
+	{
+		Array.SetNum(FBudget::GetGlobalListSize());
+		for (int32 i = 0; i < Array.Num(); i++)
+		{
+			Array[i].Value = Value;
+		}
+	}
+
+	inline const Type& operator [](const FBudget& Budget) const
+	{
+		check(Array.Num() == FBudget::GetGlobalListSize());
+		return Array[Budget.BudgetId].Value;
+	}
+
+	inline Type& operator [](const FBudget& Budget)
+	{
+		check(Array.Num() == FBudget::GetGlobalListSize());
+		return Array[Budget.BudgetId].Value;
+	}
+
+	inline const Type& operator [](const int32 BudgetId) const
+	{
+		check(Array.Num() == FBudget::GetGlobalListSize());
+		return Array[BudgetId].Value;
+	}
+
+	inline Type& operator [](const int32 BudgetId)
+	{
+		check(Array.Num() == FBudget::GetGlobalListSize());
+		return Array[BudgetId].Value;
+	}
+
+private:
+	struct FItem
+	{
+		const TCHAR* Name = nullptr;
+		Type Value;
+	};
+	TArray<FItem, TInlineAllocator<kInlineAllocatedBudgets>> Array;
+};
+
+
+constexpr float FractionToPercentage(float Fraction)
+{
+	return Fraction * 100.0f;
+}
+
+constexpr float PercentageToFraction(float Percentage)
+{
+	return Percentage / 100.0f;
+}
+
+inline float GetPercentageCVarToFraction(const TAutoConsoleVariable<float>& Percentage)
+{
+	return PercentageToFraction(Percentage.GetValueOnAnyThread());
+}
+
+/** Updates all FBudget::GetSettings() with their new FHeuristicSettings. */
+RENDERCORE_API void UpdateHeuristicsSettings();
+
+/** Begins recording GPU timings in RDG for all the different FBudgets. */
+RENDERCORE_API void BeginFrame(const TMap<bool>& bIsBudgetEnabled);
+
+/** Ends recording GPU timings in RDG for all the different FBudgets. */
+RENDERCORE_API void EndFrame();
+
+/** Returns the latest available timings. */
+RENDERCORE_API const TMap<uint64>& GetLastestTimings();
+
+
+} // namespace DynamicRenderScaling
diff --git a/Engine/Source/Runtime/RenderCore/Public/RenderGraphBuilder.h b/Engine/Source/Runtime/RenderCore/Public/RenderGraphBuilder.h
--- a/Engine/Source/Runtime/RenderCore/Public/RenderGraphBuilder.h
+++ b/Engine/Source/Runtime/RenderCore/Public/RenderGraphBuilder.h
@@ -665,7 +665,7 @@
 	EAsyncComputeBudget AsyncComputeBudgetState = EAsyncComputeBudget(~0u);
 
 	IF_RDG_CPU_SCOPES(FRDGCPUScopeStacks CPUScopeStacks);
-	IF_RDG_GPU_SCOPES(FRDGGPUScopeStacksByPipeline GPUScopeStacks);
+	FRDGGPUScopeStacksByPipeline GPUScopeStacks;
 	IF_RHI_WANT_BREADCRUMB_EVENTS(FRDGBreadcrumbState* BreadcrumbState{});
 
 	IF_RDG_ENABLE_TRACE(FRDGTrace Trace);
@@ -803,6 +803,7 @@
 #endif
 
 	friend FRDGTrace;
+	friend DynamicRenderScaling::FRDGScope;
 	friend FRDGEventScopeGuard;
 	friend FRDGGPUStatScopeGuard;
 	friend FRDGAsyncComputeBudgetScopeGuard;
diff --git a/Engine/Source/Runtime/RenderCore/Public/RenderGraphBuilder.inl b/Engine/Source/Runtime/RenderCore/Public/RenderGraphBuilder.inl
--- a/Engine/Source/Runtime/RenderCore/Public/RenderGraphBuilder.inl
+++ b/Engine/Source/Runtime/RenderCore/Public/RenderGraphBuilder.inl
@@ -398,14 +398,14 @@
 
 inline void FRDGBuilder::BeginEventScope(FRDGEventName&& ScopeName)
 {
-#if RDG_GPU_SCOPES
+#if RDG_GPU_DEBUG_SCOPES
 	GPUScopeStacks.BeginEventScope(MoveTemp(ScopeName), RHICmdList.GetGPUMask());
 #endif
 }
 
 inline void FRDGBuilder::EndEventScope()
 {
-#if RDG_GPU_SCOPES
+#if RDG_GPU_DEBUG_SCOPES
 	GPUScopeStacks.EndEventScope();
 #endif
 }
diff --git a/Engine/Source/Runtime/RenderCore/Public/RenderGraphDefinitions.h b/Engine/Source/Runtime/RenderCore/Public/RenderGraphDefinitions.h
--- a/Engine/Source/Runtime/RenderCore/Public/RenderGraphDefinitions.h
+++ b/Engine/Source/Runtime/RenderCore/Public/RenderGraphDefinitions.h
@@ -53,12 +53,12 @@
 	#define RDG_EVENTS RDG_EVENTS_NONE
 #endif
 
-#define RDG_GPU_SCOPES (RDG_EVENTS || HAS_GPU_STATS)
+#define RDG_GPU_DEBUG_SCOPES (RDG_EVENTS || HAS_GPU_STATS)
 
-#if RDG_GPU_SCOPES
-	#define IF_RDG_GPU_SCOPES(Op) Op
+#if RDG_GPU_DEBUG_SCOPES
+	#define IF_RDG_GPU_DEBUG_SCOPES(Op) Op
 #else
-	#define IF_RDG_GPU_SCOPES(Op)
+	#define IF_RDG_GPU_DEBUG_SCOPES(Op)
 #endif
 
 #define RDG_CPU_SCOPES (CSV_PROFILER)
diff --git a/Engine/Source/Runtime/RenderCore/Public/RenderGraphEvent.h b/Engine/Source/Runtime/RenderCore/Public/RenderGraphEvent.h
--- a/Engine/Source/Runtime/RenderCore/Public/RenderGraphEvent.h
+++ b/Engine/Source/Runtime/RenderCore/Public/RenderGraphEvent.h
@@ -5,6 +5,7 @@
 #include "RenderGraphDefinitions.h"
 #include "RendererInterface.h"
 #include "ProfilingDebugging/RealtimeGPUProfiler.h"
+#include "DynamicRenderScaling.h"
 
 /** Macros for create render graph event names and scopes.
  *
@@ -262,6 +263,99 @@
 
 //////////////////////////////////////////////////////////////////////////
 //
+// GPU Timing
+//
+//////////////////////////////////////////////////////////////////////////
+
+class FRDGTimingScope final
+{
+public:
+	FRDGTimingScope(const FRDGTimingScope* InParentScope, const int32 InBudgetId)
+		: ParentScope(InParentScope)
+		, BudgetId(InBudgetId)
+	{}
+
+	const int32 GetBudgetId() const
+	{
+		return BudgetId;
+	}
+
+	const FRDGTimingScope* const ParentScope;
+	const int32 BudgetId;
+};
+
+using FRDGTimingScopeOp = TRDGScopeOp<FRDGTimingScope>;
+
+class RENDERCORE_API FRDGTimingScopeOpArray final
+{
+public:
+	FRDGTimingScopeOpArray() = default;
+	FRDGTimingScopeOpArray(ERHIPipeline Pipeline, const TRDGScopeOpArray<FRDGTimingScopeOp>& Ops);
+
+	void Execute(FRHIComputeCommandList& RHICmdList);
+
+private:
+	FRHIRenderQuery* TimestampQuery = nullptr;
+};
+
+class RENDERCORE_API FRDGTimingScopeStack final
+{
+public:
+	FRDGTimingScopeStack(FRDGAllocator& Allocator)
+		: ScopeStack(Allocator)
+	{}
+
+	inline void BeginScope(const DynamicRenderScaling::FBudget& Budget)
+	{
+		int32 BudgetId = Budget.GetBudgetId();
+		ScopeStack.BeginScope(BudgetId);
+	}
+
+	inline void EndScope()
+	{
+		ScopeStack.EndScope();
+	}
+
+	inline void ReserveOps()
+	{
+		ScopeStack.ReserveOps();
+	}
+
+	FRDGTimingScopeOpArray CompilePassPrologue(const FRDGPass* Pass);
+	
+	inline void EndExecute(FRHIComputeCommandList& RHICmdList)
+	{
+		return FRDGTimingScopeOpArray(RHICmdList.GetPipeline(), ScopeStack.EndCompile()).Execute(RHICmdList);
+	}
+
+	inline const FRDGTimingScope* GetCurrentScope() const
+	{
+		return ScopeStack.GetCurrentScope();
+	}
+
+	TRDGScopeStack<FRDGTimingScopeOp> ScopeStack;
+};
+
+namespace DynamicRenderScaling
+{
+
+class RENDERCORE_API FRDGScope final
+{
+public:
+	FRDGScope(FRDGBuilder& InGraphBuilder, const FBudget& InBudget);
+	FRDGScope(const FRDGScope&) = delete;
+	~FRDGScope();
+
+private:
+	FRDGBuilder& GraphBuilder;
+	const bool bIsEnabled;
+};
+
+} // namespace DynamicRenderScaling
+
+
+//////////////////////////////////////////////////////////////////////////
+//
 // GPU Events - Named hierarchical events emitted to external profiling tools.
 //
 //////////////////////////////////////////////////////////////////////////
@@ -294,7 +388,7 @@
 #endif
 };
 
-#if RDG_GPU_SCOPES
+#if RDG_GPU_DEBUG_SCOPES
 
 class FRDGEventScope final
 {
@@ -441,12 +535,16 @@
 	bool bCondition = true;
 };
 
+#endif // RDG_GPU_DEBUG_SCOPES
+
 //////////////////////////////////////////////////////////////////////////
 //
 // GPU Stats - Aggregated counters emitted to the runtime 'stat GPU' profiler.
 //
 //////////////////////////////////////////////////////////////////////////
 
+#if RDG_GPU_DEBUG_SCOPES
+
 class FRDGGPUStatScope final
 {
 public:
@@ -582,70 +680,89 @@
 	FRDGBuilder& GraphBuilder;
 };
 
+#endif // RDG_GPU_DEBUG_SCOPES
+
+//////////////////////////////////////////////////////////////////////////
+//
+// General GPU scopes
+//
+//////////////////////////////////////////////////////////////////////////
+
 struct FRDGGPUScopes
 {
-	const FRDGEventScope* Event = nullptr;
-	const FRDGGPUStatScope* Stat = nullptr;
+	const FRDGTimingScope* Timing = nullptr;
+	IF_RDG_GPU_DEBUG_SCOPES(const FRDGEventScope* Event = nullptr);
+	IF_RDG_GPU_DEBUG_SCOPES(const FRDGGPUStatScope* Stat = nullptr);
 };
 
 struct FRDGGPUScopeOpArrays
 {
 	inline void Execute(FRHIComputeCommandList& RHICmdList)
 	{
-		Event.Execute(RHICmdList);
-		Stat.Execute(RHICmdList);
+		Timing.Execute(RHICmdList);
+		IF_RDG_GPU_DEBUG_SCOPES(Event.Execute(RHICmdList));
+		IF_RDG_GPU_DEBUG_SCOPES(Stat.Execute(RHICmdList));
 	}
 
-	FRDGEventScopeOpArray Event;
-	FRDGGPUStatScopeOpArray Stat;
+	FRDGTimingScopeOpArray Timing;
+	IF_RDG_GPU_DEBUG_SCOPES(FRDGEventScopeOpArray Event);
+	IF_RDG_GPU_DEBUG_SCOPES(FRDGGPUStatScopeOpArray Stat);
 };
 
 /** The complete set of scope stack implementations. */
 struct FRDGGPUScopeStacks
 {
 	FRDGGPUScopeStacks(FRDGAllocator& Allocator)
-		: Event(Allocator)
+		: Timing(Allocator)
+#if RDG_GPU_DEBUG_SCOPES
+		, Event(Allocator)
 		, Stat(Allocator)
+#endif
 	{}
 
 	inline void ReserveOps(int32 PassCount)
 	{
-		Event.ReserveOps(PassCount);
-		Stat.ReserveOps();
+		Timing.ReserveOps();
+		IF_RDG_GPU_DEBUG_SCOPES(Event.ReserveOps(PassCount));
+		IF_RDG_GPU_DEBUG_SCOPES(Stat.ReserveOps());
 	}
 
 	inline FRDGGPUScopeOpArrays CompilePassPrologue(const FRDGPass* Pass, FRHIGPUMask GPUMask)
 	{
 		FRDGGPUScopeOpArrays Result;
-		Result.Event = Event.CompilePassPrologue(Pass);
-		Result.Stat = Stat.CompilePassPrologue(Pass, GPUMask);
+		Result.Timing = Timing.CompilePassPrologue(Pass);
+		IF_RDG_GPU_DEBUG_SCOPES(Result.Event = Event.CompilePassPrologue(Pass));
+		IF_RDG_GPU_DEBUG_SCOPES(Result.Stat = Stat.CompilePassPrologue(Pass, GPUMask));
 		return MoveTemp(Result);
 	}
 
 	inline FRDGGPUScopeOpArrays CompilePassEpilogue()
 	{
 		FRDGGPUScopeOpArrays Result;
-		Result.Event = Event.CompilePassEpilogue();
-		Result.Stat = Stat.CompilePassEpilogue();
+		IF_RDG_GPU_DEBUG_SCOPES(Result.Event = Event.CompilePassEpilogue());
+		IF_RDG_GPU_DEBUG_SCOPES(Result.Stat = Stat.CompilePassEpilogue());
 		return MoveTemp(Result);
 	}
 
 	inline void EndExecute(FRHIComputeCommandList& RHICmdList)
 	{
-		Event.EndExecute(RHICmdList);
-		Stat.EndExecute(RHICmdList);
+		Timing.EndExecute(RHICmdList);
+		IF_RDG_GPU_DEBUG_SCOPES(Event.EndExecute(RHICmdList));
+		IF_RDG_GPU_DEBUG_SCOPES(Stat.EndExecute(RHICmdList));
 	}
 
 	inline FRDGGPUScopes GetCurrentScopes() const
 	{
 		FRDGGPUScopes Scopes;
-		Scopes.Event = Event.GetCurrentScope();
-		Scopes.Stat = Stat.GetCurrentScope();
+		Scopes.Timing = Timing.GetCurrentScope();
+		IF_RDG_GPU_DEBUG_SCOPES(Scopes.Event = Event.GetCurrentScope());
+		IF_RDG_GPU_DEBUG_SCOPES(Scopes.Stat = Stat.GetCurrentScope());
 		return Scopes;
 	}
 
-	FRDGEventScopeStack Event;
-	FRDGGPUStatScopeStack Stat;
+	FRDGTimingScopeStack Timing;
+	IF_RDG_GPU_DEBUG_SCOPES(FRDGEventScopeStack Event);
+	IF_RDG_GPU_DEBUG_SCOPES(FRDGGPUStatScopeStack Stat);
 };
 
 struct RENDERCORE_API FRDGGPUScopeStacksByPipeline
@@ -655,6 +772,19 @@
 		, AsyncCompute(Allocator)
 	{}
 
+	inline void BeginTimingScope(const DynamicRenderScaling::FBudget& Budget)
+	{
+		Graphics.Timing.BeginScope(Budget);
+		AsyncCompute.Timing.BeginScope(Budget);
+	}
+
+	inline void EndTimingScope()
+	{
+		Graphics.Timing.EndScope();
+		AsyncCompute.Timing.EndScope();
+	}
+
+#if RDG_GPU_DEBUG_SCOPES
 	inline void BeginEventScope(FRDGEventName&& ScopeName, FRHIGPUMask GPUMask)
 	{
 		FRDGEventName ScopeNameCopy = ScopeName;
@@ -677,6 +807,7 @@
 	{
 		Graphics.Stat.EndScope();
 	}
+#endif
 
 	inline void ReserveOps(int32 PassCount)
 	{
@@ -701,8 +832,6 @@
 	FRDGGPUScopeStacks AsyncCompute;
 };
 
-#endif
-
 //////////////////////////////////////////////////////////////////////////
 //
 // CPU CSV Stats
diff --git a/Engine/Source/Runtime/RenderCore/Public/RenderGraphEvent.inl b/Engine/Source/Runtime/RenderCore/Public/RenderGraphEvent.inl
--- a/Engine/Source/Runtime/RenderCore/Public/RenderGraphEvent.inl
+++ b/Engine/Source/Runtime/RenderCore/Public/RenderGraphEvent.inl
@@ -156,8 +156,6 @@
 #endif
 }
 
-#if RDG_GPU_SCOPES
-
 inline const FRDGGPUScopeStacks& FRDGGPUScopeStacksByPipeline::GetScopeStacks(ERHIPipeline Pipeline) const
 {
 	switch (Pipeline)
@@ -185,5 +183,3 @@
 		return Graphics;
 	}
 }
-
-#endif
diff --git a/Engine/Source/Runtime/RenderCore/Public/RenderGraphPass.h b/Engine/Source/Runtime/RenderCore/Public/RenderGraphPass.h
--- a/Engine/Source/Runtime/RenderCore/Public/RenderGraphPass.h
+++ b/Engine/Source/Runtime/RenderCore/Public/RenderGraphPass.h
@@ -294,12 +294,10 @@
 	}
 #endif
 
-#if RDG_GPU_SCOPES
 	FRDGGPUScopes GetGPUScopes() const
 	{
 		return GPUScopes;
 	}
-#endif
 
 #if WITH_MGPU
 	FRHIGPUMask GetGPUMask() const
@@ -491,13 +489,11 @@
 	FRDGCPUScopeOpArrays CPUScopeOps;
 #endif
 
-#if RDG_GPU_SCOPES
 	FRDGGPUScopes GPUScopes;
 	FRDGGPUScopeOpArrays GPUScopeOpsPrologue;
 	FRDGGPUScopeOpArrays GPUScopeOpsEpilogue;
-#endif
 
-#if RDG_GPU_SCOPES && RDG_ENABLE_TRACE
+#if RDG_GPU_DEBUG_SCOPES && RDG_ENABLE_TRACE
 	const FRDGEventScope* TraceEventScope = nullptr;
 #endif
 
diff --git a/Engine/Source/Runtime/Renderer/Private/SceneRendering.cpp b/Engine/Source/Runtime/Renderer/Private/SceneRendering.cpp
--- a/Engine/Source/Runtime/Renderer/Private/SceneRendering.cpp
+++ b/Engine/Source/Runtime/Renderer/Private/SceneRendering.cpp
@@ -791,7 +791,7 @@ FRHICommandList* FParallelCommandListSet::NewParallelCommandList()
 	FRHICommandList* Result = AllocCommandList();
 	Result->ExecuteStat = ExecuteStat;
 
-#if RDG_GPU_SCOPES
+#if RDG_GPU_DEBUG_SCOPES
 	if (Pass->GetGPUScopes().Stat != nullptr && (**Pass->GetGPUScopes().Stat->DrawCallCounter) != -1)
 	{
 		Result->EnqueueLambda(
