Fixed 16 bit float GBuffer format to use PF_FloatRGBA.

diff --git a/Engine/Source/Runtime/RenderCore/Private/GBufferInfo.cpp b/Engine/Source/Runtime/RenderCore/Private/GBufferInfo.cpp
--- a/Engine/Source/Runtime/RenderCore/Private/GBufferInfo.cpp
+++ b/Engine/Source/Runtime/RenderCore/Private/GBufferInfo.cpp
@@ -161,7 +161,7 @@ FGBufferBinding FindGBufferBindingByName(const FGBufferInfo& GBufferInfo, const
 			PixelFormat = PF_G16R16;
 			break;
 		case GBT_Float_16_16_16_16:
-			PixelFormat = PF_A16B16G16R16;
+			PixelFormat = PF_FloatRGBA;
 			break;
 		case GBT_Invalid:
 		default:
