Removed GetUAV from IPooledRenderTarget.

diff --git a/Engine/Source/Runtime/RenderCore/Public/RendererInterface.h b/Engine/Source/Runtime/RenderCore/Public/RendererInterface.h
--- a/Engine/Source/Runtime/RenderCore/Public/RendererInterface.h
+++ b/Engine/Source/Runtime/RenderCore/Public/RendererInterface.h
@@ -570,8 +570,6 @@ struct IPooledRenderTarget
 
 	FRHITexture* GetRHI() const { return RenderTargetItem.TargetableTexture; }
 
-	FRHIUnorderedAccessView* GetUAV() const { return RenderTargetItem.UAV; }
-
 	UE_DEPRECATED(5.0, "Use GetRHI instead.")
 	FRHITexture* GetTargetableRHI() const { return GetRHI(); }
 
