Fixes linking failure on shipping in render graph

diff --git a/Engine/Source/Runtime/RenderCore/Private/RenderGraphEvent.cpp b/Engine/Source/Runtime/RenderCore/Private/RenderGraphEvent.cpp
--- a/Engine/Source/Runtime/RenderCore/Private/RenderGraphEvent.cpp
+++ b/Engine/Source/Runtime/RenderCore/Private/RenderGraphEvent.cpp
@@ -793,6 +793,8 @@
 	return {};
 }
 
+#endif // RDG_GPU_DEBUG_SCOPES
+
 FRDGGPUScopeOpArrays FRDGGPUScopeStacksByPipeline::CompilePassPrologue(const FRDGPass* Pass, FRHIGPUMask GPUMask)
 {
 	return GetScopeStacks(Pass->GetPipeline()).CompilePassPrologue(Pass, GPUMask);
@@ -803,8 +805,6 @@
 	return GetScopeStacks(Pass->GetPipeline()).CompilePassEpilogue();
 }
 
-#endif
-
 //////////////////////////////////////////////////////////////////////////
 // CPU Scopes
 //////////////////////////////////////////////////////////////////////////
