Fixes VisualizeTexture not visualizing Texture2DArray's slices

diff --git a/Engine/Source/Runtime/RenderCore/Private/VisualizeTexture.cpp b/Engine/Source/Runtime/RenderCore/Private/VisualizeTexture.cpp
--- a/Engine/Source/Runtime/RenderCore/Private/VisualizeTexture.cpp
+++ b/Engine/Source/Runtime/RenderCore/Private/VisualizeTexture.cpp
@@ -554,13 +554,13 @@ class FVisualizeTexturePS : public FGlobalShader
 		SHADER_PARAMETER(FVector3f, TextureExtent)
 		SHADER_PARAMETER_ARRAY(FVector4f, VisualizeParam, [3])
 
-		SHADER_PARAMETER_RDG_TEXTURE(Texture2D, VisualizeTexture2D)
+		SHADER_PARAMETER_RDG_TEXTURE_SRV(Texture2D, VisualizeTexture2D)
 		SHADER_PARAMETER_SAMPLER(SamplerState, VisualizeTexture2DSampler)
-		SHADER_PARAMETER_RDG_TEXTURE(Texture3D, VisualizeTexture3D)
+		SHADER_PARAMETER_RDG_TEXTURE_SRV(Texture3D, VisualizeTexture3D)
 		SHADER_PARAMETER_SAMPLER(SamplerState, VisualizeTexture3DSampler)
-		SHADER_PARAMETER_RDG_TEXTURE(TextureCube, VisualizeTextureCube)
+		SHADER_PARAMETER_RDG_TEXTURE_SRV(TextureCube, VisualizeTextureCube)
 		SHADER_PARAMETER_SAMPLER(SamplerState, VisualizeTextureCubeSampler)
-		SHADER_PARAMETER_RDG_TEXTURE(TextureCubeArray, VisualizeTextureCubeArray)
+		SHADER_PARAMETER_RDG_TEXTURE_SRV(TextureCubeArray, VisualizeTextureCubeArray)
 		SHADER_PARAMETER_SAMPLER(SamplerState, VisualizeTextureCubeArraySampler)
 		SHADER_PARAMETER_RDG_TEXTURE_SRV(Texture2D<uint4>, VisualizeDepthStencil)
 		SHADER_PARAMETER_RDG_TEXTURE(Texture2DMS<float4>, VisualizeTexture2DMS)
@@ -688,15 +688,24 @@ void FVisualizeTexture::CreateContentCapturePass(FRDGBuilder& GraphBuilder, cons
 				PassParameters->VisualizeParam[2] = FVector4f((float)InputValueMapping, 0.0f, Config.SingleChannel);
 			}
 
+			FRDGTextureSRV* InputSRV = nullptr;
 			FRHISamplerState* PointSampler = TStaticSamplerState<SF_Point, AM_Clamp, AM_Clamp, AM_Clamp>::GetRHI();
+			if (InputTexture->Desc.Dimension == ETextureDimension::Texture2DArray)
+			{
+				InputSRV = GraphBuilder.CreateSRV(FRDGTextureSRVDesc::CreateForSlice(InputTexture, FMath::Clamp(int32(Config.ArrayIndex), 0, int32(InputDesc.ArraySize) - 1)));
+			}
+			else
+			{
+				InputSRV = GraphBuilder.CreateSRV(FRDGTextureSRVDesc::Create(InputTexture));
+			}
 
-			PassParameters->VisualizeTexture2D = InputTexture;
+			PassParameters->VisualizeTexture2D = InputSRV;
 			PassParameters->VisualizeTexture2DSampler = PointSampler;
-			PassParameters->VisualizeTexture3D = InputTexture;
+			PassParameters->VisualizeTexture3D = InputSRV;
 			PassParameters->VisualizeTexture3DSampler = PointSampler;
-			PassParameters->VisualizeTextureCube = InputTexture;
+			PassParameters->VisualizeTextureCube = InputSRV;
 			PassParameters->VisualizeTextureCubeSampler = PointSampler;
-			PassParameters->VisualizeTextureCubeArray = InputTexture;
+			PassParameters->VisualizeTextureCubeArray = InputSRV;
 			PassParameters->VisualizeTextureCubeArraySampler = PointSampler;
 
 			if (VisualizeType == EVisualisePSType::Texture2DDepthStencilNoMSAA)
@@ -725,6 +734,11 @@ void FVisualizeTexture::CreateContentCapturePass(FRDGBuilder& GraphBuilder, cons
 				ExtendedDrawEvent += FString::Printf(TEXT("x%d CapturedSlice=%d"), InputDesc.Depth, Config.ArrayIndex);
 			}
 
+			if (InputDesc.IsTextureArray())
+			{
+				ExtendedDrawEvent += FString::Printf(TEXT(" ArraySize=%d CapturedSlice=%d"), InputDesc.ArraySize, Config.ArrayIndex);
+			}
+
 			// Precise the mip level being captured in the mip level when there is a mip chain.
 			if (InputDesc.IsMipChain())
 			{
