#ue5 - TruncToInt(double) and similar *ToInt(double) functions now return int64 rather than int32. Fixed up existing uses of *ToInt(double) functions to use explicit *ToInt32(double) where old result type was expected to be int32.

Enabled overloads of Max() and Min() with mixed signed integral arguments (promoting to largest type). Unsigned are not allowed to implicitly mix.
Added explicit overload of Clamp(int64, int32, int32) => int64 to allow for easier upgrades, and once type truncation warnings are enabled in more modules they can be evaluated per-case.

diff --git a/Engine/Plugins/Runtime/MassGameplay/Source/MassCommon/Public/MassCommonTypes.h b/Engine/Plugins/Runtime/MassGameplay/Source/MassCommon/Public/MassCommonTypes.h
--- a/Engine/Plugins/Runtime/MassGameplay/Source/MassCommon/Public/MassCommonTypes.h
+++ b/Engine/Plugins/Runtime/MassGameplay/Source/MassCommon/Public/MassCommonTypes.h
@@ -64,7 +64,7 @@ struct MASSCOMMON_API FMassInt16Real
 	
 	void Set(const float InValue)
 	{
-		Value = (int16)FMath::Clamp(FMath::RoundToInt(InValue), -(int32)MAX_int16, (int32)MAX_int16);
+		Value = (int16)FMath::Clamp(FMath::RoundToInt32(InValue), -(int32)MAX_int16, (int32)MAX_int16);
 	}
 
 	float Get() const
@@ -101,7 +101,7 @@ struct MASSCOMMON_API FMassInt16Real10
 	
 	void Set(const float InValue)
 	{
-		Value = (int16)FMath::Clamp(FMath::RoundToInt(InValue * Scale), -(int32)MAX_int16, (int32)MAX_int16);
+		Value = (int16)FMath::Clamp(FMath::RoundToInt32(InValue * Scale), -(int32)MAX_int16, (int32)MAX_int16);
 	}
 
 	float Get() const
@@ -138,9 +138,9 @@ struct MASSCOMMON_API FMassSnorm8Vector
 	
 	void Set(const FVector& InVector)
 	{
-		X = (int8)FMath::Clamp(FMath::RoundToInt(InVector.X * Scale), -(int32)MAX_int8, (int32)MAX_int8);
-		Y = (int8)FMath::Clamp(FMath::RoundToInt(InVector.Y * Scale), -(int32)MAX_int8, (int32)MAX_int8);
-		Z = (int8)FMath::Clamp(FMath::RoundToInt(InVector.Z * Scale), -(int32)MAX_int8, (int32)MAX_int8);
+		X = (int8)FMath::Clamp(FMath::RoundToInt32(InVector.X * Scale), -(int32)MAX_int8, (int32)MAX_int8);
+		Y = (int8)FMath::Clamp(FMath::RoundToInt32(InVector.Y * Scale), -(int32)MAX_int8, (int32)MAX_int8);
+		Z = (int8)FMath::Clamp(FMath::RoundToInt32(InVector.Z * Scale), -(int32)MAX_int8, (int32)MAX_int8);
 	}
 
 	FVector Get() const
@@ -181,8 +181,8 @@ struct MASSCOMMON_API FMassSnorm8Vector2D
 
 	void Set(const FVector2D& InVector)
 	{
-		X = (int8)FMath::Clamp(FMath::RoundToInt(InVector.X * Scale), -(int32)MAX_int8, (int32)MAX_int8);
-		Y = (int8)FMath::Clamp(FMath::RoundToInt(InVector.Y * Scale), -(int32)MAX_int8, (int32)MAX_int8);
+		X = (int8)FMath::Clamp(FMath::RoundToInt32(InVector.X * Scale), -(int32)MAX_int8, (int32)MAX_int8);
+		Y = (int8)FMath::Clamp(FMath::RoundToInt32(InVector.Y * Scale), -(int32)MAX_int8, (int32)MAX_int8);
 	}
 
 	FVector2D Get() const
@@ -220,7 +220,7 @@ struct MASSCOMMON_API FMassUnorm8Real
 
 	void Set(const float InValue)
 	{
-		Value = (int8)FMath::Clamp(FMath::RoundToInt(InValue * Scale), 0, (int32)MAX_uint8);
+		Value = (int8)FMath::Clamp(FMath::RoundToInt32(InValue * Scale), 0, (int32)MAX_uint8);
 	}
 
 	float Get() const
@@ -247,9 +247,9 @@ struct MASSCOMMON_API FMassInt16Vector
 	
 	void Set(const FVector& InVector)
 	{
-		X = (int16)FMath::Clamp(FMath::RoundToInt(InVector.X), -(int32)MAX_int16, (int32)MAX_int16);
-		Y = (int16)FMath::Clamp(FMath::RoundToInt(InVector.Y), -(int32)MAX_int16, (int32)MAX_int16);
-		Z = (int16)FMath::Clamp(FMath::RoundToInt(InVector.Z), -(int32)MAX_int16, (int32)MAX_int16);
+		X = (int16)FMath::Clamp(FMath::RoundToInt32(InVector.X), -(int32)MAX_int16, (int32)MAX_int16);
+		Y = (int16)FMath::Clamp(FMath::RoundToInt32(InVector.Y), -(int32)MAX_int16, (int32)MAX_int16);
+		Z = (int16)FMath::Clamp(FMath::RoundToInt32(InVector.Z), -(int32)MAX_int16, (int32)MAX_int16);
 	}
 
 	FVector Get() const
@@ -286,8 +286,8 @@ struct MASSCOMMON_API FMassInt16Vector2D
 	
 	void Set(const FVector2D& InVector)
 	{
-		X = (int16)FMath::Clamp(FMath::RoundToInt(InVector.X), -(int32)MAX_int16, (int32)MAX_int16);
-		Y = (int16)FMath::Clamp(FMath::RoundToInt(InVector.Y), -(int32)MAX_int16, (int32)MAX_int16);
+		X = (int16)FMath::Clamp(FMath::RoundToInt32(InVector.X), -(int32)MAX_int16, (int32)MAX_int16);
+		Y = (int16)FMath::Clamp(FMath::RoundToInt32(InVector.Y), -(int32)MAX_int16, (int32)MAX_int16);
 	}
 
 	FVector2D Get() const
diff --git a/Engine/Source/Runtime/Core/Private/Tests/Math/UnrealMathTest.cpp b/Engine/Source/Runtime/Core/Private/Tests/Math/UnrealMathTest.cpp
--- a/Engine/Source/Runtime/Core/Private/Tests/Math/UnrealMathTest.cpp
+++ b/Engine/Source/Runtime/Core/Private/Tests/Math/UnrealMathTest.cpp
@@ -3077,8 +3077,8 @@ bool FVectorRegisterAbstractionTest::RunTest(const FString& Parameters)
 	}
 
 
-	// Float function compilation with various types. Set this define to 1 to verify that warnings are generated for all code within MATHTEST_CHECK_INVALID_FLOAT_VARIANTS blocks.
-#define MATHTEST_CHECK_INVALID_FLOAT_VARIANTS 0
+	// Float function compilation with various types. Set this define to 1 to verify that warnings are generated for all code within MATHTEST_CHECK_INVALID_OVERLOAD_VARIANTS blocks.
+#define MATHTEST_CHECK_INVALID_OVERLOAD_VARIANTS 0
 	{
 		float F = 1.f, TestFloat;
 		double D = 1.0, TestDouble;
@@ -3090,7 +3090,7 @@ bool FVectorRegisterAbstractionTest::RunTest(const FString& Parameters)
 		TestFloat = FMath::Fmod(F, I);
 		TestFloat = FMath::Fmod(I, F);
 
-#if MATHTEST_CHECK_INVALID_FLOAT_VARIANTS
+#if MATHTEST_CHECK_INVALID_OVERLOAD_VARIANTS
 		// Expected to generate warnings
 		TestFloat = FMath::Fmod(F, D);
 		TestFloat = FMath::Fmod(D, F);
@@ -3110,12 +3110,12 @@ bool FVectorRegisterAbstractionTest::RunTest(const FString& Parameters)
 		TestDouble = FMath::Fmod(I, F);
 		TestDouble = FMath::Fmod(I, D);
 
-#if MATHTEST_CHECK_INVALID_FLOAT_VARIANTS
+#if MATHTEST_CHECK_INVALID_OVERLOAD_VARIANTS
 		// Expected to generate warnings
 		TestDouble = FMath::Fmod(I, I); // Should be warned to be deprecated
 #endif
 
-#if MATHTEST_CHECK_INVALID_FLOAT_VARIANTS
+#if MATHTEST_CHECK_INVALID_OVERLOAD_VARIANTS
 		// Expected to generate warnings
 		int TestInt;
 		TestInt = FMath::Fmod(F, F);
@@ -3138,7 +3138,7 @@ bool FVectorRegisterAbstractionTest::RunTest(const FString& Parameters)
 
 		TestDouble = FMath::TruncToDouble(I); // not currently a warning
 
-#if MATHTEST_CHECK_INVALID_FLOAT_VARIANTS
+#if MATHTEST_CHECK_INVALID_OVERLOAD_VARIANTS
 		// Expected to generate errors
 		TestFloat = FMath::TruncToFloat(I);
 		TestFloat = FMath::TruncToFloat(D);
@@ -3185,6 +3185,9 @@ bool FVectorRegisterAbstractionTest::RunTest(const FString& Parameters)
 		F = FMath::Clamp(F1, I1, F2);
 		D = FMath::Clamp(F1, D2, I3);
 
+		D = FMath::Max(F1, D1);
+		D = FMath::Min(D1, F1);
+
 		CheckPassing(FMath::Clamp( 1, 0, 2) == 1);
 		CheckPassing(FMath::Clamp(-1, 0, 2) == 0);
 		CheckPassing(FMath::Clamp( 3, 0, 2) == 2);
@@ -3192,13 +3195,15 @@ bool FVectorRegisterAbstractionTest::RunTest(const FString& Parameters)
 		U1 = 0; U2 = 2;
 		CheckPassing(FMath::Clamp<int32>(-1, U1, U2) == 0);
 
-#if MATHTEST_CHECK_INVALID_FLOAT_VARIANTS
+#if MATHTEST_CHECK_INVALID_OVERLOAD_VARIANTS
 		// Expected to generate errors (ambiguous arguments)
 		U = FMath::Clamp(U1, I1, U3);
 		// Expected to generate errors (double->float truncation)
 		F = FMath::Clamp(F1, F2, D3);
 		F = FMath::Clamp(F1, D2, F3);
 		F = FMath::Clamp(D1, D2, F3);
+		F = FMath::Max(F1, D1);
+		F = FMath::Min(D1, F1);
 		// Expected to generate errors (float->int truncation)
 		I = FMath::Clamp(I1, F2, I3);
 #endif
@@ -3212,9 +3217,33 @@ bool FVectorRegisterAbstractionTest::RunTest(const FString& Parameters)
 		U1 = uint32(-1);
 		CheckPassing(FMath::Min<int32>(0, U1) == -1);
 
+		// Mix float/int types
+		F = FMath::Max(Big1, F1);
+		F = FMath::Max(F1, I1);
+		D = FMath::Min(Big1, D1);
+		D = FMath::Min(D1, I1);
+
 		// Test compilation mixing int32/int64 types
-		//Big3 = FMath::Max(Big1, I2);
-		//Big3 = FMath::Max(I1, Big2);
+		static_assert(TIsSame< decltype(FMath::Max(Big1, I1)), decltype(Big1) >::Value);
+		static_assert(TIsSame< decltype(FMath::Min(I1, Big1)), decltype(Big1) >::Value);
+		Big3 = FMath::Max(Big1, I1);
+		Big3 = FMath::Min(Big1, I1);
+
+		static_assert(TIsSame< decltype(FMath::Max(I1, Big2)), decltype(Big1) >::Value);
+		static_assert(TIsSame< decltype(FMath::Min(Big2, I1)), decltype(Big1) >::Value);
+		Big3 = FMath::Max(I1, Big2);
+		Big3 = FMath::Min(I1, Big2);
+
+		Big3 = FMath::Clamp(Big1, I1, I2);
+
+#if MATHTEST_CHECK_INVALID_OVERLOAD_VARIANTS
+		// Expected to generate errors (no overload for mixing signed/unsigned types).
+		Big3 = FMath::Max(U1, I1);
+		Big3 = FMath::Min(I1, U1);
+		// Expected to generate errors or warnings (truncation / loss of data)
+		Big3 = FMath::Clamp(I1, Big1, I2);
+		I = FMath::Clamp<int32>(Big1, I1, I2);
+#endif
 	}
 
 	// Sin, Cos, Tan tests
diff --git a/Engine/Source/Runtime/Core/Public/GenericPlatform/GenericPlatformMath.h b/Engine/Source/Runtime/Core/Public/GenericPlatform/GenericPlatformMath.h
--- a/Engine/Source/Runtime/Core/Public/GenericPlatform/GenericPlatformMath.h
+++ b/Engine/Source/Runtime/Core/Public/GenericPlatform/GenericPlatformMath.h
@@ -219,7 +219,7 @@ struct FGenericPlatformMath
 	}
 
 	static CONSTEXPR FORCEINLINE int32 TruncToInt(float F) { return TruncToInt32(F); }
-	static CONSTEXPR FORCEINLINE int32 TruncToInt(double F) { return TruncToInt32(F); }
+	static CONSTEXPR FORCEINLINE int64 TruncToInt(double F) { return TruncToInt64(F); }
 
 	/**
 	 * Converts a float to an integer value with truncation towards zero.
@@ -273,7 +273,7 @@ struct FGenericPlatformMath
 	}
 
 	static FORCEINLINE int32 FloorToInt(float F) { return FloorToInt32(F); }
-	static FORCEINLINE int32 FloorToInt(double F) { return FloorToInt32(F); }
+	static FORCEINLINE int64 FloorToInt(double F) { return FloorToInt64(F); }
 	
 	
 	/**
@@ -322,7 +322,7 @@ struct FGenericPlatformMath
 	}
 	
 	static FORCEINLINE int32 RoundToInt(float F) { return RoundToInt32(F); }
-	static FORCEINLINE int32 RoundToInt(double F) { return RoundToInt32(F); }
+	static FORCEINLINE int64 RoundToInt(double F) { return RoundToInt64(F); }
 
 	/**
 	* Converts a float to the nearest integer. Rounds up when the fraction is .5
@@ -376,7 +376,7 @@ struct FGenericPlatformMath
 	}
 
 	static FORCEINLINE int32 CeilToInt(float F) { return CeilToInt32(F); }
-	static FORCEINLINE int32 CeilToInt(double F) { return CeilToInt32(F); }
+	static FORCEINLINE int64 CeilToInt(double F) { return CeilToInt64(F); }
 
 	/**
 	* Converts a float to the nearest greater or equal integer.
@@ -528,6 +528,7 @@ struct FGenericPlatformMath
 
 	static FORCEINLINE float Sin( float Value ) { return sinf(Value); }
 	static FORCEINLINE double Sin( double Value ) { return sin(Value); }
+	RESOLVE_FLOAT_AMBIGUITY(Sin);
 
 	static FORCEINLINE float Asin( float Value ) { return asinf( (Value<-1.f) ? -1.f : ((Value<1.f) ? Value : 1.f) ); }
 	static FORCEINLINE double Asin( double Value ) { return asin( (Value<-1.0) ? -1.0 : ((Value<1.0) ? Value : 1.0) ); }
@@ -943,10 +944,12 @@ struct FGenericPlatformMath
 	}
 
 	// Allow mixing of float types to promote to highest precision type
-	static CONSTEXPR FORCEINLINE double Max(const double A, const float B) { return Max<double>(A, B); }
-	static CONSTEXPR FORCEINLINE double Max(const float A, const double B) { return Max<double>(A, B); }
-	static CONSTEXPR FORCEINLINE double Min(const double A, const float B) { return Min<double>(A, B); }
-	static CONSTEXPR FORCEINLINE double Min(const float A, const double B) { return Min<double>(A, B); }
+	MIX_FLOATS_2_ARGS(Max);
+	MIX_FLOATS_2_ARGS(Min);
+
+	// Allow mixing of signed integral types.
+	MIX_SIGNED_INTS_2_ARGS_CONSTEXPR(Max);
+	MIX_SIGNED_INTS_2_ARGS_CONSTEXPR(Min);
 
 	/**
 	* Min of Array
diff --git a/Engine/Source/Runtime/Core/Public/Math/UnrealMathUtility.h b/Engine/Source/Runtime/Core/Public/Math/UnrealMathUtility.h
--- a/Engine/Source/Runtime/Core/Public/Math/UnrealMathUtility.h
+++ b/Engine/Source/Runtime/Core/Public/Math/UnrealMathUtility.h
@@ -590,6 +590,12 @@ public:
 	}
 	/** Allow mixing float/double arguments, promoting to highest precision type. */
 	MIX_FLOATS_3_ARGS(Clamp);
+	
+	/** Clamps X to be between Min and Max, inclusive. Overload to support int64/int32 mixed types. */
+	UE_NODISCARD static constexpr FORCEINLINE int64 Clamp(const int64 X, const int32 Min, const int32 Max)
+	{
+		return (X < Min) ? Min : (X < Max) ? X : Max;
+	}
 
 	/** Wraps X to be between Min and Max, inclusive. */
 	/** When X can wrap to both Min and Max, it will wrap to Min if it lies below the range and wrap to Max if it is above the range. */
diff --git a/Engine/Source/Runtime/Core/Public/Math/UnrealPlatformMathSSE.h b/Engine/Source/Runtime/Core/Public/Math/UnrealPlatformMathSSE.h
--- a/Engine/Source/Runtime/Core/Public/Math/UnrealPlatformMathSSE.h
+++ b/Engine/Source/Runtime/Core/Public/Math/UnrealPlatformMathSSE.h
@@ -274,13 +274,13 @@ struct TUnrealPlatformMathSSEBase : public Base
 	//
 
 	static FORCEINLINE int32 TruncToInt(float F) { return TruncToInt32(F); }
-	static FORCEINLINE int32 TruncToInt(double F) { return TruncToInt32(F); }
+	static FORCEINLINE int64 TruncToInt(double F) { return TruncToInt64(F); }
 	static FORCEINLINE int32 FloorToInt(float F) { return FloorToInt32(F); }
-	static FORCEINLINE int32 FloorToInt(double F) { return FloorToInt32(F); }
+	static FORCEINLINE int64 FloorToInt(double F) { return FloorToInt64(F); }
 	static FORCEINLINE int32 RoundToInt(float F) { return RoundToInt32(F); }
-	static FORCEINLINE int32 RoundToInt(double F) { return RoundToInt32(F); }
+	static FORCEINLINE int64 RoundToInt(double F) { return RoundToInt64(F); }
 	static FORCEINLINE int32 CeilToInt(float F) { return CeilToInt32(F); }
-	static FORCEINLINE int32 CeilToInt(double F) { return CeilToInt32(F); }
+	static FORCEINLINE int64 CeilToInt(double F) { return CeilToInt64(F); }
 
 
 	template<typename T>
diff --git a/Engine/Source/Runtime/Core/Public/Templates/ResolveTypeAmbiguity.h b/Engine/Source/Runtime/Core/Public/Templates/ResolveTypeAmbiguity.h
--- a/Engine/Source/Runtime/Core/Public/Templates/ResolveTypeAmbiguity.h
+++ b/Engine/Source/Runtime/Core/Public/Templates/ResolveTypeAmbiguity.h
@@ -6,6 +6,7 @@
 #include "Templates/IsFloatingPoint.h"
 #include "Templates/UnrealTypeTraits.h"
 #include "Templates/IsSigned.h"
+#include "Templates/IsIntegral.h"
 
 /**
  * This file defines a few macros which can be used to restrict (via deprecation) ambiguous calls to functions expecting float params/return type.
@@ -41,19 +42,23 @@
 	static FORCEINLINE float Func(T&& Value) { return Func((float)Value); }
 
 
-// Mixing any types with any other type results in same type as "X * Y", which is promoted to the type of the result of mixed arithmetic between the types
-#define MIX_SIGNED_TYPES_2_ARGS_ACTUAL(Func, OptionalMarkup) \
-	template<typename Arg1, typename Arg2, TEMPLATE_REQUIRES(TAnd< TIsSigned<Arg1>, TIsSigned<Arg2>, TNot<TIsSame<Arg1, Arg2>> >::Value)> \
+// Mixing any signed integral types with any other signed integral type results in same type as "X * Y", which is promoted to the type of the result of mixed arithmetic between the types
+#define MIX_SIGNED_INTS_2_ARGS_ACTUAL(Func, OptionalMarkup) \
+	template<typename Arg1, typename Arg2, TEMPLATE_REQUIRES(TAnd< \
+																TNot<TIsSame<Arg1, Arg2>>, \
+																TIsSigned<Arg1>, TIsIntegral<Arg1>, \
+																TIsSigned<Arg2>, TIsIntegral<Arg2> \
+																>::Value)> \
 	static OptionalMarkup FORCEINLINE auto Func(Arg1 X, Arg2 Y) -> decltype(X * Y) \
 	{ \
 		using ArgType = decltype(X * Y); \
 		return Func((ArgType)X, (ArgType)Y); \
 	}
 
-#define MIX_SIGNED_TYPES_2_ARGS(Func)				MIX_SIGNED_TYPES_2_ARGS_ACTUAL(Func,)
-#define MIX_SIGNED_TYPES_2_ARGS_CONSTEXPR(Func)		MIX_SIGNED_TYPES_2_ARGS_ACTUAL(Func, CONSTEXPR)
+#define MIX_SIGNED_INTS_2_ARGS(Func)				MIX_SIGNED_INTS_2_ARGS_ACTUAL(Func,)
+#define MIX_SIGNED_INTS_2_ARGS_CONSTEXPR(Func)		MIX_SIGNED_INTS_2_ARGS_ACTUAL(Func, CONSTEXPR)
 
-// Mixing any types with any other type results in same type as "X * Y", which is promoted to the type of the result of mixed arithmetic between the types
+// Mixing any signed integral types with any other signed integral type results in same type as "X * Y", which is promoted to the type of the result of mixed arithmetic between the types
 #define MIX_SIGNED_TYPES_3_ARGS_ACTUAL(Func, OptionalMarkup) \
 	template<typename Arg1, typename Arg2, typename Arg3, TEMPLATE_REQUIRES(TAnd< \
 																				TOr< \
@@ -61,9 +66,9 @@
 																					TNot<TIsSame<Arg2, Arg3>>, \
 																					TNot<TIsSame<Arg1, Arg3>> \
 																					>, \
-																				TIsSigned<Arg1>, \
-																				TIsSigned<Arg2>, \
-																				TIsSigned<Arg3> \
+																				TIsSigned<Arg1>, TIsIntegral<Arg1>, \
+																				TIsSigned<Arg2>, TIsIntegral<Arg2>, \
+																				TIsSigned<Arg3>, TIsIntegral<Arg3> \
 																				>::Value)> \
 	static OptionalMarkup FORCEINLINE auto Func(Arg1 X, Arg2 Y, Arg3 Z) -> decltype(X * Y * Z) \
 	{ \
@@ -71,8 +76,8 @@
 		return Func((ArgType)X, (ArgType)Y, (ArgType)Z); \
 	}
 
-#define MIX_SIGNED_TYPES_3_ARGS(Func)				MIX_SIGNED_TYPES_3_ARGS_ACTUAL(Func,)
-#define MIX_SIGNED_TYPES_3_ARGS_CONSTEXPR(Func)		MIX_SIGNED_TYPES_3_ARGS_ACTUAL(Func, CONSTEXPR)
+#define MIX_SIGNED_INTS_3_ARGS(Func)				MIX_SIGNED_INTS_3_ARGS_ACTUAL(Func,)
+#define MIX_SIGNED_INTS_3_ARGS_CONSTEXPR(Func)		MIX_SIGNED_INTS_3_ARGS_ACTUAL(Func, CONSTEXPR)
 
 
 // Cannot know whether the type of e.g. Fmod(int, int) should be float or double.
diff --git a/Engine/Source/Runtime/Engine/Private/ImageUtils.cpp b/Engine/Source/Runtime/Engine/Private/ImageUtils.cpp
--- a/Engine/Source/Runtime/Engine/Private/ImageUtils.cpp
+++ b/Engine/Source/Runtime/Engine/Private/ImageUtils.cpp
@@ -496,10 +496,10 @@ void FImageUtils::ImageResize(int32 SrcWidth, int32 SrcHeight, const TArrayView<
 						PixelCount++;
 					}
 				}
-				uint8 FinalR = FMath::Clamp(FMath::TruncToInt(StepColor.X / (float)PixelCount), 0, 255);
-				uint8 FinalG = FMath::Clamp(FMath::TruncToInt(StepColor.Y / (float)PixelCount), 0, 255);
-				uint8 FinalB = FMath::Clamp(FMath::TruncToInt(StepColor.Z / (float)PixelCount), 0, 255);
-				uint8 FinalA = FMath::Clamp(FMath::TruncToInt(StepColor.W / (float)PixelCount), 0, 255);
+				uint8 FinalR = FMath::Clamp(FMath::TruncToInt32(StepColor.X / (float)PixelCount), 0, 255);
+				uint8 FinalG = FMath::Clamp(FMath::TruncToInt32(StepColor.Y / (float)PixelCount), 0, 255);
+				uint8 FinalB = FMath::Clamp(FMath::TruncToInt32(StepColor.Z / (float)PixelCount), 0, 255);
+				uint8 FinalA = FMath::Clamp(FMath::TruncToInt32(StepColor.W / (float)PixelCount), 0, 255);
 				FinalColor = FColor(FinalR, FinalG, FinalB, FinalA);
 			}
 
diff --git a/Engine/Source/Runtime/Engine/Private/KismetMathLibrary.cpp b/Engine/Source/Runtime/Engine/Private/KismetMathLibrary.cpp
--- a/Engine/Source/Runtime/Engine/Private/KismetMathLibrary.cpp
+++ b/Engine/Source/Runtime/Engine/Private/KismetMathLibrary.cpp
@@ -244,7 +244,7 @@ int32 UKismetMathLibrary::FMod(float Dividend, float Divisor, float& Remainder)
 	if (Divisor != 0.f)
 	{
 		const float Quotient = Dividend / Divisor;
-		Result = (Quotient < 0.f ? -1 : 1) * FMath::FloorToInt(FMath::Abs(Quotient));
+		Result = (Quotient < 0.f ? -1 : 1) * FMath::FloorToInt32(FMath::Abs(Quotient));
 		Remainder = FMath::Fmod(Dividend, Divisor);
 	}
 	else
@@ -264,7 +264,7 @@ int32 UKismetMathLibrary::FMod(double Dividend, double Divisor, double& Remainde
 	if (Divisor != 0.f)
 	{
 		const double Quotient = Dividend / Divisor;
-		Result = (Quotient < 0.f ? -1 : 1) * FMath::FloorToInt(FMath::Abs(Quotient));
+		Result = (Quotient < 0.f ? -1 : 1) * FMath::FloorToInt32(FMath::Abs(Quotient));
 		Remainder = FMath::Fmod(Dividend, Divisor);
 	}
 	else
diff --git a/Engine/Source/Runtime/Engine/Private/ModelLight.cpp b/Engine/Source/Runtime/Engine/Private/ModelLight.cpp
--- a/Engine/Source/Runtime/Engine/Private/ModelLight.cpp
+++ b/Engine/Source/Runtime/Engine/Private/ModelLight.cpp
@@ -438,8 +438,8 @@ void UModelComponent::GetSurfaceLightMapResolution( int32 SurfaceIndex, int32 Qu
 		MaxUV.X = FMath::CeilToFloat(MaxUV.X / Scale) * Scale;
 		MaxUV.Y = FMath::CeilToFloat(MaxUV.Y / Scale) * Scale;
 
-		Width = FMath::Clamp(FMath::CeilToInt((MaxUV.X - MinUV.X) / (Surf.LightMapScale * QualityScale)),4,SHADOWMAP_MAX_WIDTH);
-		Height = FMath::Clamp(FMath::CeilToInt((MaxUV.Y - MinUV.Y) / (Surf.LightMapScale * QualityScale)),4,SHADOWMAP_MAX_HEIGHT);
+		Width = FMath::Clamp(FMath::CeilToInt32((MaxUV.X - MinUV.X) / (Surf.LightMapScale * QualityScale)),4,SHADOWMAP_MAX_WIDTH);
+		Height = FMath::Clamp(FMath::CeilToInt32((MaxUV.Y - MinUV.Y) / (Surf.LightMapScale * QualityScale)),4,SHADOWMAP_MAX_HEIGHT);
 		WorldToMap = FMatrix(
 			FPlane(MapX.X / (MaxUV.X - MinUV.X),	MapY.X / (MaxUV.Y - MinUV.Y),	Surf.Plane.X,	0),
 			FPlane(MapX.Y / (MaxUV.X - MinUV.X),	MapY.Y / (MaxUV.Y - MinUV.Y),	Surf.Plane.Y,	0),
diff --git a/Engine/Source/Runtime/Engine/Private/PhysicsEngine/ImmediatePhysics/ImmediatePhysicsChaos/ImmediatePhysicsSimulation_Chaos.cpp b/Engine/Source/Runtime/Engine/Private/PhysicsEngine/ImmediatePhysics/ImmediatePhysicsChaos/ImmediatePhysicsSimulation_Chaos.cpp
--- a/Engine/Source/Runtime/Engine/Private/PhysicsEngine/ImmediatePhysics/ImmediatePhysicsChaos/ImmediatePhysicsSimulation_Chaos.cpp
+++ b/Engine/Source/Runtime/Engine/Private/PhysicsEngine/ImmediatePhysics/ImmediatePhysicsChaos/ImmediatePhysicsSimulation_Chaos.cpp
@@ -832,7 +832,7 @@ namespace ImmediatePhysics_Chaos
 		FReal StepTime = UpdateStepTime(DeltaTime, MaxStepTime);
 
 		// Calculate number of steps to run
-		int32 NumSteps = FMath::Clamp(FMath::RoundToInt(DeltaTime / StepTime), 1, MaxSubSteps);
+		int32 NumSteps = FMath::Clamp(FMath::RoundToInt32(DeltaTime / StepTime), 1, MaxSubSteps);
 
 		// TEMP: overrides
 		{
diff --git a/Engine/Source/Runtime/NavigationSystem/Private/NavMesh/RecastNavMeshGenerator.cpp b/Engine/Source/Runtime/NavigationSystem/Private/NavMesh/RecastNavMeshGenerator.cpp
--- a/Engine/Source/Runtime/NavigationSystem/Private/NavMesh/RecastNavMeshGenerator.cpp
+++ b/Engine/Source/Runtime/NavigationSystem/Private/NavMesh/RecastNavMeshGenerator.cpp
@@ -574,10 +574,10 @@ void ExportChaosHeightFieldSlice(const FNavHeightfieldSamples& PrefetchedHeightf
 	const FBox LocalBox = SliceBox.TransformBy(LocalToWorld.Inverse());
 	const bool bMirrored = (LocalToWorld.GetDeterminant() < 0.f);
 
-	const int32 MinX = FMath::Clamp(FMath::FloorToInt(LocalBox.Min.X) - 1, 0, NumCols);
-	const int32 MinY = FMath::Clamp(FMath::FloorToInt(LocalBox.Min.Y) - 1, 0, NumRows);
-	const int32 MaxX = FMath::Clamp(FMath::CeilToInt(LocalBox.Max.X) + 1, 0, NumCols);
-	const int32 MaxY = FMath::Clamp(FMath::CeilToInt(LocalBox.Max.Y) + 1, 0, NumRows);
+	const int32 MinX = FMath::Clamp(FMath::FloorToInt32(LocalBox.Min.X) - 1, 0, NumCols);
+	const int32 MinY = FMath::Clamp(FMath::FloorToInt32(LocalBox.Min.Y) - 1, 0, NumRows);
+	const int32 MaxX = FMath::Clamp(FMath::CeilToInt32(LocalBox.Max.X) + 1, 0, NumCols);
+	const int32 MaxY = FMath::Clamp(FMath::CeilToInt32(LocalBox.Max.Y) + 1, 0, NumRows);
 	const int32 SizeX = MaxX - MinX;
 	const int32 SizeY = MaxY - MinY;
 
diff --git a/Engine/Source/Runtime/RenderCore/Public/PackedNormal.h b/Engine/Source/Runtime/RenderCore/Public/PackedNormal.h
--- a/Engine/Source/Runtime/RenderCore/Public/PackedNormal.h
+++ b/Engine/Source/Runtime/RenderCore/Public/PackedNormal.h
@@ -263,26 +263,26 @@ struct FPackedRGB10A2N
 
 FORCEINLINE void FPackedRGB10A2N::operator=(const FVector3f& InVector)
 {
-	Vector.X = FMath::Clamp(FMath::TruncToInt(InVector.X * 511.5f + 511.5f), 0, 1023);
-	Vector.Y = FMath::Clamp(FMath::TruncToInt(InVector.Y * 511.5f + 511.5f), 0, 1023);
-	Vector.Z = FMath::Clamp(FMath::TruncToInt(InVector.Z * 511.5f + 511.5f), 0, 1023);
+	Vector.X = FMath::Clamp(FMath::TruncToInt32(InVector.X * 511.5f + 511.5f), 0, 1023);
+	Vector.Y = FMath::Clamp(FMath::TruncToInt32(InVector.Y * 511.5f + 511.5f), 0, 1023);
+	Vector.Z = FMath::Clamp(FMath::TruncToInt32(InVector.Z * 511.5f + 511.5f), 0, 1023);
 	Vector.W = 3;
 }
 
 FORCEINLINE void FPackedRGB10A2N::operator=(const FVector3d& InVector)
 {
-	Vector.X = FMath::Clamp(FMath::TruncToInt(InVector.X * 511.5f + 511.5f), 0, 1023);
-	Vector.Y = FMath::Clamp(FMath::TruncToInt(InVector.Y * 511.5f + 511.5f), 0, 1023);
-	Vector.Z = FMath::Clamp(FMath::TruncToInt(InVector.Z * 511.5f + 511.5f), 0, 1023);
+	Vector.X = FMath::Clamp(FMath::TruncToInt32(InVector.X * 511.5f + 511.5f), 0, 1023);
+	Vector.Y = FMath::Clamp(FMath::TruncToInt32(InVector.Y * 511.5f + 511.5f), 0, 1023);
+	Vector.Z = FMath::Clamp(FMath::TruncToInt32(InVector.Z * 511.5f + 511.5f), 0, 1023);
 	Vector.W = 3;
 }
 
 FORCEINLINE void FPackedRGB10A2N::operator=(const FVector4f& InVector)
 {
-	Vector.X = FMath::Clamp(FMath::TruncToInt(InVector.X * 511.5f + 511.5f), 0, 1023);
-	Vector.Y = FMath::Clamp(FMath::TruncToInt(InVector.Y * 511.5f + 511.5f), 0, 1023);
-	Vector.Z = FMath::Clamp(FMath::TruncToInt(InVector.Z * 511.5f + 511.5f), 0, 1023);
-	Vector.W = FMath::Clamp(FMath::TruncToInt(InVector.W * 1.5f   + 1.5f),   0, 3);
+	Vector.X = FMath::Clamp(FMath::TruncToInt32(InVector.X * 511.5f + 511.5f), 0, 1023);
+	Vector.Y = FMath::Clamp(FMath::TruncToInt32(InVector.Y * 511.5f + 511.5f), 0, 1023);
+	Vector.Z = FMath::Clamp(FMath::TruncToInt32(InVector.Z * 511.5f + 511.5f), 0, 1023);
+	Vector.W = FMath::Clamp(FMath::TruncToInt32(InVector.W * 1.5f   + 1.5f),   0, 3);
 }
 
 FORCEINLINE bool FPackedRGB10A2N::operator==(const FPackedRGB10A2N& B) const
@@ -360,37 +360,37 @@ DECLARE_INTRINSIC_TYPE_LAYOUT(FPackedRGBA16N);
 FORCEINLINE void FPackedRGBA16N::operator=(const FVector3f& InVector)
 {
 	const float Scale = MAX_int16;
-	X = (int16)FMath::Clamp<int32>(FMath::RoundToInt(InVector.X * Scale), MIN_int16, MAX_int16);
-	Y = (int16)FMath::Clamp<int32>(FMath::RoundToInt(InVector.Y * Scale), MIN_int16, MAX_int16);
-	Z = (int16)FMath::Clamp<int32>(FMath::RoundToInt(InVector.Z * Scale), MIN_int16, MAX_int16);
+	X = (int16)FMath::Clamp<int32>(FMath::RoundToInt32(InVector.X * Scale), MIN_int16, MAX_int16);
+	Y = (int16)FMath::Clamp<int32>(FMath::RoundToInt32(InVector.Y * Scale), MIN_int16, MAX_int16);
+	Z = (int16)FMath::Clamp<int32>(FMath::RoundToInt32(InVector.Z * Scale), MIN_int16, MAX_int16);
 	W = MAX_int16;
 }
 
 FORCEINLINE void FPackedRGBA16N::operator=(const FVector3d& InVector)
 {
 	const float Scale = MAX_int16;
-	X = (int16)FMath::Clamp<int32>(FMath::RoundToInt(InVector.X * Scale), MIN_int16, MAX_int16);
-	Y = (int16)FMath::Clamp<int32>(FMath::RoundToInt(InVector.Y * Scale), MIN_int16, MAX_int16);
-	Z = (int16)FMath::Clamp<int32>(FMath::RoundToInt(InVector.Z * Scale), MIN_int16, MAX_int16);
+	X = (int16)FMath::Clamp<int32>(FMath::RoundToInt32(InVector.X * Scale), MIN_int16, MAX_int16);
+	Y = (int16)FMath::Clamp<int32>(FMath::RoundToInt32(InVector.Y * Scale), MIN_int16, MAX_int16);
+	Z = (int16)FMath::Clamp<int32>(FMath::RoundToInt32(InVector.Z * Scale), MIN_int16, MAX_int16);
 	W = MAX_int16;
 }
 
 FORCEINLINE void FPackedRGBA16N::operator=(const FVector4f& InVector)
 {
 	const float Scale = MAX_int16;
-	X = (int16)FMath::Clamp<int32>(FMath::RoundToInt(InVector.X * Scale), MIN_int16, MAX_int16);
-	Y = (int16)FMath::Clamp<int32>(FMath::RoundToInt(InVector.Y * Scale), MIN_int16, MAX_int16);
-	Z = (int16)FMath::Clamp<int32>(FMath::RoundToInt(InVector.Z * Scale), MIN_int16, MAX_int16);
-	W = (int16)FMath::Clamp<int32>(FMath::RoundToInt(InVector.W * Scale), MIN_int16, MAX_int16);
+	X = (int16)FMath::Clamp<int32>(FMath::RoundToInt32(InVector.X * Scale), MIN_int16, MAX_int16);
+	Y = (int16)FMath::Clamp<int32>(FMath::RoundToInt32(InVector.Y * Scale), MIN_int16, MAX_int16);
+	Z = (int16)FMath::Clamp<int32>(FMath::RoundToInt32(InVector.Z * Scale), MIN_int16, MAX_int16);
+	W = (int16)FMath::Clamp<int32>(FMath::RoundToInt32(InVector.W * Scale), MIN_int16, MAX_int16);
 }
 
 FORCEINLINE void FPackedRGBA16N::operator=(const FVector4d& InVector)
 {
 	const float Scale = MAX_int16;
-	X = (int16)FMath::Clamp<int32>(FMath::RoundToInt(InVector.X * Scale), MIN_int16, MAX_int16);
-	Y = (int16)FMath::Clamp<int32>(FMath::RoundToInt(InVector.Y * Scale), MIN_int16, MAX_int16);
-	Z = (int16)FMath::Clamp<int32>(FMath::RoundToInt(InVector.Z * Scale), MIN_int16, MAX_int16);
-	W = (int16)FMath::Clamp<int32>(FMath::RoundToInt(InVector.W * Scale), MIN_int16, MAX_int16);
+	X = (int16)FMath::Clamp<int32>(FMath::RoundToInt32(InVector.X * Scale), MIN_int16, MAX_int16);
+	Y = (int16)FMath::Clamp<int32>(FMath::RoundToInt32(InVector.Y * Scale), MIN_int16, MAX_int16);
+	Z = (int16)FMath::Clamp<int32>(FMath::RoundToInt32(InVector.Z * Scale), MIN_int16, MAX_int16);
+	W = (int16)FMath::Clamp<int32>(FMath::RoundToInt32(InVector.W * Scale), MIN_int16, MAX_int16);
 }
 
 FORCEINLINE bool FPackedRGBA16N::operator==(const FPackedRGBA16N& B) const
diff --git a/Engine/Source/Runtime/Renderer/Private/PlanarReflectionRendering.cpp b/Engine/Source/Runtime/Renderer/Private/PlanarReflectionRendering.cpp
--- a/Engine/Source/Runtime/Renderer/Private/PlanarReflectionRendering.cpp
+++ b/Engine/Source/Runtime/Renderer/Private/PlanarReflectionRendering.cpp
@@ -515,8 +515,8 @@ void FScene::UpdatePlanarReflectionContents(UPlanarReflectionComponent* CaptureC
 		FIntPoint DesiredBufferSize = FSceneRenderer::GetDesiredInternalBufferSize(*MainSceneRenderer.ActiveViewFamily);
 		FVector2D DesiredPlanarReflectionTextureSizeFloat = FVector2D(DesiredBufferSize.X, DesiredBufferSize.Y) * FMath::Clamp(CaptureComponent->ScreenPercentage / 100.f, 0.25f, 1.f);
 		FIntPoint DesiredPlanarReflectionTextureSize;
-		DesiredPlanarReflectionTextureSize.X = FMath::Clamp(FMath::CeilToInt(DesiredPlanarReflectionTextureSizeFloat.X), 1, static_cast<int32>(DesiredBufferSize.X));
-		DesiredPlanarReflectionTextureSize.Y = FMath::Clamp(FMath::CeilToInt(DesiredPlanarReflectionTextureSizeFloat.Y), 1, static_cast<int32>(DesiredBufferSize.Y));
+		DesiredPlanarReflectionTextureSize.X = FMath::Clamp(FMath::CeilToInt32(DesiredPlanarReflectionTextureSizeFloat.X), 1, static_cast<int32>(DesiredBufferSize.X));
+		DesiredPlanarReflectionTextureSize.Y = FMath::Clamp(FMath::CeilToInt32(DesiredPlanarReflectionTextureSizeFloat.Y), 1, static_cast<int32>(DesiredBufferSize.Y));
 
 		const bool bIsMobilePixelProjectedReflectionEnabled = IsMobilePixelProjectedReflectionEnabled(GetShaderPlatform());
 
