Moved RDG uniform buffer creation off the render thread. Added AddSetupTask helper function for async tasks that are automatically waited on prior to execution.

diff --git a/Engine/Source/Runtime/RenderCore/Private/RenderGraphBuilder.cpp b/Engine/Source/Runtime/RenderCore/Private/RenderGraphBuilder.cpp
--- a/Engine/Source/Runtime/RenderCore/Private/RenderGraphBuilder.cpp
+++ b/Engine/Source/Runtime/RenderCore/Private/RenderGraphBuilder.cpp
@@ -7,7 +7,7 @@
 #include "RenderGraphResourcePool.h"
 #include "VisualizeTexture.h"
 #include "ProfilingDebugging/CsvProfiler.h"
-
+#include "Async/ParallelFor.h"
 
 #if ENABLE_RHI_VALIDATION
 
@@ -1524,7 +1524,7 @@
 	const FRDGPassHandle ProloguePassHandle = GetProloguePassHandle();
 	const FRDGPassHandle EpiloguePassHandle = GetEpiloguePassHandle();
 
-	FGraphEventArray AsyncCompileEvents;
+	FGraphEventArray AsyncCompileEvents = ParallelSetupEvents;
 
 	IF_RDG_ENABLE_DEBUG(UserValidation.ValidateExecuteBegin());
 	IF_RDG_ENABLE_DEBUG(GRDGAllowRHIAccess = true);
@@ -1637,6 +1637,8 @@
 			}
 		}
 
+		CreateUniformBuffers(bParallelExecuteEnabled ? &AsyncCompileEvents : nullptr);
+
 		{
 			SCOPED_NAMED_EVENT_TEXT("FRDGBuilder::CollectBarriers", FColor::Magenta);
 			SCOPE_CYCLE_COUNTER(STAT_RDG_CollectBarriersTime);
@@ -1670,35 +1672,30 @@
 		});
 	}
 
-	if (bParallelExecuteEnabled)
-	{
-		// Overlap pass barrier creation with other compilation tasks, since it's not required to run on the render thread.
-		AsyncCompileEvents.Emplace(FFunctionGraphTask::CreateAndDispatchWhenReady(
-			[this](ENamedThreads::Type, const FGraphEventRef&)
-		{
-			CreatePassBarriers();
+	CreatePassBarriers(bParallelExecuteEnabled ? &AsyncCompileEvents : nullptr);
 
-		}, TStatId(), nullptr, ENamedThreads::AnyHiPriThreadHiPriTask));
-	}
-	else
-	{
-		CreatePassBarriers();
-	}
-
-	SubmitBufferUploads();
-
-	CreateUniformBuffers();
-
 	EndFlushResourcesRHI();
 
-	IF_RDG_ENABLE_TRACE(Trace.OutputGraphBegin());
-
 	IF_RDG_ENABLE_DEBUG(GRDGAllowRHIAccess = bParallelExecuteEnabled);
 
-	const ENamedThreads::Type RenderThread = ENamedThreads::GetRenderThread_Local();
-
 	FGraphEventRef DispatchParallelExecuteEvent;
 
+	if (bParallelExecuteEnabled)
+	{
+		DispatchParallelExecuteEvent = FFunctionGraphTask::CreateAndDispatchWhenReady(
+			[this, &RHICmdContext = RHICmdList.GetContext()](ENamedThreads::Type, const FGraphEventRef&)
+		{
+			DispatchParallelExecute(&RHICmdContext);
+
+		}, TStatId(), &AsyncCompileEvents, ENamedThreads::AnyHiPriThreadHiPriTask);
+	}
+
+	SubmitBufferUploads();
+
+	IF_RDG_ENABLE_TRACE(Trace.OutputGraphBegin());
+
+	const ENamedThreads::Type RenderThread = ENamedThreads::GetRenderThread_Local();
+
 	if (!IsImmediateMode())
 	{
 		SCOPED_NAMED_EVENT_TEXT("FRDGBuilder::ExecutePasses", FColor::Magenta);
@@ -1711,16 +1708,6 @@
 			FTaskGraphInterface::Get().WaitUntilTasksComplete(AsyncCompileEvents, RenderThread);
 		}
 
-		if (bParallelExecuteEnabled)
-		{
-			DispatchParallelExecuteEvent = FFunctionGraphTask::CreateAndDispatchWhenReady(
-				[this, &RHICmdContext = RHICmdList.GetContext()](ENamedThreads::Type, const FGraphEventRef&)
-			{
-				DispatchParallelExecute(&RHICmdContext);
-
-			}, TStatId(), nullptr, ENamedThreads::AnyHiPriThreadHiPriTask);
-		}
-
 		for (FRDGPassHandle PassHandle = ProloguePassHandle; PassHandle <= EpiloguePassHandle; ++PassHandle)
 		{
 			FRDGPass* Pass = Passes[PassHandle];
@@ -2103,9 +2090,11 @@
 		SubmitBufferUploads();
 		CompilePassOps(Pass);
 		BeginResourcesRHI(Pass, PassHandle);
+		CreateUniformBuffers(nullptr);
 		CollectPassBarriers(Pass, PassHandle);
-		CreatePassBarriers();
-		CreateUniformBuffers();
+		CreatePassBarriers(nullptr);
+		FTaskGraphInterface::Get().WaitUntilTasksComplete(ParallelSetupEvents);
+		ParallelSetupEvents.Reset();
 		ExecutePass(Pass, RHICmdList);
 	}
 
@@ -2164,9 +2153,9 @@
 			else
 #endif
 			{
-				void* DestPtr = RHICmdList.LockBuffer(UploadedBuffer.Buffer->GetRHI(), 0, UploadedBuffer.DataSize, RLM_WriteOnly);
+				void* DestPtr = RHICmdList.LockBuffer(UploadedBuffer.Buffer->GetRHIUnchecked(), 0, UploadedBuffer.DataSize, RLM_WriteOnly);
 				FMemory::Memcpy(DestPtr, UploadedBuffer.Data, UploadedBuffer.DataSize);
-				RHICmdList.UnlockBuffer(UploadedBuffer.Buffer->GetRHI());
+				RHICmdList.UnlockBuffer(UploadedBuffer.Buffer->GetRHIUnchecked());
 			}
 
 			if (UploadedBuffer.bUseFreeCallbacks)
@@ -2381,15 +2370,33 @@
 
 ///////////////////////////////////////////////////////////////////////////////////////////////////
 
-void FRDGBuilder::CreateUniformBuffers()
+void FRDGBuilder::CreateUniformBuffers(FGraphEventArray* AsyncCompileEvents)
 {
-	SCOPED_NAMED_EVENT_TEXT("FRDGBuilder::CreateUniformBuffers", FColor::Magenta);
+	const int32 ParallelDispatchThreshold = 4;
 
-	for (FRDGUniformBufferHandle UniformBufferHandle : UniformBuffersToCreate)
+	const auto CreateUniformBuffersFunction = [this]
 	{
-		UniformBuffers[UniformBufferHandle]->InitRHI();
+		SCOPED_NAMED_EVENT_TEXT("FRDGBuilder::CreateUniformBuffers", FColor::Magenta);
+		for (FRDGUniformBufferHandle UniformBufferHandle : UniformBuffersToCreate)
+		{
+			UniformBuffers[UniformBufferHandle]->InitRHI();
+		}
+		UniformBuffersToCreate.Reset();
+	};
+
+	if (AsyncCompileEvents && UniformBuffersToCreate.Num() > ParallelDispatchThreshold)
+	{
+		AsyncCompileEvents->Emplace(FFunctionGraphTask::CreateAndDispatchWhenReady(
+			[CreateUniformBuffersFunction](ENamedThreads::Type, const FGraphEventRef&)
+		{
+			CreateUniformBuffersFunction();
+
+		}, TStatId(), nullptr, ENamedThreads::AnyHiPriThreadHiPriTask));
 	}
-	UniformBuffersToCreate.Reset();
+	else
+	{
+		CreateUniformBuffersFunction();
+	}
 }
 
 ///////////////////////////////////////////////////////////////////////////////////////////////////
@@ -2497,7 +2504,7 @@
 
 void FRDGBuilder::ExecutePass(FRDGPass* Pass, FRHIComputeCommandList& RHICmdListPass)
 {
-#if RDG_EVENTS != RDG_EVENTS_NONE
+#if 1 || RDG_EVENTS != RDG_EVENTS_NONE
 	SCOPED_NAMED_EVENT_TCHAR(Pass->GetName(), FColor::Magenta);
 #endif
 
@@ -2634,15 +2641,43 @@
 	}
 }
 
-void FRDGBuilder::CreatePassBarriers()
+void FRDGBuilder::CreatePassBarriers(FGraphEventArray* AsyncCompileEvents)
 {
-	SCOPED_NAMED_EVENT_TEXT("FRDGBuilder::CreatePassBarriers", FColor::Magenta);
+	const int32 NumBarriersPerTask = 128;
 
-	for (FRDGBarrierBatchBegin* BarrierBatchBegin : TransitionCreateQueue)
+	if (AsyncCompileEvents && TransitionCreateQueue.Num() > NumBarriersPerTask)
 	{
-		BarrierBatchBegin->CreateTransition();
+		AsyncCompileEvents->Emplace(FFunctionGraphTask::CreateAndDispatchWhenReady(
+			[this, NumBarriersPerTask](ENamedThreads::Type, const FGraphEventRef&)
+		{
+#if 1
+			ParallelFor(TEXT("FRDGBuilder::CreatePassBarriers"), TransitionCreateQueue.Num(), NumBarriersPerTask, [this](int32 Index)
+			{
+				TransitionCreateQueue[Index]->CreateTransition();
+			});
+#else
+			SCOPED_NAMED_EVENT_TEXT("FRDGBuilder::CreatePassBarriers", FColor::Magenta);
+			for (FRDGBarrierBatchBegin* BeginBatch : TransitionCreateQueue)
+			{
+				BeginBatch->CreateTransition();
+			}
+#endif
+
+			TransitionCreateQueue.Reset();
+
+		}, TStatId(), nullptr, ENamedThreads::AnyHiPriThreadHiPriTask));
 	}
-	TransitionCreateQueue.Reset();
+	else
+	{
+		SCOPED_NAMED_EVENT_TEXT("FRDGBuilder::CreatePassBarriers", FColor::Magenta);
+
+		for (FRDGBarrierBatchBegin* BeginBatch : TransitionCreateQueue)
+		{
+			BeginBatch->CreateTransition();
+		}
+
+		TransitionCreateQueue.Reset();
+	}
 }
 
 ///////////////////////////////////////////////////////////////////////////////////////////////////
diff --git a/Engine/Source/Runtime/RenderCore/Public/RenderGraphBuilder.h b/Engine/Source/Runtime/RenderCore/Public/RenderGraphBuilder.h
--- a/Engine/Source/Runtime/RenderCore/Public/RenderGraphBuilder.h
+++ b/Engine/Source/Runtime/RenderCore/Public/RenderGraphBuilder.h
@@ -184,6 +184,19 @@
 	/** A hint to the builder to flush work to the RHI thread after the last queued pass on the execution timeline. */
 	void AddDispatchHint();
 
+
+	template <typename TaskLambda>
+	void AddSetupTask(TaskLambda&& Task)
+	{
+		ParallelSetupEvents.Emplace(FFunctionGraphTask::CreateAndDispatchWhenReady(
+			[Task = MoveTemp(Task)](ENamedThreads::Type, const FGraphEventRef&)
+		{
+			FTaskTagScope Scope(ETaskTag::EParallelRenderingThread);
+			Task();
+
+		}, TStatId(), nullptr, ENamedThreads::AnyHiPriThreadHiPriTask));
+	}
+
 	/** Tells the builder to delete unused RHI resources. The behavior of this method depends on whether RDG immediate mode is enabled:
 	 *   Deferred:  RHI resource flushes are performed prior to execution.
 	 *   Immediate: RHI resource flushes are performed immediately.
@@ -650,6 +663,9 @@
 	/** Array of all active parallel execute tasks. */
 	FGraphEventArray ParallelExecuteEvents;
 
+	/** Array of all task events requested by the user. */
+	FGraphEventArray ParallelSetupEvents;
+
 	/** Tracks the final access used on resources in order to call SetTrackedAccess. */
 	TArray<FRHITrackedAccessInfo, FRDGArrayAllocator> EpilogueResourceAccesses;
 
@@ -749,9 +765,9 @@
 
 	void CollectPassBarriers(FRDGPass* Pass, FRDGPassHandle PassHandle);
 
-	void CreatePassBarriers();
+	void CreatePassBarriers(FGraphEventArray* AsyncCompileEvents);
 
-	void CreateUniformBuffers();
+	void CreateUniformBuffers(FGraphEventArray* AsyncCompileEvents);
 
 	void AddPassDependency(FRDGPassHandle ProducerHandle, FRDGPassHandle ConsumerHandle);
 
diff --git a/Engine/Source/Runtime/RenderCore/Public/RenderGraphResources.h b/Engine/Source/Runtime/RenderCore/Public/RenderGraphResources.h
--- a/Engine/Source/Runtime/RenderCore/Public/RenderGraphResources.h
+++ b/Engine/Source/Runtime/RenderCore/Public/RenderGraphResources.h
@@ -1417,5 +1417,6 @@ inline FRDGTextureSRV* GetAsTextureSRV(FRDGView* View)
 }
 
 inline FGraphicsPipelineRenderTargetsInfo ExtractRenderTargetsInfo(const FRDGParameterStruct& ParameterStruct);
+inline FGraphicsPipelineRenderTargetsInfo ExtractRenderTargetsInfo(const FRenderTargetBindingSlots& RenderTargets);
 
 #include "RenderGraphResources.inl"
\ No newline at end of file
diff --git a/Engine/Source/Runtime/RenderCore/Public/RenderGraphResources.inl b/Engine/Source/Runtime/RenderCore/Public/RenderGraphResources.inl
--- a/Engine/Source/Runtime/RenderCore/Public/RenderGraphResources.inl
+++ b/Engine/Source/Runtime/RenderCore/Public/RenderGraphResources.inl
@@ -209,13 +209,17 @@ inline FRDGBufferUAVDesc::FRDGBufferUAVDesc(FRDGBufferRef InBuffer)
 
 inline FGraphicsPipelineRenderTargetsInfo ExtractRenderTargetsInfo(const FRDGParameterStruct& ParameterStruct)
 {
-	const FRenderTargetBindingSlots& RDGRenderTargets = ParameterStruct.GetRenderTargets();
+	return ExtractRenderTargetsInfo(ParameterStruct.GetRenderTargets());
+}
+
+inline FGraphicsPipelineRenderTargetsInfo ExtractRenderTargetsInfo(const FRenderTargetBindingSlots& RenderTargets)
+{
 	FGraphicsPipelineRenderTargetsInfo RenderTargetsInfo;
 	uint32 RenderTargetIndex = 0;
 
 	RenderTargetsInfo.NumSamples = 1;
 
-	RDGRenderTargets.Enumerate([&](FRenderTargetBinding RenderTarget)
+	RenderTargets.Enumerate([&](FRenderTargetBinding RenderTarget)
 	{
 		FRDGTextureRef Texture = RenderTarget.GetTexture();
 
@@ -231,7 +235,7 @@ inline FGraphicsPipelineRenderTargetsInfo ExtractRenderTargetsInfo(const FRDGPar
 		RenderTargetsInfo.RenderTargetFormats[RenderTargetIndex] = PF_Unknown;
 	}
 
-	const FDepthStencilBinding& DepthStencil = RDGRenderTargets.DepthStencil;
+	const FDepthStencilBinding& DepthStencil = RenderTargets.DepthStencil;
 	if (FRDGTextureRef DepthTexture = DepthStencil.GetTexture())
 	{
 		RenderTargetsInfo.DepthStencilTargetFormat = DepthTexture->Desc.Format;
@@ -251,8 +255,8 @@ inline FGraphicsPipelineRenderTargetsInfo ExtractRenderTargetsInfo(const FRDGPar
 		RenderTargetsInfo.DepthStencilTargetFormat = PF_Unknown;
 	}
 
-	RenderTargetsInfo.MultiViewCount = RDGRenderTargets.MultiViewCount;
-	RenderTargetsInfo.bHasFragmentDensityAttachment = RDGRenderTargets.ShadingRateTexture != nullptr;
+	RenderTargetsInfo.MultiViewCount = RenderTargets.MultiViewCount;
+	RenderTargetsInfo.bHasFragmentDensityAttachment = RenderTargets.ShadingRateTexture != nullptr;
 
 	return RenderTargetsInfo;
-}
+}
\ No newline at end of file
