Add missing task tag when FRDGBuilder::CreateUniformBuffers runs async.

#rnx

diff --git a/Engine/Source/Runtime/RenderCore/Private/RenderGraphBuilder.cpp b/Engine/Source/Runtime/RenderCore/Private/RenderGraphBuilder.cpp
--- a/Engine/Source/Runtime/RenderCore/Private/RenderGraphBuilder.cpp
+++ b/Engine/Source/Runtime/RenderCore/Private/RenderGraphBuilder.cpp
@@ -2389,6 +2389,7 @@ void FRDGBuilder::CreateUniformBuffers(FGraphEventArray* AsyncCompileEvents)
 		AsyncCompileEvents->Emplace(FFunctionGraphTask::CreateAndDispatchWhenReady(
 			[CreateUniformBuffersFunction](ENamedThreads::Type, const FGraphEventRef&)
 		{
+			FTaskTagScope Scope(ETaskTag::EParallelRenderingThread);
 			CreateUniformBuffersFunction();
 
 		}, TStatId(), nullptr, ENamedThreads::AnyHiPriThreadHiPriTask));
