Fix for engine test failure in DumpGPU on Vulkan

diff --git a/Engine/Source/Runtime/RenderCore/Private/DumpGPU.cpp b/Engine/Source/Runtime/RenderCore/Private/DumpGPU.cpp
--- a/Engine/Source/Runtime/RenderCore/Private/DumpGPU.cpp
+++ b/Engine/Source/Runtime/RenderCore/Private/DumpGPU.cpp
@@ -18,8 +18,8 @@
 #include "RenderGraphUtils.h"
 #include "GlobalShader.h"
 #include "RHIValidation.h"
+#include "RenderGraphPrivate.h"
 
-
 IDumpGPUUploadServiceProvider* IDumpGPUUploadServiceProvider::GProvider = nullptr;
 
 FString IDumpGPUUploadServiceProvider::FDumpParameters::DumpServiceParametersFileContent() const
@@ -1898,7 +1898,7 @@
 
 void FRDGBuilder::DumpResourcePassOutputs(const FRDGPass* Pass)
 {
-	if (bInDumpPassOutputs)
+	if (!AuxiliaryPasses.IsDumpAllowed())
 	{
 		return;
 	}
@@ -1914,7 +1914,7 @@
 		return;
 	}
 
-	bInDumpPassOutputs = true;
+	RDG_RECURSION_COUNTER_SCOPE(AuxiliaryPasses.Dump);
 
 	TArray<TSharedPtr<FJsonValue>> InputResourceNames;
 	TArray<TSharedPtr<FJsonValue>> OutputResourceNames;
@@ -2165,8 +2165,6 @@
 	}
 
 	GRDGResourceDumpContext->PassesCount++;
-
-	bInDumpPassOutputs = false;
 }
 
 #if RDG_DUMP_RESOURCES_AT_EACH_DRAW
diff --git a/Engine/Source/Runtime/RenderCore/Private/RenderGraphBuilder.cpp b/Engine/Source/Runtime/RenderCore/Private/RenderGraphBuilder.cpp
--- a/Engine/Source/Runtime/RenderCore/Private/RenderGraphBuilder.cpp
+++ b/Engine/Source/Runtime/RenderCore/Private/RenderGraphBuilder.cpp
@@ -673,12 +673,14 @@ void FRDGBuilder::UseInternalAccessMode(FRDGViewableResource* Resource)
 
 void FRDGBuilder::FlushAccessModeQueue()
 {
-	if (AccessModeQueue.IsEmpty() || bInFlushExternalAccess)
+	if (AccessModeQueue.IsEmpty() || !AuxiliaryPasses.IsFlushAccessModeQueueAllowed())
 	{
 		return;
 	}
 
-	bInFlushExternalAccess = true;
+	// Don't allow Dump GPU to dump access mode passes. We rely on FlushAccessQueue in dump GPU to transition things back to external access.
+	RDG_RECURSION_COUNTER_SCOPE(AuxiliaryPasses.Dump);
+	RDG_RECURSION_COUNTER_SCOPE(AuxiliaryPasses.FlushAccessModeQueue);
 
 	FAccessModePassParameters* ParametersByPipeline[] =
 	{
@@ -761,8 +763,6 @@ void FRDGBuilder::FlushAccessModeQueue()
 	}
 
 	AccessModeQueue.Reset();
-
-	bInFlushExternalAccess = false;
 }
 
 ///////////////////////////////////////////////////////////////////////////////////////////////////
@@ -1676,8 +1676,6 @@ void FRDGBuilder::Execute()
 
 	EndFlushResourcesRHI();
 
-	IF_RDG_ENABLE_DEBUG(GRDGAllowRHIAccess = bParallelExecuteEnabled);
-
 	FGraphEventRef DispatchParallelExecuteEvent;
 
 	if (bParallelExecuteEnabled)
@@ -1692,6 +1690,7 @@ void FRDGBuilder::Execute()
 
 	SubmitBufferUploads();
 
+	IF_RDG_ENABLE_DEBUG(GRDGAllowRHIAccess = bParallelExecuteEnabled);
 	IF_RDG_ENABLE_TRACE(Trace.OutputGraphBegin());
 
 	const ENamedThreads::Type RenderThread = ENamedThreads::GetRenderThread_Local();
@@ -1854,7 +1853,7 @@ void FRDGBuilder::Execute()
 
 FRDGPass* FRDGBuilder::SetupPass(FRDGPass* Pass)
 {
-	IF_RDG_ENABLE_DEBUG(UserValidation.ValidateAddPass(Pass, InAuxiliaryPass()));
+	IF_RDG_ENABLE_DEBUG(UserValidation.ValidateAddPass(Pass, AuxiliaryPasses.IsActive()));
 	CSV_SCOPED_TIMING_STAT_EXCLUSIVE_CONDITIONAL(RDGBuilder_SetupPass, GRDGVerboseCSVStats != 0);
 
 	const FRDGParameterStruct PassParameters = Pass->GetParameters();
@@ -3477,12 +3476,12 @@ void FRDGBuilder::EndResourceRHI(FRDGPassHandle PassHandle, FRDGBufferRef Buffer
 void FRDGBuilder::VisualizePassOutputs(const FRDGPass* Pass)
 {
 #if SUPPORTS_VISUALIZE_TEXTURE
-	if (bInVisualizePassOutputs)
+	if (!AuxiliaryPasses.IsVisualizeAllowed())
 	{
 		return;
 	}
 
-	bInVisualizePassOutputs = true;
+	RDG_RECURSION_COUNTER_SCOPE(AuxiliaryPasses.Visualize);
 
 	Pass->GetParameters().EnumerateTextures([&](FRDGParameter Parameter)
 	{
@@ -3548,19 +3547,17 @@ void FRDGBuilder::VisualizePassOutputs(const FRDGPass* Pass)
 		break;
 		}
 	});
-
-	bInVisualizePassOutputs = false;
 #endif
 }
 
 void FRDGBuilder::ClobberPassOutputs(const FRDGPass* Pass)
 {
-	if (!GRDGClobberResources || bInClobberPassOutputs)
+	if (!GRDGClobberResources || !AuxiliaryPasses.IsClobberAllowed())
 	{
 		return;
 	}
 
-	bInClobberPassOutputs = true;
+	RDG_RECURSION_COUNTER_SCOPE(AuxiliaryPasses.Clobber);
 
 	RDG_EVENT_SCOPE(*this, "RDG ClobberResources");
 
@@ -3654,8 +3651,6 @@ void FRDGBuilder::ClobberPassOutputs(const FRDGPass* Pass)
 		break;
 		}
 	});
-
-	bInClobberPassOutputs = false;
 }
 
 #endif //! RDG_ENABLE_DEBUG
diff --git a/Engine/Source/Runtime/RenderCore/Private/RenderGraphPrivate.h b/Engine/Source/Runtime/RenderCore/Private/RenderGraphPrivate.h
--- a/Engine/Source/Runtime/RenderCore/Private/RenderGraphPrivate.h
+++ b/Engine/Source/Runtime/RenderCore/Private/RenderGraphPrivate.h
@@ -20,6 +20,8 @@ DEFINE_LOG_CATEGORY_STATIC(LogRDG, Log, All);
 #define RDG_BREAKPOINT_PASS_COMPILE 2
 #define RDG_BREAKPOINT_PASS_EXECUTE 3
 
+#define RDG_RECURSION_COUNTER_SCOPE(Counter) Counter++; ON_SCOPE_EXIT { Counter--; }
+
 #if RDG_ENABLE_DEBUG
 extern int32 GRDGAsyncCompute;
 extern int32 GRDGClobberResources;
diff --git a/Engine/Source/Runtime/RenderCore/Public/RenderGraphBuilder.h b/Engine/Source/Runtime/RenderCore/Public/RenderGraphBuilder.h
--- a/Engine/Source/Runtime/RenderCore/Public/RenderGraphBuilder.h
+++ b/Engine/Source/Runtime/RenderCore/Public/RenderGraphBuilder.h
@@ -694,13 +694,22 @@ private:
 	FRDGBarrierValidation BarrierValidation;
 #endif
 
-	/** Tracks whether we are in a scope of adding passes to the builder. Used to avoid recursion. */
-	bool bInClobberPassOutputs = false;
-	bool bInVisualizePassOutputs = false;
-	bool bInDumpPassOutputs = false;
-	bool bInFlushExternalAccess = false;
+	/** Tracks stack counters of auxiliary passes to avoid calling them recursively. */
+	struct FAuxiliaryPass
+	{
+		uint8 Clobber = 0;
+		uint8 Visualize = 0;
+		uint8 Dump = 0;
+		uint8 FlushAccessModeQueue = 0;
+
+		bool IsDumpAllowed() const { return Dump == 0; }
+		bool IsVisualizeAllowed() const { return Visualize == 0; }
+		bool IsClobberAllowed() const { return Clobber == 0; }
+		bool IsFlushAccessModeQueueAllowed() const { return FlushAccessModeQueue == 0; }
+
+		bool IsActive() const { return Clobber > 0 || Visualize > 0 || Dump > 0 || FlushAccessModeQueue > 0; }
 
-	bool InAuxiliaryPass() const { return bInClobberPassOutputs || bInVisualizePassOutputs || bInDumpPassOutputs || bInFlushExternalAccess; }
+	} AuxiliaryPasses;
 
 #if WITH_MGPU
 	/** Name for the temporal effect used to synchronize multi-frame resources. */
