Missing copyright boilerplate in StereoRenderUtils.cpp

diff --git a/Engine/Source/Runtime/RenderCore/Private/StereoRenderUtils.cpp b/Engine/Source/Runtime/RenderCore/Private/StereoRenderUtils.cpp
--- a/Engine/Source/Runtime/RenderCore/Private/StereoRenderUtils.cpp
+++ b/Engine/Source/Runtime/RenderCore/Private/StereoRenderUtils.cpp
@@ -1,3 +1,5 @@
+// Copyright Epic Games, Inc. All Rights Reserved.
+
 #include "StereoRenderUtils.h"
 #include "RenderUtils.h"
 
