Restoring functionality to r.ShaderCompiler.EmitWarningsOnLoad

diff --git a/Engine/Shaders/Public/ShaderVersion.ush b/Engine/Shaders/Public/ShaderVersion.ush
--- a/Engine/Shaders/Public/ShaderVersion.ush
+++ b/Engine/Shaders/Public/ShaderVersion.ush
@@ -2,4 +2,4 @@
 // Each time the console command is executed it generates a new GUID. As this file is included
 // in Platform.ush (which should be included in any shader) it allows to invalidate the shader DDC.
 // 
-// GUID = 0658D561-29B7-48D0-AED8-4BD4CD643A90
+// GUID = E3AA61C3-6F51-4308-8DC0-46E401CFDD58
diff --git a/Engine/Source/Runtime/RenderCore/Private/Shader.cpp b/Engine/Source/Runtime/RenderCore/Private/Shader.cpp
--- a/Engine/Source/Runtime/RenderCore/Private/Shader.cpp
+++ b/Engine/Source/Runtime/RenderCore/Private/Shader.cpp
@@ -32,7 +32,6 @@
 #endif
 
 DEFINE_LOG_CATEGORY(LogShaders);
-DECLARE_LOG_CATEGORY_CLASS(LogShaderWarnings, Log, Log);
 
 IMPLEMENT_TYPE_LAYOUT(FShader);
 IMPLEMENT_TYPE_LAYOUT(FShaderParameterBindings);
@@ -93,13 +92,6 @@ static TAutoConsoleVariable<int32> CVarAllowCompilingThroughWorkers(
 	ECVF_ReadOnly
 	);
 
-static TAutoConsoleVariable<int32> CVarShaderCompilerEmitWarningsOnLoad(
-	TEXT("r.ShaderCompiler.EmitWarningsOnLoad"),
-	0,
-	TEXT("When 1, shader compiler warnings are emitted to the log for all shaders as they are loaded."),
-	ECVF_Default
-);
-
 static TAutoConsoleVariable<int32> CVarShadersForceDXC(
 	TEXT("r.Shaders.ForceDXC"),
 	0,
diff --git a/Engine/Source/Runtime/RenderCore/Private/ShaderResource.cpp b/Engine/Source/Runtime/RenderCore/Private/ShaderResource.cpp
--- a/Engine/Source/Runtime/RenderCore/Private/ShaderResource.cpp
+++ b/Engine/Source/Runtime/RenderCore/Private/ShaderResource.cpp
@@ -30,6 +30,8 @@
 #include "Interfaces/IShaderFormat.h"
 #endif
 
+DECLARE_LOG_CATEGORY_CLASS(LogShaderWarnings, Log, Log);
+
 int32 GShaderCompressionFormatChoice = 2;
 static FAutoConsoleVariableRef CVarShaderCompressionFormatChoice(
 	TEXT("r.Shaders.CompressionFormat"),
@@ -73,6 +75,14 @@ static FAutoConsoleVariableRef CVarShaderCompressionOodleAlgoChoice(
 	TEXT("  8 : Optimal4\n"),
 	ECVF_ReadOnly);
 
+static int32 GShaderCompilerEmitWarningsOnLoad = 0;
+static FAutoConsoleVariableRef CVarShaderCompilerEmitWarningsOnLoad(
+	TEXT("r.ShaderCompiler.EmitWarningsOnLoad"),
+	GShaderCompilerEmitWarningsOnLoad,
+	TEXT("When 1, shader compiler warnings are emitted to the log for all shaders as they are loaded."),
+	ECVF_Default
+);
+
 FName GetShaderCompressionFormat(const FName& ShaderFormat)
 {
 	// support an older developer-only CVar for compatibility and make it preempt
@@ -264,6 +274,10 @@ void FShaderMapResourceCode::Finalize()
 	Hasher.Final();
 	Hasher.GetHash(ResourceHash.Hash);
 	ApplyResourceStats(*this);
+
+#if WITH_EDITORONLY_DATA
+	LogShaderCompilerWarnings();
+#endif
 }
 
 uint32 FShaderMapResourceCode::GetSizeBytes() const
@@ -286,6 +300,11 @@ void FShaderMapResourceCode::AddShaderCompilerOutput(const FShaderCompilerOutput
 {
 #if WITH_EDITORONLY_DATA
 	AddPlatformDebugData(Output.PlatformDebugData);
+
+	for (const FShaderCompilerError& Error : Output.Errors)
+	{
+		CompilerWarnings.Add(Error.GetErrorString());
+	}
 #endif
 	AddShaderCode(Output.Target.GetFrequency(), Output.OutputHash, Output.ShaderCode);
 }
@@ -385,6 +404,20 @@ void FShaderMapResourceCode::AddPlatformDebugData(TConstArrayView<uint8> InPlatf
 		PlatformDebugData.EmplaceAt(Index, InPlatformDebugData.GetData(), InPlatformDebugData.Num());
 	}
 }
+
+void FShaderMapResourceCode::LogShaderCompilerWarnings()
+{
+	if (CompilerWarnings.Num() > 0 && GShaderCompilerEmitWarningsOnLoad != 0)
+	{
+		// Emit all the compiler warnings seen whilst serializing/loading this shader to the log.
+		// Since successfully compiled shaders are stored in the DDC, we'll get the compiler warnings
+		// even if we didn't compile the shader this run.
+		for (const FString& CompilerWarning : CompilerWarnings)
+		{
+			UE_LOG(LogShaderWarnings, Warning, TEXT("%s"), *CompilerWarning);
+		}
+	}
+}
 #endif // WITH_EDITORONLY_DATA
 
 void FShaderMapResourceCode::ToString(FStringBuilderBase& OutString) const
@@ -405,14 +438,22 @@ void FShaderMapResourceCode::Serialize(FArchive& Ar, bool bLoadedByCookedMateria
 	Ar << ShaderEntries;
 	check(ShaderEntries.Num() == ShaderHashes.Num());
 #if WITH_EDITORONLY_DATA
-	const bool bSerializePlatformData = !bLoadedByCookedMaterial && (!Ar.IsCooking() || Ar.CookingTarget()->HasEditorOnlyData());
-	if (bSerializePlatformData)
+	const bool bSerializeEditorOnlyData = !bLoadedByCookedMaterial && (!Ar.IsCooking() || Ar.CookingTarget()->HasEditorOnlyData());
+	if (bSerializeEditorOnlyData)
 	{
 		Ar << PlatformDebugDataHashes;
 		Ar << PlatformDebugData;
+		Ar << CompilerWarnings;
 	}
 #endif // WITH_EDITORONLY_DATA
 	ApplyResourceStats(*this);
+
+#if WITH_EDITORONLY_DATA
+	if (Ar.IsLoading())
+	{
+		LogShaderCompilerWarnings();
+	}
+#endif
 }
 
 #if WITH_EDITORONLY_DATA
diff --git a/Engine/Source/Runtime/RenderCore/Public/Shader.h b/Engine/Source/Runtime/RenderCore/Public/Shader.h
--- a/Engine/Source/Runtime/RenderCore/Public/Shader.h
+++ b/Engine/Source/Runtime/RenderCore/Public/Shader.h
@@ -396,6 +396,7 @@ public:
 	RENDERCORE_API void AddShaderCode(EShaderFrequency InFrequency, const FSHAHash& InHash, const FShaderCode& InCode);
 #if WITH_EDITORONLY_DATA
 	RENDERCORE_API void AddPlatformDebugData(TConstArrayView<uint8> InPlatformDebugData);
+	RENDERCORE_API void LogShaderCompilerWarnings();
 #endif
 
 	RENDERCORE_API void ToString(FStringBuilderBase& OutString) const;
@@ -406,6 +407,10 @@ public:
 #if WITH_EDITORONLY_DATA
 	TArray<TArray<uint8>> PlatformDebugData;
 	TArray<FSHAHash> PlatformDebugDataHashes;
+
+	/** An array of all the compiler warnings that were emitted when this shader was compiled.
+	 *  Does not contain errors since if there were any errors, this object wouldn't exist. */
+	TArray<FString> CompilerWarnings;
 #endif // WITH_EDITORONLY_DATA
 };
 	
