Sync RDG setup events prior to the PSO cache coalescing sync point.

diff --git a/Engine/Source/Runtime/RenderCore/Private/RenderGraphBuilder.cpp b/Engine/Source/Runtime/RenderCore/Private/RenderGraphBuilder.cpp
--- a/Engine/Source/Runtime/RenderCore/Private/RenderGraphBuilder.cpp
+++ b/Engine/Source/Runtime/RenderCore/Private/RenderGraphBuilder.cpp
@@ -1524,7 +1524,7 @@ void FRDGBuilder::Execute()
 	const FRDGPassHandle ProloguePassHandle = GetProloguePassHandle();
 	const FRDGPassHandle EpiloguePassHandle = GetEpiloguePassHandle();
 
-	FGraphEventArray AsyncCompileEvents = ParallelSetupEvents;
+	FGraphEventArray AsyncCompileEvents;
 
 	IF_RDG_ENABLE_DEBUG(UserValidation.ValidateExecuteBegin());
 	IF_RDG_ENABLE_DEBUG(GRDGAllowRHIAccess = true);
@@ -1674,6 +1674,11 @@ void FRDGBuilder::Execute()
 
 	CreatePassBarriers(bParallelExecuteEnabled ? &AsyncCompileEvents : nullptr);
 
+	if (!ParallelSetupEvents.IsEmpty())
+	{
+		FTaskGraphInterface::Get().WaitUntilTasksComplete(ParallelSetupEvents);
+	}
+
 	EndFlushResourcesRHI();
 
 	FGraphEventRef DispatchParallelExecuteEvent;
