VulkanRHI: Temporarily disable gpu stat accumulation during modal dialog's 'long frame' to prevent OOM from staging buffer accumulation.

diff --git a/Engine/Source/Runtime/RenderCore/Private/ProfilingDebugging/RealtimeGPUProfiler.cpp b/Engine/Source/Runtime/RenderCore/Private/ProfilingDebugging/RealtimeGPUProfiler.cpp
--- a/Engine/Source/Runtime/RenderCore/Private/ProfilingDebugging/RealtimeGPUProfiler.cpp
+++ b/Engine/Source/Runtime/RenderCore/Private/ProfilingDebugging/RealtimeGPUProfiler.cpp
@@ -1081,6 +1081,16 @@ void FRealtimeGPUProfiler::EndFrame(FRHICommandListImmediate& RHICmdList)
 	}
 }
 
+void FRealtimeGPUProfiler::SuspendFrame()
+{
+	if (!AreGPUStatsEnabled())
+	{
+		return;
+	}
+
+	bStatGatheringPaused = true;
+}
+
 FRealtimeGPUProfilerQuery FRealtimeGPUProfiler::PushEvent(FRHIGPUMask GPUMask, const FName& Name, const FName& StatName, const TCHAR* Description)
 {
 	check(IsInParallelRenderingThread());
diff --git a/Engine/Source/Runtime/RenderCore/Public/ProfilingDebugging/RealtimeGPUProfiler.h b/Engine/Source/Runtime/RenderCore/Public/ProfilingDebugging/RealtimeGPUProfiler.h
--- a/Engine/Source/Runtime/RenderCore/Public/ProfilingDebugging/RealtimeGPUProfiler.h
+++ b/Engine/Source/Runtime/RenderCore/Public/ProfilingDebugging/RealtimeGPUProfiler.h
@@ -284,6 +284,7 @@ struct RENDERCORE_API FBreadcrumbEvent
 #endif
  #define GPU_STATS_BEGINFRAME(RHICmdList) FRealtimeGPUProfiler::Get()->BeginFrame(RHICmdList);
  #define GPU_STATS_ENDFRAME(RHICmdList) FRealtimeGPUProfiler::Get()->EndFrame(RHICmdList);
+ #define GPU_STATS_SUSPENDFRAME() FRealtimeGPUProfiler::Get()->SuspendFrame();
 #else
  #define DECLARE_GPU_STAT(StatName)
  #define DECLARE_GPU_DRAWCALL_STAT(StatName)
@@ -296,6 +297,7 @@ struct RENDERCORE_API FBreadcrumbEvent
  #define SCOPED_GPU_STAT(RHICmdList, StatName) 
  #define GPU_STATS_BEGINFRAME(RHICmdList) 
  #define GPU_STATS_ENDFRAME(RHICmdList) 
+ #define GPU_STATS_SUSPENDFRAME()
 #endif
 
 RENDERCORE_API bool AreGPUStatsEnabled();
@@ -378,6 +380,7 @@ public:
 	/** Per-frame update */
 	RENDERCORE_API void BeginFrame(FRHICommandListImmediate& RHICmdList);
 	RENDERCORE_API void EndFrame(FRHICommandListImmediate& RHICmdList);
+	RENDERCORE_API void SuspendFrame();
 
 	/** Final cleanup */
 	UE_DEPRECATED(4.23, "Use FRealtimeGPUProfiler::SafeRelease() instead.")
diff --git a/Engine/Source/Runtime/SlateRHIRenderer/Private/SlateRHIRenderer.cpp b/Engine/Source/Runtime/SlateRHIRenderer/Private/SlateRHIRenderer.cpp
--- a/Engine/Source/Runtime/SlateRHIRenderer/Private/SlateRHIRenderer.cpp
+++ b/Engine/Source/Runtime/SlateRHIRenderer/Private/SlateRHIRenderer.cpp
@@ -1553,6 +1553,9 @@ void FSlateRHIRenderer::BeginFrame() const
 	   [](FRHICommandListImmediate& RHICmdList)
 	   {
 		   RHICmdList.BeginFrame();
+
+		   // Suspend stat gathering when running modal dialog 'fake' frame loops
+		   GPU_STATS_SUSPENDFRAME();
 	   }
 	);
 }
