Adding Cvar Flags that will make them automatically modify the KeyString when modified

 Cvars with this flag will modify the Shader Keystring for All Platforms
ECVF_GeneralShaderChange = 0x200,

Cvars with this flag will modify the Shader Keystring for Mobile Platforms
ECVF_MobileShaderChange = 0x400,

Cvars with this flag will modify the Shader Keystring for Desktop Platforms
ECVF_DesktopShaderChange = 0x800,

r.MobileHDR is using this system now. Others will follow.

diff --git a/Engine/Source/Runtime/Core/Private/HAL/ConsoleManager.cpp b/Engine/Source/Runtime/Core/Private/HAL/ConsoleManager.cpp
--- a/Engine/Source/Runtime/Core/Private/HAL/ConsoleManager.cpp
+++ b/Engine/Source/Runtime/Core/Private/HAL/ConsoleManager.cpp
@@ -2555,7 +2555,7 @@ static TAutoConsoleVariable<int32> CVarMobileHDR(
 	1,
 	TEXT("0: Mobile renders in LDR gamma space. (suggested for unlit games targeting low-end phones)\n"
 		 "1: Mobile renders in HDR linear space. (default)"),
-	ECVF_RenderThreadSafe | ECVF_ReadOnly);
+	ECVF_RenderThreadSafe | ECVF_ReadOnly | ECVF_MobileShaderChange);
 
 static TAutoConsoleVariable<int32> CVarMobileShadingPath(
 	TEXT("r.Mobile.ShadingPath"),
diff --git a/Engine/Source/Runtime/Core/Private/Misc/CoreGlobals.cpp b/Engine/Source/Runtime/Core/Private/Misc/CoreGlobals.cpp
--- a/Engine/Source/Runtime/Core/Private/Misc/CoreGlobals.cpp
+++ b/Engine/Source/Runtime/Core/Private/Misc/CoreGlobals.cpp
@@ -37,7 +37,9 @@ IMPLEMENT_MODULE( FCoreModule, Core );
 /*-----------------------------------------------------------------------------
 	Global variables.
 -----------------------------------------------------------------------------*/
-
+TArray<const FAutoConsoleObject*> FAutoConsoleObject::GeneralShaderChangeCvars;
+TArray<const FAutoConsoleObject*> FAutoConsoleObject::MobileShaderChangeCvars;
+TArray<const FAutoConsoleObject*> FAutoConsoleObject::DesktopShaderChangeCvars;
 CORE_API FFeedbackContext*	GWarn						= nullptr;		/* User interaction and non critical warnings */
 FConfigCacheIni*			GConfig						= nullptr;		/* Configuration database cache */
 ITransaction*				GUndo						= nullptr;		/* Transaction tracker, non-NULL when a transaction is in progress */
diff --git a/Engine/Source/Runtime/Core/Public/HAL/IConsoleManager.h b/Engine/Source/Runtime/Core/Public/HAL/IConsoleManager.h
--- a/Engine/Source/Runtime/Core/Public/HAL/IConsoleManager.h
+++ b/Engine/Source/Runtime/Core/Public/HAL/IConsoleManager.h
@@ -93,6 +93,15 @@ enum EConsoleVariableFlags
 	/* CVars with this flag will be included in the device profile previews. */
 	ECVF_Preview = 0x100,
 
+	/* Cvars with this flag will modify the Shader Keystring for All Platforms*/
+	ECVF_GeneralShaderChange = 0x200,
+
+	/* Cvars with this flag will modify the Shader Keystring for Mobile Platforms*/
+	ECVF_MobileShaderChange = 0x400,
+
+	/* Cvars with this flag will modify the Shader Keystring for Desktop Platforms*/
+	ECVF_DesktopShaderChange = 0x800,
+
 	// ------------------------------------------------
 
 	/* Set flags */
@@ -1046,6 +1055,18 @@ protected:
 		: Target(InTarget)
 	{
 		check(Target);
+		if (Target->TestFlags(ECVF_GeneralShaderChange))
+		{
+			GeneralShaderChangeCvars.Add(this);
+		}
+		else if (Target->TestFlags(ECVF_MobileShaderChange))
+		{
+			MobileShaderChangeCvars.Add(this);
+		}
+		else if (Target->TestFlags(ECVF_DesktopShaderChange))
+		{
+			DesktopShaderChangeCvars.Add(this);
+		}
 	}
 	/** Destructor, removes the console object **/
 	virtual ~FAutoConsoleObject()
@@ -1054,6 +1075,10 @@ protected:
 	}
 
 public:
+	static TArray<const FAutoConsoleObject*> GeneralShaderChangeCvars;
+	static TArray<const FAutoConsoleObject*> MobileShaderChangeCvars;
+	static TArray<const FAutoConsoleObject*> DesktopShaderChangeCvars;
+
 	/** returns the contained console object as an IConsoleVariable **/
 	FORCEINLINE IConsoleVariable* AsVariable()
 	{
diff --git a/Engine/Source/Runtime/RenderCore/Private/Shader.cpp b/Engine/Source/Runtime/RenderCore/Private/Shader.cpp
--- a/Engine/Source/Runtime/RenderCore/Private/Shader.cpp
+++ b/Engine/Source/Runtime/RenderCore/Private/Shader.cpp
@@ -1349,6 +1349,37 @@ void ShaderMapAppendKeyString(EShaderPlatform Platform, FString& KeyString)
 {
 	const FName ShaderFormatName = LegacyShaderPlatformToShaderFormat(Platform);
 
+	for (const FAutoConsoleObject* ConsoleObject : FAutoConsoleObject::GeneralShaderChangeCvars)
+	{
+		FString ConsoleObjectName = IConsoleManager::Get().FindConsoleObjectName(ConsoleObject->AsVariable());
+		KeyString += TEXT("_");
+		KeyString += ConsoleObjectName;
+		KeyString += TEXT("_");
+		KeyString += ConsoleObject->AsVariable()->GetString();
+	}
+	if (IsMobilePlatform(Platform))
+	{
+		for (const FAutoConsoleObject* ConsoleObject : FAutoConsoleObject::MobileShaderChangeCvars)
+		{
+			FString ConsoleObjectName = IConsoleManager::Get().FindConsoleObjectName(ConsoleObject->AsVariable());
+			KeyString += TEXT("_");
+			KeyString += ConsoleObjectName;
+			KeyString += TEXT("_");
+			KeyString += ConsoleObject->AsVariable()->GetString();
+		}
+	}
+	else if (IsConsolePlatform(Platform))
+	{
+		for (const FAutoConsoleObject* ConsoleObject : FAutoConsoleObject::DesktopShaderChangeCvars)
+		{
+			FString ConsoleObjectName = IConsoleManager::Get().FindConsoleObjectName(ConsoleObject->AsVariable());
+			KeyString += TEXT("_");
+			KeyString += ConsoleObjectName;
+			KeyString += TEXT("_");
+			KeyString += ConsoleObject->AsVariable()->GetString();
+		}
+	}
+
 	// Globals that should cause all shaders to recompile when changed must be appended to the key here
 	// Key should be kept as short as possible while being somewhat human readable for debugging
 
@@ -1545,11 +1576,6 @@ void ShaderMapAppendKeyString(EShaderPlatform Platform, FString& KeyString)
 			KeyString += bMobileGpuScene ? TEXT("_MobGPUSc") : TEXT("");
 		}
 
-		{
-			static IConsoleVariable* CVar = IConsoleManager::Get().FindConsoleVariable(TEXT("r.MobileHDR"));
-			KeyString += (CVar && CVar->GetInt() != 0) ? TEXT("_MobileHDR") : TEXT("");
-		}
-
 		{
 			static IConsoleVariable* MobileShadingPathCVar = IConsoleManager::Get().FindConsoleVariable(TEXT("r.Mobile.ShadingPath"));			
 			if (MobileShadingPathCVar)
