Add a version of UseVirtualTexturing() that takes a shader platform.
A future goal is to deprecate the version that takes a feature level.

diff --git a/Engine/Source/Runtime/RenderCore/Private/RenderUtils.cpp b/Engine/Source/Runtime/RenderCore/Private/RenderUtils.cpp
--- a/Engine/Source/Runtime/RenderCore/Private/RenderUtils.cpp
+++ b/Engine/Source/Runtime/RenderCore/Private/RenderUtils.cpp
@@ -1727,7 +1727,7 @@ RENDERCORE_API void QuantizeSceneBufferSize(const FIntPoint& InBufferSize, FIntP
 	OutBufferSize.Y = (InBufferSize.Y + DividableBy - 1) & Mask;
 }
 
-RENDERCORE_API bool UseVirtualTexturing(const FStaticFeatureLevel InFeatureLevel, const ITargetPlatform* TargetPlatform)
+RENDERCORE_API bool UseVirtualTexturing(const EShaderPlatform InShaderPlatform, const ITargetPlatform* TargetPlatform)
 {
 #if PLATFORM_SUPPORTS_VIRTUAL_TEXTURE_STREAMING
 	if (!FPlatformProperties::SupportsVirtualTextureStreaming())
@@ -1745,7 +1745,7 @@ RENDERCORE_API bool UseVirtualTexturing(const FStaticFeatureLevel InFeatureLevel
 
 	// mobile needs an additional switch to enable VT		
 	static const auto CVarMobileVirtualTexture = IConsoleManager::Get().FindTConsoleVariableDataInt(TEXT("r.Mobile.VirtualTextures"));
-	if (InFeatureLevel == ERHIFeatureLevel::ES3_1 && CVarMobileVirtualTexture->GetValueOnAnyThread() == 0)
+	if (IsMobilePlatform(InShaderPlatform) && CVarMobileVirtualTexture->GetValueOnAnyThread() == 0)
 	{
 		return false;
 	}
@@ -1756,6 +1756,11 @@ RENDERCORE_API bool UseVirtualTexturing(const FStaticFeatureLevel InFeatureLevel
 #endif
 }
 
+RENDERCORE_API bool UseVirtualTexturing(const FStaticFeatureLevel InFeatureLevel, const ITargetPlatform* TargetPlatform)
+{
+	return UseVirtualTexturing(GetFeatureLevelShaderPlatform(InFeatureLevel), TargetPlatform);
+}
+
 RENDERCORE_API bool UseVirtualTextureLightmap(const FStaticFeatureLevel InFeatureLevel, const ITargetPlatform* TargetPlatform)
 {
 	static const auto CVar = IConsoleManager::Get().FindTConsoleVariableDataInt(TEXT("r.VirtualTexturedLightmaps"));
diff --git a/Engine/Source/Runtime/RenderCore/Public/RenderUtils.h b/Engine/Source/Runtime/RenderCore/Public/RenderUtils.h
--- a/Engine/Source/Runtime/RenderCore/Public/RenderUtils.h
+++ b/Engine/Source/Runtime/RenderCore/Public/RenderUtils.h
@@ -706,11 +706,12 @@ RENDERCORE_API FBufferRHIRef& GetUnitCubeAABBVertexBuffer();
 RENDERCORE_API void QuantizeSceneBufferSize(const FIntPoint& InBufferSize, FIntPoint& OutBufferSize);
 
 /**
-*	Checks if virtual texturing enabled and supported
+* Checks if virtual texturing enabled and supported
+* todo: Deprecate the version of the function that takes FStaticFeatureLevel
 */
+RENDERCORE_API bool UseVirtualTexturing(const EShaderPlatform InShaderPlatform, const class ITargetPlatform* TargetPlatform = nullptr);
 RENDERCORE_API bool UseVirtualTexturing(const FStaticFeatureLevel InFeatureLevel, const class ITargetPlatform* TargetPlatform = nullptr);
 
-
 RENDERCORE_API bool DoesPlatformSupportNanite(EShaderPlatform Platform, bool bCheckForProjectSetting = true);
 
 inline bool NaniteAtomicsSupported()
