[Backout] - CL20737284 & CL20738687

- testing
  - compiled and cooked QAGame locally for Win64

#rnx

Original CL Desc for CL20737284
-----------------------------------------------------------------
Add a version of UseVirtualTexturing() that takes a shader platform.
A future goal is to deprecate the version that takes a feature level.

Original CL Desc for CL20738687
-----------------------------------------------------------------
Virtual texturing forces depth pass on.
This fixes bug when using any masked materials that sample VT.

diff --git a/Engine/Source/Runtime/RenderCore/Private/RenderUtils.cpp b/Engine/Source/Runtime/RenderCore/Private/RenderUtils.cpp
--- a/Engine/Source/Runtime/RenderCore/Private/RenderUtils.cpp
+++ b/Engine/Source/Runtime/RenderCore/Private/RenderUtils.cpp
@@ -1727,7 +1727,7 @@ RENDERCORE_API void QuantizeSceneBufferSize(const FIntPoint& InBufferSize, FIntP
 	OutBufferSize.Y = (InBufferSize.Y + DividableBy - 1) & Mask;
 }
 
-RENDERCORE_API bool UseVirtualTexturing(const EShaderPlatform InShaderPlatform, const ITargetPlatform* TargetPlatform)
+RENDERCORE_API bool UseVirtualTexturing(const FStaticFeatureLevel InFeatureLevel, const ITargetPlatform* TargetPlatform)
 {
 #if PLATFORM_SUPPORTS_VIRTUAL_TEXTURE_STREAMING
 	if (!FPlatformProperties::SupportsVirtualTextureStreaming())
@@ -1745,7 +1745,7 @@ RENDERCORE_API bool UseVirtualTexturing(const EShaderPlatform InShaderPlatform,
 
 	// mobile needs an additional switch to enable VT		
 	static const auto CVarMobileVirtualTexture = IConsoleManager::Get().FindTConsoleVariableDataInt(TEXT("r.Mobile.VirtualTextures"));
-	if (IsMobilePlatform(InShaderPlatform) && CVarMobileVirtualTexture->GetValueOnAnyThread() == 0)
+	if (InFeatureLevel == ERHIFeatureLevel::ES3_1 && CVarMobileVirtualTexture->GetValueOnAnyThread() == 0)
 	{
 		return false;
 	}
@@ -1756,11 +1756,6 @@ RENDERCORE_API bool UseVirtualTexturing(const EShaderPlatform InShaderPlatform,
 #endif
 }
 
-RENDERCORE_API bool UseVirtualTexturing(const FStaticFeatureLevel InFeatureLevel, const ITargetPlatform* TargetPlatform)
-{
-	return UseVirtualTexturing(GetFeatureLevelShaderPlatform(InFeatureLevel), TargetPlatform);
-}
-
 RENDERCORE_API bool UseVirtualTextureLightmap(const FStaticFeatureLevel InFeatureLevel, const ITargetPlatform* TargetPlatform)
 {
 	static const auto CVar = IConsoleManager::Get().FindTConsoleVariableDataInt(TEXT("r.VirtualTexturedLightmaps"));
diff --git a/Engine/Source/Runtime/RenderCore/Public/RenderUtils.h b/Engine/Source/Runtime/RenderCore/Public/RenderUtils.h
--- a/Engine/Source/Runtime/RenderCore/Public/RenderUtils.h
+++ b/Engine/Source/Runtime/RenderCore/Public/RenderUtils.h
@@ -706,12 +706,11 @@ RENDERCORE_API FBufferRHIRef& GetUnitCubeAABBVertexBuffer();
 RENDERCORE_API void QuantizeSceneBufferSize(const FIntPoint& InBufferSize, FIntPoint& OutBufferSize);
 
 /**
-* Checks if virtual texturing enabled and supported
-* todo: Deprecate the version of the function that takes FStaticFeatureLevel
+*	Checks if virtual texturing enabled and supported
 */
-RENDERCORE_API bool UseVirtualTexturing(const EShaderPlatform InShaderPlatform, const class ITargetPlatform* TargetPlatform = nullptr);
 RENDERCORE_API bool UseVirtualTexturing(const FStaticFeatureLevel InFeatureLevel, const class ITargetPlatform* TargetPlatform = nullptr);
 
+
 RENDERCORE_API bool DoesPlatformSupportNanite(EShaderPlatform Platform, bool bCheckForProjectSetting = true);
 
 inline bool NaniteAtomicsSupported()
diff --git a/Engine/Source/Runtime/Renderer/Private/DeferredShadingRenderer.cpp b/Engine/Source/Runtime/Renderer/Private/DeferredShadingRenderer.cpp
--- a/Engine/Source/Runtime/Renderer/Private/DeferredShadingRenderer.cpp
+++ b/Engine/Source/Runtime/Renderer/Private/DeferredShadingRenderer.cpp
@@ -355,19 +355,18 @@ const TCHAR* GetDepthPassReason(bool bDitheredLODTransitionsUseStencil, EShaderP
 		return TEXT("(Forced by ForwardShading)");
 	}
 
-	if (UseNanite(ShaderPlatform))
+	bool bUseNanite = UseNanite(ShaderPlatform);
+
+	if (bUseNanite)
 	{
 		return TEXT("(Forced by Nanite)");
 	}
 
-	if (IsUsingDBuffers(ShaderPlatform))
-	{
-		return TEXT("(Forced by DBuffer)");
-	}
+	bool bDBufferAllowed = IsUsingDBuffers(ShaderPlatform);
 
-	if (UseVirtualTexturing(ShaderPlatform))
+	if (bDBufferAllowed)
 	{
-		return TEXT("(Forced by VirtualTexture)");
+		return TEXT("(Forced by DBuffer)");
 	}
 
 	if (bDitheredLODTransitionsUseStencil)
diff --git a/Engine/Source/Runtime/Renderer/Private/RendererScene.cpp b/Engine/Source/Runtime/Renderer/Private/RendererScene.cpp
--- a/Engine/Source/Runtime/Renderer/Private/RendererScene.cpp
+++ b/Engine/Source/Runtime/Renderer/Private/RendererScene.cpp
@@ -3893,19 +3893,19 @@ void FScene::Release()
 bool ShouldForceFullDepthPass(EShaderPlatform ShaderPlatform)
 {
 	const bool bNaniteEnabled = UseNanite(ShaderPlatform);
+
 	const bool bDBufferAllowed = IsUsingDBuffers(ShaderPlatform);
-	const bool bVirtualTextureEnabled = UseVirtualTexturing(ShaderPlatform);
 
 	static const auto StencilLODDitherCVar = IConsoleManager::Get().FindTConsoleVariableDataInt(TEXT("r.StencilForLODDither"));
 	const bool bStencilLODDither = StencilLODDitherCVar->GetValueOnAnyThread() != 0;
 
 	static const auto AOComputeCVar = IConsoleManager::Get().FindTConsoleVariableDataInt(TEXT("r.AmbientOcclusion.Compute"));
-	const bool bAOCompute = AOComputeCVar->GetValueOnAnyThread() > 0;
+	const bool bAOCompute		= AOComputeCVar->GetValueOnAnyThread() > 0;
 
 	const bool bEarlyZMaterialMasking = MaskedInEarlyPass(ShaderPlatform);
 
 	// Note: ShouldForceFullDepthPass affects which static draw lists meshes go into, so nothing it depends on can change at runtime, unless you do a FGlobalComponentRecreateRenderStateContext to propagate the cvar change
-	return bNaniteEnabled || bAOCompute || bDBufferAllowed || bVirtualTextureEnabled || bStencilLODDither || bEarlyZMaterialMasking || IsForwardShadingEnabled(ShaderPlatform) || IsUsingSelectiveBasePassOutputs(ShaderPlatform);
+	return bNaniteEnabled || bAOCompute || bDBufferAllowed || bStencilLODDither || bEarlyZMaterialMasking || IsForwardShadingEnabled(ShaderPlatform) || IsUsingSelectiveBasePassOutputs(ShaderPlatform);
 }
 
 void FScene::UpdateEarlyZPassMode()
