Low Input latency mode fix

-Allow the frame sync to be kicked from the render thread when using the gtsync==2 (low input latency mode). On some platforms, the RHI thread will block on present. Putting the RHI as the trigger index will block the GameThread until the present is finished (vblank). In low input latency mode, some consoles uses the RHIOffsetThread to kick of the Gamethread, so we dont want it to block on present.
-Use the present mode to submit frame tokens in low input latency mode. Using GT_wait or GT_early introduce stuttering in the frame time since the gamethread will sync on the present when its double buffered. Even if we increase the frame interval (3 or 4), the WaitFrameEventX will eventually block the gamethread when the second display buffer is in used. Note that we delay the barrier transition(RT->Display) until we really use the backbuffer, so we dont technically need to throttle on the frame token, since we are in pseudo triple buffer mode. Issuing GT_Wait/GT_early mode in double buffer mode will always produce bad frame pacing when we use the FRHIFrameOffsetThread approach.
-Set the frame offset interval to zero to sync the os input pooling with the vblank.

diff --git a/Engine/Source/Runtime/RenderCore/Private/RenderingThread.cpp b/Engine/Source/Runtime/RenderCore/Private/RenderingThread.cpp
--- a/Engine/Source/Runtime/RenderCore/Private/RenderingThread.cpp
+++ b/Engine/Source/Runtime/RenderCore/Private/RenderingThread.cpp
@@ -1011,6 +1011,14 @@ TAutoConsoleVariable<int32> CVarGTSyncType(
 	TEXT(" 2 - Sync the game thread with the GPU swap chain flip (only on supported platforms).\n"),
 	ECVF_Default);
 
+TAutoConsoleVariable<bool> CVarAllowRHITriggerThread(
+	TEXT("r.AllowRHITriggerThread"),
+	true,
+	TEXT("In low latency mode, use the rhi thread to trigger the frame sync.")
+	TEXT(" true (default).\n")
+	TEXT(" false.\n"),
+	ECVF_Default);
+
 FRHICOMMAND_MACRO(FRHISyncFrameCommand)
 {
 	FGraphEventRef GraphEvent;
@@ -1065,7 +1073,7 @@ void FRenderCommandFence::BeginFence(bool bSyncToRHIAndGPU)
 			check(CVarVsync != nullptr);
 
 			if ( GTSyncType == 0 || CVarVsync->GetInt() == 0 )
-		{
+			{
 				bSyncToRHIAndGPU = false;
 			}
 		}
@@ -1075,8 +1083,10 @@ void FRenderCommandFence::BeginFence(bool bSyncToRHIAndGPU)
 		{
 			if (IsRHIThreadRunning())
 			{
-				// Change trigger thread to RHI
-				TriggerThreadIndex = ENamedThreads::RHIThread;
+				// Potentially change trigger thread to RHI
+				// On some platform RHI thread will block on present. Putting the RHI as the trigger index will block the GameThread until the present is finished.
+				// In low input latency mode, some consoles uses the RHIOffsetThread to kick of the Gamethread, so we dont want it to block on present.
+				TriggerThreadIndex = (GTSyncType == 2 && !CVarAllowRHITriggerThread.GetValueOnAnyThread()) ? TriggerThreadIndex : ENamedThreads::RHIThread;
 			}
 			
 			// Create a task graph event which we can pass to the render or RHI threads.
