Fixed long-living tasks that drive total mem usage up. Tasks use ConcurrentLinearAllocator that doesn't like long-living allocs as a single alloc can hold an entire mem page occupied. A completion handle uses a different allocator that doesn't have this problem and it can be used for waiting for its parent task completion. A completion handle is used only in rare cases where it's difficult or impossible to shorten task handle lifetime, as in this CL. All other cases were fixed by resetting a task handle when it's completed and thus releasing its memory.

diff --git a/Engine/Source/Runtime/Engine/Private/VT/VirtualTextureTranscodeCache.cpp b/Engine/Source/Runtime/Engine/Private/VT/VirtualTextureTranscodeCache.cpp
--- a/Engine/Source/Runtime/Engine/Private/VT/VirtualTextureTranscodeCache.cpp
+++ b/Engine/Source/Runtime/Engine/Private/VT/VirtualTextureTranscodeCache.cpp
@@ -422,7 +422,9 @@ FVTTranscodeTileHandle FVirtualTextureTranscodeCache::SubmitTask(
 
 	if (bNeedToLaunchTask)
 	{
-		TaskEntry.GraphEvent = TGraphTask<FTranscodeTask>::CreateTask(Prerequisites).ConstructAndDispatchWhenReady(StagingBuffer, InParams);
+		FGraphEventRef Task = TGraphTask<FTranscodeTask>::CreateTask(Prerequisites).ConstructAndDispatchWhenReady(StagingBuffer, InParams);
+		// this reference can live long, store completion handle instead of a reference to the task to reduce peak mem usage 
+		TaskEntry.GraphEvent = Task->CreateCompletionHandle();
 	}
 	else
 	{
diff --git a/Engine/Source/Runtime/Experimental/Chaos/Private/Chaos/PBDRigidsEvolution.cpp b/Engine/Source/Runtime/Experimental/Chaos/Private/Chaos/PBDRigidsEvolution.cpp
--- a/Engine/Source/Runtime/Experimental/Chaos/Private/Chaos/PBDRigidsEvolution.cpp
+++ b/Engine/Source/Runtime/Experimental/Chaos/Private/Chaos/PBDRigidsEvolution.cpp
@@ -801,7 +801,9 @@ namespace Chaos
 			if (bCanStartAsyncTasks)
 			{
 				// we run the task for both starting a new accel structure as well as for the time-slicing
-				AccelerationStructureTaskComplete = TGraphTask<FChaosAccelerationStructureTask>::CreateTask().ConstructAndDispatchWhenReady(*SpatialCollectionFactory, SpatialAccelerationCache, AsyncInternalAcceleration, AsyncExternalAcceleration, ForceFullBuild, bIsSingleThreaded, bNeedsReset);
+				FGraphEventRef Task = TGraphTask<FChaosAccelerationStructureTask>::CreateTask().ConstructAndDispatchWhenReady(*SpatialCollectionFactory, SpatialAccelerationCache, AsyncInternalAcceleration, AsyncExternalAcceleration, ForceFullBuild, bIsSingleThreaded, bNeedsReset);
+				// this reference can live long, store completion handle instead of a reference to the task to reduce peak mem usage 
+				AccelerationStructureTaskComplete = Task->CreateCompletionHandle();
 			}
 		}
 		else
diff --git a/Engine/Source/Runtime/RenderCore/Private/RenderingThread.cpp b/Engine/Source/Runtime/RenderCore/Private/RenderingThread.cpp
--- a/Engine/Source/Runtime/RenderCore/Private/RenderingThread.cpp
+++ b/Engine/Source/Runtime/RenderCore/Private/RenderingThread.cpp
@@ -981,8 +981,11 @@ void StartRenderCommandFenceBundler()
 	STAT_FNullGraphTask_FenceRenderCommandBundled,
 		STATGROUP_TaskGraphTasks);
 
-	BundledCompletionEvent = TGraphTask<FNullGraphTask>::CreateTask(&Prereqs, ENamedThreads::GameThread).ConstructAndDispatchWhenReady(
+	FGraphEventRef Task = TGraphTask<FNullGraphTask>::CreateTask(&Prereqs, ENamedThreads::GameThread).ConstructAndDispatchWhenReady(
 		GET_STATID(STAT_FNullGraphTask_FenceRenderCommandBundled), ENamedThreads::GetRenderThread());
+	// this reference can live long, store completion handle instead of a reference to the task to reduce peak mem usage 
+	BundledCompletionEvent = Task->CreateCompletionHandle();
+	
 
 	StartBatchedRelease();
 }
@@ -995,7 +998,7 @@ void StopRenderCommandFenceBundler()
 	}
 
 	EndBatchedRelease();
-	check(IsInGameThread() && BundledCompletionEvent.GetReference() && !BundledCompletionEvent->IsComplete() && BundledCompletionEventPrereq.GetReference() && !BundledCompletionEventPrereq->IsComplete()); // can't use this in a nested fashion
+	checkf(IsInGameThread() && BundledCompletionEvent.GetReference() && !BundledCompletionEvent->IsComplete() && BundledCompletionEventPrereq.GetReference() && !BundledCompletionEventPrereq->IsComplete(), TEXT("IsInGameThread: %d, BundledCompletionEvent is completed: %d, BundledCompletionEventPrereq is completed: %d"), IsInGameThread(), BundledCompletionEvent->IsComplete(), BundledCompletionEventPrereq->IsComplete()); // can't use this in a nested fashion
 	BundledCompletionEventPrereq->DispatchSubsequents();
 	BundledCompletionEventPrereq = nullptr;
 	BundledCompletionEvent = nullptr;
@@ -1115,8 +1118,8 @@ void FRenderCommandFence::BeginFence(bool bSyncToRHIAndGPU)
 			STAT_FNullGraphTask_FenceRenderCommand,
 				STATGROUP_TaskGraphTasks);
 
-			CompletionEvent = TGraphTask<FNullGraphTask>::CreateTask(NULL, ENamedThreads::GameThread).ConstructAndDispatchWhenReady(
-				GET_STATID(STAT_FNullGraphTask_FenceRenderCommand), ENamedThreads::GetRenderThread());
+			CompletionEvent = FGraphEvent::CreateGraphEvent();
+			FFunctionGraphTask::CreateAndDispatchWhenReady([CompletionEvent = CompletionEvent] { CompletionEvent->DispatchSubsequents(); }, GET_STATID(STAT_FNullGraphTask_FenceRenderCommand), nullptr, ENamedThreads::GetRenderThread());
 		}
 	}
 }
