Undo the backout changelist 20750841
This resubmits the following original changelists (with fixes).

Original CL Desc for CL20737284 -----------------------------------------------------------------
Add a version of UseVirtualTexturing() that takes a shader platform.
A future goal is to deprecate the version that takes a feature level.

Original CL Desc for CL20738687 -----------------------------------------------------------------
Virtual texturing forces depth pass on.
This fixes bug when using any masked materials that sample VT.

diff --git a/Engine/Source/Runtime/RenderCore/Private/RenderUtils.cpp b/Engine/Source/Runtime/RenderCore/Private/RenderUtils.cpp
--- a/Engine/Source/Runtime/RenderCore/Private/RenderUtils.cpp
+++ b/Engine/Source/Runtime/RenderCore/Private/RenderUtils.cpp
@@ -1727,7 +1727,7 @@ RENDERCORE_API void QuantizeSceneBufferSize(const FIntPoint& InBufferSize, FIntP
 	OutBufferSize.Y = (InBufferSize.Y + DividableBy - 1) & Mask;
 }
 
-RENDERCORE_API bool UseVirtualTexturing(const FStaticFeatureLevel InFeatureLevel, const ITargetPlatform* TargetPlatform)
+bool UseVirtualTexturing(bool bIsMobilePlatform, const ITargetPlatform* TargetPlatform)
 {
 #if PLATFORM_SUPPORTS_VIRTUAL_TEXTURE_STREAMING
 	if (!FPlatformProperties::SupportsVirtualTextureStreaming())
@@ -1745,7 +1745,7 @@ RENDERCORE_API bool UseVirtualTexturing(const FStaticFeatureLevel InFeatureLevel
 
 	// mobile needs an additional switch to enable VT		
 	static const auto CVarMobileVirtualTexture = IConsoleManager::Get().FindTConsoleVariableDataInt(TEXT("r.Mobile.VirtualTextures"));
-	if (InFeatureLevel == ERHIFeatureLevel::ES3_1 && CVarMobileVirtualTexture->GetValueOnAnyThread() == 0)
+	if (bIsMobilePlatform && CVarMobileVirtualTexture->GetValueOnAnyThread() == 0)
 	{
 		return false;
 	}
@@ -1756,6 +1756,16 @@ RENDERCORE_API bool UseVirtualTexturing(const FStaticFeatureLevel InFeatureLevel
 #endif
 }
 
+RENDERCORE_API bool UseVirtualTexturing(const EShaderPlatform InShaderPlatform, const ITargetPlatform* TargetPlatform)
+{
+	return UseVirtualTexturing(IsMobilePlatform(InShaderPlatform), TargetPlatform);
+}
+
+RENDERCORE_API bool UseVirtualTexturing(const FStaticFeatureLevel InFeatureLevel, const ITargetPlatform* TargetPlatform)
+{
+	return UseVirtualTexturing(InFeatureLevel == ERHIFeatureLevel::ES3_1, TargetPlatform);
+}
+
 RENDERCORE_API bool UseVirtualTextureLightmap(const FStaticFeatureLevel InFeatureLevel, const ITargetPlatform* TargetPlatform)
 {
 	static const auto CVar = IConsoleManager::Get().FindTConsoleVariableDataInt(TEXT("r.VirtualTexturedLightmaps"));
diff --git a/Engine/Source/Runtime/RenderCore/Public/RenderUtils.h b/Engine/Source/Runtime/RenderCore/Public/RenderUtils.h
--- a/Engine/Source/Runtime/RenderCore/Public/RenderUtils.h
+++ b/Engine/Source/Runtime/RenderCore/Public/RenderUtils.h
@@ -706,11 +706,12 @@ RENDERCORE_API FBufferRHIRef& GetUnitCubeAABBVertexBuffer();
 RENDERCORE_API void QuantizeSceneBufferSize(const FIntPoint& InBufferSize, FIntPoint& OutBufferSize);
 
 /**
-*	Checks if virtual texturing enabled and supported
+* Checks if virtual texturing enabled and supported
+* todo: Deprecate the version of the function that takes FStaticFeatureLevel
 */
+RENDERCORE_API bool UseVirtualTexturing(const EShaderPlatform InShaderPlatform, const class ITargetPlatform* TargetPlatform = nullptr);
 RENDERCORE_API bool UseVirtualTexturing(const FStaticFeatureLevel InFeatureLevel, const class ITargetPlatform* TargetPlatform = nullptr);
 
-
 RENDERCORE_API bool DoesPlatformSupportNanite(EShaderPlatform Platform, bool bCheckForProjectSetting = true);
 
 inline bool NaniteAtomicsSupported()
diff --git a/Engine/Source/Runtime/Renderer/Private/DeferredShadingRenderer.cpp b/Engine/Source/Runtime/Renderer/Private/DeferredShadingRenderer.cpp
--- a/Engine/Source/Runtime/Renderer/Private/DeferredShadingRenderer.cpp
+++ b/Engine/Source/Runtime/Renderer/Private/DeferredShadingRenderer.cpp
@@ -355,20 +355,21 @@ const TCHAR* GetDepthPassReason(bool bDitheredLODTransitionsUseStencil, EShaderP
 		return TEXT("(Forced by ForwardShading)");
 	}
 
-	bool bUseNanite = UseNanite(ShaderPlatform);
-
-	if (bUseNanite)
+	if (UseNanite(ShaderPlatform))
 	{
 		return TEXT("(Forced by Nanite)");
 	}
 
-	bool bDBufferAllowed = IsUsingDBuffers(ShaderPlatform);
-
-	if (bDBufferAllowed)
+	if (IsUsingDBuffers(ShaderPlatform))
 	{
 		return TEXT("(Forced by DBuffer)");
 	}
 
+	if (UseVirtualTexturing(ShaderPlatform))
+	{
+		return TEXT("(Forced by VirtualTexture)");
+	}
+
 	if (bDitheredLODTransitionsUseStencil)
 	{
 		return TEXT("(Forced by StencilLODDither)");
diff --git a/Engine/Source/Runtime/Renderer/Private/RendererScene.cpp b/Engine/Source/Runtime/Renderer/Private/RendererScene.cpp
--- a/Engine/Source/Runtime/Renderer/Private/RendererScene.cpp
+++ b/Engine/Source/Runtime/Renderer/Private/RendererScene.cpp
@@ -3893,19 +3893,19 @@ void FScene::Release()
 bool ShouldForceFullDepthPass(EShaderPlatform ShaderPlatform)
 {
 	const bool bNaniteEnabled = UseNanite(ShaderPlatform);
-
 	const bool bDBufferAllowed = IsUsingDBuffers(ShaderPlatform);
+	const bool bVirtualTextureEnabled = UseVirtualTexturing(ShaderPlatform);
 
 	static const auto StencilLODDitherCVar = IConsoleManager::Get().FindTConsoleVariableDataInt(TEXT("r.StencilForLODDither"));
 	const bool bStencilLODDither = StencilLODDitherCVar->GetValueOnAnyThread() != 0;
 
 	static const auto AOComputeCVar = IConsoleManager::Get().FindTConsoleVariableDataInt(TEXT("r.AmbientOcclusion.Compute"));
-	const bool bAOCompute		= AOComputeCVar->GetValueOnAnyThread() > 0;
+	const bool bAOCompute = AOComputeCVar->GetValueOnAnyThread() > 0;
 
 	const bool bEarlyZMaterialMasking = MaskedInEarlyPass(ShaderPlatform);
 
 	// Note: ShouldForceFullDepthPass affects which static draw lists meshes go into, so nothing it depends on can change at runtime, unless you do a FGlobalComponentRecreateRenderStateContext to propagate the cvar change
-	return bNaniteEnabled || bAOCompute || bDBufferAllowed || bStencilLODDither || bEarlyZMaterialMasking || IsForwardShadingEnabled(ShaderPlatform) || IsUsingSelectiveBasePassOutputs(ShaderPlatform);
+	return bNaniteEnabled || bAOCompute || bDBufferAllowed || bVirtualTextureEnabled || bStencilLODDither || bEarlyZMaterialMasking || IsForwardShadingEnabled(ShaderPlatform) || IsUsingSelectiveBasePassOutputs(ShaderPlatform);
 }
 
 void FScene::UpdateEarlyZPassMode()
