Fixed RHI validation error in cloud rendering.

diff --git a/Engine/Source/Runtime/RenderCore/Private/RenderGraphBuilder.cpp b/Engine/Source/Runtime/RenderCore/Private/RenderGraphBuilder.cpp
--- a/Engine/Source/Runtime/RenderCore/Private/RenderGraphBuilder.cpp
+++ b/Engine/Source/Runtime/RenderCore/Private/RenderGraphBuilder.cpp
@@ -1135,6 +1135,13 @@ void FRDGBuilder::Compile()
 				Buffer->bUsedByAsyncComputePass |= bAsyncComputePass;
 				Buffer->bCulled = false;
 
+				if (Buffer->FirstBarrier == FRDGBuffer::EFirstBarrier::ImmediateRequested)
+				{
+					Buffer->FirstBarrier = FRDGBuffer::EFirstBarrier::ImmediateConfirmed;
+					Buffer->FirstPass = PassHandle;
+					Buffer->State->SetPass(ERHIPipeline::Graphics, PassHandle);
+				}
+
 			#if STATS
 				GRDGStatBufferReferenceCount += PassState.ReferenceCount;
 			#endif
@@ -1674,9 +1681,11 @@ void FRDGBuilder::Execute()
 
 	CreatePassBarriers(bParallelExecuteEnabled ? &AsyncCompileEvents : nullptr);
 
+	const ENamedThreads::Type RenderThread = ENamedThreads::GetRenderThread_Local();
+
 	if (!ParallelSetupEvents.IsEmpty())
 	{
-		FTaskGraphInterface::Get().WaitUntilTasksComplete(ParallelSetupEvents);
+		FTaskGraphInterface::Get().WaitUntilTasksComplete(ParallelSetupEvents, RenderThread);
 	}
 
 	EndFlushResourcesRHI();
@@ -1698,8 +1707,6 @@ void FRDGBuilder::Execute()
 	IF_RDG_ENABLE_DEBUG(GRDGAllowRHIAccess = bParallelExecuteEnabled);
 	IF_RDG_ENABLE_TRACE(Trace.OutputGraphBegin());
 
-	const ENamedThreads::Type RenderThread = ENamedThreads::GetRenderThread_Local();
-
 	if (!IsImmediateMode())
 	{
 		SCOPED_NAMED_EVENT_TEXT("FRDGBuilder::ExecutePasses", FColor::Magenta);
diff --git a/Engine/Source/Runtime/RenderCore/Private/RenderGraphResources.cpp b/Engine/Source/Runtime/RenderCore/Private/RenderGraphResources.cpp
--- a/Engine/Source/Runtime/RenderCore/Private/RenderGraphResources.cpp
+++ b/Engine/Source/Runtime/RenderCore/Private/RenderGraphResources.cpp
@@ -21,7 +21,7 @@ inline bool NeedsUAVBarrier(const FRDGSubresourceState& Previous, const FRDGSubr
 	return NeedsUAVBarrier(Previous.NoUAVBarrierFilter.GetUniqueHandle(), Next.NoUAVBarrierFilter.GetUniqueHandle());
 }
 
-FRDGViewableResource::FRDGViewableResource(const TCHAR* InName, const ERDGViewableResourceType InType, bool bSkipTracking)
+FRDGViewableResource::FRDGViewableResource(const TCHAR* InName, const ERDGViewableResourceType InType, bool bSkipTracking, bool bImmediateFirstBarrier)
 	: FRDGResource(InName)
 	, Type(InType)
 	, bExternal(0)
@@ -42,6 +42,11 @@ FRDGViewableResource::FRDGViewableResource(const TCHAR* InName, const ERDGViewab
 		SetExternalAccessMode(ERHIAccess::ReadOnlyExclusiveMask, ERHIPipeline::All);
 		AccessModeState.bLocked = 1;
 	}
+
+	if (bImmediateFirstBarrier)
+	{
+		FirstBarrier = EFirstBarrier::ImmediateRequested;
+	}
 }
 
 bool FRDGProducerState::IsDependencyRequired(FRDGProducerState LastProducer, ERHIPipeline LastPipeline, FRDGProducerState NextState, ERHIPipeline NextPipeline)
diff --git a/Engine/Source/Runtime/RenderCore/Public/RenderGraphDefinitions.h b/Engine/Source/Runtime/RenderCore/Public/RenderGraphDefinitions.h
--- a/Engine/Source/Runtime/RenderCore/Public/RenderGraphDefinitions.h
+++ b/Engine/Source/Runtime/RenderCore/Public/RenderGraphDefinitions.h
@@ -138,6 +138,11 @@
 	 */
 	SkipTracking = 1 << 1,
 
+	/** When set, RDG will perform its first barrier without splitting. Practically, this means the resource is left in its initial state
+	 *  until the first pass it's used within the graph. Without this flag, the resource is split-transitioned at the start of the graph.
+	 */
+	ForceImmediateFirstBarrier = 1 << 2,
+
 	// Deprecated Enums
 	ReadOnly      UE_DEPRECATED(5.1, "ReadOnly is deprecated. Use SkipTracking instead.")                                            = 0,
 	ForceTracking UE_DEPRECATED(5.1, "ForceTracking is deprecated. Resources now opt-out of tracking explicitly with SkipTracking.") = 0
@@ -158,9 +163,14 @@
 	 *  readable state for use with RDG passes, as RDG does not know the exact state of the resource.
 	 */
 	SkipTracking = 1 << 1,
+	
+	/** When set, RDG will perform its first barrier without splitting. Practically, this means the resource is left in its initial state
+	 *  until the first pass it's used within the graph. Without this flag, the resource is split-transitioned at the start of the graph.
+	 */
+	ForceImmediateFirstBarrier = 1 << 2,
 
 	/** Prevents metadata decompression on this texture. */
-	MaintainCompression = 1 << 2,
+	MaintainCompression = 1 << 3,
 
 	// Deprecated Enums
 	ReadOnly      UE_DEPRECATED(5.1, "ReadOnly is deprecated. Use SkipTracking instead.")                                            = 0,
diff --git a/Engine/Source/Runtime/RenderCore/Public/RenderGraphResources.h b/Engine/Source/Runtime/RenderCore/Public/RenderGraphResources.h
--- a/Engine/Source/Runtime/RenderCore/Public/RenderGraphResources.h
+++ b/Engine/Source/Runtime/RenderCore/Public/RenderGraphResources.h
@@ -267,7 +267,7 @@ public:
 	}
 
 protected:
-	FRDGViewableResource(const TCHAR* InName, ERDGViewableResourceType InType, bool bSkipTracking);
+	FRDGViewableResource(const TCHAR* InName, ERDGViewableResourceType InType, bool bSkipTracking, bool bImmediateFirstBarrier);
 
 	static const ERHIAccess DefaultEpilogueAccess = ERHIAccess::SRVMask;
 
@@ -517,7 +517,7 @@ public:
 
 private:
 	FRDGTexture(const TCHAR* InName, const FRDGTextureDesc& InDesc, ERDGTextureFlags InFlags)
-		: FRDGViewableResource(InName, ERDGViewableResourceType::Texture, EnumHasAnyFlags(InFlags, ERDGTextureFlags::SkipTracking))
+		: FRDGViewableResource(InName, ERDGViewableResourceType::Texture, EnumHasAnyFlags(InFlags, ERDGTextureFlags::SkipTracking), EnumHasAnyFlags(InFlags, ERDGTextureFlags::ForceImmediateFirstBarrier) || EnumHasAnyFlags(InDesc.Flags, ETextureCreateFlags::Presentable))
 		, Desc(InDesc)
 		, Flags(InFlags)
 		, Layout(InDesc)
@@ -529,11 +529,6 @@ private:
 		LastProducers.Reserve(SubresourceCount);
 		LastProducers.SetNum(SubresourceCount);
 
-		if (EnumHasAnyFlags(Desc.Flags, ETextureCreateFlags::Presentable))
-		{
-			FirstBarrier = EFirstBarrier::ImmediateRequested;
-		}
-
 		if (EnumHasAnyFlags(Desc.Flags, ETextureCreateFlags::Foveation))
 		{
 			EpilogueAccess = ERHIAccess::ShadingRateSource;
@@ -1242,7 +1237,7 @@ public:
 
 private:
 	FRDGBuffer(const TCHAR* InName, const FRDGBufferDesc& InDesc, ERDGBufferFlags InFlags)
-		: FRDGViewableResource(InName, ERDGViewableResourceType::Buffer, EnumHasAnyFlags(InFlags, ERDGBufferFlags::SkipTracking))
+		: FRDGViewableResource(InName, ERDGViewableResourceType::Buffer, EnumHasAnyFlags(InFlags, ERDGBufferFlags::SkipTracking), EnumHasAnyFlags(InFlags, ERDGBufferFlags::ForceImmediateFirstBarrier))
 		, Desc(InDesc)
 		, Flags(InFlags)
 	{}
diff --git a/Engine/Source/Runtime/Renderer/Private/ReflectionEnvironmentRealTimeCapture.cpp b/Engine/Source/Runtime/Renderer/Private/ReflectionEnvironmentRealTimeCapture.cpp
--- a/Engine/Source/Runtime/Renderer/Private/ReflectionEnvironmentRealTimeCapture.cpp
+++ b/Engine/Source/Runtime/Renderer/Private/ReflectionEnvironmentRealTimeCapture.cpp
@@ -836,7 +836,11 @@ void FScene::AllocateAndCaptureFrameSkyEnvMap(
 	{
 		FRDGTextureRef SourceCubemapTexture = GraphBuilder.RegisterExternalTexture(SourceCubemap);
 		FRDGTextureSRVRef SourceCubemapTextureSRV = GraphBuilder.CreateSRV(FRDGTextureSRVDesc::Create(SourceCubemapTexture));
-		FRDGBuffer* SkyIrradianceEnvironmentMapRDG = GraphBuilder.RegisterExternalBuffer(SkyIrradianceEnvironmentMap);
+		
+		// ForceImmediateFirstBarrier is required because the RHI resource is used as an SRV outside of RDG prior to this UAV pass. Without
+		// the flag, RDG will split the transition to UAV to the start of the graph, which results in a validation error. With the flag, RDG
+		// will transition to UAV at the start of the pass instead.
+		FRDGBuffer* SkyIrradianceEnvironmentMapRDG = GraphBuilder.RegisterExternalBuffer(SkyIrradianceEnvironmentMap, ERDGBufferFlags::ForceImmediateFirstBarrier);
 
 		TShaderMapRef<FComputeSkyEnvMapDiffuseIrradianceCS> ComputeShader(GetGlobalShaderMap(FeatureLevel));
 
