Move RDG buffer uploads off the render thread.

diff --git a/Engine/Source/Runtime/RenderCore/Private/RenderGraphBuilder.cpp b/Engine/Source/Runtime/RenderCore/Private/RenderGraphBuilder.cpp
--- a/Engine/Source/Runtime/RenderCore/Private/RenderGraphBuilder.cpp
+++ b/Engine/Source/Runtime/RenderCore/Private/RenderGraphBuilder.cpp
@@ -1532,6 +1532,7 @@ void FRDGBuilder::Execute()
 	const FRDGPassHandle EpiloguePassHandle = GetEpiloguePassHandle();
 
 	FGraphEventArray AsyncCompileEvents;
+	FGraphEventArray* AsyncCompileEventsIfEnabled = bParallelExecuteEnabled ? &AsyncCompileEvents : nullptr;
 
 	IF_RDG_ENABLE_DEBUG(UserValidation.ValidateExecuteBegin());
 	IF_RDG_ENABLE_DEBUG(GRDGAllowRHIAccess = true);
@@ -1560,7 +1561,7 @@ void FRDGBuilder::Execute()
 			}, TStatId(), nullptr, ENamedThreads::AnyHiPriThreadHiPriTask));
 		}
 
-		SetupBufferUploads();
+		SubmitBufferUploads(AsyncCompileEventsIfEnabled);
 
 		{
 			SCOPE_CYCLE_COUNTER(STAT_RDG_CollectResourcesTime);
@@ -1644,7 +1645,7 @@ void FRDGBuilder::Execute()
 			}
 		}
 
-		CreateUniformBuffers(bParallelExecuteEnabled ? &AsyncCompileEvents : nullptr);
+		CreateUniformBuffers(AsyncCompileEventsIfEnabled);
 
 		{
 			SCOPED_NAMED_EVENT_TEXT("FRDGBuilder::CollectBarriers", FColor::Magenta);
@@ -1679,16 +1680,18 @@ void FRDGBuilder::Execute()
 		});
 	}
 
-	CreatePassBarriers(bParallelExecuteEnabled ? &AsyncCompileEvents : nullptr);
-
 	const ENamedThreads::Type RenderThread = ENamedThreads::GetRenderThread_Local();
 
-	if (!ParallelSetupEvents.IsEmpty())
+	CreatePassBarriers([this, RenderThread]
 	{
-		FTaskGraphInterface::Get().WaitUntilTasksComplete(ParallelSetupEvents, RenderThread);
-	}
-
-	EndFlushResourcesRHI();
+		if (!ParallelSetupEvents.IsEmpty())
+		{
+			FTaskGraphInterface::Get().WaitUntilTasksComplete(ParallelSetupEvents, RenderThread);
+		}
+		
+		// Process RHI thread flush before helping with barrier compilation on the render thread.
+		EndFlushResourcesRHI();
+	});
 
 	FGraphEventRef DispatchParallelExecuteEvent;
 
@@ -1702,8 +1705,6 @@ void FRDGBuilder::Execute()
 		}, TStatId(), &AsyncCompileEvents, ENamedThreads::AnyHiPriThreadHiPriTask);
 	}
 
-	SubmitBufferUploads();
-
 	IF_RDG_ENABLE_DEBUG(GRDGAllowRHIAccess = bParallelExecuteEnabled);
 	IF_RDG_ENABLE_TRACE(Trace.OutputGraphBegin());
 
@@ -2097,13 +2098,12 @@ void FRDGBuilder::SetupPassInternal(FRDGPass* Pass, FRDGPassHandle PassHandle, E
 
 		check(!EnumHasAnyFlags(PassPipeline, ERHIPipeline::AsyncCompute));
 
-		SetupBufferUploads();
-		SubmitBufferUploads();
+		SubmitBufferUploads(nullptr);
 		CompilePassOps(Pass);
 		BeginResourcesRHI(Pass, PassHandle);
 		CreateUniformBuffers(nullptr);
 		CollectPassBarriers(Pass, PassHandle);
-		CreatePassBarriers(nullptr);
+		CreatePassBarriers([] {});
 		FTaskGraphInterface::Get().WaitUntilTasksComplete(ParallelSetupEvents);
 		ParallelSetupEvents.Reset();
 		ExecutePass(Pass, RHICmdList);
@@ -2118,64 +2118,82 @@ void FRDGBuilder::SetupPassInternal(FRDGPass* Pass, FRDGPassHandle PassHandle, E
 
 ///////////////////////////////////////////////////////////////////////////////////////////////////
 
-void FRDGBuilder::SetupBufferUploads()
+void FRDGBuilder::SubmitBufferUploads(FGraphEventArray* AsyncCompileEventsIfEnabled)
 {
-	SCOPED_NAMED_EVENT_TEXT("FRDGBuilder::PrepareBufferUploads", FColor::Magenta);
-
-	for (FUploadedBuffer& UploadedBuffer : UploadedBuffers)
 	{
-		if (UploadedBuffer.bUseDataCallbacks)
-		{
-			UploadedBuffer.Data = UploadedBuffer.DataCallback();
-			UploadedBuffer.DataSize = UploadedBuffer.DataSizeCallback();
-		}
+		SCOPED_NAMED_EVENT_TEXT("FRDGBuilder::PrepareBufferUploads", FColor::Magenta);
 
-		if (UploadedBuffer.Data && UploadedBuffer.DataSize)
+		for (FUploadedBuffer& UploadedBuffer : UploadedBuffers)
 		{
-			ConvertToExternalBuffer(UploadedBuffer.Buffer);
-			check(UploadedBuffer.DataSize <= UploadedBuffer.Buffer->Desc.GetSize());
+			if (UploadedBuffer.bUseDataCallbacks)
+			{
+				UploadedBuffer.Data = UploadedBuffer.DataCallback();
+				UploadedBuffer.DataSize = UploadedBuffer.DataSizeCallback();
+			}
+
+			if (UploadedBuffer.Data && UploadedBuffer.DataSize)
+			{
+				ConvertToExternalBuffer(UploadedBuffer.Buffer);
+				check(UploadedBuffer.DataSize <= UploadedBuffer.Buffer->Desc.GetSize());
+			}
 		}
 	}
-}
-
-void FRDGBuilder::SubmitBufferUploads()
-{
-	SCOPED_NAMED_EVENT_TEXT("FRDGBuilder::SubmitBufferUploads", FColor::Magenta);
 
-	for (const FUploadedBuffer& UploadedBuffer : UploadedBuffers)
+	const auto SubmitUploadsLambda = [this](FRHICommandList& RHICmdListUpload)
 	{
-		if (UploadedBuffer.Data && UploadedBuffer.DataSize)
+		SCOPED_NAMED_EVENT_TEXT("FRDGBuilder::SubmitBufferUploads", FColor::Magenta);
+
+		for (const FUploadedBuffer& UploadedBuffer : UploadedBuffers)
 		{
+			if (UploadedBuffer.Data && UploadedBuffer.DataSize)
+			{
 #if PLATFORM_NEEDS_GPU_UAV_RESOURCE_INIT_WORKAROUND
-			if (UploadedBuffer.Buffer->bUAVAccessed)
-			{
-				FRHIResourceCreateInfo CreateInfo(UploadedBuffer.Buffer->Name);
-				FBufferRHIRef TempBuffer = RHICreateVertexBuffer(UploadedBuffer.DataSize, BUF_Static | BUF_ShaderResource, CreateInfo);
-				void* DestPtr = RHICmdList.LockBuffer(TempBuffer, 0, UploadedBuffer.DataSize, RLM_WriteOnly);
-				FMemory::Memcpy(DestPtr, UploadedBuffer.Data, UploadedBuffer.DataSize);
-				RHICmdList.UnlockBuffer(TempBuffer);
-				RHICmdList.Transition(
+				if (UploadedBuffer.Buffer->bUAVAccessed)
 				{
-					FRHITransitionInfo(TempBuffer, ERHIAccess::Unknown, ERHIAccess::CopySrc | ERHIAccess::SRVMask),
-					FRHITransitionInfo(UploadedBuffer.Buffer->GetRHI(), ERHIAccess::Unknown, ERHIAccess::CopyDest)
-				});
-				RHICmdList.CopyBufferRegion(UploadedBuffer.Buffer->GetRHI(), 0, TempBuffer, 0, UploadedBuffer.DataSize);
-			}
-			else
+					FRHIResourceCreateInfo CreateInfo(UploadedBuffer.Buffer->Name);
+					FBufferRHIRef TempBuffer = RHICmdListUpload.CreateBuffer(UploadedBuffer.DataSize, BUF_VertexBuffer | BUF_Static | BUF_ShaderResource, 0, ERHIAccess::CopySrc | ERHIAccess::SRVMask, CreateInfo);
+					void* DestPtr = RHICmdListUpload.LockBuffer(TempBuffer, 0, UploadedBuffer.DataSize, RLM_WriteOnly);
+					FMemory::Memcpy(DestPtr, UploadedBuffer.Data, UploadedBuffer.DataSize);
+					RHICmdListUpload.UnlockBuffer(TempBuffer);
+					RHICmdListUpload.Transition(FRHITransitionInfo(UploadedBuffer.Buffer->GetRHI(), ERHIAccess::Unknown, ERHIAccess::CopyDest));
+					RHICmdListUpload.CopyBufferRegion(UploadedBuffer.Buffer->GetRHI(), 0, TempBuffer, 0, UploadedBuffer.DataSize);
+				}
+				else
 #endif
-			{
-				void* DestPtr = RHICmdList.LockBuffer(UploadedBuffer.Buffer->GetRHIUnchecked(), 0, UploadedBuffer.DataSize, RLM_WriteOnly);
-				FMemory::Memcpy(DestPtr, UploadedBuffer.Data, UploadedBuffer.DataSize);
-				RHICmdList.UnlockBuffer(UploadedBuffer.Buffer->GetRHIUnchecked());
-			}
+				{
+					void* DestPtr = RHICmdListUpload.LockBuffer(UploadedBuffer.Buffer->GetRHIUnchecked(), 0, UploadedBuffer.DataSize, RLM_WriteOnly);
+					FMemory::Memcpy(DestPtr, UploadedBuffer.Data, UploadedBuffer.DataSize);
+					RHICmdListUpload.UnlockBuffer(UploadedBuffer.Buffer->GetRHIUnchecked());
+				}
 
-			if (UploadedBuffer.bUseFreeCallbacks)
-			{
-				UploadedBuffer.DataFreeCallback(UploadedBuffer.Data);
+				if (UploadedBuffer.bUseFreeCallbacks)
+				{
+					UploadedBuffer.DataFreeCallback(UploadedBuffer.Data);
+				}
 			}
 		}
+	};
+
+	if (AsyncCompileEventsIfEnabled)
+	{
+		FRHICommandList* RHICmdListUpload = new FRHICommandList(FRHIGPUMask::All());
+
+		FGraphEventRef Event = FFunctionGraphTask::CreateAndDispatchWhenReady(
+			[SubmitUploadsLambda, RHICmdListUpload](ENamedThreads::Type, const FGraphEventRef&)
+		{
+			FTaskTagScope Scope(ETaskTag::EParallelRenderingThread);
+			SubmitUploadsLambda(*RHICmdListUpload);
+
+		}, TStatId(), nullptr, ENamedThreads::AnyHiPriThreadHiPriTask);
+
+		RHICmdList.QueueRenderThreadCommandListSubmit(Event, RHICmdListUpload);
+
+		AsyncCompileEventsIfEnabled->Emplace(Event);
+	}
+	else
+	{
+		SubmitUploadsLambda(RHICmdList);
 	}
-	UploadedBuffers.Reset();
 }
 
 ///////////////////////////////////////////////////////////////////////////////////////////////////
@@ -2516,7 +2534,7 @@ void FRDGBuilder::ExecutePassEpilogue(FRHIComputeCommandList& RHICmdListPass, FR
 
 void FRDGBuilder::ExecutePass(FRDGPass* Pass, FRHIComputeCommandList& RHICmdListPass)
 {
-#if 1 || RDG_EVENTS != RDG_EVENTS_NONE
+#if RDG_EVENTS != RDG_EVENTS_NONE
 	SCOPED_NAMED_EVENT_TCHAR(Pass->GetName(), FColor::Magenta);
 #endif
 
@@ -2653,43 +2671,31 @@ void FRDGBuilder::CollectPassBarriers(FRDGPass* Pass, FRDGPassHandle PassHandle)
 	}
 }
 
-void FRDGBuilder::CreatePassBarriers(FGraphEventArray* AsyncCompileEvents)
+void FRDGBuilder::CreatePassBarriers(TFunctionRef<void()> PreWork)
 {
-	const int32 NumBarriersPerTask = 128;
+	const int32 NumBarriersPerTask = 32;
 
-	if (AsyncCompileEvents && TransitionCreateQueue.Num() > NumBarriersPerTask)
+	if (bParallelExecuteEnabled && TransitionCreateQueue.Num() > NumBarriersPerTask)
 	{
-		AsyncCompileEvents->Emplace(FFunctionGraphTask::CreateAndDispatchWhenReady(
-			[this, NumBarriersPerTask](ENamedThreads::Type, const FGraphEventRef&)
+		ParallelForWithPreWork(TEXT("FRDGBuilder::CreatePassBarriers"), TransitionCreateQueue.Num(), NumBarriersPerTask, [this](int32 Index)
 		{
-#if 1
-			ParallelFor(TEXT("FRDGBuilder::CreatePassBarriers"), TransitionCreateQueue.Num(), NumBarriersPerTask, [this](int32 Index)
-			{
-				TransitionCreateQueue[Index]->CreateTransition();
-			});
-#else
-			SCOPED_NAMED_EVENT_TEXT("FRDGBuilder::CreatePassBarriers", FColor::Magenta);
-			for (FRDGBarrierBatchBegin* BeginBatch : TransitionCreateQueue)
-			{
-				BeginBatch->CreateTransition();
-			}
-#endif
+			TransitionCreateQueue[Index]->CreateTransition();
 
-			TransitionCreateQueue.Reset();
-
-		}, TStatId(), nullptr, ENamedThreads::AnyHiPriThreadHiPriTask));
+		}, PreWork);
 	}
 	else
 	{
+		PreWork();
+
 		SCOPED_NAMED_EVENT_TEXT("FRDGBuilder::CreatePassBarriers", FColor::Magenta);
 
 		for (FRDGBarrierBatchBegin* BeginBatch : TransitionCreateQueue)
 		{
 			BeginBatch->CreateTransition();
 		}
-
-		TransitionCreateQueue.Reset();
 	}
+
+	TransitionCreateQueue.Reset();
 }
 
 ///////////////////////////////////////////////////////////////////////////////////////////////////
diff --git a/Engine/Source/Runtime/RenderCore/Public/RenderGraphBuilder.h b/Engine/Source/Runtime/RenderCore/Public/RenderGraphBuilder.h
--- a/Engine/Source/Runtime/RenderCore/Public/RenderGraphBuilder.h
+++ b/Engine/Source/Runtime/RenderCore/Public/RenderGraphBuilder.h
@@ -756,8 +756,7 @@ private:
 	void SetupParallelExecute();
 	void DispatchParallelExecute(IRHICommandContext* RHICmdContext);
 
-	void SetupBufferUploads();
-	void SubmitBufferUploads();
+	void SubmitBufferUploads(FGraphEventArray* AsyncCompileEvents);
 	void BeginFlushResourcesRHI();
 	void EndFlushResourcesRHI();
 
@@ -774,7 +773,7 @@ private:
 
 	void CollectPassBarriers(FRDGPass* Pass, FRDGPassHandle PassHandle);
 
-	void CreatePassBarriers(FGraphEventArray* AsyncCompileEvents);
+	void CreatePassBarriers(TFunctionRef<void()> PreWork);
 
 	void CreateUniformBuffers(FGraphEventArray* AsyncCompileEvents);
 
