Always use UINT for NiagaraFX SortedIndices parameter.
To fix vulkan validation warning: Validation Error: [ VUID-vkCmdDrawIndexed-None-02699 ] Object 0: handle = 0xb60d06000004faeb, type = VK_OBJECT_TYPE_DESCRIPTOR_SET; | MessageID = 0xa44449d4 | Descriptor set VkDescriptorSet 0xb60d06000004faeb[] encountered the following validation error at vkCmdDrawIndexed time: Descriptor in binding #11 index 0 requires UINT component type, but bound descriptor format is VK_FORMAT_R32_SINT. The Vulkan spec states: Descriptors in each bound descriptor set, specified via vkCmdBindDescriptorSets, must be valid if they are statically used by the VkPipeline bound to the pipeline bind point used by this command (https://www.khronos.org/registry/vulkan/specs/1.3-extensions/html/vkspec.html#VUID-vkCmdDrawIndexed-None-02699)

diff --git a/Engine/Plugins/FX/Niagara/Shaders/Private/NiagaraUpdateMeshGPUSceneInstances.usf b/Engine/Plugins/FX/Niagara/Shaders/Private/NiagaraUpdateMeshGPUSceneInstances.usf
--- a/Engine/Plugins/FX/Niagara/Shaders/Private/NiagaraUpdateMeshGPUSceneInstances.usf
+++ b/Engine/Plugins/FX/Niagara/Shaders/Private/NiagaraUpdateMeshGPUSceneInstances.usf
@@ -5,7 +5,7 @@
 #include "/Engine/Private/GPUScene/GPUSceneWriter.ush"
 #include "NiagaraMeshParticleUtils.ush"
 
-Buffer<int> SortedIndices;
+Buffer<uint> SortedIndices;
 int SortedIndicesOffset;
 
 float3 SystemLWCTile;
diff --git a/Engine/Plugins/FX/Niagara/Source/Niagara/Private/NiagaraRendererMeshes.cpp b/Engine/Plugins/FX/Niagara/Source/Niagara/Private/NiagaraRendererMeshes.cpp
--- a/Engine/Plugins/FX/Niagara/Source/Niagara/Private/NiagaraRendererMeshes.cpp
+++ b/Engine/Plugins/FX/Niagara/Source/Niagara/Private/NiagaraRendererMeshes.cpp
@@ -634,7 +634,7 @@ uint32 FNiagaraRendererMeshes::PerformSortAndCull(FParticleMeshRenderData& Parti
 		else
 		{
 			FGlobalDynamicReadBuffer::FAllocation SortedIndices;
-			SortedIndices = ReadBuffer.AllocateInt32(NumInstances);
+			SortedIndices = ReadBuffer.AllocateUInt32(NumInstances);
 			NumInstances = SortAndCullIndices(SortInfo, *ParticleMeshRenderData.SourceParticleData, SortedIndices);
 			ParticleMeshRenderData.ParticleSortedIndicesSRV = SortedIndices.SRV;
 			ParticleMeshRenderData.ParticleSortedIndicesOffset = 0;
diff --git a/Engine/Plugins/FX/Niagara/Source/Niagara/Private/NiagaraRendererRibbons.cpp b/Engine/Plugins/FX/Niagara/Source/Niagara/Private/NiagaraRendererRibbons.cpp
--- a/Engine/Plugins/FX/Niagara/Source/Niagara/Private/NiagaraRendererRibbons.cpp
+++ b/Engine/Plugins/FX/Niagara/Source/Niagara/Private/NiagaraRendererRibbons.cpp
@@ -1944,7 +1944,7 @@ inline void FNiagaraRendererRibbons::SetupMeshBatchAndCollectorResourceForView(c
 	VFLooseParams.Shape = static_cast<uint32>(ShapeState.Shape);
 	VFLooseParams.NeedsPreciseMotionVectors = GenerationConfig.NeedsPreciseMotionVectors();
 
-	VFLooseParams.IndirectDrawOutput = bShouldUseGPUInitIndices? (FRHIShaderResourceView*)RenderingViewResources->IndirectDrawBuffer.SRV : GetDummyIntBuffer();
+	VFLooseParams.IndirectDrawOutput = bShouldUseGPUInitIndices? (FRHIShaderResourceView*)RenderingViewResources->IndirectDrawBuffer.SRV : GetDummyUIntBuffer();
 	VFLooseParams.IndirectDrawOutputOffset = bShouldUseGPUInitIndices? 0 : -1;
 
 	// Collector.AllocateOneFrameResource uses default ctor, initialize the vertex factory
diff --git a/Engine/Plugins/FX/Niagara/Source/Niagara/Private/NiagaraRendererSprites.cpp b/Engine/Plugins/FX/Niagara/Source/Niagara/Private/NiagaraRendererSprites.cpp
--- a/Engine/Plugins/FX/Niagara/Source/Niagara/Private/NiagaraRendererSprites.cpp
+++ b/Engine/Plugins/FX/Niagara/Source/Niagara/Private/NiagaraRendererSprites.cpp
@@ -1006,7 +1006,7 @@ void FNiagaraRendererSprites::GetDynamicMeshElements(const TArray<const FSceneVi
 				else
 				{
 					FGlobalDynamicReadBuffer::FAllocation SortedIndices;
-					SortedIndices = Collector.GetDynamicReadBuffer().AllocateInt32(NumInstances);
+					SortedIndices = Collector.GetDynamicReadBuffer().AllocateUInt32(NumInstances);
 					NumInstances = SortAndCullIndices(SortInfo, *ParticleSpriteRenderData.SourceParticleData, SortedIndices);
 					VertexFactory.SetSortedIndices(SortedIndices.SRV, 0);
 				}
@@ -1096,7 +1096,7 @@ void FNiagaraRendererSprites::GetDynamicRayTracingInstances(FRayTracingMaterialG
 		else
 		{
 			FGlobalDynamicReadBuffer::FAllocation SortedIndices;
-			SortedIndices = DynamicReadBuffer.AllocateInt32(NumInstances);
+			SortedIndices = DynamicReadBuffer.AllocateUInt32(NumInstances);
 			NumInstances = SortAndCullIndices(SortInfo, *ParticleSpriteRenderData.SourceParticleData, SortedIndices);
 			VertexFactory.SetSortedIndices(SortedIndices.SRV, 0);
 		}
diff --git a/Engine/Plugins/FX/Niagara/Source/NiagaraVertexFactories/Public/NiagaraMeshVertexFactory.h b/Engine/Plugins/FX/Niagara/Source/NiagaraVertexFactories/Public/NiagaraMeshVertexFactory.h
--- a/Engine/Plugins/FX/Niagara/Source/NiagaraVertexFactories/Public/NiagaraMeshVertexFactory.h
+++ b/Engine/Plugins/FX/Niagara/Source/NiagaraVertexFactories/Public/NiagaraMeshVertexFactory.h
@@ -35,7 +35,7 @@ BEGIN_SHADER_PARAMETER_STRUCT(FNiagaraMeshCommonParameters, NIAGARAVERTEXFACTORI
 	SHADER_PARAMETER(uint32, NiagaraFloatDataStride)
 	SHADER_PARAMETER(uint32, NiagaraIntDataStride)
 
-	SHADER_PARAMETER_SRV(Buffer<int>, SortedIndices)
+	SHADER_PARAMETER_SRV(Buffer<uint>, SortedIndices)
 	SHADER_PARAMETER(int, SortedIndicesOffset)
 	
 	SHADER_PARAMETER(FVector3f, SystemLWCTile)
diff --git a/Engine/Plugins/FX/Niagara/Source/NiagaraVertexFactories/Public/NiagaraSpriteVertexFactory.h b/Engine/Plugins/FX/Niagara/Source/NiagaraVertexFactories/Public/NiagaraSpriteVertexFactory.h
--- a/Engine/Plugins/FX/Niagara/Source/NiagaraVertexFactories/Public/NiagaraSpriteVertexFactory.h
+++ b/Engine/Plugins/FX/Niagara/Source/NiagaraVertexFactories/Public/NiagaraSpriteVertexFactory.h
@@ -100,7 +100,7 @@ BEGIN_GLOBAL_SHADER_PARAMETER_STRUCT(FNiagaraSpriteVFLooseParameters, NIAGARAVER
 	SHADER_PARAMETER_SRV(Buffer<float2>, CutoutGeometry)
 	SHADER_PARAMETER_SRV(Buffer<float>, NiagaraParticleDataFloat)
 	SHADER_PARAMETER_SRV(Buffer<float>, NiagaraParticleDataHalf)
-	SHADER_PARAMETER_SRV(Buffer<int>, SortedIndices)
+	SHADER_PARAMETER_SRV(Buffer<uint>, SortedIndices)
 	SHADER_PARAMETER_SRV(Buffer<uint>, IndirectArgsBuffer)
 END_GLOBAL_SHADER_PARAMETER_STRUCT()
 
diff --git a/Engine/Plugins/FX/Niagara/Source/NiagaraVertexFactories/Public/NiagaraVertexFactory.h b/Engine/Plugins/FX/Niagara/Source/NiagaraVertexFactories/Public/NiagaraVertexFactory.h
--- a/Engine/Plugins/FX/Niagara/Source/NiagaraVertexFactories/Public/NiagaraVertexFactory.h
+++ b/Engine/Plugins/FX/Niagara/Source/NiagaraVertexFactories/Public/NiagaraVertexFactory.h
@@ -25,12 +25,12 @@ public:
 	{
 		// create a static vertex buffer
 		FRHIResourceCreateInfo CreateInfo(TEXT("FNiagaraNullSortedIndicesVertexBuffer"));
-		VertexBufferRHI = RHICreateBuffer(sizeof(int32), BUF_Static | BUF_VertexBuffer | BUF_ShaderResource, 0, ERHIAccess::VertexOrIndexBuffer | ERHIAccess::SRVMask, CreateInfo);
-		void* BufferData = RHILockBuffer(VertexBufferRHI, 0, sizeof(int32), RLM_WriteOnly);
-		FMemory::Memzero(BufferData, sizeof(int32));
+		VertexBufferRHI = RHICreateBuffer(sizeof(uint32), BUF_Static | BUF_VertexBuffer | BUF_ShaderResource, 0, ERHIAccess::VertexOrIndexBuffer | ERHIAccess::SRVMask, CreateInfo);
+		void* BufferData = RHILockBuffer(VertexBufferRHI, 0, sizeof(uint32), RLM_WriteOnly);
+		FMemory::Memzero(BufferData, sizeof(uint32));
 		RHIUnlockBuffer(VertexBufferRHI);
 
-		VertexBufferSRV = RHICreateShaderResourceView(VertexBufferRHI, sizeof(int32), PF_R32_SINT);
+		VertexBufferSRV = RHICreateShaderResourceView(VertexBufferRHI, sizeof(uint32), PF_R32_UINT);
 	}
 
 	virtual void ReleaseRHI() override
diff --git a/Engine/Shaders/Public/ShaderVersion.ush b/Engine/Shaders/Public/ShaderVersion.ush
--- a/Engine/Shaders/Public/ShaderVersion.ush
+++ b/Engine/Shaders/Public/ShaderVersion.ush
@@ -2,4 +2,4 @@
 // Each time the console command is executed it generates a new GUID. As this file is included
 // in Platform.ush (which should be included in any shader) it allows to invalidate the shader DDC.
 // 
-// GUID = E3AA61C3-6F51-4308-8DC0-46E401CFDD58
+// GUID = 61380C6B-BC5B-4BEC-94A4-FE7397B0ACDF
diff --git a/Engine/Source/Runtime/Engine/Private/GPUSortManager.cpp b/Engine/Source/Runtime/Engine/Private/GPUSortManager.cpp
--- a/Engine/Source/Runtime/Engine/Private/GPUSortManager.cpp
+++ b/Engine/Source/Runtime/Engine/Private/GPUSortManager.cpp
@@ -291,7 +291,8 @@ void FGPUSortManager::FValueBuffer::Allocate(FAllocationInfo& OutInfo, int32 Val
 	check(UsedCount <= AllocatedCount);
 	if (EnumHasAnyFlags(Flags, EGPUSortFlags::ValuesAsInt32))
 	{
-		OutInfo.BufferSRV = Int32SRV;
+		// Note that the sorting shader reads the indices as UInt32 (see FParticleSortBuffers::InitRHI()).
+		OutInfo.BufferSRV = UInt32SRV;
 	}
 	else
 	{
diff --git a/Engine/Source/Runtime/RenderCore/Private/DynamicBufferAllocator.cpp b/Engine/Source/Runtime/RenderCore/Private/DynamicBufferAllocator.cpp
--- a/Engine/Source/Runtime/RenderCore/Private/DynamicBufferAllocator.cpp
+++ b/Engine/Source/Runtime/RenderCore/Private/DynamicBufferAllocator.cpp
@@ -63,6 +63,7 @@ FGlobalDynamicReadBuffer::FGlobalDynamicReadBuffer()
 {
 	FloatBufferPool = new FDynamicReadBufferPool();
 	Int32BufferPool = new FDynamicReadBufferPool();
+	UInt32BufferPool = new FDynamicReadBufferPool();
 	HalfBufferPool = new FDynamicReadBufferPool();
 }
 
@@ -85,6 +86,12 @@ void FGlobalDynamicReadBuffer::Cleanup()
 		Int32BufferPool = nullptr;
 	}
 
+	if (UInt32BufferPool)
+	{
+		delete UInt32BufferPool;
+		UInt32BufferPool = nullptr;
+	}
+
 	if (HalfBufferPool)
 	{
 		delete HalfBufferPool;
@@ -186,16 +193,22 @@ FGlobalDynamicReadBuffer::FAllocation FGlobalDynamicReadBuffer::AllocateInt32(ui
 	return FGlobalDynamicReadBuffer_AllocateInternal<PF_R32_SINT, int32>(Int32BufferPool, Num);
 }
 
+FGlobalDynamicReadBuffer::FAllocation FGlobalDynamicReadBuffer::AllocateUInt32(uint32 Num)
+{
+	IncrementTotalAllocations(Num);
+	return FGlobalDynamicReadBuffer_AllocateInternal<PF_R32_UINT, uint32>(UInt32BufferPool, Num);
+}
+
 bool FGlobalDynamicReadBuffer::IsRenderAlarmLoggingEnabled() const
 {
 	return GMaxReadBufferRenderingBytesAllocatedPerFrame > 0 && TotalAllocatedSinceLastCommit >= (size_t)GMaxReadBufferRenderingBytesAllocatedPerFrame;
 }
 
-void FGlobalDynamicReadBuffer::Commit()
+static void RemoveUnusedBuffers(FDynamicReadBufferPool* BufferPool)
 {
-	for (int32 BufferIndex = 0, NumBuffers = FloatBufferPool->Buffers.Num(); BufferIndex < NumBuffers; ++BufferIndex)
+	for (int32 BufferIndex = 0, NumBuffers = BufferPool->Buffers.Num(); BufferIndex < NumBuffers; ++BufferIndex)
 	{
-		FDynamicAllocReadBuffer& Buffer = FloatBufferPool->Buffers[BufferIndex];
+		FDynamicAllocReadBuffer& Buffer = BufferPool->Buffers[BufferIndex];
 		if (Buffer.MappedBuffer != nullptr)
 		{
 			Buffer.Unlock();
@@ -207,58 +220,27 @@ void FGlobalDynamicReadBuffer::Commit()
 			{
 				// Remove the buffer, assumes they are unordered.
 				Buffer.Release();
-				FloatBufferPool->Buffers.RemoveAtSwap(BufferIndex);
+				BufferPool->Buffers.RemoveAtSwap(BufferIndex);
 				--BufferIndex;
 				--NumBuffers;
 			}
 		}
 	}
+}
+
+void FGlobalDynamicReadBuffer::Commit()
+{
+	RemoveUnusedBuffers(FloatBufferPool);
 	FloatBufferPool->CurrentBuffer = nullptr;
 
-	for (int32 BufferIndex = 0, NumBuffers = Int32BufferPool->Buffers.Num(); BufferIndex < NumBuffers; ++BufferIndex)
-	{
-		FDynamicAllocReadBuffer& Buffer = Int32BufferPool->Buffers[BufferIndex];
-		if (Buffer.MappedBuffer != nullptr)
-		{
-			Buffer.Unlock();
-		}
-		else if (GGlobalBufferNumFramesUnusedThresold && !Buffer.AllocatedByteCount)
-		{
-			++Buffer.NumFramesUnused;
-			if (Buffer.NumFramesUnused >= GGlobalBufferNumFramesUnusedThresold)
-			{
-				// Remove the buffer, assumes they are unordered.
-				Buffer.Release();
-				Int32BufferPool->Buffers.RemoveAtSwap(BufferIndex);
-				--BufferIndex;
-				--NumBuffers;
-			}
-		}
-	}
+	RemoveUnusedBuffers(Int32BufferPool);
 	Int32BufferPool->CurrentBuffer = nullptr;
 
-	for (int32 BufferIndex = 0, NumBuffers = HalfBufferPool->Buffers.Num(); BufferIndex < NumBuffers; ++BufferIndex)
-	{
-		FDynamicAllocReadBuffer& Buffer = HalfBufferPool->Buffers[BufferIndex];
-		if (Buffer.MappedBuffer != nullptr)
-		{
-			Buffer.Unlock();
-		}
-		else if (GGlobalBufferNumFramesUnusedThresold && !Buffer.AllocatedByteCount)
-		{
-			++Buffer.NumFramesUnused;
-			if (Buffer.NumFramesUnused >= GGlobalBufferNumFramesUnusedThresold)
-			{
-				// Remove the buffer, assumes they are unordered.
-				Buffer.Release();
-				HalfBufferPool->Buffers.RemoveAtSwap(BufferIndex);
-				--BufferIndex;
-				--NumBuffers;
-			}
-		}
-	}
-	HalfBufferPool->CurrentBuffer = nullptr;
+	RemoveUnusedBuffers(UInt32BufferPool);
+	UInt32BufferPool->CurrentBuffer = nullptr;
 
+	RemoveUnusedBuffers(HalfBufferPool);
+	HalfBufferPool->CurrentBuffer = nullptr;
 
 	TotalAllocatedSinceLastCommit = 0;
 }
diff --git a/Engine/Source/Runtime/RenderCore/Public/DynamicBufferAllocator.h b/Engine/Source/Runtime/RenderCore/Public/DynamicBufferAllocator.h
--- a/Engine/Source/Runtime/RenderCore/Public/DynamicBufferAllocator.h
+++ b/Engine/Source/Runtime/RenderCore/Public/DynamicBufferAllocator.h
@@ -76,6 +76,7 @@ public:
 	FAllocation AllocateHalf(uint32 Num);
 	FAllocation AllocateFloat(uint32 Num);
 	FAllocation AllocateInt32(uint32 Num);
+	FAllocation AllocateUInt32(uint32 Num);
 
 	/**
 	* Commits allocated memory to the GPU.
@@ -98,6 +99,7 @@ protected:
 	FDynamicReadBufferPool* HalfBufferPool;
 	FDynamicReadBufferPool* FloatBufferPool;
 	FDynamicReadBufferPool* Int32BufferPool;
+	FDynamicReadBufferPool* UInt32BufferPool;
 
 	/** A total of all allocations made since the last commit. Used to alert about spikes in memory usage. */
 	size_t TotalAllocatedSinceLastCommit;
