Changed ShaderArchive, GlobalShaderCache, ShaderDebugInfo and Autogen to use ShaderPlatformName and not ShaderFormat when naming their output files.

diff --git a/Engine/Config/Android/DataDrivenPlatformInfo.ini b/Engine/Config/Android/DataDrivenPlatformInfo.ini
--- a/Engine/Config/Android/DataDrivenPlatformInfo.ini
+++ b/Engine/Config/Android/DataDrivenPlatformInfo.ini
@@ -56,6 +56,7 @@ DeviceProfileName=Android
 [ShaderPlatform OPENGL_ES3_1_ANDROID]
 Language=OpenGL
 MaxFeatureLevel=ES3_1
+ShaderFormat=GLSL_ES3_1_ANDROID
 bIsMobile=true
 bIsAndroidOpenGLES=true
 bTargetsTiledGPU=true
@@ -66,6 +67,7 @@ bSupportsManualVertexFetch=false
 [ShaderPlatform VULKAN_ES3_1_ANDROID]
 Language=Vulkan
 MaxFeatureLevel=ES3_1
+ShaderFormat=SF_VULKAN_ES31_ANDROID
 bIsMobile=true
 bTargetsTiledGPU=true
 bSupportsMobileMultiView=true
@@ -76,6 +78,7 @@ bSupportsManualVertexFetch=false
 [ShaderPlatform VULKAN_SM5_ANDROID]
 Language=Vulkan
 MaxFeatureLevel=SM5
+ShaderFormat=SF_VULKAN_SM5_ANDROID
 bTargetsTiledGPU=true
 bNeedsOfflineCompiler=true
 bSupportsIndexBufferUAVs=true
diff --git a/Engine/Config/HoloLens/DataDrivenPlatformInfo.ini b/Engine/Config/HoloLens/DataDrivenPlatformInfo.ini
--- a/Engine/Config/HoloLens/DataDrivenPlatformInfo.ini
+++ b/Engine/Config/HoloLens/DataDrivenPlatformInfo.ini
@@ -28,6 +28,7 @@ PlatformGroupName=Mobile
 [ShaderPlatform D3D_ES3_1_HOLOLENS]
 Language=D3D
 MaxFeatureLevel=ES3_1
+ShaderFormat=D3D_ES3_1_HOLOLENS
 bIsMobile=true
 bSupportsInstancedStereo=true
 bSupportsManualVertexFetch=false
diff --git a/Engine/Config/IOS/DataDrivenPlatformInfo.ini b/Engine/Config/IOS/DataDrivenPlatformInfo.ini
--- a/Engine/Config/IOS/DataDrivenPlatformInfo.ini
+++ b/Engine/Config/IOS/DataDrivenPlatformInfo.ini
@@ -48,6 +48,7 @@ DeviceProfileName=IOS
 [ShaderPlatform METAL]
 Language=Metal
 MaxFeatureLevel=ES3_1
+ShaderFormat=SF_METAL_MACES3_1
 bIsMobile=true
 bTargetsTiledGPU=true
 bSupports4ComponentUAVReadWrite=true
@@ -57,6 +58,7 @@ bSupportsVolumeTextureAtomics=false
 [ShaderPlatform METAL_MRT]
 Language=Metal
 MaxFeatureLevel=SM5
+ShaderFormat=SF_METAL_MRT
 bIsMetalMRT=true
 bTargetsTiledGPU=true
 bSupportsIndexBufferUAVs=true
diff --git a/Engine/Config/Mac/DataDrivenPlatformInfo.ini b/Engine/Config/Mac/DataDrivenPlatformInfo.ini
--- a/Engine/Config/Mac/DataDrivenPlatformInfo.ini
+++ b/Engine/Config/Mac/DataDrivenPlatformInfo.ini
@@ -26,6 +26,7 @@ PlatformGroupName=Desktop
 [ShaderPlatform METAL_SM5]
 Language=Metal
 MaxFeatureLevel=SM5
+ShaderFormat=SF_METAL_SM5
 bIsPC=true
 bSupportsDebugViewShaders=true
 bSupportsAnisotropicMaterials=true
@@ -53,6 +54,7 @@ EnablesHLSL2021ByDefault=1
 [ShaderPlatform METAL_MRT_MAC]
 Language=Metal
 MaxFeatureLevel=SM5
+ShaderFormat=SF_METAL_MRT_MAC
 bIsPC=true
 bIsMetalMRT=true
 bSupportsIndexBufferUAVs=true
@@ -63,6 +65,7 @@ bSupportsVolumeTextureAtomics=false
 [ShaderPlatform METAL_MACES3_1]
 Language=Metal
 MaxFeatureLevel=ES3_1
+ShaderFormat=SF_METAL_MACES3_1
 bIsPC=true
 bIsMobile=true
 bSupports4ComponentUAVReadWrite=true
diff --git a/Engine/Config/TVOS/DataDrivenPlatformInfo.ini b/Engine/Config/TVOS/DataDrivenPlatformInfo.ini
--- a/Engine/Config/TVOS/DataDrivenPlatformInfo.ini
+++ b/Engine/Config/TVOS/DataDrivenPlatformInfo.ini
@@ -19,6 +19,7 @@ PlatformGroupName=Mobile
 [ShaderPlatform METAL_MRT_TVOS]
 Language=Metal
 MaxFeatureLevel=SM5
+ShaderFormat=SF_METAL_MRT_TVOS
 bIsMetalMRT=true
 bTargetsTiledGPU=true
 bSupportsIndexBufferUAVs=true
@@ -28,6 +29,7 @@ bSupportsManualVertexFetch=false
 [ShaderPlatform METAL_TVOS]
 Language=Metal
 MaxFeatureLevel=ES3_1
+ShaderFormat=SF_METAL_TVOS
 bIsMobile=true
 bTargetsTiledGPU=true
 bSupports4ComponentUAVReadWrite=true
diff --git a/Engine/Config/VulkanPC/DataDrivenPlatformInfo.ini b/Engine/Config/VulkanPC/DataDrivenPlatformInfo.ini
--- a/Engine/Config/VulkanPC/DataDrivenPlatformInfo.ini
+++ b/Engine/Config/VulkanPC/DataDrivenPlatformInfo.ini
@@ -11,6 +11,7 @@ GlobalIdentifier=97F8C85EF1B0424792D201C893358C8F
 [ShaderPlatform VULKAN_SM5]
 Language=Vulkan
 MaxFeatureLevel=SM5
+ShaderFormat=SF_VULKAN_SM5
 
 bIsMobile = false
 bIsMetalMRT = false
@@ -135,6 +136,7 @@ EnablesHLSL2021ByDefault=1
 [ShaderPlatform VULKAN_PCES3_1]
 Language=Vulkan
 MaxFeatureLevel=ES3_1
+ShaderFormat=SF_VULKAN_ES31
 bIsPC=true
 bIsMobile=true
 bSupportsMobileMultiView=true
diff --git a/Engine/Config/Windows/DataDrivenPlatformInfo.ini b/Engine/Config/Windows/DataDrivenPlatformInfo.ini
--- a/Engine/Config/Windows/DataDrivenPlatformInfo.ini
+++ b/Engine/Config/Windows/DataDrivenPlatformInfo.ini
@@ -57,6 +57,7 @@ PlatformGroupName=Desktop
 [ShaderPlatform PCD3D_SM5]
 Language=D3D
 MaxFeatureLevel=SM5
+ShaderFormat=PCD3D_SM5
 bIsPC=true
 bSupportsDebugViewShaders=true
 bSupportsRayTracing=true
@@ -102,6 +103,7 @@ MenuTooltip=LOCTEXT("PreviewMenuTooltip_PCD3D_SM5", "DirectX 11, OpenGL 4.3+, PS
 FriendlyName=LOCTEXT("FriendlyShaderPlatformName_PCD3D_SM6", "Shader Model 6")
 Language=D3D
 MaxFeatureLevel=SM6
+ShaderFormat=PCD3D_SM6
 bIsPC=true
 MaxMeshShaderThreadGroupSize=128
 bSupportsDebugViewShaders=true
@@ -153,6 +155,7 @@ MenuTooltip=LOCTEXT("PreviewMenuTooltip_PCD3D_SM6", "Preview using SM6 device pr
 [ShaderPlatform PCD3D_ES3_1]
 Language=D3D
 MaxFeatureLevel=ES3_1
+ShaderFormat=PCD3D_ES31
 bIsPC=true
 bIsMobile=true
 bSupportsDistanceFields = true
@@ -164,5 +167,6 @@ bSupportsManualVertexFetch=false
 Language=OpenGL
 bIsPC=true
 MaxFeatureLevel=ES3_1
+ShaderFormat=GLSL_150_ES31
 bIsMobile=true
 bSupportsManualVertexFetch=false
diff --git a/Engine/Plugins/Compositing/OpenColorIO/Source/OpenColorIO/Private/OpenColorIOShaderCompilationManager.cpp b/Engine/Plugins/Compositing/OpenColorIO/Source/OpenColorIO/Private/OpenColorIOShaderCompilationManager.cpp
--- a/Engine/Plugins/Compositing/OpenColorIO/Source/OpenColorIO/Private/OpenColorIOShaderCompilationManager.cpp
+++ b/Engine/Plugins/Compositing/OpenColorIO/Source/OpenColorIO/Private/OpenColorIOShaderCompilationManager.cpp
@@ -82,7 +82,7 @@ OPENCOLORIO_API void FOpenColorIOShaderCompilationManager::AddJobs(TArray<FShade
 			CurrentJob->Input.Environment.CompilerFlags.Add(CFLAG_NoFastMath);
 		}
 
-		CurrentJob->Input.DumpDebugInfoRootPath = GShaderCompilingManager->GetAbsoluteShaderDebugInfoDirectory() / LegacyShaderPlatformToShaderFormat(EShaderPlatform(CurrentJob->Input.Target.Platform)).ToString();
+		CurrentJob->Input.DumpDebugInfoRootPath = GShaderCompilingManager->GetAbsoluteShaderDebugInfoDirectory() / CurrentJob->Input.ShaderPlatformName.ToString();
 		FPaths::NormalizeDirectoryName(CurrentJob->Input.DumpDebugInfoRootPath);
 		CurrentJob->Input.DebugExtension.Empty();
 		CurrentJob->Input.DumpDebugInfoPath.Empty();
diff --git a/Engine/Plugins/Experimental/ComputeFramework/Source/ComputeFramework/Private/ComputeKernelShaderCompilationManager.cpp b/Engine/Plugins/Experimental/ComputeFramework/Source/ComputeFramework/Private/ComputeKernelShaderCompilationManager.cpp
--- a/Engine/Plugins/Experimental/ComputeFramework/Source/ComputeFramework/Private/ComputeKernelShaderCompilationManager.cpp
+++ b/Engine/Plugins/Experimental/ComputeFramework/Source/ComputeFramework/Private/ComputeKernelShaderCompilationManager.cpp
@@ -115,7 +115,7 @@ void FComputeKernelShaderCompilationManager::RunCompileJobs()
 
 				UE_LOG(LogComputeKernelShaderCompiler, Log, TEXT("Compile Job processing... %s"), *CurrentJob.Input.DebugGroupName);
 
-				CurrentJob.Input.DumpDebugInfoRootPath = GShaderCompilingManager->GetAbsoluteShaderDebugInfoDirectory() / Format.ToString();
+				CurrentJob.Input.DumpDebugInfoRootPath = GShaderCompilingManager->GetAbsoluteShaderDebugInfoDirectory() / CurrentJob.Input.ShaderPlatformName.ToString();
 				FPaths::NormalizeDirectoryName(CurrentJob.Input.DumpDebugInfoRootPath);
 				const FShaderCompilingManager::EDumpShaderDebugInfo DumpShaderDebugInfo = GShaderCompilingManager->GetDumpShaderDebugInfo();
 				CurrentJob.Input.DebugExtension.Empty();
diff --git a/Engine/Plugins/FX/Niagara/Shaders/Private/NiagaraShaderVersion.ush b/Engine/Plugins/FX/Niagara/Shaders/Private/NiagaraShaderVersion.ush
--- a/Engine/Plugins/FX/Niagara/Shaders/Private/NiagaraShaderVersion.ush
+++ b/Engine/Plugins/FX/Niagara/Shaders/Private/NiagaraShaderVersion.ush
@@ -3,4 +3,4 @@
 // in the DDC key, it will automatically invalidate.
 // 
 // If you are merging streams and there is a conflict with this GUID you should make a new GUID rather than taking one or the other.
-// GUID = 32F26874A0DC4FADAA73ED5FE11D0FD7
+// GUID = 1100075F7722451B96F0AB70CCFDD49C
diff --git a/Engine/Plugins/FX/Niagara/Source/NiagaraShader/Private/NiagaraShaderCompilationManager.cpp b/Engine/Plugins/FX/Niagara/Source/NiagaraShader/Private/NiagaraShaderCompilationManager.cpp
--- a/Engine/Plugins/FX/Niagara/Source/NiagaraShader/Private/NiagaraShaderCompilationManager.cpp
+++ b/Engine/Plugins/FX/Niagara/Source/NiagaraShader/Private/NiagaraShaderCompilationManager.cpp
@@ -49,7 +49,7 @@ NIAGARASHADER_API void FNiagaraShaderCompilationManager::AddJobs(TArray<FShaderC
 
 		UE_LOG(LogNiagaraShaderCompiler, Verbose, TEXT("Adding niagara gpu shader compile job... %s"), *CurrentJob->Input.DebugGroupName);
 
-		CurrentJob->Input.DumpDebugInfoRootPath = GShaderCompilingManager->GetAbsoluteShaderDebugInfoDirectory() / LegacyShaderPlatformToShaderFormat(EShaderPlatform(CurrentJob->Input.Target.Platform)).ToString();
+		CurrentJob->Input.DumpDebugInfoRootPath = GShaderCompilingManager->GetAbsoluteShaderDebugInfoDirectory() / CurrentJob->Input.ShaderPlatformName.ToString();
 		FPaths::NormalizeDirectoryName(CurrentJob->Input.DumpDebugInfoRootPath);
 		CurrentJob->Input.DebugExtension.Empty();
 		CurrentJob->Input.DumpDebugInfoPath.Empty();
diff --git a/Engine/Shaders/Public/ShaderVersion.ush b/Engine/Shaders/Public/ShaderVersion.ush
--- a/Engine/Shaders/Public/ShaderVersion.ush
+++ b/Engine/Shaders/Public/ShaderVersion.ush
@@ -2,4 +2,4 @@
 // Each time the console command is executed it generates a new GUID. As this file is included
 // in Platform.ush (which should be included in any shader) it allows to invalidate the shader DDC.
 // 
-// GUID = 61380C6B-BC5B-4BEC-94A4-FE7397B0ACDF
+// GUID = 84806049-D47D-41D3-A61D-0D2CC08B1073
diff --git a/Engine/Source/Developer/ShaderCompilerCommon/Private/ShaderCompilerCommon.cpp b/Engine/Source/Developer/ShaderCompilerCommon/Private/ShaderCompilerCommon.cpp
--- a/Engine/Source/Developer/ShaderCompilerCommon/Private/ShaderCompilerCommon.cpp
+++ b/Engine/Source/Developer/ShaderCompilerCommon/Private/ShaderCompilerCommon.cpp
@@ -923,6 +923,10 @@ FString CreateShaderCompilerWorkerDirectCommandLine(const FShaderCompilerInput&
 	Text += Input.ShaderFormat.GetPlainNameString();
 	Text += TEXT(" -entry=");
 	Text += Input.EntryPointName;
+
+	Text += TEXT(" -shaderPlatformName=");
+	Text += Input.ShaderPlatformName.GetPlainNameString();
+
 	switch (Input.Target.Frequency)
 	{
 	case SF_Vertex:			Text += TEXT(" -vs"); break;
diff --git a/Engine/Source/Developer/ShaderPreprocessor/Private/ShaderPreprocessor.cpp b/Engine/Source/Developer/ShaderPreprocessor/Private/ShaderPreprocessor.cpp
--- a/Engine/Source/Developer/ShaderPreprocessor/Private/ShaderPreprocessor.cpp
+++ b/Engine/Source/Developer/ShaderPreprocessor/Private/ShaderPreprocessor.cpp
@@ -36,7 +36,7 @@ public:
 		, ShaderOutput(InShaderOutput)
 	{
 		FString InputShaderSource;
-		if (LoadShaderSourceFile(*InShaderInput.VirtualSourceFilePath, InShaderInput.Target.GetPlatform(),  &InputShaderSource, nullptr))
+		if (LoadShaderSourceFile(*InShaderInput.VirtualSourceFilePath, InShaderInput.Target.GetPlatform(),  &InputShaderSource, nullptr, &InShaderInput.ShaderPlatformName))
 		{
 			InputShaderSource = FString::Printf(TEXT("%s\n#line 1\n%s"), *ShaderInput.SourceFilePrefix, *InputShaderSource);
 			CachedFileContents.Add(InShaderInput.VirtualSourceFilePath, StringToArray<ANSICHAR>(*InputShaderSource, InputShaderSource.Len() + 1));
@@ -73,7 +73,7 @@ private:
 		ReplaceVirtualFilePathForShaderPlatform(VirtualFilePath, This->ShaderInput.Target.GetPlatform());
 
 		// Fixup autogen file
-		ReplaceVirtualFilePathForShaderAutogen(VirtualFilePath, This->ShaderInput.Target.GetPlatform());
+		ReplaceVirtualFilePathForShaderAutogen(VirtualFilePath, This->ShaderInput.Target.GetPlatform(), &This->ShaderInput.ShaderPlatformName);
 
 		// Collapse any relative directories to allow #include "../MyFile.ush"
 		FPaths::CollapseRelativeDirectories(VirtualFilePath);
@@ -95,7 +95,7 @@ private:
 			{
 				CheckShaderHashCacheInclude(VirtualFilePath, This->ShaderInput.Target.GetPlatform());
 
-				LoadShaderSourceFile(*VirtualFilePath, This->ShaderInput.Target.GetPlatform(), &FileContents, &This->ShaderOutput.Errors);
+				LoadShaderSourceFile(*VirtualFilePath, This->ShaderInput.Target.GetPlatform(), &FileContents, &This->ShaderOutput.Errors, &This->ShaderInput.ShaderPlatformName);
 			}
 
 			if (FileContents.Len() > 0)
diff --git a/Engine/Source/Editor/UnrealEd/Private/Commandlets/CookGlobalShaderCommandlet.cpp b/Engine/Source/Editor/UnrealEd/Private/Commandlets/CookGlobalShaderCommandlet.cpp
--- a/Engine/Source/Editor/UnrealEd/Private/Commandlets/CookGlobalShaderCommandlet.cpp
+++ b/Engine/Source/Editor/UnrealEd/Private/Commandlets/CookGlobalShaderCommandlet.cpp
@@ -49,8 +49,8 @@ void UCookGlobalShadersCommandlet::CookGlobalShaders() const
 	TArray<TPair<FString, FString>> FilesToCopy;
 	for (FName ShaderFormat : ShaderFormats)
 	{
-		const FString GlobalShaderCacheName = FPaths::Combine(OutputDir, TEXT("Engine"), TEXT("GlobalShaderCache-") + ShaderFormat.ToString() + TEXT(".bin"));
-		const FString OverrideGlobalShaderCacheName = FPaths::Combine(TEXT("Engine"), TEXT("OverrideGlobalShaderCache-") + ShaderFormat.ToString() + TEXT(".bin"));
+		const FString GlobalShaderCacheName = FPaths::Combine(OutputDir, TEXT("Engine"), TEXT("GlobalShaderCache-") + FDataDrivenShaderPlatformInfo::GetName(ShaderFormatToLegacyShaderPlatform(ShaderFormat)).ToString() + TEXT(".bin"));
+		const FString OverrideGlobalShaderCacheName = FPaths::Combine(TEXT("Engine"), TEXT("OverrideGlobalShaderCache-") + FDataDrivenShaderPlatformInfo::GetName(ShaderFormatToLegacyShaderPlatform(ShaderFormat)).ToString() + TEXT(".bin"));
 		FilesToCopy.Emplace(GlobalShaderCacheName, OverrideGlobalShaderCacheName);
 	}
 
diff --git a/Engine/Source/Programs/ShaderCompileWorker/Private/ShaderCompileWorker.cpp b/Engine/Source/Programs/ShaderCompileWorker/Private/ShaderCompileWorker.cpp
--- a/Engine/Source/Programs/ShaderCompileWorker/Private/ShaderCompileWorker.cpp
+++ b/Engine/Source/Programs/ShaderCompileWorker/Private/ShaderCompileWorker.cpp
@@ -599,8 +599,10 @@ private:
 
 			// Flush cache, to make sure we load the latest version of the input file.
 			// (Otherwise quick changes to a shader file can result in the wrong output.)
-			FlushShaderFileCache();
-
+			FName ShaderPlatformName;
+			InputFile << ShaderPlatformName;
+			FlushShaderFileCache(&ShaderPlatformName);
+			
 			for (int32 BatchIndex = 0; BatchIndex < NumBatches; BatchIndex++)
 			{
 				// Deserialize the job's inputs.
@@ -915,6 +917,7 @@ static void DirectCompile(const TArray<const class IShaderFormat*>& ShaderFormat
 	FString InputFile;
 
 	FName FormatName;
+	FName ShaderPlatformName;
 	FString Entry = TEXT("Main");
 	bool bPipeline = false;
 	bool bUseMCPP = false;
@@ -940,6 +943,10 @@ static void DirectCompile(const TArray<const class IShaderFormat*>& ShaderFormat
 					Entry = Entry.Mid(1, Entry.Len() - 2);
 				}
 			}
+			else if (Token.StartsWith(TEXT("shaderPlatformName=")))
+			{
+				ShaderPlatformName = FName(*Token.RightChop(19));
+			}
 			else if (Token.StartsWith(TEXT("cflags=")))
 			{
 				CFlags = FCString::Atoi64(*Token.RightChop(7));
@@ -1021,6 +1028,7 @@ static void DirectCompile(const TArray<const class IShaderFormat*>& ShaderFormat
 	FShaderCompilerInput Input;
 	Input.EntryPointName = Entry;
 	Input.ShaderFormat = FormatName;
+	Input.ShaderPlatformName = ShaderPlatformName;
 	Input.VirtualSourceFilePath = InputFile;
 	Input.Target.Platform =  ShaderFormatNameToShaderPlatform(FormatName);
 	Input.Target.Frequency = Frequency;
diff --git a/Engine/Source/Runtime/Engine/Private/ShaderCompiler/ShaderCompiler.cpp b/Engine/Source/Runtime/Engine/Private/ShaderCompiler/ShaderCompiler.cpp
--- a/Engine/Source/Runtime/Engine/Private/ShaderCompiler/ShaderCompiler.cpp
+++ b/Engine/Source/Runtime/Engine/Private/ShaderCompiler/ShaderCompiler.cpp
@@ -1445,6 +1445,9 @@ bool FShaderCompileUtilities::DoWriteTasks(const TArray<FShaderCommonCompileJobP
 		int32 NumBatches = QueuedSingleJobs.Num();
 		TransferFile << NumBatches;
 
+		FName ShaderPlatformName = FDataDrivenShaderPlatformInfo::GetName(GMaxRHIShaderPlatform);
+		TransferFile << ShaderPlatformName;
+
 		// Serialize all the batched jobs
 		for (int32 JobIndex = 0; JobIndex < QueuedSingleJobs.Num(); JobIndex++)
 		{
@@ -5357,6 +5360,7 @@ void GlobalBeginCompileShader(
 #endif
 
 	Input.Target = Target;
+	Input.ShaderPlatformName = FDataDrivenShaderPlatformInfo::GetName(ShaderPlatform);
 	Input.ShaderFormat = ShaderFormatName;
 	Input.CompressionFormat = GetShaderCompressionFormat();
 	GetShaderCompressionOodleSettings(Input.OodleCompressor, Input.OodleLevel);
@@ -5365,7 +5369,7 @@ void GlobalBeginCompileShader(
 	Input.bCompilingForShaderPipeline = false;
 	Input.bIncludeUsedOutputs = false;
 	Input.bGenerateDirectCompileFile = (GDumpShaderDebugInfoSCWCommandLine != 0);
-	Input.DumpDebugInfoRootPath = GShaderCompilingManager->GetAbsoluteShaderDebugInfoDirectory() / ShaderFormatName.ToString();
+	Input.DumpDebugInfoRootPath = GShaderCompilingManager->GetAbsoluteShaderDebugInfoDirectory() / Input.ShaderPlatformName.ToString();
 	// asset material name or "Global"
 	Input.DebugGroupName = DebugGroupName;
 	Input.DebugDescription = DebugDescription;
@@ -6750,7 +6754,7 @@ void VerifyGlobalShaders(EShaderPlatform Platform, bool bLoadedFromCacheFile, co
 
 #include "Misc/PreLoadFile.h"
 #include "Serialization/LargeMemoryReader.h"
-static FPreLoadFile GGlobalShaderPreLoadFile(*(FString(TEXT("../../../Engine")) / TEXT("GlobalShaderCache-SF_") + FPlatformProperties::IniPlatformName() + TEXT(".bin")));
+static FPreLoadFile GGlobalShaderPreLoadFile(*(FString(TEXT("../../../Engine")) / TEXT("GlobalShaderCache-SP_") + FPlatformProperties::IniPlatformName() + TEXT(".bin")));
 
 static const ITargetPlatform* GGlobalShaderTargetPlatform[SP_NumPlatforms] = { nullptr };
 
@@ -6765,12 +6769,12 @@ static FString GetGlobalShaderCacheOverrideFilename(EShaderPlatform Platform)
 		DirectoryPrefix = GGlobalShaderCacheOverrideDirectory / TEXT("GlobalShaderCache-");
 	}
 
-	return DirectoryPrefix + LegacyShaderPlatformToShaderFormat(Platform).ToString() + TEXT(".bin");
+	return DirectoryPrefix + FDataDrivenShaderPlatformInfo::GetName(Platform).ToString() + TEXT(".bin");
 }
 
 static FString GetGlobalShaderCacheFilename(EShaderPlatform Platform)
 {
-	return FString(TEXT("Engine")) / TEXT("GlobalShaderCache-") + LegacyShaderPlatformToShaderFormat(Platform).ToString() + TEXT(".bin");
+	return FString(TEXT("Engine")) / TEXT("GlobalShaderCache-") + FDataDrivenShaderPlatformInfo::GetName(Platform).ToString() + TEXT(".bin");
 }
 
 #if WITH_EDITOR
diff --git a/Engine/Source/Runtime/Engine/Private/ShaderCompiler/ShaderGenerationUtil.cpp b/Engine/Source/Runtime/Engine/Private/ShaderCompiler/ShaderGenerationUtil.cpp
--- a/Engine/Source/Runtime/Engine/Private/ShaderCompiler/ShaderGenerationUtil.cpp
+++ b/Engine/Source/Runtime/Engine/Private/ShaderCompiler/ShaderGenerationUtil.cpp
@@ -1530,7 +1530,7 @@ static FCriticalSection GCriticalSection;
 
 static FString GetAutoGenDirectory(EShaderPlatform TargetPlatform)
 {
-	FString PlatformName = LegacyShaderPlatformToShaderFormat(TargetPlatform).GetPlainNameString();
+	FString PlatformName = FDataDrivenShaderPlatformInfo::GetName(TargetPlatform).ToString();
 	FString AutogenHeaderDirectory = FPaths::ProjectIntermediateDir() / TEXT("ShaderAutogen") / PlatformName;
 	return AutogenHeaderDirectory;
 }
@@ -1967,7 +1967,6 @@ FGBufferInfo FShaderCompileUtilities::FetchGBufferInfoAndWriteAutogen(EShaderPla
 
 	FGBufferInfo BufferInfo = FetchFullGBufferInfo(Params);
 
-	FString PlatformName = LegacyShaderPlatformToShaderFormat(TargetPlatform).GetPlainNameString();
 	FString AutoGenDirectory = GetAutoGenDirectory(TargetPlatform);
 	FString AutogenHeaderFilename = AutoGenDirectory / TEXT("AutogenShaderHeaders.ush");
 	FString AutogenHeaderFilenameTemp = AutoGenDirectory / TEXT("AutogenShaderHeaders_temp.ush");
diff --git a/Engine/Source/Runtime/RHI/Private/RHI.cpp b/Engine/Source/Runtime/RHI/Private/RHI.cpp
--- a/Engine/Source/Runtime/RHI/Private/RHI.cpp
+++ b/Engine/Source/Runtime/RHI/Private/RHI.cpp
@@ -2431,7 +2431,7 @@ void FGenericDataDrivenShaderPlatformInfo::SetDefaultValues()
 {
 	MaxFeatureLevel = ERHIFeatureLevel::Num;
 	bSupportsMSAA = true;
-
+	 
 	bNeedsToSwitchVerticalAxisOnMobileOpenGL = true;
 	bSupportsDOFHybridScattering = true;
 	bSupportsHZBOcclusion = true;
@@ -2445,6 +2445,8 @@ void FGenericDataDrivenShaderPlatformInfo::ParseDataDrivenShaderInfo(const FConf
 {
 	Info.Language = *GetSectionString(Section, "Language");
 	Info.ShaderFormat = *GetSectionString(Section, "ShaderFormat");
+	checkf(!Info.ShaderFormat.IsNone(), TEXT("Missing ShaderFormat for ShaderPlatform %s  ShaderFormat %s"), *Info.Name.ToString(), *Info.ShaderFormat.ToString());
+
 	GetFeatureLevelFromName(*GetSectionString(Section, "MaxFeatureLevel"), Info.MaxFeatureLevel);
 
 	Info.ShaderPropertiesHash = 0;
diff --git a/Engine/Source/Runtime/RenderCore/Private/ShaderCodeArchive.cpp b/Engine/Source/Runtime/RenderCore/Private/ShaderCodeArchive.cpp
--- a/Engine/Source/Runtime/RenderCore/Private/ShaderCodeArchive.cpp
+++ b/Engine/Source/Runtime/RenderCore/Private/ShaderCodeArchive.cpp
@@ -1643,7 +1643,7 @@ void FIoStoreShaderCodeArchive::SaveIoStoreShaderCodeArchive(const FIoStoreShade
 
 FIoStoreShaderCodeArchive* FIoStoreShaderCodeArchive::Create(EShaderPlatform InPlatform, const FString& InLibraryName, FIoDispatcher& InIoDispatcher)
 {
-	const FName PlatformName = LegacyShaderPlatformToShaderFormat(InPlatform);
+	const FName PlatformName = FName(FDataDrivenShaderPlatformInfo::GetShaderFormat(InPlatform).ToString() + TEXT("-") + FDataDrivenShaderPlatformInfo::GetName(InPlatform).ToString());
 	FIoChunkId ChunkId = GetShaderCodeArchiveChunkId(InLibraryName, PlatformName);
 	if (InIoDispatcher.DoesChunkExist(ChunkId))
 	{
diff --git a/Engine/Source/Runtime/RenderCore/Private/ShaderCodeLibrary.cpp b/Engine/Source/Runtime/RenderCore/Private/ShaderCodeLibrary.cpp
--- a/Engine/Source/Runtime/RenderCore/Private/ShaderCodeLibrary.cpp
+++ b/Engine/Source/Runtime/RenderCore/Private/ShaderCodeLibrary.cpp
@@ -766,8 +766,11 @@ public:
 
 		if (!Library)
 		{
-			const FName PlatformName = LegacyShaderPlatformToShaderFormat(InShaderPlatform);
-			const FString DestFilePath = GetCodeArchiveFilename(ShaderCodeDir, InLibraryName, PlatformName);
+			const FName PlatformName = FDataDrivenShaderPlatformInfo::GetName(InShaderPlatform);
+			const FName ShaderFormatName = LegacyShaderPlatformToShaderFormat(InShaderPlatform);
+			FString ShaderFormatAndPlatform = ShaderFormatName.ToString() + TEXT("-") + PlatformName.ToString();
+
+			const FString DestFilePath = GetCodeArchiveFilename(ShaderCodeDir, InLibraryName, FName(ShaderFormatAndPlatform));
 			TUniquePtr<FArchive> Ar(IFileManager::Get().CreateFileReader(*DestFilePath));
 			if (Ar)
 			{
@@ -1119,7 +1122,13 @@ struct FEditorShaderCodeArchive
 		, Format(nullptr)
 		, bNeedsDeterministicOrder(bInNeedsDeterministicOrder)
 	{
-		Format = GetTargetPlatformManagerRef().FindShaderFormat(InFormat);
+		TArray<FString> Components;
+		FString Name = FormatName.ToString();
+		Name.ParseIntoArray(Components, TEXT("-"));
+
+		FName ShaderFormatName(Components[0]);
+		Format = GetTargetPlatformManagerRef().FindShaderFormat(ShaderFormatName);
+
 		check(Format);
 
 		SerializedShaders.ShaderHashTable.Clear(0x10000);
@@ -2331,10 +2340,12 @@ public:
 
 			EShaderPlatform Platform = ShaderFormatToLegacyShaderPlatform(Format);
 			FName PossiblyAdjustedFormat = LegacyShaderPlatformToShaderFormat(Platform);	// Vulkan and GL switch between name variants depending on CVars 
+
+			FString FormatAndPlatformName = PossiblyAdjustedFormat.ToString() + TEXT("-") + FDataDrivenShaderPlatformInfo::GetName(Platform).ToString();
 			FEditorShaderCodeArchive* CodeArchive = EditorShaderCodeArchive[Platform];
 			if (!CodeArchive)
 			{
-				CodeArchive = new FEditorShaderCodeArchive(PossiblyAdjustedFormat, Descriptor.bNeedsDeterministicOrder);
+				CodeArchive = new FEditorShaderCodeArchive(FName(FormatAndPlatformName), Descriptor.bNeedsDeterministicOrder);
 				EditorShaderCodeArchive[Platform] = CodeArchive;
 				EditorArchivePipelines[Platform] = !bNativeFormat;
 			}
@@ -2979,10 +2990,13 @@ bool FShaderLibraryCooker::CreatePatchLibrary(TArray<FString> const& OldMetaData
 		if (Name.RemoveFromStart(TEXT("ShaderArchive-")))
 		{
 			TArray<FString> Components;
-			if (Name.ParseIntoArray(Components, TEXT("-")) == 2)
+			if (Name.ParseIntoArray(Components, TEXT("-")) == 3)
 			{
 				FName Format(*Components[1]);
-				TSet<FString>& Libraries = FormatLibraryMap.FindOrAdd(Format);
+				FName Platform(*Components[2]);
+				FString FormatAndPlatform = Format.ToString() + TEXT("-") + Platform.ToString();
+
+				TSet<FString>& Libraries = FormatLibraryMap.FindOrAdd(FName(FormatAndPlatform));
 				Libraries.Add(Components[0]);
 			}
 		}
diff --git a/Engine/Source/Runtime/RenderCore/Private/ShaderCore.cpp b/Engine/Source/Runtime/RenderCore/Private/ShaderCore.cpp
--- a/Engine/Source/Runtime/RenderCore/Private/ShaderCore.cpp
+++ b/Engine/Source/Runtime/RenderCore/Private/ShaderCore.cpp
@@ -660,7 +660,7 @@ bool CheckVirtualShaderFilePath(FStringView VirtualFilePath, TArray<FShaderCompi
 * @param OutVirtualFilePaths - [out] list of shader source files to add to
 * @param ShaderFilename - shader file to add
 */
-void AddShaderSourceFileEntry(TArray<FString>& OutVirtualFilePaths, FString VirtualFilePath, EShaderPlatform ShaderPlatform)
+void AddShaderSourceFileEntry(TArray<FString>& OutVirtualFilePaths, FString VirtualFilePath, EShaderPlatform ShaderPlatform, const FName* ShaderPlatformName)
 {
 	check(CheckVirtualShaderFilePath(VirtualFilePath));
 	if (!OutVirtualFilePaths.Contains(VirtualFilePath))
@@ -668,7 +668,9 @@ void AddShaderSourceFileEntry(TArray<FString>& OutVirtualFilePaths, FString Virt
 		OutVirtualFilePaths.Add(VirtualFilePath);
 
 		TArray<FString> ShaderIncludes;
-		GetShaderIncludes(*VirtualFilePath, *VirtualFilePath, OutVirtualFilePaths, ShaderPlatform);
+
+		const uint32 DepthLimit = 100;
+		GetShaderIncludes(*VirtualFilePath, *VirtualFilePath, OutVirtualFilePaths, ShaderPlatform, DepthLimit, ShaderPlatformName);
 		for( int32 IncludeIdx=0; IncludeIdx < ShaderIncludes.Num(); IncludeIdx++ )
 		{
 			OutVirtualFilePaths.AddUnique(ShaderIncludes[IncludeIdx]);
@@ -681,7 +683,7 @@ void AddShaderSourceFileEntry(TArray<FString>& OutVirtualFilePaths, FString Virt
 *
 * @param OutVirtualFilePaths - [out] list of shader source files to add to
 */
-void GetAllVirtualShaderSourcePaths(TArray<FString>& OutVirtualFilePaths, EShaderPlatform ShaderPlatform)
+void GetAllVirtualShaderSourcePaths(TArray<FString>& OutVirtualFilePaths, EShaderPlatform ShaderPlatform, const FName* ShaderPlatformName)
 {
 	// add all shader source files for hashing
 	for( TLinkedList<FVertexFactoryType*>::TIterator FactoryIt(FVertexFactoryType::GetTypeList()); FactoryIt; FactoryIt.Next() )
@@ -690,7 +692,7 @@ void GetAllVirtualShaderSourcePaths(TArray<FString>& OutVirtualFilePaths, EShade
 		if( VertexFactoryType )
 		{
 			FString ShaderFilename(VertexFactoryType->GetShaderFilename());
-			AddShaderSourceFileEntry(OutVirtualFilePaths, ShaderFilename, ShaderPlatform);
+			AddShaderSourceFileEntry(OutVirtualFilePaths, ShaderFilename, ShaderPlatform, ShaderPlatformName);
 		}
 	}
 	for( TLinkedList<FShaderType*>::TIterator ShaderIt(FShaderType::GetTypeList()); ShaderIt; ShaderIt.Next() )
@@ -699,7 +701,7 @@ void GetAllVirtualShaderSourcePaths(TArray<FString>& OutVirtualFilePaths, EShade
 		if(ShaderType)
 		{
 			FString ShaderFilename(ShaderType->GetShaderFilename());
-			AddShaderSourceFileEntry(OutVirtualFilePaths, ShaderFilename, ShaderPlatform);
+			AddShaderSourceFileEntry(OutVirtualFilePaths, ShaderFilename, ShaderPlatform, ShaderPlatformName);
 		}
 	}
 
@@ -892,7 +894,7 @@ bool ReplaceVirtualFilePathForShaderPlatform(FString& InOutVirtualFilePath, ESha
 	return false;
 }
 
-bool ReplaceVirtualFilePathForShaderAutogen(FString& InOutVirtualFilePath, EShaderPlatform ShaderPlatform)
+bool ReplaceVirtualFilePathForShaderAutogen(FString& InOutVirtualFilePath, EShaderPlatform ShaderPlatform, const FName* InShaderPlatformName)
 {
 	const FStringView ShaderAutogenStem = TEXTVIEW("/Engine/Generated/ShaderAutogen/");
 
@@ -900,16 +902,17 @@ bool ReplaceVirtualFilePathForShaderAutogen(FString& InOutVirtualFilePath, EShad
 	// for examples, if it starts with "/Engine/Generated/ShaderAutogen/" change it to "ShaderAutogen/PCD3D_SM5/"
 	if (!FCString::Strnicmp(*InOutVirtualFilePath, ShaderAutogenStem.GetData(), ShaderAutogenStem.Len()))
 	{
+		check(FDataDrivenShaderPlatformInfo::IsValid(ShaderPlatform) || InShaderPlatformName != nullptr);
+		FName ShaderPlatformName = FDataDrivenShaderPlatformInfo::IsValid(ShaderPlatform) ? FDataDrivenShaderPlatformInfo::GetName(ShaderPlatform) : *InShaderPlatformName;
 		TStringBuilder<1024> OutputShaderName;
 
 		// Append the prefix.
 		OutputShaderName.Append(TEXTVIEW("/ShaderAutogen/"));
 
 		// Append the platform name.
-		const FName PlatformName = LegacyShaderPlatformToShaderFormat(ShaderPlatform);
-		TCHAR PlatformNameString[NAME_SIZE];
-		const uint32 PlatformNameLen = PlatformName.GetPlainNameString(PlatformNameString);
-		OutputShaderName.Append(PlatformNameString, PlatformNameLen);
+
+		FString PlatformNameString = ShaderPlatformName.ToString();
+		OutputShaderName.Append(*PlatformNameString, PlatformNameString.Len());
 
 		OutputShaderName.AppendChar(TEXT('/'));
 
@@ -924,7 +927,7 @@ bool ReplaceVirtualFilePathForShaderAutogen(FString& InOutVirtualFilePath, EShad
 	return false;
 }
 
-bool LoadShaderSourceFile(const TCHAR* InVirtualFilePath, EShaderPlatform ShaderPlatform, FString* OutFileContents, TArray<FShaderCompilerError>* OutCompileErrors) // TODO: const FString&
+bool LoadShaderSourceFile(const TCHAR* InVirtualFilePath, EShaderPlatform ShaderPlatform, FString* OutFileContents, TArray<FShaderCompilerError>* OutCompileErrors, const FName* ShaderPlatformName) // TODO: const FString&
 {
 	// it's not expected that cooked platforms get here, but if they do, this is the final out
 	if (FPlatformProperties::RequiresCookedData())
@@ -944,7 +947,7 @@ bool LoadShaderSourceFile(const TCHAR* InVirtualFilePath, EShaderPlatform Shader
 		ReplaceVirtualFilePathForShaderPlatform(VirtualFilePath, ShaderPlatform);
 
 		// Fixup autogen file
-		ReplaceVirtualFilePathForShaderAutogen(VirtualFilePath, ShaderPlatform);
+		ReplaceVirtualFilePathForShaderAutogen(VirtualFilePath, ShaderPlatform, ShaderPlatformName);
 
 		// Protect GShaderFileCache from simultaneous access by multiple threads
 		FScopeLock ScopeLock(&FileCacheCriticalSection);
@@ -1048,7 +1051,7 @@ static const TCHAR* FindFirstInclude(const TCHAR* Text)
 /**
  * Recursively populates IncludeFilenames with the unique include filenames found in the shader file named Filename.
  */
-static void GetShaderIncludes(const TCHAR* EntryPointVirtualFilePath, const TCHAR* VirtualFilePath, const FString& FileContents, TArray<FString>& IncludeVirtualFilePaths, EShaderPlatform ShaderPlatform, uint32 DepthLimit, bool AddToIncludeFile)
+static void InternalGetShaderIncludes(const TCHAR* EntryPointVirtualFilePath, const TCHAR* VirtualFilePath, const FString& FileContents, TArray<FString>& IncludeVirtualFilePaths, EShaderPlatform ShaderPlatform, uint32 DepthLimit, bool AddToIncludeFile, const FName* ShaderPlatformName)
 {
 	//avoid an infinite loop with a 0 length string
 	if (FileContents.Len() > 0)
@@ -1097,7 +1100,7 @@ static void GetShaderIncludes(const TCHAR* EntryPointVirtualFilePath, const TCHA
 					ReplaceVirtualFilePathForShaderPlatform(ExtractedIncludeFilename, ShaderPlatform);
 
 					// Fixup autogen file
-					ReplaceVirtualFilePathForShaderAutogen(ExtractedIncludeFilename, ShaderPlatform);
+					ReplaceVirtualFilePathForShaderAutogen(ExtractedIncludeFilename, ShaderPlatform, ShaderPlatformName);
 
 					// Ignore uniform buffer, vertex factory and instanced stereo includes
 					bool bIgnoreInclude = ExtractedIncludeFilename.StartsWith(TEXT("/Engine/Generated/"));
@@ -1118,8 +1121,8 @@ static void GetShaderIncludes(const TCHAR* EntryPointVirtualFilePath, const TCHA
 						if (!IncludeVirtualFilePaths.Contains(ExtractedIncludeFilename))
 						{
 							FString IncludedFileContents;
-							LoadShaderSourceFile(*ExtractedIncludeFilename, ShaderPlatform, &IncludedFileContents, nullptr);
-							GetShaderIncludes(EntryPointVirtualFilePath, *ExtractedIncludeFilename, IncludedFileContents, IncludeVirtualFilePaths, ShaderPlatform, DepthLimit - 1, true);
+							LoadShaderSourceFile(*ExtractedIncludeFilename, ShaderPlatform, &IncludedFileContents, nullptr, ShaderPlatformName);
+							InternalGetShaderIncludes(EntryPointVirtualFilePath, *ExtractedIncludeFilename, IncludedFileContents, IncludeVirtualFilePaths, ShaderPlatform, DepthLimit - 1, true, ShaderPlatformName);
 						}
 					}
 				}
@@ -1150,22 +1153,22 @@ static void GetShaderIncludes(const TCHAR* EntryPointVirtualFilePath, const TCHA
 /**
  * Recursively populates IncludeFilenames with the unique include filenames found in the shader file named Filename.
  */
-static void GetShaderIncludes(const TCHAR* EntryPointVirtualFilePath, const TCHAR* VirtualFilePath, TArray<FString>& IncludeVirtualFilePaths, EShaderPlatform ShaderPlatform, uint32 DepthLimit, bool AddToIncludeFile)
+static void InternalGetShaderIncludes(const TCHAR* EntryPointVirtualFilePath, const TCHAR* VirtualFilePath, TArray<FString>& IncludeVirtualFilePaths, EShaderPlatform ShaderPlatform, uint32 DepthLimit, bool AddToIncludeFile, const FName* ShaderPlatformName)
 {
 	FString FileContents;
-	LoadShaderSourceFile(VirtualFilePath, ShaderPlatform, &FileContents, nullptr);
+	LoadShaderSourceFile(VirtualFilePath, ShaderPlatform, &FileContents, nullptr, ShaderPlatformName);
 
-	GetShaderIncludes(EntryPointVirtualFilePath, VirtualFilePath, FileContents, IncludeVirtualFilePaths, ShaderPlatform, DepthLimit, AddToIncludeFile);
+	InternalGetShaderIncludes(EntryPointVirtualFilePath, VirtualFilePath, FileContents, IncludeVirtualFilePaths, ShaderPlatform, DepthLimit, AddToIncludeFile, ShaderPlatformName);
 }
 
-void GetShaderIncludes(const TCHAR* EntryPointVirtualFilePath, const TCHAR* VirtualFilePath, TArray<FString>& IncludeVirtualFilePaths, EShaderPlatform ShaderPlatform, uint32 DepthLimit)
+void GetShaderIncludes(const TCHAR* EntryPointVirtualFilePath, const TCHAR* VirtualFilePath, TArray<FString>& IncludeVirtualFilePaths, EShaderPlatform ShaderPlatform, uint32 DepthLimit, const FName* ShaderPlatformName)
 {
-	GetShaderIncludes(EntryPointVirtualFilePath, VirtualFilePath, IncludeVirtualFilePaths, ShaderPlatform, DepthLimit, false);
+	InternalGetShaderIncludes(EntryPointVirtualFilePath, VirtualFilePath, IncludeVirtualFilePaths, ShaderPlatform, DepthLimit, false, ShaderPlatformName);
 }
 
-void GetShaderIncludes(const TCHAR* EntryPointVirtualFilePath, const TCHAR* VirtualFilePath, const FString& FileContents, TArray<FString>& IncludeVirtualFilePaths, EShaderPlatform ShaderPlatform, uint32 DepthLimit)
+void GetShaderIncludes(const TCHAR* EntryPointVirtualFilePath, const TCHAR* VirtualFilePath, const FString& FileContents, TArray<FString>& IncludeVirtualFilePaths, EShaderPlatform ShaderPlatform, uint32 DepthLimit, const FName* ShaderPlatformName)
 {
-	GetShaderIncludes(EntryPointVirtualFilePath, VirtualFilePath, FileContents, IncludeVirtualFilePaths, ShaderPlatform, DepthLimit, false);
+	InternalGetShaderIncludes(EntryPointVirtualFilePath, VirtualFilePath, FileContents, IncludeVirtualFilePaths, ShaderPlatform, DepthLimit, false, ShaderPlatformName);
 }
 
 void HashShaderFileWithIncludes(FArchive& HashingArchive, const TCHAR* VirtualFilePath, const FString& FileContents, EShaderPlatform ShaderPlatform, bool bOnlyHashIncludedFiles)
@@ -1335,12 +1338,12 @@ const FSHAHash& GetShaderFilesHash(const TArray<FString>& VirtualFilePaths, ESha
 	}
 }
 
-void BuildShaderFileToUniformBufferMap(TMap<FString, TArray<const TCHAR*> >& ShaderFileToUniformBufferVariables)
+void BuildShaderFileToUniformBufferMap(TMap<FString, TArray<const TCHAR*> >& ShaderFileToUniformBufferVariables, const FName* ShaderPlatformName)
 {
 	if (!FPlatformProperties::RequiresCookedData())
 	{
 		TArray<FString> ShaderSourceFiles;
-		GetAllVirtualShaderSourcePaths(ShaderSourceFiles, GMaxRHIShaderPlatform);
+		GetAllVirtualShaderSourcePaths(ShaderSourceFiles, GMaxRHIShaderPlatform, ShaderPlatformName);
 
 		FScopedSlowTask SlowTask((float)ShaderSourceFiles.Num());
 
@@ -1441,7 +1444,7 @@ void UninitializeShaderTypes()
  * Flushes the shader file and CRC cache, and regenerates the binary shader files if necessary.
  * Allows shader source files to be re-read properly even if they've been modified since startup.
  */
-void FlushShaderFileCache()
+void FlushShaderFileCache(const FName* ShaderPlatformName)
 {
 	UE_LOG(LogShaders, Log, TEXT("FlushShaderFileCache() begin"));
 
@@ -1459,7 +1462,7 @@ void FlushShaderFileCache()
 		LogShaderSourceDirectoryMappings();
 
 		TMap<FString, TArray<const TCHAR*> > ShaderFileToUniformBufferVariables;
-		BuildShaderFileToUniformBufferMap(ShaderFileToUniformBufferVariables);
+		BuildShaderFileToUniformBufferMap(ShaderFileToUniformBufferVariables, ShaderPlatformName);
 
 		for (TLinkedList<FShaderPipelineType*>::TConstIterator It(FShaderPipelineType::GetTypeList()); It; It.Next())
 		{
@@ -1828,6 +1831,11 @@ FArchive& operator<<(FArchive& Ar, FShaderCompilerInput& Input)
 		Ar << CompressionFormatString;
 		Input.CompressionFormat = FName(*CompressionFormatString);
 	}
+	{
+		FString ShaderPlatformNameString(Input.ShaderPlatformName.ToString());
+		Ar << ShaderPlatformNameString;
+		Input.ShaderPlatformName = FName(*ShaderPlatformNameString);
+	}
 	Ar << Input.SourceFilePrefix;
 	Ar << Input.VirtualSourceFilePath;
 	Ar << Input.EntryPointName;
diff --git a/Engine/Source/Runtime/RenderCore/Public/ShaderCompilerCore.h b/Engine/Source/Runtime/RenderCore/Public/ShaderCompilerCore.h
--- a/Engine/Source/Runtime/RenderCore/Public/ShaderCompilerCore.h
+++ b/Engine/Source/Runtime/RenderCore/Public/ShaderCompilerCore.h
@@ -177,6 +177,7 @@ namespace FOodleDataCompression
 struct FShaderCompilerInput
 {
 	FShaderTarget Target;
+	FName ShaderPlatformName;
 	FName ShaderFormat;
 	FName CompressionFormat;
 	FString SourceFilePrefix;
@@ -527,7 +528,7 @@ extern RENDERCORE_API bool CheckVirtualShaderFilePath(FStringView VirtualPath, T
  * @param OutFileContents - If true is returned, will contain the contents of the shader file. Can be null.
  * @return True if the file was successfully loaded.
  */
-extern RENDERCORE_API bool LoadShaderSourceFile(const TCHAR* VirtualFilePath, EShaderPlatform ShaderPlatform, FString* OutFileContents, TArray<FShaderCompilerError>* OutCompileErrors);
+extern RENDERCORE_API bool LoadShaderSourceFile(const TCHAR* VirtualFilePath, EShaderPlatform ShaderPlatform, FString* OutFileContents, TArray<FShaderCompilerError>* OutCompileErrors, const FName* ShaderPlatformName = nullptr);
 
 enum class EShaderCompilerWorkerType : uint8
 {
diff --git a/Engine/Source/Runtime/RenderCore/Public/ShaderCore.h b/Engine/Source/Runtime/RenderCore/Public/ShaderCore.h
--- a/Engine/Source/Runtime/RenderCore/Public/ShaderCore.h
+++ b/Engine/Source/Runtime/RenderCore/Public/ShaderCore.h
@@ -1010,7 +1010,7 @@ extern RENDERCORE_API FString ParseVirtualShaderFilename(const FString& InFilena
 extern RENDERCORE_API bool ReplaceVirtualFilePathForShaderPlatform(FString& InOutVirtualFilePath, EShaderPlatform ShaderPlatform);
 
 /** Replaces virtual platform path with appropriate autogen path for a given ShaderPlatform. Returns true if path was changed. */
-extern RENDERCORE_API bool ReplaceVirtualFilePathForShaderAutogen(FString& InOutVirtualFilePath, EShaderPlatform ShaderPlatform);
+extern RENDERCORE_API bool ReplaceVirtualFilePathForShaderAutogen(FString& InOutVirtualFilePath, EShaderPlatform ShaderPlatform, const FName* InShaderPlatformName = nullptr);
 
 /** Loads the shader file with the given name.  If the shader file couldn't be loaded, throws a fatal error. */
 extern RENDERCORE_API void LoadShaderSourceFileChecked(const TCHAR* VirtualFilePath, EShaderPlatform ShaderPlatform, FString& OutFileContents);
@@ -1018,8 +1018,8 @@ extern RENDERCORE_API void LoadShaderSourceFileChecked(const TCHAR* VirtualFileP
 /**
  * Recursively populates IncludeFilenames with the include filenames from Filename
  */
-extern RENDERCORE_API void GetShaderIncludes(const TCHAR* EntryPointVirtualFilePath, const TCHAR* VirtualFilePath, TArray<FString>& IncludeVirtualFilePaths, EShaderPlatform ShaderPlatform, uint32 DepthLimit=100);
-extern RENDERCORE_API void GetShaderIncludes(const TCHAR* EntryPointVirtualFilePath, const TCHAR* VirtualFilePath, const FString& FileContents, TArray<FString>& IncludeVirtualFilePaths, EShaderPlatform ShaderPlatform, uint32 DepthLimit = 100);
+extern RENDERCORE_API void GetShaderIncludes(const TCHAR* EntryPointVirtualFilePath, const TCHAR* VirtualFilePath, TArray<FString>& IncludeVirtualFilePaths, EShaderPlatform ShaderPlatform, uint32 DepthLimit=100, const FName* ShaderPlatformName = nullptr);
+extern RENDERCORE_API void GetShaderIncludes(const TCHAR* EntryPointVirtualFilePath, const TCHAR* VirtualFilePath, const FString& FileContents, TArray<FString>& IncludeVirtualFilePaths, EShaderPlatform ShaderPlatform, uint32 DepthLimit = 100, const FName* ShaderPlatformName = nullptr);
 
 /**
  * Calculates a Hash for the given filename if it does not already exist in the Hash cache.
@@ -1043,13 +1043,13 @@ extern RENDERCORE_API void HashShaderFileWithIncludes(FArchive& HashingArchive,
  */
 extern RENDERCORE_API const class FSHAHash& GetShaderFilesHash(const TArray<FString>& VirtualFilePaths, EShaderPlatform ShaderPlatform);
 
-extern void BuildShaderFileToUniformBufferMap(TMap<FString, TArray<const TCHAR*> >& ShaderFileToUniformBufferVariables);
+extern void BuildShaderFileToUniformBufferMap(TMap<FString, TArray<const TCHAR*> >& ShaderFileToUniformBufferVariables, const FName* ShaderPlatformName = nullptr);
 
 /**
  * Flushes the shader file and CRC cache, and regenerates the binary shader files if necessary.
  * Allows shader source files to be re-read properly even if they've been modified since startup.
  */
-extern RENDERCORE_API void FlushShaderFileCache();
+extern RENDERCORE_API void FlushShaderFileCache(const FName* ShaderPlatformName = nullptr);
 
 extern RENDERCORE_API void VerifyShaderSourceFiles(EShaderPlatform ShaderPlatform);
 
@@ -1098,5 +1098,5 @@ extern RENDERCORE_API void ResetAllShaderSourceDirectoryMappings();
  */
 extern RENDERCORE_API void AddShaderSourceDirectoryMapping(const FString& VirtualShaderDirectory, const FString& RealShaderDirectory);
 
-extern RENDERCORE_API void AddShaderSourceFileEntry(TArray<FString>& OutVirtualFilePaths, FString VirtualFilePath, EShaderPlatform ShaderPlatform);
-extern RENDERCORE_API void GetAllVirtualShaderSourcePaths(TArray<FString>& OutVirtualFilePaths, EShaderPlatform ShaderPlatform);
+extern RENDERCORE_API void AddShaderSourceFileEntry(TArray<FString>& OutVirtualFilePaths, FString VirtualFilePath, EShaderPlatform ShaderPlatform, const FName* ShaderPlatformName = nullptr);
+extern RENDERCORE_API void GetAllVirtualShaderSourcePaths(TArray<FString>& OutVirtualFilePaths, EShaderPlatform ShaderPlatform, const FName* ShaderPlatformName = nullptr);
