Exclude ray tracing shaders from SM5 if it's disabled for D3D12

* Add ITargetPlatform::GetRayTracingShaderFormats(), with default implementation that falls back to GetAllTargetedShaderFormats() if HWRT is enabled for a platform
* Override GetRayTracingShaderFormats() for TGenericWindowsTargetPlatform to return only enabled D3D12 and VK shader formats (excluding D3D11)
* Use GetRayTracingShaderFormat() and UsesRayTracing() to configure GRayTracingPlaformMask (the mask contains one bit per shader platform)

Previously all bits of GRayTracingPlaformMask were set if project uses ray tracing, then shader platform bits were disabled for all possible shader formats for platforms that don't support ray tracing.
If SM5 was disabled for D3D12, ray tracing shaders would still get compiled for it because the corresponding bit was still set in GRayTracingPlaformMask.

With this change, SM5 bit will no longer be set in GRayTracingPlaformMask in the editor if it is disabled for D3D12 RHI for the project, so ray tracing shaders won't be compiled.
When running a packaged build, GRayTracingPlaformMask is now set to 0 or 1 based on global project setting (r.RayTracing) and whether the current machine+RHI supports ray tracing.

diff --git a/Engine/Source/Developer/TargetPlatform/Private/TargetPlatformBase.cpp b/Engine/Source/Developer/TargetPlatform/Private/TargetPlatformBase.cpp
--- a/Engine/Source/Developer/TargetPlatform/Private/TargetPlatformBase.cpp
+++ b/Engine/Source/Developer/TargetPlatform/Private/TargetPlatformBase.cpp
@@ -95,6 +95,14 @@ bool FTargetPlatformBase::UsesASTCHDR() const
 	return GASTCHDRProfile != 0;
 }
 
+void FTargetPlatformBase::GetRayTracingShaderFormats(TArray<FName>& OutFormats) const
+{
+	if (UsesRayTracing())
+	{
+		GetAllTargetedShaderFormats(OutFormats);
+	}
+}
+
 static bool IsPluginEnabledForTarget(const IPlugin& Plugin, const FProjectDescriptor* Project, const FString& Platform, EBuildConfiguration Configuration, EBuildTargetType TargetType)
 {
 	if (!Plugin.GetDescriptor().SupportsTargetPlatform(Platform))
diff --git a/Engine/Source/Developer/TargetPlatform/Public/Common/TargetPlatformBase.h b/Engine/Source/Developer/TargetPlatform/Public/Common/TargetPlatformBase.h
--- a/Engine/Source/Developer/TargetPlatform/Public/Common/TargetPlatformBase.h
+++ b/Engine/Source/Developer/TargetPlatform/Public/Common/TargetPlatformBase.h
@@ -79,6 +79,8 @@ public:
 
 	TARGETPLATFORM_API virtual bool UsesASTCHDR() const override;
 
+	TARGETPLATFORM_API virtual void GetRayTracingShaderFormats(TArray<FName>& OutFormats) const override;
+
 #if WITH_ENGINE
 	virtual void GetReflectionCaptureFormats( TArray<FName>& OutFormats ) const override
 	{
diff --git a/Engine/Source/Developer/TargetPlatform/Public/Interfaces/ITargetPlatform.h b/Engine/Source/Developer/TargetPlatform/Public/Interfaces/ITargetPlatform.h
--- a/Engine/Source/Developer/TargetPlatform/Public/Interfaces/ITargetPlatform.h
+++ b/Engine/Source/Developer/TargetPlatform/Public/Interfaces/ITargetPlatform.h
@@ -531,6 +531,13 @@ public:
 	*/
 	virtual void GetAllTargetedShaderFormats(TArray<FName>& OutFormats) const = 0;
 
+	/**
+	* Gets the shader formats that support ray tracing for this target platform.
+	* 
+	* @param OutFormats Will contain the shader formats.
+	*/
+	virtual void GetRayTracingShaderFormats(TArray<FName>& OutFormats) const = 0;
+
 #if WITH_ENGINE
 	/**
 	 * Gets the format to use for a particular body setup.
diff --git a/Engine/Source/Developer/Windows/WindowsTargetPlatform/Private/GenericWindowsTargetPlatform.h b/Engine/Source/Developer/Windows/WindowsTargetPlatform/Private/GenericWindowsTargetPlatform.h
--- a/Engine/Source/Developer/Windows/WindowsTargetPlatform/Private/GenericWindowsTargetPlatform.h
+++ b/Engine/Source/Developer/Windows/WindowsTargetPlatform/Private/GenericWindowsTargetPlatform.h
@@ -230,27 +230,21 @@ public:
 		}
 	}
 
-	virtual void GetAllTargetedShaderFormats( TArray<FName>& OutFormats ) const override 
+	void GetAllTargetedShaderFormatsInternal(TArrayView<TCHAR const*> RelevantSettings, TArray<FName>& OutFormats) const
 	{
-		// Get the Target RHIs for this platform, we do not always want all those that are supported. (reload in case user changed in the editor)
 		TArray<FString> TargetedShaderFormats;
 
-		auto AddWindowsShaderFormats = [&TargetedShaderFormats](const TCHAR* InName)
+		for (const TCHAR* Name : RelevantSettings)
 		{
 			TArray<FString> NewTargetedShaderFormats;
-			GConfig->GetArray(TEXT("/Script/WindowsTargetPlatform.WindowsTargetSettings"), InName, NewTargetedShaderFormats, GEngineIni);
+			GConfig->GetArray(TEXT("/Script/WindowsTargetPlatform.WindowsTargetSettings"), Name, NewTargetedShaderFormats, GEngineIni);
 
 			for (const FString& NewShaderFormat : NewTargetedShaderFormats)
 			{
 				TargetedShaderFormats.AddUnique(NewShaderFormat);
 			}
-		};
-
-		AddWindowsShaderFormats(TEXT("TargetedRHIs"));
-		AddWindowsShaderFormats(TEXT("D3D12TargetedShaderFormats"));
-		AddWindowsShaderFormats(TEXT("D3D11TargetedShaderFormats"));
-		AddWindowsShaderFormats(TEXT("VulkanTargetedShaderFormats"));
-
+		}
+		
 		// Gather the list of Target RHIs and filter out any that may be invalid.
 		TArray<FName> PossibleShaderFormats;
 		GetAllPossibleShaderFormats(PossibleShaderFormats);
@@ -270,6 +264,36 @@ public:
 		}
 	}
 
+	virtual void GetAllTargetedShaderFormats( TArray<FName>& OutFormats ) const override 
+	{
+		// Get the Target RHIs for this platform, we do not always want all those that are supported. (reload in case user changed in the editor)
+
+		TCHAR const* RelevantSettings[] = 
+		{
+			TEXT("TargetedRHIs"),
+			TEXT("D3D12TargetedShaderFormats"),
+			TEXT("D3D11TargetedShaderFormats"),
+			TEXT("VulkanTargetedShaderFormats")
+		};
+
+		GetAllTargetedShaderFormatsInternal(RelevantSettings, OutFormats);
+	}
+
+	virtual void GetRayTracingShaderFormats( TArray<FName>& OutFormats ) const override 
+	{
+		if (UsesRayTracing())
+		{
+			// Only D3D12 and Vulkan can possibly support ray tracing
+			TCHAR const* RelevantSettings[] = 
+			{
+				TEXT("D3D12TargetedShaderFormats"),
+				TEXT("VulkanTargetedShaderFormats")
+			};
+
+			GetAllTargetedShaderFormatsInternal(RelevantSettings, OutFormats);
+		}
+	}
+
 #if WITH_ENGINE
 	virtual void GetReflectionCaptureFormats(TArray<FName>& OutFormats) const override
 	{
diff --git a/Engine/Source/Runtime/RenderCore/Private/RenderUtils.cpp b/Engine/Source/Runtime/RenderCore/Private/RenderUtils.cpp
--- a/Engine/Source/Runtime/RenderCore/Private/RenderUtils.cpp
+++ b/Engine/Source/Runtime/RenderCore/Private/RenderUtils.cpp
@@ -1379,10 +1379,6 @@ RENDERCORE_API void RenderUtilsInit()
 	}
 
 	static IConsoleVariable* RayTracingCVar = IConsoleManager::Get().FindConsoleVariable(TEXT("r.RayTracing"));
-	if (RayTracingCVar && RayTracingCVar->GetInt())
-	{
-		GRayTracingPlaformMask = ~0ull;
-	}
 
 	static IConsoleVariable* MobileAmbientOcclusionCVar = IConsoleManager::Get().FindConsoleVariable(TEXT("r.Mobile.AmbientOcclusion"));
 	if (MobileAmbientOcclusionCVar && MobileAmbientOcclusionCVar->GetInt())
@@ -1396,10 +1392,10 @@ RENDERCORE_API void RenderUtilsInit()
 	{
 		for (ITargetPlatform* TargetPlatform : TargetPlatformManager->GetTargetPlatforms())
 		{
-			TArray<FName> PlaformShaderFormats;
-			TargetPlatform->GetAllPossibleShaderFormats(PlaformShaderFormats);
+			TArray<FName> PlaformPossibleShaderFormats;
+			TargetPlatform->GetAllPossibleShaderFormats(PlaformPossibleShaderFormats);
 
-			for (FName Format : PlaformShaderFormats)
+			for (FName Format : PlaformPossibleShaderFormats)
 			{
 				EShaderPlatform ShaderPlatform = ShaderFormatNameToShaderPlatform(Format);
 				uint32 ShaderPlatformIndex = static_cast<uint32>(ShaderPlatform);
@@ -1451,15 +1447,6 @@ RENDERCORE_API void RenderUtilsInit()
 					GDistanceFieldsPlatformMask &= ~Mask;
 				}
 
-				if (TargetPlatform->UsesRayTracing())
-				{
-					GRayTracingPlaformMask |= Mask;
-				}
-				else
-				{
-					GRayTracingPlaformMask &= ~Mask;
-				}
-
 				if (TargetPlatform->ForcesSimpleSkyDiffuse())
 				{
 					GSimpleSkyDiffusePlatformMask |= Mask;
@@ -1487,9 +1474,27 @@ RENDERCORE_API void RenderUtilsInit()
 					GMobileAmbientOcclusionPlatformMask &= ~Mask;
 				}
 			}
+
+
+			if (TargetPlatform->UsesRayTracing())
+			{
+				TArray<FName> PlaformRayTracingShaderFormats;
+				TargetPlatform->GetRayTracingShaderFormats(PlaformRayTracingShaderFormats);
+
+				for (FName Format : PlaformRayTracingShaderFormats)
+				{
+					EShaderPlatform ShaderPlatform = ShaderFormatNameToShaderPlatform(Format);
+					uint32 ShaderPlatformIndex = static_cast<uint32>(ShaderPlatform);
+
+					uint64 Mask = 1ull << ShaderPlatformIndex;
+
+					GRayTracingPlaformMask |= Mask;
+				}
+			}
 		}
 	}
-#else
+
+#else // WITH_EDITOR
 
 	if (IsMobilePlatform(GMaxRHIShaderPlatform))
 	{
@@ -1497,6 +1502,11 @@ RENDERCORE_API void RenderUtilsInit()
 		GBasePassVelocityPlatformMask = 0;
 	}
 
+	if (RayTracingCVar && RayTracingCVar->GetInt() && GRHISupportsRayTracing)
+	{
+		GRayTracingPlaformMask = ~0ull;
+	}
+
 	// Load runtime values from and *.ini file used by a current platform
 	// Should be code shared between cook and game, but unfortunately can't be done before we untangle non data driven platforms
 	const FString PlatformName(FPlatformProperties::IniPlatformName());
