Backout 20871665 to fix Mac cook

diff --git a/Engine/Config/Android/DataDrivenPlatformInfo.ini b/Engine/Config/Android/DataDrivenPlatformInfo.ini
--- a/Engine/Config/Android/DataDrivenPlatformInfo.ini
+++ b/Engine/Config/Android/DataDrivenPlatformInfo.ini
@@ -56,7 +56,6 @@ DeviceProfileName=Android
 [ShaderPlatform OPENGL_ES3_1_ANDROID]
 Language=OpenGL
 MaxFeatureLevel=ES3_1
-ShaderFormat=GLSL_ES3_1_ANDROID
 bIsMobile=true
 bIsAndroidOpenGLES=true
 bTargetsTiledGPU=true
@@ -67,7 +66,6 @@ bSupportsManualVertexFetch=false
 [ShaderPlatform VULKAN_ES3_1_ANDROID]
 Language=Vulkan
 MaxFeatureLevel=ES3_1
-ShaderFormat=SF_VULKAN_ES31_ANDROID
 bIsMobile=true
 bTargetsTiledGPU=true
 bSupportsMobileMultiView=true
@@ -78,7 +76,6 @@ bSupportsManualVertexFetch=false
 [ShaderPlatform VULKAN_SM5_ANDROID]
 Language=Vulkan
 MaxFeatureLevel=SM5
-ShaderFormat=SF_VULKAN_SM5_ANDROID
 bTargetsTiledGPU=true
 bNeedsOfflineCompiler=true
 bSupportsIndexBufferUAVs=true
diff --git a/Engine/Config/HoloLens/DataDrivenPlatformInfo.ini b/Engine/Config/HoloLens/DataDrivenPlatformInfo.ini
--- a/Engine/Config/HoloLens/DataDrivenPlatformInfo.ini
+++ b/Engine/Config/HoloLens/DataDrivenPlatformInfo.ini
@@ -28,7 +28,6 @@ PlatformGroupName=Mobile
 [ShaderPlatform D3D_ES3_1_HOLOLENS]
 Language=D3D
 MaxFeatureLevel=ES3_1
-ShaderFormat=D3D_ES3_1_HOLOLENS
 bIsMobile=true
 bSupportsInstancedStereo=true
 bSupportsManualVertexFetch=false
diff --git a/Engine/Config/IOS/DataDrivenPlatformInfo.ini b/Engine/Config/IOS/DataDrivenPlatformInfo.ini
--- a/Engine/Config/IOS/DataDrivenPlatformInfo.ini
+++ b/Engine/Config/IOS/DataDrivenPlatformInfo.ini
@@ -48,7 +48,6 @@ DeviceProfileName=IOS
 [ShaderPlatform METAL]
 Language=Metal
 MaxFeatureLevel=ES3_1
-ShaderFormat=SF_METAL_MACES3_1
 bIsMobile=true
 bTargetsTiledGPU=true
 bSupports4ComponentUAVReadWrite=true
@@ -58,7 +57,6 @@ bSupportsVolumeTextureAtomics=false
 [ShaderPlatform METAL_MRT]
 Language=Metal
 MaxFeatureLevel=SM5
-ShaderFormat=SF_METAL_MRT
 bIsMetalMRT=true
 bTargetsTiledGPU=true
 bSupportsIndexBufferUAVs=true
diff --git a/Engine/Config/Mac/DataDrivenPlatformInfo.ini b/Engine/Config/Mac/DataDrivenPlatformInfo.ini
--- a/Engine/Config/Mac/DataDrivenPlatformInfo.ini
+++ b/Engine/Config/Mac/DataDrivenPlatformInfo.ini
@@ -26,7 +26,6 @@ PlatformGroupName=Desktop
 [ShaderPlatform METAL_SM5]
 Language=Metal
 MaxFeatureLevel=SM5
-ShaderFormat=SF_METAL_SM5
 bIsPC=true
 bSupportsDebugViewShaders=true
 bSupportsAnisotropicMaterials=true
@@ -54,7 +53,6 @@ EnablesHLSL2021ByDefault=1
 [ShaderPlatform METAL_MRT_MAC]
 Language=Metal
 MaxFeatureLevel=SM5
-ShaderFormat=SF_METAL_MRT_MAC
 bIsPC=true
 bIsMetalMRT=true
 bSupportsIndexBufferUAVs=true
@@ -65,7 +63,6 @@ bSupportsVolumeTextureAtomics=false
 [ShaderPlatform METAL_MACES3_1]
 Language=Metal
 MaxFeatureLevel=ES3_1
-ShaderFormat=SF_METAL_MACES3_1
 bIsPC=true
 bIsMobile=true
 bSupports4ComponentUAVReadWrite=true
diff --git a/Engine/Config/TVOS/DataDrivenPlatformInfo.ini b/Engine/Config/TVOS/DataDrivenPlatformInfo.ini
--- a/Engine/Config/TVOS/DataDrivenPlatformInfo.ini
+++ b/Engine/Config/TVOS/DataDrivenPlatformInfo.ini
@@ -19,7 +19,6 @@ PlatformGroupName=Mobile
 [ShaderPlatform METAL_MRT_TVOS]
 Language=Metal
 MaxFeatureLevel=SM5
-ShaderFormat=SF_METAL_MRT_TVOS
 bIsMetalMRT=true
 bTargetsTiledGPU=true
 bSupportsIndexBufferUAVs=true
@@ -29,7 +28,6 @@ bSupportsManualVertexFetch=false
 [ShaderPlatform METAL_TVOS]
 Language=Metal
 MaxFeatureLevel=ES3_1
-ShaderFormat=SF_METAL_TVOS
 bIsMobile=true
 bTargetsTiledGPU=true
 bSupports4ComponentUAVReadWrite=true
diff --git a/Engine/Config/VulkanPC/DataDrivenPlatformInfo.ini b/Engine/Config/VulkanPC/DataDrivenPlatformInfo.ini
--- a/Engine/Config/VulkanPC/DataDrivenPlatformInfo.ini
+++ b/Engine/Config/VulkanPC/DataDrivenPlatformInfo.ini
@@ -11,7 +11,6 @@ GlobalIdentifier=97F8C85EF1B0424792D201C893358C8F
 [ShaderPlatform VULKAN_SM5]
 Language=Vulkan
 MaxFeatureLevel=SM5
-ShaderFormat=SF_VULKAN_SM5
 
 bIsMobile = false
 bIsMetalMRT = false
@@ -136,7 +135,6 @@ EnablesHLSL2021ByDefault=1
 [ShaderPlatform VULKAN_PCES3_1]
 Language=Vulkan
 MaxFeatureLevel=ES3_1
-ShaderFormat=SF_VULKAN_ES31
 bIsPC=true
 bIsMobile=true
 bSupportsMobileMultiView=true
diff --git a/Engine/Config/Windows/DataDrivenPlatformInfo.ini b/Engine/Config/Windows/DataDrivenPlatformInfo.ini
--- a/Engine/Config/Windows/DataDrivenPlatformInfo.ini
+++ b/Engine/Config/Windows/DataDrivenPlatformInfo.ini
@@ -57,7 +57,6 @@ PlatformGroupName=Desktop
 [ShaderPlatform PCD3D_SM5]
 Language=D3D
 MaxFeatureLevel=SM5
-ShaderFormat=PCD3D_SM5
 bIsPC=true
 bSupportsDebugViewShaders=true
 bSupportsRayTracing=true
@@ -103,7 +102,6 @@ MenuTooltip=LOCTEXT("PreviewMenuTooltip_PCD3D_SM5", "DirectX 11, OpenGL 4.3+, PS
 FriendlyName=LOCTEXT("FriendlyShaderPlatformName_PCD3D_SM6", "Shader Model 6")
 Language=D3D
 MaxFeatureLevel=SM6
-ShaderFormat=PCD3D_SM6
 bIsPC=true
 MaxMeshShaderThreadGroupSize=128
 bSupportsDebugViewShaders=true
@@ -155,7 +153,6 @@ MenuTooltip=LOCTEXT("PreviewMenuTooltip_PCD3D_SM6", "Preview using SM6 device pr
 [ShaderPlatform PCD3D_ES3_1]
 Language=D3D
 MaxFeatureLevel=ES3_1
-ShaderFormat=PCD3D_ES31
 bIsPC=true
 bIsMobile=true
 bSupportsDistanceFields = true
@@ -167,6 +164,5 @@ bSupportsManualVertexFetch=false
 Language=OpenGL
 bIsPC=true
 MaxFeatureLevel=ES3_1
-ShaderFormat=GLSL_150_ES31
 bIsMobile=true
 bSupportsManualVertexFetch=false
diff --git a/Engine/Plugins/Compositing/OpenColorIO/Source/OpenColorIO/Private/OpenColorIOShaderCompilationManager.cpp b/Engine/Plugins/Compositing/OpenColorIO/Source/OpenColorIO/Private/OpenColorIOShaderCompilationManager.cpp
--- a/Engine/Plugins/Compositing/OpenColorIO/Source/OpenColorIO/Private/OpenColorIOShaderCompilationManager.cpp
+++ b/Engine/Plugins/Compositing/OpenColorIO/Source/OpenColorIO/Private/OpenColorIOShaderCompilationManager.cpp
@@ -82,7 +82,7 @@ OPENCOLORIO_API void FOpenColorIOShaderCompilationManager::AddJobs(TArray<FShade
 			CurrentJob->Input.Environment.CompilerFlags.Add(CFLAG_NoFastMath);
 		}
 
-		CurrentJob->Input.DumpDebugInfoRootPath = GShaderCompilingManager->GetAbsoluteShaderDebugInfoDirectory() / CurrentJob->Input.ShaderPlatformName.ToString();
+		CurrentJob->Input.DumpDebugInfoRootPath = GShaderCompilingManager->GetAbsoluteShaderDebugInfoDirectory() / LegacyShaderPlatformToShaderFormat(EShaderPlatform(CurrentJob->Input.Target.Platform)).ToString();
 		FPaths::NormalizeDirectoryName(CurrentJob->Input.DumpDebugInfoRootPath);
 		CurrentJob->Input.DebugExtension.Empty();
 		CurrentJob->Input.DumpDebugInfoPath.Empty();
diff --git a/Engine/Plugins/Experimental/ComputeFramework/Source/ComputeFramework/Private/ComputeKernelShaderCompilationManager.cpp b/Engine/Plugins/Experimental/ComputeFramework/Source/ComputeFramework/Private/ComputeKernelShaderCompilationManager.cpp
--- a/Engine/Plugins/Experimental/ComputeFramework/Source/ComputeFramework/Private/ComputeKernelShaderCompilationManager.cpp
+++ b/Engine/Plugins/Experimental/ComputeFramework/Source/ComputeFramework/Private/ComputeKernelShaderCompilationManager.cpp
@@ -115,7 +115,7 @@ void FComputeKernelShaderCompilationManager::RunCompileJobs()
 
 				UE_LOG(LogComputeKernelShaderCompiler, Log, TEXT("Compile Job processing... %s"), *CurrentJob.Input.DebugGroupName);
 
-				CurrentJob.Input.DumpDebugInfoRootPath = GShaderCompilingManager->GetAbsoluteShaderDebugInfoDirectory() / CurrentJob.Input.ShaderPlatformName.ToString();
+				CurrentJob.Input.DumpDebugInfoRootPath = GShaderCompilingManager->GetAbsoluteShaderDebugInfoDirectory() / Format.ToString();
 				FPaths::NormalizeDirectoryName(CurrentJob.Input.DumpDebugInfoRootPath);
 				const FShaderCompilingManager::EDumpShaderDebugInfo DumpShaderDebugInfo = GShaderCompilingManager->GetDumpShaderDebugInfo();
 				CurrentJob.Input.DebugExtension.Empty();
diff --git a/Engine/Plugins/FX/Niagara/Shaders/Private/NiagaraShaderVersion.ush b/Engine/Plugins/FX/Niagara/Shaders/Private/NiagaraShaderVersion.ush
--- a/Engine/Plugins/FX/Niagara/Shaders/Private/NiagaraShaderVersion.ush
+++ b/Engine/Plugins/FX/Niagara/Shaders/Private/NiagaraShaderVersion.ush
@@ -3,4 +3,4 @@
 // in the DDC key, it will automatically invalidate.
 // 
 // If you are merging streams and there is a conflict with this GUID you should make a new GUID rather than taking one or the other.
-// GUID = 1100075F7722451B96F0AB70CCFDD49C
+// GUID = 97DA38DB603E47EA9743AFC9AFCA732D
\ No newline at end of file
diff --git a/Engine/Plugins/FX/Niagara/Source/NiagaraShader/Private/NiagaraShaderCompilationManager.cpp b/Engine/Plugins/FX/Niagara/Source/NiagaraShader/Private/NiagaraShaderCompilationManager.cpp
--- a/Engine/Plugins/FX/Niagara/Source/NiagaraShader/Private/NiagaraShaderCompilationManager.cpp
+++ b/Engine/Plugins/FX/Niagara/Source/NiagaraShader/Private/NiagaraShaderCompilationManager.cpp
@@ -49,7 +49,7 @@ NIAGARASHADER_API void FNiagaraShaderCompilationManager::AddJobs(TArray<FShaderC
 
 		UE_LOG(LogNiagaraShaderCompiler, Verbose, TEXT("Adding niagara gpu shader compile job... %s"), *CurrentJob->Input.DebugGroupName);
 
-		CurrentJob->Input.DumpDebugInfoRootPath = GShaderCompilingManager->GetAbsoluteShaderDebugInfoDirectory() / CurrentJob->Input.ShaderPlatformName.ToString();
+		CurrentJob->Input.DumpDebugInfoRootPath = GShaderCompilingManager->GetAbsoluteShaderDebugInfoDirectory() / LegacyShaderPlatformToShaderFormat(EShaderPlatform(CurrentJob->Input.Target.Platform)).ToString();
 		FPaths::NormalizeDirectoryName(CurrentJob->Input.DumpDebugInfoRootPath);
 		CurrentJob->Input.DebugExtension.Empty();
 		CurrentJob->Input.DumpDebugInfoPath.Empty();
diff --git a/Engine/Shaders/Public/ShaderVersion.ush b/Engine/Shaders/Public/ShaderVersion.ush
--- a/Engine/Shaders/Public/ShaderVersion.ush
+++ b/Engine/Shaders/Public/ShaderVersion.ush
@@ -2,4 +2,4 @@
 // Each time the console command is executed it generates a new GUID. As this file is included
 // in Platform.ush (which should be included in any shader) it allows to invalidate the shader DDC.
 // 
-// GUID = 84806049-D47D-41D3-A61D-0D2CC08B1073
+// GUID = 6E2FADCF-B7B1-4C25-8103-433C76BC8AE8
diff --git a/Engine/Source/Developer/ShaderCompilerCommon/Private/ShaderCompilerCommon.cpp b/Engine/Source/Developer/ShaderCompilerCommon/Private/ShaderCompilerCommon.cpp
--- a/Engine/Source/Developer/ShaderCompilerCommon/Private/ShaderCompilerCommon.cpp
+++ b/Engine/Source/Developer/ShaderCompilerCommon/Private/ShaderCompilerCommon.cpp
@@ -923,10 +923,6 @@ FString CreateShaderCompilerWorkerDirectCommandLine(const FShaderCompilerInput&
 	Text += Input.ShaderFormat.GetPlainNameString();
 	Text += TEXT(" -entry=");
 	Text += Input.EntryPointName;
-
-	Text += TEXT(" -shaderPlatformName=");
-	Text += Input.ShaderPlatformName.GetPlainNameString();
-
 	switch (Input.Target.Frequency)
 	{
 	case SF_Vertex:			Text += TEXT(" -vs"); break;
diff --git a/Engine/Source/Developer/ShaderPreprocessor/Private/ShaderPreprocessor.cpp b/Engine/Source/Developer/ShaderPreprocessor/Private/ShaderPreprocessor.cpp
--- a/Engine/Source/Developer/ShaderPreprocessor/Private/ShaderPreprocessor.cpp
+++ b/Engine/Source/Developer/ShaderPreprocessor/Private/ShaderPreprocessor.cpp
@@ -36,7 +36,7 @@ public:
 		, ShaderOutput(InShaderOutput)
 	{
 		FString InputShaderSource;
-		if (LoadShaderSourceFile(*InShaderInput.VirtualSourceFilePath, InShaderInput.Target.GetPlatform(),  &InputShaderSource, nullptr, &InShaderInput.ShaderPlatformName))
+		if (LoadShaderSourceFile(*InShaderInput.VirtualSourceFilePath, InShaderInput.Target.GetPlatform(),  &InputShaderSource, nullptr))
 		{
 			InputShaderSource = FString::Printf(TEXT("%s\n#line 1\n%s"), *ShaderInput.SourceFilePrefix, *InputShaderSource);
 			CachedFileContents.Add(InShaderInput.VirtualSourceFilePath, StringToArray<ANSICHAR>(*InputShaderSource, InputShaderSource.Len() + 1));
@@ -73,7 +73,7 @@ private:
 		ReplaceVirtualFilePathForShaderPlatform(VirtualFilePath, This->ShaderInput.Target.GetPlatform());
 
 		// Fixup autogen file
-		ReplaceVirtualFilePathForShaderAutogen(VirtualFilePath, This->ShaderInput.Target.GetPlatform(), &This->ShaderInput.ShaderPlatformName);
+		ReplaceVirtualFilePathForShaderAutogen(VirtualFilePath, This->ShaderInput.Target.GetPlatform());
 
 		// Collapse any relative directories to allow #include "../MyFile.ush"
 		FPaths::CollapseRelativeDirectories(VirtualFilePath);
@@ -95,7 +95,7 @@ private:
 			{
 				CheckShaderHashCacheInclude(VirtualFilePath, This->ShaderInput.Target.GetPlatform());
 
-				LoadShaderSourceFile(*VirtualFilePath, This->ShaderInput.Target.GetPlatform(), &FileContents, &This->ShaderOutput.Errors, &This->ShaderInput.ShaderPlatformName);
+				LoadShaderSourceFile(*VirtualFilePath, This->ShaderInput.Target.GetPlatform(), &FileContents, &This->ShaderOutput.Errors);
 			}
 
 			if (FileContents.Len() > 0)
diff --git a/Engine/Source/Editor/UnrealEd/Private/Commandlets/CookGlobalShaderCommandlet.cpp b/Engine/Source/Editor/UnrealEd/Private/Commandlets/CookGlobalShaderCommandlet.cpp
--- a/Engine/Source/Editor/UnrealEd/Private/Commandlets/CookGlobalShaderCommandlet.cpp
+++ b/Engine/Source/Editor/UnrealEd/Private/Commandlets/CookGlobalShaderCommandlet.cpp
@@ -49,8 +49,8 @@ void UCookGlobalShadersCommandlet::CookGlobalShaders() const
 	TArray<TPair<FString, FString>> FilesToCopy;
 	for (FName ShaderFormat : ShaderFormats)
 	{
-		const FString GlobalShaderCacheName = FPaths::Combine(OutputDir, TEXT("Engine"), TEXT("GlobalShaderCache-") + FDataDrivenShaderPlatformInfo::GetName(ShaderFormatToLegacyShaderPlatform(ShaderFormat)).ToString() + TEXT(".bin"));
-		const FString OverrideGlobalShaderCacheName = FPaths::Combine(TEXT("Engine"), TEXT("OverrideGlobalShaderCache-") + FDataDrivenShaderPlatformInfo::GetName(ShaderFormatToLegacyShaderPlatform(ShaderFormat)).ToString() + TEXT(".bin"));
+		const FString GlobalShaderCacheName = FPaths::Combine(OutputDir, TEXT("Engine"), TEXT("GlobalShaderCache-") + ShaderFormat.ToString() + TEXT(".bin"));
+		const FString OverrideGlobalShaderCacheName = FPaths::Combine(TEXT("Engine"), TEXT("OverrideGlobalShaderCache-") + ShaderFormat.ToString() + TEXT(".bin"));
 		FilesToCopy.Emplace(GlobalShaderCacheName, OverrideGlobalShaderCacheName);
 	}
 
diff --git a/Engine/Source/Programs/ShaderCompileWorker/Private/ShaderCompileWorker.cpp b/Engine/Source/Programs/ShaderCompileWorker/Private/ShaderCompileWorker.cpp
--- a/Engine/Source/Programs/ShaderCompileWorker/Private/ShaderCompileWorker.cpp
+++ b/Engine/Source/Programs/ShaderCompileWorker/Private/ShaderCompileWorker.cpp
@@ -599,12 +599,8 @@ private:
 
 			// Flush cache, to make sure we load the latest version of the input file.
 			// (Otherwise quick changes to a shader file can result in the wrong output.)
-			FName ShaderPlatformName;
-			InputFile << ShaderPlatformName;
+			FlushShaderFileCache();
 
-			// temp-fix for non-Windows platforms
-			//FlushShaderFileCache(&ShaderPlatformName);
-			
 			for (int32 BatchIndex = 0; BatchIndex < NumBatches; BatchIndex++)
 			{
 				// Deserialize the job's inputs.
@@ -919,7 +915,6 @@ static void DirectCompile(const TArray<const class IShaderFormat*>& ShaderFormat
 	FString InputFile;
 
 	FName FormatName;
-	FName ShaderPlatformName;
 	FString Entry = TEXT("Main");
 	bool bPipeline = false;
 	bool bUseMCPP = false;
@@ -945,10 +940,6 @@ static void DirectCompile(const TArray<const class IShaderFormat*>& ShaderFormat
 					Entry = Entry.Mid(1, Entry.Len() - 2);
 				}
 			}
-			else if (Token.StartsWith(TEXT("shaderPlatformName=")))
-			{
-				ShaderPlatformName = FName(*Token.RightChop(19));
-			}
 			else if (Token.StartsWith(TEXT("cflags=")))
 			{
 				CFlags = FCString::Atoi64(*Token.RightChop(7));
@@ -1030,7 +1021,6 @@ static void DirectCompile(const TArray<const class IShaderFormat*>& ShaderFormat
 	FShaderCompilerInput Input;
 	Input.EntryPointName = Entry;
 	Input.ShaderFormat = FormatName;
-	Input.ShaderPlatformName = ShaderPlatformName;
 	Input.VirtualSourceFilePath = InputFile;
 	Input.Target.Platform =  ShaderFormatNameToShaderPlatform(FormatName);
 	Input.Target.Frequency = Frequency;
diff --git a/Engine/Source/Runtime/Engine/Private/ShaderCompiler/ShaderCompiler.cpp b/Engine/Source/Runtime/Engine/Private/ShaderCompiler/ShaderCompiler.cpp
--- a/Engine/Source/Runtime/Engine/Private/ShaderCompiler/ShaderCompiler.cpp
+++ b/Engine/Source/Runtime/Engine/Private/ShaderCompiler/ShaderCompiler.cpp
@@ -1445,9 +1445,6 @@ bool FShaderCompileUtilities::DoWriteTasks(const TArray<FShaderCommonCompileJobP
 		int32 NumBatches = QueuedSingleJobs.Num();
 		TransferFile << NumBatches;
 
-		FName ShaderPlatformName = FDataDrivenShaderPlatformInfo::GetName(GMaxRHIShaderPlatform);
-		TransferFile << ShaderPlatformName;
-
 		// Serialize all the batched jobs
 		for (int32 JobIndex = 0; JobIndex < QueuedSingleJobs.Num(); JobIndex++)
 		{
@@ -5360,7 +5357,6 @@ void GlobalBeginCompileShader(
 #endif
 
 	Input.Target = Target;
-	Input.ShaderPlatformName = FDataDrivenShaderPlatformInfo::GetName(ShaderPlatform);
 	Input.ShaderFormat = ShaderFormatName;
 	Input.CompressionFormat = GetShaderCompressionFormat();
 	GetShaderCompressionOodleSettings(Input.OodleCompressor, Input.OodleLevel);
@@ -5369,7 +5365,7 @@ void GlobalBeginCompileShader(
 	Input.bCompilingForShaderPipeline = false;
 	Input.bIncludeUsedOutputs = false;
 	Input.bGenerateDirectCompileFile = (GDumpShaderDebugInfoSCWCommandLine != 0);
-	Input.DumpDebugInfoRootPath = GShaderCompilingManager->GetAbsoluteShaderDebugInfoDirectory() / Input.ShaderPlatformName.ToString();
+	Input.DumpDebugInfoRootPath = GShaderCompilingManager->GetAbsoluteShaderDebugInfoDirectory() / ShaderFormatName.ToString();
 	// asset material name or "Global"
 	Input.DebugGroupName = DebugGroupName;
 	Input.DebugDescription = DebugDescription;
@@ -6754,7 +6750,7 @@ void VerifyGlobalShaders(EShaderPlatform Platform, bool bLoadedFromCacheFile, co
 
 #include "Misc/PreLoadFile.h"
 #include "Serialization/LargeMemoryReader.h"
-static FPreLoadFile GGlobalShaderPreLoadFile(*(FString(TEXT("../../../Engine")) / TEXT("GlobalShaderCache-SP_") + FPlatformProperties::IniPlatformName() + TEXT(".bin")));
+static FPreLoadFile GGlobalShaderPreLoadFile(*(FString(TEXT("../../../Engine")) / TEXT("GlobalShaderCache-SF_") + FPlatformProperties::IniPlatformName() + TEXT(".bin")));
 
 static const ITargetPlatform* GGlobalShaderTargetPlatform[SP_NumPlatforms] = { nullptr };
 
@@ -6769,12 +6765,12 @@ static FString GetGlobalShaderCacheOverrideFilename(EShaderPlatform Platform)
 		DirectoryPrefix = GGlobalShaderCacheOverrideDirectory / TEXT("GlobalShaderCache-");
 	}
 
-	return DirectoryPrefix + FDataDrivenShaderPlatformInfo::GetName(Platform).ToString() + TEXT(".bin");
+	return DirectoryPrefix + LegacyShaderPlatformToShaderFormat(Platform).ToString() + TEXT(".bin");
 }
 
 static FString GetGlobalShaderCacheFilename(EShaderPlatform Platform)
 {
-	return FString(TEXT("Engine")) / TEXT("GlobalShaderCache-") + FDataDrivenShaderPlatformInfo::GetName(Platform).ToString() + TEXT(".bin");
+	return FString(TEXT("Engine")) / TEXT("GlobalShaderCache-") + LegacyShaderPlatformToShaderFormat(Platform).ToString() + TEXT(".bin");
 }
 
 #if WITH_EDITOR
diff --git a/Engine/Source/Runtime/Engine/Private/ShaderCompiler/ShaderGenerationUtil.cpp b/Engine/Source/Runtime/Engine/Private/ShaderCompiler/ShaderGenerationUtil.cpp
--- a/Engine/Source/Runtime/Engine/Private/ShaderCompiler/ShaderGenerationUtil.cpp
+++ b/Engine/Source/Runtime/Engine/Private/ShaderCompiler/ShaderGenerationUtil.cpp
@@ -1530,7 +1530,7 @@ static FCriticalSection GCriticalSection;
 
 static FString GetAutoGenDirectory(EShaderPlatform TargetPlatform)
 {
-	FString PlatformName = FDataDrivenShaderPlatformInfo::GetName(TargetPlatform).ToString();
+	FString PlatformName = LegacyShaderPlatformToShaderFormat(TargetPlatform).GetPlainNameString();
 	FString AutogenHeaderDirectory = FPaths::ProjectIntermediateDir() / TEXT("ShaderAutogen") / PlatformName;
 	return AutogenHeaderDirectory;
 }
@@ -1967,6 +1967,7 @@ FGBufferInfo FShaderCompileUtilities::FetchGBufferInfoAndWriteAutogen(EShaderPla
 
 	FGBufferInfo BufferInfo = FetchFullGBufferInfo(Params);
 
+	FString PlatformName = LegacyShaderPlatformToShaderFormat(TargetPlatform).GetPlainNameString();
 	FString AutoGenDirectory = GetAutoGenDirectory(TargetPlatform);
 	FString AutogenHeaderFilename = AutoGenDirectory / TEXT("AutogenShaderHeaders.ush");
 	FString AutogenHeaderFilenameTemp = AutoGenDirectory / TEXT("AutogenShaderHeaders_temp.ush");
diff --git a/Engine/Source/Runtime/RHI/Private/RHI.cpp b/Engine/Source/Runtime/RHI/Private/RHI.cpp
--- a/Engine/Source/Runtime/RHI/Private/RHI.cpp
+++ b/Engine/Source/Runtime/RHI/Private/RHI.cpp
@@ -2431,7 +2431,7 @@ void FGenericDataDrivenShaderPlatformInfo::SetDefaultValues()
 {
 	MaxFeatureLevel = ERHIFeatureLevel::Num;
 	bSupportsMSAA = true;
-	 
+
 	bNeedsToSwitchVerticalAxisOnMobileOpenGL = true;
 	bSupportsDOFHybridScattering = true;
 	bSupportsHZBOcclusion = true;
@@ -2445,8 +2445,6 @@ void FGenericDataDrivenShaderPlatformInfo::ParseDataDrivenShaderInfo(const FConf
 {
 	Info.Language = *GetSectionString(Section, "Language");
 	Info.ShaderFormat = *GetSectionString(Section, "ShaderFormat");
-	checkf(!Info.ShaderFormat.IsNone(), TEXT("Missing ShaderFormat for ShaderPlatform %s  ShaderFormat %s"), *Info.Name.ToString(), *Info.ShaderFormat.ToString());
-
 	GetFeatureLevelFromName(*GetSectionString(Section, "MaxFeatureLevel"), Info.MaxFeatureLevel);
 
 	Info.ShaderPropertiesHash = 0;
diff --git a/Engine/Source/Runtime/RenderCore/Private/ShaderCodeArchive.cpp b/Engine/Source/Runtime/RenderCore/Private/ShaderCodeArchive.cpp
--- a/Engine/Source/Runtime/RenderCore/Private/ShaderCodeArchive.cpp
+++ b/Engine/Source/Runtime/RenderCore/Private/ShaderCodeArchive.cpp
@@ -1643,7 +1643,7 @@ void FIoStoreShaderCodeArchive::SaveIoStoreShaderCodeArchive(const FIoStoreShade
 
 FIoStoreShaderCodeArchive* FIoStoreShaderCodeArchive::Create(EShaderPlatform InPlatform, const FString& InLibraryName, FIoDispatcher& InIoDispatcher)
 {
-	const FName PlatformName = FName(FDataDrivenShaderPlatformInfo::GetShaderFormat(InPlatform).ToString() + TEXT("-") + FDataDrivenShaderPlatformInfo::GetName(InPlatform).ToString());
+	const FName PlatformName = LegacyShaderPlatformToShaderFormat(InPlatform);
 	FIoChunkId ChunkId = GetShaderCodeArchiveChunkId(InLibraryName, PlatformName);
 	if (InIoDispatcher.DoesChunkExist(ChunkId))
 	{
diff --git a/Engine/Source/Runtime/RenderCore/Private/ShaderCodeLibrary.cpp b/Engine/Source/Runtime/RenderCore/Private/ShaderCodeLibrary.cpp
--- a/Engine/Source/Runtime/RenderCore/Private/ShaderCodeLibrary.cpp
+++ b/Engine/Source/Runtime/RenderCore/Private/ShaderCodeLibrary.cpp
@@ -766,11 +766,8 @@ public:
 
 		if (!Library)
 		{
-			const FName PlatformName = FDataDrivenShaderPlatformInfo::GetName(InShaderPlatform);
-			const FName ShaderFormatName = LegacyShaderPlatformToShaderFormat(InShaderPlatform);
-			FString ShaderFormatAndPlatform = ShaderFormatName.ToString() + TEXT("-") + PlatformName.ToString();
-
-			const FString DestFilePath = GetCodeArchiveFilename(ShaderCodeDir, InLibraryName, FName(ShaderFormatAndPlatform));
+			const FName PlatformName = LegacyShaderPlatformToShaderFormat(InShaderPlatform);
+			const FString DestFilePath = GetCodeArchiveFilename(ShaderCodeDir, InLibraryName, PlatformName);
 			TUniquePtr<FArchive> Ar(IFileManager::Get().CreateFileReader(*DestFilePath));
 			if (Ar)
 			{
@@ -1122,13 +1119,7 @@ struct FEditorShaderCodeArchive
 		, Format(nullptr)
 		, bNeedsDeterministicOrder(bInNeedsDeterministicOrder)
 	{
-		TArray<FString> Components;
-		FString Name = FormatName.ToString();
-		Name.ParseIntoArray(Components, TEXT("-"));
-
-		FName ShaderFormatName(Components[0]);
-		Format = GetTargetPlatformManagerRef().FindShaderFormat(ShaderFormatName);
-
+		Format = GetTargetPlatformManagerRef().FindShaderFormat(InFormat);
 		check(Format);
 
 		SerializedShaders.ShaderHashTable.Clear(0x10000);
@@ -2340,12 +2331,10 @@ public:
 
 			EShaderPlatform Platform = ShaderFormatToLegacyShaderPlatform(Format);
 			FName PossiblyAdjustedFormat = LegacyShaderPlatformToShaderFormat(Platform);	// Vulkan and GL switch between name variants depending on CVars 
-
-			FString FormatAndPlatformName = PossiblyAdjustedFormat.ToString() + TEXT("-") + FDataDrivenShaderPlatformInfo::GetName(Platform).ToString();
 			FEditorShaderCodeArchive* CodeArchive = EditorShaderCodeArchive[Platform];
 			if (!CodeArchive)
 			{
-				CodeArchive = new FEditorShaderCodeArchive(FName(FormatAndPlatformName), Descriptor.bNeedsDeterministicOrder);
+				CodeArchive = new FEditorShaderCodeArchive(PossiblyAdjustedFormat, Descriptor.bNeedsDeterministicOrder);
 				EditorShaderCodeArchive[Platform] = CodeArchive;
 				EditorArchivePipelines[Platform] = !bNativeFormat;
 			}
@@ -2990,13 +2979,10 @@ bool FShaderLibraryCooker::CreatePatchLibrary(TArray<FString> const& OldMetaData
 		if (Name.RemoveFromStart(TEXT("ShaderArchive-")))
 		{
 			TArray<FString> Components;
-			if (Name.ParseIntoArray(Components, TEXT("-")) == 3)
+			if (Name.ParseIntoArray(Components, TEXT("-")) == 2)
 			{
 				FName Format(*Components[1]);
-				FName Platform(*Components[2]);
-				FString FormatAndPlatform = Format.ToString() + TEXT("-") + Platform.ToString();
-
-				TSet<FString>& Libraries = FormatLibraryMap.FindOrAdd(FName(FormatAndPlatform));
+				TSet<FString>& Libraries = FormatLibraryMap.FindOrAdd(Format);
 				Libraries.Add(Components[0]);
 			}
 		}
diff --git a/Engine/Source/Runtime/RenderCore/Private/ShaderCore.cpp b/Engine/Source/Runtime/RenderCore/Private/ShaderCore.cpp
--- a/Engine/Source/Runtime/RenderCore/Private/ShaderCore.cpp
+++ b/Engine/Source/Runtime/RenderCore/Private/ShaderCore.cpp
@@ -660,7 +660,7 @@ bool CheckVirtualShaderFilePath(FStringView VirtualFilePath, TArray<FShaderCompi
 * @param OutVirtualFilePaths - [out] list of shader source files to add to
 * @param ShaderFilename - shader file to add
 */
-void AddShaderSourceFileEntry(TArray<FString>& OutVirtualFilePaths, FString VirtualFilePath, EShaderPlatform ShaderPlatform, const FName* ShaderPlatformName)
+void AddShaderSourceFileEntry(TArray<FString>& OutVirtualFilePaths, FString VirtualFilePath, EShaderPlatform ShaderPlatform)
 {
 	check(CheckVirtualShaderFilePath(VirtualFilePath));
 	if (!OutVirtualFilePaths.Contains(VirtualFilePath))
@@ -668,9 +668,7 @@ void AddShaderSourceFileEntry(TArray<FString>& OutVirtualFilePaths, FString Virt
 		OutVirtualFilePaths.Add(VirtualFilePath);
 
 		TArray<FString> ShaderIncludes;
-
-		const uint32 DepthLimit = 100;
-		GetShaderIncludes(*VirtualFilePath, *VirtualFilePath, OutVirtualFilePaths, ShaderPlatform, DepthLimit, ShaderPlatformName);
+		GetShaderIncludes(*VirtualFilePath, *VirtualFilePath, OutVirtualFilePaths, ShaderPlatform);
 		for( int32 IncludeIdx=0; IncludeIdx < ShaderIncludes.Num(); IncludeIdx++ )
 		{
 			OutVirtualFilePaths.AddUnique(ShaderIncludes[IncludeIdx]);
@@ -683,7 +681,7 @@ void AddShaderSourceFileEntry(TArray<FString>& OutVirtualFilePaths, FString Virt
 *
 * @param OutVirtualFilePaths - [out] list of shader source files to add to
 */
-void GetAllVirtualShaderSourcePaths(TArray<FString>& OutVirtualFilePaths, EShaderPlatform ShaderPlatform, const FName* ShaderPlatformName)
+void GetAllVirtualShaderSourcePaths(TArray<FString>& OutVirtualFilePaths, EShaderPlatform ShaderPlatform)
 {
 	// add all shader source files for hashing
 	for( TLinkedList<FVertexFactoryType*>::TIterator FactoryIt(FVertexFactoryType::GetTypeList()); FactoryIt; FactoryIt.Next() )
@@ -692,7 +690,7 @@ void GetAllVirtualShaderSourcePaths(TArray<FString>& OutVirtualFilePaths, EShade
 		if( VertexFactoryType )
 		{
 			FString ShaderFilename(VertexFactoryType->GetShaderFilename());
-			AddShaderSourceFileEntry(OutVirtualFilePaths, ShaderFilename, ShaderPlatform, ShaderPlatformName);
+			AddShaderSourceFileEntry(OutVirtualFilePaths, ShaderFilename, ShaderPlatform);
 		}
 	}
 	for( TLinkedList<FShaderType*>::TIterator ShaderIt(FShaderType::GetTypeList()); ShaderIt; ShaderIt.Next() )
@@ -701,7 +699,7 @@ void GetAllVirtualShaderSourcePaths(TArray<FString>& OutVirtualFilePaths, EShade
 		if(ShaderType)
 		{
 			FString ShaderFilename(ShaderType->GetShaderFilename());
-			AddShaderSourceFileEntry(OutVirtualFilePaths, ShaderFilename, ShaderPlatform, ShaderPlatformName);
+			AddShaderSourceFileEntry(OutVirtualFilePaths, ShaderFilename, ShaderPlatform);
 		}
 	}
 
@@ -894,7 +892,7 @@ bool ReplaceVirtualFilePathForShaderPlatform(FString& InOutVirtualFilePath, ESha
 	return false;
 }
 
-bool ReplaceVirtualFilePathForShaderAutogen(FString& InOutVirtualFilePath, EShaderPlatform ShaderPlatform, const FName* InShaderPlatformName)
+bool ReplaceVirtualFilePathForShaderAutogen(FString& InOutVirtualFilePath, EShaderPlatform ShaderPlatform)
 {
 	const FStringView ShaderAutogenStem = TEXTVIEW("/Engine/Generated/ShaderAutogen/");
 
@@ -902,17 +900,16 @@ bool ReplaceVirtualFilePathForShaderAutogen(FString& InOutVirtualFilePath, EShad
 	// for examples, if it starts with "/Engine/Generated/ShaderAutogen/" change it to "ShaderAutogen/PCD3D_SM5/"
 	if (!FCString::Strnicmp(*InOutVirtualFilePath, ShaderAutogenStem.GetData(), ShaderAutogenStem.Len()))
 	{
-		check(FDataDrivenShaderPlatformInfo::IsValid(ShaderPlatform) || InShaderPlatformName != nullptr);
-		FName ShaderPlatformName = FDataDrivenShaderPlatformInfo::IsValid(ShaderPlatform) ? FDataDrivenShaderPlatformInfo::GetName(ShaderPlatform) : *InShaderPlatformName;
 		TStringBuilder<1024> OutputShaderName;
 
 		// Append the prefix.
 		OutputShaderName.Append(TEXTVIEW("/ShaderAutogen/"));
 
 		// Append the platform name.
-
-		FString PlatformNameString = ShaderPlatformName.ToString();
-		OutputShaderName.Append(*PlatformNameString, PlatformNameString.Len());
+		const FName PlatformName = LegacyShaderPlatformToShaderFormat(ShaderPlatform);
+		TCHAR PlatformNameString[NAME_SIZE];
+		const uint32 PlatformNameLen = PlatformName.GetPlainNameString(PlatformNameString);
+		OutputShaderName.Append(PlatformNameString, PlatformNameLen);
 
 		OutputShaderName.AppendChar(TEXT('/'));
 
@@ -927,7 +924,7 @@ bool ReplaceVirtualFilePathForShaderAutogen(FString& InOutVirtualFilePath, EShad
 	return false;
 }
 
-bool LoadShaderSourceFile(const TCHAR* InVirtualFilePath, EShaderPlatform ShaderPlatform, FString* OutFileContents, TArray<FShaderCompilerError>* OutCompileErrors, const FName* ShaderPlatformName) // TODO: const FString&
+bool LoadShaderSourceFile(const TCHAR* InVirtualFilePath, EShaderPlatform ShaderPlatform, FString* OutFileContents, TArray<FShaderCompilerError>* OutCompileErrors) // TODO: const FString&
 {
 	// it's not expected that cooked platforms get here, but if they do, this is the final out
 	if (FPlatformProperties::RequiresCookedData())
@@ -947,7 +944,7 @@ bool LoadShaderSourceFile(const TCHAR* InVirtualFilePath, EShaderPlatform Shader
 		ReplaceVirtualFilePathForShaderPlatform(VirtualFilePath, ShaderPlatform);
 
 		// Fixup autogen file
-		ReplaceVirtualFilePathForShaderAutogen(VirtualFilePath, ShaderPlatform, ShaderPlatformName);
+		ReplaceVirtualFilePathForShaderAutogen(VirtualFilePath, ShaderPlatform);
 
 		// Protect GShaderFileCache from simultaneous access by multiple threads
 		FScopeLock ScopeLock(&FileCacheCriticalSection);
@@ -1051,7 +1048,7 @@ static const TCHAR* FindFirstInclude(const TCHAR* Text)
 /**
  * Recursively populates IncludeFilenames with the unique include filenames found in the shader file named Filename.
  */
-static void InternalGetShaderIncludes(const TCHAR* EntryPointVirtualFilePath, const TCHAR* VirtualFilePath, const FString& FileContents, TArray<FString>& IncludeVirtualFilePaths, EShaderPlatform ShaderPlatform, uint32 DepthLimit, bool AddToIncludeFile, const FName* ShaderPlatformName)
+static void GetShaderIncludes(const TCHAR* EntryPointVirtualFilePath, const TCHAR* VirtualFilePath, const FString& FileContents, TArray<FString>& IncludeVirtualFilePaths, EShaderPlatform ShaderPlatform, uint32 DepthLimit, bool AddToIncludeFile)
 {
 	//avoid an infinite loop with a 0 length string
 	if (FileContents.Len() > 0)
@@ -1100,7 +1097,7 @@ static void InternalGetShaderIncludes(const TCHAR* EntryPointVirtualFilePath, co
 					ReplaceVirtualFilePathForShaderPlatform(ExtractedIncludeFilename, ShaderPlatform);
 
 					// Fixup autogen file
-					ReplaceVirtualFilePathForShaderAutogen(ExtractedIncludeFilename, ShaderPlatform, ShaderPlatformName);
+					ReplaceVirtualFilePathForShaderAutogen(ExtractedIncludeFilename, ShaderPlatform);
 
 					// Ignore uniform buffer, vertex factory and instanced stereo includes
 					bool bIgnoreInclude = ExtractedIncludeFilename.StartsWith(TEXT("/Engine/Generated/"));
@@ -1121,8 +1118,8 @@ static void InternalGetShaderIncludes(const TCHAR* EntryPointVirtualFilePath, co
 						if (!IncludeVirtualFilePaths.Contains(ExtractedIncludeFilename))
 						{
 							FString IncludedFileContents;
-							LoadShaderSourceFile(*ExtractedIncludeFilename, ShaderPlatform, &IncludedFileContents, nullptr, ShaderPlatformName);
-							InternalGetShaderIncludes(EntryPointVirtualFilePath, *ExtractedIncludeFilename, IncludedFileContents, IncludeVirtualFilePaths, ShaderPlatform, DepthLimit - 1, true, ShaderPlatformName);
+							LoadShaderSourceFile(*ExtractedIncludeFilename, ShaderPlatform, &IncludedFileContents, nullptr);
+							GetShaderIncludes(EntryPointVirtualFilePath, *ExtractedIncludeFilename, IncludedFileContents, IncludeVirtualFilePaths, ShaderPlatform, DepthLimit - 1, true);
 						}
 					}
 				}
@@ -1153,22 +1150,22 @@ static void InternalGetShaderIncludes(const TCHAR* EntryPointVirtualFilePath, co
 /**
  * Recursively populates IncludeFilenames with the unique include filenames found in the shader file named Filename.
  */
-static void InternalGetShaderIncludes(const TCHAR* EntryPointVirtualFilePath, const TCHAR* VirtualFilePath, TArray<FString>& IncludeVirtualFilePaths, EShaderPlatform ShaderPlatform, uint32 DepthLimit, bool AddToIncludeFile, const FName* ShaderPlatformName)
+static void GetShaderIncludes(const TCHAR* EntryPointVirtualFilePath, const TCHAR* VirtualFilePath, TArray<FString>& IncludeVirtualFilePaths, EShaderPlatform ShaderPlatform, uint32 DepthLimit, bool AddToIncludeFile)
 {
 	FString FileContents;
-	LoadShaderSourceFile(VirtualFilePath, ShaderPlatform, &FileContents, nullptr, ShaderPlatformName);
+	LoadShaderSourceFile(VirtualFilePath, ShaderPlatform, &FileContents, nullptr);
 
-	InternalGetShaderIncludes(EntryPointVirtualFilePath, VirtualFilePath, FileContents, IncludeVirtualFilePaths, ShaderPlatform, DepthLimit, AddToIncludeFile, ShaderPlatformName);
+	GetShaderIncludes(EntryPointVirtualFilePath, VirtualFilePath, FileContents, IncludeVirtualFilePaths, ShaderPlatform, DepthLimit, AddToIncludeFile);
 }
 
-void GetShaderIncludes(const TCHAR* EntryPointVirtualFilePath, const TCHAR* VirtualFilePath, TArray<FString>& IncludeVirtualFilePaths, EShaderPlatform ShaderPlatform, uint32 DepthLimit, const FName* ShaderPlatformName)
+void GetShaderIncludes(const TCHAR* EntryPointVirtualFilePath, const TCHAR* VirtualFilePath, TArray<FString>& IncludeVirtualFilePaths, EShaderPlatform ShaderPlatform, uint32 DepthLimit)
 {
-	InternalGetShaderIncludes(EntryPointVirtualFilePath, VirtualFilePath, IncludeVirtualFilePaths, ShaderPlatform, DepthLimit, false, ShaderPlatformName);
+	GetShaderIncludes(EntryPointVirtualFilePath, VirtualFilePath, IncludeVirtualFilePaths, ShaderPlatform, DepthLimit, false);
 }
 
-void GetShaderIncludes(const TCHAR* EntryPointVirtualFilePath, const TCHAR* VirtualFilePath, const FString& FileContents, TArray<FString>& IncludeVirtualFilePaths, EShaderPlatform ShaderPlatform, uint32 DepthLimit, const FName* ShaderPlatformName)
+void GetShaderIncludes(const TCHAR* EntryPointVirtualFilePath, const TCHAR* VirtualFilePath, const FString& FileContents, TArray<FString>& IncludeVirtualFilePaths, EShaderPlatform ShaderPlatform, uint32 DepthLimit)
 {
-	InternalGetShaderIncludes(EntryPointVirtualFilePath, VirtualFilePath, FileContents, IncludeVirtualFilePaths, ShaderPlatform, DepthLimit, false, ShaderPlatformName);
+	GetShaderIncludes(EntryPointVirtualFilePath, VirtualFilePath, FileContents, IncludeVirtualFilePaths, ShaderPlatform, DepthLimit, false);
 }
 
 void HashShaderFileWithIncludes(FArchive& HashingArchive, const TCHAR* VirtualFilePath, const FString& FileContents, EShaderPlatform ShaderPlatform, bool bOnlyHashIncludedFiles)
@@ -1338,12 +1335,12 @@ const FSHAHash& GetShaderFilesHash(const TArray<FString>& VirtualFilePaths, ESha
 	}
 }
 
-void BuildShaderFileToUniformBufferMap(TMap<FString, TArray<const TCHAR*> >& ShaderFileToUniformBufferVariables, const FName* ShaderPlatformName)
+void BuildShaderFileToUniformBufferMap(TMap<FString, TArray<const TCHAR*> >& ShaderFileToUniformBufferVariables)
 {
 	if (!FPlatformProperties::RequiresCookedData())
 	{
 		TArray<FString> ShaderSourceFiles;
-		GetAllVirtualShaderSourcePaths(ShaderSourceFiles, GMaxRHIShaderPlatform, ShaderPlatformName);
+		GetAllVirtualShaderSourcePaths(ShaderSourceFiles, GMaxRHIShaderPlatform);
 
 		FScopedSlowTask SlowTask((float)ShaderSourceFiles.Num());
 
@@ -1444,7 +1441,7 @@ void UninitializeShaderTypes()
  * Flushes the shader file and CRC cache, and regenerates the binary shader files if necessary.
  * Allows shader source files to be re-read properly even if they've been modified since startup.
  */
-void FlushShaderFileCache(const FName* ShaderPlatformName)
+void FlushShaderFileCache()
 {
 	UE_LOG(LogShaders, Log, TEXT("FlushShaderFileCache() begin"));
 
@@ -1462,7 +1459,7 @@ void FlushShaderFileCache(const FName* ShaderPlatformName)
 		LogShaderSourceDirectoryMappings();
 
 		TMap<FString, TArray<const TCHAR*> > ShaderFileToUniformBufferVariables;
-		BuildShaderFileToUniformBufferMap(ShaderFileToUniformBufferVariables, ShaderPlatformName);
+		BuildShaderFileToUniformBufferMap(ShaderFileToUniformBufferVariables);
 
 		for (TLinkedList<FShaderPipelineType*>::TConstIterator It(FShaderPipelineType::GetTypeList()); It; It.Next())
 		{
@@ -1831,11 +1828,6 @@ FArchive& operator<<(FArchive& Ar, FShaderCompilerInput& Input)
 		Ar << CompressionFormatString;
 		Input.CompressionFormat = FName(*CompressionFormatString);
 	}
-	{
-		FString ShaderPlatformNameString(Input.ShaderPlatformName.ToString());
-		Ar << ShaderPlatformNameString;
-		Input.ShaderPlatformName = FName(*ShaderPlatformNameString);
-	}
 	Ar << Input.SourceFilePrefix;
 	Ar << Input.VirtualSourceFilePath;
 	Ar << Input.EntryPointName;
diff --git a/Engine/Source/Runtime/RenderCore/Public/ShaderCompilerCore.h b/Engine/Source/Runtime/RenderCore/Public/ShaderCompilerCore.h
--- a/Engine/Source/Runtime/RenderCore/Public/ShaderCompilerCore.h
+++ b/Engine/Source/Runtime/RenderCore/Public/ShaderCompilerCore.h
@@ -177,7 +177,6 @@ namespace FOodleDataCompression
 struct FShaderCompilerInput
 {
 	FShaderTarget Target;
-	FName ShaderPlatformName;
 	FName ShaderFormat;
 	FName CompressionFormat;
 	FString SourceFilePrefix;
@@ -528,7 +527,7 @@ extern RENDERCORE_API bool CheckVirtualShaderFilePath(FStringView VirtualPath, T
  * @param OutFileContents - If true is returned, will contain the contents of the shader file. Can be null.
  * @return True if the file was successfully loaded.
  */
-extern RENDERCORE_API bool LoadShaderSourceFile(const TCHAR* VirtualFilePath, EShaderPlatform ShaderPlatform, FString* OutFileContents, TArray<FShaderCompilerError>* OutCompileErrors, const FName* ShaderPlatformName = nullptr);
+extern RENDERCORE_API bool LoadShaderSourceFile(const TCHAR* VirtualFilePath, EShaderPlatform ShaderPlatform, FString* OutFileContents, TArray<FShaderCompilerError>* OutCompileErrors);
 
 enum class EShaderCompilerWorkerType : uint8
 {
diff --git a/Engine/Source/Runtime/RenderCore/Public/ShaderCore.h b/Engine/Source/Runtime/RenderCore/Public/ShaderCore.h
--- a/Engine/Source/Runtime/RenderCore/Public/ShaderCore.h
+++ b/Engine/Source/Runtime/RenderCore/Public/ShaderCore.h
@@ -1010,7 +1010,7 @@ extern RENDERCORE_API FString ParseVirtualShaderFilename(const FString& InFilena
 extern RENDERCORE_API bool ReplaceVirtualFilePathForShaderPlatform(FString& InOutVirtualFilePath, EShaderPlatform ShaderPlatform);
 
 /** Replaces virtual platform path with appropriate autogen path for a given ShaderPlatform. Returns true if path was changed. */
-extern RENDERCORE_API bool ReplaceVirtualFilePathForShaderAutogen(FString& InOutVirtualFilePath, EShaderPlatform ShaderPlatform, const FName* InShaderPlatformName = nullptr);
+extern RENDERCORE_API bool ReplaceVirtualFilePathForShaderAutogen(FString& InOutVirtualFilePath, EShaderPlatform ShaderPlatform);
 
 /** Loads the shader file with the given name.  If the shader file couldn't be loaded, throws a fatal error. */
 extern RENDERCORE_API void LoadShaderSourceFileChecked(const TCHAR* VirtualFilePath, EShaderPlatform ShaderPlatform, FString& OutFileContents);
@@ -1018,8 +1018,8 @@ extern RENDERCORE_API void LoadShaderSourceFileChecked(const TCHAR* VirtualFileP
 /**
  * Recursively populates IncludeFilenames with the include filenames from Filename
  */
-extern RENDERCORE_API void GetShaderIncludes(const TCHAR* EntryPointVirtualFilePath, const TCHAR* VirtualFilePath, TArray<FString>& IncludeVirtualFilePaths, EShaderPlatform ShaderPlatform, uint32 DepthLimit=100, const FName* ShaderPlatformName = nullptr);
-extern RENDERCORE_API void GetShaderIncludes(const TCHAR* EntryPointVirtualFilePath, const TCHAR* VirtualFilePath, const FString& FileContents, TArray<FString>& IncludeVirtualFilePaths, EShaderPlatform ShaderPlatform, uint32 DepthLimit = 100, const FName* ShaderPlatformName = nullptr);
+extern RENDERCORE_API void GetShaderIncludes(const TCHAR* EntryPointVirtualFilePath, const TCHAR* VirtualFilePath, TArray<FString>& IncludeVirtualFilePaths, EShaderPlatform ShaderPlatform, uint32 DepthLimit=100);
+extern RENDERCORE_API void GetShaderIncludes(const TCHAR* EntryPointVirtualFilePath, const TCHAR* VirtualFilePath, const FString& FileContents, TArray<FString>& IncludeVirtualFilePaths, EShaderPlatform ShaderPlatform, uint32 DepthLimit = 100);
 
 /**
  * Calculates a Hash for the given filename if it does not already exist in the Hash cache.
@@ -1043,13 +1043,13 @@ extern RENDERCORE_API void HashShaderFileWithIncludes(FArchive& HashingArchive,
  */
 extern RENDERCORE_API const class FSHAHash& GetShaderFilesHash(const TArray<FString>& VirtualFilePaths, EShaderPlatform ShaderPlatform);
 
-extern void BuildShaderFileToUniformBufferMap(TMap<FString, TArray<const TCHAR*> >& ShaderFileToUniformBufferVariables, const FName* ShaderPlatformName = nullptr);
+extern void BuildShaderFileToUniformBufferMap(TMap<FString, TArray<const TCHAR*> >& ShaderFileToUniformBufferVariables);
 
 /**
  * Flushes the shader file and CRC cache, and regenerates the binary shader files if necessary.
  * Allows shader source files to be re-read properly even if they've been modified since startup.
  */
-extern RENDERCORE_API void FlushShaderFileCache(const FName* ShaderPlatformName = nullptr);
+extern RENDERCORE_API void FlushShaderFileCache();
 
 extern RENDERCORE_API void VerifyShaderSourceFiles(EShaderPlatform ShaderPlatform);
 
@@ -1098,5 +1098,5 @@ extern RENDERCORE_API void ResetAllShaderSourceDirectoryMappings();
  */
 extern RENDERCORE_API void AddShaderSourceDirectoryMapping(const FString& VirtualShaderDirectory, const FString& RealShaderDirectory);
 
-extern RENDERCORE_API void AddShaderSourceFileEntry(TArray<FString>& OutVirtualFilePaths, FString VirtualFilePath, EShaderPlatform ShaderPlatform, const FName* ShaderPlatformName = nullptr);
-extern RENDERCORE_API void GetAllVirtualShaderSourcePaths(TArray<FString>& OutVirtualFilePaths, EShaderPlatform ShaderPlatform, const FName* ShaderPlatformName = nullptr);
+extern RENDERCORE_API void AddShaderSourceFileEntry(TArray<FString>& OutVirtualFilePaths, FString VirtualFilePath, EShaderPlatform ShaderPlatform);
+extern RENDERCORE_API void GetAllVirtualShaderSourcePaths(TArray<FString>& OutVirtualFilePaths, EShaderPlatform ShaderPlatform);
