Implemented support for asynchronous creation and deletion of FRenderResource.
The resource list uses a thread local free list that is coalesced later.

diff --git a/Engine/Source/Runtime/Engine/Private/World.cpp b/Engine/Source/Runtime/Engine/Private/World.cpp
--- a/Engine/Source/Runtime/Engine/Private/World.cpp
+++ b/Engine/Source/Runtime/Engine/Private/World.cpp
@@ -8429,6 +8429,8 @@ void UWorld::ChangeFeatureLevel(ERHIFeatureLevel::Type InFeatureLevel, bool bSho
 
 			SlowTask.EnterProgressFrame(10.0f);
 			FGlobalComponentReregisterContext RecreateComponents;
+			// Finish any deferred / async render cleanup work.
+			GetRendererModule().PerFrameCleanupIfSkipRenderer();
 			FlushRenderingCommands();
 
 			FeatureLevel = InFeatureLevel;
diff --git a/Engine/Source/Runtime/RenderCore/Private/RenderResource.cpp b/Engine/Source/Runtime/RenderCore/Private/RenderResource.cpp
--- a/Engine/Source/Runtime/RenderCore/Private/RenderResource.cpp
+++ b/Engine/Source/Runtime/RenderCore/Private/RenderResource.cpp
@@ -42,28 +42,194 @@ FAutoConsoleVariableRef CVarDebugForceRuntimeBLAS(
 	TEXT("Force building BLAS at runtime."),
 	ECVF_ReadOnly);
 
-FThreadSafeCounter FRenderResource::ResourceListIterationActive;
+/** Tracks render resources in a list. The implementation is optimized to allow fast allocation / deallocation from any thread,
+ *  at the cost of period coalescing of thread-local data at a sync point each frame. Furthermore, iteration is not thread safe
+ *  and must be performed at sync points.
+ */
+class FRenderResourceList
+{
+public:
+	static FRenderResourceList& Get()
+	{
+		static FRenderResourceList Instance;
+		return Instance;
+	}
+
+	~FRenderResourceList()
+	{
+		for (FFreeList* FreeList : LocalFreeLists)
+		{
+			delete FreeList;
+		}
+		FPlatformTLS::FreeTlsSlot(TLSSlot);
+	}
+
+	int32 Allocate(FRenderResource* Resource)
+	{
+		if (bIsIterating)
+		{
+			// This part is not thread safe. Iteration requires that no adds / removals are happening concurrently. The only
+			// supported case is recursive adds on the same thread (i.e. a parent resource initializes a child resource). In
+			// this case, we need to add the resource to the end so that it gets iterated as well.
+			check(IsInRenderingThread());
+			return ResourceList.AddElement(Resource);
+		}
+
+		FFreeList& LocalFreeList = GetLocalFreeList();
+
+		if (LocalFreeList.IsEmpty())
+		{
+			FScopeLock Lock(&CS);
+
+			// Try to allocate free slots from the global free list.
+			const int32 OldFreeListSize = GlobalFreeList.Num();
+			const int32 NewFreeListSize = FMath::Max(OldFreeListSize - ChunkSize, 0);
+			int32 NumElements = OldFreeListSize - NewFreeListSize;
+
+			if (NumElements > 0)
+			{
+				LocalFreeList.Append(GlobalFreeList.GetData() + NewFreeListSize, NumElements);
+				GlobalFreeList.SetNum(NewFreeListSize, false);
+			}
+
+			// Allocate more if we didn't get a full chunk from the global list.
+			while (NumElements < ChunkSize)
+			{
+				LocalFreeList.Emplace(ResourceList.AddElement(nullptr));
+				NumElements++;
+			}
+		}
+
+		int32 Index = LocalFreeList.Pop(false);
+		ResourceList[Index] = Resource;
+		return Index;
+	}
+
+	void Deallocate(int32 Index)
+	{
+		GetLocalFreeList().Emplace(Index);
+		ResourceList[Index] = nullptr;
+	}
+
+	//////////////////////////////////////////////////////////////////////////////
+	// These methods must be called at sync points where allocations / deallocations can't occur from another thread.
+
+	void Clear()
+	{
+		check(IsInRenderingThread());
+
+		for (FFreeList* FreeList : LocalFreeLists)
+		{
+			FreeList->Empty();
+		}
+		ResourceList.Empty();
+	}
+
+	void Coalesce()
+	{
+		check(IsInRenderingThread());
+
+		for (FFreeList* FreeList : LocalFreeLists)
+		{
+			GlobalFreeList.Append(*FreeList);
+			FreeList->Empty();
+		}
+	}
+
+	template<typename FunctionType>
+	void ForEach(const FunctionType& Function)
+	{
+		check(IsInRenderingThread());
+		check(!bIsIterating);
+		bIsIterating = true;
+		for (int32 Index = 0; Index < ResourceList.Num(); ++Index)
+		{
+			FRenderResource* Resource = ResourceList[Index];
+			if (Resource)
+			{
+				check(Resource->GetListIndex() == Index);
+				Function(Resource);
+			}
+		}
+		bIsIterating = false;
+	}
+
+	template<typename FunctionType>
+	void ForEachReverse(const FunctionType& Function)
+	{
+		check(IsInRenderingThread());
+		check(!bIsIterating);
+		bIsIterating = true;
+		for (int32 Index = ResourceList.Num() - 1; Index >= 0; --Index)
+		{
+			FRenderResource* Resource = ResourceList[Index];
+			if (Resource)
+			{
+				check(Resource->GetListIndex() == Index);
+				Function(Resource);
+			}
+		}
+		bIsIterating = false;
+	}
+
+	//////////////////////////////////////////////////////////////////////////////
+
+private:
+	FRenderResourceList()
+	{
+		TLSSlot = FPlatformTLS::AllocTlsSlot();
+	}
+
+	const int32 ChunkSize = 1024;
+
+	using FFreeList = TArray<int32>;
+
+	FFreeList& GetLocalFreeList()
+	{
+		void* TLSValue = FPlatformTLS::GetTlsValue(TLSSlot);
+		if (TLSValue == nullptr)
+		{
+			FFreeList* TLSCache = new FFreeList();
+			FPlatformTLS::SetTlsValue(TLSSlot, (void*)(TLSCache));
+			FScopeLock S(&CS);
+			LocalFreeLists.Add(TLSCache);
+			return *TLSCache;
+		}
+		return *((FFreeList*)TLSValue);
+	}
+
+	uint32 TLSSlot;
+	FCriticalSection CS;
+	TArray<FFreeList*> LocalFreeLists;
+	FFreeList GlobalFreeList;
+	TChunkedArray<FRenderResource*> ResourceList;
+	bool bIsIterating = false;
+};
 
-TArray<int32>& GetFreeIndicesList()
+void FRenderResource::CoalesceResourceList()
 {
-	static TArray<int32> FreeIndicesList;
-	return FreeIndicesList;
+	FRenderResourceList::Get().Coalesce();
 }
 
-TArray<FRenderResource*>& FRenderResource::GetResourceList()
+void FRenderResource::ReleaseRHIForAllResources()
 {
-	static TArray<FRenderResource*> RenderResourceList;
-	return RenderResourceList;
+	FRenderResourceList& ResourceList = FRenderResourceList::Get();
+	ResourceList.ForEachReverse([](FRenderResource* Resource) { check(Resource->IsInitialized()); Resource->ReleaseRHI(); });
+	ResourceList.ForEachReverse([](FRenderResource* Resource) { Resource->ReleaseDynamicRHI(); });
 }
 
 /** Initialize all resources initialized before the RHI was initialized */
 void FRenderResource::InitPreRHIResources()
-{	
+{
+	FRenderResourceList& ResourceList = FRenderResourceList::Get();
+
 	// Notify all initialized FRenderResources that there's a valid RHI device to create their RHI resources for now.
-	FRenderResource::InitRHIForAllResources();
+	ResourceList.ForEach([](FRenderResource* Resource) { Resource->InitRHI(); });
+	// Dynamic resources can have dependencies on static resources (with uniform buffers) and must initialized last!
+	ResourceList.ForEach([](FRenderResource* Resource) { Resource->InitDynamicRHI(); });
 
 #if !PLATFORM_NEEDS_RHIRESOURCELIST
-	FRenderResource::GetResourceList().Empty();
+	ResourceList.Clear();
 #endif
 }
 
@@ -72,7 +238,7 @@ void FRenderResource::ChangeFeatureLevel(ERHIFeatureLevel::Type NewFeatureLevel)
 	ENQUEUE_RENDER_COMMAND(FRenderResourceChangeFeatureLevel)(
 		[NewFeatureLevel](FRHICommandList& RHICmdList)
 	{
-		FRenderResource::ForAllResources([NewFeatureLevel](FRenderResource* Resource)
+		FRenderResourceList::Get().ForEach([NewFeatureLevel](FRenderResource* Resource)
 		{
 			// Only resources configured for a specific feature level need to be updated
 			if (Resource->HasValidFeatureLevel() && (Resource->FeatureLevel != NewFeatureLevel))
@@ -89,7 +255,6 @@ void FRenderResource::ChangeFeatureLevel(ERHIFeatureLevel::Type NewFeatureLevel)
 
 void FRenderResource::InitResource()
 {
-	check(IsInRenderingThread());
 	if (ListIndex == INDEX_NONE)
 	{
 		int32 LocalListIndex = INDEX_NONE;
@@ -97,21 +262,7 @@ void FRenderResource::InitResource()
 		if (PLATFORM_NEEDS_RHIRESOURCELIST || !GIsRHIInitialized)
 		{
 			LLM_SCOPE(ELLMTag::SceneRender);
-			TArray<FRenderResource*>& ResourceList = GetResourceList();
-			TArray<int32>& FreeIndicesList = GetFreeIndicesList();
-
-			// If resource list is currently being iterated, new resources must be added to the end of the list, to ensure they're processed during the iteration
-			// Otherwise empty slots in the list may be re-used for new resources
-			if (FreeIndicesList.Num() > 0 && ResourceListIterationActive.GetValue() == 0)
-			{
-				LocalListIndex = FreeIndicesList.Pop();
-				check(ResourceList[LocalListIndex] == nullptr);
-				ResourceList[LocalListIndex] = this;
-			}
-			else
-			{
-				LocalListIndex = ResourceList.Add(this);
-			}
+			LocalListIndex = FRenderResourceList::Get().Allocate(this);
 		}
 		else
 		{
@@ -135,7 +286,6 @@ void FRenderResource::ReleaseResource()
 {
 	if ( !GIsCriticalError )
 	{
-		check(IsInRenderingThread());
 		if(ListIndex != INDEX_NONE)
 		{
 			if(GIsRHIInitialized)
@@ -145,10 +295,7 @@ void FRenderResource::ReleaseResource()
 			}
 
 #if PLATFORM_NEEDS_RHIRESOURCELIST
-			TArray<FRenderResource*>& ResourceList = GetResourceList();
-			TArray<int32>& FreeIndicesList = GetFreeIndicesList();
-			ResourceList[ListIndex] = nullptr;
-			FreeIndicesList.Add(ListIndex);
+			FRenderResourceList::Get().Deallocate(ListIndex);
 #endif
 			ListIndex = INDEX_NONE;
 		}
diff --git a/Engine/Source/Runtime/RenderCore/Public/RenderResource.h b/Engine/Source/Runtime/RenderCore/Public/RenderResource.h
--- a/Engine/Source/Runtime/RenderCore/Public/RenderResource.h
+++ b/Engine/Source/Runtime/RenderCore/Public/RenderResource.h
@@ -27,55 +27,23 @@ extern RENDERCORE_API bool GFreeStructuresOnRHIBufferCreation;
 class RENDERCORE_API FRenderResource
 {
 public:
-	template<typename FunctionType>
-	static void ForAllResources(const FunctionType& Function)
-	{
-		const TArray<FRenderResource*>& ResourceList = GetResourceList();
-		ResourceListIterationActive.Increment();
-		for (int32 Index = 0; Index < ResourceList.Num(); ++Index)
-		{
-			FRenderResource* Resource = ResourceList[Index];
-			if (Resource)
-			{
-				checkSlow(Resource->ListIndex == Index);
-				Function(Resource);
-			}
-		}
-		ResourceListIterationActive.Decrement();
-	}
+	////////////////////////////////////////////////////////////////////////////////////
+	// The following methods may not be called while asynchronously initializing / releasing render resources.
 
-	template<typename FunctionType>
-	static void ForAllResourcesReverse(const FunctionType& Function)
-	{
-		const TArray<FRenderResource*>& ResourceList = GetResourceList();
-		ResourceListIterationActive.Increment();
-		for (int32 Index = ResourceList.Num() - 1; Index >= 0; --Index)
-		{
-			FRenderResource* Resource = ResourceList[Index];
-			if (Resource)
-			{
-				checkSlow(Resource->ListIndex == Index);
-				Function(Resource);
-			}
-		}
-		ResourceListIterationActive.Decrement();
-	}
+	/** Release all render resources that are currently initialized. */
+	static void ReleaseRHIForAllResources();
 
-	static void InitRHIForAllResources()
-	{
-		ForAllResources([](FRenderResource* Resource) { Resource->InitRHI(); });
-		// Dynamic resources can have dependencies on static resources (with uniform buffers) and must initialized last!
-		ForAllResources([](FRenderResource* Resource) { Resource->InitDynamicRHI(); });
-	}
+	/** Initialize all resources initialized before the RHI was initialized. */
+	static void InitPreRHIResources();
 
-	static void ReleaseRHIForAllResources()
-	{
-		ForAllResourcesReverse([](FRenderResource* Resource) { check(Resource->IsInitialized()); Resource->ReleaseRHI(); });
-		ForAllResourcesReverse([](FRenderResource* Resource) { Resource->ReleaseDynamicRHI(); });
-	}
+	/** Call periodically to coalesce the render resource list. */
+	static void CoalesceResourceList();
 
+	/** Reinitializes render resources at a new feature level. */
 	static void ChangeFeatureLevel(ERHIFeatureLevel::Type NewFeatureLevel);
 
+	////////////////////////////////////////////////////////////////////////////////////
+
 	/** Default constructor. */
 	FRenderResource()
 		: ListIndex(INDEX_NONE)
@@ -145,15 +113,9 @@ public:
 	// Accessors.
 	FORCEINLINE bool IsInitialized() const { return ListIndex != INDEX_NONE; }
 
-	/** Initialize all resources initialized before the RHI was initialized */
-	static void InitPreRHIResources();
+	int32 GetListIndex() const { return ListIndex; }
 
 private:
-	/** @return The global initialized resource list. */
-	static TArray<FRenderResource*>& GetResourceList();
-
-	static FThreadSafeCounter ResourceListIterationActive;
-
 	int32 ListIndex;
 
 protected:
diff --git a/Engine/Source/Runtime/Renderer/Private/SceneRendering.cpp b/Engine/Source/Runtime/Renderer/Private/SceneRendering.cpp
--- a/Engine/Source/Runtime/Renderer/Private/SceneRendering.cpp
+++ b/Engine/Source/Runtime/Renderer/Private/SceneRendering.cpp
@@ -4037,6 +4037,7 @@ static void DeleteSceneRenderers(FRHICommandListImmediate& RHICmdList, const TAr
 	// Can release only after all mesh pass tasks are finished.
 	GPrimitiveIdVertexBufferPool.DiscardAll();
 	FGraphicsMinimalPipelineStateId::ResetLocalPipelineIdTableSize();
+	FRenderResource::CoalesceResourceList();
 }
 
 static void ReleaseSceneRenderers(FRHICommandListImmediate& RHICmdList, const TArray<FSceneRenderer*>& SceneRenderers)
