Make RDG transient pooled render target refcounting thread safe.

diff --git a/Engine/Source/Runtime/RenderCore/Private/RenderGraphResourcePool.cpp b/Engine/Source/Runtime/RenderCore/Private/RenderGraphResourcePool.cpp
--- a/Engine/Source/Runtime/RenderCore/Private/RenderGraphResourcePool.cpp
+++ b/Engine/Source/Runtime/RenderCore/Private/RenderGraphResourcePool.cpp
@@ -185,16 +185,14 @@ TGlobalResource<FRDGBufferPool> GRenderGraphResourcePool;
 
 uint32 FRDGTransientRenderTarget::AddRef() const
 {
-	check(IsInRenderingThread());
 	check(LifetimeState == ERDGTransientResourceLifetimeState::Allocated);
-	return ++RefCount;
+	return uint32(FPlatformAtomics::InterlockedIncrement(&RefCount));
 }
 
 uint32 FRDGTransientRenderTarget::Release()
 {
-	check(IsInRenderingThread());
-	check(RefCount > 0 && LifetimeState == ERDGTransientResourceLifetimeState::Allocated);
-	const uint32 Refs = --RefCount;
+	const int32 Refs = FPlatformAtomics::InterlockedDecrement(&RefCount);
+	check(Refs >= 0 && LifetimeState == ERDGTransientResourceLifetimeState::Allocated);
 	if (Refs == 0)
 	{
 		if (GRDGTransientResourceAllocator.IsValid())
@@ -285,6 +283,8 @@ void FRDGTransientResourceAllocator::AddPendingDeallocation(FRDGTransientRenderT
 	check(RenderTarget);
 	check(RenderTarget->GetRefCount() == 0);
 
+	FScopeLock Lock(&CS);
+
 	if (RenderTarget->Texture)
 	{
 		RenderTarget->LifetimeState = ERDGTransientResourceLifetimeState::PendingDeallocation;
@@ -299,6 +299,8 @@ void FRDGTransientResourceAllocator::AddPendingDeallocation(FRDGTransientRenderT
 
 void FRDGTransientResourceAllocator::ReleasePendingDeallocations()
 {
+	FScopeLock Lock(&CS);
+
 	if (!PendingDeallocationList.IsEmpty())
 	{
 		TArray<FRHITransitionInfo, SceneRenderingAllocator> Transitions;
diff --git a/Engine/Source/Runtime/RenderCore/Private/RenderGraphResourcePool.h b/Engine/Source/Runtime/RenderCore/Private/RenderGraphResourcePool.h
--- a/Engine/Source/Runtime/RenderCore/Private/RenderGraphResourcePool.h
+++ b/Engine/Source/Runtime/RenderCore/Private/RenderGraphResourcePool.h
@@ -111,6 +111,7 @@ private:
 
 	IRHITransientResourceAllocator* Allocator = nullptr;
 
+	FCriticalSection CS;
 	TArray<FRDGTransientRenderTarget*> FreeList;
 	TArray<FRDGTransientRenderTarget*> PendingDeallocationList;
 	TArray<FRDGTransientRenderTarget*> DeallocatedList;
