Fixed build break.

diff --git a/Engine/Source/Runtime/RenderCore/Private/RenderGraphResourcePool.h b/Engine/Source/Runtime/RenderCore/Private/RenderGraphResourcePool.h
--- a/Engine/Source/Runtime/RenderCore/Private/RenderGraphResourcePool.h
+++ b/Engine/Source/Runtime/RenderCore/Private/RenderGraphResourcePool.h
@@ -85,7 +85,7 @@ private:
 	FRHITransientTexture* Texture;
 	FPooledRenderTargetDesc Desc;
 	ERDGTransientResourceLifetimeState LifetimeState;
-	mutable uint32 RefCount = 0;
+	mutable int32 RefCount = 0;
 
 	friend class FRDGTransientResourceAllocator;
 };
