Moved RDG builder destruction off the render thread.

diff --git a/Engine/Source/Runtime/RenderCore/Private/RenderGraphAllocator.cpp b/Engine/Source/Runtime/RenderCore/Private/RenderGraphAllocator.cpp
--- a/Engine/Source/Runtime/RenderCore/Private/RenderGraphAllocator.cpp
+++ b/Engine/Source/Runtime/RenderCore/Private/RenderGraphAllocator.cpp
@@ -11,11 +11,7 @@ FRDGAllocator& FRDGAllocator::Get()
 
 FRDGAllocator::~FRDGAllocator()
 {
-#if RDG_USE_MALLOC
-	check(Context.RawAllocs.IsEmpty() && ContextForTasks.RawAllocs.IsEmpty());
-#else
-	check(Context.MemStack.IsEmpty() && ContextForTasks.MemStack.IsEmpty());
-#endif
+	ReleaseAll();
 }
 
 void FRDGAllocator::FContext::ReleaseAll()
diff --git a/Engine/Source/Runtime/RenderCore/Private/RenderGraphBuilder.cpp b/Engine/Source/Runtime/RenderCore/Private/RenderGraphBuilder.cpp
--- a/Engine/Source/Runtime/RenderCore/Private/RenderGraphBuilder.cpp
+++ b/Engine/Source/Runtime/RenderCore/Private/RenderGraphBuilder.cpp
@@ -566,12 +566,30 @@ FRDGBuilder::~FRDGBuilder()
 {
 	SCOPED_NAMED_EVENT(FRDGBuilder_Clear, FColor::Emerald);
 
-	Passes.Clear();
-	Buffers.Clear();
-	UniformBuffers.Clear();
-	Blackboard.Clear();
-	ActivePooledTextures.Empty();
-	ActivePooledBuffers.Empty();
+	if (bParallelExecuteEnabled)
+	{
+		FFunctionGraphTask::CreateAndDispatchWhenReady(
+			[Passes = MoveTemp(Passes), Blackboard = MoveTemp(Blackboard), Buffers = MoveTemp(Buffers), UniformBuffers = MoveTemp(UniformBuffers), ActivePooledTextures = MoveTemp(ActivePooledTextures), ActivePooledBuffers = MoveTemp(ActivePooledBuffers), Allocator = MoveTemp(Allocator)](ENamedThreads::Type, const FGraphEventRef&) mutable
+		{
+			SCOPED_NAMED_EVENT(FRDGBuilder_Clear, FColor::Emerald);
+			Passes.Clear();
+			Buffers.Clear();
+			UniformBuffers.Clear();
+			Blackboard.Clear();
+			ActivePooledTextures.Empty();
+			ActivePooledBuffers.Empty();
+
+		}, TStatId(), nullptr, ENamedThreads::AnyThread);
+	}
+	else
+	{
+		Passes.Clear();
+		Buffers.Clear();
+		UniformBuffers.Clear();
+		Blackboard.Clear();
+		ActivePooledTextures.Empty();
+		ActivePooledBuffers.Empty();
+	}
 }
 
 ///////////////////////////////////////////////////////////////////////////////////////////////////
@@ -2097,6 +2115,7 @@ void FRDGBuilder::SetupPassInternal(FRDGPass* Pass, FRDGPassHandle PassHandle, E
 		}
 
 		check(!EnumHasAnyFlags(PassPipeline, ERHIPipeline::AsyncCompute));
+		check(ParallelSetupEvents.IsEmpty());
 
 		SubmitBufferUploads(nullptr);
 		CompilePassOps(Pass);
@@ -2104,8 +2123,6 @@ void FRDGBuilder::SetupPassInternal(FRDGPass* Pass, FRDGPassHandle PassHandle, E
 		CreateUniformBuffers(nullptr);
 		CollectPassBarriers(Pass, PassHandle);
 		CreatePassBarriers([] {});
-		FTaskGraphInterface::Get().WaitUntilTasksComplete(ParallelSetupEvents);
-		ParallelSetupEvents.Reset();
 		ExecutePass(Pass, RHICmdList);
 	}
 
diff --git a/Engine/Source/Runtime/RenderCore/Public/RenderGraphAllocator.h b/Engine/Source/Runtime/RenderCore/Public/RenderGraphAllocator.h
--- a/Engine/Source/Runtime/RenderCore/Public/RenderGraphAllocator.h
+++ b/Engine/Source/Runtime/RenderCore/Public/RenderGraphAllocator.h
@@ -11,6 +11,7 @@
 class RENDERCORE_API FRDGAllocator final
 {
 public:
+	static FRDGAllocator& Get();
 	~FRDGAllocator();
 
 	/** Allocates raw memory. */
@@ -68,8 +69,10 @@ public:
 	}
 
 private:
-	static FRDGAllocator& Get();
 	FRDGAllocator() = default;
+	FRDGAllocator(FRDGAllocator&&) = default;
+	FRDGAllocator& operator = (FRDGAllocator&&) = default;
+
 	void ReleaseAll();
 
 	struct FTrackedAlloc
@@ -112,6 +115,7 @@ private:
 	template <uint32>
 	friend class TRDGArrayAllocator;
 	friend class FRDGAllocatorScope;
+	friend class FRDGBuilder;
 };
 
 /** Base class for RDG builder which scopes the allocations and releases them in the destructor. */
diff --git a/Engine/Source/Runtime/RenderCore/Public/RenderGraphBuilder.h b/Engine/Source/Runtime/RenderCore/Public/RenderGraphBuilder.h
--- a/Engine/Source/Runtime/RenderCore/Public/RenderGraphBuilder.h
+++ b/Engine/Source/Runtime/RenderCore/Public/RenderGraphBuilder.h
@@ -184,18 +184,9 @@ public:
 	/** A hint to the builder to flush work to the RHI thread after the last queued pass on the execution timeline. */
 	void AddDispatchHint();
 
-
+	/** Launches a task that is synced prior to graph execution. If parallel execution is not enabled, the lambda is run immediately. */
 	template <typename TaskLambda>
-	void AddSetupTask(TaskLambda&& Task)
-	{
-		ParallelSetupEvents.Emplace(FFunctionGraphTask::CreateAndDispatchWhenReady(
-			[Task = MoveTemp(Task)](ENamedThreads::Type, const FGraphEventRef&)
-		{
-			FTaskTagScope Scope(ETaskTag::EParallelRenderingThread);
-			Task();
-
-		}, TStatId(), nullptr, ENamedThreads::AnyHiPriThreadHiPriTask));
-	}
+	void AddSetupTask(TaskLambda&& Task);
 
 	/** Tells the builder to delete unused RHI resources. The behavior of this method depends on whether RDG immediate mode is enabled:
 	 *   Deferred:  RHI resource flushes are performed prior to execution.
diff --git a/Engine/Source/Runtime/RenderCore/Public/RenderGraphBuilder.inl b/Engine/Source/Runtime/RenderCore/Public/RenderGraphBuilder.inl
--- a/Engine/Source/Runtime/RenderCore/Public/RenderGraphBuilder.inl
+++ b/Engine/Source/Runtime/RenderCore/Public/RenderGraphBuilder.inl
@@ -362,6 +362,25 @@ inline void FRDGBuilder::AddDispatchHint()
 	}
 }
 
+template <typename TaskLambda>
+void FRDGBuilder::AddSetupTask(TaskLambda&& Task)
+{
+	if (bParallelExecuteEnabled)
+	{
+		ParallelSetupEvents.Emplace(FFunctionGraphTask::CreateAndDispatchWhenReady(
+			[Task = MoveTemp(Task)](ENamedThreads::Type, const FGraphEventRef&)
+		{
+			FTaskTagScope Scope(ETaskTag::EParallelRenderingThread);
+			Task();
+
+		}, TStatId(), nullptr, ENamedThreads::AnyHiPriThreadHiPriTask));
+	}
+	else
+	{
+		Task();
+	}
+}
+
 inline const TRefCountPtr<IPooledRenderTarget>& FRDGBuilder::GetPooledTexture(FRDGTextureRef Texture) const
 {
 	IF_RDG_ENABLE_DEBUG(UserValidation.ValidateGetPooledTexture(Texture));
