Put RDG::ExecutePass into its own RDG_Execute CSV stat to make it easier to measure (instead of using RenderOther)

diff --git a/Engine/Source/Runtime/RenderCore/Private/RenderGraphBuilder.cpp b/Engine/Source/Runtime/RenderCore/Private/RenderGraphBuilder.cpp
--- a/Engine/Source/Runtime/RenderCore/Private/RenderGraphBuilder.cpp
+++ b/Engine/Source/Runtime/RenderCore/Private/RenderGraphBuilder.cpp
@@ -1730,7 +1730,7 @@ void FRDGBuilder::Execute()
 	{
 		SCOPED_NAMED_EVENT_TEXT("FRDGBuilder::ExecutePasses", FColor::Magenta);
 		SCOPE_CYCLE_COUNTER(STAT_RDG_ExecuteTime);
-		CSV_SCOPED_TIMING_STAT_EXCLUSIVE(RenderOther);
+		CSV_SCOPED_TIMING_STAT_EXCLUSIVE(RDG_Execute);
 
 		// Wait on all async compilation tasks before executing any passes.
 		if (!AsyncCompileEvents.IsEmpty())
