Fixed async RDG clear to delete allocator after the FRDGBuilder destructor executes.

diff --git a/Engine/Source/Runtime/RenderCore/Private/RenderGraphAllocator.cpp b/Engine/Source/Runtime/RenderCore/Private/RenderGraphAllocator.cpp
--- a/Engine/Source/Runtime/RenderCore/Private/RenderGraphAllocator.cpp
+++ b/Engine/Source/Runtime/RenderCore/Private/RenderGraphAllocator.cpp
@@ -31,7 +31,7 @@ void FRDGAllocator::FContext::ReleaseAll()
 	{
 		FMemory::Free(RawAlloc);
 	}
-	RawAllocs.Reset();
+	RawAllocs.Reset();builder.cpp
 #else
 	MemStack.Flush();
 #endif
@@ -43,4 +43,25 @@ void FRDGAllocator::ReleaseAll()
 	CSV_SCOPED_TIMING_STAT_EXCLUSIVE_CONDITIONAL(RDGAllocator_Clear, GRDGVerboseCSVStats != 0);
 	Context.ReleaseAll();
 	ContextForTasks.ReleaseAll();
-}
\ No newline at end of file
+}
+
+FRDGAllocatorScope::~FRDGAllocatorScope()
+{
+	if (AsyncDeleteFunction)
+	{
+		//FFunctionGraphTask::CreateAndDispatchWhenReady(
+		//	[Allocator = MoveTemp(Allocator), AsyncDeleteFunction = MoveTemp(AsyncDeleteFunction)] () mutable
+		//{
+
+		FRDGAllocator LocalAllocator = MoveTemp(Allocator);
+			SCOPED_NAMED_EVENT(FRDGAllocatorScope_AsyncDelete, FColor::Emerald);
+			AsyncDeleteFunction();
+			AsyncDeleteFunction = {};
+
+		//}, TStatId(), nullptr, ENamedThreads::AnyThread);
+	}
+	else
+	{
+		Allocator.ReleaseAll();
+	}
+}
diff --git a/Engine/Source/Runtime/RenderCore/Private/RenderGraphBuilder.cpp b/Engine/Source/Runtime/RenderCore/Private/RenderGraphBuilder.cpp
--- a/Engine/Source/Runtime/RenderCore/Private/RenderGraphBuilder.cpp
+++ b/Engine/Source/Runtime/RenderCore/Private/RenderGraphBuilder.cpp
@@ -564,31 +564,20 @@ FRDGBuilder::FRDGBuilder(FRHICommandListImmediate& InRHICmdList, FRDGEventName I
 
 FRDGBuilder::~FRDGBuilder()
 {
-	SCOPED_NAMED_EVENT(FRDGBuilder_Clear, FColor::Emerald);
-
 	if (bParallelExecuteEnabled)
 	{
-		FFunctionGraphTask::CreateAndDispatchWhenReady(
-			[Passes = MoveTemp(Passes), Blackboard = MoveTemp(Blackboard), Buffers = MoveTemp(Buffers), UniformBuffers = MoveTemp(UniformBuffers), ActivePooledTextures = MoveTemp(ActivePooledTextures), ActivePooledBuffers = MoveTemp(ActivePooledBuffers), Allocator = MoveTemp(Allocator)](ENamedThreads::Type, const FGraphEventRef&) mutable
-		{
-			SCOPED_NAMED_EVENT(FRDGBuilder_Clear, FColor::Emerald);
-			Passes.Clear();
-			Buffers.Clear();
-			UniformBuffers.Clear();
-			Blackboard.Clear();
-			ActivePooledTextures.Empty();
-			ActivePooledBuffers.Empty();
-
-		}, TStatId(), nullptr, ENamedThreads::AnyThread);
-	}
-	else
-	{
-		Passes.Clear();
-		Buffers.Clear();
-		UniformBuffers.Clear();
-		Blackboard.Clear();
-		ActivePooledTextures.Empty();
-		ActivePooledBuffers.Empty();
+		// Move expensive operations into the async deleter, which will be called in the base class destructor.
+		BeginAsyncDelete([
+			Passes					= MoveTemp(Passes),
+			Textures				= MoveTemp(Textures),
+			Buffers					= MoveTemp(Buffers),
+			Views					= MoveTemp(Views),
+			UniformBuffers			= MoveTemp(UniformBuffers),
+			Blackboard				= MoveTemp(Blackboard),
+			ActivePooledTextures	= MoveTemp(ActivePooledTextures),
+			ActivePooledBuffers		= MoveTemp(ActivePooledBuffers),
+			UploadedBuffers			= MoveTemp(UploadedBuffers)
+		] () mutable {});
 	}
 }
 
diff --git a/Engine/Source/Runtime/RenderCore/Public/RenderGraphAllocator.h b/Engine/Source/Runtime/RenderCore/Public/RenderGraphAllocator.h
--- a/Engine/Source/Runtime/RenderCore/Public/RenderGraphAllocator.h
+++ b/Engine/Source/Runtime/RenderCore/Public/RenderGraphAllocator.h
@@ -115,26 +115,28 @@ private:
 	template <uint32>
 	friend class TRDGArrayAllocator;
 	friend class FRDGAllocatorScope;
-	friend class FRDGBuilder;
 };
 
 /** Base class for RDG builder which scopes the allocations and releases them in the destructor. */
-class FRDGAllocatorScope
+class RENDERCORE_API FRDGAllocatorScope
 {
-protected:
-	FRDGAllocator& Allocator;
-
-private:
+public:
 	FRDGAllocatorScope()
 		: Allocator(FRDGAllocator::Get())
 	{}
 
-	~FRDGAllocatorScope()
+	~FRDGAllocatorScope();
+
+protected:
+	FRDGAllocator& Allocator;
+
+	void BeginAsyncDelete(TUniqueFunction<void()>&& InAsyncDeleteFunction)
 	{
-		Allocator.ReleaseAll();
+		AsyncDeleteFunction = MoveTemp(InAsyncDeleteFunction);
 	}
 
-	friend class FRDGBuilder;
+private:
+	TUniqueFunction<void()> AsyncDeleteFunction;
 };
 
 /** A container allocator that allocates from a global RDG allocator instance. */
