Fix for ASAN use after free in RDG allocator.

diff --git a/Engine/Source/Runtime/RenderCore/Private/RenderGraphAllocator.cpp b/Engine/Source/Runtime/RenderCore/Private/RenderGraphAllocator.cpp
--- a/Engine/Source/Runtime/RenderCore/Private/RenderGraphAllocator.cpp
+++ b/Engine/Source/Runtime/RenderCore/Private/RenderGraphAllocator.cpp
@@ -11,7 +11,8 @@ FRDGAllocator& FRDGAllocator::Get()
 
 FRDGAllocator::~FRDGAllocator()
 {
-	ReleaseAll();
+	Context.ReleaseAll();
+	ContextForTasks.ReleaseAll();
 }
 
 void FRDGAllocator::FContext::ReleaseAll()
@@ -31,7 +32,7 @@ void FRDGAllocator::FContext::ReleaseAll()
 	{
 		FMemory::Free(RawAlloc);
 	}
-	RawAllocs.Reset();builder.cpp
+	RawAllocs.Reset();
 #else
 	MemStack.Flush();
 #endif
@@ -49,16 +50,14 @@ FRDGAllocatorScope::~FRDGAllocatorScope()
 {
 	if (AsyncDeleteFunction)
 	{
-		//FFunctionGraphTask::CreateAndDispatchWhenReady(
-		//	[Allocator = MoveTemp(Allocator), AsyncDeleteFunction = MoveTemp(AsyncDeleteFunction)] () mutable
-		//{
-
-		FRDGAllocator LocalAllocator = MoveTemp(Allocator);
+		FFunctionGraphTask::CreateAndDispatchWhenReady(
+			[Allocator = MoveTemp(Allocator), AsyncDeleteFunction = MoveTemp(AsyncDeleteFunction)] () mutable
+		{
 			SCOPED_NAMED_EVENT(FRDGAllocatorScope_AsyncDelete, FColor::Emerald);
 			AsyncDeleteFunction();
 			AsyncDeleteFunction = {};
 
-		//}, TStatId(), nullptr, ENamedThreads::AnyThread);
+		}, TStatId(), nullptr, ENamedThreads::AnyThread);
 	}
 	else
 	{
