Moved RDG builder destruction off the render thread.

diff --git a/Engine/Source/Runtime/RenderCore/Public/RenderGraphAllocator.h b/Engine/Source/Runtime/RenderCore/Public/RenderGraphAllocator.h
--- a/Engine/Source/Runtime/RenderCore/Public/RenderGraphAllocator.h
+++ b/Engine/Source/Runtime/RenderCore/Public/RenderGraphAllocator.h
@@ -115,6 +115,7 @@ private:
 	template <uint32>
 	friend class TRDGArrayAllocator;
 	friend class FRDGAllocatorScope;
+	friend class FRDGBuilder;
 };
 
 /** Base class for RDG builder which scopes the allocations and releases them in the destructor. */
