Fixed async RDG clear to delete allocator after the FRDGBuilder destructor executes.

diff --git a/Engine/Source/Runtime/RenderCore/Private/RenderGraphAllocator.cpp b/Engine/Source/Runtime/RenderCore/Private/RenderGraphAllocator.cpp
--- a/Engine/Source/Runtime/RenderCore/Private/RenderGraphAllocator.cpp
+++ b/Engine/Source/Runtime/RenderCore/Private/RenderGraphAllocator.cpp
@@ -11,8 +11,7 @@ FRDGAllocator& FRDGAllocator::Get()
 
 FRDGAllocator::~FRDGAllocator()
 {
-	Context.ReleaseAll();
-	ContextForTasks.ReleaseAll();
+	ReleaseAll();
 }
 
 void FRDGAllocator::FContext::ReleaseAll()
@@ -32,7 +31,7 @@ void FRDGAllocator::FContext::ReleaseAll()
 	{
 		FMemory::Free(RawAlloc);
 	}
-	RawAllocs.Reset();
+	RawAllocs.Reset();builder.cpp
 #else
 	MemStack.Flush();
 #endif
@@ -50,14 +49,16 @@ FRDGAllocatorScope::~FRDGAllocatorScope()
 {
 	if (AsyncDeleteFunction)
 	{
-		FFunctionGraphTask::CreateAndDispatchWhenReady(
-			[Allocator = MoveTemp(Allocator), AsyncDeleteFunction = MoveTemp(AsyncDeleteFunction)] () mutable
-		{
+		//FFunctionGraphTask::CreateAndDispatchWhenReady(
+		//	[Allocator = MoveTemp(Allocator), AsyncDeleteFunction = MoveTemp(AsyncDeleteFunction)] () mutable
+		//{
+
+		FRDGAllocator LocalAllocator = MoveTemp(Allocator);
 			SCOPED_NAMED_EVENT(FRDGAllocatorScope_AsyncDelete, FColor::Emerald);
 			AsyncDeleteFunction();
 			AsyncDeleteFunction = {};
 
-		}, TStatId(), nullptr, ENamedThreads::AnyThread);
+		//}, TStatId(), nullptr, ENamedThreads::AnyThread);
 	}
 	else
 	{
diff --git a/Engine/Source/Runtime/RenderCore/Public/RenderGraphAllocator.h b/Engine/Source/Runtime/RenderCore/Public/RenderGraphAllocator.h
--- a/Engine/Source/Runtime/RenderCore/Public/RenderGraphAllocator.h
+++ b/Engine/Source/Runtime/RenderCore/Public/RenderGraphAllocator.h
@@ -115,7 +115,6 @@ private:
 	template <uint32>
 	friend class TRDGArrayAllocator;
 	friend class FRDGAllocatorScope;
-	friend class FRDGBuilder;
 };
 
 /** Base class for RDG builder which scopes the allocations and releases them in the destructor. */
