Fixed memory leak in RDG.

diff --git a/Engine/Source/Runtime/RenderCore/Public/RenderGraphDefinitions.h b/Engine/Source/Runtime/RenderCore/Public/RenderGraphDefinitions.h
--- a/Engine/Source/Runtime/RenderCore/Public/RenderGraphDefinitions.h
+++ b/Engine/Source/Runtime/RenderCore/Public/RenderGraphDefinitions.h
@@ -419,6 +419,17 @@ public:
 	using ObjectType = typename HandleType::ObjectType;
 	using IndexType = typename HandleType::IndexType;
 
+	TRDGHandleRegistry() = default;
+	TRDGHandleRegistry(const TRDGHandleRegistry&) = delete;
+	TRDGHandleRegistry(TRDGHandleRegistry&&) = default;
+	TRDGHandleRegistry& operator=(TRDGHandleRegistry&&) = default;
+	TRDGHandleRegistry& operator=(const TRDGHandleRegistry&) = delete;
+
+	~TRDGHandleRegistry()
+	{
+		Clear();
+	}
+
 	void Insert(ObjectType* Object)
 	{
 		Array.Emplace(Object);
