Removed unacceptable words from RenderCore.

diff --git a/Engine/Source/Runtime/Core/Private/Stats/Stats2.cpp b/Engine/Source/Runtime/Core/Private/Stats/Stats2.cpp
--- a/Engine/Source/Runtime/Core/Private/Stats/Stats2.cpp
+++ b/Engine/Source/Runtime/Core/Private/Stats/Stats2.cpp
@@ -206,9 +206,9 @@ void FStats::AdvanceFrame( bool bDiscardCallstack, const FOnAdvanceRenderingThre
 	}
 	if( MasterDisableChangeTagStartFrame == -1 )
 	{
-		MasterDisableChangeTagStartFrame = FThreadStats::MasterDisableChangeTag();
+		MasterDisableChangeTagStartFrame = FThreadStats::PrimaryDisableChangeTag();
 	}
-	if( !FThreadStats::IsCollectingData() || MasterDisableChangeTagStartFrame != FThreadStats::MasterDisableChangeTag() )
+	if( !FThreadStats::IsCollectingData() || MasterDisableChangeTagStartFrame != FThreadStats::PrimaryDisableChangeTag() )
 	{
 		Frame = -Frame; // mark this as a bad frame
 	}
@@ -230,7 +230,7 @@ void FStats::AdvanceFrame( bool bDiscardCallstack, const FOnAdvanceRenderingThre
 
 	FThreadStats::ExplicitFlush( bDiscardCallstack );
 	FThreadStats::WaitForStats();
-	MasterDisableChangeTagStartFrame = FThreadStats::MasterDisableChangeTag();
+	MasterDisableChangeTagStartFrame = FThreadStats::PrimaryDisableChangeTag();
 #endif
 }
 
diff --git a/Engine/Source/Runtime/Core/Public/Stats/Stats2.h b/Engine/Source/Runtime/Core/Public/Stats/Stats2.h
--- a/Engine/Source/Runtime/Core/Public/Stats/Stats2.h
+++ b/Engine/Source/Runtime/Core/Public/Stats/Stats2.h
@@ -1515,7 +1515,7 @@ public:
 	}
 
 	/** Everytime master enable changes, this number increases. This is used to determine full frames. **/
-	static FORCEINLINE_STATS int32 MasterDisableChangeTag()
+	static FORCEINLINE_STATS int32 PrimaryDisableChangeTag()
 	{
 		if (MasterDisableChangeTagLock.GetValue())
 		{
diff --git a/Engine/Source/Runtime/RenderCore/Private/DumpGPU.cpp b/Engine/Source/Runtime/RenderCore/Private/DumpGPU.cpp
--- a/Engine/Source/Runtime/RenderCore/Private/DumpGPU.cpp
+++ b/Engine/Source/Runtime/RenderCore/Private/DumpGPU.cpp
@@ -99,7 +99,7 @@ static TAutoConsoleVariable<int32> GDumpRenderingConsoleVariablesCVar(
 
 static TAutoConsoleVariable<int32> GDumpTestEnableDiskWrite(
 	TEXT("r.DumpGPU.Test.EnableDiskWrite"), 1,
-	TEXT("Master switch whether any files should be written to disk, used for r.DumpGPU automation tests to not fill up workers' hard drive."),
+	TEXT("Main switch whether any files should be written to disk, used for r.DumpGPU automation tests to not fill up workers' hard drive."),
 	ECVF_Default);
 
 static TAutoConsoleVariable<int32> GDumpTestPrettifyResourceFileNames(
diff --git a/Engine/Source/Runtime/RenderCore/Private/RenderUtils.cpp b/Engine/Source/Runtime/RenderCore/Private/RenderUtils.cpp
--- a/Engine/Source/Runtime/RenderCore/Private/RenderUtils.cpp
+++ b/Engine/Source/Runtime/RenderCore/Private/RenderUtils.cpp
@@ -18,7 +18,7 @@
 #include "RHIShaderFormatDefinitions.inl"
 #endif
 
-// This is a per-project master switch for Nanite, that influences the shader permutations compiled. Changing it will cause shaders to be recompiled.
+// This is a per-project main switch for Nanite, that influences the shader permutations compiled. Changing it will cause shaders to be recompiled.
 int32 GNaniteProjectEnabled = 1;
 FAutoConsoleVariableRef CVarAllowNanite(
 	TEXT("r.Nanite.ProjectEnabled"),
diff --git a/Engine/Source/Runtime/RenderCore/Private/RenderingThread.cpp b/Engine/Source/Runtime/RenderCore/Private/RenderingThread.cpp
--- a/Engine/Source/Runtime/RenderCore/Private/RenderingThread.cpp
+++ b/Engine/Source/Runtime/RenderCore/Private/RenderingThread.cpp
@@ -398,11 +398,11 @@ void RenderingThreadMain( FEvent* TaskGraphBoundSyncEvent )
 /**
  * Advances stats for the rendering thread.
  */
-static void AdvanceRenderingThreadStats(int64 StatsFrame, int32 MasterDisableChangeTagStartFrame)
+static void AdvanceRenderingThreadStats(int64 StatsFrame, int32 DisableChangeTagStartFrame)
 {
 #if STATS
 	int64 Frame = StatsFrame;
-	if (!FThreadStats::IsCollectingData() ||  MasterDisableChangeTagStartFrame != FThreadStats::MasterDisableChangeTag())
+	if (!FThreadStats::IsCollectingData() || DisableChangeTagStartFrame != FThreadStats::PrimaryDisableChangeTag())
 	{
 		Frame = -StatsFrame; // mark this as a bad frame
 	}
@@ -417,12 +417,12 @@ static void AdvanceRenderingThreadStats(int64 StatsFrame, int32 MasterDisableCha
 /**
  * Advances stats for the rendering thread. Called from the game thread.
  */
-void AdvanceRenderingThreadStatsGT( bool bDiscardCallstack, int64 StatsFrame, int32 MasterDisableChangeTagStartFrame )
+void AdvanceRenderingThreadStatsGT( bool bDiscardCallstack, int64 StatsFrame, int32 DisableChangeTagStartFrame )
 {
 	ENQUEUE_RENDER_COMMAND(RenderingThreadTickCommand)(
-		[StatsFrame, MasterDisableChangeTagStartFrame](FRHICommandList& RHICmdList)
+		[StatsFrame, DisableChangeTagStartFrame](FRHICommandList& RHICmdList)
 		{
-			AdvanceRenderingThreadStats(StatsFrame, MasterDisableChangeTagStartFrame );
+			AdvanceRenderingThreadStats(StatsFrame, DisableChangeTagStartFrame);
 		}
 	);
 	if( bDiscardCallstack )
diff --git a/Engine/Source/Runtime/RenderCore/Private/ShaderCodeLibrary.cpp b/Engine/Source/Runtime/RenderCore/Private/ShaderCodeLibrary.cpp
--- a/Engine/Source/Runtime/RenderCore/Private/ShaderCodeLibrary.cpp
+++ b/Engine/Source/Runtime/RenderCore/Private/ShaderCodeLibrary.cpp
@@ -1845,7 +1845,7 @@ struct FEditorShaderStableInfo
 				}
 			}
 
-			// Only the master cooker needs to write to the output directory, child cookers only write to the Saved directory
+			// Only the primary cooker needs to write to the output directory, child cookers only write to the Saved directory
 			FString OutputFilePath = GetStableInfoArchiveFilename(OutputDir, LibraryName, FormatName);
 
 			// Copy to output location - support for iterative native library cooking
diff --git a/Engine/Source/Runtime/RenderCore/Public/RenderingThread.h b/Engine/Source/Runtime/RenderCore/Public/RenderingThread.h
--- a/Engine/Source/Runtime/RenderCore/Public/RenderingThread.h
+++ b/Engine/Source/Runtime/RenderCore/Public/RenderingThread.h
@@ -63,7 +63,7 @@ extern RENDERCORE_API bool IsRenderingThreadHealthy();
 /**
  * Advances stats for the rendering thread. Called from the game thread.
  */
-extern RENDERCORE_API void AdvanceRenderingThreadStatsGT( bool bDiscardCallstack, int64 StatsFrame, int32 MasterDisableChangeTagStartFrame );
+extern RENDERCORE_API void AdvanceRenderingThreadStatsGT( bool bDiscardCallstack, int64 StatsFrame, int32 DisableChangeTagStartFrame );
 
 /**
  * Waits for the rendering thread to finish executing all pending rendering commands.  Should only be used from the game thread.
