[Backout] - CL20968686
Original CL Desc
-----------------------------------------------------------------
Move RDG buffer uploads off the render thread.

Maybe be related to gpu crash in EMT (UE-157850)

diff --git a/Engine/Source/Runtime/RenderCore/Private/RenderGraphBuilder.cpp b/Engine/Source/Runtime/RenderCore/Private/RenderGraphBuilder.cpp
--- a/Engine/Source/Runtime/RenderCore/Private/RenderGraphBuilder.cpp
+++ b/Engine/Source/Runtime/RenderCore/Private/RenderGraphBuilder.cpp
@@ -1539,7 +1539,6 @@ void FRDGBuilder::Execute()
 	const FRDGPassHandle EpiloguePassHandle = GetEpiloguePassHandle();
 
 	FGraphEventArray AsyncCompileEvents;
-	FGraphEventArray* AsyncCompileEventsIfEnabled = bParallelExecuteEnabled ? &AsyncCompileEvents : nullptr;
 
 	IF_RDG_ENABLE_DEBUG(UserValidation.ValidateExecuteBegin());
 	IF_RDG_ENABLE_DEBUG(GRDGAllowRHIAccess = true);
@@ -1568,7 +1567,7 @@ void FRDGBuilder::Execute()
 			}, TStatId(), nullptr, ENamedThreads::AnyHiPriThreadHiPriTask));
 		}
 
-		SubmitBufferUploads(AsyncCompileEventsIfEnabled);
+		SetupBufferUploads();
 
 		{
 			SCOPE_CYCLE_COUNTER(STAT_RDG_CollectResourcesTime);
@@ -1652,7 +1651,7 @@ void FRDGBuilder::Execute()
 			}
 		}
 
-		CreateUniformBuffers(AsyncCompileEventsIfEnabled);
+		CreateUniformBuffers(bParallelExecuteEnabled ? &AsyncCompileEvents : nullptr);
 
 		{
 			SCOPED_NAMED_EVENT_TEXT("FRDGBuilder::CollectBarriers", FColor::Magenta);
@@ -1687,18 +1686,14 @@ void FRDGBuilder::Execute()
 		});
 	}
 
-	const ENamedThreads::Type RenderThread = ENamedThreads::GetRenderThread_Local();
+	CreatePassBarriers(bParallelExecuteEnabled ? &AsyncCompileEvents : nullptr);
 
-	CreatePassBarriers([this, RenderThread]
+	if (!ParallelSetupEvents.IsEmpty())
 	{
-		if (!ParallelSetupEvents.IsEmpty())
-		{
-			FTaskGraphInterface::Get().WaitUntilTasksComplete(ParallelSetupEvents, RenderThread);
-		}
-		
-		// Process RHI thread flush before helping with barrier compilation on the render thread.
-		EndFlushResourcesRHI();
-	});
+		FTaskGraphInterface::Get().WaitUntilTasksComplete(ParallelSetupEvents);
+	}
+
+	EndFlushResourcesRHI();
 
 	FGraphEventRef DispatchParallelExecuteEvent;
 
@@ -1712,9 +1707,13 @@ void FRDGBuilder::Execute()
 		}, TStatId(), &AsyncCompileEvents, ENamedThreads::AnyHiPriThreadHiPriTask);
 	}
 
+	SubmitBufferUploads();
+
 	IF_RDG_ENABLE_DEBUG(GRDGAllowRHIAccess = bParallelExecuteEnabled);
 	IF_RDG_ENABLE_TRACE(Trace.OutputGraphBegin());
 
+	const ENamedThreads::Type RenderThread = ENamedThreads::GetRenderThread_Local();
+
 	if (!IsImmediateMode())
 	{
 		SCOPED_NAMED_EVENT_TEXT("FRDGBuilder::ExecutePasses", FColor::Magenta);
@@ -2106,12 +2105,15 @@ void FRDGBuilder::SetupPassInternal(FRDGPass* Pass, FRDGPassHandle PassHandle, E
 		check(!EnumHasAnyFlags(PassPipeline, ERHIPipeline::AsyncCompute));
 		check(ParallelSetupEvents.IsEmpty());
 
-		SubmitBufferUploads(nullptr);
+		SetupBufferUploads();
+		SubmitBufferUploads();
 		CompilePassOps(Pass);
 		BeginResourcesRHI(Pass, PassHandle);
 		CreateUniformBuffers(nullptr);
 		CollectPassBarriers(Pass, PassHandle);
-		CreatePassBarriers([] {});
+		CreatePassBarriers(nullptr);
+		FTaskGraphInterface::Get().WaitUntilTasksComplete(ParallelSetupEvents);
+		ParallelSetupEvents.Reset();
 		ExecutePass(Pass, RHICmdList);
 	}
 
@@ -2124,82 +2126,64 @@ void FRDGBuilder::SetupPassInternal(FRDGPass* Pass, FRDGPassHandle PassHandle, E
 
 ///////////////////////////////////////////////////////////////////////////////////////////////////
 
-void FRDGBuilder::SubmitBufferUploads(FGraphEventArray* AsyncCompileEventsIfEnabled)
+void FRDGBuilder::SetupBufferUploads()
 {
-	{
-		SCOPED_NAMED_EVENT_TEXT("FRDGBuilder::PrepareBufferUploads", FColor::Magenta);
+	SCOPED_NAMED_EVENT_TEXT("FRDGBuilder::PrepareBufferUploads", FColor::Magenta);
 
-		for (FUploadedBuffer& UploadedBuffer : UploadedBuffers)
+	for (FUploadedBuffer& UploadedBuffer : UploadedBuffers)
+	{
+		if (UploadedBuffer.bUseDataCallbacks)
 		{
-			if (UploadedBuffer.bUseDataCallbacks)
-			{
-				UploadedBuffer.Data = UploadedBuffer.DataCallback();
-				UploadedBuffer.DataSize = UploadedBuffer.DataSizeCallback();
-			}
+			UploadedBuffer.Data = UploadedBuffer.DataCallback();
+			UploadedBuffer.DataSize = UploadedBuffer.DataSizeCallback();
+		}
 
-			if (UploadedBuffer.Data && UploadedBuffer.DataSize)
-			{
-				ConvertToExternalBuffer(UploadedBuffer.Buffer);
-				check(UploadedBuffer.DataSize <= UploadedBuffer.Buffer->Desc.GetSize());
-			}
+		if (UploadedBuffer.Data && UploadedBuffer.DataSize)
+		{
+			ConvertToExternalBuffer(UploadedBuffer.Buffer);
+			check(UploadedBuffer.DataSize <= UploadedBuffer.Buffer->Desc.GetSize());
 		}
 	}
+}
 
-	const auto SubmitUploadsLambda = [this](FRHICommandList& RHICmdListUpload)
-	{
-		SCOPED_NAMED_EVENT_TEXT("FRDGBuilder::SubmitBufferUploads", FColor::Magenta);
+void FRDGBuilder::SubmitBufferUploads()
+{
+	SCOPED_NAMED_EVENT_TEXT("FRDGBuilder::SubmitBufferUploads", FColor::Magenta);
 
-		for (const FUploadedBuffer& UploadedBuffer : UploadedBuffers)
+	for (const FUploadedBuffer& UploadedBuffer : UploadedBuffers)
+	{
+		if (UploadedBuffer.Data && UploadedBuffer.DataSize)
 		{
-			if (UploadedBuffer.Data && UploadedBuffer.DataSize)
-			{
 #if PLATFORM_NEEDS_GPU_UAV_RESOURCE_INIT_WORKAROUND
-				if (UploadedBuffer.Buffer->bUAVAccessed)
+			if (UploadedBuffer.Buffer->bUAVAccessed)
+			{
+				FRHIResourceCreateInfo CreateInfo(UploadedBuffer.Buffer->Name);
+				FBufferRHIRef TempBuffer = RHICreateVertexBuffer(UploadedBuffer.DataSize, BUF_Static | BUF_ShaderResource, CreateInfo);
+				void* DestPtr = RHICmdList.LockBuffer(TempBuffer, 0, UploadedBuffer.DataSize, RLM_WriteOnly);
+				FMemory::Memcpy(DestPtr, UploadedBuffer.Data, UploadedBuffer.DataSize);
+				RHICmdList.UnlockBuffer(TempBuffer);
+				RHICmdList.Transition(
 				{
-					FRHIResourceCreateInfo CreateInfo(UploadedBuffer.Buffer->Name);
-					FBufferRHIRef TempBuffer = RHICmdListUpload.CreateBuffer(UploadedBuffer.DataSize, BUF_VertexBuffer | BUF_Static | BUF_ShaderResource, 0, ERHIAccess::CopySrc | ERHIAccess::SRVMask, CreateInfo);
-					void* DestPtr = RHICmdListUpload.LockBuffer(TempBuffer, 0, UploadedBuffer.DataSize, RLM_WriteOnly);
-					FMemory::Memcpy(DestPtr, UploadedBuffer.Data, UploadedBuffer.DataSize);
-					RHICmdListUpload.UnlockBuffer(TempBuffer);
-					RHICmdListUpload.Transition(FRHITransitionInfo(UploadedBuffer.Buffer->GetRHI(), ERHIAccess::Unknown, ERHIAccess::CopyDest));
-					RHICmdListUpload.CopyBufferRegion(UploadedBuffer.Buffer->GetRHI(), 0, TempBuffer, 0, UploadedBuffer.DataSize);
-				}
-				else
+					FRHITransitionInfo(TempBuffer, ERHIAccess::Unknown, ERHIAccess::CopySrc | ERHIAccess::SRVMask),
+					FRHITransitionInfo(UploadedBuffer.Buffer->GetRHI(), ERHIAccess::Unknown, ERHIAccess::CopyDest)
+				});
+				RHICmdList.CopyBufferRegion(UploadedBuffer.Buffer->GetRHI(), 0, TempBuffer, 0, UploadedBuffer.DataSize);
+			}
+			else
 #endif
-				{
-					void* DestPtr = RHICmdListUpload.LockBuffer(UploadedBuffer.Buffer->GetRHIUnchecked(), 0, UploadedBuffer.DataSize, RLM_WriteOnly);
-					FMemory::Memcpy(DestPtr, UploadedBuffer.Data, UploadedBuffer.DataSize);
-					RHICmdListUpload.UnlockBuffer(UploadedBuffer.Buffer->GetRHIUnchecked());
-				}
+			{
+				void* DestPtr = RHICmdList.LockBuffer(UploadedBuffer.Buffer->GetRHIUnchecked(), 0, UploadedBuffer.DataSize, RLM_WriteOnly);
+				FMemory::Memcpy(DestPtr, UploadedBuffer.Data, UploadedBuffer.DataSize);
+				RHICmdList.UnlockBuffer(UploadedBuffer.Buffer->GetRHIUnchecked());
+			}
 
-				if (UploadedBuffer.bUseFreeCallbacks)
-				{
-					UploadedBuffer.DataFreeCallback(UploadedBuffer.Data);
-				}
+			if (UploadedBuffer.bUseFreeCallbacks)
+			{
+				UploadedBuffer.DataFreeCallback(UploadedBuffer.Data);
 			}
 		}
-	};
-
-	if (AsyncCompileEventsIfEnabled)
-	{
-		FRHICommandList* RHICmdListUpload = new FRHICommandList(FRHIGPUMask::All());
-
-		FGraphEventRef Event = FFunctionGraphTask::CreateAndDispatchWhenReady(
-			[SubmitUploadsLambda, RHICmdListUpload](ENamedThreads::Type, const FGraphEventRef&)
-		{
-			FTaskTagScope Scope(ETaskTag::EParallelRenderingThread);
-			SubmitUploadsLambda(*RHICmdListUpload);
-
-		}, TStatId(), nullptr, ENamedThreads::AnyHiPriThreadHiPriTask);
-
-		RHICmdList.QueueRenderThreadCommandListSubmit(Event, RHICmdListUpload);
-
-		AsyncCompileEventsIfEnabled->Emplace(Event);
-	}
-	else
-	{
-		SubmitUploadsLambda(RHICmdList);
 	}
+	UploadedBuffers.Reset();
 }
 
 ///////////////////////////////////////////////////////////////////////////////////////////////////
@@ -2678,31 +2662,43 @@ void FRDGBuilder::CollectPassBarriers(FRDGPass* Pass, FRDGPassHandle PassHandle)
 	}
 }
 
-void FRDGBuilder::CreatePassBarriers(TFunctionRef<void()> PreWork)
+void FRDGBuilder::CreatePassBarriers(FGraphEventArray* AsyncCompileEvents)
 {
-	const int32 NumBarriersPerTask = 32;
+	const int32 NumBarriersPerTask = 128;
 
-	if (bParallelExecuteEnabled && TransitionCreateQueue.Num() > NumBarriersPerTask)
+	if (AsyncCompileEvents && TransitionCreateQueue.Num() > NumBarriersPerTask)
 	{
-		ParallelForWithPreWork(TEXT("FRDGBuilder::CreatePassBarriers"), TransitionCreateQueue.Num(), NumBarriersPerTask, [this](int32 Index)
+		AsyncCompileEvents->Emplace(FFunctionGraphTask::CreateAndDispatchWhenReady(
+			[this, NumBarriersPerTask](ENamedThreads::Type, const FGraphEventRef&)
 		{
-			TransitionCreateQueue[Index]->CreateTransition();
+#if 1
+			ParallelFor(TEXT("FRDGBuilder::CreatePassBarriers"), TransitionCreateQueue.Num(), NumBarriersPerTask, [this](int32 Index)
+			{
+				TransitionCreateQueue[Index]->CreateTransition();
+			});
+#else
+			SCOPED_NAMED_EVENT_TEXT("FRDGBuilder::CreatePassBarriers", FColor::Magenta);
+			for (FRDGBarrierBatchBegin* BeginBatch : TransitionCreateQueue)
+			{
+				BeginBatch->CreateTransition();
+			}
+#endif
 
-		}, PreWork);
+			TransitionCreateQueue.Reset();
+
+		}, TStatId(), nullptr, ENamedThreads::AnyHiPriThreadHiPriTask));
 	}
 	else
 	{
-		PreWork();
-
 		SCOPED_NAMED_EVENT_TEXT("FRDGBuilder::CreatePassBarriers", FColor::Magenta);
 
 		for (FRDGBarrierBatchBegin* BeginBatch : TransitionCreateQueue)
 		{
 			BeginBatch->CreateTransition();
 		}
-	}
 
-	TransitionCreateQueue.Reset();
+		TransitionCreateQueue.Reset();
+	}
 }
 
 ///////////////////////////////////////////////////////////////////////////////////////////////////
diff --git a/Engine/Source/Runtime/RenderCore/Public/RenderGraphBuilder.h b/Engine/Source/Runtime/RenderCore/Public/RenderGraphBuilder.h
--- a/Engine/Source/Runtime/RenderCore/Public/RenderGraphBuilder.h
+++ b/Engine/Source/Runtime/RenderCore/Public/RenderGraphBuilder.h
@@ -747,7 +747,8 @@ private:
 	void SetupParallelExecute();
 	void DispatchParallelExecute(IRHICommandContext* RHICmdContext);
 
-	void SubmitBufferUploads(FGraphEventArray* AsyncCompileEvents);
+	void SetupBufferUploads();
+	void SubmitBufferUploads();
 	void BeginFlushResourcesRHI();
 	void EndFlushResourcesRHI();
 
@@ -764,7 +765,7 @@ private:
 
 	void CollectPassBarriers(FRDGPass* Pass, FRDGPassHandle PassHandle);
 
-	void CreatePassBarriers(TFunctionRef<void()> PreWork);
+	void CreatePassBarriers(FGraphEventArray* AsyncCompileEvents);
 
 	void CreateUniformBuffers(FGraphEventArray* AsyncCompileEvents);
 
