rename r.HDR.DisplayNitsLevel to r.HDR.Display.MaxLuminance

diff --git a/Engine/Source/Runtime/RenderCore/Private/RenderCore.cpp b/Engine/Source/Runtime/RenderCore/Private/RenderCore.cpp
--- a/Engine/Source/Runtime/RenderCore/Private/RenderCore.cpp
+++ b/Engine/Source/Runtime/RenderCore/Private/RenderCore.cpp
@@ -345,8 +345,8 @@ TAutoConsoleVariable<int32> CVarDisplayOutputDevice(
 	TEXT("9: Linear final color with tone curve\n"),
 	ECVF_Scalability | ECVF_RenderThreadSafe);
 
-static TAutoConsoleVariable<int32> CVarHDRNits(
-	TEXT("r.HDR.DisplayNitsLevel"),
+static TAutoConsoleVariable<int32> CVarHDRDisplayMaxLuminance(
+	TEXT("r.HDR.Display.MaxLuminance"),
 	0,
 	TEXT("The configured display output nit level, assuming HDR output is enabled."),
 	ECVF_RenderThreadSafe
@@ -406,7 +406,7 @@ inline FHDRMetaData HDRGetDefaultMetaData()
 	HDRMetaData.DisplayOutputFormat = HDRGetDefaultDisplayOutputFormat();
 	HDRMetaData.DisplayColorGamut = HDRGetDefaultDisplayColorGamut();
 	HDRMetaData.bHDRSupported = IsHDREnabled() && GRHISupportsHDROutput;
-	HDRMetaData.MaximumLuminanceInNits = CVarHDRNits.GetValueOnAnyThread();
+	HDRMetaData.MaximumLuminanceInNits = CVarHDRDisplayMaxLuminance.GetValueOnAnyThread();
 	return HDRMetaData;
 }
 
@@ -430,7 +430,7 @@ RENDERCORE_API void HDRAddCustomMetaData(void* OSWindow, EDisplayOutputFormat Di
 	HDRMetaData.DisplayOutputFormat = DisplayOutputFormat;
 	HDRMetaData.DisplayColorGamut = DisplayColorGamut;
 	HDRMetaData.bHDRSupported = bHDREnabled;
-	HDRMetaData.MaximumLuminanceInNits = CVarHDRNits.GetValueOnAnyThread();
+	HDRMetaData.MaximumLuminanceInNits = CVarHDRDisplayMaxLuminance.GetValueOnAnyThread();
 
 	FScopeLock Lock(&GWindowsWithDefaultParamsCS);
 	GWindowsWithDefaultParams.Add(OSWindow, HDRMetaData);
@@ -515,7 +515,7 @@ RENDERCORE_API void HDRGetMetaData(EDisplayOutputFormat& OutDisplayOutputFormat,
 
 	if (OutbHDRSupported)
 	{
-		FPlatformMisc::ChooseHDRDeviceAndColorGamut(GRHIVendorId, CVarHDRNits.GetValueOnAnyThread(), OutDisplayOutputFormat, OutDisplayColorGamut);
+		FPlatformMisc::ChooseHDRDeviceAndColorGamut(GRHIVendorId, CVarHDRDisplayMaxLuminance.GetValueOnAnyThread(), OutDisplayOutputFormat, OutDisplayColorGamut);
 	}
 
 }
@@ -538,16 +538,16 @@ RENDERCORE_API void HDRConfigureCVars(bool bIsHDREnabled, uint32 DisplayNits, bo
 		FPlatformMisc::ChooseHDRDeviceAndColorGamut(GRHIVendorId, DisplayNits, OutputDevice, ColorGamut);
 	}
 
-	//CVarHDRNits is ECVF_SetByCode as it's only a mean of communicating the information from UGameUserSettings to the rest of the engine
+	//CVarHDRDisplayMaxLuminance is ECVF_SetByCode as it's only a mean of communicating the information from UGameUserSettings to the rest of the engine
 	if (bIsHDREnabled)
 	{
 		CVarHDROutputEnabled->Set(1, bFromGameSettings ? ECVF_SetByGameSetting : ECVF_SetByCode);
-		CVarHDRNits->Set((int32)DisplayNits, ECVF_SetByCode);
+		CVarHDRDisplayMaxLuminance->Set((int32)DisplayNits, ECVF_SetByCode);
 	}
 	else
 	{
 		CVarHDROutputEnabled->Set(0, bFromGameSettings ? ECVF_SetByGameSetting : ECVF_SetByCode);
-		CVarHDRNits->Set(0, ECVF_SetByCode);
+		CVarHDRDisplayMaxLuminance->Set(0, ECVF_SetByCode);
 	}
 
 	CVarDisplayOutputDevice->Set((int32)OutputDevice, ECVF_SetByDeviceProfile);
