Allow FScopedGPUStatEvent and SCOPED_GPU_STAT to be used with any command list type.

* Stat push/pop will still only execute when running on immediate command list (no behavior change), but code that uses non-RDG mechanism for GPU stats will now always compile
* Remove rendering thread check from FScopedGPUStatEvent::Begin/End (it is implied by immediate command list)

diff --git a/Engine/Source/Runtime/RenderCore/Private/ProfilingDebugging/RealtimeGPUProfiler.cpp b/Engine/Source/Runtime/RenderCore/Private/ProfilingDebugging/RealtimeGPUProfiler.cpp
--- a/Engine/Source/Runtime/RenderCore/Private/ProfilingDebugging/RealtimeGPUProfiler.cpp
+++ b/Engine/Source/Runtime/RenderCore/Private/ProfilingDebugging/RealtimeGPUProfiler.cpp
@@ -1175,9 +1175,8 @@ void FRealtimeGPUProfiler::PopStat(FRHICommandListImmediate& RHICmdList, FRHIDra
 /*-----------------------------------------------------------------------------
 FScopedGPUStatEvent
 -----------------------------------------------------------------------------*/
-void FScopedGPUStatEvent::Begin(FRHICommandList& InRHICmdList, const FName& Name, const FName& StatName, const TCHAR* Description, FRHIDrawCallsStatPtr InNumDrawCallsPtr)
+void FScopedGPUStatEvent::Begin(FRHICommandListBase& InRHICmdList, const FName& Name, const FName& StatName, const TCHAR* Description, FRHIDrawCallsStatPtr InNumDrawCallsPtr)
 {
-	check(IsInRenderingThread());
 	if (!AreGPUStatsEnabled())
 	{
 		return;
@@ -1187,21 +1186,21 @@ void FScopedGPUStatEvent::Begin(FRHICommandList& InRHICmdList, const FName& Name
 	if (InRHICmdList.IsImmediate())
 	{
 		NumDrawCallsPtr = InNumDrawCallsPtr;
-		RHICmdList = &static_cast<FRHICommandListImmediate&>(InRHICmdList);
-		FRealtimeGPUProfiler::Get()->PushStat(*RHICmdList, Name, StatName, Description, InNumDrawCallsPtr);
+		RHICmdList = &InRHICmdList;
+		FRealtimeGPUProfiler::Get()->PushStat(InRHICmdList.GetAsImmediate(), Name, StatName, Description, InNumDrawCallsPtr);
 	}
 }
 
 void FScopedGPUStatEvent::End()
 {
-	check(IsInRenderingThread());
 	if (!AreGPUStatsEnabled())
 	{
 		return;
 	}
 	if (RHICmdList != nullptr)
 	{
-		FRealtimeGPUProfiler::Get()->PopStat(*RHICmdList, NumDrawCallsPtr);
+		// Command list is initialized only if it is immediate during Begin() and GetAsImmediate() also internally checks this.
+		FRealtimeGPUProfiler::Get()->PopStat(RHICmdList->GetAsImmediate(), NumDrawCallsPtr);
 	}
 }
 #endif // HAS_GPU_STATS
diff --git a/Engine/Source/Runtime/RenderCore/Public/ProfilingDebugging/RealtimeGPUProfiler.h b/Engine/Source/Runtime/RenderCore/Public/ProfilingDebugging/RealtimeGPUProfiler.h
--- a/Engine/Source/Runtime/RenderCore/Public/ProfilingDebugging/RealtimeGPUProfiler.h
+++ b/Engine/Source/Runtime/RenderCore/Public/ProfilingDebugging/RealtimeGPUProfiler.h
@@ -433,10 +433,13 @@ private:
 class FScopedGPUStatEvent
 {
 	/** Cmdlist to push onto. */
-	FRHICommandListImmediate* RHICmdList;
+	FRHICommandListBase* RHICmdList;
 
 	FRHIDrawCallsStatPtr NumDrawCallsPtr;
 public:
+
+	UE_NONCOPYABLE(FScopedGPUStatEvent)
+
 	/** Default constructor, initializing all member variables. */
 	FORCEINLINE FScopedGPUStatEvent()
 		: RHICmdList(nullptr)
@@ -457,7 +460,7 @@ public:
 	/**
 	* Start/Stop functions for timer stats
 	*/
-	RENDERCORE_API void Begin(FRHICommandList& InRHICmdList, const FName& Name, const FName& StatName, const TCHAR* Description, FRHIDrawCallsStatPtr InNumDrawCallsPtr);
+	RENDERCORE_API void Begin(FRHICommandListBase& InRHICmdList, const FName& Name, const FName& StatName, const TCHAR* Description, FRHIDrawCallsStatPtr InNumDrawCallsPtr);
 	RENDERCORE_API void End();
 };
 #endif // HAS_GPU_STATS
