Initialise value to fix UBSan warning.

diff --git a/Engine/Source/Runtime/RenderCore/Public/VirtualTexturing.h b/Engine/Source/Runtime/RenderCore/Public/VirtualTexturing.h
--- a/Engine/Source/Runtime/RenderCore/Public/VirtualTexturing.h
+++ b/Engine/Source/Runtime/RenderCore/Public/VirtualTexturing.h
@@ -247,7 +247,7 @@
 	 */
 	struct IPooledRenderTarget* PooledRenderTarget = nullptr;
 	/** Location within the texture to write */
-	FIntVector pPageLocation;
+	FIntVector pPageLocation = FIntVector::ZeroValue;
 };
 
 /**
