Fixed -NoPCH -DisableUnity compile issue in ShaderCompilerCore.h

diff --git a/Engine/Source/Runtime/RenderCore/Public/ShaderCompilerCore.h b/Engine/Source/Runtime/RenderCore/Public/ShaderCompilerCore.h
--- a/Engine/Source/Runtime/RenderCore/Public/ShaderCompilerCore.h
+++ b/Engine/Source/Runtime/RenderCore/Public/ShaderCompilerCore.h
@@ -12,6 +12,7 @@
 #include "Misc/SecureHash.h"
 #include "Misc/Paths.h"
 #include "Misc/CoreStats.h"
+#include "ShaderCore.h"
 
 class Error;
 
