Fix race in BatchReleasedResources

diff --git a/Engine/Source/Runtime/RenderCore/Private/RenderResource.cpp b/Engine/Source/Runtime/RenderCore/Private/RenderResource.cpp
--- a/Engine/Source/Runtime/RenderCore/Private/RenderResource.cpp
+++ b/Engine/Source/Runtime/RenderCore/Private/RenderResource.cpp
@@ -347,50 +347,36 @@ struct FBatchedReleaseResources
 	{
 		NumPerBatch = 16
 	};
-	int32 NumBatch;
-	FRenderResource* Resources[NumPerBatch];
-	FBatchedReleaseResources()
-	{
-		Reset();
-	}
-	void Reset()
-	{
-		NumBatch = 0;
-	}
-	void Execute()
-	{
-		for (int32 Index = 0; Index < NumBatch; Index++)
-		{
-			Resources[Index]->ReleaseResource();
-		}
-		Reset();
-	}
+	TArray<FRenderResource*, TInlineAllocator<NumPerBatch>> Resources;
+
 	void Flush()
 	{
-		if (NumBatch)
+		if (Resources.Num())
 		{
-			const FBatchedReleaseResources BatchedReleaseResources = *this;
 			ENQUEUE_RENDER_COMMAND(BatchReleaseCommand)(
-				[BatchedReleaseResources](FRHICommandList& RHICmdList)
+			[BatchedReleaseResources = MoveTemp(Resources)](FRHICommandList& RHICmdList)
+			{
+				for (FRenderResource* Resource : BatchedReleaseResources)
 				{
-					((FBatchedReleaseResources&)BatchedReleaseResources).Execute();
-				});
-			Reset();
+					Resource->ReleaseResource();
+				}
+			});
 		}
 	}
+
 	void Add(FRenderResource* Resource)
 	{
-		if (NumBatch >= NumPerBatch)
+		if (Resources.Num() >= NumPerBatch)
 		{
 			Flush();
 		}
-		check(NumBatch < NumPerBatch);
-		Resources[NumBatch] = Resource;
-		NumBatch++;
+		check(Resources.Num() < NumPerBatch);
+		Resources.Push(Resource);
 	}
+
 	bool IsEmpty()
 	{
-		return !NumBatch;
+		return Resources.Num() == 0;
 	}
 };
 
