Add extra debug info for FRenderResource lifetime.

diff --git a/Engine/Source/Runtime/RenderCore/Private/RenderResource.cpp b/Engine/Source/Runtime/RenderCore/Private/RenderResource.cpp
--- a/Engine/Source/Runtime/RenderCore/Private/RenderResource.cpp
+++ b/Engine/Source/Runtime/RenderCore/Private/RenderResource.cpp
@@ -316,6 +316,8 @@ void FRenderResource::UpdateRHI()
 
 FRenderResource::~FRenderResource()
 {
+	checkf(ResourceState == ERenderResourceState::Default, TEXT(" Invalid Resource State: %s"), ResourceState == ERenderResourceState::BatchReleased ? TEXT("BatchReleased") : TEXT("Deleted"));
+	ResourceState = ERenderResourceState::Deleted;
 	if (IsInitialized() && !GIsCriticalError)
 	{
 		// Deleting an initialized FRenderResource will result in a crash later since it is still linked
@@ -358,7 +360,9 @@ struct FBatchedReleaseResources
 			{
 				for (FRenderResource* Resource : BatchedReleaseResources)
 				{
+					check(Resource->ResourceState == ERenderResourceState::BatchReleased);
 					Resource->ReleaseResource();
+					Resource->ResourceState = ERenderResourceState::Default;
 				}
 			});
 		}
@@ -371,6 +375,8 @@ struct FBatchedReleaseResources
 			Flush();
 		}
 		check(Resources.Num() < NumPerBatch);
+		check(Resource->ResourceState == ERenderResourceState::Default);
+		Resource->ResourceState = ERenderResourceState::BatchReleased;
 		Resources.Push(Resource);
 	}
 
diff --git a/Engine/Source/Runtime/RenderCore/Public/RenderResource.h b/Engine/Source/Runtime/RenderCore/Public/RenderResource.h
--- a/Engine/Source/Runtime/RenderCore/Public/RenderResource.h
+++ b/Engine/Source/Runtime/RenderCore/Public/RenderResource.h
@@ -33,6 +33,13 @@ extern int32 GGlobalBufferNumFramesUnusedThresold;
 /** Experimental: whether we free helper structures after submitting to RHI. */
 extern RENDERCORE_API bool GFreeStructuresOnRHIBufferCreation;
 
+enum class ERenderResourceState : uint8
+{
+	Default,
+	BatchReleased,
+	Deleted,
+};
+
 /**
  * A rendering resource which is owned by the rendering thread.
  * NOTE - Adding new virtual methods to this class may require stubs added to FViewport/FDummyViewport, otherwise certain modules may have link errors
@@ -172,6 +179,9 @@ protected:
 
 private:
 	TEnumAsByte<ERHIFeatureLevel::Type> FeatureLevel;
+
+public:
+	ERenderResourceState ResourceState = ERenderResourceState::Default;
 };
 
 /**
