Working around DXC issues that caused constant buffers to be marked used when we initialized bindless resources for Uniform Buffers.
Uniform Buffer struct initializing now just use a global resource/sampler to instead of calling the getter function.
Adding bindless support to FShaderResourceParameter.

diff --git a/Engine/Shaders/Public/BindlessResources.ush b/Engine/Shaders/Public/BindlessResources.ush
--- a/Engine/Shaders/Public/BindlessResources.ush
+++ b/Engine/Shaders/Public/BindlessResources.ush
@@ -33,12 +33,16 @@
 		DEFINE_RESOURCE_VARIABLE(Type, Name); \
 		static const Type Name = GET_RESOURCE(Name);
 
-	COMPILER_DEFINE AUTO_BINDLESS_RESOURCE_INDEX(Name)          DEFINE_RESOURCE_INDEX(Name)
-	COMPILER_DEFINE AUTO_BINDLESS_RESOURCE_VARIABLE(Type, Name) DEFINE_RESOURCE_VARIABLE(Type, Name); static const Type Name = GET_RESOURCE(Name)
+	COMPILER_DEFINE AUTO_BINDLESS_RESOURCE_INDEX(Name)          uint BindlessResource_##Name
+	COMPILER_DEFINE AUTO_BINDLESS_RESOURCE_VARIABLE(Type, Name) Type GetBindlessResource##Name() { return GetResourceFromHeap(Type, BindlessResource_##Name); }; static const Type Name = GetBindlessResource##Name()
+
+	#define UB_RESOURCE_MEMBER_RESOURCE(Type, Name) AUTO_BINDLESS_RESOURCE_VARIABLE(Type, Name)
 #else
 	#define DEFINE_RESOURCE_VARIABLE(Type, Name)  Type Name
 	#define DEFINE_RESOURCE(Type, Name)           Type Name
 	#define GET_RESOURCE(Name)                    Name
+
+	#define UB_RESOURCE_MEMBER_RESOURCE(Type, Name) Type Name
 #endif
 
 #if ENABLE_BINDLESS_SAMPLERS && PLATFORM_SUPPORTS_BINDLESS
@@ -57,10 +61,14 @@
 
 	COMPILER_DEFINE AUTO_BINDLESS_SAMPLER_INDEX(Name)          uint BindlessSampler_##Name
 	COMPILER_DEFINE AUTO_BINDLESS_SAMPLER_VARIABLE(Type, Name) Type GetBindlessSampler##Name() { return GetSamplerFromHeap(Type, BindlessSampler_##Name); }; static const Type Name = GetBindlessSampler##Name()
+
+	#define UB_RESOURCE_MEMBER_SAMPLER(Type, Name) AUTO_BINDLESS_SAMPLER_VARIABLE(Type, Name)
 #else
 	#define DEFINE_SAMPLER_VARIABLE(Type, Name)  Type Name
 	#define DEFINE_SAMPLER(Type, Name)           Type Name
 	#define GET_SAMPLER(Name)                    Name
+
+	#define UB_RESOURCE_MEMBER_SAMPLER(Type, Name) Type Name
 #endif
 
 /*
@@ -71,21 +79,20 @@ cbuffer UniformBuffer_MyUniformBuffer
 	DEFINE_SAMPLER_INDEX(MyUniformBuffer_MySampler);
 };
 
-DEFINE_RESOURCE_VARIABLE(Texture2D, MyUniformBuffer_MyTexture);
-DEFINE_SAMPLER_VARIABLE(SamplerState, MyUniformBuffer_MySampler);
+UB_RESOURCE_MEMBER_RESOURCE(Texture2D, MyUniformBuffer_MyTexture);
+UB_RESOURCE_MEMBER_SAMPLER(SamplerState, MyUniformBuffer_MySampler);
 
 static const struct
 {
 	Texture2D MyTexture;
 	SamplerState MySampler;
 }
-MyUniformBuffer = { GET_RESOURCE(MyUniformBuffer_MyTexture), GET_SAMPLER(MyUniformBuffer_MySampler) };
+MyUniformBuffer = { MyUniformBuffer_MyTexture, MyUniformBuffer_MySampler };
 
 {	// No wrappers necessary to use
 	MyUniformBuffer.MyTexture.Sample(MyUniformBuffer.MySampler, ...);
 }
 
-
 // Normal uses will look like:
 DEFINE_RESOURCE(Texture2D, MyTexture);
 DEFINE_SAMPLER(SamplerState, MySampler);
diff --git a/Engine/Source/Runtime/Core/Public/HAL/Platform.h b/Engine/Source/Runtime/Core/Public/HAL/Platform.h
--- a/Engine/Source/Runtime/Core/Public/HAL/Platform.h
+++ b/Engine/Source/Runtime/Core/Public/HAL/Platform.h
@@ -350,6 +350,10 @@
 	#define PLATFORM_SUPPORTS_MESH_SHADERS 0
 #endif
 
+#ifndef PLATFORM_SUPPORTS_BINDLESS_RENDERING
+	#define PLATFORM_SUPPORTS_BINDLESS_RENDERING 0
+#endif
+
 #ifndef PLATFORM_BUILTIN_VERTEX_HALF_FLOAT
 	#define PLATFORM_BUILTIN_VERTEX_HALF_FLOAT	1
 #endif
diff --git a/Engine/Source/Runtime/Core/Public/Windows/WindowsPlatform.h b/Engine/Source/Runtime/Core/Public/Windows/WindowsPlatform.h
--- a/Engine/Source/Runtime/Core/Public/Windows/WindowsPlatform.h
+++ b/Engine/Source/Runtime/Core/Public/Windows/WindowsPlatform.h
@@ -88,6 +88,7 @@ typedef FWindowsPlatformTypes FPlatformTypes;
 #define PLATFORM_SUPPORTS_VIRTUAL_TEXTURE_STREAMING			1
 #define PLATFORM_SUPPORTS_VARIABLE_RATE_SHADING				1
 #define PLATFORM_SUPPORTS_MESH_SHADERS						1
+#define PLATFORM_SUPPORTS_BINDLESS_RENDERING				1
 #define PLATFORM_USES__ALIGNED_MALLOC						1
 
 #if WITH_EDITOR
diff --git a/Engine/Source/Runtime/RenderCore/Private/ShaderParameterStruct.cpp b/Engine/Source/Runtime/RenderCore/Private/ShaderParameterStruct.cpp
--- a/Engine/Source/Runtime/RenderCore/Private/ShaderParameterStruct.cpp
+++ b/Engine/Source/Runtime/RenderCore/Private/ShaderParameterStruct.cpp
@@ -726,6 +726,7 @@ inline void SetShaderParametersInternal(
 		RHICmdList.SetShaderParameter(ShaderRHI, Parameter.BufferIndex, Parameter.BaseIndex, Parameter.ByteSize, DataPtr);
 	}
 
+#if PLATFORM_SUPPORTS_BINDLESS_RENDERING
 	for (const FShaderParameterBindings::FBindlessResourceParameter& BindlessResource : Bindings.BindlessResourceParameters)
 	{
 		const FRHIDescriptorHandle Handle = GetBindlessParameterHandle(Reader, BindlessResource);
@@ -735,6 +736,7 @@ inline void SetShaderParametersInternal(
 			RHICmdList.SetShaderParameter(ShaderRHI, 0, BindlessResource.GlobalConstantOffset, 4, &BindlessIndex);
 		}
 	}
+#endif
 
 	TArray<FShaderParameterBindings::FResourceParameter, TInlineAllocator<16>> GraphSRVs;
 
diff --git a/Engine/Source/Runtime/RenderCore/Private/ShaderParameters.cpp b/Engine/Source/Runtime/RenderCore/Private/ShaderParameters.cpp
--- a/Engine/Source/Runtime/RenderCore/Private/ShaderParameters.cpp
+++ b/Engine/Source/Runtime/RenderCore/Private/ShaderParameters.cpp
@@ -43,11 +43,19 @@ FArchive& operator<<(FArchive& Ar,FShaderParameter& P)
 	return Ar << P.BaseIndex << P.NumBytes << PBufferIndex;
 }
 
-void FShaderResourceParameter::Bind(const FShaderParameterMap& ParameterMap,const TCHAR* ParameterName,EShaderParameterFlags Flags)
+void FShaderResourceParameter::Bind(const FShaderParameterMap& ParameterMap, const TCHAR* ParameterName, EShaderParameterFlags Flags)
 {
-	uint16 UnusedBufferIndex = 0;
-
-	if(!ParameterMap.FindParameterAllocation(ParameterName,UnusedBufferIndex,BaseIndex,NumResources) && Flags == SPF_Mandatory)
+	if (TOptional<FParameterAllocation> Allocation = ParameterMap.FindParameterAllocation(ParameterName))
+	{
+		if (Allocation->Type == EShaderParameterType::BindlessResourceIndex || Allocation->Type == EShaderParameterType::BindlessSamplerIndex)
+		{
+			checkf(Allocation->BufferIndex == 0, TEXT("Unexpected buffer index (%d) for bindless index. Global bindless parameters are expected to be in the global constant buffer (buffer index 0)."), Allocation->BufferIndex);
+		}
+		BaseIndex = Allocation->BaseIndex;
+		NumResources = Allocation->Size;
+		Type = Allocation->Type;
+	}
+	else if (Flags == SPF_Mandatory)
 	{
 		if (!UE_LOG_ACTIVE(LogShaders, Log))
 		{
@@ -277,21 +285,21 @@ static void CreateHLSLUniformBufferStructMembersDeclaration(
 
 			if (Member.GetBaseType() == UBMT_SAMPLER)
 			{
-				Decl.ResourceMembers += FString::Printf(TEXT("DEFINE_SAMPLER_VARIABLE(%s, %s);\r\n"), Member.GetShaderType(), *ParameterName);
+				Decl.ResourceMembers += FString::Printf(TEXT("UB_RESOURCE_MEMBER_SAMPLER(%s, %s);\r\n"), Member.GetShaderType(), *ParameterName);
 				Decl.StructMembers += FString::Printf(TEXT("\t%s %s;\r\n"), Member.GetShaderType(), Member.GetName());
-				Decl.Initializer += FString::Printf(TEXT("GET_SAMPLER(%s),"), *ParameterName);
+				Decl.Initializer += FString::Printf(TEXT("%s,"), *ParameterName);
 			}
 			else if (Member.GetBaseType() == UBMT_SRV)
 			{
-				Decl.ResourceMembers += FString::Printf(TEXT("PLATFORM_SUPPORTS_SRV_UB_MACRO( DEFINE_RESOURCE_VARIABLE(%s, %s); ) \r\n"), Member.GetShaderType(), *ParameterName);
+				Decl.ResourceMembers += FString::Printf(TEXT("PLATFORM_SUPPORTS_SRV_UB_MACRO( UB_RESOURCE_MEMBER_RESOURCE(%s, %s); ) \r\n"), Member.GetShaderType(), *ParameterName);
 				Decl.StructMembers += FString::Printf(TEXT("\tPLATFORM_SUPPORTS_SRV_UB_MACRO( %s %s; ) \r\n"), Member.GetShaderType(), Member.GetName());
-				Decl.Initializer += FString::Printf(TEXT(" PLATFORM_SUPPORTS_SRV_UB_MACRO( GET_RESOURCE(%s), ) "), *ParameterName);
+				Decl.Initializer += FString::Printf(TEXT(" PLATFORM_SUPPORTS_SRV_UB_MACRO( %s, ) "), *ParameterName);
 			}
 			else
 			{
-				Decl.ResourceMembers += FString::Printf(TEXT("DEFINE_RESOURCE_VARIABLE(%s, %s);\r\n"), Member.GetShaderType(), *ParameterName);
+				Decl.ResourceMembers += FString::Printf(TEXT("UB_RESOURCE_MEMBER_RESOURCE(%s, %s);\r\n"), Member.GetShaderType(), *ParameterName);
 				Decl.StructMembers += FString::Printf(TEXT("\t%s %s;\r\n"), Member.GetShaderType(), Member.GetName());
-				Decl.Initializer += FString::Printf(TEXT("GET_RESOURCE(%s),"), *ParameterName);
+				Decl.Initializer += FString::Printf(TEXT("%s,"), *ParameterName);
 			}
 		}
 	}
diff --git a/Engine/Source/Runtime/RenderCore/Public/ShaderParameterUtils.h b/Engine/Source/Runtime/RenderCore/Public/ShaderParameterUtils.h
--- a/Engine/Source/Runtime/RenderCore/Public/ShaderParameterUtils.h
+++ b/Engine/Source/Runtime/RenderCore/Public/ShaderParameterUtils.h
@@ -80,52 +80,6 @@ void SetShaderValueOnContext(
 	}
 }
 
-
-/** Specialization of the above for C++ bool type. */
-template<typename ShaderRHIParamRef>
-UE_DEPRECATED(4.24, "Please use integer values for boolean shader parameters instead.")
-void SetShaderValue(
-	FRHICommandList& RHICmdList, 
-	const ShaderRHIParamRef& Shader,
-	const FShaderParameter& Parameter,
-	bool Value,
-	uint32 ElementIndex = 0
-	)
-{
-	const uint32 BoolValue = Value;
-	SetShaderValue(RHICmdList, Shader, Parameter, BoolValue, ElementIndex);
-}
-
-/** Specialization of the above for C++ bool type. */
-template<typename ShaderRHIParamRef>
-UE_DEPRECATED(4.24, "Please use integer values for boolean shader parameters instead.")
-void SetShaderValue(
-	FRHIComputeCommandList& RHICmdList,
-	const ShaderRHIParamRef& Shader,
-	const FShaderParameter& Parameter,
-	bool Value,
-	uint32 ElementIndex = 0
-	)
-{
-	const uint32 BoolValue = Value;
-	SetShaderValue(RHICmdList, Shader, Parameter, BoolValue, ElementIndex);
-}
-
-/** Specialization of the above for C++ bool type. */
-template<typename ShaderRHIParamRef>
-UE_DEPRECATED(4.24, "Please use integer values for boolean shader parameters instead.")
-void SetShaderValue(
-	FRHICommandListImmediate& RHICmdList,
-	const ShaderRHIParamRef& Shader,
-	const FShaderParameter& Parameter,
-	bool Value,
-	uint32 ElementIndex = 0
-	)
-{
-	const uint32 BoolValue = Value;
-	SetShaderValue(RHICmdList, Shader, Parameter, BoolValue, ElementIndex);
-}
-
 /**
  * Sets the value of a shader parameter array.  Template'd on shader type
  * A template parameter specified the type of the parameter value.
@@ -198,76 +152,63 @@ GUARD_SETSHADERVALUE(FQuat4)
 GUARD_SETSHADERVALUE(FSphere3)
 GUARD_SETSHADERVALUE(FBox3)
 
+template<typename TRHIShader, typename TRHICmdList>
+FORCEINLINE void SetBindlessParameter(TRHICmdList& RHICmdList, TRHIShader* Shader, const FShaderResourceParameter& Parameter, FRHIDescriptorHandle Handle)
+{
+	if (Handle.IsValid())
+	{
+		const uint32 BindlessIndex = Handle.GetIndex();
+		RHICmdList.SetShaderParameter(Shader, 0, Parameter.GetBaseIndex(), sizeof(BindlessIndex), &BindlessIndex);
+	}
+}
+
 
 /**
- * Sets the value of a pixel shader bool parameter.
+ * Sets the value of a shader surface parameter (e.g. to access MSAA samples).
+ * Template'd on shader type (e.g. pixel shader or compute shader).
  */
-UE_DEPRECATED(4.24, "Please use integer values for boolean shader parameters instead.")
-inline void SetPixelShaderBool(
-	FRHICommandList& RHICmdList, 
-	FRHIPixelShader* PixelShader,
-	const FShaderParameter& Parameter,
-	bool Value
-	)
+template<typename TRHIShader, typename TRHICmdList>
+FORCEINLINE void SetTextureParameter(TRHICmdList& RHICmdList, TRHIShader* Shader, const FShaderResourceParameter& Parameter, FRHITexture* NewTextureRHI, uint32 ElementIndex = 0)
 {
-	// This will trigger if the parameter was not serialized
-	checkSlow(Parameter.IsInitialized());
-
-	if (Parameter.GetNumBytes() > 0)
+	if (Parameter.IsBound() && ElementIndex < Parameter.GetNumResources())
 	{
-		// Convert to uint32 before passing to RHI
-		uint32 BoolValue = Value;
-		RHICmdList.SetShaderParameter(
-			PixelShader,
-			Parameter.GetBufferIndex(),
-			Parameter.GetBaseIndex(),
-			sizeof(BoolValue),
-			&BoolValue
-			);
+#if PLATFORM_SUPPORTS_BINDLESS_RENDERING
+		if (Parameter.GetType() == EShaderParameterType::BindlessResourceIndex)
+		{
+			const FRHIDescriptorHandle Handle = NewTextureRHI ? NewTextureRHI->GetDefaultBindlessHandle() : FRHIDescriptorHandle();
+			SetBindlessParameter(RHICmdList, Shader, Parameter, Handle);
+		}
+		else
+#endif
+		{
+			RHICmdList.SetShaderTexture(Shader, Parameter.GetBaseIndex() + ElementIndex, NewTextureRHI);
+		}
 	}
 }
 
 /**
- * Sets the value of a shader texture parameter.  Template'd on shader type
+ * Sets the value of a shader sampler parameter. Template'd on shader type.
  */
 template<typename TRHIShader, typename TRHICmdList>
-FORCEINLINE void SetTextureParameter(
-	TRHICmdList& RHICmdList,
-	TRHIShader* Shader,
-	const FShaderResourceParameter& TextureParameter,
-	const FShaderResourceParameter& SamplerParameter,
-	const FTexture* Texture,
-	uint32 ElementIndex = 0
-	)
+FORCEINLINE void SetSamplerParameter(TRHICmdList& RHICmdList, TRHIShader* Shader, const FShaderResourceParameter& Parameter, FRHISamplerState* SamplerStateRHI, uint32 ElementIndex = 0)
 {
-	// This will trigger if the parameter was not serialized
-	checkSlow(TextureParameter.IsInitialized());
-	checkSlow(SamplerParameter.IsInitialized());
-	if(TextureParameter.IsBound())
+	if (Parameter.IsBound() && ElementIndex < Parameter.GetNumResources())
 	{
-		Texture->LastRenderTime = FApp::GetCurrentTime();
-
-		if (ElementIndex < TextureParameter.GetNumResources())
+#if PLATFORM_SUPPORTS_BINDLESS_RENDERING
+		if (Parameter.GetType() == EShaderParameterType::BindlessSamplerIndex)
 		{
-			RHICmdList.SetShaderTexture( Shader, TextureParameter.GetBaseIndex() + ElementIndex, Texture->TextureRHI);
+			const FRHIDescriptorHandle Handle = SamplerStateRHI ? SamplerStateRHI->GetBindlessHandle() : FRHIDescriptorHandle();
+			SetBindlessParameter(RHICmdList, Shader, Parameter, Handle);
 		}
-	}
-	
-	// @todo UE samplerstate Should we maybe pass in two separate values? SamplerElement and TextureElement? Or never allow an array of samplers? Unsure best
-	// if there is a matching sampler for this texture array index (ElementIndex), then set it. This will help with this case:
-	//			Texture2D LightMapTextures[NUM_LIGHTMAP_COEFFICIENTS];
-	//			SamplerState LightMapTexturesSampler;
-	// In this case, we only set LightMapTexturesSampler when ElementIndex is 0, we don't set the sampler state for all 4 textures
-	// This assumes that the all textures want to use the same sampler state
-	if(SamplerParameter.IsBound())
-	{
-		if (ElementIndex < SamplerParameter.GetNumResources())
+		else
+#endif
 		{
-			RHICmdList.SetShaderSampler( Shader, SamplerParameter.GetBaseIndex() + ElementIndex, Texture->SamplerStateRHI);
+			RHICmdList.SetShaderSampler(Shader, Parameter.GetBaseIndex() + ElementIndex, SamplerStateRHI);
 		}
 	}
 }
 
+
 /**
  * Sets the value of a shader texture parameter. Template'd on shader type.
  */
@@ -282,72 +223,37 @@ FORCEINLINE void SetTextureParameter(
 	uint32 ElementIndex = 0
 	)
 {
-	// This will trigger if the parameter was not serialized
-	checkSlow(TextureParameter.IsInitialized());
-	checkSlow(SamplerParameter.IsInitialized());
-	if(TextureParameter.IsBound())
-	{
-		if (ElementIndex < TextureParameter.GetNumResources())
-		{
-			RHICmdList.SetShaderTexture( Shader, TextureParameter.GetBaseIndex() + ElementIndex, TextureRHI);
-		}
-	}
+	SetTextureParameter(RHICmdList, Shader, TextureParameter, TextureRHI, ElementIndex);
+	
 	// @todo UE samplerstate Should we maybe pass in two separate values? SamplerElement and TextureElement? Or never allow an array of samplers? Unsure best
 	// if there is a matching sampler for this texture array index (ElementIndex), then set it. This will help with this case:
 	//			Texture2D LightMapTextures[NUM_LIGHTMAP_COEFFICIENTS];
 	//			SamplerState LightMapTexturesSampler;
 	// In this case, we only set LightMapTexturesSampler when ElementIndex is 0, we don't set the sampler state for all 4 textures
 	// This assumes that the all textures want to use the same sampler state
-	if(SamplerParameter.IsBound())
-	{
-		if (ElementIndex < SamplerParameter.GetNumResources())
-		{
-			RHICmdList.SetShaderSampler( Shader, SamplerParameter.GetBaseIndex() + ElementIndex, SamplerStateRHI);
-		}
-	}
+
+	SetSamplerParameter(RHICmdList, Shader, SamplerParameter, SamplerStateRHI, ElementIndex);
 }
 
 /**
- * Sets the value of a shader surface parameter (e.g. to access MSAA samples).
- * Template'd on shader type (e.g. pixel shader or compute shader).
+ * Sets the value of a shader texture parameter.  Template'd on shader type
  */
 template<typename TRHIShader, typename TRHICmdList>
 FORCEINLINE void SetTextureParameter(
 	TRHICmdList& RHICmdList,
 	TRHIShader* Shader,
-	const FShaderResourceParameter& Parameter,
-	FRHITexture* NewTextureRHI
-	)
+	const FShaderResourceParameter& TextureParameter,
+	const FShaderResourceParameter& SamplerParameter,
+	const FTexture* Texture,
+	uint32 ElementIndex = 0
+)
 {
-	if(Parameter.IsBound())
+	if (TextureParameter.IsBound())
 	{
-		RHICmdList.SetShaderTexture(
-			Shader,
-			Parameter.GetBaseIndex(),
-			NewTextureRHI
-			);
+		Texture->LastRenderTime = FApp::GetCurrentTime();
 	}
-}
 
-/**
- * Sets the value of a shader sampler parameter. Template'd on shader type.
- */
-template<typename TRHIShader, typename TRHICmdList>
-FORCEINLINE void SetSamplerParameter(
-	TRHICmdList& RHICmdList,
-	TRHIShader* Shader,
-	const FShaderResourceParameter& Parameter,
-	FRHISamplerState* SamplerStateRHI
-	)
-{
-	if(Parameter.IsBound())
-	{
-		RHICmdList.SetShaderSampler(
-			Shader,
-			Parameter.GetBaseIndex(),
-			SamplerStateRHI
-			);
-	}
+	SetTextureParameter(RHICmdList, Shader, TextureParameter, SamplerParameter, Texture->SamplerStateRHI, Texture->TextureRHI, ElementIndex);
 }
 
 /**
@@ -362,13 +268,9 @@ FORCEINLINE void SetSRVParameter(
 	FRHIShaderResourceView* NewShaderResourceViewRHI
 	)
 {
-	if(Parameter.IsBound())
+	if (Parameter.IsBound())
 	{
-		RHICmdList.SetShaderResourceViewParameter(
-			Shader,
-			Parameter.GetBaseIndex(),
-			NewShaderResourceViewRHI
-			);
+		RHICmdList.SetShaderResourceViewParameter(Shader, Parameter.GetBaseIndex(), NewShaderResourceViewRHI);
 	}
 }
 
@@ -383,40 +285,31 @@ FORCEINLINE void SetSRVParameter(
 {
 	if (Parameter.IsBound())
 	{
-		RHICmdList.SetShaderResourceViewParameter(
-			Shader.GetReference(),
-			Parameter.GetBaseIndex(),
-			NewShaderResourceViewRHI
-		);
+		RHICmdList.SetShaderResourceViewParameter(Shader.GetReference(), Parameter.GetBaseIndex(), NewShaderResourceViewRHI);
 	}
 }
+
 /**
  * Sets the value of a unordered access view parameter
  */
-template<typename TRHICmdList>
 FORCEINLINE void SetUAVParameter(
-	TRHICmdList& RHICmdList,
+	FRHIComputeCommandList& RHICmdList,
 	FRHIComputeShader* ComputeShader,
 	const FShaderResourceParameter& Parameter,
 	FRHIUnorderedAccessView* NewUnorderedAccessViewRHI
 	)
 {
-	if(Parameter.IsBound())
+	if (Parameter.IsBound())
 	{
-		RHICmdList.SetUAVParameter(
-			ComputeShader,
-			Parameter.GetBaseIndex(),
-			NewUnorderedAccessViewRHI
-			);
+		RHICmdList.SetUAVParameter(ComputeShader, Parameter.GetBaseIndex(), NewUnorderedAccessViewRHI);
 	}
 }
 
 /**
  * Sets the value of a unordered access view parameter
  */
-template<typename TRHICmdList>
 FORCEINLINE void SetUAVParameter(
-	TRHICmdList& RHICmdList,
+	FRHICommandList& RHICmdList,
 	FRHIPixelShader* PixelShader,
 	const FShaderResourceParameter& Parameter,
 	FRHIUnorderedAccessView* NewUnorderedAccessViewRHI
@@ -424,16 +317,11 @@ FORCEINLINE void SetUAVParameter(
 {
 	if (Parameter.IsBound())
 	{
-		RHICmdList.SetUAVParameter(
-			PixelShader,
-			Parameter.GetBaseIndex(),
-			NewUnorderedAccessViewRHI
-		);
+		RHICmdList.SetUAVParameter(PixelShader, Parameter.GetBaseIndex(), NewUnorderedAccessViewRHI);
 	}
 }
 
 
-
 template<typename TRHICmdList>
 inline bool SetUAVParameterIfCS(TRHICmdList& RHICmdList, FRHIVertexShader* Shader, const FShaderResourceParameter& UAVParameter, FRHIUnorderedAccessView* UAV)
 {
diff --git a/Engine/Source/Runtime/RenderCore/Public/ShaderParameters.h b/Engine/Source/Runtime/RenderCore/Public/ShaderParameters.h
--- a/Engine/Source/Runtime/RenderCore/Public/ShaderParameters.h
+++ b/Engine/Source/Runtime/RenderCore/Public/ShaderParameters.h
@@ -26,6 +26,9 @@ struct FRWBuffer;
 struct FRWBufferStructured;
 struct FShaderCompilerEnvironment;
 
+enum class EShaderParameterType : uint8;
+DECLARE_INTRINSIC_TYPE_LAYOUT(EShaderParameterType);
+
 RENDERCORE_API void CacheUniformBufferIncludes(TMap<const TCHAR*, struct FCachedUniformBufferDeclaration>& Cache, EShaderPlatform Platform);
 
 
@@ -73,26 +76,22 @@ class FShaderResourceParameter
 {
 	DECLARE_EXPORTED_TYPE_LAYOUT(FShaderResourceParameter, RENDERCORE_API, NonVirtual);
 public:
-	FShaderResourceParameter()
-	:	BaseIndex(0)
-	,	NumResources(0) 
-	{}
+	FShaderResourceParameter() = default;
 
 	RENDERCORE_API void Bind(const FShaderParameterMap& ParameterMap,const TCHAR* ParameterName,EShaderParameterFlags Flags = SPF_Optional);
 	friend RENDERCORE_API FArchive& operator<<(FArchive& Ar,FShaderResourceParameter& P);
-	bool IsBound() const { return NumResources > 0; }
 
-	inline bool IsInitialized() const 
-	{ 
-		return true;
-	}
+	inline bool IsBound() const { return NumResources > 0; }
+	inline bool IsInitialized() const { return true; }
 
-	uint32 GetBaseIndex() const { return BaseIndex; }
-	uint32 GetNumResources() const { return NumResources; }
+	inline uint32 GetBaseIndex() const { return BaseIndex; }
+	inline uint32 GetNumResources() const { return NumResources; }
+	inline EShaderParameterType GetType() const { return Type; }
 
 private:
-	LAYOUT_FIELD(uint16, BaseIndex);
-	LAYOUT_FIELD(uint16, NumResources);
+	LAYOUT_FIELD_INITIALIZED(uint16, BaseIndex, 0);
+	LAYOUT_FIELD_INITIALIZED(uint8, NumResources, 0);
+	LAYOUT_FIELD_INITIALIZED(EShaderParameterType, Type, {});
 };
 
 /** A class that binds either a UAV or SRV of a resource. */
