Filter Frequencies for mobile platforms(only allow vertex, pixel and compute)

diff --git a/Engine/Source/Runtime/RenderCore/Private/Shader.cpp b/Engine/Source/Runtime/RenderCore/Private/Shader.cpp
--- a/Engine/Source/Runtime/RenderCore/Private/Shader.cpp
+++ b/Engine/Source/Runtime/RenderCore/Private/Shader.cpp
@@ -356,9 +356,19 @@ FShader* FShaderType::ConstructCompiled(const FShader::CompiledShaderInitializer
 	return (*ConstructCompiledRef)(Initializer);
 }
 
+static bool ShouldCompileShaderFrequency(EShaderFrequency Frequency, EShaderPlatform ShaderPlatform)
+{
+	if (IsMobilePlatform(ShaderPlatform))
+	{
+		return Frequency == SF_Vertex || Frequency == SF_Pixel || Frequency == SF_Compute;
+	}
+
+	return true;
+}
+
 bool FShaderType::ShouldCompilePermutation(const FShaderPermutationParameters& Parameters) const
 {
-	return (*ShouldCompilePermutationRef)(Parameters);
+	return ShouldCompileShaderFrequency((EShaderFrequency)Frequency, Parameters.Platform) && (*ShouldCompilePermutationRef)(Parameters);
 }
 
 void FShaderType::ModifyCompilationEnvironment(const FShaderPermutationParameters& Parameters, FShaderCompilerEnvironment& OutEnvironment) const
