Make "r.Mobile.ShadingPath" per-platform option. Require DXC for GL to enable deferred shading

diff --git a/Engine/Source/Runtime/RenderCore/Private/RenderUtils.cpp b/Engine/Source/Runtime/RenderCore/Private/RenderUtils.cpp
--- a/Engine/Source/Runtime/RenderCore/Private/RenderUtils.cpp
+++ b/Engine/Source/Runtime/RenderCore/Private/RenderUtils.cpp
@@ -1045,8 +1045,11 @@ RENDERCORE_API bool MobileSupportsGPUScene()
 
 RENDERCORE_API bool IsMobileDeferredShadingEnabled(const FStaticShaderPlatform Platform)
 {
-	static auto* MobileShadingPathCvar = IConsoleManager::Get().FindTConsoleVariableDataInt(TEXT("r.Mobile.ShadingPath"));
-	return MobileShadingPathCvar->GetValueOnAnyThread() == 1;
+	static FShaderPlatformCachedIniValue<bool> MobileShadingPathIniValue(TEXT("r.Mobile.ShadingPath"));
+	return 
+		MobileShadingPathIniValue.Get(Platform) == 1 && 
+		// OpenGL requires DXC for deferred shading
+		(!IsOpenGLPlatform(Platform) || IsDxcEnabledForPlatform(Platform));
 }
 
 RENDERCORE_API bool MobileRequiresSceneDepthAux(const FStaticShaderPlatform Platform)
