Limits the size of staged buffer when dump large resources with DumpGPU

diff --git a/Engine/Source/Runtime/RenderCore/Private/DumpGPU.cpp b/Engine/Source/Runtime/RenderCore/Private/DumpGPU.cpp
--- a/Engine/Source/Runtime/RenderCore/Private/DumpGPU.cpp
+++ b/Engine/Source/Runtime/RenderCore/Private/DumpGPU.cpp
@@ -72,6 +72,11 @@ static TAutoConsoleVariable<int32> GDumpBufferCVar(
 	TEXT(" 2: Dump buffers' descriptors and binaries (default)"),
 	ECVF_RenderThreadSafe);
 
+static TAutoConsoleVariable<int32> GDumpMaxStagingSize(
+	TEXT("r.DumpGPU.MaxStagingSize"), 64,
+	TEXT("Maximum size of stating resource in MB (default=64)."),
+	ECVF_RenderThreadSafe);
+
 static TAutoConsoleVariable<int32> GDumpGPUPassParameters(
 	TEXT("r.DumpGPU.PassParameters"), 1,
 	TEXT("Whether to dump the pass parameters."),
@@ -244,17 +249,17 @@ public:
 	class FFileWriteCtx : public FTimeBucketMeasure
 	{
 	public:
-		FFileWriteCtx(FRDGResourceDumpContext* InDumpCtx, ETimingBucket InBucket, int64 WriteSizeBytes)
+		FFileWriteCtx(FRDGResourceDumpContext* InDumpCtx, ETimingBucket InBucket, int64 WriteSizeBytes, int32 FilesOpened = 1)
 			: FTimeBucketMeasure(InDumpCtx, InBucket)
 		{
 			if (InBucket == ETimingBucket::MetadataFileWrite)
 			{
-				DumpCtx->MetadataFilesOpened += 1;
+				DumpCtx->MetadataFilesOpened += FilesOpened;
 				DumpCtx->MetadataFilesWriteBytes += WriteSizeBytes;
 			}
 			else if (InBucket == ETimingBucket::ResourceBinaryFileWrite)
 			{
-				DumpCtx->ResourceBinaryFilesOpened += 1;
+				DumpCtx->ResourceBinaryFilesOpened += FilesOpened;
 				DumpCtx->ResourceBinaryWriteBytes += WriteSizeBytes;
 			}
 			else
@@ -1421,9 +1426,11 @@ public:
 		// Dump the resource's binary to a .bin file.
 		if (bDumpResourceBinary && DumpBufferMode == 2)
 		{
+			int32 StagingResourceByteSize = FMath::Min(ByteSize, GDumpMaxStagingSize.GetValueOnRenderThread() * 1024 * 1024);
+
 			FString DumpFilePath = kResourcesDir / FString::Printf(TEXT("%s.v%016x.bin"), *UniqueResourceName, PtrToUint(bIsOutputResource ? Pass : nullptr));
 
-			if (IsUnsafeToDumpResource(ByteSize, 1.2f))
+			if (IsUnsafeToDumpResource(StagingResourceByteSize, 1.2f))
 			{
 				UE_LOG(LogRendererCore, Warning, TEXT("Not dumping %s because of insuficient memory available for staging buffer."), *DumpFilePath);
 				return;
@@ -1450,39 +1457,72 @@ public:
 				RDG_EVENT_NAME("RDG DumpBuffer(%s -> %s)", Buffer->Name, *DumpFilePath),
 				PassParameters,
 				ERDGPassFlags::Readback,
-				[this, DumpFilePath, Buffer, ByteSize](FRHICommandListImmediate& RHICmdList)
+				[this, DumpFilePath, Buffer, ByteSize, StagingResourceByteSize](FRHICommandListImmediate& RHICmdList)
 			{
 				check(IsInRenderingThread());
+
 				FStagingBufferRHIRef StagingBuffer;
+				FGPUFenceRHIRef Fence;
 				{
 					FTimeBucketMeasure TimeBucketMeasure(this, ETimingBucket::RHIReadbackCommands);
 					StagingBuffer = RHICreateStagingBuffer();
 
-					// Transfer memory GPU -> CPU
-					RHICmdList.CopyToStagingBuffer(Buffer->GetRHI(), StagingBuffer, 0, ByteSize);
+					static const FName FenceName(TEXT("DumpGPU.BufferFence"));
+					Fence = RHICreateGPUFence(FenceName);
 				}
-
-				// Submit to GPU and wait for completion.
-				static const FName FenceName(TEXT("DumpGPU.BufferFence"));
-				FGPUFenceRHIRef Fence = RHICreateGPUFence(FenceName);
+				
+				TUniquePtr<FArchive> Ar;
+				if (bEnableDiskWrite)
 				{
-					FTimeBucketMeasure TimeBucketMeasure(this, ETimingBucket::GPUWait);
-					Fence->Clear();
-					RHICmdList.WriteGPUFence(Fence);
-					RHICmdList.SubmitCommandsAndFlushGPU();
-					RHICmdList.BlockUntilGPUIdle();
+					FString FullPath = GetDumpFullPath(DumpFilePath);
+					FFileWriteCtx WriteCtx(this, ETimingBucket::ResourceBinaryFileWrite, /* ByteSize = */ 0);
+
+					Ar = TUniquePtr<FArchive>(IFileManager::Get().CreateFileWriter(*FullPath, /* WriteFlags = */ 0));
 				}
 
-				void* Content = RHICmdList.LockStagingBuffer(StagingBuffer, Fence.GetReference(), 0, ByteSize);
-				if (Content)
+				for (int32 Offset = 0; Offset < ByteSize; Offset += StagingResourceByteSize)
 				{
-					DumpResourceBinaryToFile(reinterpret_cast<const uint8*>(Content), ByteSize, DumpFilePath);
+					int32 CopyByteSize = FMath::Min(StagingResourceByteSize, ByteSize - Offset);
+
+					// Transfer memory GPU -> CPU
+					{
+						FTimeBucketMeasure TimeBucketMeasure(this, ETimingBucket::RHIReadbackCommands);
+						RHICmdList.CopyToStagingBuffer(Buffer->GetRHI(), StagingBuffer, Offset, CopyByteSize);
+					}
 
-					RHICmdList.UnlockStagingBuffer(StagingBuffer);
+					// Submit to GPU and wait for completion.
+					{
+						FTimeBucketMeasure TimeBucketMeasure(this, ETimingBucket::GPUWait);
+						Fence->Clear();
+						RHICmdList.WriteGPUFence(Fence);
+						RHICmdList.SubmitCommandsAndFlushGPU();
+						RHICmdList.BlockUntilGPUIdle();
+					}
+
+					void* Content = RHICmdList.LockStagingBuffer(StagingBuffer, Fence.GetReference(), 0, CopyByteSize);
+					if (Content)
+					{
+						if (Ar)
+						{
+							FFileWriteCtx WriteCtx(this, ETimingBucket::ResourceBinaryFileWrite, CopyByteSize, /* OpenedFiles = */ 0);
+							Ar->Serialize((void*)Content, CopyByteSize);
+							Ar->Flush();
+						}
+						RHICmdList.UnlockStagingBuffer(StagingBuffer);
+					}
+					else
+					{
+						UE_LOG(LogRendererCore, Warning, TEXT("RHICmdList.LockStagingBuffer() to dump buffer %s failed."), Buffer->Name);
+						break;
+					}
 				}
-				else
+
+				if (Ar)
 				{
-					UE_LOG(LogRendererCore, Warning, TEXT("RHICmdList.LockStagingBuffer() to dump buffer %s failed."), Buffer->Name);
+					FFileWriteCtx WriteCtx(this, ETimingBucket::ResourceBinaryFileWrite, /* ByteSize = */ 0, /* OpenedFiles = */ 0);
+					// Always explicitly close to catch errors from flush/close
+					Ar->Close();
+					Ar = nullptr;
 				}
 
 				StagingBuffer = nullptr;
