Removing reference of a temporary.

diff --git a/Engine/Source/Runtime/RenderCore/Public/ShaderParameterMacros.h b/Engine/Source/Runtime/RenderCore/Public/ShaderParameterMacros.h
--- a/Engine/Source/Runtime/RenderCore/Public/ShaderParameterMacros.h
+++ b/Engine/Source/Runtime/RenderCore/Public/ShaderParameterMacros.h
@@ -208,9 +208,9 @@ public:
 #endif
 	}
 
-	FORCEINLINE const TUniformBufferRef<TBufferStruct>& GetUniformBuffer() const
+	FORCEINLINE TUniformBufferRef<TBufferStruct> GetUniformBuffer() const
 	{
-		return static_cast<const TUniformBufferRef<TBufferStruct>&>(FUniformBufferBinding::GetUniformBuffer());
+		return TUniformBufferRef<TBufferStruct>(FUniformBufferBinding::GetUniformBuffer());
 	}
 
 	FORCEINLINE TRDGUniformBuffer<TBufferStruct>* operator->() const
