Force inline functions that are often called.
The compiler on some platforms doesn't do this automatically.

diff --git a/Engine/Source/Runtime/Experimental/ChaosCore/Public/Chaos/ArrayCollectionArray.h b/Engine/Source/Runtime/Experimental/ChaosCore/Public/Chaos/ArrayCollectionArray.h
--- a/Engine/Source/Runtime/Experimental/ChaosCore/Public/Chaos/ArrayCollectionArray.h
+++ b/Engine/Source/Runtime/Experimental/ChaosCore/Public/Chaos/ArrayCollectionArray.h
@@ -59,17 +59,17 @@ public:
 		SetNum(Num);
 	}
 
-	void RemoveAt(const int Idx, const int Count) override
+	FORCEINLINE void RemoveAt(const int Idx, const int Count) override
 	{
 		TArray<T>::RemoveAt(Idx, Count);
 	}
 
-	void RemoveAtSwap(const int Idx) override
+	FORCEINLINE void RemoveAtSwap(const int Idx) override
 	{
 		TArray<T>::RemoveAtSwap(Idx);
 	}
 
-	void MoveToOtherArray(const int Idx, TArrayCollectionArrayBase& Other)
+	FORCEINLINE void MoveToOtherArray(const int Idx, TArrayCollectionArrayBase& Other)
 	{
 		//todo: add developer check to make sure this is ok?
 		auto& OtherTArray = static_cast<TArrayCollectionArray<T>&>(Other);
@@ -77,7 +77,7 @@ public:
 		TArray<T>::RemoveAtSwap(Idx);
 	}
 
-	uint64 SizeOfElem() const override
+	FORCEINLINE uint64 SizeOfElem() const override
 	{
 		return sizeof(T);
 	}
diff --git a/Engine/Source/Runtime/RenderCore/Public/ShaderParameters.h b/Engine/Source/Runtime/RenderCore/Public/ShaderParameters.h
--- a/Engine/Source/Runtime/RenderCore/Public/ShaderParameters.h
+++ b/Engine/Source/Runtime/RenderCore/Public/ShaderParameters.h
@@ -53,7 +53,7 @@ public:
 
 	RENDERCORE_API void Bind(const FShaderParameterMap& ParameterMap,const TCHAR* ParameterName, EShaderParameterFlags Flags = SPF_Optional);
 	friend RENDERCORE_API FArchive& operator<<(FArchive& Ar,FShaderParameter& P);
-	bool IsBound() const { return NumBytes > 0; }
+	FORCEINLINE bool IsBound() const { return NumBytes > 0; }
 	
 	inline bool IsInitialized() const 
 	{ 
@@ -112,12 +112,12 @@ public:
 		checkf(!(SRVParameter.GetNumResources() && UAVParameter.GetNumResources()),TEXT("Shader binds SRV and UAV of the same resource: %s"),BaseName);
 	}
 
-	bool IsBound() const
+	FORCEINLINE bool IsBound() const
 	{
 		return SRVParameter.IsBound() || UAVParameter.IsBound();
 	}
 
-	bool IsUAVBound() const
+	FORCEINLINE bool IsUAVBound() const
 	{
 		return UAVParameter.IsBound();
 	}
@@ -170,7 +170,7 @@ public:
 		return Ar;
 	}
 
-	bool IsBound() const { return BaseIndex != 0xffff; }
+	FORCEINLINE bool IsBound() const { return BaseIndex != 0xffff; }
 
 	void Serialize(FArchive& Ar)
 	{
