Skip caching MobileBasePassCSM commands when deferred shading is enabled (missed one file)

diff --git a/Engine/Source/Runtime/RenderCore/Private/RenderUtils.cpp b/Engine/Source/Runtime/RenderCore/Private/RenderUtils.cpp
--- a/Engine/Source/Runtime/RenderCore/Private/RenderUtils.cpp
+++ b/Engine/Source/Runtime/RenderCore/Private/RenderUtils.cpp
@@ -1188,7 +1188,15 @@ RENDERCORE_API bool MobileUsesGBufferCustomData(FStaticShaderPlatform ShaderPlat
 RENDERCORE_API bool MobileBasePassAlwaysUsesCSM(const FStaticShaderPlatform Platform)
 {
 	static TConsoleVariableData<int32>* CVar = IConsoleManager::Get().FindTConsoleVariableDataInt(TEXT("r.Mobile.Shadow.CSMShaderCullingMethod"));
-	return CVar && (CVar->GetValueOnAnyThread() & 0xF) == 5 && IsMobileDistanceFieldEnabled(Platform);
+	if (IsMobileDeferredShadingEnabled(Platform))
+	{
+		// deferred shading does not need CSM culling
+		return true;
+	}
+	else
+	{
+		return CVar && (CVar->GetValueOnAnyThread() & 0xF) == 5 && IsMobileDistanceFieldEnabled(Platform);
+	}
 }
 
 RENDERCORE_API bool SupportsGen4TAA(const FStaticShaderPlatform Platform)
