-Use Depth and velocity pass on Xb1 and Ps4 (VelocityParallel).
   - Saves ~0.3ms on PS4/XB1 (Dynres = 100%).
-Add possibility to disable grass velocity output when WPO velocity is enable (only effective when r.velocity.enablevertexdeformation != 0 and r.OptimizedWPO = true)

diff --git a/Engine/Source/Developer/TargetPlatform/Public/Common/TargetPlatformBase.h b/Engine/Source/Developer/TargetPlatform/Public/Common/TargetPlatformBase.h
--- a/Engine/Source/Developer/TargetPlatform/Public/Common/TargetPlatformBase.h
+++ b/Engine/Source/Developer/TargetPlatform/Public/Common/TargetPlatformBase.h
@@ -79,6 +79,7 @@ public:
 
 	TARGETPLATFORM_API virtual bool UsesDBuffer() const override;
 
+	UE_DEPRECATED(5.1, "Use IsUsingBasePassVelocity(EShaderPlatform Platform) in renderutils that will uses FShaderPlatformCachedIniValue to retrieve the cvar value per platform.")
 	TARGETPLATFORM_API virtual bool UsesBasePassVelocity() const override;
 
     TARGETPLATFORM_API virtual bool VelocityEncodeDepth() const override;
diff --git a/Engine/Source/Runtime/Engine/Classes/Components/PrimitiveComponent.h b/Engine/Source/Runtime/Engine/Classes/Components/PrimitiveComponent.h
--- a/Engine/Source/Runtime/Engine/Classes/Components/PrimitiveComponent.h
+++ b/Engine/Source/Runtime/Engine/Classes/Components/PrimitiveComponent.h
@@ -2091,6 +2091,9 @@ public:
 	 */
 	virtual bool IsShown(const FEngineShowFlags& ShowFlags) const;
 
+	/** Returns false if this primitive should never output velocity based on its WPO state. */
+	virtual bool SupportsWorldPositionOffsetVelocity() const { return true; }
+
 #if WITH_EDITOR
 	/**
 	 * Determines whether the supplied bounding box intersects with the component.
diff --git a/Engine/Source/Runtime/Engine/Private/PrimitiveSceneProxy.cpp b/Engine/Source/Runtime/Engine/Private/PrimitiveSceneProxy.cpp
--- a/Engine/Source/Runtime/Engine/Private/PrimitiveSceneProxy.cpp
+++ b/Engine/Source/Runtime/Engine/Private/PrimitiveSceneProxy.cpp
@@ -375,7 +375,7 @@ FPrimitiveSceneProxy::FPrimitiveSceneProxy(const UPrimitiveComponent* InComponen
 
 	bAlwaysHasVelocity = CVarVelocityForceOutput.GetValueOnAnyThread();
 
-	if (!bAlwaysHasVelocity && VertexDeformationOutputsVelocity())
+	if (!bAlwaysHasVelocity && InComponent->SupportsWorldPositionOffsetVelocity() && VertexDeformationOutputsVelocity())
 	{
 		// Find if we have any WPO materials.
 		ERHIFeatureLevel::Type FeatureLevel = GetScene().GetFeatureLevel();
diff --git a/Engine/Source/Runtime/Foliage/Private/GrassInstancedStaticMeshComponent.cpp b/Engine/Source/Runtime/Foliage/Private/GrassInstancedStaticMeshComponent.cpp
new file mode 100644
--- /dev/null
+++ b/Engine/Source/Runtime/Foliage/Private/GrassInstancedStaticMeshComponent.cpp
@@ -0,0 +1,16 @@
+// Copyright Epic Games, Inc. All Rights Reserved.
+#include "GrassInstancedStaticMeshComponent.h"
+
+
+static TAutoConsoleVariable<bool> CVarEnableGrassInstancedWPOVelocity(
+	TEXT("r.Velocity.EnableLandscapeGrass"),
+	true,
+	TEXT("Specify if you want to output velocity for the grass component for WPO.\n")
+	TEXT(" True (default)\n")
+	TEXT(" False")
+	);
+
+bool UGrassInstancedStaticMeshComponent::SupportsWorldPositionOffsetVelocity() const 
+{ 
+	return CVarEnableGrassInstancedWPOVelocity.GetValueOnAnyThread();
+}
\ No newline at end of file
diff --git a/Engine/Source/Runtime/Foliage/Public/GrassInstancedStaticMeshComponent.h b/Engine/Source/Runtime/Foliage/Public/GrassInstancedStaticMeshComponent.h
--- a/Engine/Source/Runtime/Foliage/Public/GrassInstancedStaticMeshComponent.h
+++ b/Engine/Source/Runtime/Foliage/Public/GrassInstancedStaticMeshComponent.h
@@ -15,5 +15,7 @@ public:
 	static void BuildTreeAnyThread(TArray<FMatrix>& InstanceTransforms, TArray<float>& InstanceCustomDataFloats, int32 NumCustomDataFloats, const FBox& MeshBox, TArray<FClusterNode>& OutClusterTree, TArray<int32>& OutSortedInstances, TArray<int32>& OutInstanceReorderTable, int32& OutOcclusionLayerNum, int32 MaxInstancesPerLeaf, bool InGenerateInstanceScalingRange);
 	
 	void AcceptPrebuiltTree(TArray<FInstancedStaticMeshInstanceData>& InInstanceData, TArray<FClusterNode>& InClusterTree, int32 InOcclusionLayerNumNodes, int32 InNumBuiltRenderInstances);
+
+	bool SupportsWorldPositionOffsetVelocity() const override;
 };
 
diff --git a/Engine/Source/Runtime/RenderCore/Private/RenderUtils.cpp b/Engine/Source/Runtime/RenderCore/Private/RenderUtils.cpp
--- a/Engine/Source/Runtime/RenderCore/Private/RenderUtils.cpp
+++ b/Engine/Source/Runtime/RenderCore/Private/RenderUtils.cpp
@@ -1247,9 +1247,6 @@ static_assert(SP_NumPlatforms <= sizeof(GForwardShadingPlatformMask) * 8, "GForw
 RENDERCORE_API uint64 GDBufferPlatformMask = 0;
 static_assert(SP_NumPlatforms <= sizeof(GDBufferPlatformMask) * 8, "GDBufferPlatformMask must be large enough to support all shader platforms");
 
-RENDERCORE_API uint64 GBasePassVelocityPlatformMask = 0;
-static_assert(SP_NumPlatforms <= sizeof(GBasePassVelocityPlatformMask) * 8, "GBasePassVelocityPlatformMask must be large enough to support all shader platforms");
-
 RENDERCORE_API uint64 GVelocityEncodeDepthPlatformMask = 0;
 static_assert(SP_NumPlatforms <= sizeof(GVelocityEncodeDepthPlatformMask) * 8, "GVelocityEncodeDepthPlatformMask must be large enough to support all shader platforms");
 
@@ -1289,12 +1286,6 @@ RENDERCORE_API void RenderUtilsInit()
 		GDBufferPlatformMask = ~0ull;
 	}
 
-	static IConsoleVariable* VelocityPassCVar = IConsoleManager::Get().FindConsoleVariable(TEXT("r.VelocityOutputPass"));
-	if (VelocityPassCVar && VelocityPassCVar->GetInt() == 1)
-	{
-		GBasePassVelocityPlatformMask = ~0ull;
-	}
-
 	static IConsoleVariable* SelectiveBasePassOutputsCVar = IConsoleManager::Get().FindConsoleVariable(TEXT("r.SelectiveBasePassOutputs"));
 	if (SelectiveBasePassOutputsCVar && SelectiveBasePassOutputsCVar->GetInt())
 	{
@@ -1349,15 +1340,6 @@ RENDERCORE_API void RenderUtilsInit()
 					GDBufferPlatformMask &= ~Mask;
 				}
 
-				if (TargetPlatform->UsesBasePassVelocity() && !IsMobilePlatform(ShaderPlatform))
-				{
-					GBasePassVelocityPlatformMask |= Mask;
-				}
-				else
-				{
-					GBasePassVelocityPlatformMask &= ~Mask;
-				}
-
 				if (TargetPlatform->UsesSelectiveBasePassOutputs())
 				{
 					GSelectiveBasePassOutputsPlatformMask |= Mask;
@@ -1428,7 +1410,6 @@ RENDERCORE_API void RenderUtilsInit()
 	if (IsMobilePlatform(GMaxRHIShaderPlatform))
 	{
 		GDBufferPlatformMask = 0;
-		GBasePassVelocityPlatformMask = 0;
 	}
 
 	if (RayTracingCVar && RayTracingCVar->GetInt() && GRHISupportsRayTracing)
@@ -1791,3 +1772,16 @@ RENDERCORE_API bool DoesRuntimeSupportOnePassPointLightShadows(EShaderPlatform P
 	return RHISupportsVertexShaderLayer(Platform)
 		|| (CVar->GetValueOnAnyThread() != 0 && GRHISupportsArrayIndexFromAnyShader != 0);
 }
+
+bool IsUsingBasePassVelocity(const FStaticShaderPlatform Platform)
+{
+	static FShaderPlatformCachedIniValue<bool> PerPlatformCVar(TEXT("r.VelocityOutputPass"));
+	if (IsMobilePlatform(Platform))
+	{
+		return false;
+	}
+	else
+	{
+		return (PerPlatformCVar.Get(Platform) == 1);
+	}
+}
\ No newline at end of file
diff --git a/Engine/Source/Runtime/RenderCore/Public/RenderUtils.h b/Engine/Source/Runtime/RenderCore/Public/RenderUtils.h
--- a/Engine/Source/Runtime/RenderCore/Public/RenderUtils.h
+++ b/Engine/Source/Runtime/RenderCore/Public/RenderUtils.h
@@ -614,11 +614,7 @@ inline bool IsUsingGBuffers(const FStaticShaderPlatform Platform)
 }
 
 /** Returns whether the base pass should output to the velocity buffer is enabled for a given shader platform */
-inline bool IsUsingBasePassVelocity(const FStaticShaderPlatform Platform)
-{
-	extern RENDERCORE_API uint64 GBasePassVelocityPlatformMask;
-	return !!(GBasePassVelocityPlatformMask & (1ull << Platform));
-}
+RENDERCORE_API bool IsUsingBasePassVelocity(const FStaticShaderPlatform Platform);
 
 /** Returns whether the base pass should use selective outputs for a given shader platform */
 inline bool IsUsingSelectiveBasePassOutputs(const FStaticShaderPlatform Platform)
