Fix for race condition with FRDGPooledBuffer reference counting.

diff --git a/Engine/Source/Runtime/RenderCore/Public/RenderGraphResources.h b/Engine/Source/Runtime/RenderCore/Public/RenderGraphResources.h
--- a/Engine/Source/Runtime/RenderCore/Public/RenderGraphResources.h
+++ b/Engine/Source/Runtime/RenderCore/Public/RenderGraphResources.h
@@ -472,7 +472,7 @@ PRAGMA_ENABLE_DEPRECATION_WARNINGS
 }
 
 class RENDERCORE_API FRDGPooledTexture final
-	: public FRefCountedObject
+	: public FRefCountBase
 {
 public:
 	FRDGPooledTexture(FRHITexture* InTexture)
@@ -1147,7 +1147,7 @@ struct FRDGBufferUAVDesc final
 inline FRHIBufferCreateInfo Translate(const FRDGBufferDesc& InDesc);
 
 class RENDERCORE_API FRDGPooledBuffer final
-	: public FRefCountedObject
+	: public FRefCountBase
 {
 public:
 	FRDGPooledBuffer(TRefCountPtr<FRHIBuffer> InBuffer, const FRDGBufferDesc& InDesc, uint32 InNumAllocatedElements, const TCHAR* InName)
