ShaderCodeLibraryCooker: Separate concerns so that writing shaderlibraries to disk is separate from population during the cook. This supports upcoming multiprocess cooking changes.
Tweak some functions for reduction of duplication and readability.
#rnx

diff --git a/Engine/Plugins/FX/Niagara/Source/Niagara/Private/NiagaraScript.cpp b/Engine/Plugins/FX/Niagara/Source/Niagara/Private/NiagaraScript.cpp
--- a/Engine/Plugins/FX/Niagara/Source/Niagara/Private/NiagaraScript.cpp
+++ b/Engine/Plugins/FX/Niagara/Source/Niagara/Private/NiagaraScript.cpp
@@ -31,6 +31,7 @@
 	#include "Interfaces/ITargetPlatform.h"
 	#include "NiagaraSettings.h"
 	#include "Internationalization/Regex.h"
+	#include "ShaderCodeLibrary.h"
 
 	// This is a version string that mimics the old versioning scheme. In case of merge conflicts with DDC versions,
 	// you MUST generate a new GUID and set this new version. If you want to bump this version, generate a new guid
diff --git a/Engine/Plugins/FX/Niagara/Source/NiagaraShader/Private/NiagaraShared.cpp b/Engine/Plugins/FX/Niagara/Source/NiagaraShader/Private/NiagaraShared.cpp
--- a/Engine/Plugins/FX/Niagara/Source/NiagaraShader/Private/NiagaraShared.cpp
+++ b/Engine/Plugins/FX/Niagara/Source/NiagaraShader/Private/NiagaraShared.cpp
@@ -23,6 +23,7 @@
 #include "NiagaraCustomVersion.h"
 #if WITH_EDITOR
 #include "Interfaces/ITargetPlatform.h"
+#include "ShaderCodeLibrary.h"
 #endif
 #include "ShaderParameterMetadataBuilder.h"
 
diff --git a/Engine/Source/Editor/UnrealEd/Classes/CookOnTheSide/CookOnTheFlyServer.h b/Engine/Source/Editor/UnrealEd/Classes/CookOnTheSide/CookOnTheFlyServer.h
--- a/Engine/Source/Editor/UnrealEd/Classes/CookOnTheSide/CookOnTheFlyServer.h
+++ b/Engine/Source/Editor/UnrealEd/Classes/CookOnTheSide/CookOnTheFlyServer.h
@@ -858,6 +858,10 @@ public:
 	void CookerAddReferencedObjects(FReferenceCollector& Ar);
 	void PostGarbageCollect();
 
+	/** Calculate the ShaderLibrary codedir and metadatadir */
+	void GetShaderLibraryPaths(const ITargetPlatform* TargetPlatform, FString& OutShaderCodeDir,
+		FString& OutMetaDataPath, bool bUseProjectDirForDLC=false);
+
 private:
 
 	/**
@@ -935,6 +939,12 @@ private:
 	/** Invokes the necessary FShaderCodeLibrary functions to open a named code library. */
 	void OpenShaderLibrary(FString const& Name);
     
+	/**
+	 * Finishes collection of data that should be in the named code library (e.g. load data from previous
+	 * iterative cook)
+	 */
+	void FinishPopulateShaderLibrary(const ITargetPlatform* TargetPlatform, FString const& Name);
+
 	/** Invokes the necessary FShaderCodeLibrary functions to save and close a named code library. */
 	void SaveShaderLibrary(const ITargetPlatform* TargetPlatform, FString const& Name);
 
diff --git a/Engine/Source/Editor/UnrealEd/Private/CookOnTheFlyServer.cpp b/Engine/Source/Editor/UnrealEd/Private/CookOnTheFlyServer.cpp
--- a/Engine/Source/Editor/UnrealEd/Private/CookOnTheFlyServer.cpp
+++ b/Engine/Source/Editor/UnrealEd/Private/CookOnTheFlyServer.cpp
@@ -1090,7 +1090,8 @@ bool UCookOnTheFlyServer::IsDirectorCookByTheBook() const
 
 bool UCookOnTheFlyServer::IsUsingShaderCodeLibrary() const
 {
-	return IsDirectorCookByTheBook() && AllowShaderCompiling();
+	const UProjectPackagingSettings* const PackagingSettings = GetDefault<UProjectPackagingSettings>();
+	return IsDirectorCookByTheBook() && AllowShaderCompiling() && PackagingSettings->bShareMaterialShaderCode;
 }
 
 bool UCookOnTheFlyServer::IsUsingZenStore() const
@@ -7800,8 +7801,8 @@ void UCookOnTheFlyServer::BeginCookStartShaderCodeLibrary(FBeginCookContext& Beg
 {
     const UProjectPackagingSettings* const PackagingSettings = GetDefault<UProjectPackagingSettings>();
 	bool const bCacheShaderLibraries = IsUsingShaderCodeLibrary();
-    if (bCacheShaderLibraries && PackagingSettings->bShareMaterialShaderCode)
-    {
+	if (bCacheShaderLibraries)
+	{
 		FShaderLibraryCooker::InitForCooking(PackagingSettings->bSharedMaterialNativeLibraries);
         
 		bool bAllPlatformsNeedStableKeys = false;
@@ -7876,16 +7877,12 @@ void UCookOnTheFlyServer::RegisterShaderChunkDataGenerator()
 {
 	check(!IsCookWorkerMode());
 	// add shader library and PSO cache chunkers
-	const UProjectPackagingSettings* const PackagingSettings = GetDefault<UProjectPackagingSettings>();
 	FString LibraryName = GetProjectShaderLibraryName();
-	if (PackagingSettings->bShareMaterialShaderCode && IsUsingShaderCodeLibrary())
+	for (const ITargetPlatform* TargetPlatform : PlatformManager->GetSessionPlatforms())
 	{
-		for (const ITargetPlatform* TargetPlatform : PlatformManager->GetSessionPlatforms())
-		{
-			FAssetRegistryGenerator& RegistryGenerator = *(PlatformManager->GetPlatformData(TargetPlatform)->RegistryGenerator);
-			RegistryGenerator.RegisterChunkDataGenerator(MakeShared<FShaderLibraryChunkDataGenerator>(TargetPlatform));
-			RegistryGenerator.RegisterChunkDataGenerator(MakeShared<FPipelineCacheChunkDataGenerator>(TargetPlatform, LibraryName));
-		}
+		FAssetRegistryGenerator& RegistryGenerator = *(PlatformManager->GetPlatformData(TargetPlatform)->RegistryGenerator);
+		RegistryGenerator.RegisterChunkDataGenerator(MakeShared<FShaderLibraryChunkDataGenerator>(*this, TargetPlatform));
+		RegistryGenerator.RegisterChunkDataGenerator(MakeShared<FPipelineCacheChunkDataGenerator>(TargetPlatform, LibraryName));
 	}
 }
 
@@ -7902,9 +7899,8 @@ static FString GenerateShaderCodeLibraryName(FString const& Name, bool bIsIterat
 
 void UCookOnTheFlyServer::OpenGlobalShaderLibrary()
 {
-	const UProjectPackagingSettings* const PackagingSettings = GetDefault<UProjectPackagingSettings>();
-	bool const bCacheShaderLibraries = IsUsingShaderCodeLibrary();
-	if (bCacheShaderLibraries && PackagingSettings->bShareMaterialShaderCode)
+	const bool bCacheShaderLibraries = IsUsingShaderCodeLibrary();
+	if (bCacheShaderLibraries)
 	{
 		const TCHAR* GlobalShaderLibName = TEXT("Global");
 		FString ActualName = GenerateShaderCodeLibraryName(GlobalShaderLibName, IsCookFlagSet(ECookInitializationFlags::IterateSharedBuild));
@@ -7916,9 +7912,8 @@ void UCookOnTheFlyServer::OpenGlobalShaderLibrary()
 
 void UCookOnTheFlyServer::OpenShaderLibrary(FString const& Name)
 {
-	const UProjectPackagingSettings* const PackagingSettings = GetDefault<UProjectPackagingSettings>();
-	bool const bCacheShaderLibraries = IsUsingShaderCodeLibrary();
-	if (bCacheShaderLibraries && PackagingSettings->bShareMaterialShaderCode)
+	const bool bCacheShaderLibraries = IsUsingShaderCodeLibrary();
+	if (bCacheShaderLibraries)
 	{
 		FString ActualName = GenerateShaderCodeLibraryName(Name, IsCookFlagSet(ECookInitializationFlags::IterateSharedBuild));
 
@@ -8068,16 +8063,16 @@ void UCookOnTheFlyServer::CreatePipelineCache(const ITargetPlatform* TargetPlatf
 
 void UCookOnTheFlyServer::SaveAndCloseGlobalShaderLibrary()
 {
-	const TCHAR* GlobalShaderLibName = TEXT("Global");
-	FString ActualName = GenerateShaderCodeLibraryName(GlobalShaderLibName, IsCookFlagSet(ECookInitializationFlags::IterateSharedBuild));
-
-	const UProjectPackagingSettings* const PackagingSettings = GetDefault<UProjectPackagingSettings>();
-	bool const bCacheShaderLibraries = IsUsingShaderCodeLibrary();
-	if (bCacheShaderLibraries && PackagingSettings->bShareMaterialShaderCode)
+	const bool bCacheShaderLibraries = IsUsingShaderCodeLibrary();
+	if (bCacheShaderLibraries)
 	{
+		const TCHAR* GlobalShaderLibName = TEXT("Global");
+		FString ActualName = GenerateShaderCodeLibraryName(GlobalShaderLibName, IsCookFlagSet(ECookInitializationFlags::IterateSharedBuild));
+
 		// Save shader code map - cleaning directories is deliberately a separate loop here as we open the cache once per shader platform and we don't assume that they can't be shared across target platforms.
 		for (const ITargetPlatform* TargetPlatform : PlatformManager->GetSessionPlatforms())
 		{
+			FinishPopulateShaderLibrary(TargetPlatform, GlobalShaderLibName);
 			SaveShaderLibrary(TargetPlatform, GlobalShaderLibName);
 		}
 
@@ -8085,50 +8080,60 @@ void UCookOnTheFlyServer::SaveAndCloseGlobalShaderLibrary()
 	}
 }
 
-void UCookOnTheFlyServer::SaveShaderLibrary(const ITargetPlatform* TargetPlatform, FString const& Name)
+void UCookOnTheFlyServer::GetShaderLibraryPaths(const ITargetPlatform* TargetPlatform,
+	FString& OutShaderCodeDir, FString& OutMetaDataPath, bool bUseProjectDirForDLC)
 {
-	TArray<FName> ShaderFormats;
-	TargetPlatform->GetAllTargetedShaderFormats(ShaderFormats);
-	if (ShaderFormats.Num() > 0)
-	{
-		FString ActualName = GenerateShaderCodeLibraryName(Name, IsCookFlagSet(ECookInitializationFlags::IterateSharedBuild));
-		FString BasePath = !IsCookingDLC() ? FPaths::ProjectContentDir() : GetContentDirectoryForDLC();
+	// TODO: Saving ShaderChunks into the DLC directory currently does not work, so we have the bUseProjectDirForDLC arg to save to Project
+	const FString BasePath = (!IsCookingDLC() || bUseProjectDirForDLC) ? FPaths::ProjectContentDir() : GetContentDirectoryForDLC();
+	OutShaderCodeDir = ConvertToFullSandboxPath(*BasePath, true, TargetPlatform->PlatformName());
 
-		FString ShaderCodeDir = ConvertToFullSandboxPath(*BasePath, true, TargetPlatform->PlatformName());
+	const FString RootMetaDataPath = FPaths::ProjectDir() / TEXT("Metadata") / TEXT("PipelineCaches");
+	OutMetaDataPath = ConvertToFullSandboxPath(*RootMetaDataPath, true, *TargetPlatform->PlatformName()); 
+}
 
-		const FString RootMetaDataPath = FPaths::ProjectDir() / TEXT("Metadata") / TEXT("PipelineCaches");
-		const FString MetaDataPathSB = ConvertToFullSandboxPath(*RootMetaDataPath, true);
-		const FString MetaDataPath = MetaDataPathSB.Replace(TEXT("[Platform]"), *TargetPlatform->PlatformName());
+void UCookOnTheFlyServer::FinishPopulateShaderLibrary(const ITargetPlatform* TargetPlatform, const FString& Name)
+{
+	FString ShaderCodeDir;
+	FString MetaDataPath;
+	GetShaderLibraryPaths(TargetPlatform, ShaderCodeDir, MetaDataPath);
 
-		TArray<FString>& PlatformSCLCSVPaths = OutSCLCSVPaths.FindOrAdd(FName(TargetPlatform->PlatformName()));
-		const UProjectPackagingSettings* const PackagingSettings = GetDefault<UProjectPackagingSettings>();
-		FString ErrorString;
-		if (!FShaderLibraryCooker::SaveShaderLibraryWithoutChunking(TargetPlatform, Name, ShaderCodeDir, MetaDataPath, PlatformSCLCSVPaths, ErrorString))
-		{
-			// This is fatal - In this case we should cancel any launch on device operation or package write but we don't want to assert and crash the editor
-			LogCookerMessage(FString::Printf(TEXT("%s"), *ErrorString), EMessageSeverity::Error);
-		}
-		else
+	FShaderLibraryCooker::FinishPopulateShaderLibrary(TargetPlatform, Name, ShaderCodeDir, MetaDataPath);
+}
+
+void UCookOnTheFlyServer::SaveShaderLibrary(const ITargetPlatform* TargetPlatform, const FString& Name)
+{
+	FString ShaderCodeDir;
+	FString MetaDataPath;
+	GetShaderLibraryPaths(TargetPlatform, ShaderCodeDir, MetaDataPath);
+
+	TArray<FString>& PlatformSCLCSVPaths = OutSCLCSVPaths.FindOrAdd(FName(TargetPlatform->PlatformName()));
+	FString ErrorString;
+	bool bHasData;
+	if (!FShaderLibraryCooker::SaveShaderLibraryWithoutChunking(TargetPlatform, Name, ShaderCodeDir, MetaDataPath,
+		PlatformSCLCSVPaths, ErrorString, bHasData))
+	{
+		// This is fatal - In this case we should cancel any launch on device operation or package write but we don't want to assert and crash the editor
+		LogCookerMessage(FString::Printf(TEXT("%s"), *ErrorString), EMessageSeverity::Error);
+	}
+	else if (bHasData)
+	{
+		for (const FString& Item : PlatformSCLCSVPaths)
 		{
-			for (const FString& Item : PlatformSCLCSVPaths)
-			{
-				UE_LOG(LogCook, Display, TEXT("Saved scl.csv %s for platform %s, %d bytes"), *Item, *TargetPlatform->PlatformName(),
-					IFileManager::Get().FileSize(*Item));
-			}
+			UE_LOG(LogCook, Display, TEXT("Saved scl.csv %s for platform %s, %d bytes"), *Item, *TargetPlatform->PlatformName(),
+				IFileManager::Get().FileSize(*Item));
 		}
 	}
 }
 
 void UCookOnTheFlyServer::CleanShaderCodeLibraries()
 {
-	const UProjectPackagingSettings* const PackagingSettings = GetDefault<UProjectPackagingSettings>();
-	bool const bCacheShaderLibraries = IsUsingShaderCodeLibrary();
+	const bool bCacheShaderLibraries = IsUsingShaderCodeLibrary();
 
 	for (const ITargetPlatform* TargetPlatform : PlatformManager->GetSessionPlatforms())
 	{
 		UE::Cook::FPlatformData* PlatformData = PlatformManager->GetPlatformData(TargetPlatform);
 		// If this is a full non-iterative build then clean up our temporary files
-		if (bCacheShaderLibraries && PackagingSettings->bShareMaterialShaderCode && PlatformData->bFullBuild)
+		if (bCacheShaderLibraries && PlatformData->bFullBuild)
 		{
 			TArray<FName> ShaderFormats;
 			TargetPlatform->GetAllTargetedShaderFormats(ShaderFormats);
@@ -8169,9 +8174,9 @@ void UCookOnTheFlyServer::CookByTheBookFinished()
 	
 	UCookerSettings const* CookerSettings = GetDefault<UCookerSettings>();
 
-	const UProjectPackagingSettings* const PackagingSettings = GetDefault<UProjectPackagingSettings>();
-	bool const bCacheShaderLibraries = IsUsingShaderCodeLibrary();
 	FString LibraryName = GetProjectShaderLibraryName();
+	check(!LibraryName.IsEmpty());
+	const bool bCacheShaderLibraries = IsUsingShaderCodeLibrary();
 
 	{
 		// Save modified asset registry with all streaming chunk info generated during cook
@@ -8193,7 +8198,10 @@ void UCookOnTheFlyServer::CookByTheBookFinished()
 			SCOPED_BOOT_TIMING("SavingAssetRegistry");
 
 			RegisterLocalizationChunkDataGenerator();
-			RegisterShaderChunkDataGenerator();
+			if (bCacheShaderLibraries)
+			{
+				RegisterShaderChunkDataGenerator();
+			}
 
 			// if we are cooking DLC, the DevelopmentAR isn't needed - it's used when making DLC against shipping, so there's no need to make it
 			// again, as we don't make DLC against DLC (but allow an override just in case)
@@ -8274,6 +8282,11 @@ void UCookOnTheFlyServer::CookByTheBookFinished()
 					}
 				}
 				
+				if (bCacheShaderLibraries)
+				{
+					FinishPopulateShaderLibrary(TargetPlatform, LibraryName);
+				}
+
 				// Add the package hashes to the relevant AssetPackageDatas.
 				// PackageHashes are gated by requiring UPackage::WaitForAsyncFileWrites(), which is called above.
 				FCookSavePackageContext& SaveContext = FindOrCreateSaveContext(TargetPlatform);
@@ -8329,23 +8342,17 @@ void UCookOnTheFlyServer::CookByTheBookFinished()
 						Generator.WriteCookerOpenOrder(*SandboxFile);
 					}
 				}
-				// now that we have the asset registry and cooking open order, we have enough information to split the shader library
-				// into parts for each chunk and (possibly) lay out the code in accordance with the file order
-				if (bCacheShaderLibraries && PackagingSettings->bShareMaterialShaderCode)
+				if (bCacheShaderLibraries)
 				{
+					// now that we have the asset registry and cooking open order, we have enough information to split the shader library
+					// into parts for each chunk and (possibly) lay out the code in accordance with the file order
 					// Save shader code map
-					if (LibraryName.Len() > 0)
-					{
-						SaveShaderLibrary(TargetPlatform, LibraryName);
-
-						CreatePipelineCache(TargetPlatform, LibraryName);
-					}
+					SaveShaderLibrary(TargetPlatform, LibraryName);
+					CreatePipelineCache(TargetPlatform, LibraryName);
 				}
+				if (FParse::Param(FCommandLine::Get(), TEXT("fastcook")))
 				{
-					if (FParse::Param(FCommandLine::Get(), TEXT("fastcook")))
-					{
-						FFileHelper::SaveStringToFile(FString(), *(GetSandboxDirectory(PlatformNameString) / TEXT("fastcook.txt")));
-					}
+					FFileHelper::SaveStringToFile(FString(), *(GetSandboxDirectory(PlatformNameString) / TEXT("fastcook.txt")));
 				}
 				if (IsCreatingReleaseVersion())
 				{
diff --git a/Engine/Source/Editor/UnrealEd/Private/ShaderLibraryChunkDataGenerator.cpp b/Engine/Source/Editor/UnrealEd/Private/ShaderLibraryChunkDataGenerator.cpp
--- a/Engine/Source/Editor/UnrealEd/Private/ShaderLibraryChunkDataGenerator.cpp
+++ b/Engine/Source/Editor/UnrealEd/Private/ShaderLibraryChunkDataGenerator.cpp
@@ -2,15 +2,17 @@
 
 #include "ShaderLibraryChunkDataGenerator.h"
 
-#include "IPlatformFileSandboxWrapper.h"
+#include "CookOnTheSide/CookOnTheFlyServer.h"
 #include "Interfaces/ITargetPlatform.h"
+#include "IPlatformFileSandboxWrapper.h"
 #include "Misc/AssertionMacros.h"
 #include "Misc/ConfigCacheIni.h"
 #include "ShaderCodeLibrary.h"
 
 class FName;
 
-FShaderLibraryChunkDataGenerator::FShaderLibraryChunkDataGenerator(const ITargetPlatform* TargetPlatform)
+FShaderLibraryChunkDataGenerator::FShaderLibraryChunkDataGenerator(UCookOnTheFlyServer& InCOTFS, const ITargetPlatform* TargetPlatform)
+	: COTFS(InCOTFS)
 {
 	// Find out if this platform requires stable shader keys, by reading the platform setting file.
 	bOptedOut = false;
@@ -21,19 +23,22 @@ FShaderLibraryChunkDataGenerator::FShaderLibraryChunkDataGenerator(const ITarget
 	PlatformIniFile.GetBool(TEXT("DevOptions.Shaders"), TEXT("bDoNotChunkShaderLib"), bOptedOut);
 }
 
-
-void FShaderLibraryChunkDataGenerator::GenerateChunkDataFiles(const int32 InChunkId, const TSet<FName>& InPackagesInChunk, const ITargetPlatform* TargetPlatform, FSandboxPlatformFile* InSandboxFile, TArray<FString>& OutChunkFilenames)
+void FShaderLibraryChunkDataGenerator::GenerateChunkDataFiles(const int32 InChunkId, const TSet<FName>& InPackagesInChunk,
+	const ITargetPlatform* TargetPlatform, FSandboxPlatformFile* InSandboxFile, TArray<FString>& OutChunkFilenames)
 {
 	if (!bOptedOut && InPackagesInChunk.Num() > 0)
 	{
-		checkf(PlatformNameUsedForIni == TargetPlatform->IniPlatformName(), TEXT("Mismatch between platform names in shaderlib chunk generator. Ini settings might have been applied incorrectly."));
+		checkf(PlatformNameUsedForIni == TargetPlatform->IniPlatformName(),
+			TEXT("Mismatch between platform names in shaderlib chunk generator. Ini settings might have been applied incorrectly."));
 
 		// get the sandbox content directory here, to relieve shaderlib from including the wrapper
-		const FString InPlatformName = TargetPlatform->PlatformName();
-		const FString ShaderlibContentSandboxRoot = (InSandboxFile->GetSandboxDirectory() / InSandboxFile->GetGameSandboxDirectoryName() / TEXT("Content")).Replace(TEXT("[Platform]"), *InPlatformName);
-		const FString ShaderlibMetadataSandboxRoot = (InSandboxFile->GetSandboxDirectory() / InSandboxFile->GetGameSandboxDirectoryName() / TEXT("Metadata") / TEXT("PipelineCaches")).Replace(TEXT("[Platform]"), *InPlatformName);
+		FString ShaderCodeDir;
+		FString MetaDataPath;
+		COTFS.GetShaderLibraryPaths(TargetPlatform, ShaderCodeDir, MetaDataPath, true /* bUseProjectDirForDLC */);
 
-		FShaderLibraryCooker::SaveShaderLibraryChunk(InChunkId, InPackagesInChunk, TargetPlatform, ShaderlibContentSandboxRoot, ShaderlibMetadataSandboxRoot, OutChunkFilenames);
+		bool bHasData;
+		FShaderLibraryCooker::SaveShaderLibraryChunk(InChunkId, InPackagesInChunk, TargetPlatform, ShaderCodeDir,
+			MetaDataPath, OutChunkFilenames, bHasData);
 	}
 }
 	
\ No newline at end of file
diff --git a/Engine/Source/Editor/UnrealEd/Private/ShaderLibraryChunkDataGenerator.h b/Engine/Source/Editor/UnrealEd/Private/ShaderLibraryChunkDataGenerator.h
--- a/Engine/Source/Editor/UnrealEd/Private/ShaderLibraryChunkDataGenerator.h
+++ b/Engine/Source/Editor/UnrealEd/Private/ShaderLibraryChunkDataGenerator.h
@@ -15,6 +15,7 @@
 class FName;
 class FSandboxPlatformFile;
 class ITargetPlatform;
+class UCookOnTheFlyServer;
 
 /**
  * Implementation for splitting shader library into chunks when creating streaming install manifests.
@@ -27,12 +28,15 @@ class FShaderLibraryChunkDataGenerator : public IChunkDataGenerator
 	/** Temporary/transitional - this holds the platform name whose ini we checked for the opt out. */
 	FString PlatformNameUsedForIni;
 
+	UCookOnTheFlyServer& COTFS;
+
 public:
 
 	/** Target platform passed in MUST match target platform to generate chunks for. This is checked.*/
-	FShaderLibraryChunkDataGenerator(const ITargetPlatform* TargetPlatform);
+	FShaderLibraryChunkDataGenerator(UCookOnTheFlyServer& InCOTFS, const ITargetPlatform* TargetPlatform);
 	virtual ~FShaderLibraryChunkDataGenerator() = default;
 
 	//~ IChunkDataGenerator
-	virtual void GenerateChunkDataFiles(const int32 InChunkId, const TSet<FName>& InPackagesInChunk, const ITargetPlatform* TargetPlatform, FSandboxPlatformFile* InSandboxFile, TArray<FString>& OutChunkFilenames) override;
+	virtual void GenerateChunkDataFiles(const int32 InChunkId, const TSet<FName>& InPackagesInChunk, const ITargetPlatform* TargetPlatform,
+		FSandboxPlatformFile* InSandboxFile, TArray<FString>& OutChunkFilenames) override;
 };
diff --git a/Engine/Source/Runtime/RenderCore/Private/GlobalShader.cpp b/Engine/Source/Runtime/RenderCore/Private/GlobalShader.cpp
--- a/Engine/Source/Runtime/RenderCore/Private/GlobalShader.cpp
+++ b/Engine/Source/Runtime/RenderCore/Private/GlobalShader.cpp
@@ -20,6 +20,7 @@
 #include "Misc/ConfigCacheIni.h"
 #include "Interfaces/ITargetPlatformManagerModule.h"
 #include "Containers/StaticBitArray.h"
+#include "ShaderCodeLibrary.h"
 #include "StaticBoundShaderState.h"
 
 /** The global shader map. */
diff --git a/Engine/Source/Runtime/RenderCore/Private/ShaderCodeArchive.cpp b/Engine/Source/Runtime/RenderCore/Private/ShaderCodeArchive.cpp
--- a/Engine/Source/Runtime/RenderCore/Private/ShaderCodeArchive.cpp
+++ b/Engine/Source/Runtime/RenderCore/Private/ShaderCodeArchive.cpp
@@ -1,28 +1,30 @@
 // Copyright Epic Games, Inc. All Rights Reserved.
 
 #include "ShaderCodeArchive.h"
-#include "ShaderCodeLibrary.h"
-#include "Stats/Stats.h"
-#include "ProfilingDebugging/LoadTimeTracker.h"
-#include "Misc/ScopeRWLock.h"
+
+#include "Async/ParallelFor.h"
+#include "Compression/OodleDataCompression.h"
+#include "Misc/FileHelper.h"
 #include "Misc/MemStack.h"
+#include "Misc/ScopeRWLock.h"
 #include "Policies/PrettyJsonPrintPolicy.h"
+#include "ProfilingDebugging/LoadTimeTracker.h"
+#include "RenderUtils.h"
+#include "RHIShaderFormatDefinitions.inl"
 #include "Serialization/JsonSerializer.h"
-#include "Misc/FileHelper.h"
 #include "Serialization/MemoryReader.h"
-#include "Compression/OodleDataCompression.h"
+#include "ShaderCodeLibrary.h"
+#include "Stats/Stats.h"
+
 #if WITH_EDITOR
 #include "Misc/StringBuilder.h"
 #include "Templates/Greater.h"
 #endif
-#include "RenderUtils.h"
-#include "RHIShaderFormatDefinitions.inl"
-#include "Async/ParallelFor.h"
 
 #if UE_SCA_VISUALIZE_SHADER_USAGE
-#include "Modules/ModuleManager.h"
-#include "IImageWrapperModule.h"
 #include "IImageWrapper.h"
+#include "IImageWrapperModule.h"
+#include "Modules/ModuleManager.h"
 #endif // UE_SCA_VISUALIZE_SHADER_USAGE
 
 
@@ -92,40 +94,21 @@ bool FSerializedShaderArchive::FindOrAddShaderMap(const FSHAHash& Hash, int32& O
 {
 	const uint32 Key = GetTypeHash(Hash);
 	int32 Index = FindShaderMapWithKey(Hash, Key);
-	bool bAdded = false;
-	if (Index == INDEX_NONE)
+	bool bAdded = Index == INDEX_NONE;
+	if (bAdded)
 	{
 		Index = ShaderMapHashes.Add(Hash);
 		ShaderMapEntries.AddDefaulted();
 		check(ShaderMapEntries.Num() == ShaderMapHashes.Num());
 		ShaderMapHashTable.Add(Key, Index);
-#if WITH_EDITOR
-		if (AssociatedAssets && AssociatedAssets->Num() > 0)
-		{
-			ShaderCodeToAssets.Add(Hash, *AssociatedAssets);
-		}
-#endif
-		bAdded = true;
 	}
-	else
-	{
 #if WITH_EDITOR
-		// check if we need to replace or merge assets
-		if (AssociatedAssets && AssociatedAssets->Num())
-		{
-			FShaderMapAssetPaths* PrevAssets = ShaderCodeToAssets.Find(Hash);
-			if (PrevAssets)
-			{
-				int PrevAssetsNum = PrevAssets->Num();
-				PrevAssets->Append(*AssociatedAssets);
-			}
-			else
-			{
-				ShaderCodeToAssets.Add(Hash, *AssociatedAssets);
-			}
-		}
-#endif
+	if (AssociatedAssets && AssociatedAssets->Num())
+	{
+		ShaderCodeToAssets.FindOrAdd(Hash).Append(*AssociatedAssets);
 	}
+#endif
+
 
 	OutIndex = Index;
 	return bAdded;
@@ -152,19 +135,17 @@ int32 FSerializedShaderArchive::FindShader(const FSHAHash& Hash) const
 bool FSerializedShaderArchive::FindOrAddShader(const FSHAHash& Hash, int32& OutIndex)
 {
 	const uint32 Key = GetTypeHash(Hash);
-	int32 Index = FindShaderWithKey(Hash, Key);
-	bool bAdded = false;
-	if (Index == INDEX_NONE)
+	OutIndex = FindShaderWithKey(Hash, Key);
+	if (OutIndex == INDEX_NONE)
 	{
-		Index = ShaderHashes.Add(Hash);
+		OutIndex = ShaderHashes.Add(Hash);
 		ShaderEntries.AddDefaulted();
 		check(ShaderEntries.Num() == ShaderHashes.Num());
-		ShaderHashTable.Add(Key, Index);
-		bAdded = true;
+		ShaderHashTable.Add(Key, OutIndex);
+		return true;
 	}
 
-	OutIndex = Index;
-	return bAdded;
+	return false;
 }
 
 #if UE_SCA_VISUALIZE_SHADER_USAGE
@@ -345,6 +326,7 @@ void FSerializedShaderArchive::Finalize()
 		}
 	}
 
+	constexpr int32 MaxByteGapAllowedInAPreload = 1024;
 	PreloadEntries.Empty();
 	for (FShaderMapEntry& ShaderMapEntry : ShaderMapEntries)
 	{
@@ -368,7 +350,7 @@ void FSerializedShaderArchive::Finalize()
 			const int64 Gap = PreloadEntry.Offset - CurrentPreloadEntry.Offset - CurrentPreloadEntry.Size;
 			checkf(Gap >= 0, TEXT("Overlapping preload entries, [%lld-%lld), [%lld-%lld)"),
 				CurrentPreloadEntry.Offset, CurrentPreloadEntry.Offset + CurrentPreloadEntry.Size, PreloadEntry.Offset, PreloadEntry.Offset + PreloadEntry.Size);
-			if (Gap > 1024)
+			if (Gap > MaxByteGapAllowedInAPreload)
 			{
 				++ShaderMapEntry.NumPreloadEntries;
 				PreloadEntries.Add(CurrentPreloadEntry);
diff --git a/Engine/Source/Runtime/RenderCore/Private/ShaderCodeLibrary.cpp b/Engine/Source/Runtime/RenderCore/Private/ShaderCodeLibrary.cpp
--- a/Engine/Source/Runtime/RenderCore/Private/ShaderCodeLibrary.cpp
+++ b/Engine/Source/Runtime/RenderCore/Private/ShaderCodeLibrary.cpp
@@ -5,49 +5,49 @@ ShaderCodeLibrary.cpp: Bound shader state cache implementation.
 =============================================================================*/
 
 #include "ShaderCodeLibrary.h"
-#include "ShaderCodeArchive.h"
-#include "Shader.h"
+
 #include "Algo/Replace.h"
+#include "Async/AsyncFileHandle.h"
+#include "Async/ParallelFor.h"
+#include "Containers/HashTable.h"
 #include "Containers/StringView.h"
-#include "Misc/SecureHash.h"
-#include "Misc/Paths.h"
-#include "Math/UnitConversion.h"
+#include "FileCache/FileCache.h"
 #include "HAL/FileManagerGeneric.h"
 #include "HAL/LowLevelMemTracker.h"
 #include "HAL/PlatformSplash.h"
+#include "Hash/CityHash.h"
+#include "Interfaces/IPluginManager.h"
+#include "Math/UnitConversion.h"
+#include "Misc/CommandLine.h"
+#include "Misc/ConfigCacheIni.h"
+#include "Misc/CoreDelegates.h"
+#include "Misc/FileHelper.h"
+#include "Misc/Paths.h"
 #include "Misc/ScopeLock.h"
 #include "Misc/ScopeRWLock.h"
+#include "Misc/SecureHash.h"
 #include "Misc/StringBuilder.h"
-#include "String/ParseTokens.h"
-#include "Async/AsyncFileHandle.h"
 #include "PipelineFileCache.h"
-#include "Interfaces/IPluginManager.h"
-#include "Hash/CityHash.h"
-#include "Containers/HashTable.h"
-#include "FileCache/FileCache.h"
-#include "Misc/CoreDelegates.h"
-#include "Async/ParallelFor.h"
-
-#include "ShaderPipelineCache.h"
-#include "Misc/FileHelper.h"
-#include "Misc/ConfigCacheIni.h"
-#include "Misc/CommandLine.h"
 #include "ProfilingDebugging/LoadTimeTracker.h"
 #include "Serialization/LargeMemoryWriter.h"
 #include "Serialization/MemoryReader.h"
+#include "Shader.h"
+#include "ShaderCodeArchive.h"
+#include "ShaderPipelineCache.h"
+#include "String/ParseTokens.h"
 
 #if WITH_EDITORONLY_DATA
-#include "Modules/ModuleManager.h"
 #include "Interfaces/IShaderFormat.h"
 #include "Interfaces/IShaderFormatModule.h"
 #include "Interfaces/ITargetPlatform.h"
 #include "Interfaces/ITargetPlatformManagerModule.h"
+#include "Modules/ModuleManager.h"
 #endif
 
 #if WITH_EDITOR
 #include "PipelineCacheUtilities.h"
-#include "UObject/Object.h"
 #include "UObject/Class.h"
+#include "UObject/Object.h"
 #endif
 
 // allow introspection (e.g. dumping the contents) for easier debugging
@@ -1268,17 +1268,19 @@ struct FEditorShaderCodeArchive
 		FArchive& Ar)
 	{
 		int32 ShaderMapIndex = 0;
-		if (SerializedShaders.FindOrAddShaderMap(OtherShaders.ShaderMapHashes[OtherShaderMapIndex], ShaderMapIndex, 
-				OtherShaders.ShaderCodeToAssets.Find(OtherShaders.ShaderMapHashes[OtherShaderMapIndex])))
+		const FSHAHash& OtherShaderMapHash = OtherShaders.ShaderMapHashes[OtherShaderMapIndex];
+		const FShaderMapAssetPaths* OtherAssets = OtherShaders.ShaderCodeToAssets.Find(OtherShaderMapHash);
+		bool bIsNew = SerializedShaders.FindOrAddShaderMap(OtherShaderMapHash, ShaderMapIndex, OtherAssets);
+		if (bIsNew)
 		{
-			const FShaderMapEntry& PrevShaderMapEntry = OtherShaders.ShaderMapEntries[OtherShaderMapIndex];
+			const FShaderMapEntry& OtherEntry = OtherShaders.ShaderMapEntries[OtherShaderMapIndex];
 			FShaderMapEntry& ShaderMapEntry = SerializedShaders.ShaderMapEntries[ShaderMapIndex];
-			ShaderMapEntry.NumShaders = PrevShaderMapEntry.NumShaders;
+			ShaderMapEntry.NumShaders = OtherEntry.NumShaders;
 			ShaderMapEntry.ShaderIndicesOffset = SerializedShaders.ShaderIndices.AddZeroed(ShaderMapEntry.NumShaders);
 
 			for (uint32 i = 0; i < ShaderMapEntry.NumShaders; ++i)
 			{
-				const int32 OtherShaderIndex = OtherShaders.ShaderIndices[PrevShaderMapEntry.ShaderIndicesOffset + i];
+				const int32 OtherShaderIndex = OtherShaders.ShaderIndices[OtherEntry.ShaderIndicesOffset + i];
 				int32 ShaderIndex = 0;
 				if (SerializedShaders.FindOrAddShader(OtherShaders.ShaderHashes[OtherShaderIndex], ShaderIndex))
 				{
@@ -1345,30 +1347,6 @@ struct FEditorShaderCodeArchive
 			UE_LOG(LogShaderLibrary, Error, TEXT("Failed to open shader code library from %s"), *IntermediateFormatPath);
 		}
 		
-#if 0
-		if (bOK)
-		{
-			FString PipelinesPath = GetPipelinesArchiveFilename(MetaDataDir / TEXT("ShaderLibrarySource"), LibraryName, FormatName);
-			FArchive* PipelinesArchive = IFileManager::Get().CreateFileReader(*PipelinesPath);
-			if (PipelinesArchive)
-			{
-				uint32 ArchiveVersion = 0;
-				*PipelinesArchive << ArchiveVersion;
-				if (ArchiveVersion == GShaderPipelineArchiveVersion)
-				{
-					*PipelinesArchive << Pipelines;
-				}
-				else
-				{
-					bOK = false;
-					UE_LOG(LogShaderLibrary, Warning, TEXT("Failed to deserialize shader pipelines from %s because the archive format %u is incompatible with the current version %u"), *PipelinesPath, ArchiveVersion, GShaderPipelineArchiveVersion);
-				}
-				
-				PipelinesArchive->Close();
-				delete PipelinesArchive;
-			}
-		}
-#endif
 		return bOK;
 	}
 
@@ -1429,11 +1407,14 @@ struct FEditorShaderCodeArchive
 		}
 	}
 
-	bool Finalize(FString OutputDir, const FString& MetaOutputDir, bool bSaveOnlyAssetInfo = false, TArray<FString>* OutputFilenames = nullptr)
+	void FinishPopulate(const FString& OutputDir)
 	{
-		check(LibraryName.Len() > 0);
-
 		AddExistingShaderCodeLibrary(OutputDir);
+	}
+
+	bool SaveToDisk(const FString& OutputDir, const FString& MetaOutputDir, bool bSaveOnlyAssetInfo = false, TArray<FString>* OutputFilenames = nullptr)
+	{
+		check(LibraryName.Len() > 0);
 
 		bool bSuccess = IFileManager::Get().MakeDirectory(*OutputDir, true);
 
@@ -1612,7 +1593,8 @@ struct FEditorShaderCodeArchive
 			if (bOK)
 			{
 				FString Empty;
-				bOK = OutLibrary.Finalize(OutDir, Empty);
+				OutLibrary.FinishPopulate(OutDir);
+				bOK = OutLibrary.SaveToDisk(OutDir, Empty);
 				UE_CLOG(!bOK, LogShaderLibrary, Error, TEXT("Failed to save %s shader patch library %s, %s, %s"), bNativeFormat ? TEXT("native") : TEXT(""), *FormatName.ToString(), *LibraryName, *OutDir);
 				
 				if (bOK && bNativeFormat && OutLibrary.GetFormat()->SupportsShaderArchives())
@@ -1767,9 +1749,6 @@ struct FEditorShaderStableInfo
 	{
 		check(LibraryName.Len() > 0);
 
-		TMap<uint32, FName> NameCache;
-		NameCache.Reserve(2048);
-
 		const FString ShaderIntermediateLocation = FPaths::ProjectSavedDir() / TEXT("Shaders") / FormatName.ToString();
 
 		TArray<FString> ShaderFiles;
@@ -1791,13 +1770,16 @@ struct FEditorShaderStableInfo
 		}
 	}
 
-	bool Finalize(FString OutputDir, FString& OutSCLCSVPath)
+	void FinishPopulate(FString const& OutputDir)
+	{
+		AddExistingShaderCodeLibrary(OutputDir);
+	}
+
+	bool SaveToDisk(FString const& OutputDir, FString& OutSCLCSVPath)
 	{
 		check(LibraryName.Len() > 0);
 		OutSCLCSVPath = FString();
 
-		AddExistingShaderCodeLibrary(OutputDir);
-
 		bool bSuccess = IFileManager::Get().MakeDirectory(*OutputDir, true);
 
 		EShaderPlatform Platform = ShaderFormatToLegacyShaderPlatform(FormatName);
@@ -2433,15 +2415,34 @@ public:
 		StableArchive->AddShader(StableKeyValue);
 	}
 
+	void FinishPopulateShaderCode(const FString& ShaderCodeDir, const FString& MetaOutputDir, const TArray<FName>& ShaderFormats)
+	{
+		for (FName ShaderFormatName : ShaderFormats)
+		{
+			EShaderPlatform SPlatform = ShaderFormatToLegacyShaderPlatform(ShaderFormatName);
+
+			FEditorShaderCodeArchive* CodeArchive = EditorShaderCodeArchive[SPlatform];
+			if (CodeArchive)
+			{
+				CodeArchive->FinishPopulate(ShaderCodeDir);
+			}
+
+			FEditorShaderStableInfo* StableArchive = EditorShaderStableInfo[SPlatform];
+			if (StableArchive)
+			{
+				StableArchive->FinishPopulate(MetaOutputDir);
+			}
+		}
+	}
+
 	bool SaveShaderCode(const FString& ShaderCodeDir, const FString& MetaOutputDir, const TArray<FName>& ShaderFormats, TArray<FString>& OutSCLCSVPath)
 	{
 		bool bOk = ShaderFormats.Num() > 0;
 
 		FScopeLock ScopeLock(&ShaderCodeCS);
 
-		for (int32 i = 0; i < ShaderFormats.Num(); ++i)
+		for (const FName ShaderFormatName : ShaderFormats)
 		{
-			FName ShaderFormatName = ShaderFormats[i];
 			EShaderPlatform SPlatform = ShaderFormatToLegacyShaderPlatform(ShaderFormatName);
 
 			FEditorShaderCodeArchive* CodeArchive = EditorShaderCodeArchive[SPlatform];
@@ -2454,7 +2455,7 @@ public:
 				{
 					// always save shaders in our format even if the platform will use native one. This is needed for iterative cooks (Launch On et al)
 					// to reload previously cooked shaders
-					bOk = CodeArchive->Finalize(ShaderCodeDir, MetaOutputDir) && bOk;
+					bOk = CodeArchive->SaveToDisk(ShaderCodeDir, MetaOutputDir) && bOk;
 
 					bool bShouldWriteInNativeFormat = bOk && bNativeFormat && CodeArchive->GetFormat()->SupportsShaderArchives();
 					if (bShouldWriteInNativeFormat)
@@ -2470,33 +2471,31 @@ public:
 				else
 				{
 					// save asset info only, for debugging
-					bOk = CodeArchive->Finalize(ShaderCodeDir, MetaOutputDir, true) && bOk;
+					bOk = CodeArchive->SaveToDisk(ShaderCodeDir, MetaOutputDir, true) && bOk;
 				}
 			}
-			// Stable shader info is not saved per-chunk (it is not needed runtime), so save it always
+			FEditorShaderStableInfo* StableArchive = EditorShaderStableInfo[SPlatform];
+			if (StableArchive)
 			{
-				FEditorShaderStableInfo* StableArchive = EditorShaderStableInfo[SPlatform];
-				if (StableArchive)
-				{
-					FString SCLCSVPath;
-					bOk &= StableArchive->Finalize(MetaOutputDir, SCLCSVPath);
-					OutSCLCSVPath.Add(SCLCSVPath);
-				}
+				// Stable shader info is not saved per-chunk (it is not needed at runtime), so save it always
+				FString SCLCSVPath;
+				bOk &= StableArchive->SaveToDisk(MetaOutputDir, SCLCSVPath);
+				OutSCLCSVPath.Add(SCLCSVPath);
 			}
 		}
 
 		return bOk;
 	}
 
-	bool SaveShaderCodeChunk(int32 ChunkId, const TSet<FName>& InPackagesInChunk, const TArray<FName>& ShaderFormats, const FString& SandboxDestinationPath, const FString& SandboxMetadataPath, TArray<FString>& OutChunkFilenames)
+	bool SaveShaderCodeChunk(int32 ChunkId, const TSet<FName>& InPackagesInChunk, const TArray<FName>& ShaderFormats,
+		const FString& SandboxDestinationPath, const FString& SandboxMetadataPath, TArray<FString>& OutChunkFilenames)
 	{
 		bool bOk = ShaderFormats.Num() > 0;
 
 		FScopeLock ScopeLock(&ShaderCodeCS);
 
-		for (int32 i = 0; i < ShaderFormats.Num(); ++i)
+		for (const FName ShaderFormatName : ShaderFormats)
 		{
-			FName ShaderFormatName = ShaderFormats[i];
 			EShaderPlatform SPlatform = ShaderFormatToLegacyShaderPlatform(ShaderFormatName);
 
 			// we may get duplicate calls for the same Chunk Id because the cooker sometimes calls asset registry SaveManifests twice.
@@ -2525,7 +2524,7 @@ public:
 
 				// always save shaders in our format even if the platform will use native one. This is needed for iterative cooks (Launch On et al)
 				// to reload previously cooked shaders
-				bOk = PerChunkArchive->Finalize(SandboxDestinationPath, SandboxMetadataPath, false, &OutChunkFilenames) && bOk;
+				bOk = PerChunkArchive->SaveToDisk(SandboxDestinationPath, SandboxMetadataPath, false, &OutChunkFilenames) && bOk;
 
 				bool bShouldWriteInNativeFormat = bOk && bNativeFormat && PerChunkArchive->GetFormat()->SupportsShaderArchives();
 				if (bShouldWriteInNativeFormat)
@@ -2547,9 +2546,8 @@ public:
 	bool PackageNativeShaderLibrary(const FString& ShaderCodeDir, const TArray<FName>& ShaderFormats)
 	{
 		bool bOK = true;
-		for (int32 i = 0; i < ShaderFormats.Num(); ++i)
+		for (const FName ShaderFormatName : ShaderFormats)
 		{
-			FName ShaderFormatName = ShaderFormats[i];
 			EShaderPlatform SPlatform = ShaderFormatToLegacyShaderPlatform(ShaderFormatName);
 			FEditorShaderCodeArchive* CodeArchive = EditorShaderCodeArchive[SPlatform];
 
@@ -3034,11 +3032,11 @@ bool FShaderLibraryCooker::CreatePatchLibrary(TArray<FString> const& OldMetaData
 	return bOK;
 }
 
-bool FShaderLibraryCooker::SaveShaderLibraryWithoutChunking(const ITargetPlatform* TargetPlatform, FString const& Name, FString const& SandboxDestinationPath, FString const& SandboxMetadataPath, TArray<FString>& PlatformSCLCSVPaths, FString& OutErrorMessage)
+void FShaderLibraryCooker::FinishPopulateShaderLibrary(const ITargetPlatform* TargetPlatform, FString const& Name,
+	FString const& SandboxDestinationPath, FString const& SandboxMetadataPath)
 {
-	const FString ActualName = Name;
-	const FString ShaderCodeDir = SandboxDestinationPath;
-	const FString MetaDataPath = SandboxMetadataPath;
+	const FString& ShaderCodeDir = SandboxDestinationPath;
+	const FString& MetaDataPath = SandboxMetadataPath;
 
 	checkf(FShaderLibrariesCollection::Impl != nullptr, TEXT("FShaderLibraryCooker was not initialized properly"));
 	checkf(TargetPlatform, TEXT("A valid TargetPlatform is expected"));
@@ -3048,32 +3046,57 @@ bool FShaderLibraryCooker::SaveShaderLibraryWithoutChunking(const ITargetPlatfor
 	TargetPlatform->GetAllTargetedShaderFormats(ShaderFormats);
 	if (ShaderFormats.Num() > 0)
 	{
-		FString TargetPlatformName = TargetPlatform->PlatformName();
-		bool bSaved = FShaderLibrariesCollection::Impl->SaveShaderCode(ShaderCodeDir, MetaDataPath, ShaderFormats, PlatformSCLCSVPaths);
-
-		if (UNLIKELY(!bSaved))
-		{
-			OutErrorMessage = FString::Printf(TEXT("Saving shared material shader code library failed for %s."), *TargetPlatformName);
-			return false;
-		}
+		FShaderLibrariesCollection::Impl->FinishPopulateShaderCode(ShaderCodeDir, MetaDataPath, ShaderFormats);
 	}
-
-	return true;
 }
 
-bool FShaderLibraryCooker::SaveShaderLibraryChunk(int32 ChunkId, const TSet<FName>& InPackagesInChunk, const ITargetPlatform* TargetPlatform, const FString& SandboxDestinationPath, const FString& SandboxMetadataPath, TArray<FString>& OutChunkFilenames)
+bool FShaderLibraryCooker::SaveShaderLibraryWithoutChunking(const ITargetPlatform* TargetPlatform, FString const& Name,
+	FString const& SandboxDestinationPath, FString const& SandboxMetadataPath, TArray<FString>& PlatformSCLCSVPaths, FString& OutErrorMessage,
+	bool& bOutHasData)
 {
+	// note that shader formats can be shared across the target platforms
+	TArray<FName> ShaderFormats;
+	TargetPlatform->GetAllTargetedShaderFormats(ShaderFormats);
+	if (ShaderFormats.IsEmpty())
+	{
+		bOutHasData = false;
+		return true;
+	}
+	bOutHasData = true;
+
+	const FString& ShaderCodeDir = SandboxDestinationPath;
+	const FString& MetaDataPath = SandboxMetadataPath;
+
 	checkf(FShaderLibrariesCollection::Impl != nullptr, TEXT("FShaderLibraryCooker was not initialized properly"));
 	checkf(TargetPlatform, TEXT("A valid TargetPlatform is expected"));
-	bool bResult = true;
 
+	const bool bSaved = FShaderLibrariesCollection::Impl->SaveShaderCode(ShaderCodeDir, MetaDataPath, ShaderFormats, PlatformSCLCSVPaths);
+	if (UNLIKELY(!bSaved))
+	{
+		OutErrorMessage = FString::Printf(TEXT("Saving shared material shader code library failed for %s."),
+			*TargetPlatform->PlatformName());
+		return false;
+	}
+
+	return true;
+}
+
+bool FShaderLibraryCooker::SaveShaderLibraryChunk(int32 ChunkId, const TSet<FName>& InPackagesInChunk, const ITargetPlatform* TargetPlatform,
+	const FString& SandboxDestinationPath, const FString& SandboxMetadataPath, TArray<FString>& OutChunkFilenames, bool& bOutHasData)
+{
 	TArray<FName> ShaderFormats;
 	TargetPlatform->GetAllTargetedShaderFormats(ShaderFormats);
-	if (ShaderFormats.Num() > 0)
+	if (ShaderFormats.IsEmpty())
 	{
-		bResult = FShaderLibrariesCollection::Impl->SaveShaderCodeChunk(ChunkId, InPackagesInChunk, ShaderFormats, SandboxDestinationPath, SandboxMetadataPath, OutChunkFilenames);
+		bOutHasData = false;
+		return true;
 	}
-	return bResult;
+	bOutHasData = true;
+
+	checkf(FShaderLibrariesCollection::Impl != nullptr, TEXT("FShaderLibraryCooker was not initialized properly"));
+	checkf(TargetPlatform, TEXT("A valid TargetPlatform is expected"));
+	return FShaderLibrariesCollection::Impl->SaveShaderCodeChunk(ChunkId, InPackagesInChunk, ShaderFormats,
+		SandboxDestinationPath, SandboxMetadataPath, OutChunkFilenames);
 }
 
 #endif// WITH_EDITOR
diff --git a/Engine/Source/Runtime/RenderCore/Public/Shader.h b/Engine/Source/Runtime/RenderCore/Public/Shader.h
--- a/Engine/Source/Runtime/RenderCore/Public/Shader.h
+++ b/Engine/Source/Runtime/RenderCore/Public/Shader.h
@@ -37,7 +37,6 @@
 #include "Serialization/ArchiveProxy.h"
 #include "Serialization/MemoryImage.h"
 #include "Serialization/MemoryLayout.h"
-#include "ShaderCodeLibrary.h" // TODO - move to cpp
 #include "ShaderCore.h"
 #include "ShaderParameterMetadata.h"
 #include "ShaderParameters.h"
@@ -77,6 +76,7 @@ class FShaderType;
 class FVertexFactoryType;
 struct FShaderCompiledShaderInitializerType;
 struct FShaderCompilerOutput;
+using FShaderMapAssetPaths = TSet<FName>; // Copied from ShaderCodeLibrary.h
 
 UE_DEPRECATED(4.26, "FShadereCompiledShaderInitializerType is deprecated. Use FShaderCompiledShaderInitializerType.")
 typedef FShaderCompiledShaderInitializerType FShadereCompiledShaderInitializerType;
@@ -505,6 +505,7 @@ public:
 
 	RENDERCORE_API void ToString(FStringBuilderBase& OutString) const;
 
+	/** A hash describing the total contents of *this. Constructed from the contents of ShaderHashes during Finalize. */
 	FSHAHash ResourceHash;
 	TArray<FSHAHash> ShaderHashes;
 	TArray<FShaderEntry> ShaderEntries;
diff --git a/Engine/Source/Runtime/RenderCore/Public/ShaderCodeArchive.h b/Engine/Source/Runtime/RenderCore/Public/ShaderCodeArchive.h
--- a/Engine/Source/Runtime/RenderCore/Public/ShaderCodeArchive.h
+++ b/Engine/Source/Runtime/RenderCore/Public/ShaderCodeArchive.h
@@ -81,7 +81,8 @@ public:
 	/** An array of all shaders descriptors, deduplicated */
 	TArray<FShaderCodeEntry> ShaderEntries;
 
-	/** An array of preload entries*/
+	/** An array of entries for the bytes of shadercode that need to be preloaded for a shadermap.
+	  * Each shadermap has a range in this array, beginning of which is stored in FShaderMapEntry.FirstPreloadIndex. */
 	TArray<FFileCachePreloadEntry> PreloadEntries;
 
 	/** Flat array of shaders referenced by all shadermaps. Each shadermap has a range in this array, beginning of which is
diff --git a/Engine/Source/Runtime/RenderCore/Public/ShaderCodeLibrary.h b/Engine/Source/Runtime/RenderCore/Public/ShaderCodeLibrary.h
--- a/Engine/Source/Runtime/RenderCore/Public/ShaderCodeLibrary.h
+++ b/Engine/Source/Runtime/RenderCore/Public/ShaderCodeLibrary.h
@@ -14,7 +14,6 @@
 #include "Containers/SparseArray.h"
 #include "Containers/StringFwd.h"
 #include "Containers/UnrealString.h"
-#include "CoreMinimal.h"
 #include "Delegates/Delegate.h"
 #include "HAL/Platform.h"
 #include "HAL/PlatformCrt.h"
@@ -291,6 +290,10 @@ struct RENDERCORE_API FShaderLibraryCooker
 	// At cook time, add the human readable key value information
 	static void AddShaderStableKeyValue(EShaderPlatform ShaderPlatform, FStableShaderKeyAndValue& StableKeyValue);
 
+	/** Finishes collection of data that should be in the named code library. This includes loading data from a previous iterative cook. */
+	static void FinishPopulateShaderLibrary(const ITargetPlatform* TargetPlatform, FString const& Name, FString const& SandboxDestinationPath,
+		FString const& SandboxMetadataPath);
+
 	/**
 	 * Saves collected shader code to a single file per shader platform
 	 * When chunking is enabled, this call will not write the shader code, only the SCL.CSV file with the stable shader info.
@@ -301,9 +304,11 @@ struct RENDERCORE_API FShaderLibraryCooker
 	 * @param SandboxMetadataPath path for the metadata (not a part of the build itself, but produced together with the build)
 	 * @param PlatformSCLCSVPaths path where to put the information about the shader hashes
 	 * @param OutErrorMessage used to return the details of the failure (if failed)
+	 * @param bOutHasData Reports whether any files were written to PlatformSCLCSVPaths
 	 * @return true if successful
 	 */
-	static bool SaveShaderLibraryWithoutChunking(const ITargetPlatform* TargetPlatform, FString const& Name, FString const& SandboxDestinationPath, FString const& SandboxMetadataPath, TArray<FString>& PlatformSCLCSVPaths, FString& OutErrorMessage);
+	static bool SaveShaderLibraryWithoutChunking(const ITargetPlatform* TargetPlatform, FString const& Name, FString const& SandboxDestinationPath,
+		FString const& SandboxMetadataPath, TArray<FString>& PlatformSCLCSVPaths, FString& OutErrorMessage, bool& bOutHasData);
 
 	/** 
 	 * Saves a single chunk of the collected shader code (per shader platform). Does not save SCL.CSV info.
@@ -315,9 +320,11 @@ struct RENDERCORE_API FShaderLibraryCooker
 	 * @param TargetPlatform target platform
 	 * @param SandboxDestinationPath where to put the .ushaderbytecode file(s)
 	 * @param OutChunkFilenames array where the function will append the full paths of the written files
+	 * @param bOutHasData Reports whether any files were written to OutChunkFilenames
 	 * @return true if successful
 	 */
-	static bool SaveShaderLibraryChunk(int32 ChunkId, const TSet<FName>& InPackagesInChunk, const ITargetPlatform* TargetPlatform, const FString& SandboxDestinationPath, const FString& SandboxMetadataPath, TArray<FString>& OutChunkFilenames);
+	static bool SaveShaderLibraryChunk(int32 ChunkId, const TSet<FName>& InPackagesInChunk, const ITargetPlatform* TargetPlatform,
+		const FString& SandboxDestinationPath, const FString& SandboxMetadataPath, TArray<FString>& OutChunkFilenames, bool& bOutHasData);
 
 	// Dump collected stats for each shader platform
 	static void DumpShaderCodeStats();
