Disable full string events in RDG in test builds, even when profile gpu is enabled, and disable per-pass tracing as well.

diff --git a/Engine/Source/Runtime/RenderCore/Private/RenderGraphBuilder.cpp b/Engine/Source/Runtime/RenderCore/Private/RenderGraphBuilder.cpp
--- a/Engine/Source/Runtime/RenderCore/Private/RenderGraphBuilder.cpp
+++ b/Engine/Source/Runtime/RenderCore/Private/RenderGraphBuilder.cpp
@@ -2524,7 +2524,7 @@ void FRDGBuilder::ExecutePassEpilogue(FRHIComputeCommandList& RHICmdListPass, FR
 
 void FRDGBuilder::ExecutePass(FRDGPass* Pass, FRHIComputeCommandList& RHICmdListPass)
 {
-#if RDG_EVENTS != RDG_EVENTS_NONE
+#if RDG_ENABLE_DEBUG && RDG_EVENTS != RDG_EVENTS_NONE
 	SCOPED_NAMED_EVENT_TCHAR(Pass->GetName(), FColor::Magenta);
 #endif
 
diff --git a/Engine/Source/Runtime/RenderCore/Public/RenderGraphDefinitions.h b/Engine/Source/Runtime/RenderCore/Public/RenderGraphDefinitions.h
--- a/Engine/Source/Runtime/RenderCore/Public/RenderGraphDefinitions.h
+++ b/Engine/Source/Runtime/RenderCore/Public/RenderGraphDefinitions.h
@@ -46,7 +46,11 @@
 
 /** Whether render graph GPU events are enabled. */
 #if WITH_PROFILEGPU
-	#define RDG_EVENTS RDG_EVENTS_STRING_COPY
+	#if UE_BUILD_TEST
+		#define RDG_EVENTS RDG_EVENTS_STRING_REF
+	#else
+		#define RDG_EVENTS RDG_EVENTS_STRING_COPY
+	#endif
 #elif RHI_WANT_BREADCRUMB_EVENTS
 	#define RDG_EVENTS RDG_EVENTS_STRING_REF
 #else
