Added parallel setup for RDG passes
 - Passes are added to a queue that is consumed by a task to perform pass setup actions.
 - Moved CompilePassBarriers to overlap with resource collection.
 - Converted most tasks to new task system.

diff --git a/Engine/Source/Runtime/RenderCore/Private/RenderGraphBuilder.cpp b/Engine/Source/Runtime/RenderCore/Private/RenderGraphBuilder.cpp
--- a/Engine/Source/Runtime/RenderCore/Private/RenderGraphBuilder.cpp
+++ b/Engine/Source/Runtime/RenderCore/Private/RenderGraphBuilder.cpp
@@ -396,7 +396,6 @@ void FRDGBuilder::TickPoolElements()
 
 #if STATS
 	SET_DWORD_STAT(STAT_RDG_PassCount, GRDGStatPassCount);
-	SET_DWORD_STAT(STAT_RDG_PassWithParameterCount, GRDGStatPassWithParameterCount);
 	SET_DWORD_STAT(STAT_RDG_PassCullCount, GRDGStatPassCullCount);
 	SET_DWORD_STAT(STAT_RDG_RenderPassMergeCount, GRDGStatRenderPassMergeCount);
 	SET_DWORD_STAT(STAT_RDG_PassDependencyCount, GRDGStatPassDependencyCount);
@@ -414,7 +413,6 @@ void FRDGBuilder::TickPoolElements()
 	SET_DWORD_STAT(STAT_RDG_TransitionBatchCount, GRDGStatTransitionBatchCount);
 	SET_MEMORY_STAT(STAT_RDG_MemoryWatermark, int64(GRDGStatMemoryWatermark));
 	GRDGStatPassCount = 0;
-	GRDGStatPassWithParameterCount = 0;
 	GRDGStatPassCullCount = 0;
 	GRDGStatRenderPassMergeCount = 0;
 	GRDGStatPassDependencyCount = 0;
@@ -536,11 +534,13 @@ FRDGBuilder::FRDGBuilder(FRHICommandListImmediate& InRHICmdList, FRDGEventName I
 	, Blackboard(Allocator)
 	, RHICmdListAsyncCompute(FRHICommandListExecutor::GetImmediateAsyncComputeCommandList())
 	, BuilderName(InName)
+	, CompilePipe(TEXT("RDG_CompilePipe"))
 #if RDG_CPU_SCOPES
 	, CPUScopeStacks(Allocator)
 #endif
 	, GPUScopeStacks(Allocator)
 	, bParallelExecuteEnabled(IsParallelExecuteEnabled() && EnumHasAnyFlags(InFlags, ERDGBuilderFlags::AllowParallelExecute))
+	, bParallelSetupEnabled(IsParallelSetupEnabled() && EnumHasAnyFlags(InFlags, ERDGBuilderFlags::AllowParallelExecute))
 #if RDG_ENABLE_DEBUG
 	, UserValidation(Allocator, bParallelExecuteEnabled)
 	, BarrierValidation(&Passes, BuilderName)
@@ -629,7 +629,7 @@ void FRDGBuilder::UseExternalAccessMode(FRDGViewableResource* Resource, ERHIAcce
 
 	// We already validated that back-to-back calls to UseExternalAccessMode are valid only if the parameters match,
 	// so we can safely no-op this call.
-	if (AccessModeState.Mode == FRDGViewableResource::EAccessMode::External)
+	if (AccessModeState.Mode == FRDGViewableResource::EAccessMode::External || AccessModeState.bLocked)
 	{
 		return;
 	}
@@ -656,7 +656,7 @@ void FRDGBuilder::UseInternalAccessMode(FRDGViewableResource* Resource)
 	auto& AccessModeState = Resource->AccessModeState;
 
 	// Just no-op if the resource is already in (or queued for) the Internal state.
-	if (AccessModeState.Mode == FRDGViewableResource::EAccessMode::Internal)
+	if (AccessModeState.Mode == FRDGViewableResource::EAccessMode::Internal || AccessModeState.bLocked)
 	{
 		return;
 	}
@@ -703,9 +703,22 @@ void FRDGBuilder::FlushAccessModeQueue()
 
 	ERHIPipeline ParameterPipelines = ERHIPipeline::None;
 
+	TArray<FRDGPass::FExternalAccessOp, FRDGArrayAllocator> Ops;
+	Ops.Reserve(bParallelSetupEnabled ? AccessModeQueue.Num() : 0);
+
 	for (FRDGViewableResource* Resource : AccessModeQueue)
 	{
 		const auto& AccessModeState = Resource->AccessModeState;
+		Resource->AccessModeState.bQueued = false;
+
+		if (bParallelSetupEnabled)
+		{
+			Ops.Emplace(Resource, AccessModeState.Mode);
+		}
+		else
+		{
+			Resource->AccessModeState.ActiveMode = Resource->AccessModeState.Mode;
+		}
 
 		ParameterPipelines |= AccessModeState.Pipelines;
 
@@ -742,31 +755,43 @@ void FRDGBuilder::FlushAccessModeQueue()
 
 	if (EnumHasAnyFlags(ParameterPipelines, ERHIPipeline::Graphics))
 	{
+		auto ExecuteLambda = [](FRHIComputeCommandList&) {};
+		using LambdaPassType = TRDGLambdaPass<FAccessModePassParameters, decltype(ExecuteLambda)>;
+
 		FAccessModePassParameters* Parameters = ParametersByPipeline[GetRHIPipelineIndex(ERHIPipeline::Graphics)];
 
-		AddPass(
+		FRDGPass* Pass = Passes.Allocate<LambdaPassType>(
+			Allocator,
 			RDG_EVENT_NAME("AccessModePass[Graphics] (Textures: %d, Buffers: %d)", Parameters->Textures.Num(), Parameters->Buffers.Num()),
+			FAccessModePassParameters::FTypeInfo::GetStructMetadata(),
 			Parameters,
 			// Use all of the work flags so that any access is valid.
 			ERDGPassFlags::Copy | ERDGPassFlags::Compute | ERDGPassFlags::Raster | ERDGPassFlags::SkipRenderPass | ERDGPassFlags::NeverCull,
-			[](FRHICommandList&) {});
+			MoveTemp(ExecuteLambda));
+
+		Pass->ExternalAccessOps = MoveTemp(Ops);
+		Pass->bExternalAccessPass = 1;
+		SetupParameterPass(Pass);
 	}
 
 	if (EnumHasAnyFlags(ParameterPipelines, ERHIPipeline::AsyncCompute))
 	{
+		auto ExecuteLambda = [](FRHIComputeCommandList&) {};
+		using LambdaPassType = TRDGLambdaPass<FAccessModePassParameters, decltype(ExecuteLambda)>;
+
 		FAccessModePassParameters* Parameters = ParametersByPipeline[GetRHIPipelineIndex(ERHIPipeline::AsyncCompute)];
 
-		AddPass(
+		FRDGPass* Pass = Passes.Allocate<LambdaPassType>(
+			Allocator,
 			RDG_EVENT_NAME("AccessModePass[AsyncCompute] (Textures: %d, Buffers: %d)", Parameters->Textures.Num(), Parameters->Buffers.Num()),
+			FAccessModePassParameters::FTypeInfo::GetStructMetadata(),
 			Parameters,
 			ERDGPassFlags::AsyncCompute | ERDGPassFlags::NeverCull,
-			[](FRHIComputeCommandList&) {});
-	}
+			MoveTemp(ExecuteLambda));
 
-	// This has to be done *after* adding the passes so they don't inherit the new access mode.
-	for (FRDGViewableResource* Resource : AccessModeQueue)
-	{
-		Resource->AccessModeState.bQueued = false;
+		Pass->ExternalAccessOps = MoveTemp(Ops);
+		Pass->bExternalAccessPass = 1;
+		SetupParameterPass(Pass);
 	}
 
 	AccessModeQueue.Reset();
@@ -876,141 +901,241 @@ void FRDGBuilder::AddPassDependency(FRDGPassHandle ProducerHandle, FRDGPassHandl
 #endif
 }
 
-void FRDGBuilder::Compile()
+void FRDGBuilder::AddPassDependency(FRDGPass* Producer, FRDGPass* Consumer)
 {
-	SCOPE_CYCLE_COUNTER(STAT_RDG_CompileTime);
-	CSV_SCOPED_TIMING_STAT_EXCLUSIVE_CONDITIONAL(RDG_Compile, GRDGVerboseCSVStats != 0);
-	SCOPED_NAMED_EVENT(Compile, FColor::Emerald);
+	auto& Producers = Consumer->Producers;
+	if (Producers.Find(Producer->Handle) == INDEX_NONE)
+	{
+		Producers.Add(Producer->Handle);
+	}
 
-	const FRDGPassHandle ProloguePassHandle = GetProloguePassHandle();
-	const FRDGPassHandle EpiloguePassHandle = GetEpiloguePassHandle();
+#if STATS
+	GRDGStatPassDependencyCount++;
+#endif
+}
 
-	const uint32 CompilePassCount = Passes.Num();
+void FRDGBuilder::AddCullingDependency(FRDGProducerStatesByPipeline& LastProducers, const FRDGProducerState& NextState, ERHIPipeline NextPipeline)
+{
+	for (ERHIPipeline LastPipeline : GetRHIPipelines())
+	{
+		FRDGProducerState& LastProducer = LastProducers[LastPipeline];
 
-	const bool bCullPasses = GRDGCullPasses > 0;
+		if (LastProducer.Access == ERHIAccess::Unknown)
+		{
+			continue;
+		}
 
-	TArray<FRDGPassHandle, FRDGArrayAllocator> PassStack;
+		if (FRDGProducerState::IsDependencyRequired(LastProducer, LastPipeline, NextState, NextPipeline))
+		{
+			AddPassDependency(LastProducer.Pass, NextState.Pass);
+		}
+	}
 
-	if (bCullPasses)
+	if (IsWritableAccess(NextState.Access))
 	{
-		PassStack.Reserve(CompilePassCount);
+		LastProducers[NextPipeline] = NextState;
 	}
+}
 
-	TransitionCreateQueue.Reserve(CompilePassCount);
+void FRDGBuilder::CompilePassBarriers()
+{
+	// Walk the culled graph and compile barriers for each subresource. Certain transitions are redundant; read-to-read, for example.
+	// We can avoid them by traversing and merging compatible states together. The merging states removes a transition, but the merging
+	// heuristic is conservative and choosing not to merge doesn't necessarily mean a transition is performed. They are two distinct steps.
+	// Merged states track the first and last pass interval. Pass references are also accumulated onto each resource. This must happen
+	// after culling since culled passes can't contribute references.
 
-	FRDGPassBitArray PassesOnAsyncCompute(false, CompilePassCount);
+	const FRDGPassHandle ProloguePassHandle = GetProloguePassHandle();
+	const FRDGPassHandle EpiloguePassHandle = GetEpiloguePassHandle();
 
-	// Build producer / consumer dependencies across the graph and construct packed bit-arrays of metadata
-	// for better cache coherency when searching for passes meeting specific criteria. Search roots are also
-	// identified for culling. Passes with untracked RHI output (e.g. SHADER_PARAMETER_{BUFFER, TEXTURE}_UAV)
-	// cannot be culled, nor can any pass which writes to an external resource. Resource extractions extend the
-	// lifetime to the epilogue pass which is always a root of the graph. The prologue and epilogue are helper
-	// passes and therefore never culled.
+	SCOPED_NAMED_EVENT(CompileBarriers, FColor::Emerald);
 
-	if (bCullPasses || AsyncComputePassCount > 0)
+	for (FRDGPassHandle PassHandle = ProloguePassHandle + 1; PassHandle < EpiloguePassHandle; ++PassHandle)
 	{
-		SCOPED_NAMED_EVENT(PassDependencies, FColor::Emerald);
+		FRDGPass* Pass = Passes[PassHandle];
 
-		const auto AddCullingDependency = [&](FRDGProducerStatesByPipeline& LastProducers, const FRDGProducerState& NextState, ERHIPipeline NextPipeline)
+		if (Pass->bCulled || Pass->bEmptyParameters)
 		{
-			for (ERHIPipeline LastPipeline : GetRHIPipelines())
-			{
-				FRDGProducerState& LastProducer = LastProducers[LastPipeline];
+			continue;
+		}
 
-				if (LastProducer.Access == ERHIAccess::Unknown)
+		const ERHIPipeline PassPipeline = Pass->Pipeline;
+
+		const auto MergeSubresourceStates = [&](ERDGViewableResourceType ResourceType, FRDGSubresourceState*& PassMergeState, FRDGSubresourceState*& ResourceMergeState, const FRDGSubresourceState& PassState)
+		{
+			if (!ResourceMergeState || !FRDGSubresourceState::IsMergeAllowed(ResourceType, *ResourceMergeState, PassState))
+			{
+				// Cross-pipeline, non-mergable state changes require a new pass dependency for fencing purposes.
+				if (ResourceMergeState)
 				{
-					continue;
+					for (ERHIPipeline Pipeline : GetRHIPipelines())
+					{
+						if (Pipeline != PassPipeline && ResourceMergeState->LastPass[Pipeline].IsValid())
+						{
+							// Add a dependency from the other pipe to this pass to join back.
+							AddPassDependency(ResourceMergeState->LastPass[Pipeline], PassHandle);
+						}
+					}
 				}
 
-				if (FRDGProducerState::IsDependencyRequired(LastProducer, LastPipeline, NextState, NextPipeline))
+				// Allocate a new pending merge state and assign it to the pass state.
+				ResourceMergeState = AllocSubresource(PassState);
+			}
+			else
+			{
+				// Merge the pass state into the merged state.
+				ResourceMergeState->Access |= PassState.Access;
+
+				FRDGPassHandle& FirstPassHandle = ResourceMergeState->FirstPass[PassPipeline];
+
+				if (FirstPassHandle.IsNull())
 				{
-					AddPassDependency(LastProducer.PassHandle, NextState.PassHandle);
+					FirstPassHandle = PassHandle;
 				}
-			}
 
-			if (IsWritableAccess(NextState.Access))
-			{
-				LastProducers[NextPipeline] = NextState;
+				ResourceMergeState->LastPass[PassPipeline] = PassHandle;
 			}
+
+			PassMergeState = ResourceMergeState;
 		};
 
-		for (FRDGPassHandle PassHandle = ProloguePassHandle + 1; PassHandle < EpiloguePassHandle; ++PassHandle)
+		for (auto& PassState : Pass->TextureStates)
 		{
-			FRDGPass* Pass = Passes[PassHandle];
-			const ERHIPipeline PassPipeline = Pass->Pipeline;
-
-			bool bUntrackedOutputs = Pass->bHasExternalOutputs;
+			FRDGTextureRef Texture = PassState.Texture;
 
-			for (auto& PassState : Pass->TextureStates)
+			if (Texture->FirstBarrier == FRDGTexture::EFirstBarrier::ImmediateRequested)
 			{
-				FRDGTextureRef Texture = PassState.Texture;
-				auto& LastProducers = Texture->LastProducers;
+				check(Texture->bExternal);
+				Texture->FirstBarrier = FRDGTexture::EFirstBarrier::ImmediateConfirmed;
+				Texture->FirstPass = PassHandle;
 
-				for (uint32 Index = 0, Count = LastProducers.Num(); Index < Count; ++Index)
+				for (FRDGSubresourceState& SubresourceState : *Texture->State)
 				{
-					const auto& SubresourceState = PassState.State[Index];
-
-					if (SubresourceState.Access == ERHIAccess::Unknown)
-					{
-						continue;
-					}
+					SubresourceState.SetPass(ERHIPipeline::Graphics, PassHandle);
+				}
+			}
 
-					FRDGProducerState ProducerState;
-					ProducerState.Access = SubresourceState.Access;
-					ProducerState.PassHandle = PassHandle;
-					ProducerState.NoUAVBarrierHandle = SubresourceState.NoUAVBarrierFilter.GetUniqueHandle();
+		#if STATS
+			GRDGStatTextureReferenceCount += PassState.ReferenceCount;
+		#endif
 
-					AddCullingDependency(LastProducers[Index], ProducerState, PassPipeline);
+			for (int32 Index = 0; Index < PassState.State.Num(); ++Index)
+			{
+				if (PassState.State[Index].Access == ERHIAccess::Unknown)
+				{
+					continue;
 				}
 
-				bUntrackedOutputs |= Texture->bExternal;
+				MergeSubresourceStates(ERDGViewableResourceType::Texture, PassState.MergeState[Index], Texture->MergeState[Index], PassState.State[Index]);
 			}
+		}
 
-			for (auto& PassState : Pass->BufferStates)
+		for (auto& PassState : Pass->BufferStates)
+		{
+			FRDGBufferRef Buffer = PassState.Buffer;
+
+			if (Buffer->FirstBarrier == FRDGBuffer::EFirstBarrier::ImmediateRequested)
 			{
-				FRDGBufferRef Buffer = PassState.Buffer;
-				const auto& SubresourceState = PassState.State;
+				check(Buffer->bExternal);
+				Buffer->FirstBarrier = FRDGBuffer::EFirstBarrier::ImmediateConfirmed;
+				Buffer->FirstPass = PassHandle;
+				Buffer->State->SetPass(ERHIPipeline::Graphics, PassHandle);
+			}
 
-				FRDGProducerState ProducerState;
-				ProducerState.Access = SubresourceState.Access;
-				ProducerState.PassHandle = PassHandle;
-				ProducerState.NoUAVBarrierHandle = SubresourceState.NoUAVBarrierFilter.GetUniqueHandle();
+		#if STATS
+			GRDGStatBufferReferenceCount += PassState.ReferenceCount;
+		#endif
 
-				AddCullingDependency(Buffer->LastProducer, ProducerState, PassPipeline);
-				bUntrackedOutputs |= Buffer->bExternal;
-			}
+			MergeSubresourceStates(ERDGViewableResourceType::Buffer, PassState.MergeState, Buffer->MergeState, PassState.State);
+		}
+	}
+}
+
+void FRDGBuilder::Compile()
+{
+	SCOPE_CYCLE_COUNTER(STAT_RDG_CompileTime);
+	CSV_SCOPED_TIMING_STAT_EXCLUSIVE_CONDITIONAL(RDG_Compile, GRDGVerboseCSVStats != 0);
+
+	const FRDGPassHandle ProloguePassHandle = GetProloguePassHandle();
+	const FRDGPassHandle EpiloguePassHandle = GetEpiloguePassHandle();
+
+	const uint32 CompilePassCount = Passes.Num();
+
+	TransitionCreateQueue.Reserve(CompilePassCount);
+
+	const bool bCullPasses = GRDGCullPasses > 0;
+
+	if (bParallelSetupEnabled)
+	{
+		SetupPassQueue.Flush(TEXT("FRDGBuilder::SetupPassResources"), [this](FRDGPass* Pass) { SetupPassResources(Pass); });
+	}
+
+	if (bCullPasses || AsyncComputePassCount > 0)
+	{
+		SCOPED_NAMED_EVENT(PassDependencies, FColor::Emerald);
 
-			PassesOnAsyncCompute[PassHandle] = EnumHasAnyFlags(Pass->Flags, ERDGPassFlags::AsyncCompute);
-			Pass->bCulled = bCullPasses;
+		if (!bParallelSetupEnabled)
+		{
+			if (bCullPasses)
+			{
+				CullPassStack.Reserve(CompilePassCount);
+			}
 
-			if (bCullPasses && (bUntrackedOutputs || EnumHasAnyFlags(Pass->Flags, ERDGPassFlags::NeverCull)))
+			for (FRDGPassHandle PassHandle = ProloguePassHandle + 1; PassHandle < EpiloguePassHandle; ++PassHandle)
 			{
-				PassStack.Emplace(PassHandle);
+				SetupPassDependencies(Passes[PassHandle]);
 			}
 		}
 
+		const auto AddLastProducersToCullStack = [&](const FRDGProducerStatesByPipeline& LastProducers)
+		{
+			for (const FRDGProducerState& LastProducer : LastProducers)
+			{
+				if (LastProducer.Pass)
+				{
+					CullPassStack.Emplace(LastProducer.Pass->Handle);
+				}
+			}
+		};
+
+		// The last producer of a extracted or external resource is a cull graph root as it's not contained within the graph.
+
 		for (const FExtractedTexture& ExtractedTexture : ExtractedTextures)
 		{
 			FRDGTextureRef Texture = ExtractedTexture.Texture;
+
 			for (auto& LastProducer : Texture->LastProducers)
 			{
-				FRDGProducerState StateFinal;
-				StateFinal.Access = Texture->EpilogueAccess;
-				StateFinal.PassHandle = EpiloguePassHandle;
-
-				AddCullingDependency(LastProducer, StateFinal, ERHIPipeline::Graphics);
+				AddLastProducersToCullStack(LastProducer);
 			}
+
+			Texture->ReferenceCount++;
 		}
 
 		for (const FExtractedBuffer& ExtractedBuffer : ExtractedBuffers)
 		{
 			FRDGBufferRef Buffer = ExtractedBuffer.Buffer;
 
-			FRDGProducerState StateFinal;
-			StateFinal.Access = Buffer->EpilogueAccess;
-			StateFinal.PassHandle = EpiloguePassHandle;
+			AddLastProducersToCullStack(Buffer->LastProducer);
+
+			Buffer->ReferenceCount++;
+		}
+
+		for (const auto& Pair : ExternalTextures)
+		{
+			FRDGTexture* Texture = Pair.Value;
+
+			for (auto& LastProducer : Texture->LastProducers)
+			{
+				AddLastProducersToCullStack(LastProducer);
+			}
+		}
+
+		for (const auto& Pair : ExternalBuffers)
+		{
+			FRDGBuffer* Buffer = Pair.Value;
 
-			AddCullingDependency(Buffer->LastProducer, StateFinal, ERHIPipeline::Graphics);
+			AddLastProducersToCullStack(Buffer->LastProducer);
 		}
 	}
 
@@ -1022,7 +1147,7 @@ void FRDGBuilder::Compile()
 	{
 		SCOPED_NAMED_EVENT(PassCulling, FColor::Emerald);
 
-		PassStack.Emplace(EpiloguePassHandle);
+		CullPassStack.Emplace(EpiloguePassHandle);
 
 		// Mark the epilogue pass as culled so that it is traversed.
 		EpiloguePass->bCulled = 1;
@@ -1030,130 +1155,37 @@ void FRDGBuilder::Compile()
 		// Manually mark the prologue passes as not culled.
 		ProloguePass->bCulled = 0;
 
-		while (PassStack.Num())
+		while (CullPassStack.Num())
 		{
-			FRDGPass* Pass = Passes[PassStack.Pop()];
+			FRDGPass* Pass = Passes[CullPassStack.Pop()];
 
 			if (Pass->bCulled)
 			{
 				Pass->bCulled = 0;
-				PassStack.Append(Pass->Producers);
+
+				CullPassStack.Append(Pass->Producers);
 			}
 		}
-	}
-
-	// Walk the culled graph and compile barriers for each subresource. Certain transitions are redundant; read-to-read, for example.
-	// We can avoid them by traversing and merging compatible states together. The merging states removes a transition, but the merging
-	// heuristic is conservative and choosing not to merge doesn't necessarily mean a transition is performed. They are two distinct steps.
-	// Merged states track the first and last pass interval. Pass references are also accumulated onto each resource. This must happen
-	// after culling since culled passes can't contribute references.
-
-	{
-		SCOPED_NAMED_EVENT(CompileBarriers, FColor::Emerald);
 
 		for (FRDGPassHandle PassHandle = ProloguePassHandle + 1; PassHandle < EpiloguePassHandle; ++PassHandle)
 		{
 			FRDGPass* Pass = Passes[PassHandle];
 
-			if (Pass->bCulled || Pass->bEmptyParameters)
+			if (!Pass->bCulled)
 			{
 				continue;
 			}
 
-			const bool bAsyncComputePass = PassesOnAsyncCompute[PassHandle];
-
-			const ERHIPipeline PassPipeline = Pass->Pipeline;
-
-			const auto MergeSubresourceStates = [&](ERDGViewableResourceType ResourceType, FRDGSubresourceState*& PassMergeState, FRDGSubresourceState*& ResourceMergeState, const FRDGSubresourceState& PassState)
-			{
-				if (!ResourceMergeState || !FRDGSubresourceState::IsMergeAllowed(ResourceType, *ResourceMergeState, PassState))
-				{
-					// Cross-pipeline, non-mergable state changes require a new pass dependency for fencing purposes.
-					if (ResourceMergeState)
-					{
-						for (ERHIPipeline Pipeline : GetRHIPipelines())
-						{
-							if (Pipeline != PassPipeline && ResourceMergeState->LastPass[Pipeline].IsValid())
-							{
-								// Add a dependency from the other pipe to this pass to join back.
-								AddPassDependency(ResourceMergeState->LastPass[Pipeline], PassHandle);
-							}
-						}
-					}
-
-					// Allocate a new pending merge state and assign it to the pass state.
-					ResourceMergeState = AllocSubresource(PassState);
-				}
-				else
-				{
-					// Merge the pass state into the merged state.
-					ResourceMergeState->Access |= PassState.Access;
-
-					FRDGPassHandle& FirstPassHandle = ResourceMergeState->FirstPass[PassPipeline];
-
-					if (FirstPassHandle.IsNull())
-					{
-						FirstPassHandle = PassHandle;
-					}
-
-					ResourceMergeState->LastPass[PassPipeline] = PassHandle;
-				}
-
-				PassMergeState = ResourceMergeState;
-			};
-
+			// Subtract reference counts from culled passes that were added during pass setup.
+		
 			for (auto& PassState : Pass->TextureStates)
 			{
-				FRDGTextureRef Texture = PassState.Texture;
-				Texture->ReferenceCount += PassState.ReferenceCount;
-				Texture->bUsedByAsyncComputePass |= bAsyncComputePass;
-				Texture->bCulled = false;
-
-				if (Texture->FirstBarrier == FRDGTexture::EFirstBarrier::ImmediateRequested)
-				{
-					Texture->FirstBarrier = FRDGTexture::EFirstBarrier::ImmediateConfirmed;
-					Texture->FirstPass = PassHandle;
-
-					for (FRDGSubresourceState& SubresourceState : *Texture->State)
-					{
-						SubresourceState.SetPass(ERHIPipeline::Graphics, PassHandle);
-					}
-				}
-
-			#if STATS
-				GRDGStatTextureReferenceCount += PassState.ReferenceCount;
-			#endif
-
-				for (int32 Index = 0; Index < PassState.State.Num(); ++Index)
-				{
-					if (PassState.State[Index].Access == ERHIAccess::Unknown)
-					{
-						continue;
-					}
-
-					MergeSubresourceStates(ERDGViewableResourceType::Texture, PassState.MergeState[Index], Texture->MergeState[Index], PassState.State[Index]);
-				}
+				PassState.Texture->ReferenceCount -= PassState.ReferenceCount;
 			}
 
 			for (auto& PassState : Pass->BufferStates)
 			{
-				FRDGBufferRef Buffer = PassState.Buffer;
-				Buffer->ReferenceCount += PassState.ReferenceCount;
-				Buffer->bUsedByAsyncComputePass |= bAsyncComputePass;
-				Buffer->bCulled = false;
-
-				if (Buffer->FirstBarrier == FRDGBuffer::EFirstBarrier::ImmediateRequested)
-				{
-					Buffer->FirstBarrier = FRDGBuffer::EFirstBarrier::ImmediateConfirmed;
-					Buffer->FirstPass = PassHandle;
-					Buffer->State->SetPass(ERHIPipeline::Graphics, PassHandle);
-				}
-
-			#if STATS
-				GRDGStatBufferReferenceCount += PassState.ReferenceCount;
-			#endif
-
-				MergeSubresourceStates(ERDGViewableResourceType::Buffer, PassState.MergeState, Buffer->MergeState, PassState.State);
+				PassState.Buffer->ReferenceCount -= PassState.ReferenceCount;
 			}
 		}
 	}
@@ -1298,6 +1330,8 @@ void FRDGBuilder::Compile()
 	{
 		SCOPED_NAMED_EVENT(AsyncComputeFences, FColor::Emerald);
 
+		FRDGPassBitArray PassesOnAsyncCompute(false, CompilePassCount);
+
 		// Traverse the active passes in execution order to find latest cross-pipeline producer and the earliest
 		// cross-pipeline consumer for each pass. This helps narrow the search space later when building async
 		// compute overlap regions.
@@ -1319,6 +1353,8 @@ void FRDGBuilder::Compile()
 				continue;
 			}
 
+			PassesOnAsyncCompute[PassHandle] = Pass->IsAsyncCompute();
+
 			for (FRDGPassHandle ProducerHandle : Pass->GetProducers())
 			{
 				const FRDGPassHandle ConsumerHandle = PassHandle;
@@ -1516,6 +1552,61 @@ void FRDGBuilder::Compile()
 	}
 }
 
+///////////////////////////////////////////////////////////////////////////////////////////////////
+
+template <typename LambdaType>
+void FRDGBuilder::FPassQueue::Flush(UE::Tasks::FPipe& Pipe, const TCHAR* Name, LambdaType&& Lambda)
+{
+	if (LastTask.IsCompleted())
+	{
+		LastTask = Pipe.Launch(Name, [this, Lambda, Name]
+		{
+			SCOPED_NAMED_EVENT_TCHAR(Name, FColor::Magenta);
+			while (FRDGPass* Pass = Queue.Pop())
+			{
+				Lambda(Pass);
+			}
+		});
+	}
+}
+
+template <typename LambdaType>
+void FRDGBuilder::FPassQueue::Flush(const TCHAR* Name, LambdaType&& Lambda)
+{
+	SCOPED_NAMED_EVENT_TCHAR(Name, FColor::Magenta);
+	LastTask.Wait();
+	while (FRDGPass* Pass = Queue.Pop())
+	{
+		Lambda(Pass);
+	}
+}
+
+
+void FRDGBuilder::FlushSetupQueue()
+{
+	if (bParallelSetupEnabled)
+	{
+		SetupPassQueue.Flush(CompilePipe, TEXT("FRDGBuilder::SetupPassResources"), [this](FRDGPass* Pass) { SetupPassResources(Pass); });
+	}
+}
+
+template <typename LambdaType>
+UE::Tasks::FTask FRDGBuilder::LaunchCompileTask(const TCHAR* Name, bool bCondition, LambdaType&& Lambda)
+{
+	if (bCondition)
+	{
+		SCOPED_NAMED_EVENT_TCHAR(Name, FColor::Magenta);
+		return CompilePipe.Launch(Name, [Lambda] { Lambda(); }, LowLevelTasks::ETaskPriority::High);
+	}
+	else
+	{
+		Lambda();
+		return {};
+	}
+}
+
+///////////////////////////////////////////////////////////////////////////////////////////////////
+
 void FRDGBuilder::Execute()
 {
 	CSV_SCOPED_TIMING_STAT_EXCLUSIVE(RDG);
@@ -1538,16 +1629,22 @@ void FRDGBuilder::Execute()
 	const FRDGPassHandle ProloguePassHandle = GetProloguePassHandle();
 	const FRDGPassHandle EpiloguePassHandle = GetEpiloguePassHandle();
 
-	FGraphEventArray AsyncCompileEvents;
+	FGraphEventRef SubmitBufferUploadsTask;
+	UE::Tasks::FTask CompilePassBarriersTask;
+	UE::Tasks::FTask CreateUniformBuffersTask;
+	UE::Tasks::FTask SetupParallelExecuteTask;
 
 	IF_RDG_ENABLE_DEBUG(UserValidation.ValidateExecuteBegin());
 	IF_RDG_ENABLE_DEBUG(GRDGAllowRHIAccess = true);
 
 	if (!IsImmediateMode())
 	{
-		BeginFlushResourcesRHI();
+		UE::Tasks::FTask CompileTask = LaunchCompileTask(TEXT("FRDGBuilder::Compile"), bParallelSetupEnabled, [this] { Compile(); });
 
-		Compile();
+		CompilePassBarriersTask = LaunchCompileTask(TEXT("FRDGBuilder::CompilePassBarriers"), bParallelSetupEnabled, [this] { CompilePassBarriers(); });
+
+		BeginFlushResourcesRHI();
+		PrepareBufferUploads();
 
 		GPUScopeStacks.ReserveOps(Passes.Num());
 		IF_RDG_CPU_SCOPES(CPUScopeStacks.ReserveOps());
@@ -1559,15 +1656,10 @@ void FRDGBuilder::Execute()
 		#endif
 
 			// Parallel execute setup can be done off the render thread and synced prior to dispatch.
-			AsyncCompileEvents.Emplace(FFunctionGraphTask::CreateAndDispatchWhenReady(
-				[this](ENamedThreads::Type, const FGraphEventRef&)
-			{
-				SetupParallelExecute();
-
-			}, TStatId(), nullptr, ENamedThreads::AnyHiPriThreadHiPriTask));
+			SetupParallelExecuteTask = LaunchCompileTask(TEXT("FRDGBuilder::SetupParallelExecute"), bParallelExecuteEnabled, [this] { SetupParallelExecute(); });
 		}
 
-		SetupBufferUploads();
+		CompileTask.Wait();
 
 		{
 			SCOPE_CYCLE_COUNTER(STAT_RDG_CollectResourcesTime);
@@ -1592,9 +1684,8 @@ void FRDGBuilder::Execute()
 			{
 				FRDGTexture* Texture = Pair.Value;
 
-				if (Texture->bCulled)
+				if (Texture->IsCulled())
 				{
-					check(Texture->ReferenceCount == 0);
 					EndResourceRHI(ProloguePassHandle, Texture, 0);
 				}
 			}
@@ -1603,9 +1694,8 @@ void FRDGBuilder::Execute()
 			{
 				FRDGBuffer* Buffer = Pair.Value;
 
-				if (Buffer->bCulled)
+				if (Buffer->IsCulled())
 				{
-					check(Buffer->ReferenceCount == 0);
 					EndResourceRHI(ProloguePassHandle, Buffer, 0);
 				}
 			}
@@ -1651,20 +1741,26 @@ void FRDGBuilder::Execute()
 			}
 		}
 
-		CreateUniformBuffers(bParallelExecuteEnabled ? &AsyncCompileEvents : nullptr);
+		// We have to wait until after view creation to launch uploads because we can't lock / unlock while creating views simultaneously.
+		SubmitBufferUploadsTask = SubmitBufferUploads();
+
+		// Uniform buffer creation depends on view creation.
+		CreateUniformBuffersTask = CreateUniformBuffers();
 
 		{
 			SCOPED_NAMED_EVENT_TEXT("FRDGBuilder::CollectBarriers", FColor::Magenta);
 			SCOPE_CYCLE_COUNTER(STAT_RDG_CollectBarriersTime);
 			CSV_SCOPED_TIMING_STAT_EXCLUSIVE_CONDITIONAL(RDG_CollectBarriers, GRDGVerboseCSVStats != 0);
 
+			CompilePassBarriersTask.Wait();
+
 			for (FRDGPassHandle PassHandle = ProloguePassHandle + 1; PassHandle < EpiloguePassHandle; ++PassHandle)
 			{
 				FRDGPass* Pass = Passes[PassHandle];
 
 				if (!Pass->bCulled && !Pass->bEmptyParameters)
 				{
-					CollectPassBarriers(Pass, PassHandle);
+					CollectPassBarriers(Pass);
 				}
 			}
 		}
@@ -1686,47 +1782,44 @@ void FRDGBuilder::Execute()
 		});
 	}
 
-	CreatePassBarriers(bParallelExecuteEnabled ? &AsyncCompileEvents : nullptr);
+	const ENamedThreads::Type RenderThread = ENamedThreads::GetRenderThread_Local();
 
-	if (!ParallelSetupEvents.IsEmpty())
+	CreatePassBarriers([&]
 	{
-		FTaskGraphInterface::Get().WaitUntilTasksComplete(ParallelSetupEvents);
-	}
+		if (!ParallelSetupEvents.IsEmpty())
+		{
+			FTaskGraphInterface::Get().WaitUntilTasksComplete(ParallelSetupEvents, RenderThread);
+		}
 
-	EndFlushResourcesRHI();
+		// Process RHI thread flush before helping with barrier compilation on the render thread.
+		EndFlushResourcesRHI();
+	});
 
-	FGraphEventRef DispatchParallelExecuteEvent;
+	CreateUniformBuffersTask.Wait();
 
-	if (bParallelExecuteEnabled)
+	if (SubmitBufferUploadsTask)
 	{
-		DispatchParallelExecuteEvent = FFunctionGraphTask::CreateAndDispatchWhenReady(
-			[this, &RHICmdContext = RHICmdList.GetContext()](ENamedThreads::Type, const FGraphEventRef&)
-		{
-			DispatchParallelExecute(&RHICmdContext);
-
-		}, TStatId(), &AsyncCompileEvents, ENamedThreads::AnyHiPriThreadHiPriTask);
+		SubmitBufferUploadsTask->Wait(RenderThread);
 	}
 
-	SubmitBufferUploads();
+	UE::Tasks::FTask ParallelExecuteTask;
+
+	if (bParallelExecuteEnabled)
+	{
+		SetupParallelExecuteTask.Wait();
+		ParallelExecuteTask = CompilePipe.Launch(TEXT("DispatchParallelExecute"), [this, &RHICmdContext = RHICmdList.GetContext()] { DispatchParallelExecute(&RHICmdContext); }, LowLevelTasks::ETaskPriority::High);
+	}
 
 	IF_RDG_ENABLE_DEBUG(GRDGAllowRHIAccess = bParallelExecuteEnabled);
 	IF_RDG_ENABLE_TRACE(Trace.OutputGraphBegin());
 
-	const ENamedThreads::Type RenderThread = ENamedThreads::GetRenderThread_Local();
-
 	if (!IsImmediateMode())
 	{
 		SCOPED_NAMED_EVENT_TEXT("FRDGBuilder::ExecutePasses", FColor::Magenta);
 		SCOPE_CYCLE_COUNTER(STAT_RDG_ExecuteTime);
 		CSV_SCOPED_TIMING_STAT_EXCLUSIVE(RDG_Execute);
 
-		// Wait on all async compilation tasks before executing any passes.
-		if (!AsyncCompileEvents.IsEmpty())
-		{
-			FTaskGraphInterface::Get().WaitUntilTasksComplete(AsyncCompileEvents, RenderThread);
-		}
-
-		for (FRDGPassHandle PassHandle = ProloguePassHandle; PassHandle <= EpiloguePassHandle; ++PassHandle)
+		for (FRDGPassHandle PassHandle = ProloguePassHandle; PassHandle <= EpiloguePassHandle; ++PassHandle)
 		{
 			FRDGPass* Pass = Passes[PassHandle];
 
@@ -1793,12 +1886,6 @@ void FRDGBuilder::Execute()
 		ExecutePass(EpiloguePass, RHICmdList);
 	}
 
-	// Wait for the parallel dispatch task before attempting to wait on the execute event array (the former mutates the array).
-	if (DispatchParallelExecuteEvent)
-	{
-		DispatchParallelExecuteEvent->Wait(RenderThread);
-	}
-
 	RHICmdList.SetStaticUniformBuffers({});
 
 #if WITH_MGPU
@@ -1823,6 +1910,9 @@ void FRDGBuilder::Execute()
 
 	RHICmdList.SetTrackedAccess(EpilogueResourceAccesses);
 
+	// Wait for the parallel dispatch task before attempting to wait on the execute event array (the former mutates the array).
+	ParallelExecuteTask.Wait();
+
 	// Wait on the actual parallel execute tasks in the Execute call. When draining is okay to let them overlap with other graph setup.
 	// This also needs to be done before extraction of external resources to be consistent with non-parallel rendering.
 	if (!ParallelExecuteEvents.IsEmpty())
@@ -1870,11 +1960,150 @@ void FRDGBuilder::Execute()
 
 ///////////////////////////////////////////////////////////////////////////////////////////////////
 
-FRDGPass* FRDGBuilder::SetupPass(FRDGPass* Pass)
+void FRDGBuilder::MarkResourcesAsProduced(FRDGPass* Pass)
 {
-	IF_RDG_ENABLE_DEBUG(UserValidation.ValidateAddPass(Pass, AuxiliaryPasses.IsActive()));
-	CSV_SCOPED_TIMING_STAT_EXCLUSIVE_CONDITIONAL(RDGBuilder_SetupPass, GRDGVerboseCSVStats != 0);
+	const auto MarkAsProduced = [&](FRDGViewableResource* Resource)
+	{
+		Resource->bProduced = true;
+	};
+
+	const auto MarkAsProducedIfWritable = [&](FRDGViewableResource* Resource, ERHIAccess Access)
+	{
+		if (IsWritableAccess(Access))
+		{
+			Resource->bProduced = true;
+		}
+	};
+
+	Pass->GetParameters().Enumerate([&](FRDGParameter Parameter)
+	{
+		switch (Parameter.GetType())
+		{
+		case UBMT_RDG_TEXTURE_UAV:
+			if (FRDGTextureUAV* UAV = Parameter.GetAsTextureUAV())
+			{
+				MarkAsProduced(UAV->GetParent());
+			}
+		break;
+		case UBMT_RDG_BUFFER_UAV:
+			if (FRDGBufferUAV* UAV = Parameter.GetAsBufferUAV())
+			{
+				MarkAsProduced(UAV->GetParent());
+			}
+		break;
+		case UBMT_RDG_TEXTURE_ACCESS:
+		{
+			if (FRDGTextureAccess TextureAccess = Parameter.GetAsTextureAccess())
+			{
+				MarkAsProducedIfWritable(TextureAccess.GetTexture(), TextureAccess.GetAccess());
+			}
+		}
+		break;
+		case UBMT_RDG_TEXTURE_ACCESS_ARRAY:
+		{
+			const FRDGTextureAccessArray& TextureAccessArray = Parameter.GetAsTextureAccessArray();
 
+			for (FRDGTextureAccess TextureAccess : TextureAccessArray)
+			{
+				MarkAsProducedIfWritable(TextureAccess.GetTexture(), TextureAccess.GetAccess());
+			}
+		}
+		break;
+		case UBMT_RDG_BUFFER_ACCESS:
+			if (FRDGBufferAccess BufferAccess = Parameter.GetAsBufferAccess())
+			{
+				MarkAsProducedIfWritable(BufferAccess.GetBuffer(), BufferAccess.GetAccess());
+			}
+		break;
+		case UBMT_RDG_BUFFER_ACCESS_ARRAY:
+		{
+			const FRDGBufferAccessArray& BufferAccessArray = Parameter.GetAsBufferAccessArray();
+
+			for (FRDGBufferAccess BufferAccess : BufferAccessArray)
+			{
+				MarkAsProducedIfWritable(BufferAccess.GetBuffer(), BufferAccess.GetAccess());
+			}
+		}
+		break;
+		case UBMT_RENDER_TARGET_BINDING_SLOTS:
+		{
+			const FRenderTargetBindingSlots& RenderTargets = Parameter.GetAsRenderTargetBindingSlots();
+
+			RenderTargets.Enumerate([&](FRenderTargetBinding RenderTarget)
+			{
+				MarkAsProduced(RenderTarget.GetTexture());
+
+				if (FRDGTexture* ResolveTexture = RenderTarget.GetResolveTexture())
+				{
+					MarkAsProduced(ResolveTexture);
+				}
+			});
+
+			const FDepthStencilBinding& DepthStencil = RenderTargets.DepthStencil;
+
+			if (DepthStencil.GetDepthStencilAccess().IsAnyWrite())
+			{
+				MarkAsProduced(DepthStencil.GetTexture());
+			}
+		}
+		break;
+		}
+	});
+}
+
+void FRDGBuilder::SetupPassDependencies(FRDGPass* Pass)
+{
+	for (auto& PassState : Pass->TextureStates)
+	{
+		FRDGTextureRef Texture = PassState.Texture;
+		auto& LastProducers = Texture->LastProducers;
+
+		Texture->ReferenceCount += PassState.ReferenceCount;
+
+		for (uint32 Index = 0, Count = LastProducers.Num(); Index < Count; ++Index)
+		{
+			const auto& SubresourceState = PassState.State[Index];
+
+			if (SubresourceState.Access == ERHIAccess::Unknown)
+			{
+				continue;
+			}
+
+			FRDGProducerState ProducerState;
+			ProducerState.Pass = Pass;
+			ProducerState.Access = SubresourceState.Access;
+			ProducerState.NoUAVBarrierHandle = SubresourceState.NoUAVBarrierFilter.GetUniqueHandle();
+
+			AddCullingDependency(LastProducers[Index], ProducerState, Pass->Pipeline);
+		}
+	}
+
+	for (auto& PassState : Pass->BufferStates)
+	{
+		FRDGBufferRef Buffer = PassState.Buffer;
+		const auto& SubresourceState = PassState.State;
+
+		Buffer->ReferenceCount += PassState.ReferenceCount;
+
+		FRDGProducerState ProducerState;
+		ProducerState.Pass = Pass;
+		ProducerState.Access = SubresourceState.Access;
+		ProducerState.NoUAVBarrierHandle = SubresourceState.NoUAVBarrierFilter.GetUniqueHandle();
+
+		AddCullingDependency(Buffer->LastProducer, ProducerState, Pass->Pipeline);
+	}
+
+	const bool bCullPasses = GRDGCullPasses > 0;
+	Pass->bCulled = bCullPasses;
+
+	if (bCullPasses && (Pass->bHasExternalOutputs || EnumHasAnyFlags(Pass->Flags, ERDGPassFlags::NeverCull)))
+	{
+		CullPassStack.Emplace(Pass->Handle);
+	}
+}
+
+void FRDGBuilder::SetupPassResources(FRDGPass* Pass)
+{
 	const FRDGParameterStruct PassParameters = Pass->GetParameters();
 	const FRDGPassHandle PassHandle = Pass->Handle;
 	const ERDGPassFlags PassFlags = Pass->Flags;
@@ -1897,7 +2126,7 @@ FRDGPass* FRDGBuilder::SetupPass(FRDGPass* Pass)
 	{
 		TryAddView(TextureView);
 
-		if (Texture->AccessModeState.IsExternalAccess())
+		if (Texture->AccessModeState.IsExternalAccess() && !Pass->bExternalAccessPass)
 		{
 			// Resources in external access mode are expected to remain in the same state and are ignored by the graph.
 			// As only External | Extracted resources can be set as external by the user, the graph doesn't need to track
@@ -1935,9 +2164,16 @@ FRDGPass* FRDGBuilder::SetupPass(FRDGPass* Pass)
 			State.SetPass(PassPipeline, PassHandle);
 		});
 
-		const bool bWritableAccess = IsWritableAccess(Access);
-		bRenderPassOnlyWrites &= (!bWritableAccess || EnumHasAnyFlags(AccessFlags, ERDGTextureAccessFlags::RenderTarget));
-		Texture->bProduced |= bWritableAccess;
+		if (IsWritableAccess(Access))
+		{
+			bRenderPassOnlyWrites &= EnumHasAnyFlags(AccessFlags, ERDGTextureAccessFlags::RenderTarget);
+
+			// When running in parallel this is set via MarkResourcesAsProduced. We also can't touch this as its a bitfield and not atomic.
+			if (!bParallelSetupEnabled)
+			{
+				Texture->bProduced = true;
+			}
+		}
 	});
 
 	Pass->BufferStates.Reserve(PassParameters.GetBufferParameterCount());
@@ -1945,7 +2181,7 @@ FRDGPass* FRDGBuilder::SetupPass(FRDGPass* Pass)
 	{
 		TryAddView(BufferView);
 
-		if (Buffer->AccessModeState.IsExternalAccess())
+		if (Buffer->AccessModeState.IsExternalAccess() && !Pass->bExternalAccessPass)
 		{
 			// Resources in external access mode are expected to remain in the same state and are ignored by the graph.
 			// As only External | Extracted resources can be set as external by the user, the graph doesn't need to track
@@ -1954,6 +2190,11 @@ FRDGPass* FRDGBuilder::SetupPass(FRDGPass* Pass)
 			return;
 		}
 
+		if (!FCString::Stricmp(Buffer->Name, TEXT("DefaultBuffer")))
+		{
+			UE_DEBUG_BREAK();
+		}
+
 		const FRDGViewHandle NoUAVBarrierHandle = GetHandleIfNoUAVBarrier(BufferView);
 
 		FRDGPass::FBufferState* PassState;
@@ -1977,68 +2218,45 @@ FRDGPass* FRDGBuilder::SetupPass(FRDGPass* Pass)
 		PassState->State.NoUAVBarrierFilter.AddHandle(NoUAVBarrierHandle);
 		PassState->State.SetPass(PassPipeline, PassHandle);
 
-		const bool bWritableAccess = IsWritableAccess(Access);
-		bRenderPassOnlyWrites &= !bWritableAccess;
-		Buffer->bProduced |= bWritableAccess;
+		if (IsWritableAccess(Access))
+		{
+			bRenderPassOnlyWrites = false;
+
+			// When running in parallel this is set via MarkResourcesAsProduced. We also can't touch this as its a bitfield and not atomic.
+			if (!bParallelSetupEnabled)
+			{
+				Buffer->bProduced = true;
+			}
+		}
 	});
 
+	Pass->bEmptyParameters = !Pass->TextureStates.Num() && !Pass->BufferStates.Num();
+	Pass->bRenderPassOnlyWrites = bRenderPassOnlyWrites;
+	Pass->bHasExternalOutputs = PassParameters.HasExternalOutputs();
+
 	Pass->UniformBuffers.Reserve(PassParameters.GetUniformBufferParameterCount());
 	PassParameters.EnumerateUniformBuffers([&](FRDGUniformBufferBinding UniformBuffer)
 	{
 		Pass->UniformBuffers.Emplace(UniformBuffer.GetUniformBuffer()->Handle);
 	});
 
-	Pass->bRenderPassOnlyWrites = bRenderPassOnlyWrites;
-	Pass->bHasExternalOutputs = PassParameters.HasExternalOutputs();
-
-	const bool bEmptyParameters = !Pass->TextureStates.Num() && !Pass->BufferStates.Num();
-	SetupPassInternal(Pass, PassHandle, PassPipeline, bEmptyParameters);
-	return Pass;
-}
-
-FRDGPass* FRDGBuilder::SetupEmptyPass(FRDGPass* Pass)
-{
-	const bool bEmptyParameters = true;
-	SetupPassInternal(Pass, Pass->Handle, Pass->Pipeline, bEmptyParameters);
-	return Pass;
-}
-
-void FRDGBuilder::CompilePassOps(FRDGPass* Pass)
-{
-#if WITH_MGPU
-	if (!bWaitedForTemporalEffect && NameForTemporalEffect != NAME_None && Pass->Pipeline == ERHIPipeline::Graphics)
+	if (bParallelSetupEnabled)
 	{
-		bWaitedForTemporalEffect = true;
-		Pass->bWaitForTemporalEffect = 1;
-	}
-
-	FRHIGPUMask GPUMask = Pass->GPUMask;
-#else
-	FRHIGPUMask GPUMask = FRHIGPUMask::All();
-#endif
+		SetupPassDependencies(Pass);
 
-#if RDG_CMDLIST_STATS
-	if (CommandListStatState != Pass->CommandListStat && !Pass->bSentinel)
-	{
-		CommandListStatState = Pass->CommandListStat;
-		Pass->bSetCommandListStat = 1;
+		for (FRDGPass::FExternalAccessOp Op : Pass->ExternalAccessOps)
+		{
+			Op.Resource->AccessModeState.ActiveMode = Op.Mode;
+		}
 	}
-#endif
-
-#if RDG_CPU_SCOPES
-	Pass->CPUScopeOps = CPUScopeStacks.CompilePassPrologue(Pass);
-#endif
-
-	Pass->GPUScopeOpsPrologue = GPUScopeStacks.CompilePassPrologue(Pass, GPUMask);
-	Pass->GPUScopeOpsEpilogue = GPUScopeStacks.CompilePassEpilogue(Pass);
 }
 
-void FRDGBuilder::SetupPassInternal(FRDGPass* Pass, FRDGPassHandle PassHandle, ERHIPipeline PassPipeline, bool bEmptyParameters)
+void FRDGBuilder::SetupPassInternals(FRDGPass* Pass)
 {
-	check(Pass->Handle == PassHandle);
-	check(Pass->Pipeline == PassPipeline);
+	const FRDGPassHandle PassHandle = Pass->Handle;
+	const ERDGPassFlags PassFlags = Pass->Flags;
+	const ERHIPipeline PassPipeline = Pass->Pipeline;
 
-	Pass->bEmptyParameters = bEmptyParameters;
 	Pass->GraphicsForkPass = PassHandle;
 	Pass->GraphicsJoinPass = PassHandle;
 	Pass->PrologueBarrierPass = PassHandle;
@@ -2050,8 +2268,8 @@ void FRDGBuilder::SetupPassInternal(FRDGPass* Pass, FRDGPassHandle PassHandle, E
 		Pass->ResourcesToEnd.Add(Pass);
 	}
 
-	AsyncComputePassCount += EnumHasAnyFlags(Pass->Flags, ERDGPassFlags::AsyncCompute) ? 1 : 0;
-	RasterPassCount       += EnumHasAnyFlags(Pass->Flags, ERDGPassFlags::Raster)       ? 1 : 0;
+	AsyncComputePassCount += EnumHasAnyFlags(PassFlags, ERDGPassFlags::AsyncCompute) ? 1 : 0;
+	RasterPassCount += EnumHasAnyFlags(PassFlags, ERDGPassFlags::Raster) ? 1 : 0;
 
 #if WITH_MGPU
 	Pass->GPUMask = RHICmdList.GetGPUMask();
@@ -2061,7 +2279,6 @@ void FRDGBuilder::SetupPassInternal(FRDGPass* Pass, FRDGPassHandle PassHandle, E
 	Pass->CommandListStat = CommandListStatScope;
 
 	GRDGStatPassCount++;
-	GRDGStatPassWithParameterCount += !bEmptyParameters ? 1 : 0;
 #endif
 
 	IF_RDG_CPU_SCOPES(Pass->CPUScopes = CPUScopeStacks.GetCurrentScopes());
@@ -2077,7 +2294,10 @@ void FRDGBuilder::SetupPassInternal(FRDGPass* Pass, FRDGPassHandle PassHandle, E
 		Pass->FullPathIfDebug = Scope->GetPath(Pass->Name);
 	}
 #endif
+}
 
+void FRDGBuilder::SetupAuxiliaryPasses(FRDGPass* Pass)
+{
 	if (IsImmediateMode() && !Pass->bSentinel)
 	{
 		SCOPED_NAMED_EVENT(FRDGBuilder_ExecutePass, FColor::Emerald);
@@ -2102,16 +2322,16 @@ void FRDGBuilder::SetupPassInternal(FRDGPass* Pass, FRDGPassHandle PassHandle, E
 			PassState.MergeState = &PassState.State;
 		}
 
-		check(!EnumHasAnyFlags(PassPipeline, ERHIPipeline::AsyncCompute));
+		check(!EnumHasAnyFlags(Pass->Pipeline, ERHIPipeline::AsyncCompute));
 		check(ParallelSetupEvents.IsEmpty());
 
-		SetupBufferUploads();
+		PrepareBufferUploads();
 		SubmitBufferUploads();
 		CompilePassOps(Pass);
-		BeginResourcesRHI(Pass, PassHandle);
-		CreateUniformBuffers(nullptr);
-		CollectPassBarriers(Pass, PassHandle);
-		CreatePassBarriers(nullptr);
+		BeginResourcesRHI(Pass, Pass->Handle);
+		CreateUniformBuffers();
+		CollectPassBarriers(Pass);
+		CreatePassBarriers([] {});
 		FTaskGraphInterface::Get().WaitUntilTasksComplete(ParallelSetupEvents);
 		ParallelSetupEvents.Reset();
 		ExecutePass(Pass, RHICmdList);
@@ -2124,66 +2344,157 @@ void FRDGBuilder::SetupPassInternal(FRDGPass* Pass, FRDGPassHandle PassHandle, E
 #endif
 }
 
+FRDGPass* FRDGBuilder::SetupParameterPass(FRDGPass* Pass)
+{
+	IF_RDG_ENABLE_DEBUG(UserValidation.ValidateAddPass(Pass, AuxiliaryPasses.IsActive()));
+	CSV_SCOPED_TIMING_STAT_EXCLUSIVE_CONDITIONAL(RDGBuilder_SetupPass, GRDGVerboseCSVStats != 0);
+
+	SetupPassInternals(Pass);
+
+	if (bParallelSetupEnabled)
+	{
+		MarkResourcesAsProduced(Pass);
+		SetupPassQueue.Push(Pass);
+	}
+	else
+	{
+		SetupPassResources(Pass);
+	}
+
+	SetupAuxiliaryPasses(Pass);
+	return Pass;
+}
+
+FRDGPass* FRDGBuilder::SetupEmptyPass(FRDGPass* Pass)
+{
+	IF_RDG_ENABLE_DEBUG(UserValidation.ValidateAddPass(Pass, AuxiliaryPasses.IsActive()));
+	CSV_SCOPED_TIMING_STAT_EXCLUSIVE_CONDITIONAL(RDGBuilder_SetupPass, GRDGVerboseCSVStats != 0);
+
+	Pass->bEmptyParameters = true;
+	SetupPassInternals(Pass);
+	SetupAuxiliaryPasses(Pass);
+	return Pass;
+}
+
+void FRDGBuilder::CompilePassOps(FRDGPass* Pass)
+{
+#if WITH_MGPU
+	if (!bWaitedForTemporalEffect && NameForTemporalEffect != NAME_None && Pass->Pipeline == ERHIPipeline::Graphics)
+	{
+		bWaitedForTemporalEffect = true;
+		Pass->bWaitForTemporalEffect = 1;
+	}
+
+	FRHIGPUMask GPUMask = Pass->GPUMask;
+#else
+	FRHIGPUMask GPUMask = FRHIGPUMask::All();
+#endif
+
+#if RDG_CMDLIST_STATS
+	if (CommandListStatState != Pass->CommandListStat && !Pass->bSentinel)
+	{
+		CommandListStatState = Pass->CommandListStat;
+		Pass->bSetCommandListStat = 1;
+	}
+#endif
+
+#if RDG_CPU_SCOPES
+	Pass->CPUScopeOps = CPUScopeStacks.CompilePassPrologue(Pass);
+#endif
+
+	Pass->GPUScopeOpsPrologue = GPUScopeStacks.CompilePassPrologue(Pass, GPUMask);
+	Pass->GPUScopeOpsEpilogue = GPUScopeStacks.CompilePassEpilogue(Pass);
+}
+
 ///////////////////////////////////////////////////////////////////////////////////////////////////
 
-void FRDGBuilder::SetupBufferUploads()
+void FRDGBuilder::PrepareBufferUploads()
 {
 	SCOPED_NAMED_EVENT_TEXT("FRDGBuilder::PrepareBufferUploads", FColor::Magenta);
 
 	for (FUploadedBuffer& UploadedBuffer : UploadedBuffers)
 	{
-		if (UploadedBuffer.bUseDataCallbacks)
-		{
-			UploadedBuffer.Data = UploadedBuffer.DataCallback();
-			UploadedBuffer.DataSize = UploadedBuffer.DataSizeCallback();
-		}
+		FRDGBuffer* Buffer = UploadedBuffer.Buffer;
 
-		if (UploadedBuffer.Data && UploadedBuffer.DataSize)
+		if (!Buffer->HasRHI())
 		{
-			ConvertToExternalBuffer(UploadedBuffer.Buffer);
-			check(UploadedBuffer.DataSize <= UploadedBuffer.Buffer->Desc.GetSize());
+			BeginResourceRHI(GetProloguePassHandle(), Buffer);
 		}
+
+		check(UploadedBuffer.DataSize <= Buffer->Desc.GetSize());
 	}
 }
 
-void FRDGBuilder::SubmitBufferUploads()
+FGraphEventRef FRDGBuilder::SubmitBufferUploads()
 {
-	SCOPED_NAMED_EVENT_TEXT("FRDGBuilder::SubmitBufferUploads", FColor::Magenta);
-
-	for (const FUploadedBuffer& UploadedBuffer : UploadedBuffers)
+	const auto SubmitUploadsLambda = [this](FRHICommandList& RHICmdListUpload)
 	{
-		if (UploadedBuffer.Data && UploadedBuffer.DataSize)
+		SCOPED_NAMED_EVENT_TEXT("FRDGBuilder::SubmitBufferUploads", FColor::Magenta);
+
+		for (FUploadedBuffer& UploadedBuffer : UploadedBuffers)
 		{
-#if PLATFORM_NEEDS_GPU_UAV_RESOURCE_INIT_WORKAROUND
-			if (UploadedBuffer.Buffer->bUAVAccessed)
-			{
-				FRHIResourceCreateInfo CreateInfo(UploadedBuffer.Buffer->Name);
-				FBufferRHIRef TempBuffer = RHICreateVertexBuffer(UploadedBuffer.DataSize, BUF_Static | BUF_ShaderResource, CreateInfo);
-				void* DestPtr = RHICmdList.LockBuffer(TempBuffer, 0, UploadedBuffer.DataSize, RLM_WriteOnly);
-				FMemory::Memcpy(DestPtr, UploadedBuffer.Data, UploadedBuffer.DataSize);
-				RHICmdList.UnlockBuffer(TempBuffer);
-				RHICmdList.Transition(
-				{
-					FRHITransitionInfo(TempBuffer, ERHIAccess::Unknown, ERHIAccess::CopySrc | ERHIAccess::SRVMask),
-					FRHITransitionInfo(UploadedBuffer.Buffer->GetRHI(), ERHIAccess::Unknown, ERHIAccess::CopyDest)
-				});
-				RHICmdList.CopyBufferRegion(UploadedBuffer.Buffer->GetRHI(), 0, TempBuffer, 0, UploadedBuffer.DataSize);
-			}
-			else
-#endif
+			if (UploadedBuffer.bUseDataCallbacks)
 			{
-				void* DestPtr = RHICmdList.LockBuffer(UploadedBuffer.Buffer->GetRHIUnchecked(), 0, UploadedBuffer.DataSize, RLM_WriteOnly);
-				FMemory::Memcpy(DestPtr, UploadedBuffer.Data, UploadedBuffer.DataSize);
-				RHICmdList.UnlockBuffer(UploadedBuffer.Buffer->GetRHIUnchecked());
+				UploadedBuffer.Data = UploadedBuffer.DataCallback();
+				UploadedBuffer.DataSize = UploadedBuffer.DataSizeCallback();
 			}
 
-			if (UploadedBuffer.bUseFreeCallbacks)
+			if (UploadedBuffer.Data && UploadedBuffer.DataSize)
 			{
-				UploadedBuffer.DataFreeCallback(UploadedBuffer.Data);
+#if PLATFORM_NEEDS_GPU_UAV_RESOURCE_INIT_WORKAROUND
+				if (UploadedBuffer.Buffer->bUAVAccessed)
+				{
+					FRHIResourceCreateInfo CreateInfo(UploadedBuffer.Buffer->Name);
+					FBufferRHIRef TempBuffer = RHICmdListUpload.CreateBuffer(UploadedBuffer.DataSize, BUF_VertexBuffer | BUF_Static | BUF_ShaderResource, 0, ERHIAccess::CopySrc | ERHIAccess::SRVMask, CreateInfo);
+					void* DestPtr = RHICmdListUpload.LockBuffer(TempBuffer, 0, UploadedBuffer.DataSize, RLM_WriteOnly);
+					FMemory::Memcpy(DestPtr, UploadedBuffer.Data, UploadedBuffer.DataSize);
+					RHICmdListUpload.UnlockBuffer(TempBuffer);
+					RHICmdListUpload.Transition(FRHITransitionInfo(UploadedBuffer.Buffer->GetRHI(), ERHIAccess::Unknown, ERHIAccess::CopyDest));
+					RHICmdListUpload.CopyBufferRegion(UploadedBuffer.Buffer->GetRHI(), 0, TempBuffer, 0, UploadedBuffer.DataSize);
+				}
+				else
+#endif
+				{
+					void* DestPtr = RHICmdListUpload.LockBuffer(UploadedBuffer.Buffer->GetRHIUnchecked(), 0, UploadedBuffer.DataSize, RLM_WriteOnly);
+					FMemory::Memcpy(DestPtr, UploadedBuffer.Data, UploadedBuffer.DataSize);
+					RHICmdListUpload.UnlockBuffer(UploadedBuffer.Buffer->GetRHIUnchecked());
+				}
+
+				if (UploadedBuffer.bUseFreeCallbacks)
+				{
+					UploadedBuffer.DataFreeCallback(UploadedBuffer.Data);
+				}
 			}
 		}
+
+		UploadedBuffers.Reset();
+	};
+
+	FGraphEventRef Event;
+
+	if (bParallelSetupEnabled)
+	{
+		SCOPED_NAMED_EVENT_TEXT("FRDGBuilder::SubmitBufferUploads", FColor::Magenta);
+
+		FRHICommandList* RHICmdListUpload = new FRHICommandList(FRHIGPUMask::All());
+
+		Event = FFunctionGraphTask::CreateAndDispatchWhenReady(
+			[SubmitUploadsLambda, RHICmdListUpload](ENamedThreads::Type, const FGraphEventRef&)
+		{
+			FTaskTagScope Scope(ETaskTag::EParallelRenderingThread);
+			SubmitUploadsLambda(*RHICmdListUpload);
+			RHICmdListUpload->FinishRecording();
+
+		}, TStatId(), nullptr, ENamedThreads::AnyHiPriThreadHiPriTask);
+
+		RHICmdList.QueueRenderThreadCommandListSubmit(Event, RHICmdListUpload);
+	}
+	else
+	{
+		SubmitUploadsLambda(RHICmdList);
 	}
-	UploadedBuffers.Reset();
+
+	return Event;
 }
 
 ///////////////////////////////////////////////////////////////////////////////////////////////////
@@ -2389,8 +2700,10 @@ void FRDGBuilder::DispatchParallelExecute(IRHICommandContext* RHICmdContext)
 
 ///////////////////////////////////////////////////////////////////////////////////////////////////
 
-void FRDGBuilder::CreateUniformBuffers(FGraphEventArray* AsyncCompileEvents)
+UE::Tasks::FTask FRDGBuilder::CreateUniformBuffers()
 {
+	SCOPED_NAMED_EVENT_TEXT("FRDGBuilder::CreateUniformBuffers", FColor::Magenta);
+
 	const int32 ParallelDispatchThreshold = 4;
 
 	const auto CreateUniformBuffersFunction = [this]
@@ -2403,20 +2716,24 @@ void FRDGBuilder::CreateUniformBuffers(FGraphEventArray* AsyncCompileEvents)
 		UniformBuffersToCreate.Reset();
 	};
 
-	if (AsyncCompileEvents && UniformBuffersToCreate.Num() > ParallelDispatchThreshold)
+	UE::Tasks::FTask Task;
+
+	if (bParallelSetupEnabled && UniformBuffersToCreate.Num() > ParallelDispatchThreshold)
 	{
-		AsyncCompileEvents->Emplace(FFunctionGraphTask::CreateAndDispatchWhenReady(
-			[CreateUniformBuffersFunction](ENamedThreads::Type, const FGraphEventRef&)
+		Task = UE::Tasks::Launch(TEXT("FRDGBuilder::CreateUniformBuffer"),
+			[CreateUniformBuffersFunction]
 		{
 			FTaskTagScope Scope(ETaskTag::EParallelRenderingThread);
 			CreateUniformBuffersFunction();
 
-		}, TStatId(), nullptr, ENamedThreads::AnyHiPriThreadHiPriTask));
+		}, LowLevelTasks::ETaskPriority::High);
 	}
 	else
 	{
 		CreateUniformBuffersFunction();
 	}
+
+	return Task;
 }
 
 ///////////////////////////////////////////////////////////////////////////////////////////////////
@@ -2567,8 +2884,6 @@ void FRDGBuilder::ExecutePass(FRDGPass* Pass, FRHIComputeCommandList& RHICmdList
 	
 	Pass->GPUScopeOpsEpilogue.Execute(RHICmdListPass);
 
-
-
 	if (!Pass->bParallelExecute && Pass->bDispatchAfterExecute)
 	{
 		if (RHICmdListAsyncCompute.HasCommands())
@@ -2642,14 +2957,14 @@ void FRDGBuilder::EndResourcesRHI(FRDGPass* ResourcePass, FRDGPassHandle Execute
 
 ///////////////////////////////////////////////////////////////////////////////////////////////////
 
-void FRDGBuilder::CollectPassBarriers(FRDGPass* Pass, FRDGPassHandle PassHandle)
+void FRDGBuilder::CollectPassBarriers(FRDGPass* Pass)
 {
 	IF_RDG_ENABLE_DEBUG(ConditionalDebugBreak(RDG_BREAKPOINT_PASS_COMPILE, BuilderName.GetTCHAR(), Pass->GetName()));
 
 	for (auto& PassState : Pass->TextureStates)
 	{
 		FRDGTextureRef Texture = PassState.Texture;
-		AddTransition(PassHandle, Texture, PassState.MergeState);
+		AddTransition(Pass->Handle, Texture, PassState.MergeState);
 
 		IF_RDG_ENABLE_TRACE(Trace.AddTexturePassDependency(Texture, Pass));
 	}
@@ -2657,56 +2972,45 @@ void FRDGBuilder::CollectPassBarriers(FRDGPass* Pass, FRDGPassHandle PassHandle)
 	for (auto& PassState : Pass->BufferStates)
 	{
 		FRDGBufferRef Buffer = PassState.Buffer;
-		AddTransition(PassHandle, Buffer, *PassState.MergeState);
+		AddTransition(Pass->Handle, Buffer, *PassState.MergeState);
 
 		IF_RDG_ENABLE_TRACE(Trace.AddBufferPassDependency(Buffer, Pass));
 	}
 }
 
-void FRDGBuilder::CreatePassBarriers(FGraphEventArray* AsyncCompileEvents)
+void FRDGBuilder::CreatePassBarriers(TFunctionRef<void()> PreWork)
 {
-	const int32 NumBarriersPerTask = 128;
+	const int32 NumBarriersPerTask = 32;
 
-	if (AsyncCompileEvents && TransitionCreateQueue.Num() > NumBarriersPerTask)
+	if (bParallelSetupEnabled && TransitionCreateQueue.Num() > NumBarriersPerTask)
 	{
-		AsyncCompileEvents->Emplace(FFunctionGraphTask::CreateAndDispatchWhenReady(
-			[this, NumBarriersPerTask](ENamedThreads::Type, const FGraphEventRef&)
+		ParallelForWithPreWork(TEXT("FRDGBuilder::CreatePassBarriers"), TransitionCreateQueue.Num(), NumBarriersPerTask, [this](int32 Index)
 		{
-#if 1
-			ParallelFor(TEXT("FRDGBuilder::CreatePassBarriers"), TransitionCreateQueue.Num(), NumBarriersPerTask, [this](int32 Index)
-			{
-				TransitionCreateQueue[Index]->CreateTransition();
-			});
-#else
-			SCOPED_NAMED_EVENT_TEXT("FRDGBuilder::CreatePassBarriers", FColor::Magenta);
-			for (FRDGBarrierBatchBegin* BeginBatch : TransitionCreateQueue)
-			{
-				BeginBatch->CreateTransition();
-			}
-#endif
-
-			TransitionCreateQueue.Reset();
+			TransitionCreateQueue[Index]->CreateTransition();
 
-		}, TStatId(), nullptr, ENamedThreads::AnyHiPriThreadHiPriTask));
+		}, PreWork);
 	}
 	else
 	{
+		PreWork();
+
 		SCOPED_NAMED_EVENT_TEXT("FRDGBuilder::CreatePassBarriers", FColor::Magenta);
 
 		for (FRDGBarrierBatchBegin* BeginBatch : TransitionCreateQueue)
 		{
 			BeginBatch->CreateTransition();
 		}
-
-		TransitionCreateQueue.Reset();
 	}
+
+	TransitionCreateQueue.Reset();
 }
 
+
 ///////////////////////////////////////////////////////////////////////////////////////////////////
 
 void FRDGBuilder::AddEpilogueTransition(FRDGTextureRef Texture)
 {
-	if (!Texture->HasRHI() || Texture->bCulled || !Texture->bLastOwner)
+	if (!Texture->HasRHI() || Texture->IsCulled() || !Texture->bLastOwner)
 	{
 		return;
 	}
@@ -2754,7 +3058,7 @@ void FRDGBuilder::AddEpilogueTransition(FRDGTextureRef Texture)
 
 void FRDGBuilder::AddEpilogueTransition(FRDGBufferRef Buffer)
 {
-	if (!Buffer->HasRHI() || Buffer->bCulled || !Buffer->bLastOwner)
+	if (!Buffer->HasRHI() || Buffer->IsCulled() || !Buffer->bLastOwner)
 	{
 		return;
 	}
@@ -3309,7 +3613,7 @@ void FRDGBuilder::BeginResourceRHI(FRDGPassHandle PassHandle, FRDGBufferRef Buff
 		return;
 	}
 
-	check(Buffer->ReferenceCount > 0 || Buffer->bExternal || IsImmediateMode());
+	check(Buffer->ReferenceCount > 0 || Buffer->bExternal || Buffer->bQueuedForUpload || IsImmediateMode());
 
 #if RDG_ENABLE_DEBUG
 	{
@@ -3434,6 +3738,7 @@ void FRDGBuilder::InitRHI(FRDGView* View)
 void FRDGBuilder::EndResourceRHI(FRDGPassHandle PassHandle, FRDGTextureRef Texture, uint32 ReferenceCount)
 {
 	check(Texture);
+	check(Texture->ReferenceCount != FRDGViewableResource::DeallocatedReferenceCount);
 	check(Texture->ReferenceCount >= ReferenceCount || IsImmediateMode());
 	Texture->ReferenceCount -= ReferenceCount;
 
@@ -3467,12 +3772,14 @@ void FRDGBuilder::EndResourceRHI(FRDGPassHandle PassHandle, FRDGTextureRef Textu
 		}
 
 		Texture->LastPass = PassHandle;
+		Texture->ReferenceCount = FRDGViewableResource::DeallocatedReferenceCount;
 	}
 }
 
 void FRDGBuilder::EndResourceRHI(FRDGPassHandle PassHandle, FRDGBufferRef Buffer, uint32 ReferenceCount)
 {
 	check(Buffer);
+	check(Buffer->ReferenceCount != FRDGViewableResource::DeallocatedReferenceCount);
 	check(Buffer->ReferenceCount >= ReferenceCount || IsImmediateMode());
 	Buffer->ReferenceCount -= ReferenceCount;
 
@@ -3488,6 +3795,7 @@ void FRDGBuilder::EndResourceRHI(FRDGPassHandle PassHandle, FRDGBufferRef Buffer
 		}
 
 		Buffer->LastPass = PassHandle;
+		Buffer->ReferenceCount = FRDGViewableResource::DeallocatedReferenceCount;
 	}
 }
 
@@ -3768,4 +4076,4 @@ void FRDGBuilder::ForceCopyCrossGPU()
 		RHICmdList.TransferResources(Transfers);
 	}
 }
-#endif  // WITH_MGPU
+#endif  // WITH_MGPU
\ No newline at end of file
diff --git a/Engine/Source/Runtime/RenderCore/Private/RenderGraphPrivate.cpp b/Engine/Source/Runtime/RenderCore/Private/RenderGraphPrivate.cpp
--- a/Engine/Source/Runtime/RenderCore/Private/RenderGraphPrivate.cpp
+++ b/Engine/Source/Runtime/RenderCore/Private/RenderGraphPrivate.cpp
@@ -328,6 +328,16 @@ FAutoConsoleVariableRef CVarRDGTransientExtractedResource(
 	TEXT(" 2: force enables all external transient resources (not recommended);"),
 	ECVF_RenderThreadSafe);
 
+#if RDG_ENABLE_PARALLEL_TASKS
+
+int32 GRDGParallelSetup = 1;
+FAutoConsoleVariableRef CVarRDGParallelSetup(
+	TEXT("r.RDG.ParallelSetup"), GRDGParallelSetup,
+	TEXT("RDG will setup passes in parallel when prompted by calls to FRDGBuilder::FlushSetupQueue.")
+	TEXT(" 0: pass setup is done synchronously in AddPass;")
+	TEXT(" 1: pass setup is done asynchronously (default);"),
+	ECVF_RenderThreadSafe);
+
 int32 GRDGParallelExecute = 1;
 FAutoConsoleVariableRef CVarRDGParallelExecute(
 	TEXT("r.RDG.ParallelExecute"), GRDGParallelExecute,
@@ -363,13 +373,6 @@ FAutoConsoleVariableRef CVarRDGParallelExecutePassMax(
 	TEXT("The maximum span of contiguous passes eligible for parallel execution for the span to be offloaded to a task."),
 	ECVF_RenderThreadSafe);
 
-// Fix for random GPU crashes on draw indirects on multiple IHVs. Force all indirect arg buffers as non transient (see UE-115982)
-int32 GRDGTransientIndirectArgBuffers = 0;
-FAutoConsoleVariableRef CVarRDGIndirectArgBufferTransientAllocated(
-	TEXT("r.RDG.TransientAllocator.IndirectArgumentBuffers"), GRDGTransientIndirectArgBuffers,
-	TEXT("Whether indirect argument buffers should use transient resource allocator. Default: 0"),
-	ECVF_RenderThreadSafe);
-
 int32 GRDGParallelExecuteStress = 0;
 FAutoConsoleVariableRef CVarRDGDebugParallelExecute(
 	TEXT("r.RDG.ParallelExecuteStress"),
@@ -407,6 +410,15 @@ FAutoConsoleVariableRef CVarRDGDebugParallelExecute(
 	}),
 	ECVF_RenderThreadSafe);
 
+#endif //!RDG_ENABLE_PARALLEL_TASKS
+
+// Fix for random GPU crashes on draw indirects on multiple IHVs. Force all indirect arg buffers as non transient (see UE-115982)
+int32 GRDGTransientIndirectArgBuffers = 0;
+FAutoConsoleVariableRef CVarRDGIndirectArgBufferTransientAllocated(
+	TEXT("r.RDG.TransientAllocator.IndirectArgumentBuffers"), GRDGTransientIndirectArgBuffers,
+	TEXT("Whether indirect argument buffers should use transient resource allocator. Default: 0"),
+	ECVF_RenderThreadSafe);
+
 #if CSV_PROFILER
 int32 GRDGVerboseCSVStats = 0;
 FAutoConsoleVariableRef CVarRDGVerboseCSVStats(
@@ -514,11 +526,19 @@ void InitRenderGraph()
 		GRDGCullPasses = CullPassesValue;
 	}
 
+#if RDG_ENABLE_PARALLEL_TASKS
+	int32 ParallelSetupValue = 0;
+	if (FParse::Value(FCommandLine::Get(), TEXT("rdgparallelsetup="), ParallelSetupValue))
+	{
+		GRDGParallelSetup = ParallelSetupValue;
+	}
+
 	int32 ParallelExecuteValue = 0;
 	if (FParse::Value(FCommandLine::Get(), TEXT("rdgparallelexecute="), ParallelExecuteValue))
 	{
 		GRDGParallelExecute = ParallelExecuteValue;
 	}
+#endif
 
 	int32 MergeRenderPassesValue = 0;
 	if (FParse::Value(FCommandLine::Get(), TEXT("rdgmergerenderpasses="), MergeRenderPassesValue))
diff --git a/Engine/Source/Runtime/RenderCore/Private/RenderGraphPrivate.h b/Engine/Source/Runtime/RenderCore/Private/RenderGraphPrivate.h
--- a/Engine/Source/Runtime/RenderCore/Private/RenderGraphPrivate.h
+++ b/Engine/Source/Runtime/RenderCore/Private/RenderGraphPrivate.h
@@ -20,6 +20,10 @@ DEFINE_LOG_CATEGORY_STATIC(LogRDG, Log, All);
 #define RDG_BREAKPOINT_PASS_COMPILE 2
 #define RDG_BREAKPOINT_PASS_EXECUTE 3
 
+#ifndef RDG_ENABLE_PARALLEL_TASKS
+#define RDG_ENABLE_PARALLEL_TASKS 1
+#endif
+
 #define RDG_RECURSION_COUNTER_SCOPE(Counter) Counter++; ON_SCOPE_EXIT { Counter--; }
 
 #if RDG_ENABLE_DEBUG
@@ -111,10 +115,23 @@ extern int32 GRDGMergeRenderPasses;
 extern int32 GRDGTransientAllocator;
 extern int32 GRDGTransientExtractedResources;
 extern int32 GRDGTransientIndirectArgBuffers;
+
+#if RDG_ENABLE_PARALLEL_TASKS
+
+extern int32 GRDGParallelSetup;
 extern int32 GRDGParallelExecute;
 extern int32 GRDGParallelExecutePassMin;
 extern int32 GRDGParallelExecutePassMax;
 
+#else
+
+const int32 GRDGParallelSetup = 0;
+const int32 GRDGParallelExecute = 0;
+const int32 GRDGParallelExecutePassMin = 0;
+const int32 GRDGParallelExecutePassMax = 0;
+
+#endif
+
 #if CSV_PROFILER
 extern int32 GRDGVerboseCSVStats;
 #else
@@ -122,7 +139,6 @@ const int32 GRDGVerboseCSVStats = 0;
 #endif
 
 #if STATS
-extern int32 GRDGStatPassWithParameterCount;
 extern int32 GRDGStatPassCount;
 extern int32 GRDGStatPassCullCount;
 extern int32 GRDGStatRenderPassMergeCount;
@@ -222,6 +238,23 @@ FORCEINLINE bool IsParallelExecuteEnabled()
 		;
 }
 
+FORCEINLINE bool IsParallelSetupEnabled()
+{
+	return GRDGParallelSetup > 0
+		&& !GRHICommandList.Bypass()
+		&& !IsImmediateMode()
+		&& !GRDGDebug
+		&& !GRDGTransitionLog
+		&& !IsMobilePlatform(GMaxRHIShaderPlatform)
+		&& GRHISupportsMultithreadedShaderCreation
+#if WITH_DUMPGPU
+		&& !IsDumpingRDGResources()
+#endif
+		// Only run parallel RDG if we have a rendering thread.
+		&& IsInActualRenderingThread()
+		;
+}
+
 template <typename ResourceRegistryType, typename FunctionType>
 inline void EnumerateExtendedLifetimeResources(ResourceRegistryType& Registry, FunctionType Function)
 {
diff --git a/Engine/Source/Runtime/RenderCore/Private/RenderGraphResources.cpp b/Engine/Source/Runtime/RenderCore/Private/RenderGraphResources.cpp
--- a/Engine/Source/Runtime/RenderCore/Private/RenderGraphResources.cpp
+++ b/Engine/Source/Runtime/RenderCore/Private/RenderGraphResources.cpp
@@ -31,16 +31,16 @@ FRDGViewableResource::FRDGViewableResource(const TCHAR* InName, const ERDGViewab
 	, bForceNonTransient(0)
 	, TransientExtractionHint(ETransientExtractionHint::None)
 	, bLastOwner(1)
-	, bCulled(!IsImmediateMode())
-	, bUsedByAsyncComputePass(0)
 	, bQueuedForUpload(0)
 	, FirstBarrier(EFirstBarrier::Split)
 	, bUAVAccessed(0)
+	, ReferenceCount(IsImmediateMode() ? 1 : 0)
 {
 	if (bSkipTracking)
 	{
 		SetExternalAccessMode(ERHIAccess::ReadOnlyExclusiveMask, ERHIPipeline::All);
 		AccessModeState.bLocked = 1;
+		AccessModeState.ActiveMode = AccessModeState.Mode;
 	}
 
 	if (bImmediateFirstBarrier)
diff --git a/Engine/Source/Runtime/RenderCore/Private/RenderGraphTrace.cpp b/Engine/Source/Runtime/RenderCore/Private/RenderGraphTrace.cpp
--- a/Engine/Source/Runtime/RenderCore/Private/RenderGraphTrace.cpp
+++ b/Engine/Source/Runtime/RenderCore/Private/RenderGraphTrace.cpp
@@ -338,7 +338,7 @@ void FRDGTrace::OutputGraphEnd(const FRDGBuilder& GraphBuilder)
 			<< TextureMessage.NumSamples(Texture->Desc.NumSamples)
 			<< TextureMessage.IsExternal(bool(Texture->bExternal))
 			<< TextureMessage.IsExtracted(bool(Texture->bExtracted))
-			<< TextureMessage.IsCulled(bool(Texture->bCulled))
+			<< TextureMessage.IsCulled(bool(Texture->IsCulled()))
 			<< TextureMessage.IsTrackingSkipped(EnumHasAnyFlags(Texture->Flags, ERDGTextureFlags::SkipTracking))
 			<< TextureMessage.IsTransient(bool(Texture->bTransient))
 			<< TextureMessage.IsTransientUntracked(false)
@@ -372,7 +372,7 @@ void FRDGTrace::OutputGraphEnd(const FRDGBuilder& GraphBuilder)
 			<< BufferMessage.NumElements(Buffer->Desc.NumElements)
 			<< BufferMessage.IsExternal(bool(Buffer->bExternal))
 			<< BufferMessage.IsExtracted(bool(Buffer->bExtracted))
-			<< BufferMessage.IsCulled(bool(Buffer->bCulled))
+			<< BufferMessage.IsCulled(bool(Buffer->IsCulled()))
 			<< BufferMessage.IsTrackingSkipped(EnumHasAnyFlags(Buffer->Flags, ERDGBufferFlags::SkipTracking))
 			<< BufferMessage.IsTransient(bool(Buffer->bTransient))
 			<< BufferMessage.IsTransientUntracked(false)
diff --git a/Engine/Source/Runtime/RenderCore/Private/RenderGraphValidation.cpp b/Engine/Source/Runtime/RenderCore/Private/RenderGraphValidation.cpp
--- a/Engine/Source/Runtime/RenderCore/Private/RenderGraphValidation.cpp
+++ b/Engine/Source/Runtime/RenderCore/Private/RenderGraphValidation.cpp
@@ -850,8 +850,6 @@ void FRDGUserValidation::ValidateExecuteEnd()
 	{
 		auto ValidateResourceAtExecuteEnd = [](const FRDGViewableResource* Resource)
 		{
-			check(Resource->ReferenceCount == Resource->bExtracted ? 1 : 0);
-
 			const auto& ViewableDebugData = Resource->GetViewableDebugData();
 			const bool bProducedButNeverUsed = ViewableDebugData.PassAccessCount == 1 && ViewableDebugData.FirstProducer;
 
diff --git a/Engine/Source/Runtime/RenderCore/Public/RenderGraphBuilder.h b/Engine/Source/Runtime/RenderCore/Public/RenderGraphBuilder.h
--- a/Engine/Source/Runtime/RenderCore/Public/RenderGraphBuilder.h
+++ b/Engine/Source/Runtime/RenderCore/Public/RenderGraphBuilder.h
@@ -31,7 +31,7 @@
 #include "Stats/Stats2.h"
 #include "Templates/RefCounting.h"
 #include "Templates/UnrealTemplate.h"
-#include "UObject/NameTypes.h"
+#include "Tasks/Pipe.h"
 
 /** Use the render graph builder to build up a graph of passes and then call Execute() to process them. Resource barriers
  *  and lifetimes are derived from _RDG_ parameters in the pass parameter struct provided to each AddPass call. The resulting
@@ -310,6 +310,9 @@ public:
 	/** Manually ends the current GPU event scope. */
 	void EndEventScope();
 
+	/** Flushes all queued passes to an async task to perform setup work. */
+	void FlushSetupQueue();
+
 	/** Executes the queued passes, managing setting of render targets (RHI RenderPasses), resource transitions and queued texture extraction. */
 	void Execute();
 
@@ -568,6 +571,34 @@ private:
 	/** Set of all active barrier batch begin instances; used to create transitions. */
 	FRDGTransitionCreateQueue TransitionCreateQueue;
 
+	template <typename LambdaType>
+	UE::Tasks::FTask LaunchCompileTask(const TCHAR* Name, bool bCondition, LambdaType&& Lambda);
+
+	UE::Tasks::FPipe CompilePipe;
+
+	class FPassQueue
+	{
+	public:
+		void Push(FRDGPass* Pass)
+		{
+			Queue.Push(Pass);
+		}
+
+		template <typename LambdaType>
+		void Flush(UE::Tasks::FPipe& Pipe, const TCHAR* Name, LambdaType&& Lambda);
+
+		template <typename LambdaType>
+		void Flush(const TCHAR* Name, LambdaType&& Lambda);
+
+	private:
+		TLockFreePointerListFIFO<FRDGPass, PLATFORM_CACHE_LINE_SIZE> Queue;
+		UE::Tasks::FTask LastTask;
+	};
+
+	FPassQueue SetupPassQueue;
+
+	TArray<FRDGPassHandle, FRDGArrayAllocator> CullPassStack;
+
 	/** The epilogue and prologue passes are sentinels that are used to simplify graph logic around barriers
 	 *  and traversal. The prologue pass is used exclusively for barriers before the graph executes, while the
 	 *  epilogue pass is used for resource extraction barriers--a property that also makes it the main root of
@@ -700,6 +731,7 @@ private:
 
 	bool bFlushResourcesRHI = false;
 	bool bParallelExecuteEnabled = false;
+	bool bParallelSetupEnabled = false;
 
 #if RDG_ENABLE_DEBUG
 	FRDGUserValidation UserValidation;
@@ -742,6 +774,8 @@ private:
 
 	IRHITransientResourceAllocator* TransientResourceAllocator = nullptr;
 
+	void MarkResourcesAsProduced(FRDGPass* Pass);
+
 	void Compile();
 	void Clear();
 
@@ -768,29 +802,36 @@ private:
 	void SetupParallelExecute();
 	void DispatchParallelExecute(IRHICommandContext* RHICmdContext);
 
-	void SetupBufferUploads();
-	void SubmitBufferUploads();
+	void PrepareBufferUploads();
+	FGraphEventRef SubmitBufferUploads();
 	void BeginFlushResourcesRHI();
 	void EndFlushResourcesRHI();
 
 	void FlushAccessModeQueue();
 
-	void SetupPassInternal(FRDGPass* Pass, FRDGPassHandle PassHandle, ERHIPipeline PassPipeline, bool bEmptyParameters);
-	FRDGPass* SetupPass(FRDGPass* Pass);
 	FRDGPass* SetupEmptyPass(FRDGPass* Pass);
+	FRDGPass* SetupParameterPass(FRDGPass* Pass);
+
+	void SetupPassInternals(FRDGPass* Pass);
+	void SetupPassResources(FRDGPass* Pass);
+	void SetupAuxiliaryPasses(FRDGPass* Pass);
+	void SetupPassDependencies(FRDGPass* Pass);
+
 	void CompilePassOps(FRDGPass* Pass);
 	void ExecutePass(FRDGPass* Pass, FRHIComputeCommandList& RHICmdListPass);
 
 	void ExecutePassPrologue(FRHIComputeCommandList& RHICmdListPass, FRDGPass* Pass);
 	void ExecutePassEpilogue(FRHIComputeCommandList& RHICmdListPass, FRDGPass* Pass);
 
-	void CollectPassBarriers(FRDGPass* Pass, FRDGPassHandle PassHandle);
-
-	void CreatePassBarriers(FGraphEventArray* AsyncCompileEvents);
+	void CompilePassBarriers();
+	void CollectPassBarriers(FRDGPass* Pass);
+	void CreatePassBarriers(TFunctionRef<void()> PreWork);
 
-	void CreateUniformBuffers(FGraphEventArray* AsyncCompileEvents);
+	UE::Tasks::FTask CreateUniformBuffers();
 
 	void AddPassDependency(FRDGPassHandle ProducerHandle, FRDGPassHandle ConsumerHandle);
+	void AddPassDependency(FRDGPass* Producer, FRDGPass* Consumer);
+	void AddCullingDependency(FRDGProducerStatesByPipeline& LastProducers, const FRDGProducerState& NextState, ERHIPipeline NextPipeline);
 
 	void AddEpilogueTransition(FRDGTextureRef Texture);
 	void AddEpilogueTransition(FRDGBufferRef Buffer);
diff --git a/Engine/Source/Runtime/RenderCore/Public/RenderGraphBuilder.inl b/Engine/Source/Runtime/RenderCore/Public/RenderGraphBuilder.inl
--- a/Engine/Source/Runtime/RenderCore/Public/RenderGraphBuilder.inl
+++ b/Engine/Source/Runtime/RenderCore/Public/RenderGraphBuilder.inl
@@ -227,7 +227,7 @@ FRDGPassRef FRDGBuilder::AddPassInternal(
 
 	IF_RDG_ENABLE_DEBUG(ClobberPassOutputs(Pass));
 	Passes.Insert(Pass);
-	SetupPass(Pass);
+	SetupParameterPass(Pass);
 	return Pass;
 }
 
@@ -265,6 +265,7 @@ inline void FRDGBuilder::QueueBufferUpload(FRDGBufferRef Buffer, const void* Ini
 
 	UploadedBuffers.Emplace(Buffer, InitialData, InitialDataSize);
 	Buffer->bQueuedForUpload = 1;
+	Buffer->bForceNonTransient = 1;
 }
 
 inline void FRDGBuilder::QueueBufferUpload(FRDGBufferRef Buffer, const void* InitialData, uint64 InitialDataSize, FRDGBufferInitialDataFreeCallback&& InitialDataFreeCallback)
@@ -278,6 +279,7 @@ inline void FRDGBuilder::QueueBufferUpload(FRDGBufferRef Buffer, const void* Ini
 
 	UploadedBuffers.Emplace(Buffer, InitialData, InitialDataSize, MoveTemp(InitialDataFreeCallback));
 	Buffer->bQueuedForUpload = 1;
+	Buffer->bForceNonTransient = 1;
 }
 
 inline void FRDGBuilder::QueueBufferUpload(FRDGBufferRef Buffer, FRDGBufferInitialDataCallback&& InitialDataCallback, FRDGBufferInitialDataSizeCallback&& InitialDataSizeCallback)
@@ -286,6 +288,7 @@ inline void FRDGBuilder::QueueBufferUpload(FRDGBufferRef Buffer, FRDGBufferIniti
 
 	UploadedBuffers.Emplace(Buffer, MoveTemp(InitialDataCallback), MoveTemp(InitialDataSizeCallback));
 	Buffer->bQueuedForUpload = 1;
+	Buffer->bForceNonTransient = 1;
 }
 
 inline void FRDGBuilder::QueueBufferUpload(FRDGBufferRef Buffer, FRDGBufferInitialDataCallback&& InitialDataCallback, FRDGBufferInitialDataSizeCallback&& InitialDataSizeCallback, FRDGBufferInitialDataFreeCallback&& InitialDataFreeCallback)
@@ -294,6 +297,7 @@ inline void FRDGBuilder::QueueBufferUpload(FRDGBufferRef Buffer, FRDGBufferIniti
 
 	UploadedBuffers.Emplace(Buffer, MoveTemp(InitialDataCallback), MoveTemp(InitialDataSizeCallback), MoveTemp(InitialDataFreeCallback));
 	Buffer->bQueuedForUpload = 1;
+	Buffer->bForceNonTransient = 1;
 }
 
 inline void FRDGBuilder::QueueTextureExtraction(FRDGTextureRef Texture, TRefCountPtr<IPooledRenderTarget>* OutTexturePtr, ERHIAccess AccessFinal, ERDGResourceExtractionFlags Flags)
@@ -308,9 +312,7 @@ inline void FRDGBuilder::QueueTextureExtraction(FRDGTextureRef Texture, TRefCoun
 
 	*OutTexturePtr = nullptr;
 
-	Texture->ReferenceCount++;
 	Texture->bExtracted = true;
-	Texture->bCulled = false;
 
 	if (EnumHasAnyFlags(Flags, ERDGResourceExtractionFlags::AllowTransient))
 	{
@@ -333,9 +335,7 @@ inline void FRDGBuilder::QueueBufferExtraction(FRDGBufferRef Buffer, TRefCountPt
 
 	*OutBufferPtr = nullptr;
 
-	Buffer->ReferenceCount++;
 	Buffer->bExtracted = true;
-	Buffer->bCulled = false;
 	Buffer->bForceNonTransient = true;
 	ExtractedBuffers.Emplace(Buffer, OutBufferPtr);
 }
diff --git a/Engine/Source/Runtime/RenderCore/Public/RenderGraphPass.h b/Engine/Source/Runtime/RenderCore/Public/RenderGraphPass.h
--- a/Engine/Source/Runtime/RenderCore/Public/RenderGraphPass.h
+++ b/Engine/Source/Runtime/RenderCore/Public/RenderGraphPass.h
@@ -387,18 +387,9 @@ protected:
 			/** Whether the pass is allowed to execute in parallel. */
 			uint32 bParallelExecuteAllowed : 1;
 
-			/** Whether this pass has non-RDG UAV outputs. */
-			uint32 bHasExternalOutputs : 1;
-
 			/** Whether this pass is a sentinel (prologue / epilogue) pass. */
 			uint32 bSentinel : 1;
 
-			/** Whether this pass has been culled. */
-			uint32 bCulled : 1;
-
-			/** Whether this pass does not contain parameters. */
-			uint32 bEmptyParameters : 1;
-
 			/** If set, dispatches to the RHI thread after executing this pass. */
 			uint32 bDispatchAfterExecute : 1;
 
@@ -422,6 +413,18 @@ protected:
 
 			/** If set, marks that a pass is executing in parallel. */
 			uint32 bParallelExecute : 1;
+
+			/** Whether this pass does not contain parameters. */
+			uint32 bEmptyParameters : 1;
+
+			/** Whether this pass has external UAVs that are not tracked by RDG. */
+			uint32 bHasExternalOutputs : 1;
+
+			/** Whether this pass has been culled. */
+			uint32 bCulled : 1;
+
+			/** Whether this pass is used for external access transitions. */
+			uint32 bExternalAccessPass : 1;
 		};
 		uint32 PacketBits2 = 0;
 	};
@@ -481,6 +484,21 @@ protected:
 	TArray<FRDGViewHandle, FRDGArrayAllocator> Views;
 	TArray<FRDGUniformBufferHandle, FRDGArrayAllocator> UniformBuffers;
 
+	struct FExternalAccessOp
+	{
+		FExternalAccessOp() = default;
+
+		FExternalAccessOp(FRDGViewableResource* InResource, FRDGViewableResource::EAccessMode InMode)
+			: Resource(InResource)
+			, Mode(InMode)
+		{}
+
+		FRDGViewableResource* Resource;
+		FRDGViewableResource::EAccessMode Mode;
+	};
+
+	TArray<FExternalAccessOp, FRDGArrayAllocator> ExternalAccessOps;
+
 	/** Lists of pass parameters scheduled for begin during execution of this pass. */
 	TArray<FRDGPass*, TInlineAllocator<1, FRDGArrayAllocator>> ResourcesToBegin;
 	TArray<FRDGPass*, TInlineAllocator<1, FRDGArrayAllocator>> ResourcesToEnd;
diff --git a/Engine/Source/Runtime/RenderCore/Public/RenderGraphResources.h b/Engine/Source/Runtime/RenderCore/Public/RenderGraphResources.h
--- a/Engine/Source/Runtime/RenderCore/Public/RenderGraphResources.h
+++ b/Engine/Source/Runtime/RenderCore/Public/RenderGraphResources.h
@@ -50,8 +50,8 @@ struct FRDGProducerState
 
 	FRDGProducerState() = default;
 
+	FRDGPass* Pass = nullptr;
 	ERHIAccess Access = ERHIAccess::Unknown;
-	FRDGPassHandle PassHandle;
 	FRDGViewHandle NoUAVBarrierHandle;
 };
 
@@ -286,7 +286,7 @@ public:
 
 	bool IsCulled() const
 	{
-		return bCulled;
+		return ReferenceCount == 0;
 	}
 
 	/** Whether a prior pass added to the graph produced contents for this resource. External resources are not considered produced
@@ -324,7 +324,7 @@ protected:
 
 	struct FAccessModeState
 	{
-		bool IsExternalAccess() const { return Mode == EAccessMode::External && !bQueued; }
+		bool IsExternalAccess() const { return ActiveMode == EAccessMode::External; }
 
 		FAccessModeState()
 			: Pipelines(ERHIPipeline::None)
@@ -333,11 +333,14 @@ protected:
 			, bQueued(0)
 		{}
 
-		ERHIAccess			Access			= ERHIAccess::None;
-		ERHIPipeline		Pipelines		: 2;
-		EAccessMode			Mode			: 1;
-		uint8				bLocked			: 1;
-		uint8				bQueued			: 1;
+		ERHIAccess			Access = ERHIAccess::None;
+		ERHIPipeline		Pipelines : 2;
+		EAccessMode			Mode : 1;
+		uint8				bLocked : 1;
+		uint8				bQueued : 1;
+
+		/** The actual access mode replayed on the setup pass timeline. */
+		EAccessMode ActiveMode = EAccessMode::Internal;
 
 	} AccessModeState;
 
@@ -362,12 +365,6 @@ protected:
 	/** Whether this resource is the last owner of its allocation (i.e. nothing aliases the allocation later in the execution timeline). */
 	uint8 bLastOwner : 1;
 
-	/** If true, the resource was not used by any pass not culled by the graph. */
-	uint8 bCulled : 1;
-
-	/** If true, the resource has been used on an async compute pass and may have async compute states. */
-	uint8 bUsedByAsyncComputePass : 1;
-
 	/** If true, the resource has been queued for an upload operation. */
 	uint8 bQueuedForUpload : 1;
 
@@ -384,8 +381,10 @@ protected:
 	ERHIAccess EpilogueAccess = DefaultEpilogueAccess;
 
 private:
+	static const uint16 DeallocatedReferenceCount = ~0;
+
 	/** Number of references in passes and deferred queries. */
-	uint16 ReferenceCount = 0;
+	uint16 ReferenceCount;
 
 	/** Scratch index allocated for the resource in the pass being setup. */
 	uint16 PassStateIndex = 0;
@@ -398,9 +397,6 @@ private:
 		AccessModeState.Access = InReadOnlyAccess;
 		AccessModeState.Pipelines = InPipelines;
 
-		// External access resources are not always added to the pass states (unless marked as such within the graph), so marked as not culled here.
-		bCulled = 0;
-
 		EpilogueAccess = InReadOnlyAccess;
 	}
 
@@ -419,6 +415,7 @@ private:
 	friend FRDGBarrierBatchBegin;
 	friend FRDGResourceDumpContext;
 	friend FRDGTrace;
+	friend FRDGPass;
 };
 
 /** A render graph resource (e.g. a view) which references a single viewable resource (e.g. a texture / buffer). Provides an abstract way to access the viewable resource. */
diff --git a/Engine/Source/Runtime/Renderer/Private/DeferredShadingRenderer.cpp b/Engine/Source/Runtime/Renderer/Private/DeferredShadingRenderer.cpp
--- a/Engine/Source/Runtime/Renderer/Private/DeferredShadingRenderer.cpp
+++ b/Engine/Source/Runtime/Renderer/Private/DeferredShadingRenderer.cpp
@@ -3095,6 +3095,8 @@ void FDeferredShadingSceneRenderer::Render(FRDGBuilder& GraphBuilder)
 	}
 #endif // RHI_RAYTRACING
 
+	GraphBuilder.FlushSetupQueue();
+
 	if (bRenderDeferredLighting)
 	{
 		RDG_GPU_STAT_SCOPE(GraphBuilder, RenderDeferredLighting);
@@ -3182,6 +3184,8 @@ void FDeferredShadingSceneRenderer::Render(FRDGBuilder& GraphBuilder)
 		RenderHeterogeneousVolumes(GraphBuilder, SceneTextures);
 	}
 
+	GraphBuilder.FlushSetupQueue();
+
 	if (bShouldRenderVolumetricCloud && IsVolumetricRenderTargetEnabled() && !bHasHalfResCheckerboardMinMaxDepth && !bHasRayTracedOverlay)
 	{
 		HalfResolutionDepthCheckerboardMinMaxTexture = CreateHalfResolutionDepthCheckerboardMinMax(GraphBuilder, Views, SceneTextures.Depth.Resolve);
@@ -3523,6 +3527,8 @@ void FDeferredShadingSceneRenderer::Render(FRDGBuilder& GraphBuilder)
 		PostProcessingInputs.CustomDepthTexture = SceneTextures.CustomDepth.Depth;
 		PostProcessingInputs.SceneTextures = SceneTextures.UniformBuffer;
 
+		GraphBuilder.FlushSetupQueue();
+
 		if (ViewFamily.UseDebugViewPS())
 		{
 			for (int32 ViewIndex = 0; ViewIndex < Views.Num(); ViewIndex++)
@@ -3594,6 +3600,7 @@ void FDeferredShadingSceneRenderer::Render(FRDGBuilder& GraphBuilder)
 		RenderFinish(GraphBuilder, ViewFamilyTexture);
 		GraphBuilder.SetCommandListStat(GET_STATID(STAT_CLM_AfterFrame));
 		GraphBuilder.AddDispatchHint();
+		GraphBuilder.FlushSetupQueue();
 	}
 
 	QueueSceneTextureExtractions(GraphBuilder, SceneTextures);
