Add comment directing users towards r.RayTracing.NonBlockingPipelineCreation to reduce RTPSO stalls during gameplay

diff --git a/Engine/Source/Runtime/RenderCore/Private/ShaderPipelineCache.cpp b/Engine/Source/Runtime/RenderCore/Private/ShaderPipelineCache.cpp
--- a/Engine/Source/Runtime/RenderCore/Private/ShaderPipelineCache.cpp
+++ b/Engine/Source/Runtime/RenderCore/Private/ShaderPipelineCache.cpp
@@ -1057,6 +1057,7 @@ bool FShaderPipelineCacheTask::Precompile(FRHICommandListImmediate& RHICmdList,
 			  // If ray tracing PSO file cache is generated using one payload size but later shaders were re-compiled with a different payload declaration 
 			  // it is possible for the wrong size to be used here, which leads to a D3D run-time error when attempting to create the PSO.
 			  // Ray tracing shader pre-compilation is disabled until a robust solution is found.
+			  // Consider setting r.RayTracing.NonBlockingPipelineCreation=1 meanwhile to avoid most of the RTPSO creation stalls during gameplay.
 			if (IsRayTracingEnabled())
 			{
 				FRayTracingPipelineStateInitializer Initializer;
