Preview Shader Platform using different ShaderMaps and the DDPI/Cvar values of the Previewable Actual Platform
NumPlatforms limit changed from 64 to 65536, limited by NumBits

diff --git a/Engine/Config/Android/DataDrivenPlatformInfo.ini b/Engine/Config/Android/DataDrivenPlatformInfo.ini
--- a/Engine/Config/Android/DataDrivenPlatformInfo.ini
+++ b/Engine/Config/Android/DataDrivenPlatformInfo.ini
@@ -23,6 +23,7 @@ PlatformGroupName=Mobile
 [PreviewPlatform AndroidES31]
 EnabledCVar=ini:Engine:Android:/Script/AndroidRuntimeSettings.AndroidRuntimeSettings:bBuildForES31
 PlatformName=Android
+ShaderPlatform=OPENGL_ES3_1_ANDROID
 ShaderFormat=GLSL_ES3_1_ANDROID
 ActiveIconName=LevelEditor.PreviewMode.AndroidES31.Enabled
 InactiveIconName=LevelEditor.PreviewMode.AndroidES31.Disabled
@@ -34,6 +35,7 @@ DeviceProfileName=Android_Preview_OpenGL
 [PreviewPlatform AndroidVulkan]
 EnabledCVar=ini:Engine:Android:/Script/AndroidRuntimeSettings.AndroidRuntimeSettings:bSupportsVulkan
 PlatformName=Android
+ShaderPlatform=VULKAN_ES3_1_ANDROID
 ShaderFormat=SF_VULKAN_ES31_ANDROID
 ActiveIconName=LevelEditor.PreviewMode.AndroidVulkan.Enabled
 InactiveIconName=LevelEditor.PreviewMode.AndroidVulkan.Disabled
@@ -45,6 +47,7 @@ DeviceProfileName=Android_Preview_Vulkan
 [PreviewPlatform AndroidVulkanSM5]
 EnabledCVar=ini:Engine:Android:/Script/AndroidRuntimeSettings.AndroidRuntimeSettings:bSupportsVulkanSM5
 PlatformName=Android
+ShaderPlatform=VULKAN_SM5_ANDROID
 ShaderFormat=SF_VULKAN_SM5_ANDROID
 ActiveIconName=LevelEditor.PreviewMode.AndroidVulkanSM5.Enabled
 InactiveIconName=LevelEditor.PreviewMode.AndroidVulkanSM5.Disabled
diff --git a/Engine/Config/IOS/DataDrivenPlatformInfo.ini b/Engine/Config/IOS/DataDrivenPlatformInfo.ini
--- a/Engine/Config/IOS/DataDrivenPlatformInfo.ini
+++ b/Engine/Config/IOS/DataDrivenPlatformInfo.ini
@@ -27,6 +27,7 @@ PlatformGroupName=Mobile
 [PreviewPlatform IOSMetal]
 PlatformName=IOS
 ShaderFormat=SF_METAL
+ShaderPlatform=METAL
 ActiveIconName=LevelEditor.PreviewMode.iOS.Enabled
 InactiveIconName=LevelEditor.PreviewMode.iOS.Disabled
 MenuText=NSLOCTEXT("PreviewPlatform", "PreviewMenuText_iOS", "iOS")
@@ -37,6 +38,7 @@ DeviceProfileName=IOS
 [PreviewPlatform IOSMetalSM5]
 PlatformName=IOS
 ShaderFormat=SF_METAL_MRT
+ShaderPlatform=METAL_MRT
 EnabledCVar=ini:Engine:IOS:/Script/IOSRuntimeSettings.IOSRuntimeSettings:bSupportsMetalMRT
 ActiveIconName=LevelEditor.PreviewMode.iOSSM5.Enabled
 InactiveIconName=LevelEditor.PreviewMode.iOSSM5.Disabled
diff --git a/Engine/Config/Windows/DataDrivenPlatformInfo.ini b/Engine/Config/Windows/DataDrivenPlatformInfo.ini
--- a/Engine/Config/Windows/DataDrivenPlatformInfo.ini
+++ b/Engine/Config/Windows/DataDrivenPlatformInfo.ini
@@ -78,6 +78,7 @@ SupportsMultiViewport=RuntimeDependent
 [PreviewPlatform PCD3D_SM5]
 PlatformName=PCD3D_SM5
 ShaderFormat=PCD3D_SM5
+ShaderPlatform=PCD3D_SM5
 MenuText=LOCTEXT("PreviewMenuText_PCD3D_SM5", "Shader Model 5")
 MenuTooltip=LOCTEXT("PreviewMenuTooltip_PCD3D_SM5", "DirectX 11, OpenGL 4.3+, PS4, XB1")
 
@@ -135,8 +136,9 @@ SupportsMultiViewport=RuntimeGuaranteed
 [PreviewPlatform PCD3D_SM6]
 PlatformName=PCD3D_SM6
 ShaderFormat=PCD3D_SM6
+ShaderPlatform=PCD3D_SM6
 MenuText=LOCTEXT("PreviewMenuText_PCD3D_SM6", "Shader Model 6")
-MenuTooltip=LOCTEXT("PreviewMenuTooltip_PCD3D_SM6", "Preview using SM6 device profile.")
+MenuTooltip=LOCTEXT("PreviewMenuTooltip_PCD3D_SM6", "Windows using SM6 profile")
 
 [ShaderPlatform PCD3D_ES3_1]
 Language=D3D
@@ -151,9 +153,17 @@ bSupportsManualVertexFetch=false
 bSupportsSceneDataCompressedTransforms=true
 bSupportsDebugViewShaders=true
 
+[PreviewPlatform PCD3D_ES3_1]
+PlatformName=PCD3D_ES3_1
+ShaderFormat=PCD3D_ES31
+ShaderPlatform=PCD3D_ES3_1
+MenuText=LOCTEXT("PreviewMenuText_PCD3D_ES3_1", "Shader Model 3_1")
+MenuTooltip=LOCTEXT("PreviewMenuTooltip_PCD3D_ES3_1", "Windows using ES3_1 profile.")
+
 [ShaderPlatform OPENGL_PCES3_1]
 Language=OpenGL
 bIsPC=true
+ShaderFormat=GLSL_150_ES31
 MaxFeatureLevel=ES3_1
 ShaderFormat=GLSL_150_ES31
 bIsMobile=true
diff --git a/Engine/Plugins/Experimental/GPULightmass/Source/GPULightmassEditor/Private/GPULightmassEditorModule.cpp b/Engine/Plugins/Experimental/GPULightmass/Source/GPULightmassEditor/Private/GPULightmassEditorModule.cpp
--- a/Engine/Plugins/Experimental/GPULightmass/Source/GPULightmassEditor/Private/GPULightmassEditorModule.cpp
+++ b/Engine/Plugins/Experimental/GPULightmass/Source/GPULightmassEditor/Private/GPULightmassEditorModule.cpp
@@ -85,13 +85,13 @@ ERayTracingDisabledReason GetRayTracingDisabledReason()
 		return ERayTracingDisabledReason::DISABLED_BY_PROJECT_SETTINGS;
 	}
 
-	extern RENDERCORE_API uint64 GRayTracingPlaformMask;
+	extern RENDERCORE_API ShaderPlatformMaskType GRayTracingPlaformMask;
 	if (!RHISupportsRayTracing(GMaxRHIShaderPlatform))
 	{
 		return ERayTracingDisabledReason::INCOMPATIBLE_SHADER_PLATFORM;
 	}
 	// Shader platform statically supports ray tracing, but disabled by target platform at runtime
-	else if (RHISupportsRayTracing(GMaxRHIShaderPlatform) && (GRayTracingPlaformMask & (1ull << GMaxRHIShaderPlatform)) == 0)
+	else if (RHISupportsRayTracing(GMaxRHIShaderPlatform) && (GRayTracingPlaformMask[(int)GMaxRHIShaderPlatform]) == 0)
 	{
 		return ERayTracingDisabledReason::DISABLED_BY_TARGET_PLATFORM;
 	}
diff --git a/Engine/Plugins/FX/Niagara/Shaders/Private/NiagaraShaderVersion.ush b/Engine/Plugins/FX/Niagara/Shaders/Private/NiagaraShaderVersion.ush
--- a/Engine/Plugins/FX/Niagara/Shaders/Private/NiagaraShaderVersion.ush
+++ b/Engine/Plugins/FX/Niagara/Shaders/Private/NiagaraShaderVersion.ush
@@ -3,4 +3,4 @@
 // in the DDC key, it will automatically invalidate.
 // 
 // If you are merging streams and there is a conflict with this GUID you should make a new GUID rather than taking one or the other.
-// GUID = 8E831E6316F34A4BAAA213E11EE571E3
+// GUID = 22B2FD1A565942EAB86A4B4A6FB56DEC
diff --git a/Engine/Shaders/Public/ShaderVersion.ush b/Engine/Shaders/Public/ShaderVersion.ush
--- a/Engine/Shaders/Public/ShaderVersion.ush
+++ b/Engine/Shaders/Public/ShaderVersion.ush
@@ -2,4 +2,4 @@
 // Each time the console command is executed it generates a new GUID. As this file is included
 // in Platform.ush (which should be included in any shader) it allows to invalidate the shader DDC.
 // 
-// GUID = E4136F6D-EC93-4BC0-9E76-AB60D1E5BABF
+// GUID = 7553B411-6440-4492-9C57-1CD175A19DB3
diff --git a/Engine/Source/Developer/Windows/ShaderFormatD3D/Private/D3DShaderCompiler.cpp b/Engine/Source/Developer/Windows/ShaderFormatD3D/Private/D3DShaderCompiler.cpp
--- a/Engine/Source/Developer/Windows/ShaderFormatD3D/Private/D3DShaderCompiler.cpp
+++ b/Engine/Source/Developer/Windows/ShaderFormatD3D/Private/D3DShaderCompiler.cpp
@@ -81,11 +81,12 @@ static void D3D11FilterShaderCompileWarnings(const FString& CompileWarnings, TAr
 }
 
 // @return 0 if not recognized
-static const TCHAR* GetShaderProfileName(FShaderTarget Target, bool bForceSM6)
+static const TCHAR* GetShaderProfileName(ELanguage Language, uint32 Frequency, bool bForceSM6)
 {
-	if (Target.Platform == SP_PCD3D_SM6)
+
+	if (Language == ELanguage::SM6)
 	{
-		switch (Target.Frequency)
+		switch (Frequency)
 		{
 		default:
 			checkfSlow(false, TEXT("Unexpected shader frequency"));
@@ -109,10 +110,10 @@ static const TCHAR* GetShaderProfileName(FShaderTarget Target, bool bForceSM6)
 			return USE_SHADER_MODEL_6_6 ? TEXT("lib_6_6") : TEXT("lib_6_5");
 		}
 	}
-	else if(Target.Platform == SP_PCD3D_SM5)
+	else if(Language == ELanguage::SM5)
 	{
 		//set defines and profiles for the appropriate shader paths
-		switch(Target.Frequency)
+		switch(Frequency)
 		{
 		default:
 			checkfSlow(false, TEXT("Unexpected shader frequency"));
@@ -132,15 +133,15 @@ static const TCHAR* GetShaderProfileName(FShaderTarget Target, bool bForceSM6)
 			return TEXT("lib_6_5");
 		}
 	}
-	else if ((Target.Platform == SP_PCD3D_ES3_1) || (Target.Platform == SP_D3D_ES3_1_HOLOLENS))
+	else if (Language == ELanguage::ES3_1)
 	{
-		checkSlow(Target.Frequency == SF_Vertex ||
-			Target.Frequency == SF_Pixel ||
-			Target.Frequency == SF_Geometry ||
-			Target.Frequency == SF_Compute);
+		checkSlow(Frequency == SF_Vertex ||
+			Frequency == SF_Pixel ||
+			Frequency == SF_Geometry ||
+			Frequency == SF_Compute);
 
 		//set defines and profiles for the appropriate shader paths
-		switch(Target.Frequency)
+		switch(Frequency)
 		{
 		case SF_Pixel:
 			return TEXT("ps_5_0");
@@ -1193,7 +1194,7 @@ void CompileD3DShader(const FShaderCompilerInput& Input, FShaderCompilerOutput&
 		|| Input.Environment.CompilerFlags.Contains(CFLAG_WaveOperations)
 		|| Input.Environment.CompilerFlags.Contains(CFLAG_ForceDXC)
 		|| Input.Environment.CompilerFlags.Contains(CFLAG_InlineRayTracing);
-	const TCHAR* ShaderProfile = GetShaderProfileName(Input.Target, bUseDXC);
+	const TCHAR* ShaderProfile = GetShaderProfileName(Language, Input.Target.Frequency, bUseDXC);
 
 	if(!ShaderProfile)
 	{
@@ -1421,17 +1422,14 @@ void CompileShader_Windows(const FShaderCompilerInput& Input,FShaderCompilerOutp
 	FShaderCompilerDefinitions AdditionalDefines;
 	if (Language == ELanguage::SM6)
 	{
-		check(Input.Target.Platform == SP_PCD3D_SM6);
 		AdditionalDefines.SetDefine(TEXT("SM6_PROFILE"), 1);
 	}
 	else if (Language == ELanguage::SM5)
 	{
-		check(Input.Target.Platform == SP_PCD3D_SM5);
 		AdditionalDefines.SetDefine(TEXT("SM5_PROFILE"), 1);
 	}
 	else if (Language == ELanguage::ES3_1)
 	{
-		check((Input.Target.Platform == SP_PCD3D_ES3_1) || (Input.Target.Platform == SP_D3D_ES3_1_HOLOLENS));
 		AdditionalDefines.SetDefine(TEXT("ES3_1_PROFILE"), 1);
 	}
 	else
diff --git a/Engine/Source/Editor/LevelEditor/Private/LevelEditor.cpp b/Engine/Source/Editor/LevelEditor/Private/LevelEditor.cpp
--- a/Engine/Source/Editor/LevelEditor/Private/LevelEditor.cpp
+++ b/Engine/Source/Editor/LevelEditor/Private/LevelEditor.cpp
@@ -1795,7 +1795,22 @@ void FLevelEditorModule::BindGlobalLevelEditorCommands()
 		const bool IsDefaultActive = ShaderPlatform == GMaxRHIShaderPlatform;
 		const bool AllowPreview = !IsDefaultActive;
 
-		FPreviewPlatformInfo PreviewFeatureLevelInfo(FeatureLevel, IsDefaultActive ? NAME_None : Item.PlatformName, IsDefaultActive ? NAME_None : Item.ShaderFormat, IsDefaultActive ? NAME_None : Item.DeviceProfileName, AllowPreview);
+		EShaderPlatform InShaderPlatform = EShaderPlatform::SP_NumPlatforms;
+		for (uint32 j = 0; j < SP_NumPlatforms; ++j)
+		{
+			InShaderPlatform = (EShaderPlatform)j;
+			if (!FDataDrivenShaderPlatformInfo::IsValid(InShaderPlatform))
+			{
+				continue;
+			}
+
+			FName ShaderPlatformTest = FDataDrivenShaderPlatformInfo::GetName(InShaderPlatform);
+			if (Item.ShaderPlatformPreview == ShaderPlatformTest)
+			{
+				break;
+			}
+		}
+		FPreviewPlatformInfo PreviewFeatureLevelInfo(FeatureLevel, (EShaderPlatform)InShaderPlatform, IsDefaultActive ? NAME_None : Item.PlatformName, IsDefaultActive ? NAME_None : Item.ShaderFormat, IsDefaultActive ? NAME_None : Item.DeviceProfileName, AllowPreview, Item.ShaderPlatformPreview);
 
 		ActionList.MapAction(
 			Commands.PreviewPlatformOverrides[Index],
diff --git a/Engine/Source/Editor/LevelEditor/Private/LevelEditorToolBar.cpp b/Engine/Source/Editor/LevelEditor/Private/LevelEditorToolBar.cpp
--- a/Engine/Source/Editor/LevelEditor/Private/LevelEditorToolBar.cpp
+++ b/Engine/Source/Editor/LevelEditor/Private/LevelEditorToolBar.cpp
@@ -1403,7 +1403,7 @@ void FLevelEditorToolBar::RegisterLevelEditorToolBar( const TSharedRef<FUIComman
 			{
 				const FPreviewPlatformMenuItem* Item = FDataDrivenPlatformInfoRegistry::GetAllPreviewPlatformMenuItems().FindByPredicate([](const FPreviewPlatformMenuItem& TestItem)
 					{
-						return GEditor->PreviewPlatform.PreviewPlatformName == TestItem.PlatformName && GEditor->PreviewPlatform.PreviewShaderFormatName == TestItem.ShaderFormat;
+						return GEditor->PreviewPlatform.PreviewPlatformName == TestItem.PlatformName && GEditor->PreviewPlatform.PreviewShaderFormatName == TestItem.ShaderFormat && GEditor->PreviewPlatform.ShaderPlatformPreview == TestItem.ShaderPlatformPreview;
 					});
 				return Item ? Item->IconText : FText();
 			}
@@ -1438,7 +1438,7 @@ void FLevelEditorToolBar::RegisterLevelEditorToolBar( const TSharedRef<FUIComman
 			{
 				const FPreviewPlatformMenuItem* Item = FDataDrivenPlatformInfoRegistry::GetAllPreviewPlatformMenuItems().FindByPredicate([](const FPreviewPlatformMenuItem& TestItem)
 					{
-						return GEditor->PreviewPlatform.PreviewPlatformName == TestItem.PlatformName && GEditor->PreviewPlatform.PreviewShaderFormatName == TestItem.ShaderFormat;
+						return GEditor->PreviewPlatform.PreviewPlatformName == TestItem.PlatformName && GEditor->PreviewPlatform.PreviewShaderFormatName == TestItem.ShaderFormat && GEditor->PreviewPlatform.ShaderPlatformPreview == TestItem.ShaderPlatformPreview;
 					});
 				if(Item)
 				{
diff --git a/Engine/Source/Editor/UnrealEd/Classes/Editor/EditorEngine.h b/Engine/Source/Editor/UnrealEd/Classes/Editor/EditorEngine.h
--- a/Engine/Source/Editor/UnrealEd/Classes/Editor/EditorEngine.h
+++ b/Engine/Source/Editor/UnrealEd/Classes/Editor/EditorEngine.h
@@ -209,26 +209,33 @@ struct FPreviewPlatformInfo
 {
 	FPreviewPlatformInfo()
 	:	PreviewFeatureLevel(ERHIFeatureLevel::SM5)
+	,	ShaderPlatform(EShaderPlatform::SP_NumPlatforms)
 	,	PreviewPlatformName(NAME_None)
 	,	PreviewShaderFormatName(NAME_None)
 	,	bPreviewFeatureLevelActive(false)
+	,	ShaderPlatformPreview(NAME_None)
 	{}
 
-	FPreviewPlatformInfo(ERHIFeatureLevel::Type InFeatureLevel, FName InPreviewPlatformName = NAME_None, FName InPreviewShaderFormatName = NAME_None, FName InDeviceProfileName = NAME_None, bool InbPreviewFeatureLevelActive = false)
+	FPreviewPlatformInfo(ERHIFeatureLevel::Type InFeatureLevel, EShaderPlatform InShaderPlatform = EShaderPlatform::SP_NumPlatforms, FName InPreviewPlatformName = NAME_None, FName InPreviewShaderFormatName = NAME_None, FName InDeviceProfileName = NAME_None, bool InbPreviewFeatureLevelActive = false, FName InShaderPlatformName = NAME_None)
 	:	PreviewFeatureLevel(InFeatureLevel)
+	,	ShaderPlatform(InShaderPlatform)
 	,	PreviewPlatformName(InPreviewPlatformName)
 	,	PreviewShaderFormatName(InPreviewShaderFormatName)
 	,	DeviceProfileName(InDeviceProfileName)
 	,	bPreviewFeatureLevelActive(InbPreviewFeatureLevelActive)
+	,	ShaderPlatformPreview(InShaderPlatformName)
 	{}
 
 	/** The feature level we should use when loading or creating a new world */
 	ERHIFeatureLevel::Type PreviewFeatureLevel;
 
+	/** The ShaderPlatform to be used when in preview */
+	EShaderPlatform ShaderPlatform;
+
 	/** The the platform to preview, or NAME_None if there is no preview platform */
 	FName PreviewPlatformName;
 
-	/** The shader platform to preview, or NAME_None if there is no shader preview platform */
+	/** The shader Format to preview, or NAME_None if there is no shader preview format */
 	FName PreviewShaderFormatName;
 
 	/** The device profile to preview. */
@@ -236,11 +243,14 @@ struct FPreviewPlatformInfo
 
 	/** Is feature level preview currently active */
 	bool bPreviewFeatureLevelActive;
+	
+	/** The shader platform to preview*/
+	FName ShaderPlatformPreview;
 
 	/** Checks if two FPreviewPlatformInfos are for the same preview platform. Note, this does NOT compare the bPreviewFeatureLevelActive flag */
 	bool Matches(const FPreviewPlatformInfo& Other) const
 	{
-		return PreviewFeatureLevel == Other.PreviewFeatureLevel && PreviewPlatformName == Other.PreviewPlatformName && PreviewShaderFormatName == Other.PreviewShaderFormatName && DeviceProfileName == Other.DeviceProfileName;
+		return PreviewFeatureLevel == Other.PreviewFeatureLevel && ShaderPlatform == Other.ShaderPlatform && PreviewPlatformName == Other.PreviewPlatformName && PreviewShaderFormatName == Other.PreviewShaderFormatName && DeviceProfileName == Other.DeviceProfileName && ShaderPlatformPreview == Other.ShaderPlatformPreview;
 	}
 
 	/** Return platform name like "Android", or NAME_None if none is set or the preview feature level is not active */
@@ -532,6 +542,9 @@ public:
 	/** The feature level and platform we should use when loading or creating a new world */
 	FPreviewPlatformInfo PreviewPlatform;
 	
+	/** Cached ShaderPlatform so the editor can go back to after previewing */
+	EShaderPlatform CachedEditorShaderPlatform;
+
 	/** A delegate that is called when the preview feature level changes. Primarily used to switch a viewport's feature level. */
 	DECLARE_MULTICAST_DELEGATE_OneParam(FPreviewFeatureLevelChanged, ERHIFeatureLevel::Type);
 	FPreviewFeatureLevelChanged PreviewFeatureLevelChanged;
diff --git a/Engine/Source/Editor/UnrealEd/Classes/Editor/EditorPerProjectUserSettings.h b/Engine/Source/Editor/UnrealEd/Classes/Editor/EditorPerProjectUserSettings.h
--- a/Engine/Source/Editor/UnrealEd/Classes/Editor/EditorPerProjectUserSettings.h
+++ b/Engine/Source/Editor/UnrealEd/Classes/Editor/EditorPerProjectUserSettings.h
@@ -162,14 +162,22 @@ public:
 	UPROPERTY(config)
 	int32 PreviewFeatureLevel;
 
+	/** The Preview EShaderPlatform used for Preview */
+	UPROPERTY(config)
+	int32 PreviewShaderPlatform;
+
 	/** The platform name to preview, or NAME_None if there is no preview platform */
 	UPROPERTY(config)
 	FName PreviewPlatformName;
 
-	/** The shader platform to preview, or NAME_None if there is no preview platform */
+	/** The shader format to preview, or NAME_None if there is no preview platform */
 	UPROPERTY(config)
 	FName PreviewShaderFormatName;
 
+	/** The shader platform to preview, or NAME_None if there is no preview platform */
+	UPROPERTY(config)
+	FName PreviewShaderPlatformName;
+	
 	/** Is feature level preview currently active */
 	UPROPERTY(config)
 	bool bPreviewFeatureLevelActive;
diff --git a/Engine/Source/Editor/UnrealEd/Private/EditorEngine.cpp b/Engine/Source/Editor/UnrealEd/Private/EditorEngine.cpp
--- a/Engine/Source/Editor/UnrealEd/Private/EditorEngine.cpp
+++ b/Engine/Source/Editor/UnrealEd/Private/EditorEngine.cpp
@@ -378,7 +378,8 @@ UEditorEngine::UEditorEngine(const FObjectInitializer& ObjectInitializer)
 	bAllowMultiplePIEWorlds = true;
 	bIsEndingPlay = false;
 	DefaultWorldFeatureLevel = GMaxRHIFeatureLevel;
-	PreviewPlatform = FPreviewPlatformInfo(DefaultWorldFeatureLevel);
+	PreviewPlatform = FPreviewPlatformInfo(DefaultWorldFeatureLevel, GMaxRHIShaderPlatform);
+	CachedEditorShaderPlatform = GMaxRHIShaderPlatform;
 
 	FCoreDelegates::OnFeatureLevelDisabled.AddLambda([this](int RHIType, const FName& PreviewPlatformName)
 		{
@@ -391,7 +392,7 @@ UEditorEngine::UEditorEngine(const FObjectInitializer& ObjectInitializer)
 					return;
 				}
 				
-				SetPreviewPlatform(FPreviewPlatformInfo(GMaxRHIFeatureLevel), false);
+				SetPreviewPlatform(FPreviewPlatformInfo(GMaxRHIFeatureLevel, GMaxRHIShaderPlatform), false);
 			}
 		});
 
@@ -7655,14 +7656,13 @@ void UEditorEngine::OnSceneMaterialsModified()
 
 void UEditorEngine::SetPreviewPlatform(const FPreviewPlatformInfo& NewPreviewPlatform, bool bSaveSettings)
 {
-	// If we have specified a MaterialQualityPlatform ensure its feature level matches the requested feature level.
-	EShaderPlatform ShaderPlatform = ShaderFormatToLegacyShaderPlatform(NewPreviewPlatform.PreviewShaderFormatName);
+	// Get the requested preview platform, make sure it is valid.
+	EShaderPlatform ShaderPlatform = NewPreviewPlatform.ShaderPlatform;
+	check(FDataDrivenShaderPlatformInfo::IsValid(ShaderPlatform));
 	ERHIFeatureLevel::Type MaxFeatureLevel = NewPreviewPlatform.PreviewShaderFormatName != NAME_None ? (ERHIFeatureLevel::Type)GetMaxSupportedFeatureLevel(ShaderPlatform) : ERHIFeatureLevel::SM5;
 	check(NewPreviewPlatform.PreviewShaderFormatName.IsNone() || MaxFeatureLevel == NewPreviewPlatform.PreviewFeatureLevel);
 
-	const bool bChangedPreviewShaderPlatform = NewPreviewPlatform.PreviewShaderFormatName != PreviewPlatform.PreviewShaderFormatName;
-	const bool bChangedFeatureLevel = NewPreviewPlatform.PreviewFeatureLevel != PreviewPlatform.PreviewFeatureLevel ||
-		NewPreviewPlatform.bPreviewFeatureLevelActive != PreviewPlatform.bPreviewFeatureLevelActive;
+	const bool bChangedPreviewShaderPlatform = NewPreviewPlatform.ShaderPlatform != PreviewPlatform.ShaderPlatform;
 	const ERHIFeatureLevel::Type EffectiveFeatureLevel = NewPreviewPlatform.GetEffectivePreviewFeatureLevel();
 
 	if (NewPreviewPlatform.PreviewShaderFormatName != NAME_None)
@@ -7686,8 +7686,8 @@ void UEditorEngine::SetPreviewPlatform(const FPreviewPlatformInfo& NewPreviewPla
 	constexpr bool bUpdateProgressDialog = true;
 	constexpr bool bCacheAllRemainingShaders = false;
 
-	if (bChangedFeatureLevel)
 	{
+		// Set the correct SP preview for the FeatureLevel that is being previewed
 		FScopedSlowTask SlowTask(100.f, NSLOCTEXT("Engine", "ChangingPreviewPlatform", "Changing Preview Platform"), true);
 		SlowTask.Visibility = ESlowTaskVisibility::ForceVisible;
 		SlowTask.MakeDialog();
@@ -7708,6 +7708,8 @@ void UEditorEngine::SetPreviewPlatform(const FPreviewPlatformInfo& NewPreviewPla
 			UMaterialInterface::SetGlobalRequiredFeatureLevel(FeatureLevel, FeatureLevel == PreviewPlatform.PreviewFeatureLevel || FeatureLevel == GMaxRHIFeatureLevel);
 		}
 
+		GShaderPlatformForFeatureLevel[PreviewPlatform.PreviewFeatureLevel] = ShaderPlatform;
+
 		SlowTask.EnterProgressFrame(35.0f);
 		UMaterial::AllMaterialsCacheResourceShadersForRendering(bUpdateProgressDialog, bCacheAllRemainingShaders);
 
@@ -7723,15 +7725,6 @@ void UEditorEngine::SetPreviewPlatform(const FPreviewPlatformInfo& NewPreviewPla
 		DefaultWorldFeatureLevel = EffectiveFeatureLevel;
 		PreviewFeatureLevelChanged.Broadcast(EffectiveFeatureLevel);
 	}
-	else if (bChangedPreviewShaderPlatform)
-	{
-		// Rebuild materials if we have the same feature level but a different 'material quality platform'
-		FGlobalComponentRecreateRenderStateContext Recreate;
-		FlushRenderingCommands();
-		UMaterial::AllMaterialsCacheResourceShadersForRendering(bUpdateProgressDialog, bCacheAllRemainingShaders);
-		UMaterialInstance::AllMaterialsCacheResourceShadersForRendering(bUpdateProgressDialog, bCacheAllRemainingShaders);
-		PreviewFeatureLevelChanged.Broadcast(EffectiveFeatureLevel);
-	}
 
 	Scalability::ChangeScalabilityPreviewPlatform(PreviewPlatform.GetEffectivePreviewPlatformName());
 
@@ -7769,6 +7762,8 @@ void UEditorEngine::ToggleFeatureLevelPreview()
 
 	if (PreviewPlatform.bPreviewFeatureLevelActive)
 	{
+		GShaderPlatformForFeatureLevel[PreviewPlatform.PreviewFeatureLevel] = PreviewPlatform.ShaderPlatform;
+
 		if (PreviewPlatform.DeviceProfileName != NAME_None)
 		{
 			if (UDeviceProfile* DP = UDeviceProfileManager::Get().FindProfile(PreviewPlatform.DeviceProfileName.ToString(), false))
@@ -7783,6 +7778,11 @@ void UEditorEngine::ToggleFeatureLevelPreview()
 	}
 	else
 	{
+		// If the Preview FeatureLevel is the same as the Editor, restore the SP
+		if (PreviewPlatform.PreviewFeatureLevel == GMaxRHIFeatureLevel)
+		{
+			GShaderPlatformForFeatureLevel[PreviewPlatform.PreviewFeatureLevel] = CachedEditorShaderPlatform;
+		}
 		UDeviceProfileManager::Get().RestorePreviewDeviceProfile();
 	}
 
@@ -7797,7 +7797,7 @@ void UEditorEngine::ToggleFeatureLevelPreview()
 
 bool UEditorEngine::IsFeatureLevelPreviewEnabled() const
 {
-	return PreviewPlatform.PreviewFeatureLevel != GMaxRHIFeatureLevel || PreviewPlatform.PreviewShaderFormatName != NAME_None;
+	return PreviewPlatform.PreviewFeatureLevel != GMaxRHIFeatureLevel || PreviewPlatform.PreviewShaderFormatName != NAME_None || PreviewPlatform.ShaderPlatformPreview != NAME_None;
 }
 
 bool UEditorEngine::IsFeatureLevelPreviewActive() const
@@ -7813,7 +7813,7 @@ ERHIFeatureLevel::Type UEditorEngine::GetActiveFeatureLevelPreviewType() const
 void UEditorEngine::LoadEditorFeatureLevel()
 {
 	auto* Settings = GetMutableDefault<UEditorPerProjectUserSettings>();
-	if (Settings->PreviewFeatureLevel >= 0 && Settings->PreviewFeatureLevel < (int32)ERHIFeatureLevel::Num)
+	if (Settings->PreviewFeatureLevel >= 0 && Settings->PreviewFeatureLevel < (int32)ERHIFeatureLevel::Num && Settings->PreviewShaderPlatform < (int32)EShaderPlatform::SP_NumPlatforms)
 	{
 		// Try to map a saved ShaderFormatName to the PreviewPlatformName using ITargetPlatform if we don't have one. 
 		// We now store the PreviewPlatformName explicitly to support preview for platforms we don't have an ITargetPlatform of.
@@ -7826,7 +7826,13 @@ void UEditorEngine::LoadEditorFeatureLevel()
 			}
 		}
 
-		SetPreviewPlatform(FPreviewPlatformInfo((ERHIFeatureLevel::Type)Settings->PreviewFeatureLevel, Settings->PreviewPlatformName, Settings->PreviewShaderFormatName, Settings->PreviewDeviceProfileName, Settings->bPreviewFeatureLevelActive), false);
+		FName PreviewShaderPlatformName = Settings->PreviewShaderPlatformName;
+		// If we have an old .ini file in Saved we just will use what is the SP of the current FeatureLevel
+		if (Settings->PreviewShaderPlatformName == NAME_None)
+		{
+			PreviewShaderPlatformName = *(LexToString(GShaderPlatformForFeatureLevel[Settings->PreviewFeatureLevel]));
+		}
+		SetPreviewPlatform(FPreviewPlatformInfo((ERHIFeatureLevel::Type)Settings->PreviewFeatureLevel, (EShaderPlatform)Settings->PreviewShaderPlatform, Settings->PreviewPlatformName, Settings->PreviewShaderFormatName, Settings->PreviewDeviceProfileName, Settings->bPreviewFeatureLevelActive, PreviewShaderPlatformName), false);
 	}
 }
 
@@ -7834,11 +7840,13 @@ void UEditorEngine::SaveEditorFeatureLevel()
 {
 	auto* Settings = GetMutableDefault<UEditorPerProjectUserSettings>();
 	Settings->PreviewFeatureLevel = (int32)PreviewPlatform.PreviewFeatureLevel;
+	Settings->PreviewShaderPlatform = (int32)PreviewPlatform.ShaderPlatform;
 	Settings->PreviewPlatformName = PreviewPlatform.PreviewPlatformName;
 	Settings->PreviewShaderFormatName = PreviewPlatform.PreviewShaderFormatName;
 	Settings->bPreviewFeatureLevelActive = PreviewPlatform.bPreviewFeatureLevelActive;
 	Settings->bPreviewFeatureLevelWasDefault = (PreviewPlatform.PreviewFeatureLevel == GMaxRHIFeatureLevel);
 	Settings->PreviewDeviceProfileName = PreviewPlatform.DeviceProfileName;
+	Settings->PreviewShaderPlatformName = PreviewPlatform.ShaderPlatformPreview;
 	Settings->PostEditChange();
 }
 
diff --git a/Engine/Source/Editor/UnrealEd/Private/Settings/EditorPerProjectUserSettings.cpp b/Engine/Source/Editor/UnrealEd/Private/Settings/EditorPerProjectUserSettings.cpp
--- a/Engine/Source/Editor/UnrealEd/Private/Settings/EditorPerProjectUserSettings.cpp
+++ b/Engine/Source/Editor/UnrealEd/Private/Settings/EditorPerProjectUserSettings.cpp
@@ -28,10 +28,12 @@ UEditorPerProjectUserSettings::UEditorPerProjectUserSettings(const FObjectInitia
 	SwarmMaxUploadChunkSizeInMB = 100;
 	SwarmIntermediateFolder = FPaths::ConvertRelativePathToFull(FPaths::ProjectIntermediateDir() + TEXT("Simplygon/"));
 	PreviewFeatureLevel = GMaxRHIFeatureLevel;
+	PreviewPlatformName = NAME_None;
 	PreviewShaderFormatName = NAME_None;
 	bPreviewFeatureLevelActive = false;
 	bPreviewFeatureLevelWasDefault = true;
 	PreviewDeviceProfileName = NAME_None;
+	PreviewShaderPlatform = EShaderPlatform::SP_NumPlatforms;
 }
 
 void UEditorPerProjectUserSettings::PostInitProperties()
diff --git a/Engine/Source/Programs/ShaderCompileWorker/Private/ShaderCompileWorker.cpp b/Engine/Source/Programs/ShaderCompileWorker/Private/ShaderCompileWorker.cpp
--- a/Engine/Source/Programs/ShaderCompileWorker/Private/ShaderCompileWorker.cpp
+++ b/Engine/Source/Programs/ShaderCompileWorker/Private/ShaderCompileWorker.cpp
@@ -611,6 +611,14 @@ private:
 				InputFile << CompilerInput;
 				CompilerInput.DeserializeSharedInputs(InputFile, ExternalIncludes, SharedEnvironments, ParameterStructures);
 
+				// SCW doesn't run DDPI, GShaderHasCache Initialize is run  at start with no knowledge of the CustomPlatforms
+				// CustomPlatforms are known when we parse the WorkerInput so we populate the Directory here
+				if (IsCustomPlatform((EShaderPlatform)CompilerInput.Target.Platform))
+				{
+					const EShaderPlatform ShaderPlatform = ShaderFormatNameToShaderPlatform(CompilerInput.ShaderFormat);
+					UpdateIncludeDirectoryForPreviewPlatform((EShaderPlatform)CompilerInput.Target.Platform, ShaderPlatform);
+				}
+
 				if (IsValidRef(CompilerInput.SharedEnvironment))
 				{
 					// Merge the shared environment into the per-shader environment before calling into the compile function
@@ -657,6 +665,14 @@ private:
 					InputFile << CompilerInputs[StageIndex];
 					CompilerInputs[StageIndex].DeserializeSharedInputs(InputFile, ExternalIncludes, SharedEnvironments, ParameterStructures);
 
+					// SCW doesn't run DDPI, GShaderHasCache Initialize is run  at start with no knowledge of the CustomPlatforms
+					// CustomPlatforms are known when we parse the WorkerInput so we populate the Directory here
+					if (IsCustomPlatform((EShaderPlatform)CompilerInputs[StageIndex].Target.Platform))
+					{
+						const EShaderPlatform ShaderPlatform = ShaderFormatNameToShaderPlatform(CompilerInputs[StageIndex].ShaderFormat);
+						UpdateIncludeDirectoryForPreviewPlatform((EShaderPlatform)CompilerInputs[StageIndex].Target.Platform, ShaderPlatform);
+					}
+
 					if (IsValidRef(CompilerInputs[StageIndex].SharedEnvironment))
 					{
 						// Merge the shared environment into the per-shader environment before calling into the compile function
diff --git a/Engine/Source/Runtime/Core/Private/Misc/DataDrivenPlatformInfoRegistry.cpp b/Engine/Source/Runtime/Core/Private/Misc/DataDrivenPlatformInfoRegistry.cpp
--- a/Engine/Source/Runtime/Core/Private/Misc/DataDrivenPlatformInfoRegistry.cpp
+++ b/Engine/Source/Runtime/Core/Private/Misc/DataDrivenPlatformInfoRegistry.cpp
@@ -284,6 +284,7 @@ static void ParsePreviewPlatforms(const FConfigFile& IniFile)
 			Item.InactiveIconPath = GetSectionString(Section.Value, FName("InactiveIconPath"));
 			Item.InactiveIconName = *GetSectionString(Section.Value, FName("InactiveIconName"));
 			Item.DeviceProfileName = *GetSectionString(Section.Value, FName("DeviceProfileName"));
+			Item.ShaderPlatformPreview = *(GetSectionString(Section.Value, FName("ShaderPlatform")) + TEXT("_PREVIEW"));
 			FTextStringHelper::ReadFromBuffer(*GetSectionString(Section.Value, FName("MenuText")), Item.MenuText);
 			FTextStringHelper::ReadFromBuffer(*GetSectionString(Section.Value, FName("MenuTooltip")), Item.MenuTooltip);
 			FTextStringHelper::ReadFromBuffer(*GetSectionString(Section.Value, FName("IconText")), Item.IconText);
diff --git a/Engine/Source/Runtime/Core/Public/Misc/DataDrivenPlatformInfoRegistry.h b/Engine/Source/Runtime/Core/Public/Misc/DataDrivenPlatformInfoRegistry.h
--- a/Engine/Source/Runtime/Core/Public/Misc/DataDrivenPlatformInfoRegistry.h
+++ b/Engine/Source/Runtime/Core/Public/Misc/DataDrivenPlatformInfoRegistry.h
@@ -99,6 +99,7 @@ struct FPreviewPlatformMenuItem
 {
 	FName PlatformName;
 	FName ShaderFormat;
+	FName ShaderPlatformPreview;
 	FString ActiveIconPath;
 	FName ActiveIconName;
 	FString InactiveIconPath;
diff --git a/Engine/Source/Runtime/Engine/Private/Materials/MaterialShader.cpp b/Engine/Source/Runtime/Engine/Private/Materials/MaterialShader.cpp
--- a/Engine/Source/Runtime/Engine/Private/Materials/MaterialShader.cpp
+++ b/Engine/Source/Runtime/Engine/Private/Materials/MaterialShader.cpp
@@ -2198,7 +2198,7 @@ void FMaterialShaderMap::Compile(
 	FString WorkingDebugDescription = FString::Printf(
 		TEXT("Compiling %s: Platform=%s, Usage=%s"),
 		*NewContent->FriendlyName,
-		*LegacyShaderPlatformToShaderFormat(InPlatform).ToString(),
+		*FDataDrivenShaderPlatformInfo::GetName(InPlatform).ToString(),
 		*MaterialUsage
 	);
 	for (int32 StaticSwitchIndex = 0; StaticSwitchIndex < ShaderMapId.GetStaticSwitchParameters().Num(); ++StaticSwitchIndex)
diff --git a/Engine/Source/Runtime/Engine/Private/Materials/MaterialShared.cpp b/Engine/Source/Runtime/Engine/Private/Materials/MaterialShared.cpp
--- a/Engine/Source/Runtime/Engine/Private/Materials/MaterialShared.cpp
+++ b/Engine/Source/Runtime/Engine/Private/Materials/MaterialShared.cpp
@@ -2543,10 +2543,11 @@ bool FMaterial::CacheShaders(const FMaterialShaderMapId& ShaderMapId, EShaderPla
 					ShaderMapCondition = TEXT("Missing");
 				}
 #if WITH_EDITOR
+				FString ShaderPlatformName = FGenericDataDrivenShaderPlatformInfo::GetName(Platform).ToString();
 				UE_LOG(LogMaterial, Display, TEXT("%s cached shadermap for %s in %s, %s, %s (DDC key hash: %s), compiling. %s"),
 					ShaderMapCondition,
 					*GetAssetName(),
-					*LexToString(Platform),
+					*ShaderPlatformName,
 					*LexToString(ShaderMapId.QualityLevel),
 					*LexToString(ShaderMapId.FeatureLevel),
 					*DDCKeyHash,
diff --git a/Engine/Source/Runtime/Engine/Private/ShaderCompiler/ShaderCompiler.cpp b/Engine/Source/Runtime/Engine/Private/ShaderCompiler/ShaderCompiler.cpp
--- a/Engine/Source/Runtime/Engine/Private/ShaderCompiler/ShaderCompiler.cpp
+++ b/Engine/Source/Runtime/Engine/Private/ShaderCompiler/ShaderCompiler.cpp
@@ -4807,11 +4807,11 @@ bool FShaderCompilingManager::HandlePotentialRetryOnError(TMap<int32, FShaderMap
 				{
 					if (TargetShaderPlatformString.IsEmpty())
 					{
-						TargetShaderPlatformString = LegacyShaderPlatformToShaderFormat(ErrorPlatforms[PlatformIndex]).ToString();
+						TargetShaderPlatformString = FDataDrivenShaderPlatformInfo::GetName(ErrorPlatforms[PlatformIndex]).ToString();
 					}
 					else
 					{
-						TargetShaderPlatformString += FString(TEXT(", ")) + LegacyShaderPlatformToShaderFormat(ErrorPlatforms[PlatformIndex]).ToString();
+						TargetShaderPlatformString += FString(TEXT(", ")) + FDataDrivenShaderPlatformInfo::GetName(ErrorPlatforms[PlatformIndex]).ToString();
 					}
 				}
 
diff --git a/Engine/Source/Runtime/Engine/Private/ShaderCompiler/ShaderCompilerDistributed.cpp b/Engine/Source/Runtime/Engine/Private/ShaderCompiler/ShaderCompilerDistributed.cpp
--- a/Engine/Source/Runtime/Engine/Private/ShaderCompiler/ShaderCompilerDistributed.cpp
+++ b/Engine/Source/Runtime/Engine/Private/ShaderCompiler/ShaderCompilerDistributed.cpp
@@ -106,8 +106,8 @@ TArray<FString> FShaderCompileDistributedThreadRunnable_Interface::GetDependency
 	TArray<FShaderCommonCompileJobPtr>& Jobs)
 {
 	TArray<FString> Dependencies;
-	uint64 ShaderPlatformMask = 0;
-	static_assert(EShaderPlatform::SP_NumPlatforms <= sizeof(ShaderPlatformMask) * 8, "Insufficient bits in ShaderPlatformMask.");
+	TBitArray<> ShaderPlatformMask;
+	ShaderPlatformMask.Init(false, EShaderPlatform::SP_NumPlatforms);
 	for (const FShaderCommonCompileJobPtr& Job : Jobs)
 	{
 		EShaderPlatform ShaderPlatform = EShaderPlatform::SP_PCD3D_SM5;
@@ -139,9 +139,9 @@ TArray<FString> FShaderCompileDistributedThreadRunnable_Interface::GetDependency
 			}
 		}
 		// Add base dependencies for the platform only once.
-		if (!(ShaderPlatformMask & (static_cast<uint64>(1) << ShaderPlatform)))
+		if (!(ShaderPlatformMask[(int)ShaderPlatform]))
 		{
-			ShaderPlatformMask |= (static_cast<uint64>(1) << ShaderPlatform);
+			ShaderPlatformMask[(int)ShaderPlatform] = true;
 			TArray<FString>& ShaderPlatformCacheEntry = PlatformShaderInputFilesCache.FindOrAdd(ShaderPlatform);
 			if (!ShaderPlatformCacheEntry.Num())
 			{
diff --git a/Engine/Source/Runtime/Launch/Private/LaunchEngineLoop.cpp b/Engine/Source/Runtime/Launch/Private/LaunchEngineLoop.cpp
--- a/Engine/Source/Runtime/Launch/Private/LaunchEngineLoop.cpp
+++ b/Engine/Source/Runtime/Launch/Private/LaunchEngineLoop.cpp
@@ -2839,6 +2839,7 @@ PRAGMA_ENABLE_DEPRECATION_WARNINGS
 		SCOPED_BOOT_TIMING("RHIInit");
 		// Initialize the RHI.
 		RHIInit(bHasEditorToken);
+		FGenericDataDrivenShaderPlatformInfo::UpdatePreviewPlatforms();
 	}
 
 	{
diff --git a/Engine/Source/Runtime/RHI/Private/RHI.cpp b/Engine/Source/Runtime/RHI/Private/RHI.cpp
--- a/Engine/Source/Runtime/RHI/Private/RHI.cpp
+++ b/Engine/Source/Runtime/RHI/Private/RHI.cpp
@@ -1770,7 +1770,7 @@ FName ShaderPlatformToPlatformName(EShaderPlatform Platform)
 
 FName LegacyShaderPlatformToShaderFormat(EShaderPlatform Platform)
 {
-	return ShaderPlatformToShaderFormatName(Platform);
+	return  FDataDrivenShaderPlatformInfo::GetShaderFormat(Platform);
 }
 
 EShaderPlatform ShaderFormatToLegacyShaderPlatform(FName ShaderFormat)
@@ -2299,7 +2299,7 @@ void LexFromString(EShaderPlatform& Value, const TCHAR* String)
 {
 	Value = EShaderPlatform::SP_NumPlatforms;
 
-	for (uint8 i = 0; i < (uint8)EShaderPlatform::SP_NumPlatforms; ++i)
+	for (int i = 0; i < (int)EShaderPlatform::SP_NumPlatforms; ++i)
 	{
 		if (LexToString((EShaderPlatform)i, false).Equals(String, ESearchCase::IgnoreCase))
 		{
@@ -2437,6 +2437,7 @@ void FGenericDataDrivenShaderPlatformInfo::SetDefaultValues()
 	bSupportsAsyncPipelineCompilation = true;
 	bSupportsManualVertexFetch = true;
 	bSupportsVolumeTextureAtomics = true;
+	PreviewShaderPlatformParent = EShaderPlatform::SP_NumPlatforms;
 }
 
 void FGenericDataDrivenShaderPlatformInfo::ParseDataDrivenShaderInfo(const FConfigSection& Section, FGenericDataDrivenShaderPlatformInfo& Info)
@@ -2571,6 +2572,15 @@ void FGenericDataDrivenShaderPlatformInfo::Initialize()
 
 	// look for the standard DataDriven ini files
 	int32 NumDDInfoFiles = FDataDrivenPlatformInfoRegistry::GetNumDataDrivenIniFiles();
+	int32 CustomShaderPlatform = EShaderPlatform::SP_CUSTOM_PLATFORM_FIRST;
+
+	struct PlatformInfoAndPlatformEnum
+	{
+		FGenericDataDrivenShaderPlatformInfo Info;
+		EShaderPlatform ShaderPlatform;
+	};
+
+	TMap<FName, PlatformInfoAndPlatformEnum> ShaderPlatformNameToConfigSection;
 	for (int32 Index = 0; Index < NumDDInfoFiles; Index++)
 	{
 		FConfigFile IniFile;
@@ -2603,6 +2613,54 @@ void FGenericDataDrivenShaderPlatformInfo::Initialize()
 				Infos[ShaderPlatform].Name = *SectionName.Mid(15);
 				ParseDataDrivenShaderInfo(Section.Value, Infos[ShaderPlatform]);	
 				Infos[ShaderPlatform].bContainsValidPlatformInfo = true;
+
+				PlatformInfoAndPlatformEnum ItemToAdd;
+				ItemToAdd.Info = Infos[ShaderPlatform];
+				ItemToAdd.ShaderPlatform = ShaderPlatform;
+
+				ShaderPlatformNameToConfigSection.Add(Infos[ShaderPlatform].Name, ItemToAdd);
+			}
+		}
+	}
+
+	if (GIsEditor)
+	{
+		for (int32 Index = 0; Index < NumDDInfoFiles; Index++)
+		{
+			// find all PreviewPlatform sections for all platforms
+			// copy the [PreeviewPlatform]ShaderPlatform= customshaderplatform to a new customshaderplatform with name (orignal)_PREVIEW including its featurelevel.
+			// replace shader format with current editor's shader format.
+
+			FConfigFile IniFile;
+			FString PlatformName;
+
+			FDataDrivenPlatformInfoRegistry::LoadDataDrivenIniFile(Index, IniFile, PlatformName);
+
+			for (auto Section : IniFile)
+			{
+				if (Section.Key.StartsWith(TEXT("PreviewPlatform ")))
+				{
+					const FString& SectionName = Section.Key;
+					EShaderPlatform ShaderPlatform = EShaderPlatform(CustomShaderPlatform++);
+					// get enum value for the string name
+					//LexFromString(ShaderPlatform, *SectionName.Mid(15));
+					if (ShaderPlatform == EShaderPlatform::SP_NumPlatforms)
+					{
+						UE_LOG(LogRHI, Warning, TEXT("Found an unknown shader platform %s in a DataDriven ini file"), *SectionName.Mid(15));
+						continue;
+					}
+
+					// at this point, we can start pulling information out
+					/*Infos[ShaderPlatform].Name*/ 
+					FName ShaderPlatformName = *GetSectionString(Section.Value, FName("ShaderPlatform"));
+					FMemory::Memcpy(&Infos[ShaderPlatform], &ShaderPlatformNameToConfigSection[ShaderPlatformName].Info, sizeof(FGenericDataDrivenShaderPlatformInfo));
+
+					Infos[ShaderPlatform].Name = *(ShaderPlatformName.ToString() + TEXT("_PREVIEW"));
+					Infos[ShaderPlatform].PreviewShaderPlatformParent = ShaderPlatformNameToConfigSection[ShaderPlatformName].ShaderPlatform;
+
+					//set previewable to true, bContainsValidPlatformInfo is set in the update
+					Infos[ShaderPlatform].bIsPreviewPlatform = true;
+				}
 			}
 		}
 	}
@@ -2610,6 +2668,36 @@ void FGenericDataDrivenShaderPlatformInfo::Initialize()
 	bInitialized = true;
 }
 
+void FGenericDataDrivenShaderPlatformInfo::UpdatePreviewPlatforms()
+{
+	for (int i = 0; i < EShaderPlatform::SP_NumPlatforms; ++i)
+	{
+		EShaderPlatform ShaderPlatform = EShaderPlatform(i);
+		if (IsValid(ShaderPlatform))
+		{
+			if (Infos[ShaderPlatform].bIsPreviewPlatform)
+			{
+				Infos[ShaderPlatform].ShaderFormat = Infos[GShaderPlatformForFeatureLevel[Infos[ShaderPlatform].MaxFeatureLevel]].ShaderFormat;
+				Infos[ShaderPlatform].Language = Infos[GShaderPlatformForFeatureLevel[Infos[ShaderPlatform].MaxFeatureLevel]].Language;
+				Infos[ShaderPlatform].bIsHlslcc = Infos[GShaderPlatformForFeatureLevel[Infos[ShaderPlatform].MaxFeatureLevel]].bIsHlslcc;
+				Infos[ShaderPlatform].bSupportsDxc = Infos[GShaderPlatformForFeatureLevel[Infos[ShaderPlatform].MaxFeatureLevel]].bSupportsDxc;
+				Infos[ShaderPlatform].bSupportsGPUScene = Infos[GShaderPlatformForFeatureLevel[Infos[ShaderPlatform].MaxFeatureLevel]].bSupportsGPUScene;
+				Infos[ShaderPlatform].bIsPC = true;
+				Infos[ShaderPlatform].bIsConsole = false;
+				Infos[ShaderPlatform].bSupportsSceneDataCompressedTransforms = Infos[GShaderPlatformForFeatureLevel[Infos[ShaderPlatform].MaxFeatureLevel]].bSupportsSceneDataCompressedTransforms;
+				Infos[ShaderPlatform].bSupportsNanite = Infos[GShaderPlatformForFeatureLevel[Infos[ShaderPlatform].MaxFeatureLevel]].bSupportsNanite;
+				Infos[ShaderPlatform].bSupportsUInt64ImageAtomics = Infos[GShaderPlatformForFeatureLevel[Infos[ShaderPlatform].MaxFeatureLevel]].bSupportsUInt64ImageAtomics;
+				Infos[ShaderPlatform].bSupportsGen5TemporalAA = Infos[GShaderPlatformForFeatureLevel[Infos[ShaderPlatform].MaxFeatureLevel]].bSupportsGen5TemporalAA;
+				Infos[ShaderPlatform].bSupportsRenderTargetWriteMask = false;
+				Infos[ShaderPlatform].bSupportsIntrinsicWaveOnce = false;
+				Infos[ShaderPlatform].bSupportsDOFHybridScattering = false;
+				Infos[ShaderPlatform].bSupports4ComponentUAVReadWrite = false;
+				Infos[ShaderPlatform].bContainsValidPlatformInfo = true;
+			}
+		}
+	}
+}
+
 //
 //	MSAA sample offsets.
 //
diff --git a/Engine/Source/Runtime/RHI/Public/RHIDefinitions.h b/Engine/Source/Runtime/RHI/Public/RHIDefinitions.h
--- a/Engine/Source/Runtime/RHI/Public/RHIDefinitions.h
+++ b/Engine/Source/Runtime/RHI/Public/RHIDefinitions.h
@@ -119,8 +119,11 @@ enum EShaderPlatform
 	SP_PCD3D_SM6,
 	SP_D3D_ES3_1_HOLOLENS,
 
+	SP_CUSTOM_PLATFORM_FIRST,
+	SP_CUSTOM_PLATFORM_LAST = (SP_CUSTOM_PLATFORM_FIRST + 100),
+
 	SP_NumPlatforms,
-	SP_NumBits						= 7,
+	SP_NumBits						= 16,
 };
 static_assert(SP_NumPlatforms <= (1 << SP_NumBits), "SP_NumPlatforms will not fit on SP_NumBits");
 
@@ -336,6 +339,7 @@ class RHI_API FGenericDataDrivenShaderPlatformInfo
 	FName Language;
 	ERHIFeatureLevel::Type MaxFeatureLevel;
 	FName ShaderFormat;
+	EShaderPlatform PreviewShaderPlatformParent;
 	uint32 ShaderPropertiesHash;
 	uint32 bIsMobile: 1;
 	uint32 bIsMetalMRT: 1;
@@ -421,6 +425,7 @@ class RHI_API FGenericDataDrivenShaderPlatformInfo
 	uint32 bSupportsRealTypes : int32(ERHIFeatureSupport::NumBits);
 	uint32 EnablesHLSL2021ByDefault : 2; // 0: disabled, 1: global shaders only, 2: all shaders
 	uint32 bSupportsSceneDataCompressedTransforms : 1;
+	uint32 bIsPreviewPlatform : 1;
 		
 #if WITH_EDITOR
 	FText FriendlyName;
@@ -440,6 +445,7 @@ class RHI_API FGenericDataDrivenShaderPlatformInfo
 
 public:
 	static void Initialize();
+	static void UpdatePreviewPlatforms();
 	static void ParseDataDrivenShaderInfo(const FConfigSection& Section, FGenericDataDrivenShaderPlatformInfo& Info);
 
 	static FORCEINLINE_DEBUGGABLE const FName GetName(const FStaticShaderPlatform Platform)
@@ -453,6 +459,12 @@ public:
 		check(IsValid(Platform));
 		return Infos[Platform].ShaderFormat;
 	}
+	
+	static FORCEINLINE_DEBUGGABLE const EShaderPlatform GetPreviewShaderPlatformParent(const FStaticShaderPlatform Platform)
+	{
+		check(IsValid(Platform));
+		return Infos[Platform].PreviewShaderPlatformParent;
+	}
 
 	static FORCEINLINE_DEBUGGABLE uint32 GetShaderPlatformPropertiesHash(const FStaticShaderPlatform Platform)
 	{
@@ -997,6 +1009,11 @@ public:
 	{
 		return Infos[Platform].bSupportsOIT;
 	}
+	
+	static FORCEINLINE_DEBUGGABLE const bool GetIsPreviewPlatform(const FStaticShaderPlatform Platform)
+	{
+		return Infos[Platform].bIsPreviewPlatform;
+	}
 
 	static FORCEINLINE_DEBUGGABLE const ERHIFeatureSupport GetSupportsRealTypes(const FStaticShaderPlatform Platform)
 	{
@@ -2088,6 +2105,11 @@ inline bool IsMobilePlatform(const FStaticShaderPlatform Platform)
 	return FDataDrivenShaderPlatformInfo::GetMaxFeatureLevel(Platform) == ERHIFeatureLevel::ES3_1;
 }
 
+inline bool IsCustomPlatform(const FStaticShaderPlatform Platform)
+{
+	return (Platform >= EShaderPlatform::SP_CUSTOM_PLATFORM_FIRST) && (Platform < EShaderPlatform::SP_CUSTOM_PLATFORM_LAST);
+}
+
 inline bool IsOpenGLPlatform(const FStaticShaderPlatform Platform)
 {
 	return FDataDrivenShaderPlatformInfo::GetIsLanguageOpenGL(Platform);
@@ -2186,35 +2208,14 @@ inline FStaticFeatureLevel GetMaxSupportedFeatureLevel(const FStaticShaderPlatfo
 /* Returns true if the shader platform Platform is used to simulate a mobile feature level on a PC platform. */
 inline bool IsSimulatedPlatform(const FStaticShaderPlatform Platform)
 {
-	switch (Platform)
-	{
-		case SP_PCD3D_ES3_1:
-		case SP_OPENGL_PCES3_1:
-		case SP_METAL_MACES3_1:
-		case SP_VULKAN_PCES3_1:
-			return true;
-		break;
-
-		default:
-			return false;
-		break;
-	}
-
-	return false;
+	return FDataDrivenShaderPlatformInfo::GetIsPreviewPlatform(Platform);
 }
 
 inline EShaderPlatform GetSimulatedPlatform(FStaticShaderPlatform Platform)
 {
-	switch (Platform)
+	if (IsSimulatedPlatform(Platform))
 	{
-		case SP_PCD3D_ES3_1:
-		case SP_OPENGL_PCES3_1:
-			return SP_OPENGL_ES3_1_ANDROID;
-		break;
-
-		default:
-			return Platform;
-		break;
+		return FDataDrivenShaderPlatformInfo::GetPreviewShaderPlatformParent(Platform);
 	}
 
 	return Platform;
diff --git a/Engine/Source/Runtime/RenderCore/Private/RenderUtils.cpp b/Engine/Source/Runtime/RenderCore/Private/RenderUtils.cpp
--- a/Engine/Source/Runtime/RenderCore/Private/RenderUtils.cpp
+++ b/Engine/Source/Runtime/RenderCore/Private/RenderUtils.cpp
@@ -1127,12 +1127,11 @@ RENDERCORE_API uint32 GetPlatformShadingModelsMask(const FStaticShaderPlatform P
 	return 0xFFFFFFFF;
 }
 
-RENDERCORE_API uint64 GMobileAmbientOcclusionPlatformMask = 0;
-static_assert(SP_NumPlatforms <= sizeof(GMobileAmbientOcclusionPlatformMask) * 8, "GMobileAmbientOcclusionPlatformMask must be large enough to support all shader platforms");
+RENDERCORE_API ShaderPlatformMaskType GMobileAmbientOcclusionPlatformMask;
 
 RENDERCORE_API bool IsMobileAmbientOcclusionEnabled(const FStaticShaderPlatform Platform)
 {
-	return (IsMobilePlatform(Platform) && !!(GMobileAmbientOcclusionPlatformMask & (1ull << Platform)));
+	return IsMobilePlatform(Platform) && GMobileAmbientOcclusionPlatformMask[(int)Platform];
 }
 
 RENDERCORE_API bool IsMobileDistanceFieldEnabled(const FStaticShaderPlatform Platform)
@@ -1241,29 +1240,17 @@ static TAutoConsoleVariable<int32> CVarDistanceFields(
 	); 
 
 
-RENDERCORE_API uint64 GForwardShadingPlatformMask = 0;
-static_assert(SP_NumPlatforms <= sizeof(GForwardShadingPlatformMask) * 8, "GForwardShadingPlatformMask must be large enough to support all shader platforms");
-
-RENDERCORE_API uint64 GDBufferPlatformMask = 0;
-static_assert(SP_NumPlatforms <= sizeof(GDBufferPlatformMask) * 8, "GDBufferPlatformMask must be large enough to support all shader platforms");
-
-RENDERCORE_API uint64 GVelocityEncodeDepthPlatformMask = 0;
-static_assert(SP_NumPlatforms <= sizeof(GVelocityEncodeDepthPlatformMask) * 8, "GVelocityEncodeDepthPlatformMask must be large enough to support all shader platforms");
-
-RENDERCORE_API uint64 GSelectiveBasePassOutputsPlatformMask = 0;
-static_assert(SP_NumPlatforms <= sizeof(GSelectiveBasePassOutputsPlatformMask) * 8, "GSelectiveBasePassOutputsPlatformMask must be large enough to support all shader platforms");
-
-RENDERCORE_API uint64 GDistanceFieldsPlatformMask = 0;
-static_assert(SP_NumPlatforms <= sizeof(GDistanceFieldsPlatformMask) * 8, "GDistanceFieldsPlatformMask must be large enough to support all shader platforms");
-
-RENDERCORE_API uint64 GSimpleSkyDiffusePlatformMask = 0;
-static_assert(SP_NumPlatforms <= sizeof(GSimpleSkyDiffusePlatformMask) * 8, "GSimpleSkyDiffusePlatformMask must be large enough to support all shader platforms");
+RENDERCORE_API ShaderPlatformMaskType GForwardShadingPlatformMask;
+RENDERCORE_API ShaderPlatformMaskType GDBufferPlatformMask;
+RENDERCORE_API ShaderPlatformMaskType GVelocityEncodeDepthPlatformMask;
+RENDERCORE_API ShaderPlatformMaskType GSelectiveBasePassOutputsPlatformMask;
+RENDERCORE_API ShaderPlatformMaskType GDistanceFieldsPlatformMask;
+RENDERCORE_API ShaderPlatformMaskType GSimpleSkyDiffusePlatformMask;
 
 // Specifies whether ray tracing *can* be enabled on a particular platform.
 // This takes into account whether RT is globally enabled for the project and specifically enabled on a target platform.
 // Safe to use to make cook-time decisions, such as whether to compile ray tracing shaders.
-RENDERCORE_API uint64 GRayTracingPlaformMask = 0;
-static_assert(SP_NumPlatforms <= sizeof(GRayTracingPlaformMask) * 8, "GRayTracingPlaformMask must be large enough to support all shader platforms");
+RENDERCORE_API ShaderPlatformMaskType GRayTracingPlaformMask;
 
 // Specifies whether ray tracing *is* enabled on the current running system (in current game or editor process).
 // This takes into account additional factors, such as concrete current GPU/OS/Driver capability, user-set game graphics options, etc.
@@ -1271,40 +1258,46 @@ static_assert(SP_NumPlatforms <= sizeof(GRayTracingPlaformMask) * 8, "GRayTracin
 // Value may be queried using IsRayTracingEnabled().
 RENDERCORE_API bool GUseRayTracing = false;
 
+void GetAllPossiblePreviewPlatformsForMainShaderPlatform(TArray<EShaderPlatform>& OutPreviewPlatforms, FName ShaderPlatformName)
+{
+	FName PreviewPlatformName = FName(ShaderPlatformName.ToString() + TEXT("_PREVIEW"));
+	for (int i = 0; i < EShaderPlatform::SP_NumPlatforms; ++i)
+	{
+		EShaderPlatform ShaderPlatform = EShaderPlatform(i);
+		if (FDataDrivenShaderPlatformInfo::IsValid(ShaderPlatform))
+		{
+			bool bIsPreviewPlatform = FDataDrivenShaderPlatformInfo::GetIsPreviewPlatform(ShaderPlatform) && (FDataDrivenShaderPlatformInfo::GetName(ShaderPlatform) == PreviewPlatformName);
+			if (bIsPreviewPlatform)
+			{
+				OutPreviewPlatforms.Add(ShaderPlatform);
+			}
+		}
+	}
+}
+
 RENDERCORE_API void RenderUtilsInit()
 {
 	checkf(GIsRHIInitialized, TEXT("RenderUtilsInit() may only be called once RHI is initialized."));
 
-	if (GUseForwardShading)
-	{
-		GForwardShadingPlatformMask = ~0ull;
-	}
+	GForwardShadingPlatformMask.Init(GUseForwardShading == 1, EShaderPlatform::SP_NumPlatforms);
 
 	static IConsoleVariable* DBufferVar = IConsoleManager::Get().FindConsoleVariable(TEXT("r.DBuffer"));
-	if (DBufferVar && DBufferVar->GetInt())
-	{
-		GDBufferPlatformMask = ~0ull;
-	}
+	GDBufferPlatformMask.Init(DBufferVar && DBufferVar->GetInt(), EShaderPlatform::SP_NumPlatforms);
 
 	static IConsoleVariable* SelectiveBasePassOutputsCVar = IConsoleManager::Get().FindConsoleVariable(TEXT("r.SelectiveBasePassOutputs"));
-	if (SelectiveBasePassOutputsCVar && SelectiveBasePassOutputsCVar->GetInt())
-	{
-		GSelectiveBasePassOutputsPlatformMask = ~0ull;
-	}
+	GSelectiveBasePassOutputsPlatformMask.Init(SelectiveBasePassOutputsCVar && SelectiveBasePassOutputsCVar->GetInt(), EShaderPlatform::SP_NumPlatforms);
 
 	static IConsoleVariable* DistanceFieldsCVar = IConsoleManager::Get().FindConsoleVariable(TEXT("r.DistanceFields")); 
-	if (DistanceFieldsCVar && DistanceFieldsCVar->GetInt())
-	{
-		GDistanceFieldsPlatformMask = ~0ull;
-	}
+	GDistanceFieldsPlatformMask.Init(DistanceFieldsCVar && DistanceFieldsCVar->GetInt(), EShaderPlatform::SP_NumPlatforms);
 
 	static IConsoleVariable* RayTracingCVar = IConsoleManager::Get().FindConsoleVariable(TEXT("r.RayTracing"));
 
+	GSimpleSkyDiffusePlatformMask.Init(false, EShaderPlatform::SP_NumPlatforms);
+	GVelocityEncodeDepthPlatformMask.Init(false, EShaderPlatform::SP_NumPlatforms);
+	GRayTracingPlaformMask.Init(false, EShaderPlatform::SP_NumPlatforms);
+
 	static IConsoleVariable* MobileAmbientOcclusionCVar = IConsoleManager::Get().FindConsoleVariable(TEXT("r.Mobile.AmbientOcclusion"));
-	if (MobileAmbientOcclusionCVar && MobileAmbientOcclusionCVar->GetInt())
-	{
-		GMobileAmbientOcclusionPlatformMask = ~0ull;
-	}
+	GMobileAmbientOcclusionPlatformMask.Init(MobileAmbientOcclusionCVar && MobileAmbientOcclusionCVar->GetInt(), EShaderPlatform::SP_NumPlatforms);
 
 #if WITH_EDITOR
 	ITargetPlatformManagerModule* TargetPlatformManager = GetTargetPlatformManager();
@@ -1318,71 +1311,34 @@ RENDERCORE_API void RenderUtilsInit()
 			for (FName Format : PlaformPossibleShaderFormats)
 			{
 				EShaderPlatform ShaderPlatform = ShaderFormatNameToShaderPlatform(Format);
-				uint32 ShaderPlatformIndex = static_cast<uint32>(ShaderPlatform);
-
-				uint64 Mask = 1ull << ShaderPlatformIndex;
+				TArray<EShaderPlatform> PossiblePreviewPlatformsAndMainPlatform;
+				GetAllPossiblePreviewPlatformsForMainShaderPlatform(PossiblePreviewPlatformsAndMainPlatform, FDataDrivenShaderPlatformInfo::GetName(ShaderPlatform));
+				PossiblePreviewPlatformsAndMainPlatform.Add(ShaderPlatform);
 
-				if (TargetPlatform->UsesForwardShading())
-				{
-					GForwardShadingPlatformMask |= Mask;
-				}
-				else
+				for (EShaderPlatform ShaderPlatformToEdit : PossiblePreviewPlatformsAndMainPlatform)
 				{
-					GForwardShadingPlatformMask &= ~Mask;
-				}
+					uint32 ShaderPlatformIndex = static_cast<uint32>(ShaderPlatformToEdit);
 
-				if (TargetPlatform->UsesDBuffer() && !IsMobilePlatform(ShaderPlatform))
-				{
-					GDBufferPlatformMask |= Mask;
-				}
-				else
-				{
-					GDBufferPlatformMask &= ~Mask;
-				}
+					uint64 Mask = 1ull << ShaderPlatformToEdit;
 
-				if (TargetPlatform->UsesSelectiveBasePassOutputs())
-				{
-					GSelectiveBasePassOutputsPlatformMask |= Mask;
-				}
-				else
-				{
-					GSelectiveBasePassOutputsPlatformMask &= ~Mask;
-				}
+					if (!FDataDrivenShaderPlatformInfo::IsValid(ShaderPlatformToEdit))
+					{
+						continue;
+					}
 
-				if (TargetPlatform->UsesDistanceFields())
-				{
-					GDistanceFieldsPlatformMask |= Mask;
-				}
-				else
-				{
-					GDistanceFieldsPlatformMask &= ~Mask;
-				}
+					GForwardShadingPlatformMask[ShaderPlatformIndex] = TargetPlatform->UsesForwardShading();
 
-				if (TargetPlatform->ForcesSimpleSkyDiffuse())
-				{
-					GSimpleSkyDiffusePlatformMask |= Mask;
-				}
-				else
-				{
-					GSimpleSkyDiffusePlatformMask &= ~Mask;
-				}
+					GDBufferPlatformMask[ShaderPlatformIndex] = TargetPlatform->UsesDBuffer() && !IsMobilePlatform(ShaderPlatformToEdit);
 
-				if (TargetPlatform->VelocityEncodeDepth())
-				{
-					GVelocityEncodeDepthPlatformMask |= Mask;
-				}
-				else
-				{
-					GVelocityEncodeDepthPlatformMask &= ~Mask;
-				}
+					GSelectiveBasePassOutputsPlatformMask[ShaderPlatformIndex] = TargetPlatform->UsesSelectiveBasePassOutputs();
 
-				if (TargetPlatform->UsesMobileAmbientOcclusion())
-				{
-					GMobileAmbientOcclusionPlatformMask |= Mask;
-				}
-				else
-				{
-					GMobileAmbientOcclusionPlatformMask &= ~Mask;
+					GDistanceFieldsPlatformMask[ShaderPlatformIndex] = TargetPlatform->UsesDistanceFields();
+
+					GSimpleSkyDiffusePlatformMask[ShaderPlatformIndex] = TargetPlatform->ForcesSimpleSkyDiffuse();
+
+					GVelocityEncodeDepthPlatformMask[ShaderPlatformIndex] = TargetPlatform->VelocityEncodeDepth();
+
+					GMobileAmbientOcclusionPlatformMask[ShaderPlatformIndex] = TargetPlatform->UsesMobileAmbientOcclusion();
 				}
 			}
 
@@ -1397,9 +1353,7 @@ RENDERCORE_API void RenderUtilsInit()
 					EShaderPlatform ShaderPlatform = ShaderFormatNameToShaderPlatform(Format);
 					uint32 ShaderPlatformIndex = static_cast<uint32>(ShaderPlatform);
 
-					uint64 Mask = 1ull << ShaderPlatformIndex;
-
-					GRayTracingPlaformMask |= Mask;
+					GRayTracingPlaformMask[ShaderPlatformIndex] = true;
 				}
 			}
 		}
@@ -1409,12 +1363,12 @@ RENDERCORE_API void RenderUtilsInit()
 
 	if (IsMobilePlatform(GMaxRHIShaderPlatform))
 	{
-		GDBufferPlatformMask = 0;
+		GDBufferPlatformMask.Init(false, EShaderPlatform::SP_NumPlatforms);
 	}
 
 	if (RayTracingCVar && RayTracingCVar->GetInt() && GRHISupportsRayTracing)
 	{
-		GRayTracingPlaformMask = ~0ull;
+		GRayTracingPlaformMask.Init(true, EShaderPlatform::SP_NumPlatforms);
 	}
 
 	// Load runtime values from and *.ini file used by a current platform
@@ -1431,13 +1385,13 @@ RENDERCORE_API void RenderUtilsInit()
 			bool bDistanceFields = false;
 			if (PlatformIniFile.GetBool(*CategoryName, TEXT("bEnableDistanceFields"), bDistanceFields) && !bDistanceFields)
 			{
-				GDistanceFieldsPlatformMask = 0;
+				GDistanceFieldsPlatformMask.Init(false, EShaderPlatform::SP_NumPlatforms);
 			}
 
 			bool bRayTracing = false;
 			if (PlatformIniFile.GetBool(*CategoryName, TEXT("bEnableRayTracing"), bRayTracing) && !bRayTracing)
 			{
-				GRayTracingPlaformMask = 0;
+				GRayTracingPlaformMask.Init(false, EShaderPlatform::SP_NumPlatforms);
 			}
 		}
 	}
@@ -1456,7 +1410,7 @@ RENDERCORE_API void RenderUtilsInit()
 
 	if (RayTracingCVar && RayTracingCVar->GetBool())
 	{
-		const bool bRayTracingAllowedOnCurrentPlatform = !!(GRayTracingPlaformMask & (1ull << GMaxRHIShaderPlatform));
+		const bool bRayTracingAllowedOnCurrentPlatform = (GRayTracingPlaformMask[(int)GMaxRHIShaderPlatform]);
 		if (GRHISupportsRayTracing && bRayTracingAllowedOnCurrentPlatform)
 		{
 			if (GIsEditor)
@@ -1755,8 +1709,8 @@ RENDERCORE_API bool DoesPlatformSupportNanite(EShaderPlatform Platform, bool bCh
 /** Returns whether DBuffer decals are enabled for a given shader platform */
 RENDERCORE_API bool IsUsingDBuffers(const FStaticShaderPlatform Platform)
 {
-	extern RENDERCORE_API uint64 GDBufferPlatformMask;
-	return !!(GDBufferPlatformMask & (1ull << Platform));
+	extern RENDERCORE_API ShaderPlatformMaskType GDBufferPlatformMask;
+	return (GDBufferPlatformMask[(int)Platform]);
 }
 
 RENDERCORE_API bool AreSkinCacheShadersEnabled(EShaderPlatform Platform)
diff --git a/Engine/Source/Runtime/RenderCore/Private/ShaderCodeLibrary.cpp b/Engine/Source/Runtime/RenderCore/Private/ShaderCodeLibrary.cpp
--- a/Engine/Source/Runtime/RenderCore/Private/ShaderCodeLibrary.cpp
+++ b/Engine/Source/Runtime/RenderCore/Private/ShaderCodeLibrary.cpp
@@ -1880,7 +1880,7 @@ class FShaderLibrariesCollection
 	// At cook time, shader code collection for each shader platform
 	FEditorShaderStableInfo* EditorShaderStableInfo[EShaderPlatform::SP_NumPlatforms];
 	// Cached bit field for shader formats that require stable keys
-	uint64_t bShaderFormatsThatNeedStableKeys = 0;
+	TBitArray<> ShaderFormatsThatNeedStableKeys;
 	// At cook time, shader stats for each shader platform
 	FShaderCodeStats EditorShaderCodeStats[EShaderPlatform::SP_NumPlatforms];
 	// At cook time, whether the shader archive supports pipelines (only OpenGL should)
@@ -1913,6 +1913,7 @@ public:
 		, bNativeFormat(bInNativeFormat)
 	{
 #if WITH_EDITOR
+		ShaderFormatsThatNeedStableKeys.Init(false, EShaderPlatform::SP_NumPlatforms);
 		FMemory::Memzero(EditorShaderCodeArchive);
 		FMemory::Memzero(EditorShaderStableInfo);
 		FMemory::Memzero(EditorShaderCodeStats);
@@ -2371,8 +2372,7 @@ public:
 			{
 				StableArchive = new FEditorShaderStableInfo(PossiblyAdjustedFormat);
 				EditorShaderStableInfo[Platform] = StableArchive;
-				bShaderFormatsThatNeedStableKeys |= (uint64_t(1u) << (uint32_t)Platform);
-				static_assert(SP_NumPlatforms < 64u, "ShaderPlatform will no longer fit into bitfield.");
+				ShaderFormatsThatNeedStableKeys[(int)Platform] = true;
 			}
 		}
 	}
@@ -2381,9 +2381,9 @@ public:
 	{
 		if (Platform == EShaderPlatform::SP_NumPlatforms)
 		{
-			return bShaderFormatsThatNeedStableKeys != 0;
+			return ShaderFormatsThatNeedStableKeys.Find(true) != INDEX_NONE;
 		}
-		return (bShaderFormatsThatNeedStableKeys & (uint64_t(1u) << (uint32_t) Platform)) != 0;
+		return (ShaderFormatsThatNeedStableKeys[(int)Platform]) != 0;
 	}
 
 	void AddShaderCode(EShaderPlatform Platform, const FShaderMapResourceCode* Code, const FShaderMapAssetPaths& AssociatedAssets)
diff --git a/Engine/Source/Runtime/RenderCore/Private/ShaderCore.cpp b/Engine/Source/Runtime/RenderCore/Private/ShaderCore.cpp
--- a/Engine/Source/Runtime/RenderCore/Private/ShaderCore.cpp
+++ b/Engine/Source/Runtime/RenderCore/Private/ShaderCore.cpp
@@ -161,6 +161,11 @@ public:
 		bInitialized = true;
 	}
 
+	void UpdateIncludeDirectoryForPreviewPlatform(EShaderPlatform PreviewShaderPlatform, EShaderPlatform ParentShaderPlatform)
+	{
+		Platforms[PreviewShaderPlatform].IncludeDirectory = Platforms[ParentShaderPlatform].IncludeDirectory;
+	}
+
 	FSHAHash* FindHash(EShaderPlatform ShaderPlatform, const FString& VirtualFilePath)
 	{
 		check(ShaderPlatform < UE_ARRAY_COUNT(Platforms));
@@ -1400,6 +1405,12 @@ void InitializeShaderHashCache()
 	GShaderHashCache.Initialize();
 }
 
+void UpdateIncludeDirectoryForPreviewPlatform(EShaderPlatform PreviewPlatform, EShaderPlatform ActualPlatform)
+{
+	FScopeLock ShaderHashAccessLock(&GShaderHashAccessGuard);
+	GShaderHashCache.UpdateIncludeDirectoryForPreviewPlatform(PreviewPlatform, ActualPlatform);
+}
+
 void CheckShaderHashCacheInclude(const FString& VirtualFilePath, EShaderPlatform ShaderPlatform)
 {
 	FScopeLock ShaderHashAccessLock(&GShaderHashAccessGuard);
diff --git a/Engine/Source/Runtime/RenderCore/Public/RenderResource.h b/Engine/Source/Runtime/RenderCore/Public/RenderResource.h
--- a/Engine/Source/Runtime/RenderCore/Public/RenderResource.h
+++ b/Engine/Source/Runtime/RenderCore/Public/RenderResource.h
@@ -27,6 +27,8 @@
 #include "RenderingThread.h"
 #include "Serialization/MemoryLayout.h"
 
+typedef TBitArray<TInlineAllocator<EShaderPlatform::SP_NumPlatforms / 8>>	ShaderPlatformMaskType;
+
 /** Number of frames after which unused global resource allocations will be discarded. */
 extern int32 GGlobalBufferNumFramesUnusedThresold;
 
@@ -653,8 +655,8 @@ FORCEINLINE bool IsRayTracingEnabledForProject(EShaderPlatform ShaderPlatform)
 {
 	if (RHISupportsRayTracing(ShaderPlatform))
 	{
-		extern RENDERCORE_API uint64 GRayTracingPlaformMask;
-		return !!(GRayTracingPlaformMask & (1ull << ShaderPlatform));
+		extern RENDERCORE_API TBitArray<TInlineAllocator<EShaderPlatform::SP_NumPlatforms / 8>> GRayTracingPlaformMask;
+		return (GRayTracingPlaformMask[(int)ShaderPlatform]);
 	}
 	else
 	{
diff --git a/Engine/Source/Runtime/RenderCore/Public/RenderUtils.h b/Engine/Source/Runtime/RenderCore/Public/RenderUtils.h
--- a/Engine/Source/Runtime/RenderCore/Public/RenderUtils.h
+++ b/Engine/Source/Runtime/RenderCore/Public/RenderUtils.h
@@ -540,9 +540,18 @@ struct RENDERCORE_API FShaderPlatformCachedIniValue
 
 	Type Get(EShaderPlatform ShaderPlatform)
 	{
+		EShaderPlatform ActualShaderPlatform = ShaderPlatform;
 		Type Value{};
 
-		FName IniPlatformName = ShaderPlatformToPlatformName(ShaderPlatform);
+		if (GIsEditor)
+		{
+			if (FDataDrivenShaderPlatformInfo::GetIsPreviewPlatform(ActualShaderPlatform))
+			{
+				ActualShaderPlatform = FDataDrivenShaderPlatformInfo::GetPreviewShaderPlatformParent(ActualShaderPlatform);
+			}
+		}
+
+		FName IniPlatformName = ShaderPlatformToPlatformName(ActualShaderPlatform);
 		// find the cvar if needed
 		if (CVar == nullptr)
 		{
@@ -552,7 +561,7 @@ struct RENDERCORE_API FShaderPlatformCachedIniValue
 		// if we are looking up our own platform, just use the current value, however
 		// ShaderPlatformToPlatformName can return the wrong platform than expected - for instance, Linux Vulkan will return Windows
 		// so instead of hitting an asser below, we detect that the request SP is the current SP, and use the CVar value that is set currently
-		if (IniPlatformName == FPlatformProperties::IniPlatformName() || ShaderPlatform == GMaxRHIShaderPlatform)
+		if (IniPlatformName == FPlatformProperties::IniPlatformName() || ActualShaderPlatform == GMaxRHIShaderPlatform)
 		{
 			checkf(CVar != nullptr, TEXT("Failed to find CVar %s when getting current value for FShaderPlatformCachedIniValue"));
 
@@ -594,8 +603,8 @@ private:
 /** Returns if ForwardShading is enabled. Only valid for the current platform (otherwise call ITargetPlatform::UsesForwardShading()). */
 inline bool IsForwardShadingEnabled(const FStaticShaderPlatform Platform)
 {
-	extern RENDERCORE_API uint64 GForwardShadingPlatformMask;
-	return !!(GForwardShadingPlatformMask & (1ull << Platform))
+	extern RENDERCORE_API ShaderPlatformMaskType GForwardShadingPlatformMask;
+	return (GForwardShadingPlatformMask[(int)Platform])
 		// Culling uses compute shader
 		&& GetMaxSupportedFeatureLevel(Platform) >= ERHIFeatureLevel::SM5;
 }
@@ -619,15 +628,15 @@ RENDERCORE_API bool IsUsingBasePassVelocity(const FStaticShaderPlatform Platform
 /** Returns whether the base pass should use selective outputs for a given shader platform */
 inline bool IsUsingSelectiveBasePassOutputs(const FStaticShaderPlatform Platform)
 {
-	extern RENDERCORE_API uint64 GSelectiveBasePassOutputsPlatformMask;
-	return !!(GSelectiveBasePassOutputsPlatformMask & (1ull << Platform));
+	extern RENDERCORE_API ShaderPlatformMaskType GSelectiveBasePassOutputsPlatformMask;
+	return (GSelectiveBasePassOutputsPlatformMask[(int)Platform]);
 }
 
 /** Returns whether distance fields are enabled for a given shader platform */
 inline bool IsUsingDistanceFields(const FStaticShaderPlatform Platform)
 {
-	extern RENDERCORE_API uint64 GDistanceFieldsPlatformMask;
-	return !!(GDistanceFieldsPlatformMask & (1ull << Platform));
+	extern RENDERCORE_API ShaderPlatformMaskType GDistanceFieldsPlatformMask;
+	return (GDistanceFieldsPlatformMask[(int)Platform]);
 }
 
 /** Returns if water should render distance field shadow a second time for the water surface. This is for a platofrm so can be used at cook time. */
@@ -664,14 +673,14 @@ inline bool UseGPUScene(const FStaticShaderPlatform Platform)
 
 inline bool ForceSimpleSkyDiffuse(const FStaticShaderPlatform Platform)
 {
-	extern RENDERCORE_API uint64 GSimpleSkyDiffusePlatformMask;
-	return !!(GSimpleSkyDiffusePlatformMask & (1ull << Platform));
+	extern RENDERCORE_API ShaderPlatformMaskType GSimpleSkyDiffusePlatformMask;
+	return (GSimpleSkyDiffusePlatformMask[(int)Platform]);
 }
 
 inline bool VelocityEncodeDepth(const FStaticShaderPlatform Platform)
 {
-	extern RENDERCORE_API uint64 GVelocityEncodeDepthPlatformMask;
-	return !!(GVelocityEncodeDepthPlatformMask & (1ull << Platform));
+	extern RENDERCORE_API ShaderPlatformMaskType GVelocityEncodeDepthPlatformMask;
+	return (GVelocityEncodeDepthPlatformMask[(int)Platform]);
 }
 
 /** Unit cube vertex buffer (VertexDeclarationFVector4) */
diff --git a/Engine/Source/Runtime/RenderCore/Public/ShaderCore.h b/Engine/Source/Runtime/RenderCore/Public/ShaderCore.h
--- a/Engine/Source/Runtime/RenderCore/Public/ShaderCore.h
+++ b/Engine/Source/Runtime/RenderCore/Public/ShaderCore.h
@@ -113,6 +113,9 @@ inline TStatId GetMemoryStatType(EShaderFrequency ShaderFrequency)
 /** Initializes shader hash cache from IShaderFormatModules. This must be called before reading any shader include. */
 extern RENDERCORE_API void InitializeShaderHashCache();
 
+/** Updates the PreviewPlatform's IncludeDirectory to match that of the Parent Platform*/
+extern RENDERCORE_API void UpdateIncludeDirectoryForPreviewPlatform(EShaderPlatform PreviewPlatform, EShaderPlatform ActualPlatform);
+
 /** Checks if shader include isn't skipped by a shader hash cache. */
 extern RENDERCORE_API void CheckShaderHashCacheInclude(const FString& VirtualFilePath, EShaderPlatform ShaderPlatform);
 
diff --git a/Engine/Source/Runtime/SlateRHIRenderer/Private/SlateRHIRenderingPolicy.cpp b/Engine/Source/Runtime/SlateRHIRenderer/Private/SlateRHIRenderingPolicy.cpp
--- a/Engine/Source/Runtime/SlateRHIRenderer/Private/SlateRHIRenderingPolicy.cpp
+++ b/Engine/Source/Runtime/SlateRHIRenderer/Private/SlateRHIRenderingPolicy.cpp
@@ -394,7 +394,7 @@ static bool UpdateScissorRect(
 					RHICmdList.BeginRenderPass(RPInfo, TEXT("SlateUpdateScissorRect_ClearStencil"));
 				}
 
-				FGlobalShaderMap* MaxFeatureLevelShaderMap = GetGlobalShaderMap(GMaxRHIFeatureLevel);
+				FGlobalShaderMap* MaxFeatureLevelShaderMap = GetGlobalShaderMap(GMaxRHIShaderPlatform);
 
 				// Set the new shaders
 				TShaderMapRef<FSlateMaskingVS> VertexShader(MaxFeatureLevelShaderMap);
@@ -659,7 +659,7 @@ void FSlateRHIRenderingPolicy::DrawElements(
 		SceneViews[NumScenes - 1] = CreateSceneView(SceneViewFamilyContexts[NumScenes - 1], BackBuffer, FMatrix(Params.ViewProjectionMatrix));
 	}
 
-	TShaderMapRef<FSlateElementVS> GlobalVertexShader(GetGlobalShaderMap(GMaxRHIFeatureLevel));
+	TShaderMapRef<FSlateElementVS> GlobalVertexShader(GetGlobalShaderMap(GMaxRHIShaderPlatform));
 
 	FSamplerStateRHIRef BilinearClamp = TStaticSamplerState<SF_Bilinear, AM_Clamp, AM_Clamp, AM_Clamp>::GetRHI();
 
@@ -672,7 +672,7 @@ void FSlateRHIRenderingPolicy::DrawElements(
 	const FSlateRenderDataHandle* LastHandle = nullptr;
 
 	const ERHIFeatureLevel::Type FeatureLevel = GMaxRHIFeatureLevel;
-	FGlobalShaderMap* ShaderMap = GetGlobalShaderMap(FeatureLevel);
+	FGlobalShaderMap* ShaderMap = GetGlobalShaderMap(GMaxRHIShaderPlatform);
 
 #if WITH_SLATE_VISUALIZERS
 	FRandomStream BatchColors(1337);
