Add RHI scope tag when setting up the RHI thread.
This fixes IsInRHIThread() when calling the function too early.

diff --git a/Engine/Source/Runtime/RenderCore/Private/RenderingThread.cpp b/Engine/Source/Runtime/RenderCore/Private/RenderingThread.cpp
--- a/Engine/Source/Runtime/RenderCore/Private/RenderingThread.cpp
+++ b/Engine/Source/Runtime/RenderCore/Private/RenderingThread.cpp
@@ -321,7 +321,10 @@ public:
 #endif
 
 		FMemory::SetupTLSCachesOnCurrentThread();
-		FPlatformProcess::SetupRHIThread();
+		{
+			FTaskTagScope Scope(ETaskTag::ERhiThread);
+			FPlatformProcess::SetupRHIThread();
+		}
 		FTaskGraphInterface::Get().AttachToThread(ENamedThreads::RHIThread);
 		FTaskGraphInterface::Get().ProcessThreadUntilRequestReturn(ENamedThreads::RHIThread);
 		FMemory::ClearAndDisableTLSCachesOnCurrentThread();
