Removed RDG pass CPU trace to reduce memory pressure on insights.

diff --git a/Engine/Source/Runtime/RenderCore/Private/RenderGraphBuilder.cpp b/Engine/Source/Runtime/RenderCore/Private/RenderGraphBuilder.cpp
--- a/Engine/Source/Runtime/RenderCore/Private/RenderGraphBuilder.cpp
+++ b/Engine/Source/Runtime/RenderCore/Private/RenderGraphBuilder.cpp
@@ -2841,7 +2841,7 @@ void FRDGBuilder::ExecutePassEpilogue(FRHIComputeCommandList& RHICmdListPass, FR
 
 void FRDGBuilder::ExecutePass(FRDGPass* Pass, FRHIComputeCommandList& RHICmdListPass)
 {
-#if RDG_ENABLE_DEBUG && RDG_EVENTS != RDG_EVENTS_NONE
+#if 0 // Disabled by default to reduce memory usage in Insights.
 	SCOPED_NAMED_EVENT_TCHAR(Pass->GetName(), FColor::Magenta);
 #endif
 
