HDR composition improvements:
* De-hardcode UI luminance, now driven by r.HDR.UI.Luminance
* Expose hard-coded ACES tonemapper parameters:
  - r.HDR.Display.MinLuminanceLog10 (-4.0) defines the luminance for the darkest point in the scene, and will also compute its ACES value (through lookup_ACESmin). The same is happening with the existing CVar r.HDR.Display.MaxLuminance, through lookup_ACESmax
  - r.HDR.Display.MidLuminance (15.0) defines the luminance for 18% gray. It works by offseting the dakest and brightest points in the scene
  - r.HDR.Aces.SceneColorMultiplier (1.5) . This acts like a gain to ensure consistency between ACES (for HDR) and UE (for SDR) tonemappers.

From my tests,  r.HDR.Aces.SceneColorMultiplier and r.HDR.Display.MidLuminance are a bit redundant, but the former acts as a gain on the input color, whereas the latter changes the anchor points for the tonemapper. Since they don't have a linear relationship, it felt useful to keep them both.

some examples:

Default Values (used in initial implementation of ACES 1.3)
r.HDR.Display.MinLuminanceLog10 = -4
r.HDR.Display.MidLuminance = 15
r.HDR.Aces.SceneColorMultiplier = 1.5

ACES 1.3 reference values (match reference images for r.HDR.Display.MaxLuminance=1000, 2000, 4000)
r.HDR.Display.MinLuminanceLog10 = -4
r.HDR.Display.MidLuminance = 15
r.HDR.Aces.SceneColorMultiplier = 1.0

ACES 1.3 with better fit to UE SDR tonemapper (scaled to 300 nits)
r.HDR.Display.MinLuminanceLog10 = -4
r.HDR.Display.MidLuminance = 54
r.HDR.Aces.SceneColorMultiplier = 1.0

diff --git a/Engine/Shaders/Private/ACES.ush b/Engine/Shaders/Private/ACES.ush
--- a/Engine/Shaders/Private/ACES.ush
+++ b/Engine/Shaders/Private/ACES.ush
@@ -1304,9 +1304,7 @@ float3 dim_to_dark(float3 XYZ)
 float3 outputTransform
 (
 	float3 inValue,
-	float Y_MIN,
-	float Y_MID,
-	float Y_MAX,
+	TsParams PARAMS,
 	FColorSpace DISPLAY_PRI,
 	FColorSpace LIMITING_PRI,
 	int EOTF,
@@ -1322,9 +1320,9 @@ float3 outputTransform
 		NOTE: This is a bit of a hack - probably a more direct way to do this.
 		TODO: Fix in future version
 	*/
-	TsParams PARAMS_DEFAULT = init_TsParams(Y_MIN, Y_MAX);
-	float expShift = log2(inv_ssts(Y_MID, PARAMS_DEFAULT)) - log2(0.18);
-	TsParams PARAMS = init_TsParams(Y_MIN, Y_MAX, expShift);
+
+	float Y_MIN = PARAMS.Min.y;
+	float Y_MAX = PARAMS.Max.y;
 
 	// RRT sweeteners
 	float3 rgbPre = rrt_sweeteners(inValue);
@@ -1480,6 +1478,32 @@ float3 outputTransform
 	return outputCV;
 }
 
+float3 outputTransform
+(
+	float3 inValue,
+	float Y_MIN,
+	float Y_MID,
+	float Y_MAX,
+	FColorSpace DISPLAY_PRI,
+	FColorSpace LIMITING_PRI,
+	int EOTF,
+	int SURROUND,
+	bool STRETCH_BLACK = true,
+	bool D60_SIM = false,
+	bool LEGAL_RANGE = false
+)
+{
+	/*
+		NOTE: This is a bit of a hack - probably a more direct way to do this.
+		TODO: Fix in future version
+	*/
+	TsParams PARAMS_DEFAULT = init_TsParams(Y_MIN, Y_MAX);
+	float expShift = log2(inv_ssts(Y_MID, PARAMS_DEFAULT)) - log2(0.18);
+	TsParams PARAMS = init_TsParams(Y_MIN, Y_MAX, expShift);
+
+	return outputTransform(inValue, PARAMS, DISPLAY_PRI, LIMITING_PRI, EOTF, SURROUND, STRETCH_BLACK, D60_SIM, LEGAL_RANGE);
+}
+
 float3 invOutputTransform
 (
 	float3 inValue,
diff --git a/Engine/Shaders/Private/CompositeUIPixelShader.usf b/Engine/Shaders/Private/CompositeUIPixelShader.usf
--- a/Engine/Shaders/Private/CompositeUIPixelShader.usf
+++ b/Engine/Shaders/Private/CompositeUIPixelShader.usf
@@ -17,6 +17,13 @@ uint OutputDevice;
 uint OutputGamut;
 float OutputMaxLuminance;
 
+float4 ACESMinMaxData;
+float4 ACESMidData;
+float4 ACESCoefsLow_0;
+float4 ACESCoefsHigh_0;
+float ACESCoefsLow_4;
+float ACESCoefsHigh_4;
+
 void GenerateLUTPS(
 	FWriteToSliceGeometryOutput Input,
 	out float4 OutColor : SV_Target0
@@ -30,10 +37,12 @@ void GenerateLUTPS(
 
 	float3 LinearColor = ST2084ToLinear(Neutral);
 
+	FACESTonemapParams LUTTsParams = ComputeACESTonemapParams(ACESMinMaxData, ACESMidData, ACESCoefsLow_0, ACESCoefsHigh_0, ACESCoefsLow_4, ACESCoefsHigh_4, 1.0);
+
 	if( OutputDevice == TONEMAPPER_OUTPUT_ACES1000nitST2084 || OutputDevice == TONEMAPPER_OUTPUT_ACES1000nitScRGB)
 	{		
 		// ACES 1000nit transform with PQ/2084 encoding, user specified gamut 
-		float3 ODTColor = ACESOutputTransforms1000(LinearColor, OutputMaxLuminance);
+		float3 ODTColor = ACESOutputTransforms1000(LinearColor, LUTTsParams);
 		// do not apply scRGB here because LUT is not 16F in this format, and scRGB can output negative values
 		if (OutputDevice != TONEMAPPER_OUTPUT_ACES1000nitScRGB)
 		{
@@ -44,7 +53,7 @@ void GenerateLUTPS(
 	else if( OutputDevice == TONEMAPPER_OUTPUT_ACES2000nitST2084 || OutputDevice == TONEMAPPER_OUTPUT_ACES2000nitScRGB)
 	{		
 		// ACES 2000nit transform with PQ/2084 encoding, user specified gamut 
-		float3 ODTColor = ACESOutputTransforms2000(LinearColor, OutputMaxLuminance);
+		float3 ODTColor = ACESOutputTransforms2000(LinearColor, LUTTsParams);
 		// do not apply scRGB here because LUT is not 16F in this format, and scRGB can output negative values
 		// 
 		if (OutputDevice != TONEMAPPER_OUTPUT_ACES2000nitScRGB)
@@ -77,6 +86,7 @@ Texture3D ColorSpaceLUT;
 SamplerState ColorSpaceLUTSampler;
 
 float UILevel;
+float UILuminance;
 
 // Invert the encoding applied by the ApplyGammaCorrection function.
 float3 LinearizeColor(float3 EncodedColor)
@@ -95,7 +105,6 @@ float3 LinearizeColor(float3 EncodedColor)
 }
 
 static const float ScRGBScaleFactor = 80.0f; // used to convert between ScRGB and nits
-static const float UIBrightnessInNits = 300.0f; // todo stop hardcoding this value and try to fetch SDR/HDR brightness balance instead (automatically or though calibration step)
 
 float3 ColorLookupTable(float3 LinearColor)
 {
@@ -175,11 +184,11 @@ float4 ComputeHDRUIColor(float2 InputUV)
 	UIColor.xyz = ST2084ToLinear(UIColor.xyz);
 #else
 	UIColor.xyz = ST2084ToScRGB(UIColor.xyz, OutputDevice, OutputMaxLuminance);
-	UIColor.xyz = mul(sRGB_2_Rec2020, UIColor.xyz) * UIBrightnessInNits;
+	UIColor.xyz = mul(sRGB_2_Rec2020, UIColor.xyz) * UILuminance;
 #endif
 
 #else // #if USE_LUT_FOR_UI
-	UIColor.xyz = mul(sRGB_2_Rec2020, UIColor.xyz) * UIBrightnessInNits;
+	UIColor.xyz = mul(sRGB_2_Rec2020, UIColor.xyz) * UILuminance;
 #endif // #if USE_LUT_FOR_UI
 
 	return UIColor;
diff --git a/Engine/Shaders/Private/PostProcessCombineLUTs.usf b/Engine/Shaders/Private/PostProcessCombineLUTs.usf
--- a/Engine/Shaders/Private/PostProcessCombineLUTs.usf
+++ b/Engine/Shaders/Private/PostProcessCombineLUTs.usf
@@ -67,6 +67,14 @@ float  ColorCorrectionShadowsMax;
 float  ColorCorrectionHighlightsMin;
 float  ColorCorrectionHighlightsMax;
 
+float4 ACESMinMaxData;
+float4 ACESMidData;
+float4 ACESCoefsLow_0;
+float4 ACESCoefsHigh_0;
+float ACESCoefsLow_4;
+float ACESCoefsHigh_4;
+float ACESSceneColorMultiplier;
+
 float3 ColorCorrect( float3 WorkingColor,
 	float4 ColorSaturation,
 	float4 ColorContrast,
@@ -341,7 +349,8 @@ float4 CombineLUTsCommon(float2 InUV, uint InLayerIndex)
 	else if( GetOutputDevice() == TONEMAPPER_OUTPUT_ACES1000nitST2084 || GetOutputDevice() == TONEMAPPER_OUTPUT_ACES1000nitScRGB)
 	{		
 		// 1000 nit ODT
-		float3 ODTColor = ACESOutputTransforms1000( GradedColor, (float3x3)WorkingColorSpace.ToAP0, OutputMaxLuminance);
+		FACESTonemapParams AcesParams = ComputeACESTonemapParams(ACESMinMaxData, ACESMidData, ACESCoefsLow_0, ACESCoefsHigh_0, ACESCoefsLow_4, ACESCoefsHigh_4, ACESSceneColorMultiplier);
+		float3 ODTColor = ACESOutputTransforms1000( GradedColor, (float3x3)WorkingColorSpace.ToAP0, AcesParams);
 
 		// Convert from AP1 to specified output gamut except for ScRGB as it might contain negative values. 
 		// In this case, the AP1->output conversion happens through ST2084ToScRGB in the tonemap function
@@ -358,7 +367,8 @@ float4 CombineLUTsCommon(float2 InUV, uint InLayerIndex)
 	else if( GetOutputDevice() == TONEMAPPER_OUTPUT_ACES2000nitST2084 || GetOutputDevice() == TONEMAPPER_OUTPUT_ACES2000nitScRGB)
 	{		
 		// 2000 nit ODT
-		float3 ODTColor = ACESOutputTransforms2000( GradedColor, (float3x3)WorkingColorSpace.ToAP0, OutputMaxLuminance);
+		FACESTonemapParams AcesParams = ComputeACESTonemapParams(ACESMinMaxData, ACESMidData, ACESCoefsLow_0, ACESCoefsHigh_0, ACESCoefsLow_4, ACESCoefsHigh_4, ACESSceneColorMultiplier);
+		float3 ODTColor = ACESOutputTransforms2000( GradedColor, (float3x3)WorkingColorSpace.ToAP0, AcesParams);
 
 		// Convert from AP1 to specified output gamut except for ScRGB as it might contain negative values. 
 		// In this case, the AP1->output conversion happens through ST2084ToScRGB in the tonemap function
diff --git a/Engine/Shaders/Private/PostProcessDeviceEncodingOnly.usf b/Engine/Shaders/Private/PostProcessDeviceEncodingOnly.usf
--- a/Engine/Shaders/Private/PostProcessDeviceEncodingOnly.usf
+++ b/Engine/Shaders/Private/PostProcessDeviceEncodingOnly.usf
@@ -33,6 +33,14 @@ uint OutputDevice;
 uint OutputGamut;
 float OutputMaxLuminance;
 
+float4 ACESMinMaxData;
+float4 ACESMidData;
+float4 ACESCoefsLow_0;
+float4 ACESCoefsHigh_0;
+float ACESCoefsLow_4;
+float ACESCoefsHigh_4;
+float ACESSceneColorMultiplier;
+
 uint GetOutputDevice()
 {
 #if FEATURE_LEVEL > FEATURE_LEVEL_ES3_1
@@ -95,9 +103,10 @@ float4 DeviceEncodingOnlyCommonPS(
 
 	// ACES 1000nit transform with PQ/2084 encoding, user specified gamut 
 	else if( GetOutputDevice() == TONEMAPPER_OUTPUT_ACES1000nitST2084 )
-	{		
+	{		 
 		// 1000 nit ODT
-		float3 ODTColor = ACESOutputTransforms1000( FilmColor, (float3x3)WorkingColorSpace.ToAP0, OutputMaxLuminance);
+		FACESTonemapParams AcesParams = ComputeACESTonemapParams(ACESMinMaxData, ACESMidData, ACESCoefsLow_0, ACESCoefsHigh_0, ACESCoefsLow_4, ACESCoefsHigh_4, ACESSceneColorMultiplier);
+		float3 ODTColor = ACESOutputTransforms1000( FilmColor, (float3x3)WorkingColorSpace.ToAP0, AcesParams);
 
 		// Convert from AP1 to specified output gamut
 		ODTColor = mul( AP1_2_Output, ODTColor );
@@ -105,12 +114,14 @@ float4 DeviceEncodingOnlyCommonPS(
 		// Apply conversion to ST-2084 (Dolby PQ)
 		OutDeviceColor = LinearToST2084( ODTColor );
 	}
-
+	 
 	// ACES 2000nit transform with PQ/2084 encoding, user specified gamut 
 	else if( GetOutputDevice() == TONEMAPPER_OUTPUT_ACES2000nitST2084 )
 	{		
+		FACESTonemapParams AcesParams = ComputeACESTonemapParams(ACESMinMaxData, ACESMidData, ACESCoefsLow_0, ACESCoefsHigh_0, ACESCoefsLow_4, ACESCoefsHigh_4, ACESSceneColorMultiplier);
+		
 		// 2000 nit ODT
-		float3 ODTColor = ACESOutputTransforms2000( FilmColor, (float3x3)WorkingColorSpace.ToAP0, OutputMaxLuminance);
+		float3 ODTColor = ACESOutputTransforms2000( FilmColor, (float3x3)WorkingColorSpace.ToAP0, AcesParams);
 
 		// Convert from AP1 to specified output gamut
 		ODTColor = mul( AP1_2_Output, ODTColor );
@@ -122,8 +133,10 @@ float4 DeviceEncodingOnlyCommonPS(
 	// ACES 1000nit transform to linear ScRGB
 	else if(GetOutputDevice() == TONEMAPPER_OUTPUT_ACES1000nitScRGB)
 	{
+		FACESTonemapParams AcesParams = ComputeACESTonemapParams(ACESMinMaxData, ACESMidData, ACESCoefsLow_0, ACESCoefsHigh_0, ACESCoefsLow_4, ACESCoefsHigh_4, ACESSceneColorMultiplier);
+
 		// 1000 nit ODT
-		float3 ODTColor = ACESOutputTransforms1000( FilmColor, (float3x3)WorkingColorSpace.ToAP0, OutputMaxLuminance );
+		float3 ODTColor = ACESOutputTransforms1000( FilmColor, (float3x3)WorkingColorSpace.ToAP0, AcesParams);
 
 		// Apply conversion to ScRGB
 		OutDeviceColor = ST2084ToScRGB( LinearToST2084( ODTColor ), TONEMAPPER_OUTPUT_ACES1000nitScRGB, OutputMaxLuminance );
@@ -132,8 +145,10 @@ float4 DeviceEncodingOnlyCommonPS(
 	// ACES 2000nit transform to linear ScRGB
 	else if(GetOutputDevice() == TONEMAPPER_OUTPUT_ACES2000nitScRGB)
 	{
+		FACESTonemapParams AcesParams = ComputeACESTonemapParams(ACESMinMaxData, ACESMidData, ACESCoefsLow_0, ACESCoefsHigh_0, ACESCoefsLow_4, ACESCoefsHigh_4, ACESSceneColorMultiplier);
+
 		// 2000 nit ODT
-		float3 ODTColor = ACESOutputTransforms2000( FilmColor, (float3x3)WorkingColorSpace.ToAP0, OutputMaxLuminance );
+		float3 ODTColor = ACESOutputTransforms2000( FilmColor, (float3x3)WorkingColorSpace.ToAP0, AcesParams);
 
 		// Apply conversion to ScRGB
 		OutDeviceColor = ST2084ToScRGB( LinearToST2084( ODTColor ), TONEMAPPER_OUTPUT_ACES2000nitScRGB, OutputMaxLuminance );
diff --git a/Engine/Shaders/Private/TonemapCommon.ush b/Engine/Shaders/Private/TonemapCommon.ush
--- a/Engine/Shaders/Private/TonemapCommon.ush
+++ b/Engine/Shaders/Private/TonemapCommon.ush
@@ -352,27 +352,68 @@ float3 InverseACESOutputTransformsRGBD65( float3 OutputReferredLinearsRGBColor )
 	return SceneReferredLinearsRGBColor;
 }
 
+struct FACESTonemapParams
+{
+	TsParams AcesParams;
+	float SceneColorMultiplier;
+};
+
+FACESTonemapParams ComputeACESTonemapParams(float4 InACESMinMaxData, float4 InACESMidData, float4 InACESCoefsLow_0, float4 InACESCoefsHigh_0, float InACESCoefsLow_4, float InACESCoefsHigh_4, float SceneColorMultiplier)
+{
+	FACESTonemapParams OutACESTonemapParams;
+	OutACESTonemapParams.AcesParams.Min.x = InACESMinMaxData.x;
+	OutACESTonemapParams.AcesParams.Min.y = InACESMinMaxData.y;
+	OutACESTonemapParams.AcesParams.Min.slope = 0;
+
+	OutACESTonemapParams.AcesParams.Mid.x = InACESMidData.x;
+	OutACESTonemapParams.AcesParams.Mid.y = InACESMidData.y;
+	OutACESTonemapParams.AcesParams.Mid.slope = InACESMidData.z;
+
+	OutACESTonemapParams.AcesParams.Max.x = InACESMinMaxData.z;
+	OutACESTonemapParams.AcesParams.Max.y = InACESMinMaxData.w;
+	OutACESTonemapParams.AcesParams.Max.slope = 0;
+
+	OutACESTonemapParams.AcesParams.coefsLow[0] = InACESCoefsLow_0.x;
+	OutACESTonemapParams.AcesParams.coefsLow[1] = InACESCoefsLow_0.y;
+	OutACESTonemapParams.AcesParams.coefsLow[2] = InACESCoefsLow_0.z;
+	OutACESTonemapParams.AcesParams.coefsLow[3] = InACESCoefsLow_0.w;
+	OutACESTonemapParams.AcesParams.coefsLow[4] = InACESCoefsLow_4;
+	OutACESTonemapParams.AcesParams.coefsLow[5] = InACESCoefsLow_4;
+
+	OutACESTonemapParams.AcesParams.coefsHigh[0] = InACESCoefsHigh_0.x;
+	OutACESTonemapParams.AcesParams.coefsHigh[1] = InACESCoefsHigh_0.y;
+	OutACESTonemapParams.AcesParams.coefsHigh[2] = InACESCoefsHigh_0.z;
+	OutACESTonemapParams.AcesParams.coefsHigh[3] = InACESCoefsHigh_0.w;
+	OutACESTonemapParams.AcesParams.coefsHigh[4] = InACESCoefsHigh_4;
+	OutACESTonemapParams.AcesParams.coefsHigh[5] = InACESCoefsHigh_4;
+
+	OutACESTonemapParams.SceneColorMultiplier = SceneColorMultiplier;
+
+	return OutACESTonemapParams;
+
+}
+
 //
 // ACES D65 1000 nit Output Transform - Forward
 //  Input is scene-referred linear values in the working space gamut
 //  Output is output-referred linear values in the AP1 gamut
 //
-float3 ACESOutputTransforms1000( float3 SceneReferredLinearColor, const float3x3 WCS_2_AP0, float MaxLuminance ) 
+float3 ACESOutputTransforms1000( float3 SceneReferredLinearColor, const float3x3 WCS_2_AP0, FACESTonemapParams InACESTonemapParams)
 {
-	float3 aces = mul( WCS_2_AP0, SceneReferredLinearColor * 1.5 ); // x 1.5 to get similar brightness between HDR and SDR
+	float3 aces = mul( WCS_2_AP0, SceneReferredLinearColor * InACESTonemapParams.SceneColorMultiplier);
 	const FColorSpace UE_ACES_AP1_DISPLAY_PRI = GetColorSpace(TONEMAPPER_GAMUT_ACEScg_D60);
 	const FColorSpace UE_ACES_AP1_LIMITING_PRI = GetColorSpace(TONEMAPPER_GAMUT_ACEScg_D60);
-	float3 OutputReferredLinearAP1Color = outputTransform(aces, UE_ACES_HDR_Y_MIN, UE_ACES_HDR_Y_MID, MaxLuminance, UE_ACES_AP1_DISPLAY_PRI, UE_ACES_AP1_LIMITING_PRI, UE_ACES_EOTF,
+	float3 OutputReferredLinearAP1Color = outputTransform(aces, InACESTonemapParams.AcesParams, UE_ACES_AP1_DISPLAY_PRI, UE_ACES_AP1_LIMITING_PRI, UE_ACES_EOTF,
 											 			  UE_ACES_SURROUND, UE_ACES_STRETCH_BLACK, UE_ACES_D60_SIM, UE_ACES_LEGAL_RANGE);
 
 	return OutputReferredLinearAP1Color;
 }
 
-float3 ACESOutputTransforms1000( float3 SceneReferredLinearsRGBColor, float MaxLuminance )
+float3 ACESOutputTransforms1000( float3 SceneReferredLinearsRGBColor, FACESTonemapParams InACESTonemapParams)
 {
 	const float3x3 sRGB_2_AP0 = mul( XYZ_2_AP0_MAT, mul( D65_2_D60_CAT, sRGB_2_XYZ_MAT ) );
 
-	return ACESOutputTransforms1000( SceneReferredLinearsRGBColor, sRGB_2_AP0, MaxLuminance );
+	return ACESOutputTransforms1000( SceneReferredLinearsRGBColor, sRGB_2_AP0, InACESTonemapParams);
 }
 
 //
@@ -380,22 +421,22 @@ float3 ACESOutputTransforms1000( float3 SceneReferredLinearsRGBColor, float MaxL
 //  Input is scene-referred linear values in the working space gamut
 //  Output is output-referred linear values in the AP1 gamut
 //
-float3 ACESOutputTransforms2000( float3 SceneReferredLinearColor, const float3x3 WCS_2_AP0, float MaxLuminance ) 
+float3 ACESOutputTransforms2000( float3 SceneReferredLinearColor, const float3x3 WCS_2_AP0, FACESTonemapParams InACESTonemapParams)
 {
-	float3 aces = mul( WCS_2_AP0, SceneReferredLinearColor * 1.5 ); // x 1.5 to get similar brightness between HDR and SDR
+	float3 aces = mul( WCS_2_AP0, SceneReferredLinearColor * InACESTonemapParams.SceneColorMultiplier);
 	const FColorSpace UE_ACES_AP1_DISPLAY_PRI = GetColorSpace(TONEMAPPER_GAMUT_ACEScg_D60);
 	const FColorSpace UE_ACES_AP1_LIMITING_PRI = GetColorSpace(TONEMAPPER_GAMUT_ACEScg_D60);
-	float3 OutputReferredLinearAP1Color = outputTransform(aces, UE_ACES_HDR_Y_MIN, UE_ACES_HDR_Y_MID, MaxLuminance, UE_ACES_AP1_DISPLAY_PRI, UE_ACES_AP1_LIMITING_PRI, UE_ACES_EOTF,
+	float3 OutputReferredLinearAP1Color = outputTransform(aces, InACESTonemapParams.AcesParams, UE_ACES_AP1_DISPLAY_PRI, UE_ACES_AP1_LIMITING_PRI, UE_ACES_EOTF,
 											 			  UE_ACES_SURROUND, UE_ACES_STRETCH_BLACK, UE_ACES_D60_SIM, UE_ACES_LEGAL_RANGE);
 
 	return OutputReferredLinearAP1Color;
 }
 
-float3 ACESOutputTransforms2000( float3 SceneReferredLinearsRGBColor, float MaxLuminance )
+float3 ACESOutputTransforms2000( float3 SceneReferredLinearsRGBColor, FACESTonemapParams InACESTonemapParams)
 {
 	const float3x3 sRGB_2_AP0 = mul( XYZ_2_AP0_MAT, mul( D65_2_D60_CAT, sRGB_2_XYZ_MAT ) );
 
-	return ACESOutputTransforms2000( SceneReferredLinearsRGBColor, sRGB_2_AP0, MaxLuminance );
+	return ACESOutputTransforms2000( SceneReferredLinearsRGBColor, sRGB_2_AP0, InACESTonemapParams);
 }
 
 static const float3x3 GamutMappingIdentityMatrix = { 1.0, 0.0, 0.0, 0.0, 1.0, 0.0, 0.0, 0.0, 1.0 };
diff --git a/Engine/Source/Runtime/Engine/Private/UnrealEngine.cpp b/Engine/Source/Runtime/Engine/Private/UnrealEngine.cpp
--- a/Engine/Source/Runtime/Engine/Private/UnrealEngine.cpp
+++ b/Engine/Source/Runtime/Engine/Private/UnrealEngine.cpp
@@ -872,7 +872,11 @@ void ENGINE_API HDRSettingChangedSinkCallback()
 {
 	static const auto CVarHDROutputEnabled = IConsoleManager::Get().FindTConsoleVariableDataInt(TEXT("r.HDR.EnableHDROutput"));
 	check(CVarHDROutputEnabled);
-	
+	static const auto CVarHDRMinLuminanceLog10 = IConsoleManager::Get().FindTConsoleVariableDataFloat(TEXT("r.HDR.Display.MinLuminanceLog10"));
+	static const auto CVarHDRMidLuminance = IConsoleManager::Get().FindTConsoleVariableDataFloat(TEXT("r.HDR.Display.MidLuminance"));
+	static const auto CVarHDRMaxLuminance = IConsoleManager::Get().FindTConsoleVariableDataInt(TEXT("r.HDR.Display.MaxLuminance"));
+	static const auto CVarHDRSceneColorMultiplier = IConsoleManager::Get().FindTConsoleVariableDataFloat(TEXT("r.HDR.Aces.SceneColorMultiplier"));
+
 	if (GRHIVendorId == 0)
 	{
 		return;
@@ -885,7 +889,13 @@ void ENGINE_API HDRSettingChangedSinkCallback()
 		bIsHDREnabled = false;
 	}
 
-	if(bIsHDREnabled != GRHIIsHDREnabled)
+	static float GHDRMinLuminnanceLog10 = CVarHDRMinLuminanceLog10->GetValueOnAnyThread();
+	static float GHDRMidLuminnance = CVarHDRMidLuminance->GetValueOnAnyThread();
+	static int32 GHDRMaxLuminnance = CVarHDRMaxLuminance->GetValueOnAnyThread();
+	static float GHDRSceneColorMultiplier = CVarHDRSceneColorMultiplier->GetValueOnAnyThread();
+
+	if(bIsHDREnabled != GRHIIsHDREnabled || !FMath::IsNearlyEqual(GHDRMinLuminnanceLog10, CVarHDRMinLuminanceLog10->GetValueOnAnyThread()) || !FMath::IsNearlyEqual(GHDRMidLuminnance, CVarHDRMidLuminance->GetValueOnAnyThread()) ||
+									        (GHDRMaxLuminnance != CVarHDRMaxLuminance->GetValueOnAnyThread()) || !FMath::IsNearlyEqual(GHDRSceneColorMultiplier, CVarHDRSceneColorMultiplier->GetValueOnAnyThread()))
 	{
 		// We'll naively fall back to 1000 if DisplayNits is 0
 		uint32 DisplayNitLevel = 0;
@@ -904,6 +914,10 @@ void ENGINE_API HDRSettingChangedSinkCallback()
 
 		// Now set the HDR setting.
 		GRHIIsHDREnabled = IsHDREnabled();
+		GHDRMinLuminnanceLog10 = CVarHDRMinLuminanceLog10->GetValueOnAnyThread();
+		GHDRMidLuminnance = CVarHDRMidLuminance->GetValueOnAnyThread();
+		GHDRMaxLuminnance = CVarHDRMaxLuminance->GetValueOnAnyThread();
+		GHDRSceneColorMultiplier = CVarHDRSceneColorMultiplier->GetValueOnAnyThread();
 	}
 }
 
diff --git a/Engine/Source/Runtime/RenderCore/Private/RenderCore.cpp b/Engine/Source/Runtime/RenderCore/Private/RenderCore.cpp
--- a/Engine/Source/Runtime/RenderCore/Private/RenderCore.cpp
+++ b/Engine/Source/Runtime/RenderCore/Private/RenderCore.cpp
@@ -345,6 +345,20 @@ TAutoConsoleVariable<int32> CVarDisplayOutputDevice(
 	TEXT("9: Linear final color with tone curve\n"),
 	ECVF_Scalability | ECVF_RenderThreadSafe);
 
+static TAutoConsoleVariable<float> CVarHDRDisplayMinLuminanceLog10(
+	TEXT("r.HDR.Display.MinLuminanceLog10"),
+	-4.0f,
+	TEXT("The configured minimum display output nit level (log10 value)"),
+	ECVF_RenderThreadSafe
+);
+
+static TAutoConsoleVariable<float> CVarHDRDisplayMidLuminance(
+	TEXT("r.HDR.Display.MidLuminance"),
+	15.0f,
+	TEXT("The configured display output nit level for 18% gray"),
+	ECVF_RenderThreadSafe
+);
+
 static TAutoConsoleVariable<int32> CVarHDRDisplayMaxLuminance(
 	TEXT("r.HDR.Display.MaxLuminance"),
 	0,
@@ -352,6 +366,13 @@ static TAutoConsoleVariable<int32> CVarHDRDisplayMaxLuminance(
 	ECVF_RenderThreadSafe
 );
 
+static TAutoConsoleVariable<float> CVarHDRAcesColorMultiplier(
+	TEXT("r.HDR.Aces.SceneColorMultiplier"),
+	1.5f,
+	TEXT("Multiplier applied to scene color. Helps to"),
+	ECVF_RenderThreadSafe
+);
+
 TAutoConsoleVariable<int32> CVarHDROutputEnabled(
 	TEXT("r.HDR.EnableHDROutput"),
 	0,
@@ -525,6 +546,15 @@ RENDERCORE_API float HDRGetDisplayMaximumLuminance()
 	return CVarHDRDisplayMaxLuminance.GetValueOnAnyThread();
 }
 
+static FACESTonemapParams GACESTonemapParams;
+
+RENDERCORE_API FACESTonemapParams HDRGetACESTonemapParams()
+{
+	return GACESTonemapParams;
+}
+
+void ConfigureACESTonemapParams(FACESTonemapParams& OutTsParams, float MinLuminance, float MidLuminance, float MaxLuminance);
+
 RENDERCORE_API void HDRConfigureCVars(bool bIsHDREnabled, uint32 DisplayNits, bool bFromGameSettings)
 {
 	if (bIsHDREnabled && !GRHISupportsHDROutput)
@@ -557,6 +587,9 @@ RENDERCORE_API void HDRConfigureCVars(bool bIsHDREnabled, uint32 DisplayNits, bo
 
 	CVarDisplayOutputDevice->Set((int32)OutputDevice, ECVF_SetByDeviceProfile);
 	CVarDisplayColorGamut->Set((int32)ColorGamut, ECVF_SetByDeviceProfile);
+
+	float HDRDisplayMinLuminance = FMath::Pow(10, CVarHDRDisplayMinLuminanceLog10.GetValueOnAnyThread());
+	ConfigureACESTonemapParams(GACESTonemapParams, HDRDisplayMinLuminance, CVarHDRDisplayMidLuminance.GetValueOnAnyThread(), CVarHDRDisplayMaxLuminance.GetValueOnAnyThread());
 }
 
 RENDERCORE_API FMatrix44f GamutToXYZMatrix(EDisplayColorGamut ColorGamut)
@@ -627,3 +660,298 @@ RENDERCORE_API FMatrix44f XYZToGamutMatrix(EDisplayColorGamut ColorGamut)
 	}
 
 }
+
+static const float MIN_STOP_SDR = -6.5;
+static const float MAX_STOP_SDR = 6.5;
+
+static const float MIN_STOP_RRT = -15.;
+static const float MAX_STOP_RRT = 18.;
+
+static const float MIN_LUM_SDR = 0.02;
+static const float MAX_LUM_SDR = 48.0;
+
+static const float MIN_LUM_RRT = 0.0001;
+static const float MAX_LUM_RRT = 10000.0;
+
+///+ TODO: check if valid
+float interpolate1D(const FVector2f table[2], float value)
+{
+	float t = FMath::Clamp((value - table[0].X) / (table[1].X - table[0].X), 0.0f, 1.0f);
+	return FMath::Lerp(table[0].Y, table[1].Y, t);
+}
+///-
+
+float lookup_ACESmin(float minLum)
+{
+	const FVector2f minTable[2] = { FVector2f(FMath::LogX(10, MIN_LUM_RRT), MIN_STOP_RRT),
+								 FVector2f(FMath::LogX(10, MIN_LUM_SDR), MIN_STOP_SDR) };
+
+	return 0.18 * pow(2., interpolate1D(minTable, FMath::LogX(10, minLum)));
+}
+
+float lookup_ACESmax(float maxLum)
+{
+	const FVector2f maxTable[2] = { FVector2f(FMath::LogX(10, MAX_LUM_SDR), MAX_STOP_SDR),
+								   FVector2f(FMath::LogX(10, MAX_LUM_RRT), MAX_STOP_RRT) };
+
+	return 0.18 * pow(2., interpolate1D(maxTable, FMath::LogX(10, maxLum)));
+}
+
+struct TsPoint
+{
+	float x;        // ACES
+	float y;        // luminance
+	float slope;    // 
+};
+
+struct TsParams
+{
+	TsPoint Min;
+	TsPoint Mid;
+	TsPoint Max;
+	float coefsLow[5];
+	float coefsHigh[5];
+};
+
+void init_coefsLow(
+	const TsPoint& TsPointLow,
+	const TsPoint& TsPointMid,
+	float coefsLow[5]
+)
+{
+
+	float knotIncLow = (FMath::LogX(10, TsPointMid.x) - FMath::LogX(10, TsPointLow.x)) / 3.;
+	// float halfKnotInc = (FMath::LogX(10, TsPointMid.x) - FMath::LogX(10, TsPointLow.x)) / 6.;
+
+	// Determine two lowest coefficients (straddling minPt)
+	coefsLow[0] = (TsPointLow.slope * (FMath::LogX(10, TsPointLow.x) - 0.5 * knotIncLow)) + (FMath::LogX(10, TsPointLow.y) - TsPointLow.slope * FMath::LogX(10, TsPointLow.x));
+	coefsLow[1] = (TsPointLow.slope * (FMath::LogX(10, TsPointLow.x) + 0.5 * knotIncLow)) + (FMath::LogX(10, TsPointLow.y) - TsPointLow.slope * FMath::LogX(10, TsPointLow.x));
+	// NOTE: if slope=0, then the above becomes just 
+		// coefsLow[0] = FMath::LogX(10, TsPointLow.y);
+		// coefsLow[1] = FMath::LogX(10, TsPointLow.y);
+	// leaving it as a variable for now in case we decide we need non-zero slope extensions
+
+	// Determine two highest coefficients (straddling midPt)
+	coefsLow[3] = (TsPointMid.slope * (FMath::LogX(10, TsPointMid.x) - 0.5 * knotIncLow)) + (FMath::LogX(10, TsPointMid.y) - TsPointMid.slope * FMath::LogX(10, TsPointMid.x));
+	coefsLow[4] = (TsPointMid.slope * (FMath::LogX(10, TsPointMid.x) + 0.5 * knotIncLow)) + (FMath::LogX(10, TsPointMid.y) - TsPointMid.slope * FMath::LogX(10, TsPointMid.x));
+
+	// Middle coefficient (which defines the "sharpness of the bend") is linearly interpolated
+	FVector2f bendsLow[2] = { FVector2f(MIN_STOP_RRT, 0.18),
+						   FVector2f(MIN_STOP_SDR, 0.35) };
+	float pctLow = interpolate1D(bendsLow, log2(TsPointLow.x / 0.18));
+	coefsLow[2] = FMath::LogX(10, TsPointLow.y) + pctLow * (FMath::LogX(10, TsPointMid.y) - FMath::LogX(10, TsPointLow.y));
+}
+
+void init_coefsHigh(
+	const TsPoint& TsPointMid,
+	const TsPoint& TsPointMax,
+	float coefsHigh[5]
+)
+{
+
+	float knotIncHigh = (FMath::LogX(10, TsPointMax.x) - FMath::LogX(10, TsPointMid.x)) / 3.;
+	// float halfKnotInc = (FMath::LogX(10, TsPointMax.x) - FMath::LogX(10, TsPointMid.x)) / 6.;
+
+	// Determine two lowest coefficients (straddling midPt)
+	coefsHigh[0] = (TsPointMid.slope * (FMath::LogX(10, TsPointMid.x) - 0.5 * knotIncHigh)) + (FMath::LogX(10, TsPointMid.y) - TsPointMid.slope * FMath::LogX(10, TsPointMid.x));
+	coefsHigh[1] = (TsPointMid.slope * (FMath::LogX(10, TsPointMid.x) + 0.5 * knotIncHigh)) + (FMath::LogX(10, TsPointMid.y) - TsPointMid.slope * FMath::LogX(10, TsPointMid.x));
+
+	// Determine two highest coefficients (straddling maxPt)
+	coefsHigh[3] = (TsPointMax.slope * (FMath::LogX(10, TsPointMax.x) - 0.5 * knotIncHigh)) + (FMath::LogX(10, TsPointMax.y) - TsPointMax.slope * FMath::LogX(10, TsPointMax.x));
+	coefsHigh[4] = (TsPointMax.slope * (FMath::LogX(10, TsPointMax.x) + 0.5 * knotIncHigh)) + (FMath::LogX(10, TsPointMax.y) - TsPointMax.slope * FMath::LogX(10, TsPointMax.x));
+	// NOTE: if slope=0, then the above becomes just
+		// coefsHigh[0] = FMath::LogX(10, TsPointHigh.y);
+		// coefsHigh[1] = FMath::LogX(10, TsPointHigh.y);
+	// leaving it as a variable for now in case we decide we need non-zero slope extensions
+
+	// Middle coefficient (which defines the "sharpness of the bend") is linearly interpolated
+	FVector2f bendsHigh[2] = { FVector2f(MAX_STOP_SDR, 0.89) ,
+							FVector2f(MAX_STOP_RRT, 0.90) };
+	float pctHigh = interpolate1D(bendsHigh, log2(TsPointMax.x / 0.18));
+	coefsHigh[2] = FMath::LogX(10, TsPointMid.y) + pctHigh * (FMath::LogX(10, TsPointMax.y) - FMath::LogX(10, TsPointMid.y));
+}
+
+float shift(float inValue, float expShift)
+{
+	return pow(2., (log2(inValue) - expShift));
+}
+
+TsParams init_TsParams(
+	float minLum,
+	float maxLum,
+	float expShift = 0
+)
+{
+	TsPoint MIN_PT = { lookup_ACESmin(minLum), minLum, 0.0 };
+	TsPoint MID_PT = { 0.18, 4.8, 1.55 };
+	TsPoint MAX_PT = { lookup_ACESmax(maxLum), maxLum, 0.0 };
+	float cLow[5];
+	init_coefsLow(MIN_PT, MID_PT, cLow);
+	float cHigh[5];
+	init_coefsHigh(MID_PT, MAX_PT, cHigh);
+	MIN_PT.x = shift(lookup_ACESmin(minLum), expShift);
+	MID_PT.x = shift(0.18, expShift);
+	MAX_PT.x = shift(lookup_ACESmax(maxLum), expShift);
+
+	check(MIN_PT.slope == 0.0f);
+	check(MAX_PT.slope == 0.0f);
+
+	TsParams P = {
+		{MIN_PT.x, MIN_PT.y, MIN_PT.slope},
+		{MID_PT.x, MID_PT.y, MID_PT.slope},
+		{MAX_PT.x, MAX_PT.y, MAX_PT.slope},
+		{cLow[0],  cLow[1],  cLow[2],  cLow[3],  cLow[4]},
+		{cHigh[0], cHigh[1], cHigh[2], cHigh[3], cHigh[4]}
+	};
+
+	return P;
+}
+
+FVector3f ApplyM1Matrix(const FVector3f& cf)
+{
+	return FVector3f(
+		0.5f * cf[0] - 1.0f * cf[1] + 0.5f * cf[2],
+	   -1.0f * cf[0] + 1.0f * cf[1] + 0.0f * cf[2],
+		0.5f * cf[0] + 0.5f * cf[1] + 0.0f * cf[2]
+	);
+}
+
+float inv_ssts
+(
+	const float y,
+	const TsParams C
+)
+{
+	const int N_KNOTS_LOW = 4;
+	const int N_KNOTS_HIGH = 4;
+
+	const float KNOT_INC_LOW = (FMath::LogX(10, C.Mid.x) - FMath::LogX(10, C.Min.x)) / (N_KNOTS_LOW - 1.);
+	const float KNOT_INC_HIGH = (FMath::LogX(10, C.Max.x) - FMath::LogX(10, C.Mid.x)) / (N_KNOTS_HIGH - 1.);
+
+	// KNOT_Y is luminance of the spline at each knot
+	float KNOT_Y_LOW[N_KNOTS_LOW];
+	///+warning: redefinition of 'i'
+	{
+		for (int i = 0; i < N_KNOTS_LOW; i = i + 1) {
+			KNOT_Y_LOW[i] = (C.coefsLow[i] + C.coefsLow[i + 1]) / 2.;
+		};
+	}
+	///-
+
+	float KNOT_Y_HIGH[N_KNOTS_HIGH];
+	///+ warning: redefinition of 'i'
+	{
+		for (int i = 0; i < N_KNOTS_HIGH; i = i + 1) {
+			KNOT_Y_HIGH[i] = (C.coefsHigh[i] + C.coefsHigh[i + 1]) / 2.;
+		};
+	}
+	///-
+	float logy = FMath::LogX(10, FMath::Max(y, 1e-10));
+
+	float logx;
+	if (logy <= FMath::LogX(10, C.Min.y)) {
+
+		logx = FMath::LogX(10, C.Min.x);
+
+	}
+	else if ((logy > FMath::LogX(10, C.Min.y)) && (logy <= FMath::LogX(10, C.Mid.y))) {
+
+		int j=0;
+		FVector4f cf;
+		if (logy > KNOT_Y_LOW[0] && logy <= KNOT_Y_LOW[1]) {
+			cf[0] = C.coefsLow[0];  cf[1] = C.coefsLow[1];  cf[2] = C.coefsLow[2];  j = 0;
+		}
+		else if (logy > KNOT_Y_LOW[1] && logy <= KNOT_Y_LOW[2]) {
+			cf[0] = C.coefsLow[1];  cf[1] = C.coefsLow[2];  cf[2] = C.coefsLow[3];  j = 1;
+		}
+		else if (logy > KNOT_Y_LOW[2] && logy <= KNOT_Y_LOW[3]) {
+			cf[0] = C.coefsLow[2];  cf[1] = C.coefsLow[3];  cf[2] = C.coefsLow[4];  j = 2;
+		}
+
+		cf[3] = 0;
+
+		const FVector4f tmp = ApplyM1Matrix(cf);
+
+		float a = tmp[0];
+		float b = tmp[1];
+		float c = tmp[2];
+		c = c - logy;
+
+		const float d = sqrt(b * b - 4. * a * c);
+
+		const float t = (2. * c) / (-d - b);
+
+		logx = FMath::LogX(10, C.Min.x) + (t + j) * KNOT_INC_LOW;
+
+	}
+	else if ((logy > FMath::LogX(10, C.Mid.y)) && (logy < FMath::LogX(10, C.Max.y))) {
+
+		int j=0;
+		FVector4f cf;
+		if (logy >= KNOT_Y_HIGH[0] && logy <= KNOT_Y_HIGH[1]) {
+			cf[0] = C.coefsHigh[0];  cf[1] = C.coefsHigh[1];  cf[2] = C.coefsHigh[2];  j = 0;
+		}
+		else if (logy > KNOT_Y_HIGH[1] && logy <= KNOT_Y_HIGH[2]) {
+			cf[0] = C.coefsHigh[1];  cf[1] = C.coefsHigh[2];  cf[2] = C.coefsHigh[3];  j = 1;
+		}
+		else if (logy > KNOT_Y_HIGH[2] && logy <= KNOT_Y_HIGH[3]) {
+			cf[0] = C.coefsHigh[2];  cf[1] = C.coefsHigh[3];  cf[2] = C.coefsHigh[4];  j = 2;
+		}
+		cf[3] = 0;
+
+		const FVector4f tmp = ApplyM1Matrix(cf);
+
+		float a = tmp[0];
+		float b = tmp[1];
+		float c = tmp[2];
+		c = c - logy;
+
+		const float d = sqrt(b * b - 4. * a * c);
+
+		const float t = (2. * c) / (-d - b);
+
+		logx = FMath::LogX(10, C.Mid.x) + (t + j) * KNOT_INC_HIGH;
+
+	}
+	else { //if ( logy >= FMath::LogX(10, C.Max.y) ) {
+
+		logx = FMath::LogX(10, C.Max.x);
+
+	}
+
+	return FMath::Pow(10, logx);
+
+}
+
+void ConfigureACESTonemapParams(FACESTonemapParams& OutACESTonemapParams, float MinLuminance, float MidLuminance, float MaxLuminance)
+{
+	TsParams PARAMS_DEFAULT = init_TsParams(MinLuminance, MaxLuminance);
+	float tmp = inv_ssts(MidLuminance, PARAMS_DEFAULT);
+	float expShift = log2(tmp) - log2(0.18);
+	TsParams PARAMS = init_TsParams(MinLuminance, MaxLuminance, expShift);
+
+	OutACESTonemapParams.ACESMinMaxData.X = PARAMS.Min.x;
+	OutACESTonemapParams.ACESMinMaxData.Y = PARAMS.Min.y;
+	OutACESTonemapParams.ACESMinMaxData.Z = PARAMS.Max.x;
+	OutACESTonemapParams.ACESMinMaxData.W = PARAMS.Max.y;
+	check(FMath::Abs(PARAMS.Min.slope) < 1e-6f);
+	check(FMath::Abs(PARAMS.Max.slope) < 1e-6f);
+
+	OutACESTonemapParams.ACESMidData.X = PARAMS.Mid.x;
+	OutACESTonemapParams.ACESMidData.Y = PARAMS.Mid.y;
+	OutACESTonemapParams.ACESMidData.Z = PARAMS.Mid.slope;
+
+	OutACESTonemapParams.ACESCoefsLow_0.X = PARAMS.coefsLow[0];
+	OutACESTonemapParams.ACESCoefsLow_0.Y = PARAMS.coefsLow[1];
+	OutACESTonemapParams.ACESCoefsLow_0.Z = PARAMS.coefsLow[2];
+	OutACESTonemapParams.ACESCoefsLow_0.W = PARAMS.coefsLow[3];
+	OutACESTonemapParams.ACESCoefsLow_4 = PARAMS.coefsLow[4];
+	OutACESTonemapParams.ACESCoefsHigh_0.X = PARAMS.coefsHigh[0];
+	OutACESTonemapParams.ACESCoefsHigh_0.Y = PARAMS.coefsHigh[1];
+	OutACESTonemapParams.ACESCoefsHigh_0.Z = PARAMS.coefsHigh[2];
+	OutACESTonemapParams.ACESCoefsHigh_0.W = PARAMS.coefsHigh[3];
+	OutACESTonemapParams.ACESCoefsHigh_4 = PARAMS.coefsHigh[4];
+
+	OutACESTonemapParams.ACESSceneColorMultiplier = CVarHDRAcesColorMultiplier.GetValueOnAnyThread();
+
+}
\ No newline at end of file
diff --git a/Engine/Source/Runtime/RenderCore/Public/HDRHelper.h b/Engine/Source/Runtime/RenderCore/Public/HDRHelper.h
--- a/Engine/Source/Runtime/RenderCore/Public/HDRHelper.h
+++ b/Engine/Source/Runtime/RenderCore/Public/HDRHelper.h
@@ -26,3 +26,45 @@ struct FDisplayInformation
 	FIntRect DesktopCoordinates;
 	bool   bHDRSupported;
 };
+
+struct FACESTonemapParams
+{
+	// Default Values such as MinLuminanceLog10=-4, MidLuminance=15, MaxLuminance=1000, SceneColorMultiplier=1.5
+	FACESTonemapParams()
+	: ACESMinMaxData(2.50898620e-06f, 1e-4f, 692.651123f, 1000.0f)
+	, ACESMidData(0.0822144598f, 4.80000019f, 1.55f, 1.0f)
+	, ACESCoefsLow_0(-4.00000000f, -4.00000000f, -3.15737653f, -0.485249996f)
+	, ACESCoefsHigh_0(-0.332863420f, 1.69534576f, 2.75812411f, 3.00000000f)
+	, ACESCoefsLow_4(1.84773231f)
+	, ACESCoefsHigh_4(3.0f)
+	, ACESSceneColorMultiplier(1.5f)
+	{
+
+	}
+	FVector4f ACESMinMaxData;
+	FVector4f ACESMidData;
+	FVector4f ACESCoefsLow_0;
+	FVector4f ACESCoefsHigh_0;
+	float ACESCoefsLow_4;
+	float ACESCoefsHigh_4;
+	float ACESSceneColorMultiplier;
+};
+
+UE_DEPRECATED(all, "For internal use only.")
+RENDERCORE_API FACESTonemapParams HDRGetACESTonemapParams();
+
+template<class FACESTonemapParametersType>
+inline void GetACESTonemapParameters(FACESTonemapParametersType& ACESTonemapParameters)
+{
+	PRAGMA_DISABLE_DEPRECATION_WARNINGS
+	FACESTonemapParams TonemapperParams = HDRGetACESTonemapParams();
+	PRAGMA_ENABLE_DEPRECATION_WARNINGS
+	
+	ACESTonemapParameters.ACESMinMaxData = TonemapperParams.ACESMinMaxData;
+	ACESTonemapParameters.ACESMidData = TonemapperParams.ACESMidData;
+	ACESTonemapParameters.ACESCoefsLow_0 = TonemapperParams.ACESCoefsLow_0;
+	ACESTonemapParameters.ACESCoefsHigh_0 = TonemapperParams.ACESCoefsHigh_0;
+	ACESTonemapParameters.ACESCoefsLow_4 = TonemapperParams.ACESCoefsLow_4;
+	ACESTonemapParameters.ACESCoefsHigh_4 = TonemapperParams.ACESCoefsHigh_4;
+	ACESTonemapParameters.ACESSceneColorMultiplier = TonemapperParams.ACESSceneColorMultiplier;
+}
\ No newline at end of file
diff --git a/Engine/Source/Runtime/Renderer/Private/PostProcess/PostProcessCombineLUTs.cpp b/Engine/Source/Runtime/Renderer/Private/PostProcess/PostProcessCombineLUTs.cpp
--- a/Engine/Source/Runtime/Renderer/Private/PostProcess/PostProcessCombineLUTs.cpp
+++ b/Engine/Source/Runtime/Renderer/Private/PostProcess/PostProcessCombineLUTs.cpp
@@ -4,6 +4,7 @@
 #include "PostProcess/PostProcessTonemap.h"
 #include "ScenePrivate.h"
 #include "VolumeRendering.h"
+#include "HDRHelper.h"
 
 namespace
 {
@@ -90,11 +91,23 @@ FColorRemapParameters GetColorRemapParameters()
 	return Parameters;
 }
 
+BEGIN_SHADER_PARAMETER_STRUCT(FACESTonemapShaderParameters, )
+	SHADER_PARAMETER(FVector4f, ACESMinMaxData) // xy = min ACES/luminance, zw = max ACES/luminance
+	SHADER_PARAMETER(FVector4f, ACESMidData) // x = mid ACES, y = mid luminance, z = mid slope
+	SHADER_PARAMETER(FVector4f, ACESCoefsLow_0) // coeflow 0-3
+	SHADER_PARAMETER(FVector4f, ACESCoefsHigh_0) // coefhigh 0-3
+	SHADER_PARAMETER(float, ACESCoefsLow_4)
+	SHADER_PARAMETER(float, ACESCoefsHigh_4)
+	SHADER_PARAMETER(float, ACESSceneColorMultiplier)
+END_SHADER_PARAMETER_STRUCT()
+
+
 BEGIN_SHADER_PARAMETER_STRUCT(FCombineLUTParameters, )
 	SHADER_PARAMETER_TEXTURE_ARRAY(Texture2D, Textures, [GMaxLUTBlendCount])
 	SHADER_PARAMETER_SAMPLER_ARRAY(SamplerState, Samplers, [GMaxLUTBlendCount])
 	SHADER_PARAMETER_SCALAR_ARRAY(float, LUTWeights, [GMaxLUTBlendCount])
 	SHADER_PARAMETER_STRUCT_REF(FWorkingColorSpaceShaderParameters, WorkingColorSpace)
+	SHADER_PARAMETER_STRUCT_INCLUDE(FACESTonemapShaderParameters, ACESTonemapParameters)
 	SHADER_PARAMETER(FVector4f, OverlayColor)
 	SHADER_PARAMETER(FVector3f, ColorScale)
 	SHADER_PARAMETER(FVector4f, ColorSaturation)
@@ -170,6 +183,7 @@ void GetCombineLUTParameters(
 	}
 
 	Parameters.WorkingColorSpace = GDefaultWorkingColorSpaceUniformBuffer.GetUniformBufferRef();
+	GetACESTonemapParameters(Parameters.ACESTonemapParameters);
 
 	Parameters.ColorScale = FVector4f(View.ColorScale);
 	Parameters.OverlayColor = View.OverlayColor;
diff --git a/Engine/Source/Runtime/Renderer/Private/PostProcess/PostProcessDeviceEncodingOnly.cpp b/Engine/Source/Runtime/Renderer/Private/PostProcess/PostProcessDeviceEncodingOnly.cpp
--- a/Engine/Source/Runtime/Renderer/Private/PostProcess/PostProcessDeviceEncodingOnly.cpp
+++ b/Engine/Source/Runtime/Renderer/Private/PostProcess/PostProcessDeviceEncodingOnly.cpp
@@ -83,6 +83,15 @@ FDeviceEncodingOnlyOutputDeviceParameters GetDeviceEncodingOnlyOutputDeviceParam
 	return Parameters;
 }
 
+BEGIN_SHADER_PARAMETER_STRUCT(FDeviceEncodingACESTonemapShaderParameters, )
+	SHADER_PARAMETER(FVector4f, ACESMinMaxData) // xy = min ACES/luminance, zw = max ACES/luminance
+	SHADER_PARAMETER(FVector4f, ACESMidData) // x = mid ACES, y = mid luminance, z = mid slope
+	SHADER_PARAMETER(FVector4f, ACESCoefsLow_0) // coeflow 0-3
+	SHADER_PARAMETER(FVector4f, ACESCoefsHigh_0) // coefhigh 0-3
+	SHADER_PARAMETER(float, ACESCoefsLow_4)
+	SHADER_PARAMETER(float, ACESCoefsHigh_4)
+	SHADER_PARAMETER(float, ACESSceneColorMultiplier)
+END_SHADER_PARAMETER_STRUCT()
 
 BEGIN_SHADER_PARAMETER_STRUCT(FDeviceEncodingOnlyParameters, )
 	SHADER_PARAMETER_STRUCT_REF(FViewUniformShaderParameters, View)
@@ -90,6 +99,7 @@ BEGIN_SHADER_PARAMETER_STRUCT(FDeviceEncodingOnlyParameters, )
 	SHADER_PARAMETER_STRUCT_INCLUDE(FDeviceEncodingOnlyOutputDeviceParameters, OutputDevice)
 	SHADER_PARAMETER_STRUCT(FScreenPassTextureViewportParameters, Color)
 	SHADER_PARAMETER_STRUCT(FScreenPassTextureViewportParameters, Output)
+	SHADER_PARAMETER_STRUCT_INCLUDE(FDeviceEncodingACESTonemapShaderParameters, ACESTonemapParameters)
 	SHADER_PARAMETER_RDG_TEXTURE(Texture2D, ColorTexture)
 	SHADER_PARAMETER_SAMPLER(SamplerState, ColorSampler)
 	SHADER_PARAMETER(float, EditorNITLevel)
@@ -211,6 +221,7 @@ FScreenPassTexture AddDeviceEncodingOnlyPass(FRDGBuilder& GraphBuilder, const FV
 	CommonParameters.OutputDevice = GetDeviceEncodingOnlyOutputDeviceParameters(ViewFamily);
 	CommonParameters.Color = GetScreenPassTextureViewportParameters(SceneColorViewport);
 	CommonParameters.Output = GetScreenPassTextureViewportParameters(OutputViewport);
+	GetACESTonemapParameters(CommonParameters.ACESTonemapParameters);
 	CommonParameters.ColorTexture = Inputs.SceneColor.Texture;
 	CommonParameters.ColorSampler = BilinearClampSampler;
 	CommonParameters.EditorNITLevel = EditorNITLevel;
diff --git a/Engine/Source/Runtime/SlateRHIRenderer/Private/SlatePostProcessor.cpp b/Engine/Source/Runtime/SlateRHIRenderer/Private/SlatePostProcessor.cpp
--- a/Engine/Source/Runtime/SlateRHIRenderer/Private/SlatePostProcessor.cpp
+++ b/Engine/Source/Runtime/SlateRHIRenderer/Private/SlatePostProcessor.cpp
@@ -48,6 +48,7 @@ public:
 		UIWriteMaskTexture.Bind(Initializer.ParameterMap, TEXT("UIWriteMaskTexture"));
 		UISampler.Bind(Initializer.ParameterMap, TEXT("UISampler"));
 		UILevel.Bind(Initializer.ParameterMap, TEXT("UILevel"));
+		UILuminance.Bind(Initializer.ParameterMap, TEXT("UILuminance"));
 		UITextureSize.Bind(Initializer.ParameterMap, TEXT("UITextureSize"));
 	}
 	FBlitUIToHDRPSBase() = default;
@@ -57,7 +58,9 @@ public:
 		SetTextureParameter(RHICmdList, RHICmdList.GetBoundPixelShader(), HDRTexture, UISampler, TStaticSamplerState<SF_Bilinear>::GetRHI(), HDRTextureRHI);
 		SetTextureParameter(RHICmdList, RHICmdList.GetBoundPixelShader(), UITexture, UISampler, TStaticSamplerState<SF_Bilinear>::GetRHI(), UITextureRHI);
 		static auto CVarHDRUILevel = IConsoleManager::Get().FindConsoleVariable(TEXT("r.HDR.UI.Level"));
+		static auto CVarHDRUILuminance = IConsoleManager::Get().FindConsoleVariable(TEXT("r.HDR.UI.Luminance"));
 		SetShaderValue(RHICmdList, RHICmdList.GetBoundPixelShader(), UILevel, CVarHDRUILevel ? CVarHDRUILevel->GetFloat() : 1.0f);
+		SetShaderValue(RHICmdList, RHICmdList.GetBoundPixelShader(), UILuminance, CVarHDRUILuminance ? CVarHDRUILuminance->GetFloat() : 300.0f);
 		SetShaderValue(RHICmdList, RHICmdList.GetBoundPixelShader(), UITextureSize, InUITextureSize);
 		if (UITextureWriteMaskRHI != nullptr)
 		{
@@ -87,6 +90,7 @@ private:
 	LAYOUT_FIELD(FShaderResourceParameter, UISampler);
 	LAYOUT_FIELD(FShaderResourceParameter, UIWriteMaskTexture);
 	LAYOUT_FIELD(FShaderParameter, UILevel);
+	LAYOUT_FIELD(FShaderParameter, UILuminance);
 	LAYOUT_FIELD(FShaderParameter, UITextureSize);
 };
 
diff --git a/Engine/Source/Runtime/SlateRHIRenderer/Private/SlateRHIRenderer.cpp b/Engine/Source/Runtime/SlateRHIRenderer/Private/SlateRHIRenderer.cpp
--- a/Engine/Source/Runtime/SlateRHIRenderer/Private/SlateRHIRenderer.cpp
+++ b/Engine/Source/Runtime/SlateRHIRenderer/Private/SlateRHIRenderer.cpp
@@ -56,6 +56,12 @@ static TAutoConsoleVariable<float> CVarUILevel(
 	TEXT("Luminance level for UI elements when compositing into HDR framebuffer (default: 1.0)."),
 	ECVF_RenderThreadSafe);
 
+static TAutoConsoleVariable<float> CVarHDRUILuminance(
+	TEXT("r.HDR.UI.Luminance"),
+	300.0f,
+	TEXT("Base Luminance in nits for UI elements when compositing into HDR framebuffer. Gets multiplied by r.HDR.UI.Level"),
+	ECVF_RenderThreadSafe);
+
 static TAutoConsoleVariable<int32> CVarUICompositeMode(
 	TEXT("r.HDR.UI.CompositeMode"),
 	1,
@@ -594,6 +600,13 @@ public:
 		OutputDevice.Bind(Initializer.ParameterMap, TEXT("OutputDevice"));
 		OutputGamut.Bind(Initializer.ParameterMap, TEXT("OutputGamut"));
 		OutputMaxLuminance.Bind(Initializer.ParameterMap, TEXT("OutputMaxLuminance"));
+		ACESMinMaxData.Bind(Initializer.ParameterMap, TEXT("ACESMinMaxData"));
+		ACESMidData.Bind(Initializer.ParameterMap, TEXT("ACESMidData"));
+		ACESCoefsLow_0.Bind(Initializer.ParameterMap, TEXT("ACESCoefsLow_0"));
+		ACESCoefsHigh_0.Bind(Initializer.ParameterMap, TEXT("ACESCoefsHigh_0"));
+		ACESCoefsLow_4.Bind(Initializer.ParameterMap, TEXT("ACESCoefsLow_4"));
+		ACESCoefsHigh_4.Bind(Initializer.ParameterMap, TEXT("ACESCoefsHigh_4"));
+		ACESSceneColorMultiplier.Bind(Initializer.ParameterMap, TEXT("ACESSceneColorMultiplier"));
 	}
 	FCompositeLUTGenerationPS() {}
 
@@ -620,6 +633,18 @@ public:
 		SetShaderValue(RHICmdList, RHICmdList.GetBoundPixelShader(), OutputDevice, OutputDeviceValue);
 		SetShaderValue(RHICmdList, RHICmdList.GetBoundPixelShader(), OutputGamut, OutputGamutValue);
 		SetShaderValue(RHICmdList, RHICmdList.GetBoundPixelShader(), OutputMaxLuminance, DisplayMaxLuminance);
+
+		FACESTonemapParams TmpInternalACESParams;
+		GetACESTonemapParameters(TmpInternalACESParams);
+
+		SetShaderValue(RHICmdList, RHICmdList.GetBoundPixelShader(), ACESMinMaxData, TmpInternalACESParams.ACESMinMaxData);
+		SetShaderValue(RHICmdList, RHICmdList.GetBoundPixelShader(), ACESMidData, TmpInternalACESParams.ACESMidData);
+		SetShaderValue(RHICmdList, RHICmdList.GetBoundPixelShader(), ACESCoefsLow_0, TmpInternalACESParams.ACESCoefsLow_0);
+		SetShaderValue(RHICmdList, RHICmdList.GetBoundPixelShader(), ACESCoefsHigh_0, TmpInternalACESParams.ACESCoefsHigh_0);
+		SetShaderValue(RHICmdList, RHICmdList.GetBoundPixelShader(), ACESCoefsLow_4, TmpInternalACESParams.ACESCoefsLow_4);
+		SetShaderValue(RHICmdList, RHICmdList.GetBoundPixelShader(), ACESCoefsHigh_4, TmpInternalACESParams.ACESCoefsHigh_4);
+		SetShaderValue(RHICmdList, RHICmdList.GetBoundPixelShader(), ACESSceneColorMultiplier, TmpInternalACESParams.ACESSceneColorMultiplier);
+
 	}
 
 	static const TCHAR* GetSourceFilename()
@@ -636,6 +661,13 @@ private:
 	LAYOUT_FIELD(FShaderParameter, OutputDevice);
 	LAYOUT_FIELD(FShaderParameter, OutputGamut);
 	LAYOUT_FIELD(FShaderParameter, OutputMaxLuminance);
+	LAYOUT_FIELD(FShaderParameter, ACESMinMaxData) // xy = min ACES/luminance, zw = max ACES/luminance
+	LAYOUT_FIELD(FShaderParameter, ACESMidData) // x = mid ACES, y = mid luminance, z = mid slope
+	LAYOUT_FIELD(FShaderParameter, ACESCoefsLow_0) // coeflow 0-3
+	LAYOUT_FIELD(FShaderParameter, ACESCoefsHigh_0) // coefhigh 0-3
+	LAYOUT_FIELD(FShaderParameter, ACESCoefsLow_4)
+	LAYOUT_FIELD(FShaderParameter, ACESCoefsHigh_4)
+	LAYOUT_FIELD(FShaderParameter, ACESSceneColorMultiplier)
 };
 
 IMPLEMENT_SHADER_TYPE(, FCompositeLUTGenerationPS, TEXT("/Engine/Private/CompositeUIPixelShader.usf"), TEXT("GenerateLUTPS"), SF_Pixel);
@@ -660,6 +692,7 @@ public:
 		ColorSpaceLUT.Bind(Initializer.ParameterMap, TEXT("ColorSpaceLUT"));
 		ColorSpaceLUTSampler.Bind(Initializer.ParameterMap, TEXT("ColorSpaceLUTSampler"));
 		UILevel.Bind(Initializer.ParameterMap, TEXT("UILevel"));
+		UILuminance.Bind(Initializer.ParameterMap, TEXT("UILuminance"));
 		OutputDevice.Bind(Initializer.ParameterMap, TEXT("OutputDevice"));
 	}
 	FCompositeShaderBase() = default;
@@ -673,7 +706,8 @@ public:
 		SetTextureParameter(RHICmdList, BoundShader, ColorSpaceLUT, ColorSpaceLUTSampler, TStaticSamplerState<SF_Bilinear, AM_Clamp, AM_Clamp, AM_Clamp>::GetRHI(), ColorSpaceLUTRHI);
 		SetShaderValue(RHICmdList, BoundShader, UILevel, CVarUILevel.GetValueOnRenderThread());
 		SetShaderValue(RHICmdList, BoundShader, OutputDevice, CVarOutputDevice->GetValueOnRenderThread());
-		
+		SetShaderValue(RHICmdList, RHICmdList.GetBoundPixelShader(), UILuminance, CVarHDRUILuminance.GetValueOnRenderThread());
+
 		if (UITextureWriteMaskRHI != nullptr)
 		{
 			SetTextureParameter(RHICmdList, BoundShader, UIWriteMaskTexture, UITextureWriteMaskRHI);
@@ -692,6 +726,7 @@ private:
 	LAYOUT_FIELD(FShaderResourceParameter, ColorSpaceLUT);
 	LAYOUT_FIELD(FShaderResourceParameter, ColorSpaceLUTSampler);
 	LAYOUT_FIELD(FShaderParameter, UILevel);
+	LAYOUT_FIELD(FShaderParameter, UILuminance);
 	LAYOUT_FIELD(FShaderParameter, OutputDevice);
 };
 IMPLEMENT_TYPE_LAYOUT(FCompositeShaderBase);
