Change check ==0 to an ensure < 1e-6 to not block anyone for some FP imprecision

diff --git a/Engine/Source/Runtime/RenderCore/Private/RenderCore.cpp b/Engine/Source/Runtime/RenderCore/Private/RenderCore.cpp
--- a/Engine/Source/Runtime/RenderCore/Private/RenderCore.cpp
+++ b/Engine/Source/Runtime/RenderCore/Private/RenderCore.cpp
@@ -793,8 +793,8 @@ TsParams init_TsParams(
 	MID_PT.x = shift(0.18, expShift);
 	MAX_PT.x = shift(lookup_ACESmax(maxLum), expShift);
 
-	check(MIN_PT.slope == 0.0f);
-	check(MAX_PT.slope == 0.0f);
+	ensure(FMath::Abs(MIN_PT.slope) <= 1e-6f);
+	ensure(FMath::Abs(MAX_PT.slope) <= 1e-6f);
 
 	TsParams P = {
 		{MIN_PT.x, MIN_PT.y, MIN_PT.slope},
@@ -934,8 +934,8 @@ void ConfigureACESTonemapParams(FACESTonemapParams& OutACESTonemapParams, float
 	OutACESTonemapParams.ACESMinMaxData.Y = PARAMS.Min.y;
 	OutACESTonemapParams.ACESMinMaxData.Z = PARAMS.Max.x;
 	OutACESTonemapParams.ACESMinMaxData.W = PARAMS.Max.y;
-	check(FMath::Abs(PARAMS.Min.slope) < 1e-6f);
-	check(FMath::Abs(PARAMS.Max.slope) < 1e-6f);
+	ensure(FMath::Abs(PARAMS.Min.slope) <= 1e-6f);
+	ensure(FMath::Abs(PARAMS.Max.slope) <= 1e-6f);
 
 	OutACESTonemapParams.ACESMidData.X = PARAMS.Mid.x;
 	OutACESTonemapParams.ACESMidData.Y = PARAMS.Mid.y;
