#ue5 - Fix type truncation warnings in misc Engine files.

diff --git a/Engine/Source/Runtime/Core/Public/Math/Vector2D.h b/Engine/Source/Runtime/Core/Public/Math/Vector2D.h
--- a/Engine/Source/Runtime/Core/Public/Math/Vector2D.h
+++ b/Engine/Source/Runtime/Core/Public/Math/Vector2D.h
@@ -1189,7 +1189,10 @@ FORCEINLINE TVector2<T> TVector2<T>::RoundToVector() const
 	{
 		return TVector2<T>(FMath::RoundToFloat(X), FMath::RoundToFloat(Y));
 	}
-	return TVector2<T>(FMath::RoundToDouble(X), FMath::RoundToDouble(Y));
+	else
+	{
+		return TVector2<T>(FMath::RoundToDouble(X), FMath::RoundToDouble(Y));
+	}
 }
 
 template<typename T>
diff --git a/Engine/Source/Runtime/Engine/Classes/Engine/EngineTypes.h b/Engine/Source/Runtime/Engine/Classes/Engine/EngineTypes.h
--- a/Engine/Source/Runtime/Engine/Classes/Engine/EngineTypes.h
+++ b/Engine/Source/Runtime/Engine/Classes/Engine/EngineTypes.h
@@ -697,7 +697,7 @@ public:
 	void AddShadingModel(EStrataShadingModel InShadingModel) { check(InShadingModel < SSM_NUM); ShadingModelField |= (1 << (uint16)InShadingModel); }
 	bool HasShadingModel(EStrataShadingModel InShadingModel) const { return (ShadingModelField & (1 << (uint16)InShadingModel)) != 0; }
 	bool HasOnlyShadingModel(EStrataShadingModel InShadingModel) const { return ShadingModelField == (1 << (uint16)InShadingModel); }
-	uint8 GetShadingModelField() const { return ShadingModelField; }
+	uint16 GetShadingModelField() const { return ShadingModelField; }
 	int32 CountShadingModels() const { return FMath::CountBits(ShadingModelField); }
 
 	// Subsurface profiles
@@ -1464,8 +1464,8 @@ FORCEINLINE ECollisionEnabled::Type CollisionEnabledIntersection(ECollisionEnabl
 	// Subtract 1 because the first index is NoCollision.
 	// If both indices indicate _some_ collision setting (ie, greater than -1),
 	// lookup their intersection setting in the matrix.
-	const int8 IndexA = (int32)CollisionEnabledA - 1;
-	const int8 IndexB = (int32)CollisionEnabledB - 1;
+	const int32 IndexA = (int32)CollisionEnabledA - 1;
+	const int32 IndexB = (int32)CollisionEnabledB - 1;
 	if (IndexA >= 0 && IndexB >= 0)
 	{
 		return IntersectionMatrix[IndexA][IndexB];
diff --git a/Engine/Source/Runtime/Engine/Public/Matrix3x4.h b/Engine/Source/Runtime/Engine/Public/Matrix3x4.h
--- a/Engine/Source/Runtime/Engine/Public/Matrix3x4.h
+++ b/Engine/Source/Runtime/Engine/Public/Matrix3x4.h
@@ -18,20 +18,20 @@ MS_ALIGN(16) struct FMatrix3x4
 		const FMatrix::FReal* RESTRICT Src = &(Mat.M[0][0]);
 		float* RESTRICT Dest = &(M[0][0]);
 
-		Dest[0] = (FMatrix::FReal)Src[0];   // [0][0]
-		Dest[1] = (FMatrix::FReal)Src[1];   // [0][1]
-		Dest[2] = (FMatrix::FReal)Src[2];   // [0][2]
-		Dest[3] = (FMatrix::FReal)Src[3];   // [0][3]
+		Dest[0] = (float)Src[0];   // [0][0]
+		Dest[1] = (float)Src[1];   // [0][1]
+		Dest[2] = (float)Src[2];   // [0][2]
+		Dest[3] = (float)Src[3];   // [0][3]
 
-		Dest[4] = (FMatrix::FReal)Src[4];   // [1][0]
-		Dest[5] = (FMatrix::FReal)Src[5];   // [1][1]
-		Dest[6] = (FMatrix::FReal)Src[6];   // [1][2]
-		Dest[7] = (FMatrix::FReal)Src[7];   // [1][3]
+		Dest[4] = (float)Src[4];   // [1][0]
+		Dest[5] = (float)Src[5];   // [1][1]
+		Dest[6] = (float)Src[6];   // [1][2]
+		Dest[7] = (float)Src[7];   // [1][3]
 
-		Dest[8] = (FMatrix::FReal)Src[8];   // [2][0]
-		Dest[9] = (FMatrix::FReal)Src[9];   // [2][1]
-		Dest[10] = (FMatrix::FReal)Src[10]; // [2][2]
-		Dest[11] = (FMatrix::FReal)Src[11]; // [2][3]
+		Dest[8] = (float)Src[8];   // [2][0]
+		Dest[9] = (float)Src[9];   // [2][1]
+		Dest[10] = (float)Src[10]; // [2][2]
+		Dest[11] = (float)Src[11]; // [2][3]
 	}
 
 	FORCEINLINE void SetMatrixTranspose(const FMatrix& Mat)
@@ -39,20 +39,20 @@ MS_ALIGN(16) struct FMatrix3x4
 		const FMatrix::FReal* RESTRICT Src = &(Mat.M[0][0]);
 		float* RESTRICT Dest = &(M[0][0]);
 
-		Dest[0] = (FMatrix::FReal)Src[0];   // [0][0]
-		Dest[1] = (FMatrix::FReal)Src[4];   // [1][0]
-		Dest[2] = (FMatrix::FReal)Src[8];   // [2][0]
-		Dest[3] = (FMatrix::FReal)Src[12];  // [3][0]
+		Dest[0] = (float)Src[0];   // [0][0]
+		Dest[1] = (float)Src[4];   // [1][0]
+		Dest[2] = (float)Src[8];   // [2][0]
+		Dest[3] = (float)Src[12];  // [3][0]
 
-		Dest[4] = (FMatrix::FReal)Src[1];   // [0][1]
-		Dest[5] = (FMatrix::FReal)Src[5];   // [1][1]
-		Dest[6] = (FMatrix::FReal)Src[9];   // [2][1]
-		Dest[7] = (FMatrix::FReal)Src[13];  // [3][1]
+		Dest[4] = (float)Src[1];   // [0][1]
+		Dest[5] = (float)Src[5];   // [1][1]
+		Dest[6] = (float)Src[9];   // [2][1]
+		Dest[7] = (float)Src[13];  // [3][1]
 
-		Dest[8] = (FMatrix::FReal)Src[2];   // [0][2]
-		Dest[9] = (FMatrix::FReal)Src[6];   // [1][2]
-		Dest[10] = (FMatrix::FReal)Src[10]; // [2][2]
-		Dest[11] = (FMatrix::FReal)Src[14]; // [3][2]
+		Dest[8] = (float)Src[2];   // [0][2]
+		Dest[9] = (float)Src[6];   // [1][2]
+		Dest[10] = (float)Src[10]; // [2][2]
+		Dest[11] = (float)Src[14]; // [3][2]
 	}
 } GCC_ALIGN(16);
 
diff --git a/Engine/Source/Runtime/RenderCore/Public/PackedNormal.h b/Engine/Source/Runtime/RenderCore/Public/PackedNormal.h
--- a/Engine/Source/Runtime/RenderCore/Public/PackedNormal.h
+++ b/Engine/Source/Runtime/RenderCore/Public/PackedNormal.h
@@ -116,37 +116,37 @@ public:
 FORCEINLINE void FPackedNormal::operator=(const FVector3f& InVector)
 {
 	const float Scale = MAX_int8;
-	Vector.X = (int8)FMath::Clamp<int32>(FMath::RoundToInt(InVector.X * Scale), MIN_int8, MAX_int8);
-	Vector.Y = (int8)FMath::Clamp<int32>(FMath::RoundToInt(InVector.Y * Scale), MIN_int8, MAX_int8);
-	Vector.Z = (int8)FMath::Clamp<int32>(FMath::RoundToInt(InVector.Z * Scale), MIN_int8, MAX_int8);
+	Vector.X = (int8)FMath::Clamp<int32>(FMath::RoundToInt32(InVector.X * Scale), MIN_int8, MAX_int8);
+	Vector.Y = (int8)FMath::Clamp<int32>(FMath::RoundToInt32(InVector.Y * Scale), MIN_int8, MAX_int8);
+	Vector.Z = (int8)FMath::Clamp<int32>(FMath::RoundToInt32(InVector.Z * Scale), MIN_int8, MAX_int8);
 	Vector.W = MAX_int8;
 }
 
 FORCEINLINE void FPackedNormal::operator=(const FVector3d& InVector)
 {
 	const float Scale = MAX_int8;
-	Vector.X = (int8)FMath::Clamp<int32>(FMath::RoundToInt(InVector.X * Scale), MIN_int8, MAX_int8);
-	Vector.Y = (int8)FMath::Clamp<int32>(FMath::RoundToInt(InVector.Y * Scale), MIN_int8, MAX_int8);
-	Vector.Z = (int8)FMath::Clamp<int32>(FMath::RoundToInt(InVector.Z * Scale), MIN_int8, MAX_int8);
+	Vector.X = (int8)FMath::Clamp<int32>(FMath::RoundToInt32(InVector.X * Scale), MIN_int8, MAX_int8);
+	Vector.Y = (int8)FMath::Clamp<int32>(FMath::RoundToInt32(InVector.Y * Scale), MIN_int8, MAX_int8);
+	Vector.Z = (int8)FMath::Clamp<int32>(FMath::RoundToInt32(InVector.Z * Scale), MIN_int8, MAX_int8);
 	Vector.W = MAX_int8;
 }
 
 FORCEINLINE void FPackedNormal::operator=(const FVector4f& InVector)
 {
 	const float Scale = MAX_int8;
-	Vector.X = (int8)FMath::Clamp<int32>(FMath::RoundToInt(InVector.X * Scale), MIN_int8, MAX_int8);
-	Vector.Y = (int8)FMath::Clamp<int32>(FMath::RoundToInt(InVector.Y * Scale), MIN_int8, MAX_int8);
-	Vector.Z = (int8)FMath::Clamp<int32>(FMath::RoundToInt(InVector.Z * Scale), MIN_int8, MAX_int8);
-	Vector.W = (int8)FMath::Clamp<int32>(FMath::RoundToInt(InVector.W * Scale), MIN_int8, MAX_int8);
+	Vector.X = (int8)FMath::Clamp<int32>(FMath::RoundToInt32(InVector.X * Scale), MIN_int8, MAX_int8);
+	Vector.Y = (int8)FMath::Clamp<int32>(FMath::RoundToInt32(InVector.Y * Scale), MIN_int8, MAX_int8);
+	Vector.Z = (int8)FMath::Clamp<int32>(FMath::RoundToInt32(InVector.Z * Scale), MIN_int8, MAX_int8);
+	Vector.W = (int8)FMath::Clamp<int32>(FMath::RoundToInt32(InVector.W * Scale), MIN_int8, MAX_int8);
 }
 
 FORCEINLINE void FPackedNormal::operator=(const FVector4d& InVector)
 {
 	const float Scale = MAX_int8;
-	Vector.X = (int8)FMath::Clamp<int32>(FMath::RoundToInt(InVector.X * Scale), MIN_int8, MAX_int8);
-	Vector.Y = (int8)FMath::Clamp<int32>(FMath::RoundToInt(InVector.Y * Scale), MIN_int8, MAX_int8);
-	Vector.Z = (int8)FMath::Clamp<int32>(FMath::RoundToInt(InVector.Z * Scale), MIN_int8, MAX_int8);
-	Vector.W = (int8)FMath::Clamp<int32>(FMath::RoundToInt(InVector.W * Scale), MIN_int8, MAX_int8);
+	Vector.X = (int8)FMath::Clamp<int32>(FMath::RoundToInt32(InVector.X * Scale), MIN_int8, MAX_int8);
+	Vector.Y = (int8)FMath::Clamp<int32>(FMath::RoundToInt32(InVector.Y * Scale), MIN_int8, MAX_int8);
+	Vector.Z = (int8)FMath::Clamp<int32>(FMath::RoundToInt32(InVector.Z * Scale), MIN_int8, MAX_int8);
+	Vector.W = (int8)FMath::Clamp<int32>(FMath::RoundToInt32(InVector.W * Scale), MIN_int8, MAX_int8);
 }
 
 FORCEINLINE bool FPackedNormal::operator==(const FPackedNormal& B) const
