Add r.RayTracing.EnableInEditor and r.RayTracing.EnableInGame CVars to control whether ray tracing effects should be enabled by default in the editor and in game (if corresponding game user setting does not exist). Both default to enabled, preserving previous behavior.

diff --git a/Engine/Source/Runtime/RenderCore/Private/RenderUtils.cpp b/Engine/Source/Runtime/RenderCore/Private/RenderUtils.cpp
--- a/Engine/Source/Runtime/RenderCore/Private/RenderUtils.cpp
+++ b/Engine/Source/Runtime/RenderCore/Private/RenderUtils.cpp
@@ -27,6 +27,24 @@ FAutoConsoleVariableRef CVarAllowNanite(
 	ECVF_ReadOnly | ECVF_RenderThreadSafe
 );
 
+int32 GRayTracingEnableInGame = 1;
+FAutoConsoleVariableRef CVarRayTracingEnableInGame(
+	TEXT("r.RayTracing.EnableInGame"),
+	GRayTracingEnableInGame,
+	TEXT("Controls the default state of ray tracing effects when running the game. This setting is overridden by its counterpart in GameUserSettings.ini (if it exists) to allow control through in-game UI. ")
+	TEXT("(default = 1)"),
+	ECVF_ReadOnly
+);
+
+int32 GRayTracingEnableInEditor = 1;
+FAutoConsoleVariableRef CVarRayTracingEnableInEditor(
+	TEXT("r.RayTracing.EnableInEditor"),
+	GRayTracingEnableInEditor,
+	TEXT("Controls whether ray tracing effects are available by default when running the editor. This can be useful to improve editor performance when only some people require ray tracing features. ")
+	TEXT("(default = 1)"),
+	ECVF_ReadOnly
+);
+
 FBufferWithRDG::FBufferWithRDG() = default;
 FBufferWithRDG::FBufferWithRDG(const FBufferWithRDG & Other) = default;
 FBufferWithRDG& FBufferWithRDG::operator=(const FBufferWithRDG & Other) = default;
@@ -1401,7 +1419,8 @@ RENDERCORE_API void RenderUtilsInit()
 	// - Ray tracing must be enabled for the project
 	// - Skin cache must be enabled for the project
 	// - Current GPU, OS and driver must support ray tracing
-	// - User is running the Editor *OR* running the game with ray tracing enabled in graphics options
+	// - User is running the Editor and r.RayTracing.EnableInEditor=1 
+	//   *OR* running the game with ray tracing enabled in graphics options
 
 	// When ray tracing is enabled, we must load additional shaders and build acceleration structures for meshes.
 	// For this reason it is only possible to enable RT at startup and changing the state requires restart.
@@ -1418,9 +1437,11 @@ RENDERCORE_API void RenderUtilsInit()
 				// Ray tracing is enabled for the project and we are running on RT-capable machine,
 				// therefore the core ray tracing features are also enabled, so that required shaders
 				// are loaded, acceleration structures are built, etc.
-				GUseRayTracing = true;
+				GUseRayTracing = GRayTracingEnableInEditor != 0;
 
-				UE_LOG(LogRendererCore, Log, TEXT("Ray tracing is enabled for the editor. Reason: r.RayTracing=1."));
+				UE_LOG(LogRendererCore, Log, TEXT("Ray tracing is %s for the editor. Reason: r.RayTracing=1 and r.RayTracing.EnableInEditor=%d."),
+					GUseRayTracing ? TEXT("enabled") : TEXT("disabled"),
+					GRayTracingEnableInEditor);
 			}
 			else
 			{
@@ -1428,19 +1449,17 @@ RENDERCORE_API void RenderUtilsInit()
 				// Otherwise the current value is preserved.
 				if (GConfig->GetBool(TEXT("RayTracing"), TEXT("r.RayTracing.EnableInGame"), GUseRayTracing, GGameUserSettingsIni))
 				{
-					UE_LOG(LogRendererCore, Log, TEXT("Ray tracing is %s for the game. Reason: user setting r.RayTracing.EnableInGame=%d."),
+					UE_LOG(LogRendererCore, Log, TEXT("Ray tracing is %s for the game. Reason: game user setting r.RayTracing.EnableInGame=%d."),
 						GUseRayTracing ? TEXT("enabled") : TEXT("disabled"),
 						(int)GUseRayTracing);
 				}
 				else
 				{
-					GUseRayTracing = true;
-
-					UE_LOG(LogRendererCore, Log, TEXT("Ray tracing is enabled for the game. Reason: r.RayTracing=1, and r.RayTracing.EnableInGame is not present (default true)."));
-
-					//GUseRayTracing = false;
+					GUseRayTracing = GRayTracingEnableInGame != 0;
 
-					//UE_LOG(LogRendererCore, Log, TEXT("Ray tracing is disabled for the game. Reason: r.RayTracing=1, and r.RayTracing.EnableInGame is not present (default false)."));
+					UE_LOG(LogRendererCore, Log, TEXT("Ray tracing is %s for the game. Reason: r.RayTracing=1, and r.RayTracing.EnableInGame game user setting does not exist (using default from CVar: %d)."),
+						GUseRayTracing ? TEXT("enabled") : TEXT("disabled"), 
+						GRayTracingEnableInGame);
 				}
 			}
 
@@ -1469,7 +1488,7 @@ RENDERCORE_API void RenderUtilsInit()
 	}
 	else
 	{
-		UE_LOG(LogRendererCore, Log, TEXT("Ray tracing is disabled. Reason: r.RayTracing=0."));
+		UE_LOG(LogRendererCore, Log, TEXT("Ray tracing is disabled. Reason: disabled through project setting (r.RayTracing=0)."));
 	}
 }
 
