Fixed GetMipIndex typo in FRenderTargetBindings.

#9134

diff --git a/Engine/Source/Runtime/RenderCore/Public/ShaderParameterMacros.h b/Engine/Source/Runtime/RenderCore/Public/ShaderParameterMacros.h
--- a/Engine/Source/Runtime/RenderCore/Public/ShaderParameterMacros.h
+++ b/Engine/Source/Runtime/RenderCore/Public/ShaderParameterMacros.h
@@ -542,7 +542,7 @@ struct FRenderTargetBinding
 		check(Validate());
 	}
 
-	void GetMipIndex(uint8 InMipIndex)
+	void SetMipIndex(uint8 InMipIndex)
 	{
 		MipIndex = InMipIndex;
 		check(Validate());
@@ -554,6 +554,9 @@ struct FRenderTargetBinding
 		check(Validate());
 	}
 
+	UE_DEPRECATED(5.1, "GetMipIndex has been renamed to SetMipIndex.")
+	void GetMipIndex(uint8 InMipIndex) { SetMipIndex(InMipIndex); }
+
 private:
 	/** Intentionally kept private to force setting the load action in the constructor. */
 	TAlignedShaderParameterPtr<FRDGTexture*> Texture		= nullptr;
