Fix for ensure when using RDG clobber resources.

diff --git a/Engine/Source/Runtime/Core/Public/PixelFormat.h b/Engine/Source/Runtime/Core/Public/PixelFormat.h
--- a/Engine/Source/Runtime/Core/Public/PixelFormat.h
+++ b/Engine/Source/Runtime/Core/Public/PixelFormat.h
@@ -214,6 +214,31 @@ FORCEINLINE bool IsHDR(EPixelFormat PixelFormat)
 	return PixelFormat == PF_FloatRGBA || PixelFormat == PF_BC6H || PixelFormat == PF_R16F || PixelFormat == PF_R32_FLOAT || PixelFormat == PF_A32B32G32R32F;
 }
 
+FORCEINLINE bool IsInteger(EPixelFormat PixelFormat)
+{
+	switch (PixelFormat)
+	{
+	case PF_R32_UINT:
+	case PF_R32_SINT:
+	case PF_R16_UINT:
+	case PF_R16_SINT:
+	case PF_R16G16B16A16_UINT:
+	case PF_R16G16B16A16_SINT:
+	case PF_R32G32B32A32_UINT:
+	case PF_R16G16_UINT:
+	case PF_R8_UINT:
+	case PF_R8G8B8A8_UINT:
+	case PF_R32G32_UINT:
+	case PF_R8G8_UINT:
+	case PF_R32G32B32_UINT:
+	case PF_R32G32B32_SINT:
+	case PF_R8_SINT:
+	case PF_R64_UINT:
+		return true;
+	}
+	return false;
+}
+
 /** 
 * Information about a pixel format. The majority of this structure is valid after static init, however RHI does keep
 * some state in here that is initialized by that module and should not be used by general use programs that don't
diff --git a/Engine/Source/Runtime/RenderCore/Private/RenderGraphBuilder.cpp b/Engine/Source/Runtime/RenderCore/Private/RenderGraphBuilder.cpp
--- a/Engine/Source/Runtime/RenderCore/Private/RenderGraphBuilder.cpp
+++ b/Engine/Source/Runtime/RenderCore/Private/RenderGraphBuilder.cpp
@@ -3893,6 +3893,18 @@ void FRDGBuilder::ClobberPassOutputs(const FRDGPass* Pass)
 
 	const FLinearColor ClobberColor = GetClobberColor();
 
+	const auto ClobberTextureUAV = [&](FRDGTextureUAV* TextureUAV)
+	{
+		if (IsInteger(TextureUAV->GetParent()->Desc.Format))
+		{
+			AddClearUAVPass(*this, TextureUAV, GetClobberBufferValue());
+		}
+		else
+		{
+			AddClearUAVPass(*this, TextureUAV, ClobberColor);
+		}
+	};
+
 	Pass->GetParameters().Enumerate([&](FRDGParameter Parameter)
 	{
 		switch (Parameter.GetType())
@@ -3922,7 +3934,7 @@ void FRDGBuilder::ClobberPassOutputs(const FRDGPass* Pass)
 					{
 						for (int32 MipLevel = 0; MipLevel < Texture->Desc.NumMips; MipLevel++)
 						{
-							AddClearUAVPass(*this, CreateUAV(FRDGTextureUAVDesc(Texture, MipLevel)), ClobberColor);
+							ClobberTextureUAV(CreateUAV(FRDGTextureUAVDesc(Texture, MipLevel)));
 						}
 					}
 					else if (EnumHasAnyFlags(TextureAccess.GetAccess(), ERHIAccess::RTV))
@@ -3943,13 +3955,13 @@ void FRDGBuilder::ClobberPassOutputs(const FRDGPass* Pass)
 				{
 					if (Texture->Desc.NumMips == 1)
 					{
-						AddClearUAVPass(*this, UAV, ClobberColor);
+						ClobberTextureUAV(UAV);
 					}
 					else
 					{
 						for (int32 MipLevel = 0; MipLevel < Texture->Desc.NumMips; MipLevel++)
 						{
-							AddClearUAVPass(*this, CreateUAV(FRDGTextureUAVDesc(Texture, MipLevel)), ClobberColor);
+							ClobberTextureUAV(CreateUAV(FRDGTextureUAVDesc(Texture, MipLevel)));
 						}
 					}
 				}
