- Add IsRayTracingEnabled(ShaderPlatform) overload that returns true when Project has RT on and ShaderPlatform supports RT and has it on for the project
- Don't create FRayTracingMeshProcessor when IsRayTracingEnabled is false
- Only Add RayTracingGeometries when IsRayTracingEnabled returns true
They are added to the FRayTracingMeshProcessor and will try to compile SF_RayHitGroup shaders

diff --git a/Engine/Source/Runtime/Engine/Private/SkeletalMesh.cpp b/Engine/Source/Runtime/Engine/Private/SkeletalMesh.cpp
--- a/Engine/Source/Runtime/Engine/Private/SkeletalMesh.cpp
+++ b/Engine/Source/Runtime/Engine/Private/SkeletalMesh.cpp
@@ -5835,7 +5835,7 @@ FSkeletalMeshSceneProxy::FSkeletalMeshSceneProxy(const USkinnedMeshComponent* Co
 	}
 
 #if RHI_RAYTRACING
-	if (IsRayTracingEnabled())
+	if (IsRayTracingEnabled(GetScene().GetShaderPlatform()))
 	{
 		if (bRenderStatic)
 		{
diff --git a/Engine/Source/Runtime/Engine/Private/StaticMeshRender.cpp b/Engine/Source/Runtime/Engine/Private/StaticMeshRender.cpp
--- a/Engine/Source/Runtime/Engine/Private/StaticMeshRender.cpp
+++ b/Engine/Source/Runtime/Engine/Private/StaticMeshRender.cpp
@@ -309,7 +309,7 @@ FStaticMeshSceneProxy::FStaticMeshSceneProxy(UStaticMeshComponent* InComponent,
 	bSupportRayTracing = InComponent->GetStaticMesh()->bSupportRayTracing;
 	bDynamicRayTracingGeometry = false;
 	
-	if (IsRayTracingEnabled() && bSupportRayTracing)
+	if (IsRayTracingEnabled(GetScene().GetShaderPlatform()) && bSupportRayTracing)
 	{
 		if (CVarRayTracingStaticMeshesWPO.GetValueOnAnyThread() > 0)
 		{
diff --git a/Engine/Source/Runtime/RenderCore/Private/RenderResource.cpp b/Engine/Source/Runtime/RenderCore/Private/RenderResource.cpp
--- a/Engine/Source/Runtime/RenderCore/Private/RenderResource.cpp
+++ b/Engine/Source/Runtime/RenderCore/Private/RenderResource.cpp
@@ -1178,3 +1178,7 @@ bool IsRayTracingEnabled()
 	return GUseRayTracing;
 }
 
+bool IsRayTracingEnabled(EShaderPlatform ShaderPlatform)
+{
+	return IsRayTracingEnabled() && ShouldCompileRayTracingCallableShadersForProject(ShaderPlatform);
+}
diff --git a/Engine/Source/Runtime/RenderCore/Public/RenderResource.h b/Engine/Source/Runtime/RenderCore/Public/RenderResource.h
--- a/Engine/Source/Runtime/RenderCore/Public/RenderResource.h
+++ b/Engine/Source/Runtime/RenderCore/Public/RenderResource.h
@@ -684,6 +684,9 @@ FORCEINLINE bool ShouldCompileRayTracingCallableShadersForProject(EShaderPlatfor
 // Returns `true` when running on RT-capable machine, RT support is enabled for the project and by game graphics options.
 // This function may only be called at runtime, never during cooking.
 extern RENDERCORE_API bool IsRayTracingEnabled();
+// Returns 'true' when running on RT-capable machine, RT support is enabled for the project and by game graphics options, ShaderPlatform supports RT and ShaderPlatform has RT enabled for the project
+// This function may only be called at runtime, never during cooking.
+RENDERCORE_API bool IsRayTracingEnabled(EShaderPlatform ShaderPlatform);
 
 enum class ERTAccelerationStructureBuildPriority
 {
diff --git a/Engine/Source/Runtime/Renderer/Private/PrimitiveSceneInfo.cpp b/Engine/Source/Runtime/Renderer/Private/PrimitiveSceneInfo.cpp
--- a/Engine/Source/Runtime/Renderer/Private/PrimitiveSceneInfo.cpp
+++ b/Engine/Source/Runtime/Renderer/Private/PrimitiveSceneInfo.cpp
@@ -874,7 +874,7 @@ void CacheRayTracingPrimitive(
 
 void FPrimitiveSceneInfo::CacheRayTracingPrimitives(FScene* Scene, const TArrayView<FPrimitiveSceneInfo*>& SceneInfos)
 {
-	if (IsRayTracingEnabled())
+	if (IsRayTracingEnabled(Scene->GetShaderPlatform()))
 	{
 		CSV_SCOPED_TIMING_STAT_EXCLUSIVE(FPrimitiveSceneInfo_CacheRayTracingPrimitives)
 		SCOPED_NAMED_EVENT(FPrimitiveSceneInfo_CacheRayTracingPrimitives, FColor::Emerald);
