Use format string in RDG events with shipping builds.

diff --git a/Engine/Source/Runtime/RenderCore/Public/RenderGraphDefinitions.h b/Engine/Source/Runtime/RenderCore/Public/RenderGraphDefinitions.h
--- a/Engine/Source/Runtime/RenderCore/Public/RenderGraphDefinitions.h
+++ b/Engine/Source/Runtime/RenderCore/Public/RenderGraphDefinitions.h
@@ -46,7 +46,7 @@
 
 /** Whether render graph GPU events are enabled. */
 #if WITH_PROFILEGPU
-	#if UE_BUILD_TEST
+	#if UE_BUILD_TEST || UE_BUILD_SHIPPING
 		#define RDG_EVENTS RDG_EVENTS_STRING_REF
 	#else
 		#define RDG_EVENTS RDG_EVENTS_STRING_COPY
