Fix RT for other platforms besides PC

diff --git a/Engine/Source/Runtime/RenderCore/Private/RenderResource.cpp b/Engine/Source/Runtime/RenderCore/Private/RenderResource.cpp
--- a/Engine/Source/Runtime/RenderCore/Private/RenderResource.cpp
+++ b/Engine/Source/Runtime/RenderCore/Private/RenderResource.cpp
@@ -1180,5 +1180,5 @@ bool IsRayTracingEnabled()
 
 bool IsRayTracingEnabled(EShaderPlatform ShaderPlatform)
 {
-	return IsRayTracingEnabled() && ShouldCompileRayTracingCallableShadersForProject(ShaderPlatform);
+	return IsRayTracingEnabled() && RHISupportsRayTracing(ShaderPlatform);
 }
diff --git a/Engine/Source/Runtime/RenderCore/Public/RenderResource.h b/Engine/Source/Runtime/RenderCore/Public/RenderResource.h
--- a/Engine/Source/Runtime/RenderCore/Public/RenderResource.h
+++ b/Engine/Source/Runtime/RenderCore/Public/RenderResource.h
@@ -684,7 +684,7 @@ FORCEINLINE bool ShouldCompileRayTracingCallableShadersForProject(EShaderPlatfor
 // Returns `true` when running on RT-capable machine, RT support is enabled for the project and by game graphics options.
 // This function may only be called at runtime, never during cooking.
 extern RENDERCORE_API bool IsRayTracingEnabled();
-// Returns 'true' when running on RT-capable machine, RT support is enabled for the project and by game graphics options, ShaderPlatform supports RT and ShaderPlatform has RT enabled for the project
+// Returns 'true' when running on RT-capable machine, RT support is enabled for the project and by game graphics options and ShaderPlatform supports RT
 // This function may only be called at runtime, never during cooking.
 RENDERCORE_API bool IsRayTracingEnabled(EShaderPlatform ShaderPlatform);
 
