MPCook: Replicate shaderlibrary from CookWorker to Director.
#rnx

diff --git a/Engine/Source/Editor/UnrealEd/Private/CookOnTheFlyServer.cpp b/Engine/Source/Editor/UnrealEd/Private/CookOnTheFlyServer.cpp
--- a/Engine/Source/Editor/UnrealEd/Private/CookOnTheFlyServer.cpp
+++ b/Engine/Source/Editor/UnrealEd/Private/CookOnTheFlyServer.cpp
@@ -7807,21 +7807,61 @@ const FString UCookOnTheFlyServer::GetCookedAssetRegistryFilename(const FString&
 	return CookedAssetRegistryFilename;
 }
 
+namespace UE::Cook
+{
+
+/** CookMultiprocess collector for ShaderLibrary data. */
+class FShaderLibraryCollector : public IMPCollector
+{
+public:
+	virtual FGuid GetMessageType() override { return MessageType; }
+	virtual const TCHAR* GetDebugName() override { return TEXT("FShaderLibraryCollector"); }
+
+	virtual void ClientTick(FClientContext& Context) override;
+	virtual void ReceiveMessage(FServerContext& Context, FCbObjectView Message) override;
+
+	static FGuid MessageType;
+};
+FGuid FShaderLibraryCollector::MessageType(TEXT("4DF3B36BBA2F4E04A846E894E24EB2C4"));
+
+void FShaderLibraryCollector::ClientTick(FClientContext& Context)
+{
+	FCbWriter Writer;
+	bool bHasData;
+	Writer.BeginObject();
+	Writer.SetName("S");
+	FShaderLibraryCooker::CopyToCompactBinaryAndClear(Writer, bHasData);
+	if (bHasData)
+	{
+		Writer.EndObject();
+		Context.AddMessage(Writer.Save().AsObject());
+	}
+}
+
+void FShaderLibraryCollector::ReceiveMessage(FServerContext& Context, FCbObjectView Message)
+{
+	bool bSuccessful = FShaderLibraryCooker::AppendFromCompactBinary(Message["S"]);
+	UE_CLOG(!bSuccessful, LogCook, Error,
+		TEXT("Corrupt message received from CookWorker when replicating ShaderLibrary. Shaders will be missing from the cook."));
+}
+
+}
+
 void UCookOnTheFlyServer::BeginCookStartShaderCodeLibrary(FBeginCookContext& BeginContext)
 {
-    const UProjectPackagingSettings* const PackagingSettings = GetDefault<UProjectPackagingSettings>();
+	const UProjectPackagingSettings* const PackagingSettings = GetDefault<UProjectPackagingSettings>();
 	bool const bCacheShaderLibraries = IsUsingShaderCodeLibrary();
 	if (bCacheShaderLibraries)
 	{
 		FShaderLibraryCooker::InitForCooking(PackagingSettings->bSharedMaterialNativeLibraries);
-        
+
 		bool bAllPlatformsNeedStableKeys = false;
 		// support setting without Hungarian prefix for the compatibility, but allow newer one to override
 		GConfig->GetBool(TEXT("DevOptions.Shaders"), TEXT("NeedsShaderStableKeys"), bAllPlatformsNeedStableKeys, GEngineIni);
 		GConfig->GetBool(TEXT("DevOptions.Shaders"), TEXT("bNeedsShaderStableKeys"), bAllPlatformsNeedStableKeys, GEngineIni);
 
-        for (const ITargetPlatform* TargetPlatform : BeginContext.TargetPlatforms)
-        {
+ 		for (const ITargetPlatform* TargetPlatform : BeginContext.TargetPlatforms)
+ 		{
 			// Find out if this platform requires stable shader keys, by reading the platform setting file.
 			// Stable shader keys are needed if we are going to create a PSO cache.
 			bool bNeedShaderStableKeys = bAllPlatformsNeedStableKeys;
@@ -7835,8 +7875,8 @@ void UCookOnTheFlyServer::BeginCookStartShaderCodeLibrary(FBeginCookContext& Beg
 			FConfigCacheIni::LoadLocalIniFile(PlatformGameIniFile, TEXT("Game"), true, *TargetPlatform->IniPlatformName());
 			PlatformGameIniFile.GetBool(TEXT("/Script/UnrealEd.ProjectPackagingSettings"), TEXT("bDeterministicShaderCodeOrder"), bNeedsDeterministicOrder);
 
-            TArray<FName> ShaderFormats;
-            TargetPlatform->GetAllTargetedShaderFormats(ShaderFormats);
+ 			TArray<FName> ShaderFormats;
+ 			TargetPlatform->GetAllTargetedShaderFormats(ShaderFormats);
 			TArray<FShaderLibraryCooker::FShaderFormatDescriptor> ShaderFormatsWithStableKeys;
 			for (FName& Format : ShaderFormats)
 			{
@@ -7847,12 +7887,21 @@ void UCookOnTheFlyServer::BeginCookStartShaderCodeLibrary(FBeginCookContext& Beg
 				ShaderFormatsWithStableKeys.Push(NewDesc);
 			}
 
-            if (ShaderFormats.Num() > 0)
+			if (ShaderFormats.Num() > 0)
 			{
 				FShaderLibraryCooker::CookShaderFormats(ShaderFormatsWithStableKeys);
 			}
-        }
-    }
+		}
+
+		if (CookDirector)
+		{
+			CookDirector->Register(new UE::Cook::FShaderLibraryCollector());
+		}
+		else if (CookWorkerClient)
+		{
+			CookWorkerClient->Register(new UE::Cook::FShaderLibraryCollector());
+		}
+	}
 
 	if (!IsCookWorkerMode())
 	{
@@ -7876,10 +7925,8 @@ void UCookOnTheFlyServer::BeginCookFinishShaderCodeLibrary(FBeginCookContext& Be
 
 	// Open the shader code library for the current project or the current DLC pack, depending on which we are cooking
 	FString LibraryName = GetProjectShaderLibraryName();
-	if (LibraryName.Len() > 0)
-	{
-		OpenShaderLibrary(LibraryName);
-	}
+	check(!LibraryName.IsEmpty());
+	OpenShaderLibrary(LibraryName);
 }
 
 
@@ -9270,7 +9317,6 @@ void UCookOnTheFlyServer::StartCookAsCookWorker()
 	CookWorkerClient->DoneWithInitialSettings();
 
 	// Initialize systems referenced by later stages or that need to start early for async performance
-	// MPCOOKTODO: Need to send information from ShaderCodeLibrary to director
 	if (IsDirectorCookByTheBook())
 	{
 		BeginCookStartShaderCodeLibrary(BeginContext); // start shader code library cooking asynchronously; we block on it later
diff --git a/Engine/Source/Runtime/RenderCore/Private/ShaderCodeArchive.cpp b/Engine/Source/Runtime/RenderCore/Private/ShaderCodeArchive.cpp
--- a/Engine/Source/Runtime/RenderCore/Private/ShaderCodeArchive.cpp
+++ b/Engine/Source/Runtime/RenderCore/Private/ShaderCodeArchive.cpp
@@ -13,11 +13,15 @@
 #include "RHIShaderFormatDefinitions.inl"
 #include "Serialization/JsonSerializer.h"
 #include "Serialization/MemoryReader.h"
+#include "Serialization/MemoryWriter.h"
 #include "ShaderCodeLibrary.h"
 #include "Stats/Stats.h"
 
 #if WITH_EDITOR
+#include "Misc/Optional.h"
 #include "Misc/StringBuilder.h"
+#include "Serialization/CompactBinarySerialization.h"
+#include "Serialization/CompactBinaryWriter.h"
 #include "Templates/Greater.h"
 #endif
 
@@ -27,7 +31,6 @@
 #include "Modules/ModuleManager.h"
 #endif // UE_SCA_VISUALIZE_SHADER_USAGE
 
-
 int32 GShaderCodeLibraryAsyncLoadingPriority = int32(AIOP_Normal);
 static FAutoConsoleVariableRef CVarShaderCodeLibraryAsyncLoadingPriority(
 	TEXT("r.ShaderCodeLibrary.DefaultAsyncIOPriority"),
@@ -148,6 +151,83 @@ bool FSerializedShaderArchive::FindOrAddShader(const FSHAHash& Hash, int32& OutI
 	return false;
 }
 
+#if WITH_EDITOR
+FCbWriter& operator<<(FCbWriter& Writer, const FSerializedShaderArchive& Archive)
+{
+	TArray64<uint8> SerializedBytes;
+	{
+		FMemoryWriter64 SerializeArchive(SerializedBytes);
+		const_cast<FSerializedShaderArchive&>(Archive).Serialize(SerializeArchive);
+	}
+
+	Writer.BeginObject();
+	{
+		Writer << "SerializedBytes";
+		Writer.AddBinary(FMemoryView(SerializedBytes.GetData(), SerializedBytes.Num()));
+		SerializedBytes.Empty();
+
+		// Serialize is meant for runtime fields only, so copy the editor-only fields separately
+		Writer.BeginArray("ShaderCodeToAssets");
+		for (const TPair<FSHAHash, FShaderMapAssetPaths>& Pair : Archive.ShaderCodeToAssets)
+		{
+			Writer << Pair.Key;
+			Writer.BeginArray();
+			for (FName AssetName : Pair.Value)
+			{
+				Writer << AssetName;
+			}
+			Writer.EndArray();
+		}
+		Writer.EndArray();
+	}
+	Writer.EndObject();
+	return Writer;
+}
+
+bool LoadFromCompactBinary(FCbFieldView Field, FSerializedShaderArchive& OutArchive)
+{
+	FMemoryView SerializedBytes = Field["SerializedBytes"].AsBinaryView();
+	{
+		FMemoryReaderView SerializeArchive(SerializedBytes);
+		OutArchive.Serialize(SerializeArchive);
+		if (SerializeArchive.IsError())
+		{
+			OutArchive = FSerializedShaderArchive();
+			return false;
+		}
+	}
+
+	FCbFieldView ShaderCodeToAssetsField = Field["ShaderCodeToAssets"];
+	// Map size is array size divided by two because pairs are written as successive elements
+	int32 NumShaderCodeToAssets = ShaderCodeToAssetsField.AsArrayView().Num()/2;
+	bool bOk = !ShaderCodeToAssetsField.HasError();
+	OutArchive.ShaderCodeToAssets.Empty(NumShaderCodeToAssets);
+	FCbFieldViewIterator It = ShaderCodeToAssetsField.CreateViewIterator();
+	while (It)
+	{
+		FSHAHash ShaderMapHash;
+		if (!LoadFromCompactBinary(*It++, ShaderMapHash))
+		{
+			bOk = false;
+			continue;
+		}
+		FShaderMapAssetPaths& Paths = OutArchive.ShaderCodeToAssets.FindOrAdd(ShaderMapHash);
+		FCbFieldView AssetNameArrayField = *It++;
+		Paths.Reserve((*It).AsArrayView().Num());
+		bOk = !AssetNameArrayField.HasError() & bOk;
+		for (FCbFieldView AssetNameField : AssetNameArrayField)
+		{
+			FName AssetName;
+			bOk = LoadFromCompactBinary(AssetNameField, AssetName) & bOk;
+			Paths.Add(AssetName);
+		}
+	}
+
+	return bOk;
+}
+#endif
+
+
 #if UE_SCA_VISUALIZE_SHADER_USAGE
 void FShaderUsageVisualizer::Initialize(const int32 InNumShaders)
 {
diff --git a/Engine/Source/Runtime/RenderCore/Private/ShaderCodeLibrary.cpp b/Engine/Source/Runtime/RenderCore/Private/ShaderCodeLibrary.cpp
--- a/Engine/Source/Runtime/RenderCore/Private/ShaderCodeLibrary.cpp
+++ b/Engine/Source/Runtime/RenderCore/Private/ShaderCodeLibrary.cpp
@@ -10,6 +10,7 @@ ShaderCodeLibrary.cpp: Bound shader state cache implementation.
 #include "Async/AsyncFileHandle.h"
 #include "Async/ParallelFor.h"
 #include "Containers/HashTable.h"
+#include "Containers/Set.h"
 #include "Containers/StringView.h"
 #include "FileCache/FileCache.h"
 #include "HAL/FileManagerGeneric.h"
@@ -48,6 +49,8 @@ ShaderCodeLibrary.cpp: Bound shader state cache implementation.
 #include "PipelineCacheUtilities.h"
 #include "UObject/Class.h"
 #include "UObject/Object.h"
+#include "Serialization/CompactBinarySerialization.h"
+#include "Serialization/CompactBinaryWriter.h"
 #endif
 
 // allow introspection (e.g. dumping the contents) for easier debugging
@@ -720,6 +723,83 @@ void FStableShaderKeyAndValue::SetPipelineHash(const FShaderPipeline* Pipeline)
 	}
 }
 
+#if WITH_EDITOR
+void WriteToCompactBinary(FCbWriter& Writer, const FStableShaderKeyAndValue& Key,
+	const TMap<FSHAHash, int32>& HashToIndex)
+{
+	Writer.BeginArray();
+	{
+		TStringBuilder<128> ClassNameAndObjectPathStr;
+		Key.ClassNameAndObjectPath.AppendString(ClassNameAndObjectPathStr);
+		Writer << ClassNameAndObjectPathStr;
+	}
+	Writer << Key.ShaderType;
+	Writer << Key.ShaderClass;
+	Writer << Key.MaterialDomain;
+	Writer << Key.FeatureLevel;
+	Writer << Key.QualityLevel;
+	Writer << Key.TargetFrequency;
+	Writer << Key.TargetPlatform;
+	Writer << Key.VFType;
+	Writer << Key.PermutationId;
+	Writer << HashToIndex[Key.OutputHash];
+	Writer << HashToIndex[Key.PipelineHash];
+	Writer.EndArray();
+}
+
+bool LoadFromCompactBinary(FCbFieldView Field, FStableShaderKeyAndValue& Key,
+	const TArray<FSHAHash>& IndexToHash)
+{
+	int32 NumFields = Field.AsArrayView().Num();
+	if (NumFields != 12)
+	{
+		Key = FStableShaderKeyAndValue();
+		return false;
+	}
+	FCbFieldViewIterator It = Field.CreateViewIterator();
+	bool bOk = true;
+
+	{
+		TStringBuilder<128> ClassNameAndObjectPathStr;
+		bOk = LoadFromCompactBinary(*It++, ClassNameAndObjectPathStr) & bOk;
+		Key.ClassNameAndObjectPath.ParseFromString(ClassNameAndObjectPathStr.ToView());
+	}
+
+	auto LoadIndexedHashFromCompactBinary = [&IndexToHash](FCbFieldView HashField, FSHAHash& OutHash)
+	{
+		int32 Index;
+		if (!LoadFromCompactBinary(HashField, Index))
+		{
+			OutHash = FSHAHash();
+			return false;
+		}
+		if (!IndexToHash.IsValidIndex(Index))
+		{
+			OutHash = FSHAHash();
+			return false;
+		}
+		OutHash = IndexToHash.GetData()[Index];
+		return true;
+	};
+
+	bOk = LoadFromCompactBinary(*It++, Key.ShaderType) & bOk;
+	bOk = LoadFromCompactBinary(*It++, Key.ShaderClass) & bOk;
+	bOk = LoadFromCompactBinary(*It++, Key.MaterialDomain) & bOk;
+	bOk = LoadFromCompactBinary(*It++, Key.FeatureLevel) & bOk;
+	bOk = LoadFromCompactBinary(*It++, Key.QualityLevel) & bOk;
+	bOk = LoadFromCompactBinary(*It++, Key.TargetFrequency) & bOk;
+	bOk = LoadFromCompactBinary(*It++, Key.TargetPlatform) & bOk;
+	bOk = LoadFromCompactBinary(*It++, Key.VFType) & bOk;
+	bOk = LoadFromCompactBinary(*It++, Key.PermutationId) & bOk;
+	bOk = LoadIndexedHashFromCompactBinary(*It++, Key.OutputHash) & bOk;
+	bOk = LoadIndexedHashFromCompactBinary(*It++, Key.PipelineHash) & bOk;
+
+	Key.ComputeKeyHash();
+	return bOk;
+}
+#endif
+
+
 void FShaderCodeLibraryPipeline::Initialize(const FShaderPipeline* Pipeline)
 {
 	check(Pipeline != nullptr);
@@ -1193,6 +1273,7 @@ struct FEditorShaderCodeArchive
 					SerializedShaderEntry.Frequency = SourceShaderEntry.Frequency;
 					SerializedShaderEntry.Size = SourceShaderEntry.Code.Num();
 					SerializedShaderEntry.UncompressedSize = SourceShaderEntry.UncompressedSize;
+					check(!SourceShaderEntry.Code.IsEmpty());
 					ShaderCode.Add(SourceShaderEntry.Code);
 					check(ShaderCode.Num() == SerializedShaders.ShaderEntries.Num());
 
@@ -1207,13 +1288,225 @@ struct FEditorShaderCodeArchive
 			// the cook because of serialization path being reused for other purposes than actual saving, so counting them every time artificially inflates number of shaders.
 			CodeStats.NumShaders += Code->ShaderEntries.Num();
 			CodeStats.NumShaderMaps++;
+			MarkShaderMapDirty(ShaderMapIndex);
 		}
 		return ShaderMapIndex;
 	}
 
+	void MarkShaderMapDirty(int32 ShaderMapIndex)
+	{
+		if (bHasCopiedAndCleared)
+		{
+			ShaderMapsToCopy.Add(ShaderMapIndex);
+		}
+	}
+
+	void MarkAllShaderMapsDirty()
+	{
+		if (bHasCopiedAndCleared)
+		{
+			for (int32 ShaderMapIndex = 0; ShaderMapIndex < SerializedShaders.ShaderMapHashes.Num(); ++ShaderMapIndex)
+			{
+				ShaderMapsToCopy.Add(ShaderMapIndex);
+			}
+		}
+	}
+
+	bool HasDataToCopy() const
+	{
+		// After the first copy dirty elements are added to ShaderMapsToCopy; during the first copy every ShaderMap in SerializedShaders is copied
+		return bHasCopiedAndCleared ? !ShaderMapsToCopy.IsEmpty() : !SerializedShaders.ShaderMapHashes.IsEmpty();
+	}
+
+	void CopyToCompactBinaryAndClear(FCbWriter& Writer)
+	{
+		FSerializedShaderArchive TransferArchive;
+		TArray<uint8> TransferCode;
+		CopyToArchiveAndClear(TransferArchive, TransferCode);
+
+		Writer.BeginObject();
+		Writer << "SerializedShaders" << TransferArchive;
+		Writer << "ShaderCode";
+		Writer.AddBinary(FMemoryView(TransferCode.GetData(), TransferCode.Num()));
+		Writer.EndObject();
+	}
+
+	bool AppendFromCompactBinary(FCbFieldView Field, FShaderCodeStats& CodeStats)
+	{
+		FSerializedShaderArchive TransferArchive;
+		if (!LoadFromCompactBinary(Field["SerializedShaders"], TransferArchive))
+		{
+			return false;
+		}
+
+		FMemoryView FlatShaderCodeView = Field["ShaderCode"].AsBinaryView();
+		TConstArrayView64<uint8> TransferCode(reinterpret_cast<const uint8*>(FlatShaderCodeView.GetData()), FlatShaderCodeView.GetSize());
+		return AppendFromArchive(MoveTemp(TransferArchive), TransferCode, CodeStats);
+	}
+
+	void CopyToArchiveAndClear(FSerializedShaderArchive& TargetArchive, TArray<uint8>& TargetFlatShaderCode)
+	{
+		if (!bHasCopiedAndCleared)
+		{
+			// First CopyAndClear adds all shadermaps to ShaderMapsToCopy, because as an optimization we do not write to
+			// ShaderMapsToCopy until the first CopyAndClear call.
+			for (int32 ShaderMapIndex = 0; ShaderMapIndex < SerializedShaders.ShaderMapHashes.Num(); ++ShaderMapIndex)
+			{
+				ShaderMapsToCopy.Add(ShaderMapIndex);
+			}
+		}
+
+		TArray<int32> LocalShaderMapsToCopy = ShaderMapsToCopy.Array();
+		LocalShaderMapsToCopy.Sort(); // Maintain the same order that the shadermaps were added in
+
+		const FSerializedShaderArchive& SourceArchive = this->SerializedShaders;
+		TArray<TArray<uint8>>& SourceShaderCodes = this->ShaderCode;
+		for (int32 SourceShaderMapIndex : LocalShaderMapsToCopy)
+		{
+			const FSHAHash& SourceShaderMapHash = SourceArchive.ShaderMapHashes[SourceShaderMapIndex];
+			const FShaderMapEntry& SourceEntry = SourceArchive.ShaderMapEntries[SourceShaderMapIndex];
+			const FShaderMapAssetPaths* const SourceAssetPaths = SourceArchive.ShaderCodeToAssets.Find(SourceShaderMapHash);
+
+			int32 TargetShaderMapIndex;
+			const bool bIsNewShaderMap = TargetArchive.FindOrAddShaderMap(SourceShaderMapHash, TargetShaderMapIndex, SourceAssetPaths);
+			if (!bIsNewShaderMap)
+			{
+				continue;
+			}
+			FShaderMapEntry& TargetEntry = TargetArchive.ShaderMapEntries[TargetShaderMapIndex];
+			const int32 NumShaders = SourceEntry.NumShaders;
+
+			TargetEntry.NumShaders = NumShaders;
+			TargetEntry.ShaderIndicesOffset = TargetArchive.ShaderIndices.Num();
+			TargetArchive.ShaderIndices.AddUninitialized(NumShaders);
+
+			for (int32 ShaderIndexIndex = 0; ShaderIndexIndex < NumShaders; ++ShaderIndexIndex)
+			{
+				const uint32 SourceShaderIndex = SourceArchive.ShaderIndices[SourceEntry.ShaderIndicesOffset + ShaderIndexIndex];
+				const FSHAHash& SourceShaderHash = SourceArchive.ShaderHashes[SourceShaderIndex];
+				int32 TargetShaderIndex;
+				const bool bIsNewShader = TargetArchive.FindOrAddShader(SourceShaderHash, TargetShaderIndex);
+				TargetArchive.ShaderIndices[TargetEntry.ShaderIndicesOffset + ShaderIndexIndex] = TargetShaderIndex;
+				if (bIsNewShader)
+				{
+					const FShaderCodeEntry& SourceShaderEntry = SourceArchive.ShaderEntries[SourceShaderIndex];
+					TArray<uint8>& SourceShaderCode = SourceShaderCodes[SourceShaderIndex];
+					FShaderCodeEntry& TargetShaderEntry = TargetArchive.ShaderEntries[TargetShaderIndex];
+
+					TargetShaderEntry = SourceShaderEntry;
+					if (!SourceShaderCode.IsEmpty())
+					{
+						TargetShaderEntry.Offset = TargetFlatShaderCode.Num();
+						check(SourceShaderEntry.Size == SourceShaderCode.Num());
+						TargetFlatShaderCode.Append(SourceShaderCode);
+
+						// Empty the ShaderCode to save memory in the local process. The consumer of the TargetArchive and
+						// TargetFlatShaderCode will be the only one that needs to read it.
+						SourceShaderCode.Empty();
+					}
+					else
+					{
+						// The shadercode was already copied in an earlier CopyAndClear operation; we just need to note that
+						// ShaderMaps in this call to CopyAndClear reference it.
+						check(bHasCopiedAndCleared);
+						TargetShaderEntry.Offset = INDEX_NONE;
+					}
+				}
+			}
+		}
+
+		ShaderMapsToCopy.Empty();
+		bHasCopiedAndCleared = true;
+	}
+
+	bool AppendFromArchive(FSerializedShaderArchive&& SourceArchive, TConstArrayView64<uint8> SourceFlatShaderCode,
+		FShaderCodeStats& CodeStats)
+	{
+		bool bOk = true;
+		FSerializedShaderArchive& TargetArchive = this->SerializedShaders;
+		TArray<TArray<uint8>>& TargetShaderCodes = this->ShaderCode;
+
+		for (int32 SourceShaderMapIndex = 0; SourceShaderMapIndex < SourceArchive.ShaderMapHashes.Num(); ++SourceShaderMapIndex)
+		{
+			const FSHAHash& SourceShaderMapHash = SourceArchive.ShaderMapHashes[SourceShaderMapIndex];
+			const FShaderMapAssetPaths* const SourceAssetPaths = SourceArchive.ShaderCodeToAssets.Find(SourceShaderMapHash);
+			int32 TargetShaderMapIndex;
+			const bool bIsNewShaderMap = TargetArchive.FindOrAddShaderMap(SourceShaderMapHash, TargetShaderMapIndex, SourceAssetPaths);
+			if (!bIsNewShaderMap)
+			{
+				// ShaderMap has already been loaded in this process, e.g. from another CookWorker loading the same Material
+				continue;
+			}
+
+			const FShaderMapEntry& SourceEntry = SourceArchive.ShaderMapEntries[SourceShaderMapIndex];
+			FShaderMapEntry& TargetEntry = TargetArchive.ShaderMapEntries[TargetShaderMapIndex];
+			const int32 NumShaders = SourceEntry.NumShaders;
+
+			TargetEntry.NumShaders = NumShaders;
+			TargetEntry.ShaderIndicesOffset = TargetArchive.ShaderIndices.Num();
+			TargetArchive.ShaderIndices.AddUninitialized(NumShaders);
+
+			for (int32 ShaderIndexIndex = 0; ShaderIndexIndex < NumShaders; ++ShaderIndexIndex)
+			{
+				const uint32 SourceShaderIndex = SourceArchive.ShaderIndices[SourceEntry.ShaderIndicesOffset + ShaderIndexIndex];
+				const FSHAHash& SourceShaderHash = SourceArchive.ShaderHashes[SourceShaderIndex];
+				int32 TargetShaderIndex;
+				const bool bShaderIsNew = TargetArchive.FindOrAddShader(SourceShaderHash, TargetShaderIndex);
+				TargetArchive.ShaderIndices[TargetEntry.ShaderIndicesOffset + ShaderIndexIndex] = TargetShaderIndex;
+				if (!bShaderIsNew)
+				{
+					// Shader has already been loaded for another ShaderMap
+					continue;
+				}
+				check(TargetShaderIndex == TargetArchive.ShaderEntries.Num() - 1 &&
+					TargetShaderCodes.Num() == TargetArchive.ShaderEntries.Num() - 1);
+				TArray<uint8>& TargetShaderCode = TargetShaderCodes.Emplace_GetRef();
+
+				const FShaderCodeEntry& SourceShaderEntry = SourceArchive.ShaderEntries[SourceShaderIndex];
+				FShaderCodeEntry& TargetShaderEntry = TargetArchive.ShaderEntries[TargetShaderIndex];
+				TargetShaderEntry = SourceShaderEntry;
+				TargetShaderEntry.Offset = 0;
+
+				CodeStats.NumUniqueShaders++;
+				CodeStats.ShadersUniqueSize += SourceShaderEntry.Size;
+
+				if (SourceShaderEntry.Offset == INDEX_NONE)
+				{
+					UE_LOG(LogShaderLibrary, Error, TEXT("ShaderMapLibrary transfer received from a remote machine has incomplete record for ShaderMap %s, Shader %s. ")
+						TEXT("The remote machine thought this machine already had the shader and did not send the ShaderCode for the Shader, but the shader is not found. ")
+						TEXT("The ShaderMap will be corrupt."),
+						*SourceShaderMapHash.ToString(), *SourceShaderHash.ToString());
+					bOk = false;
+					TargetShaderEntry.Size = 0;
+				}
+				else if (SourceShaderEntry.Offset + SourceShaderEntry.Size > static_cast<uint64>(SourceFlatShaderCode.Num()))
+				{
+					UE_LOG(LogShaderLibrary, Error, TEXT("ShaderMapLibrary transfer received from a remote machine has corrupt record for ShaderMap %s, Shader %s. ")
+						TEXT("The (Offset, Size) specified by the ShaderEntry does not fit in the TransferCode (an array of bytes that should contain the shader code for all transferred shaders. ")
+						TEXT("The ShaderMap will be corrupt."),
+						*SourceShaderMapHash.ToString(), *SourceShaderHash.ToString());
+					bOk = false;
+					TargetShaderEntry.Size = 0;
+				}
+				else
+				{
+					const TConstArrayView<uint8> SourceShaderCode(SourceFlatShaderCode.GetData() + SourceShaderEntry.Offset, SourceShaderEntry.Size);
+					TargetShaderCode = SourceShaderCode; // Copy from source's flat list to the target's separate TArray<uint8> for each shader
+				}
+			}
+
+			CodeStats.NumShaders += TargetEntry.NumShaders; // Sum of shader counts used by each ShaderMap, without removing duplicates
+			CodeStats.NumShaderMaps++;
+			MarkShaderMapDirty(TargetShaderMapIndex);
+		}
+		return bOk;
+	}
+
 	/** Produces another archive that contains the code only for these assets */
 	FEditorShaderCodeArchive* CreateChunk(int ChunkId, const TSet<FName>& PackagesInChunk)
 	{
+		checkf(!bHasCopiedAndCleared, TEXT("It is not valid to call CreateChunk on an FEditorShaderCodeArchive that has sent its shaders to another process by calling CopyAndClear."));
+
 		FEditorShaderCodeArchive* NewChunk = new FEditorShaderCodeArchive(FormatName, bNeedsDeterministicOrder);
 		NewChunk->OpenLibrary(UE::ShaderLibrary::Private::GetShaderLibraryNameForChunk(LibraryName, ChunkId));
 
@@ -1227,7 +1520,9 @@ struct FEditorShaderCodeArchive
 		NewChunk->ShaderCode.Empty();
 		for (int32 NewArchiveIdx = 0, NumIndices = ShaderCodeEntriesNeeded.Num(); NewArchiveIdx < NumIndices; ++NewArchiveIdx)
 		{
-			NewChunk->ShaderCode.Add(ShaderCode[ShaderCodeEntriesNeeded[NewArchiveIdx]]);
+			TArray<uint8>& SourceShaderCodeEntry = ShaderCode[ShaderCodeEntriesNeeded[NewArchiveIdx]];
+			check(!SourceShaderCodeEntry.IsEmpty());
+			NewChunk->ShaderCode.Add(SourceShaderCodeEntry);
 		}
 
 		return NewChunk;
@@ -1235,6 +1530,9 @@ struct FEditorShaderCodeArchive
 
 	int32 AddShaderCode(int32 OtherShaderMapIndex, const FEditorShaderCodeArchive& OtherArchive)
 	{
+		checkf(!OtherArchive.bHasCopiedAndCleared,
+			TEXT("It is not valid to call AddShaderCode from an FEditorShaderCodeArchive that has sent its shaders to another process by calling CopyAndClear."));
+
 		int32 ShaderMapIndex = 0;
 		if (SerializedShaders.FindOrAddShaderMap(OtherArchive.SerializedShaders.ShaderMapHashes[OtherShaderMapIndex], ShaderMapIndex, 
 				OtherArchive.SerializedShaders.ShaderCodeToAssets.Find(OtherArchive.SerializedShaders.ShaderMapHashes[OtherShaderMapIndex])))
@@ -1253,11 +1551,14 @@ struct FEditorShaderCodeArchive
 					const FShaderCodeEntry& OtherShaderEntry = OtherArchive.SerializedShaders.ShaderEntries[OtherShaderIndex];
 					SerializedShaders.ShaderEntries[ShaderIndex] = OtherShaderEntry;
 
-					ShaderCode.Add(OtherArchive.ShaderCode[OtherShaderIndex]);
+					const TArray<uint8>& OtherShaderCodeEntry = OtherArchive.ShaderCode[OtherShaderIndex];
+					check(!OtherShaderCodeEntry.IsEmpty());
+					ShaderCode.Add(OtherShaderCodeEntry);
 					check(ShaderCode.Num() == SerializedShaders.ShaderEntries.Num());
 				}
 				SerializedShaders.ShaderIndices[ShaderMapEntry.ShaderIndicesOffset + i] = ShaderIndex;
 			}
+			MarkShaderMapDirty(ShaderMapIndex);
 		}
 		return ShaderMapIndex;
 	}
@@ -1292,6 +1593,7 @@ struct FEditorShaderCodeArchive
 
 					// Read shader code from archive and add shader to set
 					const int64 ReadSize = OtherShaderEntry.Size;
+					check(ReadSize > 0);
 					const int64 ReadOffset = OtherShaderCodeOffset + OtherShaderEntry.Offset;
 					Code.SetNumUninitialized(ReadSize);
 					Ar.Seek(ReadOffset);
@@ -1299,6 +1601,7 @@ struct FEditorShaderCodeArchive
 				}
 				SerializedShaders.ShaderIndices[ShaderMapEntry.ShaderIndicesOffset + i] = ShaderIndex;
 			}
+			MarkShaderMapDirty(ShaderMapIndex);
 		}
 		return ShaderMapIndex;
 	}
@@ -1316,12 +1619,14 @@ struct FEditorShaderCodeArchive
 			{
 				// Read shader library
 				*PrevCookedAr << SerializedShaders;
+				MarkAllShaderMapsDirty();
 
 				ShaderCode.AddDefaulted(SerializedShaders.ShaderEntries.Num());
-				for(int32 Index = 0; Index < ShaderCode.Num(); ++Index)
+				for (int32 Index = 0; Index < ShaderCode.Num(); ++Index)
 				{
 					const FShaderCodeEntry& Entry = SerializedShaders.ShaderEntries[Index];
 					TArray<uint8>& Code = ShaderCode[Index];
+					check(Entry.Size > 0);
 					Code.SetNumUninitialized(Entry.Size);
 					PrevCookedAr->Serialize(Code.GetData(), Entry.Size);
 					bOK = !PrevCookedAr->GetError();
@@ -1415,6 +1720,7 @@ struct FEditorShaderCodeArchive
 	bool SaveToDisk(const FString& OutputDir, const FString& MetaOutputDir, bool bSaveOnlyAssetInfo = false, TArray<FString>* OutputFilenames = nullptr)
 	{
 		check(LibraryName.Len() > 0);
+		checkf(!bHasCopiedAndCleared, TEXT("It is not valid to call SaveToDisk on an FEditorShaderCodeArchive that has sent its shaders to another process by calling CopyAndClear."));
 
 		bool bSuccess = IFileManager::Get().MakeDirectory(*OutputDir, true);
 
@@ -1458,6 +1764,7 @@ struct FEditorShaderCodeArchive
 					*FileWriter << SerializedShaders;
 					for (auto& Code : ShaderCode)
 					{
+						check(!Code.IsEmpty());
 						FileWriter->Serialize(Code.GetData(), Code.Num());
 					}
 
@@ -1513,6 +1820,7 @@ struct FEditorShaderCodeArchive
 		{
 			return true;
 		}
+		checkf(!bHasCopiedAndCleared, TEXT("It is not valid to call PackageNativeShaderLibrary on an FEditorShaderCodeArchive that has sent its shaders to another process by calling CopyAndClear."));
 
 		FString IntermediateFormatPath = GetShaderDebugFolder(FPaths::ProjectSavedDir() / TEXT("Shaders") / FormatName.ToString(), LibraryName, FormatName);
 		FString TempPath = IntermediateFormatPath / TEXT("NativeLibrary");
@@ -1706,7 +2014,16 @@ private:
 	FName FormatName;
 	FString LibraryName;
 	FSerializedShaderArchive SerializedShaders;
+	/**
+	 * The element at index N holds the ShaderCode for the element of SerializedShaders.ShaderEntries at index N.
+	 * In MultiprocessCooking elements can be empty if they have been transferred to the Director (bHasCopiedAndCleared will be true in this case).
+	 */
 	TArray<TArray<uint8>> ShaderCode;
+	/** A list of ShaderMaps that have not yet been copied to the CookDirector. Used only by MultiprocessCookWorkers. */
+	TSet<int32> ShaderMapsToCopy;
+	/** True if CopyToArchiveAndClear has been called, otherwise false. If false we avoid doing some tracking since it might never be used. */
+	bool bHasCopiedAndCleared = false;
+
 	const IShaderFormat* Format;
 	bool bNeedsDeterministicOrder;
 };
@@ -1724,6 +2041,8 @@ struct FEditorShaderStableInfo
 		check(Name.Len() > 0);
 		LibraryName = Name;
 		StableMap.Empty();
+		StableMapsToCopy.Empty();
+		bHasCopied = false;
 	}
 
 	void CloseLibrary(FString const& Name)
@@ -1743,6 +2062,7 @@ struct FEditorShaderStableInfo
 			return;
 		}
 		StableMap.Add(StableKeyValue);
+		MarkKeyValueDirty(StableKeyValue);
 	}
 
 	void AddExistingShaderCodeLibrary(FString const& OutputDir)
@@ -1846,10 +2166,96 @@ struct FEditorShaderStableInfo
 		return bSuccess;
 	}
 
+	void MarkKeyValueDirty(const FStableShaderKeyAndValue& StableKeyValue)
+	{
+		if (bHasCopied)
+		{
+			StableMapsToCopy.Add(StableKeyValue);
+		}
+	}
+
+	bool HasDataToCopy() const
+	{
+		// After the first copy dirty elements are added to StableMapsToCopy; during the first copy every element from StableMap is copied
+		return bHasCopied ? !StableMapsToCopy.IsEmpty() : !StableMap.IsEmpty();
+	}
+
+	void CopyToCompactBinary(FCbWriter& Writer)
+	{
+		if (!bHasCopied)
+		{
+			// First Copy adds all StableMaps to StableMapsToCopy, because as an optimization we do not write to
+			// StableMapsToCopy until the first Copy call.
+			StableMapsToCopy = StableMap;
+		}
+
+		// Deduplicate the hashes; there will be a significant number (up to 50%) of duplicate hashes.
+		TArray<FSHAHash> Hashes;
+		TMap<FSHAHash, int32> HashToIndex;
+		auto IndexHash = [&Hashes, &HashToIndex](const FSHAHash& Hash)
+		{
+			if (HashToIndex.Find(Hash) == nullptr)
+			{
+				Hashes.Add(Hash);
+				HashToIndex.Add(Hash, Hashes.Num() - 1);
+			}
+		};
+		for (const FStableShaderKeyAndValue& StableInfo : StableMapsToCopy)
+		{
+			IndexHash(StableInfo.PipelineHash);
+			IndexHash(StableInfo.OutputHash);
+		}
+
+		Writer.BeginObject();
+		{
+			Writer << "Hashes" << Hashes;
+			Writer.BeginArray("StableShaderKeyAndValues");
+			for (const FStableShaderKeyAndValue& StableInfo : StableMapsToCopy)
+			{
+				WriteToCompactBinary(Writer, StableInfo, HashToIndex);
+			}
+			Writer.EndArray();
+		}
+		Writer.EndObject();
+
+		StableMapsToCopy.Empty();
+		bHasCopied = true;
+	}
+
+	bool AppendFromCompactBinary(FCbFieldView Field)
+	{
+		TArray<FSHAHash> Hashes;
+		if (!LoadFromCompactBinary(Field["Hashes"], Hashes))
+		{
+			// Can't load the StableShaderKeyAndValues if the hashes for them are missing
+			return false;
+		}
+		FCbFieldView StableShaderKeyAndValuesField = Field["StableShaderKeyAndValues"];
+		int32 NumInfos = StableShaderKeyAndValuesField.AsArrayView().Num();
+		bool bOk = !StableShaderKeyAndValuesField.HasError();
+		for (FCbFieldView InfoView : Field)
+		{
+			FStableShaderKeyAndValue StableInfo;
+			if (LoadFromCompactBinary(InfoView, StableInfo, Hashes))
+			{
+				AddShader(StableInfo);
+			}
+			else
+			{
+				bOk = false;
+			}
+		}
+		return bOk;
+	}
+
 private:
 	FName FormatName;
 	FString LibraryName;
 	TSet<FStableShaderKeyAndValue> StableMap;
+	/** Copies of elements in StableMap that have not yet been copied to the CookDirector. Used only by MultiprocessCookWorkers. */
+	TSet<FStableShaderKeyAndValue> StableMapsToCopy;
+	/** True if CopyToCompactBinary has been called, otherwise false. If false we avoid doing some tracking since it might never be used. */
+	bool bHasCopied = false;
 };
 #endif //WITH_EDITOR
 
@@ -2391,16 +2797,113 @@ public:
 		FScopeLock ScopeLock(&ShaderCodeCS);
 		checkf(Platform < UE_ARRAY_COUNT(EditorShaderCodeStats), TEXT("FShaderCodeLibrary::AddShaderCode can only be called with a valid shader platform (expected no more than %d, passed: %d)"), 
 			static_cast<int32>(UE_ARRAY_COUNT(EditorShaderCodeStats)), static_cast<int32>(Platform));
-		FShaderCodeStats& CodeStats = EditorShaderCodeStats[Platform];
+		static_assert(UE_ARRAY_COUNT(EditorShaderCodeStats) == UE_ARRAY_COUNT(EditorShaderCodeArchive), "Size of EditorShaderCodeStats must match size of EditorShaderCodeArchive");
 
-		checkf(Platform < UE_ARRAY_COUNT(EditorShaderCodeArchive), TEXT("FShaderCodeLibrary::AddShaderCode can only be called with a valid shader platform (expected no more than %d, passed: %d)"),
-			static_cast<int32>(UE_ARRAY_COUNT(EditorShaderCodeArchive)), static_cast<int32>(Platform));
+		FShaderCodeStats& CodeStats = EditorShaderCodeStats[Platform];
 		FEditorShaderCodeArchive* CodeArchive = EditorShaderCodeArchive[Platform];
 		checkf(CodeArchive, TEXT("EditorShaderCodeArchive for (EShaderPlatform)%d is null!"), (int32)Platform);
 
 		CodeArchive->AddShaderCode(CodeStats, Code, AssociatedAssets);
 	}
 
+#if WITH_EDITOR
+	void CopyToCompactBinaryAndClear(FCbWriter& Writer, bool& bOutHasData)
+	{
+		TArray<EShaderPlatform, TInlineAllocator<10>> PlatformsToCopy;
+		FScopeLock ScopeLock(&ShaderCodeCS);
+		for (EShaderPlatform Platform = (EShaderPlatform)0; Platform < (EShaderPlatform)UE_ARRAY_COUNT(EditorShaderCodeStats);
+			Platform = (EShaderPlatform)((int)Platform + 1))
+		{
+			FEditorShaderCodeArchive* CodeArchive = EditorShaderCodeArchive[Platform];
+			FEditorShaderStableInfo* StableInfo = EditorShaderStableInfo[Platform];
+			if ((CodeArchive && CodeArchive->HasDataToCopy()) || (StableInfo && StableInfo->HasDataToCopy()))
+			{
+				PlatformsToCopy.Add(Platform);
+			}
+		}
+
+		if (PlatformsToCopy.IsEmpty())
+		{
+			bOutHasData = false;
+			return;
+		}
+		bOutHasData = true;
+		Writer.BeginArray();
+		for (EShaderPlatform Platform : PlatformsToCopy)
+		{
+			Writer.BeginObject();
+			{
+				Writer << "Platform" << (uint32)Platform;
+				FEditorShaderCodeArchive* CodeArchive = EditorShaderCodeArchive[Platform];
+				FEditorShaderStableInfo* StableInfo = EditorShaderStableInfo[Platform];
+
+				if ((CodeArchive && CodeArchive->HasDataToCopy()))
+				{
+					Writer.SetName("EditorShaderCodeArchive");
+					CodeArchive->CopyToCompactBinaryAndClear(Writer);
+				}
+				if (StableInfo && StableInfo->HasDataToCopy())
+				{
+					Writer.SetName("EditorShaderStableInfo");
+					StableInfo->CopyToCompactBinary(Writer);
+				}
+
+				// ChunksSaved is not copied; it is constructed at end of cook after copying from remote workers is complete
+				// EditorShaderCodeStats is not copied; it is gathered from the shadercode we send when the director receives it
+				// bShaderFormatsThatNeedStableKeys is not copied; it is constructed during BeginCook and is the same on all machines
+				// EditorArchivePipelines is not copied; it is constructed during BeginCook and is the same on all machines
+
+			}
+			Writer.EndObject();
+		}
+		Writer.EndArray();
+	}
+
+	bool AppendFromCompactBinary(FCbFieldView Field)
+	{
+		bool bOk = true;
+		for (FCbFieldView PlatformField : Field)
+		{
+			uint32 PlatformInt;
+			if (!LoadFromCompactBinary(PlatformField["Platform"], PlatformInt) ||
+				PlatformInt >= UE_ARRAY_COUNT(EditorShaderCodeStats))
+			{
+				bOk = false;
+				continue;
+			}
+			EShaderPlatform Platform = (EShaderPlatform)PlatformInt;
+
+			FCbFieldView CodeArchiveField = PlatformField["EditorShaderCodeArchive"];
+			if (CodeArchiveField.HasValue())
+			{
+				FEditorShaderCodeArchive* CodeArchive = EditorShaderCodeArchive[Platform];
+				if (CodeArchive)
+				{
+					FShaderCodeStats& CodeStats = EditorShaderCodeStats[Platform];
+					bOk = CodeArchive->AppendFromCompactBinary(CodeArchiveField, CodeStats) & bOk;
+				}
+				else
+				{
+					UE_LOG(LogShaderLibrary, Error, TEXT("ShaderMapLibrary transfer received from a remote machine includes data for Platform %d, but the ShaderMapLibrary has not been initialized for that platform in the local process. ")
+						TEXT("The information will be ignored."), (int32)Platform);
+					bOk = false;
+				}
+			}
+
+			FCbFieldView StableArchiveField = PlatformField["EditorShaderStableInfo"];
+			if (StableArchiveField.HasValue())
+			{
+				FEditorShaderStableInfo* StableArchive = EditorShaderStableInfo[Platform];
+				if (StableArchive)
+				{
+					bOk = StableArchive->AppendFromCompactBinary(StableArchiveField) & bOk;
+				}
+			}
+		}
+		return bOk;
+	}
+#endif
+
 	void AddShaderStableKeyValue(EShaderPlatform InShaderPlatform, FStableShaderKeyAndValue& StableKeyValue)
 	{
 		FEditorShaderStableInfo* StableArchive = EditorShaderStableInfo[InShaderPlatform];
@@ -2972,6 +3475,25 @@ bool FShaderLibraryCooker::AddShaderCode(EShaderPlatform ShaderPlatform, const F
 	return false;
 }
 
+#if WITH_EDITOR
+void FShaderLibraryCooker::CopyToCompactBinaryAndClear(FCbWriter& Writer, bool& bOutHasData)
+{
+	if (FShaderLibrariesCollection::Impl)
+	{
+		FShaderLibrariesCollection::Impl->CopyToCompactBinaryAndClear(Writer, bOutHasData);
+	}
+}
+
+bool FShaderLibraryCooker::AppendFromCompactBinary(FCbFieldView Field)
+{
+	if (FShaderLibrariesCollection::Impl)
+	{
+		return FShaderLibrariesCollection::Impl->AppendFromCompactBinary(Field);
+	}
+	return false;
+}
+#endif
+
 bool FShaderLibraryCooker::NeedsShaderStableKeys(EShaderPlatform ShaderPlatform)
 {
 	if (FShaderLibrariesCollection::Impl)
diff --git a/Engine/Source/Runtime/RenderCore/Public/ShaderCodeArchive.h b/Engine/Source/Runtime/RenderCore/Public/ShaderCodeArchive.h
--- a/Engine/Source/Runtime/RenderCore/Public/ShaderCodeArchive.h
+++ b/Engine/Source/Runtime/RenderCore/Public/ShaderCodeArchive.h
@@ -25,6 +25,11 @@
 #include "Templates/RefCounting.h"
 #include "UObject/NameTypes.h"
 
+#if WITH_EDITOR
+class FCbFieldView;
+class FCbWriter;
+#endif
+
 // enable visualization in the desktop Development builds only as it has a memory hit and writes files
 #define UE_SCA_VISUALIZE_SHADER_USAGE			(!WITH_EDITOR && UE_BUILD_DEVELOPMENT && PLATFORM_DESKTOP)
 
@@ -192,6 +197,8 @@ public:
 	void CreateAsChunkFrom(const FSerializedShaderArchive& Parent, const TSet<FName>& PackagesInChunk, TArray<int32>& OutShaderCodeEntriesNeeded);
 	void CollectStatsAndDebugInfo(FDebugStats& OutDebugStats, FExtendedDebugStats* OutExtendedDebugStats);
 	void DumpContentsInPlaintext(FString& OutText) const;
+	RENDERCORE_API friend FCbWriter& operator<<(FCbWriter& Writer, const FSerializedShaderArchive& Archive);
+	RENDERCORE_API friend bool LoadFromCompactBinary(FCbFieldView Field, FSerializedShaderArchive& OutArchive);
 #endif
 
 	friend FArchive& operator<<(FArchive& Ar, FSerializedShaderArchive& Ref)
diff --git a/Engine/Source/Runtime/RenderCore/Public/ShaderCodeLibrary.h b/Engine/Source/Runtime/RenderCore/Public/ShaderCodeLibrary.h
--- a/Engine/Source/Runtime/RenderCore/Public/ShaderCodeLibrary.h
+++ b/Engine/Source/Runtime/RenderCore/Public/ShaderCodeLibrary.h
@@ -37,6 +37,11 @@ class FIoBuffer;
 class FIoChunkId;
 class UObject;
 
+#if WITH_EDITOR
+class FCbFieldView;
+class FCbWriter;
+#endif
+
 struct RENDERCORE_API FShaderCodeLibraryPipeline
 {
 	FSHAHash Shaders[SF_NumGraphicsFrequencies];
@@ -161,8 +166,14 @@ struct RENDERCORE_API FStableShaderKeyAndValue
 	{
 		return Key.KeyHash;
 	}
-
 };
+#if WITH_EDITOR
+RENDERCORE_API void WriteToCompactBinary(FCbWriter& Writer, const FStableShaderKeyAndValue& Key, 
+	const TMap<FSHAHash, int32>& HashToIndex);
+RENDERCORE_API bool LoadFromCompactBinary(FCbFieldView Field, FStableShaderKeyAndValue& Key,
+	const TArray<FSHAHash>& IndexToHash);
+#endif
+
 
 DECLARE_MULTICAST_DELEGATE_TwoParams(FSharedShaderCodeRequest, const FSHAHash&, FArchive*);
 DECLARE_MULTICAST_DELEGATE_OneParam(FSharedShaderCodeRelease, const FSHAHash&);
@@ -282,6 +293,14 @@ struct RENDERCORE_API FShaderLibraryCooker
 	// At cook time, add shader code to collection
 	static bool AddShaderCode(EShaderPlatform ShaderPlatform, const FShaderMapResourceCode* Code, const FShaderMapAssetPaths& AssociatedAssets);
 
+#if WITH_EDITOR
+	/** Called from a CookWorker to send all contents of the ShaderLibrary to the CookDirector */
+	static void CopyToCompactBinaryAndClear(FCbWriter& Writer, bool& bOutHasData);
+
+	/** Called On the CookDirector to receive ShaderLibrary contents from a CookWorker */
+	static bool AppendFromCompactBinary(FCbFieldView Field);
+#endif
+
 	// We check this early in the callstack to avoid creating a bunch of FName and keys and things we will never save anyway. 
 	// Pass the shader platform to check or EShaderPlatform::SP_NumPlatforms to check if any of the registered types require
 	// stable keys.
