Fix build error: ! and & together are ambiguous precedence so wrap in parentheses to clarify.
#rnx

diff --git a/Engine/Source/Runtime/RenderCore/Private/ShaderCodeArchive.cpp b/Engine/Source/Runtime/RenderCore/Private/ShaderCodeArchive.cpp
--- a/Engine/Source/Runtime/RenderCore/Private/ShaderCodeArchive.cpp
+++ b/Engine/Source/Runtime/RenderCore/Private/ShaderCodeArchive.cpp
@@ -214,7 +214,7 @@ bool LoadFromCompactBinary(FCbFieldView Field, FSerializedShaderArchive& OutArch
 		FShaderMapAssetPaths& Paths = OutArchive.ShaderCodeToAssets.FindOrAdd(ShaderMapHash);
 		FCbFieldView AssetNameArrayField = *It++;
 		Paths.Reserve((*It).AsArrayView().Num());
-		bOk = !AssetNameArrayField.HasError() & bOk;
+		bOk = (!AssetNameArrayField.HasError()) & bOk;
 		for (FCbFieldView AssetNameField : AssetNameArrayField)
 		{
 			FName AssetName;
