ShaderCodeLibrary for CookMultiprocess: Fix typo in AppendFromCompactBinary that caused all calls to fail.
#rnx

diff --git a/Engine/Source/Runtime/RenderCore/Private/ShaderCodeLibrary.cpp b/Engine/Source/Runtime/RenderCore/Private/ShaderCodeLibrary.cpp
--- a/Engine/Source/Runtime/RenderCore/Private/ShaderCodeLibrary.cpp
+++ b/Engine/Source/Runtime/RenderCore/Private/ShaderCodeLibrary.cpp
@@ -2233,7 +2233,7 @@ struct FEditorShaderStableInfo
 		FCbFieldView StableShaderKeyAndValuesField = Field["StableShaderKeyAndValues"];
 		int32 NumInfos = StableShaderKeyAndValuesField.AsArrayView().Num();
 		bool bOk = !StableShaderKeyAndValuesField.HasError();
-		for (FCbFieldView InfoView : Field)
+		for (FCbFieldView InfoView : StableShaderKeyAndValuesField)
 		{
 			FStableShaderKeyAndValue StableInfo;
 			if (LoadFromCompactBinary(InfoView, StableInfo, Hashes))
