Rename CVarRayTracingSM5 -> CVarRayTracingRequireSM6

No functional changes, just reflecting the name of the CVar itself (r.RayTracing.RequireSM6).

diff --git a/Engine/Source/Runtime/RenderCore/Private/RenderUtils.cpp b/Engine/Source/Runtime/RenderCore/Private/RenderUtils.cpp
--- a/Engine/Source/Runtime/RenderCore/Private/RenderUtils.cpp
+++ b/Engine/Source/Runtime/RenderCore/Private/RenderUtils.cpp
@@ -46,7 +46,7 @@ FAutoConsoleVariableRef CVarRayTracingEnableInEditor(
 );
 
 static int32 GRayTracingRequireSM6 = 0;
-static FAutoConsoleVariableRef CVarRayTracingSM5(
+static FAutoConsoleVariableRef CVarRayTracingRequireSM6(
 	TEXT("r.RayTracing.RequireSM6"),
 	GRayTracingRequireSM6,
 	TEXT("Whether ray tracing shaders and features should only be available when targetting and running SM6. If disabled, ray tracing shaders will also be available when running in SM5 mode. (default = 0, allow SM5 and SM6)"),
