Fix oldValue of RHIThread by setting the ThreadId on the RHI Thread itself as it was racing between the Launching and RHI Thread and add missing TaskTag for the RenderingThread

diff --git a/Engine/Source/Runtime/OpenGLDrv/Private/OpenGLQuery.cpp b/Engine/Source/Runtime/OpenGLDrv/Private/OpenGLQuery.cpp
--- a/Engine/Source/Runtime/OpenGLDrv/Private/OpenGLQuery.cpp
+++ b/Engine/Source/Runtime/OpenGLDrv/Private/OpenGLQuery.cpp
@@ -523,7 +523,7 @@ public:
 
 	void DoTask(ENamedThreads::Type CurrentThread, const FGraphEventRef& MyCompletionGraphEvent)
 	{
-		FTaskTagScope TagScope(ETaskTag::ERhiThread);
+		check(IsInRHIThread());
 		check(IsRunningRHIInDedicatedThread() && IsInRHIThread()); // this should never be used on a platform that doesn't support the RHI thread, and it can't quite work when running the RHI stuff on task threads
 		if (bWait)
 		{
diff --git a/Engine/Source/Runtime/RHI/Private/RHICommandList.cpp b/Engine/Source/Runtime/RHI/Private/RHICommandList.cpp
--- a/Engine/Source/Runtime/RHI/Private/RHICommandList.cpp
+++ b/Engine/Source/Runtime/RHI/Private/RHICommandList.cpp
@@ -685,8 +685,7 @@ void FRHICommandListImmediate::ExecuteAndReset()
 			    [RHICmdList = MoveTemp(static_cast<FRHICommandListBase&>(*this))]() mutable
 			    {
 				    // RHI thread tasks main entry point
-    
-				    FTaskTagScope Scope(ETaskTag::ERhiThread);
+					check(IsInRHIThread());
 				    SCOPE_CYCLE_COUNTER(STAT_RHIThreadExecute);
 				    if (IsRunningRHIInTaskThread())
 				    {
diff --git a/Engine/Source/Runtime/RenderCore/Private/RenderingThread.cpp b/Engine/Source/Runtime/RenderCore/Private/RenderingThread.cpp
--- a/Engine/Source/Runtime/RenderCore/Private/RenderingThread.cpp
+++ b/Engine/Source/Runtime/RenderCore/Private/RenderingThread.cpp
@@ -308,6 +308,14 @@ public:
 		check(IsInGameThread());
 	}
 
+	virtual bool Init(void) override
+	{
+		PRAGMA_DISABLE_DEPRECATION_WARNINGS
+		GRHIThreadId = FPlatformTLS::GetCurrentThreadId();
+		PRAGMA_ENABLE_DEPRECATION_WARNINGS
+		return true;
+	}
+
 	virtual uint32 Run() override
 	{
 		LLM_SCOPE(ELLMTag::RHIMisc);
@@ -320,9 +328,9 @@ public:
 		{
 			FTaskTagScope Scope(ETaskTag::ERhiThread);
 			FPlatformProcess::SetupRHIThread();
+			FTaskGraphInterface::Get().AttachToThread(ENamedThreads::RHIThread);
+			FTaskGraphInterface::Get().ProcessThreadUntilRequestReturn(ENamedThreads::RHIThread);
 		}
-		FTaskGraphInterface::Get().AttachToThread(ENamedThreads::RHIThread);
-		FTaskGraphInterface::Get().ProcessThreadUntilRequestReturn(ENamedThreads::RHIThread);
 		FMemory::ClearAndDisableTLSCachesOnCurrentThread();
 		return 0;
 	}
@@ -508,11 +516,18 @@ PRAGMA_ENABLE_DEPRECATION_WARNINGS
 		return ReportCrash(ExceptionInfo);
 	}
 #endif
+	
+	void SetupRenderThread()
+	{
+		FTaskTagScope Scope(ETaskTag::ERenderingThread);
+		FPlatformProcess::SetupRenderThread();
+	}
+
 
 	virtual uint32 Run(void) override
 	{
 		FMemory::SetupTLSCachesOnCurrentThread();
-		FPlatformProcess::SetupRenderThread();
+		SetupRenderThread();
 
 #if PLATFORM_WINDOWS
 		bool bNoExceptionHandler = FParse::Param(FCommandLine::Get(), TEXT("noexceptionhandler"));
@@ -717,9 +732,7 @@ public:
 	**/
 	void DoTask(ENamedThreads::Type CurrentThread, const FGraphEventRef& MyCompletionGraphEvent)
 	{
-		// note that this task is the first task run on the thread, before GRHIThread_InternalUseOnly is assigned, so we can't check IsInRHIThread()
-
-		FTaskTagScope Scope(ETaskTag::ERhiThread);
+		check(IsInRHIThread());
 		if (bAcquireOwnership)
 		{
 			GDynamicRHI->RHIAcquireThreadOwnership();
@@ -762,7 +775,6 @@ void StartRenderingThread()
 
 PRAGMA_DISABLE_DEPRECATION_WARNINGS
 		GRHIThread_InternalUseOnly = FRHIThread::Get().Thread;
-		GRHIThreadId = GRHIThread_InternalUseOnly->GetThreadID();
 PRAGMA_ENABLE_DEPRECATION_WARNINGS
 		GIsRunningRHIInDedicatedThread_InternalUseOnly = true;
 		GIsRunningRHIInSeparateThread_InternalUseOnly = true;
