add missing PixelFormatToChannelFlags entry for PF_P010

 - mapped to 'none', assuming it's a similar format to NV12

#rnx

diff --git a/Engine/Source/Runtime/RenderCore/Private/RenderUtils.cpp b/Engine/Source/Runtime/RenderCore/Private/RenderUtils.cpp
--- a/Engine/Source/Runtime/RenderCore/Private/RenderUtils.cpp
+++ b/Engine/Source/Runtime/RenderCore/Private/RenderUtils.cpp
@@ -939,6 +939,7 @@ EPixelFormatChannelFlags GetPixelFormatValidChannels(EPixelFormat InPixelFormat)
 		EPixelFormatChannelFlags::R,		// PF_R8_SINT
 		EPixelFormatChannelFlags::R,		// PF_R64_UINT
 		EPixelFormatChannelFlags::RGB,		// PF_R9G9B9EXP5
+		EPixelFormatChannelFlags::None,		// PF_P010
 	};
 	static_assert(UE_ARRAY_COUNT(PixelFormatToChannelFlags) == (uint8)PF_MAX, "Missing pixel format");
 	return (InPixelFormat < PF_MAX) ? PixelFormatToChannelFlags[(uint8)InPixelFormat] : EPixelFormatChannelFlags::None;
