Fix GRayTracingPlaformMask initialization for preview platforms

diff --git a/Engine/Source/Runtime/RenderCore/Private/RenderUtils.cpp b/Engine/Source/Runtime/RenderCore/Private/RenderUtils.cpp
--- a/Engine/Source/Runtime/RenderCore/Private/RenderUtils.cpp
+++ b/Engine/Source/Runtime/RenderCore/Private/RenderUtils.cpp
@@ -1377,12 +1377,19 @@ RENDERCORE_API void RenderUtilsInit()
 				TArray<FName> PlaformRayTracingShaderFormats;
 				TargetPlatform->GetRayTracingShaderFormats(PlaformRayTracingShaderFormats);
 
-				for (FName Format : PlaformRayTracingShaderFormats)
+				for (FName FormatName : PlaformRayTracingShaderFormats)
 				{
-					EShaderPlatform ShaderPlatform = ShaderFormatNameToShaderPlatform(Format);
-					uint32 ShaderPlatformIndex = static_cast<uint32>(ShaderPlatform);
+					TArray<EShaderPlatform> PossiblePreviewPlatformsAndMainPlatform;
+					GetAllPossiblePreviewPlatformsForMainShaderPlatform(PossiblePreviewPlatformsAndMainPlatform, FormatName);
 
-					GRayTracingPlaformMask[ShaderPlatformIndex] = true;
+					EShaderPlatform MainShaderPlatform = ShaderFormatNameToShaderPlatform(FormatName);
+					PossiblePreviewPlatformsAndMainPlatform.Add(MainShaderPlatform);
+
+					for (EShaderPlatform ShaderPlatform : PossiblePreviewPlatformsAndMainPlatform)
+					{
+						uint32 ShaderPlatformIndex = static_cast<uint32>(ShaderPlatform);
+						GRayTracingPlaformMask[ShaderPlatformIndex] = true;
+					}
 				}
 			}
 		}
