Fix spelling Plaform -> Platform

diff --git a/Engine/Plugins/Experimental/GPULightmass/Source/GPULightmassEditor/Private/GPULightmassEditorModule.cpp b/Engine/Plugins/Experimental/GPULightmass/Source/GPULightmassEditor/Private/GPULightmassEditorModule.cpp
--- a/Engine/Plugins/Experimental/GPULightmass/Source/GPULightmassEditor/Private/GPULightmassEditorModule.cpp
+++ b/Engine/Plugins/Experimental/GPULightmass/Source/GPULightmassEditor/Private/GPULightmassEditorModule.cpp
@@ -85,13 +85,13 @@ ERayTracingDisabledReason GetRayTracingDisabledReason()
 		return ERayTracingDisabledReason::DISABLED_BY_PROJECT_SETTINGS;
 	}
 
-	extern RENDERCORE_API ShaderPlatformMaskType GRayTracingPlaformMask;
+	extern RENDERCORE_API ShaderPlatformMaskType GRayTracingPlatformMask;
 	if (!RHISupportsRayTracing(GMaxRHIShaderPlatform))
 	{
 		return ERayTracingDisabledReason::INCOMPATIBLE_SHADER_PLATFORM;
 	}
 	// Shader platform statically supports ray tracing, but disabled by target platform at runtime
-	else if (RHISupportsRayTracing(GMaxRHIShaderPlatform) && (GRayTracingPlaformMask[(int)GMaxRHIShaderPlatform]) == 0)
+	else if (RHISupportsRayTracing(GMaxRHIShaderPlatform) && (GRayTracingPlatformMask[(int)GMaxRHIShaderPlatform]) == 0)
 	{
 		return ERayTracingDisabledReason::DISABLED_BY_TARGET_PLATFORM;
 	}
diff --git a/Engine/Source/Runtime/RenderCore/Private/RenderUtils.cpp b/Engine/Source/Runtime/RenderCore/Private/RenderUtils.cpp
--- a/Engine/Source/Runtime/RenderCore/Private/RenderUtils.cpp
+++ b/Engine/Source/Runtime/RenderCore/Private/RenderUtils.cpp
@@ -1279,7 +1279,7 @@ RENDERCORE_API ShaderPlatformMaskType GSimpleSkyDiffusePlatformMask;
 // Specifies whether ray tracing *can* be enabled on a particular platform.
 // This takes into account whether RT is globally enabled for the project and specifically enabled on a target platform.
 // Safe to use to make cook-time decisions, such as whether to compile ray tracing shaders.
-RENDERCORE_API ShaderPlatformMaskType GRayTracingPlaformMask;
+RENDERCORE_API ShaderPlatformMaskType GRayTracingPlatformMask;
 
 // Specifies whether ray tracing *is* enabled on the current running system (in current game or editor process).
 // This takes into account additional factors, such as concrete current GPU/OS/Driver capability, user-set game graphics options, etc.
@@ -1323,7 +1323,7 @@ RENDERCORE_API void RenderUtilsInit()
 
 	GSimpleSkyDiffusePlatformMask.Init(false, EShaderPlatform::SP_NumPlatforms);
 	GVelocityEncodeDepthPlatformMask.Init(false, EShaderPlatform::SP_NumPlatforms);
-	GRayTracingPlaformMask.Init(false, EShaderPlatform::SP_NumPlatforms);
+	GRayTracingPlatformMask.Init(false, EShaderPlatform::SP_NumPlatforms);
 
 	static IConsoleVariable* MobileAmbientOcclusionCVar = IConsoleManager::Get().FindConsoleVariable(TEXT("r.Mobile.AmbientOcclusion"));
 	GMobileAmbientOcclusionPlatformMask.Init(MobileAmbientOcclusionCVar && MobileAmbientOcclusionCVar->GetInt(), EShaderPlatform::SP_NumPlatforms);
@@ -1334,10 +1334,10 @@ RENDERCORE_API void RenderUtilsInit()
 	{
 		for (ITargetPlatform* TargetPlatform : TargetPlatformManager->GetTargetPlatforms())
 		{
-			TArray<FName> PlaformPossibleShaderFormats;
-			TargetPlatform->GetAllPossibleShaderFormats(PlaformPossibleShaderFormats);
+			TArray<FName> PlatformPossibleShaderFormats;
+			TargetPlatform->GetAllPossibleShaderFormats(PlatformPossibleShaderFormats);
 
-			for (FName Format : PlaformPossibleShaderFormats)
+			for (FName Format : PlatformPossibleShaderFormats)
 			{
 				EShaderPlatform ShaderPlatform = ShaderFormatNameToShaderPlatform(Format);
 				TArray<EShaderPlatform> PossiblePreviewPlatformsAndMainPlatform;
@@ -1374,10 +1374,10 @@ RENDERCORE_API void RenderUtilsInit()
 
 			if (TargetPlatform->UsesRayTracing())
 			{
-				TArray<FName> PlaformRayTracingShaderFormats;
-				TargetPlatform->GetRayTracingShaderFormats(PlaformRayTracingShaderFormats);
+				TArray<FName> PlatformRayTracingShaderFormats;
+				TargetPlatform->GetRayTracingShaderFormats(PlatformRayTracingShaderFormats);
 
-				for (FName FormatName : PlaformRayTracingShaderFormats)
+				for (FName FormatName : PlatformRayTracingShaderFormats)
 				{
 					TArray<EShaderPlatform> PossiblePreviewPlatformsAndMainPlatform;
 					GetAllPossiblePreviewPlatformsForMainShaderPlatform(PossiblePreviewPlatformsAndMainPlatform, FormatName);
@@ -1388,7 +1388,7 @@ RENDERCORE_API void RenderUtilsInit()
 					for (EShaderPlatform ShaderPlatform : PossiblePreviewPlatformsAndMainPlatform)
 					{
 						uint32 ShaderPlatformIndex = static_cast<uint32>(ShaderPlatform);
-						GRayTracingPlaformMask[ShaderPlatformIndex] = true;
+						GRayTracingPlatformMask[ShaderPlatformIndex] = true;
 					}
 				}
 			}
@@ -1404,7 +1404,7 @@ RENDERCORE_API void RenderUtilsInit()
 
 	if (RayTracingCVar && RayTracingCVar->GetInt() && GRHISupportsRayTracing)
 	{
-		GRayTracingPlaformMask.Init(true, EShaderPlatform::SP_NumPlatforms);
+		GRayTracingPlatformMask.Init(true, EShaderPlatform::SP_NumPlatforms);
 	}
 
 	// Load runtime values from and *.ini file used by a current platform
@@ -1427,7 +1427,7 @@ RENDERCORE_API void RenderUtilsInit()
 			bool bRayTracing = false;
 			if (PlatformIniFile.GetBool(*CategoryName, TEXT("bEnableRayTracing"), bRayTracing) && !bRayTracing)
 			{
-				GRayTracingPlaformMask.Init(false, EShaderPlatform::SP_NumPlatforms);
+				GRayTracingPlatformMask.Init(false, EShaderPlatform::SP_NumPlatforms);
 			}
 		}
 	}
@@ -1447,7 +1447,7 @@ RENDERCORE_API void RenderUtilsInit()
 
 	if (RayTracingCVar && RayTracingCVar->GetBool())
 	{
-		const bool bRayTracingAllowedOnCurrentPlatform = (GRayTracingPlaformMask[(int)GMaxRHIShaderPlatform]);
+		const bool bRayTracingAllowedOnCurrentPlatform = (GRayTracingPlatformMask[(int)GMaxRHIShaderPlatform]);
 		if (GRHISupportsRayTracing && bRayTracingAllowedOnCurrentPlatform)
 		{
 			if (GIsEditor)
diff --git a/Engine/Source/Runtime/RenderCore/Public/RenderResource.h b/Engine/Source/Runtime/RenderCore/Public/RenderResource.h
--- a/Engine/Source/Runtime/RenderCore/Public/RenderResource.h
+++ b/Engine/Source/Runtime/RenderCore/Public/RenderResource.h
@@ -655,8 +655,8 @@ FORCEINLINE bool IsRayTracingEnabledForProject(EShaderPlatform ShaderPlatform)
 {
 	if (RHISupportsRayTracing(ShaderPlatform))
 	{
-		extern RENDERCORE_API TBitArray<TInlineAllocator<EShaderPlatform::SP_NumPlatforms / 8>> GRayTracingPlaformMask;
-		return (GRayTracingPlaformMask[(int)ShaderPlatform]);
+		extern RENDERCORE_API TBitArray<TInlineAllocator<EShaderPlatform::SP_NumPlatforms / 8>> GRayTracingPlatformMask;
+		return (GRayTracingPlatformMask[(int)ShaderPlatform]);
 	}
 	else
 	{
