Non-Engine module public header truncation fixes to reduce noise when compiling with truncation warnings

diff --git a/Engine/Source/Runtime/CoreUObject/Public/Serialization/ObjectReader.h b/Engine/Source/Runtime/CoreUObject/Public/Serialization/ObjectReader.h
--- a/Engine/Source/Runtime/CoreUObject/Public/Serialization/ObjectReader.h
+++ b/Engine/Source/Runtime/CoreUObject/Public/Serialization/ObjectReader.h
@@ -60,7 +60,7 @@ public:
 			// Only serialize if we have the requested amount of data
 			if (Offset + Num <= TotalSize())
 			{
-				FMemory::Memcpy(Data, &Bytes[Offset], Num);
+				FMemory::Memcpy(Data, &Bytes[IntCastChecked<int32>(Offset)], Num);
 				Offset += Num;
 			}
 			else
diff --git a/Engine/Source/Runtime/HeadMountedDisplay/Public/HeadMountedDisplayTypes.h b/Engine/Source/Runtime/HeadMountedDisplay/Public/HeadMountedDisplayTypes.h
--- a/Engine/Source/Runtime/HeadMountedDisplay/Public/HeadMountedDisplayTypes.h
+++ b/Engine/Source/Runtime/HeadMountedDisplay/Public/HeadMountedDisplayTypes.h
@@ -212,12 +212,12 @@ struct FSpectatorScreenModeTexturePlusEyeLayout
 
 	FIntRect GetScaledEyeRect(int SizeX, int SizeY) const
 	{
-		return FIntRect(EyeRectMin.X * SizeX, EyeRectMin.Y * SizeY, EyeRectMax.X * SizeX, EyeRectMax.Y * SizeY);
+		return FIntRect(FIntRect::IntType(EyeRectMin.X * SizeX), FIntRect::IntType(EyeRectMin.Y * SizeY), FIntRect::IntType(EyeRectMax.X * SizeX), FIntRect::IntType(EyeRectMax.Y * SizeY));
 	}
 
 	FIntRect GetScaledTextureRect(int SizeX, int SizeY) const
 	{
-		return FIntRect(TextureRectMin.X * SizeX, TextureRectMin.Y * SizeY, TextureRectMax.X * SizeX, TextureRectMax.Y * SizeY);
+		return FIntRect(FIntRect::IntType(TextureRectMin.X * SizeX), FIntRect::IntType(TextureRectMin.Y * SizeY), FIntRect::IntType(TextureRectMax.X * SizeX), FIntRect::IntType(TextureRectMax.Y * SizeY));
 	}
 
 	FVector2D EyeRectMin;
diff --git a/Engine/Source/Runtime/Interchange/Core/Public/Types/AttributeStorage.h b/Engine/Source/Runtime/Interchange/Core/Public/Types/AttributeStorage.h
--- a/Engine/Source/Runtime/Interchange/Core/Public/Types/AttributeStorage.h
+++ b/Engine/Source/Runtime/Interchange/Core/Public/Types/AttributeStorage.h
@@ -1080,7 +1080,7 @@ namespace UE
 				}
 
 				//Allocate the ByteArray
-				OutValue.AddZeroed(AttributeAllocationInfo->Size);
+				OutValue.AddZeroed(static_cast<typename ArrayType::SizeType>(AttributeAllocationInfo->Size));
 				//Copy into the buffer
 				const uint8* StorageData = AttributeStorage.GetData();
 				FMemory::Memcpy(OutValue.GetData(), &StorageData[AttributeAllocationInfo->Offset], AttributeAllocationInfo->Size);
diff --git a/Engine/Source/Runtime/Json/Public/Policies/JsonPrintPolicy.h b/Engine/Source/Runtime/Json/Public/Policies/JsonPrintPolicy.h
--- a/Engine/Source/Runtime/Json/Public/Policies/JsonPrintPolicy.h
+++ b/Engine/Source/Runtime/Json/Public/Policies/JsonPrintPolicy.h
@@ -35,7 +35,7 @@ struct TJsonPrintPolicy
 
 		for (int32 CharIndex = 0; CharIndex < String.Len(); ++CharIndex, ++CharPtr)
 		{
-			WriteChar(Stream, *CharPtr);
+			WriteChar(Stream, static_cast<CharType>(*CharPtr));
 		}
 	}
 
diff --git a/Engine/Source/Runtime/MeshUtilitiesCommon/Public/MeshUtilitiesCommon.h b/Engine/Source/Runtime/MeshUtilitiesCommon/Public/MeshUtilitiesCommon.h
--- a/Engine/Source/Runtime/MeshUtilitiesCommon/Public/MeshUtilitiesCommon.h
+++ b/Engine/Source/Runtime/MeshUtilitiesCommon/Public/MeshUtilitiesCommon.h
@@ -33,7 +33,7 @@ struct FBoneVertInfo
 /** Helper struct for building acceleration structures. */
 struct FIndexAndZ
 {
-	float Z;
+	FVector::FReal Z;
 	int32 Index;
 
 	/** Default constructor. */
@@ -45,6 +45,11 @@ struct FIndexAndZ
 		Z = 0.30f * V.X + 0.33f * V.Y + 0.37f * V.Z;
 		Index = InIndex;
 	}
+	FIndexAndZ(int32 InIndex, FVector3f V)
+	{
+		Z = 0.30f * V.X + 0.33f * V.Y + 0.37f * V.Z;
+		Index = InIndex;
+	}
 };
 
 /** Sorting function for vertex Z/index pairs. */
@@ -83,7 +88,7 @@ namespace TriangleUtilities
 	/*
 	 * This function compute the area of a triangle, it will return zero if the triangle is degenerated
 	 */
-	static float ComputeTriangleArea(const FVector& PointA, const FVector& PointB, const FVector& PointC)
+	static double ComputeTriangleArea(const FVector& PointA, const FVector& PointB, const FVector& PointC)
 	{
 		return FVector::CrossProduct((PointB - PointA), (PointC - PointA)).Size() / 2.0f;
 	}
@@ -91,7 +96,7 @@ namespace TriangleUtilities
 	/*
 	 * This function compute the angle of a triangle corner, it will return zero if the triangle is degenerated
 	 */
-	static float ComputeTriangleCornerAngle(const FVector& PointA, const FVector& PointB, const FVector& PointC)
+	static double ComputeTriangleCornerAngle(const FVector& PointA, const FVector& PointB, const FVector& PointC)
 	{
 		FVector E1 = (PointB - PointA);
 		FVector E2 = (PointC - PointA);
@@ -101,7 +106,6 @@ namespace TriangleUtilities
 			//Return a null ratio if the polygon is degenerate
 			return 0.0f;
 		}
-		float DotProduct = FVector::DotProduct(E1, E2);
-		return FMath::Acos(DotProduct);
+		return FMath::Acos(FVector::DotProduct(E1, E2));
 	}
 }
\ No newline at end of file
diff --git a/Engine/Source/Runtime/Net/Core/Public/Net/Core/Misc/ResizableCircularQueue.h b/Engine/Source/Runtime/Net/Core/Public/Net/Core/Misc/ResizableCircularQueue.h
--- a/Engine/Source/Runtime/Net/Core/Public/Net/Core/Misc/ResizableCircularQueue.h
+++ b/Engine/Source/Runtime/Net/Core/Public/Net/Core/Misc/ResizableCircularQueue.h
@@ -267,7 +267,9 @@ void TResizableCircularQueue<T, AllocatorT>::PopNoCheck(SIZE_T Count)
 template<typename T, typename AllocatorT>
 void TResizableCircularQueue<T, AllocatorT>::SetCapacity(SIZE_T RequiredCapacity)
 {
-	SIZE_T NewCapacity = FMath::RoundUpToPowerOfTwo(RequiredCapacity);
+	using SizeType = typename StorageT::SizeType; 
+
+ 	SIZE_T NewCapacity = FMath::RoundUpToPowerOfTwo64(RequiredCapacity);
 
 	if ((NewCapacity == Storage.Num()) || (NewCapacity < Count()))
 	{
@@ -277,7 +279,7 @@ void TResizableCircularQueue<T, AllocatorT>::SetCapacity(SIZE_T RequiredCapacity
 	if (Storage.Num() > 0)
 	{
 		StorageT NewStorage;
-		NewStorage.Empty(NewCapacity);
+		NewStorage.Empty(static_cast<SizeType>(NewCapacity));
 		
 		// copy data to new storage
 		const IndexT MaskedTail = Tail & IndexMask;
@@ -292,25 +294,25 @@ void TResizableCircularQueue<T, AllocatorT>::SetCapacity(SIZE_T RequiredCapacity
 		if ((SrcSize > 0) & (MaskedTail >= MaskedHead))
 		{
 			const SIZE_T CopyCount = (SrcCapacity - MaskedTail);
-			NewStorage.Append(SrcData + MaskedTail, CopyCount);
+			NewStorage.Append(SrcData + MaskedTail, static_cast<SizeType>(CopyCount));
 			NewStorage.Append(SrcData, MaskedHead);
 		}
 		else
 		{
-			NewStorage.Append(SrcData + MaskedTail, SrcSize);
+			NewStorage.Append(SrcData + MaskedTail, static_cast<SizeType>(SrcSize));
 		}
 
-		NewStorage.AddUninitialized(NewCapacity - SrcSize);
+		NewStorage.AddUninitialized(static_cast<SizeType>(NewCapacity - SrcSize));
 
 		this->Storage = MoveTemp(NewStorage);
-		IndexMask = NewCapacity - 1;
+		IndexMask = static_cast<IndexT>(NewCapacity - 1);
 		Tail = 0u;
-		Head = SrcSize;
+		Head = static_cast<IndexT>(SrcSize);
 	}
 	else
 	{
-		IndexMask = NewCapacity - 1;
-		Storage.AddUninitialized(NewCapacity);
+		IndexMask = static_cast<IndexT>(NewCapacity - 1);
+		Storage.AddUninitialized(static_cast<SizeType>(NewCapacity));
 	}
 }
 
diff --git a/Engine/Source/Runtime/Net/Core/Public/Net/Core/PropertyConditions/RepChangedPropertyTracker.h b/Engine/Source/Runtime/Net/Core/Public/Net/Core/PropertyConditions/RepChangedPropertyTracker.h
--- a/Engine/Source/Runtime/Net/Core/Public/Net/Core/PropertyConditions/RepChangedPropertyTracker.h
+++ b/Engine/Source/Runtime/Net/Core/Public/Net/Core/PropertyConditions/RepChangedPropertyTracker.h
@@ -63,7 +63,7 @@ public:
 	UE_DEPRECATED(5.1, "No longer used, ActiveState must be constructed with the correct number of properties.")
 	void InitActiveParents(int32 ParentCount) {}
 
-	bool IsParentActive(int32 ParentIndex) const
+	bool IsParentActive(uint16 ParentIndex) const
 	{
 		return ActiveState.GetActiveState(ParentIndex);
 	}
diff --git a/Engine/Source/Runtime/Net/Core/Public/Net/Core/Trace/Private/NetTraceInternal.h b/Engine/Source/Runtime/Net/Core/Public/Net/Core/Trace/Private/NetTraceInternal.h
--- a/Engine/Source/Runtime/Net/Core/Public/Net/Core/Trace/Private/NetTraceInternal.h
+++ b/Engine/Source/Runtime/Net/Core/Public/Net/Core/Trace/Private/NetTraceInternal.h
@@ -494,7 +494,7 @@ FNetTraceBunchScope<T>::~FNetTraceBunchScope()
 	// Only report if we have a valid collector
 	if (Collector)
 	{
-		FNetTrace::EndBunch(*Collector, FNetTrace::TraceName(Bunch.ChName), StartPos, HeaderBits, Bunch.GetNumBits(), MakeBunchInfo(Bunch));
+		FNetTrace::EndBunch(*Collector, FNetTrace::TraceName(Bunch.ChName), StartPos, HeaderBits, (uint32)Bunch.GetNumBits(), MakeBunchInfo(Bunch));	// LWC_TODO: Precision loss. FNetTrace::EndBunch et al. should take int64 BunchBits?
 	}
 }
 
diff --git a/Engine/Source/Runtime/RenderCore/Public/RenderGraphDefinitions.h b/Engine/Source/Runtime/RenderCore/Public/RenderGraphDefinitions.h
--- a/Engine/Source/Runtime/RenderCore/Public/RenderGraphDefinitions.h
+++ b/Engine/Source/Runtime/RenderCore/Public/RenderGraphDefinitions.h
@@ -606,8 +606,8 @@ struct FRDGTextureDesc : public FRHITextureDesc
 		, uint32              ExtData = 0
 	)
 	{
-		const uint32 Depth = 1;
-		const uint32 ArraySize = 1;
+		const uint16 Depth = 1;
+		const uint16 ArraySize = 1;
 		return FRDGTextureDesc(ETextureDimension::Texture2D, Flags, Format, ClearValue, { Size.X, Size.Y }, Depth, ArraySize, NumMips, NumSamples, ExtData);
 	}
 
@@ -622,7 +622,7 @@ struct FRDGTextureDesc : public FRHITextureDesc
 		, uint32              ExtData = 0
 	)
 	{
-		const uint32 Depth = 1;
+		const uint16 Depth = 1;
 		return FRDGTextureDesc(ETextureDimension::Texture2DArray, Flags, Format, ClearValue, { Size.X, Size.Y }, Depth, ArraySize, NumMips, NumSamples, ExtData);
 	}
 
@@ -635,12 +635,12 @@ struct FRDGTextureDesc : public FRHITextureDesc
 		, uint32              ExtData = 0
 	)
 	{
-		const uint32 ArraySize = 1;
-		const uint32 LocalNumSamples = 1;
+		const uint16 ArraySize = 1;
+		const uint8 LocalNumSamples = 1;
 
 		checkf(Size.Z <= TNumericLimits<decltype(FRDGTextureDesc::Depth)>::Max(), TEXT("Depth parameter (Size.Z) exceeds valid range"));
 
-		return FRDGTextureDesc(ETextureDimension::Texture3D, Flags, Format, ClearValue, { Size.X, Size.Y }, Size.Z, ArraySize, NumMips, LocalNumSamples, ExtData);
+		return FRDGTextureDesc(ETextureDimension::Texture3D, Flags, Format, ClearValue, { Size.X, Size.Y }, (uint16)Size.Z, ArraySize, NumMips, LocalNumSamples, ExtData);
 	}
 
 	static FRDGTextureDesc CreateCube(
@@ -655,8 +655,8 @@ struct FRDGTextureDesc : public FRHITextureDesc
 	{
 		checkf(Size <= (uint32)TNumericLimits<int32>::Max(), TEXT("Size parameter exceeds valid range"));
 
-		const uint32 Depth = 1;
-		const uint32 ArraySize = 1;
+		const uint16 Depth = 1;
+		const uint16 ArraySize = 1;
 		return FRDGTextureDesc(ETextureDimension::TextureCube, Flags, Format, ClearValue, { (int32)Size, (int32)Size }, Depth, ArraySize, NumMips, NumSamples, ExtData);
 	}
 
@@ -673,7 +673,7 @@ struct FRDGTextureDesc : public FRHITextureDesc
 	{
 		checkf(Size <= (uint32)TNumericLimits<int32>::Max(), TEXT("Size parameter exceeds valid range"));
 
-		const uint32 Depth = 1;
+		const uint16 Depth = 1;
 		return FRDGTextureDesc(ETextureDimension::TextureCubeArray, Flags, Format, ClearValue, { (int32)Size, (int32)Size }, Depth, ArraySize, NumMips, NumSamples, ExtData);
 	}
 
diff --git a/Engine/Source/Runtime/RenderCore/Public/RenderGraphResources.h b/Engine/Source/Runtime/RenderCore/Public/RenderGraphResources.h
--- a/Engine/Source/Runtime/RenderCore/Public/RenderGraphResources.h
+++ b/Engine/Source/Runtime/RenderCore/Public/RenderGraphResources.h
@@ -715,7 +715,7 @@ public:
 		check(Texture->Desc.IsTextureArray());
 		check(SliceIndex >= 0 && SliceIndex < Texture->Desc.ArraySize);
 		FRDGTextureSRVDesc Desc = FRDGTextureSRVDesc::Create(Texture);
-		Desc.FirstArraySlice = SliceIndex;
+		Desc.FirstArraySlice = (uint16)SliceIndex;
 		Desc.NumArraySlices = 1;
 		return Desc;
 	}
diff --git a/Engine/Source/Runtime/Renderer/Public/MeshPassProcessor.h b/Engine/Source/Runtime/Renderer/Public/MeshPassProcessor.h
--- a/Engine/Source/Runtime/Renderer/Public/MeshPassProcessor.h
+++ b/Engine/Source/Runtime/Renderer/Public/MeshPassProcessor.h
@@ -918,7 +918,7 @@ private:
 		}
 	}
 
-	void AllocateZeroed(uint32 InSize) 
+	void AllocateZeroed(uint16 InSize) 
 	{
 		Allocate(InSize);
 
@@ -1078,7 +1078,7 @@ public:
 #if PLATFORM_64BITS
 				// Ignoring the lower 4 bits since they are likely zero anyway.
 				// Higher bits are more significant in 64 bit builds.
-				return reinterpret_cast<UPTRINT>(Key) >> 4;
+				return static_cast<uint32>(reinterpret_cast<UPTRINT>(Key) >> 4);
 #else
 				return reinterpret_cast<UPTRINT>(Key);
 #endif
@@ -1132,7 +1132,7 @@ public:
 
 	inline void SetStencilRef(uint32 InStencilRef)
 	{
-		StencilRef = InStencilRef;
+		StencilRef = static_cast<uint8>(InStencilRef);
 		// Verify no overflow
 		checkSlow((uint32)StencilRef == InStencilRef);
 	}
diff --git a/Engine/Source/Runtime/SignalProcessing/Public/SampleBuffer.h b/Engine/Source/Runtime/SignalProcessing/Public/SampleBuffer.h
--- a/Engine/Source/Runtime/SignalProcessing/Public/SampleBuffer.h
+++ b/Engine/Source/Runtime/SignalProcessing/Public/SampleBuffer.h
@@ -118,7 +118,7 @@ namespace Audio
 				// Convert from int to float:
 				for (int32 SampleIndex = 0; SampleIndex < NumSamples; SampleIndex++)
 				{
-					RawPCMData[SampleIndex] = ((float)InBufferPtr[SampleIndex]) / 32767.0f;
+					RawPCMData[SampleIndex] = (InBufferPtr[SampleIndex] / 32767.0f);
 				}
 			}
 			else
@@ -365,7 +365,7 @@ namespace Audio
 
 				for (int32 SampleIndex = 0; SampleIndex < RawPCMData.Num(); SampleIndex++)
 				{
-					RawPCMData[SampleIndex] = FMath::Clamp<float>(RawPCMData[SampleIndex], ClampMin, Ceiling);
+					RawPCMData[SampleIndex] = static_cast<SampleType>(FMath::Clamp<float>(RawPCMData[SampleIndex], ClampMin, Ceiling));
 				}
 			}
 			else if (TIsSame<SampleType, int16>::Value)
@@ -373,8 +373,8 @@ namespace Audio
 				// int16 case:
 				Ceiling = FMath::Clamp(Ceiling, 0.0f, 1.0f);
 
-				int16 ClampMax = Ceiling * 32767.0f;
-				int16 ClampMin = Ceiling * -32767.0f;
+				int16 ClampMax = static_cast<int16>(Ceiling * 32767.0f);
+				int16 ClampMin = static_cast<int16>(Ceiling * -32767.0f);
 
 				for (int32 SampleIndex = 0; SampleIndex < NumSamples; SampleIndex++)
 				{
@@ -388,7 +388,7 @@ namespace Audio
 
 				for (int32 SampleIndex = 0; SampleIndex < RawPCMData.Num(); SampleIndex++)
 				{
-					RawPCMData[SampleIndex] = FMath::Clamp<SampleType>(RawPCMData[SampleIndex], ClampMin, Ceiling);
+					RawPCMData[SampleIndex] = static_cast<SampleType>(FMath::Clamp<SampleType>(RawPCMData[SampleIndex], ClampMin, Ceiling));
 				}
 			}
 		}
