Changed uniform buffer binding to use a reference to reduce atomics.

diff --git a/Engine/Source/Runtime/RenderCore/Public/ShaderParameterMacros.h b/Engine/Source/Runtime/RenderCore/Public/ShaderParameterMacros.h
--- a/Engine/Source/Runtime/RenderCore/Public/ShaderParameterMacros.h
+++ b/Engine/Source/Runtime/RenderCore/Public/ShaderParameterMacros.h
@@ -148,7 +148,7 @@ class alignas(SHADER_PARAMETER_STRUCT_ALIGNMENT) FUniformBufferBinding
 public:
 	FUniformBufferBinding() = default;
 
-	FORCEINLINE FUniformBufferRHIRef GetUniformBuffer() const
+	FORCEINLINE const FUniformBufferRHIRef& GetUniformBuffer() const
 	{
 		return UniformBuffer;
 	}
