[shaders] fix determinism issue in serializing shader maps to DDC - previously warnings were appended to a single array in order of job completion which is inherently unstable.
- modified serialization to instead add the warnings along with debug info to an "editor only data" struct, and store one of these per deduplicated bytecode entry (note the previous separate deduplication step for debug data was determined to be redundant)
- append additional context information to each warning via a DebugName argument passed to AddShaderCompilerOutput and store all warnings associated with a single bytecode (including those encountered for different permutations with potentially different source code resulting in the same bytecode) in sorted order
- add a simple ToString helper to FShaderCompileJobKey for easy construction of the DebugName argument for the above
- modify all engine usage of AddShaderCompilerOutput to construct this debug name from the job key

diff --git a/Engine/Plugins/Compositing/OpenColorIO/Source/OpenColorIO/Private/OpenColorIOShaderMap.cpp b/Engine/Plugins/Compositing/OpenColorIO/Source/OpenColorIO/Private/OpenColorIOShaderMap.cpp
--- a/Engine/Plugins/Compositing/OpenColorIO/Source/OpenColorIO/Private/OpenColorIOShaderMap.cpp
+++ b/Engine/Plugins/Compositing/OpenColorIO/Source/OpenColorIO/Private/OpenColorIOShaderMap.cpp
@@ -470,7 +470,7 @@ FShader* FOpenColorIOShaderMap::ProcessCompilationResultsForSingleJob(const TRef
 	FShaderCompileJob* CurrentJob = SingleJob->GetSingleShaderJob();
 	check(CurrentJob->Id == CompilingId);
 
-	GetResourceCode()->AddShaderCompilerOutput(CurrentJob->Output);
+	GetResourceCode()->AddShaderCompilerOutput(CurrentJob->Output, CurrentJob->Key.ToString());
 
 	FShader* Shader = nullptr;
 
diff --git a/Engine/Plugins/FX/Niagara/Source/NiagaraShader/Private/NiagaraShader.cpp b/Engine/Plugins/FX/Niagara/Source/NiagaraShader/Private/NiagaraShader.cpp
--- a/Engine/Plugins/FX/Niagara/Source/NiagaraShader/Private/NiagaraShader.cpp
+++ b/Engine/Plugins/FX/Niagara/Source/NiagaraShader/Private/NiagaraShader.cpp
@@ -780,7 +780,7 @@ FShader* FNiagaraShaderMap::ProcessCompilationResultsForSingleJob(const TRefCoun
 	auto CurrentJob = SingleJob->GetSingleShaderJob();
 	check(CurrentJob->Id == CompilingId);
 
-	GetResourceCode()->AddShaderCompilerOutput(CurrentJob->Output);
+	GetResourceCode()->AddShaderCompilerOutput(CurrentJob->Output, CurrentJob->Key.ToString());
 
 	FShader* Shader = nullptr;
 
diff --git a/Engine/Plugins/Runtime/ComputeFramework/Source/ComputeFramework/Private/ComputeKernelShaderMap.cpp b/Engine/Plugins/Runtime/ComputeFramework/Source/ComputeFramework/Private/ComputeKernelShaderMap.cpp
--- a/Engine/Plugins/Runtime/ComputeFramework/Source/ComputeFramework/Private/ComputeKernelShaderMap.cpp
+++ b/Engine/Plugins/Runtime/ComputeFramework/Source/ComputeFramework/Private/ComputeKernelShaderMap.cpp
@@ -471,7 +471,7 @@ FShader* FComputeKernelShaderMap::ProcessCompilationResultsForSingleJob(FShaderC
 {
 	check(CurrentJob.Id == CompilingId);
 
-	GetResourceCode()->AddShaderCompilerOutput(CurrentJob.Output);
+	GetResourceCode()->AddShaderCompilerOutput(CurrentJob.Output, CurrentJob.Key.ToString());
 
 	FShader* Shader = nullptr;
 
diff --git a/Engine/Shaders/Public/ShaderVersion.ush b/Engine/Shaders/Public/ShaderVersion.ush
--- a/Engine/Shaders/Public/ShaderVersion.ush
+++ b/Engine/Shaders/Public/ShaderVersion.ush
@@ -2,4 +2,4 @@
 // Each time the console command is executed it generates a new GUID. As this file is included
 // in Platform.ush (which should be included in any shader) it allows to invalidate the shader DDC.
 // 
-// GUID = CE2E5BA6-4FFF-4321-9A10-4A3BCEB06687
\ No newline at end of file
+// GUID = D22FFC72-E8D1-4B88-B4D2-A05F60748A10
\ No newline at end of file
diff --git a/Engine/Source/Runtime/Engine/Private/Materials/MaterialShader.cpp b/Engine/Source/Runtime/Engine/Private/Materials/MaterialShader.cpp
--- a/Engine/Source/Runtime/Engine/Private/Materials/MaterialShader.cpp
+++ b/Engine/Source/Runtime/Engine/Private/Materials/MaterialShader.cpp
@@ -2323,7 +2323,7 @@ FShader* FMaterialShaderMap::ProcessCompilationResultsForSingleJob(FShaderCompil
 	const FShaderCompileJob& CurrentJob = *SingleJob;
 	check(CurrentJob.Id == CompilingId);
 
-	GetResourceCode()->AddShaderCompilerOutput(CurrentJob.Output);
+	GetResourceCode()->AddShaderCompilerOutput(CurrentJob.Output, CurrentJob.Key.ToString());
 
 #if ALLOW_SHADERMAP_DEBUG_DATA
 	CompileTime += SingleJob->Output.CompileTime;
diff --git a/Engine/Source/Runtime/Engine/Private/ShaderCompiler/ShaderCompiler.cpp b/Engine/Source/Runtime/Engine/Private/ShaderCompiler/ShaderCompiler.cpp
--- a/Engine/Source/Runtime/Engine/Private/ShaderCompiler/ShaderCompiler.cpp
+++ b/Engine/Source/Runtime/Engine/Private/ShaderCompiler/ShaderCompiler.cpp
@@ -369,6 +369,15 @@ FShaderPipelineCompileJob::FShaderPipelineCompileJob(uint32 InHash, uint32 InId,
 	}
 }
 
+FString FShaderCompileJobKey::ToString() const
+{
+	return FString::Printf(
+		TEXT("ShaderType:%s VertexFactoryType:%s PermutationId:%d"),
+		ShaderType ? ShaderType->GetName() : TEXT("None"),
+		VFType ? VFType->GetName() : TEXT("None"),
+		PermutationId);
+}
+
 FShaderCompileJobCollection::FShaderCompileJobCollection()
 {
 	FMemory::Memzero(PendingJobs);
@@ -6787,7 +6796,7 @@ FShader* FGlobalShaderTypeCompiler::FinishCompileShader(const FGlobalShaderType*
 		EShaderPlatform Platform = CurrentJob.Input.Target.GetPlatform();
 		FGlobalShaderMapSection* Section = GGlobalShaderMap[Platform]->FindOrAddSection(ShaderType);
 
-		Section->GetResourceCode()->AddShaderCompilerOutput(CurrentJob.Output);
+		Section->GetResourceCode()->AddShaderCompilerOutput(CurrentJob.Output, CurrentJob.Key.ToString());
 
 		if (ShaderPipelineType && !ShaderPipelineType->ShouldOptimizeUnusedOutputs(CurrentJob.Input.Target.GetPlatform()))
 		{
diff --git a/Engine/Source/Runtime/Engine/Public/ShaderCompiler.h b/Engine/Source/Runtime/Engine/Public/ShaderCompiler.h
--- a/Engine/Source/Runtime/Engine/Public/ShaderCompiler.h
+++ b/Engine/Source/Runtime/Engine/Public/ShaderCompiler.h
@@ -191,7 +191,7 @@ struct FShaderCompileJobKey
 	{}
 
 	uint32 MakeHash(uint32 Id) const { return HashCombine(HashCombine(HashCombine(GetTypeHash(Id), GetTypeHash(VFType)), GetTypeHash(ShaderType)), GetTypeHash(PermutationId)); }
-
+	ENGINE_API FString ToString() const;
 	const FShaderType* ShaderType;
 	const FVertexFactoryType* VFType;
 	int32 PermutationId;
diff --git a/Engine/Source/Runtime/RenderCore/Private/ShaderResource.cpp b/Engine/Source/Runtime/RenderCore/Private/ShaderResource.cpp
--- a/Engine/Source/Runtime/RenderCore/Private/ShaderResource.cpp
+++ b/Engine/Source/Runtime/RenderCore/Private/ShaderResource.cpp
@@ -257,8 +257,7 @@ FShaderMapResourceCode::FShaderMapResourceCode(const FShaderMapResourceCode& Oth
 	ShaderEntries = Other.ShaderEntries;
 
 #if WITH_EDITORONLY_DATA
-	PlatformDebugData = Other.PlatformDebugData;
-	PlatformDebugDataHashes = Other.PlatformDebugDataHashes;
+	ShaderEditorOnlyDataEntries = Other.ShaderEditorOnlyDataEntries;
 #endif // WITH_EDITORONLY_DATA
 }
 
@@ -296,30 +295,24 @@ int32 FShaderMapResourceCode::FindShaderIndex(const FSHAHash& InHash) const
 	return Algo::BinarySearch(ShaderHashes, InHash);
 }
 
-void FShaderMapResourceCode::AddShaderCompilerOutput(const FShaderCompilerOutput& Output)
-{
-#if WITH_EDITORONLY_DATA
-	AddPlatformDebugData(Output.PlatformDebugData);
-
-	for (const FShaderCompilerError& Error : Output.Errors)
-	{
-		CompilerWarnings.Add(Error.GetErrorString());
-	}
-#endif
-	AddShaderCode(Output.Target.GetFrequency(), Output.OutputHash, Output.ShaderCode);
-}
-
-void FShaderMapResourceCode::AddShaderCode(EShaderFrequency InFrequency, const FSHAHash& InHash, const FShaderCode& InCode)
+void FShaderMapResourceCode::AddShaderCompilerOutput(const FShaderCompilerOutput& Output, const FString& DebugName)
 {
 	TRACE_CPUPROFILER_EVENT_SCOPE(FShaderMapResourceCode::AddShaderCode);
 
+	const FSHAHash& InHash = Output.OutputHash;
+	const FShaderCode& InCode = Output.ShaderCode;
 	const int32 Index = Algo::LowerBound(ShaderHashes, InHash);
 	if (Index >= ShaderHashes.Num() || ShaderHashes[Index] != InHash)
 	{
 		ShaderHashes.Insert(InHash, Index);
 
+#if WITH_EDITORONLY_DATA
+		// Output.Errors contains warnings in the case any exist (no errors since if there were the job would have failed)
+		AddEditorOnlyData(Index, DebugName, Output.PlatformDebugData, Output.Errors);
+#endif
+
 		FShaderEntry& Entry = ShaderEntries.InsertDefaulted_GetRef(Index);
-		Entry.Frequency = InFrequency;
+		Entry.Frequency = Output.Target.GetFrequency();
 		const TArray<uint8>& ShaderCode = InCode.GetReadAccess();
 
 		FName ShaderCompressionFormat = GetShaderCompressionFormat();
@@ -379,42 +372,54 @@ void FShaderMapResourceCode::AddShaderCode(EShaderFrequency InFrequency, const F
 
 		Entry.Code = ShaderCode;
 	}
+#if WITH_EDITORONLY_DATA
+	else
+	{
+		// Output.Errors contains warnings in the case any exist (no errors since if there were the job would have failed)
+		// We append the warnings for any additional jobs which resulted in the same bytecode for the sake of determinism in the
+		// results saved to DDC. 
+		AppendWarningsToEditorOnlyData(Index, DebugName, Output.Errors);
+	}
+#endif
 }
 
 #if WITH_EDITORONLY_DATA
-void FShaderMapResourceCode::AddPlatformDebugData(TConstArrayView<uint8> InPlatformDebugData)
+void FShaderMapResourceCode::AddEditorOnlyData(int32 Index, const FString& DebugName, TConstArrayView<uint8> InPlatformDebugData, TConstArrayView<FShaderCompilerError> InCompilerWarnings)
 {
-	if (InPlatformDebugData.Num() == 0)
-	{
-		return;
-	}
+	FShaderEditorOnlyDataEntry& Entry = ShaderEditorOnlyDataEntries.InsertDefaulted_GetRef(Index);
+	Entry.PlatformDebugData = InPlatformDebugData;
 
-	FSHAHash Hash;
-	{
-		FSHA1 Hasher;
-		Hasher.Update(InPlatformDebugData.GetData(), InPlatformDebugData.Num());
-		Hasher.Final();
-		Hasher.GetHash(Hash.Hash);
-	}
+	AppendWarningsToEditorOnlyData(Index, DebugName, InCompilerWarnings);
+}
 
-	const int32 Index = Algo::LowerBound(PlatformDebugDataHashes, Hash);
-	if (Index >= PlatformDebugDataHashes.Num() || PlatformDebugDataHashes[Index] != Hash)
+void FShaderMapResourceCode::AppendWarningsToEditorOnlyData(int32 Index, const FString& DebugName, TConstArrayView<FShaderCompilerError> InCompilerWarnings)
+{
+	FShaderEditorOnlyDataEntry& Entry = ShaderEditorOnlyDataEntries[Index];
+	for (const FShaderCompilerError& Warning : InCompilerWarnings)
 	{
-		PlatformDebugDataHashes.Insert(Hash, Index);
-		PlatformDebugData.EmplaceAt(Index, InPlatformDebugData.GetData(), InPlatformDebugData.Num());
+		FString ModifiedWarning = !DebugName.IsEmpty() ? FString::Printf(TEXT("%s [%s]"), *Warning.GetErrorString(), *DebugName) : Warning.GetErrorString();
+		// Maintain sorted order in Entry.CompilerWarnings & deduplicate
+		const int32 WarningIndex = Algo::LowerBound(Entry.CompilerWarnings, ModifiedWarning);
+		if (WarningIndex >= Entry.CompilerWarnings.Num() || Entry.CompilerWarnings[WarningIndex] != ModifiedWarning)
+		{
+			Entry.CompilerWarnings.Insert(ModifiedWarning, WarningIndex);
+		}
 	}
 }
 
 void FShaderMapResourceCode::LogShaderCompilerWarnings()
 {
-	if (CompilerWarnings.Num() > 0 && GShaderCompilerEmitWarningsOnLoad != 0)
+	if (ShaderEditorOnlyDataEntries.Num() > 0 && GShaderCompilerEmitWarningsOnLoad != 0)
 	{
 		// Emit all the compiler warnings seen whilst serializing/loading this shader to the log.
 		// Since successfully compiled shaders are stored in the DDC, we'll get the compiler warnings
 		// even if we didn't compile the shader this run.
-		for (const FString& CompilerWarning : CompilerWarnings)
+		for (const FShaderEditorOnlyDataEntry& Entry : ShaderEditorOnlyDataEntries)
 		{
-			UE_LOG(LogShaderWarnings, Warning, TEXT("%s"), *CompilerWarning);
+			for (const FString& CompilerWarning : Entry.CompilerWarnings)
+			{
+				UE_LOG(LogShaderWarnings, Warning, TEXT("%s"), *CompilerWarning);
+			}
 		}
 	}
 }
@@ -441,9 +446,7 @@ void FShaderMapResourceCode::Serialize(FArchive& Ar, bool bLoadedByCookedMateria
 	const bool bSerializeEditorOnlyData = !bLoadedByCookedMaterial && (!Ar.IsCooking() || Ar.CookingTarget()->HasEditorOnlyData());
 	if (bSerializeEditorOnlyData)
 	{
-		Ar << PlatformDebugDataHashes;
-		Ar << PlatformDebugData;
-		Ar << CompilerWarnings;
+		Ar << ShaderEditorOnlyDataEntries;
 	}
 #endif // WITH_EDITORONLY_DATA
 	ApplyResourceStats(*this);
@@ -462,13 +465,13 @@ void FShaderMapResourceCode::NotifyShadersCompiled(FName FormatName)
 #if WITH_ENGINE
 	// Notify the platform shader format that this particular shader is being used in the cook.
 	// We discard this data in cooked builds unless Ar.CookingTarget()->HasEditorOnlyData() is true.
-	if (PlatformDebugData.Num())
+	if (ShaderEditorOnlyDataEntries.Num())
 	{
 		if (const IShaderFormat* ShaderFormat = GetTargetPlatformManagerRef().FindShaderFormat(FormatName))
 		{
-			for (const TArray<uint8>& Entry : PlatformDebugData)
+			for (const FShaderEditorOnlyDataEntry& Entry : ShaderEditorOnlyDataEntries)
 			{
-				ShaderFormat->NotifyShaderCompiled(Entry, FormatName);
+				ShaderFormat->NotifyShaderCompiled(Entry.PlatformDebugData, FormatName);
 			}
 		}
 	}
diff --git a/Engine/Source/Runtime/RenderCore/Public/Shader.h b/Engine/Source/Runtime/RenderCore/Public/Shader.h
--- a/Engine/Source/Runtime/RenderCore/Public/Shader.h
+++ b/Engine/Source/Runtime/RenderCore/Public/Shader.h
@@ -466,14 +466,18 @@ public:
 	};
 
 #if WITH_EDITORONLY_DATA
-	struct FPlatformDebugEntry
+	struct FShaderEditorOnlyDataEntry
 	{
-		uint32 Offset;
-		uint32 Size;
+		TArray<uint8> PlatformDebugData;
+		/** A (deduplicated/sorted) array of all the compiler warnings that were emitted when all shaders resulting 
+		 *  in the associated bytecode were compiled (i.e. if multiple shader sources have warnings but compile to
+		 *  the same code, all warnings for each unique source will be reported).
+		 *  Does not contain errors since if there were any errors, this object wouldn't exist. */
+		TArray<FString> CompilerWarnings;
 
-		friend FArchive& operator<<(FArchive& Ar, FPlatformDebugEntry& Entry)
+		friend FArchive& operator<<(FArchive& Ar, FShaderEditorOnlyDataEntry& Entry)
 		{
-			return Ar << Entry.Offset << Entry.Size;
+			return Ar << Entry.PlatformDebugData << Entry.CompilerWarnings;
 		}
 	};
 #endif // WITH_EDITORONLY_DATA
@@ -493,13 +497,17 @@ public:
 
 	RENDERCORE_API uint32 GetSizeBytes() const;
 
-	RENDERCORE_API void AddShaderCompilerOutput(const FShaderCompilerOutput& Output);
+	RENDERCORE_API void AddShaderCompilerOutput(const FShaderCompilerOutput& Output, const FString& DebugName = FString());
 
 	int32 FindShaderIndex(const FSHAHash& InHash) const;
 
-	RENDERCORE_API void AddShaderCode(EShaderFrequency InFrequency, const FSHAHash& InHash, const FShaderCode& InCode);
+	UE_DEPRECATED(5.1, "Do not call AddShaderCode directly, AddShaderCompilerOutput manages this automatically.")
+	RENDERCORE_API void AddShaderCode(EShaderFrequency InFrequency, const FSHAHash& InHash, const FShaderCode& InCode) {}
 #if WITH_EDITORONLY_DATA
-	RENDERCORE_API void AddPlatformDebugData(TConstArrayView<uint8> InPlatformDebugData);
+	UE_DEPRECATED(5.1, "Do not call AddPlatformDebugData directly; AddShaderCompilerOutput manages this automatically.")
+	RENDERCORE_API void AddPlatformDebugData(TConstArrayView<uint8> InPlatformDebugData) {}
+	void AddEditorOnlyData(int32 Index, const FString& DebugName, TConstArrayView<uint8> InPlatformDebugData, TConstArrayView<FShaderCompilerError> InCompilerWarnings);
+	void AppendWarningsToEditorOnlyData(int32 Index, const FString& DebugName, TConstArrayView<FShaderCompilerError> InCompilerWarnings);
 	RENDERCORE_API void LogShaderCompilerWarnings();
 #endif
 
@@ -510,12 +518,9 @@ public:
 	TArray<FSHAHash> ShaderHashes;
 	TArray<FShaderEntry> ShaderEntries;
 #if WITH_EDITORONLY_DATA
-	TArray<TArray<uint8>> PlatformDebugData;
-	TArray<FSHAHash> PlatformDebugDataHashes;
-
-	/** An array of all the compiler warnings that were emitted when this shader was compiled.
-	 *  Does not contain errors since if there were any errors, this object wouldn't exist. */
-	TArray<FString> CompilerWarnings;
+	// Optional array of editor-only data indexed in the same order as ShaderEntries (sorted by the shader hash)
+	// Empty in the cases where the editor-only data is not serialized.
+	TArray<FShaderEditorOnlyDataEntry> ShaderEditorOnlyDataEntries;
 #endif // WITH_EDITORONLY_DATA
 };
 	
