Fixes warning for missing FDumpTextureCS in shipping builds

diff --git a/Engine/Source/Runtime/RenderCore/Private/DumpGPU.cpp b/Engine/Source/Runtime/RenderCore/Private/DumpGPU.cpp
--- a/Engine/Source/Runtime/RenderCore/Private/DumpGPU.cpp
+++ b/Engine/Source/Runtime/RenderCore/Private/DumpGPU.cpp
@@ -136,6 +136,13 @@ static TAutoConsoleVariable<FString> GDumpGPUVisualizeResource(
 	TEXT("Name of RDG output resource to automatically open in the dump viewer."),
 	ECVF_Default);
 
+} // namespace
+
+#endif // RDG_DUMP_RESOURCES
+
+namespace
+{
+
 class FDumpTextureCS : public FGlobalShader
 {
 	DECLARE_GLOBAL_SHADER(FDumpTextureCS);
@@ -159,6 +166,13 @@ class FDumpTextureCS : public FGlobalShader
 
 IMPLEMENT_GLOBAL_SHADER(FDumpTextureCS, "/Engine/Private/Tools/DumpTexture.usf", "MainCS", SF_Compute);
 
+} // namespace
+
+#if RDG_DUMP_RESOURCES
+
+namespace
+{
+
 BEGIN_SHADER_PARAMETER_STRUCT(FDumpTexturePass, )
 	SHADER_PARAMETER_RDG_TEXTURE_SRV(Texture2D, Texture)
 	RDG_TEXTURE_ACCESS_DYNAMIC(TextureAccess)
