Fixed flickering shadow artifact on some platforms. Fixed crash when r.RDG.CullPasses is 0.

#rnx

diff --git a/Engine/Source/Runtime/RenderCore/Private/RenderGraphBuilder.cpp b/Engine/Source/Runtime/RenderCore/Private/RenderGraphBuilder.cpp
--- a/Engine/Source/Runtime/RenderCore/Private/RenderGraphBuilder.cpp
+++ b/Engine/Source/Runtime/RenderCore/Private/RenderGraphBuilder.cpp
@@ -1133,6 +1133,25 @@ void FRDGBuilder::Compile()
 			AddLastProducersToCullStack(Buffer->LastProducer);
 		}
 	}
+	else
+	{
+		for (FRDGPassHandle PassHandle = ProloguePassHandle + 1; PassHandle < EpiloguePassHandle; ++PassHandle)
+		{
+			FRDGPass* Pass = Passes[PassHandle];
+
+			// Add reference counts for passes.
+
+			for (auto& PassState : Pass->TextureStates)
+			{
+				PassState.Texture->ReferenceCount += PassState.ReferenceCount;
+			}
+
+			for (auto& PassState : Pass->BufferStates)
+			{
+				PassState.Buffer->ReferenceCount += PassState.ReferenceCount;
+			}
+		}
+	}
 
 	// All dependencies in the raw graph have been specified; if enabled, all passes are marked as culled and a
 	// depth first search is employed to find reachable regions of the graph. Roots of the search are those passes
@@ -1787,20 +1806,20 @@ void FRDGBuilder::Execute()
 			ParallelSetupEvents.Empty();
 		}
 
+		CreateUniformBuffersTask.Wait();
+
+		if (SubmitBufferUploadsTask.IsValid())
+		{
+			SubmitBufferUploadsTask.Wait();
+			check(RHICmdListBufferUploads);
+			RHICmdList.QueueAsyncCommandListSubmit(RHICmdListBufferUploads);
+			RHICmdListBufferUploads = nullptr;
+		}
+
 		// Process RHI thread flush before helping with barrier compilation on the render thread.
 		EndFlushResourcesRHI();
 	});
 
-	CreateUniformBuffersTask.Wait();
-
-	if (SubmitBufferUploadsTask.IsValid())
-	{
-		SubmitBufferUploadsTask.Wait();
-		check(RHICmdListBufferUploads);
-		RHICmdList.QueueAsyncCommandListSubmit(RHICmdListBufferUploads);
-		RHICmdListBufferUploads = nullptr;
-	}
-
 	UE::Tasks::FTask ParallelExecuteTask;
 
 	if (bParallelExecuteEnabled)
