Fix Lyra crash due to shader decompression failure (UE-165785).

- PCD3D_SM5 and PCD3D_SM6 had a unique situation of having identical (bytecode-wise and output hash-wise) shaders (raytracing is allowed in PCD3D_SM5). This resulted in those IoStore chunks being deduplicated, however, the actual contents were not identical due to optional data (which isn't included in the shader hash, see UE-166389).
- This fix makes sure we don't use just the shader hash alone to produce IoStore chunk hashes. Hash will now include the shader size and the shader format name.
- A few early guards added here and there to detect the problem early.

Also fixed a memory leak in UnrealPak pointed by CM.

diff --git a/Engine/Source/Developer/IoStoreUtilities/Private/IoStoreUtilities.cpp b/Engine/Source/Developer/IoStoreUtilities/Private/IoStoreUtilities.cpp
--- a/Engine/Source/Developer/IoStoreUtilities/Private/IoStoreUtilities.cpp
+++ b/Engine/Source/Developer/IoStoreUtilities/Private/IoStoreUtilities.cpp
@@ -1866,6 +1866,7 @@ bool ConvertToIoStoreShaderLibrary(
 				FMemory::Memcpy(OutCodeIoChunk.Get<1>().GetData(), UncompressedGroupMemory, Group.UncompressedSize);
 				Group.CompressedSize = Group.UncompressedSize;
 			}
+			FMemory::Free(UncompressedGroupMemory);
 		},
 		EParallelForFlags::Unbalanced
 	);
@@ -1985,6 +1986,12 @@ void ProcessShaderLibraries(const FIoStoreArguments& Arguments, TArray<FContaine
 					}
 					else
 					{
+						// first, make sure that the code is exactly the same
+						if (ShaderInfo->CodeIoBuffer != CodeChunk.Get<1>())
+						{
+							UE_LOG(LogIoStore, Error, TEXT("Collision of two shader code chunks, same Id (%s), different code. Packaged game will likely crash, not being able to decompress the shaders."), *LexToString(ShaderChunkId) );
+						}
+
 						// If we already exist, then we have two separate LoadOrderFactors,
 						// which one we got first affects build determinism. Take the lower.
 						if (CodeChunk.Get<2>() < ShaderInfo->LoadOrderFactor)
diff --git a/Engine/Source/Runtime/Core/Public/IO/IoDispatcher.h b/Engine/Source/Runtime/Core/Public/IO/IoDispatcher.h
--- a/Engine/Source/Runtime/Core/Public/IO/IoDispatcher.h
+++ b/Engine/Source/Runtime/Core/Public/IO/IoDispatcher.h
@@ -481,6 +481,8 @@ public:
 
 	inline void			EnsureOwned() const		{ if (!CorePtr->IsMemoryOwned()) { MakeOwned(); } }
 
+	inline bool			operator!=(const FIoBuffer& Rhs) const { return DataSize() != Rhs.DataSize() || FMemory::Memcmp(GetData(), Rhs.GetData(), DataSize()) != 0; }
+
 	CORE_API void		MakeOwned() const;
 	
 	/**
diff --git a/Engine/Source/Runtime/RenderCore/Private/ShaderCodeArchive.cpp b/Engine/Source/Runtime/RenderCore/Private/ShaderCodeArchive.cpp
--- a/Engine/Source/Runtime/RenderCore/Private/ShaderCodeArchive.cpp
+++ b/Engine/Source/Runtime/RenderCore/Private/ShaderCodeArchive.cpp
@@ -1426,7 +1426,7 @@ void FIoStoreShaderCodeArchive::CreateIoStoreShaderCodeArchiveHeader(const FName
 	TArray<TArray<uint32>> StoredGroupShaderIndices;
 
 	/** Third and last stage of processing the shader group. We actually add it here, and do the book-keeping. */
-	auto ProcessShaderGroup_AddNewGroup = [&OutHeader, &SerializedShaders, &StoredGroupShaderIndices, &Stats_TotalUncompressedMemory, &Stats_MinGroupSize, &Stats_MaxGroupSize](TArray<uint32>& ShaderIndicesInGroup)
+	auto ProcessShaderGroup_AddNewGroup = [&Format, &OutHeader, &SerializedShaders, &StoredGroupShaderIndices, &Stats_TotalUncompressedMemory, &Stats_MinGroupSize, &Stats_MaxGroupSize](TArray<uint32>& ShaderIndicesInGroup)
 	{
 		// first, sort the shaders by uncompressed size, as this was found to compress better
 		ShaderIndicesInGroup.Sort(
@@ -1453,11 +1453,19 @@ void FIoStoreShaderCodeArchive::CreateIoStoreShaderCodeArchiveHeader(const FName
 			IoStoreShaderEntry.ShaderGroupIndex = CurrentGroupIdx;
 			IoStoreShaderEntry.UncompressedOffsetInGroup = CurrentGroupSize;
 
-			// group hash is constructed from hashing the shaders in the group
+			// group hash is constructed from hashing the shaders in the group.
 			GroupHasher.Update(SerializedShaders.ShaderHashes[ShaderIndex].Hash, sizeof(FSHAHash));
+			// shader hash as of now excludes optional data, so we cannot rely on it, especially across the shader formats. Make the group hash a bit more robust by including the shader size in it.
+			GroupHasher.Update(reinterpret_cast<const uint8*>(&SerializedShaders.ShaderEntries[ShaderIndex].UncompressedSize), sizeof(FShaderCodeEntry::UncompressedSize));
 
 			CurrentGroupSize += SerializedShaders.ShaderEntries[ShaderIndex].UncompressedSize;
 		}
+		// Shader hashes cannot be used to uniquely identify across shader formats due to aforementioned exclusion of optional data from it.
+		// Include the shader format (in a platform-agnostic way) into the group hash to lower the risk of collision of shaders of different formats.
+		const TStringConversion<TStringConvert<TCHAR, UTF8CHAR>> Utf8Name(Format.ToString());
+		GroupHasher.Update(reinterpret_cast<const uint8*>(Utf8Name.Get()), Utf8Name.Length());
+		static_assert(sizeof(uint8) == sizeof(UTF8CHAR), "Unexpected UTF-8 char size.");
+
 		GroupEntry.UncompressedSize = CurrentGroupSize;
 		OutHeader.ShaderGroupIoHashes.Add(FIoStoreShaderCodeArchive::GetShaderCodeChunkId(GroupHasher.Finalize()));
 
@@ -2063,6 +2071,7 @@ TRefCountPtr<FRHIShader> FIoStoreShaderCodeArchive::CreateShader(int32 ShaderInd
 	FMemStackBase& MemStack = FMemStack::Get();
 	FMemMark Mark(MemStack);
 	FIoStoreShaderGroupEntry& GroupEntry = Header.ShaderGroupEntries[GroupIndex];
+	ensureMsgf(GroupEntry.CompressedSize == PreloadEntryPtr->IoRequest.GetResultOrDie().DataSize(), TEXT("Shader archive header does not match the actual IoStore buffer size, decompression failure likely imminent."));
 	if (GroupEntry.IsGroupCompressed())
 	{
 		uint8* UncompressedCode = reinterpret_cast<uint8*>(MemStack.Alloc(GroupEntry.UncompressedSize, 16));
