Implemented render graph util function to copy region of buffer.
- Specify source/destination offsets and number of bytes.

diff --git a/Engine/Source/Runtime/RenderCore/Private/RenderGraphUtils.cpp b/Engine/Source/Runtime/RenderCore/Private/RenderGraphUtils.cpp
--- a/Engine/Source/Runtime/RenderCore/Private/RenderGraphUtils.cpp
+++ b/Engine/Source/Runtime/RenderCore/Private/RenderGraphUtils.cpp
@@ -403,7 +403,7 @@ BEGIN_SHADER_PARAMETER_STRUCT(FCopyBufferParameters, )
 	RDG_BUFFER_ACCESS(DstBuffer, ERHIAccess::CopyDest)
 END_SHADER_PARAMETER_STRUCT()
 
-void AddCopyBufferPass(FRDGBuilder& GraphBuilder, FRDGBufferRef DstBuffer, FRDGBufferRef SrcBuffer)
+void AddCopyBufferPass(FRDGBuilder& GraphBuilder, FRDGBufferRef DstBuffer, uint64 DstOffset, FRDGBufferRef SrcBuffer, uint64 SrcOffset, uint64 NumBytes)
 {
 	check(SrcBuffer);
 	check(DstBuffer);
@@ -411,16 +411,25 @@ void AddCopyBufferPass(FRDGBuilder& GraphBuilder, FRDGBufferRef DstBuffer, FRDGB
 	FCopyBufferParameters* Parameters = GraphBuilder.AllocParameters<FCopyBufferParameters>();
 	Parameters->SrcBuffer = SrcBuffer;
 	Parameters->DstBuffer = DstBuffer;
-	const uint64 NumBytes = Parameters->SrcBuffer->Desc.NumElements * Parameters->SrcBuffer->Desc.BytesPerElement;
 
 	GraphBuilder.AddPass(
 		RDG_EVENT_NAME("CopyBuffer(%s Size=%ubytes)", SrcBuffer->Name, SrcBuffer->Desc.GetSize()),
 		Parameters,
 		ERDGPassFlags::Copy,
-		[&Parameters, SrcBuffer, DstBuffer, NumBytes](FRHICommandList& RHICmdList)
-	{
-		RHICmdList.CopyBufferRegion(DstBuffer->GetRHI(), 0, SrcBuffer->GetRHI(), 0, NumBytes);
-	});
+		[&Parameters, SrcBuffer, DstBuffer, SrcOffset, DstOffset, NumBytes](FRHICommandList& RHICmdList)
+		{
+			RHICmdList.CopyBufferRegion(DstBuffer->GetRHI(), DstOffset, SrcBuffer->GetRHI(), SrcOffset, NumBytes);
+		});
+}
+
+void AddCopyBufferPass(FRDGBuilder& GraphBuilder, FRDGBufferRef DstBuffer, FRDGBufferRef SrcBuffer)
+{
+	check(SrcBuffer);
+	check(DstBuffer);
+
+	const uint64 NumBytes = SrcBuffer->Desc.NumElements * SrcBuffer->Desc.BytesPerElement;
+
+	AddCopyBufferPass(GraphBuilder, DstBuffer, 0, SrcBuffer, 0, NumBytes);
 }
 
 BEGIN_SHADER_PARAMETER_STRUCT(FClearBufferUAVParameters, )
diff --git a/Engine/Source/Runtime/RenderCore/Public/RenderGraphUtils.h b/Engine/Source/Runtime/RenderCore/Public/RenderGraphUtils.h
--- a/Engine/Source/Runtime/RenderCore/Public/RenderGraphUtils.h
+++ b/Engine/Source/Runtime/RenderCore/Public/RenderGraphUtils.h
@@ -827,6 +827,8 @@ RENDERCORE_API void AddClearStencilPass(FRDGBuilder& GraphBuilder, FRDGTextureRe
 RENDERCORE_API void AddResummarizeHTilePass(FRDGBuilder& GraphBuilder, FRDGTextureRef DepthTexture);
 
 /** Adds a render graph pass to copy SrcBuffer content into DstBuffer. */
+RENDERCORE_API void AddCopyBufferPass(FRDGBuilder& GraphBuilder, FRDGBufferRef DstBuffer, uint64 DstOffset, FRDGBufferRef SrcBuffer, uint64 SrcOffset, uint64 NumBytes);
+
 RENDERCORE_API void AddCopyBufferPass(FRDGBuilder& GraphBuilder, FRDGBufferRef DstBuffer, FRDGBufferRef SrcBuffer);
 
 /** Adds a pass to readback contents of an RDG texture. */
