Fix MSAA being broken in forward (UE-166561).

diff --git a/Engine/Source/Runtime/RenderCore/Private/RenderUtils.cpp b/Engine/Source/Runtime/RenderCore/Private/RenderUtils.cpp
--- a/Engine/Source/Runtime/RenderCore/Private/RenderUtils.cpp
+++ b/Engine/Source/Runtime/RenderCore/Private/RenderUtils.cpp
@@ -1773,8 +1773,11 @@ RENDERCORE_API bool DoesRuntimeSupportOnePassPointLightShadows(EShaderPlatform P
 
 bool IsUsingBasePassVelocity(const FStaticShaderPlatform Platform)
 {
-	static FShaderPlatformCachedIniValue<bool> PerPlatformCVar(TEXT("r.VelocityOutputPass"));
-	if (IsMobilePlatform(Platform))
+	static FShaderPlatformCachedIniValue<int32> PerPlatformCVar(TEXT("r.VelocityOutputPass"));
+	const IConsoleVariable* AntiAliasingCVar = IConsoleManager::Get().FindConsoleVariable(TEXT("r.AntiAliasingMethod"));
+	// Writing velocity in base pass is disabled for desktop forward when using MSAA
+	const int32 AAM_MSAA = 3;	// see EAntiAliasingMethod in SceneUtils.h
+	if (IsMobilePlatform(Platform) || (AntiAliasingCVar && IsForwardShadingEnabled(Platform) && AntiAliasingCVar->GetInt() == AAM_MSAA))
 	{
 		return false;
 	}
