Fix too frequent FindConsoleVariables calls.

diff --git a/Engine/Source/Runtime/RenderCore/Private/RenderUtils.cpp b/Engine/Source/Runtime/RenderCore/Private/RenderUtils.cpp
--- a/Engine/Source/Runtime/RenderCore/Private/RenderUtils.cpp
+++ b/Engine/Source/Runtime/RenderCore/Private/RenderUtils.cpp
@@ -1774,7 +1774,7 @@ RENDERCORE_API bool DoesRuntimeSupportOnePassPointLightShadows(EShaderPlatform P
 bool IsUsingBasePassVelocity(const FStaticShaderPlatform Platform)
 {
 	static FShaderPlatformCachedIniValue<int32> PerPlatformCVar(TEXT("r.VelocityOutputPass"));
-	const IConsoleVariable* AntiAliasingCVar = IConsoleManager::Get().FindConsoleVariable(TEXT("r.AntiAliasingMethod"));
+	static const IConsoleVariable* AntiAliasingCVar = IConsoleManager::Get().FindConsoleVariable(TEXT("r.AntiAliasingMethod"));
 	// Writing velocity in base pass is disabled for desktop forward when using MSAA
 	const int32 AAM_MSAA = 3;	// see EAntiAliasingMethod in SceneUtils.h
 	if (IsMobilePlatform(Platform) || (AntiAliasingCVar && IsForwardShadingEnabled(Platform) && AntiAliasingCVar->GetInt() == AAM_MSAA))
