Fix a dangling shader precompilation task crash on exit. on behalf of Allan

diff --git a/Engine/Source/Runtime/RenderCore/Private/ShaderPipelineCache.cpp b/Engine/Source/Runtime/RenderCore/Private/ShaderPipelineCache.cpp
--- a/Engine/Source/Runtime/RenderCore/Private/ShaderPipelineCache.cpp
+++ b/Engine/Source/Runtime/RenderCore/Private/ShaderPipelineCache.cpp
@@ -804,6 +804,8 @@ bool FShaderPipelineCache::OpenPipelineFileCache(EShaderPlatform Platform)
 	bool bFileOpen = false;
 	if (ShaderPipelineCache)
 	{
+		FScopeLock Lock(&ShaderPipelineCache->Mutex);
+
 		FString CacheName = UE::ShaderPipeline::GetShaderPipelineBaseName(false);
 		// for Bundled files the Key == CacheName, only the user cache differs.
 		const FString& Key = CacheName;
@@ -1556,14 +1558,20 @@ bool FShaderPipelineCache::ReadyForAutoSave() const
 
 FShaderPipelineCache::~FShaderPipelineCache()
 {
+	// The render thread tick should be dead now and we are safe to destroy everything that needs to wait or manual destruction
+
 	FScopeLock Lock(&Mutex);
+	if (NextPrecompileCacheTask.IsValid() && !NextPrecompileCacheTask->IsComplete())
+	{
+		NextPrecompileCacheTask->Wait();
+	}
+	NextPrecompileCacheTask = FGraphEventRef();
 
 	FShaderPipelineCache::CloseUserPipelineFileCache();
 
 	for (auto& PSOTask : ShaderCacheTasks)
 	{
 		PSOTask.Value->Close();
-		// The render thread tick should be dead now and we are safe to destroy everything that needs to wait or manual destruction
 		// Close() should have called though to Flush() - all work is now in the shutdown lists
 	}
 
@@ -1592,24 +1600,24 @@ FShaderPipelineCache::~FShaderPipelineCache()
 
  void FShaderPipelineCache::BeginNextPrecompileCacheTask()
 {
-	// call begin on the GT, some of the precompile delegate handling is expecting the call to come from the GT.
-	// TODO: we should be able to call this from any thread.
-	if (IsInGameThread())
-	{
-		ShaderPipelineCache->BeginNextPrecompileCacheTaskInternal();
-	}
-	else
-	{
-		static std::atomic_bool bHasDispatched = false;
-		// kick off next PSOFC compile on GT, ensure we only have one request in flight.
-		bool bExpected = false;
-		if (bHasDispatched.compare_exchange_strong(bExpected, true))
+
+	 // Mutex is always held here.
+	if (!NextPrecompileCacheTask.IsValid() || NextPrecompileCacheTask->IsComplete())
+	{	 
+		NextPrecompileCacheTask = FGraphEventRef();
+		// call begin on the GT, some of the precompile delegate handling is expecting the call to come from the GT.
+		// TODO: we should be able to call this from any thread.
+		if (IsInGameThread())
+		{
+			ShaderPipelineCache->BeginNextPrecompileCacheTaskInternal();
+		}
+		else
 		{
-			FFunctionGraphTask::CreateAndDispatchWhenReady(
+			// kick off next PSOFC compile on GT, ensure we only have one request in flight.
+			NextPrecompileCacheTask = FFunctionGraphTask::CreateAndDispatchWhenReady(
 				[]()
 				{
 					ShaderPipelineCache->BeginNextPrecompileCacheTaskInternal();
-					bHasDispatched = false;
 				}
 			, TStatId(), NULL, ENamedThreads::GameThread);
 		}
@@ -1657,7 +1665,7 @@ void FShaderPipelineCache::BeginNextPrecompileCacheTaskInternal()
 		PendingPrecompilePSOCacheKeys.RemoveAtSwap(0);
 		FShaderPipelineCacheTask* NextPrecompileTask = ShaderCacheTasks.Find(CurrentPrecompilingPSOFCKey)->Get();
 		check(NextPrecompileTask->bReady);
-		UE_LOG(LogRHI, Warning, TEXT("FShaderPipelineCache::BeginNextPrecompileCacheTask() - %s begining compile."), *CurrentPrecompilingPSOFCKey);			
+		UE_LOG(LogRHI, Log, TEXT("FShaderPipelineCache::BeginNextPrecompileCacheTask() - %s begining compile."), *CurrentPrecompilingPSOFCKey);			
 		NextPrecompileTask->BeginPrecompilingPipelineCache();
 
 		// FShaderPipelineCache::OnCachePreOpen can set bReady=false, if so we cannot precompile that task.
diff --git a/Engine/Source/Runtime/RenderCore/Public/ShaderPipelineCache.h b/Engine/Source/Runtime/RenderCore/Public/ShaderPipelineCache.h
--- a/Engine/Source/Runtime/RenderCore/Public/ShaderPipelineCache.h
+++ b/Engine/Source/Runtime/RenderCore/Public/ShaderPipelineCache.h
@@ -267,6 +267,7 @@ private:
 	std::atomic_bool bHasShaderPipelineTask = false;
 
 	mutable FCriticalSection Mutex;
+	FGraphEventRef NextPrecompileCacheTask;
 
 	static FShaderCachePreOpenDelegate OnCachePreOpen;
     static FShaderCacheOpenedDelegate OnCachedOpened;
