If a plugin contains a shader code library make sure to close it when unmounting the plugin

diff --git a/Engine/Source/Runtime/Projects/Private/PluginManager.cpp b/Engine/Source/Runtime/Projects/Private/PluginManager.cpp
--- a/Engine/Source/Runtime/Projects/Private/PluginManager.cpp
+++ b/Engine/Source/Runtime/Projects/Private/PluginManager.cpp
@@ -2433,6 +2433,11 @@ IPluginManager::FNewPluginMountedEvent& FPluginManager::OnPluginEdited()
 	return PluginEditedEvent;
 }
 
+IPluginManager::FNewPluginMountedEvent& FPluginManager::OnPluginUnmounted()
+{
+	return PluginUnmountedEvent;
+}
+
 void FPluginManager::MountNewlyCreatedPlugin(const FString& PluginName)
 {
 	TSharedPtr<FPlugin> Plugin = FindPluginInstance(PluginName);
@@ -2616,6 +2621,11 @@ bool FPluginManager::UnmountPluginFromExternalSource(const TSharedPtr<FPlugin>&
 		}
 
 		UnRegisterMountPointDelegate.Execute(Plugin->GetMountedAssetPath(), Plugin->GetContentDir());
+
+		if (PluginUnmountedEvent.IsBound())
+		{
+			PluginUnmountedEvent.Broadcast(*Plugin);
+		}
 	}
 
 	Plugin->SetIsMounted(false);
diff --git a/Engine/Source/Runtime/Projects/Private/PluginManager.h b/Engine/Source/Runtime/Projects/Private/PluginManager.h
--- a/Engine/Source/Runtime/Projects/Private/PluginManager.h
+++ b/Engine/Source/Runtime/Projects/Private/PluginManager.h
@@ -167,6 +167,7 @@ public:
 	virtual FNewPluginMountedEvent& OnNewPluginCreated() override;
 	virtual FNewPluginMountedEvent& OnNewPluginMounted() override;
 	virtual FNewPluginMountedEvent& OnPluginEdited() override;
+	virtual FNewPluginMountedEvent& OnPluginUnmounted() override;
 	virtual void MountNewlyCreatedPlugin(const FString& PluginName) override;
 	virtual bool MountExplicitlyLoadedPlugin(const FString& PluginName) override;
 	virtual bool MountExplicitlyLoadedPlugin_FromFileName(const FString& PluginFileName) override;
@@ -299,6 +300,7 @@ private:
 	FNewPluginMountedEvent NewPluginCreatedEvent;
 	FNewPluginMountedEvent NewPluginMountedEvent;
 	FNewPluginMountedEvent PluginEditedEvent;
+	FNewPluginMountedEvent PluginUnmountedEvent;
 
 	/** Callback for notifications that a loading phase was completed */
 	FLoadingModulesForPhaseEvent LoadingPhaseCompleteEvent;
diff --git a/Engine/Source/Runtime/Projects/Public/Interfaces/IPluginManager.h b/Engine/Source/Runtime/Projects/Public/Interfaces/IPluginManager.h
--- a/Engine/Source/Runtime/Projects/Public/Interfaces/IPluginManager.h
+++ b/Engine/Source/Runtime/Projects/Public/Interfaces/IPluginManager.h
@@ -392,6 +392,11 @@ public:
 	 */
 	virtual FNewPluginMountedEvent& OnPluginEdited() = 0;
 
+	/**
+	* Event for being notified that a plugin has been unmounted
+	*/
+	virtual FNewPluginMountedEvent& OnPluginUnmounted() = 0;
+
 	/**
 	 * Marks a newly created plugin as enabled, mounts its content and tries to load its modules
 	 */
diff --git a/Engine/Source/Runtime/RenderCore/Private/ShaderCodeLibrary.cpp b/Engine/Source/Runtime/RenderCore/Private/ShaderCodeLibrary.cpp
--- a/Engine/Source/Runtime/RenderCore/Private/ShaderCodeLibrary.cpp
+++ b/Engine/Source/Runtime/RenderCore/Private/ShaderCodeLibrary.cpp
@@ -92,6 +92,10 @@ namespace UE
 				ECVF_Default
 			);
 
+			static FDelegateHandle OnPakFileMountedDelegateHandle;
+			static FDelegateHandle OnPluginMountedDelegateHandle;
+			static FDelegateHandle OnPluginUnmountedDelegateHandle;
+
 			/** Helper function shared between the cooker and runtime */
 			FString GetShaderLibraryNameForChunk(FString const& BaseName, int32 ChunkId)
 			{
@@ -3129,6 +3133,15 @@ static void FShaderCodeLibraryPluginMountedCallback(IPlugin& Plugin)
 	}
 }
 
+static void FShaderCodeLibraryPluginUnmountedCallback(IPlugin& Plugin)
+{
+	if (Plugin.CanContainContent() && Plugin.IsEnabled())
+	{
+		// unload any shader libraries that may exist in this plugin
+		FShaderCodeLibrary::CloseLibrary(Plugin.GetName());
+	}
+}
+
 static void FShaderLibraryPakFileMountedCallback(const IPakFile& PakFile)
 {
 	using namespace UE::ShaderLibrary::Private;
@@ -3156,7 +3169,7 @@ static void FShaderLibraryPakFileMountedCallback(const IPakFile& PakFile)
 void FShaderCodeLibrary::PreInit()
 {
 	// add a callback for opening later chunks
-	FCoreDelegates::OnPakFileMounted2.AddStatic(&FShaderLibraryPakFileMountedCallback);
+	UE::ShaderLibrary::Private::OnPakFileMountedDelegateHandle = FCoreDelegates::OnPakFileMounted2.AddStatic(&FShaderLibraryPakFileMountedCallback);
 }
 
 void FShaderCodeLibrary::InitForRuntime(EShaderPlatform ShaderPlatform)
@@ -3184,7 +3197,8 @@ void FShaderCodeLibrary::InitForRuntime(EShaderPlatform ShaderPlatform)
 		FShaderLibrariesCollection::Impl = new FShaderLibrariesCollection(ShaderPlatform, false);
 		if (FShaderLibrariesCollection::Impl->OpenLibrary(TEXT("Global"), FPaths::ProjectContentDir()))
 		{
-			IPluginManager::Get().OnNewPluginMounted().AddStatic(&FShaderCodeLibraryPluginMountedCallback);
+			UE::ShaderLibrary::Private::OnPluginMountedDelegateHandle = IPluginManager::Get().OnNewPluginMounted().AddStatic(&FShaderCodeLibraryPluginMountedCallback);
+			UE::ShaderLibrary::Private::OnPluginUnmountedDelegateHandle = IPluginManager::Get().OnPluginUnmounted().AddStatic(&FShaderCodeLibraryPluginUnmountedCallback);
 		
 #if !UE_BUILD_SHIPPING
 			// support shared cooked builds by also opening the shared cooked build shader code file
@@ -3223,6 +3237,22 @@ void FShaderCodeLibrary::InitForRuntime(EShaderPlatform ShaderPlatform)
 
 void FShaderCodeLibrary::Shutdown()
 {
+	if (UE::ShaderLibrary::Private::OnPakFileMountedDelegateHandle.IsValid())
+	{
+		FCoreDelegates::OnPakFileMounted2.Remove(UE::ShaderLibrary::Private::OnPakFileMountedDelegateHandle);
+		UE::ShaderLibrary::Private::OnPakFileMountedDelegateHandle.Reset();
+	}
+	if (UE::ShaderLibrary::Private::OnPluginMountedDelegateHandle.IsValid())
+	{
+		IPluginManager::Get().OnNewPluginMounted().Remove(UE::ShaderLibrary::Private::OnPluginMountedDelegateHandle);
+		UE::ShaderLibrary::Private::OnPluginMountedDelegateHandle.Reset();
+	}
+	if (UE::ShaderLibrary::Private::OnPluginUnmountedDelegateHandle.IsValid())
+	{
+		IPluginManager::Get().OnPluginUnmounted().Remove(UE::ShaderLibrary::Private::OnPluginUnmountedDelegateHandle);
+		UE::ShaderLibrary::Private::OnPluginUnmountedDelegateHandle.Reset();
+	}
+
 	if (FShaderLibrariesCollection::Impl)
 	{
 		delete FShaderLibrariesCollection::Impl;
