Changing static `TArray`s to be declared as function statics instead of global class statics in order
   to avoid a memory leak.
   #rnx

diff --git a/Engine/Source/Runtime/Core/Private/HAL/ConsoleManager.cpp b/Engine/Source/Runtime/Core/Private/HAL/ConsoleManager.cpp
--- a/Engine/Source/Runtime/Core/Private/HAL/ConsoleManager.cpp
+++ b/Engine/Source/Runtime/Core/Private/HAL/ConsoleManager.cpp
@@ -72,6 +72,27 @@ static const TCHAR* GetSetByTCHAR(EConsoleVariableFlags InSetBy)
 	return TEXT("<UNKNOWN>");
 }
 
+TArray<const FAutoConsoleObject*>& FAutoConsoleObject::AccessGeneralShaderChangeCvars()
+{
+	// This variable cannot be global because it is accessed from the constructor of other global variables.
+	static TArray<const FAutoConsoleObject*> GeneralShaderChangeCvars;
+	return GeneralShaderChangeCvars;
+}
+
+TArray<const FAutoConsoleObject*>& FAutoConsoleObject::AccessMobileShaderChangeCvars()
+{
+	// This variable cannot be global because it is accessed from the constructor of other global variables.
+	static TArray<const FAutoConsoleObject*> MobileShaderChangeCvars;
+	return MobileShaderChangeCvars;
+}
+
+TArray<const FAutoConsoleObject*>& FAutoConsoleObject::AccessDesktopShaderChangeCvars()
+{
+	// This variable cannot be global because it is accessed from the constructor of other global variables.
+	static TArray<const FAutoConsoleObject*> DesktopShaderChangeCvars;
+	return DesktopShaderChangeCvars;
+}
+
 class FConsoleVariableBase : public IConsoleVariable
 {
 public:
@@ -3322,4 +3343,4 @@ static TAutoConsoleVariable<int32> CVarMobileSupportsGen4TAA(
 	TEXT("Support desktop Gen4 TAA with mobile rendering\n"
 		 "0: Fallback to FXAA"
 		 "1: Support Desktop Gen4 TAA (default)"),
-	ECVF_ReadOnly | ECVF_RenderThreadSafe);
\ No newline at end of file
+	ECVF_ReadOnly | ECVF_RenderThreadSafe);
diff --git a/Engine/Source/Runtime/Core/Private/Misc/CoreGlobals.cpp b/Engine/Source/Runtime/Core/Private/Misc/CoreGlobals.cpp
--- a/Engine/Source/Runtime/Core/Private/Misc/CoreGlobals.cpp
+++ b/Engine/Source/Runtime/Core/Private/Misc/CoreGlobals.cpp
@@ -38,9 +38,6 @@ IMPLEMENT_MODULE( FCoreModule, Core );
 /*-----------------------------------------------------------------------------
 	Global variables.
 -----------------------------------------------------------------------------*/
-TArray<const FAutoConsoleObject*> FAutoConsoleObject::GeneralShaderChangeCvars;
-TArray<const FAutoConsoleObject*> FAutoConsoleObject::MobileShaderChangeCvars;
-TArray<const FAutoConsoleObject*> FAutoConsoleObject::DesktopShaderChangeCvars;
 CORE_API FFeedbackContext*	GWarn						= nullptr;		/* User interaction and non critical warnings */
 FConfigCacheIni*			GConfig						= nullptr;		/* Configuration database cache */
 ITransaction*				GUndo						= nullptr;		/* Transaction tracker, non-NULL when a transaction is in progress */
diff --git a/Engine/Source/Runtime/Core/Public/HAL/IConsoleManager.h b/Engine/Source/Runtime/Core/Public/HAL/IConsoleManager.h
--- a/Engine/Source/Runtime/Core/Public/HAL/IConsoleManager.h
+++ b/Engine/Source/Runtime/Core/Public/HAL/IConsoleManager.h
@@ -1057,27 +1057,27 @@ protected:
 		check(Target);
 		if (Target->TestFlags(ECVF_GeneralShaderChange))
 		{
-			GeneralShaderChangeCvars.Add(this);
+			AccessGeneralShaderChangeCvars().Add(this);
 		}
 		else if (Target->TestFlags(ECVF_MobileShaderChange))
 		{
-			MobileShaderChangeCvars.Add(this);
+			AccessMobileShaderChangeCvars().Add(this);
 		}
 		else if (Target->TestFlags(ECVF_DesktopShaderChange))
 		{
-			DesktopShaderChangeCvars.Add(this);
+			AccessDesktopShaderChangeCvars().Add(this);
 		}
 	}
+public:
 	/** Destructor, removes the console object **/
 	virtual ~FAutoConsoleObject()
 	{
 		IConsoleManager::Get().UnregisterConsoleObject(Target);
 	}
 
-public:
-	static TArray<const FAutoConsoleObject*> GeneralShaderChangeCvars;
-	static TArray<const FAutoConsoleObject*> MobileShaderChangeCvars;
-	static TArray<const FAutoConsoleObject*> DesktopShaderChangeCvars;
+	static TArray<const FAutoConsoleObject*>& AccessGeneralShaderChangeCvars();
+	static TArray<const FAutoConsoleObject*>& AccessMobileShaderChangeCvars();
+	static TArray<const FAutoConsoleObject*>& AccessDesktopShaderChangeCvars();
 
 	/** returns the contained console object as an IConsoleVariable **/
 	FORCEINLINE IConsoleVariable* AsVariable()
diff --git a/Engine/Source/Runtime/RenderCore/Private/Shader.cpp b/Engine/Source/Runtime/RenderCore/Private/Shader.cpp
--- a/Engine/Source/Runtime/RenderCore/Private/Shader.cpp
+++ b/Engine/Source/Runtime/RenderCore/Private/Shader.cpp
@@ -1361,7 +1361,7 @@ void ShaderMapAppendKeyString(EShaderPlatform Platform, FString& KeyString)
 {
 	const FName ShaderFormatName = LegacyShaderPlatformToShaderFormat(Platform);
 
-	for (const FAutoConsoleObject* ConsoleObject : FAutoConsoleObject::GeneralShaderChangeCvars)
+	for (const FAutoConsoleObject* ConsoleObject : FAutoConsoleObject::AccessGeneralShaderChangeCvars())
 	{
 		FString ConsoleObjectName = IConsoleManager::Get().FindConsoleObjectName(ConsoleObject->AsVariable());
 		KeyString += TEXT("_");
@@ -1371,7 +1371,7 @@ void ShaderMapAppendKeyString(EShaderPlatform Platform, FString& KeyString)
 	}
 	if (IsMobilePlatform(Platform))
 	{
-		for (const FAutoConsoleObject* ConsoleObject : FAutoConsoleObject::MobileShaderChangeCvars)
+		for (const FAutoConsoleObject* ConsoleObject : FAutoConsoleObject::AccessMobileShaderChangeCvars())
 		{
 			FString ConsoleObjectName = IConsoleManager::Get().FindConsoleObjectName(ConsoleObject->AsVariable());
 			KeyString += TEXT("_");
@@ -1382,7 +1382,7 @@ void ShaderMapAppendKeyString(EShaderPlatform Platform, FString& KeyString)
 	}
 	else if (IsConsolePlatform(Platform))
 	{
-		for (const FAutoConsoleObject* ConsoleObject : FAutoConsoleObject::DesktopShaderChangeCvars)
+		for (const FAutoConsoleObject* ConsoleObject : FAutoConsoleObject::AccessDesktopShaderChangeCvars())
 		{
 			FString ConsoleObjectName = IConsoleManager::Get().FindConsoleObjectName(ConsoleObject->AsVariable());
 			KeyString += TEXT("_");
