Removing PLATFORM_SUPPORTS_SRV_UB since it should be set to the same value for all shader platforms.

diff --git a/Engine/Plugins/FX/Niagara/Source/NiagaraShader/Private/NiagaraShader.cpp b/Engine/Plugins/FX/Niagara/Source/NiagaraShader/Private/NiagaraShader.cpp
--- a/Engine/Plugins/FX/Niagara/Source/NiagaraShader/Private/NiagaraShader.cpp
+++ b/Engine/Plugins/FX/Niagara/Source/NiagaraShader/Private/NiagaraShader.cpp
@@ -446,12 +446,6 @@ void FNiagaraShaderType::AddReferencedUniformBufferIncludes(FShaderCompilerEnvir
 
 	FString& GeneratedUniformBuffersInclude = OutEnvironment.IncludeVirtualPathToContentsMap.FindOrAdd("/Engine/Generated/GeneratedUniformBuffers.ush");
 	GeneratedUniformBuffersInclude.Append(UniformBufferIncludes);
-
-	ERHIFeatureLevel::Type MaxFeatureLevel = GetMaxSupportedFeatureLevel(Platform);
-	if (MaxFeatureLevel >= ERHIFeatureLevel::SM5)
-	{
-		OutEnvironment.SetDefine(TEXT("PLATFORM_SUPPORTS_SRV_UB"), TEXT("1"));
-	}
 }
 
 
diff --git a/Engine/Shaders/Public/Platform.ush b/Engine/Shaders/Public/Platform.ush
--- a/Engine/Shaders/Public/Platform.ush
+++ b/Engine/Shaders/Public/Platform.ush
@@ -97,10 +97,6 @@
 #define COMPILER_FORCE_WAVE32_MODE
 #endif
 
-#ifndef PLATFORM_SUPPORTS_SRV_UB
-#define PLATFORM_SUPPORTS_SRV_UB 0
-#endif
-
 #ifndef PLATFORM_SUPPORTS_ROV
 #define PLATFORM_SUPPORTS_ROV 0
 #define RasterizerOrderedTexture2D RWTexture2D
@@ -169,13 +165,6 @@
 #endif
 #endif
 
-
-#if PLATFORM_SUPPORTS_SRV_UB
-#define PLATFORM_SUPPORTS_SRV_UB_MACRO(...) __VA_ARGS__
-#else
-#define PLATFORM_SUPPORTS_SRV_UB_MACRO(...)
-#endif
-
 #ifndef PLATFORM_SUPPORTS_CALLABLE_SHADERS
 #define PLATFORM_SUPPORTS_CALLABLE_SHADERS 0
 #endif
diff --git a/Engine/Source/Runtime/RenderCore/Private/ShaderParameters.cpp b/Engine/Source/Runtime/RenderCore/Private/ShaderParameters.cpp
--- a/Engine/Source/Runtime/RenderCore/Private/ShaderParameters.cpp
+++ b/Engine/Source/Runtime/RenderCore/Private/ShaderParameters.cpp
@@ -291,9 +291,9 @@ static void CreateHLSLUniformBufferStructMembersDeclaration(
 			}
 			else if (Member.GetBaseType() == UBMT_SRV)
 			{
-				Decl.ResourceMembers += FString::Printf(TEXT("PLATFORM_SUPPORTS_SRV_UB_MACRO( UB_RESOURCE_MEMBER_RESOURCE(%s, %s); ) \r\n"), Member.GetShaderType(), *ParameterName);
-				Decl.StructMembers += FString::Printf(TEXT("\tPLATFORM_SUPPORTS_SRV_UB_MACRO( %s %s; ) \r\n"), Member.GetShaderType(), Member.GetName());
-				Decl.Initializer += FString::Printf(TEXT(" PLATFORM_SUPPORTS_SRV_UB_MACRO( %s, ) "), *ParameterName);
+				Decl.ResourceMembers += FString::Printf(TEXT("UB_RESOURCE_MEMBER_RESOURCE(%s, %s);\r\n"), Member.GetShaderType(), *ParameterName);
+				Decl.StructMembers += FString::Printf(TEXT("\t%s %s;\r\n"), Member.GetShaderType(), Member.GetName());
+				Decl.Initializer += FString::Printf(TEXT("%s,"), *ParameterName);
 			}
 			else
 			{
@@ -449,7 +449,6 @@ void FShaderType::AddReferencedUniformBufferIncludes(FShaderCompilerEnvironment&
 
 	FString& GeneratedUniformBuffersInclude = OutEnvironment.IncludeVirtualPathToContentsMap.FindOrAdd("/Engine/Generated/GeneratedUniformBuffers.ush");
 	GeneratedUniformBuffersInclude.Append(UniformBufferIncludes);
-	OutEnvironment.SetDefine(TEXT("PLATFORM_SUPPORTS_SRV_UB"), TEXT("1"));
 }
 
 void FShaderType::DumpDebugInfo()
@@ -588,5 +587,4 @@ void FVertexFactoryType::AddReferencedUniformBufferIncludes(FShaderCompilerEnvir
 
 	FString& GeneratedUniformBuffersInclude = OutEnvironment.IncludeVirtualPathToContentsMap.FindOrAdd("/Engine/Generated/GeneratedUniformBuffers.ush");
 	GeneratedUniformBuffersInclude.Append(UniformBufferIncludes);
-	OutEnvironment.SetDefine(TEXT("PLATFORM_SUPPORTS_SRV_UB"), TEXT("1"));
 }
