Fix Hololens rendering (UE-167020).

- Fixed Mobile MultiView (fallback to ISR) stereo rendering on D3D12 (UE-167020).

diff --git a/Engine/Plugins/Experimental/Water/Source/Runtime/Private/WaterInfoRendering.cpp b/Engine/Plugins/Experimental/Water/Source/Runtime/Private/WaterInfoRendering.cpp
--- a/Engine/Plugins/Experimental/Water/Source/Runtime/Private/WaterInfoRendering.cpp
+++ b/Engine/Plugins/Experimental/Water/Source/Runtime/Private/WaterInfoRendering.cpp
@@ -652,6 +652,7 @@ static FSceneRenderer* CreateWaterInfoSceneRenderer(const FCreateWaterInfoSceneR
 	ViewInitOptions.ProjectionMatrix = Params.ProjectionMatrix;
 	ViewInitOptions.LODDistanceFactor = 0.001f;
 	ViewInitOptions.OverlayColor = FLinearColor::Black;
+	ViewInitOptions.bIsSceneCapture = true;
 
 	if (ViewFamily.Scene->GetWorld() != nullptr && ViewFamily.Scene->GetWorld()->GetWorldSettings() != nullptr)
 	{
@@ -659,7 +660,6 @@ static FSceneRenderer* CreateWaterInfoSceneRenderer(const FCreateWaterInfoSceneR
 	}
 
 	FSceneView* View = new FSceneView(ViewInitOptions);
-	View->bIsSceneCapture = true;
 	View->GPUMask = FRHIGPUMask::All();
 	View->bOverrideGPUMask = true;
 	View->AntiAliasingMethod = EAntiAliasingMethod::AAM_None;
diff --git a/Engine/Plugins/Runtime/OpenXR/Source/OpenXRHMD/Private/OpenXRHMD.cpp b/Engine/Plugins/Runtime/OpenXR/Source/OpenXRHMD/Private/OpenXRHMD.cpp
--- a/Engine/Plugins/Runtime/OpenXR/Source/OpenXRHMD/Private/OpenXRHMD.cpp
+++ b/Engine/Plugins/Runtime/OpenXR/Source/OpenXRHMD/Private/OpenXRHMD.cpp
@@ -1187,7 +1187,8 @@ FOpenXRHMD::FOpenXRHMD(const FAutoRegister& AutoRegister, XrInstance InInstance,
 	, SelectedEnvironmentBlendMode(XR_ENVIRONMENT_BLEND_MODE_MAX_ENUM)
 	, RenderBridge(InRenderBridge)
 	, RendererModule(nullptr)
-	, LastRequestedSwapchainFormat(0)
+	, LastRequestedColorSwapchainFormat(0)
+	, LastActualColorSwapchainFormat(0)
 	, LastRequestedDepthSwapchainFormat(PF_DepthStencil)
 	, bTrackingSpaceInvalid(true)
 	, bUseCustomReferenceSpace(false)
@@ -2074,13 +2075,14 @@ bool FOpenXRHMD::AllocateRenderTargetTexture(uint32 Index, uint32 SizeX, uint32
 
 	FClearValueBinding ClearColor = FClearValueBinding::Transparent;
 
+	uint8 ActualFormat = Format;
 	FXRSwapChainPtr& Swapchain = PipelinedLayerStateRendering.ColorSwapchain;
 	const FRHITexture2D* const SwapchainTexture = Swapchain == nullptr ? nullptr : Swapchain->GetTexture2DArray() ? Swapchain->GetTexture2DArray() : Swapchain->GetTexture2D();
-	if (Swapchain == nullptr || SwapchainTexture == nullptr || Format != LastRequestedSwapchainFormat || SwapchainTexture->GetSizeX() != SizeX || SwapchainTexture->GetSizeY() != SizeY)
+	if (Swapchain == nullptr || SwapchainTexture == nullptr || Format != LastRequestedColorSwapchainFormat || SwapchainTexture->GetSizeX() != SizeX || SwapchainTexture->GetSizeY() != SizeY)
 	{
 		ensureMsgf(NumSamples == 1, TEXT("OpenXR supports MSAA swapchains, but engine logic expects the swapchain target to be 1x."));
 
-		Swapchain = RenderBridge->CreateSwapchain(Session, Format, SizeX, SizeY, bIsMobileMultiViewEnabled ? 2 : 1, NumMips, NumSamples, UnifiedCreateFlags, ClearColor);
+		Swapchain = RenderBridge->CreateSwapchain(Session, Format, ActualFormat, SizeX, SizeY, bIsMobileMultiViewEnabled ? 2 : 1, NumMips, NumSamples, UnifiedCreateFlags, ClearColor);
 		if (!Swapchain)
 		{
 			return false;
@@ -2089,7 +2091,8 @@ bool FOpenXRHMD::AllocateRenderTargetTexture(uint32 Index, uint32 SizeX, uint32
 
 	// Grab the presentation texture out of the swapchain.
 	OutTargetableTexture = OutShaderResourceTexture = (FTexture2DRHIRef&)Swapchain->GetTextureRef();
-	LastRequestedSwapchainFormat = Format;
+	LastRequestedColorSwapchainFormat = Format;
+	LastActualColorSwapchainFormat = ActualFormat;
 
 	// TODO: Pass in known depth parameters (format + flags)? Do we know that at viewport setup time?
 	AllocateDepthTextureInternal(Index, SizeX, SizeY, NumSamples);
@@ -2130,7 +2133,8 @@ void FOpenXRHMD::AllocateDepthTextureInternal(uint32 Index, uint32 SizeX, uint32
 		constexpr uint32 NumSamplesExpected = 1;
 		constexpr uint32 NumMipsExpected = 1;
 
-		DepthSwapchain = RenderBridge->CreateSwapchain(Session, PF_DepthStencil, SizeX, SizeY, bIsMobileMultiViewEnabled ? 2 : 1, NumMipsExpected, NumSamplesExpected, UnifiedCreateFlags, FClearValueBinding::DepthFar);
+		uint8 UnusedActualFormat = 0;
+		DepthSwapchain = RenderBridge->CreateSwapchain(Session, PF_DepthStencil, UnusedActualFormat, SizeX, SizeY, bIsMobileMultiViewEnabled ? 2 : 1, NumMipsExpected, NumSamplesExpected, UnifiedCreateFlags, FClearValueBinding::DepthFar);
 		if (!DepthSwapchain)
 		{
 			return;
@@ -2186,8 +2190,10 @@ void CreateNativeLayerSwapchain(FOpenXRLayer& Layer, TRefCountPtr<FOpenXRRenderB
 {
 	auto CreateSwapchain = [&](FRHITexture2D* Texture, ETextureCreateFlags Flags)
 	{
+		uint8 UnusedActualFormat = 0;
 		return RenderBridge->CreateSwapchain(Session,
 			IStereoRenderTargetManager::GetStereoLayerPixelFormat(),
+			UnusedActualFormat,
 			Texture->GetSizeX(),
 			Texture->GetSizeY(),
 			1,
diff --git a/Engine/Plugins/Runtime/OpenXR/Source/OpenXRHMD/Private/OpenXRHMD.h b/Engine/Plugins/Runtime/OpenXR/Source/OpenXRHMD/Private/OpenXRHMD.h
--- a/Engine/Plugins/Runtime/OpenXR/Source/OpenXRHMD/Private/OpenXRHMD.h
+++ b/Engine/Plugins/Runtime/OpenXR/Source/OpenXRHMD/Private/OpenXRHMD.h
@@ -312,6 +312,7 @@ public:
 	virtual bool AllocateRenderTargetTexture(uint32 Index, uint32 SizeX, uint32 SizeY, uint8 Format, uint32 NumMips, ETextureCreateFlags Flags, ETextureCreateFlags TargetableTextureFlags, FTexture2DRHIRef& OutTargetableTexture, FTexture2DRHIRef& OutShaderResourceTexture, uint32 NumSamples = 1) override;
 	virtual bool NeedReAllocateDepthTexture(const TRefCountPtr<IPooledRenderTarget>& DepthTarget) override final { return bNeedReAllocatedDepth; }
 	virtual bool AllocateDepthTexture(uint32 Index, uint32 SizeX, uint32 SizeY, uint8 Format, uint32 NumMips, ETextureCreateFlags InTexFlags, ETextureCreateFlags TargetableTextureFlags, FTexture2DRHIRef& OutTargetableTexture, FTexture2DRHIRef& OutShaderResourceTexture, uint32 NumSamples = 1) override final;
+	virtual EPixelFormat GetActualColorSwapchainFormat() const override { return static_cast<EPixelFormat>(LastActualColorSwapchainFormat); }
 
 	/** FXRRenderTargetManager */
 	virtual FXRRenderBridge* GetActiveRenderBridge_GameThread(bool bUseSeparateRenderTarget) override;
@@ -415,7 +416,8 @@ private:
 	TRefCountPtr<FOpenXRRenderBridge> RenderBridge;
 	IRendererModule*		RendererModule;
 
-	uint8					LastRequestedSwapchainFormat;
+	uint8					LastRequestedColorSwapchainFormat;
+	uint8					LastActualColorSwapchainFormat;
 	uint8					LastRequestedDepthSwapchainFormat;
 
 	TArray<FHMDViewMesh>	HiddenAreaMeshes;
diff --git a/Engine/Plugins/Runtime/OpenXR/Source/OpenXRHMD/Private/OpenXRHMD_RenderBridge.cpp b/Engine/Plugins/Runtime/OpenXR/Source/OpenXRHMD/Private/OpenXRHMD_RenderBridge.cpp
--- a/Engine/Plugins/Runtime/OpenXR/Source/OpenXRHMD/Private/OpenXRHMD_RenderBridge.cpp
+++ b/Engine/Plugins/Runtime/OpenXR/Source/OpenXRHMD/Private/OpenXRHMD_RenderBridge.cpp
@@ -47,9 +47,9 @@ public:
 		return &Binding;
 	}
 
-	virtual FXRSwapChainPtr CreateSwapchain(XrSession InSession, uint8 Format, uint32 SizeX, uint32 SizeY, uint32 ArraySize, uint32 NumMips, uint32 NumSamples, ETextureCreateFlags CreateFlags, const FClearValueBinding& ClearValueBinding) override final
+	virtual FXRSwapChainPtr CreateSwapchain(XrSession InSession, uint8 Format, uint8& OutActualFormat, uint32 SizeX, uint32 SizeY, uint32 ArraySize, uint32 NumMips, uint32 NumSamples, ETextureCreateFlags CreateFlags, const FClearValueBinding& ClearValueBinding) override final
 	{
-		return CreateSwapchain_D3D11(InSession, Format, SizeX, SizeY, ArraySize, NumMips, NumSamples, CreateFlags, ClearValueBinding);
+		return CreateSwapchain_D3D11(InSession, Format, OutActualFormat, SizeX, SizeY, ArraySize, NumMips, NumSamples, CreateFlags, ClearValueBinding);
 	}
 
 private:
@@ -87,9 +87,9 @@ public:
 		return &Binding;
 	}
 
-	virtual FXRSwapChainPtr CreateSwapchain(XrSession InSession, uint8 Format, uint32 SizeX, uint32 SizeY, uint32 ArraySize, uint32 NumMips, uint32 NumSamples, ETextureCreateFlags CreateFlags, const FClearValueBinding& ClearValueBinding) override final
+	virtual FXRSwapChainPtr CreateSwapchain(XrSession InSession, uint8 Format, uint8& OutActualFormat, uint32 SizeX, uint32 SizeY, uint32 ArraySize, uint32 NumMips, uint32 NumSamples, ETextureCreateFlags CreateFlags, const FClearValueBinding& ClearValueBinding) override final
 	{
-		return CreateSwapchain_D3D12(InSession, Format, SizeX, SizeY, ArraySize, NumMips, NumSamples, CreateFlags, ClearValueBinding);
+		return CreateSwapchain_D3D12(InSession, Format, OutActualFormat, SizeX, SizeY, ArraySize, NumMips, NumSamples, CreateFlags, ClearValueBinding);
 	}
 
 private:
@@ -141,9 +141,9 @@ public:
 		return nullptr;
 	}
 
-	virtual FXRSwapChainPtr CreateSwapchain(XrSession InSession, uint8 Format, uint32 SizeX, uint32 SizeY, uint32 ArraySize, uint32 NumMips, uint32 NumSamples, ETextureCreateFlags CreateFlags, const FClearValueBinding& ClearValueBinding) override final
+	virtual FXRSwapChainPtr CreateSwapchain(XrSession InSession, uint8 Format, uint8& OutActualFormat, uint32 SizeX, uint32 SizeY, uint32 ArraySize, uint32 NumMips, uint32 NumSamples, ETextureCreateFlags CreateFlags, const FClearValueBinding& ClearValueBinding) override final
 	{
-		return CreateSwapchain_OpenGL(InSession, Format, SizeX, SizeY, ArraySize, NumMips, NumSamples, CreateFlags, ClearValueBinding);
+		return CreateSwapchain_OpenGL(InSession, Format, OutActualFormat, SizeX, SizeY, ArraySize, NumMips, NumSamples, CreateFlags, ClearValueBinding);
 	}
 
 private:
@@ -201,9 +201,9 @@ public:
 		return nullptr;
 	}
 
-	virtual FXRSwapChainPtr CreateSwapchain(XrSession InSession, uint8 Format, uint32 SizeX, uint32 SizeY, uint32 ArraySize, uint32 NumMips, uint32 NumSamples, ETextureCreateFlags CreateFlags, const FClearValueBinding& ClearValueBinding) override final
+	virtual FXRSwapChainPtr CreateSwapchain(XrSession InSession, uint8 Format, uint8& OutActualFormat, uint32 SizeX, uint32 SizeY, uint32 ArraySize, uint32 NumMips, uint32 NumSamples, ETextureCreateFlags CreateFlags, const FClearValueBinding& ClearValueBinding) override final
 	{
-		return CreateSwapchain_OpenGLES(InSession, Format, SizeX, SizeY, ArraySize, NumMips, NumSamples, CreateFlags, ClearValueBinding);
+		return CreateSwapchain_OpenGLES(InSession, Format, OutActualFormat, SizeX, SizeY, ArraySize, NumMips, NumSamples, CreateFlags, ClearValueBinding);
 	}
 
 private:
@@ -282,9 +282,9 @@ public:
 		return AdapterLuid;
 	}
 
-	virtual FXRSwapChainPtr CreateSwapchain(XrSession InSession, uint8 Format, uint32 SizeX, uint32 SizeY, uint32 ArraySize, uint32 NumMips, uint32 NumSamples, ETextureCreateFlags CreateFlags, const FClearValueBinding& ClearValueBinding) override final
+	virtual FXRSwapChainPtr CreateSwapchain(XrSession InSession, uint8 Format, uint8& OutActualFormat, uint32 SizeX, uint32 SizeY, uint32 ArraySize, uint32 NumMips, uint32 NumSamples, ETextureCreateFlags CreateFlags, const FClearValueBinding& ClearValueBinding) override final
 	{
-		return CreateSwapchain_Vulkan(InSession, Format, SizeX, SizeY, ArraySize, NumMips, NumSamples, CreateFlags, ClearValueBinding);
+		return CreateSwapchain_Vulkan(InSession, Format, OutActualFormat, SizeX, SizeY, ArraySize, NumMips, NumSamples, CreateFlags, ClearValueBinding);
 	}
 
 private:
diff --git a/Engine/Plugins/Runtime/OpenXR/Source/OpenXRHMD/Private/OpenXRHMD_RenderBridge.h b/Engine/Plugins/Runtime/OpenXR/Source/OpenXRHMD/Private/OpenXRHMD_RenderBridge.h
--- a/Engine/Plugins/Runtime/OpenXR/Source/OpenXRHMD/Private/OpenXRHMD_RenderBridge.h
+++ b/Engine/Plugins/Runtime/OpenXR/Source/OpenXRHMD/Private/OpenXRHMD_RenderBridge.h
@@ -25,7 +25,7 @@ public:
 	virtual uint64 GetGraphicsAdapterLuid() { return AdapterLuid; }
 
 	virtual void* GetGraphicsBinding() = 0;
-	virtual FXRSwapChainPtr CreateSwapchain(XrSession InSession, uint8 Format, uint32 SizeX, uint32 SizeY, uint32 ArraySize, uint32 NumMips, uint32 NumSamples, ETextureCreateFlags CreateFlags, const FClearValueBinding& ClearValueBinding) = 0;
+	virtual FXRSwapChainPtr CreateSwapchain(XrSession InSession, uint8 Format, uint8& OutActualFormat, uint32 SizeX, uint32 SizeY, uint32 ArraySize, uint32 NumMips, uint32 NumSamples, ETextureCreateFlags CreateFlags, const FClearValueBinding& ClearValueBinding) = 0;
 
 	FXRSwapChainPtr CreateSwapchain(XrSession InSession, FRHITexture2D* Template, ETextureCreateFlags CreateFlags)
 	{
@@ -34,8 +34,10 @@ public:
 			return nullptr;
 		}
 
+		uint8 UnusedOutFormat = 0;
 		return CreateSwapchain(InSession,
 			Template->GetFormat(),
+			UnusedOutFormat,
 			Template->GetSizeX(),
 			Template->GetSizeY(),
 			1,
diff --git a/Engine/Plugins/Runtime/OpenXR/Source/OpenXRHMD/Private/OpenXRHMD_Swapchain.cpp b/Engine/Plugins/Runtime/OpenXR/Source/OpenXRHMD/Private/OpenXRHMD_Swapchain.cpp
--- a/Engine/Plugins/Runtime/OpenXR/Source/OpenXRHMD/Private/OpenXRHMD_Swapchain.cpp
+++ b/Engine/Plugins/Runtime/OpenXR/Source/OpenXRHMD/Private/OpenXRHMD_Swapchain.cpp
@@ -257,7 +257,7 @@ TArray<T> EnumerateImages(XrSwapchain InSwapchain, XrStructureType InType)
 }
 
 #ifdef XR_USE_GRAPHICS_API_D3D11
-FXRSwapChainPtr CreateSwapchain_D3D11(XrSession InSession, uint8 Format, uint32 SizeX, uint32 SizeY, uint32 ArraySize, uint32 NumMips, uint32 NumSamples, ETextureCreateFlags CreateFlags, const FClearValueBinding& ClearValueBinding)
+FXRSwapChainPtr CreateSwapchain_D3D11(XrSession InSession, uint8 Format, uint8& OutActualFormat, uint32 SizeX, uint32 SizeY, uint32 ArraySize, uint32 NumMips, uint32 NumSamples, ETextureCreateFlags CreateFlags, const FClearValueBinding& ClearValueBinding)
 {
 	TFunction<uint32(uint8)> ToPlatformFormat = [](uint8 InFormat)
 	{
@@ -270,6 +270,7 @@ FXRSwapChainPtr CreateSwapchain_D3D11(XrSession InSession, uint8 Format, uint32
 		return nullptr;
 	}
 
+	OutActualFormat = Format;
 	XrSwapchain Swapchain = FOpenXRSwapchain::CreateSwapchain(InSession, ToPlatformFormat(Format), SizeX, SizeY, ArraySize, NumMips, NumSamples, CreateFlags);
 	if (!Swapchain)
 	{
@@ -291,7 +292,7 @@ FXRSwapChainPtr CreateSwapchain_D3D11(XrSession InSession, uint8 Format, uint32
 #endif
 
 #ifdef XR_USE_GRAPHICS_API_D3D12
-FXRSwapChainPtr CreateSwapchain_D3D12(XrSession InSession, uint8 Format, uint32 SizeX, uint32 SizeY, uint32 ArraySize, uint32 NumMips, uint32 NumSamples, ETextureCreateFlags CreateFlags, const FClearValueBinding& ClearValueBinding)
+FXRSwapChainPtr CreateSwapchain_D3D12(XrSession InSession, uint8 Format, uint8& OutActualFormat, uint32 SizeX, uint32 SizeY, uint32 ArraySize, uint32 NumMips, uint32 NumSamples, ETextureCreateFlags CreateFlags, const FClearValueBinding& ClearValueBinding)
 {
 	TFunction<uint32(uint8)> ToPlatformFormat = [](uint8 InFormat)
 	{
@@ -304,6 +305,7 @@ FXRSwapChainPtr CreateSwapchain_D3D12(XrSession InSession, uint8 Format, uint32
 		return nullptr;
 	}
 
+	OutActualFormat = Format;
 	XrSwapchain Swapchain = FOpenXRSwapchain::CreateSwapchain(InSession, ToPlatformFormat(Format), SizeX, SizeY, ArraySize, NumMips, NumSamples, CreateFlags);
 	if (!Swapchain)
 	{
@@ -324,7 +326,7 @@ FXRSwapChainPtr CreateSwapchain_D3D12(XrSession InSession, uint8 Format, uint32
 #endif
 
 #ifdef XR_USE_GRAPHICS_API_OPENGL
-FXRSwapChainPtr CreateSwapchain_OpenGL(XrSession InSession, uint8 Format, uint32 SizeX, uint32 SizeY, uint32 ArraySize, uint32 NumMips, uint32 NumSamples, ETextureCreateFlags CreateFlags, const FClearValueBinding& ClearValueBinding)
+FXRSwapChainPtr CreateSwapchain_OpenGL(XrSession InSession, uint8 Format, uint8& OutActualFormat, uint32 SizeX, uint32 SizeY, uint32 ArraySize, uint32 NumMips, uint32 NumSamples, ETextureCreateFlags CreateFlags, const FClearValueBinding& ClearValueBinding)
 {
 	Format = FOpenXRSwapchain::GetNearestSupportedSwapchainFormat(InSession, Format);
 	if (!Format)
@@ -332,6 +334,7 @@ FXRSwapChainPtr CreateSwapchain_OpenGL(XrSession InSession, uint8 Format, uint32
 		return nullptr;
 	}
 
+	OutActualFormat = Format;
 	XrSwapchain Swapchain = FOpenXRSwapchain::CreateSwapchain(InSession, GPixelFormats[Format].PlatformFormat, SizeX, SizeY, ArraySize, NumMips, NumSamples, CreateFlags);
 	if (!Swapchain)
 	{
@@ -353,7 +356,7 @@ FXRSwapChainPtr CreateSwapchain_OpenGL(XrSession InSession, uint8 Format, uint32
 #endif
 
 #ifdef XR_USE_GRAPHICS_API_OPENGL_ES
-FXRSwapChainPtr CreateSwapchain_OpenGLES(XrSession InSession, uint8 Format, uint32 SizeX, uint32 SizeY, uint32 ArraySize, uint32 NumMips, uint32 NumSamples, ETextureCreateFlags CreateFlags, const FClearValueBinding& ClearValueBinding)
+FXRSwapChainPtr CreateSwapchain_OpenGLES(XrSession InSession, uint8 Format, uint8& OutActualFormat, uint32 SizeX, uint32 SizeY, uint32 ArraySize, uint32 NumMips, uint32 NumSamples, ETextureCreateFlags CreateFlags, const FClearValueBinding& ClearValueBinding)
 {
 	Format = FOpenXRSwapchain::GetNearestSupportedSwapchainFormat(InSession, Format);
 	if (!Format)
@@ -361,6 +364,7 @@ FXRSwapChainPtr CreateSwapchain_OpenGLES(XrSession InSession, uint8 Format, uint
 		return nullptr;
 	}
 
+	OutActualFormat = Format;
 	XrSwapchain Swapchain = FOpenXRSwapchain::CreateSwapchain(InSession, GPixelFormats[Format].PlatformFormat, SizeX, SizeY, ArraySize, NumMips, NumSamples, CreateFlags);
 	if (!Swapchain)
 	{
@@ -382,7 +386,7 @@ FXRSwapChainPtr CreateSwapchain_OpenGLES(XrSession InSession, uint8 Format, uint
 #endif
 
 #ifdef XR_USE_GRAPHICS_API_VULKAN
-FXRSwapChainPtr CreateSwapchain_Vulkan(XrSession InSession, uint8 Format, uint32 SizeX, uint32 SizeY, uint32 ArraySize, uint32 NumMips, uint32 NumSamples, ETextureCreateFlags CreateFlags, const FClearValueBinding& ClearValueBinding)
+FXRSwapChainPtr CreateSwapchain_Vulkan(XrSession InSession, uint8 Format, uint8& OutActualFormat, uint32 SizeX, uint32 SizeY, uint32 ArraySize, uint32 NumMips, uint32 NumSamples, ETextureCreateFlags CreateFlags, const FClearValueBinding& ClearValueBinding)
 {
 	TFunction<uint32(uint8)> ToPlatformFormat = [](uint8 InFormat)
 	{
@@ -395,6 +399,7 @@ FXRSwapChainPtr CreateSwapchain_Vulkan(XrSession InSession, uint8 Format, uint32
 		return nullptr;
 	}
 
+	OutActualFormat = Format;
 	// When we specify the mutable format flag we want to inform the runtime that we'll only use a
 	// linear and an sRGB view format to allow for efficiency optimizations.
 	TArray<VkFormat> ViewFormatList;
diff --git a/Engine/Plugins/Runtime/OpenXR/Source/OpenXRHMD/Private/OpenXRHMD_Swapchain.h b/Engine/Plugins/Runtime/OpenXR/Source/OpenXRHMD/Private/OpenXRHMD_Swapchain.h
--- a/Engine/Plugins/Runtime/OpenXR/Source/OpenXRHMD/Private/OpenXRHMD_Swapchain.h
+++ b/Engine/Plugins/Runtime/OpenXR/Source/OpenXRHMD/Private/OpenXRHMD_Swapchain.h
@@ -34,17 +34,17 @@ protected:
 };
 
 #ifdef XR_USE_GRAPHICS_API_D3D11
-FXRSwapChainPtr CreateSwapchain_D3D11(XrSession InSession, uint8 Format, uint32 SizeX, uint32 SizeY, uint32 ArraySize, uint32 NumMips, uint32 NumSamples, ETextureCreateFlags CreateFlags, const FClearValueBinding& ClearValueBinding);
+FXRSwapChainPtr CreateSwapchain_D3D11(XrSession InSession, uint8 Format, uint8& OutActualFormat, uint32 SizeX, uint32 SizeY, uint32 ArraySize, uint32 NumMips, uint32 NumSamples, ETextureCreateFlags CreateFlags, const FClearValueBinding& ClearValueBinding);
 #endif
 #ifdef XR_USE_GRAPHICS_API_D3D12
-FXRSwapChainPtr CreateSwapchain_D3D12(XrSession InSession, uint8 Format, uint32 SizeX, uint32 SizeY, uint32 ArraySize, uint32 NumMips, uint32 NumSamples, ETextureCreateFlags CreateFlags, const FClearValueBinding& ClearValueBinding);
+FXRSwapChainPtr CreateSwapchain_D3D12(XrSession InSession, uint8 Format, uint8& OutActualFormat, uint32 SizeX, uint32 SizeY, uint32 ArraySize, uint32 NumMips, uint32 NumSamples, ETextureCreateFlags CreateFlags, const FClearValueBinding& ClearValueBinding);
 #endif
 #ifdef XR_USE_GRAPHICS_API_OPENGL
-FXRSwapChainPtr CreateSwapchain_OpenGL(XrSession InSession, uint8 Format, uint32 SizeX, uint32 SizeY, uint32 ArraySize, uint32 NumMips, uint32 NumSamples, ETextureCreateFlags CreateFlags, const FClearValueBinding& ClearValueBinding);
+FXRSwapChainPtr CreateSwapchain_OpenGL(XrSession InSession, uint8 Format, uint8& OutActualFormat, uint32 SizeX, uint32 SizeY, uint32 ArraySize, uint32 NumMips, uint32 NumSamples, ETextureCreateFlags CreateFlags, const FClearValueBinding& ClearValueBinding);
 #endif
 #ifdef XR_USE_GRAPHICS_API_OPENGL_ES
-FXRSwapChainPtr CreateSwapchain_OpenGLES(XrSession InSession, uint8 Format, uint32 SizeX, uint32 SizeY, uint32 ArraySize, uint32 NumMips, uint32 NumSamples, ETextureCreateFlags CreateFlags, const FClearValueBinding& ClearValueBinding);
+FXRSwapChainPtr CreateSwapchain_OpenGLES(XrSession InSession, uint8 Format, uint8& OutActualFormat, uint32 SizeX, uint32 SizeY, uint32 ArraySize, uint32 NumMips, uint32 NumSamples, ETextureCreateFlags CreateFlags, const FClearValueBinding& ClearValueBinding);
 #endif
 #ifdef XR_USE_GRAPHICS_API_VULKAN
-FXRSwapChainPtr CreateSwapchain_Vulkan(XrSession InSession, uint8 Format, uint32 SizeX, uint32 SizeY, uint32 ArraySize, uint32 NumMips, uint32 NumSamples, ETextureCreateFlags CreateFlags, const FClearValueBinding& ClearValueBinding);
+FXRSwapChainPtr CreateSwapchain_Vulkan(XrSession InSession, uint8 Format, uint8& OutActualFormat, uint32 SizeX, uint32 SizeY, uint32 ArraySize, uint32 NumMips, uint32 NumSamples, ETextureCreateFlags CreateFlags, const FClearValueBinding& ClearValueBinding);
 #endif
\ No newline at end of file
diff --git a/Engine/Plugins/Runtime/nDisplay/Source/DisplayCluster/Private/Render/Viewport/DisplayClusterViewport_EditorPreview.cpp b/Engine/Plugins/Runtime/nDisplay/Source/DisplayCluster/Private/Render/Viewport/DisplayClusterViewport_EditorPreview.cpp
--- a/Engine/Plugins/Runtime/nDisplay/Source/DisplayCluster/Private/Render/Viewport/DisplayClusterViewport_EditorPreview.cpp
+++ b/Engine/Plugins/Runtime/nDisplay/Source/DisplayCluster/Private/Render/Viewport/DisplayClusterViewport_EditorPreview.cpp
@@ -118,15 +118,11 @@ FSceneView* FDisplayClusterViewport::ImplCalcScenePreview(FSceneViewFamilyContex
 			ViewInitOptions.OverlayColor = FLinearColor::Black;
 		}
 
-		FSceneView* View = new FSceneView(ViewInitOptions);
-
-		View->bIsSceneCapture = true;
-		View->bSceneCaptureUsesRayTracing = false;
-		View->bIsPlanarReflection = false;
+		ViewInitOptions.bIsSceneCapture = true;
+		ViewInitOptions.bSceneCaptureUsesRayTracing = false;
+		ViewInitOptions.bIsPlanarReflection = false;
 
-		// Note: this has to be set before EndFinalPostprocessSettings
-		// Needs to be reconfigured now that bIsPlanarReflection has changed.
-		View->SetupAntiAliasingMethod();
+		FSceneView* View = new FSceneView(ViewInitOptions);
 
 		InOutViewFamily.Views.Add(View);
 
diff --git a/Engine/Source/Runtime/D3D12RHI/Private/D3D12CommandList.cpp b/Engine/Source/Runtime/D3D12RHI/Private/D3D12CommandList.cpp
--- a/Engine/Source/Runtime/D3D12RHI/Private/D3D12CommandList.cpp
+++ b/Engine/Source/Runtime/D3D12RHI/Private/D3D12CommandList.cpp
@@ -439,12 +439,13 @@ void FD3D12ContextCommon::TransitionResource(FD3D12RenderTargetView* View, D3D12
 		break;
 
 	case D3D12_RTV_DIMENSION_TEXTURE2DARRAY:
+	case D3D12_RTV_DIMENSION_TEXTURE2DMSARRAY:
 		// Multiple subresources to transition
 		TransitionResource(Resource, D3D12_RESOURCE_STATE_TBD, After, View->GetViewSubresourceSubset());
 		break;
 
 	default:
-		checkNoEntry();
+		UE_LOG(LogRHI, Fatal, TEXT("Unsupported resource dimension %d"), Desc.ViewDimension);
 		break;
 	}
 }
diff --git a/Engine/Source/Runtime/D3D12RHI/Private/D3D12RenderTarget.cpp b/Engine/Source/Runtime/D3D12RHI/Private/D3D12RenderTarget.cpp
--- a/Engine/Source/Runtime/D3D12RHI/Private/D3D12RenderTarget.cpp
+++ b/Engine/Source/Runtime/D3D12RHI/Private/D3D12RenderTarget.cpp
@@ -622,8 +622,8 @@ void FD3D12CommandContext::ResolveTexture(UE::RHICore::FResolveTextureInfo Info)
 
 		for (int32 ArraySlice = ArraySliceBegin; ArraySlice < ArraySliceEnd; ArraySlice++)
 		{
-			int32 DestSubresource   = CalcSubresource(Info.MipLevel, ArraySlice, DestDesc.ArraySize);
-			int32 SourceSubresource = CalcSubresource(Info.MipLevel, ArraySlice, SourceDesc.ArraySize);
+			int32 DestSubresource   = CalcSubresource(Info.MipLevel, ArraySlice, DestDesc.NumMips);
+			int32 SourceSubresource = CalcSubresource(Info.MipLevel, ArraySlice, SourceDesc.NumMips);
 
 			FScopedResourceBarrier ConditionalScopeResourceBarrierDst(*this, DestTexture->GetResource()  , D3D12_RESOURCE_STATE_RESOLVE_DEST  , DestSubresource  );
 			FScopedResourceBarrier ConditionalScopeResourceBarrierSrc(*this, SourceTexture->GetResource(), D3D12_RESOURCE_STATE_RESOLVE_SOURCE, SourceSubresource);
diff --git a/Engine/Source/Runtime/D3D12RHI/Private/D3D12Texture.cpp b/Engine/Source/Runtime/D3D12RHI/Private/D3D12Texture.cpp
--- a/Engine/Source/Runtime/D3D12RHI/Private/D3D12Texture.cpp
+++ b/Engine/Source/Runtime/D3D12RHI/Private/D3D12Texture.cpp
@@ -2708,9 +2708,8 @@ FD3D12Texture* FD3D12DynamicRHI::CreateTextureFromResource(bool bTextureArray, b
 	TextureResource->AddRef();
 
 	FRHITextureCreateDesc CreateDesc =
-		FRHITextureCreateDesc::Create2D(TEXT("TextureFromResource"))
-		.SetExtent(FIntPoint(SizeX, SizeY))
-		.SetFormat(Format)
+		((SizeZ > 1) ? FRHITextureCreateDesc::Create2DArray(TEXT("TextureFromResource"), FIntPoint(SizeX, SizeY), SizeZ, Format) :
+			FRHITextureCreateDesc::Create2D(TEXT("TextureFromResource"), FIntPoint(SizeX, SizeY), Format))
 		.SetClearValue(ClearValueBinding)
 		.SetFlags(TexCreateFlags)
 		.SetNumMips(NumMips)
diff --git a/Engine/Source/Runtime/Engine/Private/SceneTexturesConfig.cpp b/Engine/Source/Runtime/Engine/Private/SceneTexturesConfig.cpp
--- a/Engine/Source/Runtime/Engine/Private/SceneTexturesConfig.cpp
+++ b/Engine/Source/Runtime/Engine/Private/SceneTexturesConfig.cpp
@@ -6,6 +6,7 @@
 #include "SceneInterface.h"
 #include "Engine/Engine.h"
 #include "StereoRendering.h"
+#include "StereoRenderTargetManager.h"
 
 FSceneTexturesConfig FSceneTexturesConfig::GlobalInstance;
 
@@ -29,7 +30,24 @@ FSceneTextureShaderParameters GetSceneTextureShaderParameters(TRDGUniformBufferR
 
 static EPixelFormat GetDefaultMobileSceneColorLowPrecisionFormat()
 {
-	return (GEngine && GEngine->XRSystem.IsValid() && GEngine->StereoRenderingDevice.IsValid() && GEngine->StereoRenderingDevice->IsStandaloneStereoOnlyDevice()) ? PF_R8G8B8A8 : PF_B8G8R8A8;
+	if (UNLIKELY(GEngine && GEngine->XRSystem.IsValid() && GEngine->StereoRenderingDevice.IsValid()))
+	{
+		IStereoRenderTargetManager* RTManager = GEngine->StereoRenderingDevice->GetRenderTargetManager();
+		if (LIKELY(RTManager))
+		{
+			EPixelFormat ActualFormat = RTManager->GetActualColorSwapchainFormat();
+			if (LIKELY(ActualFormat != PF_Unknown))
+			{
+				return ActualFormat;
+			}
+		}
+		
+		if (GEngine->StereoRenderingDevice->IsStandaloneStereoOnlyDevice())
+		{
+			return PF_R8G8B8A8;
+		}
+	}
+	return PF_B8G8R8A8;
 }
 
 static EPixelFormat GetMobileSceneColorFormat(bool bRequiresAlphaChannel)
diff --git a/Engine/Source/Runtime/Engine/Private/SceneView.cpp b/Engine/Source/Runtime/Engine/Private/SceneView.cpp
--- a/Engine/Source/Runtime/Engine/Private/SceneView.cpp
+++ b/Engine/Source/Runtime/Engine/Private/SceneView.cpp
@@ -763,11 +763,11 @@ FSceneView::FSceneView(const FSceneViewInitOptions& InitOptions)
 	, CursorPos(InitOptions.CursorPos)
 	, bIsGameView(false)
 	, bIsViewInfo(false)
-	, bIsSceneCapture(false)
-	, bIsSceneCaptureCube(false)
-	, bSceneCaptureUsesRayTracing(false)
-	, bIsReflectionCapture(false)
-	, bIsPlanarReflection(false)
+	, bIsSceneCapture(InitOptions.bIsSceneCapture)
+	, bIsSceneCaptureCube(InitOptions.bIsSceneCaptureCube)
+	, bSceneCaptureUsesRayTracing(InitOptions.bSceneCaptureUsesRayTracing)
+	, bIsReflectionCapture(InitOptions.bIsReflectionCapture)
+	, bIsPlanarReflection(InitOptions.bIsPlanarReflection)
 	, bIsVirtualTexture(false)
 	, bIsOfflineRender(false)
 	, bRenderSceneTwoSided(false)
@@ -865,40 +865,29 @@ PRAGMA_ENABLE_DEPRECATION_WARNINGS
 #endif
 
 	// Query instanced stereo and multi-view state
-	static const auto CVar = IConsoleManager::Get().FindTConsoleVariableDataInt(TEXT("vr.InstancedStereo"));
-	bIsInstancedStereoEnabled = !bUsingMobileRenderer && RHISupportsInstancedStereo(ShaderPlatform) && (CVar && CVar->GetValueOnAnyThread() != 0);
-	bIsMultiViewportEnabled = GRHISupportsArrayIndexFromAnyShader && RHISupportsMultiViewport(ShaderPlatform) && bIsInstancedStereoEnabled;
-
-	PRAGMA_DISABLE_DEPRECATION_WARNINGS
-	bIsMultiViewEnabled = bIsMultiViewportEnabled;	// temporary, as a graceful way to support plugins/licensee mods
-	PRAGMA_ENABLE_DEPRECATION_WARNINGS
-
-	bIsMobileMultiViewEnabled = Family && Family->bRequireMultiView;
-	if (bIsMobileMultiViewEnabled && !RHISupportsMobileMultiView(ShaderPlatform))
-	{
-		// Native mobile multi-view is not supported, attempt to fall back to instancing on compatible RHIs
-		if (RHISupportsInstancedStereo(ShaderPlatform) && !GRHISupportsArrayIndexFromAnyShader)
-		{
-			UE_LOG(LogMultiView, Fatal, TEXT("Mobile Multi-View not supported by the RHI and no fallback is available."));
-		}
-		bIsMobileMultiViewEnabled = bIsInstancedStereoEnabled = RHISupportsInstancedStereo(ShaderPlatform);
-	}
-
-	// If instanced stereo is enabled, we should also have either multiviewport enabled, or mmv fallback enabled.
-	// Assert this since a more graceful handling is done earlier (see PostInitRHI)
-	checkf(!bIsInstancedStereoEnabled || (bIsMultiViewportEnabled || bIsMobileMultiViewEnabled), 
-		TEXT("If instanced stereo is enabled, either multi-viewport or mobile multi-view needs to be enabled as well."));
-
 	{
 		// The shader variants that are compiled have ISR or MMV _enabled_ in the shader, even if the current ViewFamily doesn't
 		// require multiple views functionality.
 		const UE::StereoRenderUtils::FStereoShaderAspects Aspects(ShaderPlatform);
 		bShouldBindInstancedViewUB = Aspects.IsInstancedStereoEnabled() || Aspects.IsMobileMultiViewEnabled();
 
-		// make sure that FStereoShaderAspects logic matches the above logic (TODO: above should be replaced by setting them through aspects after more testing)
-		ensure(bIsInstancedStereoEnabled == Aspects.IsInstancedStereoEnabled());
-		ensure(bIsMultiViewportEnabled == Aspects.IsInstancedMultiViewportEnabled());
-		ensure(bIsMobileMultiViewEnabled == (Family && Family->bRequireMultiView && Aspects.IsMobileMultiViewEnabled()));
+		if (Family && Family->bRequireMultiView && !Aspects.IsMobileMultiViewEnabled())
+		{
+			UE_LOG(LogMultiView, Fatal, TEXT("Family requires Mobile Multi-View, but it is not supported by the RHI and no fallback is available."));
+		}
+
+		bIsInstancedStereoEnabled = Aspects.IsInstancedStereoEnabled();
+		bIsMultiViewportEnabled = Aspects.IsInstancedMultiViewportEnabled();
+		bIsMobileMultiViewEnabled = (!bIsSceneCapture && !bIsReflectionCapture && !bIsPlanarReflection && Family && Family->bRequireMultiView && Aspects.IsMobileMultiViewEnabled());
+
+		PRAGMA_DISABLE_DEPRECATION_WARNINGS
+		bIsMultiViewEnabled = bIsMultiViewportEnabled;	// temporary, as a graceful way to support plugins/licensee mods
+		PRAGMA_ENABLE_DEPRECATION_WARNINGS
+
+		// If instanced stereo is enabled, we should also have either multiviewport enabled, or mmv fallback enabled.
+		// Assert this since a more graceful handling is done earlier (see PostInitRHI)
+		checkf(!bIsInstancedStereoEnabled || (bIsMultiViewportEnabled || Aspects.IsMobileMultiViewEnabled()),
+			TEXT("If instanced stereo is enabled, either multi-viewport or mobile multi-view needs to be enabled as well."));
 	}
 
 	SetupAntiAliasingMethod();
@@ -2926,8 +2915,6 @@ FSceneViewFamily::FSceneViewFamily(const ConstructionValues& CVS)
 		EngineShowFlags.ScreenPercentage = false;
 	}
 
-	// TODO: Re-enable Mobile Multi-View on all platforms when all desktop XR plugins support it
-#if (PLATFORM_HOLOLENS || PLATFORM_ANDROID)
 	if (GEngine && GEngine->IsStereoscopic3D())
 	{
 		static const auto MobileMultiViewCVar = IConsoleManager::Get().FindTConsoleVariableDataInt(TEXT("vr.MobileMultiView"));
@@ -2935,7 +2922,6 @@ FSceneViewFamily::FSceneViewFamily(const ConstructionValues& CVS)
 		const bool bUsingMobileRenderer = FSceneInterface::GetShadingPath(GetFeatureLevel()) == EShadingPath::Mobile;
 		bRequireMultiView = (GSupportsMobileMultiView || GRHISupportsArrayIndexFromAnyShader) && bUsingMobileRenderer && bSkipPostprocessing && (MobileMultiViewCVar && MobileMultiViewCVar->GetValueOnAnyThread() != 0);
 	}
-#endif
 
 	// Check if the translucency are allowed to be rendered after DOF, if not, translucency after DOF will be rendered in standard translucency.
 	{
diff --git a/Engine/Source/Runtime/Engine/Private/ShaderCompiler/ShaderCompiler.cpp b/Engine/Source/Runtime/Engine/Private/ShaderCompiler/ShaderCompiler.cpp
--- a/Engine/Source/Runtime/Engine/Private/ShaderCompiler/ShaderCompiler.cpp
+++ b/Engine/Source/Runtime/Engine/Private/ShaderCompiler/ShaderCompiler.cpp
@@ -5786,12 +5786,26 @@ void GlobalBeginCompileShader(
 		Input.Environment.SetDefine(TEXT("MULTI_VIEW"), Aspects.IsInstancedMultiViewportEnabled());
 		Input.Environment.SetDefine(TEXT("MOBILE_MULTI_VIEW"), Aspects.IsMobileMultiViewEnabled());
 
-		// Throw a warning if we are silently disabling ISR due to missing platform support.
+		// Throw a warning if we are silently disabling ISR due to missing platform support (but don't have MMV enabled).
 		static const auto CVarInstancedStereo = IConsoleManager::Get().FindTConsoleVariableDataInt(TEXT("vr.InstancedStereo"));
-		const bool bIsInstancedStereoCVar = CVarInstancedStereo ? (CVarInstancedStereo->GetValueOnAnyThread() != 0) : false;
-		if (bIsInstancedStereoCVar && !Aspects.IsInstancedStereoEnabled() && !GShaderCompilingManager->AreWarningsSuppressed(ShaderPlatform))
+		const bool bIsInstancedStereoEnabledInSettings = CVarInstancedStereo ? (CVarInstancedStereo->GetValueOnAnyThread() != 0) : false;
+		static const auto CVarMultiview = IConsoleManager::Get().FindTConsoleVariableDataInt(TEXT("vr.MobileMultiView"));
+		const bool bIsMultiviewEnabledInSettings = CVarMultiview ? (CVarMultiview->GetValueOnAnyThread() != 0) : false;
+		bool bWarningIssued = false;
+		// warn if ISR was enabled in settings, but aspects show that it's not enabled AND we don't use Mobile MultiView as an alternative
+		if (bIsInstancedStereoEnabledInSettings && !Aspects.IsInstancedStereoEnabled() && !(bIsMultiviewEnabledInSettings && Aspects.IsMobileMultiViewEnabled()) && !GShaderCompilingManager->AreWarningsSuppressed(ShaderPlatform))
+		{
+			UE_LOG(LogShaderCompilers, Warning, TEXT("Instanced stereo rendering is not supported for %s shader platform."), *ShaderFormatName.ToString());
+			bWarningIssued = true;
+		}
+		// Warn if MMV was enabled in settings, but aspects show that it's not enabled AND we don't use Instanced Stereo as an alternative
+		if (bIsMultiviewEnabledInSettings && !Aspects.IsMobileMultiViewEnabled() && !(bIsInstancedStereoEnabledInSettings && Aspects.IsInstancedStereoEnabled()) && !GShaderCompilingManager->AreWarningsSuppressed(ShaderPlatform))
+		{
+			UE_LOG(LogShaderCompilers, Warning, TEXT("Multiview rendering is not supported for %s shader platform."), *ShaderFormatName.ToString());
+			bWarningIssued = true;
+		}
+		if (bWarningIssued)
 		{
-			UE_LOG(LogShaderCompilers, Warning, TEXT("Instanced stereo rendering is not supported for the %s shader platform."), *ShaderFormatName.ToString());
 			GShaderCompilingManager->SuppressWarnings(ShaderPlatform);
 		}
 	}
diff --git a/Engine/Source/Runtime/Engine/Public/SceneView.h b/Engine/Source/Runtime/Engine/Public/SceneView.h
--- a/Engine/Source/Runtime/Engine/Public/SceneView.h
+++ b/Engine/Source/Runtime/Engine/Public/SceneView.h
@@ -176,6 +176,21 @@ struct FSceneViewInitOptions : public FSceneViewProjectionData
 	/** In case of ortho, generate a fake view position that has a non-zero W component. The view position will be derived based on the view matrix. */
 	bool bUseFauxOrthoViewPos;
 
+	/** Whether this view is being used to render a scene capture. */
+	bool bIsSceneCapture;
+
+	/** Whether the scene capture is a cube map (bIsSceneCapture will also be set). */
+	bool bIsSceneCaptureCube;
+
+	/** Whether this view uses ray tracing, for views that are used to render a scene capture. */
+	bool bSceneCaptureUsesRayTracing;
+
+	/** Whether this view is being used to render a reflection capture. */
+	bool bIsReflectionCapture;
+
+	/** Whether this view is being used to render a planar reflection. */
+	bool bIsPlanarReflection;
+
 #if WITH_EDITOR
 	/** default to 0'th view index, which is a bitfield of 1 */
 	uint64 EditorViewBitflag;
@@ -210,6 +225,11 @@ struct FSceneViewInitOptions : public FSceneViewProjectionData
 		, FOV(90.f)
 		, DesiredFOV(90.f)
 		, bUseFauxOrthoViewPos(false)
+		, bIsSceneCapture(false)
+		, bIsSceneCaptureCube(false)
+		, bSceneCaptureUsesRayTracing(false)
+		, bIsReflectionCapture(false)
+		, bIsPlanarReflection(false)
 #if WITH_EDITOR
 		, EditorViewBitflag(1)
 		, OverrideLODViewOrigin(ForceInitToZero)
diff --git a/Engine/Source/Runtime/Engine/Public/StereoRenderTargetManager.h b/Engine/Source/Runtime/Engine/Public/StereoRenderTargetManager.h
--- a/Engine/Source/Runtime/Engine/Public/StereoRenderTargetManager.h
+++ b/Engine/Source/Runtime/Engine/Public/StereoRenderTargetManager.h
@@ -72,6 +72,11 @@ public:
 	 */
 	virtual bool AllocateRenderTargetTexture(uint32 Index, uint32 SizeX, uint32 SizeY, uint8 Format, uint32 NumMips, ETextureCreateFlags Flags, ETextureCreateFlags TargetableTextureFlags, FTexture2DRHIRef& OutTargetableTexture, FTexture2DRHIRef& OutShaderResourceTexture, uint32 NumSamples = 1) { return false; }
 
+	/**
+	 * Returns pixel format that the device created its swapchain with (which can be different than what was requested in AllocateRenderTargetTexture)
+	 */
+	virtual EPixelFormat GetActualColorSwapchainFormat() const { return PF_Unknown; }
+	
 	/**
 	 * Allocates a depth texture.
 	 *
diff --git a/Engine/Source/Runtime/RHI/Private/RHI.cpp b/Engine/Source/Runtime/RHI/Private/RHI.cpp
--- a/Engine/Source/Runtime/RHI/Private/RHI.cpp
+++ b/Engine/Source/Runtime/RHI/Private/RHI.cpp
@@ -2334,12 +2334,18 @@ FString LexToString(EShaderPlatform Platform, bool bError)
 	case SP_VULKAN_SM5: return TEXT("VULKAN_SM5");
 	case SP_VULKAN_SM5_ANDROID: return TEXT("VULKAN_SM5_ANDROID");
 	case SP_D3D_ES3_1_HOLOLENS: return TEXT("D3D_ES3_1_HOLOLENS");
+	case SP_CUSTOM_PLATFORM_FIRST: return TEXT("SP_CUSTOM_PLATFORM_FIRST");
+	case SP_CUSTOM_PLATFORM_LAST: return TEXT("SP_CUSTOM_PLATFORM_LAST");
 
 	default:
 		if (FStaticShaderPlatformNames::IsStaticPlatform(Platform))
 		{
 			return FStaticShaderPlatformNames::Get().GetShaderPlatform(Platform).ToString();
 		}
+		else if (Platform > SP_CUSTOM_PLATFORM_FIRST && Platform < SP_CUSTOM_PLATFORM_LAST)
+		{
+			return FString::Printf(TEXT("SP_CUSTOM_PLATFORM_%d"), Platform - SP_CUSTOM_PLATFORM_FIRST);
+		}
 		else
 		{
 			checkf(!bError, TEXT("Unknown or removed EShaderPlatform %d!"), (int32)Platform);
diff --git a/Engine/Source/Runtime/RenderCore/Private/StereoRenderUtils.cpp b/Engine/Source/Runtime/RenderCore/Private/StereoRenderUtils.cpp
--- a/Engine/Source/Runtime/RenderCore/Private/StereoRenderUtils.cpp
+++ b/Engine/Source/Runtime/RenderCore/Private/StereoRenderUtils.cpp
@@ -3,6 +3,39 @@
 #include "StereoRenderUtils.h"
 #include "RenderUtils.h"
 
+// enable this to printf-debug stereo rendering aspects on a device
+#define UE_DEBUG_STEREO_ASPECTS			(!UE_BUILD_SHIPPING && !UE_BUILD_TEST)
+
+#if UE_DEBUG_STEREO_ASPECTS
+	#define UE_DEBUG_SSA_LOG_INIT(Platform) \
+	bool bSSALogEnable = false; \
+	{ \
+		static EShaderPlatform LastLogged[2] = { SP_NumPlatforms, SP_NumPlatforms };	\
+		if (UNLIKELY(LastLogged[0] != Platform && LastLogged[1] != Platform)) \
+		{ \
+			LastLogged[1] = LastLogged[0]; \
+			LastLogged[0] = Platform; \
+			bSSALogEnable = true; \
+		} \
+	}
+
+	#define UE_DEBUG_SSA_LOG(Verbosity, Format, ...) \
+	{ \
+		if (UNLIKELY(bSSALogEnable)) \
+		{ \
+			UE_LOG(LogInit, Verbosity, TEXT("FStereoShaderAspects: %s"), *FString::Printf(Format, ##__VA_ARGS__)); \
+		} \
+	}
+#else
+	#define UE_DEBUG_SSA_LOG_INIT(Platform)
+	#define UE_DEBUG_SSA_LOG(Verbosity, Format, ...)
+#endif
+
+#define UE_DEBUG_SSA_LOG_BOOL(Boolean) \
+{ \
+	UE_DEBUG_SSA_LOG(Log, TEXT(#Boolean) TEXT(" = %d"), Boolean); \
+}
+
 namespace UE::StereoRenderUtils
 {
 	
@@ -15,14 +48,15 @@ RENDERCORE_API FStereoShaderAspects::FStereoShaderAspects(EShaderPlatform Platfo
 	, bMobileMultiViewFallback(false)
 {
 	check(Platform < EShaderPlatform::SP_NumPlatforms);
+	UE_DEBUG_SSA_LOG_INIT(Platform);
 	
 	// Would be nice to use URendererSettings, but not accessible in RenderCore
 	static FShaderPlatformCachedIniValue<bool> CVarInstancedStereo(TEXT("vr.InstancedStereo"));
 	static FShaderPlatformCachedIniValue<bool> CVarMobileMultiView(TEXT("vr.MobileMultiView"));
 	static FShaderPlatformCachedIniValue<bool> CVarMobileHDR(TEXT("r.MobileHDR"));
-	
+
 	const bool bInstancedStereo = CVarInstancedStereo.Get(Platform);
-	
+
 	const bool bMobilePlatform = IsMobilePlatform(Platform);
 	const bool bMobilePostprocessing = CVarMobileHDR.Get(Platform);
 	const bool bMobileMultiView = CVarMobileMultiView.Get(Platform);
@@ -30,20 +64,45 @@ RENDERCORE_API FStereoShaderAspects::FStereoShaderAspects(EShaderPlatform Platfo
 	const bool bMultiViewportCapable = (GRHISupportsArrayIndexFromAnyShader || IsRunningCookCommandlet()) && RHISupportsMultiViewport(Platform);
 
 	bInstancedStereoNative = !bMobilePlatform && bInstancedStereo && RHISupportsInstancedStereo(Platform);
-	
+
+	UE_DEBUG_SSA_LOG(Log, TEXT("--- StereoAspects begin ---"));
+	UE_DEBUG_SSA_LOG(Log, TEXT("Platform=%s (%d)"), *LexToString(Platform), static_cast<int32>(Platform));
+	UE_DEBUG_SSA_LOG_BOOL(bInstancedStereo);
+	UE_DEBUG_SSA_LOG_BOOL(bMobilePlatform);
+	UE_DEBUG_SSA_LOG_BOOL(bMobilePostprocessing);
+	UE_DEBUG_SSA_LOG_BOOL(bMobileMultiView);
+	UE_DEBUG_SSA_LOG_BOOL(bMultiViewportCapable);
+	UE_DEBUG_SSA_LOG_BOOL(bInstancedStereoNative);
+	UE_DEBUG_SSA_LOG(Log, TEXT("---"));
+
 	const bool bMobileMultiViewCoreSupport = bMobilePlatform && bMobileMultiView && !bMobilePostprocessing;
 	if (bMobileMultiViewCoreSupport)
 	{
+		UE_DEBUG_SSA_LOG(Log, TEXT("RHISupportsMobileMultiView(%s) = %d."), *LexToString(Platform), RHISupportsMobileMultiView(Platform));
+		UE_DEBUG_SSA_LOG(Log, TEXT("RHISupportsInstancedStereo(%s) = %d."), *LexToString(Platform), RHISupportsInstancedStereo(Platform));
+
 		if (RHISupportsMobileMultiView(Platform))
 		{
 			bMobileMultiViewNative = true;
 		}
-		else if (RHISupportsInstancedStereo(Platform) && RHISupportsVertexShaderLayer(Platform))
+		else if (RHISupportsInstancedStereo(Platform))
 		{
-			bMobileMultiViewFallback = true;
+			// ISR in mobile shaders is achieved via layered RTs and eye-dependent layer index, so unlike on desktop, it does not depend on multi-viewport
+			const bool bVertexShaderLayerCapable = (GRHISupportsArrayIndexFromAnyShader || IsRunningCookCommandlet()) && RHISupportsVertexShaderLayer(Platform);
+			UE_DEBUG_SSA_LOG_BOOL(bVertexShaderLayerCapable);
+
+			if (bVertexShaderLayerCapable)
+			{
+				bMobileMultiViewFallback = true;
+			}
 		}
 	}
 
+	UE_DEBUG_SSA_LOG_BOOL(bMobileMultiViewCoreSupport);
+	UE_DEBUG_SSA_LOG_BOOL(bMobileMultiViewNative);
+	UE_DEBUG_SSA_LOG_BOOL(bMobileMultiViewFallback);
+	UE_DEBUG_SSA_LOG(Log, TEXT("---"));
+
 	// "instanced stereo" is confusingly used to refer to two two modes:
 	// 1) regular aka "native" ISR, where the views are selected via SV_ViewportArrayIndex - uses non-mobile shaders
 	// 2) "mobile multiview fallback" ISR, which writes to a texture layer via SV_RenderTargetArrayIndex - uses mobile shaders
@@ -53,10 +112,16 @@ RENDERCORE_API FStereoShaderAspects::FStereoShaderAspects(EShaderPlatform Platfo
 	// Since instanced stereo now relies on multi-viewport capability, it cannot be separately enabled from it.
 	bInstancedStereoEnabled = bInstancedStereoNative || bMobileMultiViewFallback;
 	bMobileMultiViewEnabled = bMobileMultiViewNative || bMobileMultiViewFallback;
+
+	UE_DEBUG_SSA_LOG_BOOL(bInstancedMultiViewportEnabled);
+	UE_DEBUG_SSA_LOG_BOOL(bInstancedStereoEnabled);
+	UE_DEBUG_SSA_LOG_BOOL(bMobileMultiViewEnabled);
+	UE_DEBUG_SSA_LOG(Log, TEXT("--- StereoAspects end ---"));
+
+	// check the following invariants
+	checkf(!bMobileMultiViewNative || !bInstancedStereoEnabled, TEXT("When a platform supports MMV natively, ISR should not be enabled."));
+	checkf(!(bMobileMultiViewEnabled && bInstancedStereoEnabled) || bMobileMultiViewFallback, TEXT("If a platform uses MMV fallback, both MMV and ISR should be enabled."));
+	checkf(!bInstancedStereoEnabled || (bInstancedMultiViewportEnabled || bMobileMultiViewFallback), TEXT("If ISR is enabled, we need either multi-viewport (since we no longer support clip-distance method) or MMV fallback (which uses vertex layer)."));
 }
 	
-RENDERCORE_API bool FStereoShaderAspects::IsInstancedStereoEnabled() const { return bInstancedStereoEnabled; }
-RENDERCORE_API bool FStereoShaderAspects::IsMobileMultiViewEnabled() const { return bMobileMultiViewEnabled; }
-RENDERCORE_API bool FStereoShaderAspects::IsInstancedMultiViewportEnabled() const { return bInstancedMultiViewportEnabled; }
-	
 } // namespace UE::RenderUtils
diff --git a/Engine/Source/Runtime/RenderCore/Public/StereoRenderUtils.h b/Engine/Source/Runtime/RenderCore/Public/StereoRenderUtils.h
--- a/Engine/Source/Runtime/RenderCore/Public/StereoRenderUtils.h
+++ b/Engine/Source/Runtime/RenderCore/Public/StereoRenderUtils.h
@@ -21,22 +21,22 @@ namespace UE::StereoRenderUtils
 
 		/**
 		 * Whether instanced stereo rendering is enabled - i.e. using a single instanced drawcall to render to both stereo views.
-		 * Can redirect output either via viewport index or X coordinate shift + clip planes.
+		 * The output is redirected via the viewport index.
 		 */
-		RENDERCORE_API bool IsInstancedStereoEnabled() const;
+		inline bool IsInstancedStereoEnabled() const { return bInstancedStereoEnabled; }
 
 		/**
 		 * Whether mobile multiview is enabled - i.e. using VK_KHR_multiview. Another drawcall reduction technique, independent of instanced stereo.
 		 * Mobile multiview generates view indices to index into texture arrays.
 		 * Can be internally emulated using instanced stereo if native support is unavailable, by using ISR-generated view indices to index into texture arrays.
 		 */
-		RENDERCORE_API bool IsMobileMultiViewEnabled() const;
+		inline bool IsMobileMultiViewEnabled() const { return bMobileMultiViewEnabled; };
 
 		/**
 		 * Whether multiviewport rendering is enabled - i.e. using ViewportIndex to index into viewport.
 		 * Relies on instanced stereo rendering being enabled.
 		 */
-		RENDERCORE_API bool IsInstancedMultiViewportEnabled() const;
+		inline bool IsInstancedMultiViewportEnabled() const { return bInstancedMultiViewportEnabled; };
 
 	private:
 		bool bInstancedStereoEnabled : 1;
diff --git a/Engine/Source/Runtime/Renderer/Private/SceneCaptureRendering.cpp b/Engine/Source/Runtime/Renderer/Private/SceneCaptureRendering.cpp
--- a/Engine/Source/Runtime/Renderer/Private/SceneCaptureRendering.cpp
+++ b/Engine/Source/Runtime/Renderer/Private/SceneCaptureRendering.cpp
@@ -622,6 +622,10 @@ void SetupViewFamilyForSceneCapture(
 		ViewInitOptions.ProjectionMatrix = SceneCaptureViewInfo.ProjectionMatrix;
 		ViewInitOptions.LODDistanceFactor = FMath::Clamp(SceneCaptureComponent->LODDistanceFactor, .01f, 100.0f);
 		ViewInitOptions.bUseFauxOrthoViewPos = bUseFauxOrthoViewPos;
+		ViewInitOptions.bIsSceneCapture = true;
+		ViewInitOptions.bIsSceneCaptureCube = SceneCaptureComponent->IsCube();
+		ViewInitOptions.bSceneCaptureUsesRayTracing = SceneCaptureComponent->bUseRayTracingIfEnabled;
+		ViewInitOptions.bIsPlanarReflection = bIsPlanarReflection;
 
 		if (ViewFamily.Scene->GetWorld() != nullptr && ViewFamily.Scene->GetWorld()->GetWorldSettings() != nullptr)
 		{
@@ -636,14 +640,6 @@ void SetupViewFamilyForSceneCapture(
 
 		FSceneView* View = new FSceneView(ViewInitOptions);
 
-		View->bIsSceneCapture = true;
-		View->bIsSceneCaptureCube = SceneCaptureComponent->IsCube();
-		View->bSceneCaptureUsesRayTracing = SceneCaptureComponent->bUseRayTracingIfEnabled;
-		// Note: this has to be set before EndFinalPostprocessSettings
-		View->bIsPlanarReflection = bIsPlanarReflection;
-        // Needs to be reconfigured now that bIsPlanarReflection has changed.
-		View->SetupAntiAliasingMethod();
-
 		check(SceneCaptureComponent);
 		for (auto It = SceneCaptureComponent->HiddenComponents.CreateConstIterator(); It; ++It)
 		{
