Implements r.DumpGPU.Delay

diff --git a/Engine/Source/Runtime/RenderCore/Private/DumpGPU.cpp b/Engine/Source/Runtime/RenderCore/Private/DumpGPU.cpp
--- a/Engine/Source/Runtime/RenderCore/Private/DumpGPU.cpp
+++ b/Engine/Source/Runtime/RenderCore/Private/DumpGPU.cpp
@@ -82,6 +82,11 @@ static TAutoConsoleVariable<int32> GDumpGPUPassParameters(
 	TEXT("Whether to dump the pass parameters."),
 	ECVF_RenderThreadSafe);
 
+static TAutoConsoleVariable<float> GDumpGPUDelay(
+	TEXT("r.DumpGPU.Delay"), 0.0f,
+	TEXT("Delay in seconds before dumping the frame."),
+	ECVF_RenderThreadSafe);
+
 static TAutoConsoleVariable<int32> GDumpGPUDraws(
 	TEXT("r.DumpGPU.Draws"), 0,
 	TEXT("Whether to dump resource after each individual draw call (disabled by default)."),
@@ -1582,6 +1587,9 @@ public:
 static uint64 DumpingFrameCounter_GameThread = 0;
 static FRDGResourceDumpContext* GRDGResourceDumpContext = nullptr;
 
+static float GNextDumpingRemainingTime = -1.0f;
+static FRDGResourceDumpContext* GNextRDGResourceDumpContext = nullptr;
+
 bool IsDumpingRDGResources()
 {
 	return GRDGResourceDumpContext != nullptr;
@@ -1599,7 +1607,7 @@ FString FRDGBuilder::BeginResourceDump(const TCHAR* Cmd)
 {
 	check(IsInGameThread());
 
-	if (DumpingFrameCounter_GameThread != 0)
+	if (DumpingFrameCounter_GameThread != 0 || GNextDumpingRemainingTime > 0.0f)
 	{
 		return FString();
 	}
@@ -1639,59 +1647,76 @@ FString FRDGBuilder::BeginResourceDump(const TCHAR* Cmd)
 		}
 	}
 
-	FRDGResourceDumpContext* NewResourceDumpContext = new FRDGResourceDumpContext;
-
-	NewResourceDumpContext->Time = FDateTime::Now();
+	FRDGResourceDumpContext* NewResourceDumpContext;
+	if (GNextRDGResourceDumpContext == nullptr)
 	{
-		FString CVarDirectoryPath = GDumpGPUDirectoryCVar.GetValueOnGameThread();
-		FString EnvDirectoryPath = FPlatformMisc::GetEnvironmentVariable(TEXT("UE-DumpGPUPath"));
+		NewResourceDumpContext = new FRDGResourceDumpContext;
 
-		FString DirectoryPath;
-		if (!CVarDirectoryPath.IsEmpty())
-		{
-			DirectoryPath = CVarDirectoryPath;
-		}
-		else if (!EnvDirectoryPath.IsEmpty())
+		NewResourceDumpContext->Time = FDateTime::Now();
 		{
-			DirectoryPath = EnvDirectoryPath;
-		}
-		else
-		{
-			DirectoryPath = FPaths::ProjectSavedDir() / TEXT("GPUDumps/");
+			FString CVarDirectoryPath = GDumpGPUDirectoryCVar.GetValueOnGameThread();
+			FString EnvDirectoryPath = FPlatformMisc::GetEnvironmentVariable(TEXT("UE-DumpGPUPath"));
+
+			FString DirectoryPath;
+			if (!CVarDirectoryPath.IsEmpty())
+			{
+				DirectoryPath = CVarDirectoryPath;
+			}
+			else if (!EnvDirectoryPath.IsEmpty())
+			{
+				DirectoryPath = EnvDirectoryPath;
+			}
+			else
+			{
+				DirectoryPath = FPaths::ProjectSavedDir() / TEXT("GPUDumps/");
+			}
+			NewResourceDumpContext->DumpingDirectoryPath = DirectoryPath / FApp::GetProjectName() + TEXT("-") + FPlatformProperties::PlatformName() + TEXT("-") + NewResourceDumpContext->Time.ToString() + TEXT("/");
 		}
-		NewResourceDumpContext->DumpingDirectoryPath = DirectoryPath / FApp::GetProjectName() + TEXT("-") + FPlatformProperties::PlatformName() + TEXT("-") + NewResourceDumpContext->Time.ToString() + TEXT("/");
-	}
-	NewResourceDumpContext->bEnableDiskWrite = GDumpTestEnableDiskWrite.GetValueOnGameThread() != 0;
+		NewResourceDumpContext->bEnableDiskWrite = GDumpTestEnableDiskWrite.GetValueOnGameThread() != 0;
 
-	if (Switches.Contains(TEXT("upload")))
-	{
-		if (!IDumpGPUUploadServiceProvider::GProvider || !NewResourceDumpContext->bEnableDiskWrite)
+		if (Switches.Contains(TEXT("upload")))
 		{
-			UE_LOG(LogRendererCore, Warning, TEXT("DumpGPU upload services are not set up."));
+			if (!IDumpGPUUploadServiceProvider::GProvider || !NewResourceDumpContext->bEnableDiskWrite)
+			{
+				UE_LOG(LogRendererCore, Warning, TEXT("DumpGPU upload services are not set up."));
+			}
+			else if (GDumpGPUUploadCVar.GetValueOnGameThread() == 0)
+			{
+				UE_LOG(LogRendererCore, Warning, TEXT("DumpGPU upload services are not available because r.DumpGPU.Upload=0."));
+			}
+			else
+			{
+				NewResourceDumpContext->bUpload = true;
+			}
 		}
-		else if (GDumpGPUUploadCVar.GetValueOnGameThread() == 0)
+
+		if (NewResourceDumpContext->bUpload)
 		{
-			UE_LOG(LogRendererCore, Warning, TEXT("DumpGPU upload services are not available because r.DumpGPU.Upload=0."));
+			if (GDumpGPUUploadCompressResources.GetValueOnGameThread() == 1)
+			{
+				NewResourceDumpContext->UploadResourceCompressionName = NAME_Zlib;
+			}
+			else if (GDumpGPUUploadCompressResources.GetValueOnGameThread() == 2)
+			{
+				NewResourceDumpContext->UploadResourceCompressionName = NAME_Gzip;
+			}
 		}
-		else
+
+		NewResourceDumpContext->bShowInExplore = NewResourceDumpContext->bEnableDiskWrite && GDumpExploreCVar.GetValueOnGameThread() != 0 && !NewResourceDumpContext->bUpload;
+	
+		if (GDumpGPUDelay.GetValueOnGameThread() > 0.0f)
 		{
-			NewResourceDumpContext->bUpload = true;
+			GNextDumpingRemainingTime = GDumpGPUDelay.GetValueOnGameThread();
+			GNextRDGResourceDumpContext = NewResourceDumpContext;
+			return NewResourceDumpContext->DumpingDirectoryPath;
 		}
 	}
-
-	if (NewResourceDumpContext->bUpload)
+	else
 	{
-		if (GDumpGPUUploadCompressResources.GetValueOnGameThread() == 1)
-		{
-			NewResourceDumpContext->UploadResourceCompressionName = NAME_Zlib;
-		}
-		else if (GDumpGPUUploadCompressResources.GetValueOnGameThread() == 2)
-		{
-			NewResourceDumpContext->UploadResourceCompressionName = NAME_Gzip;
-		}
+		NewResourceDumpContext = GNextRDGResourceDumpContext;
+		GNextRDGResourceDumpContext = nullptr;
 	}
-
-	NewResourceDumpContext->bShowInExplore = NewResourceDumpContext->bEnableDiskWrite && GDumpExploreCVar.GetValueOnGameThread() != 0 && !NewResourceDumpContext->bUpload;
+	
 	NewResourceDumpContext->MemoryConstants = FPlatformMemory::GetConstants();
 	NewResourceDumpContext->MemoryStats = FPlatformMemory::GetStats();
 
@@ -1825,6 +1850,17 @@ void FRDGBuilder::EndResourceDump()
 		DumpingFrameCounter_GameThread == MAX_uint64 ||
 		DumpingFrameCounter_GameThread >= GFrameCounter)
 	{
+		if (GNextDumpingRemainingTime > 0.0f)
+		{
+			check(GNextRDGResourceDumpContext);
+			GNextDumpingRemainingTime -= FApp::GetDeltaTime();
+			if (GNextDumpingRemainingTime <= 0.0)
+			{
+				GNextDumpingRemainingTime = -1.0f;
+				FRDGBuilder::BeginResourceDump(TEXT(""));
+				check(!GNextRDGResourceDumpContext);
+			}
+		}
 		return;
 	}
 
