Exclude more non-rendering apps from RHI runtime caps check (UE-168344).

diff --git a/Engine/Source/Runtime/RenderCore/Private/StereoRenderUtils.cpp b/Engine/Source/Runtime/RenderCore/Private/StereoRenderUtils.cpp
--- a/Engine/Source/Runtime/RenderCore/Private/StereoRenderUtils.cpp
+++ b/Engine/Source/Runtime/RenderCore/Private/StereoRenderUtils.cpp
@@ -60,8 +60,8 @@ RENDERCORE_API FStereoShaderAspects::FStereoShaderAspects(EShaderPlatform Platfo
 	const bool bMobilePlatform = IsMobilePlatform(Platform);
 	const bool bMobilePostprocessing = CVarMobileHDR.Get(Platform);
 	const bool bMobileMultiView = CVarMobileMultiView.Get(Platform);
-	// If we're a cooker, don't check GRHI* setting, as it reflects runtime RHI capabilities.
-	const bool bMultiViewportCapable = (GRHISupportsArrayIndexFromAnyShader || IsRunningCookCommandlet()) && RHISupportsMultiViewport(Platform);
+	// If we're in a non-rendering run (cooker, DDC commandlet, anything with -nullrhi), don't check GRHI* setting, as it reflects runtime RHI capabilities.
+	const bool bMultiViewportCapable = (GRHISupportsArrayIndexFromAnyShader || !FApp::CanEverRender()) && RHISupportsMultiViewport(Platform);
 
 	bInstancedStereoNative = !bMobilePlatform && bInstancedStereo && RHISupportsInstancedStereo(Platform);
 
@@ -88,7 +88,8 @@ RENDERCORE_API FStereoShaderAspects::FStereoShaderAspects(EShaderPlatform Platfo
 		else if (RHISupportsInstancedStereo(Platform))
 		{
 			// ISR in mobile shaders is achieved via layered RTs and eye-dependent layer index, so unlike on desktop, it does not depend on multi-viewport
-			const bool bVertexShaderLayerCapable = (GRHISupportsArrayIndexFromAnyShader || IsRunningCookCommandlet()) && RHISupportsVertexShaderLayer(Platform);
+			// If we're in a non-rendering run (cooker, DDC commandlet, anything with -nullrhi), don't check GRHI* setting, as it reflects runtime RHI capabilities.
+			const bool bVertexShaderLayerCapable = (GRHISupportsArrayIndexFromAnyShader || !FApp::CanEverRender()) && RHISupportsVertexShaderLayer(Platform);
 			UE_DEBUG_SSA_LOG_BOOL(bVertexShaderLayerCapable);
 
 			if (bVertexShaderLayerCapable)
