Implemented DrawClearQuadAlpha to clear only the alpha channel of a render target.

diff --git a/Engine/Source/Runtime/RenderCore/Private/ClearQuad.cpp b/Engine/Source/Runtime/RenderCore/Private/ClearQuad.cpp
--- a/Engine/Source/Runtime/RenderCore/Private/ClearQuad.cpp
+++ b/Engine/Source/Runtime/RenderCore/Private/ClearQuad.cpp
@@ -13,17 +13,18 @@ TGlobalResource<FClearVertexBuffer> GClearVertexBuffer;
 
 DEFINE_LOG_CATEGORY_STATIC(LogClearQuad, Log, Log)
 
-static void ClearQuadSetup( FRHICommandList& RHICmdList, bool bClearColor, int32 NumClearColors, const FLinearColor* ClearColorArray, bool bClearDepth, float Depth, bool bClearStencil, uint32 Stencil, TFunction<void(FGraphicsPipelineStateInitializer&)> PSOModifier = nullptr)
+template<EColorWriteMask ColorWriteMask>
+static void ClearQuadSetup(FRHICommandList& RHICmdList, int32 NumClearColors, const FLinearColor* ClearColorArray, bool bClearDepth, float Depth, bool bClearStencil, uint32 Stencil, TFunction<void(FGraphicsPipelineStateInitializer&)> PSOModifier = nullptr)
 {
 	if (UNLIKELY(!FApp::CanEverRender()))
 	{
 		return;
 	}
 
+	check(ColorWriteMask == CW_NONE || NumClearColors > 0);
+
 	// Set new states
-	FRHIBlendState* BlendStateRHI = bClearColor
-		? TStaticBlendState<>::GetRHI()
-		: TStaticBlendStateWriteMask<CW_NONE,CW_NONE,CW_NONE,CW_NONE,CW_NONE,CW_NONE,CW_NONE,CW_NONE>::GetRHI();
+	FRHIBlendState* BlendStateRHI = TStaticBlendStateWriteMask<ColorWriteMask, ColorWriteMask, ColorWriteMask, ColorWriteMask, ColorWriteMask, ColorWriteMask, ColorWriteMask, ColorWriteMask>::GetRHI();
 	
 	FRHIDepthStencilState* DepthStencilStateRHI =
 		(bClearDepth && bClearStencil)
@@ -79,9 +80,25 @@ static void ClearQuadSetup( FRHICommandList& RHICmdList, bool bClearColor, int32
 	PixelShader->SetColors(RHICmdList, PixelShader, ClearColorArray, NumClearColors);
 }
 
+void DrawClearQuadAlpha(FRHICommandList& RHICmdList, float Alpha)
+{
+	FLinearColor Color(0, 0, 0, Alpha);
+	ClearQuadSetup<CW_ALPHA>(RHICmdList, 1, &Color, false, 0, false, 0);
+
+	RHICmdList.SetStreamSource(0, GClearVertexBuffer.VertexBufferRHI, 0);
+	RHICmdList.DrawPrimitive(0, 2, 1);
+}
+
 void DrawClearQuadMRT(FRHICommandList& RHICmdList, bool bClearColor, int32 NumClearColors, const FLinearColor* ClearColorArray, bool bClearDepth, float Depth, bool bClearStencil, uint32 Stencil)
 {
-	ClearQuadSetup(RHICmdList, bClearColor, NumClearColors, ClearColorArray, bClearDepth, Depth, bClearStencil, Stencil);
+	if (bClearColor)
+	{
+		ClearQuadSetup<CW_RGBA>(RHICmdList, NumClearColors, ClearColorArray, bClearDepth, Depth, bClearStencil, Stencil);
+	}
+	else
+	{
+		ClearQuadSetup<CW_NONE>(RHICmdList, NumClearColors, ClearColorArray, bClearDepth, Depth, bClearStencil, Stencil);
+	}
 
 	RHICmdList.SetStreamSource(0, GClearVertexBuffer.VertexBufferRHI, 0);
 	RHICmdList.DrawPrimitive(0, 2, 1);
@@ -89,7 +106,14 @@ void DrawClearQuadMRT(FRHICommandList& RHICmdList, bool bClearColor, int32 NumCl
 
 void DrawClearQuadMRT(FRHICommandList& RHICmdList, bool bClearColor, int32 NumClearColors, const FLinearColor* ClearColorArray, bool bClearDepth, float Depth, bool bClearStencil, uint32 Stencil, FClearQuadCallbacks ClearQuadCallbacks)
 {
-	ClearQuadSetup(RHICmdList, bClearColor, NumClearColors, ClearColorArray, bClearDepth, Depth, bClearStencil, Stencil, ClearQuadCallbacks.PSOModifier);
+	if (bClearColor)
+	{
+		ClearQuadSetup<CW_RGBA>(RHICmdList, NumClearColors, ClearColorArray, bClearDepth, Depth, bClearStencil, Stencil, ClearQuadCallbacks.PSOModifier);
+	}
+	else
+	{
+		ClearQuadSetup<CW_NONE>(RHICmdList, NumClearColors, ClearColorArray, bClearDepth, Depth, bClearStencil, Stencil, ClearQuadCallbacks.PSOModifier);
+	}
 
 	if (ClearQuadCallbacks.PreClear)
 	{
@@ -114,7 +138,14 @@ void DrawClearQuadMRT(FRHICommandList& RHICmdList, bool bClearColor, int32 NumCl
 		return;
 	}
 
-	ClearQuadSetup(RHICmdList, bClearColor, NumClearColors, ClearColorArray, bClearDepth, Depth, bClearStencil, Stencil);
+	if (bClearColor)
+	{
+		ClearQuadSetup<CW_RGBA>(RHICmdList, NumClearColors, ClearColorArray, bClearDepth, Depth, bClearStencil, Stencil);
+	}
+	else
+	{
+		ClearQuadSetup<CW_NONE>(RHICmdList, NumClearColors, ClearColorArray, bClearDepth, Depth, bClearStencil, Stencil);
+	}
 
 	// Draw a fullscreen quad
 	if (ExcludeRect.Width() > 0 && ExcludeRect.Height() > 0)
diff --git a/Engine/Source/Runtime/RenderCore/Public/ClearQuad.h b/Engine/Source/Runtime/RenderCore/Public/ClearQuad.h
--- a/Engine/Source/Runtime/RenderCore/Public/ClearQuad.h
+++ b/Engine/Source/Runtime/RenderCore/Public/ClearQuad.h
@@ -52,6 +52,7 @@ struct FClearQuadCallbacks
 };
 
 extern RENDERCORE_API void ClearUAV(FRHICommandList& RHICmdList, FRHIUnorderedAccessView* Buffer, uint32 NumBytes, uint32 Value, bool bBarriers = true);
+extern RENDERCORE_API void DrawClearQuadAlpha(FRHICommandList& RHICmdList, float Alpha);
 extern RENDERCORE_API void DrawClearQuadMRT(FRHICommandList& RHICmdList, bool bClearColor, int32 NumClearColors, const FLinearColor* ClearColorArray, bool bClearDepth, float Depth, bool bClearStencil, uint32 Stencil);
 extern RENDERCORE_API void DrawClearQuadMRT(FRHICommandList& RHICmdList, bool bClearColor, int32 NumClearColors, const FLinearColor* ClearColorArray, bool bClearDepth, float Depth, bool bClearStencil, uint32 Stencil, FClearQuadCallbacks ClearQuadCallbacks);
 extern RENDERCORE_API void DrawClearQuadMRT(FRHICommandList& RHICmdList, bool bClearColor, int32 NumClearColors, const FLinearColor* ClearColorArray, bool bClearDepth, float Depth, bool bClearStencil, uint32 Stencil, FIntPoint ViewSize, FIntRect ExcludeRect);
