Fixes unity build in DynamicRenderScaling.h

diff --git a/Engine/Source/Runtime/RenderCore/Public/DynamicRenderScaling.h b/Engine/Source/Runtime/RenderCore/Public/DynamicRenderScaling.h
--- a/Engine/Source/Runtime/RenderCore/Public/DynamicRenderScaling.h
+++ b/Engine/Source/Runtime/RenderCore/Public/DynamicRenderScaling.h
@@ -9,6 +9,7 @@
 #include "CoreTypes.h"
 #include "HAL/IConsoleManager.h"
 #include "Misc/AssertionMacros.h"
+#include "Stats/Stats.h"
 
 namespace DynamicRenderScaling
 {
