Adding ISR & Multiview Logging at runtime initialization::
Moved MMV exit code to it's own function

diff --git a/Engine/Source/Runtime/Launch/Private/LaunchEngineLoop.cpp b/Engine/Source/Runtime/Launch/Private/LaunchEngineLoop.cpp
--- a/Engine/Source/Runtime/Launch/Private/LaunchEngineLoop.cpp
+++ b/Engine/Source/Runtime/Launch/Private/LaunchEngineLoop.cpp
@@ -6485,29 +6485,11 @@ void FEngineLoop::PostInitRHI()
 	{
 		// perform an early check of hardware capabilities
 		EShaderPlatform ShaderPlatform = GMaxRHIShaderPlatform;
-		const UE::StereoRenderUtils::FStereoShaderAspects Aspects(ShaderPlatform);
 
-		// If instanced stereo is enabled, we should also have either multiviewport enabled, or mmv fallback enabled.
-		// Otherwise, exit gracefully with a message box
-		if (Aspects.IsInstancedStereoEnabled() && !Aspects.IsInstancedMultiViewportEnabled() && !Aspects.IsMobileMultiViewEnabled())
-		{
-			UE_LOG(LogInit, Log, TEXT("ShaderPlatform=%d RHISupportsInstancedStereo()=%d GRHISupportsArrayIndexFromAnyShader=%d"),
-				ShaderPlatform, RHISupportsMultiViewport(ShaderPlatform), GRHISupportsArrayIndexFromAnyShader);
+		const UE::StereoRenderUtils::FStereoShaderAspects Aspects(ShaderPlatform);
 
-			// MessageBoxExt may not yet handle unattended runs itself (see UE-165694), so special-case here to avoid getting stuck in a commandlet with rendering enabled
-			const FText MessageText = NSLOCTEXT("InstancedStereo", "UnableToUseInstancedStereoRenderingText", "Cannot render an Instanced Stereo-enabled project due to a missing functionality on the system. Please check log files for more info.");
-			if (!FApp::IsUnattended())
-			{
-				FPlatformMisc::MessageBoxExt(EAppMsgType::Ok, *MessageText.ToString(),
-					*NSLOCTEXT("InstancedStereo", "UnableToUseInstancedStereoRendering", "Unable to use Instanced Stereo Rendering.").ToString());
-			}
-			else
-			{
-				UE_LOG(LogInit, Error, TEXT("%s"), *MessageText.ToString());					
-			}
-			FPlatformMisc::RequestExitWithStatus(true, 1);
-			// unreachable
-		}
+		UE::StereoRenderUtils::LogISRInit(Aspects);
+		UE::StereoRenderUtils::VerifyISRConfig(Aspects, ShaderPlatform);
 	}
 
 #endif
diff --git a/Engine/Source/Runtime/RenderCore/Private/StereoRenderUtils.cpp b/Engine/Source/Runtime/RenderCore/Private/StereoRenderUtils.cpp
--- a/Engine/Source/Runtime/RenderCore/Private/StereoRenderUtils.cpp
+++ b/Engine/Source/Runtime/RenderCore/Private/StereoRenderUtils.cpp
@@ -125,4 +125,37 @@ RENDERCORE_API FStereoShaderAspects::FStereoShaderAspects(EShaderPlatform Platfo
 	checkf(!bInstancedStereoEnabled || (bInstancedMultiViewportEnabled || bMobileMultiViewFallback), TEXT("If ISR is enabled, we need either multi-viewport (since we no longer support clip-distance method) or MMV fallback (which uses vertex layer)."));
 }
 	
-} // namespace UE::RenderUtils
+
+RENDERCORE_API void LogISRInit(const UE::StereoRenderUtils::FStereoShaderAspects& Aspects)
+{
+	UE_LOG(LogInit, Log, TEXT("XR: Instanced Stereo Rendering is %s"), (Aspects.IsInstancedStereoEnabled() ? TEXT("Enabled") : TEXT("Disabled")));
+	UE_LOG(LogInit, Log, TEXT("XR: MultiViewport is %s"), (Aspects.IsInstancedMultiViewportEnabled() ? TEXT("Enabled") : TEXT("Disabled")));
+	UE_LOG(LogInit, Log, TEXT("XR: Mobile Multiview is %s"), (Aspects.IsMobileMultiViewEnabled() ? TEXT("Enabled") : TEXT("Disabled")));
+}
+
+RENDERCORE_API void VerifyISRConfig(const UE::StereoRenderUtils::FStereoShaderAspects& Aspects, EShaderPlatform ShaderPlatform)
+{
+	// If instanced stereo is enabled, we should also have either multiviewport enabled, or mmv fallback enabled.
+	// Otherwise, exit gracefully with a message box
+	if (Aspects.IsInstancedStereoEnabled() && !Aspects.IsInstancedMultiViewportEnabled() && !Aspects.IsMobileMultiViewEnabled())
+	{
+		UE_LOG(LogInit, Log, TEXT("ShaderPlatform=%d RHISupportsInstancedStereo()=%d GRHISupportsArrayIndexFromAnyShader=%d"),
+			ShaderPlatform, RHISupportsMultiViewport(ShaderPlatform), GRHISupportsArrayIndexFromAnyShader);
+
+		// MessageBoxExt may not yet handle unattended runs itself (see UE-165694), so special-case here to avoid getting stuck in a commandlet with rendering enabled
+		const FText MessageText = NSLOCTEXT("InstancedStereo", "UnableToUseInstancedStereoRenderingText", "Cannot render an Instanced Stereo-enabled project due to a missing functionality on the system. Please check log files for more info.");
+		if (!FApp::IsUnattended())
+		{
+			FPlatformMisc::MessageBoxExt(EAppMsgType::Ok, *MessageText.ToString(),
+				*NSLOCTEXT("InstancedStereo", "UnableToUseInstancedStereoRendering", "Unable to use Instanced Stereo Rendering.").ToString());
+		}
+		else
+		{
+			UE_LOG(LogInit, Error, TEXT("%s"), *MessageText.ToString());
+		}
+		FPlatformMisc::RequestExitWithStatus(true, 1);
+		// unreachable
+	}
+}
+
+} // namespace UE::StereoRenderUtils
diff --git a/Engine/Source/Runtime/RenderCore/Public/StereoRenderUtils.h b/Engine/Source/Runtime/RenderCore/Public/StereoRenderUtils.h
--- a/Engine/Source/Runtime/RenderCore/Public/StereoRenderUtils.h
+++ b/Engine/Source/Runtime/RenderCore/Public/StereoRenderUtils.h
@@ -7,45 +7,50 @@
 namespace UE::StereoRenderUtils
 {
 
-	/*
-	 * Detect the single-draw stereo shader variant, in order to support usage across different platforms
-	 */
-	class FStereoShaderAspects
-	{
-	public:
-		/**
-		* Determines the stereo aspects of the shader pipeline based on the input shader platform
-		* @param Platform	Target shader platform used to determine stereo shader variant
-		*/
-		RENDERCORE_API FStereoShaderAspects(EShaderPlatform Platform);
-
-		/**
-		 * Whether instanced stereo rendering is enabled - i.e. using a single instanced drawcall to render to both stereo views.
-		 * The output is redirected via the viewport index.
-		 */
-		inline bool IsInstancedStereoEnabled() const { return bInstancedStereoEnabled; }
-
-		/**
-		 * Whether mobile multiview is enabled - i.e. using VK_KHR_multiview. Another drawcall reduction technique, independent of instanced stereo.
-		 * Mobile multiview generates view indices to index into texture arrays.
-		 * Can be internally emulated using instanced stereo if native support is unavailable, by using ISR-generated view indices to index into texture arrays.
-		 */
-		inline bool IsMobileMultiViewEnabled() const { return bMobileMultiViewEnabled; };
-
-		/**
-		 * Whether multiviewport rendering is enabled - i.e. using ViewportIndex to index into viewport.
-		 * Relies on instanced stereo rendering being enabled.
-		 */
-		inline bool IsInstancedMultiViewportEnabled() const { return bInstancedMultiViewportEnabled; };
-
-	private:
-		bool bInstancedStereoEnabled : 1;
-		bool bMobileMultiViewEnabled : 1;
-		bool bInstancedMultiViewportEnabled : 1;
-
-		bool bInstancedStereoNative : 1;
-		bool bMobileMultiViewNative : 1;
-		bool bMobileMultiViewFallback : 1;
-	};
-
-}
+/*
+	* Detect the single-draw stereo shader variant, in order to support usage across different platforms
+	*/
+class FStereoShaderAspects
+{
+public:
+	/**
+	* Determines the stereo aspects of the shader pipeline based on the input shader platform
+	* @param Platform	Target shader platform used to determine stereo shader variant
+	*/
+	RENDERCORE_API FStereoShaderAspects(EShaderPlatform Platform);
+
+	/**
+	* Whether instanced stereo rendering is enabled - i.e. using a single instanced drawcall to render to both stereo views.
+	* The output is redirected via the viewport index.
+	*/
+	inline bool IsInstancedStereoEnabled() const { return bInstancedStereoEnabled; }
+
+	/**
+	* Whether mobile multiview is enabled - i.e. using VK_KHR_multiview. Another drawcall reduction technique, independent of instanced stereo.
+	* Mobile multiview generates view indices to index into texture arrays.
+	* Can be internally emulated using instanced stereo if native support is unavailable, by using ISR-generated view indices to index into texture arrays.
+	*/
+	inline bool IsMobileMultiViewEnabled() const { return bMobileMultiViewEnabled; };
+
+	/**
+	* Whether multiviewport rendering is enabled - i.e. using ViewportIndex to index into viewport.
+	* Relies on instanced stereo rendering being enabled.
+	*/
+	inline bool IsInstancedMultiViewportEnabled() const { return bInstancedMultiViewportEnabled; };
+
+private:
+	bool bInstancedStereoEnabled : 1;
+	bool bMobileMultiViewEnabled : 1;
+	bool bInstancedMultiViewportEnabled : 1;
+
+	bool bInstancedStereoNative : 1;
+	bool bMobileMultiViewNative : 1;
+	bool bMobileMultiViewFallback : 1;
+
+};
+
+RENDERCORE_API void LogISRInit(const UE::StereoRenderUtils::FStereoShaderAspects& Aspects);
+
+RENDERCORE_API void VerifyISRConfig(const UE::StereoRenderUtils::FStereoShaderAspects& Aspects, EShaderPlatform ShaderPlatform);
+
+} // namespace UE::StereoRenderUtils
