Enabling non-PQ >8 bit material for H265 in Electra

diff --git a/Engine/Plugins/Media/ElectraUtil/Source/ElectraSamples/Private/ElectraTextureSamplePC.cpp b/Engine/Plugins/Media/ElectraUtil/Source/ElectraSamples/Private/ElectraTextureSamplePC.cpp
--- a/Engine/Plugins/Media/ElectraUtil/Source/ElectraSamples/Private/ElectraTextureSamplePC.cpp
+++ b/Engine/Plugins/Media/ElectraUtil/Source/ElectraSamples/Private/ElectraTextureSamplePC.cpp
@@ -52,7 +52,8 @@ void FElectraTextureSample::Initialize(FVideoDecoderOutput *InVideoDecoderOutput
 
 	// Prepare YUV -> RGB matrix containing all necessary offsets and scales to produce RGB straight from sample data
 	const FMatrix* Mtx = bFullRange ? &MediaShaders::YuvToRgbRec709Unscaled : &MediaShaders::YuvToRgbRec709Scaled;
-	FVector Off = bFullRange ? MediaShaders::YUVOffsetNoScale8bits : MediaShaders::YUVOffset8bits;
+	FVector Off = (SampleFormat == EMediaTextureSampleFormat::P010) ? (bFullRange ? MediaShaders::YUVOffsetNoScale16bits : MediaShaders::YUVOffset16bits)
+																	: (bFullRange ? MediaShaders::YUVOffsetNoScale8bits : MediaShaders::YUVOffset8bits);
 	float Scale = 1.0f;
 
 	if (auto PinnedHDRInfo = HDRInfo.Pin())
diff --git a/Engine/Shaders/Private/MediaShaders.usf b/Engine/Shaders/Private/MediaShaders.usf
--- a/Engine/Shaders/Private/MediaShaders.usf
+++ b/Engine/Shaders/Private/MediaShaders.usf
@@ -352,7 +352,7 @@ void P010ConvertPS(
 	Color.z = dot(P010ConvertUB.ColorTransform[2], YUV4);
 
 	// Go linear
-	Color = ST2084ToLinear(Color);
+	Color = P010ConvertUB.IsST2084 ? ST2084ToLinear(Color) : Rec709ToLinear(Color);
 
 	// Adjust colorspace
 	float4 Col4 = float4(Color, 1);
@@ -394,7 +394,7 @@ void P010ConvertAsUINT16sPS(
 	Color.z = dot(P010ConvertAsUINT16sUB.ColorTransform[2], YUV4);
 
 	// Go linear
-	Color = ST2084ToLinear(Color);
+	Color = P010ConvertAsUINT16sUB.IsST2084 ? ST2084ToLinear(Color) : Rec709ToLinear(Color);
 
 	// Adjust colorspace
 	float4 Col4 = float4(Color, 1);
@@ -439,7 +439,7 @@ void P010_2101010ConvertPS(
 	Color.z = dot(P010_2101010ConvertUB.ColorTransform[2], YUV4);
 
 	// Go linear
-	Color = ST2084ToLinear(Color);
+	Color = P010_2101010ConvertUB.IsST2084 ? ST2084ToLinear(Color) : Rec709ToLinear(Color);
 
 	// Adjust colorspace
 	float4 Col4 = float4(Color, 1);
@@ -451,7 +451,7 @@ void P010_2101010ConvertPS(
 	OutColor = float4(Color, 1.0f); 
 #else
 	OutColor = float4(1,0,1,1);
-#endif // COMPILER_PSSL
+#endif // COMPILER_GLSL_ES3_1
 }
 
 
diff --git a/Engine/Source/Runtime/MediaAssets/Private/Misc/MediaTextureResource.cpp b/Engine/Source/Runtime/MediaAssets/Private/Misc/MediaTextureResource.cpp
--- a/Engine/Source/Runtime/MediaAssets/Private/Misc/MediaTextureResource.cpp
+++ b/Engine/Source/Runtime/MediaAssets/Private/Misc/MediaTextureResource.cpp
@@ -1027,10 +1027,6 @@ void FMediaTextureResource::ConvertSample(const TSharedPtr<IMediaTextureSample,
 						ColorSpaceMtx = ColorSpaceMtx.ApplyScale(1.0f / 80.0f);
 					}
 
-					// We for now hard code this in the shaders!
-					// (skipping this assert as we might get dummy frames passing through here, that are not flagging this correctly)
-					//check(Sample->GetEncodingType() == UE::Color::EEncoding::ST2084);
-
 					FIntPoint TexDim = InputTexture->GetSizeXY();
 					if (InputTexture->GetFormat() == PF_P010)
 					{
@@ -1040,7 +1036,7 @@ void FMediaTextureResource::ConvertSample(const TSharedPtr<IMediaTextureSample,
 
 						FShaderResourceViewRHIRef Y_SRV = RHICreateShaderResourceView(InputTexture, 0, 1, PF_G16);
 						FShaderResourceViewRHIRef UV_SRV = RHICreateShaderResourceView(InputTexture, 0, 1, PF_G16R16);
-						ConvertShader->SetParameters(CommandList, TexDim, Y_SRV, UV_SRV, OutputDim, YUVMtx, ColorSpaceMtx, bIsSampleOutputSrgb);
+						ConvertShader->SetParameters(CommandList, TexDim, Y_SRV, UV_SRV, OutputDim, YUVMtx, ColorSpaceMtx, Sample->GetEncodingType() == UE::Color::EEncoding::ST2084);
 					}
 					else
 					{
@@ -1048,7 +1044,7 @@ void FMediaTextureResource::ConvertSample(const TSharedPtr<IMediaTextureSample,
 						GraphicsPSOInit.BoundShaderState.PixelShaderRHI = ConvertShader.GetPixelShader();
 						SetGraphicsPipelineState(CommandList, GraphicsPSOInit, 0);
 
-						ConvertShader->SetParameters(CommandList, TexDim, InputTexture, OutputDim, YUVMtx, ColorSpaceMtx, bIsSampleOutputSrgb);
+						ConvertShader->SetParameters(CommandList, TexDim, InputTexture, OutputDim, YUVMtx, ColorSpaceMtx, Sample->GetEncodingType() == UE::Color::EEncoding::ST2084);
 					}
 				}
 				break;
diff --git a/Engine/Source/Runtime/RenderCore/Private/MediaShaders.cpp b/Engine/Source/Runtime/RenderCore/Private/MediaShaders.cpp
--- a/Engine/Source/Runtime/RenderCore/Private/MediaShaders.cpp
+++ b/Engine/Source/Runtime/RenderCore/Private/MediaShaders.cpp
@@ -276,7 +276,7 @@ void FNV21ConvertPS::SetParameters(FRHICommandList& CommandList, TRefCountPtr<FR
 BEGIN_GLOBAL_SHADER_PARAMETER_STRUCT(FP010ConvertUB, )
 SHADER_PARAMETER(FMatrix44f, ColorTransform)
 SHADER_PARAMETER(FMatrix44f, CSTransform)
-SHADER_PARAMETER(uint32, SrgbToLinear)
+SHADER_PARAMETER(uint32, IsST2084)
 SHADER_PARAMETER(FVector2f, UVScale)
 SHADER_PARAMETER_SRV(Texture2D, SRV_Y)
 SHADER_PARAMETER_SRV(Texture2D, SRV_UV)
@@ -287,7 +287,7 @@ END_GLOBAL_SHADER_PARAMETER_STRUCT()
 IMPLEMENT_GLOBAL_SHADER_PARAMETER_STRUCT(FP010ConvertUB, "P010ConvertUB");
 IMPLEMENT_SHADER_TYPE(, FP010ConvertPS, TEXT("/Engine/Private/MediaShaders.usf"), TEXT("P010ConvertPS"), SF_Pixel);
 
-void FP010ConvertPS::SetParameters(FRHICommandList& CommandList, const FIntPoint& TexDim, FShaderResourceViewRHIRef SRV_Y, FShaderResourceViewRHIRef SRV_UV, const FIntPoint& OutputDimensions, const FMatrix44f& ColorTransform, const FMatrix44f& CSTransform, bool SrgbToLinear)
+void FP010ConvertPS::SetParameters(FRHICommandList& CommandList, const FIntPoint& TexDim, FShaderResourceViewRHIRef SRV_Y, FShaderResourceViewRHIRef SRV_UV, const FIntPoint& OutputDimensions, const FMatrix44f& ColorTransform, const FMatrix44f& CSTransform, bool bIsST2084)
 {
 	FP010ConvertUB UB;
 	{
@@ -295,7 +295,7 @@ void FP010ConvertPS::SetParameters(FRHICommandList& CommandList, const FIntPoint
 		UB.CSTransform = CSTransform;
 		UB.SamplerB = TStaticSamplerState<SF_Bilinear>::GetRHI();
 		UB.SamplerP = TStaticSamplerState<SF_Point>::GetRHI();
-		UB.SrgbToLinear = SrgbToLinear;
+		UB.IsST2084 = bIsST2084;
 		UB.SRV_Y = SRV_Y;
 		UB.SRV_UV = SRV_UV;
 		UB.UVScale = FVector2f((float)OutputDimensions.X / (float)TexDim.X, (float)OutputDimensions.Y / (float)TexDim.Y);
@@ -312,7 +312,7 @@ BEGIN_GLOBAL_SHADER_PARAMETER_STRUCT(FP010ConvertAsUINT16sUB, )
 SHADER_PARAMETER(FMatrix44f, ColorTransform)
 SHADER_PARAMETER(FMatrix44f, CSTransform)
 SHADER_PARAMETER(uint32, OutputWidth)
-SHADER_PARAMETER(uint32, SrgbToLinear)
+SHADER_PARAMETER(uint32, IsST2084)
 SHADER_PARAMETER(FVector2f, UVScale)
 SHADER_PARAMETER_TEXTURE(Texture2D, Texture)
 SHADER_PARAMETER_SAMPLER(SamplerState, SamplerB)
@@ -322,7 +322,7 @@ END_GLOBAL_SHADER_PARAMETER_STRUCT()
 IMPLEMENT_GLOBAL_SHADER_PARAMETER_STRUCT(FP010ConvertAsUINT16sUB, "P010ConvertAsUINT16sUB");
 IMPLEMENT_SHADER_TYPE(, FP010ConvertAsUINT16sPS, TEXT("/Engine/Private/MediaShaders.usf"), TEXT("P010ConvertAsUINT16sPS"), SF_Pixel);
 
-void FP010ConvertAsUINT16sPS::SetParameters(FRHICommandList& CommandList, const FIntPoint& TexDim, TRefCountPtr<FRHITexture2D> NV12Texture, const FIntPoint& OutputDimensions, const FMatrix44f& ColorTransform, const FMatrix44f& CSTransform, bool SrgbToLinear)
+void FP010ConvertAsUINT16sPS::SetParameters(FRHICommandList& CommandList, const FIntPoint& TexDim, TRefCountPtr<FRHITexture2D> NV12Texture, const FIntPoint& OutputDimensions, const FMatrix44f& ColorTransform, const FMatrix44f& CSTransform, bool bIsST2084)
 {
 	FP010ConvertAsUINT16sUB UB;
 	{
@@ -331,7 +331,7 @@ void FP010ConvertAsUINT16sPS::SetParameters(FRHICommandList& CommandList, const
 		UB.OutputWidth = OutputDimensions.X;
 		UB.SamplerB = TStaticSamplerState<SF_Bilinear>::GetRHI();
 		UB.SamplerP = TStaticSamplerState<SF_Point>::GetRHI();
-		UB.SrgbToLinear = SrgbToLinear;
+		UB.IsST2084 = bIsST2084;
 		UB.Texture = NV12Texture;
 		UB.UVScale = FVector2f((float)OutputDimensions.X / (float)TexDim.X, (float)OutputDimensions.Y / (float)TexDim.Y);
 	}
@@ -346,7 +346,7 @@ void FP010ConvertAsUINT16sPS::SetParameters(FRHICommandList& CommandList, const
 BEGIN_GLOBAL_SHADER_PARAMETER_STRUCT(FP010_2101010ConvertUB, )
 SHADER_PARAMETER(FMatrix44f, ColorTransform)
 SHADER_PARAMETER(FMatrix44f, CSTransform)
-SHADER_PARAMETER(uint32, SrgbToLinear)
+SHADER_PARAMETER(uint32, IsST2084)
 SHADER_PARAMETER(FVector2f, UVScaleY)
 SHADER_PARAMETER(FVector2f, UVScaleUV)
 SHADER_PARAMETER_SRV(Texture2D, SRV_Y)
@@ -360,14 +360,14 @@ END_GLOBAL_SHADER_PARAMETER_STRUCT()
 IMPLEMENT_GLOBAL_SHADER_PARAMETER_STRUCT(FP010_2101010ConvertUB, "P010_2101010ConvertUB");
 IMPLEMENT_SHADER_TYPE(, FP010_2101010ConvertPS, TEXT("/Engine/Private/MediaShaders.usf"), TEXT("P010_2101010ConvertPS"), SF_Pixel);
 
-void FP010_2101010ConvertPS::SetParameters(FRHICommandList& CommandList, const FIntPoint& TexDim, FShaderResourceViewRHIRef SRV_Y, FShaderResourceViewRHIRef SRV_U, FShaderResourceViewRHIRef SRV_V, const FIntPoint& OutputDimensions, const FMatrix44f& ColorTransform, const FMatrix44f& CSTransform, bool SrgbToLinear)
+void FP010_2101010ConvertPS::SetParameters(FRHICommandList& CommandList, const FIntPoint& TexDim, FShaderResourceViewRHIRef SRV_Y, FShaderResourceViewRHIRef SRV_U, FShaderResourceViewRHIRef SRV_V, const FIntPoint& OutputDimensions, const FMatrix44f& ColorTransform, const FMatrix44f& CSTransform, bool bIsST2084)
 {
 	FP010_2101010ConvertUB UB;
 	{
 		UB.ColorTransform = ColorTransform;
 		UB.CSTransform = CSTransform;
 		UB.SamplerP = TStaticSamplerState<SF_Point>::GetRHI();
-		UB.SrgbToLinear = SrgbToLinear;
+		UB.IsST2084 = bIsST2084;
 		UB.SRV_Y = SRV_Y;
 		UB.SRV_U = SRV_U;
 		UB.SRV_V = SRV_V;
diff --git a/Engine/Source/Runtime/RenderCore/Public/MediaShaders.h b/Engine/Source/Runtime/RenderCore/Public/MediaShaders.h
--- a/Engine/Source/Runtime/RenderCore/Public/MediaShaders.h
+++ b/Engine/Source/Runtime/RenderCore/Public/MediaShaders.h
@@ -338,7 +338,7 @@ public:
 		: FGlobalShader(Initializer)
 	{ }
 
-	RENDERCORE_API void SetParameters(FRHICommandList& RHICmdList, const FIntPoint& TexDim, FShaderResourceViewRHIRef SRV_Y, FShaderResourceViewRHIRef SRV_UV, const FIntPoint& OutputDimensions, const FMatrix44f& ColorTransform, const FMatrix44f& CSTransform, bool SrgbToLinear);
+	RENDERCORE_API void SetParameters(FRHICommandList& RHICmdList, const FIntPoint& TexDim, FShaderResourceViewRHIRef SRV_Y, FShaderResourceViewRHIRef SRV_UV, const FIntPoint& OutputDimensions, const FMatrix44f& ColorTransform, const FMatrix44f& CSTransform, bool bIsST2084);
 };
 
 /**
@@ -362,7 +362,7 @@ public:
 		: FGlobalShader(Initializer)
 	{ }
 
-	RENDERCORE_API void SetParameters(FRHICommandList& RHICmdList, const FIntPoint& TexDim, TRefCountPtr<FRHITexture2D> NV12Texture, const FIntPoint& OutputDimensions, const FMatrix44f& ColorTransform, const FMatrix44f& CSTransform, bool SrgbToLinear);
+	RENDERCORE_API void SetParameters(FRHICommandList& RHICmdList, const FIntPoint& TexDim, TRefCountPtr<FRHITexture2D> NV12Texture, const FIntPoint& OutputDimensions, const FMatrix44f& ColorTransform, const FMatrix44f& CSTransform, bool bIsST2084);
 };
 
 /**
@@ -386,7 +386,7 @@ public:
 		: FGlobalShader(Initializer)
 	{ }
 
-	RENDERCORE_API void SetParameters(FRHICommandList& RHICmdList, const FIntPoint& TexDim, FShaderResourceViewRHIRef SRV_Y, FShaderResourceViewRHIRef SRV_U, FShaderResourceViewRHIRef SRV_V, const FIntPoint& OutputDimensions, const FMatrix44f& ColorTransform, const FMatrix44f& CSTransform, bool SrgbToLinear);
+	RENDERCORE_API void SetParameters(FRHICommandList& RHICmdList, const FIntPoint& TexDim, FShaderResourceViewRHIRef SRV_Y, FShaderResourceViewRHIRef SRV_U, FShaderResourceViewRHIRef SRV_V, const FIntPoint& OutputDimensions, const FMatrix44f& ColorTransform, const FMatrix44f& CSTransform, bool bIsST2084);
 };
 
 /**
