Use a RW lock to protect GShaderFileCache instead of a critical section to reduce lock contention

diff --git a/Engine/Source/Runtime/RenderCore/Private/ShaderCore.cpp b/Engine/Source/Runtime/RenderCore/Private/ShaderCore.cpp
--- a/Engine/Source/Runtime/RenderCore/Private/ShaderCore.cpp
+++ b/Engine/Source/Runtime/RenderCore/Private/ShaderCore.cpp
@@ -102,7 +102,7 @@ DEFINE_STAT(STAT_Shaders_NumShadersRegistered);
 DEFINE_STAT(STAT_Shaders_NumShadersDuplicated);
 
 /** Protects GShaderFileCache from simultaneous access by multiple threads. */
-FCriticalSection FileCacheCriticalSection;
+FRWLock FileCacheRWLock;
 
 /** The shader file cache, used to minimize shader file reads */
 TMap<FString, FString> GShaderFileCache;
@@ -945,6 +945,7 @@ bool LoadShaderSourceFile(const TCHAR* InVirtualFilePath, EShaderPlatform Shader
 	bool bResult = false;
 
 	STAT(double ShaderFileLoadingTime = 0);
+
 	{
 		SCOPE_SECONDS_COUNTER(ShaderFileLoadingTime);
 
@@ -956,12 +957,15 @@ bool LoadShaderSourceFile(const TCHAR* InVirtualFilePath, EShaderPlatform Shader
 		// Fixup autogen file
 		ReplaceVirtualFilePathForShaderAutogen(VirtualFilePath, ShaderPlatform, ShaderPlatformName);
 
-		// Protect GShaderFileCache from simultaneous access by multiple threads
-		FScopeLock ScopeLock(&FileCacheCriticalSection);
+		FString* CachedFile = nullptr;
 
-		FString* CachedFile = GShaderFileCache.Find(VirtualFilePath);
+		// First try a shared lock and only acquire exclusive access if element is not found in cache
+		
+		{
+			FRWScopeLock ScopeLock(FileCacheRWLock, SLT_ReadOnly);
+			CachedFile = GShaderFileCache.Find(VirtualFilePath);
+		}
 
-		//if this file has already been loaded and cached, use that
 		if (CachedFile)
 		{
 			if (OutFileContents)
@@ -970,25 +974,43 @@ bool LoadShaderSourceFile(const TCHAR* InVirtualFilePath, EShaderPlatform Shader
 			}
 			bResult = true;
 		}
-		else
+		else 
 		{
-			FString ShaderFilePath = GetShaderSourceFilePath(VirtualFilePath, OutCompileErrors);
+			FRWScopeLock ScopeLock(FileCacheRWLock, SLT_Write);
 
-			// verify SHA hash of shader files on load. missing entries trigger an error
-			FString FileContents;
-			if (!ShaderFilePath.IsEmpty() && FFileHelper::LoadFileToString(FileContents, *ShaderFilePath, FFileHelper::EHashOptions::EnableVerify|FFileHelper::EHashOptions::ErrorMissingHash) )
-			{
-				//update the shader file cache
-				GShaderFileCache.Add(VirtualFilePath, FileContents);
+			// Double-check the cache while holding exclusive lock as another thread may have added the item we're looking for
+			CachedFile = GShaderFileCache.Find(VirtualFilePath);
 
+			// if this file has already been loaded and cached, use that
+			if (CachedFile)
+			{
 				if (OutFileContents)
 				{
-					*OutFileContents = MoveTemp(FileContents);
+					*OutFileContents = *CachedFile;
 				}
 				bResult = true;
 			}
+			else
+			{
+				FString ShaderFilePath = GetShaderSourceFilePath(VirtualFilePath, OutCompileErrors);
+
+				// verify SHA hash of shader files on load. missing entries trigger an error
+				FString FileContents;
+				if (!ShaderFilePath.IsEmpty() && FFileHelper::LoadFileToString(FileContents, *ShaderFilePath, FFileHelper::EHashOptions::EnableVerify|FFileHelper::EHashOptions::ErrorMissingHash) )
+				{
+					//update the shader file cache
+					GShaderFileCache.Add(VirtualFilePath, FileContents);
+
+					if (OutFileContents)
+					{
+						*OutFileContents = MoveTemp(FileContents);
+					}
+					bResult = true;
+				}
+			}
 		}
 	}
+
 	INC_FLOAT_STAT_BY(STAT_ShaderCompiling_LoadingShaderFiles,(float)ShaderFileLoadingTime);
 
 	return bResult;
@@ -1466,7 +1488,7 @@ void FlushShaderFileCache(const FName* ShaderPlatformName)
 		GShaderHashCache.Empty();
 	}
 	{
-		FScopeLock ScopeLock(&FileCacheCriticalSection);
+		FRWScopeLock ScopeLock(FileCacheRWLock, SLT_Write);
 		GShaderFileCache.Empty();
 	}
 
