Use a RW lock to protect GShaderHashCache instead of a critical section to reduce lock contention

A critical section is still used on cache miss to prevent redundant calculations, while the cache remains accessible for readers.

diff --git a/Engine/Source/Runtime/RenderCore/Private/ShaderCore.cpp b/Engine/Source/Runtime/RenderCore/Private/ShaderCore.cpp
--- a/Engine/Source/Runtime/RenderCore/Private/ShaderCore.cpp
+++ b/Engine/Source/Runtime/RenderCore/Private/ShaderCore.cpp
@@ -168,7 +168,7 @@ public:
 		Platforms[PreviewShaderPlatform].IncludeDirectory = Platforms[ParentShaderPlatform].IncludeDirectory;
 	}
 
-	FSHAHash* FindHash(EShaderPlatform ShaderPlatform, const FString& VirtualFilePath)
+	const FSHAHash* FindHash(EShaderPlatform ShaderPlatform, const FString& VirtualFilePath) const
 	{
 		check(ShaderPlatform < UE_ARRAY_COUNT(Platforms));
 		checkf(bInitialized, TEXT("GShaderHashCache::Initialize needs to be called before GShaderHashCache::FindHash."));
@@ -184,7 +184,7 @@ public:
 		return Platforms[ShaderPlatform].ShaderHashCache.Add(VirtualFilePath, FSHAHash());
 	}
 
-	bool ShouldIgnoreInclude(const FString& VirtualFilePath, EShaderPlatform ShaderPlatform)
+	bool ShouldIgnoreInclude(const FString& VirtualFilePath, EShaderPlatform ShaderPlatform) const
 	{
 		// Ignore only platform specific files, which won't be used by the target platform.
 		if (VirtualFilePath.StartsWith(TEXT("/Engine/Private/Platform/"))
@@ -209,7 +209,7 @@ public:
 		}
 	}
 
-	const FString& GetPlatformIncludeDirectory(EShaderPlatform ShaderPlatform)
+	const FString& GetPlatformIncludeDirectory(EShaderPlatform ShaderPlatform) const
 	{
 		check(ShaderPlatform < EShaderPlatform::SP_NumPlatforms);
 		checkf(bInitialized, TEXT("GShaderHashCache::Initialize needs to be called before GShaderHashCache::GetPlatformIncludeDirectory."));
@@ -235,7 +235,7 @@ private:
 };
 
 /** Protects GShaderHashCache from simultaneous modification by multiple threads. Note that it can cover more than one method of the class, e.g. a block of code doing Find() then Add() can be guarded */
-FCriticalSection GShaderHashAccessGuard;
+FRWLock GShaderHashAccessRWLock;
 
 FShaderHashCache GShaderHashCache;
 
@@ -856,7 +856,7 @@ bool ReplaceVirtualFilePathForShaderPlatform(FString& InOutVirtualFilePath, ESha
 	// as of 2021-03-01, it'd be safe to access just the include directory without the lock... but the lock (and copy) is here for the consistency's and future-proofness' sake
 	const FString PlatformIncludeDirectory( [ShaderPlatform]()
 		{
-			FScopeLock ShaderHashAccessLock(&GShaderHashAccessGuard);
+			FRWScopeLock ShaderHashAccessLock(GShaderHashAccessRWLock, SLT_ReadOnly);
 			return GShaderHashCache.GetPlatformIncludeDirectory(ShaderPlatform);
 		}()
 	);
@@ -1139,7 +1139,7 @@ static void InternalGetShaderIncludes(const TCHAR* EntryPointVirtualFilePath, co
 
 					// Include only platform specific files, which will be used by the target platform.
 					{
-						FScopeLock ShaderHashAccessLock(&GShaderHashAccessGuard);
+						FRWScopeLock ShaderHashAccessLock(GShaderHashAccessRWLock, SLT_ReadOnly);
 						bIgnoreInclude = bIgnoreInclude || GShaderHashCache.ShouldIgnoreInclude(ExtractedIncludeFilename, ShaderPlatform);
 					}
 
@@ -1207,13 +1207,13 @@ void HashShaderFileWithIncludes(FArchive& HashingArchive, const TCHAR* VirtualFi
 		// first, a "soft" check
 		bool bFoundInCache = false;
 		{
-			FScopeLock ShaderHashAccessLock(&GShaderHashAccessGuard);
-			FSHAHash* CachedHash = GShaderHashCache.FindHash(ShaderPlatform, VirtualFilePath);
+			FRWScopeLock ShaderHashAccessLock(GShaderHashAccessRWLock, SLT_ReadOnly);
+			const FSHAHash* CachedHash = GShaderHashCache.FindHash(ShaderPlatform, VirtualFilePath);
 			// If a hash for this filename has been cached, use that
 			if (CachedHash)
 			{
 				bFoundInCache = true;
-				HashingArchive << *CachedHash;
+				HashingArchive << const_cast<FSHAHash&>(*CachedHash);
 			}
 		}
 
@@ -1285,6 +1285,12 @@ static void UpdateSingleShaderFilehash(FSHA1& InOutHashState, const TCHAR* Virtu
 	InOutHashState.UpdateWithString(*FileContents, FileContents.Len());
 }
 
+/** 
+* Prevents multiple threads from trying to redundantly call UpdateSingleShaderFilehash in GetShaderFileHash / GetShaderFilesHash.
+* Must be used in conjunction with GShaderHashAccessRWLock, which protects actual GShaderHashCache operations.
+*/
+static FCriticalSection GShaderFileHashCalculationGuard;
+
 /**
  * Calculates a Hash for the given filename and its includes if it does not already exist in the Hash cache.
  * @param Filename - shader file to Hash
@@ -1295,13 +1301,27 @@ const FSHAHash& GetShaderFileHash(const TCHAR* VirtualFilePath, EShaderPlatform
 	// Make sure we are only accessing GShaderHashCache from one thread
 	//check(IsInGameThread() || IsAsyncLoading());
 	STAT(double HashTime = 0);
-	FScopeLock ShaderHashAccessLock(&GShaderHashAccessGuard);
 	{
 		SCOPE_SECONDS_COUNTER(HashTime);
 
-		FSHAHash* CachedHash = GShaderHashCache.FindHash(ShaderPlatform, VirtualFilePath);
+		{
+			FRWScopeLock ShaderHashAccessLock(GShaderHashAccessRWLock, SLT_ReadOnly);
+			const FSHAHash* CachedHash = GShaderHashCache.FindHash(ShaderPlatform, VirtualFilePath);
+			// If a hash for this filename has been cached, use that
+			if (CachedHash)
+			{
+				return *CachedHash;
+			}
+		}
+
+		// We don't want UpdateSingleShaderFilehash to be called redundantly from multiple threads,
+		// while minimiziong GShaderHashAccessRWLock exclusive lock time.
+		// We can use a dedicated critical section around the hash calculation and cache update, 
+		// while keeping the cache itself available for reading.
+		FScopeLock FileHashCalculationAccessLock(&GShaderFileHashCalculationGuard);
 
-		// If a hash for this filename has been cached, use that
+		// Double-check the cache while holding exclusive lock as another thread may have added the item we're looking for
+		const FSHAHash* CachedHash = GShaderHashCache.FindHash(ShaderPlatform, VirtualFilePath);
 		if (CachedHash)
 		{
 			return *CachedHash;
@@ -1312,6 +1332,7 @@ const FSHAHash& GetShaderFileHash(const TCHAR* VirtualFilePath, EShaderPlatform
 		HashState.Final();
 
 		// Update the hash cache
+		FRWScopeLock ShaderHashAccessLock(GShaderHashAccessRWLock, SLT_Write);
 		FSHAHash& NewHash = GShaderHashCache.AddHash(ShaderPlatform, VirtualFilePath);
 		HashState.GetHash(&NewHash.Hash[0]);
 
@@ -1333,7 +1354,7 @@ const FSHAHash& GetShaderFilesHash(const TArray<FString>& VirtualFilePaths, ESha
 	// Make sure we are only accessing GShaderHashCache from one thread
 	//check(IsInGameThread() || IsAsyncLoading());
 	STAT(double HashTime = 0);
-	FScopeLock ShaderHashAccessLock(&GShaderHashAccessGuard);
+
 	{
 		SCOPE_SECONDS_COUNTER(HashTime);
 
@@ -1343,9 +1364,25 @@ const FSHAHash& GetShaderFilesHash(const TArray<FString>& VirtualFilePaths, ESha
 			Key += Filename;
 		}
 
-		FSHAHash* CachedHash = GShaderHashCache.FindHash(ShaderPlatform, Key);
+		{
+			FRWScopeLock ShaderHashAccessLock(GShaderHashAccessRWLock, SLT_ReadOnly);
+			const FSHAHash* CachedHash = GShaderHashCache.FindHash(ShaderPlatform, Key);
+
+			// If a hash for this filename has been cached, use that
+			if (CachedHash)
+			{
+				return *CachedHash;
+			}
+		}
+
+		// We don't want UpdateSingleShaderFilehash to be called redundantly from multiple threads,
+		// while minimiziong GShaderHashAccessRWLock exclusive lock time.
+		// We can use a dedicated critical section around the hash calculation and cache update, 
+		// while keeping the cache itself available for reading.
+		FScopeLock FileHashCalculationAccessLock(&GShaderFileHashCalculationGuard);
 
-		// If a hash for this filename has been cached, use that
+		// Double-check the cache while holding exclusive lock as another thread may have added the item we're looking for
+		const FSHAHash* CachedHash = GShaderHashCache.FindHash(ShaderPlatform, Key);
 		if (CachedHash)
 		{
 			return *CachedHash;
@@ -1359,6 +1396,7 @@ const FSHAHash& GetShaderFilesHash(const TArray<FString>& VirtualFilePaths, ESha
 		HashState.Final();
 
 		// Update the hash cache
+		FRWScopeLock ShaderHashAccessLock(GShaderHashAccessRWLock, SLT_Write);
 		FSHAHash& NewHash = GShaderHashCache.AddHash(ShaderPlatform, Key);
 		HashState.GetHash(&NewHash.Hash[0]);
 
@@ -1425,19 +1463,19 @@ void BuildShaderFileToUniformBufferMap(TMap<FString, TArray<const TCHAR*> >& Sha
 
 void InitializeShaderHashCache()
 {
-	FScopeLock ShaderHashAccessLock(&GShaderHashAccessGuard);
+	FRWScopeLock ShaderHashAccessLock(GShaderHashAccessRWLock, SLT_Write);
 	GShaderHashCache.Initialize();
 }
 
 void UpdateIncludeDirectoryForPreviewPlatform(EShaderPlatform PreviewPlatform, EShaderPlatform ActualPlatform)
 {
-	FScopeLock ShaderHashAccessLock(&GShaderHashAccessGuard);
+	FRWScopeLock ShaderHashAccessLock(GShaderHashAccessRWLock, SLT_Write);
 	GShaderHashCache.UpdateIncludeDirectoryForPreviewPlatform(PreviewPlatform, ActualPlatform);
 }
 
 void CheckShaderHashCacheInclude(const FString& VirtualFilePath, EShaderPlatform ShaderPlatform)
 {
-	FScopeLock ShaderHashAccessLock(&GShaderHashAccessGuard);
+	FRWScopeLock ShaderHashAccessLock(GShaderHashAccessRWLock, SLT_ReadOnly);
 	bool bIgnoreInclude = GShaderHashCache.ShouldIgnoreInclude(VirtualFilePath, ShaderPlatform);
 
 	checkf(!bIgnoreInclude,
@@ -1484,7 +1522,7 @@ void FlushShaderFileCache(const FName* ShaderPlatformName)
 	UE_LOG(LogShaders, Log, TEXT("FlushShaderFileCache() begin"));
 
 	{
-		FScopeLock ShaderHashAccessLock(&GShaderHashAccessGuard);
+		FRWScopeLock ShaderHashAccessLock(GShaderHashAccessRWLock, SLT_Write);
 		GShaderHashCache.Empty();
 	}
 	{
