Improve performance of RHI GPU draw call / num primitive stats
 - Old code relied on atomic increments in the RHI_DRAW_CALL_INC /  RHI_DRAW_CALL_STATS macros, which is expensive, particularly on platforms with poor atomic performance.
 - New system replaces the atomic writes with a context-specific stats structure, which is accumulated by the RHI thread into the global structure.
 - Contexts write stat data through a "Stats" pointer on the IRHIComputeContext, which is set automatically by the command list management code. The pointer is replaced whenever a new "draw call category" is pushed, to redirect the counts.
 - Also moved some macro definitions around so more of the system can be removed when HAS_GPU_STATS is 0. Removed dependencies on CSV_PROFILER.

Moved Begin/EndFrame, Begin/EndScene, Begin/EndDrawingViewport into the immediate RHICmdList
 - They were already immediate-only functions due to a check() they contain.

diff --git a/Engine/Source/Runtime/D3D12RHI/Private/D3D12CommandContext.h b/Engine/Source/Runtime/D3D12RHI/Private/D3D12CommandContext.h
--- a/Engine/Source/Runtime/D3D12RHI/Private/D3D12CommandContext.h
+++ b/Engine/Source/Runtime/D3D12RHI/Private/D3D12CommandContext.h
@@ -1087,6 +1087,14 @@ public:
 		return PhysicalContexts[GPUIndex];
 	}
 
+	virtual void StatsSetCategory(FRHIDrawStats* InStats, uint32 InCategoryID) final override
+	{
+		for (uint32 GPUIndex : PhysicalGPUMask)
+		{
+			PhysicalContexts[GPUIndex]->StatsSetCategory(InStats, InCategoryID, GPUIndex);
+		}
+	}
+
 private:
 	TStaticArray<FD3D12CommandContext*, MAX_NUM_GPUS> PhysicalContexts;
 };
diff --git a/Engine/Source/Runtime/D3D12RHI/Private/D3D12Commands.cpp b/Engine/Source/Runtime/D3D12RHI/Private/D3D12Commands.cpp
--- a/Engine/Source/Runtime/D3D12RHI/Private/D3D12Commands.cpp
+++ b/Engine/Source/Runtime/D3D12RHI/Private/D3D12Commands.cpp
@@ -1862,7 +1862,7 @@ void FD3D12CommandContext::CommitComputeResourceTables()
 
 void FD3D12CommandContext::RHIDrawPrimitive(uint32 BaseVertexIndex, uint32 NumPrimitives, uint32 NumInstances)
 {
-	RHI_DRAW_CALL_STATS_MGPU(GetGPUIndex(), StateCache.GetGraphicsPipelinePrimitiveType(), FMath::Max(NumInstances, 1U) * NumPrimitives);
+	RHI_DRAW_CALL_STATS(StateCache.GetGraphicsPipelinePrimitiveType(), FMath::Max(NumInstances, 1U) * NumPrimitives);
 
 	CommitGraphicsResourceTables();
 	CommitNonComputeShaderConstants();
@@ -1887,7 +1887,7 @@ void FD3D12CommandContext::RHIDrawPrimitiveIndirect(FRHIBuffer* ArgumentBufferRH
 {
 	FD3D12Buffer* ArgumentBuffer = RetrieveObject<FD3D12Buffer>(ArgumentBufferRHI);
 
-	RHI_DRAW_CALL_INC_MGPU(GetGPUIndex());
+	RHI_DRAW_CALL_INC();
 	if (bTrackingEvents)
 	{
 		GetParentDevice()->RegisterGPUWork(0);
@@ -1929,7 +1929,7 @@ void FD3D12CommandContext::RHIDrawIndexedIndirect(FRHIBuffer* IndexBufferRHI, FR
 	FD3D12Buffer* IndexBuffer = RetrieveObject<FD3D12Buffer>(IndexBufferRHI);
 	FD3D12Buffer* ArgumentsBuffer = RetrieveObject<FD3D12Buffer>(ArgumentsBufferRHI);
 
-	RHI_DRAW_CALL_INC_MGPU(GetGPUIndex());
+	RHI_DRAW_CALL_INC();
 	if (bTrackingEvents)
 	{
 		GetParentDevice()->RegisterGPUWork(1);
@@ -1972,7 +1972,7 @@ void FD3D12CommandContext::RHIDrawIndexedPrimitive(FRHIBuffer* IndexBufferRHI, i
 {
 	// called should make sure the input is valid, this avoid hidden bugs
 	ensure(NumPrimitives > 0);
-	RHI_DRAW_CALL_STATS_MGPU(GetGPUIndex(), StateCache.GetGraphicsPipelinePrimitiveType(), FMath::Max(NumInstances, 1U) * NumPrimitives);
+	RHI_DRAW_CALL_STATS(StateCache.GetGraphicsPipelinePrimitiveType(), FMath::Max(NumInstances, 1U) * NumPrimitives);
 
 	NumInstances = FMath::Max<uint32>(1, NumInstances);
 
@@ -2011,7 +2011,7 @@ void FD3D12CommandContext::RHIDrawIndexedPrimitiveIndirect(FRHIBuffer* IndexBuff
 	FD3D12Buffer* IndexBuffer = RetrieveObject<FD3D12Buffer>(IndexBufferRHI);
 	FD3D12Buffer* ArgumentBuffer = RetrieveObject<FD3D12Buffer>(ArgumentBufferRHI);
 
-	RHI_DRAW_CALL_INC_MGPU(GetGPUIndex());
+	RHI_DRAW_CALL_INC();
 	if (bTrackingEvents)
 	{
 		GetParentDevice()->RegisterGPUWork(0);
@@ -2073,7 +2073,7 @@ void FD3D12CommandContext::RHIDispatchIndirectMeshShader(FRHIBuffer* ArgumentBuf
 {
 	FD3D12Buffer* ArgumentBuffer = RetrieveObject<FD3D12Buffer>(ArgumentBufferRHI);
 
-	RHI_DRAW_CALL_INC_MGPU(GetGPUIndex());
+	RHI_DRAW_CALL_INC();
 	if (bTrackingEvents)
 	{
 		GetParentDevice()->RegisterGPUWork(0);
diff --git a/Engine/Source/Runtime/Engine/Private/UnrealEngine.cpp b/Engine/Source/Runtime/Engine/Private/UnrealEngine.cpp
--- a/Engine/Source/Runtime/Engine/Private/UnrealEngine.cpp
+++ b/Engine/Source/Runtime/Engine/Private/UnrealEngine.cpp
@@ -17638,7 +17638,7 @@ int32 UEngine::RenderStatUnit(UWorld* World, FViewport* Viewport, FCanvas* Canva
 // DRAWCOUNT
 int32 UEngine::RenderStatDrawCount(UWorld* World, FViewport* Viewport, FCanvas* Canvas, int32 X, int32 Y, const FVector* ViewLocation, const FRotator* ViewRotation)
 {
-#if CSV_PROFILER
+#if HAS_GPU_STATS
 	int32 TotalCount[MAX_NUM_GPUS] = { 0 };
 	// Display all the categories of draw counts. This may always report 0 in some modes if AreGPUStatsEnabled is not enabled.
 	// Most likely because we are not currently capturing a CSV.
@@ -17675,7 +17675,7 @@ int32 UEngine::RenderStatDrawCount(UWorld* World, FViewport* Viewport, FCanvas*
 #else
 	Canvas->DrawShadowedString(X - 200,
 		Y,
-		*FString::Printf(TEXT("DrawCount only supported with CSV_PROFILER")),
+		*FString::Printf(TEXT("DrawCount only supported with HAS_GPU_STATS")),
 		GetSmallFont(),
 		FColor::Red);
 #endif
diff --git a/Engine/Source/Runtime/RHI/Private/RHI.cpp b/Engine/Source/Runtime/RHI/Private/RHI.cpp
--- a/Engine/Source/Runtime/RHI/Private/RHI.cpp
+++ b/Engine/Source/Runtime/RHI/Private/RHI.cpp
@@ -34,9 +34,11 @@ CSV_DEFINE_CATEGORY(DrawCall, true);
 #endif
 
 // Define counter stats.
-DEFINE_STAT(STAT_RHIDrawPrimitiveCalls);
-DEFINE_STAT(STAT_RHITriangles);
-DEFINE_STAT(STAT_RHILines);
+#if HAS_GPU_STATS
+	DEFINE_STAT(STAT_RHIDrawPrimitiveCalls);
+	DEFINE_STAT(STAT_RHITriangles);
+	DEFINE_STAT(STAT_RHILines);
+#endif
 
 // Define memory stats.
 DEFINE_STAT(STAT_RenderTargetMemory2D);
@@ -296,9 +298,13 @@ std::atomic<TClosableMpscQueue<FRHIResource*>*> FRHIResource::PendingDeletes { n
 FHazardPointerCollection FRHIResource::PendingDeletesHPC;
 FRHIResource* FRHIResource::CurrentlyDeleting = nullptr;
 
-RHI_API FDrawCallCategoryName* FDrawCallCategoryName::Array[FDrawCallCategoryName::MAX_DRAWCALL_CATEGORY];
-RHI_API int32 FDrawCallCategoryName::DisplayCounts[FDrawCallCategoryName::MAX_DRAWCALL_CATEGORY][MAX_NUM_GPUS];
-RHI_API int32 FDrawCallCategoryName::NumCategory = 0;
+#if HAS_GPU_STATS
+
+	TStaticArray<FDrawCallCategoryName*, FDrawCallCategoryName::MAX_DRAWCALL_CATEGORY> FDrawCallCategoryName::Array { InPlace, nullptr };
+	TStaticArray<TStaticArray<int32, MAX_NUM_GPUS>, FDrawCallCategoryName::MAX_DRAWCALL_CATEGORY> FDrawCallCategoryName::DisplayCounts { InPlace };
+	int32 FDrawCallCategoryName::NumCategory = 0;
+
+#endif
 
 TRefCountPtr<FRHITexture> FRHITextureReference::DefaultTexture;
 
@@ -1591,11 +1597,8 @@ RHI_API EShaderPlatform GShaderPlatformForFeatureLevel[ERHIFeatureLevel::Num] =
 
 // simple stats about draw calls. GNum is the previous frame and 
 // GCurrent is the current frame.
-// GCurrentNumDrawCallsRHIPtr points to the drawcall counter to increment
-RHI_API int32 GCurrentNumDrawCallsRHI[MAX_NUM_GPUS] = {};
-RHI_API int32 GNumDrawCallsRHI[MAX_NUM_GPUS] = {};
-thread_local FRHIDrawCallsStatPtr GCurrentNumDrawCallsRHIPtr = &GCurrentNumDrawCallsRHI;
-RHI_API int32 GCurrentNumPrimitivesDrawnRHI[MAX_NUM_GPUS] = {};
+
+RHI_API int32 GNumDrawCallsRHI      [MAX_NUM_GPUS] = {};
 RHI_API int32 GNumPrimitivesDrawnRHI[MAX_NUM_GPUS] = {};
 
 RHI_API uint64 GRHITransitionPrivateData_SizeInBytes = 0;
@@ -1607,57 +1610,116 @@ ERHIAccess GRHIMergeableAccessMask = ERHIAccess::ReadOnlyMask | ERHIAccess::UAVM
 // By default, only exclusively read only accesses are allowed.
 ERHIAccess GRHIMultiPipelineMergeableAccessMask = ERHIAccess::ReadOnlyExclusiveMask;
 
-/** Called once per frame only from within an RHI. */
-void RHIPrivateBeginFrame()
+RHI_API void FRHIDrawStats::Accumulate(FRHIDrawStats& Other)
 {
-	for (int32 GPUIndex = 0; GPUIndex < MAX_NUM_GPUS; GPUIndex++)
+	for (uint32 GPUIndex = 0; GPUIndex < GNumExplicitGPUsForRendering; ++GPUIndex)
 	{
-		GNumDrawCallsRHI[GPUIndex] = GCurrentNumDrawCallsRHI[GPUIndex];
+		FPerGPUStats& LeftGPU = GPUs[GPUIndex];
+		FPerGPUStats& RightGPU = Other.GPUs[GPUIndex];
+
+		for (int32 CategoryIndex = 0; CategoryIndex < NumCategories; ++CategoryIndex)
+		{
+			FPerCategoryStats& LeftCategory = LeftGPU.Categories[CategoryIndex];
+			FPerCategoryStats& RightCategory = RightGPU.Categories[CategoryIndex];
+
+			LeftCategory.Draws += RightCategory.Draws;
+
+			for (uint32 PrimitiveType = 0; PrimitiveType < PT_Num; ++PrimitiveType)
+			{
+				LeftCategory.Primitives[PrimitiveType] += RightCategory.Primitives[PrimitiveType];
+			}
+		}
 	}
-	
-#if CSV_PROFILER
-	// Only copy the display counters every so many frames to keep things more stable.
-	const int32 FramesUntilDisplayCopy = 30;
-	static int32 FrameCount = 0;
+}
+
+// Called from RHIBeginFrame
+void FRHICommandListImmediate::ProcessStats()
+{
+#if HAS_GPU_STATS
+	// Only copy the display counters every half second keep things more stable.
+	constexpr float TimeoutSeconds = 0.5;
+
+	static double LastTime = 0.0;
+	double CurrentTime = FPlatformTime::Seconds();
+
 	bool bCopyDisplayFrames = false;
-	++FrameCount;
-	if (FrameCount >= FramesUntilDisplayCopy)
+	if (CurrentTime - LastTime > TimeoutSeconds)
 	{
+		LastTime = CurrentTime;
 		bCopyDisplayFrames = true;
-		FrameCount = 0;
 	}
+#endif
+
+	// Summed stats across all GPUs
+	uint32 TotalDraws     = 0;
+	uint32 TotalTriangles = 0;
+	uint32 TotalLines     = 0;
+	TStaticArray<uint32, FRHIDrawStats::NumCategories> TotalDrawsPerCategory { InPlace, 0 };
 
-	for (int32 Index=0; Index<FDrawCallCategoryName::NumCategory; ++Index)
+	for (int32 GPUIndex = 0; GPUIndex < MAX_NUM_GPUS; ++GPUIndex)
 	{
-		FDrawCallCategoryName* CategoryName = FDrawCallCategoryName::Array[Index];
-		for (int32 GPUIndex = 0; GPUIndex < MAX_NUM_GPUS; GPUIndex++)
+		uint32 TotalDraws_ThisGPU      = 0;
+		uint32 TotalPrimitives_ThisGPU = 0;
+
+		FRHIDrawStats::FPerGPUStats& GPUStats = FrameDrawStats.GPUs[GPUIndex];
+
+		for (int32 CategoryIndex = 0; CategoryIndex < FRHIDrawStats::NumCategories; ++CategoryIndex)
 		{
-			if (bCopyDisplayFrames)
+			FRHIDrawStats::FPerCategoryStats& Category = GPUStats.Categories[CategoryIndex];
+
+			TotalDraws                           += Category.Draws;
+			TotalDraws_ThisGPU                   += Category.Draws;
+			TotalDrawsPerCategory[CategoryIndex] += Category.Draws;
+
+			for (uint32 PrimitiveType = 0; PrimitiveType < PT_Num; ++PrimitiveType)
 			{
-				FDrawCallCategoryName::DisplayCounts[Index][GPUIndex] = CategoryName->Counters[GPUIndex];
+				uint32 PrimitiveCount = Category.Primitives[PrimitiveType];
+				TotalPrimitives_ThisGPU += PrimitiveCount;
+
+				switch (PrimitiveType)
+				{
+				case PT_TriangleList:
+				case PT_TriangleStrip:
+					TotalTriangles += PrimitiveCount;
+					break;
+
+				case PT_LineList:
+					TotalLines += PrimitiveCount;
+					break;
+				}
 			}
-			GNumDrawCallsRHI[GPUIndex] += CategoryName->Counters[GPUIndex];
-		}
-		// Multi-GPU support : CSV stats do not support MGPU yet
-		FCsvProfiler::RecordCustomStat(CategoryName->Name, CSV_CATEGORY_INDEX(DrawCall), CategoryName->Counters[0], ECsvCustomStatOp::Set);
-		for (int32 GPUIndex = 0; GPUIndex < MAX_NUM_GPUS; GPUIndex++)
-		{
-			CategoryName->Counters[GPUIndex] = 0;
+
+#if HAS_GPU_STATS
+			if (bCopyDisplayFrames && CategoryIndex < FDrawCallCategoryName::NumCategory)
+			{
+				FDrawCallCategoryName::DisplayCounts[CategoryIndex][GPUIndex] = Category.Draws;
+			}
+#endif // HAS_GPU_STATS
 		}
+		
+		GNumDrawCallsRHI      [GPUIndex] = TotalDraws_ThisGPU;
+		GNumPrimitivesDrawnRHI[GPUIndex] = TotalPrimitives_ThisGPU;
 	}
-#endif
 
-	for (int32 GPUIndex = 0; GPUIndex < MAX_NUM_GPUS; GPUIndex++)
-	{
-		GNumPrimitivesDrawnRHI[GPUIndex] = GCurrentNumPrimitivesDrawnRHI[GPUIndex];
-	}
-	// Multi-GPU support : CSV stats do not support MGPU yet
-	CSV_CUSTOM_STAT(RHI, DrawCalls, GNumDrawCallsRHI[0], ECsvCustomStatOp::Set);
-	CSV_CUSTOM_STAT(RHI, PrimitivesDrawn, GNumPrimitivesDrawnRHI[0], ECsvCustomStatOp::Set);
-	for (int32 GPUIndex = 0; GPUIndex < MAX_NUM_GPUS; GPUIndex++)
+	// Multi-GPU support : CSV stats do not support MGPU yet. We're summing the totals across all GPUs here.
+	CSV_CUSTOM_STAT(RHI, DrawCalls      , int32(TotalDraws                 ), ECsvCustomStatOp::Set);
+	CSV_CUSTOM_STAT(RHI, PrimitivesDrawn, int32(TotalTriangles + TotalLines), ECsvCustomStatOp::Set);
+
+#if HAS_GPU_STATS
+	SET_DWORD_STAT(STAT_RHITriangles         , TotalTriangles);
+	SET_DWORD_STAT(STAT_RHILines             , TotalLines    );
+	SET_DWORD_STAT(STAT_RHIDrawPrimitiveCalls, TotalDraws    );
+
+	#if CSV_PROFILER
+	for (int32 CategoryIndex = 0; CategoryIndex < FDrawCallCategoryName::NumCategory; ++CategoryIndex)
 	{
-		GCurrentNumDrawCallsRHI[GPUIndex] = GCurrentNumPrimitivesDrawnRHI[GPUIndex] = 0;
+		FDrawCallCategoryName* CategoryName = FDrawCallCategoryName::Array[CategoryIndex];
+		FCsvProfiler::RecordCustomStat(CategoryName->Name, CSV_CATEGORY_INDEX(DrawCall), int32(TotalDrawsPerCategory[CategoryIndex]), ECsvCustomStatOp::Set);
 	}
+	#endif
+#endif // HAS_GPU_STATS
+
+	FrameDrawStats.Reset();
 }
 
 /** Whether to initialize 3D textures using a bulk data (or through a mip update if false). */
@@ -3134,13 +3196,3 @@ void FDebugName::AppendString(FStringBuilderBase& Builder) const
 		Builder << '_' << Number;
 	}
 }
-
-RHI_API void RHISetCurrentNumDrawCallPtr(FRHIDrawCallsStatPtr InNumDrawCallsRHIPtr)
-{
-	GCurrentNumDrawCallsRHIPtr = InNumDrawCallsRHIPtr;
-}
-
-RHI_API void RHIIncCurrentNumDrawCallPtr(uint32 GPUIndex)
-{
-	FPlatformAtomics::InterlockedIncrement(&(*GCurrentNumDrawCallsRHIPtr)[GPUIndex]);
-}
\ No newline at end of file
diff --git a/Engine/Source/Runtime/RHI/Private/RHICommandList.cpp b/Engine/Source/Runtime/RHI/Private/RHICommandList.cpp
--- a/Engine/Source/Runtime/RHI/Private/RHICommandList.cpp
+++ b/Engine/Source/Runtime/RHI/Private/RHICommandList.cpp
@@ -129,6 +129,7 @@ RHI_API FRHICommandListExecutor GRHICommandList;
 
 FGraphEventArray FRHICommandListImmediate::WaitOutstandingTasks;
 FGraphEventRef   FRHICommandListImmediate::RHIThreadTask;
+FRHIDrawStats    FRHICommandListImmediate::FrameDrawStats;
 
 static FGraphEventRef GRHIThreadEndDrawingViewportFences[2];
 static uint32 GRHIThreadEndDrawingViewportFenceIndex = 0;
@@ -307,6 +308,27 @@ void FRHICommandListBase::SetCurrentStat(TStatId Stat)
 	}
 }
 
+#if HAS_GPU_STATS
+void FRHICommandListBase::SetStatsCategory(FDrawCallCategoryName* Category)
+{
+	check(!Category || Category->ShouldCountDraws());
+
+	PersistentState.Stats.CategoryTOP = Category;
+	EnqueueLambda([Category](FRHICommandListBase& ExecutingCmdList)
+	{
+		ExecutingCmdList.PersistentState.Stats.CategoryBOP = Category;
+
+		for (IRHIComputeContext* Context : ExecutingCmdList.Contexts)
+		{
+			if (Context)
+			{
+				ExecutingCmdList.PersistentState.Stats.ApplyToContext(Context);
+			}
+		}
+	});
+}
+#endif
+
 ERHIPipeline FRHICommandListBase::SwitchPipeline(ERHIPipeline Pipeline)
 {
 	checkf(Pipeline == ERHIPipeline::None || FMath::IsPowerOfTwo((__underlying_type(ERHIPipeline))Pipeline), TEXT("Only one pipeline may be active at a time."));
@@ -367,6 +389,9 @@ ERHIPipeline FRHICommandListBase::SwitchPipeline(ERHIPipeline Pipeline)
 
 				// (Re-)apply the current GPU mask.
 				Context->RHISetGPUMask(ExecutingCmdList.PersistentState.CurrentGPUMask);
+
+				// Ensure the context has an up-to-date stat pointer.
+				ExecutingCmdList.PersistentState.Stats.ApplyToContext(Context);
 			}
 		});
 	}
@@ -374,13 +399,14 @@ ERHIPipeline FRHICommandListBase::SwitchPipeline(ERHIPipeline Pipeline)
 	return Pipeline;
 }
 
-void FRHICommandListBase::Execute(TRHIPipelineArray<IRHIComputeContext*>& InOutContexts)
+void FRHICommandListBase::Execute(TRHIPipelineArray<IRHIComputeContext*>& InOutContexts, FPersistentState::FGPUStats* ParentStats)
 {
 	check(!IsExecuting());
 	bExecuting = true;
 
 	Contexts = InOutContexts;
 	PersistentState.CurrentGPUMask = PersistentState.InitialGPUMask;
+	PersistentState.Stats.InitFrom(ParentStats);
 
 	ON_SCOPE_EXIT
 	{
@@ -476,6 +502,7 @@ void FRHICommandListImmediate::QueueAsyncCommandListSubmit(TArrayView<FQueuedCom
 			FGraphEventRef Event;
 			TArrayView<FRHICommandListBase*> InCmdLists;
 			TArray<IRHIPlatformCommandList*, TInlineAllocator<GetRHIPipelineCount()>> OutCmdLists;
+			FRHIDrawStats Stats;
 		};
 
 		uint32 NumTasks = 0;
@@ -557,7 +584,7 @@ void FRHICommandListImmediate::QueueAsyncCommandListSubmit(TArrayView<FQueuedCom
 
 			// Start a parallel translate task to replay the command list batch into the given pipeline contexts
 			Task.Event = FFunctionGraphTask::CreateAndDispatchWhenReady(
-				[&Task]()
+				[&Task, GPUStatsInitial = PersistentState.Stats]()
 				{
 					FOptionalTaskTagScope Scope(ETaskTag::EParallelRhiThread);
 					SCOPE_CYCLE_COUNTER(STAT_ParallelTranslate);
@@ -569,7 +596,11 @@ void FRHICommandListImmediate::QueueAsyncCommandListSubmit(TArrayView<FQueuedCom
 					// contexts depending on the SwitchPipeline commands that were recorded.
 					for (FRHICommandListBase* RHICmdList : Task.InCmdLists)
 					{
-						RHICmdList->Execute(Contexts);
+						// Redirect the output stats to this parallel task's copy
+						FPersistentState::FGPUStats Stats = GPUStatsInitial;
+						Stats.Ptr = &Task.Stats;
+
+						RHICmdList->Execute(Contexts, &Stats);
 						delete RHICmdList;
 					}
 
@@ -600,7 +631,7 @@ void FRHICommandListImmediate::QueueAsyncCommandListSubmit(TArrayView<FQueuedCom
 
 		// Finally, add an RHI thread task to submit the completed platform command lists.
 		// The task blocks for each parallel translate completion, in the order they will be submitted in.
-		EnqueueLambda([Tasks](FRHICommandListBase&)
+		EnqueueLambda([Tasks](FRHICommandListBase& ExecutingCmdList)
 		{
 			TArray<IRHIPlatformCommandList*> AllCmdLists;
 
@@ -615,6 +646,7 @@ void FRHICommandListImmediate::QueueAsyncCommandListSubmit(TArrayView<FQueuedCom
 				}
 
 				AllCmdLists.Append(Task.OutCmdLists);
+				ExecutingCmdList.PersistentState.Stats.Ptr->Accumulate(Task.Stats);
 
 				Task.~FTask();
 			}
@@ -642,7 +674,7 @@ void FRHICommandListImmediate::QueueAsyncCommandListSubmit(TArrayView<FQueuedCom
 			for (FRHICommandListBase* CmdList : CmdListsView)
 			{
 				CmdList->WaitForDispatchEvent();
-				CmdList->Execute(ParentCmdList.Contexts);
+				CmdList->Execute(ParentCmdList.Contexts, &ParentCmdList.PersistentState.Stats);
 				delete CmdList;
 			}
 		});
@@ -739,7 +771,7 @@ void FRHICommandListImmediate::ExecuteAndReset()
 					    FScopeLock Lock(&GRHIThreadOnTasksCritical);
 					    GWorkingRHIThreadStartCycles = FPlatformTime::Cycles();
     
-					    RHICmdList.Execute(RHICmdList.Contexts);
+					    RHICmdList.Execute(RHICmdList.Contexts, nullptr);
     
 					    GWorkingRHIThreadTime += (FPlatformTime::Cycles() - GWorkingRHIThreadStartCycles);
 				    }
@@ -760,7 +792,7 @@ void FRHICommandListImmediate::ExecuteAndReset()
 	    {
 		    // We're going to be executing the command list on the render thread.
 		    WaitForDispatchEvent();
-		    FRHICommandListBase::Execute(Contexts);
+		    FRHICommandListBase::Execute(Contexts, nullptr);
 	    }
 	}
 }
@@ -840,7 +872,7 @@ FRHICommandList_RecursiveHazardous::~FRHICommandList_RecursiveHazardous()
 
 	if (HasCommands())
 	{
-		Execute(Contexts);
+		Execute(Contexts, nullptr);
 	}
 }
 
@@ -868,7 +900,7 @@ FRHIComputeCommandList_RecursiveHazardous::~FRHIComputeCommandList_RecursiveHaza
 
 	if (HasCommands())
 	{
-		Execute(Contexts);
+		Execute(Contexts, nullptr);
 	}
 }
 	
@@ -1119,7 +1151,7 @@ void FRHICommandListImmediate::UnStallRHIThread()
 	}
 }
 
-void FRHICommandList::BeginScene()
+void FRHICommandListImmediate::BeginScene()
 {
 	check(IsImmediate() && IsInRenderingThread());
 	if (Bypass())
@@ -1137,7 +1169,7 @@ void FRHICommandList::BeginScene()
 	}
 }
 
-void FRHICommandList::EndScene()
+void FRHICommandListImmediate::EndScene()
 {
 	check(IsImmediate() && IsInRenderingThread());
 	if (Bypass())
@@ -1155,7 +1187,7 @@ void FRHICommandList::EndScene()
 	}
 }
 
-void FRHICommandList::BeginDrawingViewport(FRHIViewport* Viewport, FRHITexture* RenderTargetRHI)
+void FRHICommandListImmediate::BeginDrawingViewport(FRHIViewport* Viewport, FRHITexture* RenderTargetRHI)
 {
 	check(IsImmediate() && IsInRenderingThread());
 	if (Bypass())
@@ -1173,7 +1205,7 @@ void FRHICommandList::BeginDrawingViewport(FRHIViewport* Viewport, FRHITexture*
 	}
 }
 
-void FRHICommandList::EndDrawingViewport(FRHIViewport* Viewport, bool bPresent, bool bLockToVsync)
+void FRHICommandListImmediate::EndDrawingViewport(FRHIViewport* Viewport, bool bPresent, bool bLockToVsync)
 {
 	// Make sure all prior graphics and async compute work has been submitted.
 	// This is necessary because platform RHIs often submit additional work on the graphics queue during present, and we need to ensure we won't deadlock on async work that wasn't yet submitted by the renderer.
@@ -1217,12 +1249,12 @@ void FRHICommandList::EndDrawingViewport(FRHIViewport* Viewport, bool bPresent,
 	RHIAdvanceFrameForGetViewportBackBuffer(Viewport);
 }
 
-void FRHICommandList::BeginFrame()
+void FRHICommandListImmediate::BeginFrame()
 {
 	check(IsImmediate() && IsInRenderingThread());
 	if (Bypass())
 	{
-		RHIPrivateBeginFrame();
+		ProcessStats();
 		GetContext().RHIBeginFrame();
 		return;
 	}
@@ -1237,7 +1269,7 @@ void FRHICommandList::BeginFrame()
 
 }
 
-void FRHICommandList::EndFrame()
+void FRHICommandListImmediate::EndFrame()
 {
 	check(IsImmediate() && IsInRenderingThread());
 	if (Bypass())
diff --git a/Engine/Source/Runtime/RHI/Public/RHI.h b/Engine/Source/Runtime/RHI/Public/RHI.h
--- a/Engine/Source/Runtime/RHI/Public/RHI.h
+++ b/Engine/Source/Runtime/RHI/Public/RHI.h
@@ -16,6 +16,8 @@
 #include "Math/PerspectiveMatrix.h"
 #include "Math/ScaleMatrix.h"
 #include "Math/TranslationMatrix.h"
+#include "ProfilingDebugging/CsvProfilerConfig.h"
+#include "GpuProfilerTrace.h"
 
 #ifndef RHI_COMMAND_LIST_DEBUG_TRACES
 #define RHI_COMMAND_LIST_DEBUG_TRACES 0
@@ -638,18 +640,6 @@ extern RHI_API int32 GPoolSizeVRAMPercentage;
 /** Amount of local video memory demoted to system memory. In bytes. */
 extern RHI_API uint64 GDemotedLocalMemorySize;
 
-/** Some simple runtime stats, reset on every call to RHIBeginFrame */
-/** Num draw calls & primitives on previous frame (accurate on any thread)*/
-extern RHI_API int32 GNumDrawCallsRHI[MAX_NUM_GPUS];
-extern RHI_API int32 GNumPrimitivesDrawnRHI[MAX_NUM_GPUS];
-
-/** Num draw calls and primitives this frame (only accurate on RenderThread) */
-extern RHI_API int32 GCurrentNumDrawCallsRHI[MAX_NUM_GPUS];
-extern RHI_API int32 GCurrentNumPrimitivesDrawnRHI[MAX_NUM_GPUS];
-using FRHIDrawCallsStatPtr = int32(*)[MAX_NUM_GPUS];
-RHI_API void RHIIncCurrentNumDrawCallPtr(uint32 GPUIndex);
-RHI_API void RHISetCurrentNumDrawCallPtr(FRHIDrawCallsStatPtr InNumDrawCallsRHIPtr);
-
 /** Whether or not the RHI can handle a non-zero BaseVertexIndex - extra SetStreamSource calls will be needed if this is false */
 extern RHI_API bool GRHISupportsBaseVertexIndex;
 
@@ -872,10 +862,6 @@ inline bool RHIIsTypedUAVStoreSupported(EPixelFormat InFormat)
 */
 extern RHI_API SIZE_T CalculateImageBytes(uint32 SizeX,uint32 SizeY,uint32 SizeZ,uint8 Format);
 
-/** Called once per frame only from RHIBeginFrame's command execution. */
-void RHIPrivateBeginFrame();
-
-
 RHI_API FName LegacyShaderPlatformToShaderFormat(EShaderPlatform Platform);
 RHI_API EShaderPlatform ShaderFormatToLegacyShaderPlatform(FName ShaderFormat);
 RHI_API FName ShaderPlatformToPlatformName(EShaderPlatform Platform);
@@ -2666,68 +2652,105 @@ struct FTextureMemoryStats
 	}
 };
 
-struct RHI_API FDrawCallCategoryName
-{
-	FDrawCallCategoryName() 
-	{
-		for (int32& Counter : Counters)
-		{
-			Counter = -1;
-		}
-	}
+// GPU stats
+#ifndef HAS_GPU_STATS
+#define HAS_GPU_STATS ((STATS || CSV_PROFILER || GPUPROFILERTRACE_ENABLED) && (!UE_BUILD_SHIPPING))
+#endif
+
+extern RHI_API int32 GNumDrawCallsRHI      [MAX_NUM_GPUS];
+extern RHI_API int32 GNumPrimitivesDrawnRHI[MAX_NUM_GPUS];
+
+#if HAS_GPU_STATS
 
-	FDrawCallCategoryName(FName InName)
-		: Name(InName)
+	struct RHI_API FDrawCallCategoryName
 	{
-		for (int32& Counter : Counters)
-		{
-			Counter = 0;
-		}
+		FDrawCallCategoryName()
+			: Name(NAME_None)
+			, Index(-1)
+		{}
 
-		check(NumCategory < MAX_DRAWCALL_CATEGORY);
-		if (NumCategory < MAX_DRAWCALL_CATEGORY)
+		FDrawCallCategoryName(FName InName)
+			: Name(InName)
+			, Index(NumCategory++)
 		{
-			Array[NumCategory] = this;
-			NumCategory++;
+			check(Index < MAX_DRAWCALL_CATEGORY);
+			if (Index < MAX_DRAWCALL_CATEGORY)
+			{
+				Array[Index] = this;
+			}
 		}
-	}
 
-	FName Name;
-	int32 Counters[MAX_NUM_GPUS];
+		bool ShouldCountDraws() const { return Index != -1; }
 
-	static constexpr int32 MAX_DRAWCALL_CATEGORY = 256;
-	static FDrawCallCategoryName* Array[MAX_DRAWCALL_CATEGORY];
-	static int32 DisplayCounts[MAX_DRAWCALL_CATEGORY][MAX_NUM_GPUS]; // A backup of the counts that can be used to display on screen to avoid flickering.
-	static int32 NumCategory;
-};
+		FName  const Name;
+		uint32 const Index;
+
+		static constexpr int32 MAX_DRAWCALL_CATEGORY = 256;
+
+		static TStaticArray<FDrawCallCategoryName*, MAX_DRAWCALL_CATEGORY> Array;
+		static TStaticArray<TStaticArray<int32, MAX_NUM_GPUS>, MAX_DRAWCALL_CATEGORY> DisplayCounts; // A backup of the counts that can be used to display on screen to avoid flickering.
+		static int32 NumCategory;
+	};
+
+	// RHI counter stats.
+	DECLARE_DWORD_COUNTER_STAT_EXTERN(TEXT("DrawPrimitive calls"),STAT_RHIDrawPrimitiveCalls,STATGROUP_RHI,RHI_API);
+	DECLARE_DWORD_COUNTER_STAT_EXTERN(TEXT("Triangles drawn")    ,STAT_RHITriangles         ,STATGROUP_RHI,RHI_API);
+	DECLARE_DWORD_COUNTER_STAT_EXTERN(TEXT("Lines drawn")        ,STAT_RHILines             ,STATGROUP_RHI,RHI_API);
 
-// RHI counter stats.
-DECLARE_DWORD_COUNTER_STAT_EXTERN(TEXT("DrawPrimitive calls"),STAT_RHIDrawPrimitiveCalls,STATGROUP_RHI,RHI_API);
-DECLARE_DWORD_COUNTER_STAT_EXTERN(TEXT("Triangles drawn"),STAT_RHITriangles,STATGROUP_RHI,RHI_API);
-DECLARE_DWORD_COUNTER_STAT_EXTERN(TEXT("Lines drawn"),STAT_RHILines,STATGROUP_RHI,RHI_API);
-
-#if STATS
-#define RHI_DRAW_CALL_INC_MGPU(GPUIndex) \
-		INC_DWORD_STAT(STAT_RHIDrawPrimitiveCalls); \
-		RHIIncCurrentNumDrawCallPtr(GPUIndex);
-
-#define RHI_DRAW_CALL_STATS_MGPU(GPUIndex,PrimitiveType,NumPrimitives) \
-		RHI_DRAW_CALL_INC_MGPU(GPUIndex); \
-		INC_DWORD_STAT_BY(STAT_RHITriangles,(uint32)(PrimitiveType != PT_LineList ? (NumPrimitives) : 0)); \
-		INC_DWORD_STAT_BY(STAT_RHILines,(uint32)(PrimitiveType == PT_LineList ? (NumPrimitives) : 0)); \
-		FPlatformAtomics::InterlockedAdd(&GCurrentNumPrimitivesDrawnRHI[GPUIndex], NumPrimitives);
 #else
-#define RHI_DRAW_CALL_INC_MGPU(GPUIndex) \
-		RHIIncCurrentNumDrawCallPtr(GPUIndex);
 
-#define RHI_DRAW_CALL_STATS_MGPU(GPUIndex,PrimitiveType,NumPrimitives) \
-		FPlatformAtomics::InterlockedAdd(&GCurrentNumPrimitivesDrawnRHI[GPUIndex], NumPrimitives); \
-		RHIIncCurrentNumDrawCallPtr(GPUIndex);
+	struct FDrawCallCategoryName {};
+
+#endif
+
+// Macros for use inside RHI context Draw functions.
+// Updates the Stats structure on the underlying RHI context class (IRHICommandContext)
+
+#define RHI_DRAW_CALL_INC() do { Stats->Draws++; } while (false)
+#define RHI_DRAW_CALL_STATS(PrimitiveType,NumPrimitives)   \
+	do                                                     \
+	{												       \
+		Stats->Primitives[PrimitiveType] += NumPrimitives; \
+		Stats->Draws++;                                    \
+	} while(false)
+
+struct FRHIDrawStats
+{
+#if HAS_GPU_STATS
+	// The +1 is for "uncategorised"
+	static constexpr int32 NumCategories = FDrawCallCategoryName::MAX_DRAWCALL_CATEGORY + 1;
+#else
+	static constexpr int32 NumCategories = 1;
 #endif
 
-#define RHI_DRAW_CALL_INC() RHI_DRAW_CALL_INC_MGPU(0)
-#define RHI_DRAW_CALL_STATS(PrimitiveType,NumPrimitives) RHI_DRAW_CALL_STATS_MGPU(0, PrimitiveType,NumPrimitives)
+	static constexpr int32 NoCategory = NumCategories - 1;
+
+	struct FPerCategoryStats
+	{
+		uint32 Draws = 0;
+		TStaticArray<uint32, PT_Num> Primitives { InPlace, 0 };
+	};
 
+	struct FPerGPUStats
+	{
+		TStaticArray<FPerCategoryStats, NumCategories> Categories { InPlace };
+	};
+
+	TStaticArray<FPerGPUStats, MAX_NUM_GPUS> GPUs { InPlace };
+
+	FRHIDrawStats()
+	{
+		Reset();
+	}
+
+	void Reset()
+	{
+		FMemory::Memzero(*this);
+	}
+
+	// Adds the values from Other into the current instance.
+	RHI_API void Accumulate(FRHIDrawStats& Other);
+};
 
 // RHI memory stats.
 DECLARE_MEMORY_STAT_POOL_EXTERN(TEXT("Render target memory 2D"),STAT_RenderTargetMemory2D,STATGROUP_RHI,FPlatformMemory::MCR_GPU,RHI_API);
diff --git a/Engine/Source/Runtime/RHI/Public/RHICommandList.h b/Engine/Source/Runtime/RHI/Public/RHICommandList.h
--- a/Engine/Source/Runtime/RHI/Public/RHICommandList.h
+++ b/Engine/Source/Runtime/RHI/Public/RHICommandList.h
@@ -583,6 +583,10 @@ public:
 
 	void SetExecuteStat(TStatId Stat) { ExecuteStat = Stat; }
 
+#if HAS_GPU_STATS
+	void SetStatsCategory(FDrawCallCategoryName* Category);
+#endif
+
 	FGraphEventRef RHIThreadFence(bool bSetLockFence = false);
 
 	FORCEINLINE void* LockBuffer(FRHIBuffer* Buffer, uint32 Offset, uint32 SizeRHI, EResourceLockMode LockMode)
@@ -659,11 +663,6 @@ protected:
 		PersistentState.SubpassIndex = 0;
 	}
 
-private:
-	// Replays recorded commands into the specified contexts.
-	// Used internally, do not call directly.
-	void Execute(TRHIPipelineArray<IRHIComputeContext*>& InOutPipeContexts);
-
 protected:
 	// Blocks the calling thread until the dispatch event is completed.
 	// Used internally, do not call directly.
@@ -741,6 +740,41 @@ protected:
 		FGraphEventArray QueuedFenceCandidateEvents;
 		TArray<TRefCountPtr<FFenceCandidate>, FConcurrentLinearArrayAllocator> QueuedFenceCandidates;
 
+		struct FGPUStats
+		{
+#if HAS_GPU_STATS
+			FDrawCallCategoryName* CategoryTOP = nullptr;
+			FDrawCallCategoryName* CategoryBOP = nullptr;
+#endif
+
+			FRHIDrawStats* Ptr = nullptr;
+
+			void InitFrom(FGPUStats* Other)
+			{
+				if (!Other)
+					return;
+
+				Ptr = Other->Ptr;
+#if HAS_GPU_STATS
+				CategoryBOP = Other->CategoryTOP;
+#endif
+			}
+
+			void ApplyToContext(IRHIComputeContext* Context)
+			{
+				uint32 CategoryID = FRHIDrawStats::NoCategory;
+#if HAS_GPU_STATS
+				check(!CategoryBOP || CategoryBOP->ShouldCountDraws());
+				if (CategoryBOP)
+				{
+					CategoryID = CategoryBOP->Index;
+				}
+#endif
+
+				Context->StatsSetCategory(Ptr, CategoryID);
+			}
+		} Stats;
+
 		FPersistentState(FRHIGPUMask InInitialGPUMask, ERecordingThread InRecordingThread)
 			: RecordingThread(InRecordingThread)
 			, CurrentGPUMask(InInitialGPUMask)
@@ -798,6 +832,9 @@ public:
 private:
 	FRHICommandListBase(FPersistentState&& InPersistentState);
 
+	// Replays recorded commands into the specified contexts. Used internally, do not call directly.
+	void Execute(TRHIPipelineArray<IRHIComputeContext*>& InOutPipeContexts, FPersistentState::FGPUStats* ParentStats);
+
 	friend class FRHICommandListExecutor;
 	friend class FRHICommandListIterator;
 	friend class FRHICommandListScopedFlushAndExecute;
@@ -2224,7 +2261,7 @@ template<> RHI_API void FRHICommandSetShaderUniformBuffer        <FRHIComputeSha
 template<> RHI_API void FRHICommandSetShaderTexture              <FRHIComputeShader>::Execute(FRHICommandListBase& CmdList);
 template<> RHI_API void FRHICommandSetShaderResourceViewParameter<FRHIComputeShader>::Execute(FRHICommandListBase& CmdList);
 template<> RHI_API void FRHICommandSetShaderSampler              <FRHIComputeShader>::Execute(FRHICommandListBase& CmdList);
-template<> RHI_API void FRHICommandSetUAVParameter<FRHIComputeShader>::Execute(FRHICommandListBase& CmdList);
+template<> RHI_API void FRHICommandSetUAVParameter               <FRHIComputeShader>::Execute(FRHICommandListBase& CmdList);
 
 class RHI_API FRHIComputeCommandList : public FRHICommandListBase
 {
@@ -2909,12 +2946,12 @@ PRAGMA_ENABLE_DEPRECATION_WARNINGS
 	}
 };
 
-template<> RHI_API void FRHICommandSetShaderParameter<FRHIGraphicsShader>::Execute(FRHICommandListBase& CmdList);
-template<> RHI_API void FRHICommandSetShaderUniformBuffer<FRHIGraphicsShader>::Execute(FRHICommandListBase& CmdList);
-template<> RHI_API void FRHICommandSetShaderTexture<FRHIGraphicsShader>::Execute(FRHICommandListBase& CmdList);
+template<> RHI_API void FRHICommandSetShaderParameter            <FRHIGraphicsShader>::Execute(FRHICommandListBase& CmdList);
+template<> RHI_API void FRHICommandSetShaderUniformBuffer        <FRHIGraphicsShader>::Execute(FRHICommandListBase& CmdList);
+template<> RHI_API void FRHICommandSetShaderTexture              <FRHIGraphicsShader>::Execute(FRHICommandListBase& CmdList);
 template<> RHI_API void FRHICommandSetShaderResourceViewParameter<FRHIGraphicsShader>::Execute(FRHICommandListBase& CmdList);
-template<> RHI_API void FRHICommandSetShaderSampler<FRHIGraphicsShader>::Execute(FRHICommandListBase& CmdList);
-template<> RHI_API void FRHICommandSetUAVParameter<FRHIPixelShader>::Execute(FRHICommandListBase& CmdList);
+template<> RHI_API void FRHICommandSetShaderSampler              <FRHIGraphicsShader>::Execute(FRHICommandListBase& CmdList);
+template<> RHI_API void FRHICommandSetUAVParameter               <FRHIPixelShader   >::Execute(FRHICommandListBase& CmdList);
 
 class RHI_API FRHICommandList : public FRHIComputeCommandList
 {
@@ -3525,14 +3562,6 @@ public:
 		IncrementSubpass();
 	}
 
-	// These 6 are special in that they must be called on the immediate command list and they force a flush only when we are not doing RHI thread
-	void BeginScene();
-	void EndScene();
-	void BeginDrawingViewport(FRHIViewport* Viewport, FRHITexture* RenderTargetRHI);
-	void EndDrawingViewport(FRHIViewport* Viewport, bool bPresent, bool bLockToVsync);
-	void BeginFrame();
-	void EndFrame();
-
 	FORCEINLINE_DEBUGGABLE void RHIInvalidateCachedState()
 	{
 		if (Bypass())
@@ -3821,14 +3850,17 @@ extern RHI_API ERHIAccess RHIGetDefaultResourceState(EBufferUsageFlags InUsage,
 class RHI_API FRHICommandListImmediate : public FRHICommandList
 {
 	friend class FRHICommandListExecutor;
+	friend struct FRHICommandBeginFrame;
 
 	static FGraphEventArray WaitOutstandingTasks;
 	static FGraphEventRef   RHIThreadTask;
+	static FRHIDrawStats    FrameDrawStats;
 
 	FRHICommandListImmediate()
 		: FRHICommandList(FRHIGPUMask::All(), ERecordingThread::Render)
 	{
 		PersistentState.bImmediate = true;
+		PersistentState.Stats.Ptr = &FrameDrawStats;
 	}
 
 	~FRHICommandListImmediate()
@@ -3862,7 +3894,19 @@ class RHI_API FRHICommandListImmediate : public FRHICommandList
 	//
 	void Reset();
 
+	//
+	// Called on RHIBeginFrame. Updates the draw call counters / stats.
+	//
+	void ProcessStats();
+
 public:
+	void BeginScene();
+	void EndScene();
+	void BeginDrawingViewport(FRHIViewport* Viewport, FRHITexture* RenderTargetRHI);
+	void EndDrawingViewport(FRHIViewport* Viewport, bool bPresent, bool bLockToVsync);
+	void BeginFrame();
+	void EndFrame();
+
 	struct FQueuedCommandList
 	{
 		// The command list to enqueue.
diff --git a/Engine/Source/Runtime/RHI/Public/RHICommandListCommandExecutes.inl b/Engine/Source/Runtime/RHI/Public/RHICommandListCommandExecutes.inl
--- a/Engine/Source/Runtime/RHI/Public/RHICommandListCommandExecutes.inl
+++ b/Engine/Source/Runtime/RHI/Public/RHICommandListCommandExecutes.inl
@@ -733,7 +733,7 @@ void FRHICommandEndScene::Execute(FRHICommandListBase& CmdList)
 void FRHICommandBeginFrame::Execute(FRHICommandListBase& CmdList)
 {
 	RHISTAT(BeginFrame);
-	RHIPrivateBeginFrame();
+	CmdList.GetAsImmediate().ProcessStats();
 	INTERNAL_DECORATOR(RHIBeginFrame)();
 }
 
diff --git a/Engine/Source/Runtime/RHI/Public/RHIContext.h b/Engine/Source/Runtime/RHI/Public/RHIContext.h
--- a/Engine/Source/Runtime/RHI/Public/RHIContext.h
+++ b/Engine/Source/Runtime/RHI/Public/RHIContext.h
@@ -446,10 +446,9 @@ public:
 #endif
 
 #if ENABLE_RHI_VALIDATION
-	virtual void SetTrackedAccess(const FRHITrackedAccessInfo& Info)
-#else
-	inline  void SetTrackedAccess(const FRHITrackedAccessInfo& Info)
+	virtual
 #endif
+	void SetTrackedAccess(const FRHITrackedAccessInfo& Info)
 	{
 		check(Info.Resource != nullptr);
 		check(Info.Access != ERHIAccess::Unknown);
@@ -464,6 +463,25 @@ public:
 
 	virtual void* RHIGetNativeCommandBuffer() { return nullptr; }
 	virtual void RHIPostExternalCommandsReset() { }
+
+protected:
+	FRHIDrawStats::FPerCategoryStats* Stats = nullptr;
+
+public:
+	void StatsSetCategory(FRHIDrawStats* InStats, uint32 InCategoryID, uint32 InGPUIndex)
+	{
+		FRHIDrawStats::FPerGPUStats& GPU = InStats->GPUs[InGPUIndex];
+		FRHIDrawStats::FPerCategoryStats& Category = GPU.Categories[InCategoryID];
+		Stats = &Category;
+	}
+
+#if WITH_MGPU || ENABLE_RHI_VALIDATION
+	virtual
+#endif
+	void StatsSetCategory(FRHIDrawStats* InStats, uint32 InCategoryID)
+	{
+		StatsSetCategory(InStats, InCategoryID, 0);
+	}
 };
 
 enum class EAccelerationStructureBuildMode
diff --git a/Engine/Source/Runtime/RHI/Public/RHIValidationContext.h b/Engine/Source/Runtime/RHI/Public/RHIValidationContext.h
--- a/Engine/Source/Runtime/RHI/Public/RHIValidationContext.h
+++ b/Engine/Source/Runtime/RHI/Public/RHIValidationContext.h
@@ -1107,6 +1107,11 @@ public:
 		RHIContext->RHISetRayTracingMissShader(Scene, ShaderSlotInScene, Pipeline, ShaderIndexInPipeline, NumUniformBuffers, UniformBuffers, UserData);
 	}
 
+	virtual void StatsSetCategory(FRHIDrawStats* InStats, uint32 InCategoryID) final override
+	{
+		RHIContext->StatsSetCategory(InStats, InCategoryID);
+	}
+
 	void SetupDrawing()
 	{
 		// nothing to validate right now
diff --git a/Engine/Source/Runtime/RenderCore/Private/ProfilingDebugging/RealtimeGPUProfiler.cpp b/Engine/Source/Runtime/RenderCore/Private/ProfilingDebugging/RealtimeGPUProfiler.cpp
--- a/Engine/Source/Runtime/RenderCore/Private/ProfilingDebugging/RealtimeGPUProfiler.cpp
+++ b/Engine/Source/Runtime/RenderCore/Private/ProfilingDebugging/RealtimeGPUProfiler.cpp
@@ -1146,36 +1146,30 @@ void FRealtimeGPUProfiler::PopEventOverride()
 	}
 }
 
-void FRealtimeGPUProfiler::PushStat(FRHICommandListImmediate& RHICmdList, const FName& Name, const FName& StatName, const TCHAR* Description, FRHIDrawCallsStatPtr InNumDrawCallsPtr)
+void FRealtimeGPUProfiler::PushStat(FRHICommandListImmediate& RHICmdList, const FName& Name, const FName& StatName, const TCHAR* Description, FDrawCallCategoryName& Category)
 {
-	PushEvent(RHICmdList.GetGPUMask(), Name, StatName, Description).Submit(RHICmdList);
-
-	if (InNumDrawCallsPtr && (**InNumDrawCallsPtr) != -1)
+	if (Category.ShouldCountDraws())
 	{
-		RHICmdList.EnqueueLambda([InNumDrawCallsPtr](FRHICommandListImmediate&)
-		{
-			RHISetCurrentNumDrawCallPtr(InNumDrawCallsPtr);
-		});
+		RHICmdList.SetStatsCategory(&Category);
 	}
+
+	PushEvent(RHICmdList.GetGPUMask(), Name, StatName, Description).Submit(RHICmdList);
 }
 
-void FRealtimeGPUProfiler::PopStat(FRHICommandListImmediate& RHICmdList, FRHIDrawCallsStatPtr InNumDrawCallsPtr)
+void FRealtimeGPUProfiler::PopStat(FRHICommandListImmediate& RHICmdList, FDrawCallCategoryName& Category)
 {
 	PopEvent().Submit(RHICmdList);
 
-	if (InNumDrawCallsPtr && (**InNumDrawCallsPtr) != -1)
+	if (Category.ShouldCountDraws())
 	{
-		RHICmdList.EnqueueLambda([](FRHICommandListImmediate&)
-		{
-			RHISetCurrentNumDrawCallPtr(&GCurrentNumDrawCallsRHI);
-		});
+		RHICmdList.SetStatsCategory(nullptr);
 	}
 }
 
 /*-----------------------------------------------------------------------------
 FScopedGPUStatEvent
 -----------------------------------------------------------------------------*/
-void FScopedGPUStatEvent::Begin(FRHICommandListBase& InRHICmdList, const FName& Name, const FName& StatName, const TCHAR* Description, FRHIDrawCallsStatPtr InNumDrawCallsPtr)
+void FScopedGPUStatEvent::Begin(FRHICommandListBase& InRHICmdList, const FName& Name, const FName& StatName, const TCHAR* Description, FDrawCallCategoryName& InCategory)
 {
 	if (!AreGPUStatsEnabled())
 	{
@@ -1185,9 +1179,9 @@ void FScopedGPUStatEvent::Begin(FRHICommandListBase& InRHICmdList, const FName&
 	// Non-immediate command lists are not supported (silently fail)
 	if (InRHICmdList.IsImmediate())
 	{
-		NumDrawCallsPtr = InNumDrawCallsPtr;
+		Category = &InCategory;
 		RHICmdList = &InRHICmdList;
-		FRealtimeGPUProfiler::Get()->PushStat(InRHICmdList.GetAsImmediate(), Name, StatName, Description, InNumDrawCallsPtr);
+		FRealtimeGPUProfiler::Get()->PushStat(InRHICmdList.GetAsImmediate(), Name, StatName, Description, *Category);
 	}
 }
 
@@ -1200,7 +1194,7 @@ void FScopedGPUStatEvent::End()
 	if (RHICmdList != nullptr)
 	{
 		// Command list is initialized only if it is immediate during Begin() and GetAsImmediate() also internally checks this.
-		FRealtimeGPUProfiler::Get()->PopStat(RHICmdList->GetAsImmediate(), NumDrawCallsPtr);
+		FRealtimeGPUProfiler::Get()->PopStat(RHICmdList->GetAsImmediate(), *Category);
 	}
 }
 #endif // HAS_GPU_STATS
diff --git a/Engine/Source/Runtime/RenderCore/Private/RenderGraphEvent.cpp b/Engine/Source/Runtime/RenderCore/Private/RenderGraphEvent.cpp
--- a/Engine/Source/Runtime/RenderCore/Private/RenderGraphEvent.cpp
+++ b/Engine/Source/Runtime/RenderCore/Private/RenderGraphEvent.cpp
@@ -685,10 +685,10 @@ FRDGEventScopeOpArray FRDGEventScopeStack::CompilePassEpilogue()
 	return MoveTemp(Ops);
 }
 
-FRDGGPUStatScopeGuard::FRDGGPUStatScopeGuard(FRDGBuilder& InGraphBuilder, const FName& Name, const FName& StatName, const TCHAR* Description, FRHIDrawCallsStatPtr InNumDrawCallsPtr)
+FRDGGPUStatScopeGuard::FRDGGPUStatScopeGuard(FRDGBuilder& InGraphBuilder, const FName& Name, const FName& StatName, const TCHAR* Description, FDrawCallCategoryName& Category)
 	: GraphBuilder(InGraphBuilder)
 {
-	GraphBuilder.GPUScopeStacks.BeginStatScope(Name, StatName, Description, InNumDrawCallsPtr);
+	GraphBuilder.GPUScopeStacks.BeginStatScope(Name, StatName, Description, Category);
 }
 
 FRDGGPUStatScopeGuard::~FRDGGPUStatScopeGuard()
@@ -749,14 +749,12 @@ void FRDGGPUStatScopeOpArray::Execute(FRHIComputeCommandList& RHICmdListCompute)
 		const FRDGGPUStatScopeOp Op = Ops[Index];
 		const FRDGGPUStatScope* Scope = Op.Scope;
 
-		if (Scope->DrawCallCounter != nullptr && (**Scope->DrawCallCounter) != -1)
+		if (Scope->Category.ShouldCountDraws())
 		{
-			RHICmdList.EnqueueLambda(
-				[DrawCallCounter = Scope->DrawCallCounter, bPush = Op.IsPush()](auto&)
-			{
-				RHISetCurrentNumDrawCallPtr(bPush ? DrawCallCounter : &GCurrentNumDrawCallsRHI);
-			});
-			break;
+		    RHICmdList.SetStatsCategory(Op.IsPush()
+			    ? &Scope->Category
+			    : nullptr
+		    );
 		}
 	}
 #endif
diff --git a/Engine/Source/Runtime/RenderCore/Public/ProfilingDebugging/RealtimeGPUProfiler.h b/Engine/Source/Runtime/RenderCore/Public/ProfilingDebugging/RealtimeGPUProfiler.h
--- a/Engine/Source/Runtime/RenderCore/Public/ProfilingDebugging/RealtimeGPUProfiler.h
+++ b/Engine/Source/Runtime/RenderCore/Public/ProfilingDebugging/RealtimeGPUProfiler.h
@@ -267,51 +267,52 @@ struct RENDERCORE_API FBreadcrumbEvent
 #define SCOPED_CONDITIONAL_COMPUTE_EVENTF SCOPED_CONDITIONAL_GPU_EVENTF
 #define SCOPED_CONDITIONAL_COMPUTE_EVENTF_COLOR SCOPED_CONDITIONAL_GPU_EVENTF_COLOR
 
-// GPU stats
-#ifndef HAS_GPU_STATS
-	#if ( STATS || CSV_PROFILER || GPUPROFILERTRACE_ENABLED ) && (!UE_BUILD_SHIPPING)
-	#define HAS_GPU_STATS 1
+#if HAS_GPU_STATS
+
+	CSV_DECLARE_CATEGORY_MODULE_EXTERN(RENDERCORE_API,GPU);
+
+	// The DECLARE_GPU_STAT macros both declare and define a stat (for use in a single CPP)
+	#define DECLARE_GPU_STAT(StatName)                            DECLARE_FLOAT_COUNTER_STAT(TEXT(#StatName)       , Stat_GPU_##StatName, STATGROUP_GPU  ); CSV_DEFINE_STAT(GPU,StatName);         static FDrawCallCategoryName DrawcallCountCategory_##StatName;
+	#define DECLARE_GPU_STAT_NAMED(StatName, NameString)          DECLARE_FLOAT_COUNTER_STAT(NameString            , Stat_GPU_##StatName, STATGROUP_GPU  ); CSV_DEFINE_STAT(GPU,StatName);         static FDrawCallCategoryName DrawcallCountCategory_##StatName;
+	#define DECLARE_GPU_DRAWCALL_STAT(StatName)                   DECLARE_FLOAT_COUNTER_STAT(TEXT(#StatName)       , Stat_GPU_##StatName, STATGROUP_GPU  ); CSV_DEFINE_STAT(GPU,StatName);         static FDrawCallCategoryName DrawcallCountCategory_##StatName((TCHAR*)TEXT(#StatName));
+	#define DECLARE_GPU_DRAWCALL_STAT_NAMED(StatName, NameString) DECLARE_FLOAT_COUNTER_STAT(NameString            , Stat_GPU_##StatName, STATGROUP_GPU  ); CSV_DEFINE_STAT(GPU,StatName);         static FDrawCallCategoryName DrawcallCountCategory_##StatName((TCHAR*)TEXT(#StatName));
+	#define DECLARE_GPU_DRAWCALL_STAT_EXTERN(StatName)            DECLARE_FLOAT_COUNTER_STAT_EXTERN(TEXT(#StatName), Stat_GPU_##StatName, STATGROUP_GPU, ); CSV_DECLARE_STAT_EXTERN(GPU,StatName); extern FDrawCallCategoryName DrawcallCountCategory_##StatName;
+
+	// Extern GPU stats are needed where a stat is used in multiple CPPs. Use the DECLARE_GPU_STAT_NAMED_EXTERN in the header and DEFINE_GPU_STAT in the CPPs
+	#define DECLARE_GPU_STAT_NAMED_EXTERN(StatName, NameString) DECLARE_FLOAT_COUNTER_STAT_EXTERN(NameString, Stat_GPU_##StatName, STATGROUP_GPU, ); CSV_DECLARE_STAT_EXTERN(GPU,StatName); extern FDrawCallCategoryName DrawcallCountCategory_##StatName;
+	#define DEFINE_GPU_STAT(StatName)                           DEFINE_STAT(Stat_GPU_##StatName);                                                    CSV_DEFINE_STAT(GPU,StatName);                FDrawCallCategoryName DrawcallCountCategory_##StatName;
+	#define DEFINE_GPU_DRAWCALL_STAT(StatName)                  DEFINE_STAT(Stat_GPU_##StatName);                                                    CSV_DEFINE_STAT(GPU,StatName);                FDrawCallCategoryName DrawcallCountCategory_##StatName((TCHAR*)TEXT(#StatName));
+
+	#if STATS
+		#define SCOPED_GPU_STAT(RHICmdList, StatName)                      FScopedGPUStatEvent PREPROCESSOR_JOIN(GPUStatEvent_##StatName,__LINE__); PREPROCESSOR_JOIN(GPUStatEvent_##StatName,__LINE__).Begin(RHICmdList, CSV_STAT_FNAME(StatName), GET_STATID( Stat_GPU_##StatName ).GetName(), nullptr    , DrawcallCountCategory_##StatName);
+		#define SCOPED_GPU_STAT_VERBOSE(RHICmdList, StatName, Description) FScopedGPUStatEvent PREPROCESSOR_JOIN(GPUStatEvent_##StatName,__LINE__); PREPROCESSOR_JOIN(GPUStatEvent_##StatName,__LINE__).Begin(RHICmdList, CSV_STAT_FNAME(StatName), GET_STATID( Stat_GPU_##StatName ).GetName(), Description, DrawcallCountCategory_##StatName);
 	#else
-	#define HAS_GPU_STATS 0
+		#define SCOPED_GPU_STAT(RHICmdList, StatName)                      FScopedGPUStatEvent PREPROCESSOR_JOIN(GPUStatEvent_##StatName,__LINE__); PREPROCESSOR_JOIN(GPUStatEvent_##StatName,__LINE__).Begin(RHICmdList, CSV_STAT_FNAME(StatName), FName(), nullptr    , DrawcallCountCategory_##StatName);
+		#define SCOPED_GPU_STAT_VERBOSE(RHICmdList, StatName, Description) FScopedGPUStatEvent PREPROCESSOR_JOIN(GPUStatEvent_##StatName,__LINE__); PREPROCESSOR_JOIN(GPUStatEvent_##StatName,__LINE__).Begin(RHICmdList, CSV_STAT_FNAME(StatName), FName(), Description, DrawcallCountCategory_##StatName);
 	#endif
-#endif
 
-#if HAS_GPU_STATS
- CSV_DECLARE_CATEGORY_MODULE_EXTERN(RENDERCORE_API,GPU);
- // The DECLARE_GPU_STAT macros both declare and define a stat (for use in a single CPP)
- #define DECLARE_GPU_STAT(StatName) DECLARE_FLOAT_COUNTER_STAT(TEXT(#StatName), Stat_GPU_##StatName, STATGROUP_GPU); CSV_DEFINE_STAT(GPU,StatName); static FDrawCallCategoryName DrawcallCountCategory_##StatName;
- #define DECLARE_GPU_DRAWCALL_STAT(StatName) DECLARE_FLOAT_COUNTER_STAT(TEXT(#StatName), Stat_GPU_##StatName, STATGROUP_GPU); CSV_DEFINE_STAT(GPU,StatName); static FDrawCallCategoryName DrawcallCountCategory_##StatName((TCHAR*)TEXT(#StatName));
- #define DECLARE_GPU_DRAWCALL_STAT_EXTERN(StatName) DECLARE_FLOAT_COUNTER_STAT_EXTERN(TEXT(#StatName), Stat_GPU_##StatName, STATGROUP_GPU, ); CSV_DECLARE_STAT_EXTERN(GPU,StatName); extern FDrawCallCategoryName DrawcallCountCategory_##StatName;
- #define DECLARE_GPU_STAT_NAMED(StatName, NameString) DECLARE_FLOAT_COUNTER_STAT(NameString, Stat_GPU_##StatName, STATGROUP_GPU); CSV_DEFINE_STAT(GPU,StatName); static FDrawCallCategoryName DrawcallCountCategory_##StatName;
- #define DECLARE_GPU_DRAWCALL_STAT_NAMED(StatName, NameString) DECLARE_FLOAT_COUNTER_STAT(NameString, Stat_GPU_##StatName, STATGROUP_GPU); CSV_DEFINE_STAT(GPU,StatName); static FDrawCallCategoryName DrawcallCountCategory_##StatName((TCHAR*)TEXT(#StatName));
-
- // Extern GPU stats are needed where a stat is used in multiple CPPs. Use the DECLARE_GPU_STAT_NAMED_EXTERN in the header and DEFINE_GPU_STAT in the CPPs
- #define DECLARE_GPU_STAT_NAMED_EXTERN(StatName, NameString) DECLARE_FLOAT_COUNTER_STAT_EXTERN(NameString, Stat_GPU_##StatName, STATGROUP_GPU, ); CSV_DECLARE_STAT_EXTERN(GPU,StatName); extern FDrawCallCategoryName DrawcallCountCategory_##StatName;
- #define DEFINE_GPU_STAT(StatName) DEFINE_STAT(Stat_GPU_##StatName); CSV_DEFINE_STAT(GPU,StatName); FDrawCallCategoryName DrawcallCountCategory_##StatName;
- #define DEFINE_GPU_DRAWCALL_STAT(StatName) DEFINE_STAT(Stat_GPU_##StatName); CSV_DEFINE_STAT(GPU,StatName); FDrawCallCategoryName DrawcallCountCategory_##StatName((TCHAR*)TEXT(#StatName));
-#if STATS
-  #define SCOPED_GPU_STAT(RHICmdList, StatName) FScopedGPUStatEvent PREPROCESSOR_JOIN(GPUStatEvent_##StatName,__LINE__); PREPROCESSOR_JOIN(GPUStatEvent_##StatName,__LINE__).Begin(RHICmdList, CSV_STAT_FNAME(StatName), GET_STATID( Stat_GPU_##StatName ).GetName(), nullptr, &DrawcallCountCategory_##StatName.Counters);
-  #define SCOPED_GPU_STAT_VERBOSE(RHICmdList, StatName, Description) FScopedGPUStatEvent PREPROCESSOR_JOIN(GPUStatEvent_##StatName,__LINE__); PREPROCESSOR_JOIN(GPUStatEvent_##StatName,__LINE__).Begin(RHICmdList, CSV_STAT_FNAME(StatName), GET_STATID( Stat_GPU_##StatName ).GetName(), Description, &DrawcallCountCategory_##StatName.Counters);
-#else
-  #define SCOPED_GPU_STAT(RHICmdList, StatName) FScopedGPUStatEvent PREPROCESSOR_JOIN(GPUStatEvent_##StatName,__LINE__); PREPROCESSOR_JOIN(GPUStatEvent_##StatName,__LINE__).Begin(RHICmdList, CSV_STAT_FNAME(StatName), FName(), nullptr, &DrawcallCountCategory_##StatName.Counters);
-  #define SCOPED_GPU_STAT_VERBOSE(RHICmdList, StatName, Description) FScopedGPUStatEvent PREPROCESSOR_JOIN(GPUStatEvent_##StatName,__LINE__); PREPROCESSOR_JOIN(GPUStatEvent_##StatName,__LINE__).Begin(RHICmdList, CSV_STAT_FNAME(StatName), FName(), Description, &DrawcallCountCategory_##StatName.Counters);
-#endif
- #define GPU_STATS_BEGINFRAME(RHICmdList) FRealtimeGPUProfiler::Get()->BeginFrame(RHICmdList);
- #define GPU_STATS_ENDFRAME(RHICmdList) FRealtimeGPUProfiler::Get()->EndFrame(RHICmdList);
- #define GPU_STATS_SUSPENDFRAME() FRealtimeGPUProfiler::Get()->SuspendFrame();
+	#define GPU_STATS_BEGINFRAME(RHICmdList) FRealtimeGPUProfiler::Get()->BeginFrame(RHICmdList);
+	#define GPU_STATS_ENDFRAME(RHICmdList)   FRealtimeGPUProfiler::Get()->EndFrame(RHICmdList);
+	#define GPU_STATS_SUSPENDFRAME()         FRealtimeGPUProfiler::Get()->SuspendFrame();
+
 #else
- #define DECLARE_GPU_STAT(StatName)
- #define DECLARE_GPU_DRAWCALL_STAT(StatName)
- #define DECLARE_GPU_DRAWCALL_STAT_EXTERN(StatName)
- #define DECLARE_GPU_STAT_NAMED(StatName, NameString)
- #define DECLARE_GPU_DRAWCALL_STAT_NAMED(StatName, NameString)
- #define DECLARE_GPU_STAT_NAMED_EXTERN(StatName, NameString)
- #define DEFINE_GPU_STAT(StatName)
- #define DEFINE_GPU_DRAWCALL_STAT(StatName)
- #define SCOPED_GPU_STAT(RHICmdList, StatName) 
- #define GPU_STATS_BEGINFRAME(RHICmdList) 
- #define GPU_STATS_ENDFRAME(RHICmdList) 
- #define GPU_STATS_SUSPENDFRAME()
+
+	#define DECLARE_GPU_STAT(StatName)
+	#define DECLARE_GPU_STAT_NAMED(StatName, NameString)
+	#define DECLARE_GPU_DRAWCALL_STAT(StatName)
+	#define DECLARE_GPU_DRAWCALL_STAT_NAMED(StatName, NameString)
+	#define DECLARE_GPU_DRAWCALL_STAT_EXTERN(StatName)
+
+	#define DECLARE_GPU_STAT_NAMED_EXTERN(StatName, NameString)
+	#define DEFINE_GPU_STAT(StatName)
+	#define DEFINE_GPU_DRAWCALL_STAT(StatName)
+	#define SCOPED_GPU_STAT(RHICmdList, StatName) 
+	#define SCOPED_GPU_STAT_VERBOSE(RHICmdList, StatName, Description)
+
+	#define GPU_STATS_BEGINFRAME(RHICmdList) 
+	#define GPU_STATS_ENDFRAME(RHICmdList) 
+	#define GPU_STATS_SUSPENDFRAME()
+
 #endif
 
 RENDERCORE_API bool AreGPUStatsEnabled();
@@ -410,8 +411,8 @@ public:
 	void PopEventOverride();
 
 	/** Push/pop stats which do additional draw call tracking on top of events. */
-	void PushStat(FRHICommandListImmediate& RHICmdList, const FName& Name, const FName& StatName, const TCHAR* Description, FRHIDrawCallsStatPtr InNumDrawCallsPtr);
-	void PopStat(FRHICommandListImmediate& RHICmdList, FRHIDrawCallsStatPtr InNumDrawCallsPtr);
+	void PushStat(FRHICommandListImmediate& RHICmdList, const FName& Name, const FName& StatName, const TCHAR* Description, FDrawCallCategoryName& Category);
+	void PopStat(FRHICommandListImmediate& RHICmdList, FDrawCallCategoryName& Category);
 
 #if GPUPROFILERTRACE_ENABLED
 	RENDERCORE_API void FetchPerfByDescription(TArray<FRealtimeGPUProfilerDescriptionResult> & OutResults) const;
@@ -447,18 +448,13 @@ private:
 class FScopedGPUStatEvent
 {
 	/** Cmdlist to push onto. */
-	FRHICommandListBase* RHICmdList;
+	FRHICommandListBase* RHICmdList = nullptr;
+	FDrawCallCategoryName* Category = nullptr;
 
-	FRHIDrawCallsStatPtr NumDrawCallsPtr;
 public:
-
 	UE_NONCOPYABLE(FScopedGPUStatEvent)
 
-	/** Default constructor, initializing all member variables. */
-	FORCEINLINE FScopedGPUStatEvent()
-		: RHICmdList(nullptr)
-		, NumDrawCallsPtr(nullptr)
-	{}
+	FORCEINLINE FScopedGPUStatEvent() = default;
 
 	/**
 	* Terminate the event based upon scope
@@ -474,7 +470,7 @@ public:
 	/**
 	* Start/Stop functions for timer stats
 	*/
-	RENDERCORE_API void Begin(FRHICommandListBase& InRHICmdList, const FName& Name, const FName& StatName, const TCHAR* Description, FRHIDrawCallsStatPtr InNumDrawCallsPtr);
+	RENDERCORE_API void Begin(FRHICommandListBase& InRHICmdList, const FName& Name, const FName& StatName, const TCHAR* Description, FDrawCallCategoryName& InCategory);
 	RENDERCORE_API void End();
 };
 #endif // HAS_GPU_STATS
diff --git a/Engine/Source/Runtime/RenderCore/Public/RenderGraphEvent.h b/Engine/Source/Runtime/RenderCore/Public/RenderGraphEvent.h
--- a/Engine/Source/Runtime/RenderCore/Public/RenderGraphEvent.h
+++ b/Engine/Source/Runtime/RenderCore/Public/RenderGraphEvent.h
@@ -43,11 +43,11 @@
 
 #if HAS_GPU_STATS
 	#if STATS
-		#define RDG_GPU_STAT_SCOPE(GraphBuilder, StatName) FRDGGPUStatScopeGuard PREPROCESSOR_JOIN(__RDG_GPUStatEvent_##StatName,__LINE__) ((GraphBuilder), CSV_STAT_FNAME(StatName), GET_STATID(Stat_GPU_##StatName).GetName(), nullptr, &DrawcallCountCategory_##StatName.Counters);
-		#define RDG_GPU_STAT_SCOPE_VERBOSE(GraphBuilder, StatName, Description) FRDGGPUStatScopeGuard PREPROCESSOR_JOIN(__RDG_GPUStatEvent_##StatName,__LINE__) ((GraphBuilder), CSV_STAT_FNAME(StatName), GET_STATID(Stat_GPU_##StatName).GetName(), Description, &DrawcallCountCategory_##StatName.Counters);
+		#define RDG_GPU_STAT_SCOPE(GraphBuilder, StatName) FRDGGPUStatScopeGuard PREPROCESSOR_JOIN(__RDG_GPUStatEvent_##StatName,__LINE__) ((GraphBuilder), CSV_STAT_FNAME(StatName), GET_STATID(Stat_GPU_##StatName).GetName(), nullptr, DrawcallCountCategory_##StatName);
+		#define RDG_GPU_STAT_SCOPE_VERBOSE(GraphBuilder, StatName, Description) FRDGGPUStatScopeGuard PREPROCESSOR_JOIN(__RDG_GPUStatEvent_##StatName,__LINE__) ((GraphBuilder), CSV_STAT_FNAME(StatName), GET_STATID(Stat_GPU_##StatName).GetName(), Description, DrawcallCountCategory_##StatName);
 	#else
-		#define RDG_GPU_STAT_SCOPE(GraphBuilder, StatName) FRDGGPUStatScopeGuard PREPROCESSOR_JOIN(__RDG_GPUStatEvent_##StatName,__LINE__) ((GraphBuilder), CSV_STAT_FNAME(StatName), FName(), nullptr, &DrawcallCountCategory_##StatName.Counters);
-		#define RDG_GPU_STAT_SCOPE_VERBOSE(GraphBuilder, StatName, Description) FRDGGPUStatScopeGuard PREPROCESSOR_JOIN(__RDG_GPUStatEvent_##StatName,__LINE__) ((GraphBuilder), CSV_STAT_FNAME(StatName), FName(), Description, &DrawcallCountCategory_##StatName.Counters);
+		#define RDG_GPU_STAT_SCOPE(GraphBuilder, StatName) FRDGGPUStatScopeGuard PREPROCESSOR_JOIN(__RDG_GPUStatEvent_##StatName,__LINE__) ((GraphBuilder), CSV_STAT_FNAME(StatName), FName(), nullptr, DrawcallCountCategory_##StatName);
+		#define RDG_GPU_STAT_SCOPE_VERBOSE(GraphBuilder, StatName, Description) FRDGGPUStatScopeGuard PREPROCESSOR_JOIN(__RDG_GPUStatEvent_##StatName,__LINE__) ((GraphBuilder), CSV_STAT_FNAME(StatName), FName(), Description, DrawcallCountCategory_##StatName);
 	#endif
 #else
 	#define RDG_GPU_STAT_SCOPE(GraphBuilder, StatName)
@@ -223,7 +223,7 @@ public:
 	}
 
 	template <typename... TScopeConstructArgs>
-	inline void BeginScope(TScopeConstructArgs... ScopeConstructArgs)
+	inline void BeginScope(TScopeConstructArgs&&... ScopeConstructArgs)
 	{
 		auto Scope = Allocator.AllocNoDestruct<ScopeType>(CurrentScope, Forward<TScopeConstructArgs>(ScopeConstructArgs)...);
 		Scopes.Add(Scope);
@@ -570,11 +570,11 @@ private:
 class FRDGGPUStatScope final
 {
 public:
-	FRDGGPUStatScope(const FRDGGPUStatScope* InParentScope, const FName& InName, const FName& InStatName, const TCHAR* InDescription, FRHIDrawCallsStatPtr InDrawCallCounter)
+	FRDGGPUStatScope(const FRDGGPUStatScope* InParentScope, const FName& InName, const FName& InStatName, const TCHAR* InDescription, FDrawCallCategoryName& InCategory)
 		: ParentScope(InParentScope)
 		, Name(InName)
 		, StatName(InStatName)
-		, DrawCallCounter(InDrawCallCounter)
+		, Category(InCategory)
 	{
 		if (InDescription)
 		{
@@ -586,7 +586,7 @@ public:
 	const FName Name;
 	const FName StatName;
 	FString Description;
-	FRHIDrawCallsStatPtr DrawCallCounter;
+	FDrawCallCategoryName& Category;
 };
 
 class FRDGGPUStatScopeOp : public TRDGScopeOp<FRDGGPUStatScope>
@@ -634,12 +634,11 @@ public:
 #endif
 	{}
 
-	inline void BeginScope(const FName& Name, const FName& StatName, const TCHAR* Description, FRHIDrawCallsStatPtr DrawCallCounter)
+	inline void BeginScope(const FName& Name, const FName& StatName, const TCHAR* Description, FDrawCallCategoryName& Category)
 	{
 		if (IsEnabled())
 		{
-			check(DrawCallCounter != nullptr);
-			ScopeStack.BeginScope(Name, StatName, Description, DrawCallCounter);
+			ScopeStack.BeginScope(Name, StatName, Description, Category);
 		}
 	}
 
@@ -696,7 +695,7 @@ private:
 class RENDERCORE_API FRDGGPUStatScopeGuard final
 {
 public:
-	FRDGGPUStatScopeGuard(FRDGBuilder& InGraphBuilder, const FName& Name, const FName& StatName, const TCHAR* Description, FRHIDrawCallsStatPtr DrawCallCounter);
+	FRDGGPUStatScopeGuard(FRDGBuilder& InGraphBuilder, const FName& Name, const FName& StatName, const TCHAR* Description, FDrawCallCategoryName& InCategory);
 	FRDGGPUStatScopeGuard(const FRDGGPUStatScopeGuard&) = delete;
 	~FRDGGPUStatScopeGuard();
 
@@ -833,9 +832,9 @@ struct RENDERCORE_API FRDGGPUScopeStacksByPipeline
 		AsyncCompute.Event.EndScope();
 	}
 
-	inline void BeginStatScope(const FName& Name, const FName& StatName, const TCHAR* Description, FRHIDrawCallsStatPtr DrawCallCounter)
+	inline void BeginStatScope(const FName& Name, const FName& StatName, const TCHAR* Description, FDrawCallCategoryName& Category)
 	{
-		Graphics.Stat.BeginScope(Name, StatName, Description, DrawCallCounter);
+		Graphics.Stat.BeginScope(Name, StatName, Description, Category);
 	}
 
 	inline void EndStatScope()
diff --git a/Engine/Source/Runtime/Renderer/Private/ReflectionEnvironmentCapture.cpp b/Engine/Source/Runtime/Renderer/Private/ReflectionEnvironmentCapture.cpp
--- a/Engine/Source/Runtime/Renderer/Private/ReflectionEnvironmentCapture.cpp
+++ b/Engine/Source/Runtime/Renderer/Private/ReflectionEnvironmentCapture.cpp
@@ -1228,7 +1228,7 @@ void CaptureSceneIntoScratchCubemap(
 			// Alert the RHI that we're rendering a new frame
 			// Not really a new frame, but it will allow pooling mechanisms to update, like the uniform buffer pool
 			ENQUEUE_RENDER_COMMAND(BeginFrame)(
-				[](FRHICommandList& RHICmdList)
+				[](FRHICommandListImmediate& RHICmdList)
 			{
 				GFrameNumberRenderThread++;
 				RHICmdList.BeginFrame();
diff --git a/Engine/Source/Runtime/Renderer/Private/SceneRendering.cpp b/Engine/Source/Runtime/Renderer/Private/SceneRendering.cpp
--- a/Engine/Source/Runtime/Renderer/Private/SceneRendering.cpp
+++ b/Engine/Source/Runtime/Renderer/Private/SceneRendering.cpp
@@ -763,17 +763,6 @@ FRHICommandList* FParallelCommandListSet::NewParallelCommandList()
 	Result->SwitchPipeline(ERHIPipeline::Graphics);
 	Result->SetExecuteStat(ExecuteStat);
 
-#if RDG_GPU_DEBUG_SCOPES
-	if (Pass->GetGPUScopes().Stat != nullptr && (**Pass->GetGPUScopes().Stat->DrawCallCounter) != -1)
-	{
-		Result->EnqueueLambda(
-			[DrawCallCounter = Pass->GetGPUScopes().Stat->DrawCallCounter](auto&)
-		{
-			RHISetCurrentNumDrawCallPtr(DrawCallCounter);
-		});
-	}
-#endif
-
 #if RHI_WANT_BREADCRUMB_EVENTS
 	if (GParallelCmdListInheritBreadcrumbs)
 	{
