More shader compiling code wrapped with WITH_EDITOR checks
- Focused around moving GlobalBeginCompileShader and friends.
- ModifyCompilationEnvironment and ValidateCompiledResult now only compiled in Editor builds.
- Measured 0.5MB to 1.0MB elf size reduction depending on platform.

diff --git a/Engine/Plugins/Compositing/OpenColorIO/Source/OpenColorIO/Public/OpenColorIOShaderType.h b/Engine/Plugins/Compositing/OpenColorIO/Source/OpenColorIO/Public/OpenColorIOShaderType.h
--- a/Engine/Plugins/Compositing/OpenColorIO/Source/OpenColorIO/Public/OpenColorIOShaderType.h
+++ b/Engine/Plugins/Compositing/OpenColorIO/Source/OpenColorIO/Public/OpenColorIOShaderType.h
@@ -71,20 +71,24 @@ public:
 		int32 InTotalPermutationCount,
 		ConstructSerializedType InConstructSerializedRef,
 		ConstructCompiledType InConstructCompiledRef,
-		ModifyCompilationEnvironmentType InModifyCompilationEnvironmentRef,
 		ShouldCompilePermutationType InShouldCompilePermutationRef,
-		ValidateCompiledResultType InValidateCompiledResultRef,
 		GetRayTracingPayloadTypeType InGetRayTracingPayloadTypeRef,
+#if WITH_EDITOR
+		ModifyCompilationEnvironmentType InModifyCompilationEnvironmentRef,
+		ValidateCompiledResultType InValidateCompiledResultRef,
+#endif // WITH_EDITOR
 		uint32 InTypeSize,
 		const FShaderParametersMetadata* InRootParametersMetadata = nullptr
 		):
 		FShaderType(EShaderTypeForDynamicCast::OCIO, InTypeLayout, InName, InSourceFilename, InFunctionName, SF_Pixel, InTotalPermutationCount,
 			InConstructSerializedRef,
 			InConstructCompiledRef,
-			InModifyCompilationEnvironmentRef,
 			InShouldCompilePermutationRef,
-			InValidateCompiledResultRef,
 			InGetRayTracingPayloadTypeRef,
+#if WITH_EDITOR
+			InModifyCompilationEnvironmentRef,
+			InValidateCompiledResultRef,
+#endif // WITH_EDITOR
 			InTypeSize,
 			InRootParametersMetadata)
 	{
diff --git a/Engine/Plugins/FX/Niagara/Source/NiagaraShader/Public/NiagaraShaderType.h b/Engine/Plugins/FX/Niagara/Source/NiagaraShader/Public/NiagaraShaderType.h
--- a/Engine/Plugins/FX/Niagara/Source/NiagaraShader/Public/NiagaraShaderType.h
+++ b/Engine/Plugins/FX/Niagara/Source/NiagaraShader/Public/NiagaraShaderType.h
@@ -89,20 +89,24 @@ public:
 		int32 InTotalPermutationCount,
 		ConstructSerializedType InConstructSerializedRef,
 		ConstructCompiledType InConstructCompiledRef,
-		ModifyCompilationEnvironmentType InModifyCompilationEnvironmentRef,
 		ShouldCompilePermutationType InShouldCompilePermutationRef,
-		ValidateCompiledResultType InValidateCompiledResultRef,
 		GetRayTracingPayloadTypeType InGetRayTracingPayloadTypeRef,
+#if WITH_EDITOR
+		ModifyCompilationEnvironmentType InModifyCompilationEnvironmentRef,
+		ValidateCompiledResultType InValidateCompiledResultRef,
+#endif // WITH_EDITOR
 		uint32 InTypeSize,
 		const FShaderParametersMetadata* InRootParametersMetadata = nullptr
 	)
 		: FShaderType(EShaderTypeForDynamicCast::Niagara, InTypeLayout, InName, InSourceFilename, InFunctionName, SF_Compute, InTotalPermutationCount,
 			InConstructSerializedRef,
 			InConstructCompiledRef,
-			InModifyCompilationEnvironmentRef,
 			InShouldCompilePermutationRef,
-			InValidateCompiledResultRef,
 			InGetRayTracingPayloadTypeRef,
+#if WITH_EDITOR
+			InModifyCompilationEnvironmentRef,
+			InValidateCompiledResultRef,
+#endif // WITH_EDITOR
 			InTypeSize,
 			InRootParametersMetadata
 		)
diff --git a/Engine/Plugins/Runtime/ComputeFramework/Source/ComputeFramework/Private/ComputeFramework/ComputeKernelShaderType.h b/Engine/Plugins/Runtime/ComputeFramework/Source/ComputeFramework/Private/ComputeFramework/ComputeKernelShaderType.h
--- a/Engine/Plugins/Runtime/ComputeFramework/Source/ComputeFramework/Private/ComputeFramework/ComputeKernelShaderType.h
+++ b/Engine/Plugins/Runtime/ComputeFramework/Source/ComputeFramework/Private/ComputeFramework/ComputeKernelShaderType.h
@@ -73,10 +73,12 @@ public:
 		int32 InTotalPermutationCount,
 		ConstructSerializedType InConstructSerializedRef,
 		ConstructCompiledType InConstructCompiledRef,
-		ModifyCompilationEnvironmentType InModifyCompilationEnvironmentRef,
 		ShouldCompilePermutationType InShouldCompilePermutationRef,
-		ValidateCompiledResultType InValidateCompiledResultRef,
 		GetRayTracingPayloadTypeType InGetRayTracingPayloadTypeRef,
+#if WITH_EDITOR
+		ModifyCompilationEnvironmentType InModifyCompilationEnvironmentRef,
+		ValidateCompiledResultType InValidateCompiledResultRef,
+#endif // WITH_EDITOR
 		uint32 InTypeSize,
 		const FShaderParametersMetadata* InRootParametersMetadata = nullptr
 		):
@@ -90,10 +92,12 @@ public:
 			InTotalPermutationCount,
 			InConstructSerializedRef,
 			InConstructCompiledRef,
-			InModifyCompilationEnvironmentRef,
 			InShouldCompilePermutationRef,
-			InValidateCompiledResultRef,
 			InGetRayTracingPayloadTypeRef,
+#if WITH_EDITOR
+			InModifyCompilationEnvironmentRef,
+			InValidateCompiledResultRef,
+#endif // WITH_EDITOR
 			InTypeSize,
 			InRootParametersMetadata
 			)
diff --git a/Engine/Source/Programs/ShaderCompileWorker/Private/ShaderCompileWorker.cpp b/Engine/Source/Programs/ShaderCompileWorker/Private/ShaderCompileWorker.cpp
--- a/Engine/Source/Programs/ShaderCompileWorker/Private/ShaderCompileWorker.cpp
+++ b/Engine/Source/Programs/ShaderCompileWorker/Private/ShaderCompileWorker.cpp
@@ -714,13 +714,8 @@ private:
 
 			int32 NumBatches = 0;
 			InputFile << NumBatches;
-			// Flush cache, to make sure we load the latest version of the input file.
-			// (Otherwise quick changes to a shader file can result in the wrong output.)
-			FString ShaderPlatformNameString;
-			InputFile << ShaderPlatformNameString;
-			FName ShaderPlatformName = FName(*ShaderPlatformNameString);
 
-			FlushShaderFileCache(&ShaderPlatformName);
+			FlushShaderFileCache();
 			
 			for (int32 BatchIndex = 0; BatchIndex < NumBatches; BatchIndex++)
 			{
diff --git a/Engine/Source/Runtime/Engine/Private/ShaderCompiler/ShaderCompiler.cpp b/Engine/Source/Runtime/Engine/Private/ShaderCompiler/ShaderCompiler.cpp
--- a/Engine/Source/Runtime/Engine/Private/ShaderCompiler/ShaderCompiler.cpp
+++ b/Engine/Source/Runtime/Engine/Private/ShaderCompiler/ShaderCompiler.cpp
@@ -1517,9 +1517,6 @@ bool FShaderCompileUtilities::DoWriteTasks(const TArray<FShaderCommonCompileJobP
 		int32 NumBatches = QueuedSingleJobs.Num();
 		TransferFile << NumBatches;
 
-		FString ShaderPlatformNameString = FDataDrivenShaderPlatformInfo::GetName(GMaxRHIShaderPlatform).ToString();
-		TransferFile << ShaderPlatformNameString;
-
 		// Serialize all the batched jobs
 		for (int32 JobIndex = 0; JobIndex < QueuedSingleJobs.Num(); JobIndex++)
 		{
@@ -2040,6 +2037,7 @@ void FShaderCompileUtilities::DoReadTaskResults(const TArray<FShaderCommonCompil
 	}
 }
 
+#if WITH_EDITOR
 static bool CheckSingleJob(const FShaderCompileJob& SingleJob, TArray<FString>& OutErrors, FString* OutDumpedSource)
 {
 	if (SingleJob.bSucceeded)
@@ -2112,6 +2110,7 @@ static void AddErrorsForFailedJob(FShaderCompileJob& CurrentJob, TArray<EShaderP
 		}
 	}
 }
+#endif // WITH_EDITOR
 
 /** Information tracked for each shader compile worker process instance. */
 struct FShaderCompileWorkerInfo
@@ -5405,6 +5404,8 @@ bool FShaderCompilingManager::IsShaderCompilerWorkerRunning(FProcHandle & Worker
 	return FPlatformProcess::IsProcRunning(WorkerHandle);
 }
 
+#if WITH_EDITOR
+
 /* Generates a uniform buffer struct member hlsl declaration using the member's metadata. */
 static void GenerateUniformBufferStructMember(FString& Result, const FShaderParametersMetadata::FMember& Member, EShaderPlatform ShaderPlatform)
 {
@@ -5423,7 +5424,7 @@ static void GenerateUniformBufferStructMember(FString& Result, const FShaderPara
 }
 
 /* Generates the instanced stereo hlsl code that's dependent on view uniform declarations. */
-ENGINE_API void GenerateInstancedStereoCode(FString& Result, EShaderPlatform ShaderPlatform)
+void GenerateInstancedStereoCode(FString& Result, EShaderPlatform ShaderPlatform)
 {
 	// Find the InstancedView uniform buffer struct
 	const FShaderParametersMetadata* InstancedView = nullptr;
@@ -5579,9 +5580,7 @@ void GlobalBeginCompileShader(
 	EShaderPlatform ShaderPlatform = EShaderPlatform(Target.Platform);
 	const FName ShaderFormatName = LegacyShaderPlatformToShaderFormat(ShaderPlatform);
 
-#if WITH_EDITOR
 	FShaderCompileUtilities::GenerateBrdfHeaders(ShaderPlatform);
-#endif
 
 	Input.Target = Target;
 	Input.ShaderPlatformName = FDataDrivenShaderPlatformInfo::GetName(ShaderPlatform);
@@ -6092,13 +6091,11 @@ void GlobalBeginCompileShader(
 	{
 		// Allow GBuffer containing a velocity target to be overridden at a higher level with GBUFFER_LAYOUT
 		bool bGBufferHasVelocity = IsUsingBasePassVelocity((EShaderPlatform)Target.Platform);
-	#if WITH_EDITOR
 		if (!bGBufferHasVelocity)
 		{
 			const EGBufferLayout Layout = FShaderCompileUtilities::FetchGBufferLayout(Input.Environment);
 			bGBufferHasVelocity |= (Layout == GBL_ForceVelocity);
 		}
-	#endif
 		Input.Environment.SetDefine(TEXT("GBUFFER_HAS_VELOCITY"), bGBufferHasVelocity ? 1 : 0);
 	}
 
@@ -6366,11 +6363,10 @@ void GlobalBeginCompileShader(
 
 	// Allow the GBuffer and other shader defines to cause dependend environment changes, but minimizing the #ifdef magic in the shaders, which
 	// is nearly impossible to debug when it goes wrong.
-#if WITH_EDITOR
 	FShaderCompileUtilities::ApplyDerivedDefines(Input.Environment, Input.SharedEnvironment, (EShaderPlatform)Target.Platform);
-#endif
 }
 
+#endif // WITH_EDITOR
 
 /** Timer class used to report information on the 'recompileshaders' console command. */
 class FRecompileShadersTimer
@@ -6690,7 +6686,7 @@ bool RecompileShaders(const TCHAR* Cmd, FOutputDevice& Ar)
 				if( Material && Material->GetName() == RequestedMaterialName)
 				{
 					bMaterialFound = true;
-#if WITH_EDITOR
+
 					// <Pre/Post>EditChange will force a re-creation of the resource,
 					// in turn recompiling the shader.
 					if (TargetPlatform)
@@ -6708,7 +6704,7 @@ bool RecompileShaders(const TCHAR* Cmd, FOutputDevice& Ar)
 						Material->PreEditChange(nullptr);
 						Material->PostEditChange();
 					}
-#endif // WITH_EDITOR
+
 					break;
 				}
 			}
@@ -6732,12 +6728,11 @@ bool RecompileShaders(const TCHAR* Cmd, FOutputDevice& Ar)
 				{
 					UE_LOG(LogShaderCompilers, Log, TEXT("recompiling [%s]"),*Material->GetFullName());
 					UpdateContext.AddMaterial(Material);
-#if WITH_EDITOR
+
 					// <Pre/Post>EditChange will force a re-creation of the resource,
 					// in turn recompiling the shader.
 					Material->PreEditChange(nullptr);
 					Material->PostEditChange();
-#endif // WITH_EDITOR
 				}
 			}
 		}
@@ -7195,7 +7190,7 @@ static UE::DerivedData::FSharedString GetGlobalShaderMapName(const FGlobalShader
 {
 	return UE::DerivedData::FSharedString(WriteToString<256>(TEXTVIEW("GlobalShaderMap ["), LegacyShaderPlatformToShaderFormat(Platform), TEXTVIEW(", "), Key, TEXTVIEW("]")));
 }
-#endif
+#endif // WITH_EDITOR
 
 /** Saves the platform's shader map to the DDC. It is assumed that the caller will check IsComplete() first before calling the function. */
 static void SaveGlobalShaderMapToDerivedDataCache(EShaderPlatform Platform)
@@ -7260,7 +7255,7 @@ FString SaveGlobalShaderFile(EShaderPlatform Platform, FString SavePath, class I
 		}
 		
 		MemoryWriter.SetCookData(CookData.GetPtrOrNull());
-#endif
+#endif // WITH_EDITOR
 
 		GlobalShaderMap->SaveToGlobalArchive(MemoryWriter);
 	}
@@ -7719,7 +7714,7 @@ bool RecompileChangedShadersForPlatform(const FString& PlatformName)
 				(*It)->ClearCachedCookedPlatformData(TargetPlatform);
 			}
 		}
-#endif
+#endif // WITH_EDITOR
 	}
 
 	if (OutdatedFactoryTypes.Num() || OutdatedShaderTypes.Num())
diff --git a/Engine/Source/Runtime/Engine/Public/MaterialShaderType.h b/Engine/Source/Runtime/Engine/Public/MaterialShaderType.h
--- a/Engine/Source/Runtime/Engine/Public/MaterialShaderType.h
+++ b/Engine/Source/Runtime/Engine/Public/MaterialShaderType.h
@@ -98,20 +98,24 @@ public:
 		int32 InTotalPermutationCount,
 		ConstructSerializedType InConstructSerializedRef,
 		ConstructCompiledType InConstructCompiledRef,
-		ModifyCompilationEnvironmentType InModifyCompilationEnvironmentRef,
 		ShouldCompilePermutationType InShouldCompilePermutationRef,
-		ValidateCompiledResultType InValidateCompiledResultRef,
 		GetRayTracingPayloadTypeType InGetRayTracingPayloadTypeRef,
+#if WITH_EDITOR
+		ModifyCompilationEnvironmentType InModifyCompilationEnvironmentRef,
+		ValidateCompiledResultType InValidateCompiledResultRef,
+#endif // WITH_EDITOR
 		uint32 InTypeSize,
 		const FShaderParametersMetadata* InRootParametersMetadata = nullptr
 		):
 		FShaderType(EShaderTypeForDynamicCast::Material, InTypeLayout, InName, InSourceFilename, InFunctionName, InFrequency, InTotalPermutationCount,
 			InConstructSerializedRef,
 			InConstructCompiledRef,
-			InModifyCompilationEnvironmentRef,
 			InShouldCompilePermutationRef,
-			InValidateCompiledResultRef,
 			InGetRayTracingPayloadTypeRef,
+#if WITH_EDITOR
+			InModifyCompilationEnvironmentRef,
+			InValidateCompiledResultRef,
+#endif // WITH_EDITOR
 			InTypeSize,
 			InRootParametersMetadata
 		)
diff --git a/Engine/Source/Runtime/Engine/Public/MeshMaterialShaderType.h b/Engine/Source/Runtime/Engine/Public/MeshMaterialShaderType.h
--- a/Engine/Source/Runtime/Engine/Public/MeshMaterialShaderType.h
+++ b/Engine/Source/Runtime/Engine/Public/MeshMaterialShaderType.h
@@ -51,20 +51,24 @@ public:
 		int32 InTotalPermutationCount,
 		ConstructSerializedType InConstructSerializedRef,
 		ConstructCompiledType InConstructCompiledRef,
-		ModifyCompilationEnvironmentType InModifyCompilationEnvironmentRef,
 		ShouldCompilePermutationType InShouldCompilePermutationRef,
-		ValidateCompiledResultType InValidateCompiledResultRef,
 		GetRayTracingPayloadTypeType InGetRayTracingPayloadTypeRef,
+#if WITH_EDITOR
+		ModifyCompilationEnvironmentType InModifyCompilationEnvironmentRef,
+		ValidateCompiledResultType InValidateCompiledResultRef,
+#endif // WITH_EDITOR
 		uint32 InTypeSize,
 		const FShaderParametersMetadata* InRootParametersMetadata = nullptr
 		):
 		FShaderType(EShaderTypeForDynamicCast::MeshMaterial, InTypeLayout, InName,InSourceFilename,InFunctionName,InFrequency,InTotalPermutationCount,
 			InConstructSerializedRef,
 			InConstructCompiledRef,
-			InModifyCompilationEnvironmentRef,
 			InShouldCompilePermutationRef,
-			InValidateCompiledResultRef,
 			InGetRayTracingPayloadTypeRef,
+#if WITH_EDITOR
+			InModifyCompilationEnvironmentRef,
+			InValidateCompiledResultRef,
+#endif // WITH_EDITOR
 			InTypeSize,
 			InRootParametersMetadata
 		)
diff --git a/Engine/Source/Runtime/Engine/Public/ShaderCompiler.h b/Engine/Source/Runtime/Engine/Public/ShaderCompiler.h
--- a/Engine/Source/Runtime/Engine/Public/ShaderCompiler.h
+++ b/Engine/Source/Runtime/Engine/Public/ShaderCompiler.h
@@ -1160,8 +1160,7 @@ extern ENGINE_API FShaderCompilingManager* GShaderCompilingManager;
 /** The global shader compiling stats */
 extern ENGINE_API FShaderCompilerStats* GShaderCompilerStats;
 
-/** The shader precompilers for each platform.  These are only set during the console shader compilation while cooking or in the PrecompileShaders commandlet. */
-extern class FConsoleShaderPrecompiler* GConsoleShaderPrecompilers[SP_NumPlatforms];
+#if WITH_EDITOR
 
 /** Enqueues a shader compile job with GShaderCompilingManager. */
 extern ENGINE_API void GlobalBeginCompileShader(
@@ -1195,6 +1194,8 @@ extern ENGINE_API void GlobalBeginCompileShader(
 	const TCHAR* DebugExtension = nullptr
 );
 
+#endif // WITH_EDITOR
+
 extern void GetOutdatedShaderTypes(TArray<const FShaderType*>& OutdatedShaderTypes, TArray<const FShaderPipelineType*>& OutdatedShaderPipelineTypes, TArray<const FVertexFactoryType*>& OutdatedFactoryTypes);
 
 /** Implementation of the 'recompileshaders' console command.  Recompiles shaders at runtime based on various criteria. */
diff --git a/Engine/Source/Runtime/RenderCore/Private/Shader.cpp b/Engine/Source/Runtime/RenderCore/Private/Shader.cpp
--- a/Engine/Source/Runtime/RenderCore/Private/Shader.cpp
+++ b/Engine/Source/Runtime/RenderCore/Private/Shader.cpp
@@ -216,10 +216,12 @@ FShaderType::FShaderType(
 	int32 InTotalPermutationCount,
 	ConstructSerializedType InConstructSerializedRef,
 	ConstructCompiledType InConstructCompiledRef,
-	ModifyCompilationEnvironmentType InModifyCompilationEnvironmentRef,
 	ShouldCompilePermutationType InShouldCompilePermutationRef,
-	ValidateCompiledResultType InValidateCompiledResultRef,
 	GetRayTracingPayloadTypeType InGetRayTracingPayloadTypeRef,
+#if WITH_EDITOR
+	ModifyCompilationEnvironmentType InModifyCompilationEnvironmentRef,
+	ValidateCompiledResultType InValidateCompiledResultRef,
+#endif // WITH_EDITOR
 	uint32 InTypeSize,
 	const FShaderParametersMetadata* InRootParametersMetadata
 ):
@@ -236,17 +238,17 @@ FShaderType::FShaderType(
 	TotalPermutationCount(InTotalPermutationCount),
 	ConstructSerializedRef(InConstructSerializedRef),
 	ConstructCompiledRef(InConstructCompiledRef),
-	ModifyCompilationEnvironmentRef(InModifyCompilationEnvironmentRef),
 	ShouldCompilePermutationRef(InShouldCompilePermutationRef),
-	ValidateCompiledResultRef(InValidateCompiledResultRef),
 	GetRayTracingPayloadTypeRef(InGetRayTracingPayloadTypeRef),
+#if WITH_EDITOR
+	ModifyCompilationEnvironmentRef(InModifyCompilationEnvironmentRef),
+	ValidateCompiledResultRef(InValidateCompiledResultRef),
+#endif // WITH_EDITOR
 	RootParametersMetadata(InRootParametersMetadata),
 	GlobalListLink(this)
 {
 	FTypeLayoutDesc::Register(InTypeLayout);
 
-	CachedUniformBufferPlatform = SP_NumPlatforms;
-
 	// This will trigger if an IMPLEMENT_SHADER_TYPE was in a module not loaded before InitializeShaderTypes
 	// Shader types need to be implemented in modules that are loaded before that
 	checkf(!bInitializedSerializationHistory, TEXT("Shader type was loaded after engine init, use ELoadingPhase::PostConfigInit on your module to cause it to load earlier."));
@@ -379,6 +381,8 @@ bool FShaderType::ShouldCompilePermutation(const FShaderPermutationParameters& P
 	return ShouldCompileShaderFrequency((EShaderFrequency)Frequency, Parameters.Platform) && (*ShouldCompilePermutationRef)(Parameters);
 }
 
+#if WITH_EDITOR
+
 void FShaderType::ModifyCompilationEnvironment(const FShaderPermutationParameters& Parameters, FShaderCompilerEnvironment& OutEnvironment) const
 {
 	(*ModifyCompilationEnvironmentRef)(Parameters, OutEnvironment);
@@ -420,6 +424,8 @@ bool FShaderType::ValidateCompiledResult(EShaderPlatform Platform, const FShader
 	return (*ValidateCompiledResultRef)(Platform, ParameterMap, OutError);
 }
 
+#endif // WITH_EDITOR
+
 ERayTracingPayloadType FShaderType::GetRayTracingPayloadType(const int32 PermutationId) const
 {
 #if RHI_RAYTRACING
@@ -437,6 +443,7 @@ const FSHAHash& FShaderType::GetSourceHash(EShaderPlatform ShaderPlatform) const
 
 void FShaderType::Initialize(const TMap<FString, TArray<const TCHAR*> >& ShaderFileToUniformBufferVariables)
 {
+#if WITH_EDITOR
 	//#todo-rco: Need to call this only when Initializing from a Pipeline once it's removed from the global linked list
 	if (!FPlatformProperties::RequiresCookedData())
 	{
@@ -461,6 +468,7 @@ void FShaderType::Initialize(const TMap<FString, TArray<const TCHAR*> >& ShaderF
 		}
 #endif
 	}
+#endif // WITH_EDITOR
 
 	bInitializedSerializationHistory = true;
 }
diff --git a/Engine/Source/Runtime/RenderCore/Private/ShaderCore.cpp b/Engine/Source/Runtime/RenderCore/Private/ShaderCore.cpp
--- a/Engine/Source/Runtime/RenderCore/Private/ShaderCore.cpp
+++ b/Engine/Source/Runtime/RenderCore/Private/ShaderCore.cpp
@@ -726,6 +726,7 @@ void GetAllVirtualShaderSourcePaths(TArray<FString>& OutVirtualFilePaths, EShade
 */
 void VerifyShaderSourceFiles(EShaderPlatform ShaderPlatform)
 {
+#if WITH_EDITORONLY_DATA
 	if (!FPlatformProperties::RequiresCookedData() && AllowShaderCompiling())
 	{
 		// get the list of shader files that can be used
@@ -739,6 +740,7 @@ void VerifyShaderSourceFiles(EShaderPlatform ShaderPlatform)
 			LoadShaderSourceFile(*VirtualShaderSourcePaths[ShaderFileIdx], ShaderPlatform, nullptr, nullptr);
 		}
 	}
+#endif // WITH_EDITORONLY_DATA
 }
 
 static void LogShaderSourceDirectoryMappings()
@@ -936,6 +938,7 @@ bool ReplaceVirtualFilePathForShaderAutogen(FString& InOutVirtualFilePath, EShad
 
 bool LoadShaderSourceFile(const TCHAR* InVirtualFilePath, EShaderPlatform ShaderPlatform, FString* OutFileContents, TArray<FShaderCompilerError>* OutCompileErrors, const FName* ShaderPlatformName) // TODO: const FString&
 {
+#if WITH_EDITORONLY_DATA
 	// it's not expected that cooked platforms get here, but if they do, this is the final out
 	if (FPlatformProperties::RequiresCookedData())
 	{
@@ -1014,6 +1017,9 @@ bool LoadShaderSourceFile(const TCHAR* InVirtualFilePath, EShaderPlatform Shader
 	INC_FLOAT_STAT_BY(STAT_ShaderCompiling_LoadingShaderFiles,(float)ShaderFileLoadingTime);
 
 	return bResult;
+#else
+	return false;
+#endif // WITH_EDITORONLY_DATA
 }
 
 void LoadShaderSourceFileChecked(const TCHAR* VirtualFilePath, EShaderPlatform ShaderPlatform, FString& OutFileContents, const FName* ShaderPlatformName)
@@ -1405,12 +1411,13 @@ const FSHAHash& GetShaderFilesHash(const TArray<FString>& VirtualFilePaths, ESha
 	}
 }
 
-void BuildShaderFileToUniformBufferMap(TMap<FString, TArray<const TCHAR*> >& ShaderFileToUniformBufferVariables, const FName* ShaderPlatformName)
+#if WITH_EDITOR
+void BuildShaderFileToUniformBufferMap(TMap<FString, TArray<const TCHAR*> >& ShaderFileToUniformBufferVariables)
 {
 	if (!FPlatformProperties::RequiresCookedData())
 	{
 		TArray<FString> ShaderSourceFiles;
-		GetAllVirtualShaderSourcePaths(ShaderSourceFiles, GMaxRHIShaderPlatform, ShaderPlatformName);
+		GetAllVirtualShaderSourcePaths(ShaderSourceFiles, GMaxRHIShaderPlatform);
 
 		FScopedSlowTask SlowTask((float)ShaderSourceFiles.Num());
 
@@ -1441,7 +1448,7 @@ void BuildShaderFileToUniformBufferMap(TMap<FString, TArray<const TCHAR*> >& Sha
 			SlowTask.EnterProgressFrame(1);
 
  			FString ShaderFileContents;
-			LoadShaderSourceFileChecked(*ShaderSourceFiles[FileIndex], GMaxRHIShaderPlatform, ShaderFileContents, ShaderPlatformName);
+			LoadShaderSourceFileChecked(*ShaderSourceFiles[FileIndex], GMaxRHIShaderPlatform, ShaderFileContents);
 
 			// To allow case sensitive search which is way faster on some platforms (no need to look up locale, etc)
 			ShaderFileContents.ToUpperInline();
@@ -1460,6 +1467,7 @@ void BuildShaderFileToUniformBufferMap(TMap<FString, TArray<const TCHAR*> >& Sha
 		}
 	}
 }
+#endif // WITH_EDITOR
 
 void InitializeShaderHashCache()
 {
@@ -1491,7 +1499,9 @@ void InitializeShaderTypes()
 	LogShaderSourceDirectoryMappings();
 
 	TMap<FString, TArray<const TCHAR*> > ShaderFileToUniformBufferVariables;
+#if WITH_EDITOR
 	BuildShaderFileToUniformBufferMap(ShaderFileToUniformBufferVariables);
+#endif // WITH_EDITOR
 
 	FShaderType::Initialize(ShaderFileToUniformBufferVariables);
 	FVertexFactoryType::Initialize(ShaderFileToUniformBufferVariables);
@@ -1517,7 +1527,7 @@ void UninitializeShaderTypes()
  * Flushes the shader file and CRC cache, and regenerates the binary shader files if necessary.
  * Allows shader source files to be re-read properly even if they've been modified since startup.
  */
-void FlushShaderFileCache(const FName* ShaderPlatformName)
+void FlushShaderFileCache()
 {
 	UE_LOG(LogShaders, Log, TEXT("FlushShaderFileCache() begin"));
 
@@ -1530,12 +1540,13 @@ void FlushShaderFileCache(const FName* ShaderPlatformName)
 		GShaderFileCache.Empty();
 	}
 
+#if WITH_EDITOR
 	if (!FPlatformProperties::RequiresCookedData())
 	{
 		LogShaderSourceDirectoryMappings();
 
 		TMap<FString, TArray<const TCHAR*> > ShaderFileToUniformBufferVariables;
-		BuildShaderFileToUniformBufferMap(ShaderFileToUniformBufferVariables, ShaderPlatformName);
+		BuildShaderFileToUniformBufferMap(ShaderFileToUniformBufferVariables);
 
 		for (TLinkedList<FShaderPipelineType*>::TConstIterator It(FShaderPipelineType::GetTypeList()); It; It.Next())
 		{
@@ -1556,10 +1567,12 @@ void FlushShaderFileCache(const FName* ShaderPlatformName)
 			It->FlushShaderFileCache(ShaderFileToUniformBufferVariables);
 		}
 	}
+#endif // WITH_EDITOR
 
 	UE_LOG(LogShaders, Log, TEXT("FlushShaderFileCache() end"));
 }
 
+#if WITH_EDITOR
 void GenerateReferencedUniformBuffers(
 	const TCHAR* SourceFilename, 
 	const TCHAR* ShaderTypeName, 
@@ -1720,6 +1733,7 @@ void SerializeUniformBufferInfo_Internal(FShaderSaveArchive& Ar, const TSortedMa
 
 } // anonymous namespace
 
+
 FSHAHash GetShaderTypeLayoutHash(const FTypeLayoutDesc& TypeDesc, FPlatformTypeLayoutParameters LayoutParameters)
 {
 	static FFrozenMaterialLayoutHashCache GFrozenMaterialLayoutHashes;
@@ -1752,9 +1766,7 @@ void AppendKeyStringShaderDependencies(
 	for (const FShaderTypeDependency& ShaderTypeDependency : ShaderTypeDependencies)
 	{
 		const FShaderType* ShaderType = FindShaderTypeByName(ShaderTypeDependency.ShaderTypeName);
-#if WITH_EDITORONLY_DATA
 		checkf(ShaderType != nullptr, TEXT("Failed to find FShaderType for dependency %s (total in the NameToTypeMap: %d)"), ShaderTypeDependency.ShaderTypeName.GetDebugString().String.Get(), FShaderType::GetNameToTypeMap().Num());
-#endif
 
 		OutKeyString.AppendChar('_');
 		OutKeyString.Append(ShaderType->GetName());
@@ -1783,9 +1795,7 @@ void AppendKeyStringShaderDependencies(
 	for (const FShaderPipelineTypeDependency& Dependency : ShaderPipelineTypeDependencies)
 	{
 		const FShaderPipelineType* ShaderPipelineType = FShaderPipelineType::GetShaderPipelineTypeByName(Dependency.ShaderPipelineTypeName);
-#if WITH_EDITORONLY_DATA
 		checkf(ShaderPipelineType != nullptr, TEXT("Failed to find FShaderPipelineType for dependency %s (total in the NameToTypeMap: %d)"), Dependency.ShaderPipelineTypeName.GetDebugString().String.Get(), FShaderType::GetNameToTypeMap().Num());
-#endif
 
 		OutKeyString.AppendChar('_');
 		OutKeyString.Append(ShaderPipelineType->GetName());
@@ -1848,6 +1858,8 @@ void SerializeUniformBufferInfo(FShaderSaveArchive& Ar, const TSortedMap<const T
 	SerializeUniformBufferInfo_Internal(Ar, UniformBufferEntries);
 }
 
+#endif // WITH_EDITOR
+
 FString MakeInjectedShaderCodeBlock(const TCHAR* BlockName, const FString& CodeToInject)
 {
 	return FString::Printf(TEXT("#line 1 \"%s\"\n%s"), BlockName, *CodeToInject);
diff --git a/Engine/Source/Runtime/RenderCore/Private/ShaderParameters.cpp b/Engine/Source/Runtime/RenderCore/Private/ShaderParameters.cpp
--- a/Engine/Source/Runtime/RenderCore/Private/ShaderParameters.cpp
+++ b/Engine/Source/Runtime/RenderCore/Private/ShaderParameters.cpp
@@ -76,6 +76,7 @@ FArchive& operator<<(FArchive& Ar,FShaderResourceParameter& P)
 	return Ar << P.BaseIndex << P.NumResources;
 }
 
+#if WITH_EDITOR
 void FShaderUniformBufferParameter::ModifyCompilationEnvironment(const TCHAR* ParameterName,const FShaderParametersMetadata& Struct,EShaderPlatform Platform,FShaderCompilerEnvironment& OutEnvironment)
 {
 	const FString IncludeName = FString::Printf(TEXT("/Engine/Generated/UniformBuffers/%s.ush"),ParameterName);
@@ -90,6 +91,7 @@ void FShaderUniformBufferParameter::ModifyCompilationEnvironment(const TCHAR* Pa
 	GeneratedUniformBuffersInclude.Append(Include);
 	Struct.AddResourceTableEntries(OutEnvironment.ResourceTableMap, OutEnvironment.UniformBufferMap);
 }
+#endif // WITH_EDITOR
 
 void FShaderUniformBufferParameter::Bind(const FShaderParameterMap& ParameterMap,const TCHAR* ParameterName,EShaderParameterFlags Flags)
 {
@@ -120,6 +122,8 @@ void FShaderUniformBufferParameter::Bind(const FShaderParameterMap& ParameterMap
 	}
 }
 
+#if WITH_EDITOR
+
 /** The individual bits of a uniform buffer declaration. */
 struct FUniformBufferDecl
 {
@@ -451,6 +455,8 @@ void FShaderType::AddReferencedUniformBufferIncludes(FShaderCompilerEnvironment&
 	GeneratedUniformBuffersInclude.Append(UniformBufferIncludes);
 }
 
+#endif // WITH_EDITOR
+
 void FShaderType::DumpDebugInfo()
 {
 	UE_LOG(LogConsoleResponse, Display, TEXT("----------------------------- GLobalShader %s"), GetName());
@@ -488,9 +494,9 @@ void FShaderType::DumpDebugInfo()
 #endif
 }
 
+#if WITH_EDITOR
 void FShaderType::GetShaderStableKeyParts(FStableShaderKeyAndValue& SaveKeyVal)
 {
-#if WITH_EDITOR
 	static FName NAME_Material(TEXT("Material"));
 	static FName NAME_MeshMaterial(TEXT("MeshMaterial"));
 	static FName NAME_Niagara(TEXT("Niagara"));
@@ -510,7 +516,6 @@ void FShaderType::GetShaderStableKeyParts(FStableShaderKeyAndValue& SaveKeyVal)
 		break;
 	}
 	SaveKeyVal.ShaderType = FName(GetName() ? GetName() : TEXT("null"));
-#endif
 }
 
 void FVertexFactoryType::FlushShaderFileCache(const TMap<FString, TArray<const TCHAR*> >& ShaderFileToUniformBufferVariables)
@@ -588,3 +593,5 @@ void FVertexFactoryType::AddReferencedUniformBufferIncludes(FShaderCompilerEnvir
 	FString& GeneratedUniformBuffersInclude = OutEnvironment.IncludeVirtualPathToContentsMap.FindOrAdd("/Engine/Generated/GeneratedUniformBuffers.ush");
 	GeneratedUniformBuffersInclude.Append(UniformBufferIncludes);
 }
+
+#endif // WITH_EDITOR
diff --git a/Engine/Source/Runtime/RenderCore/Private/VertexFactory.cpp b/Engine/Source/Runtime/RenderCore/Private/VertexFactory.cpp
--- a/Engine/Source/Runtime/RenderCore/Private/VertexFactory.cpp
+++ b/Engine/Source/Runtime/RenderCore/Private/VertexFactory.cpp
@@ -54,6 +54,7 @@ FVertexFactoryType* FVertexFactoryType::GetVFByName(const FHashedName& VFName)
 
 void FVertexFactoryType::Initialize(const TMap<FString, TArray<const TCHAR*> >& ShaderFileToUniformBufferVariables)
 {
+#if WITH_EDITOR
 	if (!FPlatformProperties::RequiresCookedData())
 	{
 		// Cache serialization history for each VF type
@@ -64,6 +65,7 @@ void FVertexFactoryType::Initialize(const TMap<FString, TArray<const TCHAR*> >&
 			GenerateReferencedUniformBuffers(Type->ShaderFilename, Type->Name, ShaderFileToUniformBufferVariables, Type->ReferencedUniformBufferStructsCache);
 		}
 	}
+#endif // WITH_EDITOR
 
 	bInitializedSerializationHistory = true;
 }
@@ -81,9 +83,11 @@ FVertexFactoryType::FVertexFactoryType(
 	GetParameterTypeLayoutType InGetParameterTypeLayout,
 	GetParameterTypeElementShaderBindingsType InGetParameterTypeElementShaderBindings,
 	GetPSOPrecacheVertexFetchElementsType InGetPSOPrecacheVertexFetchElements,
-	ShouldCacheType InShouldCache,
-	ModifyCompilationEnvironmentType InModifyCompilationEnvironment,
-	ValidateCompiledResultType InValidateCompiledResult
+	ShouldCacheType InShouldCache
+#if WITH_EDITOR
+	, ModifyCompilationEnvironmentType InModifyCompilationEnvironment
+	, ValidateCompiledResultType InValidateCompiledResult
+#endif // WITH_EDITOR
 	):
 	Name(InName),
 	ShaderFilename(InShaderFilename),
@@ -95,8 +99,10 @@ FVertexFactoryType::FVertexFactoryType(
 	GetParameterTypeElementShaderBindings(InGetParameterTypeElementShaderBindings),
 	GetPSOPrecacheVertexFetchElements(InGetPSOPrecacheVertexFetchElements),
 	ShouldCacheRef(InShouldCache),
+#if WITH_EDITOR
 	ModifyCompilationEnvironmentRef(InModifyCompilationEnvironment),
 	ValidateCompiledResultRef(InValidateCompiledResult),
+#endif // WITH_EDITOR
 	GlobalListLink(this)
 {
 	// Make sure the format of the source file path is right.
@@ -106,8 +112,6 @@ FVertexFactoryType::FVertexFactoryType(
 		TEXT("Incorrect virtual shader path extension for vertex factory shader header '%s': Only .ush files should be included."),
 		InShaderFilename);
 
-	CachedUniformBufferPlatform = SP_NumPlatforms;
-
 	// This will trigger if an IMPLEMENT_VERTEX_FACTORY_TYPE was in a module not loaded before InitializeShaderTypes
 	// Vertex factory types need to be implemented in modules that are loaded before that
 	checkf(!bInitializedSerializationHistory, TEXT("VF type was loaded after engine init, use ELoadingPhase::PostConfigInit on your module to cause it to load earlier."));
diff --git a/Engine/Source/Runtime/RenderCore/Public/GlobalShader.h b/Engine/Source/Runtime/RenderCore/Public/GlobalShader.h
--- a/Engine/Source/Runtime/RenderCore/Public/GlobalShader.h
+++ b/Engine/Source/Runtime/RenderCore/Public/GlobalShader.h
@@ -93,20 +93,24 @@ public:
 		int32 InTotalPermutationCount,
 		ConstructSerializedType InConstructSerializedRef,
 		ConstructCompiledType InConstructCompiledRef,
-		ModifyCompilationEnvironmentType InModifyCompilationEnvironmentRef,
 		ShouldCompilePermutationType InShouldCompilePermutationRef,
-		ValidateCompiledResultType InValidateCompiledResultRef,
 		GetRayTracingPayloadTypeType InGetRayTracingPayloadTypeRef,
+#if WITH_EDITOR
+		ModifyCompilationEnvironmentType InModifyCompilationEnvironmentRef,
+		ValidateCompiledResultType InValidateCompiledResultRef,
+#endif // WITH_EDITOR
 		uint32 InTypeSize,
 		const FShaderParametersMetadata* InRootParametersMetadata = nullptr
 		):
 		FShaderType(EShaderTypeForDynamicCast::Global, InTypeLayout, InName, InSourceFilename, InFunctionName, InFrequency, InTotalPermutationCount,
 			InConstructSerializedRef,
 			InConstructCompiledRef,
-			InModifyCompilationEnvironmentRef,
 			InShouldCompilePermutationRef,
-			InValidateCompiledResultRef,
 			InGetRayTracingPayloadTypeRef,
+#if WITH_EDITOR
+			InModifyCompilationEnvironmentRef,
+			InValidateCompiledResultRef,
+#endif // WITH_EDITOR
 			InTypeSize,
 			InRootParametersMetadata
 		)
diff --git a/Engine/Source/Runtime/RenderCore/Public/Shader.h b/Engine/Source/Runtime/RenderCore/Public/Shader.h
--- a/Engine/Source/Runtime/RenderCore/Public/Shader.h
+++ b/Engine/Source/Runtime/RenderCore/Public/Shader.h
@@ -1194,9 +1194,11 @@ public:
 	typedef class FShader* (*ConstructSerializedType)();
 	typedef FShader* (*ConstructCompiledType)(const FShader::CompiledShaderInitializerType& Initializer);
 	typedef bool (*ShouldCompilePermutationType)(const FShaderPermutationParameters&);
+	typedef ERayTracingPayloadType(*GetRayTracingPayloadTypeType)(const int32 PermutationId);
+#if WITH_EDITOR
 	typedef void (*ModifyCompilationEnvironmentType)(const FShaderPermutationParameters&, FShaderCompilerEnvironment&);
 	typedef bool (*ValidateCompiledResultType)(EShaderPlatform, const FShaderParameterMap&, TArray<FString>&);
-	typedef ERayTracingPayloadType (*GetRayTracingPayloadTypeType)(const int32 PermutationId);
+#endif // WITH_EDITOR
 
 	/** @return The global shader factory list. */
 	static TLinkedList<FShaderType*>*& GetTypeList();
@@ -1226,10 +1228,12 @@ public:
 		int32 TotalPermutationCount,
 		ConstructSerializedType InConstructSerializedRef,
 		ConstructCompiledType InConstructCompiledRef,
-		ModifyCompilationEnvironmentType InModifyCompilationEnvironmentRef,
 		ShouldCompilePermutationType InShouldCompilePermutationRef,
-		ValidateCompiledResultType InValidateCompiledResultRef,
 		GetRayTracingPayloadTypeType InGetRayTracingPayloadTypeRef,
+#if WITH_EDITOR
+		ModifyCompilationEnvironmentType InModifyCompilationEnvironmentRef,
+		ValidateCompiledResultType InValidateCompiledResultRef,
+#endif // WITH_EDITOR
 		uint32 InTypeSize,
 		const FShaderParametersMetadata* InRootParametersMetadata
 	);
@@ -1241,6 +1245,8 @@ public:
 	FShader* ConstructCompiled(const FShader::CompiledShaderInitializerType& Initializer) const;
 
 	bool ShouldCompilePermutation(const FShaderPermutationParameters& Parameters) const;
+
+#if WITH_EDITOR
 	void ModifyCompilationEnvironment(const FShaderPermutationParameters& Parameters, FShaderCompilerEnvironment& OutEnvironment) const;
 
 	/**
@@ -1250,6 +1256,7 @@ public:
 	* @param OutError - List for appending validation errors.
 	*/
 	bool ValidateCompiledResult(EShaderPlatform Platform, const FShaderParameterMap& ParameterMap, TArray<FString>& OutError) const;
+#endif // WITH_EDITOR
 
 	ERayTracingPayloadType GetRayTracingPayloadType(const int32 PermutationId) const;
 
@@ -1399,24 +1406,27 @@ public:
 		return TotalPermutationCount;
 	}
 
-	inline const TMap<const TCHAR*, FCachedUniformBufferDeclaration>& GetReferencedUniformBufferStructsCache() const
-	{
-		return ReferencedUniformBufferStructsCache;
-	}
-
 	/** Returns the meta data for the root shader parameter struct. */
 	inline const FShaderParametersMetadata* GetRootParametersMetadata() const
 	{
 		return RootParametersMetadata;
 	}
 
+#if WITH_EDITOR
+	inline const TMap<const TCHAR*, FCachedUniformBufferDeclaration>& GetReferencedUniformBufferStructsCache() const
+	{
+		return ReferencedUniformBufferStructsCache;
+	}
+
 	/** Adds include statements for uniform buffers that this shader type references, and builds a prefix for the shader file with the include statements. */
 	void AddReferencedUniformBufferIncludes(FShaderCompilerEnvironment& OutEnvironment, FString& OutSourceFilePrefix, EShaderPlatform Platform) const;
 
 	void FlushShaderFileCache(const TMap<FString, TArray<const TCHAR*> >& ShaderFileToUniformBufferVariables);
 
-	void DumpDebugInfo();
 	void GetShaderStableKeyParts(struct FStableShaderKeyAndValue& SaveKeyVal);
+#endif // WITH_EDITOR
+
+	void DumpDebugInfo();
 
 private:
 	EShaderTypeForDynamicCast ShaderTypeForDynamicCast;
@@ -1433,10 +1443,12 @@ private:
 
 	ConstructSerializedType ConstructSerializedRef;
 	ConstructCompiledType ConstructCompiledRef;
-	ModifyCompilationEnvironmentType ModifyCompilationEnvironmentRef;
 	ShouldCompilePermutationType ShouldCompilePermutationRef;
-	ValidateCompiledResultType ValidateCompiledResultRef;
 	GetRayTracingPayloadTypeType GetRayTracingPayloadTypeRef;
+#if WITH_EDITOR
+	ModifyCompilationEnvironmentType ModifyCompilationEnvironmentRef;
+	ValidateCompiledResultType ValidateCompiledResultRef;
+#endif
 	const FShaderParametersMetadata* const RootParametersMetadata;
 
 	TLinkedList<FShaderType*> GlobalListLink;
@@ -1446,9 +1458,10 @@ private:
 	/** Tracks whether serialization history for all shader types has been initialized. */
 	static bool bInitializedSerializationHistory;
 
+#if WITH_EDITOR
 protected:
 	/** Tracks what platforms ReferencedUniformBufferStructsCache has had declarations cached for. */
-	mutable std::atomic<EShaderPlatform> CachedUniformBufferPlatform;
+	mutable std::atomic<EShaderPlatform> CachedUniformBufferPlatform{ SP_NumPlatforms };
 
 	/**
 	* Cache of referenced uniform buffer includes.
@@ -1456,7 +1469,7 @@ protected:
 	* FShaderType::Initialize will add an entry for each referenced uniform buffer, but the declarations are added on demand as shaders are compiled.
 	*/
 	mutable TMap<const TCHAR*, FCachedUniformBufferDeclaration> ReferencedUniformBufferStructsCache;
-
+#endif // WITH_EDITOR
 };
 
 struct FShaderCompiledShaderInitializerType
@@ -1486,18 +1499,25 @@ struct FShaderCompiledShaderInitializerType
 	);
 };
 
+#if WITH_EDITOR
+	#define SHADER_DECLARE_EDITOR_VTABLE(ShaderClass) \
+		static void ModifyCompilationEnvironmentImpl(const FShaderPermutationParameters& Parameters, FShaderCompilerEnvironment& OutEnvironment) \
+		{ \
+			const typename ShaderClass::FPermutationDomain PermutationVector(Parameters.PermutationId); \
+			PermutationVector.ModifyCompilationEnvironment(OutEnvironment); \
+			ShaderClass::ModifyCompilationEnvironment(static_cast<const typename ShaderClass::FPermutationParameters&>(Parameters), OutEnvironment); \
+		}
+#else
+	#define SHADER_DECLARE_EDITOR_VTABLE(ShaderClass)
+#endif // WITH_EDITOR
+
 #define SHADER_DECLARE_VTABLE(ShaderClass) \
 	static FShader* ConstructSerializedInstance() { return new ShaderClass(); } \
 	static FShader* ConstructCompiledInstance(const typename FShader::CompiledShaderInitializerType& Initializer) \
 	{ return new ShaderClass(static_cast<const typename ShaderMetaType::CompiledShaderInitializerType&>(Initializer)); }\
-	static void ModifyCompilationEnvironmentImpl(const FShaderPermutationParameters& Parameters, FShaderCompilerEnvironment& OutEnvironment) \
-	{ \
-		const typename ShaderClass::FPermutationDomain PermutationVector(Parameters.PermutationId); \
-		PermutationVector.ModifyCompilationEnvironment(OutEnvironment); \
-		ShaderClass::ModifyCompilationEnvironment(static_cast<const typename ShaderClass::FPermutationParameters&>(Parameters), OutEnvironment); \
-	} \
 	static bool ShouldCompilePermutationImpl(const FShaderPermutationParameters& Parameters) \
-	{ return ShaderClass::ShouldCompilePermutation(static_cast<const typename ShaderClass::FPermutationParameters&>(Parameters)); }
+	{ return ShaderClass::ShouldCompilePermutation(static_cast<const typename ShaderClass::FPermutationParameters&>(Parameters)); } \
+	SHADER_DECLARE_EDITOR_VTABLE(ShaderClass)
 
 
 #define INTERNAL_DECLARE_SHADER_TYPE_COMMON(ShaderClass,ShaderMetaTypeShortcut,RequiredAPI) \
@@ -1528,13 +1548,20 @@ struct FShaderCompiledShaderInitializerType
 	DECLARE_EXPORTED_TYPE_LAYOUT_EXPLICIT_BASES(ShaderClass,, NonVirtual, __VA_ARGS__); \
 	public:
 
+#if WITH_EDITOR
+#define SHADER_TYPE_EDITOR_VTABLE(ShaderClass) \
+	, ShaderClass::ModifyCompilationEnvironmentImpl \
+	, ShaderClass::ValidateCompiledResult
+#else
+#define SHADER_TYPE_EDITOR_VTABLE(ShaderClass)
+#endif
+
 #define SHADER_TYPE_VTABLE(ShaderClass) \
 	ShaderClass::ConstructSerializedInstance, \
 	ShaderClass::ConstructCompiledInstance, \
-	ShaderClass::ModifyCompilationEnvironmentImpl, \
 	ShaderClass::ShouldCompilePermutationImpl, \
-	ShaderClass::ValidateCompiledResult, \
-	ShaderClass::GetRayTracingPayloadType
+	ShaderClass::GetRayTracingPayloadType \
+	SHADER_TYPE_EDITOR_VTABLE(ShaderClass)
 
 #if !UE_BUILD_DOCS
 /** A macro to implement a shader type. */
diff --git a/Engine/Source/Runtime/RenderCore/Public/ShaderCore.h b/Engine/Source/Runtime/RenderCore/Public/ShaderCore.h
--- a/Engine/Source/Runtime/RenderCore/Public/ShaderCore.h
+++ b/Engine/Source/Runtime/RenderCore/Public/ShaderCore.h
@@ -1074,16 +1074,22 @@ extern RENDERCORE_API void HashShaderFileWithIncludes(FArchive& HashingArchive,
  */
 extern RENDERCORE_API const class FSHAHash& GetShaderFilesHash(const TArray<FString>& VirtualFilePaths, EShaderPlatform ShaderPlatform);
 
-extern void BuildShaderFileToUniformBufferMap(TMap<FString, TArray<const TCHAR*> >& ShaderFileToUniformBufferVariables, const FName* ShaderPlatformName = nullptr);
-
 /**
  * Flushes the shader file and CRC cache, and regenerates the binary shader files if necessary.
  * Allows shader source files to be re-read properly even if they've been modified since startup.
  */
-extern RENDERCORE_API void FlushShaderFileCache(const FName* ShaderPlatformName = nullptr);
+extern RENDERCORE_API void FlushShaderFileCache();
+
+UE_DEPRECATED(5.2, "FlushShaderFileCache no longer needs a ShaderPlatformName argument")
+inline void FlushShaderFileCache(const FName* ShaderPlatformName)
+{
+	FlushShaderFileCache();
+}
 
 extern RENDERCORE_API void VerifyShaderSourceFiles(EShaderPlatform ShaderPlatform);
 
+#if WITH_EDITOR
+
 struct FCachedUniformBufferDeclaration
 {
 	// Using SharedPtr so we can hand off lifetime ownership to FShaderCompilerEnvironment::IncludeVirtualPathToExternalContentsMap when invalidating this cache
@@ -1110,7 +1116,8 @@ using FSortedMapUniformBufferDeclaration = TSortedMap<const TCHAR*, FCachedUnifo
 /** Records information about all the uniform buffer layouts referenced by UniformBufferEntries. This function is now deprecated as there now a unified way of preparing
  * shadermap keys that includes this and more: AppendKeyStringShaderDependencies.
  */
-UE_DEPRECATED(5.2, "SerializeUniformBufferInfo is depreceated. For creating shadermap keys please use AppendKeyStringShaderDependencies") extern RENDERCORE_API void SerializeUniformBufferInfo(class FShaderSaveArchive& Ar, const FSortedMapUniformBufferDeclaration& UniformBufferEntries);
+UE_DEPRECATED(5.2, "SerializeUniformBufferInfo is depreceated. For creating shadermap keys please use AppendKeyStringShaderDependencies")
+extern RENDERCORE_API void SerializeUniformBufferInfo(class FShaderSaveArchive& Ar, const FSortedMapUniformBufferDeclaration& UniformBufferEntries);
 
 /**
  * Return the hash of the given type layout for a partical platform type layout. This function employs caching to avoid re-hashing the same parameters several times.
@@ -1127,12 +1134,14 @@ extern RENDERCORE_API void AppendKeyStringShaderDependencies(
 	TConstArrayView<FShaderTypeDependency> ShaderTypeDependencies,
 	FPlatformTypeLayoutParameters LayoutParams,
 	FString& OutKeyString);
+
 extern RENDERCORE_API void AppendKeyStringShaderDependencies(
 	TConstArrayView<FShaderTypeDependency> ShaderTypeDependencies,
 	TConstArrayView<FShaderPipelineTypeDependency> ShaderPipelineTypeDependencies,
 	TConstArrayView<FVertexFactoryTypeDependency> VertexFactoryTypeDependencies,
 	FPlatformTypeLayoutParameters LayoutParams,
 	FString& OutKeyString);
+#endif // WITH_EDITOR
 
 /** Create a block of source code to be injected in the preprocessed shader code. The Block will be put into a #line directive
  * to show up in case shader compilation failures happen in this code block.
diff --git a/Engine/Source/Runtime/RenderCore/Public/ShaderParameters.h b/Engine/Source/Runtime/RenderCore/Public/ShaderParameters.h
--- a/Engine/Source/Runtime/RenderCore/Public/ShaderParameters.h
+++ b/Engine/Source/Runtime/RenderCore/Public/ShaderParameters.h
@@ -29,7 +29,12 @@ struct FShaderCompilerEnvironment;
 enum class EShaderParameterType : uint8;
 DECLARE_INTRINSIC_TYPE_LAYOUT(EShaderParameterType);
 
+#if WITH_EDITOR
+/** Creates a shader code declaration of this struct for the given shader platform. */
+RENDERCORE_API void CreateUniformBufferShaderDeclaration(const TCHAR* Name, const FShaderParametersMetadata& UniformBufferStruct, EShaderPlatform Platform, FString& OutDeclaration);
+
 RENDERCORE_API void CacheUniformBufferIncludes(TMap<const TCHAR*, struct FCachedUniformBufferDeclaration>& Cache, EShaderPlatform Platform);
+#endif
 
 
 enum EShaderParameterFlags
@@ -149,9 +154,6 @@ private:
 	LAYOUT_FIELD(FShaderResourceParameter, UAVParameter);
 };
 
-/** Creates a shader code declaration of this struct for the given shader platform. */
-extern RENDERCORE_API void CreateUniformBufferShaderDeclaration(const TCHAR* Name,const FShaderParametersMetadata& UniformBufferStruct, EShaderPlatform Platform, FString& OutDeclaration);
-
 class FShaderUniformBufferParameter
 {
 	DECLARE_EXPORTED_TYPE_LAYOUT(FShaderUniformBufferParameter, RENDERCORE_API, NonVirtual);
@@ -160,7 +162,9 @@ public:
 	:	BaseIndex(0xffff)
 	{}
 
+#if WITH_EDITOR
 	static RENDERCORE_API void ModifyCompilationEnvironment(const TCHAR* ParameterName,const FShaderParametersMetadata& Struct,EShaderPlatform Platform,FShaderCompilerEnvironment& OutEnvironment);
+#endif // WITH_EDITOR
 
 	RENDERCORE_API void Bind(const FShaderParameterMap& ParameterMap,const TCHAR* ParameterName,EShaderParameterFlags Flags = SPF_Optional);
 
@@ -192,10 +196,12 @@ template<typename TBufferStruct>
 class TShaderUniformBufferParameter : public FShaderUniformBufferParameter
 {
 public:
+#if WITH_EDITOR
 	static void ModifyCompilationEnvironment(const TCHAR* ParameterName,EShaderPlatform Platform, FShaderCompilerEnvironment& OutEnvironment)
 	{
 		FShaderUniformBufferParameter::ModifyCompilationEnvironment(ParameterName,TBufferStruct::StaticStruct,Platform,OutEnvironment);
 	}
+#endif // WITH_EDITOR
 
 	friend FArchive& operator<<(FArchive& Ar,TShaderUniformBufferParameter& P)
 	{
diff --git a/Engine/Source/Runtime/RenderCore/Public/VertexFactory.h b/Engine/Source/Runtime/RenderCore/Public/VertexFactory.h
--- a/Engine/Source/Runtime/RenderCore/Public/VertexFactory.h
+++ b/Engine/Source/Runtime/RenderCore/Public/VertexFactory.h
@@ -349,10 +349,12 @@ public:
 		class FMeshDrawSingleShaderBindings& ShaderBindings,
 		FVertexInputStreamArray& VertexStreams);
 	typedef void (*GetPSOPrecacheVertexFetchElementsType)(EVertexInputStreamType VertexInputStreamType, FVertexDeclarationElementList& Elements);
-
 	typedef bool (*ShouldCacheType)(const FVertexFactoryShaderPermutationParameters&);
+
+#if WITH_EDITOR
 	typedef void (*ModifyCompilationEnvironmentType)(const FVertexFactoryShaderPermutationParameters&, FShaderCompilerEnvironment&);
 	typedef void (*ValidateCompiledResultType)(const FVertexFactoryType*, EShaderPlatform, const FShaderParameterMap& ParameterMap, TArray<FString>& OutErrors);
+#endif // WITH_EDITOR
 
 	static int32 GetNumVertexFactoryTypes() { return NumVertexFactories; }
 
@@ -382,9 +384,11 @@ public:
 		GetParameterTypeLayoutType InGetParameterTypeLayout,
 		GetParameterTypeElementShaderBindingsType InGetParameterTypeElementShaderBindings,
 		GetPSOPrecacheVertexFetchElementsType InGetPSOPrecacheVertexFetchElements,
-		ShouldCacheType InShouldCache,
-		ModifyCompilationEnvironmentType InModifyCompilationEnvironment,
-		ValidateCompiledResultType InValidateCompiledResult
+		ShouldCacheType InShouldCache
+#if WITH_EDITOR
+		, ModifyCompilationEnvironmentType InModifyCompilationEnvironment
+		, ValidateCompiledResultType InValidateCompiledResult
+#endif // WITH_EDITOR
 		);
 
 	RENDERCORE_API virtual ~FVertexFactoryType();
@@ -461,6 +465,7 @@ public:
 		return (*ShouldCacheRef)(Parameters);
 	}
 
+#if WITH_EDITOR
 	/**
 	* Calls the function ptr for the shader type on the given environment
 	* @param Environment - shader compile environment to modify
@@ -485,10 +490,12 @@ public:
 	RENDERCORE_API void AddReferencedUniformBufferIncludes(FShaderCompilerEnvironment& OutEnvironment, FString& OutSourceFilePrefix, EShaderPlatform Platform) const;
 
 	RENDERCORE_API void FlushShaderFileCache(const TMap<FString, TArray<const TCHAR*> >& ShaderFileToUniformBufferVariables);
+
 	const TMap<const TCHAR*, FCachedUniformBufferDeclaration>& GetReferencedUniformBufferStructsCache() const
 	{
 		return ReferencedUniformBufferStructsCache;
 	}
+#endif // WITH_EDITOR
 
 private:
 	static RENDERCORE_API uint32 NumVertexFactories;
@@ -506,11 +513,14 @@ private:
 	GetParameterTypeElementShaderBindingsType GetParameterTypeElementShaderBindings;
 	GetPSOPrecacheVertexFetchElementsType GetPSOPrecacheVertexFetchElements;
 	ShouldCacheType ShouldCacheRef;
+#if WITH_EDITOR
 	ModifyCompilationEnvironmentType ModifyCompilationEnvironmentRef;
 	ValidateCompiledResultType ValidateCompiledResultRef;
+#endif // WITH_EDITOR
 
 	TLinkedList<FVertexFactoryType*> GlobalListLink;
 
+#if WITH_EDITOR
 	/** 
 	 * Cache of referenced uniform buffer includes.  
 	 * These are derived from source files so they need to be flushed when editing and recompiling shaders on the fly. 
@@ -519,7 +529,8 @@ private:
 	mutable TMap<const TCHAR*, FCachedUniformBufferDeclaration> ReferencedUniformBufferStructsCache;
 
 	/** Tracks what platforms ReferencedUniformBufferStructsCache has had declarations cached for. */
-	mutable std::atomic<EShaderPlatform> CachedUniformBufferPlatform;
+	mutable std::atomic<EShaderPlatform> CachedUniformBufferPlatform{ SP_NumPlatforms };
+#endif // WITH_EDITOR
 };
 
 /**
@@ -541,14 +552,21 @@ extern RENDERCORE_API FVertexFactoryType* FindVertexFactoryType(const FHashedNam
 	static FVertexFactoryType StaticType; \
 	virtual FVertexFactoryType* GetType() const override;
 
+#if WITH_EDITOR
+	#define IMPLEMENT_VERTEX_FACTORY_EDITOR_VTABLE(FactoryClass) \
+		, FactoryClass::ModifyCompilationEnvironment \
+		, FactoryClass::ValidateCompiledResult
+#else
+	#define IMPLEMENT_VERTEX_FACTORY_EDITOR_VTABLE(FactoryClass)
+#endif // WITH_EDITOR
+
 #define IMPLEMENT_VERTEX_FACTORY_VTABLE(FactoryClass) \
 	&ConstructVertexFactoryParameters<FactoryClass>, \
 	&GetVertexFactoryParametersLayout<FactoryClass>, \
 	&GetVertexFactoryParametersElementShaderBindings<FactoryClass>, \
 	FactoryClass::GetPSOPrecacheVertexFetchElements, \
-	FactoryClass::ShouldCompilePermutation, \
-	FactoryClass::ModifyCompilationEnvironment, \
-	FactoryClass::ValidateCompiledResult
+	FactoryClass::ShouldCompilePermutation \
+	IMPLEMENT_VERTEX_FACTORY_EDITOR_VTABLE(FactoryClass)
 
 /**
  * A macro for implementing the static vertex factory type object, and specifying parameters used by the type.
