Do not log stereo aspects in non-rendering runs and commandlets.

- Logging in the interactive editor is still left for the time being.

diff --git a/Engine/Source/Runtime/RenderCore/Private/StereoRenderUtils.cpp b/Engine/Source/Runtime/RenderCore/Private/StereoRenderUtils.cpp
--- a/Engine/Source/Runtime/RenderCore/Private/StereoRenderUtils.cpp
+++ b/Engine/Source/Runtime/RenderCore/Private/StereoRenderUtils.cpp
@@ -15,7 +15,7 @@
 		{ \
 			LastLogged[1] = LastLogged[0]; \
 			LastLogged[0] = Platform; \
-			bSSALogEnable = true; \
+			bSSALogEnable = FApp::CanEverRender() && !IsRunningCommandlet(); \
 		} \
 	}
 
