Required changes to enable dynamic ray tracing by default on PC:
* r.RayTracing.EnableOnDemand - controls whether ray tracing can be toggled on demand at runtime.
* r.RayTracing.Enable - now controls if ray tracing is enabled
* Set appropriate defaults for the engine in WindowsEngine.ini to match previous state
* Enable toggling ray tracing on/off only in frontend

diff --git a/Engine/Source/Runtime/RenderCore/Private/RenderUtils.cpp b/Engine/Source/Runtime/RenderCore/Private/RenderUtils.cpp
--- a/Engine/Source/Runtime/RenderCore/Private/RenderUtils.cpp
+++ b/Engine/Source/Runtime/RenderCore/Private/RenderUtils.cpp
@@ -53,6 +53,16 @@ FAutoConsoleVariableRef CVarRayTracingEnableInEditor(
 	ECVF_ReadOnly
 );
 
+static int32 GRayTracingEnableOnDemand = 0;
+static FAutoConsoleVariableRef CVarRayTracingEnableOnDemand(
+	TEXT("r.RayTracing.EnableOnDemand"),
+	GRayTracingEnableOnDemand,
+	TEXT("Controls whether ray tracing features can be toggled on demand at runtime without restarting the game (experimental).\n")
+	TEXT("Requires r.RayTracing=1 and will override GameUserSettings in game. Requires r.RayTracing.EnableInEditor=1 in editor. Has a small performance and memory overhead.\n")
+	TEXT(" 0: off (default)\n")
+	TEXT(" 1: on"),
+	ECVF_RenderThreadSafe | ECVF_ReadOnly);
+
 static int32 GRayTracingRequireSM6 = 0;
 static FAutoConsoleVariableRef CVarRayTracingRequireSM6(
 	TEXT("r.RayTracing.RequireSM6"),
@@ -1466,7 +1476,19 @@ RENDERCORE_API void RenderUtilsInit()
 	if (RayTracingCVar && RayTracingCVar->GetInt())
 	{
 		const int32 RayTracingInt = RayTracingCVar->GetInt();
-		ERayTracingMode DesiredRayTracingMode = (ERayTracingMode)FMath::Clamp<int32>(RayTracingInt, 0, 2);
+
+		ERayTracingMode DesiredRayTracingMode = ERayTracingMode::Disabled;
+		if (RayTracingInt != 0)
+		{
+			if (GRayTracingEnableOnDemand == 1)
+			{
+				DesiredRayTracingMode = ERayTracingMode::Dynamic;
+			}
+			else
+			{
+				DesiredRayTracingMode = ERayTracingMode::Enabled;
+			}
+		}
 
 		const bool bRayTracingAllowedOnCurrentPlatform = (GRayTracingPlatformMask[(int)GMaxRHIShaderPlatform]);
 		if (GRHISupportsRayTracing && bRayTracingAllowedOnCurrentPlatform)
@@ -1488,7 +1510,15 @@ RENDERCORE_API void RenderUtilsInit()
 				// If user preference exists in game settings file, the bRayTracingEnabled will be set based on its value.
 				// Otherwise the current value is preserved.
 				bool bUseRayTracing = false;
-				if (GConfig->GetBool(TEXT("RayTracing"), TEXT("r.RayTracing.EnableInGame"), bUseRayTracing, GGameUserSettingsIni))
+				if (GRayTracingEnableOnDemand == 1)
+				{
+					GRayTracingMode = DesiredRayTracingMode;
+
+					UE_LOG(LogRendererCore, Log, TEXT("Ray tracing is %s for the game. Reason: r.RayTracing=%d and r.RayTracing.EnableOnDemand=1."),
+						GetRayTracingModeName(GRayTracingMode),
+						RayTracingInt);
+				}
+				else if (GConfig->GetBool(TEXT("RayTracing"), TEXT("r.RayTracing.EnableInGame"), bUseRayTracing, GGameUserSettingsIni))
 				{
 					GRayTracingMode = bUseRayTracing ? DesiredRayTracingMode : ERayTracingMode::Disabled;
 
@@ -1500,7 +1530,7 @@ RENDERCORE_API void RenderUtilsInit()
 				{
 					GRayTracingMode = GRayTracingEnableInGame != 0 ? DesiredRayTracingMode : ERayTracingMode::Disabled;
 
-					UE_LOG(LogRendererCore, Log, TEXT("Ray tracing is %s for the game. Reason: r.RayTracing=%d, and r.RayTracing.EnableInGame game user setting does not exist (using default from CVar: %d)."),						
+					UE_LOG(LogRendererCore, Log, TEXT("Ray tracing is %s for the game. Reason: CVar r.RayTracing=%d, and r.RayTracing.EnableInGame game user setting does not exist (using default from CVar: %d)."),
 						GetRayTracingModeName(GRayTracingMode),
 						RayTracingInt,
 						GRayTracingEnableInGame);
