Add fatal log to ~FShaderLibraryInstance.  It is impossible for it to safely release resources from it's destructor.

diff --git a/Engine/Source/Runtime/RenderCore/Private/ShaderCodeLibrary.cpp b/Engine/Source/Runtime/RenderCore/Private/ShaderCodeLibrary.cpp
--- a/Engine/Source/Runtime/RenderCore/Private/ShaderCodeLibrary.cpp
+++ b/Engine/Source/Runtime/RenderCore/Private/ShaderCodeLibrary.cpp
@@ -899,13 +899,11 @@ public:
 
 	~FShaderLibraryInstance()
 	{
-		// release RHI on all of the resources
+		// FShaderLibraryInstance cant release RHI on all of the resources because it does not own them
+		// and its lifetime must also be longer than any of the resources.
 		for (FShaderMapResource_SharedCode* Resource : Resources)
 		{
-			if (Resource)
-			{
-				BeginReleaseResource(Resource);
-			}
+			UE_CLOG(Resource, LogShaderLibrary, Fatal, TEXT("FShaderLibraryInstance %s still has resources in use at time of destruction!"), *Library->GetName());
 		}
 		
 		Library->Teardown();
