Cache uniform buffer definitions
- Remove shader platform specific info from uniform buffers, we can generate code with macros that translates to what each platform needs.
- With the per-platform differences removed, uniform buffer declarations (read: generated text) can be created once per shader parameter type that needs them instead of once per shader type's shader parameter type usage.
- FShaderType and FVertexFactoryType now hold a TSet of used uniform buffer names instead of a TMap that held potentially duplicated FStrings of uniform buffer defintions.
- Lists of uniform buffers are *potentially* different per platform, so the annoying CachedUniformBufferPlatform variables need to stick around for now.
- GenerateReferencedUniformBuffers renamed to GenerateReferencedUniformBufferNames and now populates a TSet instead of a TMap.
- ShaderType::AddReferencedUniformBufferIncludes renamed to ShaderType::AddUniformBufferIncludesToEnvironment which calls the shared UE::ShaderParameters::AddUniformBufferIncludesToEnvironment.
- CreateUniformBufferShaderDeclaration moved to UE::ShaderParameters::CreateUniformBufferShaderDeclaration and now returns its FString.
- CacheUniformBufferIncludes was removed since it was only populated values from CreateUniformBufferShaderDeclaration.
- Deprecated FShaderCompilerInput::SourceFilePrefix since nothing has populated it in 9 years.

diff --git a/Engine/Plugins/Compositing/OpenColorIO/Source/OpenColorIO/Private/OpenColorIOShaderMap.cpp b/Engine/Plugins/Compositing/OpenColorIO/Source/OpenColorIO/Private/OpenColorIOShaderMap.cpp
--- a/Engine/Plugins/Compositing/OpenColorIO/Source/OpenColorIO/Private/OpenColorIOShaderMap.cpp
+++ b/Engine/Plugins/Compositing/OpenColorIO/Source/OpenColorIO/Private/OpenColorIOShaderMap.cpp
@@ -169,7 +169,7 @@ void FOpenColorIOShaderType::BeginCompileShader(
 	NewJob->Input.EntryPointName = TEXT("MainPS");
 	NewJob->Input.Environment.IncludeVirtualPathToContentsMap.Add(TEXT("/Engine/Generated/OpenColorIOTransformShader.ush"), InColorTransform->ShaderCode);
 
-	AddReferencedUniformBufferIncludes(NewJob->Input.Environment, NewJob->Input.SourceFilePrefix, InPlatform);
+	AddUniformBufferIncludesToEnvironment(NewJob->Input.Environment, InPlatform);
 
 	FShaderCompilerEnvironment& ShaderEnvironment = NewJob->Input.Environment;
 
diff --git a/Engine/Plugins/FX/Niagara/Source/NiagaraShader/Private/NiagaraShader.cpp b/Engine/Plugins/FX/Niagara/Source/NiagaraShader/Private/NiagaraShader.cpp
--- a/Engine/Plugins/FX/Niagara/Source/NiagaraShader/Private/NiagaraShader.cpp
+++ b/Engine/Plugins/FX/Niagara/Source/NiagaraShader/Private/NiagaraShader.cpp
@@ -360,7 +360,7 @@ void FNiagaraShaderType::BeginCompileShader(
 
 	Script->ModifyCompilationEnvironment(Platform, NewJob->Input.Environment);
 
-	AddReferencedUniformBufferIncludes(NewJob->Input.Environment, NewJob->Input.SourceFilePrefix, (EShaderPlatform)Target.Platform);
+	AddUniformBufferIncludesToEnvironment(NewJob->Input.Environment, Platform);
 	
 	FShaderCompilerEnvironment& ShaderEnvironment = NewJob->Input.Environment;
 
@@ -390,57 +390,18 @@ void FNiagaraShaderType::BeginCompileShader(
 
 void FNiagaraShaderType::CacheUniformBufferIncludes(TMap<const TCHAR*, FCachedUniformBufferDeclaration>& Cache, EShaderPlatform Platform) const
 {
-	for (TMap<const TCHAR*, FCachedUniformBufferDeclaration>::TIterator It(Cache); It; ++It)
-	{
-		FCachedUniformBufferDeclaration& BufferDeclaration = It.Value();
-
-		for (TLinkedList<FShaderParametersMetadata*>::TIterator StructIt(FShaderParametersMetadata::GetStructList()); StructIt; StructIt.Next())
-		{
-			if (It.Key() == StructIt->GetShaderVariableName())
-			{
-				BufferDeclaration.Declaration = MakeShared<FString, ESPMode::ThreadSafe>();
-				CreateUniformBufferShaderDeclaration(StructIt->GetShaderVariableName(), **StructIt, Platform, *BufferDeclaration.Declaration.Get());
-				break;
-			}
-		}
-	}
+PRAGMA_DISABLE_DEPRECATION_WARNINGS
+	::CacheUniformBufferIncludes(Cache, Platform);
+PRAGMA_ENABLE_DEPRECATION_WARNINGS
 }
 
-void FNiagaraShaderType::AddReferencedUniformBufferIncludes(FShaderCompilerEnvironment& OutEnvironment, FString& OutSourceFilePrefix, EShaderPlatform Platform) const
+void FNiagaraShaderType::AddUniformBufferIncludesToEnvironment(FShaderCompilerEnvironment& OutEnvironment, EShaderPlatform Platform) const
 {
-	// Cache uniform buffer struct declarations referenced by this shader type's files
-	if (CachedUniformBufferPlatform != Platform)
-	{
-		CacheUniformBufferIncludes(ReferencedUniformBufferStructsCache, Platform);
-		CachedUniformBufferPlatform = Platform;
-	}
+	CachedUniformBufferPlatform = Platform;
 
-	FString UniformBufferIncludes;
-
-	for (TMap<const TCHAR*, FCachedUniformBufferDeclaration>::TConstIterator It(ReferencedUniformBufferStructsCache); It; ++It)
-	{
-		check(It.Value().Declaration.IsValid());
-		UniformBufferIncludes += FString::Printf(TEXT("#include \"/Engine/Generated/UniformBuffers/%s.ush\"") LINE_TERMINATOR, It.Key());
-		FString* Declaration = It.Value().Declaration.Get();
-		check(Declaration);
-		OutEnvironment.IncludeVirtualPathToContentsMap.Add(
-			*FString::Printf(TEXT("/Engine/Generated/UniformBuffers/%s.ush"), It.Key()), *Declaration
-		);
-
-		for (TLinkedList<FShaderParametersMetadata*>::TIterator StructIt(FShaderParametersMetadata::GetStructList()); StructIt; StructIt.Next())
-		{
-			if (It.Key() == StructIt->GetShaderVariableName())
-			{
-				StructIt->AddResourceTableEntries(OutEnvironment.ResourceTableMap, OutEnvironment.UniformBufferMap);
-			}
-		}
-	}
-
-	FString& GeneratedUniformBuffersInclude = OutEnvironment.IncludeVirtualPathToContentsMap.FindOrAdd("/Engine/Generated/GeneratedUniformBuffers.ush");
-	GeneratedUniformBuffersInclude.Append(UniformBufferIncludes);
+	UE::ShaderParameters::AddUniformBufferIncludesToEnvironment(OutEnvironment, ReferencedUniformBufferNames);
 }
 
-
 /**
  * Either creates a new instance of this type or returns an equivalent existing shader.
  * @param CurrentJob - Compile job that was enqueued by BeginCompileShader.
diff --git a/Engine/Plugins/FX/Niagara/Source/NiagaraShader/Public/NiagaraShaderType.h b/Engine/Plugins/FX/Niagara/Source/NiagaraShader/Public/NiagaraShaderType.h
--- a/Engine/Plugins/FX/Niagara/Source/NiagaraShader/Public/NiagaraShaderType.h
+++ b/Engine/Plugins/FX/Niagara/Source/NiagaraShader/Public/NiagaraShaderType.h
@@ -154,9 +154,16 @@ public:
 
 #if WITH_EDITOR
 	/** Adds include statements for uniform buffers that this shader type references, and builds a prefix for the shader file with the include statements. */
-	void AddReferencedUniformBufferIncludes(FShaderCompilerEnvironment& OutEnvironment, FString& OutSourceFilePrefix, EShaderPlatform Platform) const;
+	void AddUniformBufferIncludesToEnvironment(FShaderCompilerEnvironment& OutEnvironment, EShaderPlatform Platform) const;
 
-	void CacheUniformBufferIncludes(TMap<const TCHAR*, FCachedUniformBufferDeclaration>& Cache, EShaderPlatform Platform) const;
+	UE_DEPRECATED(5.2, "AddReferencedUniformBufferIncludes has moved to AddUniformBufferIncludesToEnvironment and no longer takes a prefix argument.")
+	inline void AddReferencedUniformBufferIncludes(FShaderCompilerEnvironment& OutEnvironment, FString& OutSourceFilePrefix, EShaderPlatform Platform) const
+	{
+		AddUniformBufferIncludesToEnvironment(OutEnvironment, Platform);
+	}
+
+	UE_DEPRECATED(5.2, "CacheUniformBufferIncludes should no longer be used.")
+	inline void CacheUniformBufferIncludes(TMap<const TCHAR*, FCachedUniformBufferDeclaration>& Cache, EShaderPlatform Platform) const;
 
 protected:
 
diff --git a/Engine/Plugins/Runtime/ComputeFramework/Source/ComputeFramework/Private/ComputeKernelShaderMap.cpp b/Engine/Plugins/Runtime/ComputeFramework/Source/ComputeFramework/Private/ComputeKernelShaderMap.cpp
--- a/Engine/Plugins/Runtime/ComputeFramework/Source/ComputeFramework/Private/ComputeKernelShaderMap.cpp
+++ b/Engine/Plugins/Runtime/ComputeFramework/Source/ComputeFramework/Private/ComputeKernelShaderMap.cpp
@@ -158,7 +158,7 @@ void FComputeKernelShaderType::BeginCompileShader(
 	NewJob->Input.Environment.IncludeVirtualPathToContentsMap.Add(virtualSourcePath, InKernel->GetHLSLSource());
 	UE_LOG(LogComputeFramework, Verbose, TEXT("%s"), *InKernel->GetHLSLSource());
 	
-	AddReferencedUniformBufferIncludes(NewJob->Input.Environment, NewJob->Input.SourceFilePrefix, InPlatform);
+	AddUniformBufferIncludesToEnvironment(NewJob->Input.Environment, InPlatform);
 
 	FShaderCompilerEnvironment& ShaderEnvironment = NewJob->Input.Environment;
 
diff --git a/Engine/Shaders/Public/Platform.ush b/Engine/Shaders/Public/Platform.ush
--- a/Engine/Shaders/Public/Platform.ush
+++ b/Engine/Shaders/Public/Platform.ush
@@ -356,6 +356,22 @@
 
 // ---------------------------------------------------- Platform dependent supports
 
+// Type macros for Uniform Buffer code generation.
+// Hides the platform specific 'half' and 'fixed' support but has to extend to all types and dimensions
+#if (FEATURE_LEVEL == FEATURE_LEVEL_ES3_1)
+	#define UB_INT(Dim)         int##Dim
+	#define UB_UINT(Dim)        uint##Dim
+	#define UB_FLOAT(Dim)       float##Dim
+	#define UB_HALF_FLOAT(Dim)  half##Dim
+	#define UB_FIXED_FLOAT(Dim) fixed##Dim
+#else
+	#define UB_INT(Dim)         int##Dim
+	#define UB_UINT(Dim)        uint##Dim
+	#define UB_FLOAT(Dim)       float##Dim
+	#define UB_HALF_FLOAT(Dim)  float##Dim
+	#define UB_FIXED_FLOAT(Dim) float##Dim
+#endif
+
 // On mobile use interpolator for ClipDistance as it does not work/supported on all mobile platforms
 #if (FEATURE_LEVEL == FEATURE_LEVEL_ES3_1)
 #define SV_ClipDistance OUTCLIPDIST
diff --git a/Engine/Source/Developer/ShaderPreprocessor/Private/ShaderPreprocessor.cpp b/Engine/Source/Developer/ShaderPreprocessor/Private/ShaderPreprocessor.cpp
--- a/Engine/Source/Developer/ShaderPreprocessor/Private/ShaderPreprocessor.cpp
+++ b/Engine/Source/Developer/ShaderPreprocessor/Private/ShaderPreprocessor.cpp
@@ -38,7 +38,7 @@ public:
 		FString InputShaderSource;
 		if (LoadShaderSourceFile(*InShaderInput.VirtualSourceFilePath, InShaderInput.Target.GetPlatform(),  &InputShaderSource, nullptr, &InShaderInput.ShaderPlatformName))
 		{
-			InputShaderSource = FString::Printf(TEXT("%s\n#line 1\n%s"), *ShaderInput.SourceFilePrefix, *InputShaderSource);
+			InputShaderSource = FString::Printf(TEXT("#line 1\n%s"), *InputShaderSource);
 			CachedFileContents.Add(InShaderInput.VirtualSourceFilePath, StringToArray<ANSICHAR>(*InputShaderSource, InputShaderSource.Len() + 1));
 		}
 	}
diff --git a/Engine/Source/Runtime/Engine/Private/ShaderCompiler/ShaderCompiler.cpp b/Engine/Source/Runtime/Engine/Private/ShaderCompiler/ShaderCompiler.cpp
--- a/Engine/Source/Runtime/Engine/Private/ShaderCompiler/ShaderCompiler.cpp
+++ b/Engine/Source/Runtime/Engine/Private/ShaderCompiler/ShaderCompiler.cpp
@@ -5878,11 +5878,11 @@ void GlobalBeginCompileShader(
 		}
 	}
 
-	ShaderType->AddReferencedUniformBufferIncludes(Input.Environment, Input.SourceFilePrefix, ShaderPlatform);
+	ShaderType->AddUniformBufferIncludesToEnvironment(Input.Environment, ShaderPlatform);
 
 	if (VFType)
 	{
-		VFType->AddReferencedUniformBufferIncludes(Input.Environment, Input.SourceFilePrefix, ShaderPlatform);
+		VFType->AddUniformBufferIncludesToEnvironment(Input.Environment, ShaderPlatform);
 	}
 
 	// Add generated instanced stereo code
diff --git a/Engine/Source/Runtime/RenderCore/Private/Shader.cpp b/Engine/Source/Runtime/RenderCore/Private/Shader.cpp
--- a/Engine/Source/Runtime/RenderCore/Private/Shader.cpp
+++ b/Engine/Source/Runtime/RenderCore/Private/Shader.cpp
@@ -456,7 +456,7 @@ void FShaderType::Initialize(const TMap<FString, TArray<const TCHAR*> >& ShaderF
 #if UE_BUILD_DEBUG
 			UniqueShaderTypes.Add(Type);
 #endif
-			GenerateReferencedUniformBuffers(Type->SourceFilename, Type->Name, ShaderFileToUniformBufferVariables, Type->ReferencedUniformBufferStructsCache);
+			GenerateReferencedUniformBufferNames(Type->SourceFilename, Type->Name, ShaderFileToUniformBufferVariables, Type->ReferencedUniformBufferNames);
 		}
 	
 #if UE_BUILD_DEBUG
diff --git a/Engine/Source/Runtime/RenderCore/Private/ShaderCore.cpp b/Engine/Source/Runtime/RenderCore/Private/ShaderCore.cpp
--- a/Engine/Source/Runtime/RenderCore/Private/ShaderCore.cpp
+++ b/Engine/Source/Runtime/RenderCore/Private/ShaderCore.cpp
@@ -1573,23 +1573,22 @@ void FlushShaderFileCache()
 }
 
 #if WITH_EDITOR
-void GenerateReferencedUniformBuffers(
-	const TCHAR* SourceFilename, 
-	const TCHAR* ShaderTypeName, 
+void GenerateReferencedUniformBufferNames(
+	const TCHAR* SourceFilename,
+	const TCHAR* ShaderTypeName,
 	const TMap<FString, TArray<const TCHAR*> >& ShaderFileToUniformBufferVariables,
-	TMap<const TCHAR*,FCachedUniformBufferDeclaration>& UniformBufferEntries)
+	TSet<const TCHAR*>& UniformBufferNames)
 {
 	TArray<FString> FilesToSearch;
 	GetShaderIncludes(SourceFilename, SourceFilename, FilesToSearch, GMaxRHIShaderPlatform);
-	FilesToSearch.Add(SourceFilename);
+	FilesToSearch.Emplace(SourceFilename);
 
-	for (int32 FileIndex = 0; FileIndex < FilesToSearch.Num(); FileIndex++)
+	for (const FString& FileToSearch : FilesToSearch)
 	{
-		const TArray<const TCHAR*>& FoundUniformBufferVariables = ShaderFileToUniformBufferVariables.FindChecked(FilesToSearch[FileIndex]);
-
-		for (int32 VariableIndex = 0; VariableIndex < FoundUniformBufferVariables.Num(); VariableIndex++)
+		const TArray<const TCHAR*>& FoundUniformBufferVariables = ShaderFileToUniformBufferVariables.FindChecked(FileToSearch);
+		for (const TCHAR* UniformBufferName : FoundUniformBufferVariables)
 		{
-			UniformBufferEntries.Add(FoundUniformBufferVariables[VariableIndex], FCachedUniformBufferDeclaration());
+			UniformBufferNames.Emplace(UniformBufferName);
 		}
 	}
 }
@@ -1701,18 +1700,18 @@ static FFShaderParameterPointerLookupCache GShaderParameterMemberLookupCache;
 
 
 // This copy is only used internally once - it could be inlined once the public API version is removed
-void SerializeUniformBufferInfo_Internal(FShaderSaveArchive& Ar, const TSortedMap<const TCHAR*, FCachedUniformBufferDeclaration, FDefaultAllocator, FUniformBufferNameSortOrder>& UniformBufferEntries)
+void SerializeUniformBufferInfo_Internal(FShaderSaveArchive& Ar, const TArray<const TCHAR*>& UniformBufferNames)
 {
-	if (UniformBufferEntries.IsEmpty())
+	if (UniformBufferNames.IsEmpty())
 	{
 		return;
 	}
 
 	TSharedPtr<const FShaderParameterMemberLookup> ShaderParameterMembers = GShaderParameterMemberLookupCache.Get();
 
-	for (const TPair<const TCHAR*, FCachedUniformBufferDeclaration>& Entry : UniformBufferEntries)
+	for (const TCHAR* UniformBufferName : UniformBufferNames)
 	{
-		if (const TConstArrayView<FShaderParametersMetadata::FMember>* Members = ShaderParameterMembers->FindMembersByPointer(Entry.Key))
+		if (const TConstArrayView<FShaderParametersMetadata::FMember>* Members = ShaderParameterMembers->FindMembersByPointer(UniformBufferName))
 		{
 			// Serialize information about the struct layout so we can detect when it changes
 			int32 NumMembers = Members->Num();
@@ -1761,7 +1760,7 @@ void AppendKeyStringShaderDependencies(
 	FPlatformTypeLayoutParameters LayoutParams,
 	FString& OutKeyString)
 {
-	FSortedMapUniformBufferDeclaration ReferencedUniformBuffers;
+	TSet<const TCHAR*> ReferencedUniformBufferNames;
 
 	for (const FShaderTypeDependency& ShaderTypeDependency : ShaderTypeDependencies)
 	{
@@ -1785,9 +1784,9 @@ void AppendKeyStringShaderDependencies(
 		const FSHAHash LayoutHash = GetShaderTypeLayoutHash(ShaderType->GetLayout(), LayoutParams);
 		LayoutHash.AppendString(OutKeyString);
 
-		for (const auto& Entry : ShaderType->GetReferencedUniformBufferStructsCache())
+		for (const TCHAR* UniformBufferName : ShaderType->GetReferencedUniformBufferNames())
 		{
-			ReferencedUniformBuffers.Add(Entry.Key, Entry.Value);
+			ReferencedUniformBufferNames.Add(UniformBufferName);
 		}
 	}
 
@@ -1808,9 +1807,9 @@ void AppendKeyStringShaderDependencies(
 				OutKeyString.Appendf(TEXT("%08x"), ParameterStructMetadata->GetLayoutHash());
 			}
 
-			for (const auto& Entry : ShaderType->GetReferencedUniformBufferStructsCache())
+			for (const TCHAR* UniformBufferName : ShaderType->GetReferencedUniformBufferNames())
 			{
-				ReferencedUniformBuffers.Add(Entry.Key, Entry.Value);
+				ReferencedUniformBufferNames.Add(UniformBufferName);
 			}
 		}
 	}
@@ -1834,9 +1833,9 @@ void AppendKeyStringShaderDependencies(
 			}
 		}
 
-		for (const auto& Entry : VertexFactoryType->GetReferencedUniformBufferStructsCache())
+		for (const TCHAR* UniformBufferName : VertexFactoryType->GetReferencedUniformBufferNames())
 		{
-			ReferencedUniformBuffers.Add(Entry.Key, Entry.Value);
+			ReferencedUniformBufferNames.Add(UniformBufferName);
 		}
 	}
 
@@ -1846,8 +1845,11 @@ void AppendKeyStringShaderDependencies(
 		FMemoryWriter Ar(TempData, true);
 		FShaderSaveArchive SaveArchive(Ar, SerializationHistory);
 
+		TArray<const TCHAR*> SortedUniformBufferNames = ReferencedUniformBufferNames.Array();
+		Algo::Sort(SortedUniformBufferNames, FUniformBufferNameSortOrder());
+
 		// Save uniform buffer member info so we can detect when layout has changed
-		SerializeUniformBufferInfo_Internal(SaveArchive, ReferencedUniformBuffers);
+		SerializeUniformBufferInfo_Internal(SaveArchive, SortedUniformBufferNames);
 
 		SerializationHistory.AppendKeyString(OutKeyString);
 	}
@@ -1855,7 +1857,14 @@ void AppendKeyStringShaderDependencies(
 
 void SerializeUniformBufferInfo(FShaderSaveArchive& Ar, const TSortedMap<const TCHAR*, FCachedUniformBufferDeclaration, FDefaultAllocator, FUniformBufferNameSortOrder>& UniformBufferEntries)
 {
-	SerializeUniformBufferInfo_Internal(Ar, UniformBufferEntries);
+	TArray<const TCHAR*> UniformBufferNames;
+	for (const TPair<const TCHAR*, FCachedUniformBufferDeclaration>& Entry : UniformBufferEntries)
+	{
+		UniformBufferNames.Emplace(Entry.Key);
+	}
+	Algo::Sort(UniformBufferNames, FUniformBufferNameSortOrder());
+
+	SerializeUniformBufferInfo_Internal(Ar, UniformBufferNames);
 }
 
 #endif // WITH_EDITOR
@@ -2112,7 +2121,6 @@ FArchive& operator<<(FArchive& Ar, FShaderCompilerInput& Input)
 		Ar << ShaderPlatformNameString;
 		Input.ShaderPlatformName = FName(*ShaderPlatformNameString);
 	}
-	Ar << Input.SourceFilePrefix;
 	Ar << Input.VirtualSourceFilePath;
 	Ar << Input.EntryPointName;
 	Ar << Input.ShaderName;
diff --git a/Engine/Source/Runtime/RenderCore/Private/ShaderParameterMetadata.cpp b/Engine/Source/Runtime/RenderCore/Private/ShaderParameterMetadata.cpp
--- a/Engine/Source/Runtime/RenderCore/Private/ShaderParameterMetadata.cpp
+++ b/Engine/Source/Runtime/RenderCore/Private/ShaderParameterMetadata.cpp
@@ -7,6 +7,7 @@
 #include "ShaderParameterMetadata.h"
 #include "RenderCore.h"
 #include "ShaderCore.h"
+#include "ShaderParameters.h"
 
 FUniformBufferStaticSlotRegistrar::FUniformBufferStaticSlotRegistrar(const TCHAR* InName)
 {
@@ -374,6 +375,10 @@ FShaderParametersMetadata::FShaderParametersMetadata(
 	}
 	else
 	{
+#if WITH_EDITOR
+		InitializeUniformBufferDeclaration();
+#endif
+
 		// We cannot initialize the layout during global initialization, since we have to walk nested struct members.
 		// Structs created during global initialization will have bRegisterForAutoBinding==false, and are initialized during startup.
 		// Structs created at runtime with bRegisterForAutoBinding==true can be initialized now.
@@ -410,6 +415,13 @@ void FShaderParametersMetadata::InitializeAllUniformBufferStructs()
 {
 	for (TLinkedList<FShaderParametersMetadata*>::TIterator StructIt(FShaderParametersMetadata::GetStructList()); StructIt; StructIt.Next())
 	{
+#if WITH_EDITOR
+		if (!StructIt->IsUniformBufferDeclarationInitialized())
+		{
+			StructIt->InitializeUniformBufferDeclaration();
+		}
+#endif // WITH_EDITOR
+
 		if (!StructIt->IsLayoutInitialized())
 		{
 			StructIt->InitializeLayout();
@@ -772,6 +784,19 @@ void FShaderParametersMetadata::InitializeLayout(FRHIUniformBufferLayoutInitiali
 	Layout = RHICreateUniformBufferLayout(LayoutInitializer);
 }
 
+#if WITH_EDITOR
+void FShaderParametersMetadata::InitializeUniformBufferDeclaration()
+{
+	if (UseCase != EUseCase::ShaderParameterStruct)
+	{
+		FString* NewDeclaration = new FString(UE::ShaderParameters::CreateUniformBufferShaderDeclaration(ShaderVariableName, *this));
+		check(!NewDeclaration->IsEmpty());
+
+		UniformBufferDeclaration = MakeShareable(NewDeclaration);
+	}
+}
+#endif // WITH_EDITOR
+
 void FShaderParametersMetadata::GetNestedStructs(TArray<const FShaderParametersMetadata*>& OutNestedStructs) const
 {
 	for (int32 i = 0; i < Members.Num(); ++i)
@@ -788,61 +813,65 @@ void FShaderParametersMetadata::GetNestedStructs(TArray<const FShaderParametersM
 	}
 }
 
-void FShaderParametersMetadata::AddResourceTableEntries(TMap<FString, FResourceTableEntry>& ResourceTableMap, TMap<FString, FUniformBufferEntry>& UniformBufferMap) const
-{
-	uint16 ResourceIndex = 0;
-	FString Prefix = FString::Printf(TEXT("%s_"), ShaderVariableName);
-	AddResourceTableEntriesRecursive(ShaderVariableName, *Prefix, ResourceIndex, ResourceTableMap);
-	
-	FUniformBufferEntry UniformBufferEntry;
-	UniformBufferEntry.StaticSlotName = StaticSlotName;
-	UniformBufferEntry.LayoutHash = IsLayoutInitialized() ? GetLayout().GetHash() : 0;
-	UniformBufferEntry.BindingFlags = BindingFlags;
-	UniformBufferEntry.bNoEmulatedUniformBuffer = UsageFlags & (uint32)EUsageFlags::NoEmulatedUniformBuffer;
-	UniformBufferMap.Add(ShaderVariableName, UniformBufferEntry);
-}
+#if WITH_EDITOR
 
-void FShaderParametersMetadata::AddResourceTableEntriesRecursive(const TCHAR* UniformBufferName, const TCHAR* Prefix, uint16& ResourceIndex, TMap<FString, FResourceTableEntry>& ResourceTableMap) const
+static void AddResourceTableEntriesRecursive(const TArray<FShaderParametersMetadata::FMember>& Members, const TCHAR* UniformBufferName, const TCHAR* Prefix, uint16& ResourceIndex, TMap<FString, FResourceTableEntry>& ResourceTableMap)
 {
-	for (int32 MemberIndex = 0; MemberIndex < Members.Num(); ++MemberIndex)
+	for (const FShaderParametersMetadata::FMember& Member : Members)
 	{
-		const FMember& Member = Members[MemberIndex];
-		uint32 NumElements = Member.GetNumElements();
+		const EUniformBufferBaseType BaseType = Member.GetBaseType();
+		const uint32 NumElements = Member.GetNumElements();
 
-		if (IsShaderParameterTypeForUniformBufferLayout(Member.GetBaseType()))
+		if (IsShaderParameterTypeForUniformBufferLayout(BaseType))
 		{
 			FResourceTableEntry& Entry = ResourceTableMap.FindOrAdd(FString::Printf(TEXT("%s%s"), Prefix, Member.GetName()));
 			if (Entry.UniformBufferName.IsEmpty())
 			{
 				Entry.UniformBufferName = UniformBufferName;
-				Entry.Type = Member.GetBaseType();
+				Entry.Type = BaseType;
 				Entry.ResourceIndex = ResourceIndex++;
 			}
 		}
-		else if (Member.GetBaseType() == UBMT_NESTED_STRUCT && NumElements == 0)
+		else if (BaseType == UBMT_NESTED_STRUCT && NumElements == 0)
 		{
 			check(Member.GetStructMetadata());
 			FString MemberPrefix = FString::Printf(TEXT("%s%s_"), Prefix, Member.GetName());
-			Member.GetStructMetadata()->AddResourceTableEntriesRecursive(UniformBufferName, *MemberPrefix, ResourceIndex, ResourceTableMap);
+			AddResourceTableEntriesRecursive(Member.GetStructMetadata()->GetMembers(), UniformBufferName, *MemberPrefix, ResourceIndex, ResourceTableMap);
 		}
-		else if (Member.GetBaseType() == UBMT_NESTED_STRUCT && NumElements > 0)
+		else if (BaseType == UBMT_NESTED_STRUCT && NumElements > 0)
 		{
 			for (uint32 ArrayElementId = 0; ArrayElementId < NumElements; ArrayElementId++)
 			{
 				check(Member.GetStructMetadata());
 				FString MemberPrefix = FString::Printf(TEXT("%s%s_%u_"), Prefix, Member.GetName(), ArrayElementId);
-				Member.GetStructMetadata()->AddResourceTableEntriesRecursive(UniformBufferName, *MemberPrefix, ResourceIndex, ResourceTableMap);
+				AddResourceTableEntriesRecursive(Member.GetStructMetadata()->GetMembers(), UniformBufferName, *MemberPrefix, ResourceIndex, ResourceTableMap);
 			}
 		}
-		else if (Member.GetBaseType() == UBMT_INCLUDED_STRUCT)
+		else if (BaseType == UBMT_INCLUDED_STRUCT)
 		{
 			check(Member.GetStructMetadata());
 			check(NumElements == 0);
-			Member.GetStructMetadata()->AddResourceTableEntriesRecursive(UniformBufferName, Prefix, ResourceIndex, ResourceTableMap);
+			AddResourceTableEntriesRecursive(Member.GetStructMetadata()->GetMembers(), UniformBufferName, Prefix, ResourceIndex, ResourceTableMap);
 		}
 	}
 }
 
+void FShaderParametersMetadata::AddResourceTableEntries(TMap<FString, FResourceTableEntry>& ResourceTableMap, TMap<FString, FUniformBufferEntry>& UniformBufferMap) const
+{
+	uint16 ResourceIndex = 0;
+	const FString Prefix = FString::Printf(TEXT("%s_"), ShaderVariableName);
+	AddResourceTableEntriesRecursive(GetMembers(), ShaderVariableName, *Prefix, ResourceIndex, ResourceTableMap);
+	
+	FUniformBufferEntry UniformBufferEntry;
+	UniformBufferEntry.StaticSlotName = StaticSlotName;
+	UniformBufferEntry.LayoutHash = IsLayoutInitialized() ? GetLayout().GetHash() : 0;
+	UniformBufferEntry.BindingFlags = BindingFlags;
+	UniformBufferEntry.bNoEmulatedUniformBuffer = UsageFlags & (uint32)EUsageFlags::NoEmulatedUniformBuffer;
+	UniformBufferMap.Add(ShaderVariableName, UniformBufferEntry);
+}
+
+#endif // WITH_EDITOR
+
 void FShaderParametersMetadata::FindMemberFromOffset(uint16 MemberOffset, const FShaderParametersMetadata** OutContainingStruct, const FShaderParametersMetadata::FMember** OutMember, int32* ArrayElementId, FString* NamePrefix) const
 {
 	check(MemberOffset < GetSize());
diff --git a/Engine/Source/Runtime/RenderCore/Private/ShaderParameters.cpp b/Engine/Source/Runtime/RenderCore/Private/ShaderParameters.cpp
--- a/Engine/Source/Runtime/RenderCore/Private/ShaderParameters.cpp
+++ b/Engine/Source/Runtime/RenderCore/Private/ShaderParameters.cpp
@@ -77,16 +77,24 @@ FArchive& operator<<(FArchive& Ar,FShaderResourceParameter& P)
 }
 
 #if WITH_EDITOR
-void FShaderUniformBufferParameter::ModifyCompilationEnvironment(const TCHAR* ParameterName,const FShaderParametersMetadata& Struct,EShaderPlatform Platform,FShaderCompilerEnvironment& OutEnvironment)
+void FShaderUniformBufferParameter::ModifyCompilationEnvironment(const TCHAR* ParameterName, const FShaderParametersMetadata& Struct, EShaderPlatform Platform, FShaderCompilerEnvironment& OutEnvironment)
 {
-	const FString IncludeName = FString::Printf(TEXT("/Engine/Generated/UniformBuffers/%s.ush"),ParameterName);
-	// Add the uniform buffer declaration to the compilation environment as an include: UniformBuffers/<ParameterName>.usf
-	FString Declaration;
-	CreateUniformBufferShaderDeclaration(ParameterName, Struct, Platform, Declaration);
-	OutEnvironment.IncludeVirtualPathToContentsMap.Add(IncludeName, Declaration);
+	// Add the uniform buffer declaration to the compilation environment as an include: /Engine/Generated/UniformBuffers/<ParameterName>.usf
+	const FString IncludeName = FString::Printf(TEXT("/Engine/Generated/UniformBuffers/%s.ush"), ParameterName);
+
+	// if the name matches the struct's name, use the struct's cached version; otherwise, generate it now with the correct variable name.
+	if (FCString::Strcmp(ParameterName, Struct.GetShaderVariableName()) != 0)
+	{
+		const FString Declaration = UE::ShaderParameters::CreateUniformBufferShaderDeclaration(ParameterName, Struct);
+		OutEnvironment.IncludeVirtualPathToContentsMap.Add(IncludeName, Declaration);
+	}
+	else
+	{
+		OutEnvironment.IncludeVirtualPathToContentsMap.Add(IncludeName, Struct.GetUniformBufferDeclaration());
+	}
 
 	FString& GeneratedUniformBuffersInclude = OutEnvironment.IncludeVirtualPathToContentsMap.FindOrAdd("/Engine/Generated/GeneratedUniformBuffers.ush");
-	FString Include = FString::Printf(TEXT("#include \"/Engine/Generated/UniformBuffers/%s.ush\"") LINE_TERMINATOR, ParameterName);
+	const FString Include = FString::Printf(TEXT("#include \"%s\"") LINE_TERMINATOR, *IncludeName);
 
 	GeneratedUniformBuffersInclude.Append(Include);
 	Struct.AddResourceTableEntries(OutEnvironment.ResourceTableMap, OutEnvironment.UniformBufferMap);
@@ -142,13 +150,12 @@ static void CreateHLSLUniformBufferStructMembersDeclaration(
 	const FShaderParametersMetadata& UniformBufferStruct, 
 	const FString& NamePrefix, 
 	uint32 StructOffset, 
-	EShaderPlatform Platform,
 	FUniformBufferDecl& Decl, 
 	uint32& HLSLBaseOffset)
 {
 	const TArray<FShaderParametersMetadata::FMember>& StructMembers = UniformBufferStruct.GetMembers();
 	
-	FString PreviousBaseTypeName = TEXT("float");
+	const TCHAR* PreviousBaseTypeName = TEXT("UB_FLOAT");
 	for (int32 MemberIndex = 0; MemberIndex < StructMembers.Num(); ++MemberIndex)
 	{
 		const FShaderParametersMetadata::FMember& Member = StructMembers[MemberIndex];
@@ -164,13 +171,13 @@ static void CreateHLSLUniformBufferStructMembersDeclaration(
 			checkf(Member.GetNumElements() == 0, TEXT("SHADER_PARAMETER_STRUCT_ARRAY() is not supported in uniform buffer yet."));
 			Decl.StructMembers += TEXT("struct {\r\n");
 			Decl.Initializer += TEXT("{");
-			CreateHLSLUniformBufferStructMembersDeclaration(*Member.GetStructMetadata(), FString::Printf(TEXT("%s%s_"), *NamePrefix, Member.GetName()), StructOffset + Member.GetOffset(), Platform, Decl, HLSLBaseOffset);
+			CreateHLSLUniformBufferStructMembersDeclaration(*Member.GetStructMetadata(), FString::Printf(TEXT("%s%s_"), *NamePrefix, Member.GetName()), StructOffset + Member.GetOffset(), Decl, HLSLBaseOffset);
 			Decl.Initializer += TEXT("},");
 			Decl.StructMembers += FString::Printf(TEXT("} %s%s;\r\n"),Member.GetName(),*ArrayDim);
 		}
 		else if (Member.GetBaseType() == UBMT_INCLUDED_STRUCT)
 		{
-			CreateHLSLUniformBufferStructMembersDeclaration(*Member.GetStructMetadata(), NamePrefix, StructOffset + Member.GetOffset(), Platform, Decl, HLSLBaseOffset);
+			CreateHLSLUniformBufferStructMembersDeclaration(*Member.GetStructMetadata(), NamePrefix, StructOffset + Member.GetOffset(), Decl, HLSLBaseOffset);
 		}
 		else if (IsShaderParameterTypeForUniformBufferLayout(Member.GetBaseType()))
 		{
@@ -185,12 +192,12 @@ static void CreateHLSLUniformBufferStructMembersDeclaration(
 				check(HLSLBaseOffset < AbsoluteMemberOffset);
 				while (HLSLBaseOffset < AbsoluteMemberOffset)
 				{
-					Decl.ConstantBufferMembers += FString::Printf(TEXT("\t%s PrePadding_%s%u;\r\n"), *PreviousBaseTypeName, *NamePrefix, HLSLBaseOffset);
+					Decl.ConstantBufferMembers += FString::Printf(TEXT("\t%s() PrePadding_%s%u;\r\n"), PreviousBaseTypeName, *NamePrefix, HLSLBaseOffset);
 					HLSLBaseOffset += 4;
 				};
 				check(HLSLBaseOffset == AbsoluteMemberOffset);
 			}
-			PreviousBaseTypeName = TEXT("uint");
+			PreviousBaseTypeName = TEXT("UB_UINT");
 			HLSLBaseOffset = AbsoluteMemberOffset + HLSLMemberSize;
 
 			// Generate the member declaration.
@@ -208,23 +215,23 @@ static void CreateHLSLUniformBufferStructMembersDeclaration(
 		else 
 		{
 			// Generate the base type name.
-			FString BaseTypeName;
+			const TCHAR* BaseTypeName = TEXT("");
 			switch(Member.GetBaseType())
 			{
-			case UBMT_INT32:   BaseTypeName = TEXT("int"); break;
-			case UBMT_UINT32:  BaseTypeName = TEXT("uint"); break;
+			case UBMT_INT32:   BaseTypeName = TEXT("UB_INT"); break;
+			case UBMT_UINT32:  BaseTypeName = TEXT("UB_UINT"); break;
 			case UBMT_FLOAT32: 
-				if (Member.GetPrecision() == EShaderPrecisionModifier::Float || !SupportShaderPrecisionModifier(Platform))
+				if (Member.GetPrecision() == EShaderPrecisionModifier::Float)
 				{
-					BaseTypeName = TEXT("float"); 
+					BaseTypeName = TEXT("UB_FLOAT"); 
 				}
 				else if (Member.GetPrecision() == EShaderPrecisionModifier::Half)
 				{
-					BaseTypeName = TEXT("half"); 
+					BaseTypeName = TEXT("UB_HALF_FLOAT"); 
 				}
 				else if (Member.GetPrecision() == EShaderPrecisionModifier::Fixed)
 				{
-					BaseTypeName = TEXT("fixed"); 
+					BaseTypeName = TEXT("UB_FIXED_FLOAT"); 
 				}
 				break;
 			default:           UE_LOG(LogShaders, Fatal,TEXT("Unrecognized uniform buffer struct member base type."));
@@ -260,7 +267,7 @@ static void CreateHLSLUniformBufferStructMembersDeclaration(
 				check(HLSLBaseOffset < AbsoluteMemberOffset);
 				while(HLSLBaseOffset < AbsoluteMemberOffset)
 				{
-					Decl.ConstantBufferMembers += FString::Printf(TEXT("\t%s PrePadding_%s%u;\r\n"), *PreviousBaseTypeName, *NamePrefix, HLSLBaseOffset);
+					Decl.ConstantBufferMembers += FString::Printf(TEXT("\t%s() PrePadding_%s%u;\r\n"), PreviousBaseTypeName, *NamePrefix, HLSLBaseOffset);
 					HLSLBaseOffset += 4;
 				};
 				check(HLSLBaseOffset == AbsoluteMemberOffset);
@@ -270,8 +277,8 @@ static void CreateHLSLUniformBufferStructMembersDeclaration(
 
 			// Generate the member declaration.
 			FString ParameterName = FString::Printf(TEXT("%s%s"),*NamePrefix,Member.GetName());
-			Decl.ConstantBufferMembers += FString::Printf(TEXT("\t%s%s %s%s;\r\n"),*BaseTypeName,*TypeDim,*ParameterName,*ArrayDim);
-			Decl.StructMembers += FString::Printf(TEXT("\t%s%s %s%s;\r\n"),*BaseTypeName,*TypeDim,Member.GetName(),*ArrayDim);
+			Decl.ConstantBufferMembers += FString::Printf(TEXT("\t%s(%s) %s%s;\r\n"),BaseTypeName,*TypeDim,*ParameterName,*ArrayDim);
+			Decl.StructMembers += FString::Printf(TEXT("\t%s(%s) %s%s;\r\n"),BaseTypeName,*TypeDim,Member.GetName(),*ArrayDim);
 			Decl.Initializer += FString::Printf(TEXT("%s,"),*ParameterName);
 		}
 	}
@@ -310,7 +317,7 @@ static void CreateHLSLUniformBufferStructMembersDeclaration(
 }
 
 /** Creates a HLSL declaration of a uniform buffer with the given structure. */
-static FString CreateHLSLUniformBufferDeclaration(const TCHAR* Name,const FShaderParametersMetadata& UniformBufferStruct, EShaderPlatform Platform)
+static FString CreateHLSLUniformBufferDeclaration(const TCHAR* Name, const FShaderParametersMetadata& UniformBufferStruct)
 {
 	// If the uniform buffer has no members, we don't want to write out anything.  Shader compilers throw errors when faced with empty cbuffers and structs.
 	if (UniformBufferStruct.GetMembers().Num() > 0)
@@ -318,7 +325,7 @@ static FString CreateHLSLUniformBufferDeclaration(const TCHAR* Name,const FShade
 		FString NamePrefix(FString(Name) + FString(TEXT("_")));
 		FUniformBufferDecl Decl;
 		uint32 HLSLBaseOffset = 0;
-		CreateHLSLUniformBufferStructMembersDeclaration(UniformBufferStruct, NamePrefix, 0, Platform, Decl, HLSLBaseOffset);
+		CreateHLSLUniformBufferStructMembersDeclaration(UniformBufferStruct, NamePrefix, 0, Decl, HLSLBaseOffset);
 
 		return FString::Printf(
 			TEXT("#ifndef __UniformBuffer_%s_Definition__\r\n")
@@ -347,112 +354,121 @@ static FString CreateHLSLUniformBufferDeclaration(const TCHAR* Name,const FShade
 	return FString(TEXT("\n"));
 }
 
-RENDERCORE_API void CreateUniformBufferShaderDeclaration(const TCHAR* Name,const FShaderParametersMetadata& UniformBufferStruct, EShaderPlatform Platform, FString& OutDeclaration)
+FString UE::ShaderParameters::CreateUniformBufferShaderDeclaration(const TCHAR* Name, const FShaderParametersMetadata& UniformBufferStruct)
+{
+	return CreateHLSLUniformBufferDeclaration(Name, UniformBufferStruct);
+}
+
+static FShaderParametersMetadata* FindShaderParametersMetadataWithVariableName(const TCHAR* InVariableName)
 {
-	OutDeclaration = CreateHLSLUniformBufferDeclaration(Name, UniformBufferStruct, Platform);
+	for (FShaderParametersMetadata* Metadata : *FShaderParametersMetadata::GetStructList())
+	{
+		if (InVariableName == Metadata->GetShaderVariableName())
+		{
+			return Metadata;
+		}
+	}
+	return nullptr;
 }
 
-RENDERCORE_API void CacheUniformBufferIncludes(TMap<const TCHAR*,FCachedUniformBufferDeclaration>& Cache, EShaderPlatform Platform)
+/* deprecated */
+void CacheUniformBufferIncludes(TMap<const TCHAR*, FCachedUniformBufferDeclaration>& Cache, EShaderPlatform Platform)
 {
-	for (TMap<const TCHAR*,FCachedUniformBufferDeclaration>::TIterator It(Cache); It; ++It)
+	for (TMap<const TCHAR*, FCachedUniformBufferDeclaration>::TIterator It(Cache); It; ++It)
 	{
+		const TCHAR* UniformBufferName = It.Key();
 		FCachedUniformBufferDeclaration& BufferDeclaration = It.Value();
 		check(BufferDeclaration.Declaration.Get() == NULL);
 
-		for (TLinkedList<FShaderParametersMetadata*>::TIterator StructIt(FShaderParametersMetadata::GetStructList()); StructIt; StructIt.Next())
+		if (const FShaderParametersMetadata* Metadata = FindShaderParametersMetadataWithVariableName(UniformBufferName))
 		{
-			if (It.Key() == StructIt->GetShaderVariableName())
-			{
-				FString* NewDeclaration = new FString();
-				CreateUniformBufferShaderDeclaration(StructIt->GetShaderVariableName(), **StructIt, Platform, *NewDeclaration);
-				check(!NewDeclaration->IsEmpty());
-				BufferDeclaration.Declaration = MakeShareable(NewDeclaration);
-				break;
-			}
+			FString* NewDeclaration = new FString(UE::ShaderParameters::CreateUniformBufferShaderDeclaration(UniformBufferName, *Metadata));
+			check(!NewDeclaration->IsEmpty());
+			BufferDeclaration.Declaration = MakeShareable(NewDeclaration);
+		}
+	}
+}
+
+void UE::ShaderParameters::AddUniformBufferIncludesToEnvironment(FShaderCompilerEnvironment& OutEnvironment, const TSet<const TCHAR*>& InUniformBufferNames)
+{
+	TRACE_CPUPROFILER_EVENT_SCOPE(UE::ShaderParameters::AddUniformBufferIncludesToEnvironment);
+
+	FString UniformBufferIncludes;
+
+	for (const TCHAR* UniformBufferName : InUniformBufferNames)
+	{
+		if (const FShaderParametersMetadata* Metadata = FindShaderParametersMetadataWithVariableName(UniformBufferName))
+		{
+			const FThreadSafeSharedStringPtr UniformBufferDeclaration = Metadata->GetUniformBufferDeclarationPtr();
+
+			check(UniformBufferDeclaration.Get() != NULL);
+			check(!UniformBufferDeclaration.Get()->IsEmpty());
+
+			const FString UniformBufferPath = FString::Printf(TEXT("/Engine/Generated/UniformBuffers/%s.ush"), UniformBufferName);
+
+			UniformBufferIncludes += FString::Printf(TEXT("#include \"%s\"") LINE_TERMINATOR, *UniformBufferPath);
+
+			OutEnvironment.IncludeVirtualPathToExternalContentsMap.Add(UniformBufferPath, UniformBufferDeclaration);
+
+			Metadata->AddResourceTableEntries(OutEnvironment.ResourceTableMap, OutEnvironment.UniformBufferMap);
 		}
 	}
+
+	FString& GeneratedUniformBuffersInclude = OutEnvironment.IncludeVirtualPathToContentsMap.FindOrAdd("/Engine/Generated/GeneratedUniformBuffers.ush");
+	GeneratedUniformBuffersInclude.Append(UniformBufferIncludes);
 }
 
 static const uint32 NumUniformBufferLocks = 16u;
 static FRWLock UniformBufferLocks[NumUniformBufferLocks];
 
+template<typename ShaderTypename>
+FRWLock& GetUniformBufferLock(const ShaderTypename* InShaderType)
+{
+	const uint32 LockIndex = InShaderType->GetHashedName().GetHash() % NumUniformBufferLocks;
+	return UniformBufferLocks[LockIndex];
+}
+
 void FShaderType::FlushShaderFileCache(const TMap<FString, TArray<const TCHAR*> >& ShaderFileToUniformBufferVariables)
 {
+	// If we've ever been used to populate a shader compiler environment, reset and repopulate the list of uniform buffers
 	if (CachedUniformBufferPlatform != SP_NumPlatforms)
 	{
-		const uint32 LockIndex = HashedName.GetHash() % NumUniformBufferLocks;
-		FWriteScopeLock Lock(UniformBufferLocks[LockIndex]);
+		FWriteScopeLock Lock(GetUniformBufferLock(this));
 		if (CachedUniformBufferPlatform != SP_NumPlatforms)
 		{
-			ReferencedUniformBufferStructsCache.Empty();
-			GenerateReferencedUniformBuffers(SourceFilename, Name, ShaderFileToUniformBufferVariables, ReferencedUniformBufferStructsCache);
+			ReferencedUniformBufferNames.Empty();
+			GenerateReferencedUniformBufferNames(SourceFilename, Name, ShaderFileToUniformBufferVariables, ReferencedUniformBufferNames);
 			CachedUniformBufferPlatform = SP_NumPlatforms;
 		}
 	}
 }
 
-void FShaderType::AddReferencedUniformBufferIncludes(FShaderCompilerEnvironment& OutEnvironment, FString& OutSourceFilePrefix, EShaderPlatform Platform) const
+void FShaderType::AddUniformBufferIncludesToEnvironment(FShaderCompilerEnvironment& OutEnvironment, EShaderPlatform Platform) const
 {
 	TRACE_CPUPROFILER_EVENT_SCOPE(FShaderType::AddReferencedUniformBufferIncludes);
 
-	const uint32 LockIndex = HashedName.GetHash() % NumUniformBufferLocks;
 	TOptional<FReadScopeLock> ReadScopeLock; // TOptional is so we can unlock and lock manually; the FReadScopeLock api does not support that
-	ReadScopeLock.Emplace(UniformBufferLocks[LockIndex]);
-	for (;;)
+	ReadScopeLock.Emplace(GetUniformBufferLock(this));
+
+	while (CachedUniformBufferPlatform != Platform)
 	{
 		// Cache uniform buffer struct declarations referenced by this shader type's files
-		if (CachedUniformBufferPlatform == Platform)
-		{
-			break;
-		}
 
 		// Drop the ReadLock, acquire a write lock
 		//     Write the ReferencedUniformBufferStructsCache for the new Platform
 		// Drop the WriteLock, reacquire the ReadLock
 		ReadScopeLock.Reset();
 		{
-			FWriteScopeLock WriteScopeLock(UniformBufferLocks[LockIndex]);
+			FWriteScopeLock WriteScopeLock(GetUniformBufferLock(this));
 			if (CachedUniformBufferPlatform != Platform)
 			{
-				// If there is already a cache but for another platform, keep the keys but reset the values
-				if (CachedUniformBufferPlatform != SP_NumPlatforms)
-				{
-					for (TMap<const TCHAR*, FCachedUniformBufferDeclaration>::TIterator It(ReferencedUniformBufferStructsCache); It; ++It)
-					{
-						It.Value() = FCachedUniformBufferDeclaration();
-					}
-				}
-				CacheUniformBufferIncludes(ReferencedUniformBufferStructsCache, Platform);
 				CachedUniformBufferPlatform = Platform;
 			}
 		}
-		ReadScopeLock.Emplace(UniformBufferLocks[LockIndex]);
-	}
-
-	FString UniformBufferIncludes;
-
-	for (TMap<const TCHAR*,FCachedUniformBufferDeclaration>::TConstIterator It(ReferencedUniformBufferStructsCache); It; ++It)
-	{
-		const FCachedUniformBufferDeclaration& Value = It.Value();
-		check(Value.Declaration.Get() != NULL);
-		check(!Value.Declaration.Get()->IsEmpty());
-		UniformBufferIncludes += FString::Printf(TEXT("#include \"/Engine/Generated/UniformBuffers/%s.ush\"") LINE_TERMINATOR, It.Key());
-		OutEnvironment.IncludeVirtualPathToExternalContentsMap.Add(
-			FString::Printf(TEXT("/Engine/Generated/UniformBuffers/%s.ush"),It.Key()),
-			Value.Declaration
-			);
-
-		for (TLinkedList<FShaderParametersMetadata*>::TIterator StructIt(FShaderParametersMetadata::GetStructList()); StructIt; StructIt.Next())
-		{
-			if (It.Key() == StructIt->GetShaderVariableName())
-			{
-				StructIt->AddResourceTableEntries(OutEnvironment.ResourceTableMap, OutEnvironment.UniformBufferMap);
-			}
-		}
+		ReadScopeLock.Emplace(GetUniformBufferLock(this));
 	}
 
-	FString& GeneratedUniformBuffersInclude = OutEnvironment.IncludeVirtualPathToContentsMap.FindOrAdd("/Engine/Generated/GeneratedUniformBuffers.ush");
-	GeneratedUniformBuffersInclude.Append(UniformBufferIncludes);
+	UE::ShaderParameters::AddUniformBufferIncludesToEnvironment(OutEnvironment, ReferencedUniformBufferNames);
 }
 
 #endif // WITH_EDITOR
@@ -520,78 +536,43 @@ void FShaderType::GetShaderStableKeyParts(FStableShaderKeyAndValue& SaveKeyVal)
 
 void FVertexFactoryType::FlushShaderFileCache(const TMap<FString, TArray<const TCHAR*> >& ShaderFileToUniformBufferVariables)
 {
+	// If we've ever been used to populate a shader compiler environment, reset and repopulate the list of uniform buffers
 	if (CachedUniformBufferPlatform != SP_NumPlatforms)
 	{
-		const uint32 LockIndex = HashedName.GetHash() % NumUniformBufferLocks;
-		FWriteScopeLock Lock(UniformBufferLocks[LockIndex]);
+		FWriteScopeLock Lock(GetUniformBufferLock(this));
 		if (CachedUniformBufferPlatform != SP_NumPlatforms)
 		{
-			ReferencedUniformBufferStructsCache.Empty();
-			GenerateReferencedUniformBuffers(ShaderFilename, Name, ShaderFileToUniformBufferVariables, ReferencedUniformBufferStructsCache);
+			ReferencedUniformBufferNames.Empty();
+			GenerateReferencedUniformBufferNames(ShaderFilename, Name, ShaderFileToUniformBufferVariables, ReferencedUniformBufferNames);
 			CachedUniformBufferPlatform = SP_NumPlatforms;
 		}
 	}
 }
 
-void FVertexFactoryType::AddReferencedUniformBufferIncludes(FShaderCompilerEnvironment& OutEnvironment, FString& OutSourceFilePrefix, EShaderPlatform Platform) const
+void FVertexFactoryType::AddUniformBufferIncludesToEnvironment(FShaderCompilerEnvironment& OutEnvironment, EShaderPlatform Platform) const
 {
-	const uint32 LockIndex = HashedName.GetHash() % NumUniformBufferLocks;
 	TOptional<FReadScopeLock> ReadScopeLock; // TOptional is so we can unlock and lock manually; the FReadScopeLock api does not support that
-	ReadScopeLock.Emplace(UniformBufferLocks[LockIndex]);
-	for (;;)
+	ReadScopeLock.Emplace(GetUniformBufferLock(this));
+	while (CachedUniformBufferPlatform != Platform)
 	{
 		// Cache uniform buffer struct declarations referenced by this shader type's files
-		if (CachedUniformBufferPlatform == Platform)
-		{
-			break;
-		}
 
 		// Drop the ReadLock, acquire a write lock
 		//     Write the ReferencedUniformBufferStructsCache for the new Platform
 		// Drop the WriteLock, reacquire the ReadLock
 		ReadScopeLock.Reset();
 		{
-			FWriteScopeLock WriteScopeLock(UniformBufferLocks[LockIndex]);
+			FWriteScopeLock WriteScopeLock(GetUniformBufferLock(this));
 			if (CachedUniformBufferPlatform != Platform)
 			{
-				// If there is already a cache but for another platform, keep the keys but reset the values
-				if (CachedUniformBufferPlatform != SP_NumPlatforms)
-				{
-					for (TMap<const TCHAR*, FCachedUniformBufferDeclaration>::TIterator It(ReferencedUniformBufferStructsCache); It; ++It)
-					{
-						It.Value() = FCachedUniformBufferDeclaration();
-					}
-				}
-				CacheUniformBufferIncludes(ReferencedUniformBufferStructsCache, Platform);
 				CachedUniformBufferPlatform = Platform;
 			}
 		}
-		ReadScopeLock.Emplace(UniformBufferLocks[LockIndex]);
+		ReadScopeLock.Emplace(GetUniformBufferLock(this));
 	}
 
-	FString UniformBufferIncludes;
-
-	for (TMap<const TCHAR*,FCachedUniformBufferDeclaration>::TConstIterator It(ReferencedUniformBufferStructsCache); It; ++It)
-	{
-		check(It.Value().Declaration.Get() != NULL);
-		check(!It.Value().Declaration.Get()->IsEmpty());
-		UniformBufferIncludes += FString::Printf(TEXT("#include \"/Engine/Generated/UniformBuffers/%s.ush\"") LINE_TERMINATOR, It.Key());
-		OutEnvironment.IncludeVirtualPathToExternalContentsMap.Add(
-			*FString::Printf(TEXT("/Engine/Generated/UniformBuffers/%s.ush"),It.Key()),
-			It.Value().Declaration
-		);
-
-		for (TLinkedList<FShaderParametersMetadata*>::TIterator StructIt(FShaderParametersMetadata::GetStructList()); StructIt; StructIt.Next())
-		{
-			if (It.Key() == StructIt->GetShaderVariableName())
-			{
-				StructIt->AddResourceTableEntries(OutEnvironment.ResourceTableMap, OutEnvironment.UniformBufferMap);
-			}
-		}
-	}
-
-	FString& GeneratedUniformBuffersInclude = OutEnvironment.IncludeVirtualPathToContentsMap.FindOrAdd("/Engine/Generated/GeneratedUniformBuffers.ush");
-	GeneratedUniformBuffersInclude.Append(UniformBufferIncludes);
+	UE::ShaderParameters::AddUniformBufferIncludesToEnvironment(OutEnvironment, ReferencedUniformBufferNames);
 }
 
+
 #endif // WITH_EDITOR
diff --git a/Engine/Source/Runtime/RenderCore/Private/VertexFactory.cpp b/Engine/Source/Runtime/RenderCore/Private/VertexFactory.cpp
--- a/Engine/Source/Runtime/RenderCore/Private/VertexFactory.cpp
+++ b/Engine/Source/Runtime/RenderCore/Private/VertexFactory.cpp
@@ -62,7 +62,7 @@ void FVertexFactoryType::Initialize(const TMap<FString, TArray<const TCHAR*> >&
 		for(TLinkedList<FVertexFactoryType*>::TIterator It(FVertexFactoryType::GetTypeList()); It; It.Next())
 		{
 			FVertexFactoryType* Type = *It;
-			GenerateReferencedUniformBuffers(Type->ShaderFilename, Type->Name, ShaderFileToUniformBufferVariables, Type->ReferencedUniformBufferStructsCache);
+			GenerateReferencedUniformBufferNames(Type->ShaderFilename, Type->Name, ShaderFileToUniformBufferVariables, Type->ReferencedUniformBufferNames);
 		}
 	}
 #endif // WITH_EDITOR
diff --git a/Engine/Source/Runtime/RenderCore/Public/Shader.h b/Engine/Source/Runtime/RenderCore/Public/Shader.h
--- a/Engine/Source/Runtime/RenderCore/Public/Shader.h
+++ b/Engine/Source/Runtime/RenderCore/Public/Shader.h
@@ -1413,14 +1413,17 @@ public:
 	}
 
 #if WITH_EDITOR
-	inline const TMap<const TCHAR*, FCachedUniformBufferDeclaration>& GetReferencedUniformBufferStructsCache() const
+	inline const TSet<const TCHAR*>& GetReferencedUniformBufferNames() const { return ReferencedUniformBufferNames; };
+
+	/** Adds include statements for uniform buffers that this shader type references. */
+	void AddUniformBufferIncludesToEnvironment(FShaderCompilerEnvironment& OutEnvironment, EShaderPlatform Platform) const;
+
+	UE_DEPRECATED(5.2, "AddReferencedUniformBufferIncludes has moved to AddUniformBufferIncludesToEnvironment and no longer takes a prefix argument.")
+	inline void AddReferencedUniformBufferIncludes(FShaderCompilerEnvironment& OutEnvironment, FString& OutSourceFilePrefix, EShaderPlatform Platform) const
 	{
-		return ReferencedUniformBufferStructsCache;
+		AddUniformBufferIncludesToEnvironment(OutEnvironment, Platform);
 	}
 
-	/** Adds include statements for uniform buffers that this shader type references, and builds a prefix for the shader file with the include statements. */
-	void AddReferencedUniformBufferIncludes(FShaderCompilerEnvironment& OutEnvironment, FString& OutSourceFilePrefix, EShaderPlatform Platform) const;
-
 	void FlushShaderFileCache(const TMap<FString, TArray<const TCHAR*> >& ShaderFileToUniformBufferVariables);
 
 	void GetShaderStableKeyParts(struct FStableShaderKeyAndValue& SaveKeyVal);
@@ -1460,15 +1463,15 @@ private:
 
 #if WITH_EDITOR
 protected:
-	/** Tracks what platforms ReferencedUniformBufferStructsCache has had declarations cached for. */
+	/** Tracks the last platform ReferencedUniformBufferNames was read from in AddUniformBufferIncludesToEnvironment. */
 	mutable std::atomic<EShaderPlatform> CachedUniformBufferPlatform{ SP_NumPlatforms };
 
 	/**
 	* Cache of referenced uniform buffer includes.
 	* These are derived from source files so they need to be flushed when editing and recompiling shaders on the fly.
-	* FShaderType::Initialize will add an entry for each referenced uniform buffer, but the declarations are added on demand as shaders are compiled.
+	* FShaderType::Initialize will add the referenced uniform buffers, but this set may be updated by FlushShaderFileCache.
 	*/
-	mutable TMap<const TCHAR*, FCachedUniformBufferDeclaration> ReferencedUniformBufferStructsCache;
+	TSet<const TCHAR*> ReferencedUniformBufferNames;
 #endif // WITH_EDITOR
 };
 
diff --git a/Engine/Source/Runtime/RenderCore/Public/ShaderCompilerCore.h b/Engine/Source/Runtime/RenderCore/Public/ShaderCompilerCore.h
--- a/Engine/Source/Runtime/RenderCore/Public/ShaderCompilerCore.h
+++ b/Engine/Source/Runtime/RenderCore/Public/ShaderCompilerCore.h
@@ -185,7 +185,6 @@ struct FShaderCompilerInput
 	FName CompressionFormat;
 	FName ShaderPlatformName;
 	
-	FString SourceFilePrefix;
 	FString VirtualSourceFilePath;
 	FString EntryPointName;
 	FString ShaderName;
diff --git a/Engine/Source/Runtime/RenderCore/Public/ShaderCore.h b/Engine/Source/Runtime/RenderCore/Public/ShaderCore.h
--- a/Engine/Source/Runtime/RenderCore/Public/ShaderCore.h
+++ b/Engine/Source/Runtime/RenderCore/Public/ShaderCore.h
@@ -1096,12 +1096,12 @@ struct FCachedUniformBufferDeclaration
 	FThreadSafeSharedStringPtr Declaration;
 };
 
-/** Parses the given source file and its includes for references of uniform buffers, which are then stored in UniformBufferEntries. */
-extern void GenerateReferencedUniformBuffers(
+/** Parses the given source file and its includes for references of uniform buffers. */
+extern void GenerateReferencedUniformBufferNames(
 	const TCHAR* SourceFilename,
 	const TCHAR* ShaderTypeName,
 	const TMap<FString, TArray<const TCHAR*> >& ShaderFileToUniformBufferVariables,
-	TMap<const TCHAR*,FCachedUniformBufferDeclaration>& UniformBufferEntries);
+	TSet<const TCHAR*>& UniformBufferNames);
 
 struct FUniformBufferNameSortOrder
 {
diff --git a/Engine/Source/Runtime/RenderCore/Public/ShaderParameterMetadata.h b/Engine/Source/Runtime/RenderCore/Public/ShaderParameterMetadata.h
--- a/Engine/Source/Runtime/RenderCore/Public/ShaderParameterMetadata.h
+++ b/Engine/Source/Runtime/RenderCore/Public/ShaderParameterMetadata.h
@@ -23,6 +23,9 @@
 #include "Serialization/MemoryLayout.h"
 #include "Templates/AlignmentTemplates.h"
 #include "UObject/NameTypes.h"
+#include "Templates/SharedPointer.h"
+
+using FThreadSafeSharedStringPtr = TSharedPtr<FString, ESPMode::ThreadSafe>;
 
 namespace EShaderPrecisionModifier
 {
@@ -269,7 +272,9 @@ public:
 
 	void GetNestedStructs(TArray<const FShaderParametersMetadata*>& OutNestedStructs) const;
 
+#if WITH_EDITOR
 	void AddResourceTableEntries(TMap<FString, FResourceTableEntry>& ResourceTableMap, TMap<FString, FUniformBufferEntry>& UniformBufferMap) const;
+#endif
 
 	const TCHAR* GetStructTypeName() const { return StructTypeName; }
 	const TCHAR* GetShaderVariableName() const { return ShaderVariableName; }
@@ -311,6 +316,12 @@ public:
 	}
 	const TArray<FMember>& GetMembers() const { return Members; }
 
+#if WITH_EDITOR
+	inline bool IsUniformBufferDeclarationInitialized() const { return UniformBufferDeclaration.IsValid(); }
+	FThreadSafeSharedStringPtr GetUniformBufferDeclarationPtr() const { return UniformBufferDeclaration; }
+	const FString& GetUniformBufferDeclaration() const { return *UniformBufferDeclaration; }
+#endif // WITH_EDITOR
+
 	/** Find a member for a given offset. */
 	void FindMemberFromOffset(
 		uint16 MemberOffset,
@@ -388,6 +399,11 @@ private:
 	/** List of all members. */
 	TArray<FMember> Members;
 
+#if WITH_EDITOR
+	/** Uniform buffer declaration, created once */
+	FThreadSafeSharedStringPtr UniformBufferDeclaration;
+#endif
+
 	/** Shackle elements in global link list of globally named shader parameters. */
 	TLinkedList<FShaderParametersMetadata*> GlobalListLink;
 
@@ -399,5 +415,7 @@ private:
 
 	void InitializeLayout(FRHIUniformBufferLayoutInitializer* OutLayoutInitializer = nullptr);
 
-	void AddResourceTableEntriesRecursive(const TCHAR* UniformBufferName, const TCHAR* Prefix, uint16& ResourceIndex, TMap<FString, FResourceTableEntry>& ResourceTableMap) const;
+#if WITH_EDITOR
+	void InitializeUniformBufferDeclaration();
+#endif
 };
\ No newline at end of file
diff --git a/Engine/Source/Runtime/RenderCore/Public/ShaderParameters.h b/Engine/Source/Runtime/RenderCore/Public/ShaderParameters.h
--- a/Engine/Source/Runtime/RenderCore/Public/ShaderParameters.h
+++ b/Engine/Source/Runtime/RenderCore/Public/ShaderParameters.h
@@ -22,6 +22,7 @@ class FRHITexture;
 class FRHIUnorderedAccessView;
 class FShaderParameterMap;
 class FShaderParametersMetadata;
+struct FCachedUniformBufferDeclaration;
 struct FRWBuffer;
 struct FRWBufferStructured;
 struct FShaderCompilerEnvironment;
@@ -30,12 +31,23 @@ enum class EShaderParameterType : uint8;
 DECLARE_INTRINSIC_TYPE_LAYOUT(EShaderParameterType);
 
 #if WITH_EDITOR
-/** Creates a shader code declaration of this struct for the given shader platform. */
-RENDERCORE_API void CreateUniformBufferShaderDeclaration(const TCHAR* Name, const FShaderParametersMetadata& UniformBufferStruct, EShaderPlatform Platform, FString& OutDeclaration);
+namespace UE::ShaderParameters
+{
+	/** Creates a shader code declaration of this struct for the given shader platform. */
+	RENDERCORE_API FString CreateUniformBufferShaderDeclaration(const TCHAR* Name, const FShaderParametersMetadata& UniformBufferStruct);
 
-RENDERCORE_API void CacheUniformBufferIncludes(TMap<const TCHAR*, struct FCachedUniformBufferDeclaration>& Cache, EShaderPlatform Platform);
-#endif
+	RENDERCORE_API void AddUniformBufferIncludesToEnvironment(FShaderCompilerEnvironment& OutEnvironment, const TSet<const TCHAR*>& InUniformBufferNames);
+}
 
+UE_DEPRECATED(5.2, "CreateUniformBufferShaderDeclaration has moved to UE::ShaderParameters::CreateUniformBufferShaderDeclaration, does not take a EShaderPlatform argument and now returns the Declaration.")
+inline void CreateUniformBufferShaderDeclaration(const TCHAR* Name, const FShaderParametersMetadata& UniformBufferStruct, EShaderPlatform Platform, FString& OutDeclaration)
+{
+	OutDeclaration = UE::ShaderParameters::CreateUniformBufferShaderDeclaration(Name, UniformBufferStruct);
+}
+
+UE_DEPRECATED(5.2, "CacheUniformBufferIncludes should no longer be used.")
+RENDERCORE_API void CacheUniformBufferIncludes(TMap<const TCHAR*, FCachedUniformBufferDeclaration>& Cache, EShaderPlatform Platform);
+#endif
 
 enum EShaderParameterFlags
 {
diff --git a/Engine/Source/Runtime/RenderCore/Public/VertexFactory.h b/Engine/Source/Runtime/RenderCore/Public/VertexFactory.h
--- a/Engine/Source/Runtime/RenderCore/Public/VertexFactory.h
+++ b/Engine/Source/Runtime/RenderCore/Public/VertexFactory.h
@@ -487,14 +487,17 @@ public:
 	}
 
 	/** Adds include statements for uniform buffers that this shader type references, and builds a prefix for the shader file with the include statements. */
-	RENDERCORE_API void AddReferencedUniformBufferIncludes(FShaderCompilerEnvironment& OutEnvironment, FString& OutSourceFilePrefix, EShaderPlatform Platform) const;
+	RENDERCORE_API void AddUniformBufferIncludesToEnvironment(FShaderCompilerEnvironment& OutEnvironment, EShaderPlatform Platform) const;
 
-	RENDERCORE_API void FlushShaderFileCache(const TMap<FString, TArray<const TCHAR*> >& ShaderFileToUniformBufferVariables);
-
-	const TMap<const TCHAR*, FCachedUniformBufferDeclaration>& GetReferencedUniformBufferStructsCache() const
+	UE_DEPRECATED(5.2, "AddReferencedUniformBufferIncludes has moved to AddUniformBufferIncludesToEnvironment and no longer takes a prefix argument.")
+	inline void AddReferencedUniformBufferIncludes(FShaderCompilerEnvironment& OutEnvironment, FString& OutSourceFilePrefix, EShaderPlatform Platform) const
 	{
-		return ReferencedUniformBufferStructsCache;
+		AddUniformBufferIncludesToEnvironment(OutEnvironment, Platform);
 	}
+
+	RENDERCORE_API void FlushShaderFileCache(const TMap<FString, TArray<const TCHAR*> >& ShaderFileToUniformBufferVariables);
+
+	inline const TSet<const TCHAR*>& GetReferencedUniformBufferNames() const { return ReferencedUniformBufferNames; };
 #endif // WITH_EDITOR
 
 private:
@@ -524,11 +527,11 @@ private:
 	/** 
 	 * Cache of referenced uniform buffer includes.  
 	 * These are derived from source files so they need to be flushed when editing and recompiling shaders on the fly. 
-	 * FVertexFactoryType::Initialize will add an entry for each referenced uniform buffer, but the declarations are added on demand as shaders are compiled.
+	 * FShaderType::Initialize will add the referenced uniform buffers, but this set may be updated by FlushShaderFileCache
 	 */
-	mutable TMap<const TCHAR*, FCachedUniformBufferDeclaration> ReferencedUniformBufferStructsCache;
+	TSet<const TCHAR*> ReferencedUniformBufferNames;
 
-	/** Tracks what platforms ReferencedUniformBufferStructsCache has had declarations cached for. */
+	/** Tracks the last platform ReferencedUniformBufferNames was read from in AddUniformBufferIncludesToEnvironment. */
 	mutable std::atomic<EShaderPlatform> CachedUniformBufferPlatform{ SP_NumPlatforms };
 #endif // WITH_EDITOR
 };
