Change Fatal log to a warning due to Global library's behavior.

diff --git a/Engine/Source/Runtime/RenderCore/Private/ShaderCodeLibrary.cpp b/Engine/Source/Runtime/RenderCore/Private/ShaderCodeLibrary.cpp
--- a/Engine/Source/Runtime/RenderCore/Private/ShaderCodeLibrary.cpp
+++ b/Engine/Source/Runtime/RenderCore/Private/ShaderCodeLibrary.cpp
@@ -901,9 +901,13 @@ public:
 	{
 		// FShaderLibraryInstance cant release RHI on all of the resources because it does not own them
 		// and its lifetime must also be longer than any of the resources.
-		for (FShaderMapResource_SharedCode* Resource : Resources)
+		if (!IsEngineExitRequested())
 		{
-			UE_CLOG(Resource, LogShaderLibrary, Fatal, TEXT("FShaderLibraryInstance %s still has resources in use at time of destruction!"), *Library->GetName());
+			for (FShaderMapResource_SharedCode* Resource : Resources)
+			{
+				// Global library seems to not free the resources, so use Warning and not Fatal here.
+				UE_CLOG(Resource, LogShaderLibrary, Warning, TEXT("FShaderLibraryInstance %s still has resource (shader map index %d) in use at time of destruction!"), *Library->GetName(), Resource->ShaderMapIndex);
+			}
 		}
 		
 		Library->Teardown();
