Disable logging of stereo rendering aspects in the editor.

- There are too many cases where it can be noisy (shader compilation for 3+ RHIs) and it doesn't add as much value in debugging as on a device.

diff --git a/Engine/Source/Runtime/RenderCore/Private/StereoRenderUtils.cpp b/Engine/Source/Runtime/RenderCore/Private/StereoRenderUtils.cpp
--- a/Engine/Source/Runtime/RenderCore/Private/StereoRenderUtils.cpp
+++ b/Engine/Source/Runtime/RenderCore/Private/StereoRenderUtils.cpp
@@ -4,7 +4,7 @@
 #include "RenderUtils.h"
 
 // enable this to printf-debug stereo rendering aspects on a device
-#define UE_DEBUG_STEREO_ASPECTS			(!UE_BUILD_SHIPPING && !UE_BUILD_TEST)
+#define UE_DEBUG_STEREO_ASPECTS			(!UE_BUILD_SHIPPING && !UE_BUILD_TEST && !UE_EDITOR)
 
 #if UE_DEBUG_STEREO_ASPECTS
 	#define UE_DEBUG_SSA_LOG_INIT(Platform) \
@@ -15,7 +15,7 @@
 		{ \
 			LastLogged[1] = LastLogged[0]; \
 			LastLogged[0] = Platform; \
-			bSSALogEnable = FApp::CanEverRender() && !IsRunningCommandlet(); \
+			bSSALogEnable = FApp::CanEverRender(); \
 		} \
 	}
 
