ShaderParameters cleanup
- SetShaderParameters and ValidateShaderParameters now require the caller pass in the size of the parameter data. This will be useful when making SetShaderParameters run on the RHI thread.
- SetShaderParameters no longer requires a FShaderParametersMetadata argument since it was never used.

diff --git a/Engine/Source/Runtime/RenderCore/Private/ShaderParameterStruct.cpp b/Engine/Source/Runtime/RenderCore/Private/ShaderParameterStruct.cpp
--- a/Engine/Source/Runtime/RenderCore/Private/ShaderParameterStruct.cpp
+++ b/Engine/Source/Runtime/RenderCore/Private/ShaderParameterStruct.cpp
@@ -471,26 +471,29 @@ void EmitNullShaderParameterFatalError(const TShaderRef<FShader>& Shader, const
 struct FShaderParameterReader
 {
 	FShaderParameterReader() = delete;
-	FShaderParameterReader(const void* InData) : Data(reinterpret_cast<const uint8*>(InData)) { }
+	FShaderParameterReader(const void* InData, uint64 InDataSize) : Data(reinterpret_cast<const uint8*>(InData)), DataSize(InDataSize) { }
 
 	template<typename TParameterIn>
 	const void* GetRawPointer(const TParameterIn& InParameter) const
 	{
+		checkSlow(InParameter.ByteOffset <= DataSize);
 		return (Data + InParameter.ByteOffset);
 	}
 
 	template<typename TParameterOut, typename TParameterIn>
 	const TParameterOut& Read(const TParameterIn& InParameter) const
 	{
+		checkSlow(InParameter.ByteOffset + sizeof(TParameterOut) <= DataSize);
 		return *reinterpret_cast<const TParameterOut*>(Data + InParameter.ByteOffset);
 	}
 
 	const uint8* Data;
+	const uint64 DataSize;
 };
 
 #if DO_CHECK
 
-void ValidateShaderParameters(const TShaderRef<FShader>& Shader, const FShaderParametersMetadata* ParametersMetadata, const void* Parameters)
+void ValidateShaderParameters(const TShaderRef<FShader>& Shader, const FShaderParametersMetadata* ParametersMetadata, const void* ParametersData, uint64 ParametersSize)
 {
 	const FShaderParameterBindings& Bindings = Shader->Bindings;
 
@@ -499,7 +502,7 @@ void ValidateShaderParameters(const TShaderRef<FShader>& Shader, const FShaderPa
 		TEXT("Shader %s's parameter structure has changed without recompilation of the shader"),
 		Shader.GetType()->GetName());
 
-	const FShaderParameterReader Reader(Parameters);
+	const FShaderParameterReader Reader(ParametersData, ParametersSize);
 
 	const TCHAR* ShaderClassName = Shader.GetType()->GetName();
 	const TCHAR* ShaderParameterStructName = ParametersMetadata->GetStructTypeName();
@@ -712,12 +715,11 @@ inline void SetShaderParametersInternal(
 	TRHICmdList& RHICmdList,
 	TShaderRHI* ShaderRHI,
 	const FShaderParameterBindings& Bindings,
-	const FShaderParametersMetadata* ParametersMetadata,
-	const uint8* Base)
+	const void* ParametersData, uint64 ParametersSize)
 {
 	checkf(Bindings.RootParameterBufferIndex == FShaderParameterBindings::kInvalidBufferIndex, TEXT("Can't use SetShaderParameters() for root parameter buffer index."));
 
-	FShaderParameterReader Reader(Base);
+	FShaderParameterReader Reader(ParametersData, ParametersSize);
 
 	// Parameters
 	for (const FShaderParameterBindings::FParameter& Parameter : Bindings.Parameters)
@@ -829,30 +831,27 @@ void SetShaderParameters(
 	FRHIComputeCommandList& RHICmdList,
 	FRHIComputeShader* ShaderRHI,
 	const FShaderParameterBindings& Bindings,
-	const FShaderParametersMetadata* ParametersMetadata,
-	const uint8* Base)
+	const void* ParametersData, uint64 ParametersSize)
 {
-	SetShaderParametersInternal(RHICmdList, ShaderRHI, Bindings, ParametersMetadata, Base);
+	SetShaderParametersInternal(RHICmdList, ShaderRHI, Bindings, ParametersData, ParametersSize);
 }
 
 void SetShaderParameters(
 	FRHICommandList& RHICmdList,
 	FRHIGraphicsShader* ShaderRHI,
 	const FShaderParameterBindings& Bindings,
-	const FShaderParametersMetadata* ParametersMetadata,
-	const uint8* Base)
+	const void* ParametersData, uint64 ParametersSize)
 {
-	SetShaderParametersInternal(RHICmdList, ShaderRHI, Bindings, ParametersMetadata, Base);
+	SetShaderParametersInternal(RHICmdList, ShaderRHI, Bindings, ParametersData, ParametersSize);
 }
 
 void SetShaderParameters(
 	FRHICommandList& RHICmdList,
 	FRHIComputeShader* ShaderRHI,
 	const FShaderParameterBindings& Bindings,
-	const FShaderParametersMetadata* ParametersMetadata,
-	const uint8* Base)
+	const void* ParametersData, uint64 ParametersSize)
 {
-	SetShaderParametersInternal(RHICmdList, ShaderRHI, Bindings, ParametersMetadata, Base);
+	SetShaderParametersInternal(RHICmdList, ShaderRHI, Bindings, ParametersData, ParametersSize);
 }
 
 #if RHI_RAYTRACING
@@ -860,9 +859,9 @@ void SetShaderParameters(
 	FRayTracingShaderBindingsWriter& RTBindingsWriter,
 	const FShaderParameterBindings& Bindings,
 	const FRHIUniformBufferLayout* RootUniformBufferLayout,
-	const uint8* Base)
+	const void* ParametersData, uint64 ParametersSize)
 {
-	const FShaderParameterReader Reader(Base);
+	const FShaderParameterReader Reader(ParametersData, ParametersSize);
 
 	for (const FShaderParameterBindings::FResourceParameter& Parameter : Bindings.ResourceParameters)
 	{
@@ -948,7 +947,7 @@ void SetShaderParameters(
 		// Do not do any validation at some resources may have been removed from the structure because known to not be used by the shader.
 		EUniformBufferValidation Validation = EUniformBufferValidation::None;
 
-		RTBindingsWriter.RootUniformBuffer = RHICreateUniformBuffer(Base, RootUniformBufferLayout, UniformBuffer_SingleDraw, Validation);
+		RTBindingsWriter.RootUniformBuffer = RHICreateUniformBuffer(ParametersData, RootUniformBufferLayout, UniformBuffer_SingleDraw, Validation);
 		RTBindingsWriter.SetUniformBuffer(Bindings.RootParameterBufferIndex, RTBindingsWriter.RootUniformBuffer);
 	}
 }
diff --git a/Engine/Source/Runtime/RenderCore/Public/RenderGraphUtils.h b/Engine/Source/Runtime/RenderCore/Public/RenderGraphUtils.h
--- a/Engine/Source/Runtime/RenderCore/Public/RenderGraphUtils.h
+++ b/Engine/Source/Runtime/RenderCore/Public/RenderGraphUtils.h
@@ -210,7 +210,7 @@ void ClearUnusedGraphResources(
 	// Verify the shader have all the parameters it needs. This is done before the
 	// ClearUnusedGraphResourcesImpl() to not mislead user on why some resource are missing
 	// when debugging a validation failure.
-	ValidateShaderParameters(Shader, ParametersMetadata, InoutParameters);
+	ValidateShaderParameters(Shader, ParametersMetadata, InoutParameters, sizeof(*InoutParameters));
 
 	// Clear the resources the shader won't need.
 	return ClearUnusedGraphResourcesImpl(Shader->Bindings, ParametersMetadata, InoutParameters, ExcludeList);
@@ -240,8 +240,8 @@ void ClearUnusedGraphResources(
 	// Verify the shader have all the parameters it needs. This is done before the
 	// ClearUnusedGraphResourcesImpl() to not mislead user on why some resource are missing
 	// when debugging a validation failure.
-	ValidateShaderParameters(ShaderA, ParametersMetadata, InoutParameters);
-	ValidateShaderParameters(ShaderB, ParametersMetadata, InoutParameters);
+	ValidateShaderParameters(ShaderA, ParametersMetadata, InoutParameters, sizeof(*InoutParameters));
+	ValidateShaderParameters(ShaderB, ParametersMetadata, InoutParameters, sizeof(*InoutParameters));
 
 	// Clear the resources the shader won't need.
 	const FShaderParameterBindings* ShaderBindings[] = { &ShaderA->Bindings, &ShaderB->Bindings };
diff --git a/Engine/Source/Runtime/RenderCore/Public/ShaderParameterStruct.h b/Engine/Source/Runtime/RenderCore/Public/ShaderParameterStruct.h
--- a/Engine/Source/Runtime/RenderCore/Public/ShaderParameterStruct.h
+++ b/Engine/Source/Runtime/RenderCore/Public/ShaderParameterStruct.h
@@ -117,16 +117,21 @@ extern RENDERCORE_API void EmitNullShaderParameterFatalError(const TShaderRef<FS
 
 /** Validates that all resource parameters of a shader are set. */
 #if DO_CHECK
-extern RENDERCORE_API void ValidateShaderParameters(const TShaderRef<FShader>& Shader, const FShaderParametersMetadata* ParametersMetadata, const void* Parameters);
+extern RENDERCORE_API void ValidateShaderParameters(const TShaderRef<FShader>& Shader, const FShaderParametersMetadata* ParametersMetadata, const void* ParametersData, uint64 ParametersSize);
+
+UE_DEPRECATED(5.2, "ValidateShaderParameters now requires the size of the parameters data")
+inline void ValidateShaderParameters(const TShaderRef<FShader>& Shader, const FShaderParametersMetadata* ParametersMetadata, const void* ParametersData)
+{
+	ValidateShaderParameters(Shader, ParametersMetadata, ParametersData, 0);
+}
 #else
-FORCEINLINE void ValidateShaderParameters(const TShaderRef<FShader>& Shader, const FShaderParametersMetadata* ParametersMetadata, const void* Parameters) {}
+FORCEINLINE void ValidateShaderParameters(const TShaderRef<FShader>& Shader, const FShaderParametersMetadata* ParametersMetadata, const void* ParametersData, uint64 ParametersSize) {}
 #endif
 
 template<typename TShaderClass>
 FORCEINLINE void ValidateShaderParameters(const TShaderRef<TShaderClass>& Shader, const typename TShaderClass::FParameters& Parameters)
 {
-	const typename TShaderClass::FParameters* ParameterPtr = &Parameters;
-	return ValidateShaderParameters(Shader, TShaderClass::FParameters::FTypeInfo::GetStructMetadata(), ParameterPtr);
+	return ValidateShaderParameters(Shader, TShaderClass::FParameters::FTypeInfo::GetStructMetadata(), &Parameters, sizeof(Parameters));
 }
 
 /** Unset compute shader UAVs. */
@@ -172,22 +177,19 @@ RENDERCORE_API void SetShaderParameters(
 	FRHIComputeCommandList& RHICmdList,
 	FRHIComputeShader* ShaderRHI,
 	const FShaderParameterBindings& Bindings,
-	const FShaderParametersMetadata* ParametersMetadata,
-	const uint8* Base);
+	const void* ParametersData, uint64 ParametersSize);
 
 RENDERCORE_API void SetShaderParameters(
 	FRHICommandList& RHICmdList,
 	FRHIGraphicsShader* ShaderRHI,
 	const FShaderParameterBindings& Bindings,
-	const FShaderParametersMetadata* ParametersMetadata,
-	const uint8* Base);
+	const void* ParametersData, uint64 ParametersSize);
 
 RENDERCORE_API void SetShaderParameters(
 	FRHICommandList& RHICmdList,
 	FRHIComputeShader* ShaderRHI,
 	const FShaderParameterBindings& Bindings,
-	const FShaderParametersMetadata* ParametersMetadata,
-	const uint8* Base);
+	const void* ParametersData, uint64 ParametersSize);
 
 template<typename TRHICmdList, typename TShaderClass, typename TShaderRHI>
 inline void SetShaderParameters(
@@ -197,15 +199,11 @@ inline void SetShaderParameters(
 	const FShaderParametersMetadata* ParametersMetadata,
 	const typename TShaderClass::FParameters& Parameters)
 {
-	ValidateShaderParameters(Shader, ParametersMetadata, &Parameters);
+	ValidateShaderParameters(Shader, ParametersMetadata, &Parameters, sizeof(Parameters));
 
 	// TODO(RDG): Once all shader sets their parameter through this, can refactor RHI so all shader parameters get sets through a single RHI function call.
-	const FShaderParameterBindings& Bindings = Shader->Bindings;
-
-	const typename TShaderClass::FParameters* ParametersPtr = &Parameters;
-	const uint8* Base = reinterpret_cast<const uint8*>(ParametersPtr);
 
-	SetShaderParameters(RHICmdList, ShaderRHI, Bindings, ParametersMetadata, Base);
+	SetShaderParameters(RHICmdList, ShaderRHI, Shader->Bindings, &Parameters, sizeof(Parameters));
 }
 
 template<typename TRHICmdList, typename TShaderClass, typename TShaderRHI>
@@ -221,7 +219,7 @@ RENDERCORE_API void SetShaderParameters(
 	FRayTracingShaderBindingsWriter& RTBindingsWriter,
 	const FShaderParameterBindings& Bindings,
 	const FRHIUniformBufferLayout* RootUniformBufferLayout,
-	const uint8* Base);
+	const void* ParametersData, uint64 ParametersSize);
 
 /** Set shader's parameters from its parameters struct. */
 template<typename TShaderClass>
@@ -229,14 +227,9 @@ void SetShaderParameters(FRayTracingShaderBindingsWriter& RTBindingsWriter, cons
 {
 	ValidateShaderParameters(Shader, Parameters);
 
-	const FShaderParameterBindings& Bindings = Shader->Bindings;
-
-	checkf(Bindings.Parameters.Num() == 0, TEXT("Ray tracing shader should use SHADER_USE_ROOT_PARAMETER_STRUCT() to passdown the cbuffer layout to the shader compiler."));
-
-	const typename TShaderClass::FParameters* ParametersPtr = &Parameters;
-	const uint8* Base = reinterpret_cast<const uint8*>(ParametersPtr);
+	checkf(Shader->Bindings.Parameters.Num() == 0, TEXT("Ray tracing shader should use SHADER_USE_ROOT_PARAMETER_STRUCT() to passdown the cbuffer layout to the shader compiler."));
 
-	SetShaderParameters(RTBindingsWriter, Bindings, TShaderClass::FParameters::FTypeInfo::GetStructMetadata()->GetLayoutPtr(), Base);
+	SetShaderParameters(RTBindingsWriter, Shader->Bindings, TShaderClass::FParameters::FTypeInfo::GetStructMetadata()->GetLayoutPtr(), &Parameters, sizeof(Parameters));
 }
 
 #endif // RHI_RAYTRACING
