Fix crash in D3D11 when GPU crash debugging is enabled.

FRDGEventScopeStack::EndExecute did not pass the value of bRDGEvents to the scope array, so we ended up trying to pop events that were never pushed when bRDGEvents was false.

#rnx

diff --git a/Engine/Source/Runtime/RenderCore/Public/RenderGraphEvent.h b/Engine/Source/Runtime/RenderCore/Public/RenderGraphEvent.h
--- a/Engine/Source/Runtime/RenderCore/Public/RenderGraphEvent.h
+++ b/Engine/Source/Runtime/RenderCore/Public/RenderGraphEvent.h
@@ -527,7 +527,7 @@ public:
 	{
 		if (IsEnabled())
 		{
-			FRDGEventScopeOpArray Array = ScopeStack.EndCompile();
+			FRDGEventScopeOpArray Array(ScopeStack.EndCompile(), bRDGEvents);
 			if (Array.Ops.Num())
 			{
 				FRHICommandListScopedPipeline Scope(RHICmdList, Pipeline);
