Removing stat scopes that are very small (2us) and are just adding noise to profiling captures; these just enqueue a render command and don't do any busy work inside.

diff --git a/Engine/Source/Runtime/RenderCore/Private/RenderCore.cpp b/Engine/Source/Runtime/RenderCore/Private/RenderCore.cpp
--- a/Engine/Source/Runtime/RenderCore/Private/RenderCore.cpp
+++ b/Engine/Source/Runtime/RenderCore/Private/RenderCore.cpp
@@ -182,7 +182,6 @@ DEFINE_STAT(STAT_RemoveScenePrimitiveGT);
 DEFINE_STAT(STAT_AddScenePrimitiveGT);
 DEFINE_STAT(STAT_UpdatePrimitiveTransformGT);
 DEFINE_STAT(STAT_UpdatePrimitiveInstanceGT);
-DEFINE_STAT(STAT_UpdateCustomPrimitiveDataGT);
 
 DEFINE_STAT(STAT_Scene_SetShaderMapsOnMaterialResources_RT);
 DEFINE_STAT(STAT_Scene_UpdateStaticDrawLists_RT);
diff --git a/Engine/Source/Runtime/RenderCore/Public/RenderCore.h b/Engine/Source/Runtime/RenderCore/Public/RenderCore.h
--- a/Engine/Source/Runtime/RenderCore/Public/RenderCore.h
+++ b/Engine/Source/Runtime/RenderCore/Public/RenderCore.h
@@ -140,7 +140,6 @@ DECLARE_CYCLE_STAT_EXTERN(TEXT("RemovePrimitive (GT)"),STAT_RemoveScenePrimitive
 DECLARE_CYCLE_STAT_EXTERN(TEXT("AddPrimitive (GT)"),STAT_AddScenePrimitiveGT,STATGROUP_SceneUpdate, RENDERCORE_API);
 DECLARE_CYCLE_STAT_EXTERN(TEXT("UpdatePrimitiveTransform (GT)"),STAT_UpdatePrimitiveTransformGT,STATGROUP_SceneUpdate, RENDERCORE_API);
 DECLARE_CYCLE_STAT_EXTERN(TEXT("UpdatePrimitiveInstance (GT)"), STAT_UpdatePrimitiveInstanceGT, STATGROUP_SceneUpdate, RENDERCORE_API);
-DECLARE_CYCLE_STAT_EXTERN(TEXT("UpdateCustomPrimitiveData (GT)"),STAT_UpdateCustomPrimitiveDataGT,STATGROUP_SceneUpdate, RENDERCORE_API);
 
 DECLARE_CYCLE_STAT_EXTERN(TEXT("Set Shader Maps On Material Resources (RT)"),STAT_Scene_SetShaderMapsOnMaterialResources_RT,STATGROUP_SceneUpdate, RENDERCORE_API);
 DECLARE_CYCLE_STAT_EXTERN(TEXT("Update Static Draw Lists (RT)"), STAT_Scene_UpdateStaticDrawLists_RT, STATGROUP_SceneUpdate, RENDERCORE_API);
diff --git a/Engine/Source/Runtime/Renderer/Private/RendererScene.cpp b/Engine/Source/Runtime/Renderer/Private/RendererScene.cpp
--- a/Engine/Source/Runtime/Renderer/Private/RendererScene.cpp
+++ b/Engine/Source/Runtime/Renderer/Private/RendererScene.cpp
@@ -2416,9 +2416,6 @@ void FScene::UpdatePrimitiveAttachment(UPrimitiveComponent* Primitive)
 
 void FScene::UpdateCustomPrimitiveData(UPrimitiveComponent* Primitive)
 {
-	SCOPE_CYCLE_COUNTER(STAT_UpdateCustomPrimitiveDataGT);
-	SCOPED_NAMED_EVENT(FScene_UpdateCustomPrimitiveData, FColor::Yellow);
-
 	// This path updates the primitive data directly in the GPUScene. 
 	if(Primitive->SceneProxy) 
 	{
