Fix another crash on exit due ot leaked global library's resources.

diff --git a/Engine/Source/Runtime/RenderCore/Private/ShaderCodeLibrary.cpp b/Engine/Source/Runtime/RenderCore/Private/ShaderCodeLibrary.cpp
--- a/Engine/Source/Runtime/RenderCore/Private/ShaderCodeLibrary.cpp
+++ b/Engine/Source/Runtime/RenderCore/Private/ShaderCodeLibrary.cpp
@@ -899,14 +899,24 @@ public:
 
 	~FShaderLibraryInstance()
 	{
-		// FShaderLibraryInstance cant release RHI on all of the resources because it does not own them
-		// and its lifetime must also be longer than any of the resources.
-		if (!IsEngineExitRequested())
+		// release RHI on all of the resources
+		// TODO (UE-169261): figure out exact ownership and lifetime guarantees. Looks like there's a discrepancy between global and game libraries
+		// and global ends up owning the resources while game ones not.
+		bool bGlobalLibrary = Library->GetName() == TEXT("Global");
+		for (FShaderMapResource_SharedCode* Resource : Resources)
 		{
-			for (FShaderMapResource_SharedCode* Resource : Resources)
+			if (Resource)
 			{
-				// Global library seems to not free the resources, so use Warning and not Fatal here.
-				UE_CLOG(Resource, LogShaderLibrary, Warning, TEXT("FShaderLibraryInstance %s still has resource (shader map index %d) in use at time of destruction!"), *Library->GetName(), Resource->ShaderMapIndex);
+				// only release resources on a global library. Game libraries can be unloaded dynamically (e.g. plugins),
+				// and any resources left at this point should be reported
+				if (bGlobalLibrary)
+				{
+					BeginReleaseResource(Resource);
+				}
+				else
+				{
+					UE_LOG(LogShaderLibrary, Warning, TEXT("FShaderLibraryInstance %s still has resource (shader map index %d) in use at time of destruction!"), *Library->GetName(), Resource->ShaderMapIndex);
+				}
 			}
 		}
 		
