Support concurrent shader creation within a shadermap (UE-166109).

- Avoids oversubscription deadlock when multiple background threads are creating shaders during the PSO precaching.

diff --git a/Engine/Source/Runtime/RenderCore/Private/ShaderCodeArchive.cpp b/Engine/Source/Runtime/RenderCore/Private/ShaderCodeArchive.cpp
--- a/Engine/Source/Runtime/RenderCore/Private/ShaderCodeArchive.cpp
+++ b/Engine/Source/Runtime/RenderCore/Private/ShaderCodeArchive.cpp
@@ -1200,43 +1200,25 @@ TRefCountPtr<FRHIShader> FShaderCodeArchive::CreateShader(int32 Index)
 	FShaderPreloadEntry& ShaderPreloadEntry = ShaderPreloads[Index];
 	checkf(!ShaderPreloadEntry.bNeverToBePreloaded, TEXT("We are creating a shader that shouldn't be preloaded in this run (e.g. raytracing shader on D3D11)."));
 
-	void* PreloadedShaderCode = nullptr;
+	FGraphEventArray ReadCompleteEvents;
+	PreloadShader(Index, ReadCompleteEvents);
+
 	{
+		TRACE_CPUPROFILER_EVENT_SCOPE(BlockingShaderLoad);
+		double TimeStarted = FPlatformTime::Seconds();
 		const bool bNeededToWait = WaitForPreload(ShaderPreloadEntry);
 		if (bNeededToWait)
 		{
-			UE_LOG(LogShaderLibrary, Warning, TEXT("Blocking wait for shader preload, NumRefs: %d, FramePreloadStarted: %d"), ShaderPreloadEntry.NumRefs, ShaderPreloadEntry.FramePreloadStarted);
-		}
-
-		FWriteScopeLock Lock(ShaderPreloadLock);
-		if (ShaderPreloadEntry.NumRefs > 0u)
-		{
-			check(!ShaderPreloadEntry.PreloadEvent || ShaderPreloadEntry.PreloadEvent->IsComplete());
-			ShaderPreloadEntry.PreloadEvent.SafeRelease();
-
-			ShaderPreloadEntry.NumRefs++; // Hold a reference to code while we're using it to create shader
-			PreloadedShaderCode = ShaderPreloadEntry.Code;
-			check(PreloadedShaderCode);
-		}
-	}
-
-	const uint8* ShaderCode = (uint8*)PreloadedShaderCode;
-	if (!ShaderCode)
-	{
-		UE_LOG(LogShaderLibrary, Warning, TEXT("Blocking shader load, NumRefs: %d, FramePreloadStarted: %d"), ShaderPreloadEntry.NumRefs, ShaderPreloadEntry.FramePreloadStarted);
-
-		FGraphEventArray ReadCompleteEvents;
-		EAsyncIOPriorityAndFlags DontCache = GShaderCodeLibraryAsyncLoadingAllowDontCache ? AIOP_FLAG_DONTCACHE : AIOP_MIN;
-		IMemoryReadStreamRef LoadedCode = FileCacheHandle->ReadData(ReadCompleteEvents, LibraryCodeOffset + ShaderEntry.Offset, ShaderEntry.Size, AIOP_CriticalPath | DontCache);
-		if (ReadCompleteEvents.Num() > 0)
-		{
-			FTaskGraphInterface::Get().WaitUntilTasksComplete(ReadCompleteEvents);
+			double WaitDuration = FPlatformTime::Seconds() - TimeStarted;
+			// only complain if we spent more than 1ms waiting
+			if (TimeStarted > 0.001)
+			{
+				UE_LOG(LogShaderLibrary, Warning, TEXT("Spent %.2f ms in a blocking wait for shader preload, NumRefs: %d, FramePreloadStarted: %d, CurrentFrame: %d"), WaitDuration * 1000.0, ShaderPreloadEntry.NumRefs, ShaderPreloadEntry.FramePreloadStarted, GFrameNumber);
+			}
 		}
-		void* LoadedShaderCode = MemStack.Alloc(ShaderEntry.Size, 16);
-		LoadedCode->CopyTo(LoadedShaderCode, 0, ShaderEntry.Size);
-		ShaderCode = (uint8*)LoadedShaderCode;
 	}
 
+	const uint8* ShaderCode = (uint8*)ShaderPreloadEntry.Code;
 	if (ShaderEntry.UncompressedSize != ShaderEntry.Size)
 	{
 		uint8* UncompressedCode = reinterpret_cast<uint8*>(MemStack.Alloc(ShaderEntry.UncompressedSize, 16));
@@ -1270,14 +1252,8 @@ TRefCountPtr<FRHIShader> FShaderCodeArchive::CreateShader(int32 Index)
 	}
 	DebugVisualizer.MarkCreatedForVisualization(Index);
 
-	// Release the refernece we were holding
-	if (PreloadedShaderCode)
-	{
-		FWriteScopeLock Lock(ShaderPreloadLock);
-		check(ShaderPreloadEntry.NumRefs > 1u); // we shouldn't be holding the last ref here
-		--ShaderPreloadEntry.NumRefs;
-		PreloadedShaderCode = nullptr;
-	}
+	// Release the reference we were holding
+	ReleasePreloadedShader(Index);
 
 	if (Shader)
 	{
@@ -2011,44 +1987,23 @@ TRefCountPtr<FRHIShader> FIoStoreShaderCodeArchive::CreateShader(int32 ShaderInd
 
 	const FIoStoreShaderCodeEntry& ShaderEntry = Header.ShaderEntries[ShaderIndex];
 	int32 GroupIndex = GetGroupIndexForShader(ShaderIndex);
-	FShaderGroupPreloadEntry* PreloadEntryPtr = nullptr;
-	bool bShaderWasNotPreloaded = false;
 
-	FGraphEventRef Event;
-	{
-		{
-			// preloading a shader group (write lock because we will need to adjust its refcount)
-			FWriteScopeLock Lock(PreloadedShaderGroupsLock);
-			FShaderGroupPreloadEntry** ExistingEntry = PreloadedShaderGroups.Find(GroupIndex);
-			PreloadEntryPtr = ExistingEntry ? *ExistingEntry : PreloadEntryPtr;
-			// if we have an entry, addref it so it doesn't go away
-			if (PreloadEntryPtr)
-			{
-				checkf(PreloadEntryPtr->NumRefs > 0, TEXT("Existing preload entry is not referenced! (NumRefs=%d)"), PreloadEntryPtr->NumRefs);
-				++PreloadEntryPtr->NumRefs;
-			}
-		}
+	// Preload shader group if it wasn't yet. This will also addref it so we can be sure it will exist.
+	FGraphEventArray Dummy;
+	PreloadShaderGroup(GroupIndex, Dummy);
 
-		// the shader we're attempting to create could have not been preloaded
-		if (PreloadEntryPtr == nullptr)
-		{
-			bShaderWasNotPreloaded = true;
-			FGraphEventArray Dummy;
-			PreloadShaderGroup(GroupIndex, Dummy);
+	FShaderGroupPreloadEntry* PreloadEntryPtr;
+	{
+		FReadScopeLock Lock(PreloadedShaderGroupsLock);
+		PreloadEntryPtr = FindExistingPreloadEntry(GroupIndex);
+	}
 
-			// preloading a shader group should have created (and AddRef'd) the entry, so we can get and keep the pointer
-			FReadScopeLock Lock(PreloadedShaderGroupsLock);
-			PreloadEntryPtr = FindOrAddPreloadEntry(GroupIndex);
-		}
-		check(PreloadEntryPtr);
-		
-		// raise the prio if still ongoing
-		if (!PreloadEntryPtr->IoRequest.Status().IsCompleted())
-		{
-			PreloadEntryPtr->IoRequest.UpdatePriority(IoDispatcherPriority_Max);
-		}
-		Event = PreloadEntryPtr->PreloadEvent;
+	// raise the prio if still ongoing
+	if (!PreloadEntryPtr->IoRequest.Status().IsCompleted())
+	{
+		PreloadEntryPtr->IoRequest.UpdatePriority(IoDispatcherPriority_Max);
 	}
+	FGraphEventRef Event = PreloadEntryPtr->PreloadEvent;
 
 	const bool bNeededToWait = Event.IsValid() && !Event->IsComplete();
 	if (bNeededToWait)
@@ -2112,18 +2067,8 @@ TRefCountPtr<FRHIShader> FIoStoreShaderCodeArchive::CreateShader(int32 ShaderInd
 	}
 	DebugVisualizer.MarkCreatedForVisualization(ShaderIndex);
 
-	if (bShaderWasNotPreloaded)
-	{
-		ReleasePreloadEntry(GroupIndex);
-	}
-	else
-	{
-		FWriteScopeLock Lock(PreloadedShaderGroupsLock);
-		// if the entry existed before Create was called, it will be released by the higher level resource after the creation.
-		checkf(PreloadEntryPtr->NumRefs > 1, TEXT("We shouldn't be releasing the last preload entry here (NumRefs=%d)"), PreloadEntryPtr->NumRefs);
-		--PreloadEntryPtr->NumRefs;
-	}
 	PreloadEntryPtr = nullptr;
+	ReleasePreloadEntry(GroupIndex);
 
 	if (Shader)
 	{
diff --git a/Engine/Source/Runtime/RenderCore/Private/ShaderCodeLibrary.cpp b/Engine/Source/Runtime/RenderCore/Private/ShaderCodeLibrary.cpp
--- a/Engine/Source/Runtime/RenderCore/Private/ShaderCodeLibrary.cpp
+++ b/Engine/Source/Runtime/RenderCore/Private/ShaderCodeLibrary.cpp
@@ -238,7 +238,8 @@ public:
 	virtual void ReleaseRHI() override;
 
 	// FShaderMapResource interface
-	virtual TRefCountPtr<FRHIShader> CreateRHIShader(int32 ShaderIndex) override;
+	virtual FRHIShader* CreateRHIShaderOrCrash(int32 ShaderIndex) override;
+	virtual void ReleasePreloadedShaderCode(int32 ShaderIndex) override;
 	virtual bool TryRelease() override;
 	virtual uint32 GetSizeBytes() const override { return sizeof(*this) + GetAllocatedSize(); }
 
@@ -1041,7 +1042,15 @@ public:
 			Shader = Library->CreateShader(ShaderIndex);
 
 			FRWScopeLock Locker(ShaderLocks[BucketIndex], SLT_Write);
-			RHIShaders[BucketIndex].Add(ShaderIndex, Shader);
+			TRefCountPtr<FRHIShader>* ShaderPtr = RHIShaders[BucketIndex].Find(ShaderIndex);
+			if (LIKELY(ShaderPtr == nullptr))
+			{
+				RHIShaders[BucketIndex].Add(ShaderIndex, Shader);
+			}
+			else
+			{
+				Shader = *ShaderPtr;
+			}
 		}
 		return Shader;
 	}
@@ -1155,18 +1164,43 @@ FShaderMapResource_SharedCode::~FShaderMapResource_SharedCode()
 	
 }
 
-TRefCountPtr<FRHIShader> FShaderMapResource_SharedCode::CreateRHIShader(int32 ShaderIndex)
+FRHIShader* FShaderMapResource_SharedCode::CreateRHIShaderOrCrash(int32 ShaderIndex)
 {
 	SCOPED_LOADTIMER(FShaderMapResource_SharedCode_InitRHI);
+#if STATS
+	double TimeFunctionEntered = FPlatformTime::Seconds();
+	ON_SCOPE_EXIT
+	{
+		if (IsInRenderingThread())
+		{
+			double ShaderCreationTime = FPlatformTime::Seconds() - TimeFunctionEntered;
+			INC_FLOAT_STAT_BY(STAT_Shaders_TotalRTShaderInitForRenderingTime, ShaderCreationTime);
+		}
+	};
+#endif
 
 	const int32 LibraryShaderIndex = LibraryInstance->Library->GetShaderIndex(ShaderMapIndex, ShaderIndex);
-	TRefCountPtr<FRHIShader> ShaderRHI = LibraryInstance->GetOrCreateShader(LibraryShaderIndex);
-	if (bShaderMapPreloaded && ShaderRHI)
+	FRHIShader* CreatedShader = LibraryInstance->GetOrCreateShader(LibraryShaderIndex);
+	if (UNLIKELY(CreatedShader == nullptr))
+	{
+		UE_LOG(LogShaders, Fatal, TEXT("FShaderMapResource_SharedCode::InitRHI is unable to create a shader"));
+		// unreachable
+		return nullptr;
+	}
+
+	CreatedShader->AddRef();
+	return CreatedShader;
+}
+
+void FShaderMapResource_SharedCode::ReleasePreloadedShaderCode(int32 ShaderIndex)
+{
+	SCOPED_LOADTIMER(FShaderMapResource_SharedCode_InitRHI);	// part of shader initialization in a way
+
+	if (bShaderMapPreloaded)
 	{
-		// Release our preload, once we've created the shader
+		const int32 LibraryShaderIndex = LibraryInstance->Library->GetShaderIndex(ShaderMapIndex, ShaderIndex);
 		LibraryInstance->Library->ReleasePreloadedShader(LibraryShaderIndex);
 	}
-	return ShaderRHI;
 }
 
 void FShaderMapResource_SharedCode::ReleaseRHI()
diff --git a/Engine/Source/Runtime/RenderCore/Private/ShaderResource.cpp b/Engine/Source/Runtime/RenderCore/Private/ShaderResource.cpp
--- a/Engine/Source/Runtime/RenderCore/Private/ShaderResource.cpp
+++ b/Engine/Source/Runtime/RenderCore/Private/ShaderResource.cpp
@@ -25,6 +25,7 @@
 #include "Misc/MemStack.h"
 #include "ShaderCompilerCore.h"
 #include "Compression/OodleDataCompression.h"
+#include "RHIResources.h"	// Access to FRHIRayTracingShader::RayTracingPayloadType requires this
 
 #if WITH_EDITORONLY_DATA
 #include "Interfaces/IShaderFormat.h"
@@ -589,45 +590,69 @@ void FShaderMapResource::BeginCreateAllShaders()
 	});
 }
 
-FRHIShader* FShaderMapResource::CreateShader(int32 ShaderIndex)
-{	
-	check(!RHIShaders[ShaderIndex].load(std::memory_order_acquire));
+FRHIShader* FShaderMapResource::CreateShaderOrCrash(int32 ShaderIndex)
+{
+	FRHIShader* Shader = nullptr;
+	// create before taking the lock. This may cause multiple creations, but it's better
+	// than a potential oversubscription deadlock, since CreateShader can spawn async tasks
+	FRHIShader* CreatedShader = CreateRHIShaderOrCrash(ShaderIndex);	// guaranteed to return non-null
 
-	TRefCountPtr<FRHIShader> RHIShader = CreateRHIShader(ShaderIndex);
-#if RHI_RAYTRACING
-	if (GRHISupportsRayTracing && GRHISupportsRayTracingShaders && RHIShader.IsValid())
 	{
-		switch (RHIShader->GetFrequency())
+		// Most shadermaps have <100 shaders, and less than a half of them can be created. 
+		// However, if this path is often contended, you can slice this lock
+		FScopeLock ScopeLock(&RHIShadersCreationGuard);
+
+		Shader = RHIShaders[ShaderIndex].load(std::memory_order_relaxed);
+		if (UNLIKELY(Shader == nullptr))
 		{
-		case SF_RayHitGroup:
-			RayTracingLibraryIndices[ShaderIndex] = GlobalRayTracingHitGroupLibrary.AddShader(static_cast<FRHIRayTracingShader*>(RHIShader.GetReference()));
-			break;
-		case SF_RayCallable:
-			RayTracingLibraryIndices[ShaderIndex] = GlobalRayTracingCallableShaderLibrary.AddShader(static_cast<FRHIRayTracingShader*>(RHIShader.GetReference()));
-			break;
-		case SF_RayMiss:
-			RayTracingLibraryIndices[ShaderIndex] = GlobalRayTracingMissShaderLibrary.AddShader(static_cast<FRHIRayTracingShader*>(RHIShader.GetReference()));
-			break;
-		case SF_RayGen:
-			// NOTE: we do not maintain a library for raygen shaders since the list of rayshaders we care about is usually small and consistent
-			break;
-		default:
-			break;
+			Shader = CreatedShader;
+			CreatedShader = nullptr;
+			RHIShaders[ShaderIndex].store(Shader, std::memory_order_release);
+
+#if RHI_RAYTRACING
+			// Registers RT shaders in global "libraries" that track all shaders potentially usable in a scene for adding to RTPSO
+			EShaderFrequency Frequency = Shader->GetFrequency();
+			if (LIKELY(GRHISupportsRayTracing && GRHISupportsRayTracingShaders))
+			{
+				switch (Frequency)
+				{
+					case SF_RayHitGroup:
+						RayTracingLibraryIndices[ShaderIndex] = GlobalRayTracingHitGroupLibrary.AddShader(static_cast<FRHIRayTracingShader*>(Shader));
+						break;
+					case SF_RayCallable:
+						RayTracingLibraryIndices[ShaderIndex] = GlobalRayTracingCallableShaderLibrary.AddShader(static_cast<FRHIRayTracingShader*>(Shader));
+						break;
+					case SF_RayMiss:
+						RayTracingLibraryIndices[ShaderIndex] = GlobalRayTracingMissShaderLibrary.AddShader(static_cast<FRHIRayTracingShader*>(Shader));
+						break;
+					case SF_RayGen:
+						// NOTE: we do not maintain a library for raygen shaders since the list of rayshaders we care about is usually small and consistent
+						break;
+					default:
+						break;
+				}
+			}
+#endif // RHI_RAYTRACING
+
+			// When using shader library, shader code is usually preloaded during the material load. Release it
+			// since we won't need it anymore for this shader.
+			ReleasePreloadedShaderCode(ShaderIndex);
 		}
 	}
-#endif // RHI_RAYTRACING
 
-	// keep the reference alive (the caller will release)
-	if (RHIShader.IsValid())
+	if (LIKELY(CreatedShader))
 	{
-		RHIShader->AddRef();
+		// free redundantly created shader
+		checkSlow(Shader != nullptr);
+		CreatedShader->Release();
 	}
-	return RHIShader.GetReference();
+
+	return Shader;
 }
 
-TRefCountPtr<FRHIShader> FShaderMapResource_InlineCode::CreateRHIShader(int32 ShaderIndex)
+FRHIShader* FShaderMapResource_InlineCode::CreateRHIShaderOrCrash(int32 ShaderIndex)
 {
-	TRACE_CPUPROFILER_EVENT_SCOPE(FShaderMapResource_InlineCode::CreateRHIShader);
+	TRACE_CPUPROFILER_EVENT_SCOPE(FShaderMapResource_InlineCode::CreateRHIShaderOrCrash);
 #if STATS
 	double TimeFunctionEntered = FPlatformTime::Seconds();
 	ON_SCOPE_EXIT
@@ -643,12 +668,10 @@ TRefCountPtr<FRHIShader> FShaderMapResource_InlineCode::CreateRHIShader(int32 Sh
 	// we can't have this called on the wrong platform's shaders
 	if (!ArePlatformsCompatible(GMaxRHIShaderPlatform, GetPlatform()))
 	{
-		if (FPlatformProperties::RequiresCookedData())
-		{
-			UE_LOG(LogShaders, Fatal, TEXT("FShaderMapResource_InlineCode::InitRHI got platform %s but it is not compatible with %s"),
-				*LegacyShaderPlatformToShaderFormat(GetPlatform()).ToString(), *LegacyShaderPlatformToShaderFormat(GMaxRHIShaderPlatform).ToString());
-		}
-		return TRefCountPtr<FRHIShader>();
+		UE_LOG(LogShaders, Fatal, TEXT("FShaderMapResource_InlineCode::InitRHI got platform %s but it is not compatible with %s"),
+			*LegacyShaderPlatformToShaderFormat(GetPlatform()).ToString(), *LegacyShaderPlatformToShaderFormat(GMaxRHIShaderPlatform).ToString());
+		// unreachable
+		return nullptr;
 	}
 
 	FMemStackBase& MemStack = FMemStack::Get();
@@ -668,7 +691,7 @@ TRefCountPtr<FRHIShader> FShaderMapResource_InlineCode::CreateRHIShader(int32 Sh
 	const FSHAHash& ShaderHash = Code->ShaderHashes[ShaderIndex];
 	const EShaderFrequency Frequency = ShaderEntry.Frequency;
 
-	TRefCountPtr<FRHIShader> RHIShader;
+	FRHIShader* RHIShader = nullptr;
 	switch (Frequency)
 	{
 	case SF_Vertex: RHIShader = RHICreateVertexShader(ShaderCodeView, ShaderHash); break;
@@ -689,11 +712,17 @@ TRefCountPtr<FRHIShader> FShaderMapResource_InlineCode::CreateRHIShader(int32 Sh
 		checkNoEntry();
 		break;
 	}
-
-	if (RHIShader)
+	if (UNLIKELY(RHIShader == nullptr))
 	{
-		INC_DWORD_STAT(STAT_Shaders_NumShadersUsedForRendering);
-		RHIShader->SetHash(ShaderHash);
+		UE_LOG(LogShaders, Fatal, TEXT("FShaderMapResource_InlineCode::InitRHI is unable to create a shader (frequency %d)"), static_cast<int32>(Frequency));
+		// unreachable
+		return nullptr;
 	}
+
+	INC_DWORD_STAT(STAT_Shaders_NumShadersUsedForRendering);
+	RHIShader->SetHash(ShaderHash);
+
+	// contract of this function is to return a shader with an already held reference
+	RHIShader->AddRef();
 	return RHIShader;
 }
diff --git a/Engine/Source/Runtime/RenderCore/Public/Shader.h b/Engine/Source/Runtime/RenderCore/Public/Shader.h
--- a/Engine/Source/Runtime/RenderCore/Public/Shader.h
+++ b/Engine/Source/Runtime/RenderCore/Public/Shader.h
@@ -341,21 +341,10 @@ public:
 		// synchronizing two threads: one that takes a lock and another that doesn't.
 		// Without fences, there is a race between storing the shader pointer and accessing it
 		// on the other (lockless) thread.
-
 		FRHIShader* Shader = RHIShaders[ShaderIndex].load(std::memory_order_acquire);
-		if (Shader == nullptr)
+		if (UNLIKELY(Shader == nullptr))
 		{
-			// most shadermaps have <100 shaders, and less than a half of them can be created. One lock
-			// for all creation seems sufficient, but if this function is often contended, per-shader
-			// locks are easily possible.
-			FScopeLock ScopeLock(&RHIShadersCreationGuard);
-
-			Shader = RHIShaders[ShaderIndex].load(std::memory_order_relaxed);
-			if (Shader == nullptr)
-			{
-				Shader = CreateShader(ShaderIndex);
-				RHIShaders[ShaderIndex].store(Shader, std::memory_order_release);
-			}
+			Shader = CreateShaderOrCrash(ShaderIndex);
 		}
 		return Shader;
 	}
@@ -417,23 +406,28 @@ protected:
 		return Size;
 	}
 
-	/** Addrefs the reference, passing the responsibility to the caller to Release() it. */
-	FRHIShader* CreateShader(int32 ShaderIndex);
+	/** Creates RHI shader, with a reference (so the caller can release). Never returns nullptr (inability to create is Fatal) */
+	virtual FRHIShader*	CreateRHIShaderOrCrash(int32 ShaderIndex) = 0;
+
+	/** Signal the shader library that it can release compressed shader code for a shader that it keeps preloaded in memory. */
+	virtual void ReleasePreloadedShaderCode(int32 ShaderIndex) { /* no-op when not using shader library */ };
 
-	virtual TRefCountPtr<FRHIShader> CreateRHIShader(int32 ShaderIndex) = 0;
 	virtual bool TryRelease() { return true; }
 
 	void ReleaseShaders();
 
 private:
 
+	/** Creates an entry in RHIShaders array and registers it among the raytracing libs if needed. Created shader is returned. */
+	FRHIShader* CreateShaderOrCrash(int32 ShaderIndex);
+
 	/** This lock is to prevent two threads creating the same RHIShaders element. It is only taken if the element is to be created. */
 	FCriticalSection RHIShadersCreationGuard;
 
 	/** An array of shader pointers (refcount is managed manually). */
 	TUniquePtr<std::atomic<FRHIShader*>[]> RHIShaders;
 
-	/** Since the shaders are no longer a TArray, this is their count (the size of the RHIShadersArray). */
+	/** Since the shaders are no longer a TArray, this is their count (the size of the RHIShaders array). */
 	int32 NumRHIShaders;
 
 #if RHI_RAYTRACING
@@ -530,7 +524,7 @@ public:
 	{}
 
 	// FShaderMapResource interface
-	virtual TRefCountPtr<FRHIShader> CreateRHIShader(int32 ShaderIndex) override;
+	virtual FRHIShader* CreateRHIShaderOrCrash(int32 ShaderIndex) override;
 	virtual uint32 GetSizeBytes() const override { return sizeof(*this) + GetAllocatedSize(); }
 
 	TRefCountPtr<FShaderMapResourceCode> Code;
diff --git a/Engine/Source/Runtime/RenderCore/Public/ShaderCodeArchive.h b/Engine/Source/Runtime/RenderCore/Public/ShaderCodeArchive.h
--- a/Engine/Source/Runtime/RenderCore/Public/ShaderCodeArchive.h
+++ b/Engine/Source/Runtime/RenderCore/Public/ShaderCodeArchive.h
@@ -626,6 +626,14 @@ private:
 		return Ptr;
 	}
 
+	/** Finds existing preload info for a shader group. Assumes lock guarding access to the info is taken */
+	inline FShaderGroupPreloadEntry* FindExistingPreloadEntry(int32 ShaderGroupIndex)
+	{
+		FShaderGroupPreloadEntry** PtrPtr = PreloadedShaderGroups.Find(ShaderGroupIndex);
+		checkf(PtrPtr != nullptr, TEXT("The preload entry for a shader group we assumed to exist does not exist!"));
+		return *PtrPtr;
+	}
+
 	/** Returns true if the group contains only RTX shaders. We can avoid preloading it when running with RTX off. */
 	bool GroupOnlyContainsRaytracingShaders(int32 ShaderGroupIndex);
 
