Make FShaderParametersMetadata authoritative on the parameters size, not the parameter type. Systems like Niagara can have more data allocated than can fit in the passed in data type, so the metadata has to be the source of the sizes here.
This ends up reverting parts of my last changes here, including a deprecation.

diff --git a/Engine/Source/Runtime/RenderCore/Private/ShaderParameterStruct.cpp b/Engine/Source/Runtime/RenderCore/Private/ShaderParameterStruct.cpp
--- a/Engine/Source/Runtime/RenderCore/Private/ShaderParameterStruct.cpp
+++ b/Engine/Source/Runtime/RenderCore/Private/ShaderParameterStruct.cpp
@@ -493,7 +493,7 @@ struct FShaderParameterReader
 
 #if DO_CHECK
 
-void ValidateShaderParameters(const TShaderRef<FShader>& Shader, const FShaderParametersMetadata* ParametersMetadata, const void* ParametersData, uint64 ParametersSize)
+void ValidateShaderParameters(const TShaderRef<FShader>& Shader, const FShaderParametersMetadata* ParametersMetadata, const void* ParametersData)
 {
 	const FShaderParameterBindings& Bindings = Shader->Bindings;
 
@@ -502,7 +502,7 @@ void ValidateShaderParameters(const TShaderRef<FShader>& Shader, const FShaderPa
 		TEXT("Shader %s's parameter structure has changed without recompilation of the shader"),
 		Shader.GetType()->GetName());
 
-	const FShaderParameterReader Reader(ParametersData, ParametersSize);
+	const FShaderParameterReader Reader(ParametersData, ParametersMetadata->GetSize());
 
 	const TCHAR* ShaderClassName = Shader.GetType()->GetName();
 	const TCHAR* ShaderParameterStructName = ParametersMetadata->GetStructTypeName();
@@ -715,11 +715,12 @@ inline void SetShaderParametersInternal(
 	TRHICmdList& RHICmdList,
 	TShaderRHI* ShaderRHI,
 	const FShaderParameterBindings& Bindings,
-	const void* ParametersData, uint64 ParametersSize)
+	const FShaderParametersMetadata* ParametersMetadata,
+	const void* ParametersData)
 {
 	checkf(Bindings.RootParameterBufferIndex == FShaderParameterBindings::kInvalidBufferIndex, TEXT("Can't use SetShaderParameters() for root parameter buffer index."));
 
-	FShaderParameterReader Reader(ParametersData, ParametersSize);
+	FShaderParameterReader Reader(ParametersData, ParametersMetadata->GetSize());
 
 	// Parameters
 	for (const FShaderParameterBindings::FParameter& Parameter : Bindings.Parameters)
@@ -831,37 +832,40 @@ void SetShaderParameters(
 	FRHIComputeCommandList& RHICmdList,
 	FRHIComputeShader* ShaderRHI,
 	const FShaderParameterBindings& Bindings,
-	const void* ParametersData, uint64 ParametersSize)
+	const FShaderParametersMetadata* ParametersMetadata,
+	const void* ParametersData)
 {
-	SetShaderParametersInternal(RHICmdList, ShaderRHI, Bindings, ParametersData, ParametersSize);
+	SetShaderParametersInternal(RHICmdList, ShaderRHI, Bindings, ParametersMetadata, ParametersData);
 }
 
 void SetShaderParameters(
 	FRHICommandList& RHICmdList,
 	FRHIGraphicsShader* ShaderRHI,
 	const FShaderParameterBindings& Bindings,
-	const void* ParametersData, uint64 ParametersSize)
+	const FShaderParametersMetadata* ParametersMetadata,
+	const void* ParametersData)
 {
-	SetShaderParametersInternal(RHICmdList, ShaderRHI, Bindings, ParametersData, ParametersSize);
+	SetShaderParametersInternal(RHICmdList, ShaderRHI, Bindings, ParametersMetadata, ParametersData);
 }
 
 void SetShaderParameters(
 	FRHICommandList& RHICmdList,
 	FRHIComputeShader* ShaderRHI,
 	const FShaderParameterBindings& Bindings,
-	const void* ParametersData, uint64 ParametersSize)
+	const FShaderParametersMetadata* ParametersMetadata,
+	const void* ParametersData)
 {
-	SetShaderParametersInternal(RHICmdList, ShaderRHI, Bindings, ParametersData, ParametersSize);
+	SetShaderParametersInternal(RHICmdList, ShaderRHI, Bindings, ParametersMetadata, ParametersData);
 }
 
 #if RHI_RAYTRACING
 void SetShaderParameters(
 	FRayTracingShaderBindingsWriter& RTBindingsWriter,
 	const FShaderParameterBindings& Bindings,
-	const FRHIUniformBufferLayout* RootUniformBufferLayout,
-	const void* ParametersData, uint64 ParametersSize)
+	const FShaderParametersMetadata* ParametersMetadata,
+	const void* ParametersData)
 {
-	const FShaderParameterReader Reader(ParametersData, ParametersSize);
+	const FShaderParameterReader Reader(ParametersData, ParametersMetadata->GetSize());
 
 	for (const FShaderParameterBindings::FResourceParameter& Parameter : Bindings.ResourceParameters)
 	{
@@ -947,7 +951,7 @@ void SetShaderParameters(
 		// Do not do any validation at some resources may have been removed from the structure because known to not be used by the shader.
 		EUniformBufferValidation Validation = EUniformBufferValidation::None;
 
-		RTBindingsWriter.RootUniformBuffer = RHICreateUniformBuffer(ParametersData, RootUniformBufferLayout, UniformBuffer_SingleDraw, Validation);
+		RTBindingsWriter.RootUniformBuffer = RHICreateUniformBuffer(ParametersData, ParametersMetadata->GetLayoutPtr(), UniformBuffer_SingleDraw, Validation);
 		RTBindingsWriter.SetUniformBuffer(Bindings.RootParameterBufferIndex, RTBindingsWriter.RootUniformBuffer);
 	}
 }
diff --git a/Engine/Source/Runtime/RenderCore/Public/RenderGraphUtils.h b/Engine/Source/Runtime/RenderCore/Public/RenderGraphUtils.h
--- a/Engine/Source/Runtime/RenderCore/Public/RenderGraphUtils.h
+++ b/Engine/Source/Runtime/RenderCore/Public/RenderGraphUtils.h
@@ -210,7 +210,7 @@ void ClearUnusedGraphResources(
 	// Verify the shader have all the parameters it needs. This is done before the
 	// ClearUnusedGraphResourcesImpl() to not mislead user on why some resource are missing
 	// when debugging a validation failure.
-	ValidateShaderParameters(Shader, ParametersMetadata, InoutParameters, sizeof(*InoutParameters));
+	ValidateShaderParameters(Shader, ParametersMetadata, InoutParameters);
 
 	// Clear the resources the shader won't need.
 	return ClearUnusedGraphResourcesImpl(Shader->Bindings, ParametersMetadata, InoutParameters, ExcludeList);
@@ -240,8 +240,8 @@ void ClearUnusedGraphResources(
 	// Verify the shader have all the parameters it needs. This is done before the
 	// ClearUnusedGraphResourcesImpl() to not mislead user on why some resource are missing
 	// when debugging a validation failure.
-	ValidateShaderParameters(ShaderA, ParametersMetadata, InoutParameters, sizeof(*InoutParameters));
-	ValidateShaderParameters(ShaderB, ParametersMetadata, InoutParameters, sizeof(*InoutParameters));
+	ValidateShaderParameters(ShaderA, ParametersMetadata, InoutParameters);
+	ValidateShaderParameters(ShaderB, ParametersMetadata, InoutParameters);
 
 	// Clear the resources the shader won't need.
 	const FShaderParameterBindings* ShaderBindings[] = { &ShaderA->Bindings, &ShaderB->Bindings };
diff --git a/Engine/Source/Runtime/RenderCore/Public/ShaderParameterStruct.h b/Engine/Source/Runtime/RenderCore/Public/ShaderParameterStruct.h
--- a/Engine/Source/Runtime/RenderCore/Public/ShaderParameterStruct.h
+++ b/Engine/Source/Runtime/RenderCore/Public/ShaderParameterStruct.h
@@ -117,21 +117,15 @@ extern RENDERCORE_API void EmitNullShaderParameterFatalError(const TShaderRef<FS
 
 /** Validates that all resource parameters of a shader are set. */
 #if DO_CHECK
-extern RENDERCORE_API void ValidateShaderParameters(const TShaderRef<FShader>& Shader, const FShaderParametersMetadata* ParametersMetadata, const void* ParametersData, uint64 ParametersSize);
-
-UE_DEPRECATED(5.2, "ValidateShaderParameters now requires the size of the parameters data")
-inline void ValidateShaderParameters(const TShaderRef<FShader>& Shader, const FShaderParametersMetadata* ParametersMetadata, const void* ParametersData)
-{
-	ValidateShaderParameters(Shader, ParametersMetadata, ParametersData, 0);
-}
+extern RENDERCORE_API void ValidateShaderParameters(const TShaderRef<FShader>& Shader, const FShaderParametersMetadata* ParametersMetadata, const void* ParametersData);
 #else
-FORCEINLINE void ValidateShaderParameters(const TShaderRef<FShader>& Shader, const FShaderParametersMetadata* ParametersMetadata, const void* ParametersData, uint64 ParametersSize) {}
+FORCEINLINE void ValidateShaderParameters(const TShaderRef<FShader>& Shader, const FShaderParametersMetadata* ParametersMetadata, const void* ParametersData) {}
 #endif
 
 template<typename TShaderClass>
 FORCEINLINE void ValidateShaderParameters(const TShaderRef<TShaderClass>& Shader, const typename TShaderClass::FParameters& Parameters)
 {
-	return ValidateShaderParameters(Shader, TShaderClass::FParameters::FTypeInfo::GetStructMetadata(), &Parameters, sizeof(Parameters));
+	return ValidateShaderParameters(Shader, TShaderClass::FParameters::FTypeInfo::GetStructMetadata(), &Parameters);
 }
 
 /** Unset compute shader UAVs. */
@@ -177,19 +171,22 @@ RENDERCORE_API void SetShaderParameters(
 	FRHIComputeCommandList& RHICmdList,
 	FRHIComputeShader* ShaderRHI,
 	const FShaderParameterBindings& Bindings,
-	const void* ParametersData, uint64 ParametersSize);
+	const FShaderParametersMetadata* ParametersMetadata,
+	const void* ParametersData);
 
 RENDERCORE_API void SetShaderParameters(
 	FRHICommandList& RHICmdList,
 	FRHIGraphicsShader* ShaderRHI,
 	const FShaderParameterBindings& Bindings,
-	const void* ParametersData, uint64 ParametersSize);
+	const FShaderParametersMetadata* ParametersMetadata,
+	const void* ParametersData);
 
 RENDERCORE_API void SetShaderParameters(
 	FRHICommandList& RHICmdList,
 	FRHIComputeShader* ShaderRHI,
 	const FShaderParameterBindings& Bindings,
-	const void* ParametersData, uint64 ParametersSize);
+	const FShaderParametersMetadata* ParametersMetadata,
+	const void* ParametersData);
 
 template<typename TRHICmdList, typename TShaderClass, typename TShaderRHI>
 inline void SetShaderParameters(
@@ -199,11 +196,11 @@ inline void SetShaderParameters(
 	const FShaderParametersMetadata* ParametersMetadata,
 	const typename TShaderClass::FParameters& Parameters)
 {
-	ValidateShaderParameters(Shader, ParametersMetadata, &Parameters, sizeof(Parameters));
+	ValidateShaderParameters(Shader, ParametersMetadata, &Parameters);
 
 	// TODO(RDG): Once all shader sets their parameter through this, can refactor RHI so all shader parameters get sets through a single RHI function call.
 
-	SetShaderParameters(RHICmdList, ShaderRHI, Shader->Bindings, &Parameters, sizeof(Parameters));
+	SetShaderParameters(RHICmdList, ShaderRHI, Shader->Bindings, ParametersMetadata, &Parameters);
 }
 
 template<typename TRHICmdList, typename TShaderClass, typename TShaderRHI>
@@ -218,18 +215,20 @@ inline void SetShaderParameters(TRHICmdList& RHICmdList, const TShaderRef<TShade
 RENDERCORE_API void SetShaderParameters(
 	FRayTracingShaderBindingsWriter& RTBindingsWriter,
 	const FShaderParameterBindings& Bindings,
-	const FRHIUniformBufferLayout* RootUniformBufferLayout,
-	const void* ParametersData, uint64 ParametersSize);
+	const FShaderParametersMetadata* ParametersMetadata,
+	const void* ParametersData);
 
 /** Set shader's parameters from its parameters struct. */
 template<typename TShaderClass>
 void SetShaderParameters(FRayTracingShaderBindingsWriter& RTBindingsWriter, const TShaderRef<TShaderClass>& Shader, const typename TShaderClass::FParameters& Parameters)
 {
-	ValidateShaderParameters(Shader, Parameters);
+	const FShaderParametersMetadata* ParametersMetadata = TShaderClass::FParameters::FTypeInfo::GetStructMetadata();
+
+	ValidateShaderParameters(Shader, ParametersMetadata, &Parameters);
 
 	checkf(Shader->Bindings.Parameters.Num() == 0, TEXT("Ray tracing shader should use SHADER_USE_ROOT_PARAMETER_STRUCT() to passdown the cbuffer layout to the shader compiler."));
 
-	SetShaderParameters(RTBindingsWriter, Shader->Bindings, TShaderClass::FParameters::FTypeInfo::GetStructMetadata()->GetLayoutPtr(), &Parameters, sizeof(Parameters));
+	SetShaderParameters(RTBindingsWriter, Shader->Bindings, ParametersMetadata, &Parameters);
 }
 
 #endif // RHI_RAYTRACING
