Fixed bug in RDG dependency tracking with multi-pipe skip-barrier UAVs.

diff --git a/Engine/Source/Runtime/RenderCore/Private/RenderGraphBuilder.cpp b/Engine/Source/Runtime/RenderCore/Private/RenderGraphBuilder.cpp
--- a/Engine/Source/Runtime/RenderCore/Private/RenderGraphBuilder.cpp
+++ b/Engine/Source/Runtime/RenderCore/Private/RenderGraphBuilder.cpp
@@ -918,20 +918,62 @@ void FRDGBuilder::AddCullingDependency(FRDGProducerStatesByPipeline& LastProduce
 	{
 		FRDGProducerState& LastProducer = LastProducers[LastPipeline];
 
-		if (LastProducer.Access == ERHIAccess::Unknown)
+		if (LastProducer.Access != ERHIAccess::Unknown)
 		{
-			continue;
-		}
+			FRDGPass* LastProducerPass = LastProducer.Pass;
 
-		if (FRDGProducerState::IsDependencyRequired(LastProducer, LastPipeline, NextState, NextPipeline))
-		{
-			AddPassDependency(LastProducer.Pass, NextState.Pass);
+			if (LastPipeline != NextPipeline)
+			{
+				// Only certain platforms allow multi-pipe UAV access.
+				const ERHIAccess MultiPipelineUAVMask = ERHIAccess::UAVMask & GRHIMultiPipelineMergeableAccessMask;
+
+				// If skipping a UAV barrier across pipelines, use the producer pass that will emit the correct async fence.
+				if (EnumHasAnyFlags(NextState.Access, MultiPipelineUAVMask) && SkipUAVBarrier(LastProducer.NoUAVBarrierHandle, NextState.NoUAVBarrierHandle))
+				{
+					LastProducerPass = LastProducer.PassIfSkipUAVBarrier;
+				}
+			}
+
+			if (LastProducerPass)
+			{
+				AddPassDependency(LastProducerPass, NextState.Pass);
+			}
 		}
 	}
 
 	if (IsWritableAccess(NextState.Access))
 	{
-		LastProducers[NextPipeline] = NextState;
+		FRDGProducerState& LastProducer = LastProducers[NextPipeline];
+
+		// A separate producer pass is tracked for UAV -> UAV dependencies that are skipped. Consider the following scenario:
+		//
+		//     Graphics:       A   ->    B         ->         D      ->     E       ->        G         ->            I
+		//                   (UAV)   (SkipUAV0)           (SkipUAV1)    (SkipUAV1)          (SRV)                   (UAV2)
+		//
+		// Async Compute:                           C                ->               F       ->         H
+		//                                      (SkipUAV0)                        (SkipUAV1)           (SRV)
+		//
+		// Expected Cross Pipe Dependencies: [A -> C], C -> D, [B -> F], F -> G, E -> H, F -> I. The dependencies wrapped in
+		// braces are only introduced properly by tracking a different producer for cross-pipeline skip UAV dependencies, which
+		// is only updated if skip UAV is inactive, or if transitioning from one skip UAV set to another (or another writable resource).
+
+		if (LastProducer.NoUAVBarrierHandle.IsNull())
+		{
+			if (NextState.NoUAVBarrierHandle.IsNull())
+			{
+				// Assigns the next producer when no skip UAV sets are active.
+				LastProducer.PassIfSkipUAVBarrier = NextState.Pass;
+			}
+		}
+		else if (LastProducer.NoUAVBarrierHandle != NextState.NoUAVBarrierHandle)
+		{
+			// Assigns the last producer in the prior skip UAV barrier set when moving out of a skip UAV barrier set.
+			LastProducer.PassIfSkipUAVBarrier = LastProducer.Pass;
+		}
+
+		LastProducer.Access             = NextState.Access;
+		LastProducer.Pass               = NextState.Pass;
+		LastProducer.NoUAVBarrierHandle = NextState.NoUAVBarrierHandle;
 	}
 }
 
@@ -963,19 +1005,6 @@ void FRDGBuilder::CompilePassBarriers()
 		{
 			if (!ResourceMergeState || !FRDGSubresourceState::IsMergeAllowed(ResourceType, *ResourceMergeState, PassState))
 			{
-				// Cross-pipeline, non-mergable state changes require a new pass dependency for fencing purposes.
-				if (ResourceMergeState)
-				{
-					for (ERHIPipeline Pipeline : GetRHIPipelines())
-					{
-						if (Pipeline != PassPipeline && ResourceMergeState->LastPass[Pipeline].IsValid())
-						{
-							// Add a dependency from the other pipe to this pass to join back.
-							AddPassDependency(ResourceMergeState->LastPass[Pipeline], PassHandle);
-						}
-					}
-				}
-
 				// Allocate a new pending merge state and assign it to the pass state.
 				ResourceMergeState = AllocSubresource(PassState);
 			}
@@ -1068,17 +1097,17 @@ void FRDGBuilder::Compile()
 		SetupPassQueue.Flush(TEXT("FRDGBuilder::SetupPassResources"), [this](FRDGPass* Pass) { SetupPassResources(Pass); });
 	}
 
+	if (bCullPasses)
+	{
+		CullPassStack.Reserve(CompilePassCount);
+	}
+
 	if (bCullPasses || AsyncComputePassCount > 0)
 	{
 		SCOPED_NAMED_EVENT(PassDependencies, FColor::Emerald);
 
 		if (!bParallelSetupEnabled)
 		{
-			if (bCullPasses)
-			{
-				CullPassStack.Reserve(CompilePassCount);
-			}
-
 			for (FRDGPassHandle PassHandle = ProloguePassHandle + 1; PassHandle < EpiloguePassHandle; ++PassHandle)
 			{
 				SetupPassDependencies(Passes[PassHandle]);
@@ -4057,4 +4086,4 @@ void FRDGBuilder::ForceCopyCrossGPU()
 		RHICmdList.TransferResources(Transfers);
 	}
 }
-#endif  // WITH_MGPU
\ No newline at end of file
+#endif  // WITH_MGPU
diff --git a/Engine/Source/Runtime/RenderCore/Private/RenderGraphPrivate.h b/Engine/Source/Runtime/RenderCore/Private/RenderGraphPrivate.h
--- a/Engine/Source/Runtime/RenderCore/Private/RenderGraphPrivate.h
+++ b/Engine/Source/Runtime/RenderCore/Private/RenderGraphPrivate.h
@@ -205,6 +205,17 @@ inline const TCHAR* GetEpilogueBarriersToBeginDebugName(ERHIPipeline Pipelines)
 	return TEXT("");
 }
 
+inline bool SkipUAVBarrier(FRDGViewHandle PreviousHandle, FRDGViewHandle NextHandle)
+{
+	// Barrier if previous / next don't have a matching valid skip-barrier UAV handle.
+	if (GRDGOverlapUAVs != 0 && NextHandle.IsValid() && PreviousHandle == NextHandle)
+	{
+		return true;
+	}
+
+	return false;
+}
+
 FORCEINLINE bool IsImmediateMode()
 {
 	return GRDGImmediateMode != 0;
diff --git a/Engine/Source/Runtime/RenderCore/Private/RenderGraphResources.cpp b/Engine/Source/Runtime/RenderCore/Private/RenderGraphResources.cpp
--- a/Engine/Source/Runtime/RenderCore/Private/RenderGraphResources.cpp
+++ b/Engine/Source/Runtime/RenderCore/Private/RenderGraphResources.cpp
@@ -5,20 +5,9 @@
 #include "RenderGraphPrivate.h"
 #include "RenderGraphResourcePool.h"
 
-inline bool NeedsUAVBarrier(FRDGViewHandle PreviousHandle, FRDGViewHandle NextHandle)
+inline bool SkipUAVBarrier(const FRDGSubresourceState& Previous, const FRDGSubresourceState& Next)
 {
-	// Barrier if previous / next don't have a matching valid skip-barrier UAV handle.
-	if (GRDGOverlapUAVs != 0 && NextHandle.IsValid() && PreviousHandle == NextHandle)
-	{
-		return false;
-	}
-
-	return true;
-}
-
-inline bool NeedsUAVBarrier(const FRDGSubresourceState& Previous, const FRDGSubresourceState& Next)
-{
-	return NeedsUAVBarrier(Previous.NoUAVBarrierFilter.GetUniqueHandle(), Next.NoUAVBarrierFilter.GetUniqueHandle());
+	return SkipUAVBarrier(Previous.NoUAVBarrierFilter.GetUniqueHandle(), Next.NoUAVBarrierFilter.GetUniqueHandle());
 }
 
 FRDGViewableResource::FRDGViewableResource(const TCHAR* InName, const ERDGViewableResourceType InType, bool bSkipTracking, bool bImmediateFirstBarrier)
@@ -48,40 +37,6 @@ FRDGViewableResource::FRDGViewableResource(const TCHAR* InName, const ERDGViewab
 	}
 }
 
-bool FRDGProducerState::IsDependencyRequired(FRDGProducerState LastProducer, ERHIPipeline LastPipeline, FRDGProducerState NextState, ERHIPipeline NextPipeline)
-{
-	/** This function determines whether a producer-consumer relationship exists in the graph, which is used for culling and
-	 *  async-compute fence derivation. Producers are tracked per-pipeline, so it's safe to elide a cross-pipeline producer
-	 *  for the purposes of overlapping producers, as long as a dependency exists on the same pipeline. Eliding both will
-	 *  split the producer / consumer graph into two and break culling. The only current use case this is allowing multiple
-	 *  pipes to write UAVs.
-	 *
-	 *  Producer / consumer dependencies take place independent of resource state merging / transitions, so the logic must
-	 *  be carefully aligned so that cross-pipe dependencies align with transitions.
-	 */
-
-	// The first needs to be known producers.
-	check(IsWritableAccess(LastProducer.Access));
-
-	// A dependency is always applied on the same pipe to ensure that connectivity is preserved for culling purposes.
-	if (LastPipeline == NextPipeline)
-	{
-		return true;
-	}
-
-	// Only certain platforms allow multi-pipe UAV access.
-	const ERHIAccess MultiPipelineUAVMask = ERHIAccess::UAVMask & GRHIMultiPipelineMergeableAccessMask;
-
-	// Skip the dependency if the states are used as UAV on different pipes and a UAV barrier can be skipped. This elides the async fence.
-	if (EnumHasAnyFlags(NextState.Access, MultiPipelineUAVMask) && !NeedsUAVBarrier(LastProducer.NoUAVBarrierHandle, NextState.NoUAVBarrierHandle))
-	{
-		return false;
-	}
-
-	// Everything else requires a dependency.
-	return true;
-}
-
 bool FRDGSubresourceState::IsMergeAllowed(ERDGViewableResourceType ResourceType, const FRDGSubresourceState& Previous, const FRDGSubresourceState& Next)
 {
 	/** State merging occurs during compilation and before resource transitions are collected. It serves to remove the bulk
@@ -128,7 +83,7 @@ bool FRDGSubresourceState::IsMergeAllowed(ERDGViewableResourceType ResourceType,
 	}
 
 	// Not allowed if the resource is being used as a UAV and needs a barrier.
-	if (EnumHasAnyFlags(Next.Access, ERHIAccess::UAVMask) && NeedsUAVBarrier(Previous, Next))
+	if (EnumHasAnyFlags(Next.Access, ERHIAccess::UAVMask) && !SkipUAVBarrier(Previous, Next))
 	{
 		return false;
 	}
@@ -159,7 +114,7 @@ bool FRDGSubresourceState::IsTransitionRequired(const FRDGSubresourceState& Prev
 	}
 
 	// UAV is a special case as a barrier may still be required even if the states match.
-	if (EnumHasAnyFlags(Next.Access, ERHIAccess::UAVMask) && NeedsUAVBarrier(Previous, Next))
+	if (EnumHasAnyFlags(Next.Access, ERHIAccess::UAVMask) && !SkipUAVBarrier(Previous, Next))
 	{
 		return true;
 	}
diff --git a/Engine/Source/Runtime/RenderCore/Public/RenderGraphResources.h b/Engine/Source/Runtime/RenderCore/Public/RenderGraphResources.h
--- a/Engine/Source/Runtime/RenderCore/Public/RenderGraphResources.h
+++ b/Engine/Source/Runtime/RenderCore/Public/RenderGraphResources.h
@@ -45,12 +45,10 @@ struct FRDGViewableResourceDebugData;
 /** Used for tracking pass producer / consumer edges in the graph for culling and pipe fencing. */
 struct FRDGProducerState
 {
-	/** Returns whether the next state is dependent on the last producer in the producer graph. */
-	static bool IsDependencyRequired(FRDGProducerState LastProducer, ERHIPipeline LastPipeline, FRDGProducerState NextState, ERHIPipeline NextPipeline);
-
 	FRDGProducerState() = default;
 
 	FRDGPass* Pass = nullptr;
+	FRDGPass* PassIfSkipUAVBarrier = nullptr;
 	ERHIAccess Access = ERHIAccess::Unknown;
 	FRDGViewHandle NoUAVBarrierHandle;
 };
