Bring back RHI resource deletion when shader library is closed.

- Also, add flushing the pending deletes in the attempt to make sure any leftover RHI resources from the assets that used that library (which should be already queued for deletion at this point) are deleted before we remove the backing storage.

diff --git a/Engine/Source/Runtime/RenderCore/Private/ShaderCodeLibrary.cpp b/Engine/Source/Runtime/RenderCore/Private/ShaderCodeLibrary.cpp
--- a/Engine/Source/Runtime/RenderCore/Private/ShaderCodeLibrary.cpp
+++ b/Engine/Source/Runtime/RenderCore/Private/ShaderCodeLibrary.cpp
@@ -900,26 +900,18 @@ public:
 
 	~FShaderLibraryInstance()
 	{
-		// release RHI on all of the resources
-		// TODO (UE-169261): figure out exact ownership and lifetime guarantees. Looks like there's a discrepancy between global and game libraries
-		// and global ends up owning the resources while game ones not.
-		bool bGlobalLibrary = Library->GetName() == TEXT("Global");
+		// release RHI resources on all of the resources (note: this has nothing to do with their own lifetime and refcount)
 		for (FShaderMapResource_SharedCode* Resource : Resources)
 		{
 			if (Resource)
 			{
-				// only release resources on a global library. Game libraries can be unloaded dynamically (e.g. plugins),
-				// and any resources left at this point should be reported
-				if (bGlobalLibrary)
-				{
-					BeginReleaseResource(Resource);
-				}
-				else
-				{
-					UE_LOG(LogShaderLibrary, Warning, TEXT("FShaderLibraryInstance %s still has resource (shader map index %d) in use at time of destruction!"), *Library->GetName(), Resource->ShaderMapIndex);
-				}
+				BeginReleaseResource(Resource);
 			}
 		}
+
+		// if rendering thread is active, the actual teardown may happen later, so flush the rendering commands here
+		checkf(IsInGameThread(), TEXT("Shader library closure is expected to happen only on the game thread, at the \'top\' of the pipeline"));
+		FlushRenderingCommands();	// this will also flush pending deletes
 		
 		Library->Teardown();
 		DEC_DWORD_STAT_BY(STAT_Shaders_ShaderResourceMemory, GetSizeBytes());
@@ -1205,29 +1197,35 @@ void FShaderMapResource_SharedCode::ReleasePreloadedShaderCode(int32 ShaderIndex
 
 void FShaderMapResource_SharedCode::ReleaseRHI()
 {
-	const int32 NumShaders = GetNumShaders();
-	for (int32 i = 0; i < NumShaders; ++i)
+	if (LibraryInstance)
 	{
-		const int32 LibraryShaderIndex = LibraryInstance->Library->GetShaderIndex(ShaderMapIndex, i);
-		if (HasShader(i))
+		const int32 NumShaders = GetNumShaders();
+		for (int32 i = 0; i < NumShaders; ++i)
 		{
-			LibraryInstance->ReleaseShader(LibraryShaderIndex);
-		}
-		else if (bShaderMapPreloaded)
-		{
-			// Release the preloaded memory if it was preloaded, but not created yet
-			LibraryInstance->Library->ReleasePreloadedShader(LibraryShaderIndex);
+			const int32 LibraryShaderIndex = LibraryInstance->Library->GetShaderIndex(ShaderMapIndex, i);
+			if (HasShader(i))
+			{
+				LibraryInstance->ReleaseShader(LibraryShaderIndex);
+			}
+			else if (bShaderMapPreloaded)
+			{
+				// Release the preloaded memory if it was preloaded, but not created yet
+				LibraryInstance->Library->ReleasePreloadedShader(LibraryShaderIndex);
+			}
 		}
 	}
 
 	bShaderMapPreloaded = false;
 
 	FShaderMapResource::ReleaseRHI();
+
+	// on assumption that we aren't going to get resurrected
+	LibraryInstance = nullptr;
 }
 
 bool FShaderMapResource_SharedCode::TryRelease()
 {
-	if (LibraryInstance->TryRemoveResource(this))
+	if (LibraryInstance && LibraryInstance->TryRemoveResource(this))
 	{
 		return true;
 	}
