Added counters for tracking RDG stats and CSV counters for passes / resources.

diff --git a/Engine/Source/Runtime/RenderCore/Private/RenderGraphBuilder.cpp b/Engine/Source/Runtime/RenderCore/Private/RenderGraphBuilder.cpp
--- a/Engine/Source/Runtime/RenderCore/Private/RenderGraphBuilder.cpp
+++ b/Engine/Source/Runtime/RenderCore/Private/RenderGraphBuilder.cpp
@@ -394,7 +394,29 @@ void FRDGBuilder::TickPoolElements()
 	GRDGEmitDrawEvents_RenderThread = GetEmitDrawEvents();
 #endif
 
-#if STATS
+#if RDG_STATS
+	CSV_CUSTOM_STAT(RDGCount, Passes, GRDGStatPassCount, ECsvCustomStatOp::Set);
+	CSV_CUSTOM_STAT(RDGCount, Buffers, GRDGStatBufferCount, ECsvCustomStatOp::Set);
+	CSV_CUSTOM_STAT(RDGCount, Textures, GRDGStatTextureCount, ECsvCustomStatOp::Set);
+
+	TRACE_COUNTER_SET(COUNTER_RDG_PassCount, GRDGStatPassCount);
+	TRACE_COUNTER_SET(COUNTER_RDG_PassCullCount, GRDGStatPassCullCount);
+	TRACE_COUNTER_SET(COUNTER_RDG_RenderPassMergeCount, GRDGStatRenderPassMergeCount);
+	TRACE_COUNTER_SET(COUNTER_RDG_PassDependencyCount, GRDGStatPassDependencyCount);
+	TRACE_COUNTER_SET(COUNTER_RDG_TextureCount, GRDGStatTextureCount);
+	TRACE_COUNTER_SET(COUNTER_RDG_TextureReferenceCount, GRDGStatTextureReferenceCount);
+	TRACE_COUNTER_SET(COUNTER_RDG_TextureReferenceAverage, (float)(GRDGStatTextureReferenceCount / FMath::Max((float)GRDGStatTextureCount, 1.0f)));
+	TRACE_COUNTER_SET(COUNTER_RDG_BufferCount, GRDGStatBufferCount);
+	TRACE_COUNTER_SET(COUNTER_RDG_BufferReferenceCount, GRDGStatBufferReferenceCount);
+	TRACE_COUNTER_SET(COUNTER_RDG_BufferReferenceAverage, (float)(GRDGStatBufferReferenceCount / FMath::Max((float)GRDGStatBufferCount, 1.0f)));
+	TRACE_COUNTER_SET(COUNTER_RDG_ViewCount, GRDGStatViewCount);
+	TRACE_COUNTER_SET(COUNTER_RDG_TransientTextureCount, GRDGStatTransientTextureCount);
+	TRACE_COUNTER_SET(COUNTER_RDG_TransientBufferCount, GRDGStatTransientBufferCount);
+	TRACE_COUNTER_SET(COUNTER_RDG_TransitionCount, GRDGStatTransitionCount);
+	TRACE_COUNTER_SET(COUNTER_RDG_AliasingCount, GRDGStatAliasingCount);
+	TRACE_COUNTER_SET(COUNTER_RDG_TransitionBatchCount, GRDGStatTransitionBatchCount);
+	TRACE_COUNTER_SET(COUNTER_RDG_MemoryWatermark, int64(GRDGStatMemoryWatermark));
+
 	SET_DWORD_STAT(STAT_RDG_PassCount, GRDGStatPassCount);
 	SET_DWORD_STAT(STAT_RDG_PassCullCount, GRDGStatPassCullCount);
 	SET_DWORD_STAT(STAT_RDG_RenderPassMergeCount, GRDGStatRenderPassMergeCount);
@@ -412,6 +434,7 @@ void FRDGBuilder::TickPoolElements()
 	SET_DWORD_STAT(STAT_RDG_AliasingCount, GRDGStatAliasingCount);
 	SET_DWORD_STAT(STAT_RDG_TransitionBatchCount, GRDGStatTransitionBatchCount);
 	SET_MEMORY_STAT(STAT_RDG_MemoryWatermark, int64(GRDGStatMemoryWatermark));
+
 	GRDGStatPassCount = 0;
 	GRDGStatPassCullCount = 0;
 	GRDGStatRenderPassMergeCount = 0;
@@ -891,12 +914,12 @@ void FRDGBuilder::AddPassDependency(FRDGPassHandle ProducerHandle, FRDGPassHandl
 	auto& Producers = Consumer->Producers;
 	if (Producers.Find(ProducerHandle) == INDEX_NONE)
 	{
+#if RDG_STATS
+		GRDGStatPassDependencyCount++;
+#endif
+
 		Producers.Add(ProducerHandle);
 	}
-
-#if STATS
-	GRDGStatPassDependencyCount++;
-#endif
 }
 
 void FRDGBuilder::AddPassDependency(FRDGPass* Producer, FRDGPass* Consumer)
@@ -904,12 +927,12 @@ void FRDGBuilder::AddPassDependency(FRDGPass* Producer, FRDGPass* Consumer)
 	auto& Producers = Consumer->Producers;
 	if (Producers.Find(Producer->Handle) == INDEX_NONE)
 	{
+#if RDG_STATS
+		GRDGStatPassDependencyCount++;
+#endif
+
 		Producers.Add(Producer->Handle);
 	}
-
-#if STATS
-	GRDGStatPassDependencyCount++;
-#endif
 }
 
 void FRDGBuilder::AddCullingDependency(FRDGProducerStatesByPipeline& LastProducers, const FRDGProducerState& NextState, ERHIPipeline NextPipeline)
@@ -1042,7 +1065,7 @@ void FRDGBuilder::CompilePassBarriers()
 				}
 			}
 
-		#if STATS
+		#if RDG_STATS
 			GRDGStatTextureReferenceCount += PassState.ReferenceCount;
 		#endif
 
@@ -1069,7 +1092,7 @@ void FRDGBuilder::CompilePassBarriers()
 				Buffer->State->SetPass(ERHIPipeline::Graphics, PassHandle);
 			}
 
-		#if STATS
+		#if RDG_STATS
 			GRDGStatBufferReferenceCount += PassState.ReferenceCount;
 		#endif
 
@@ -1301,7 +1324,7 @@ void FRDGBuilder::Compile()
 					SetPrologueBarrierPass(Pass, FirstPassHandle);
 				}
 
-#if STATS
+#if RDG_STATS
 				GRDGStatRenderPassMergeCount += PassesToMerge.Num();
 #endif
 			}
@@ -1875,7 +1898,7 @@ void FRDGBuilder::Execute()
 
 			if (Pass->bCulled)
 			{
-			#if STATS
+			#if RDG_STATS
 				GRDGStatPassCullCount++;
 			#endif
 
@@ -1980,7 +2003,7 @@ void FRDGBuilder::Execute()
 	IF_RDG_ENABLE_DEBUG(UserValidation.ValidateExecuteEnd());
 	IF_RDG_ENABLE_DEBUG(GRDGAllowRHIAccess = false);
 
-#if STATS
+#if RDG_STATS
 	GRDGStatBufferCount += Buffers.Num();
 	GRDGStatTextureCount += Textures.Num();
 	GRDGStatViewCount += Views.Num();
@@ -2310,7 +2333,9 @@ void FRDGBuilder::SetupPassInternals(FRDGPass* Pass)
 
 #if STATS
 	Pass->CommandListStat = CommandListStatScope;
+#endif
 
+#if RDG_STATS
 	GRDGStatPassCount++;
 #endif
 
@@ -3554,7 +3579,7 @@ void FRDGBuilder::BeginResourceRHI(FRDGPassHandle PassHandle, FRDGTextureRef Tex
 			InitialState.Access = ERHIAccess::Discard;
 			InitTextureSubresources(*Texture->State, Texture->Layout, InitialState);
 
-		#if STATS
+		#if RDG_STATS
 			GRDGStatTransientTextureCount++;
 		#endif
 		}
@@ -3638,7 +3663,7 @@ void FRDGBuilder::BeginResourceRHI(FRDGPassHandle PassHandle, FRDGBufferRef Buff
 			InitialState->SetPass(ERHIPipeline::Graphics, MinAcquirePassHandle);
 			InitialState->Access = ERHIAccess::Discard;
 
-		#if STATS
+		#if RDG_STATS
 			GRDGStatTransientBufferCount++;
 		#endif
 		}
diff --git a/Engine/Source/Runtime/RenderCore/Private/RenderGraphPass.cpp b/Engine/Source/Runtime/RenderCore/Private/RenderGraphPass.cpp
--- a/Engine/Source/Runtime/RenderCore/Private/RenderGraphPass.cpp
+++ b/Engine/Source/Runtime/RenderCore/Private/RenderGraphPass.cpp
@@ -129,7 +129,7 @@ void FRDGBarrierBatchBegin::AddTransition(FRDGViewableResource* Resource, const
 	Transitions.Add(Info);
 	bTransitionNeeded = true;
 
-#if STATS
+#if RDG_STATS
 	GRDGStatTransitionCount++;
 #endif
 
@@ -143,7 +143,7 @@ void FRDGBarrierBatchBegin::AddAlias(FRDGViewableResource* Resource, const FRHIT
 	Aliases.Add(Info);
 	bTransitionNeeded = true;
 
-#if STATS
+#if RDG_STATS
 	GRDGStatAliasingCount++;
 #endif
 
@@ -165,7 +165,7 @@ void FRDGBarrierBatchBegin::Submit(FRHIComputeCommandList& RHICmdList, ERHIPipel
 		TransitionsToBegin.Emplace(Transition);
 	}
 
-#if STATS
+#if RDG_STATS
 	GRDGStatTransitionBatchCount++;
 #endif
 }
diff --git a/Engine/Source/Runtime/RenderCore/Private/RenderGraphPrivate.cpp b/Engine/Source/Runtime/RenderCore/Private/RenderGraphPrivate.cpp
--- a/Engine/Source/Runtime/RenderCore/Private/RenderGraphPrivate.cpp
+++ b/Engine/Source/Runtime/RenderCore/Private/RenderGraphPrivate.cpp
@@ -424,7 +424,7 @@ FAutoConsoleVariableRef CVarRDGVerboseCSVStats(
 	ECVF_RenderThreadSafe);
 #endif
 
-#if STATS
+#if RDG_STATS
 int32 GRDGStatPassCount = 0;
 int32 GRDGStatPassWithParameterCount = 0;
 int32 GRDGStatPassCullCount = 0;
@@ -441,6 +441,28 @@ int32 GRDGStatTransitionCount = 0;
 int32 GRDGStatAliasingCount = 0;
 int32 GRDGStatTransitionBatchCount = 0;
 int32 GRDGStatMemoryWatermark = 0;
+#endif
+
+CSV_DEFINE_CATEGORY(RDGCount, true);
+
+TRACE_DECLARE_INT_COUNTER(COUNTER_RDG_PassCount, TEXT("RDG/PassCount"));
+TRACE_DECLARE_INT_COUNTER(COUNTER_RDG_PassWithParameterCount, TEXT("RDG/PassWithParameterCount"));
+TRACE_DECLARE_INT_COUNTER(COUNTER_RDG_PassCullCount, TEXT("RDG/PassCullCount"));
+TRACE_DECLARE_INT_COUNTER(COUNTER_RDG_RenderPassMergeCount, TEXT("RDG/RenderPassMergeCount"));
+TRACE_DECLARE_INT_COUNTER(COUNTER_RDG_PassDependencyCount, TEXT("RDG/PassDependencyCount"));
+TRACE_DECLARE_INT_COUNTER(COUNTER_RDG_TextureCount, TEXT("RDG/TextureCount"));
+TRACE_DECLARE_INT_COUNTER(COUNTER_RDG_TextureReferenceCount, TEXT("RDG/TextureReferenceCount"));
+TRACE_DECLARE_FLOAT_COUNTER(COUNTER_RDG_TextureReferenceAverage, TEXT("RDG/TextureReferenceAverage"));
+TRACE_DECLARE_INT_COUNTER(COUNTER_RDG_BufferCount, TEXT("RDG/BufferCount"));
+TRACE_DECLARE_INT_COUNTER(COUNTER_RDG_BufferReferenceCount, TEXT("RDG/BufferReferenceCount"));
+TRACE_DECLARE_FLOAT_COUNTER(COUNTER_RDG_BufferReferenceAverage, TEXT("RDG/BufferReferenceAverage"));
+TRACE_DECLARE_INT_COUNTER(COUNTER_RDG_ViewCount, TEXT("RDG/ViewCount"));
+TRACE_DECLARE_INT_COUNTER(COUNTER_RDG_TransientTextureCount, TEXT("RDG/TransientTextureCount"));
+TRACE_DECLARE_INT_COUNTER(COUNTER_RDG_TransientBufferCount, TEXT("RDG/TransientBufferCount"));
+TRACE_DECLARE_INT_COUNTER(COUNTER_RDG_TransitionCount, TEXT("RDG/TransitionCount"));
+TRACE_DECLARE_INT_COUNTER(COUNTER_RDG_AliasingCount, TEXT("RDG/AliasingCount"));
+TRACE_DECLARE_INT_COUNTER(COUNTER_RDG_TransitionBatchCount, TEXT("RDG/TransitionBatchCount"));
+TRACE_DECLARE_MEMORY_COUNTER(COUNTER_RDG_MemoryWatermark, TEXT("RDG/MemoryWatermark"));
 
 DEFINE_STAT(STAT_RDG_PassCount);
 DEFINE_STAT(STAT_RDG_PassWithParameterCount);
@@ -467,7 +489,6 @@ DEFINE_STAT(STAT_RDG_CollectBarriersTime);
 DEFINE_STAT(STAT_RDG_ClearTime);
 DEFINE_STAT(STAT_RDG_FlushRHIResources);
 DEFINE_STAT(STAT_RDG_MemoryWatermark);
-#endif
 
 #if RDG_EVENTS != RDG_EVENTS_NONE
 int32 GRDGEmitDrawEvents_RenderThread = 0;
diff --git a/Engine/Source/Runtime/RenderCore/Private/RenderGraphPrivate.h b/Engine/Source/Runtime/RenderCore/Private/RenderGraphPrivate.h
--- a/Engine/Source/Runtime/RenderCore/Private/RenderGraphPrivate.h
+++ b/Engine/Source/Runtime/RenderCore/Private/RenderGraphPrivate.h
@@ -5,6 +5,7 @@
 #include "RenderGraphDefinitions.h"
 #include "RenderGraphResources.h"
 #include "RenderTargetPool.h"
+#include "ProfilingDebugging/CountersTrace.h"
 
 DEFINE_LOG_CATEGORY_STATIC(LogRDG, Log, All);
 
@@ -135,7 +136,11 @@ extern int32 GRDGVerboseCSVStats;
 const int32 GRDGVerboseCSVStats = 0;
 #endif
 
-#if STATS
+CSV_DECLARE_CATEGORY_EXTERN(RDGCount);
+
+#define RDG_STATS STATS || COUNTERSTRACE_ENABLED
+
+#if RDG_STATS
 extern int32 GRDGStatPassCount;
 extern int32 GRDGStatPassCullCount;
 extern int32 GRDGStatRenderPassMergeCount;
@@ -151,6 +156,26 @@ extern int32 GRDGStatTransitionCount;
 extern int32 GRDGStatAliasingCount;
 extern int32 GRDGStatTransitionBatchCount;
 extern int32 GRDGStatMemoryWatermark;
+#endif
+
+TRACE_DECLARE_INT_COUNTER_EXTERN(COUNTER_RDG_PassCount);
+TRACE_DECLARE_INT_COUNTER_EXTERN(COUNTER_RDG_PassWithParameterCount);
+TRACE_DECLARE_INT_COUNTER_EXTERN(COUNTER_RDG_PassCullCount);
+TRACE_DECLARE_INT_COUNTER_EXTERN(COUNTER_RDG_RenderPassMergeCount);
+TRACE_DECLARE_INT_COUNTER_EXTERN(COUNTER_RDG_PassDependencyCount);
+TRACE_DECLARE_INT_COUNTER_EXTERN(COUNTER_RDG_TextureCount);
+TRACE_DECLARE_INT_COUNTER_EXTERN(COUNTER_RDG_TextureReferenceCount);
+TRACE_DECLARE_FLOAT_COUNTER_EXTERN(COUNTER_RDG_TextureReferenceAverage);
+TRACE_DECLARE_INT_COUNTER_EXTERN(COUNTER_RDG_BufferCount);
+TRACE_DECLARE_INT_COUNTER_EXTERN(COUNTER_RDG_BufferReferenceCount);
+TRACE_DECLARE_FLOAT_COUNTER_EXTERN(COUNTER_RDG_BufferReferenceAverage);
+TRACE_DECLARE_INT_COUNTER_EXTERN(COUNTER_RDG_ViewCount);
+TRACE_DECLARE_INT_COUNTER_EXTERN(COUNTER_RDG_TransientTextureCount);
+TRACE_DECLARE_INT_COUNTER_EXTERN(COUNTER_RDG_TransientBufferCount);
+TRACE_DECLARE_INT_COUNTER_EXTERN(COUNTER_RDG_TransitionCount);
+TRACE_DECLARE_INT_COUNTER_EXTERN(COUNTER_RDG_AliasingCount);
+TRACE_DECLARE_INT_COUNTER_EXTERN(COUNTER_RDG_TransitionBatchCount);
+TRACE_DECLARE_MEMORY_COUNTER_EXTERN(COUNTER_RDG_MemoryWatermark);
 
 DECLARE_DWORD_COUNTER_STAT_EXTERN(TEXT("Passes"), STAT_RDG_PassCount, STATGROUP_RDG, RENDERCORE_API);
 DECLARE_DWORD_COUNTER_STAT_EXTERN(TEXT("Passes With Parameters"), STAT_RDG_PassWithParameterCount, STATGROUP_RDG, RENDERCORE_API);
@@ -179,7 +204,6 @@ DECLARE_CYCLE_STAT_EXTERN(TEXT("Clear"), STAT_RDG_ClearTime, STATGROUP_RDG, REND
 DECLARE_CYCLE_STAT_EXTERN(TEXT("Flush RHI Resources"), STAT_RDG_FlushRHIResources, STATGROUP_RDG, RENDERCORE_API);
 
 DECLARE_MEMORY_STAT_EXTERN(TEXT("Builder Watermark"), STAT_RDG_MemoryWatermark, STATGROUP_RDG, RENDERCORE_API);
-#endif
 
 #if RDG_GPU_DEBUG_SCOPES
 extern int32 GRDGEvents;
