Removing uniform buffer name locks from FShaderType and FVertexFactory
- FlushShaderFileCache() used to be called to force all shader types to update their list of referenced uniform buffers, but they are only updated if they were ever used to compile in the current session. This would impact *any* shader type that had been used instead of just being limited to changed shader types.
- Moved the uniform buffer reference updating into a new function: UpdateReferencedUniformBufferNames.
- UpdateReferencedUniformBufferNames should be called before all shader recompiles but after FlushShaderFileCache does reset itself. This function requires a list of outdated types so we can reduce the shader type updates to only the types that need it.

diff --git a/Engine/Plugins/FX/Niagara/Source/NiagaraShader/Private/NiagaraShader.cpp b/Engine/Plugins/FX/Niagara/Source/NiagaraShader/Private/NiagaraShader.cpp
--- a/Engine/Plugins/FX/Niagara/Source/NiagaraShader/Private/NiagaraShader.cpp
+++ b/Engine/Plugins/FX/Niagara/Source/NiagaraShader/Private/NiagaraShader.cpp
@@ -397,8 +397,6 @@ PRAGMA_ENABLE_DEPRECATION_WARNINGS
 
 void FNiagaraShaderType::AddUniformBufferIncludesToEnvironment(FShaderCompilerEnvironment& OutEnvironment, EShaderPlatform Platform) const
 {
-	CachedUniformBufferPlatform = Platform;
-
 	UE::ShaderParameters::AddUniformBufferIncludesToEnvironment(OutEnvironment, ReferencedUniformBufferNames);
 }
 
diff --git a/Engine/Source/Runtime/Engine/Private/RendererSupport.cpp b/Engine/Source/Runtime/Engine/Private/RendererSupport.cpp
--- a/Engine/Source/Runtime/Engine/Private/RendererSupport.cpp
+++ b/Engine/Source/Runtime/Engine/Private/RendererSupport.cpp
@@ -194,6 +194,10 @@ static void RestoreReferencesToRendererModuleClasses(
 	TArray<const FShaderPipelineType*> OutdatedShaderPipelineTypes;
 	GetOutdatedShaderTypes(OutdatedShaderTypes, OutdatedShaderPipelineTypes, OutdatedFactoryTypes);
 
+#if WITH_EDITOR
+	UpdateReferencedUniformBufferNames(OutdatedShaderTypes, OutdatedFactoryTypes, OutdatedShaderPipelineTypes);
+#endif
+
 	// Recompile any missing shaders
 	UMaterialInterface::IterateOverActiveFeatureLevels([&](ERHIFeatureLevel::Type FeatureLevel) 
 	{
diff --git a/Engine/Source/Runtime/Engine/Private/ShaderCompiler/ShaderCompiler.cpp b/Engine/Source/Runtime/Engine/Private/ShaderCompiler/ShaderCompiler.cpp
--- a/Engine/Source/Runtime/Engine/Private/ShaderCompiler/ShaderCompiler.cpp
+++ b/Engine/Source/Runtime/Engine/Private/ShaderCompiler/ShaderCompiler.cpp
@@ -6706,6 +6706,8 @@ bool RecompileShaders(const TCHAR* Cmd, FOutputDevice& Ar)
 			{
 				FRecompileShadersTimer TestTimer(TEXT("RecompileShaders Changed"));
 
+				UpdateReferencedUniformBufferNames(OutdatedShaderTypes, OutdatedFactoryTypes, OutdatedShaderPipelineTypes);
+
 				UMaterialInterface::IterateOverActiveFeatureLevels([&](ERHIFeatureLevel::Type InFeatureLevel) {
 					auto ShaderPlatform = GShaderPlatformForFeatureLevel[InFeatureLevel];
 					BeginRecompileGlobalShaders(OutdatedShaderTypes, OutdatedShaderPipelineTypes, ShaderPlatform);
@@ -6810,12 +6812,11 @@ bool RecompileShaders(const TCHAR* Cmd, FOutputDevice& Ar)
 			{
 				FRecompileShadersTimer TestTimer(TEXT("RecompileShaders SingleShader"));
 				
-				TArray<const FVertexFactoryType*> FactoryTypes;
+				UpdateReferencedUniformBufferNames(ShaderTypes, {}, ShaderPipelineTypes);
 
 				UMaterialInterface::IterateOverActiveFeatureLevels([&](ERHIFeatureLevel::Type InFeatureLevel) {
 					auto ShaderPlatform = GShaderPlatformForFeatureLevel[InFeatureLevel];
 					BeginRecompileGlobalShaders(ShaderTypes, ShaderPipelineTypes, ShaderPlatform);
-					//UMaterial::UpdateMaterialShaders(ShaderTypes, ShaderPipelineTypes, FactoryTypes, ShaderPlatform);
 					FinishRecompileGlobalShaders();
 				});
 			}
@@ -7760,6 +7761,10 @@ bool RecompileChangedShadersForPlatform(const FString& PlatformName)
 	GetOutdatedShaderTypes(OutdatedShaderTypes, OutdatedShaderPipelineTypes, OutdatedFactoryTypes);
 	UE_LOG(LogShaders, Display, TEXT("We found %d out of date shader types, %d outdated pipeline types, and %d out of date VF types!"), OutdatedShaderTypes.Num(), OutdatedShaderPipelineTypes.Num(), OutdatedFactoryTypes.Num());
 
+#if WITH_EDITOR
+	UpdateReferencedUniformBufferNames(OutdatedShaderTypes, OutdatedFactoryTypes, OutdatedShaderPipelineTypes);
+#endif
+
 	for (int32 FormatIndex = 0; FormatIndex < DesiredShaderFormats.Num(); FormatIndex++)
 	{
 		// get the shader platform enum
@@ -7970,11 +7975,6 @@ void RecompileShadersForRemote(
 
 	const uint32 StartTotalShadersCompiled = GShaderCompilerStats->GetTotalShadersCompiled();
 
-	// figure out which shaders are out of date
-	TArray<const FShaderType*> OutdatedShaderTypes;
-	TArray<const FVertexFactoryType*> OutdatedFactoryTypes;
-	TArray<const FShaderPipelineType*> OutdatedShaderPipelineTypes;
-
 	// Pick up new changes to shader files
 	FlushShaderFileCache();
 
@@ -8006,11 +8006,18 @@ void RecompileShadersForRemote(
 						UE_LOG(LogShaders, Display, TEXT("\t%s..."), ShaderType->GetName());
 					}
 
+					UpdateReferencedUniformBufferNames(ShaderTypes, {}, ShaderPipelineTypes);
+
 					CompileGlobalShaderMapForRemote(ShaderTypes, ShaderPipelineTypes, ShaderPlatform, TargetPlatform, Args.GlobalShaderMap);
 				}
 				else if (Args.CommandType == ODSCRecompileCommand::Global ||
 						 Args.CommandType == ODSCRecompileCommand::Changed)
 				{
+					// figure out which shaders are out of date
+					TArray<const FShaderType*> OutdatedShaderTypes;
+					TArray<const FVertexFactoryType*> OutdatedFactoryTypes;
+					TArray<const FShaderPipelineType*> OutdatedShaderPipelineTypes;
+
 					// Explicitly get outdated types for global shaders.
 					const FGlobalShaderMap* ShaderMap = GGlobalShaderMap[ShaderPlatform];
 					if (ShaderMap)
@@ -8020,6 +8027,8 @@ void RecompileShadersForRemote(
 
 					UE_LOG(LogShaders, Display, TEXT("\tFound %d outdated shader types."), OutdatedShaderTypes.Num() + OutdatedShaderPipelineTypes.Num());
 
+					UpdateReferencedUniformBufferNames(OutdatedShaderTypes, OutdatedFactoryTypes, OutdatedShaderPipelineTypes);
+
 					CompileGlobalShaderMapForRemote(OutdatedShaderTypes, OutdatedShaderPipelineTypes, ShaderPlatform, TargetPlatform, Args.GlobalShaderMap);
 				}
 
diff --git a/Engine/Source/Runtime/RenderCore/Private/Shader.cpp b/Engine/Source/Runtime/RenderCore/Private/Shader.cpp
--- a/Engine/Source/Runtime/RenderCore/Private/Shader.cpp
+++ b/Engine/Source/Runtime/RenderCore/Private/Shader.cpp
@@ -424,6 +424,11 @@ bool FShaderType::ValidateCompiledResult(EShaderPlatform Platform, const FShader
 	return (*ValidateCompiledResultRef)(Platform, ParameterMap, OutError);
 }
 
+void FShaderType::UpdateReferencedUniformBufferNames(const TMap<FString, TArray<const TCHAR*>>& ShaderFileToUniformBufferVariables)
+{
+	ReferencedUniformBufferNames.Empty();
+	GenerateReferencedUniformBufferNames(SourceFilename, Name, ShaderFileToUniformBufferVariables, ReferencedUniformBufferNames);
+}
 #endif // WITH_EDITOR
 
 ERayTracingPayloadType FShaderType::GetRayTracingPayloadType(const int32 PermutationId) const
@@ -456,7 +461,7 @@ void FShaderType::Initialize(const TMap<FString, TArray<const TCHAR*> >& ShaderF
 #if UE_BUILD_DEBUG
 			UniqueShaderTypes.Add(Type);
 #endif
-			GenerateReferencedUniformBufferNames(Type->SourceFilename, Type->Name, ShaderFileToUniformBufferVariables, Type->ReferencedUniformBufferNames);
+			Type->UpdateReferencedUniformBufferNames(ShaderFileToUniformBufferVariables);
 		}
 	
 #if UE_BUILD_DEBUG
diff --git a/Engine/Source/Runtime/RenderCore/Private/ShaderCore.cpp b/Engine/Source/Runtime/RenderCore/Private/ShaderCore.cpp
--- a/Engine/Source/Runtime/RenderCore/Private/ShaderCore.cpp
+++ b/Engine/Source/Runtime/RenderCore/Private/ShaderCore.cpp
@@ -1540,39 +1540,45 @@ void FlushShaderFileCache()
 		GShaderFileCache.Empty();
 	}
 
+	UE_LOG(LogShaders, Log, TEXT("FlushShaderFileCache() end"));
+}
+
 #if WITH_EDITOR
+
+void UpdateReferencedUniformBufferNames(
+	TArrayView<const FShaderType*> OutdatedShaderTypes,
+	TArrayView<const FVertexFactoryType*> OutdatedFactoryTypes,
+	TArrayView<const FShaderPipelineType*> OutdatedShaderPipelineTypes)
+{
 	if (!FPlatformProperties::RequiresCookedData())
 	{
+		TRACE_CPUPROFILER_EVENT_SCOPE(UpdateReferencedUniformBufferNames);
+
 		LogShaderSourceDirectoryMappings();
 
 		TMap<FString, TArray<const TCHAR*> > ShaderFileToUniformBufferVariables;
 		BuildShaderFileToUniformBufferMap(ShaderFileToUniformBufferVariables);
 
-		for (TLinkedList<FShaderPipelineType*>::TConstIterator It(FShaderPipelineType::GetTypeList()); It; It.Next())
+		for (const FShaderPipelineType* PipelineType : OutdatedShaderPipelineTypes)
 		{
-			const auto& Stages = It->GetStages();
-			for (const FShaderType* ShaderType : Stages)
+			for (const FShaderType* ShaderType : PipelineType->GetStages())
 			{
-				((FShaderType*)ShaderType)->FlushShaderFileCache(ShaderFileToUniformBufferVariables);
+				const_cast<FShaderType*>(ShaderType)->UpdateReferencedUniformBufferNames(ShaderFileToUniformBufferVariables);
 			}
 		}
 
-		for(TLinkedList<FShaderType*>::TIterator It(FShaderType::GetTypeList()); It; It.Next())
+		for (const FShaderType* ShaderType : OutdatedShaderTypes)
 		{
-			It->FlushShaderFileCache(ShaderFileToUniformBufferVariables);
+			const_cast<FShaderType*>(ShaderType)->UpdateReferencedUniformBufferNames(ShaderFileToUniformBufferVariables);
 		}
 
-		for(TLinkedList<FVertexFactoryType*>::TIterator It(FVertexFactoryType::GetTypeList()); It; It.Next())
+		for (const FVertexFactoryType* VertexFactoryType : OutdatedFactoryTypes)
 		{
-			It->FlushShaderFileCache(ShaderFileToUniformBufferVariables);
+			const_cast<FVertexFactoryType*>(VertexFactoryType)->UpdateReferencedUniformBufferNames(ShaderFileToUniformBufferVariables);
 		}
 	}
-#endif // WITH_EDITOR
-
-	UE_LOG(LogShaders, Log, TEXT("FlushShaderFileCache() end"));
 }
 
-#if WITH_EDITOR
 void GenerateReferencedUniformBufferNames(
 	const TCHAR* SourceFilename,
 	const TCHAR* ShaderTypeName,
diff --git a/Engine/Source/Runtime/RenderCore/Private/ShaderParameters.cpp b/Engine/Source/Runtime/RenderCore/Private/ShaderParameters.cpp
--- a/Engine/Source/Runtime/RenderCore/Private/ShaderParameters.cpp
+++ b/Engine/Source/Runtime/RenderCore/Private/ShaderParameters.cpp
@@ -418,56 +418,14 @@ void UE::ShaderParameters::AddUniformBufferIncludesToEnvironment(FShaderCompiler
 	GeneratedUniformBuffersInclude.Append(UniformBufferIncludes);
 }
 
-static const uint32 NumUniformBufferLocks = 16u;
-static FRWLock UniformBufferLocks[NumUniformBufferLocks];
-
-template<typename ShaderTypename>
-FRWLock& GetUniformBufferLock(const ShaderTypename* InShaderType)
-{
-	const uint32 LockIndex = InShaderType->GetHashedName().GetHash() % NumUniformBufferLocks;
-	return UniformBufferLocks[LockIndex];
-}
-
 void FShaderType::FlushShaderFileCache(const TMap<FString, TArray<const TCHAR*> >& ShaderFileToUniformBufferVariables)
 {
-	// If we've ever been used to populate a shader compiler environment, reset and repopulate the list of uniform buffers
-	if (CachedUniformBufferPlatform != SP_NumPlatforms)
-	{
-		FWriteScopeLock Lock(GetUniformBufferLock(this));
-		if (CachedUniformBufferPlatform != SP_NumPlatforms)
-		{
-			ReferencedUniformBufferNames.Empty();
-			GenerateReferencedUniformBufferNames(SourceFilename, Name, ShaderFileToUniformBufferVariables, ReferencedUniformBufferNames);
-			CachedUniformBufferPlatform = SP_NumPlatforms;
-		}
-	}
 }
 
 void FShaderType::AddUniformBufferIncludesToEnvironment(FShaderCompilerEnvironment& OutEnvironment, EShaderPlatform Platform) const
 {
 	TRACE_CPUPROFILER_EVENT_SCOPE(FShaderType::AddReferencedUniformBufferIncludes);
 
-	TOptional<FReadScopeLock> ReadScopeLock; // TOptional is so we can unlock and lock manually; the FReadScopeLock api does not support that
-	ReadScopeLock.Emplace(GetUniformBufferLock(this));
-
-	while (CachedUniformBufferPlatform != Platform)
-	{
-		// Cache uniform buffer struct declarations referenced by this shader type's files
-
-		// Drop the ReadLock, acquire a write lock
-		//     Write the ReferencedUniformBufferStructsCache for the new Platform
-		// Drop the WriteLock, reacquire the ReadLock
-		ReadScopeLock.Reset();
-		{
-			FWriteScopeLock WriteScopeLock(GetUniformBufferLock(this));
-			if (CachedUniformBufferPlatform != Platform)
-			{
-				CachedUniformBufferPlatform = Platform;
-			}
-		}
-		ReadScopeLock.Emplace(GetUniformBufferLock(this));
-	}
-
 	UE::ShaderParameters::AddUniformBufferIncludesToEnvironment(OutEnvironment, ReferencedUniformBufferNames);
 }
 
@@ -536,43 +494,10 @@ void FShaderType::GetShaderStableKeyParts(FStableShaderKeyAndValue& SaveKeyVal)
 
 void FVertexFactoryType::FlushShaderFileCache(const TMap<FString, TArray<const TCHAR*> >& ShaderFileToUniformBufferVariables)
 {
-	// If we've ever been used to populate a shader compiler environment, reset and repopulate the list of uniform buffers
-	if (CachedUniformBufferPlatform != SP_NumPlatforms)
-	{
-		FWriteScopeLock Lock(GetUniformBufferLock(this));
-		if (CachedUniformBufferPlatform != SP_NumPlatforms)
-		{
-			ReferencedUniformBufferNames.Empty();
-			GenerateReferencedUniformBufferNames(ShaderFilename, Name, ShaderFileToUniformBufferVariables, ReferencedUniformBufferNames);
-			CachedUniformBufferPlatform = SP_NumPlatforms;
-		}
-	}
 }
-
 void FVertexFactoryType::AddUniformBufferIncludesToEnvironment(FShaderCompilerEnvironment& OutEnvironment, EShaderPlatform Platform) const
 {
-	TOptional<FReadScopeLock> ReadScopeLock; // TOptional is so we can unlock and lock manually; the FReadScopeLock api does not support that
-	ReadScopeLock.Emplace(GetUniformBufferLock(this));
-	while (CachedUniformBufferPlatform != Platform)
-	{
-		// Cache uniform buffer struct declarations referenced by this shader type's files
-
-		// Drop the ReadLock, acquire a write lock
-		//     Write the ReferencedUniformBufferStructsCache for the new Platform
-		// Drop the WriteLock, reacquire the ReadLock
-		ReadScopeLock.Reset();
-		{
-			FWriteScopeLock WriteScopeLock(GetUniformBufferLock(this));
-			if (CachedUniformBufferPlatform != Platform)
-			{
-				CachedUniformBufferPlatform = Platform;
-			}
-		}
-		ReadScopeLock.Emplace(GetUniformBufferLock(this));
-	}
-
 	UE::ShaderParameters::AddUniformBufferIncludesToEnvironment(OutEnvironment, ReferencedUniformBufferNames);
 }
 
-
 #endif // WITH_EDITOR
diff --git a/Engine/Source/Runtime/RenderCore/Private/VertexFactory.cpp b/Engine/Source/Runtime/RenderCore/Private/VertexFactory.cpp
--- a/Engine/Source/Runtime/RenderCore/Private/VertexFactory.cpp
+++ b/Engine/Source/Runtime/RenderCore/Private/VertexFactory.cpp
@@ -52,6 +52,14 @@ FVertexFactoryType* FVertexFactoryType::GetVFByName(const FHashedName& VFName)
 	return Type ? *Type : nullptr;
 }
 
+#if WITH_EDITOR
+void FVertexFactoryType::UpdateReferencedUniformBufferNames(const TMap<FString, TArray<const TCHAR*>>& ShaderFileToUniformBufferVariables)
+{
+	ReferencedUniformBufferNames.Empty();
+	GenerateReferencedUniformBufferNames(ShaderFilename, Name, ShaderFileToUniformBufferVariables, ReferencedUniformBufferNames);
+}
+#endif
+
 void FVertexFactoryType::Initialize(const TMap<FString, TArray<const TCHAR*> >& ShaderFileToUniformBufferVariables)
 {
 #if WITH_EDITOR
@@ -62,7 +70,7 @@ void FVertexFactoryType::Initialize(const TMap<FString, TArray<const TCHAR*> >&
 		for(TLinkedList<FVertexFactoryType*>::TIterator It(FVertexFactoryType::GetTypeList()); It; It.Next())
 		{
 			FVertexFactoryType* Type = *It;
-			GenerateReferencedUniformBufferNames(Type->ShaderFilename, Type->Name, ShaderFileToUniformBufferVariables, Type->ReferencedUniformBufferNames);
+			Type->UpdateReferencedUniformBufferNames(ShaderFileToUniformBufferVariables);
 		}
 	}
 #endif // WITH_EDITOR
diff --git a/Engine/Source/Runtime/RenderCore/Public/Shader.h b/Engine/Source/Runtime/RenderCore/Public/Shader.h
--- a/Engine/Source/Runtime/RenderCore/Public/Shader.h
+++ b/Engine/Source/Runtime/RenderCore/Public/Shader.h
@@ -1418,6 +1418,9 @@ public:
 		AddUniformBufferIncludesToEnvironment(OutEnvironment, Platform);
 	}
 
+	void UpdateReferencedUniformBufferNames(const TMap<FString, TArray<const TCHAR*>>& ShaderFileToUniformBufferVariables);
+
+	UE_DEPRECATED(5.2, "FlushShaderFileCache is deprecated. UpdateReferencedUniformBufferNames should be used to flush any uniform buffer changes")
 	void FlushShaderFileCache(const TMap<FString, TArray<const TCHAR*> >& ShaderFileToUniformBufferVariables);
 
 	void GetShaderStableKeyParts(struct FStableShaderKeyAndValue& SaveKeyVal);
@@ -1457,9 +1460,6 @@ private:
 
 #if WITH_EDITOR
 protected:
-	/** Tracks the last platform ReferencedUniformBufferNames was read from in AddUniformBufferIncludesToEnvironment. */
-	mutable std::atomic<EShaderPlatform> CachedUniformBufferPlatform{ SP_NumPlatforms };
-
 	/**
 	* Cache of referenced uniform buffer includes.
 	* These are derived from source files so they need to be flushed when editing and recompiling shaders on the fly.
diff --git a/Engine/Source/Runtime/RenderCore/Public/ShaderCore.h b/Engine/Source/Runtime/RenderCore/Public/ShaderCore.h
--- a/Engine/Source/Runtime/RenderCore/Public/ShaderCore.h
+++ b/Engine/Source/Runtime/RenderCore/Public/ShaderCore.h
@@ -1090,6 +1090,16 @@ extern RENDERCORE_API void VerifyShaderSourceFiles(EShaderPlatform ShaderPlatfor
 
 #if WITH_EDITOR
 
+class FShaderType;
+class FVertexFactoryType;
+class FShaderPipelineType;
+
+/** Force updates each shader/pipeline type provided to update their list of referenced uniform buffers. */
+RENDERCORE_API void UpdateReferencedUniformBufferNames(
+	TArrayView<const FShaderType*> OutdatedShaderTypes,
+	TArrayView<const FVertexFactoryType*> OutdatedFactoryTypes,
+	TArrayView<const FShaderPipelineType*> OutdatedShaderPipelineTypes);
+
 struct FCachedUniformBufferDeclaration
 {
 	// Using SharedPtr so we can hand off lifetime ownership to FShaderCompilerEnvironment::IncludeVirtualPathToExternalContentsMap when invalidating this cache
diff --git a/Engine/Source/Runtime/RenderCore/Public/VertexFactory.h b/Engine/Source/Runtime/RenderCore/Public/VertexFactory.h
--- a/Engine/Source/Runtime/RenderCore/Public/VertexFactory.h
+++ b/Engine/Source/Runtime/RenderCore/Public/VertexFactory.h
@@ -433,6 +433,9 @@ public:
 		AddUniformBufferIncludesToEnvironment(OutEnvironment, Platform);
 	}
 
+	RENDERCORE_API void UpdateReferencedUniformBufferNames(const TMap<FString, TArray<const TCHAR*>>& ShaderFileToUniformBufferVariables);
+
+	UE_DEPRECATED(5.2, "FlushShaderFileCache is deprecated. UpdateReferencedUniformBufferNames should be used to flush any uniform buffer changes")
 	RENDERCORE_API void FlushShaderFileCache(const TMap<FString, TArray<const TCHAR*> >& ShaderFileToUniformBufferVariables);
 
 	inline const TSet<const TCHAR*>& GetReferencedUniformBufferNames() const { return ReferencedUniformBufferNames; };
@@ -468,9 +471,6 @@ private:
 	 * FShaderType::Initialize will add the referenced uniform buffers, but this set may be updated by FlushShaderFileCache
 	 */
 	TSet<const TCHAR*> ReferencedUniformBufferNames;
-
-	/** Tracks the last platform ReferencedUniformBufferNames was read from in AddUniformBufferIncludesToEnvironment. */
-	mutable std::atomic<EShaderPlatform> CachedUniformBufferPlatform{ SP_NumPlatforms };
 #endif // WITH_EDITOR
 };
 
