Add CVar r.SceneDepthHZBAsyncCompute to run HZB generation on async compute

* This avoids GPU under-utilization in scenes where a very large number of individual occlusion queries is issued
* Possible r.SceneDepthHZBAsyncCompute values:
  * 0: Don't use async compute (default)
  * 1: Use async compute, start as soon as possible
  * 2: Use async compute, start after ComputeLightGrid.CompactLinks pass

diff --git a/Engine/Source/Runtime/RenderCore/Public/RenderGraphUtils.h b/Engine/Source/Runtime/RenderCore/Public/RenderGraphUtils.h
--- a/Engine/Source/Runtime/RenderCore/Public/RenderGraphUtils.h
+++ b/Engine/Source/Runtime/RenderCore/Public/RenderGraphUtils.h
@@ -489,7 +489,7 @@ struct RENDERCORE_API FComputeShaderUtils
 
 	/** Dispatch a compute shader to render graph builder with its parameters. */
 	template<typename TShaderClass>
-	static void AddPass(
+	static FRDGPassRef AddPass(
 		FRDGBuilder& GraphBuilder,
 		FRDGEventName&& PassName,
 		ERDGPassFlags PassFlags,
@@ -505,7 +505,7 @@ struct RENDERCORE_API FComputeShaderUtils
 		ValidateGroupCount(GroupCount);
 		ClearUnusedGraphResources(ComputeShader, ParametersMetadata, Parameters);
 
-		GraphBuilder.AddPass(
+		return GraphBuilder.AddPass(
 			Forward<FRDGEventName>(PassName),
 			ParametersMetadata,
 			Parameters,
@@ -520,7 +520,7 @@ struct RENDERCORE_API FComputeShaderUtils
 	 *  This allows adding a dispatch with unknown GroupCount but the value must be ready before the pass is executed.
 	 */
 	template<typename TShaderClass>
-	static void AddPass(
+	static FRDGPassRef AddPass(
 		FRDGBuilder& GraphBuilder,
 		FRDGEventName&& PassName,
 		ERDGPassFlags PassFlags,
@@ -535,7 +535,7 @@ struct RENDERCORE_API FComputeShaderUtils
 
 		ClearUnusedGraphResources(ComputeShader, ParametersMetadata, Parameters);
 
-		GraphBuilder.AddPass(
+		return GraphBuilder.AddPass(
 			Forward<FRDGEventName>(PassName),
 			ParametersMetadata,
 			Parameters,
@@ -552,7 +552,7 @@ struct RENDERCORE_API FComputeShaderUtils
 	}
 
 	template<typename TShaderClass>
-	static void AddPass(
+	static FRDGPassRef AddPass(
 		FRDGBuilder& GraphBuilder,
 		FRDGEventName&& PassName,
 		ERDGPassFlags PassFlags,
@@ -561,11 +561,11 @@ struct RENDERCORE_API FComputeShaderUtils
 		FIntVector GroupCount)
 	{
 		const FShaderParametersMetadata* ParametersMetadata = TShaderClass::FParameters::FTypeInfo::GetStructMetadata();
-		AddPass(GraphBuilder, Forward<FRDGEventName>(PassName), PassFlags, ComputeShader, ParametersMetadata, Parameters, GroupCount);
+		return AddPass(GraphBuilder, Forward<FRDGEventName>(PassName), PassFlags, ComputeShader, ParametersMetadata, Parameters, GroupCount);
 	}
 
 	template <typename TShaderClass>
-	static FORCEINLINE void AddPass(
+	static FORCEINLINE FRDGPassRef AddPass(
 		FRDGBuilder& GraphBuilder,
 		FRDGEventName&& PassName,
 		const TShaderRef<TShaderClass>& ComputeShader,
@@ -573,11 +573,11 @@ struct RENDERCORE_API FComputeShaderUtils
 		FIntVector GroupCount)
 	{
 		const FShaderParametersMetadata* ParametersMetadata = TShaderClass::FParameters::FTypeInfo::GetStructMetadata();
-		AddPass(GraphBuilder, Forward<FRDGEventName>(PassName), ERDGPassFlags::Compute, ComputeShader, ParametersMetadata, Parameters, GroupCount);
+		return AddPass(GraphBuilder, Forward<FRDGEventName>(PassName), ERDGPassFlags::Compute, ComputeShader, ParametersMetadata, Parameters, GroupCount);
 	}
 
 	template <typename TShaderClass>
-	static FORCEINLINE void AddPass(
+	static FORCEINLINE FRDGPassRef AddPass(
 		FRDGBuilder& GraphBuilder,
 		FRDGEventName&& PassName,
 		const TShaderRef<TShaderClass>& ComputeShader,
@@ -585,12 +585,12 @@ struct RENDERCORE_API FComputeShaderUtils
 		FRDGDispatchGroupCountCallback&& GroupCountCallback)
 	{
 		const FShaderParametersMetadata* ParametersMetadata = TShaderClass::FParameters::FTypeInfo::GetStructMetadata();
-		AddPass(GraphBuilder, Forward<FRDGEventName>(PassName), ERDGPassFlags::Compute, ComputeShader, ParametersMetadata, Parameters, MoveTemp(GroupCountCallback));
+		return AddPass(GraphBuilder, Forward<FRDGEventName>(PassName), ERDGPassFlags::Compute, ComputeShader, ParametersMetadata, Parameters, MoveTemp(GroupCountCallback));
 	}
 
 	/** Dispatch a compute shader to render graph builder with its parameters. */
 	template<typename TShaderClass>
-	static void AddPass(
+	static FRDGPassRef AddPass(
 		FRDGBuilder& GraphBuilder,
 		FRDGEventName&& PassName,
 		ERDGPassFlags PassFlags,
@@ -605,7 +605,7 @@ struct RENDERCORE_API FComputeShaderUtils
 		ValidateIndirectArgsBuffer(IndirectArgsBuffer, IndirectArgsOffset);
 		ClearUnusedGraphResources(ComputeShader, Parameters, { IndirectArgsBuffer });
 
-		GraphBuilder.AddPass(
+		return GraphBuilder.AddPass(
 			Forward<FRDGEventName>(PassName),
 			Parameters,
 			PassFlags,
@@ -620,7 +620,7 @@ struct RENDERCORE_API FComputeShaderUtils
 	}
 
 	template<typename TShaderClass>
-	static FORCEINLINE void AddPass(
+	static FORCEINLINE FRDGPassRef AddPass(
 		FRDGBuilder& GraphBuilder,
 		FRDGEventName&& PassName,
 		const TShaderRef<TShaderClass>& ComputeShader,
@@ -628,7 +628,7 @@ struct RENDERCORE_API FComputeShaderUtils
 		FRDGBufferRef IndirectArgsBuffer,
 		uint32 IndirectArgsOffset)
 	{
-		AddPass(GraphBuilder, Forward<FRDGEventName>(PassName), ERDGPassFlags::Compute, ComputeShader, Parameters, IndirectArgsBuffer, IndirectArgsOffset);
+		return AddPass(GraphBuilder, Forward<FRDGEventName>(PassName), ERDGPassFlags::Compute, ComputeShader, Parameters, IndirectArgsBuffer, IndirectArgsOffset);
 	}
 
 	static void ClearUAV(FRDGBuilder& GraphBuilder, FGlobalShaderMap* ShaderMap, FRDGBufferUAVRef UAV, uint32 ClearValue);
diff --git a/Engine/Source/Runtime/Renderer/Private/DeferredShadingRenderer.cpp b/Engine/Source/Runtime/Renderer/Private/DeferredShadingRenderer.cpp
--- a/Engine/Source/Runtime/Renderer/Private/DeferredShadingRenderer.cpp
+++ b/Engine/Source/Runtime/Renderer/Private/DeferredShadingRenderer.cpp
@@ -257,6 +257,14 @@ static FAutoConsoleVariableRef CVarRayTracingDebugForceOpaque(
 	TEXT("Forces all ray tracing geometry instances to be opaque, effectively disabling any-hit shaders. This is useful for debugging and profiling. (default = 0)")
 );
 
+static TAutoConsoleVariable<int32> CVarSceneDepthHZBAsyncCompute(
+	TEXT("r.SceneDepthHZBAsyncCompute"), 0,
+	TEXT("Selects whether HZB for scene depth buffer should be built with async compute.\n")
+	TEXT(" 0: Don't use async compute (default)\n")
+	TEXT(" 1: Use async compute, start as soon as possible\n")
+	TEXT(" 2: Use async compute, start after ComputeLightGrid.CompactLinks pass"),
+	ECVF_RenderThreadSafe);
+
 #if RHI_RAYTRACING
 
 static bool bRefreshRayTracingInstances = false;
@@ -510,7 +518,7 @@ bool FDeferredShadingSceneRenderer::ShouldRenderPrePass() const
 	return (DepthPass.EarlyZPassMode != DDM_None || DepthPass.bEarlyZPassMovable != 0);
 }
 
-bool FDeferredShadingSceneRenderer::RenderHzb(FRDGBuilder& GraphBuilder, FRDGTextureRef SceneDepthTexture)
+bool FDeferredShadingSceneRenderer::RenderHzb(FRDGBuilder& GraphBuilder, FRDGTextureRef SceneDepthTexture, const FBuildHZBAsyncComputeParams* AsyncComputeParams)
 {
 	RDG_GPU_STAT_SCOPE(GraphBuilder, HZB);
 
@@ -541,7 +549,9 @@ bool FDeferredShadingSceneRenderer::RenderHzb(FRDGBuilder& GraphBuilder, FRDGTex
 				TEXT("HZBClosest"),
 				/* OutClosestHZBTexture = */ ViewPipelineState.bClosestHZB ? &ClosestHZBTexture : nullptr,
 				TEXT("HZBFurthest"),
-				/* OutFurthestHZBTexture = */ &FurthestHZBTexture);
+				/* OutFurthestHZBTexture = */ &FurthestHZBTexture,
+				BuildHZBDefaultPixelFormat,
+				AsyncComputeParams);
 
 			// Update the view.
 			{
@@ -2963,12 +2973,14 @@ void FDeferredShadingSceneRenderer::Render(FRDGBuilder& GraphBuilder)
 
 	AddResolveSceneDepthPass(GraphBuilder, Views, SceneTextures.Depth);
 
+	FComputeLightGridOutput ComputeLightGridOutput = {};
+
 	// NOTE: The ordering of the lights is used to select sub-sets for different purposes, e.g., those that support clustered deferred.
 	FSortedLightSetSceneInfo& SortedLightSet = *GraphBuilder.AllocObject<FSortedLightSetSceneInfo>();
 	{
 		RDG_CSV_STAT_EXCLUSIVE_SCOPE(GraphBuilder, SortLights);
 		RDG_GPU_STAT_SCOPE(GraphBuilder, SortLights);
-		GatherLightsAndComputeLightGrid(GraphBuilder, bComputeLightGrid, SortedLightSet);
+		ComputeLightGridOutput = GatherLightsAndComputeLightGrid(GraphBuilder, bComputeLightGrid, SortedLightSet);
 	}
 
 	CSV_CUSTOM_STAT(LightCount, All,  float(SortedLightSet.SortedLights.Num()), ECsvCustomStatOp::Set);
@@ -2996,7 +3008,17 @@ void FDeferredShadingSceneRenderer::Render(FRDGBuilder& GraphBuilder)
 
 	const auto RenderOcclusionLambda = [&]()
 	{
-		RenderOcclusion(GraphBuilder, SceneTextures, bIsOcclusionTesting);
+		const int32 AsyncComputeMode = CVarSceneDepthHZBAsyncCompute.GetValueOnRenderThread();
+		bool bAsyncCompute = AsyncComputeMode != 0;
+
+		FBuildHZBAsyncComputeParams AsyncComputeParams = {};
+		if (AsyncComputeMode == 2)
+		{
+			AsyncComputeParams.Prerequisite = ComputeLightGridOutput.CompactLinksPass;
+		}
+
+		RenderOcclusion(GraphBuilder, SceneTextures, bIsOcclusionTesting, 
+			bAsyncCompute ? &AsyncComputeParams : nullptr);
 
 		CompositionLighting.ProcessAfterOcclusion(GraphBuilder);
 	};
diff --git a/Engine/Source/Runtime/Renderer/Private/DeferredShadingRenderer.h b/Engine/Source/Runtime/Renderer/Private/DeferredShadingRenderer.h
--- a/Engine/Source/Runtime/Renderer/Private/DeferredShadingRenderer.h
+++ b/Engine/Source/Runtime/Renderer/Private/DeferredShadingRenderer.h
@@ -49,6 +49,7 @@ struct FVolumetricFogLocalLightFunctionInfo;
 struct FTranslucencyLightingVolumeTextures;
 struct FLumenSceneFrameTemporaries;
 struct FSingleLayerWaterPrePassResult;
+struct FBuildHZBAsyncComputeParams;
 
 /**   
  * Data for rendering meshes into Lumen Lighting Cards.
@@ -331,7 +332,7 @@ public:
 	 * Culls local lights and reflection probes to a grid in frustum space, builds one light list and grid per view in the current Views.  
 	 * Needed for forward shading or translucency using the Surface lighting mode, and clustered deferred shading. 
 	 */
-	void GatherLightsAndComputeLightGrid(FRDGBuilder& GraphBuilder, bool bNeedLightGrid, FSortedLightSetSceneInfo &SortedLightSet);
+	FComputeLightGridOutput GatherLightsAndComputeLightGrid(FRDGBuilder& GraphBuilder, bool bNeedLightGrid, FSortedLightSetSceneInfo &SortedLightSet);
 
 	/** 
 	 * Debug light grid content on screen.
@@ -396,9 +397,10 @@ public:
 	void RenderOcclusion(
 		FRDGBuilder& GraphBuilder,
 		const FSceneTextures& SceneTextures,
-		bool bIsOcclusionTesting);
+		bool bIsOcclusionTesting,
+		const FBuildHZBAsyncComputeParams* BuildHZBAsyncComputeParams = nullptr);
 
-	bool RenderHzb(FRDGBuilder& GraphBuilder, FRDGTextureRef SceneDepthTexture);
+	bool RenderHzb(FRDGBuilder& GraphBuilder, FRDGTextureRef SceneDepthTexture, const FBuildHZBAsyncComputeParams* AsyncComputeParams);
 
 	/** Renders the view family. */
 	virtual void Render(FRDGBuilder& GraphBuilder) override;
diff --git a/Engine/Source/Runtime/Renderer/Private/LightGridInjection.cpp b/Engine/Source/Runtime/Renderer/Private/LightGridInjection.cpp
--- a/Engine/Source/Runtime/Renderer/Private/LightGridInjection.cpp
+++ b/Engine/Source/Runtime/Renderer/Private/LightGridInjection.cpp
@@ -373,8 +373,10 @@ static void PackLocalLightData(
 	checkSlow(int32(Out.RectDataAndVirtualShadowMapId.Z) == VirtualShadowMapId);
 }
 
-void FSceneRenderer::ComputeLightGrid(FRDGBuilder& GraphBuilder, bool bCullLightsToGrid, FSortedLightSetSceneInfo& SortedLightSet)
+FComputeLightGridOutput FSceneRenderer::ComputeLightGrid(FRDGBuilder& GraphBuilder, bool bCullLightsToGrid, FSortedLightSetSceneInfo& SortedLightSet)
 {
+	FComputeLightGridOutput Result = {};
+
 	RDG_CSV_STAT_EXCLUSIVE_SCOPE(GraphBuilder, ComputeLightGrid);
 	QUICK_SCOPE_CYCLE_COUNTER(STAT_ComputeLightGrid);
 	RDG_EVENT_SCOPE(GraphBuilder, "ComputeLightGrid");
@@ -816,7 +818,7 @@ void FSceneRenderer::ComputeLightGrid(FRDGBuilder& GraphBuilder, bool bCullLight
 					PassParametersCompact->NumLocalLights = ForwardLightData->NumLocalLights;
 					PassParametersCompact->NumGridCells = ForwardLightData->NumGridCells;
 
-					FComputeShaderUtils::AddPass(GraphBuilder, RDG_EVENT_NAME("CompactLinks"), ComputeShaderCompact, PassParametersCompact, NumGroups);
+					Result.CompactLinksPass = FComputeShaderUtils::AddPass(GraphBuilder, RDG_EVENT_NAME("CompactLinks"), ComputeShaderCompact, PassParametersCompact, NumGroups);
 				}
 			}
 			else
@@ -841,10 +843,14 @@ void FSceneRenderer::ComputeLightGrid(FRDGBuilder& GraphBuilder, bool bCullLight
 		});
 	}
 #endif
+
+	return Result;
 }
 
-void FDeferredShadingSceneRenderer::GatherLightsAndComputeLightGrid(FRDGBuilder& GraphBuilder, bool bNeedLightGrid, FSortedLightSetSceneInfo& SortedLightSet)
+FComputeLightGridOutput FDeferredShadingSceneRenderer::GatherLightsAndComputeLightGrid(FRDGBuilder& GraphBuilder, bool bNeedLightGrid, FSortedLightSetSceneInfo& SortedLightSet)
 {
+	FComputeLightGridOutput Result = {};
+
 	bool bShadowedLightsInClustered = ShouldUseClusteredDeferredShading()
 		&& CVarVirtualShadowOnePassProjection.GetValueOnRenderThread()
 		&& VirtualShadowMapArray.IsEnabled();
@@ -860,7 +866,7 @@ void FDeferredShadingSceneRenderer::GatherLightsAndComputeLightGrid(FRDGBuilder&
 		{
 			View.ForwardLightingResources.SetUniformBuffer(ForwardLightUniformBuffer);
 		}
-		return;
+		return Result;
 	}
 
 	bool bAnyViewUsesForwardLighting = false;
@@ -879,7 +885,9 @@ void FDeferredShadingSceneRenderer::GatherLightsAndComputeLightGrid(FRDGBuilder&
 	// Store this flag if lights are injected in the grids, check with 'AreLightsInLightGrid()'
 	bAreLightsInLightGrid = bCullLightsToGrid;
 	
-	ComputeLightGrid(GraphBuilder, bCullLightsToGrid, SortedLightSet);
+	Result = ComputeLightGrid(GraphBuilder, bCullLightsToGrid, SortedLightSet);
+
+	return Result;
 }
 
 void FDeferredShadingSceneRenderer::RenderForwardShadowProjections(
diff --git a/Engine/Source/Runtime/Renderer/Private/SceneOcclusion.cpp b/Engine/Source/Runtime/Renderer/Private/SceneOcclusion.cpp
--- a/Engine/Source/Runtime/Renderer/Private/SceneOcclusion.cpp
+++ b/Engine/Source/Runtime/Renderer/Private/SceneOcclusion.cpp
@@ -81,6 +81,15 @@ static FAutoConsoleVariableRef CVarDownsampledOcclusionQueries(
 	ECVF_RenderThreadSafe
 );
 
+static int32 GOcclusionQueryDispatchOrder = 0;
+static FAutoConsoleVariableRef CVarOcclusionQueryDispatchOrder(
+	TEXT("r.OcclusionQueryDispatchOrder"),
+	GOcclusionQueryDispatchOrder,
+	TEXT("0: Grouped queries before individual (default)\n")
+	TEXT("1: Individual queries before grouped"),
+	ECVF_RenderThreadSafe
+);
+
 bool UseDownsampledOcclusionQueries()
 {
 	return GDownsampledOcclusionQueries != 0;
@@ -1331,13 +1340,27 @@ static void BeginOcclusionTests(
 		{
 			VertexShader->SetParameters(RHICmdList, View);
 
+			if (GOcclusionQueryDispatchOrder == 0)
 			{
-				SCOPED_DRAW_EVENT(RHICmdList, GroupedQueries);
-				View.GroupedOcclusionQueries.Flush(RHICmdList);
+				{
+					SCOPED_DRAW_EVENT(RHICmdList, GroupedQueries);
+					View.GroupedOcclusionQueries.Flush(RHICmdList);
+				}
+				{
+					SCOPED_DRAW_EVENT(RHICmdList, IndividualQueries);
+					View.IndividualOcclusionQueries.Flush(RHICmdList);
+				}
 			}
+			else
 			{
-				SCOPED_DRAW_EVENT(RHICmdList, IndividualQueries);
-				View.IndividualOcclusionQueries.Flush(RHICmdList);
+				{
+					SCOPED_DRAW_EVENT(RHICmdList, IndividualQueries);
+					View.IndividualOcclusionQueries.Flush(RHICmdList);
+				}
+				{
+					SCOPED_DRAW_EVENT(RHICmdList, GroupedQueries);
+					View.GroupedOcclusionQueries.Flush(RHICmdList);
+				}
 			}
 		}
 	}
@@ -1346,7 +1369,8 @@ static void BeginOcclusionTests(
 void FDeferredShadingSceneRenderer::RenderOcclusion(
 	FRDGBuilder& GraphBuilder,
 	const FSceneTextures& SceneTextures,
-	bool bIsOcclusionTesting)
+	bool bIsOcclusionTesting,
+	const FBuildHZBAsyncComputeParams* BuildHZBAsyncComputeParams)
 {
 	TRACE_CPUPROFILER_EVENT_SCOPE(FDeferredShadingSceneRenderer::RenderOcclusion);
 
@@ -1420,7 +1444,7 @@ void FDeferredShadingSceneRenderer::RenderOcclusion(
 		}
 	}
 
-	const bool bUseHzbOcclusion = RenderHzb(GraphBuilder, SceneTextures.Depth.Resolve);
+	const bool bUseHzbOcclusion = RenderHzb(GraphBuilder, SceneTextures.Depth.Resolve, BuildHZBAsyncComputeParams);
 
 	if (bUseHzbOcclusion || bIsOcclusionTesting)
 	{
diff --git a/Engine/Source/Runtime/Renderer/Private/SceneRendering.h b/Engine/Source/Runtime/Renderer/Private/SceneRendering.h
--- a/Engine/Source/Runtime/Renderer/Private/SceneRendering.h
+++ b/Engine/Source/Runtime/Renderer/Private/SceneRendering.h
@@ -1936,6 +1936,11 @@ private:
 	FSceneTextures SceneTextures;
 };
 
+struct FComputeLightGridOutput
+{
+	FRDGPassRef CompactLinksPass = {};
+};
+
 /**
  * Used as the scope for scene rendering functions.
  * It is initialized in the game thread by FSceneViewFamily::BeginRender, and then passed to the rendering thread.
@@ -2369,7 +2374,7 @@ protected:
 	 * Culls local lights and reflection probes to a grid in frustum space, builds one light list and grid per view in the current Views.  
 	 * Needed for forward shading or translucency using the Surface lighting mode, and clustered deferred shading. 
 	 */
-	void ComputeLightGrid(FRDGBuilder& GraphBuilder, bool bCullLightsToGrid, FSortedLightSetSceneInfo& SortedLightSet);
+	FComputeLightGridOutput ComputeLightGrid(FRDGBuilder& GraphBuilder, bool bCullLightsToGrid, FSortedLightSetSceneInfo& SortedLightSet);
 
 	/**
 	* Used by RenderLights to figure out if light functions need to be rendered to the attenuation buffer.
diff --git a/Engine/Source/Runtime/Renderer/Private/SceneTextureReductions.cpp b/Engine/Source/Runtime/Renderer/Private/SceneTextureReductions.cpp
--- a/Engine/Source/Runtime/Renderer/Private/SceneTextureReductions.cpp
+++ b/Engine/Source/Runtime/Renderer/Private/SceneTextureReductions.cpp
@@ -113,8 +113,8 @@ void BuildHZB(
 	FRDGTextureRef* OutClosestHZBTexture,
 	const TCHAR* FurthestHZBName,
 	FRDGTextureRef* OutFurthestHZBTexture,
-	EPixelFormat Format
-)
+	EPixelFormat Format,
+	const FBuildHZBAsyncComputeParams* AsyncComputeParams)
 {
 	RDG_EVENT_SCOPE(GraphBuilder, "BuildHZB");
 	QUICK_SCOPE_CYCLE_COUNTER(STAT_BuildHZB);
@@ -208,17 +208,26 @@ void BuildHZB(
 			PermutationVector.Set<FHZBBuildCS::FDimClosest>(bOutputClosest);
 			PermutationVector.Set<FHZBBuildCS::FDimVisBufferFormat>(VisBufferFormat);
 
+			const bool bAsyncCompute = AsyncComputeParams != nullptr;
+			const ERDGPassFlags PassFlags = bAsyncCompute ? ERDGPassFlags::AsyncCompute : ERDGPassFlags::Compute;
+
 			TShaderMapRef<FHZBBuildCS> ComputeShader(GetGlobalShaderMap(FeatureLevel), PermutationVector);
-			FComputeShaderUtils::AddPass(
+			FRDGPassRef ReduceHZBPass = FComputeShaderUtils::AddPass(
 				GraphBuilder,
 				RDG_EVENT_NAME("ReduceHZB(mips=[%d;%d]%s%s) %dx%d",
 					StartDestMip, EndDestMip - 1,
 					bOutputClosest ? TEXT(" Closest") : TEXT(""),
 					bOutputFurthest ? TEXT(" Furthest") : TEXT(""),
 					DstSize.X, DstSize.Y),
+				PassFlags,
 				ComputeShader,
 				PassParameters,
 				FComputeShaderUtils::GetGroupCount(DstSize, 8));
+
+			if (AsyncComputeParams && AsyncComputeParams->Prerequisite)
+			{
+				GraphBuilder.AddPassDependency(AsyncComputeParams->Prerequisite, ReduceHZBPass);
+			}
 		}
 		else
 		{
@@ -315,7 +324,8 @@ void BuildHZBFurthest(
 	EShaderPlatform ShaderPlatform,
 	const TCHAR* FurthestHZBName,
 	FRDGTextureRef* OutFurthestHZBTexture,
-	EPixelFormat Format
+	EPixelFormat Format,
+	const FBuildHZBAsyncComputeParams* AsyncComputeParams
 )
 {
 	BuildHZB(
@@ -329,5 +339,6 @@ void BuildHZBFurthest(
 		nullptr,
 		FurthestHZBName,
 		OutFurthestHZBTexture,
-		Format);
-}
\ No newline at end of file
+		Format,
+		AsyncComputeParams);
+}
diff --git a/Engine/Source/Runtime/Renderer/Private/SceneTextureReductions.h b/Engine/Source/Runtime/Renderer/Private/SceneTextureReductions.h
--- a/Engine/Source/Runtime/Renderer/Private/SceneTextureReductions.h
+++ b/Engine/Source/Runtime/Renderer/Private/SceneTextureReductions.h
@@ -6,6 +6,13 @@
 #include "RenderGraph.h"
 #include "SceneTextureParameters.h"
 
+struct FBuildHZBAsyncComputeParams
+{
+	FRDGPassRef Prerequisite = nullptr;
+};
+
+static constexpr EPixelFormat BuildHZBDefaultPixelFormat = PF_R16F;
+
 void BuildHZB(
 	FRDGBuilder& GraphBuilder,
 	FRDGTextureRef SceneDepth,
@@ -17,7 +24,8 @@ void BuildHZB(
 	FRDGTextureRef* OutClosestHZBTexture,
 	const TCHAR* FurthestHZBName,
 	FRDGTextureRef* OutFurthestHZBTexture,
-	EPixelFormat Format = PF_R16F);
+	EPixelFormat Format = BuildHZBDefaultPixelFormat,
+	const FBuildHZBAsyncComputeParams* AsyncComputeParams = nullptr);
 
 // Build only the furthest HZB
 void BuildHZBFurthest(
@@ -29,4 +37,5 @@ void BuildHZBFurthest(
 	EShaderPlatform ShaderPlatform,
 	const TCHAR* FurthestHZBName,
 	FRDGTextureRef* OutFurthestHZBTexture,
-	EPixelFormat Format = PF_R16F);
+	EPixelFormat Format = BuildHZBDefaultPixelFormat,
+	const FBuildHZBAsyncComputeParams* AsyncComputeParams = nullptr);
