Optimized RDG async compute fence computation.

diff --git a/Engine/Source/Runtime/RenderCore/Private/RenderGraphBuilder.cpp b/Engine/Source/Runtime/RenderCore/Private/RenderGraphBuilder.cpp
--- a/Engine/Source/Runtime/RenderCore/Private/RenderGraphBuilder.cpp
+++ b/Engine/Source/Runtime/RenderCore/Private/RenderGraphBuilder.cpp
@@ -907,30 +907,31 @@ FRDGBufferRef FRDGBuilder::RegisterExternalBuffer(
 
 ///////////////////////////////////////////////////////////////////////////////////////////////////
 
-void FRDGBuilder::AddPassDependency(FRDGPassHandle ProducerHandle, FRDGPassHandle ConsumerHandle)
-{
-	FRDGPass* Consumer = Passes[ConsumerHandle];
-
-	auto& Producers = Consumer->Producers;
-	if (Producers.Find(ProducerHandle) == INDEX_NONE)
-	{
-#if RDG_STATS
-		GRDGStatPassDependencyCount++;
-#endif
-
-		Producers.Add(ProducerHandle);
-	}
-}
-
 void FRDGBuilder::AddPassDependency(FRDGPass* Producer, FRDGPass* Consumer)
 {
 	auto& Producers = Consumer->Producers;
+
 	if (Producers.Find(Producer->Handle) == INDEX_NONE)
 	{
 #if RDG_STATS
 		GRDGStatPassDependencyCount++;
 #endif
 
+		if (Producer->Pipeline != Consumer->Pipeline)
+		{
+			// Finds the earliest consumer on the other pipeline for the producer.
+			if (Producer->CrossPipelineConsumer.IsNull() || Consumer->Handle < Producer->CrossPipelineConsumer)
+			{
+				Producer->CrossPipelineConsumer = Consumer->Handle;
+			}
+
+			// Finds the latest producer on the other pipeline for the consumer.
+			if (Consumer->CrossPipelineProducer.IsNull() || Producer->Handle > Consumer->CrossPipelineProducer)
+			{
+				Consumer->CrossPipelineProducer = Producer->Handle;
+			}
+		}
+
 		Producers.Add(Producer->Handle);
 	}
 }
@@ -1399,191 +1400,46 @@ void FRDGBuilder::Compile()
 	{
 		SCOPED_NAMED_EVENT(AsyncComputeFences, FColor::Emerald);
 
-		FRDGPassBitArray PassesOnAsyncCompute(false, CompilePassCount);
-
-		// Traverse the active passes in execution order to find latest cross-pipeline producer and the earliest
-		// cross-pipeline consumer for each pass. This helps narrow the search space later when building async
-		// compute overlap regions.
-
-		const auto IsCrossPipeline = [&](FRDGPassHandle A, FRDGPassHandle B)
-		{
-			return PassesOnAsyncCompute[A] != PassesOnAsyncCompute[B];
-		};
-
-		FRDGPassBitArray PassesWithCrossPipelineProducer(false, Passes.Num());
-		FRDGPassBitArray PassesWithCrossPipelineConsumer(false, Passes.Num());
-
-		for (FRDGPassHandle PassHandle = ProloguePassHandle + 1; PassHandle < EpiloguePassHandle; ++PassHandle)
-		{
-			FRDGPass* Pass = Passes[PassHandle];
-
-			if (Pass->bCulled || Pass->bEmptyParameters)
-			{
-				continue;
-			}
-
-			PassesOnAsyncCompute[PassHandle] = Pass->IsAsyncCompute();
-
-			for (FRDGPassHandle ProducerHandle : Pass->GetProducers())
-			{
-				const FRDGPassHandle ConsumerHandle = PassHandle;
-
-				if (!IsCrossPipeline(ProducerHandle, ConsumerHandle))
-				{
-					continue;
-				}
-
-				FRDGPass* Consumer = Pass;
-				FRDGPass* Producer = Passes[ProducerHandle];
-
-				// Finds the earliest consumer on the other pipeline for the producer.
-				if (Producer->CrossPipelineConsumer.IsNull() || ConsumerHandle < Producer->CrossPipelineConsumer)
-				{
-					Producer->CrossPipelineConsumer = PassHandle;
-					PassesWithCrossPipelineConsumer[ProducerHandle] = true;
-				}
-
-				// Finds the latest producer on the other pipeline for the consumer.
-				if (Consumer->CrossPipelineProducer.IsNull() || ProducerHandle > Consumer->CrossPipelineProducer)
-				{
-					Consumer->CrossPipelineProducer = ProducerHandle;
-					PassesWithCrossPipelineProducer[ConsumerHandle] = true;
-				}
-			}
-		}
-
 		// Establishes fork / join overlap regions for async compute. This is used for fencing as well as resource
 		// allocation / deallocation. Async compute passes can't allocate / release their resource references until
 		// the fork / join is complete, since the two pipes run in parallel. Therefore, all resource lifetimes on
 		// async compute are extended to cover the full async region.
 
-		const auto IsCrossPipelineProducer = [&](FRDGPassHandle A)
-		{
-			return PassesWithCrossPipelineConsumer[A];
-		};
-
-		const auto IsCrossPipelineConsumer = [&](FRDGPassHandle A)
-		{
-			return PassesWithCrossPipelineProducer[A];
-		};
-
-		const auto FindCrossPipelineProducer = [&](FRDGPassHandle PassHandle)
-		{
-			FRDGPassHandle LatestProducerHandle = ProloguePassHandle;
-			FRDGPassHandle ConsumerHandle = PassHandle;
-
-			// We want to find the latest producer on the other pipeline in order to establish a fork point.
-			// Since we could be consuming N resources with N producer passes, we only care about the last one.
-			while (ConsumerHandle != ProloguePassHandle)
-			{
-				if (IsCrossPipelineConsumer(ConsumerHandle) && !IsCrossPipeline(ConsumerHandle, PassHandle))
-				{
-					const FRDGPass* Consumer = Passes[ConsumerHandle];
-
-					if (Consumer->CrossPipelineProducer > LatestProducerHandle && !Consumer->bCulled)
-					{
-						LatestProducerHandle = Consumer->CrossPipelineProducer;
-					}
-				}
-				--ConsumerHandle;
-			}
-
-			return LatestProducerHandle;
-		};
-
-		const auto FindCrossPipelineConsumer = [&](FRDGPassHandle PassHandle)
-		{
-			FRDGPassHandle EarliestConsumerHandle = EpiloguePassHandle;
-			FRDGPassHandle ProducerHandle = PassHandle;
-
-			// We want to find the earliest consumer on the other pipeline, as this establishes a join point
-			// between the pipes. Since we could be producing for N consumers on the other pipeline, we only
-			// care about the first one to execute.
-			while (ProducerHandle != EpiloguePassHandle)
-			{
-				if (IsCrossPipelineProducer(ProducerHandle) && !IsCrossPipeline(ProducerHandle, PassHandle))
-				{
-					const FRDGPass* Producer = Passes[ProducerHandle];
-
-					if (Producer->CrossPipelineConsumer < EarliestConsumerHandle && !Producer->bCulled)
-					{
-						EarliestConsumerHandle = Producer->CrossPipelineConsumer;
-					}
-				}
-				++ProducerHandle;
-			}
-
-			return EarliestConsumerHandle;
-		};
-
-		FRDGPass* AsyncComputePassBeforeFork = nullptr;
-
-		const auto InsertGraphicsToAsyncComputeFork = [&](FRDGPass* GraphicsPass, FRDGPass* AsyncComputePass)
-		{
-			FRDGBarrierBatchBegin& EpilogueBarriersToBeginForAsyncCompute = GraphicsPass->GetEpilogueBarriersToBeginForAsyncCompute(Allocator, TransitionCreateQueue);
-
-			GraphicsPass->bGraphicsFork = 1;
-			EpilogueBarriersToBeginForAsyncCompute.SetUseCrossPipelineFence();
-
-			AsyncComputePass->bAsyncComputeBegin = 1;
-			AsyncComputePass->GetPrologueBarriersToEnd(Allocator).AddDependency(&EpilogueBarriersToBeginForAsyncCompute);
-
-			// Since we are fencing the graphics pipe to some new async compute work, make sure to flush any prior work.
-			if (AsyncComputePassBeforeFork)
-			{
-				AsyncComputePassBeforeFork->bDispatchAfterExecute = 1;
-			}
-		};
-
-		const auto InsertAsyncComputeToGraphicsJoin = [&](FRDGPass* AsyncComputePass, FRDGPass* GraphicsPass)
-		{
-			FRDGBarrierBatchBegin& EpilogueBarriersToBeginForGraphics = AsyncComputePass->GetEpilogueBarriersToBeginForGraphics(Allocator, TransitionCreateQueue);
-
-			AsyncComputePass->bAsyncComputeEnd = 1;
-			AsyncComputePass->bDispatchAfterExecute = 1;
-			EpilogueBarriersToBeginForGraphics.SetUseCrossPipelineFence();
-
-			GraphicsPass->bGraphicsJoin = 1;
-			GraphicsPass->GetPrologueBarriersToEnd(Allocator).AddDependency(&EpilogueBarriersToBeginForGraphics);
-		};
-
-		const auto AddResourcesToBegin = [this](FRDGPass* PassToBegin, FRDGPass* PassWithResources)
-		{
-			Passes[PassToBegin->PrologueBarrierPass]->ResourcesToBegin.Add(PassWithResources);
-		};
-
-		const auto AddResourcesToEnd = [this](FRDGPass* PassToEnd, FRDGPass* PassWithResources)
-		{
-			Passes[PassToEnd->EpilogueBarrierPass]->ResourcesToEnd.Add(PassWithResources);
-		};
-
 		FRDGPassHandle CurrentGraphicsForkPassHandle;
+		FRDGPass* AsyncComputePassBeforeFork = nullptr;
 
 		for (FRDGPassHandle PassHandle = ProloguePassHandle + 1; PassHandle < EpiloguePassHandle; ++PassHandle)
 		{
-			if (!PassesOnAsyncCompute[PassHandle])
-			{
-				continue;
-			}
-
 			FRDGPass* AsyncComputePass = Passes[PassHandle];
 
-			if (AsyncComputePass->bCulled)
+			if (!AsyncComputePass->IsAsyncCompute() || AsyncComputePass->bCulled)
 			{
 				continue;
 			}
 
-			const FRDGPassHandle GraphicsForkPassHandle = FindCrossPipelineProducer(PassHandle);
-
+			FRDGPassHandle GraphicsForkPassHandle = FRDGPassHandle::Max(AsyncComputePass->CrossPipelineProducer, FRDGPassHandle::Max(CurrentGraphicsForkPassHandle, ProloguePassHandle));
 			FRDGPass* GraphicsForkPass = Passes[GraphicsForkPassHandle];
 
 			AsyncComputePass->GraphicsForkPass = GraphicsForkPassHandle;
-			AddResourcesToBegin(GraphicsForkPass, AsyncComputePass);
+			Passes[GraphicsForkPass->PrologueBarrierPass]->ResourcesToBegin.Add(AsyncComputePass);
 
 			if (CurrentGraphicsForkPassHandle != GraphicsForkPassHandle)
 			{
 				CurrentGraphicsForkPassHandle = GraphicsForkPassHandle;
-				InsertGraphicsToAsyncComputeFork(GraphicsForkPass, AsyncComputePass);
+
+				FRDGBarrierBatchBegin& EpilogueBarriersToBeginForAsyncCompute = GraphicsForkPass->GetEpilogueBarriersToBeginForAsyncCompute(Allocator, TransitionCreateQueue);
+
+				GraphicsForkPass->bGraphicsFork = 1;
+				EpilogueBarriersToBeginForAsyncCompute.SetUseCrossPipelineFence();
+
+				AsyncComputePass->bAsyncComputeBegin = 1;
+				AsyncComputePass->GetPrologueBarriersToEnd(Allocator).AddDependency(&EpilogueBarriersToBeginForAsyncCompute);
+
+				// Since we are fencing the graphics pipe to some new async compute work, make sure to flush any prior work.
+				if (AsyncComputePassBeforeFork)
+				{
+					AsyncComputePassBeforeFork->bDispatchAfterExecute = 1;
+				}
 			}
 
 			AsyncComputePassBeforeFork = AsyncComputePass;
@@ -1593,29 +1449,31 @@ void FRDGBuilder::Compile()
 
 		for (FRDGPassHandle PassHandle = EpiloguePassHandle - 1; PassHandle > ProloguePassHandle; --PassHandle)
 		{
-			if (!PassesOnAsyncCompute[PassHandle])
-			{
-				continue;
-			}
-
 			FRDGPass* AsyncComputePass = Passes[PassHandle];
 
-			if (AsyncComputePass->bCulled)
+			if (!AsyncComputePass->IsAsyncCompute() || AsyncComputePass->bCulled)
 			{
 				continue;
 			}
 
-			const FRDGPassHandle GraphicsJoinPassHandle = FindCrossPipelineConsumer(PassHandle);
-
+			FRDGPassHandle GraphicsJoinPassHandle = FRDGPassHandle::Min(AsyncComputePass->CrossPipelineConsumer, FRDGPassHandle::Min(CurrentGraphicsJoinPassHandle, EpiloguePassHandle));
 			FRDGPass* GraphicsJoinPass = Passes[GraphicsJoinPassHandle];
 
 			AsyncComputePass->GraphicsJoinPass = GraphicsJoinPassHandle;
-			GraphicsJoinPass->ResourcesToEnd.Add(AsyncComputePass);
+			Passes[GraphicsJoinPass->EpilogueBarrierPass]->ResourcesToEnd.Add(AsyncComputePass);
 
 			if (CurrentGraphicsJoinPassHandle != GraphicsJoinPassHandle)
 			{
 				CurrentGraphicsJoinPassHandle = GraphicsJoinPassHandle;
-				InsertAsyncComputeToGraphicsJoin(AsyncComputePass, GraphicsJoinPass);
+
+				FRDGBarrierBatchBegin& EpilogueBarriersToBeginForGraphics = AsyncComputePass->GetEpilogueBarriersToBeginForGraphics(Allocator, TransitionCreateQueue);
+
+				AsyncComputePass->bAsyncComputeEnd = 1;
+				AsyncComputePass->bDispatchAfterExecute = 1;
+				EpilogueBarriersToBeginForGraphics.SetUseCrossPipelineFence();
+
+				GraphicsJoinPass->bGraphicsJoin = 1;
+				GraphicsJoinPass->GetPrologueBarriersToEnd(Allocator).AddDependency(&EpilogueBarriersToBeginForGraphics);
 			}
 		}
 	}
diff --git a/Engine/Source/Runtime/RenderCore/Public/RenderGraphBuilder.h b/Engine/Source/Runtime/RenderCore/Public/RenderGraphBuilder.h
--- a/Engine/Source/Runtime/RenderCore/Public/RenderGraphBuilder.h
+++ b/Engine/Source/Runtime/RenderCore/Public/RenderGraphBuilder.h
@@ -851,7 +851,6 @@ private:
 
 	UE::Tasks::FTask CreateUniformBuffers();
 
-	void AddPassDependency(FRDGPassHandle ProducerHandle, FRDGPassHandle ConsumerHandle);
 	void AddCullingDependency(FRDGProducerStatesByPipeline& LastProducers, const FRDGProducerState& NextState, ERHIPipeline NextPipeline);
 
 	void AddEpilogueTransition(FRDGTextureRef Texture);
