Make sure shader refcount cannot be 0 during the creation.

diff --git a/Engine/Source/Runtime/RenderCore/Private/ShaderCodeLibrary.cpp b/Engine/Source/Runtime/RenderCore/Private/ShaderCodeLibrary.cpp
--- a/Engine/Source/Runtime/RenderCore/Private/ShaderCodeLibrary.cpp
+++ b/Engine/Source/Runtime/RenderCore/Private/ShaderCodeLibrary.cpp
@@ -1172,7 +1172,7 @@ FRHIShader* FShaderMapResource_SharedCode::CreateRHIShaderOrCrash(int32 ShaderIn
 #endif
 
 	const int32 LibraryShaderIndex = LibraryInstance->Library->GetShaderIndex(ShaderMapIndex, ShaderIndex);
-	FRHIShader* CreatedShader = LibraryInstance->GetOrCreateShader(LibraryShaderIndex);
+	TRefCountPtr<FRHIShader> CreatedShader = LibraryInstance->GetOrCreateShader(LibraryShaderIndex);
 	if (UNLIKELY(CreatedShader == nullptr))
 	{
 		UE_LOG(LogShaders, Fatal, TEXT("FShaderMapResource_SharedCode::InitRHI is unable to create a shader"));
diff --git a/Engine/Source/Runtime/RenderCore/Private/ShaderResource.cpp b/Engine/Source/Runtime/RenderCore/Private/ShaderResource.cpp
--- a/Engine/Source/Runtime/RenderCore/Private/ShaderResource.cpp
+++ b/Engine/Source/Runtime/RenderCore/Private/ShaderResource.cpp
@@ -691,7 +691,7 @@ FRHIShader* FShaderMapResource_InlineCode::CreateRHIShaderOrCrash(int32 ShaderIn
 	const FSHAHash& ShaderHash = Code->ShaderHashes[ShaderIndex];
 	const EShaderFrequency Frequency = ShaderEntry.Frequency;
 
-	FRHIShader* RHIShader = nullptr;
+	TRefCountPtr<FRHIShader> RHIShader;
 	switch (Frequency)
 	{
 	case SF_Vertex: RHIShader = RHICreateVertexShader(ShaderCodeView, ShaderHash); break;
