Removing deprecated code from Renderer and RenderCore.

diff --git a/Engine/Plugins/VirtualProduction/TextureShare/Source/TextureShare/Private/Game/ViewExtension/TextureShareSceneViewExtension.cpp b/Engine/Plugins/VirtualProduction/TextureShare/Source/TextureShare/Private/Game/ViewExtension/TextureShareSceneViewExtension.cpp
--- a/Engine/Plugins/VirtualProduction/TextureShare/Source/TextureShare/Private/Game/ViewExtension/TextureShareSceneViewExtension.cpp
+++ b/Engine/Plugins/VirtualProduction/TextureShare/Source/TextureShare/Private/Game/ViewExtension/TextureShareSceneViewExtension.cpp
@@ -12,6 +12,7 @@
 #include "RHICommandList.h"
 #include "RHIResources.h"
 #include "SceneView.h"
+#include "SceneTextures.h"
 
 #include "PostProcess/SceneRenderTargets.h"
 
diff --git a/Engine/Source/Runtime/RenderCore/Private/ProfilingDebugging/RealtimeGPUProfiler.cpp b/Engine/Source/Runtime/RenderCore/Private/ProfilingDebugging/RealtimeGPUProfiler.cpp
--- a/Engine/Source/Runtime/RenderCore/Private/ProfilingDebugging/RealtimeGPUProfiler.cpp
+++ b/Engine/Source/Runtime/RenderCore/Private/ProfilingDebugging/RealtimeGPUProfiler.cpp
@@ -55,22 +55,6 @@ static TAutoConsoleVariable<int> CVarGPUStatsChildTimesIncluded(
 
 #if WANTS_DRAW_MESH_EVENTS
 
-void FDrawEvent::Start(FRHIComputeCommandList& InRHICmdList, FColor Color, const TCHAR* Fmt, ...)
-{
-	check(IsInParallelRenderingThread() || IsInRHIThread());
-	{
-		va_list ptr;
-		va_start(ptr, Fmt);
-		TCHAR TempStr[256];
-		// Build the string in the temp buffer
-		FCString::GetVarArgs(TempStr, UE_ARRAY_COUNT(TempStr), Fmt, ptr);
-		InRHICmdList.PushEvent(TempStr, Color);
-		RHICmdList = &InRHICmdList;
-		va_end(ptr);
-	}
-	bStarted = true;
-}
-
 void FDrawEvent::Start(FRHIComputeCommandList* InRHICmdList, FColor Color, const TCHAR* Fmt, ...)
 {
 	bool bIsRenderingOrRHIThread = IsInParallelRenderingThread() || IsInRHIThread();
@@ -972,11 +956,6 @@ FRealtimeGPUProfiler::FRealtimeGPUProfiler()
 	}
 }
 
-void FRealtimeGPUProfiler::Release()
-{
-	Cleanup();
-}
-
 void FRealtimeGPUProfiler::Cleanup()
 {
 	for (int Index = 0; Index < Frames.Num(); Index++)
diff --git a/Engine/Source/Runtime/RenderCore/Private/RenderGraphResources.cpp b/Engine/Source/Runtime/RenderCore/Private/RenderGraphResources.cpp
--- a/Engine/Source/Runtime/RenderCore/Private/RenderGraphResources.cpp
+++ b/Engine/Source/Runtime/RenderCore/Private/RenderGraphResources.cpp
@@ -149,9 +149,3 @@ FRDGTextureSubresourceRange FRDGTexture::GetSubresourceRangeSRV() const
 
 	return Range;
 }
-
-IPooledRenderTarget* FRDGTexture::GetPooledRenderTarget() const
-{
-	IF_RDG_ENABLE_DEBUG(ValidateRHIAccess());
-	return RenderTarget;
-}
\ No newline at end of file
diff --git a/Engine/Source/Runtime/RenderCore/Public/ProfilingDebugging/RealtimeGPUProfiler.h b/Engine/Source/Runtime/RenderCore/Public/ProfilingDebugging/RealtimeGPUProfiler.h
--- a/Engine/Source/Runtime/RenderCore/Public/ProfilingDebugging/RealtimeGPUProfiler.h
+++ b/Engine/Source/Runtime/RenderCore/Public/ProfilingDebugging/RealtimeGPUProfiler.h
@@ -119,8 +119,6 @@ struct RENDERCORE_API FBreadcrumbEvent
 		 *  command list (rendering / RHI thread) or enqueue a command on the current immediate command list (game thread)
 		 */
 		void CDECL Start(FRHIComputeCommandList* RHICmdList, FColor Color, const TCHAR* Fmt, ...);
-		UE_DEPRECATED(4.27, "This function has been deprecated in favor of the other version that takes a pointer in parameter")
-		void CDECL Start(FRHIComputeCommandList& RHICmdList, FColor Color, const TCHAR* Fmt, ...);
 		void Stop();
 	};
 
@@ -397,10 +395,6 @@ public:
 	RENDERCORE_API void EndFrame(FRHICommandListImmediate& RHICmdList);
 	RENDERCORE_API void SuspendFrame();
 
-	/** Final cleanup */
-	UE_DEPRECATED(4.23, "Use FRealtimeGPUProfiler::SafeRelease() instead.")
-	RENDERCORE_API void Release();
-
 	/** Push/pop events */
 	FRealtimeGPUProfilerQuery PushEvent(FRHIGPUMask GPUMask, const FName& Name, const FName& StatName, const TCHAR* Description);
 	FRealtimeGPUProfilerQuery PopEvent();
diff --git a/Engine/Source/Runtime/RenderCore/Public/RenderGraphBuilder.h b/Engine/Source/Runtime/RenderCore/Public/RenderGraphBuilder.h
--- a/Engine/Source/Runtime/RenderCore/Public/RenderGraphBuilder.h
+++ b/Engine/Source/Runtime/RenderCore/Public/RenderGraphBuilder.h
@@ -379,42 +379,6 @@ public:
 
 	//////////////////////////////////////////////////////////////////////////
 	// Deprecated Functions
-	UE_DEPRECATED(5.0, "PreallocateTexture has been renamed to ConvertToExternalTexture")
-	inline void PreallocateTexture(FRDGTextureRef Texture) { ConvertToExternalTexture(Texture); }
-
-	UE_DEPRECATED(5.0, "PreallocateBuffer has been renamed to ConvertToExternalBuffer")
-	inline void PreallocateBuffer(FRDGBufferRef Buffer) { ConvertToExternalBuffer(Buffer); }
-
-	UE_DEPRECATED(5.0, "RegisterExternalTexture with ERenderTargetTexture is deprecated. Use the variant without instead.")
-PRAGMA_DISABLE_DEPRECATION_WARNINGS
-	inline FRDGTextureRef RegisterExternalTexture(
-		const TRefCountPtr<IPooledRenderTarget>& ExternalPooledTexture,
-		ERenderTargetTexture Texture,
-		ERDGTextureFlags Flags = ERDGTextureFlags::None)
-PRAGMA_ENABLE_DEPRECATION_WARNINGS
-	{
-		return RegisterExternalTexture(ExternalPooledTexture, Flags);
-	}
-
-	UE_DEPRECATED(5.0, "RegisterExternalTexture with ERenderTargetTexture is deprecated. Use the variant without instead.")
-PRAGMA_DISABLE_DEPRECATION_WARNINGS
-	inline FRDGTextureRef RegisterExternalTexture(
-		const TRefCountPtr<IPooledRenderTarget>& ExternalPooledTexture,
-		const TCHAR* NameIfNotRegistered,
-		ERenderTargetTexture RenderTargetTexture,
-		ERDGTextureFlags Flags = ERDGTextureFlags::None)
-PRAGMA_ENABLE_DEPRECATION_WARNINGS
-	{
-		return RegisterExternalTexture(ExternalPooledTexture, NameIfNotRegistered, Flags);
-	}
-
-	UE_DEPRECATED(5.0, "FindExternalTexture with ERenderTargetTexture is deprecated. Use the variant without instead.")
-PRAGMA_DISABLE_DEPRECATION_WARNINGS
-	FRDGTextureRef FindExternalTexture(IPooledRenderTarget* ExternalPooledTexture, ERenderTargetTexture Texture) const
-PRAGMA_ENABLE_DEPRECATION_WARNINGS
-	{
-		return FindExternalTexture(ExternalPooledTexture);
-	}
 
 	UE_DEPRECATED(5.1, "FinalizeResourceAccess has been replaced by UseExternalAccessMode")
 	inline void FinalizeResourceAccess(FRDGTextureAccessArray&& InTextures, FRDGBufferAccessArray&& InBuffers)
diff --git a/Engine/Source/Runtime/RenderCore/Public/RenderGraphResources.h b/Engine/Source/Runtime/RenderCore/Public/RenderGraphResources.h
--- a/Engine/Source/Runtime/RenderCore/Public/RenderGraphResources.h
+++ b/Engine/Source/Runtime/RenderCore/Public/RenderGraphResources.h
@@ -455,14 +455,6 @@ inline FRDGTextureDesc Translate(const FPooledRenderTargetDesc& InDesc);
 /** Translates from an RHI/RDG texture descriptor to a pooled render target descriptor. */
 inline FPooledRenderTargetDesc Translate(const FRHITextureDesc& InDesc);
 
-UE_DEPRECATED(5.0, "Translate with ERenderTargetTexture is deprecated. Please use the single parameter variant.")
-PRAGMA_DISABLE_DEPRECATION_WARNINGS
-inline FRDGTextureDesc Translate(const FPooledRenderTargetDesc& InDesc, ERenderTargetTexture InTexture)
-PRAGMA_ENABLE_DEPRECATION_WARNINGS
-{
-	return Translate(InDesc);
-}
-
 class RENDERCORE_API FRDGPooledTexture final
 	: public FRefCountBase
 {
@@ -499,10 +491,6 @@ public:
 	//////////////////////////////////////////////////////////////////////////
 	//! The following methods may only be called during pass execution.
 
-	/** Returns the allocated pooled render target. */
-	UE_DEPRECATED(5.0, "Accessing the underlying pooled render target has been deprecated. Use GetRHI() instead.")
-	IPooledRenderTarget* GetPooledRenderTarget() const;
-
 	/** Returns the allocated RHI texture. */
 	FORCEINLINE FRHITexture* GetRHI() const
 	{
@@ -1232,20 +1220,6 @@ public:
 		return GetRHI();
 	}
 
-	/** Returns the buffer to use for RHI calls, eg RHILockBuffer. */
-	UE_DEPRECATED(5.0, "Buffers types have been consolidated; use GetRHI() instead.")
-	FORCEINLINE FRHIBuffer* GetRHIVertexBuffer() const
-	{
-		return GetRHI();
-	}
-
-	/** Returns the buffer to use for structured buffer calls. */
-	UE_DEPRECATED(5.0, "Buffers types have been consolidated; use GetRHI() instead.")
-	FORCEINLINE FRHIBuffer* GetRHIStructuredBuffer() const
-	{
-		return GetRHI();
-	}
-
 	//////////////////////////////////////////////////////////////////////////
 
 	FRDGBufferHandle GetHandle() const
diff --git a/Engine/Source/Runtime/RenderCore/Public/RenderGraphUtils.h b/Engine/Source/Runtime/RenderCore/Public/RenderGraphUtils.h
--- a/Engine/Source/Runtime/RenderCore/Public/RenderGraphUtils.h
+++ b/Engine/Source/Runtime/RenderCore/Public/RenderGraphUtils.h
@@ -118,15 +118,6 @@ inline bool IsRegistered(FRDGBuilder& GraphBuilder, const TRefCountPtr<FRDGPoole
 	return GraphBuilder.FindExternalBuffer(Buffer) != nullptr;
 }
 
-/** Returns the pooled render target from an RDG texture if it exists, or null otherwise. */
-UE_DEPRECATED(5.0, "Accessing the underlying pooled render target has been deprecated. Use TryGetRHI() instead.")
-inline IPooledRenderTarget* TryGetPooledRenderTarget(FRDGTextureRef Texture)
-{
-PRAGMA_DISABLE_DEPRECATION_WARNINGS
-	return Texture ? Texture->GetPooledRenderTarget() : nullptr;
-PRAGMA_ENABLE_DEPRECATION_WARNINGS
-}
-
 inline FRenderTargetBindingSlots GetRenderTargetBindings(ERenderTargetLoadAction ColorLoadAction, TArrayView<FRDGTextureRef> ColorTextures)
 {
 	check(ColorTextures.Num() <= MaxSimultaneousRenderTargets);
@@ -840,17 +831,6 @@ RENDERCORE_API void AddEnqueueCopyPass(FRDGBuilder& GraphBuilder, FRHIGPUTexture
 /** Adds a pass to readback contents of an RDG buffer. */
 RENDERCORE_API void AddEnqueueCopyPass(FRDGBuilder& GraphBuilder, FRHIGPUBufferReadback* Readback, FRDGBufferRef SourceBuffer, uint32 NumBytes);
 
-UE_DEPRECATED(5.0, "Please use GraphBuilder.QueueBufferUpload to perform an upload.")
-inline void AddBufferUploadPass(
-	FRDGBuilder& GraphBuilder,
-	FRDGBufferRef Buffer,
-	const void* InitialData,
-	uint64 InitialDataSize,
-	ERDGInitialDataFlags InitialDataFlags = ERDGInitialDataFlags::None)
-{
-	GraphBuilder.QueueBufferUpload(Buffer, InitialData, InitialDataSize, InitialDataFlags);
-}
-
 /** Helper class to allocate data from a GraphBuilder in order to upload said data to an RDG resource.
 *   Allocating from the GraphBuilder makes it so we don't have to copy the data before deferring the upload.
 */
@@ -1027,13 +1007,6 @@ FORCEINLINE void AddPass(FRDGBuilder& GraphBuilder, FRDGEventName&& Name, Execut
 	GraphBuilder.AddPass(MoveTemp(Name), ERDGPassFlags::None, MoveTemp(ExecuteLambda));
 }
 
-template <typename ExecuteLambdaType>
-UE_DEPRECATED(5.0, "AddPass without an RDG_EVENT_NAME is deprecated. Use the named version instead.")
-FORCEINLINE void AddPass(FRDGBuilder& GraphBuilder, ExecuteLambdaType&& ExecuteLambda)
-{
-	AddPass(GraphBuilder, {}, MoveTemp(ExecuteLambda));
-}
-
 template <typename ExecuteLambdaType>
 FORCEINLINE void AddPassIfDebug(FRDGBuilder& GraphBuilder, FRDGEventName&& Name, ExecuteLambdaType&& ExecuteLambda)
 {
@@ -1042,19 +1015,6 @@ FORCEINLINE void AddPassIfDebug(FRDGBuilder& GraphBuilder, FRDGEventName&& Name,
 #endif
 }
 
-template <typename ExecuteLambdaType>
-UE_DEPRECATED(5.0, "AddPassIfDebug without an RDG_EVENT_NAME is deprecated. Use the named version instead.")
-FORCEINLINE void AddPassIfDebug(FRDGBuilder& GraphBuilder, ExecuteLambdaType&& ExecuteLambda)
-{
-	AddPassIfDebug(GraphBuilder, {}, MoveTemp(ExecuteLambda));
-}
-
-UE_DEPRECATED(5.0, "AddSetCurrentStatPass is deprecated. Use GraphBuilder.SetCommandListStat instead.")
-FORCEINLINE void AddSetCurrentStatPass(FRDGBuilder& GraphBuilder, TStatId StatId)
-{
-	GraphBuilder.SetCommandListStat(StatId);
-}
-
 FORCEINLINE void AddDispatchToRHIThreadPass(FRDGBuilder& GraphBuilder)
 {
 	AddPass(GraphBuilder, RDG_EVENT_NAME("DispatchToRHI"), [](FRHICommandListImmediate& RHICmdList)
@@ -1063,68 +1023,6 @@ FORCEINLINE void AddDispatchToRHIThreadPass(FRDGBuilder& GraphBuilder)
 	});
 }
 
-UE_DEPRECATED(5.0, "AddBeginUAVOverlapPass is deprecated.")
-FORCEINLINE void AddBeginUAVOverlapPass(FRDGBuilder& GraphBuilder)
-{
-	AddPass(GraphBuilder, RDG_EVENT_NAME("BeginUAVOverlap"), [](FRHICommandList& RHICmdList)
-	{
-		RHICmdList.BeginUAVOverlap();
-	});
-}
-
-UE_DEPRECATED(5.0, "AddEndUAVOverlapPass is deprecated.")
-FORCEINLINE void AddEndUAVOverlapPass(FRDGBuilder& GraphBuilder)
-{
-	AddPass(GraphBuilder, RDG_EVENT_NAME("EndUAVOverlap"), [](FRHICommandList& RHICmdList)
-	{
-		RHICmdList.EndUAVOverlap();
-	});
-}
-
-UE_DEPRECATED(5.0, "AddBeginUAVOverlapPass is deprecated.")
-FORCEINLINE void AddBeginUAVOverlapPass(FRDGBuilder& GraphBuilder, FRHIUnorderedAccessView* UAV)
-{
-	AddPass(GraphBuilder, RDG_EVENT_NAME("BeginUAVOverlap"), [UAV](FRHICommandList& RHICmdList)
-	{
-		RHICmdList.BeginUAVOverlap(UAV);
-	});
-}
-
-UE_DEPRECATED(5.0, "AddEndUAVOverlapPass is deprecated.")
-FORCEINLINE void AddEndUAVOverlapPass(FRDGBuilder& GraphBuilder, FRHIUnorderedAccessView* UAV)
-{
-	AddPass(GraphBuilder, RDG_EVENT_NAME("EndUAVOverlap"), [UAV](FRHICommandList& RHICmdList)
-	{
-		RHICmdList.EndUAVOverlap(UAV);
-	});
-}
-
-UE_DEPRECATED(5.0, "AddBeginUAVOverlapPass is deprecated.")
-FORCEINLINE void AddBeginUAVOverlapPass(FRDGBuilder& GraphBuilder, TArrayView<FRHIUnorderedAccessView*> UAVs)
-{
-	uint32 AllocSize = UAVs.Num() * sizeof(FRHIUnorderedAccessView*);
-	FRHIUnorderedAccessView** LocalUAVs = (FRHIUnorderedAccessView**)GraphBuilder.Alloc(AllocSize, alignof(FRHIUnorderedAccessView*));
-	FMemory::Memcpy(LocalUAVs, UAVs.GetData(), AllocSize);
-	TArrayView<FRHIUnorderedAccessView*> LocalView(LocalUAVs, UAVs.Num());
-	AddPass(GraphBuilder, RDG_EVENT_NAME("BeginUAVOverlap"), [LocalView](FRHICommandList& RHICmdList)
-	{
-		RHICmdList.BeginUAVOverlap(LocalView);
-	});
-}
-
-UE_DEPRECATED(5.0, "AddEndUAVOverlapPass is deprecated.")
-FORCEINLINE void AddEndUAVOverlapPass(FRDGBuilder& GraphBuilder, TArrayView<FRHIUnorderedAccessView*> UAVs)
-{
-	uint32 AllocSize = UAVs.Num() * sizeof(FRHIUnorderedAccessView*);
-	FRHIUnorderedAccessView** LocalUAVs = (FRHIUnorderedAccessView**)GraphBuilder.Alloc(AllocSize, alignof(FRHIUnorderedAccessView*));
-	FMemory::Memcpy(LocalUAVs, UAVs.GetData(), AllocSize);
-	TArrayView<FRHIUnorderedAccessView*> LocalView(LocalUAVs, UAVs.Num());
-	AddPass(GraphBuilder, RDG_EVENT_NAME("EndUAVOverlap"), [LocalView](FRHICommandList& RHICmdList)
-	{
-		RHICmdList.EndUAVOverlap(LocalView);
-	});
-}
-
 BEGIN_SHADER_PARAMETER_STRUCT(FReadbackTextureParameters, )
 	RDG_TEXTURE_ACCESS(Texture, ERHIAccess::CopySrc)
 END_SHADER_PARAMETER_STRUCT()
@@ -1393,54 +1291,4 @@ inline const TRefCountPtr<FRDGPooledBuffer>& ConvertToFinalizedExternalBuffer(
 	return ConvertToExternalAccessBuffer(GraphBuilder, Buffer, AccessFinal);
 }
 
-UE_DEPRECATED(5.0, "RegisterExternalTextureWithFallback no longer requires ERenderTargetTexture")
-PRAGMA_DISABLE_DEPRECATION_WARNINGS
-inline FRDGTextureRef RegisterExternalTextureWithFallback(
-	FRDGBuilder& GraphBuilder,
-	const TRefCountPtr<IPooledRenderTarget>& ExternalPooledTexture,
-	const TRefCountPtr<IPooledRenderTarget>& FallbackPooledTexture,
-	ERenderTargetTexture ExternalTexture,
-	ERenderTargetTexture FallbackTexture = ERenderTargetTexture::ShaderResource)
-PRAGMA_ENABLE_DEPRECATION_WARNINGS
-{
-	return RegisterExternalTextureWithFallback(GraphBuilder, ExternalPooledTexture, FallbackPooledTexture);
-}
-
-UE_DEPRECATED(5.0, "TryRegisterExternalTexture no longer requires ERenderTargetTexture")
-PRAGMA_DISABLE_DEPRECATION_WARNINGS
-inline FRDGTextureRef TryRegisterExternalTexture(
-	FRDGBuilder& GraphBuilder,
-	const TRefCountPtr<IPooledRenderTarget>& ExternalPooledTexture,
-	ERenderTargetTexture RenderTargetTexture,
-	ERDGTextureFlags Flags = ERDGTextureFlags::None)
-PRAGMA_ENABLE_DEPRECATION_WARNINGS
-{
-	return ExternalPooledTexture ? GraphBuilder.RegisterExternalTexture(ExternalPooledTexture, Flags) : nullptr;
-}
-
-UE_DEPRECATED(5.0, "RegisterExternalTextureMSAA with a single pooled render target is no longer supported.")
-inline FRDGTextureMSAA RegisterExternalTextureMSAA(
-	FRDGBuilder& GraphBuilder,
-	const TRefCountPtr<IPooledRenderTarget>& ExternalPooledTexture)
-{
-	return FRDGTextureMSAA();
-}
-
-UE_DEPRECATED(5.0, "RegisterExternalTextureMSAA with a single pooled render target is no longer supported.")
-inline FRDGTextureMSAA TryRegisterExternalTextureMSAA(
-	FRDGBuilder& GraphBuilder,
-	const TRefCountPtr<IPooledRenderTarget>& ExternalPooledTexture)
-{
-	return FRDGTextureMSAA();
-}
-
-UE_DEPRECATED(5.0, "RegisterExternalTextureMSAAWithFallback is no longer supported.")
-inline FRDGTextureMSAA RegisterExternalTextureMSAAWithFallback(
-	FRDGBuilder& GraphBuilder,
-	const TRefCountPtr<IPooledRenderTarget>& ExternalPooledTexture,
-	const TRefCountPtr<IPooledRenderTarget>& FallbackPooledTexture)
-{
-	return FRDGTextureMSAA();
-}
-
 //////////////////////////////////////////////////////////////////////////
\ No newline at end of file
diff --git a/Engine/Source/Runtime/RenderCore/Public/RenderTargetPool.h b/Engine/Source/Runtime/RenderCore/Public/RenderTargetPool.h
--- a/Engine/Source/Runtime/RenderCore/Public/RenderTargetPool.h
+++ b/Engine/Source/Runtime/RenderCore/Public/RenderTargetPool.h
@@ -63,16 +63,6 @@ struct RENDERCORE_API FPooledRenderTarget final : public IPooledRenderTarget
 	bool IsTracked() const override { return RenderTargetPool != nullptr; }
 	uint32 ComputeMemorySize() const override;
 
-	UE_DEPRECATED(5.0, "This method is deprecated.")
-PRAGMA_DISABLE_DEPRECATION_WARNINGS
-	FRDGPooledTexture* GetRDG(ERenderTargetTexture Texture) { return &PooledTexture; }
-PRAGMA_ENABLE_DEPRECATION_WARNINGS
-
-	UE_DEPRECATED(5.0, "This method is deprecated.")
-PRAGMA_DISABLE_DEPRECATION_WARNINGS
-	const FRDGPooledTexture* GetRDG(ERenderTargetTexture Texture) const { return &PooledTexture; }
-PRAGMA_ENABLE_DEPRECATION_WARNINGS
-
 private:
 	/** Pointer back to the pool for render targets which are actually pooled, otherwise NULL. */
 	FRenderTargetPool* RenderTargetPool;
diff --git a/Engine/Source/Runtime/RenderCore/Public/RendererInterface.h b/Engine/Source/Runtime/RenderCore/Public/RendererInterface.h
--- a/Engine/Source/Runtime/RenderCore/Public/RendererInterface.h
+++ b/Engine/Source/Runtime/RenderCore/Public/RendererInterface.h
@@ -80,13 +80,6 @@ public:
 		check(!IsValid());
 	}
 
-PRAGMA_DISABLE_DEPRECATION_WARNINGS
-	FPooledRenderTargetDesc(const FPooledRenderTargetDesc&) = default;
-	FPooledRenderTargetDesc(FPooledRenderTargetDesc&&) = default;
-	FPooledRenderTargetDesc& operator=(const FPooledRenderTargetDesc&) = default;
-	FPooledRenderTargetDesc& operator=(FPooledRenderTargetDesc&&) = default;
-PRAGMA_ENABLE_DEPRECATION_WARNINGS
-
 	/**
 	 * Factory function to create 2D texture description
 	 * @param InFlags bit mask combined from elements of ETextureCreateFlags e.g. TexCreate_UAV
@@ -409,9 +402,6 @@ PRAGMA_ENABLE_DEPRECATION_WARNINGS
 	FClearValueBinding ClearValue;
 	/** The flags that must be set on both the shader-resource and the targetable texture. bit mask combined from elements of ETextureCreateFlags e.g. TexCreate_UAV */
 	ETextureCreateFlags Flags = TexCreate_None;
-	/** The flags that must be set on the targetable texture. bit mask combined from elements of ETextureCreateFlags e.g. TexCreate_UAV */
-	UE_DEPRECATED(5.0, "TargetableFlags has been deprecated. Use Flags instead.")
-	ETextureCreateFlags TargetableFlags = TexCreate_None;
 	/** Texture format e.g. PF_B8G8R8A8 */
 	EPixelFormat Format = PF_Unknown;
 	/** Texture format used when creating the UAV (if TexCreate_UAV is also passed in TargetableFlags, ignored otherwise). PF_Unknown == use default (same as Format) */
@@ -435,38 +425,13 @@ PRAGMA_ENABLE_DEPRECATION_WARNINGS
 			uint8 bIsArray : 1;
 			/** true if a cubemap texture */
 			uint8 bIsCubemap : 1;
-			/** Whether the shader-resource and targetable texture must be separate textures. */
-			UE_DEPRECATED(5.0, "bForceSeparateTargetAndShaderResource has been deprecated. Do not use.")
-			uint8 bForceSeparateTargetAndShaderResource : 1;
-			/** Whether the shader-resource and targetable texture must be the same resource. */
-			UE_DEPRECATED(5.0, "bForceSharedTargetAndShaderResource has been deprecated. Do not use.")
-			uint8 bForceSharedTargetAndShaderResource : 1;
-			/** automatically set to writable via barrier during */
-			UE_DEPRECATED(5.0, "AutoWritable has been deprecated. Do not use.")
-			uint8 AutoWritable : 1;
-			/** create render target write mask (supported only on specific platforms) */
-			UE_DEPRECATED(5.0, "bCreateRenderTargetWriteMask has been deprecated. Do not use.")
-			uint8 bCreateRenderTargetWriteMask : 1;
-			/** create render target fmask (supported only on specific platforms) */
-			UE_DEPRECATED(5.0, "bCreateRenderTargetFmask has been deprecated. Do not use.")
-			uint8 bCreateRenderTargetFmask : 1;
-			/** Unused flag. */
-			uint8 bReserved0 : 1;
+			/** Unused flags. */
+			uint8 bReserved0 : 6;
 		};
 		uint8 PackedBits;
 	};
 };
 
-/** Enum to select between the two RHI textures on a pooled render target. */
-enum class UE_DEPRECATED(5.0, "ERenderTargetTexture is deprecated. IPooledRenderTarget now only has shared texture.") ERenderTargetTexture : uint8
-{
-	/** Maps to the targetable RHI texture on a pooled render target item. */
-	Targetable,
-
-	/** Maps to the shader resource RHI texture on a pooled render target item. */
-	ShaderResource
-};
-
 /**
  * Single render target item consists of a render surface and its resolve texture, Render thread side
  */
@@ -482,13 +447,6 @@ struct FSceneRenderTargetItem
 		,	UAV(InUAV)
 	{}
 
-PRAGMA_DISABLE_DEPRECATION_WARNINGS
-	FSceneRenderTargetItem(const FSceneRenderTargetItem&) = default;
-	FSceneRenderTargetItem(FSceneRenderTargetItem&&) = default;
-	FSceneRenderTargetItem& operator=(const FSceneRenderTargetItem&) = default;
-	FSceneRenderTargetItem& operator=(FSceneRenderTargetItem&&) = default;
-PRAGMA_ENABLE_DEPRECATION_WARNINGS
-
 	void SafeRelease()
 	{
 		TargetableTexture.SafeRelease();
@@ -505,11 +463,6 @@ PRAGMA_ENABLE_DEPRECATION_WARNINGS
 
 	FRHITexture* GetRHI() const { return TargetableTexture; }
 
-	UE_DEPRECATED(5.0, "GetRHI with ERenderTargetTexture is deprecated. Use GetRHI() instead.")
-PRAGMA_DISABLE_DEPRECATION_WARNINGS
-	FRHITexture* GetRHI(ERenderTargetTexture) const { return TargetableTexture; }
-PRAGMA_ENABLE_DEPRECATION_WARNINGS
-
 	/** The 2D or cubemap texture that may be used as a render or depth-stencil target. */
 	FTextureRHIRef TargetableTexture;
 
@@ -518,22 +471,6 @@ PRAGMA_ENABLE_DEPRECATION_WARNINGS
 	
 	/** only created if requested through the flag. */
 	FUnorderedAccessViewRHIRef UAV;
-
-	UE_DEPRECATED(5.0, "RTWriteMaskSRV is deprecated. Use RDG instead.")
-	FShaderResourceViewRHIRef RTWriteMaskSRV;
-
-	UE_DEPRECATED(5.0, "FmaskSRV is deprecated. Use RDG instead.")
-	FShaderResourceViewRHIRef FmaskSRV;
-
-	/** only created if requested through meta data access flags */
-	UE_DEPRECATED(5.0, "HTileUAV is deprecated. Use RDG instead.")
-	FUnorderedAccessViewRHIRef HTileUAV;
-
-	UE_DEPRECATED(5.0, "HTileSRV is deprecated. Use RDG instead.")
-	FShaderResourceViewRHIRef  HTileSRV;
-
-	UE_DEPRECATED(5.0, "StencilUAV is deprecated. Use RDG instead.")
-	FUnorderedAccessViewRHIRef StencilUAV;
 };
 
 /**
@@ -576,22 +513,7 @@ struct IPooledRenderTarget
 
 	FRHITexture* GetRHI() const { return RenderTargetItem.TargetableTexture; }
 
-	UE_DEPRECATED(5.0, "Use GetRHI instead.")
-	FRHITexture* GetTargetableRHI() const { return GetRHI(); }
-
-	UE_DEPRECATED(5.0, "Use GetRHI instead.")
-	FRHITexture* GetShaderResourceRHI() const { return GetRHI(); }
-
-	UE_DEPRECATED(5.0, "Use GetRHI without ERenderTargetTexture instead.")
-PRAGMA_DISABLE_DEPRECATION_WARNINGS
-	FRHITexture* GetRHI(ERenderTargetTexture) const { return GetRHI(); }
-PRAGMA_ENABLE_DEPRECATION_WARNINGS
-
-	UE_DEPRECATED(5.0, "SetDebugName has been deprecated.")
-	virtual void SetDebugName(const TCHAR* InName) {}
-
 protected:
-
 	/** The internal references to the created render target */
 	FSceneRenderTargetItem RenderTargetItem;
 };
@@ -788,23 +710,11 @@ public:
 	virtual class FSceneViewStateInterface* AllocateViewState(ERHIFeatureLevel::Type FeatureLevel) = 0;
 	virtual class FSceneViewStateInterface* AllocateViewState(ERHIFeatureLevel::Type FeatureLevel, FSceneViewStateInterface* ShareOriginTarget) = 0;
 
-	UE_DEPRECATED(5.0, "AllocateViewState must be called with an appropriate RHI Feature Level")
-	inline FSceneViewStateInterface* AllocateViewState()
-	{
-		return AllocateViewState(GMaxRHIFeatureLevel);
-	}
-
 	/** @return The number of lights that affect a primitive. */
 	virtual uint32 GetNumDynamicLightsAffectingPrimitive(const class FPrimitiveSceneInfo* PrimitiveSceneInfo,const class FLightCacheInterface* LCI) = 0;
 
-	UE_DEPRECATED(5.0, "ReallocateSceneRenderTargets is deprecated and no longer necessary. Please remove.")
-	virtual void ReallocateSceneRenderTargets() {}
-
 	virtual void OnWorldCleanup(UWorld* World, bool bSessionEnded, bool bCleanupResources, bool bWorldChanged) = 0;
 
-	UE_DEPRECATED(5.0, "SceneRenderTargetsSetBufferSize is deprecated and no longer necessary. Please remove.")
-	virtual void SceneRenderTargetsSetBufferSize(uint32 SizeX, uint32 SizeY) {}
-
 	virtual void InitializeSystemTextures(FRHICommandListImmediate& RHICmdList) = 0;
 
 	/** Draws a tile mesh element with the specified view. */
diff --git a/Engine/Source/Runtime/RenderCore/Public/RenderingThread.h b/Engine/Source/Runtime/RenderCore/Public/RenderingThread.h
--- a/Engine/Source/Runtime/RenderCore/Public/RenderingThread.h
+++ b/Engine/Source/Runtime/RenderCore/Public/RenderingThread.h
@@ -307,9 +307,6 @@ class RENDERCORE_API FDeferredCleanupInterface
 {
 public:
 	virtual ~FDeferredCleanupInterface() {}
-
-	UE_DEPRECATED(4.20, "FinishCleanup is deprecated. Use RAII in the destructor instead.")
-	virtual void FinishCleanup() final {}
 };
 
 /**
diff --git a/Engine/Source/Runtime/RenderCore/Public/VirtualTexturing.h b/Engine/Source/Runtime/RenderCore/Public/VirtualTexturing.h
--- a/Engine/Source/Runtime/RenderCore/Public/VirtualTexturing.h
+++ b/Engine/Source/Runtime/RenderCore/Public/VirtualTexturing.h
@@ -183,9 +183,6 @@ class IVirtualTextureFinalizer
 {
 public:
 	virtual void Finalize(FRDGBuilder& GraphBuilder) = 0;
-
-	UE_DEPRECATED(5.0, "This method has been refactored to use an FRDGBuilder instead.")
-	virtual void Finalize(FRHICommandListImmediate& RHICmdList) {}
 };
 
 enum class EVTRequestPageStatus
diff --git a/Engine/Source/Runtime/Renderer/Private/GPUScene.cpp b/Engine/Source/Runtime/Renderer/Private/GPUScene.cpp
--- a/Engine/Source/Runtime/Renderer/Private/GPUScene.cpp
+++ b/Engine/Source/Runtime/Renderer/Private/GPUScene.cpp
@@ -1632,11 +1632,6 @@ void FGPUScene::UploadDynamicPrimitiveShaderDataForViewInternal(FRDGBuilder& Gra
 	}
 }
 
-void AddPrimitiveToUpdateGPU(FScene& Scene, int32 PrimitiveId)
-{
-	Scene.GPUScene.AddPrimitiveToUpdate(PrimitiveId, EPrimitiveDirtyState::ChangedAll);
-}
-
 void FGPUScene::AddPrimitiveToUpdate(int32 PrimitiveId, EPrimitiveDirtyState DirtyState)
 {
 	LLM_SCOPE_BYTAG(GPUScene);
diff --git a/Engine/Source/Runtime/Renderer/Private/GPUScene.h b/Engine/Source/Runtime/Renderer/Private/GPUScene.h
--- a/Engine/Source/Runtime/Renderer/Private/GPUScene.h
+++ b/Engine/Source/Runtime/Renderer/Private/GPUScene.h
@@ -15,11 +15,6 @@ class FRHICommandList;
 class FScene;
 class FViewInfo;
 class FLightSceneInfoCompact;
-
-UE_DEPRECATED(5.0, "Use GPUScene::AddPrimitiveToUpdate instead.4")
-extern RENDERER_API void AddPrimitiveToUpdateGPU(FScene& Scene, int32 PrimitiveId);
-
-
 class FGPUScene;
 class FGPUSceneDynamicContext;
 
diff --git a/Engine/Source/Runtime/Renderer/Private/MeshPassProcessor.cpp b/Engine/Source/Runtime/Renderer/Private/MeshPassProcessor.cpp
--- a/Engine/Source/Runtime/Renderer/Private/MeshPassProcessor.cpp
+++ b/Engine/Source/Runtime/Renderer/Private/MeshPassProcessor.cpp
@@ -1869,17 +1869,6 @@ ERasterizerCullMode FMeshPassProcessor::ComputeMeshCullMode(const FMaterial& InM
 	return bMeshRenderTwoSided ? CM_None : (bInReverseCullModeOverride ? CM_CCW : CM_CW);
 }
 
-void FMeshPassProcessor::GetDrawCommandPrimitiveId(
-	const FPrimitiveSceneInfo* RESTRICT PrimitiveSceneInfo,
-	const FMeshBatchElement& BatchElement,
-	int32& DrawPrimitiveId,
-	int32& ScenePrimitiveId) const
-{
-	FMeshDrawCommandPrimitiveIdInfo PrimitiveIdInfo = GetDrawCommandPrimitiveId(PrimitiveSceneInfo, BatchElement);
-	DrawPrimitiveId = PrimitiveIdInfo.DrawPrimitiveId;
-	ScenePrimitiveId = PrimitiveIdInfo.ScenePrimitiveId;
-}
-
 FMeshDrawCommandPrimitiveIdInfo FMeshPassProcessor::GetDrawCommandPrimitiveId(
 	const FPrimitiveSceneInfo* RESTRICT PrimitiveSceneInfo,
 	const FMeshBatchElement& BatchElement) const
diff --git a/Engine/Source/Runtime/Renderer/Private/PostProcess/SceneFilterRendering.h b/Engine/Source/Runtime/Renderer/Private/PostProcess/SceneFilterRendering.h
--- a/Engine/Source/Runtime/Renderer/Private/PostProcess/SceneFilterRendering.h
+++ b/Engine/Source/Runtime/Renderer/Private/PostProcess/SceneFilterRendering.h
@@ -13,6 +13,7 @@
 #include "Shader.h"
 #include "PostProcess/SceneRenderTargets.h"
 #include "ShaderParameterUtils.h"
+#include "SceneRenderTargetParameters.h"
 
 /** Uniform buffer for computing the vertex positional and UV adjustments in the vertex shader. */
 BEGIN_GLOBAL_SHADER_PARAMETER_STRUCT( FDrawRectangleParameters, RENDERER_API)
diff --git a/Engine/Source/Runtime/Renderer/Private/PostProcess/SceneRenderTargets.h b/Engine/Source/Runtime/Renderer/Private/PostProcess/SceneRenderTargets.h
--- a/Engine/Source/Runtime/Renderer/Private/PostProcess/SceneRenderTargets.h
+++ b/Engine/Source/Runtime/Renderer/Private/PostProcess/SceneRenderTargets.h
@@ -2,7 +2,9 @@
 
 #pragma once
 
+#if UE_ENABLE_INCLUDE_ORDER_DEPRECATED_IN_5_2
 #include "SceneTextures.h"
+#endif
 
 /*
 * Stencil layout during basepass / deferred decals:
@@ -50,59 +52,3 @@
 #define GET_STENCIL_MOBILE_SM_MASK(Value) uint8(((Value) & 0x3) << 1)
 // Sky material mask - bit 3
 #define STENCIL_MOBILE_SKY_MASK uint8(1 << 3)
-
-class FSceneRenderTargets
-{
-public:
-	UE_DEPRECATED(5.0, "FSceneRenderTargets is now deprecated from the RDG refactor. FSceneTextures should be used instead.")
-	static FSceneRenderTargets& Get()
-	{
-		static FSceneRenderTargets Instance;
-		return Instance;
-	}
-
-	UE_DEPRECATED(5.0, "FSceneRenderTargets is now deprecated from the RDG refactor. FSceneTextures should be used instead.")
-	static FSceneRenderTargets& Get(FRHICommandListImmediate&)
-	{
-		PRAGMA_DISABLE_DEPRECATION_WARNINGS
-		return Get();
-		PRAGMA_ENABLE_DEPRECATION_WARNINGS
-	}
-
-	UE_DEPRECATED(5.0, "FSceneRenderTargets is now deprecated from the RDG refactor. FSceneTextures should be used instead.")
-	static FClearValueBinding GetDefaultColorClear()
-	{
-		return FSceneTexturesConfig::Get().ColorClearValue;
-	}
-
-	UE_DEPRECATED(5.0, "FSceneRenderTargets is now deprecated from the RDG refactor. FSceneTextures should be used instead.")
-	static FClearValueBinding GetDefaultDepthClear()
-	{
-		return FSceneTexturesConfig::Get().DepthClearValue;
-	}
-
-	UE_DEPRECATED(5.0, "FSceneRenderTargets is now deprecated from the RDG refactor. FSceneTextures should be used instead.")
-	static FIntPoint GetBufferSizeXY()
-	{
-		return FSceneTexturesConfig::Get().Extent;
-	}
-
-	UE_DEPRECATED(5.0, "FSceneRenderTargets is now deprecated from the RDG refactor. FSceneTextures should be used instead.")
-	static int32 GetMSAACount()
-	{
-		return FSceneTexturesConfig::Get().NumSamples;
-	}
-
-	UE_DEPRECATED(5.0, "FSceneRenderTargets is now deprecated from the RDG refactor. FSceneTextures should be used instead.")
-	static ERHIFeatureLevel::Type GetCurrentFeatureLevel()
-	{
-		return FSceneTexturesConfig::Get().FeatureLevel;
-	}
-
-	UE_DEPRECATED(5.0, "FSceneRenderTargets is now deprecated from the RDG refactor. FSceneTextures should be used instead.")
-	static TRefCountPtr<IPooledRenderTarget> GetSceneColor()
-	{
-		checkNoEntry();
-		return nullptr;
-	}
-};
diff --git a/Engine/Source/Runtime/Renderer/Public/MeshPassProcessor.h b/Engine/Source/Runtime/Renderer/Public/MeshPassProcessor.h
--- a/Engine/Source/Runtime/Renderer/Public/MeshPassProcessor.h
+++ b/Engine/Source/Runtime/Renderer/Public/MeshPassProcessor.h
@@ -1905,30 +1905,6 @@ public:
 		return ViewUniformBuffer;
 	}
 
-	UE_DEPRECATED(5.0, "SetInstancedViewUniformBuffer is deprecated. Use View.GetShaderParameters() and bind on an RDG pass instead.")
-	FORCEINLINE_DEBUGGABLE void SetInstancedViewUniformBuffer(const TUniformBufferRef<FInstancedViewUniformShaderParameters>& InViewUniformBuffer)
-	{
-		InstancedViewUniformBuffer = InViewUniformBuffer;
-	}
-
-	UE_DEPRECATED(5.0, "GetInstancedViewUniformBuffer is deprecated. Use View.GetShaderParameters() and bind on an RDG pass instead.")
-	FORCEINLINE_DEBUGGABLE const FRHIUniformBuffer* GetInstancedViewUniformBuffer() const
-	{
-		return InstancedViewUniformBuffer != nullptr ? InstancedViewUniformBuffer : ViewUniformBuffer;
-	}
-
-	UE_DEPRECATED(5.0, "SetPassUniformBuffer is deprecated. Use a static uniform buffer and bind on an RDG pass instead.")
-	FORCEINLINE_DEBUGGABLE void SetPassUniformBuffer(const FUniformBufferRHIRef& InPassUniformBuffer)
-	{
-		PassUniformBuffer = InPassUniformBuffer;
-	}
-
-	UE_DEPRECATED(5.0, "GetPassUniformBuffer is deprecated. Use a static uniform buffer and bind on an RDG pass instead.")
-	FORCEINLINE_DEBUGGABLE FRHIUniformBuffer* GetPassUniformBuffer() const
-	{
-		return PassUniformBuffer;
-	}
-
 	UE_DEPRECATED(5.1, "GetNaniteUniformBuffer is deprecated. Use a static uniform buffer and bind on an RDG pass instead.")
 	FORCEINLINE_DEBUGGABLE void SetNaniteUniformBuffer(FRHIUniformBuffer* InNaniteUniformBuffer)
 	{
@@ -2066,13 +2042,6 @@ public:
 		TArray<FPSOPrecacheData>& PSOInitializers);
 
 protected:
-	UE_DEPRECATED(5.0, "This version of GetDrawCommandPrimitiveId is deprecated. Use the below function instead.")
-	RENDERER_API void GetDrawCommandPrimitiveId(
-		const FPrimitiveSceneInfo* RESTRICT PrimitiveSceneInfo,
-		const FMeshBatchElement& BatchElement,
-		int32& DrawPrimitiveId,
-		int32& ScenePrimitiveId) const;
-
 	RENDERER_API FMeshDrawCommandPrimitiveIdInfo GetDrawCommandPrimitiveId(
 		const FPrimitiveSceneInfo* RESTRICT PrimitiveSceneInfo,
 		const FMeshBatchElement& BatchElement) const;
diff --git a/Engine/Source/Runtime/Renderer/Public/SceneRenderTargetParameters.h b/Engine/Source/Runtime/Renderer/Public/SceneRenderTargetParameters.h
--- a/Engine/Source/Runtime/Renderer/Public/SceneRenderTargetParameters.h
+++ b/Engine/Source/Runtime/Renderer/Public/SceneRenderTargetParameters.h
@@ -152,52 +152,3 @@ extern RENDERER_API FIntPoint GetSceneTextureExtent();
 /** Returns the feature level being used by the renderer. */
 UE_DEPRECATED(5.1, "Single pass multiple view family rendering makes this obsolete.  Use ViewFamily.SceneTexturesConfig.FeatureLevel instead.")
 extern RENDERER_API ERHIFeatureLevel::Type GetSceneTextureFeatureLevel();
-
-///////////////////////////////////////////////////////////////////////////
-// Deprecated APIs
-
-class FSceneRenderTargets;
-
-UE_DEPRECATED(5.0, "SetupSceneTextureUniforParameters now requires an FRDGBuilder.")
-inline void SetupSceneTextureUniformParameters(const FSceneRenderTargets&, ERHIFeatureLevel::Type, ESceneTextureSetupMode, FSceneTextureUniformParameters&)
-{
-	checkNoEntry();
-}
-
-UE_DEPRECATED(5.0, "CreateSceneTextureUniformBuffer now requires an FRDGBuilder.")
-inline TUniformBufferRef<FSceneTextureUniformParameters> CreateSceneTextureUniformBuffer(FRHIComputeCommandList&, ERHIFeatureLevel::Type, ESceneTextureSetupMode SetupMode = ESceneTextureSetupMode::All)
-{
-	checkNoEntry();
-	return {};
-}
-
-UE_DEPRECATED(5.0, "SetupMobileSceneTextureUniformParameters now requires an FRDGBuilder.")
-inline void SetupMobileSceneTextureUniformParameters(const FSceneRenderTargets&, EMobileSceneTextureSetupMode, FMobileSceneTextureUniformParameters&)
-{
-	checkNoEntry();
-}
-
-/** Creates the RHI mobile scene texture uniform buffer with passthrough RDG resources. */
-UE_DEPRECATED(5.0, "CreateMobileSceneTextureUniformBuffer now requires an FRDGBuilder.")
-inline TUniformBufferRef<FMobileSceneTextureUniformParameters> CreateMobileSceneTextureUniformBuffer(FRHIComputeCommandList&, EMobileSceneTextureSetupMode SetupMode = EMobileSceneTextureSetupMode::All)
-{
-	checkNoEntry();
-	return {};
-}
-
-UE_DEPRECATED(5.0, "Use CreateSceneTextureShaderParameters instead.")
-inline TRefCountPtr<FRHIUniformBuffer> CreateSceneTextureUniformBufferDependentOnShadingPath(FRHIComputeCommandList&, ERHIFeatureLevel::Type, ESceneTextureSetupMode SetupMode = ESceneTextureSetupMode::All)
-{
-	checkNoEntry();
-	return {};
-}
-
-UE_DEPRECATED(5.0, "IsSceneTexturesValid no longer requires a command list.")
-inline bool IsSceneTexturesValid(FRHICommandListImmediate&)
-{
-	PRAGMA_DISABLE_DEPRECATION_WARNINGS
-	return IsSceneTexturesValid();
-	PRAGMA_ENABLE_DEPRECATION_WARNINGS
-}
-
-//////////////////////////////////////////////////////////////////////////
\ No newline at end of file
