Add proper support for streamed geometries (geometries initialized as StreamingDestination) when dynamic ray tracing is enabled:
* Add FRayTracingGeometry::EGeometryStateFlags for tracking the state instead of several bools:
 - Add new state for tracking if the geometry was streamed in - its primary use is for dynamic ray tracing.
 - When ray tracing is toggled on at runtime we only recreate geometries that were marked as streamed in.
* Streamed geometries have to be created in a different way than regular geometries:
 - Creating geometry with StreamingDestination type will not allocate any memory for BLAS buffers.
 - Those can only be initialized by InitRHIForStreaming from temporary geometry with a StreamingSource type
 - Added FRayTracingGeometry::InitRHIForDynamicRayTracing() to handle both StreamingDestination and Rendering
  - For StreamingDestination it will do similar steps as normal mesh streaming does.
* Add missing RayTracingGeometry initialization when dynamic ray tracing was enabled but ray tracing was disabled.

diff --git a/Engine/Source/Runtime/Engine/Private/SkeletalMesh.cpp b/Engine/Source/Runtime/Engine/Private/SkeletalMesh.cpp
--- a/Engine/Source/Runtime/Engine/Private/SkeletalMesh.cpp
+++ b/Engine/Source/Runtime/Engine/Private/SkeletalMesh.cpp
@@ -6018,7 +6018,7 @@ FSkeletalMeshSceneProxy::FSkeletalMeshSceneProxy(const USkinnedMeshComponent* Co
 	}
 
 #if RHI_RAYTRACING
-	if (IsRayTracingEnabled(GetScene().GetShaderPlatform()))
+	if (IsRayTracingAllowed())
 	{
 		if (bRenderStatic)
 		{
diff --git a/Engine/Source/Runtime/Engine/Private/SkeletalRenderGPUSkin.cpp b/Engine/Source/Runtime/Engine/Private/SkeletalRenderGPUSkin.cpp
--- a/Engine/Source/Runtime/Engine/Private/SkeletalRenderGPUSkin.cpp
+++ b/Engine/Source/Runtime/Engine/Private/SkeletalRenderGPUSkin.cpp
@@ -885,7 +885,7 @@ void FSkeletalMeshObjectGPUSkin::UpdateRayTracingGeometry(FSkeletalMeshLODRender
 		{
 			EAccelerationStructureBuildMode BuildMode = bRequireRecreatingRayTracingGeometry ? EAccelerationStructureBuildMode::Build : EAccelerationStructureBuildMode::Update;
 			RayTracingUpdateQueue->Add(&RayTracingGeometry, RayTracingGeometryStructureSize, BuildMode);
-			RayTracingGeometry.bRequiresBuild = false;
+			RayTracingGeometry.SetRequiresBuild(false);
 		}
 		else
 		{
diff --git a/Engine/Source/Runtime/Engine/Private/SkeletalRenderStatic.cpp b/Engine/Source/Runtime/Engine/Private/SkeletalRenderStatic.cpp
--- a/Engine/Source/Runtime/Engine/Private/SkeletalRenderStatic.cpp
+++ b/Engine/Source/Runtime/Engine/Private/SkeletalRenderStatic.cpp
@@ -46,7 +46,7 @@ void FSkeletalMeshObjectStatic::InitResources(USkinnedMeshComponent* InMeshCompo
 			SkelLOD.InitResources(CompLODInfo);
 
 #if RHI_RAYTRACING
-			if (IsRayTracingEnabled() && SkelLOD.SkelMeshRenderData->bSupportRayTracing)
+			if (IsRayTracingAllowed() && SkelLOD.SkelMeshRenderData->bSupportRayTracing)
 			{
 				if (SkelLOD.SkelMeshRenderData->LODRenderData[LODIndex].NumReferencingStaticSkeletalMeshObjects == 0)
 				{
@@ -135,7 +135,7 @@ void FSkeletalMeshObjectStatic::ReleaseResources()
 		if (SkelLOD.SkelMeshRenderData->LODRenderData[LODIndex].GetNumVertices() > 0)
 		{
 #if RHI_RAYTRACING
-			if (IsRayTracingEnabled())
+			if (IsRayTracingAllowed())
 			{
 				if (SkelLOD.SkelMeshRenderData->LODRenderData[LODIndex].NumReferencingStaticSkeletalMeshObjects > 0)
 				{
diff --git a/Engine/Source/Runtime/Engine/Private/StaticMeshRender.cpp b/Engine/Source/Runtime/Engine/Private/StaticMeshRender.cpp
--- a/Engine/Source/Runtime/Engine/Private/StaticMeshRender.cpp
+++ b/Engine/Source/Runtime/Engine/Private/StaticMeshRender.cpp
@@ -432,7 +432,7 @@ void FStaticMeshSceneProxy::SetEvaluateWorldPositionOffsetInRayTracing(bool NewV
 	if (NewValue && !bDynamicRayTracingGeometry)
 	{
 		bDynamicRayTracingGeometry = true;
-		if (IsRayTracingEnabled())
+		if (IsRayTracingAllowed())
 		{
 			DynamicRayTracingGeometries.AddDefaulted(RenderData->LODResources.Num());
 
@@ -464,7 +464,7 @@ void FStaticMeshSceneProxy::SetEvaluateWorldPositionOffsetInRayTracing(bool NewV
 	else if (!NewValue && bDynamicRayTracingGeometry)
 	{
 		bDynamicRayTracingGeometry = false;
-		if (IsRayTracingEnabled())
+		if (IsRayTracingAllowed())
 		{
 			for (auto& Geometry : DynamicRayTracingGeometries)
 			{
@@ -743,7 +743,7 @@ bool FStaticMeshSceneProxy::GetMeshElement(
 void FStaticMeshSceneProxy::CreateRenderThreadResources()
 {
 #if RHI_RAYTRACING
-	if(IsRayTracingEnabled() && bSupportRayTracing)
+	if(IsRayTracingAllowed() && bSupportRayTracing)
 	{
 		if (bDynamicRayTracingGeometry)
 		{
diff --git a/Engine/Source/Runtime/Engine/Private/Streaming/SkeletalMeshUpdate.cpp b/Engine/Source/Runtime/Engine/Private/Streaming/SkeletalMeshUpdate.cpp
--- a/Engine/Source/Runtime/Engine/Private/Streaming/SkeletalMeshUpdate.cpp
+++ b/Engine/Source/Runtime/Engine/Private/Streaming/SkeletalMeshUpdate.cpp
@@ -205,7 +205,7 @@ void FSkeletalMeshStreamIn::DoFinishUpdate(const FContext& Context)
 
 #if RHI_RAYTRACING
 		// Must happen after the batched updates have been flushed
-		if (IsRayTracingEnabled())
+		if (IsRayTracingAllowed())
 		{
 			for (int32 LODIndex = PendingFirstLODIdx; LODIndex < CurrentFirstLODIdx; ++LODIndex)
 			{
@@ -364,7 +364,7 @@ void FSkeletalMeshStreamOut::ReleaseBuffers(const FContext& Context)
 			}
 
 #if RHI_RAYTRACING
-			if (IsRayTracingEnabled())
+			if (IsRayTracingAllowed())
 			{
 				if (RenderData->LODRenderData[LODIndex].bReferencedByStaticSkeletalMeshObjects_RenderThread)
 				{
diff --git a/Engine/Source/Runtime/Engine/Private/Streaming/StaticMeshUpdate.cpp b/Engine/Source/Runtime/Engine/Private/Streaming/StaticMeshUpdate.cpp
--- a/Engine/Source/Runtime/Engine/Private/Streaming/StaticMeshUpdate.cpp
+++ b/Engine/Source/Runtime/Engine/Private/Streaming/StaticMeshUpdate.cpp
@@ -237,12 +237,11 @@ void FStaticMeshStreamIn::DoFinishUpdate(const FContext& Context)
 #if RHI_RAYTRACING
 				if (IsRayTracingEnabled() && Context.Mesh->bSupportRayTracing &&
 					LODResource.VertexBuffers.StaticMeshVertexBuffer.GetNumVertices() > 0)
-				{
-					check(LODResource.RayTracingGeometry.RayTracingGeometryRHI != nullptr);
+				{					
 					check(IntermediateRayTracingGeometry[LODIdx].RayTracingGeometryRHI != nullptr);
 					LODResource.RayTracingGeometry.InitRHIForStreaming(IntermediateRayTracingGeometry[LODIdx].RayTracingGeometryRHI, Batcher);
 
-					LODResource.RayTracingGeometry.bRequiresBuild = IntermediateRayTracingGeometry[LODIdx].bRequiresBuild;
+					LODResource.RayTracingGeometry.SetRequiresBuild(IntermediateRayTracingGeometry[LODIdx].GetRequiresBuild());
 
 					IntermediateRayTracingGeometry[LODIdx].Initializer = {};
 					IntermediateRayTracingGeometry[LODIdx].RayTracingGeometryRHI.SafeRelease();				
@@ -253,7 +252,7 @@ void FStaticMeshStreamIn::DoFinishUpdate(const FContext& Context)
 
 #if RHI_RAYTRACING
 		// Must happen after the batched updates have been flushed
-		if (IsRayTracingEnabled() && Context.Mesh->bSupportRayTracing)
+		if (IsRayTracingAllowed() && Context.Mesh->bSupportRayTracing)
 		{
 			for (int32 LODIndex = PendingFirstLODIdx; LODIndex < CurrentFirstLODIdx; ++LODIndex)
 			{
@@ -273,8 +272,13 @@ void FStaticMeshStreamIn::DoFinishUpdate(const FContext& Context)
 						LODResource.SetupRayTracingGeometryInitializer(Initializer, Context.Mesh->GetFName());
 					}
 					LODResource.RayTracingGeometry.SetInitializer(Initializer);
+					LODResource.RayTracingGeometry.SetAsStreamedIn();			
 
-					LODResource.RayTracingGeometry.RequestBuildIfNeeded(ERTAccelerationStructureBuildPriority::Normal);
+					// Under very rare circumstances that we switch ray tracing on/off right in the middle of streaming RayTracingGeometryRHI might not be valid.
+					if (IsRayTracingEnabled() && LODResource.RayTracingGeometry.RayTracingGeometryRHI.IsValid())
+					{
+						LODResource.RayTracingGeometry.RequestBuildIfNeeded(ERTAccelerationStructureBuildPriority::Normal);
+					}
 				}
 			}
 
@@ -399,7 +403,7 @@ void FStaticMeshStreamOut::ReleaseRHIBuffers(const FContext& Context)
 			LODResource.ReleaseRHIForStreaming(Batcher);
 			
 #if RHI_RAYTRACING
-			if (IsRayTracingEnabled())
+			if (IsRayTracingAllowed())
 			{
 				LODResource.RayTracingGeometry.ReleaseRHIForStreaming(Batcher);
 			}
diff --git a/Engine/Source/Runtime/Landscape/Private/LandscapeRender.cpp b/Engine/Source/Runtime/Landscape/Private/LandscapeRender.cpp
--- a/Engine/Source/Runtime/Landscape/Private/LandscapeRender.cpp
+++ b/Engine/Source/Runtime/Landscape/Private/LandscapeRender.cpp
@@ -1170,7 +1170,7 @@ void FLandscapeComponentSceneProxy::CreateRenderThreadResources()
 #endif
 
 #if RHI_RAYTRACING
-	if (IsRayTracingEnabled())
+	if (IsRayTracingAllowed())
 	{
 		for (int32 SubY = 0; SubY < NumSubsections; SubY++)
 		{
diff --git a/Engine/Source/Runtime/RenderCore/Private/RayTracingGeometryManager.cpp b/Engine/Source/Runtime/RenderCore/Private/RayTracingGeometryManager.cpp
--- a/Engine/Source/Runtime/RenderCore/Private/RayTracingGeometryManager.cpp
+++ b/Engine/Source/Runtime/RenderCore/Private/RayTracingGeometryManager.cpp
@@ -123,7 +123,7 @@ void FRayTracingGeometryManager::Tick()
 		{
 			if (Geometry->RayTracingGeometryRHI == nullptr)
 			{
-				Geometry->InitRHI();
+				Geometry->InitRHIForDynamicRayTracing();
 			}
 		}
 	}
diff --git a/Engine/Source/Runtime/RenderCore/Private/RenderResource.cpp b/Engine/Source/Runtime/RenderCore/Private/RenderResource.cpp
--- a/Engine/Source/Runtime/RenderCore/Private/RenderResource.cpp
+++ b/Engine/Source/Runtime/RenderCore/Private/RenderResource.cpp
@@ -515,7 +515,7 @@ void FRayTracingGeometry::CreateRayTracingGeometryFromCPUData(TResourceArray<uin
 		Initializer.OfflineData = nullptr;
 	}
 
-	bRequiresBuild = Initializer.OfflineData == nullptr;		
+	SetRequiresBuild(Initializer.OfflineData == nullptr);
 	RayTracingGeometryRHI = RHICreateRayTracingGeometry(Initializer);
 }
 
@@ -523,13 +523,54 @@ void FRayTracingGeometry::RequestBuildIfNeeded(ERTAccelerationStructureBuildPrio
 {
 	RayTracingGeometryRHI->SetInitializer(Initializer);
 
-	if (bRequiresBuild)
+	if (GetRequiresBuild())
 	{
 		RayTracingBuildRequestIndex = GRayTracingGeometryManager.RequestBuildAccelerationStructure(this, InBuildPriority);
-		bRequiresBuild = false;
+		SetRequiresBuild(false);
 	}	
 }
 
+void FRayTracingGeometry::InitRHIForDynamicRayTracing()
+{
+	check(GetRayTracingMode() == ERayTracingMode::Dynamic);
+
+	// Streaming BLAS needs special handling to not get their "streaming" type wiped out as it will cause issues down the line.	
+	// We only have to do this if the geometry was marked to be streamed in.
+	// In that case we will recreate the geometry as-if it was streamed in.
+	if (EnumHasAnyFlags(GeometryState, FRayTracingGeometry::EGeometryStateFlags::StreamedIn))
+	{
+		// When a mesh is streamed in (FStaticMeshStreamIn::DoFinishUpdate) we update the geometry initializer using just streamed in VB/IB.
+		// That initializer sets a Rendering type but RHI object was created as StreamingDestination and we have a mismatch between geometry initializer and RHI initializer.
+		// It's not an issue unless we try to initialize the geometry again using the geometry's initializer.
+		// We need the current geometry and RHI object to be StreamingDestination so the streaming continues to work.
+		Initializer.Type = ERayTracingGeometryInitializerType::StreamingDestination;
+
+		// Creating RHI with StreamingDestination type will only initialize RHI object but will not created the underlying BLAS buffers.
+		InitRHI();
+		
+		// Here we simulate geometry streaming: create geometry with StreamingSource type to allocate BLAS buffers (1) and swap it with the current geometry (2).
+		// Follows the same pattern as: (1) FStaticMeshStreamIn::CreateBuffers_* (2) FStaticMeshStreamIn::DoFinishUpdate
+		// There is no other way to initialize BLAS buffers for the geometry that has a StreamingDestination type.
+		{			
+			TRHIResourceUpdateBatcher<1> Batcher;
+			FRayTracingGeometryInitializer IntermediateInitializer = Initializer;
+			IntermediateInitializer.Type = ERayTracingGeometryInitializerType::StreamingSource;
+
+			FRayTracingGeometryRHIRef IntermediateRayTracingGeometry = RHICreateRayTracingGeometry(IntermediateInitializer);
+			InitRHIForStreaming(IntermediateRayTracingGeometry, Batcher);
+
+			// When Batcher goes out of scope it will add commands to copy the BLAS buffers on RHI thread.
+			// We need to do it before we build the current geometry (also on RHI thread).
+		}
+		
+		RequestBuildIfNeeded(ERTAccelerationStructureBuildPriority::Normal);
+	}
+	else
+	{
+		InitRHI();
+	}
+}
+
 void FRayTracingGeometry::CreateRayTracingGeometry(ERTAccelerationStructureBuildPriority InBuildPriority)
 {
 	// Release previous RHI object if any
@@ -560,7 +601,16 @@ void FRayTracingGeometry::CreateRayTracingGeometry(ERTAccelerationStructureBuild
 	const bool bWithoutNativeResource = Initializer.Type == ERayTracingGeometryInitializerType::StreamingDestination;
 	if (bAllSegmentsAreValid)
 	{
-		bValid = !bWithoutNativeResource;
+		// Only geometries with StreamingDestination type are initially created in invalid state until they are streamed in (see InitRHIForStreaming).
+		if (bWithoutNativeResource)
+		{
+			EnumRemoveFlags(GeometryState, EGeometryStateFlags::Valid);
+		}
+		else
+		{
+			EnumAddFlags(GeometryState, EGeometryStateFlags::Valid);
+		}
+
 		if (IsRayTracingEnabled())
 		{
 			RayTracingGeometryRHI = RHICreateRayTracingGeometry(Initializer);
@@ -583,16 +633,16 @@ void FRayTracingGeometry::CreateRayTracingGeometry(ERTAccelerationStructureBuild
 				{
 					RayTracingBuildRequestIndex = GRayTracingGeometryManager.RequestBuildAccelerationStructure(this, InBuildPriority);
 				}
-				bRequiresBuild = false;
+				SetRequiresBuild(false);
 			}
 			else
 			{
-				bRequiresBuild = true;
+				SetRequiresBuild(true);
 			}
 		}
 		else
 		{
-			bRequiresBuild = false;
+			SetRequiresBuild(false);
 
 			// Offline data ownership is transferred to the RHI, which discards it after use.
 			// It is no longer valid to use it after this point.
@@ -603,7 +653,7 @@ void FRayTracingGeometry::CreateRayTracingGeometry(ERTAccelerationStructureBuild
 
 bool FRayTracingGeometry::IsValid() const
 {
-	return RayTracingGeometryRHI != nullptr && Initializer.TotalPrimitiveCount > 0 && bValid;
+	return RayTracingGeometryRHI != nullptr && Initializer.TotalPrimitiveCount > 0 && EnumHasAnyFlags(GeometryState, EGeometryStateFlags::Valid);
 }
 
 void FRayTracingGeometry::InitRHI()
diff --git a/Engine/Source/Runtime/RenderCore/Public/RenderResource.h b/Engine/Source/Runtime/RenderCore/Public/RenderResource.h
--- a/Engine/Source/Runtime/RenderCore/Public/RenderResource.h
+++ b/Engine/Source/Runtime/RenderCore/Public/RenderResource.h
@@ -743,7 +743,24 @@ public:
 
 	FRayTracingGeometryInitializer Initializer;
 	FRayTracingGeometryRHIRef RayTracingGeometryRHI;
-	bool bRequiresBuild = false;
+
+	// Flags for tracking the state of RayTracingGeometryRHI.
+	enum class EGeometryStateFlags : uint32
+	{
+		// Initial state when the geometry was not created or was created for streaming but not yet streamed in.
+		Invalid			= 0,
+
+		// If the geometry needs to be built.
+		RequiresBuild	= 1 << 0,
+
+		// If the geometry was successfully created or streamed in.
+		Valid			= 1 << 1,
+
+		// Special flag that is used when ray tracing is dynamic to mark the streamed geometry to be recreated when ray tracing is switched on.
+		// Only set when mesh streaming is used.
+		StreamedIn		= 1 << 2
+	};
+	FRIEND_ENUM_CLASS_FLAGS(EGeometryStateFlags);
 
 	/** LOD of the mesh associated with this ray tracing geometry object (-1 if unknown) */
 	int8 LODIndex = -1;
@@ -755,13 +772,42 @@ public:
 
 	bool IsValid() const;
 
+	void SetAsStreamedIn()
+	{
+		EnumAddFlags(GeometryState, EGeometryStateFlags::StreamedIn);
+	}
+
+	bool GetRequiresBuild() const 
+	{
+		return EnumHasAnyFlags(GeometryState, EGeometryStateFlags::RequiresBuild);
+	}
+
+	void SetRequiresBuild(bool bBuild)
+	{
+		if (bBuild)
+		{
+			EnumAddFlags(GeometryState, EGeometryStateFlags::RequiresBuild);
+		}
+		else
+		{
+			EnumRemoveFlags(GeometryState, EGeometryStateFlags::RequiresBuild);
+		}
+	}
+	
+	EGeometryStateFlags GetGeometryState() const 
+	{ 
+		return GeometryState; 
+	}
+
 	template <uint32 MaxNumUpdates>
 	void InitRHIForStreaming(FRHIRayTracingGeometry* IntermediateGeometry, TRHIResourceUpdateBatcher<MaxNumUpdates>& Batcher)
 	{
+		EnumAddFlags(GeometryState, EGeometryStateFlags::StreamedIn);
+
 		if (RayTracingGeometryRHI && IntermediateGeometry)
 		{
 			Batcher.QueueUpdateRequest(RayTracingGeometryRHI, IntermediateGeometry);
-			bValid = true;
+			EnumAddFlags(GeometryState, EGeometryStateFlags::Valid);
 		}
 	}
 
@@ -771,16 +817,20 @@ public:
 		Initializer = {};
 
 		RemoveBuildRequest();
+		EnumRemoveFlags(GeometryState, EGeometryStateFlags::StreamedIn);
 
 		if (RayTracingGeometryRHI)
 		{
 			Batcher.QueueUpdateRequest(RayTracingGeometryRHI, nullptr);
-			bValid = false;
+			EnumRemoveFlags(GeometryState, EGeometryStateFlags::Valid);
 		}		
 	}
 	void CreateRayTracingGeometryFromCPUData(TResourceArray<uint8>& OfflineData);
 	void RequestBuildIfNeeded(ERTAccelerationStructureBuildPriority InBuildPriority);
 
+	// That function is only supposed to be used when dynamic ray tracing is enabled
+	void InitRHIForDynamicRayTracing();
+
 	void CreateRayTracingGeometry(ERTAccelerationStructureBuildPriority InBuildPriority);
 
 	bool HasPendingBuildRequest() const
@@ -803,10 +853,14 @@ protected:
 	friend class FRayTracingGeometryManager;
 	int32 RayTracingBuildRequestIndex = INDEX_NONE;
 	int32 RayTracingGeometryHandle = INDEX_NONE; // Only valid when ray tracing is dynamic
-	bool bValid = false;
+	EGeometryStateFlags GeometryState = EGeometryStateFlags::Invalid;
 #endif
 };
 
+#if RHI_RAYTRACING
+ENUM_CLASS_FLAGS(FRayTracingGeometry::EGeometryStateFlags);
+#endif
+
 /**
  * A system for dynamically allocating GPU memory for vertices.
  */
diff --git a/Engine/Source/Runtime/Renderer/Private/RayTracing/RayTracingDynamicGeometry.cpp b/Engine/Source/Runtime/Renderer/Private/RayTracing/RayTracingDynamicGeometry.cpp
--- a/Engine/Source/Runtime/Renderer/Private/RayTracing/RayTracingDynamicGeometry.cpp
+++ b/Engine/Source/Runtime/Renderer/Private/RayTracing/RayTracingDynamicGeometry.cpp
@@ -343,16 +343,16 @@ void FRayTracingDynamicGeometryCollection::AddDynamicMeshBatchForGeometryUpdate(
 	{
 		checkf(Geometry.Initializer.OfflineData == nullptr, TEXT("Dynamic geometry is not expected to have offline acceleration structure data"));
 		Geometry.RayTracingGeometryRHI = RHICreateRayTracingGeometry(Geometry.Initializer);
-		Geometry.bRequiresBuild = true;
+		Geometry.SetRequiresBuild(true);
 	}
 
 	FRayTracingGeometryBuildParams Params;
 	Params.Geometry = Geometry.RayTracingGeometryRHI;
-	Params.BuildMode = Geometry.bRequiresBuild
+	Params.BuildMode = Geometry.GetRequiresBuild()
 		? EAccelerationStructureBuildMode::Build
 		: EAccelerationStructureBuildMode::Update;
 
-	Geometry.bRequiresBuild = false;
+	Geometry.SetRequiresBuild(false);
 
 	if (bUseSharedVertexBuffer)
 	{
