Set all parameters and resources for SetShaderParameters in a single RHI command
- Expands all RDG and non-RDG bindings into a list of RHIResource bindings.
- Future work could batch up all non-SetShaderParameters calls into similar data structures.

diff --git a/Engine/Source/Runtime/Apple/MetalRHI/Private/MetalCommands.cpp b/Engine/Source/Runtime/Apple/MetalRHI/Private/MetalCommands.cpp
--- a/Engine/Source/Runtime/Apple/MetalRHI/Private/MetalCommands.cpp
+++ b/Engine/Source/Runtime/Apple/MetalRHI/Private/MetalCommands.cpp
@@ -20,6 +20,7 @@
 #include "EngineGlobals.h"
 #include "PipelineStateCache.h"
 #include "RHICoreShader.h"
+#include "RHIShaderParametersShared.h"
 
 static const bool GUsesInvertedZ = true;
 
@@ -364,6 +365,30 @@ void FMetalRHICommandContext::RHISetShaderParameter(FRHIComputeShader* ComputeSh
 	}
 }
 
+void FMetalRHICommandContext::RHISetShaderParameters(FRHIGraphicsShader* Shader, TConstArrayView<uint8> InParametersData, TConstArrayView<FRHIShaderParameter> InParameters, TConstArrayView<FRHIShaderParameterResource> InResourceParameters, TConstArrayView<FRHIShaderParameterResource> InBindlessParameters)
+{
+	UE::RHICore::RHISetShaderParametersShared(
+		*this
+		, Shader
+		, InParametersData
+		, InParameters
+		, InResourceParameters
+		, InBindlessParameters
+	);
+}
+
+void FMetalRHICommandContext::RHISetShaderParameters(FRHIComputeShader* Shader, TConstArrayView<uint8> InParametersData, TConstArrayView<FRHIShaderParameter> InParameters, TConstArrayView<FRHIShaderParameterResource> InResourceParameters, TConstArrayView<FRHIShaderParameterResource> InBindlessParameters)
+{
+	UE::RHICore::RHISetShaderParametersShared(
+		*this
+		, Shader
+		, InParametersData
+		, InParameters
+		, InResourceParameters
+		, InBindlessParameters
+	);
+}
+
 void FMetalRHICommandContext::RHISetStencilRef(uint32 StencilRef)
 {
 	Context->GetCurrentState().SetStencilRef(StencilRef);
diff --git a/Engine/Source/Runtime/Apple/MetalRHI/Public/MetalRHIContext.h b/Engine/Source/Runtime/Apple/MetalRHI/Public/MetalRHIContext.h
--- a/Engine/Source/Runtime/Apple/MetalRHI/Public/MetalRHIContext.h
+++ b/Engine/Source/Runtime/Apple/MetalRHI/Public/MetalRHIContext.h
@@ -160,9 +160,11 @@ public:
 	virtual void RHISetShaderUniformBuffer(FRHIComputeShader* ComputeShader, uint32 BufferIndex, FRHIUniformBuffer* Buffer) final override;
 	
 	virtual void RHISetShaderParameter(FRHIGraphicsShader* Shader, uint32 BufferIndex, uint32 BaseIndex, uint32 NumBytes, const void* NewValue) final override;
-	
 	virtual void RHISetShaderParameter(FRHIComputeShader* ComputeShader, uint32 BufferIndex, uint32 BaseIndex, uint32 NumBytes, const void* NewValue) final override;
 	
+	virtual void RHISetShaderParameters(FRHIGraphicsShader* Shader, TConstArrayView<uint8> InParametersData, TConstArrayView<FRHIShaderParameter> InParameters, TConstArrayView<FRHIShaderParameterResource> InResourceParameters, TConstArrayView<FRHIShaderParameterResource> InBindlessParameters) final override;
+	virtual void RHISetShaderParameters(FRHIComputeShader* Shader, TConstArrayView<uint8> InParametersData, TConstArrayView<FRHIShaderParameter> InParameters, TConstArrayView<FRHIShaderParameterResource> InResourceParameters, TConstArrayView<FRHIShaderParameterResource> InBindlessParameters) final override;
+
 	virtual void RHISetStencilRef(uint32 StencilRef) final override;
 
 	virtual void RHISetBlendFactor(const FLinearColor& BlendFactor) final override;
diff --git a/Engine/Source/Runtime/D3D12RHI/Private/D3D12CommandContext.h b/Engine/Source/Runtime/D3D12RHI/Private/D3D12CommandContext.h
--- a/Engine/Source/Runtime/D3D12RHI/Private/D3D12CommandContext.h
+++ b/Engine/Source/Runtime/D3D12RHI/Private/D3D12CommandContext.h
@@ -533,6 +533,7 @@ public:
 	virtual void RHISetUAVParameter(FRHIComputeShader* ComputeShader, uint32 UAVIndex, FRHIUnorderedAccessView* UAV, uint32 InitialCount) final override;
 	virtual void RHISetShaderResourceViewParameter(FRHIComputeShader* ComputeShader, uint32 SamplerIndex, FRHIShaderResourceView* SRV) final override;
 	virtual void RHISetShaderUniformBuffer(FRHIComputeShader* ComputeShader, uint32 BufferIndex, FRHIUniformBuffer* Buffer) final override;
+	virtual void RHISetShaderParameters(FRHIComputeShader* Shader, TConstArrayView<uint8> InParametersData, TConstArrayView<FRHIShaderParameter> InParameters, TConstArrayView<FRHIShaderParameterResource> InResourceParameters, TConstArrayView<FRHIShaderParameterResource> InBindlessParameters) final override;
 	virtual void RHISetShaderParameter(FRHIComputeShader* ComputeShader, uint32 BufferIndex, uint32 BaseIndex, uint32 NumBytes, const void* NewValue) final override;
 	virtual void RHIPushEvent(const TCHAR* Name, FColor Color) final override;
 	virtual void RHIPopEvent() final override;
@@ -560,6 +561,7 @@ public:
 	virtual void RHISetShaderSampler(FRHIGraphicsShader* Shader, uint32 SamplerIndex, FRHISamplerState* NewState) final override;
 	virtual void RHISetShaderResourceViewParameter(FRHIGraphicsShader* Shader, uint32 SamplerIndex, FRHIShaderResourceView* SRV) final override;
 	virtual void RHISetShaderUniformBuffer(FRHIGraphicsShader* Shader, uint32 BufferIndex, FRHIUniformBuffer* Buffer) final override;
+	virtual void RHISetShaderParameters(FRHIGraphicsShader* Shader, TConstArrayView<uint8> InParametersData, TConstArrayView<FRHIShaderParameter> InParameters, TConstArrayView<FRHIShaderParameterResource> InResourceParameters, TConstArrayView<FRHIShaderParameterResource> InBindlessParameters) final override;
 	virtual void RHISetShaderParameter(FRHIGraphicsShader* Shader, uint32 BufferIndex, uint32 BaseIndex, uint32 NumBytes, const void* NewValue) final override;
 	virtual void RHISetStencilRef(uint32 StencilRef) final override;
 	virtual void RHISetBlendFactor(const FLinearColor& BlendFactor) final override;
@@ -799,6 +801,10 @@ public:
 	{
 		ContextRedirect(RHISetShaderParameter(ComputeShader, BufferIndex, BaseIndex, NumBytes, NewValue));
 	}
+	FORCEINLINE virtual void RHISetShaderParameters(FRHIComputeShader* Shader, TConstArrayView<uint8> InParametersData, TConstArrayView<FRHIShaderParameter> InParameters, TConstArrayView<FRHIShaderParameterResource> InResourceParameters, TConstArrayView<FRHIShaderParameterResource> InBindlessParameters) final override
+	{
+		ContextRedirect(RHISetShaderParameters(Shader, InParametersData, InParameters, InResourceParameters, InBindlessParameters));
+	}
 	FORCEINLINE virtual void RHIPushEvent(const TCHAR* Name, FColor Color) final override
 	{
 		ContextRedirect(RHIPushEvent(Name, Color));
@@ -902,6 +908,10 @@ public:
 	{
 		ContextRedirect(RHISetShaderParameter(Shader, BufferIndex, BaseIndex, NumBytes, NewValue));
 	}
+	FORCEINLINE virtual void RHISetShaderParameters(FRHIGraphicsShader* Shader, TConstArrayView<uint8> InParametersData, TConstArrayView<FRHIShaderParameter> InParameters, TConstArrayView<FRHIShaderParameterResource> InResourceParameters, TConstArrayView<FRHIShaderParameterResource> InBindlessParameters) final override
+	{
+		ContextRedirect(RHISetShaderParameters(Shader, InParametersData, InParameters, InResourceParameters, InBindlessParameters));
+	}
 	FORCEINLINE virtual void RHISetStencilRef(uint32 StencilRef) final override
 	{
 		ContextRedirect(RHISetStencilRef(StencilRef));
diff --git a/Engine/Source/Runtime/D3D12RHI/Private/D3D12Commands.cpp b/Engine/Source/Runtime/D3D12RHI/Private/D3D12Commands.cpp
--- a/Engine/Source/Runtime/D3D12RHI/Private/D3D12Commands.cpp
+++ b/Engine/Source/Runtime/D3D12RHI/Private/D3D12Commands.cpp
@@ -16,6 +16,7 @@ D3D12Commands.cpp: D3D RHI commands implementation.
 #include "ResolveShader.h"
 #include "SceneUtils.h"
 #include "RenderUtils.h"
+#include "RHIShaderParametersShared.h"
 
 int32 AFRSyncTemporalResources = 1;
 static FAutoConsoleVariableRef CVarSyncTemporalResources(
@@ -1141,6 +1142,30 @@ void FD3D12CommandContext::RHISetShaderUniformBuffer(FRHIComputeShader* ComputeS
 	DirtyUniformBuffers[SF_Compute] |= (1 << BufferIndex);
 }
 
+void FD3D12CommandContext::RHISetShaderParameters(FRHIGraphicsShader* Shader, TArrayView<const uint8> InParametersData, TArrayView<const FRHIShaderParameter> InParameters, TArrayView<const FRHIShaderParameterResource> InResourceParameters, TArrayView<const FRHIShaderParameterResource> InBindlessParameters)
+{
+	UE::RHICore::RHISetShaderParametersShared(
+		*this
+		, Shader
+		, InParametersData
+		, InParameters
+		, InResourceParameters
+		, InBindlessParameters
+	);
+}
+
+void FD3D12CommandContext::RHISetShaderParameters(FRHIComputeShader* Shader, TArrayView<const uint8> InParametersData, TArrayView<const FRHIShaderParameter> InParameters, TArrayView<const FRHIShaderParameterResource> InResourceParameters, TArrayView<const FRHIShaderParameterResource> InBindlessParameters)
+{
+	UE::RHICore::RHISetShaderParametersShared(
+		*this
+		, Shader
+		, InParametersData
+		, InParameters
+		, InResourceParameters
+		, InBindlessParameters
+	);
+}
+
 void FD3D12CommandContext::RHISetShaderParameter(FRHIGraphicsShader* ShaderRHI, uint32 BufferIndex, uint32 Offset, uint32 NumBytes, const void* NewValue)
 {
 	checkSlow(BufferIndex == 0);
diff --git a/Engine/Source/Runtime/NullDrv/Public/NullRHI.h b/Engine/Source/Runtime/NullDrv/Public/NullRHI.h
--- a/Engine/Source/Runtime/NullDrv/Public/NullRHI.h
+++ b/Engine/Source/Runtime/NullDrv/Public/NullRHI.h
@@ -535,6 +535,13 @@ public:
 	{
 	}
 
+	virtual void RHISetShaderParameters(FRHIGraphicsShader* Shader, TConstArrayView<uint8> InParametersData, TConstArrayView<FRHIShaderParameter> InParameters, TConstArrayView<FRHIShaderParameterResource> InResourceParameters, TConstArrayView<FRHIShaderParameterResource> InBindlessParameters) final override
+	{
+	}
+
+	virtual void RHISetShaderParameters(FRHIComputeShader* Shader, TConstArrayView<uint8> InParametersData, TConstArrayView<FRHIShaderParameter> InParameters, TConstArrayView<FRHIShaderParameterResource> InResourceParameters, TConstArrayView<FRHIShaderParameterResource> InBindlessParameters) final override
+	{
+	}
 
 	virtual void RHISetDepthStencilState(FRHIDepthStencilState* NewState, uint32 StencilRef) final override
 	{
diff --git a/Engine/Source/Runtime/OpenGLDrv/Private/OpenGLCommands.cpp b/Engine/Source/Runtime/OpenGLDrv/Private/OpenGLCommands.cpp
--- a/Engine/Source/Runtime/OpenGLDrv/Private/OpenGLCommands.cpp
+++ b/Engine/Source/Runtime/OpenGLDrv/Private/OpenGLCommands.cpp
@@ -16,6 +16,7 @@
 #include "OpenGLDrvPrivate.h"
 #include "RenderUtils.h"
 #include "RHICoreShader.h"
+#include "RHIShaderParametersShared.h"
 
 /*
 #define DECLARE_ISBOUNDSHADER(ShaderType) template <typename TShaderType> inline void ValidateBoundShader(TRefCountPtr<FOpenGLBoundShaderState> InBoundShaderState, FRHIGraphicsShader* GfxShader, TShaderType* ShaderTypeRHI) \
@@ -1139,6 +1140,30 @@ void FOpenGLDynamicRHI::RHISetShaderParameter(FRHIComputeShader* ComputeShaderRH
 	PendingState.LinkedProgramAndDirtyFlag = nullptr;
 }
 
+void FOpenGLDynamicRHI::RHISetShaderParameters(FRHIGraphicsShader* Shader, TConstArrayView<uint8> InParametersData, TConstArrayView<FRHIShaderParameter> InParameters, TConstArrayView<FRHIShaderParameterResource> InResourceParameters, TConstArrayView<FRHIShaderParameterResource> InBindlessParameters)
+{
+	UE::RHICore::RHISetShaderParametersShared(
+		*this
+		, Shader
+		, InParametersData
+		, InParameters
+		, InResourceParameters
+		, InBindlessParameters
+	);
+}
+
+void FOpenGLDynamicRHI::RHISetShaderParameters(FRHIComputeShader* Shader, TConstArrayView<uint8> InParametersData, TConstArrayView<FRHIShaderParameter> InParameters, TConstArrayView<FRHIShaderParameterResource> InResourceParameters, TConstArrayView<FRHIShaderParameterResource> InBindlessParameters)
+{
+	UE::RHICore::RHISetShaderParametersShared(
+		*this
+		, Shader
+		, InParametersData
+		, InParameters
+		, InResourceParameters
+		, InBindlessParameters
+	);
+}
+
 void FOpenGLDynamicRHI::RHISetDepthStencilState(FRHIDepthStencilState* NewStateRHI,uint32 StencilRef)
 {
 	VERIFY_GL_SCOPE();
diff --git a/Engine/Source/Runtime/OpenGLDrv/Public/OpenGLDrv.h b/Engine/Source/Runtime/OpenGLDrv/Public/OpenGLDrv.h
--- a/Engine/Source/Runtime/OpenGLDrv/Public/OpenGLDrv.h
+++ b/Engine/Source/Runtime/OpenGLDrv/Public/OpenGLDrv.h
@@ -482,6 +482,8 @@ public:
 	virtual void RHISetShaderUniformBuffer(FRHIComputeShader* ComputeShader, uint32 BufferIndex, FRHIUniformBuffer* Buffer) final override;
 	virtual void RHISetShaderParameter(FRHIGraphicsShader* Shader, uint32 BufferIndex, uint32 BaseIndex, uint32 NumBytes, const void* NewValue) final override;
 	virtual void RHISetShaderParameter(FRHIComputeShader* ComputeShader, uint32 BufferIndex, uint32 BaseIndex, uint32 NumBytes, const void* NewValue) final override;
+	virtual void RHISetShaderParameters(FRHIGraphicsShader* Shader, TConstArrayView<uint8> InParametersData, TConstArrayView<FRHIShaderParameter> InParameters, TConstArrayView<FRHIShaderParameterResource> InResourceParameters, TConstArrayView<FRHIShaderParameterResource> InBindlessParameters) final override;
+	virtual void RHISetShaderParameters(FRHIComputeShader* Shader, TConstArrayView<uint8> InParametersData, TConstArrayView<FRHIShaderParameter> InParameters, TConstArrayView<FRHIShaderParameterResource> InResourceParameters, TConstArrayView<FRHIShaderParameterResource> InBindlessParameters) final override;
 	virtual void RHISetDepthStencilState(FRHIDepthStencilState* NewState, uint32 StencilRef) final override;
 	virtual void RHISetStencilRef(uint32 StencilRef) final override;
 	virtual void RHISetBlendState(FRHIBlendState* NewState, const FLinearColor& BlendFactor) final override;
diff --git a/Engine/Source/Runtime/RHI/Public/RHICommandList.h b/Engine/Source/Runtime/RHI/Public/RHICommandList.h
--- a/Engine/Source/Runtime/RHI/Public/RHICommandList.h
+++ b/Engine/Source/Runtime/RHI/Public/RHICommandList.h
@@ -458,6 +458,13 @@ public:
 		return MemManager.Alloc(AllocSize, Alignment);
 	}
 
+	FORCEINLINE_DEBUGGABLE void* AllocCopy(const void* InSourceData, int64 AllocSize, int64 Alignment)
+	{
+		void* NewData = Alloc(AllocSize, Alignment);
+		FMemory::Memcpy(NewData, InSourceData, AllocSize);
+		return NewData;
+	}
+
 	template <typename T>
 	FORCEINLINE_DEBUGGABLE void* Alloc()
 	{
@@ -473,9 +480,13 @@ public:
 	template <typename T>
 	FORCEINLINE_DEBUGGABLE const TArrayView<T> AllocArray(const TArrayView<T> InArray)
 	{
+		if (InArray.Num() == 0)
+		{
+			return TArrayView<T>();
+		}
+
 		// @todo static_assert(TIsTrivial<T>::Value, "Only trivially constructible / copyable types can be used in RHICmdList.");
-		void* NewArray = Alloc(InArray.Num() * sizeof(T), alignof(T));
-		FMemory::Memcpy(NewArray, InArray.GetData(), InArray.Num() * sizeof(T));
+		void* NewArray = AllocCopy(InArray.GetData(), InArray.Num() * sizeof(T), alignof(T));
 		return TArrayView<T>((T*) NewArray, InArray.Num());
 	}
 
@@ -1057,6 +1068,31 @@ FRHICOMMAND_MACRO(FRHICommandSetStencilRef)
 	RHI_API void Execute(FRHICommandListBase& CmdList);
 };
 
+FRHICOMMAND_MACRO_TPL(TRHIShader, FRHICommandSetShaderParameters)
+{
+	TRHIShader* Shader;
+	TConstArrayView<uint8> ParametersData;
+	TConstArrayView<FRHIShaderParameter> Parameters;
+	TConstArrayView<FRHIShaderParameterResource> ResourceParameters;
+	TConstArrayView<FRHIShaderParameterResource> BindlessParameters;
+
+	FORCEINLINE_DEBUGGABLE FRHICommandSetShaderParameters(
+		TRHIShader* InShader
+		, TConstArrayView<uint8> InParametersData
+		, TConstArrayView<FRHIShaderParameter> InParameters
+		, TConstArrayView<FRHIShaderParameterResource> InResourceParameters
+		, TConstArrayView<FRHIShaderParameterResource> InBindlessParameters
+	)
+		: Shader(InShader)
+		, ParametersData(InParametersData)
+		, Parameters(InParameters)
+		, ResourceParameters(InResourceParameters)
+		, BindlessParameters(InBindlessParameters)
+	{
+	}
+	RHI_API void Execute(FRHICommandListBase& CmdList);
+};
+
 FRHICOMMAND_MACRO_TPL(TRHIShader, FRHICommandSetShaderParameter)
 {
 	TRHIShader* Shader;
@@ -2222,6 +2258,7 @@ FRHICOMMAND_MACRO(FRHICommandSetRayTracingBindings)
 };
 #endif // RHI_RAYTRACING
 
+template<> RHI_API void FRHICommandSetShaderParameters           <FRHIComputeShader>::Execute(FRHICommandListBase& CmdList);
 template<> RHI_API void FRHICommandSetShaderParameter            <FRHIComputeShader>::Execute(FRHICommandListBase& CmdList);
 template<> RHI_API void FRHICommandSetShaderUniformBuffer        <FRHIComputeShader>::Execute(FRHICommandListBase& CmdList);
 template<> RHI_API void FRHICommandSetShaderTexture              <FRHIComputeShader>::Execute(FRHICommandListBase& CmdList);
@@ -2321,6 +2358,30 @@ public:
 		SetLocalShaderUniformBuffer(Shader.GetReference(), BaseIndex, UniformBuffer);
 	}
 
+	FORCEINLINE_DEBUGGABLE void SetShaderParameters(
+		FRHIComputeShader* InShader
+		, TConstArrayView<uint8> InParametersData
+		, TConstArrayView<FRHIShaderParameter> InParameters
+		, TConstArrayView<FRHIShaderParameterResource> InResourceParameters
+		, TConstArrayView<FRHIShaderParameterResource> InBindlessParameters
+	)
+	{
+		ValidateBoundShader(InShader);
+		if (Bypass())
+		{
+			GetComputeContext().RHISetShaderParameters(InShader, InParametersData, InParameters, InResourceParameters, InBindlessParameters);
+			return;
+		}
+
+		ALLOC_COMMAND(FRHICommandSetShaderParameters<FRHIComputeShader>)(
+			InShader
+			, AllocArray(InParametersData)
+			, AllocArray(InParameters)
+			, AllocArray(InResourceParameters)
+			, AllocArray(InBindlessParameters)
+		);
+	}
+
 	FORCEINLINE_DEBUGGABLE void SetShaderParameter(FRHIComputeShader* Shader, uint32 BufferIndex, uint32 BaseIndex, uint32 NumBytes, const void* NewValue)
 	{
 		ValidateBoundShader(Shader);
@@ -2915,6 +2976,7 @@ PRAGMA_ENABLE_DEPRECATION_WARNINGS
 	}
 };
 
+template<> RHI_API void FRHICommandSetShaderParameters           <FRHIGraphicsShader>::Execute(FRHICommandListBase& CmdList);
 template<> RHI_API void FRHICommandSetShaderParameter            <FRHIGraphicsShader>::Execute(FRHICommandListBase& CmdList);
 template<> RHI_API void FRHICommandSetShaderUniformBuffer        <FRHIGraphicsShader>::Execute(FRHICommandListBase& CmdList);
 template<> RHI_API void FRHICommandSetShaderTexture              <FRHIGraphicsShader>::Execute(FRHICommandListBase& CmdList);
@@ -3016,6 +3078,32 @@ public:
 		SetShaderUniformBuffer(Shader.GetReference(), BaseIndex, UniformBuffer);
 	}
 
+	using FRHIComputeCommandList::SetShaderParameters;
+
+	FORCEINLINE_DEBUGGABLE void SetShaderParameters(
+		FRHIGraphicsShader* InShader
+		, TConstArrayView<uint8> InParametersData
+		, TConstArrayView<FRHIShaderParameter> InParameters
+		, TConstArrayView<FRHIShaderParameterResource> InResourceParameters
+		, TConstArrayView<FRHIShaderParameterResource> InBindlessParameters
+	)
+	{
+		ValidateBoundShader(InShader);
+		if (Bypass())
+		{
+			GetContext().RHISetShaderParameters(InShader, InParametersData, InParameters, InResourceParameters, InBindlessParameters);
+			return;
+		}
+
+		ALLOC_COMMAND(FRHICommandSetShaderParameters<FRHIGraphicsShader>)(
+			InShader
+			, AllocArray(InParametersData)
+			, AllocArray(InParameters)
+			, AllocArray(InResourceParameters)
+			, AllocArray(InBindlessParameters)
+			);
+	}
+
 	using FRHIComputeCommandList::SetShaderParameter;
 
 	FORCEINLINE_DEBUGGABLE void SetShaderParameter(FRHIGraphicsShader* Shader, uint32 BufferIndex, uint32 BaseIndex, uint32 NumBytes, const void* NewValue)
diff --git a/Engine/Source/Runtime/RHI/Public/RHICommandListCommandExecutes.inl b/Engine/Source/Runtime/RHI/Public/RHICommandListCommandExecutes.inl
--- a/Engine/Source/Runtime/RHI/Public/RHICommandListCommandExecutes.inl
+++ b/Engine/Source/Runtime/RHI/Public/RHICommandListCommandExecutes.inl
@@ -143,6 +143,18 @@ void FRHICommandSetStencilRef::Execute(FRHICommandListBase& CmdList)
 	INTERNAL_DECORATOR(RHISetStencilRef)(StencilRef);
 }
 
+template<> RHI_API void FRHICommandSetShaderParameters<FRHIComputeShader>::Execute(FRHICommandListBase& CmdList)
+{
+	RHISTAT(SetShaderParameters);
+	INTERNAL_DECORATOR_COMPUTE(RHISetShaderParameters)(Shader, ParametersData, Parameters, ResourceParameters, BindlessParameters);
+}
+
+template<> RHI_API void FRHICommandSetShaderParameters<FRHIGraphicsShader>::Execute(FRHICommandListBase& CmdList)
+{
+	RHISTAT(SetShaderParameters);
+	INTERNAL_DECORATOR(RHISetShaderParameters)(Shader, ParametersData, Parameters, ResourceParameters, BindlessParameters);
+}
+
 template<> RHI_API void FRHICommandSetShaderParameter<FRHIComputeShader>::Execute(FRHICommandListBase& CmdList)
 {
 	RHISTAT(SetShaderParameter);
diff --git a/Engine/Source/Runtime/RHI/Public/RHIContext.h b/Engine/Source/Runtime/RHI/Public/RHIContext.h
--- a/Engine/Source/Runtime/RHI/Public/RHIContext.h
+++ b/Engine/Source/Runtime/RHI/Public/RHIContext.h
@@ -301,6 +301,8 @@ public:
 
 	virtual void RHISetShaderUniformBuffer(FRHIComputeShader* ComputeShader, uint32 BufferIndex, FRHIUniformBuffer* Buffer) = 0;
 
+	virtual void RHISetShaderParameters(FRHIComputeShader* ComputeShader, TConstArrayView<uint8> InParametersData, TConstArrayView<FRHIShaderParameter> InParameters, TConstArrayView<FRHIShaderParameterResource> InResourceParameters, TConstArrayView<FRHIShaderParameterResource> InBindlessParameters) = 0;
+
 	virtual void RHISetShaderParameter(FRHIComputeShader* ComputeShader, uint32 BufferIndex, uint32 BaseIndex, uint32 NumBytes, const void* NewValue) = 0;
 
 	virtual void RHISetStaticUniformBuffers(const FUniformBufferStaticBindings& InUniformBuffers)
@@ -674,6 +676,7 @@ public:
 	using IRHIComputeContext::RHISetUAVParameter;
 	using IRHIComputeContext::RHISetShaderResourceViewParameter;
 	using IRHIComputeContext::RHISetShaderUniformBuffer;
+	using IRHIComputeContext::RHISetShaderParameters;
 	using IRHIComputeContext::RHISetShaderParameter;
 
 	/** Set the shader resource view of a surface. */
@@ -700,6 +703,8 @@ public:
 
 	virtual void RHISetShaderUniformBuffer(FRHIGraphicsShader* Shader, uint32 BufferIndex, FRHIUniformBuffer* Buffer) = 0;
 
+	virtual void RHISetShaderParameters(FRHIGraphicsShader* Shader, TConstArrayView<uint8> InParametersData, TConstArrayView<FRHIShaderParameter> InParameters, TConstArrayView<FRHIShaderParameterResource> InResourceParameters, TConstArrayView<FRHIShaderParameterResource> InBindlessParameters) = 0;
+
 	virtual void RHISetShaderParameter(FRHIGraphicsShader* Shader, uint32 BufferIndex, uint32 BaseIndex, uint32 NumBytes, const void* NewValue) = 0;
 
 	virtual void RHISetStencilRef(uint32 StencilRef) {}
diff --git a/Engine/Source/Runtime/RHI/Public/RHIResources.h b/Engine/Source/Runtime/RHI/Public/RHIResources.h
--- a/Engine/Source/Runtime/RHI/Public/RHIResources.h
+++ b/Engine/Source/Runtime/RHI/Public/RHIResources.h
@@ -4328,3 +4328,41 @@ private:
 	TArray<TPair<FRHIBufferUAVCreateInfo, FUnorderedAccessViewRHIRef>, TInlineAllocator<1>> UAVs;
 	TArray<TPair<FRHIBufferSRVCreateInfo, FShaderResourceViewRHIRef>, TInlineAllocator<1>> SRVs;
 };
+
+struct FRHIShaderParameter
+{
+	FRHIShaderParameter(uint16 InBufferIndex, uint16 InBaseIndex, uint16 InByteOffset, uint16 InByteSize)
+		: BufferIndex(InBufferIndex)
+		, BaseIndex(InBaseIndex)
+		, ByteOffset(InByteOffset)
+		, ByteSize(InByteSize)
+	{
+	}
+	uint16 BufferIndex;
+	uint16 BaseIndex;
+	uint16 ByteOffset;
+	uint16 ByteSize;
+};
+
+struct FRHIShaderParameterResource
+{
+	enum class EType : uint8
+	{
+		Texture,
+		ResourceView,
+		UnorderedAccessView,
+		Sampler,
+		UniformBuffer,
+	};
+
+	FRHIShaderParameterResource(EType InType, FRHIResource* InResource, uint16 InIndex)
+		: Resource(InResource)
+		, Index(InIndex)
+		, Type(InType)
+	{
+	}
+
+	FRHIResource* Resource;
+	uint16        Index;
+	EType         Type;
+};
diff --git a/Engine/Source/Runtime/RHI/Public/RHIValidationContext.h b/Engine/Source/Runtime/RHI/Public/RHIValidationContext.h
--- a/Engine/Source/Runtime/RHI/Public/RHIValidationContext.h
+++ b/Engine/Source/Runtime/RHI/Public/RHIValidationContext.h
@@ -12,6 +12,37 @@
 #if ENABLE_RHI_VALIDATION
 class FValidationRHI;
 
+inline void ValidateShaderParameters(RHIValidation::FTracker* Tracker, RHIValidation::FStaticUniformBuffers& StaticUniformBuffers, TConstArrayView<FRHIShaderParameterResource> InParameters, ERHIAccess InRequiredAccess, RHIValidation::EUAVMode InRequiredUAVMode)
+{
+	for (const FRHIShaderParameterResource& Parameter : InParameters)
+	{
+		switch (Parameter.Type)
+		{
+		case FRHIShaderParameterResource::EType::Texture:
+			Tracker->Assert(static_cast<FRHITexture*>(Parameter.Resource)->GetWholeResourceIdentitySRV(), InRequiredAccess);
+			break;
+		case FRHIShaderParameterResource::EType::ResourceView:
+			if (FRHIShaderResourceView* SRV = static_cast<FRHIShaderResourceView*>(Parameter.Resource))
+			{
+				Tracker->Assert(SRV->ViewIdentity, InRequiredAccess);
+			}
+			break;
+		case FRHIShaderParameterResource::EType::UnorderedAccessView:
+			Tracker->AssertUAV(static_cast<FRHIUnorderedAccessView*>(Parameter.Resource), InRequiredUAVMode, Parameter.Index);
+			break;
+		case FRHIShaderParameterResource::EType::Sampler:
+			// No validation
+			break;
+		case FRHIShaderParameterResource::EType::UniformBuffer:
+			StaticUniformBuffers.ValidateSetShaderUniformBuffer(static_cast<FRHIUniformBuffer*>(Parameter.Resource));
+			break;
+		default:
+			checkf(false, TEXT("Unhandled resource type?"));
+			break;
+		}
+	}
+}
+
 class FValidationComputeContext final : public IRHIComputeContext
 {
 public:
@@ -222,6 +253,16 @@ public:
 		RHIContext->RHISetShaderParameter(Shader, BufferIndex, BaseIndex, NumBytes, NewValue);
 	}
 
+	virtual void RHISetShaderParameters(FRHIComputeShader* Shader, TConstArrayView<uint8> InParametersData, TConstArrayView<FRHIShaderParameter> InParameters, TConstArrayView<FRHIShaderParameterResource> InResourceParameters, TConstArrayView<FRHIShaderParameterResource> InBindlessParameters) final override
+	{
+		checkf(State.bComputePSOSet, TEXT("A Compute PSO has to be set to set resources into a shader!"));
+
+		ValidateShaderParameters(Tracker, State.StaticUniformBuffers, InResourceParameters, ERHIAccess::SRVCompute, RHIValidation::EUAVMode::Compute);
+		ValidateShaderParameters(Tracker, State.StaticUniformBuffers, InBindlessParameters, ERHIAccess::SRVCompute, RHIValidation::EUAVMode::Compute);
+
+		RHIContext->RHISetShaderParameters(Shader, InParametersData, InParameters, InResourceParameters, InBindlessParameters);
+	}
+
 	virtual void RHISetStaticUniformBuffers(const FUniformBufferStaticBindings& InUniformBuffers) override final
 	{
 		InUniformBuffers.Bind(State.StaticUniformBuffers.Bindings);
@@ -763,6 +804,26 @@ public:
 		RHIContext->RHISetShaderParameter(Shader, BufferIndex, BaseIndex, NumBytes, NewValue);
 	}
 
+	virtual void RHISetShaderParameters(FRHIGraphicsShader* Shader, TConstArrayView<uint8> InParametersData, TConstArrayView<FRHIShaderParameter> InParameters, TConstArrayView<FRHIShaderParameterResource> InResourceParameters, TConstArrayView<FRHIShaderParameterResource> InBindlessParameters) final override
+	{
+		checkf(State.bGfxPSOSet, TEXT("A Graphics PSO has to be set to set resources into a shader!"));
+
+		ValidateShaderParameters(Tracker, State.StaticUniformBuffers, InResourceParameters, ERHIAccess::SRVGraphics, RHIValidation::EUAVMode::Graphics);
+		ValidateShaderParameters(Tracker, State.StaticUniformBuffers, InBindlessParameters, ERHIAccess::SRVGraphics, RHIValidation::EUAVMode::Graphics);
+
+		RHIContext->RHISetShaderParameters(Shader, InParametersData, InParameters, InResourceParameters, InBindlessParameters);
+	}
+
+	virtual void RHISetShaderParameters(FRHIComputeShader* Shader, TConstArrayView<uint8> InParametersData, TConstArrayView<FRHIShaderParameter> InParameters, TConstArrayView<FRHIShaderParameterResource> InResourceParameters, TConstArrayView<FRHIShaderParameterResource> InBindlessParameters) final override
+	{
+		checkf(State.bComputePSOSet, TEXT("A Compute PSO has to be set to set resources into a shader!"));
+
+		ValidateShaderParameters(Tracker, State.StaticUniformBuffers, InResourceParameters, ERHIAccess::SRVCompute, RHIValidation::EUAVMode::Compute);
+		ValidateShaderParameters(Tracker, State.StaticUniformBuffers, InBindlessParameters, ERHIAccess::SRVCompute, RHIValidation::EUAVMode::Compute);
+
+		RHIContext->RHISetShaderParameters(Shader, InParametersData, InParameters, InResourceParameters, InBindlessParameters);
+	}
+
 	virtual void RHISetStaticUniformBuffers(const FUniformBufferStaticBindings& InUniformBuffers) override final
 	{
 		InUniformBuffers.Bind(State.StaticUniformBuffers.Bindings);
diff --git a/Engine/Source/Runtime/RHICore/Internal/RHIShaderParametersShared.h b/Engine/Source/Runtime/RHICore/Internal/RHIShaderParametersShared.h
new file mode 100644
--- /dev/null
+++ b/Engine/Source/Runtime/RHICore/Internal/RHIShaderParametersShared.h
@@ -0,0 +1,115 @@
+// Copyright Epic Games, Inc. All Rights Reserved.
+
+#pragma once
+
+namespace UE::RHICore
+{
+	inline FRHIDescriptorHandle GetBindlessParameterHandle(const FRHIShaderParameterResource& Parameter)
+	{
+		if (FRHIResource* Resource = Parameter.Resource)
+		{
+			switch (Parameter.Type)
+			{
+			case FRHIShaderParameterResource::EType::Texture:             return static_cast<FRHITexture*>(Resource)->GetDefaultBindlessHandle();
+			case FRHIShaderParameterResource::EType::ResourceView:        return static_cast<FRHIShaderResourceView*>(Resource)->GetBindlessHandle();
+			case FRHIShaderParameterResource::EType::UnorderedAccessView: return static_cast<FRHIUnorderedAccessView*>(Resource)->GetBindlessHandle();
+			case FRHIShaderParameterResource::EType::Sampler:             return static_cast<FRHISamplerState*>(Resource)->GetBindlessHandle();
+			default:
+				checkf(false, TEXT("Unhandled resource type?"));
+				break;
+			}
+		}
+
+		return FRHIDescriptorHandle();
+	}
+
+	template<typename FContextRHI>
+	inline void SetShaderUAV(FContextRHI& Context, FRHIGraphicsShader* ShaderRHI, uint16 Index, FRHIUnorderedAccessView* UAV)
+	{
+		if (ShaderRHI->GetType() == RRT_PixelShader)
+		{
+			Context.RHISetUAVParameter(static_cast<FRHIPixelShader*>(ShaderRHI), Index, UAV);
+		}
+		else
+		{
+			checkf(false, TEXT("TShaderRHI Can't have compute shader to be set. UAVs are not supported on vertex, tessellation and geometry shaders."));
+		}
+	}
+
+	template<typename FContextRHI>
+	inline void SetShaderUAV(FContextRHI& Context, FRHIComputeShader* ShaderRHI, uint16 Index, FRHIUnorderedAccessView* UAV)
+	{
+		Context.RHISetUAVParameter(ShaderRHI, Index, UAV);
+	}
+
+	template<typename FContextRHI, typename TShaderRHI>
+	inline void RHISetShaderParametersShared(
+		FContextRHI& Context
+		, TShaderRHI* ShaderRHI
+		, TConstArrayView<uint8> InParametersData
+		, TConstArrayView<FRHIShaderParameter> InParameters
+		, TConstArrayView<FRHIShaderParameterResource> InResourceParameters
+		, TConstArrayView<FRHIShaderParameterResource> InBindlessParameters
+		, bool bBindUAVsFirst = false // UE TODO: investigate where this is absolutely necessary
+	)
+	{
+		if (InParameters.Num())
+		{
+			for (const FRHIShaderParameter& Parameter : InParameters)
+			{
+				Context.RHISetShaderParameter(ShaderRHI, Parameter.BufferIndex, Parameter.BaseIndex, Parameter.ByteSize, &InParametersData[Parameter.ByteOffset]);
+			}
+		}
+
+#if PLATFORM_SUPPORTS_BINDLESS_RENDERING
+		for (const FRHIShaderParameterResource& Parameter : InBindlessParameters)
+		{
+			const FRHIDescriptorHandle Handle = GetBindlessParameterHandle(Parameter);
+			if (Handle.IsValid())
+			{
+				const uint32 BindlessIndex = Handle.GetIndex();
+				Context.RHISetShaderParameter(ShaderRHI, 0, Parameter.Index, 4, &BindlessIndex);
+			}
+		}
+#endif
+
+		if (bBindUAVsFirst)
+		{
+			for (const FRHIShaderParameterResource& Parameter : InResourceParameters)
+			{
+				if (Parameter.Type == FRHIShaderParameterResource::EType::UnorderedAccessView)
+				{
+					SetShaderUAV(Context, ShaderRHI, Parameter.Index, static_cast<FRHIUnorderedAccessView*>(Parameter.Resource));
+				}
+			}
+		}
+
+		for (const FRHIShaderParameterResource& Parameter : InResourceParameters)
+		{
+			switch (Parameter.Type)
+			{
+			case FRHIShaderParameterResource::EType::Texture:
+				Context.RHISetShaderTexture(ShaderRHI, Parameter.Index, static_cast<FRHITexture*>(Parameter.Resource));
+				break;
+			case FRHIShaderParameterResource::EType::ResourceView:
+				Context.RHISetShaderResourceViewParameter(ShaderRHI, Parameter.Index, static_cast<FRHIShaderResourceView*>(Parameter.Resource));
+				break;
+			case FRHIShaderParameterResource::EType::UnorderedAccessView:
+				if (!bBindUAVsFirst)
+				{
+					SetShaderUAV(Context, ShaderRHI, Parameter.Index, static_cast<FRHIUnorderedAccessView*>(Parameter.Resource));
+				}
+				break;
+			case FRHIShaderParameterResource::EType::Sampler:
+				Context.RHISetShaderSampler(ShaderRHI, Parameter.Index, static_cast<FRHISamplerState*>(Parameter.Resource));
+				break;
+			case FRHIShaderParameterResource::EType::UniformBuffer:
+				Context.RHISetShaderUniformBuffer(ShaderRHI, Parameter.Index, static_cast<FRHIUniformBuffer*>(Parameter.Resource));
+				break;
+			default:
+				checkf(false, TEXT("Unhandled resource type?"));
+				break;
+			}
+		}
+	}
+}
diff --git a/Engine/Source/Runtime/RenderCore/Private/ShaderParameterStruct.cpp b/Engine/Source/Runtime/RenderCore/Private/ShaderParameterStruct.cpp
--- a/Engine/Source/Runtime/RenderCore/Private/ShaderParameterStruct.cpp
+++ b/Engine/Source/Runtime/RenderCore/Private/ShaderParameterStruct.cpp
@@ -472,6 +472,7 @@ struct FShaderParameterReader
 {
 	FShaderParameterReader() = delete;
 	FShaderParameterReader(const void* InData, uint64 InDataSize) : Data(reinterpret_cast<const uint8*>(InData)), DataSize(InDataSize) { }
+	FShaderParameterReader(TConstArrayView<uint8> InData) : Data(reinterpret_cast<const uint8*>(InData.GetData())), DataSize(InData.Num()) { }
 
 	template<typename TParameterIn>
 	const void* GetRawPointer(const TParameterIn& InParameter) const
@@ -596,116 +597,187 @@ void ValidateShaderParameterResourcesRHI(const void* Contents, const FRHIUniform
 
 #endif // DO_CHECK
 
-inline FRHIDescriptorHandle GetBindlessParameterHandle(FShaderParameterReader Reader, const FShaderParameterBindings::FBindlessResourceParameter& Parameter)
+static void ExtractShaderParameters(
+	TArray<FRHIShaderParameter>& OutParameters,
+	TConstArrayView<uint8>& OutMinimalParametersData,
+	const FShaderParameterBindings& Bindings,
+	const FShaderParametersMetadata* ParametersMetadata,
+	TConstArrayView<uint8> ParametersData)
+{
+	if (int32 NumParameters = Bindings.Parameters.Num())
+	{
+		OutParameters.Reserve(NumParameters);
+
+		// Keep track of the highest offset of data we need so we can avoid copying everything in the RHI command
+		uint32 MaxParametersSize = 0;
+
+		for (const FShaderParameterBindings::FParameter& Parameter : Bindings.Parameters)
+		{
+			OutParameters.Emplace(Parameter.BufferIndex, Parameter.BaseIndex, Parameter.ByteOffset, Parameter.ByteSize);
+
+			MaxParametersSize = FMath::Max<uint32>(MaxParametersSize, Parameter.ByteOffset + Parameter.ByteSize);
+		}
+		check(MaxParametersSize <= ParametersMetadata->GetSize());
+
+		OutMinimalParametersData = TConstArrayView<uint8>(ParametersData.GetData(), MaxParametersSize);
+	}
+}
+
+/** Utility class to populate FRHIShaderParameterResource based on types */
+struct FRHIShaderParameterResourceEx : public FRHIShaderParameterResource
+{
+	FRHIShaderParameterResourceEx()
+		: FRHIShaderParameterResource(FRHIShaderParameterResource::EType::Texture, nullptr, 0)
+	{
+	}
+	FRHIShaderParameterResourceEx(FRHITexture* InTexture, uint16 InIndex)
+		: FRHIShaderParameterResource(FRHIShaderParameterResource::EType::Texture, InTexture, InIndex)
+	{
+	}
+	FRHIShaderParameterResourceEx(FRHIShaderResourceView* InView, uint16 InIndex)
+		: FRHIShaderParameterResource(FRHIShaderParameterResource::EType::ResourceView, InView, InIndex)
+	{
+	}
+	FRHIShaderParameterResourceEx(FRHIUnorderedAccessView* InUAV, uint16 InIndex)
+		: FRHIShaderParameterResource(FRHIShaderParameterResource::EType::UnorderedAccessView, InUAV, InIndex)
+	{
+	}
+	FRHIShaderParameterResourceEx(FRHISamplerState* InSamplerState, uint16 InIndex)
+		: FRHIShaderParameterResource(FRHIShaderParameterResource::EType::Sampler, InSamplerState, InIndex)
+	{
+	}
+	FRHIShaderParameterResourceEx(FRHIUniformBuffer* InUniformBuffer, uint16 InIndex)
+		: FRHIShaderParameterResource(FRHIShaderParameterResource::EType::UniformBuffer, InUniformBuffer, InIndex)
+	{
+	}
+};
+
+#define UE_SINGLE_SHADER_PARAMETERS_COMMAND_INCLUDE_UB (0)
+
+inline int16 GetParameterIndex(const FShaderParameterBindings::FResourceParameter& Parameter) { return Parameter.BaseIndex; }
+inline int16 GetParameterIndex(const FShaderParameterBindings::FBindlessResourceParameter& Parameter) { return Parameter.GlobalConstantOffset; }
+inline int16 GetParameterIndex(const FShaderParameterBindings::FParameterStructReference& Parameter) { return Parameter.BufferIndex; }
+
+template<typename BindingParameterType>
+FRHIShaderParameterResource ExtractShaderParameterResource(FShaderParameterReader Reader, const BindingParameterType& Parameter)
 {
 	const EUniformBufferBaseType BaseType = static_cast<EUniformBufferBaseType>(Parameter.BaseType);
+
 	switch (BaseType)
 	{
 	case UBMT_TEXTURE:
 	{
 		FRHITexture* Texture = Reader.Read<FRHITexture*>(Parameter);
 		checkSlow(Texture);
-		return Texture ? Texture->GetDefaultBindlessHandle() : FRHIDescriptorHandle();
+		return FRHIShaderParameterResourceEx(Texture, GetParameterIndex(Parameter));
 	}
-	break;
 	case UBMT_SRV:
 	{
 		FRHIShaderResourceView* ShaderResourceView = Reader.Read<FRHIShaderResourceView*>(Parameter);
 		checkSlow(ShaderResourceView);
-		return ShaderResourceView->GetBindlessHandle();
+		return FRHIShaderParameterResourceEx(ShaderResourceView, GetParameterIndex(Parameter));
+	}
+	case UBMT_UAV:
+	{
+		FRHIUnorderedAccessView* UnorderedAccessView = Reader.Read<FRHIUnorderedAccessView*>(Parameter);
+		checkSlow(UnorderedAccessView);
+		return FRHIShaderParameterResourceEx(UnorderedAccessView, GetParameterIndex(Parameter));
 	}
-	break;
 	case UBMT_SAMPLER:
 	{
 		FRHISamplerState* SamplerState = Reader.Read<FRHISamplerState*>(Parameter);
 		checkSlow(SamplerState);
-		return SamplerState->GetBindlessHandle();
+		return FRHIShaderParameterResourceEx(SamplerState, GetParameterIndex(Parameter));
 	}
-	break;
 	case UBMT_RDG_TEXTURE:
 	{
 		FRDGTexture* RDGTexture = Reader.Read<FRDGTexture*>(Parameter);
 		checkSlow(RDGTexture);
-		FRHITexture* Texture = RDGTexture->GetRHI();
-		return Texture ? Texture->GetDefaultBindlessHandle() : FRHIDescriptorHandle();
+		RDGTexture->MarkResourceAsUsed();
+		return FRHIShaderParameterResourceEx(RDGTexture->GetRHI(), GetParameterIndex(Parameter));
 	}
-	break;
 	case UBMT_RDG_TEXTURE_SRV:
 	case UBMT_RDG_BUFFER_SRV:
 	{
 		FRDGShaderResourceView* RDGShaderResourceView = Reader.Read<FRDGShaderResourceView*>(Parameter);
 		checkSlow(RDGShaderResourceView);
 		RDGShaderResourceView->MarkResourceAsUsed();
-
-		FRHIShaderResourceView* ShaderResourceView = RDGShaderResourceView->GetRHI();
-		return ShaderResourceView ? ShaderResourceView->GetBindlessHandle() : FRHIDescriptorHandle();
-	}
-	break;
-	case UBMT_UAV:
-	{
-		FRHIUnorderedAccessView* UnorderedAccessView = Reader.Read<FRHIUnorderedAccessView*>(Parameter);
-		checkSlow(UnorderedAccessView);
-		return UnorderedAccessView->GetBindlessHandle();
+		return FRHIShaderParameterResourceEx(RDGShaderResourceView->GetRHI(), GetParameterIndex(Parameter));
 	}
 	case UBMT_RDG_TEXTURE_UAV:
 	case UBMT_RDG_BUFFER_UAV:
 	{
 		FRDGUnorderedAccessView* RDGUnorderedAccessView = Reader.Read<FRDGUnorderedAccessView*>(Parameter);
 		checkSlow(RDGUnorderedAccessView);
-
-		FRHIUnorderedAccessView* UnorderedAccessView = RDGUnorderedAccessView->GetRHI();
-		return UnorderedAccessView ? UnorderedAccessView->GetBindlessHandle() : FRHIDescriptorHandle();
+		RDGUnorderedAccessView->MarkResourceAsUsed();
+		return FRHIShaderParameterResourceEx(RDGUnorderedAccessView->GetRHI(), GetParameterIndex(Parameter));
 	}
-	break;
 	default:
 		checkf(false, TEXT("Unhandled resource type?"));
-		break;
+		return FRHIShaderParameterResourceEx();
 	}
-	return FRHIDescriptorHandle();
 }
 
-template<typename TRHICmdList>
-inline void SetShaderUAV(TRHICmdList& RHICmdList, FRHIGraphicsShader* ShaderRHI, FShaderParameterReader Reader, const FShaderParameterBindings::FResourceParameter& ParameterBinding)
+static void ExtractShaderParameterResources(
+	TArray<FRHIShaderParameterResource>& OutResourceParameters,
+	TArray<FRHIShaderParameterResource>& OutBindlessParameters,
+	const FShaderParameterBindings& Bindings,
+	const FShaderParametersMetadata* ParametersMetadata,
+	TConstArrayView<uint8> ParametersData)
 {
-	if (ShaderRHI->GetType() == RRT_PixelShader)
+	const FShaderParameterReader Reader(ParametersData);
+
+#if PLATFORM_SUPPORTS_BINDLESS_RENDERING
+	if (int32 NumBindings = Bindings.BindlessResourceParameters.Num())
 	{
-		FRHIPixelShader* PixelShaderRHI = static_cast<FRHIPixelShader*>(ShaderRHI);
+		OutBindlessParameters.Reserve(NumBindings);
 
-		if (ParameterBinding.BaseType == UBMT_UAV)
-		{
-			FRHIUnorderedAccessView* ShaderParameterRef = Reader.Read<FRHIUnorderedAccessView*>(ParameterBinding);
-			RHICmdList.SetUAVParameter(PixelShaderRHI, ParameterBinding.BaseIndex, ShaderParameterRef);
-		}
-		else if (ParameterBinding.BaseType == UBMT_RDG_TEXTURE_UAV || ParameterBinding.BaseType == UBMT_RDG_BUFFER_UAV)
+		for (const FShaderParameterBindings::FBindlessResourceParameter& Parameter : Bindings.BindlessResourceParameters)
 		{
-			FRDGUnorderedAccessView* GraphUAV = Reader.Read<FRDGUnorderedAccessView*>(ParameterBinding);
-			checkSlow(GraphUAV);
-
-			GraphUAV->MarkResourceAsUsed();
-			RHICmdList.SetUAVParameter(PixelShaderRHI, ParameterBinding.BaseIndex, GraphUAV->GetRHI());
+			const FRHIShaderParameterResource ShaderParameterResource = ExtractShaderParameterResource(Reader, Parameter);
+			OutBindlessParameters.Emplace(ShaderParameterResource);
 		}
 	}
-	else
-	{
-		checkf(false, TEXT("TShaderRHI Can't have compute shader to be set. UAVs are not supported on vertex, tessellation and geometry shaders."));
-	}
-}
+#endif
 
-template<typename TRHICmdList>
-inline void SetShaderUAV(TRHICmdList& RHICmdList, FRHIComputeShader* ShaderRHI, FShaderParameterReader Reader, const FShaderParameterBindings::FResourceParameter& ParameterBinding)
-{
-	if (ParameterBinding.BaseType == UBMT_UAV)
-	{
-		FRHIUnorderedAccessView* ShaderParameterRef = Reader.Read<FRHIUnorderedAccessView*>(ParameterBinding);
-		RHICmdList.SetUAVParameter(ShaderRHI, ParameterBinding.BaseIndex, ShaderParameterRef);
-	}
-	else if (ParameterBinding.BaseType == UBMT_RDG_TEXTURE_UAV || ParameterBinding.BaseType == UBMT_RDG_BUFFER_UAV)
+	int32 NumBindings = Bindings.ResourceParameters.Num();
+#if UE_SINGLE_SHADER_PARAMETERS_COMMAND_INCLUDE_UB
+	NumBindings += Bindings.GraphUniformBuffers.Num() + Bindings.ParameterReferences.Num();
+#endif
+
+	if (NumBindings)
 	{
-		FRDGUnorderedAccessView* GraphUAV = Reader.Read<FRDGUnorderedAccessView*>(ParameterBinding);
-		checkSlow(GraphUAV);
+		OutResourceParameters.Reserve(NumBindings);
+
+		for (const FShaderParameterBindings::FResourceParameter& Parameter : Bindings.ResourceParameters)
+		{
+			const FRHIShaderParameterResource ShaderParameterResource = ExtractShaderParameterResource(Reader, Parameter);
+			OutResourceParameters.Emplace(ShaderParameterResource);
+		}
 
-		GraphUAV->MarkResourceAsUsed();
-		RHICmdList.SetUAVParameter(ShaderRHI, ParameterBinding.BaseIndex, GraphUAV->GetRHI());
+#if UE_SINGLE_SHADER_PARAMETERS_COMMAND_INCLUDE_UB
+		for (const FShaderParameterBindings::FParameterStructReference& Parameter : Bindings.GraphUniformBuffers)
+		{
+			const FRDGUniformBufferBinding& UniformBufferBinding = Reader.Read<FRDGUniformBufferBinding>(Parameter);
+			if (UniformBufferBinding.IsShader())
+			{
+				UniformBufferBinding->MarkResourceAsUsed();
+
+				FRHIShaderParameterResourceEx Binding(UniformBufferBinding->GetRHI(), GetParameterIndex(Parameter));
+				OutResourceParameters.Emplace(Binding);
+			}
+		}
+
+		for (const FShaderParameterBindings::FParameterStructReference& Parameter : Bindings.ParameterReferences)
+		{
+			const FUniformBufferBinding& UniformBufferBinding = Reader.Read<FUniformBufferBinding>(Parameter);
+			if (UniformBufferBinding.IsShader())
+			{
+				FRHIShaderParameterResourceEx Binding(UniformBufferBinding.GetUniformBuffer(), GetParameterIndex(Parameter));
+				OutResourceParameters.Emplace(Binding);
+			}
+		}
+#endif // UE_SINGLE_SHADER_PARAMETERS_COMMAND_INCLUDE_UB
 	}
 }
 
@@ -716,95 +788,24 @@ inline void SetShaderParametersInternal(
 	TShaderRHI* ShaderRHI,
 	const FShaderParameterBindings& Bindings,
 	const FShaderParametersMetadata* ParametersMetadata,
-	const void* ParametersData)
+	const void* InParametersData)
 {
 	checkf(Bindings.RootParameterBufferIndex == FShaderParameterBindings::kInvalidBufferIndex, TEXT("Can't use SetShaderParameters() for root parameter buffer index."));
 
-	FShaderParameterReader Reader(ParametersData, ParametersMetadata->GetSize());
+	TConstArrayView<uint8> FullParametersData((const uint8*)InParametersData, ParametersMetadata->GetSize());
 
-	// Parameters
-	for (const FShaderParameterBindings::FParameter& Parameter : Bindings.Parameters)
-	{
-		const void* DataPtr = Reader.GetRawPointer(Parameter);
-		RHICmdList.SetShaderParameter(ShaderRHI, Parameter.BufferIndex, Parameter.BaseIndex, Parameter.ByteSize, DataPtr);
-	}
+	TArray<FRHIShaderParameter> Parameters;
+	TConstArrayView<uint8> MinimalParametersData;
+	ExtractShaderParameters(Parameters, MinimalParametersData, Bindings, ParametersMetadata, FullParametersData);
 
-#if PLATFORM_SUPPORTS_BINDLESS_RENDERING
-	for (const FShaderParameterBindings::FBindlessResourceParameter& BindlessResource : Bindings.BindlessResourceParameters)
-	{
-		const FRHIDescriptorHandle Handle = GetBindlessParameterHandle(Reader, BindlessResource);
-		if (Handle.IsValid())
-		{
-			const uint32 BindlessIndex = Handle.GetIndex();
-			RHICmdList.SetShaderParameter(ShaderRHI, 0, BindlessResource.GlobalConstantOffset, 4, &BindlessIndex);
-		}
-	}
-#endif
+	TArray<FRHIShaderParameterResource> ResourceParameters;
+	TArray<FRHIShaderParameterResource> BindlessParameters;
+	ExtractShaderParameterResources(ResourceParameters, BindlessParameters, Bindings, ParametersMetadata, FullParametersData);
 
-	TArray<FShaderParameterBindings::FResourceParameter, TInlineAllocator<16>> GraphSRVs;
+	RHICmdList.SetShaderParameters(ShaderRHI, MinimalParametersData, Parameters, ResourceParameters, BindlessParameters);
 
-	for (const FShaderParameterBindings::FResourceParameter& Parameter : Bindings.ResourceParameters)
-	{
-		const EUniformBufferBaseType BaseType = static_cast<EUniformBufferBaseType>(Parameter.BaseType);
-		switch (BaseType)
-		{
-		case UBMT_TEXTURE:
-		{
-			FRHITexture* RHITexture = Reader.Read<FRHITexture*>(Parameter);
-
-			RHICmdList.SetShaderTexture(ShaderRHI, Parameter.BaseIndex, RHITexture);
-		}
-		break;
-		case UBMT_SRV:
-		{
-			FRHIShaderResourceView* ShaderResourceView = Reader.Read<FRHIShaderResourceView*>(Parameter);
-			RHICmdList.SetShaderResourceViewParameter(ShaderRHI, Parameter.BaseIndex, ShaderResourceView);
-		}
-		break;
-		case UBMT_SAMPLER:
-		{
-			FRHISamplerState* SamplerState = Reader.Read<FRHISamplerState*>(Parameter);
-			RHICmdList.SetShaderSampler(ShaderRHI, Parameter.BaseIndex, SamplerState);
-		}
-		break;
-		case UBMT_RDG_TEXTURE:
-		{
-			FRDGTexture* RDGTexture = Reader.Read<FRDGTexture*>(Parameter);
-			checkSlow(RDGTexture);
-			RDGTexture->MarkResourceAsUsed();
-
-			RHICmdList.SetShaderTexture(ShaderRHI, Parameter.BaseIndex, RDGTexture->GetRHI());
-		}
-		break;
-		case UBMT_RDG_TEXTURE_SRV:
-		case UBMT_RDG_BUFFER_SRV:
-		{
-			//HACKHACK: defer SRVs binding after UAVs 
-			GraphSRVs.Add(Parameter);
-		}
-		break;
-		case UBMT_UAV:
-		case UBMT_RDG_TEXTURE_UAV:
-		case UBMT_RDG_BUFFER_UAV:
-		{
-			SetShaderUAV(RHICmdList, ShaderRHI, Reader, Parameter);
-		}
-		break;
-		default:
-			checkf(false, TEXT("Unhandled resource type?"));
-			break;
-		}
-	}
-
-	//HACKHACK: Bind SRVs after UAVs as a workaround for D3D11 RHI unbinding SRVs when binding a UAV on the same resource even when the views don't overlap.
-	for (const FShaderParameterBindings::FResourceParameter& Parameter : GraphSRVs)
-	{
-		FRDGShaderResourceView* RDGShaderResourceView = Reader.Read<FRDGShaderResourceView*>(Parameter);
-		checkSlow(RDGShaderResourceView);
-		RDGShaderResourceView->MarkResourceAsUsed();
-
-		RHICmdList.SetShaderResourceViewParameter(ShaderRHI, Parameter.BaseIndex, RDGShaderResourceView->GetRHI());
-	}
+#if !UE_SINGLE_SHADER_PARAMETERS_COMMAND_INCLUDE_UB
+	const FShaderParameterReader Reader(FullParametersData);
 
 	// Graph Uniform Buffers
 	for (const FShaderParameterBindings::FParameterStructReference& Parameter : Bindings.GraphUniformBuffers)
@@ -826,6 +827,7 @@ inline void SetShaderParametersInternal(
 			RHICmdList.SetShaderUniformBuffer(ShaderRHI, Parameter.BufferIndex, UniformBufferBinding.GetUniformBuffer());
 		}
 	}
+#endif
 }
 
 void SetShaderParameters(
diff --git a/Engine/Source/Runtime/VulkanRHI/Private/VulkanCommands.cpp b/Engine/Source/Runtime/VulkanRHI/Private/VulkanCommands.cpp
--- a/Engine/Source/Runtime/VulkanRHI/Private/VulkanCommands.cpp
+++ b/Engine/Source/Runtime/VulkanRHI/Private/VulkanCommands.cpp
@@ -10,6 +10,7 @@
 #include "EngineGlobals.h"
 #include "VulkanLLM.h"
 #include "RenderUtils.h"
+#include "RHIShaderParametersShared.h"
 
 static TAutoConsoleVariable<int32> GCVarSubmitOnDispatch(
 	TEXT("r.Vulkan.SubmitOnDispatch"),
@@ -255,6 +256,30 @@ void FVulkanCommandListContext::RHISetShaderParameter(FRHIComputeShader* Compute
 	PendingComputeState->SetPackedGlobalShaderParameter(BufferIndex, BaseIndex, NumBytes, NewValue);
 }
 
+void FVulkanCommandListContext::RHISetShaderParameters(FRHIGraphicsShader* Shader, TConstArrayView<uint8> InParametersData, TConstArrayView<FRHIShaderParameter> InParameters, TConstArrayView<FRHIShaderParameterResource> InResourceParameters, TConstArrayView<FRHIShaderParameterResource> InBindlessParameters)
+{
+	UE::RHICore::RHISetShaderParametersShared(
+		*this
+		, Shader
+		, InParametersData
+		, InParameters
+		, InResourceParameters
+		, InBindlessParameters
+	);
+}
+
+void FVulkanCommandListContext::RHISetShaderParameters(FRHIComputeShader* Shader, TConstArrayView<uint8> InParametersData, TConstArrayView<FRHIShaderParameter> InParameters, TConstArrayView<FRHIShaderParameterResource> InResourceParameters, TConstArrayView<FRHIShaderParameterResource> InBindlessParameters)
+{
+	UE::RHICore::RHISetShaderParametersShared(
+		*this
+		, Shader
+		, InParametersData
+		, InParameters
+		, InResourceParameters
+		, InBindlessParameters
+	);
+}
+
 template <typename TState>
 inline void SetShaderUniformBufferResources(FVulkanCommandListContext* Context, TState* State, const FVulkanShader* Shader, const TArray<FVulkanShaderHeader::FGlobalInfo>& GlobalInfos, const TArray<TEnumAsByte<EVulkanBindingType::EType>>& DescriptorTypes, const FVulkanShaderHeader::FUniformBufferInfo& HeaderUBInfo, const FVulkanUniformBuffer* UniformBuffer, const TArray<FDescriptorSetRemappingInfo::FRemappingInfo>& GlobalRemappingInfo)
 {
diff --git a/Engine/Source/Runtime/VulkanRHI/Private/VulkanContext.h b/Engine/Source/Runtime/VulkanRHI/Private/VulkanContext.h
--- a/Engine/Source/Runtime/VulkanRHI/Private/VulkanContext.h
+++ b/Engine/Source/Runtime/VulkanRHI/Private/VulkanContext.h
@@ -54,6 +54,8 @@ public:
 	virtual void RHISetShaderUniformBuffer(FRHIComputeShader* ComputeShader, uint32 BufferIndex, FRHIUniformBuffer* Buffer) final override;
 	virtual void RHISetShaderParameter(FRHIGraphicsShader* Shader, uint32 BufferIndex, uint32 BaseIndex, uint32 NumBytes, const void* NewValue) final override;
 	virtual void RHISetShaderParameter(FRHIComputeShader* ComputeShader, uint32 BufferIndex, uint32 BaseIndex, uint32 NumBytes, const void* NewValue) final override;
+	virtual void RHISetShaderParameters(FRHIGraphicsShader* Shader, TConstArrayView<uint8> InParametersData, TConstArrayView<FRHIShaderParameter> InParameters, TConstArrayView<FRHIShaderParameterResource> InResourceParameters, TConstArrayView<FRHIShaderParameterResource> InBindlessParameters) final override;
+	virtual void RHISetShaderParameters(FRHIComputeShader* Shader, TConstArrayView<uint8> InParametersData, TConstArrayView<FRHIShaderParameter> InParameters, TConstArrayView<FRHIShaderParameterResource> InResourceParameters, TConstArrayView<FRHIShaderParameterResource> InBindlessParameters) final override;
 	virtual void RHISetStencilRef(uint32 StencilRef) final override;
 	virtual void RHIDrawPrimitive(uint32 BaseVertexIndex, uint32 NumPrimitives, uint32 NumInstances) final override;
 	virtual void RHIDrawPrimitiveIndirect(FRHIBuffer* ArgumentBuffer, uint32 ArgumentOffset) final override;
diff --git a/Engine/Source/Runtime/Windows/D3D11RHI/Private/D3D11Commands.cpp b/Engine/Source/Runtime/Windows/D3D11RHI/Private/D3D11Commands.cpp
--- a/Engine/Source/Runtime/Windows/D3D11RHI/Private/D3D11Commands.cpp
+++ b/Engine/Source/Runtime/Windows/D3D11RHI/Private/D3D11Commands.cpp
@@ -14,6 +14,7 @@
 #include "ShaderParameterUtils.h"
 #include "SceneUtils.h"
 #include "EngineGlobals.h"
+#include "RHIShaderParametersShared.h"
 
 #if !PLATFORM_HOLOLENS
 // For Depth Bounds Test interface
@@ -684,6 +685,32 @@ void FD3D11DynamicRHI::RHISetShaderParameter(FRHIComputeShader* ComputeShaderRHI
 	CSConstantBuffers[BufferIndex]->UpdateConstant((const uint8*)NewValue,BaseIndex,NumBytes);
 }
 
+void FD3D11DynamicRHI::RHISetShaderParameters(FRHIComputeShader* Shader, TConstArrayView<uint8> InParametersData, TConstArrayView<FRHIShaderParameter> InParameters, TConstArrayView<FRHIShaderParameterResource> InResourceParameters, TConstArrayView<FRHIShaderParameterResource> InBindlessParameters)
+{
+	UE::RHICore::RHISetShaderParametersShared(
+		*this
+		, Shader
+		, InParametersData
+		, InParameters
+		, InResourceParameters
+		, InBindlessParameters
+		, true /*bBindUAVsFirst*/
+	);
+}
+
+void FD3D11DynamicRHI::RHISetShaderParameters(FRHIGraphicsShader* Shader, TConstArrayView<uint8> InParametersData, TConstArrayView<FRHIShaderParameter> InParameters, TConstArrayView<FRHIShaderParameterResource> InResourceParameters, TConstArrayView<FRHIShaderParameterResource> InBindlessParameters)
+{
+	UE::RHICore::RHISetShaderParametersShared(
+		*this
+		, Shader
+		, InParametersData
+		, InParameters
+		, InResourceParameters
+		, InBindlessParameters
+		, true /*bBindUAVsFirst*/
+	);
+}
+
 void FD3D11DynamicRHI::ValidateExclusiveDepthStencilAccess(FExclusiveDepthStencil RequestedAccess) const
 {
 	const bool bSrcDepthWrite = RequestedAccess.IsDepthWrite();
diff --git a/Engine/Source/Runtime/Windows/D3D11RHI/Private/D3D11RHIPrivate.h b/Engine/Source/Runtime/Windows/D3D11RHI/Private/D3D11RHIPrivate.h
--- a/Engine/Source/Runtime/Windows/D3D11RHI/Private/D3D11RHIPrivate.h
+++ b/Engine/Source/Runtime/Windows/D3D11RHI/Private/D3D11RHIPrivate.h
@@ -596,6 +596,8 @@ public:
 	virtual void RHISetShaderUniformBuffer(FRHIComputeShader* ComputeShader, uint32 BufferIndex, FRHIUniformBuffer* Buffer) final override;
 	virtual void RHISetShaderParameter(FRHIGraphicsShader* Shader, uint32 BufferIndex, uint32 BaseIndex, uint32 NumBytes, const void* NewValue) final override;
 	virtual void RHISetShaderParameter(FRHIComputeShader* ComputeShader, uint32 BufferIndex, uint32 BaseIndex, uint32 NumBytes, const void* NewValue) final override;
+	virtual void RHISetShaderParameters(FRHIComputeShader* Shader, TConstArrayView<uint8> InParametersData, TConstArrayView<FRHIShaderParameter> InParameters, TConstArrayView<FRHIShaderParameterResource> InResourceParameters, TConstArrayView<FRHIShaderParameterResource> InBindlessParameters) final override;
+	virtual void RHISetShaderParameters(FRHIGraphicsShader* Shader, TConstArrayView<uint8> InParametersData, TConstArrayView<FRHIShaderParameter> InParameters, TConstArrayView<FRHIShaderParameterResource> InResourceParameters, TConstArrayView<FRHIShaderParameterResource> InBindlessParameters) final override;
 	virtual void RHISetDepthStencilState(FRHIDepthStencilState* NewState, uint32 StencilRef) final override;
 	virtual void RHISetStencilRef(uint32 StencilRef) final override;
 	virtual void RHISetBlendState(FRHIBlendState* NewState, const FLinearColor& BlendFactor) final override;
