Skip pso with mesh shader from file cache from file cache when GRHISupportsMeshShadersTier0 is not available

#localization none
#tests run local replays - needs to be tested on older GPUs which I don't have available here

diff --git a/Engine/Source/Runtime/RenderCore/Private/ShaderPipelineCache.cpp b/Engine/Source/Runtime/RenderCore/Private/ShaderPipelineCache.cpp
--- a/Engine/Source/Runtime/RenderCore/Private/ShaderPipelineCache.cpp
+++ b/Engine/Source/Runtime/RenderCore/Private/ShaderPipelineCache.cpp
@@ -950,6 +950,12 @@ bool FShaderPipelineCacheTask::Precompile(FRHICommandListImmediate& RHICmdList,
 
 				if (PSO.GraphicsDesc.MeshShader != FSHAHash())
 				{
+					// Skip creation of mesh shaders if RHI doesn't support them
+					if (!GRHISupportsMeshShadersTier0)
+					{
+						return false;
+					}
+
 					FMeshShaderRHIRef MeshShader = FShaderCodeLibrary::CreateMeshShader(Platform, PSO.GraphicsDesc.MeshShader);
 					GraphicsInitializer.BoundShaderState.SetMeshShader(MeshShader);
 
