SCL: fix a deadlock by deleting library contents outside of the NamedLibrariesMutex's scope.

- Library deletion flushes render commands, but the rendering thread may be trying to access the libraries stack (for which it needs to acquire NamedLibrariesMutex), so holding that mutex during the deletion results in a deadlock. This change makes the deletion happen after the mutex is released.

- The mutex only guards the NamedLibrariesStack, not contents of any given library. It is still possible that a library being deleted is accessed by another thread that got a pointer to it beforehand, this change does not change anything in that regard. However, this should not happen if the higher level removes all the content from the unmounted plugin from memory prior to unmounting it (and closing the library), which it should do.

diff --git a/Engine/Source/Runtime/RenderCore/Private/ShaderCodeLibrary.cpp b/Engine/Source/Runtime/RenderCore/Private/ShaderCodeLibrary.cpp
--- a/Engine/Source/Runtime/RenderCore/Private/ShaderCodeLibrary.cpp
+++ b/Engine/Source/Runtime/RenderCore/Private/ShaderCodeLibrary.cpp
@@ -2580,9 +2580,13 @@ public:
 	{
 		if (IsLibraryInitializedForRuntime())
 		{
-			FRWScopeLock WriteLock(NamedLibrariesMutex, SLT_Write);
 			TUniquePtr<UE::ShaderLibrary::Private::FNamedShaderLibrary> RemovedLibrary = nullptr;
-			NamedLibrariesStack.RemoveAndCopyValue(Name, RemovedLibrary);
+
+			{
+				FRWScopeLock WriteLock(NamedLibrariesMutex, SLT_Write);
+				NamedLibrariesStack.RemoveAndCopyValue(Name, RemovedLibrary);
+			}
+
 			if (RemovedLibrary)
 			{
 				UE_LOG(LogShaderLibrary, Display, TEXT("Closing logical shader library '%s' with %d components"), *Name, RemovedLibrary->GetNumComponents());
