Added profiler event to FRayTracingGeometryManager::Tick()

diff --git a/Engine/Source/Runtime/RenderCore/Private/RayTracingGeometryManager.cpp b/Engine/Source/Runtime/RenderCore/Private/RayTracingGeometryManager.cpp
--- a/Engine/Source/Runtime/RenderCore/Private/RayTracingGeometryManager.cpp
+++ b/Engine/Source/Runtime/RenderCore/Private/RayTracingGeometryManager.cpp
@@ -116,6 +116,9 @@ void FRayTracingGeometryManager::Tick()
 		return;
 	}
 
+	TRACE_CPUPROFILER_EVENT_SCOPE(FRayTracingGeometryManager::Tick);
+	QUICK_SCOPE_CYCLE_COUNTER(STAT_FRayTracingGeometryManager_Tick);
+
 	if (IsRayTracingEnabled())
 	{
 		FScopeLock ScopeLock(&RequestCS);
