Fix some logging thresholds and add a possibility to customize it.

diff --git a/Engine/Source/Runtime/RenderCore/Private/ShaderCodeArchive.cpp b/Engine/Source/Runtime/RenderCore/Private/ShaderCodeArchive.cpp
--- a/Engine/Source/Runtime/RenderCore/Private/ShaderCodeArchive.cpp
+++ b/Engine/Source/Runtime/RenderCore/Private/ShaderCodeArchive.cpp
@@ -64,6 +64,14 @@ static FAutoConsoleVariableRef CVarShaderCodeLibraryMaxShaderGroupSize(
 	ECVF_RenderThreadSafe | ECVF_ReadOnly
 );
 
+float GShaderCodeLibraryMaxShaderPreloadWaitTime = 0.001f;
+static FAutoConsoleVariableRef CVarShaderCodeLibraryMaxShaderPreloadWaitTime(
+	TEXT("r.ShaderCodeLibrary.MaxShaderPreloadWaitTime"),
+	GShaderCodeLibraryMaxShaderPreloadWaitTime,
+	TEXT("If we wait on shader preloads longer than this amount of seconds, we will log it as a warning."),
+	ECVF_RenderThreadSafe | ECVF_ReadOnly
+);
+
 #if RHI_RAYTRACING	// this function is only needed to check if we need to avoid excluding raytracing shaders
 namespace
 {
@@ -1209,9 +1217,9 @@ TRefCountPtr<FRHIShader> FShaderCodeArchive::CreateShader(int32 Index)
 		const bool bNeededToWait = WaitForPreload(ShaderPreloadEntry);
 		if (bNeededToWait)
 		{
-			double WaitDuration = FPlatformTime::Seconds() - TimeStarted;
+			const double WaitDuration = FPlatformTime::Seconds() - TimeStarted;
 			// only complain if we spent more than 1ms waiting
-			if (TimeStarted > 0.001)
+			if (WaitDuration > GShaderCodeLibraryMaxShaderPreloadWaitTime)
 			{
 				UE_LOG(LogShaderLibrary, Warning, TEXT("Spent %.2f ms in a blocking wait for shader preload, NumRefs: %d, FramePreloadStarted: %d, CurrentFrame: %d"), WaitDuration * 1000.0, ShaderPreloadEntry.NumRefs, ShaderPreloadEntry.FramePreloadStarted, GFrameNumber);
 			}
@@ -2010,12 +2018,12 @@ TRefCountPtr<FRHIShader> FIoStoreShaderCodeArchive::CreateShader(int32 ShaderInd
 	{
 		TRACE_CPUPROFILER_EVENT_SCOPE(BlockingShaderLoad);
 
-		double TimeStarted = FPlatformTime::Seconds();
+		const double TimeStarted = FPlatformTime::Seconds();
 		FTaskGraphInterface::Get().WaitUntilTaskCompletes(Event);
-		double WaitDuration = FPlatformTime::Seconds() - TimeStarted;
+		const double WaitDuration = FPlatformTime::Seconds() - TimeStarted;
 
 		// only complain if we spent more than 1ms waiting
-		if (TimeStarted > 0.001)
+		if (WaitDuration > GShaderCodeLibraryMaxShaderPreloadWaitTime)
 		{
 			UE_LOG(LogShaderLibrary, Warning, TEXT("Spent %.2f ms in a blocking wait for shader preload, NumRefs: %d, FramePreloadStarted: %d, CurrentFrame: %d"), WaitDuration * 1000.0, PreloadEntryPtr->NumRefs, PreloadEntryPtr->FramePreloadStarted, GFrameNumber);
 		}
