Flush the ParallelPassCandidates set when changing the force bParallelTranslate flag

#localization none
#tests reprod and fixes issue and ran a few replays without any others issues. perf improvements are still there when nanite parallal translate is enabled again via cvar

diff --git a/Engine/Source/Runtime/RenderCore/Private/RenderGraphBuilder.cpp b/Engine/Source/Runtime/RenderCore/Private/RenderGraphBuilder.cpp
--- a/Engine/Source/Runtime/RenderCore/Private/RenderGraphBuilder.cpp
+++ b/Engine/Source/Runtime/RenderCore/Private/RenderGraphBuilder.cpp
@@ -2500,8 +2500,14 @@ void FRDGBuilder::SetupParallelExecute()
 			continue;
 		}
 
+		bool bPassHasParallelTranslate = EnumHasAnyFlags(Pass->Flags, ERDGPassFlags::ParallelTranslate);
+		if (bParallelTranslate != bPassHasParallelTranslate)
+		{
+			FlushParallelPassCandidates();
+		}
+
 		bDispatchAfterExecute |= Pass->bDispatchAfterExecute;
-		bParallelTranslate |= EnumHasAnyFlags(Pass->Flags, ERDGPassFlags::ParallelTranslate);
+		bParallelTranslate |= bPassHasParallelTranslate;
 
 		ParallelPassCandidates.Emplace(Pass);
 		ParallelPassCandidatesWorkload += Pass->Workload;
