Added CVar to RDG to control whether to use async destruction task.

#localization none
#tests manual

diff --git a/Engine/Source/Runtime/RenderCore/Private/RenderGraphBuilder.cpp b/Engine/Source/Runtime/RenderCore/Private/RenderGraphBuilder.cpp
--- a/Engine/Source/Runtime/RenderCore/Private/RenderGraphBuilder.cpp
+++ b/Engine/Source/Runtime/RenderCore/Private/RenderGraphBuilder.cpp
@@ -585,7 +585,7 @@ FRDGBuilder::FRDGBuilder(FRHICommandListImmediate& InRHICmdList, FRDGEventName I
 
 FRDGBuilder::~FRDGBuilder()
 {
-	if (bParallelExecuteEnabled)
+	if (bParallelExecuteEnabled && GRDGParallelDestruction > 0)
 	{
 		// Move expensive operations into the async deleter, which will be called in the base class destructor.
 		BeginAsyncDelete([
diff --git a/Engine/Source/Runtime/RenderCore/Private/RenderGraphPrivate.cpp b/Engine/Source/Runtime/RenderCore/Private/RenderGraphPrivate.cpp
--- a/Engine/Source/Runtime/RenderCore/Private/RenderGraphPrivate.cpp
+++ b/Engine/Source/Runtime/RenderCore/Private/RenderGraphPrivate.cpp
@@ -324,6 +324,14 @@ FAutoConsoleVariableRef CVarRDGEvents(
 
 #if RDG_ENABLE_PARALLEL_TASKS
 
+int32 GRDGParallelDestruction = 1;
+FAutoConsoleVariableRef CVarRDGParallelDestruction(
+	TEXT("r.RDG.ParallelDestruction"), GRDGParallelDestruction,
+	TEXT("RDG will destruct the graph using an async task.")
+	TEXT(" 0: graph destruction is done synchronously;")
+	TEXT(" 1: graph destruction may be done asynchronously (default);"),
+	ECVF_RenderThreadSafe);
+
 int32 GRDGParallelSetup = 1;
 FAutoConsoleVariableRef CVarRDGParallelSetup(
 	TEXT("r.RDG.ParallelSetup"), GRDGParallelSetup,
diff --git a/Engine/Source/Runtime/RenderCore/Private/RenderGraphPrivate.h b/Engine/Source/Runtime/RenderCore/Private/RenderGraphPrivate.h
--- a/Engine/Source/Runtime/RenderCore/Private/RenderGraphPrivate.h
+++ b/Engine/Source/Runtime/RenderCore/Private/RenderGraphPrivate.h
@@ -116,6 +116,7 @@ extern int32 GRDGTransientIndirectArgBuffers;
 
 #if RDG_ENABLE_PARALLEL_TASKS
 
+extern int32 GRDGParallelDestruction;
 extern int32 GRDGParallelSetup;
 extern int32 GRDGParallelExecute;
 extern int32 GRDGParallelExecutePassMin;
@@ -123,6 +124,7 @@ extern int32 GRDGParallelExecutePassMax;
 
 #else
 
+const int32 GRDGParallelDestruction = 0;
 const int32 GRDGParallelSetup = 0;
 const int32 GRDGParallelExecute = 0;
 const int32 GRDGParallelExecutePassMin = 0;
