Fix assert in IsInActualRenderingThread() by using an optional tag instead of replacing the renderthread one

#rnx

diff --git a/Engine/Source/Runtime/RenderCore/Private/RenderGraphBuilder.cpp b/Engine/Source/Runtime/RenderCore/Private/RenderGraphBuilder.cpp
--- a/Engine/Source/Runtime/RenderCore/Private/RenderGraphBuilder.cpp
+++ b/Engine/Source/Runtime/RenderCore/Private/RenderGraphBuilder.cpp
@@ -2575,7 +2575,7 @@ void FRDGBuilder::DispatchParallelExecute()
 		ParallelExecuteEvents.Emplace(UE::Tasks::Launch(TEXT("FRDGBuilder::ParallelExecute"), [this, &ParallelPassSet, RHICmdListPass]
 		{
 			SCOPED_NAMED_EVENT(ParallelExecute, FColor::Emerald);
-			FTaskTagScope Scope(ETaskTag::EParallelRenderingThread);
+			FOptionalTaskTagScope Scope(ETaskTag::EParallelRenderingThread);
 
 			for (FRDGPass* Pass : ParallelPassSet.Passes)
 			{
