RDG uses malloc when Address Sanitizer is enabled

#rnx

diff --git a/Engine/Source/Runtime/RenderCore/Public/RenderGraphAllocator.h b/Engine/Source/Runtime/RenderCore/Public/RenderGraphAllocator.h
--- a/Engine/Source/Runtime/RenderCore/Public/RenderGraphAllocator.h
+++ b/Engine/Source/Runtime/RenderCore/Public/RenderGraphAllocator.h
@@ -14,7 +14,7 @@
 #include "Templates/Function.h"
 #include "Templates/UnrealTemplate.h"
 
-#define RDG_USE_MALLOC 0
+#define RDG_USE_MALLOC USING_ADDRESS_SANITISER
 
 /** Private allocator used by RDG to track its internal memory. All memory is released after RDG builder execution. */
 class RENDERCORE_API FRDGAllocator final
