Added facility to convert shader hashes back to their human-readable stable keys.

- A cooked game can be given an argument -shkfile=/path/to/file.shk (a "stable shader keys" file produced during the cook) and use it to provide information on compiled PSOs (and singular shaders if needed).

diff --git a/Engine/Source/Runtime/RenderCore/Private/PipelineCacheUtilities.cpp b/Engine/Source/Runtime/RenderCore/Private/PipelineCacheUtilities.cpp
--- a/Engine/Source/Runtime/RenderCore/Private/PipelineCacheUtilities.cpp
+++ b/Engine/Source/Runtime/RenderCore/Private/PipelineCacheUtilities.cpp
@@ -1,6 +1,6 @@
 // Copyright Epic Games, Inc. All Rights Reserved.
 #include "PipelineCacheUtilities.h"
-#if WITH_EDITOR	// this has no business existing in a cooked game
+#if UE_WITH_PIPELINE_CACHE_UTILITIES
 #include "Misc/SecureHash.h"
 #include "Serialization/NameAsStringIndexProxyArchive.h"
 #include "Serialization/VarInt.h"
@@ -48,7 +48,12 @@ namespace Private
 			return Ar << Info.Magic << Info.Version << Info.NumEntries;
 		}
 	};
+#pragma pack(pop)
+
+#if WITH_EDITOR
 
+#pragma pack(push)
+#pragma pack(1)
 	/** Header of the binary stable pipeline cache file. */
 	struct FStablePipelineCacheSerializedHeader
 	{
@@ -290,6 +295,8 @@ namespace Private
 			Current = 1
 		};
 	};
+
+#endif // WITH_EDITOR
 }
 }
 }
@@ -369,6 +376,8 @@ bool UE::PipelineCacheUtilities::LoadStableKeysFile(const FStringView& Filename,
 	return true;
 }
 
+#if WITH_EDITOR
+
 bool UE::PipelineCacheUtilities::SaveStableKeysFile(const FStringView& Filename, const TSet<FStableShaderKeyAndValue>& Values)
 {
 	TUniquePtr<FArchive> FileArchiveInner(IFileManager::Get().CreateFileWriter(*FString(Filename.Len(), Filename.GetData())));
@@ -1016,5 +1025,5 @@ bool UE::PipelineCacheUtilities::LoadChunkInfo(const FString& Filename, const FS
 
 	return true;
 }
-
 #endif // WITH_EDITOR
+#endif // UE_WITH_PIPELINE_CACHE_UTILITIES
diff --git a/Engine/Source/Runtime/RenderCore/Private/ShaderPipelineCache.cpp b/Engine/Source/Runtime/RenderCore/Private/ShaderPipelineCache.cpp
--- a/Engine/Source/Runtime/RenderCore/Private/ShaderPipelineCache.cpp
+++ b/Engine/Source/Runtime/RenderCore/Private/ShaderPipelineCache.cpp
@@ -382,8 +382,8 @@ class FShaderPipelineCacheTask
 	friend class FShaderPipelineCache;
 
 public:
-	FShaderPipelineCacheTask(bool bIsUserCacheIn, EShaderPlatform InShaderPlatform) :
-		ShaderPlatform(InShaderPlatform), bIsUserCache(bIsUserCacheIn)
+	FShaderPipelineCacheTask(FShaderPipelineCache* InCache, bool bIsUserCacheIn, EShaderPlatform InShaderPlatform) :
+		ShaderPlatform(InShaderPlatform), bIsUserCache(bIsUserCacheIn), Cache(InCache)
 	{
 		bReady = !CVarPSOFileCacheGameFileMaskEnabled.GetValueOnAnyThread() || CVarPSOFileCachePreOptimizeEnabled.GetValueOnAnyThread();
 	}
@@ -434,6 +434,9 @@ private:
 	bool bOpened = false;
 	bool bIsUserCache;
 
+	/** Cache which kicked off this task. */
+	FShaderPipelineCache* Cache = nullptr;
+
 	// as there is only ever one active precompile task at a time
 	// these stats are available to read anywhere.
 	static float TotalPrecompileWallTime;
@@ -836,7 +839,7 @@ bool FShaderPipelineCache::OpenPipelineFileCacheInternal(bool bWantUserCache, co
 {
 	if(!ShaderCacheTasks.Contains(PSOCacheKey))
 	{
-		TUniquePtr<FShaderPipelineCacheTask> NewPSOTask = MakeUnique<FShaderPipelineCacheTask>(bWantUserCache, Platform);
+		TUniquePtr<FShaderPipelineCacheTask> NewPSOTask = MakeUnique<FShaderPipelineCacheTask>(this, bWantUserCache, Platform);
 		bool bOpened = NewPSOTask->Open(PSOCacheKey, CacheName, Platform);
 		if(bOpened)
 		{
@@ -1030,6 +1033,14 @@ bool FShaderPipelineCacheTask::Precompile(FRHICommandListImmediate& RHICmdList,
 				// This indicates we do not want a fatal error if this compilation fails
 				// (ie, if this entry in the file cache is bad)
 				GraphicsInitializer.bFromPSOFileCache = true;
+
+#if !UE_BUILD_SHIPPING
+				// dump to log to describe
+				if (Cache && Cache->ShaderHashToStableKey.IsInitialized())
+				{
+					Cache->ShaderHashToStableKey.DumpPSOToLogIfConfigured(GraphicsInitializer);
+				}
+#endif
 			}
 			// Use SetGraphicsPipelineState to call down into PipelineStateCache and also handle the fallback case used by OpenGL.
 			// we lose track of our PSO precompile jobs through here. they may not be RHI dependencies and could persist beyond the next tick.
@@ -1538,6 +1549,12 @@ FShaderPipelineCache::FShaderPipelineCache(EShaderPlatform Platform)
 	BatchSize = CVarPSOFileCacheBatchSize.GetValueOnAnyThread();
 	BatchTime = CVarPSOFileCacheBatchTime.GetValueOnAnyThread();
 	
+	FString StableShaderKeyFile;
+	if (FParse::Value(FCommandLine::Get(), TEXT("-shkfile="), StableShaderKeyFile) && !StableShaderKeyFile.IsEmpty())
+	{
+		ShaderHashToStableKey.Initialize(StableShaderKeyFile);
+	}
+
 	uint64 PreCompileMask = (uint64)CVarPSOFileCachePreCompileMask.GetValueOnAnyThread();
 	FPipelineFileCacheManager::SetGameUsageMaskWithComparison(PreCompileMask, PreCompileMaskComparison);
 
diff --git a/Engine/Source/Runtime/RenderCore/Private/ShaderStableKeyDebugInfoReader.cpp b/Engine/Source/Runtime/RenderCore/Private/ShaderStableKeyDebugInfoReader.cpp
new file mode 100644
--- /dev/null
+++ b/Engine/Source/Runtime/RenderCore/Private/ShaderStableKeyDebugInfoReader.cpp
@@ -0,0 +1,105 @@
+// Copyright Epic Games, Inc. All Rights Reserved.
+
+/**
+ * Code for reading .shk files and using that information to map shader hashes back to human-readable identifies.
+ */
+
+#pragma once
+
+#include "ShaderStableKeyDebugInfoReader.h"
+#include "PipelineCacheUtilities.h"
+
+// Disable in Shipping just in case (also shaderpipelineutils aren't available there)
+#define UE_WITH_SHK_READER			(!UE_BUILD_SHIPPING)
+
+bool UE::ShaderUtils::FShaderStableKeyDebugInfoReader::Initialize(const FString& ShaderStableKeyFile)
+{
+	bool bResult = false;
+
+#if UE_WITH_SHK_READER
+	TArray<FStableShaderKeyAndValue> StableArray;
+	if (UE::PipelineCacheUtilities::LoadStableKeysFile(ShaderStableKeyFile, StableArray))
+	{
+		ShaderHashesToSource.Empty();
+		for (const FStableShaderKeyAndValue& StableKey : StableArray)
+		{
+			TSet<FStableShaderKeyAndValue>& DedupShaders = ShaderHashesToSource.FindOrAdd(StableKey.OutputHash);
+			DedupShaders.Add(StableKey);
+		}
+
+		bResult = true;
+	}
+#endif // UE_WITH_SHK_READER
+
+	return bResult;
+}
+
+FString UE::ShaderUtils::FShaderStableKeyDebugInfoReader::GetShaderStableNameOptions(const FSHAHash& ShaderHash, int32 MaxOptionsToInclude)
+{
+	FString Output;	// FIXME: change to builder
+
+#if UE_WITH_SHK_READER
+	if (TSet<FStableShaderKeyAndValue>* SourceSet = ShaderHashesToSource.Find(ShaderHash))
+	{
+		Output += FString::Printf(TEXT("hash %s maps to %d shaders"), *ShaderHash.ToString(), SourceSet->Num());
+
+		if (SourceSet->Num() > MaxOptionsToInclude)
+		{
+			Output += FString::Printf(TEXT(", limiting output to %d"), MaxOptionsToInclude);
+		}
+		Output += TEXT(":\n");
+
+		int32 OptionsToInclude = FMath::Min(MaxOptionsToInclude, SourceSet->Num());
+		for (const FStableShaderKeyAndValue& StableKey : *SourceSet)
+		{
+			if (--OptionsToInclude < 0)
+			{
+				break;
+			}
+			Output += FString::Printf(TEXT("\t%s\n"), *StableKey.ToString());
+		}
+	}
+	else
+	{
+		Output += FString::Printf(TEXT("hash %s not in the key file(s)! (were key file(s) produced in the same cook as this build?)\n"), *ShaderHash.ToString());
+	}
+#endif // UE_WITH_SHK_READER
+
+	return Output;
+}
+
+FString UE::ShaderUtils::FShaderStableKeyDebugInfoReader::GetPSOStableNameOptions(const FGraphicsPipelineStateInitializer& Initializer, int32 MaxOptionsToInclude)
+{
+	FString Output;	// FIXME: change to builder
+
+#if UE_WITH_SHK_READER
+
+#define UE_FPTSM_DUMPSHADER(ShaderType) \
+	if (FRHIShader* Shader = Initializer.BoundShaderState.Get##ShaderType()) \
+	{ \
+		const FSHAHash& Hash = Shader->GetHash(); \
+		Output += FString::Printf(TEXT("%s: %s"), TEXT(#ShaderType), *GetShaderStableNameOptions(Hash, MaxOptionsToInclude)); \
+	}
+
+	Output += FString::Printf(TEXT("---- begin FGraphicsPipelineStateInitializer -----\n"));
+	UE_FPTSM_DUMPSHADER(VertexShader);
+	UE_FPTSM_DUMPSHADER(MeshShader);
+	UE_FPTSM_DUMPSHADER(AmplificationShader);
+	UE_FPTSM_DUMPSHADER(PixelShader);
+	UE_FPTSM_DUMPSHADER(GeometryShader);
+
+#undef UE_FPTSM_DUMPSHADER
+
+	Output += FString::Printf(TEXT("----- end -----\n"));
+
+#endif // UE_WITH_SHK_READER
+
+	return Output;
+}
+
+void UE::ShaderUtils::FShaderStableKeyDebugInfoReader::DumpPSOToLogIfConfigured(const FGraphicsPipelineStateInitializer& Initializer)
+{
+#if UE_WITH_SHK_READER
+	UE_LOG(LogInit, Log, TEXT("%s"), *GetPSOStableNameOptions(Initializer));
+#endif // UE_WITH_SHK_READER
+}
diff --git a/Engine/Source/Runtime/RenderCore/Public/PipelineCacheUtilities.h b/Engine/Source/Runtime/RenderCore/Public/PipelineCacheUtilities.h
--- a/Engine/Source/Runtime/RenderCore/Public/PipelineCacheUtilities.h
+++ b/Engine/Source/Runtime/RenderCore/Public/PipelineCacheUtilities.h
@@ -2,7 +2,9 @@
 
 #pragma once
 
-#if WITH_EDITOR	// this has no business existing in a cooked game
+// cooked games may want to load key files for extra information about the PSOs
+#define UE_WITH_PIPELINE_CACHE_UTILITIES			(WITH_EDITOR || !UE_BUILD_SHIPPING)
+#if UE_WITH_PIPELINE_CACHE_UTILITIES
 
 #include "Containers/Array.h"
 #include "Containers/Map.h"
@@ -24,6 +26,8 @@ namespace UE
 {
 namespace PipelineCacheUtilities
 {
+#if WITH_EDITOR	// limit what's compiled for the cooked games
+
 	/** Describes a particular combination of shaders. */
 	struct FPermutation
 	{
@@ -51,6 +55,8 @@ namespace PipelineCacheUtilities
 		}
 	};
 
+#endif // WITH_EDITOR
+
 	/** 
 	 * Loads stable shader keys file (using a proprietary format). Stable key is a way to identify a shader independently of its output hash
 	 * 
@@ -60,6 +66,8 @@ namespace PipelineCacheUtilities
 	 */
 	RENDERCORE_API bool LoadStableKeysFile(const FStringView& Filename, TArray<FStableShaderKeyAndValue>& InOutArray);
 
+#if WITH_EDITOR	// limit what's compiled for the cooked games
+
 	/** 
 	 * Saves stable shader keys file (using a proprietary format). Stable key is a way to identify a shader independently of its output hash
 	 * 
@@ -133,7 +141,10 @@ namespace PipelineCacheUtilities
 	 * @param OutPackages packages in this chunk
 	 */
 	RENDERCORE_API bool LoadChunkInfo(const FString& Filename, const FString& InShaderFormat, int32 &OutChunkId, FString& OutChunkedCacheFilename, TSet<FName>& OutPackages);
+
+#endif // WITH_EDITOR
+
 }
 };
 
-#endif // WITH_EDITOR
+#endif // UE_WITH_PIPELINE_CACHE_UTILITIES
diff --git a/Engine/Source/Runtime/RenderCore/Public/ShaderPipelineCache.h b/Engine/Source/Runtime/RenderCore/Public/ShaderPipelineCache.h
--- a/Engine/Source/Runtime/RenderCore/Public/ShaderPipelineCache.h
+++ b/Engine/Source/Runtime/RenderCore/Public/ShaderPipelineCache.h
@@ -14,6 +14,7 @@
 #include "TickableObjectRenderThread.h"
 #include "PipelineFileCache.h"
 #include "Delegates/DelegateCombinations.h"
+#include "ShaderStableKeyDebugInfoReader.h"
 
 class FBlendStateInitializerRHI;
 class FRHIBlendState;
@@ -240,6 +241,9 @@ private:
 	int32 LastAutoSaveNum = 0;
 	TSet<uint32> CompiledHashes;
 
+	/** Development tool that can describe PSOs in non-Shipping builds. */
+	UE::ShaderUtils::FShaderStableKeyDebugInfoReader ShaderHashToStableKey;
+
 	static FString UserCacheTaskKey;
 
 	struct FPipelineCachePSOContext
diff --git a/Engine/Source/Runtime/RenderCore/Public/ShaderStableKeyDebugInfoReader.h b/Engine/Source/Runtime/RenderCore/Public/ShaderStableKeyDebugInfoReader.h
new file mode 100644
--- /dev/null
+++ b/Engine/Source/Runtime/RenderCore/Public/ShaderStableKeyDebugInfoReader.h
@@ -0,0 +1,43 @@
+// Copyright Epic Games, Inc. All Rights Reserved.
+
+/**
+ * Code for reading .shk files and using that information to map shader hashes back to human-readable identifies.
+ */
+
+#pragma once
+
+#include "ShaderCodeLibrary.h"
+
+namespace UE
+{
+
+namespace ShaderUtils
+{
+
+/** Class that uses build metadata (*.shk files storing mapping of stable shader keys to their hashes) to provide high level info on shaders and PSOs. */
+class FShaderStableKeyDebugInfoReader
+{
+	/** Maps shader hashes to their sources. */
+	TMap<FSHAHash, TSet<FStableShaderKeyAndValue>>  ShaderHashesToSource;
+
+public:
+
+	/** Whether the info reader is usable */
+	bool IsInitialized() const { return !ShaderHashesToSource.IsEmpty(); }
+
+	/** Initializes the class with a pointer to directory holding SHK files */
+	bool Initialize(const FString& ShaderStableKeyFile);
+
+	/** Returns a string describing a shader. Because of shader deduplication there can be several multiple (thousands) possible options */
+	FString GetShaderStableNameOptions(const FSHAHash& ShaderHash, int32 MaxOptionsToInclude = 16);
+
+	/** Returns a string describing a PSO (shaders only). Because of shader deduplication there can be several multiple (thousands) possible options. */
+	FString GetPSOStableNameOptions(const FGraphicsPipelineStateInitializer& Initializer, int32 MaxOptionsToInclude = 16);
+
+	/** Prints a pipeline info to log if configured to do so. */
+	void DumpPSOToLogIfConfigured(const FGraphicsPipelineStateInitializer& Initializer);
+};
+
+}
+
+}
