Fix non-unity compilation error.

diff --git a/Engine/Source/Runtime/RenderCore/Public/PipelineCacheUtilities.h b/Engine/Source/Runtime/RenderCore/Public/PipelineCacheUtilities.h
--- a/Engine/Source/Runtime/RenderCore/Public/PipelineCacheUtilities.h
+++ b/Engine/Source/Runtime/RenderCore/Public/PipelineCacheUtilities.h
@@ -2,6 +2,8 @@
 
 #pragma once
 
+#include "Misc/Build.h"
+
 // cooked games may want to load key files for extra information about the PSOs
 #define UE_WITH_PIPELINE_CACHE_UTILITIES			(WITH_EDITOR || !UE_BUILD_SHIPPING)
 #if UE_WITH_PIPELINE_CACHE_UTILITIES
