Fix CIS (pragma once in .cpp)

diff --git a/Engine/Source/Runtime/RenderCore/Private/ShaderStableKeyDebugInfoReader.cpp b/Engine/Source/Runtime/RenderCore/Private/ShaderStableKeyDebugInfoReader.cpp
--- a/Engine/Source/Runtime/RenderCore/Private/ShaderStableKeyDebugInfoReader.cpp
+++ b/Engine/Source/Runtime/RenderCore/Private/ShaderStableKeyDebugInfoReader.cpp
@@ -4,8 +4,6 @@
  * Code for reading .shk files and using that information to map shader hashes back to human-readable identifies.
  */
 
-#pragma once
-
 #include "ShaderStableKeyDebugInfoReader.h"
 #include "PipelineCacheUtilities.h"
 
