Removing redundant FDataDrivenShaderPlatformInfo checks.

diff --git a/Engine/Source/Runtime/RenderCore/Public/ClearReplacementShaders.h b/Engine/Source/Runtime/RenderCore/Public/ClearReplacementShaders.h
--- a/Engine/Source/Runtime/RenderCore/Public/ClearReplacementShaders.h
+++ b/Engine/Source/Runtime/RenderCore/Public/ClearReplacementShaders.h
@@ -151,8 +151,7 @@ public:
 	
 	static bool ShouldCompilePermutation(const FGlobalShaderPermutationParameters& Parameters)
 	{
-		return IsFeatureLevelSupported(Parameters.Platform, ERHIFeatureLevel::SM5)
-				|| GetMaxSupportedFeatureLevel(Parameters.Platform) == ERHIFeatureLevel::ES3_1;
+		return true;
 	}
 
 	static inline void ModifyCompilationEnvironment(const FGlobalShaderPermutationParameters& Parameters, FShaderCompilerEnvironment& OutEnvironment)
@@ -192,8 +191,7 @@ public:
 
 	static bool ShouldCompilePermutation(const FGlobalShaderPermutationParameters& Parameters)
 	{
-		return IsFeatureLevelSupported(Parameters.Platform, ERHIFeatureLevel::SM5)
-				|| GetMaxSupportedFeatureLevel(Parameters.Platform) == ERHIFeatureLevel::ES3_1;
+		return true;
 	}
 
 	static inline void ModifyCompilationEnvironment(const FGlobalShaderPermutationParameters& Parameters, FShaderCompilerEnvironment& OutEnvironment)
@@ -228,8 +226,7 @@ public:
 
 	static bool ShouldCompilePermutation(const FGlobalShaderPermutationParameters& Parameters)
 	{
-		return IsFeatureLevelSupported(Parameters.Platform, ERHIFeatureLevel::SM5)
-				|| GetMaxSupportedFeatureLevel(Parameters.Platform) == ERHIFeatureLevel::ES3_1;
+		return true;
 	}
 
 	static inline void ModifyCompilationEnvironment(const FGlobalShaderPermutationParameters& Parameters, FShaderCompilerEnvironment& OutEnvironment)
diff --git a/Engine/Source/Runtime/Renderer/Private/VT/VirtualTextureSpace.cpp b/Engine/Source/Runtime/Renderer/Private/VT/VirtualTextureSpace.cpp
--- a/Engine/Source/Runtime/Renderer/Private/VT/VirtualTextureSpace.cpp
+++ b/Engine/Source/Runtime/Renderer/Private/VT/VirtualTextureSpace.cpp
@@ -192,8 +192,7 @@ public:
 
 	static bool ShouldCompilePermutation(const FGlobalShaderPermutationParameters& Parameters)
 	{
-		return IsFeatureLevelSupported(Parameters.Platform, ERHIFeatureLevel::SM5) 
-			|| GetMaxSupportedFeatureLevel(Parameters.Platform) == ERHIFeatureLevel::ES3_1;
+		return true;
 	}
 
 	LAYOUT_FIELD(FShaderParameter, PageTableSize);
@@ -215,12 +214,8 @@ public:
 	
 	static bool ShouldCompilePermutation(const FGlobalShaderPermutationParameters& Parameters)
 	{
-		return IsFeatureLevelSupported(Parameters.Platform, ERHIFeatureLevel::SM5)
-			|| GetMaxSupportedFeatureLevel(Parameters.Platform) == ERHIFeatureLevel::ES3_1;
+		return true;
 	}
-
-	
-	
 };
 
 template<bool Use16Bits>
