Moving TBoundShaderStateHistory into RHICore.

diff --git a/Engine/Source/Runtime/Apple/MetalRHI/Public/MetalRHIContext.h b/Engine/Source/Runtime/Apple/MetalRHI/Public/MetalRHIContext.h
--- a/Engine/Source/Runtime/Apple/MetalRHI/Public/MetalRHIContext.h
+++ b/Engine/Source/Runtime/Apple/MetalRHI/Public/MetalRHIContext.h
@@ -8,6 +8,7 @@
 #include "MetalResources.h"
 #include "MetalViewport.h"
 #include "RHICore.h"
+#include "BoundShaderStateHistory.h"
 
 #define UE_METAL_RHI_SUPPORT_CLEAR_UAV_WITH_BLIT_ENCODER	1
 
diff --git a/Engine/Source/Runtime/OpenGLDrv/Public/OpenGLDrv.h b/Engine/Source/Runtime/OpenGLDrv/Public/OpenGLDrv.h
--- a/Engine/Source/Runtime/OpenGLDrv/Public/OpenGLDrv.h
+++ b/Engine/Source/Runtime/OpenGLDrv/Public/OpenGLDrv.h
@@ -16,6 +16,7 @@
 #include "GPUProfiler.h"
 #include "RenderResource.h"
 #include "Templates/EnableIf.h"
+#include "BoundShaderStateHistory.h"
 
 // @todo platplug: Replace all of these includes with a call to COMPILED_PLATFORM_HEADER(OpenGLDrvPrivate.h)
 //TODO: Move these to OpenGLDrvPrivate.h
diff --git a/Engine/Source/Runtime/RHICore/Public/BoundShaderStateHistory.h b/Engine/Source/Runtime/RHICore/Public/BoundShaderStateHistory.h
new file mode 100644
--- /dev/null
+++ b/Engine/Source/Runtime/RHICore/Public/BoundShaderStateHistory.h
@@ -0,0 +1,62 @@
+// Copyright Epic Games, Inc. All Rights Reserved.
+
+#pragma once
+
+#include "RenderResource.h"
+#include "RHI.h"
+
+/**
+ * A list of the most recently used bound shader states.
+ * This is used to keep bound shader states that have been used recently from being freed, as they're likely to be used again soon.
+ */
+template<uint32 Size, bool TThreadSafe = true>
+class TBoundShaderStateHistory : public FRenderResource
+{
+public:
+	TBoundShaderStateHistory() = default;
+
+	/** Adds a bound shader state to the history. */
+	FORCEINLINE void Add(FRHIBoundShaderState* BoundShaderState)
+	{
+		if (TThreadSafe && GRHISupportsParallelRHIExecute)
+		{
+			BoundShaderStateHistoryLock.Lock();
+		}
+		BoundShaderStates[NextBoundShaderStateIndex] = BoundShaderState;
+		NextBoundShaderStateIndex = (NextBoundShaderStateIndex + 1) % Size;
+		if (TThreadSafe && GRHISupportsParallelRHIExecute)
+		{
+			BoundShaderStateHistoryLock.Unlock();
+		}
+	}
+
+	FRHIBoundShaderState* GetLast()
+	{
+		check(!GRHISupportsParallelRHIExecute);
+		// % doesn't work as we want on negative numbers, so handle the wraparound manually
+		uint32 LastIndex = NextBoundShaderStateIndex == 0 ? Size - 1 : NextBoundShaderStateIndex - 1;
+		return BoundShaderStates[LastIndex];
+	}
+
+	// FRenderResource interface.
+	virtual void ReleaseRHI()
+	{
+		if (TThreadSafe && GRHISupportsParallelRHIExecute)
+		{
+			BoundShaderStateHistoryLock.Lock();
+		}
+		for (uint32 Index = 0; Index < Size; Index++)
+		{
+			BoundShaderStates[Index].SafeRelease();
+		}
+		if (TThreadSafe && GRHISupportsParallelRHIExecute)
+		{
+			BoundShaderStateHistoryLock.Unlock();
+		}
+	}
+
+private:
+	FBoundShaderStateRHIRef BoundShaderStates[Size];
+	uint32 NextBoundShaderStateIndex = 0;
+	FCriticalSection BoundShaderStateHistoryLock;
+};
diff --git a/Engine/Source/Runtime/RenderCore/Public/RenderResource.h b/Engine/Source/Runtime/RenderCore/Public/RenderResource.h
--- a/Engine/Source/Runtime/RenderCore/Public/RenderResource.h
+++ b/Engine/Source/Runtime/RenderCore/Public/RenderResource.h
@@ -1017,68 +1017,6 @@ private:
 	struct FDynamicIndexBufferPool* Pools[2];
 };
 
-/**
- * A list of the most recently used bound shader states.
- * This is used to keep bound shader states that have been used recently from being freed, as they're likely to be used again soon.
- */
-
-template<uint32 Size, bool TThreadSafe = true>
-class TBoundShaderStateHistory : public FRenderResource
-{
-public:
-
-	/** Initialization constructor. */
-	TBoundShaderStateHistory():
-		NextBoundShaderStateIndex(0)
-	{}
-
-	/** Adds a bound shader state to the history. */
-	FORCEINLINE void Add(FRHIBoundShaderState* BoundShaderState)
-	{
-		if (TThreadSafe && GRHISupportsParallelRHIExecute)
-		{
-			BoundShaderStateHistoryLock.Lock();
-		}
-		BoundShaderStates[NextBoundShaderStateIndex] = BoundShaderState;
-		NextBoundShaderStateIndex = (NextBoundShaderStateIndex + 1) % Size;
-		if (TThreadSafe && GRHISupportsParallelRHIExecute)
-		{
-			BoundShaderStateHistoryLock.Unlock();
-		}
-	}
-
-	FRHIBoundShaderState* GetLast()
-	{
-		check(!GRHISupportsParallelRHIExecute);
-		// % doesn't work as we want on negative numbers, so handle the wraparound manually
-		uint32 LastIndex = NextBoundShaderStateIndex == 0 ? Size - 1 : NextBoundShaderStateIndex - 1;
-		return BoundShaderStates[LastIndex];
-	}
-
-	// FRenderResource interface.
-	virtual void ReleaseRHI()
-	{
-		if (TThreadSafe && GRHISupportsParallelRHIExecute)
-		{
-			BoundShaderStateHistoryLock.Lock();
-		}
-		for(uint32 Index = 0;Index < Size;Index++)
-		{
-			BoundShaderStates[Index].SafeRelease();
-		}
-		if (TThreadSafe && GRHISupportsParallelRHIExecute)
-		{
-			BoundShaderStateHistoryLock.Unlock();
-		}
-	}
-
-private:
-
-	FBoundShaderStateRHIRef BoundShaderStates[Size];
-	uint32 NextBoundShaderStateIndex;
-	FCriticalSection BoundShaderStateHistoryLock;
-};
-
 /**Note, this should only be used when a platform requires special shader compilation for 32 bit pixel format render targets.
 Does not replace pixel format associations across the board**/
 
diff --git a/Engine/Source/Runtime/Windows/D3D11RHI/Private/D3D11RHIPrivate.h b/Engine/Source/Runtime/Windows/D3D11RHI/Private/D3D11RHIPrivate.h
--- a/Engine/Source/Runtime/Windows/D3D11RHI/Private/D3D11RHIPrivate.h
+++ b/Engine/Source/Runtime/Windows/D3D11RHI/Private/D3D11RHIPrivate.h
@@ -17,6 +17,7 @@
 #include "EngineGlobals.h"
 #include "Engine/Engine.h"
 #include "HDRHelper.h"
+#include "BoundShaderStateHistory.h"
 
 DECLARE_LOG_CATEGORY_EXTERN(LogD3D11RHI, Log, All);
 
