Moving some shader system methods into cpp files to remove header file dependencies on FDataDrivenShaderPlatformInfo.

diff --git a/Engine/Source/Runtime/RenderCore/Private/Shader.cpp b/Engine/Source/Runtime/RenderCore/Private/Shader.cpp
--- a/Engine/Source/Runtime/RenderCore/Private/Shader.cpp
+++ b/Engine/Source/Runtime/RenderCore/Private/Shader.cpp
@@ -26,6 +26,7 @@
 #include "Misc/ScopeRWLock.h"
 #include "ProfilingDebugging/LoadTimeTracker.h"
 #include "Misc/LargeWorldRenderPosition.h"
+#include "DataDrivenShaderPlatformInfo.h"
 
 #if WITH_EDITORONLY_DATA
 #include "Interfaces/IShaderFormat.h"
@@ -1078,6 +1079,11 @@ const FShaderPipelineType* FShaderPipelineType::GetShaderPipelineTypeByName(cons
 	return FoundType ? *FoundType : nullptr;
 }
 
+bool FShaderPipelineType::ShouldOptimizeUnusedOutputs(EShaderPlatform Platform) const
+{
+	return bShouldOptimizeUnusedOutputs && RHISupportsShaderPipelines(Platform);
+}
+
 const FSHAHash& FShaderPipelineType::GetSourceHash(EShaderPlatform ShaderPlatform) const
 {
 	TArray<FString> Filenames;
diff --git a/Engine/Source/Runtime/RenderCore/Private/ShaderMap.cpp b/Engine/Source/Runtime/RenderCore/Private/ShaderMap.cpp
--- a/Engine/Source/Runtime/RenderCore/Private/ShaderMap.cpp
+++ b/Engine/Source/Runtime/RenderCore/Private/ShaderMap.cpp
@@ -22,6 +22,7 @@
 #include "UObject/FortniteMainBranchObjectVersion.h"
 #include "Misc/ScopeRWLock.h"
 #include "ProfilingDebugging/LoadTimeTracker.h"
+#include "DataDrivenShaderPlatformInfo.h"
 
 #if WITH_EDITORONLY_DATA
 #include "Interfaces/IShaderFormat.h"
@@ -298,6 +299,20 @@ static uint16 MakeShaderHash(const FHashedName& TypeName, int32 PermutationId)
 	return (uint16)CityHash128to64({ TypeName.GetHash(), (uint64)PermutationId });
 }
 
+FShaderMapContent::FShaderMapContent(EShaderPlatform InPlatform)
+	: ShaderHash(128u), ShaderPlatformName(FDataDrivenShaderPlatformInfo::GetName(InPlatform).ToString())
+{}
+
+FShaderMapContent::~FShaderMapContent()
+{
+	Empty();
+}
+
+EShaderPlatform FShaderMapContent::GetShaderPlatform() const
+{
+	return FDataDrivenShaderPlatformInfo::GetShaderPlatformFromName(FName(ShaderPlatformName));
+}
+
 FShader* FShaderMapContent::GetShader(const FHashedName& TypeName, int32 PermutationId) const
 {
 	// TRACE_CPUPROFILER_EVENT_SCOPE(FShaderMapContent::GetShader); -- this function is called too frequently, so don't add the scope by default
diff --git a/Engine/Source/Runtime/RenderCore/Private/VertexFactory.cpp b/Engine/Source/Runtime/RenderCore/Private/VertexFactory.cpp
--- a/Engine/Source/Runtime/RenderCore/Private/VertexFactory.cpp
+++ b/Engine/Source/Runtime/RenderCore/Private/VertexFactory.cpp
@@ -10,6 +10,7 @@
 #include "PipelineStateCache.h"
 #include "ShaderCompilerCore.h"
 #include "RenderUtils.h"
+#include "DataDrivenShaderPlatformInfo.h"
 
 IMPLEMENT_TYPE_LAYOUT(FVertexFactoryShaderParameters);
 
@@ -155,6 +156,12 @@ FVertexFactoryType::~FVertexFactoryType()
 	--NumVertexFactories;
 }
 
+bool FVertexFactoryType::CheckManualVertexFetchSupport(ERHIFeatureLevel::Type InFeatureLevel)
+{
+	check(InFeatureLevel != ERHIFeatureLevel::Num);
+	return (InFeatureLevel > ERHIFeatureLevel::ES3_1) && RHISupportsManualVertexFetch(GShaderPlatformForFeatureLevel[InFeatureLevel]);
+}
+
 /** Calculates a Hash based on this vertex factory type's source code and includes */
 const FSHAHash& FVertexFactoryType::GetSourceHash(EShaderPlatform ShaderPlatform) const
 {
diff --git a/Engine/Source/Runtime/RenderCore/Public/Shader.h b/Engine/Source/Runtime/RenderCore/Public/Shader.h
--- a/Engine/Source/Runtime/RenderCore/Public/Shader.h
+++ b/Engine/Source/Runtime/RenderCore/Public/Shader.h
@@ -1725,10 +1725,7 @@ public:
 	bool IsMaterialTypePipeline() const { return Stages[0]->GetMaterialShaderType() != nullptr; }
 	bool IsMeshMaterialTypePipeline() const { return Stages[0]->GetMeshMaterialShaderType() != nullptr; }
 
-	FORCEINLINE bool ShouldOptimizeUnusedOutputs(EShaderPlatform Platform) const
-	{
-		return bShouldOptimizeUnusedOutputs && RHISupportsShaderPipelines(Platform);
-	}
+	bool ShouldOptimizeUnusedOutputs(EShaderPlatform Platform) const;
 
 	/** Calculates a Hash based on this shader pipeline type stages' source code and includes */
 	const FSHAHash& GetSourceHash(EShaderPlatform ShaderPlatform) const;
@@ -2056,18 +2053,14 @@ public:
 	{
 		inline FHashedName operator()(const FShaderPipeline* InShaderPipeline) { return InShaderPipeline->TypeName; }
 	};
+
 	/** Default constructor. */
-	explicit FShaderMapContent(EShaderPlatform InPlatform)
-		: ShaderHash(128u), ShaderPlatformName(FDataDrivenShaderPlatformInfo::GetName(InPlatform).ToString())
-	{}
+	explicit FShaderMapContent(EShaderPlatform InPlatform);
 
 	/** Destructor ensures pipelines cleared up. */
-	~FShaderMapContent()
-	{
-		Empty();
-	}
+	~FShaderMapContent();
 
-	EShaderPlatform GetShaderPlatform() const { return FDataDrivenShaderPlatformInfo::GetShaderPlatformFromName(FName(ShaderPlatformName)); };
+	EShaderPlatform GetShaderPlatform() const;
 
 	void Validate(const FShaderMapBase& InShaderMap) const;
 
diff --git a/Engine/Source/Runtime/RenderCore/Public/VertexFactory.h b/Engine/Source/Runtime/RenderCore/Public/VertexFactory.h
--- a/Engine/Source/Runtime/RenderCore/Public/VertexFactory.h
+++ b/Engine/Source/Runtime/RenderCore/Public/VertexFactory.h
@@ -382,10 +382,11 @@ public:
 
 	bool SupportsManualVertexFetch(ERHIFeatureLevel::Type InFeatureLevel) const 
 	{
-		check(InFeatureLevel != ERHIFeatureLevel::Num);
-		return HasFlags(EVertexFactoryFlags::SupportsManualVertexFetch) && (InFeatureLevel > ERHIFeatureLevel::ES3_1) && RHISupportsManualVertexFetch(GMaxRHIShaderPlatform);
+		return HasFlags(EVertexFactoryFlags::SupportsManualVertexFetch) && CheckManualVertexFetchSupport(InFeatureLevel);
 	}
 
+	static RENDERCORE_API bool CheckManualVertexFetchSupport(ERHIFeatureLevel::Type InFeatureLevel);
+
 	// Hash function.
 	friend uint32 GetTypeHash(const FVertexFactoryType* Type)
 	{ 
