Moving FRayTracingGeometry into its own header+source files. FRayTracingGeometry can't forward declare its RHI members yet, so we need it in its own header to remove those dependencies from RenderResource.h

diff --git a/Engine/Plugins/Enterprise/LidarPointCloud/Source/LidarPointCloudRuntime/Private/Rendering/LidarPointCloudRenderBuffers.h b/Engine/Plugins/Enterprise/LidarPointCloud/Source/LidarPointCloudRuntime/Private/Rendering/LidarPointCloudRenderBuffers.h
--- a/Engine/Plugins/Enterprise/LidarPointCloud/Source/LidarPointCloudRuntime/Private/Rendering/LidarPointCloudRenderBuffers.h
+++ b/Engine/Plugins/Enterprise/LidarPointCloud/Source/LidarPointCloudRuntime/Private/Rendering/LidarPointCloudRenderBuffers.h
@@ -4,6 +4,7 @@
 
 #include "CoreMinimal.h"
 #include "RenderResource.h"
+#include "RayTracingGeometry.h"
 #include "VertexFactory.h"
 #include "LidarPointCloudShared.h"
 #include "HAL/ThreadSafeBool.h"
diff --git a/Engine/Plugins/Experimental/GPULightmass/Source/GPULightmass/Private/Scene/Landscape.h b/Engine/Plugins/Experimental/GPULightmass/Source/GPULightmass/Private/Scene/Landscape.h
--- a/Engine/Plugins/Experimental/GPULightmass/Source/GPULightmass/Private/Scene/Landscape.h
+++ b/Engine/Plugins/Experimental/GPULightmass/Source/GPULightmass/Private/Scene/Landscape.h
@@ -8,6 +8,7 @@
 #include "LandscapeRender.h"
 #include "Materials/MaterialInstanceConstant.h"
 #include "Materials/MaterialInstanceDynamic.h"
+#include "RayTracingGeometry.h"
 
 namespace GPULightmass
 {
diff --git a/Engine/Plugins/Experimental/Water/Source/Runtime/Public/WaterMeshSceneProxy.h b/Engine/Plugins/Experimental/Water/Source/Runtime/Public/WaterMeshSceneProxy.h
--- a/Engine/Plugins/Experimental/Water/Source/Runtime/Public/WaterMeshSceneProxy.h
+++ b/Engine/Plugins/Experimental/Water/Source/Runtime/Public/WaterMeshSceneProxy.h
@@ -9,6 +9,7 @@
 #include "WaterInstanceDataBuffer.h"
 #include "WaterQuadTree.h"
 #include "WaterVertexFactory.h"
+#include "RayTracingGeometry.h"
 
 class UWaterMeshComponent;
 
diff --git a/Engine/Plugins/FX/Niagara/Source/Niagara/Public/NiagaraRenderer.h b/Engine/Plugins/FX/Niagara/Source/Niagara/Public/NiagaraRenderer.h
--- a/Engine/Plugins/FX/Niagara/Source/Niagara/Public/NiagaraRenderer.h
+++ b/Engine/Plugins/FX/Niagara/Source/Niagara/Public/NiagaraRenderer.h
@@ -17,6 +17,7 @@ NiagaraRenderer.h: Base class for Niagara render modules
 #include "SceneView.h"
 #include "NiagaraCutoutVertexBuffer.h"
 #include "NiagaraBoundsCalculator.h"
+#include "RayTracingGeometry.h"
 
 class FNiagaraDataSet;
 class FNiagaraSceneProxy;
diff --git a/Engine/Plugins/Runtime/GeometryCache/Source/GeometryCache/Public/GeometryCacheSceneProxy.h b/Engine/Plugins/Runtime/GeometryCache/Source/GeometryCache/Public/GeometryCacheSceneProxy.h
--- a/Engine/Plugins/Runtime/GeometryCache/Source/GeometryCache/Public/GeometryCacheSceneProxy.h
+++ b/Engine/Plugins/Runtime/GeometryCache/Source/GeometryCache/Public/GeometryCacheSceneProxy.h
@@ -6,6 +6,7 @@
 #include "Stats/Stats.h"
 #include "Containers/ResourceArray.h"
 #include "RenderResource.h"
+#include "RayTracingGeometry.h"
 #include "PrimitiveViewRelevance.h"
 #include "PrimitiveSceneProxy.h"
 #include "Materials/MaterialInterface.h"
diff --git a/Engine/Plugins/Runtime/HairStrands/Source/HairStrandsCore/Public/GroomResources.h b/Engine/Plugins/Runtime/HairStrands/Source/HairStrandsCore/Public/GroomResources.h
--- a/Engine/Plugins/Runtime/HairStrands/Source/HairStrandsCore/Public/GroomResources.h
+++ b/Engine/Plugins/Runtime/HairStrands/Source/HairStrandsCore/Public/GroomResources.h
@@ -8,6 +8,7 @@
 #include "HairStrandsDatas.h"
 #include "HairCardsDatas.h"
 #include "RenderResource.h"
+#include "RayTracingGeometry.h"
 #include "HairStrandsInterface.h"
 #include "HairStrandsMeshProjection.h"
 
diff --git a/Engine/Source/Runtime/Engine/Classes/Engine/InstancedStaticMesh.h b/Engine/Source/Runtime/Engine/Classes/Engine/InstancedStaticMesh.h
--- a/Engine/Source/Runtime/Engine/Classes/Engine/InstancedStaticMesh.h
+++ b/Engine/Source/Runtime/Engine/Classes/Engine/InstancedStaticMesh.h
@@ -12,6 +12,7 @@
 #include "HAL/IConsoleManager.h"
 #include "RenderingThread.h"
 #include "RenderResource.h"
+#include "RayTracingGeometry.h"
 #include "PrimitiveViewRelevance.h"
 #include "ShaderParameters.h"
 #include "SceneView.h"
diff --git a/Engine/Source/Runtime/Engine/Private/RayTracingSkinnedGeometry.cpp b/Engine/Source/Runtime/Engine/Private/RayTracingSkinnedGeometry.cpp
--- a/Engine/Source/Runtime/Engine/Private/RayTracingSkinnedGeometry.cpp
+++ b/Engine/Source/Runtime/Engine/Private/RayTracingSkinnedGeometry.cpp
@@ -6,6 +6,7 @@
 
 #include "RenderGraphBuilder.h"
 #include "RHI.h"
+#include "RayTracingGeometry.h"
 
 DECLARE_GPU_STAT(SkinnedGeometryBuildBLAS);
 DECLARE_GPU_STAT(SkinnedGeometryUpdateBLAS);
diff --git a/Engine/Source/Runtime/Engine/Private/SkeletalRenderCPUSkin.h b/Engine/Source/Runtime/Engine/Private/SkeletalRenderCPUSkin.h
--- a/Engine/Source/Runtime/Engine/Private/SkeletalRenderCPUSkin.h
+++ b/Engine/Source/Runtime/Engine/Private/SkeletalRenderCPUSkin.h
@@ -9,6 +9,7 @@
 #include "CoreMinimal.h"
 #include "ProfilingDebugging/ResourceSize.h"
 #include "RenderResource.h"
+#include "RayTracingGeometry.h"
 #include "LocalVertexFactory.h"
 #include "Components/SkinnedMeshComponent.h"
 #include "SkeletalRenderPublic.h"
diff --git a/Engine/Source/Runtime/Engine/Private/SkeletalRenderGPUSkin.h b/Engine/Source/Runtime/Engine/Private/SkeletalRenderGPUSkin.h
--- a/Engine/Source/Runtime/Engine/Private/SkeletalRenderGPUSkin.h
+++ b/Engine/Source/Runtime/Engine/Private/SkeletalRenderGPUSkin.h
@@ -9,6 +9,7 @@
 #include "CoreMinimal.h"
 #include "ProfilingDebugging/ResourceSize.h"
 #include "RenderResource.h"
+#include "RayTracingGeometry.h"
 #include "ShaderParameters.h"
 #include "Components/ExternalMorphSet.h"
 #include "Components/SkinnedMeshComponent.h"
diff --git a/Engine/Source/Runtime/Engine/Private/SkeletalRenderStatic.h b/Engine/Source/Runtime/Engine/Private/SkeletalRenderStatic.h
--- a/Engine/Source/Runtime/Engine/Private/SkeletalRenderStatic.h
+++ b/Engine/Source/Runtime/Engine/Private/SkeletalRenderStatic.h
@@ -9,6 +9,7 @@
 #include "CoreMinimal.h"
 #include "ProfilingDebugging/ResourceSize.h"
 #include "RenderResource.h"
+#include "RayTracingGeometry.h"
 #include "LocalVertexFactory.h"
 #include "Components/SkinnedMeshComponent.h"
 #include "SkeletalRenderPublic.h"
diff --git a/Engine/Source/Runtime/Engine/Private/Streaming/StaticMeshUpdate.h b/Engine/Source/Runtime/Engine/Private/Streaming/StaticMeshUpdate.h
--- a/Engine/Source/Runtime/Engine/Private/Streaming/StaticMeshUpdate.h
+++ b/Engine/Source/Runtime/Engine/Private/Streaming/StaticMeshUpdate.h
@@ -11,7 +11,7 @@ StaticMeshUpdate.h: Helpers to stream in and out static mesh LODs.
 #include "Async/AsyncFileHandle.h"
 #include "IO/IoDispatcher.h"
 #include "RenderAssetUpdate.h"
-#include "RenderResource.h"
+#include "RayTracingGeometry.h"
 
 /**
 * A context used to update or proceed with the next update step.
diff --git a/Engine/Source/Runtime/Engine/Public/StaticMeshResources.h b/Engine/Source/Runtime/Engine/Public/StaticMeshResources.h
--- a/Engine/Source/Runtime/Engine/Public/StaticMeshResources.h
+++ b/Engine/Source/Runtime/Engine/Public/StaticMeshResources.h
@@ -39,6 +39,7 @@
 #include "WeightedRandomSampler.h"
 #include "PerPlatformProperties.h"
 #include "RayTracingInstance.h"
+#include "RayTracingGeometry.h"
 
 class FDistanceFieldVolumeData;
 class UBodySetup;
diff --git a/Engine/Source/Runtime/Engine/Public/StaticMeshSceneProxy.h b/Engine/Source/Runtime/Engine/Public/StaticMeshSceneProxy.h
--- a/Engine/Source/Runtime/Engine/Public/StaticMeshSceneProxy.h
+++ b/Engine/Source/Runtime/Engine/Public/StaticMeshSceneProxy.h
@@ -10,6 +10,7 @@
 #include "PrimitiveSceneProxy.h"
 #include "RayTracingInstance.h"
 #include "SceneManagement.h"
+#include "RayTracingGeometry.h"
 
 class FLocalVertexFactoryUniformShaderParameters;
 class FRawStaticIndexBuffer;
diff --git a/Engine/Source/Runtime/GeometryFramework/Public/Components/BaseDynamicMeshSceneProxy.h b/Engine/Source/Runtime/GeometryFramework/Public/Components/BaseDynamicMeshSceneProxy.h
--- a/Engine/Source/Runtime/GeometryFramework/Public/Components/BaseDynamicMeshSceneProxy.h
+++ b/Engine/Source/Runtime/GeometryFramework/Public/Components/BaseDynamicMeshSceneProxy.h
@@ -12,6 +12,7 @@
 #include "Rendering/ColorVertexBuffer.h"
 #include "DynamicMeshBuilder.h"
 #include "Components/BaseDynamicMeshComponent.h"
+#include "RayTracingGeometry.h"
 
 using UE::Geometry::FDynamicMesh3;
 using UE::Geometry::FDynamicMeshAttributeSet;
diff --git a/Engine/Source/Runtime/RenderCore/Private/RayTracingGeometry.cpp b/Engine/Source/Runtime/RenderCore/Private/RayTracingGeometry.cpp
new file mode 100644
--- /dev/null
+++ b/Engine/Source/Runtime/RenderCore/Private/RayTracingGeometry.cpp
@@ -0,0 +1,223 @@
+// Copyright Epic Games, Inc. All Rights Reserved.
+
+#include "RayTracingGeometry.h"
+#include "RayTracingGeometryManager.h"
+#include "RenderUtils.h"
+
+int32 GVarDebugForceRuntimeBLAS = 0;
+FAutoConsoleVariableRef CVarDebugForceRuntimeBLAS(
+	TEXT("r.Raytracing.DebugForceRuntimeBLAS"),
+	GVarDebugForceRuntimeBLAS,
+	TEXT("Force building BLAS at runtime."),
+	ECVF_ReadOnly);
+
+FRayTracingGeometry::FRayTracingGeometry() = default;
+
+/** Destructor. */
+FRayTracingGeometry::~FRayTracingGeometry() = default;
+
+#if RHI_RAYTRACING
+
+void FRayTracingGeometry::CreateRayTracingGeometryFromCPUData(TResourceArray<uint8>& OfflineData)
+{
+	check(OfflineData.Num() == 0 || Initializer.OfflineData == nullptr);
+	if (OfflineData.Num())
+	{
+		Initializer.OfflineData = &OfflineData;
+	}
+
+	if (GVarDebugForceRuntimeBLAS && Initializer.OfflineData != nullptr)
+	{
+		Initializer.OfflineData->Discard();
+		Initializer.OfflineData = nullptr;
+	}
+
+	SetRequiresBuild(Initializer.OfflineData == nullptr);
+	RayTracingGeometryRHI = RHICreateRayTracingGeometry(Initializer);
+}
+
+void FRayTracingGeometry::RequestBuildIfNeeded(ERTAccelerationStructureBuildPriority InBuildPriority)
+{
+	RayTracingGeometryRHI->SetInitializer(Initializer);
+
+	if (GetRequiresBuild())
+	{
+		RayTracingBuildRequestIndex = GRayTracingGeometryManager.RequestBuildAccelerationStructure(this, InBuildPriority);
+		SetRequiresBuild(false);
+	}
+}
+
+void FRayTracingGeometry::InitRHIForDynamicRayTracing()
+{
+	check(GetRayTracingMode() == ERayTracingMode::Dynamic);
+
+	// Streaming BLAS needs special handling to not get their "streaming" type wiped out as it will cause issues down the line.	
+	// We only have to do this if the geometry was marked to be streamed in.
+	// In that case we will recreate the geometry as-if it was streamed in.
+	if (EnumHasAnyFlags(GeometryState, FRayTracingGeometry::EGeometryStateFlags::StreamedIn))
+	{
+		// When a mesh is streamed in (FStaticMeshStreamIn::DoFinishUpdate) we update the geometry initializer using just streamed in VB/IB.
+		// That initializer sets a Rendering type but RHI object was created as StreamingDestination and we have a mismatch between geometry initializer and RHI initializer.
+		// It's not an issue unless we try to initialize the geometry again using the geometry's initializer.
+		// We need the current geometry and RHI object to be StreamingDestination so the streaming continues to work.
+		Initializer.Type = ERayTracingGeometryInitializerType::StreamingDestination;
+
+		// Creating RHI with StreamingDestination type will only initialize RHI object but will not created the underlying BLAS buffers.
+		InitRHI();
+
+		// Here we simulate geometry streaming: create geometry with StreamingSource type to allocate BLAS buffers (1) and swap it with the current geometry (2).
+		// Follows the same pattern as: (1) FStaticMeshStreamIn::CreateBuffers_* (2) FStaticMeshStreamIn::DoFinishUpdate
+		// There is no other way to initialize BLAS buffers for the geometry that has a StreamingDestination type.
+		{
+			TRHIResourceUpdateBatcher<1> Batcher;
+			FRayTracingGeometryInitializer IntermediateInitializer = Initializer;
+			IntermediateInitializer.Type = ERayTracingGeometryInitializerType::StreamingSource;
+
+			FRayTracingGeometryRHIRef IntermediateRayTracingGeometry = RHICreateRayTracingGeometry(IntermediateInitializer);
+			InitRHIForStreaming(IntermediateRayTracingGeometry, Batcher);
+
+			// When Batcher goes out of scope it will add commands to copy the BLAS buffers on RHI thread.
+			// We need to do it before we build the current geometry (also on RHI thread).
+		}
+
+		RequestBuildIfNeeded(ERTAccelerationStructureBuildPriority::Normal);
+	}
+	else
+	{
+		InitRHI();
+	}
+}
+
+void FRayTracingGeometry::CreateRayTracingGeometry(ERTAccelerationStructureBuildPriority InBuildPriority)
+{
+	// Release previous RHI object if any
+	ReleaseRHI();
+
+	check(RawData.Num() == 0 || Initializer.OfflineData == nullptr);
+	if (RawData.Num())
+	{
+		Initializer.OfflineData = &RawData;
+	}
+
+	if (GVarDebugForceRuntimeBLAS && Initializer.OfflineData != nullptr)
+	{
+		Initializer.OfflineData->Discard();
+		Initializer.OfflineData = nullptr;
+	}
+
+	bool bAllSegmentsAreValid = Initializer.Segments.Num() > 0 || Initializer.OfflineData;
+	for (const FRayTracingGeometrySegment& Segment : Initializer.Segments)
+	{
+		if (!Segment.VertexBuffer)
+		{
+			bAllSegmentsAreValid = false;
+			break;
+		}
+	}
+
+	const bool bWithoutNativeResource = Initializer.Type == ERayTracingGeometryInitializerType::StreamingDestination;
+	if (bAllSegmentsAreValid)
+	{
+		// Only geometries with StreamingDestination type are initially created in invalid state until they are streamed in (see InitRHIForStreaming).
+		if (bWithoutNativeResource)
+		{
+			EnumRemoveFlags(GeometryState, EGeometryStateFlags::Valid);
+		}
+		else
+		{
+			EnumAddFlags(GeometryState, EGeometryStateFlags::Valid);
+		}
+
+		if (IsRayTracingEnabled())
+		{
+			RayTracingGeometryRHI = RHICreateRayTracingGeometry(Initializer);
+		}
+
+		// Register the geometry if it wasn't registered before and it's not using custom path
+		const bool bRegisterGeometry = (RayTracingGeometryHandle == INDEX_NONE) && (InBuildPriority != ERTAccelerationStructureBuildPriority::Immediate);
+
+		if (bRegisterGeometry)
+		{
+			RayTracingGeometryHandle = GRayTracingGeometryManager.RegisterRayTracingGeometry(this);
+		}
+
+		if (Initializer.OfflineData == nullptr)
+		{
+			// Request build if not skip
+			if (InBuildPriority != ERTAccelerationStructureBuildPriority::Skip)
+			{
+				if (IsRayTracingEnabled())
+				{
+					RayTracingBuildRequestIndex = GRayTracingGeometryManager.RequestBuildAccelerationStructure(this, InBuildPriority);
+				}
+				SetRequiresBuild(false);
+			}
+			else
+			{
+				SetRequiresBuild(true);
+			}
+		}
+		else
+		{
+			SetRequiresBuild(false);
+
+			// Offline data ownership is transferred to the RHI, which discards it after use.
+			// It is no longer valid to use it after this point.
+			Initializer.OfflineData = nullptr;
+		}
+	}
+}
+
+bool FRayTracingGeometry::IsValid() const
+{
+	return RayTracingGeometryRHI != nullptr && Initializer.TotalPrimitiveCount > 0 && EnumHasAnyFlags(GeometryState, EGeometryStateFlags::Valid);
+}
+
+void FRayTracingGeometry::InitRHI()
+{
+	if (!IsRayTracingAllowed())
+		return;
+
+	ERTAccelerationStructureBuildPriority BuildPriority = Initializer.Type != ERayTracingGeometryInitializerType::Rendering
+		? ERTAccelerationStructureBuildPriority::Skip
+		: ERTAccelerationStructureBuildPriority::Normal;
+	CreateRayTracingGeometry(BuildPriority);
+}
+
+void FRayTracingGeometry::ReleaseRHI()
+{
+	RemoveBuildRequest();
+	RayTracingGeometryRHI.SafeRelease();
+
+	if (RayTracingGeometryHandle != INDEX_NONE)
+	{
+		GRayTracingGeometryManager.ReleaseRayTracingGeometryHandle(RayTracingGeometryHandle);
+		RayTracingGeometryHandle = INDEX_NONE;
+	}
+}
+
+void FRayTracingGeometry::RemoveBuildRequest()
+{
+	if (HasPendingBuildRequest())
+	{
+		GRayTracingGeometryManager.RemoveBuildRequest(RayTracingBuildRequestIndex);
+		RayTracingBuildRequestIndex = INDEX_NONE;
+	}
+}
+
+void FRayTracingGeometry::ReleaseResource()
+{
+	// Release any resource references held by the initializer.
+	// This includes index and vertex buffers used for building the BLAS.
+	Initializer = FRayTracingGeometryInitializer{};
+
+	FRenderResource::ReleaseResource();
+}
+
+void FRayTracingGeometry::BoostBuildPriority(float InBoostValue) const
+{
+	check(HasPendingBuildRequest());
+	GRayTracingGeometryManager.BoostPriority(RayTracingBuildRequestIndex, InBoostValue);
+}
+
+#endif // RHI_RAYTRACING
diff --git a/Engine/Source/Runtime/RenderCore/Private/RayTracingGeometryManager.cpp b/Engine/Source/Runtime/RenderCore/Private/RayTracingGeometryManager.cpp
--- a/Engine/Source/Runtime/RenderCore/Private/RayTracingGeometryManager.cpp
+++ b/Engine/Source/Runtime/RenderCore/Private/RayTracingGeometryManager.cpp
@@ -6,7 +6,8 @@
 #include "RHIResources.h"
 #include "RHICommandList.h"
 
-#include "RenderResource.h"
+#include "RayTracingGeometry.h"
+#include "RenderUtils.h"
 
 #if RHI_RAYTRACING
 
diff --git a/Engine/Source/Runtime/RenderCore/Private/RenderResource.cpp b/Engine/Source/Runtime/RenderCore/Private/RenderResource.cpp
--- a/Engine/Source/Runtime/RenderCore/Private/RenderResource.cpp
+++ b/Engine/Source/Runtime/RenderCore/Private/RenderResource.cpp
@@ -35,13 +35,6 @@ FAutoConsoleVariableRef CVarFreeStructuresOnRHIBufferCreation(
 	GFreeStructuresOnRHIBufferCreation,
 	TEXT("Toggles experimental method for freeing helper structures that own the resource arrays after submitting to RHI instead of in the callback sink."));
 
-int32 GVarDebugForceRuntimeBLAS = 0;
-FAutoConsoleVariableRef CVarDebugForceRuntimeBLAS(
-	TEXT("r.Raytracing.DebugForceRuntimeBLAS"),
-	GVarDebugForceRuntimeBLAS,
-	TEXT("Force building BLAS at runtime."),
-	ECVF_ReadOnly);
-
 /** Tracks render resources in a list. The implementation is optimized to allow fast allocation / deallocation from any thread,
  *  at the cost of period coalescing of thread-local data at a sync point each frame. Furthermore, iteration is not thread safe
  *  and must be performed at sync points.
@@ -512,216 +505,6 @@ TGlobalResource<FNullColorVertexBuffer> GNullColorVertexBuffer;
 /** The global null vertex buffer, which is set with a stride of 0 on meshes */
 TGlobalResource<FNullVertexBuffer> GNullVertexBuffer;
 
-/*------------------------------------------------------------------------------
-	FRayTracingGeometry implementation.
-------------------------------------------------------------------------------*/
-
-#if RHI_RAYTRACING
-
-void FRayTracingGeometry::CreateRayTracingGeometryFromCPUData(TResourceArray<uint8>& OfflineData)
-{
-	check(OfflineData.Num() == 0 || Initializer.OfflineData == nullptr);
-	if (OfflineData.Num())
-	{
-		Initializer.OfflineData = &OfflineData;
-	}
-
-	if (GVarDebugForceRuntimeBLAS && Initializer.OfflineData != nullptr)
-	{
-		Initializer.OfflineData->Discard();
-		Initializer.OfflineData = nullptr;
-	}
-
-	SetRequiresBuild(Initializer.OfflineData == nullptr);
-	RayTracingGeometryRHI = RHICreateRayTracingGeometry(Initializer);
-}
-
-void FRayTracingGeometry::RequestBuildIfNeeded(ERTAccelerationStructureBuildPriority InBuildPriority)
-{
-	RayTracingGeometryRHI->SetInitializer(Initializer);
-
-	if (GetRequiresBuild())
-	{
-		RayTracingBuildRequestIndex = GRayTracingGeometryManager.RequestBuildAccelerationStructure(this, InBuildPriority);
-		SetRequiresBuild(false);
-	}	
-}
-
-void FRayTracingGeometry::InitRHIForDynamicRayTracing()
-{
-	check(GetRayTracingMode() == ERayTracingMode::Dynamic);
-
-	// Streaming BLAS needs special handling to not get their "streaming" type wiped out as it will cause issues down the line.	
-	// We only have to do this if the geometry was marked to be streamed in.
-	// In that case we will recreate the geometry as-if it was streamed in.
-	if (EnumHasAnyFlags(GeometryState, FRayTracingGeometry::EGeometryStateFlags::StreamedIn))
-	{
-		// When a mesh is streamed in (FStaticMeshStreamIn::DoFinishUpdate) we update the geometry initializer using just streamed in VB/IB.
-		// That initializer sets a Rendering type but RHI object was created as StreamingDestination and we have a mismatch between geometry initializer and RHI initializer.
-		// It's not an issue unless we try to initialize the geometry again using the geometry's initializer.
-		// We need the current geometry and RHI object to be StreamingDestination so the streaming continues to work.
-		Initializer.Type = ERayTracingGeometryInitializerType::StreamingDestination;
-
-		// Creating RHI with StreamingDestination type will only initialize RHI object but will not created the underlying BLAS buffers.
-		InitRHI();
-		
-		// Here we simulate geometry streaming: create geometry with StreamingSource type to allocate BLAS buffers (1) and swap it with the current geometry (2).
-		// Follows the same pattern as: (1) FStaticMeshStreamIn::CreateBuffers_* (2) FStaticMeshStreamIn::DoFinishUpdate
-		// There is no other way to initialize BLAS buffers for the geometry that has a StreamingDestination type.
-		{			
-			TRHIResourceUpdateBatcher<1> Batcher;
-			FRayTracingGeometryInitializer IntermediateInitializer = Initializer;
-			IntermediateInitializer.Type = ERayTracingGeometryInitializerType::StreamingSource;
-
-			FRayTracingGeometryRHIRef IntermediateRayTracingGeometry = RHICreateRayTracingGeometry(IntermediateInitializer);
-			InitRHIForStreaming(IntermediateRayTracingGeometry, Batcher);
-
-			// When Batcher goes out of scope it will add commands to copy the BLAS buffers on RHI thread.
-			// We need to do it before we build the current geometry (also on RHI thread).
-		}
-		
-		RequestBuildIfNeeded(ERTAccelerationStructureBuildPriority::Normal);
-	}
-	else
-	{
-		InitRHI();
-	}
-}
-
-void FRayTracingGeometry::CreateRayTracingGeometry(ERTAccelerationStructureBuildPriority InBuildPriority)
-{
-	// Release previous RHI object if any
-	ReleaseRHI();
-
-	check(RawData.Num() == 0 || Initializer.OfflineData == nullptr);
-	if (RawData.Num())
-	{
-		Initializer.OfflineData = &RawData;
-	}
-
-	if (GVarDebugForceRuntimeBLAS && Initializer.OfflineData != nullptr)
-	{
-		Initializer.OfflineData->Discard();
-		Initializer.OfflineData = nullptr;
-	}
-
-	bool bAllSegmentsAreValid = Initializer.Segments.Num() > 0 || Initializer.OfflineData;
-	for (const FRayTracingGeometrySegment& Segment : Initializer.Segments)
-	{
-		if (!Segment.VertexBuffer)
-		{
-			bAllSegmentsAreValid = false;
-			break;
-		}
-	}
-
-	const bool bWithoutNativeResource = Initializer.Type == ERayTracingGeometryInitializerType::StreamingDestination;
-	if (bAllSegmentsAreValid)
-	{
-		// Only geometries with StreamingDestination type are initially created in invalid state until they are streamed in (see InitRHIForStreaming).
-		if (bWithoutNativeResource)
-		{
-			EnumRemoveFlags(GeometryState, EGeometryStateFlags::Valid);
-		}
-		else
-		{
-			EnumAddFlags(GeometryState, EGeometryStateFlags::Valid);
-		}
-
-		if (IsRayTracingEnabled())
-		{
-			RayTracingGeometryRHI = RHICreateRayTracingGeometry(Initializer);
-		}
-		
-		// Register the geometry if it wasn't registered before and it's not using custom path
-		const bool bRegisterGeometry = (RayTracingGeometryHandle == INDEX_NONE) && (InBuildPriority != ERTAccelerationStructureBuildPriority::Immediate);			
-
-		if (bRegisterGeometry)
-		{
-			RayTracingGeometryHandle = GRayTracingGeometryManager.RegisterRayTracingGeometry(this);
-		}
-
-		if (Initializer.OfflineData == nullptr)
-		{
-			// Request build if not skip
-			if (InBuildPriority != ERTAccelerationStructureBuildPriority::Skip)
-			{
-				if (IsRayTracingEnabled())
-				{
-					RayTracingBuildRequestIndex = GRayTracingGeometryManager.RequestBuildAccelerationStructure(this, InBuildPriority);
-				}
-				SetRequiresBuild(false);
-			}
-			else
-			{
-				SetRequiresBuild(true);
-			}
-		}
-		else
-		{
-			SetRequiresBuild(false);
-
-			// Offline data ownership is transferred to the RHI, which discards it after use.
-			// It is no longer valid to use it after this point.
-			Initializer.OfflineData = nullptr;
-		}
-	}
-}
-
-bool FRayTracingGeometry::IsValid() const
-{
-	return RayTracingGeometryRHI != nullptr && Initializer.TotalPrimitiveCount > 0 && EnumHasAnyFlags(GeometryState, EGeometryStateFlags::Valid);
-}
-
-void FRayTracingGeometry::InitRHI()
-{
-	if (!IsRayTracingAllowed())
-		return;
-
-	ERTAccelerationStructureBuildPriority BuildPriority = Initializer.Type != ERayTracingGeometryInitializerType::Rendering
-		? ERTAccelerationStructureBuildPriority::Skip
-		: ERTAccelerationStructureBuildPriority::Normal;
-	CreateRayTracingGeometry(BuildPriority);
-}
-
-void FRayTracingGeometry::ReleaseRHI()
-{
-	RemoveBuildRequest();
-	RayTracingGeometryRHI.SafeRelease();
-	
-	if (RayTracingGeometryHandle != INDEX_NONE)
-	{
-		GRayTracingGeometryManager.ReleaseRayTracingGeometryHandle(RayTracingGeometryHandle);
-		RayTracingGeometryHandle = INDEX_NONE;
-	}	
-}
-
-void FRayTracingGeometry::RemoveBuildRequest()
-{
-	if (HasPendingBuildRequest())
-	{
-		GRayTracingGeometryManager.RemoveBuildRequest(RayTracingBuildRequestIndex);
-		RayTracingBuildRequestIndex = INDEX_NONE;
-	}
-}
-
-void FRayTracingGeometry::ReleaseResource()
-{
-	// Release any resource references held by the initializer.
-	// This includes index and vertex buffers used for building the BLAS.
-	Initializer = FRayTracingGeometryInitializer {};
-
-	FRenderResource::ReleaseResource();
-}
-
-void FRayTracingGeometry::BoostBuildPriority(float InBoostValue) const
-{
-	check(HasPendingBuildRequest());
-	GRayTracingGeometryManager.BoostPriority(RayTracingBuildRequestIndex, InBoostValue);
-}
-
-#endif // RHI_RAYTRACING
-
 /*------------------------------------------------------------------------------
 	FGlobalDynamicVertexBuffer implementation.
 ------------------------------------------------------------------------------*/
diff --git a/Engine/Source/Runtime/RenderCore/Public/RayTracingGeometry.h b/Engine/Source/Runtime/RenderCore/Public/RayTracingGeometry.h
new file mode 100644
--- /dev/null
+++ b/Engine/Source/Runtime/RenderCore/Public/RayTracingGeometry.h
@@ -0,0 +1,159 @@
+// Copyright Epic Games, Inc. All Rights Reserved.
+
+#pragma once
+
+#include "CoreMinimal.h"
+#include "RenderResource.h"
+#if RHI_RAYTRACING
+#include "RHI.h"
+#endif
+
+enum class ERTAccelerationStructureBuildPriority
+{
+	Immediate,
+	High,
+	Normal,
+	Low,
+	Skip
+};
+
+/** A ray tracing geometry resource */
+class RENDERCORE_API FRayTracingGeometry : public FRenderResource
+{
+public:
+	TResourceArray<uint8> RawData;
+
+	FRayTracingGeometry();
+	virtual ~FRayTracingGeometry();
+
+#if RHI_RAYTRACING
+
+	/** When set to NonSharedVertexBuffers, then shared vertex buffers are not used  */
+	static constexpr int64 NonSharedVertexBuffers = -1;
+
+	/**
+	Vertex buffers for dynamic geometries may be sub-allocated from a shared pool, which is periodically reset and its generation ID is incremented.
+	Geometries that use the shared buffer must be updated (rebuilt or refit) before they are used for rendering after the pool is reset.
+	This is validated by comparing the current shared pool generation ID against generation IDs stored in FRayTracingGeometry during latest update.
+	*/
+	int64 DynamicGeometrySharedBufferGenerationID = NonSharedVertexBuffers;
+
+	FRayTracingGeometryInitializer Initializer;
+	FRayTracingGeometryRHIRef RayTracingGeometryRHI;
+
+	// Flags for tracking the state of RayTracingGeometryRHI.
+	enum class EGeometryStateFlags : uint32
+	{
+		// Initial state when the geometry was not created or was created for streaming but not yet streamed in.
+		Invalid = 0,
+
+		// If the geometry needs to be built.
+		RequiresBuild = 1 << 0,
+
+		// If the geometry was successfully created or streamed in.
+		Valid = 1 << 1,
+
+		// Special flag that is used when ray tracing is dynamic to mark the streamed geometry to be recreated when ray tracing is switched on.
+		// Only set when mesh streaming is used.
+		StreamedIn = 1 << 2
+	};
+	FRIEND_ENUM_CLASS_FLAGS(EGeometryStateFlags);
+
+	/** LOD of the mesh associated with this ray tracing geometry object (-1 if unknown) */
+	int8 LODIndex = -1;
+
+	void SetInitializer(const FRayTracingGeometryInitializer& InInitializer)
+	{
+		Initializer = InInitializer;
+	}
+
+	bool IsValid() const;
+
+	void SetAsStreamedIn()
+	{
+		EnumAddFlags(GeometryState, EGeometryStateFlags::StreamedIn);
+	}
+
+	bool GetRequiresBuild() const
+	{
+		return EnumHasAnyFlags(GeometryState, EGeometryStateFlags::RequiresBuild);
+	}
+
+	void SetRequiresBuild(bool bBuild)
+	{
+		if (bBuild)
+		{
+			EnumAddFlags(GeometryState, EGeometryStateFlags::RequiresBuild);
+		}
+		else
+		{
+			EnumRemoveFlags(GeometryState, EGeometryStateFlags::RequiresBuild);
+		}
+	}
+
+	EGeometryStateFlags GetGeometryState() const
+	{
+		return GeometryState;
+	}
+
+	template <uint32 MaxNumUpdates>
+	void InitRHIForStreaming(FRHIRayTracingGeometry* IntermediateGeometry, TRHIResourceUpdateBatcher<MaxNumUpdates>& Batcher)
+	{
+		EnumAddFlags(GeometryState, EGeometryStateFlags::StreamedIn);
+
+		if (RayTracingGeometryRHI && IntermediateGeometry)
+		{
+			Batcher.QueueUpdateRequest(RayTracingGeometryRHI, IntermediateGeometry);
+			EnumAddFlags(GeometryState, EGeometryStateFlags::Valid);
+		}
+	}
+
+	template <uint32 MaxNumUpdates>
+	void ReleaseRHIForStreaming(TRHIResourceUpdateBatcher<MaxNumUpdates>& Batcher)
+	{
+		Initializer = {};
+
+		RemoveBuildRequest();
+		EnumRemoveFlags(GeometryState, EGeometryStateFlags::StreamedIn);
+
+		if (RayTracingGeometryRHI)
+		{
+			Batcher.QueueUpdateRequest(RayTracingGeometryRHI, nullptr);
+			EnumRemoveFlags(GeometryState, EGeometryStateFlags::Valid);
+		}
+	}
+	void CreateRayTracingGeometryFromCPUData(TResourceArray<uint8>& OfflineData);
+	void RequestBuildIfNeeded(ERTAccelerationStructureBuildPriority InBuildPriority);
+
+	// That function is only supposed to be used when dynamic ray tracing is enabled
+	void InitRHIForDynamicRayTracing();
+
+	void CreateRayTracingGeometry(ERTAccelerationStructureBuildPriority InBuildPriority);
+
+	bool HasPendingBuildRequest() const
+	{
+		return RayTracingBuildRequestIndex != INDEX_NONE;
+	}
+	void BoostBuildPriority(float InBoostValue = 0.01f) const;
+
+	// FRenderResource interface
+
+	virtual FString GetFriendlyName() const override { return TEXT("FRayTracingGeometry"); }
+
+	virtual void InitRHI() override;
+	virtual void ReleaseRHI() override;
+
+	virtual void ReleaseResource() override;
+protected:
+	void RemoveBuildRequest();
+
+	friend class FRayTracingGeometryManager;
+	int32 RayTracingBuildRequestIndex = INDEX_NONE;
+	int32 RayTracingGeometryHandle = INDEX_NONE; // Only valid when ray tracing is dynamic
+	EGeometryStateFlags GeometryState = EGeometryStateFlags::Invalid;
+#endif
+};
+
+#if RHI_RAYTRACING
+ENUM_CLASS_FLAGS(FRayTracingGeometry::EGeometryStateFlags);
+#endif
diff --git a/Engine/Source/Runtime/RenderCore/Public/RenderResource.h b/Engine/Source/Runtime/RenderCore/Public/RenderResource.h
--- a/Engine/Source/Runtime/RenderCore/Public/RenderResource.h
+++ b/Engine/Source/Runtime/RenderCore/Public/RenderResource.h
@@ -718,160 +718,6 @@ extern RENDERCORE_API bool IsRayTracingAllowed();
 // This function may only be called at runtime, never during cooking.
 extern RENDERCORE_API ERayTracingMode GetRayTracingMode();
 
-enum class ERTAccelerationStructureBuildPriority
-{
-	Immediate,
-	High,
-	Normal,
-	Low,
-	Skip
-};
-
-/** A ray tracing geometry resource */
-class RENDERCORE_API FRayTracingGeometry : public FRenderResource
-{
-public:
-	TResourceArray<uint8> RawData;
-
-	/** Default constructor. */
-	FRayTracingGeometry()
-	{}
-
-	/** Destructor. */
-	virtual ~FRayTracingGeometry() {}
-
-#if RHI_RAYTRACING
-
-	/** When set to NonSharedVertexBuffers, then shared vertex buffers are not used  */
-	static constexpr int64 NonSharedVertexBuffers = -1;
-
-	/** 
-	Vertex buffers for dynamic geometries may be sub-allocated from a shared pool, which is periodically reset and its generation ID is incremented.
-	Geometries that use the shared buffer must be updated (rebuilt or refit) before they are used for rendering after the pool is reset.
-	This is validated by comparing the current shared pool generation ID against generation IDs stored in FRayTracingGeometry during latest update.
-	*/
-	int64 DynamicGeometrySharedBufferGenerationID = NonSharedVertexBuffers;
-
-	FRayTracingGeometryInitializer Initializer;
-	FRayTracingGeometryRHIRef RayTracingGeometryRHI;
-
-	// Flags for tracking the state of RayTracingGeometryRHI.
-	enum class EGeometryStateFlags : uint32
-	{
-		// Initial state when the geometry was not created or was created for streaming but not yet streamed in.
-		Invalid			= 0,
-
-		// If the geometry needs to be built.
-		RequiresBuild	= 1 << 0,
-
-		// If the geometry was successfully created or streamed in.
-		Valid			= 1 << 1,
-
-		// Special flag that is used when ray tracing is dynamic to mark the streamed geometry to be recreated when ray tracing is switched on.
-		// Only set when mesh streaming is used.
-		StreamedIn		= 1 << 2
-	};
-	FRIEND_ENUM_CLASS_FLAGS(EGeometryStateFlags);
-
-	/** LOD of the mesh associated with this ray tracing geometry object (-1 if unknown) */
-	int8 LODIndex = -1;
-
-	void SetInitializer(const FRayTracingGeometryInitializer& InInitializer)
-	{
-		Initializer = InInitializer;
-	}
-
-	bool IsValid() const;
-
-	void SetAsStreamedIn()
-	{
-		EnumAddFlags(GeometryState, EGeometryStateFlags::StreamedIn);
-	}
-
-	bool GetRequiresBuild() const 
-	{
-		return EnumHasAnyFlags(GeometryState, EGeometryStateFlags::RequiresBuild);
-	}
-
-	void SetRequiresBuild(bool bBuild)
-	{
-		if (bBuild)
-		{
-			EnumAddFlags(GeometryState, EGeometryStateFlags::RequiresBuild);
-		}
-		else
-		{
-			EnumRemoveFlags(GeometryState, EGeometryStateFlags::RequiresBuild);
-		}
-	}
-	
-	EGeometryStateFlags GetGeometryState() const 
-	{ 
-		return GeometryState; 
-	}
-
-	template <uint32 MaxNumUpdates>
-	void InitRHIForStreaming(FRHIRayTracingGeometry* IntermediateGeometry, TRHIResourceUpdateBatcher<MaxNumUpdates>& Batcher)
-	{
-		EnumAddFlags(GeometryState, EGeometryStateFlags::StreamedIn);
-
-		if (RayTracingGeometryRHI && IntermediateGeometry)
-		{
-			Batcher.QueueUpdateRequest(RayTracingGeometryRHI, IntermediateGeometry);
-			EnumAddFlags(GeometryState, EGeometryStateFlags::Valid);
-		}
-	}
-
-	template <uint32 MaxNumUpdates>
-	void ReleaseRHIForStreaming(TRHIResourceUpdateBatcher<MaxNumUpdates>& Batcher)
-	{
-		Initializer = {};
-
-		RemoveBuildRequest();
-		EnumRemoveFlags(GeometryState, EGeometryStateFlags::StreamedIn);
-
-		if (RayTracingGeometryRHI)
-		{
-			Batcher.QueueUpdateRequest(RayTracingGeometryRHI, nullptr);
-			EnumRemoveFlags(GeometryState, EGeometryStateFlags::Valid);
-		}		
-	}
-	void CreateRayTracingGeometryFromCPUData(TResourceArray<uint8>& OfflineData);
-	void RequestBuildIfNeeded(ERTAccelerationStructureBuildPriority InBuildPriority);
-
-	// That function is only supposed to be used when dynamic ray tracing is enabled
-	void InitRHIForDynamicRayTracing();
-
-	void CreateRayTracingGeometry(ERTAccelerationStructureBuildPriority InBuildPriority);
-
-	bool HasPendingBuildRequest() const
-	{
-		return RayTracingBuildRequestIndex != INDEX_NONE;
-	}
-	void BoostBuildPriority(float InBoostValue = 0.01f) const;
-
-	// FRenderResource interface
-
-	virtual FString GetFriendlyName() const override { return TEXT("FRayTracingGeometry"); }
-
-	virtual void InitRHI() override;
-	virtual void ReleaseRHI() override;
-
-	virtual void ReleaseResource() override;
-protected:
-	void RemoveBuildRequest();
-
-	friend class FRayTracingGeometryManager;
-	int32 RayTracingBuildRequestIndex = INDEX_NONE;
-	int32 RayTracingGeometryHandle = INDEX_NONE; // Only valid when ray tracing is dynamic
-	EGeometryStateFlags GeometryState = EGeometryStateFlags::Invalid;
-#endif
-};
-
-#if RHI_RAYTRACING
-ENUM_CLASS_FLAGS(FRayTracingGeometry::EGeometryStateFlags);
-#endif
-
 /**
  * A system for dynamically allocating GPU memory for vertices.
  */
@@ -1031,4 +877,8 @@ FORCEINLINE bool PlatformRequires128bitRT(EPixelFormat PixelFormat)
 	default:
 		return false;
 	}
-}
\ No newline at end of file
+}
+
+#if UE_ENABLE_INCLUDE_ORDER_DEPRECATED_IN_5_2
+#include "RayTracingGeometry.h"
+#endif
diff --git a/Engine/Source/Runtime/Renderer/Private/RayTracing/RayTracingDynamicGeometry.cpp b/Engine/Source/Runtime/Renderer/Private/RayTracing/RayTracingDynamicGeometry.cpp
--- a/Engine/Source/Runtime/Renderer/Private/RayTracing/RayTracingDynamicGeometry.cpp
+++ b/Engine/Source/Runtime/Renderer/Private/RayTracing/RayTracingDynamicGeometry.cpp
@@ -4,6 +4,7 @@
 #include "ScenePrivate.h"
 #include "RayTracingDynamicGeometryCollection.h"
 #include "RayTracingInstance.h"
+#include "RayTracingGeometry.h"
 
 #if RHI_RAYTRACING
 
