Moving the rendering timer types out of RenderCore.h and into their own header.

diff --git a/Engine/Source/Runtime/RenderCore/Private/RenderCore.cpp b/Engine/Source/Runtime/RenderCore/Private/RenderCore.cpp
--- a/Engine/Source/Runtime/RenderCore/Private/RenderCore.cpp
+++ b/Engine/Source/Runtime/RenderCore/Private/RenderCore.cpp
@@ -221,7 +221,7 @@ static TAutoConsoleVariable<int32> CVarForceLODShadow(
 bool GPauseRenderingRealtimeClock;
 
 /** Global realtime clock for the rendering thread. */
-FTimer GRenderingRealtimeClock;
+FRenderTimer GRenderingRealtimeClock;
 
 FInputLatencyTimer GInputLatencyTimer( 2.0f );
 
diff --git a/Engine/Source/Runtime/RenderCore/Public/RenderCore.h b/Engine/Source/Runtime/RenderCore/Public/RenderCore.h
--- a/Engine/Source/Runtime/RenderCore/Public/RenderCore.h
+++ b/Engine/Source/Runtime/RenderCore/Public/RenderCore.h
@@ -16,6 +16,7 @@
 #include "RHIDefinitions.h"
 #include "Stats/Stats.h"
 #include "Stats/Stats2.h"
+#include "RenderTimer.h"
 
 DECLARE_LOG_CATEGORY_EXTERN(LogRendererCore, Log, All);
 
@@ -158,127 +159,6 @@ DECLARE_MEMORY_STAT_EXTERN(TEXT("Pool Size"), STAT_RenderTargetPoolSize, STATGRO
 DECLARE_MEMORY_STAT_EXTERN(TEXT("Pool Used"), STAT_RenderTargetPoolUsed, STATGROUP_RenderTargetPool, RENDERCORE_API);
 DECLARE_DWORD_COUNTER_STAT_EXTERN(TEXT("Pool Count"), STAT_RenderTargetPoolCount, STATGROUP_RenderTargetPool, RENDERCORE_API);
 
-/**
- *	Timer helper class.
- **/
-class FTimer
-{
-public:
-	/**
-	 *	Constructor
-	 **/
-	FTimer()
-		: CurrentDeltaTime(0.0f)
-		, CurrentTime(0.0f)
-	{
-	}
-
-	/**
-	 *	Returns the current time, in seconds.
-	 *	@return Current time, in seconds
-	 */
-	float GetCurrentTime() const
-	{
-		return CurrentTime;
-	}
-
-	/**
-	 *	Returns the current delta time.
-	 *	@return Current delta time (number of seconds that passed between the last two tick)
-	 */
-	float GetCurrentDeltaTime() const
-	{
-		return CurrentDeltaTime;
-	}
-
-	/**
-	 *	Updates the timer.
-	 *	@param DeltaTime	Number of seconds that have passed since the last tick
-	 **/
-	void Tick( float DeltaTime )
-	{
-		CurrentDeltaTime = DeltaTime;
-		CurrentTime += DeltaTime;
-	}
-
-protected:
-	/** Current delta time (number of seconds that passed between the last two tick). */
-	float CurrentDeltaTime;
-	/** Current time, in seconds. */
-	float CurrentTime;
-};
-
-/** Whether to pause the global realtime clock for the rendering thread (read and write only on main thread). */
-extern RENDERCORE_API bool GPauseRenderingRealtimeClock;
-
-/** Global realtime clock for the rendering thread. */
-extern RENDERCORE_API FTimer GRenderingRealtimeClock;
-
-/**
- * Encapsulates a latency timer that measures the time from when mouse input
- * is read on the gamethread until that frame is fully displayed by the GPU.
- */
-struct FInputLatencyTimer
-{
-	/**
-	 * Constructor
-	 * @param InUpdateFrequency	How often the timer should be updated (in seconds).
-	 */
-	FInputLatencyTimer( float InUpdateFrequency )
-	{
-		bInitialized = false;
-		RenderThreadTrigger = false;
-		GameThreadTrigger = false;
-		StartTime = 0;
-		DeltaTime = 0;
-		LastCaptureTime = 0.0f;
-		UpdateFrequency = InUpdateFrequency;
-	}
-
-	/** Potentially starts the timer on the gamethread, based on the UpdateFrequency. */
-	RENDERCORE_API void GameThreadTick();
-
-	/** @return The number of seconds of input latency. */
-	inline float GetDeltaSeconds() const
-	{
-		return FPlatformTime::ToSeconds(DeltaTime);
-	}
-
-	/** Whether GInputLatencyTimer is initialized or not. */
-	bool	bInitialized;
-
-	/** Whether a measurement has been triggered on the gamethread. */
-	bool	GameThreadTrigger;
-
-	/** Whether a measurement has been triggered on the renderthread. */
-	bool	RenderThreadTrigger;
-
-	/** Start time (in FPlatformTime::Cycles). */
-	uint32	StartTime;
-
-	/** Last delta time that was measured (in FPlatformTime::Cycles). */
-	uint32	DeltaTime;
-
-	/** Last time we did a measurement (in seconds). */
-	double	LastCaptureTime;
-
-	/** How often we should do a measurement (in seconds). */
-	float	UpdateFrequency;
-};
-
-/** Global input latency timer. Defined in UnrealClient.cpp */
-extern RENDERCORE_API FInputLatencyTimer GInputLatencyTimer;
-/** How many cycles the renderthread used (excluding idle time). It's set once per frame in FViewport::Draw. */
-extern RENDERCORE_API uint32 GRenderThreadTime;
-/** How many cycles the rhithread used (excluding idle time). */
-extern RENDERCORE_API uint32 GRHIThreadTime;
-/** How many cycles the gamethread used (excluding idle time). It's set once per frame in FViewport::Draw. */
-extern RENDERCORE_API uint32 GGameThreadTime;
-/** How many cycles it took to swap buffers to present the frame. */
-extern RENDERCORE_API uint32 GSwapBufferTime;
-/** How many cycles the renderthread used, including dependent wait time. */
-extern RENDERCORE_API uint32 GRenderThreadTimeCriticalPath;
-
 // shared by renderer and engine, compiles down to a constant in final release
 RENDERCORE_API int32 GetCVarForceLOD();
 RENDERCORE_API int32 GetCVarForceLOD_AnyThread();
@@ -331,4 +211,4 @@ struct RENDERCORE_API FVirtualTextureUniformData
 
 	static FMatrix44f Invalid;
 };
-static_assert(sizeof(FVirtualTextureUniformData) <= sizeof(FMatrix44f), "FVirtualTextureUniformData is unable to pack");
\ No newline at end of file
+static_assert(sizeof(FVirtualTextureUniformData) <= sizeof(FMatrix44f), "FVirtualTextureUniformData is unable to pack");
diff --git a/Engine/Source/Runtime/RenderCore/Public/RenderTimer.h b/Engine/Source/Runtime/RenderCore/Public/RenderTimer.h
new file mode 100644
--- /dev/null
+++ b/Engine/Source/Runtime/RenderCore/Public/RenderTimer.h
@@ -0,0 +1,120 @@
+// Copyright Epic Games, Inc. All Rights Reserved.
+
+#pragma once
+
+#include "HAL/PlatformTime.h"
+
+/** Timer helper class. **/
+class FRenderTimer
+{
+public:
+	FRenderTimer()
+		: CurrentDeltaTime(0.0f)
+		, CurrentTime(0.0f)
+	{
+	}
+
+	/**
+	 *	Returns the current time, in seconds.
+	 *	@return Current time, in seconds
+	 */
+	float GetCurrentTime() const
+	{
+		return CurrentTime;
+	}
+
+	/**
+	 *	Returns the current delta time.
+	 *	@return Current delta time (number of seconds that passed between the last two tick)
+	 */
+	float GetCurrentDeltaTime() const
+	{
+		return CurrentDeltaTime;
+	}
+
+	/**
+	 *	Updates the timer.
+	 *	@param DeltaTime	Number of seconds that have passed since the last tick
+	 **/
+	void Tick(float DeltaTime)
+	{
+		CurrentDeltaTime = DeltaTime;
+		CurrentTime += DeltaTime;
+	}
+
+protected:
+	/** Current delta time (number of seconds that passed between the last two tick). */
+	float CurrentDeltaTime;
+	/** Current time, in seconds. */
+	float CurrentTime;
+};
+
+/** Whether to pause the global realtime clock for the rendering thread (read and write only on main thread). */
+extern RENDERCORE_API bool GPauseRenderingRealtimeClock;
+
+/** Global realtime clock for the rendering thread. */
+extern RENDERCORE_API FRenderTimer GRenderingRealtimeClock;
+
+/**
+ * Encapsulates a latency timer that measures the time from when mouse input
+ * is read on the gamethread until that frame is fully displayed by the GPU.
+ */
+struct FInputLatencyTimer
+{
+	/**
+	 * Constructor
+	 * @param InUpdateFrequency	How often the timer should be updated (in seconds).
+	 */
+	FInputLatencyTimer(float InUpdateFrequency)
+		: UpdateFrequency(InUpdateFrequency)
+	{
+	}
+
+	/** Potentially starts the timer on the gamethread, based on the UpdateFrequency. */
+	RENDERCORE_API void GameThreadTick();
+
+	/** @return The number of seconds of input latency. */
+	inline float GetDeltaSeconds() const
+	{
+		return FPlatformTime::ToSeconds(DeltaTime);
+	}
+
+	/** Whether GInputLatencyTimer is initialized or not. */
+	bool	bInitialized = false;
+
+	/** Whether a measurement has been triggered on the gamethread. */
+	bool	GameThreadTrigger = false;
+
+	/** Whether a measurement has been triggered on the renderthread. */
+	bool	RenderThreadTrigger = false;
+
+	/** Start time (in FPlatformTime::Cycles). */
+	uint32	StartTime = 0;
+
+	/** Last delta time that was measured (in FPlatformTime::Cycles). */
+	uint32	DeltaTime = 0;
+
+	/** Last time we did a measurement (in seconds). */
+	double	LastCaptureTime = 0.0f;
+
+	/** How often we should do a measurement (in seconds). */
+	float	UpdateFrequency;
+};
+
+/** Global input latency timer. Defined in UnrealClient.cpp */
+extern RENDERCORE_API FInputLatencyTimer GInputLatencyTimer;
+
+/** How many cycles the renderthread used (excluding idle time). It's set once per frame in FViewport::Draw. */
+extern RENDERCORE_API uint32 GRenderThreadTime;
+
+/** How many cycles the rhithread used (excluding idle time). */
+extern RENDERCORE_API uint32 GRHIThreadTime;
+
+/** How many cycles the gamethread used (excluding idle time). It's set once per frame in FViewport::Draw. */
+extern RENDERCORE_API uint32 GGameThreadTime;
+
+/** How many cycles it took to swap buffers to present the frame. */
+extern RENDERCORE_API uint32 GSwapBufferTime;
+
+/** How many cycles the renderthread used, including dependent wait time. */
+extern RENDERCORE_API uint32 GRenderThreadTimeCriticalPath;
