Moving some utility functions out of RenderResource and into RenderUtils.

diff --git a/Engine/Source/Runtime/Engine/Private/GPUSkinCacheVisualizationData.cpp b/Engine/Source/Runtime/Engine/Private/GPUSkinCacheVisualizationData.cpp
--- a/Engine/Source/Runtime/Engine/Private/GPUSkinCacheVisualizationData.cpp
+++ b/Engine/Source/Runtime/Engine/Private/GPUSkinCacheVisualizationData.cpp
@@ -4,6 +4,7 @@
 #include "HAL/IConsoleManager.h"
 #include "Materials/Material.h"
 #include "Misc/ConfigCacheIni.h"
+#include "RenderUtils.h"
 
 #define LOCTEXT_NAMESPACE "FGPUSkinCacheVisualizationData"
 
diff --git a/Engine/Source/Runtime/MovieSceneCapture/Private/MovieSceneCapture.cpp b/Engine/Source/Runtime/MovieSceneCapture/Private/MovieSceneCapture.cpp
--- a/Engine/Source/Runtime/MovieSceneCapture/Private/MovieSceneCapture.cpp
+++ b/Engine/Source/Runtime/MovieSceneCapture/Private/MovieSceneCapture.cpp
@@ -26,6 +26,7 @@
 #include "UObject/UObjectIterator.h"
 #include "Protocols/AudioCaptureProtocol.h"
 #include "Misc/DateTime.h"
+#include "RenderUtils.h"
 
 #include UE_INLINE_GENERATED_CPP_BY_NAME(MovieSceneCapture)
 
diff --git a/Engine/Source/Runtime/RenderCore/Private/RenderResource.cpp b/Engine/Source/Runtime/RenderCore/Private/RenderResource.cpp
--- a/Engine/Source/Runtime/RenderCore/Private/RenderResource.cpp
+++ b/Engine/Source/Runtime/RenderCore/Private/RenderResource.cpp
@@ -1031,45 +1031,3 @@ FRHISamplerState* FTexture::GetOrCreateSamplerState(const FSamplerStateInitializ
 	GTextureSamplerStateCache.Samplers.Add(Initializer, NewState);
 	return NewState;
 }
-
-bool IsRayTracingEnabled()
-{
-	bool bRayTracingEnabled = true;
-	extern RENDERCORE_API ERayTracingMode GRayTracingMode;	
-	static const auto RayTracingEnableCVar = IConsoleManager::Get().FindTConsoleVariableDataInt(TEXT("r.Raytracing.Enable"));
-	if (GRayTracingMode == ERayTracingMode::Dynamic && RayTracingEnableCVar)
-	{
-		bRayTracingEnabled = RayTracingEnableCVar->GetValueOnAnyThread() != 0;
-	}
-
-	return IsRayTracingAllowed() && bRayTracingEnabled;
-}
-
-bool IsRayTracingAllowed()
-{
-	checkf(GIsRHIInitialized, TEXT("IsRayTracingAllowed() may only be called once RHI is initialized."));
-
-#if DO_CHECK && WITH_EDITOR
-	{
-		// This function must not be called while cooking
-		if (IsRunningCookCommandlet())
-		{
-			return false;
-		}
-	}
-#endif // DO_CHECK && WITH_EDITOR
-
-	extern RENDERCORE_API ERayTracingMode GRayTracingMode;
-	return (int32)GRayTracingMode >= (int32)ERayTracingMode::Enabled;
-}
-
-bool IsRayTracingEnabled(EShaderPlatform ShaderPlatform)
-{
-	return IsRayTracingEnabled() && RHISupportsRayTracing(ShaderPlatform);
-}
-
-ERayTracingMode GetRayTracingMode()
-{
-	extern RENDERCORE_API ERayTracingMode GRayTracingMode;
-	return IsRayTracingAllowed() ? GRayTracingMode : ERayTracingMode::Disabled;
-}
diff --git a/Engine/Source/Runtime/RenderCore/Private/RenderUtils.cpp b/Engine/Source/Runtime/RenderCore/Private/RenderUtils.cpp
--- a/Engine/Source/Runtime/RenderCore/Private/RenderUtils.cpp
+++ b/Engine/Source/Runtime/RenderCore/Private/RenderUtils.cpp
@@ -1841,4 +1841,46 @@ bool IsUsingBasePassVelocity(const FStaticShaderPlatform Platform)
 RENDERCORE_API bool AllowTranslucencyPerObjectShadows(const FStaticShaderPlatform& Platform)
 {
 	return IsFeatureLevelSupported(Platform, ERHIFeatureLevel::SM5) && GAllowTranslucencyShadowsInProject != 0;
-}
\ No newline at end of file
+}
+
+bool IsRayTracingEnabled()
+{
+	bool bRayTracingEnabled = true;
+	extern RENDERCORE_API ERayTracingMode GRayTracingMode;
+	static const auto RayTracingEnableCVar = IConsoleManager::Get().FindTConsoleVariableDataInt(TEXT("r.Raytracing.Enable"));
+	if (GRayTracingMode == ERayTracingMode::Dynamic && RayTracingEnableCVar)
+	{
+		bRayTracingEnabled = RayTracingEnableCVar->GetValueOnAnyThread() != 0;
+	}
+
+	return IsRayTracingAllowed() && bRayTracingEnabled;
+}
+
+bool IsRayTracingAllowed()
+{
+	checkf(GIsRHIInitialized, TEXT("IsRayTracingAllowed() may only be called once RHI is initialized."));
+
+#if DO_CHECK && WITH_EDITOR
+	{
+		// This function must not be called while cooking
+		if (IsRunningCookCommandlet())
+		{
+			return false;
+		}
+	}
+#endif // DO_CHECK && WITH_EDITOR
+
+	extern RENDERCORE_API ERayTracingMode GRayTracingMode;
+	return (int32)GRayTracingMode >= (int32)ERayTracingMode::Enabled;
+}
+
+bool IsRayTracingEnabled(EShaderPlatform ShaderPlatform)
+{
+	return IsRayTracingEnabled() && RHISupportsRayTracing(ShaderPlatform);
+}
+
+ERayTracingMode GetRayTracingMode()
+{
+	extern RENDERCORE_API ERayTracingMode GRayTracingMode;
+	return IsRayTracingAllowed() ? GRayTracingMode : ERayTracingMode::Disabled;
+}
diff --git a/Engine/Source/Runtime/RenderCore/Public/RenderResource.h b/Engine/Source/Runtime/RenderCore/Public/RenderResource.h
--- a/Engine/Source/Runtime/RenderCore/Public/RenderResource.h
+++ b/Engine/Source/Runtime/RenderCore/Public/RenderResource.h
@@ -672,52 +672,6 @@ public:
 	virtual FString GetFriendlyName() const override { return TEXT("FIndexBuffer"); }
 };
 
-FORCEINLINE bool IsRayTracingEnabledForProject(EShaderPlatform ShaderPlatform)
-{
-	if (RHISupportsRayTracing(ShaderPlatform))
-	{
-		extern RENDERCORE_API TBitArray<TInlineAllocator<EShaderPlatform::SP_NumPlatforms / 8>> GRayTracingPlatformMask;
-		return (GRayTracingPlatformMask[(int)ShaderPlatform]);
-	}
-	else
-	{
-		return false;
-	}
-}
-
-FORCEINLINE bool ShouldCompileRayTracingShadersForProject(EShaderPlatform ShaderPlatform)
-{
-	if (RHISupportsRayTracingShaders(ShaderPlatform))
-	{
-		return IsRayTracingEnabledForProject(ShaderPlatform);
-	}
-	else
-	{
-		return false;
-	}
-}
-
-FORCEINLINE bool ShouldCompileRayTracingCallableShadersForProject(EShaderPlatform ShaderPlatform)
-{
-	return RHISupportsRayTracingCallableShaders(ShaderPlatform) && ShouldCompileRayTracingShadersForProject(ShaderPlatform);
-}
-
-// Returns `true` when running on RT-capable machine, RT support is enabled for the project and by game graphics options and RT is enabled with r.Raytracing.Enable
-// This function may only be called at runtime, never during cooking.
-extern RENDERCORE_API bool IsRayTracingEnabled();
-
-// Returns 'true' when running on RT-capable machine, RT support is enabled for the project and by game graphics options and ShaderPlatform supports RT and RT is enabled with r.Raytracing.Enable
-// This function may only be called at runtime, never during cooking.
-RENDERCORE_API bool IsRayTracingEnabled(EShaderPlatform ShaderPlatform);
-
-// Returns 'true' when running on RT-capable machine, RT support is enabled for the project and by game graphics options
-// This function may only be called at runtime, never during cooking.
-extern RENDERCORE_API bool IsRayTracingAllowed();
-
-// Returns the ray tracing mode if ray tracing is allowed.
-// This function may only be called at runtime, never during cooking.
-extern RENDERCORE_API ERayTracingMode GetRayTracingMode();
-
 /**
  * A system for dynamically allocating GPU memory for vertices.
  */
@@ -863,22 +817,7 @@ private:
 	struct FDynamicIndexBufferPool* Pools[2];
 };
 
-/**Note, this should only be used when a platform requires special shader compilation for 32 bit pixel format render targets.
-Does not replace pixel format associations across the board**/
-
-FORCEINLINE bool PlatformRequires128bitRT(EPixelFormat PixelFormat)
-{
-	switch (PixelFormat)
-	{
-	case PF_R32_FLOAT:
-	case PF_G32R32F:
-	case PF_A32B32G32R32F:
-		return FDataDrivenShaderPlatformInfo::GetRequiresExplicit128bitRT(GMaxRHIShaderPlatform);
-	default:
-		return false;
-	}
-}
-
 #if UE_ENABLE_INCLUDE_ORDER_DEPRECATED_IN_5_2
 #include "RayTracingGeometry.h"
+#include "RenderUtils.h"
 #endif
diff --git a/Engine/Source/Runtime/RenderCore/Public/RenderUtils.h b/Engine/Source/Runtime/RenderCore/Public/RenderUtils.h
--- a/Engine/Source/Runtime/RenderCore/Public/RenderUtils.h
+++ b/Engine/Source/Runtime/RenderCore/Public/RenderUtils.h
@@ -858,3 +858,66 @@ RENDERCORE_API bool DoesRuntimeSupportOnePassPointLightShadows(EShaderPlatform P
  * Read-only switch to check if translucency per object shadows are enabled.
  */
 RENDERCORE_API bool AllowTranslucencyPerObjectShadows(const FStaticShaderPlatform &Platform);
+
+
+/**Note, this should only be used when a platform requires special shader compilation for 32 bit pixel format render targets.
+Does not replace pixel format associations across the board**/
+
+FORCEINLINE bool PlatformRequires128bitRT(EPixelFormat PixelFormat)
+{
+	switch (PixelFormat)
+	{
+	case PF_R32_FLOAT:
+	case PF_G32R32F:
+	case PF_A32B32G32R32F:
+		return FDataDrivenShaderPlatformInfo::GetRequiresExplicit128bitRT(GMaxRHIShaderPlatform);
+	default:
+		return false;
+	}
+}
+
+FORCEINLINE bool IsRayTracingEnabledForProject(EShaderPlatform ShaderPlatform)
+{
+	if (RHISupportsRayTracing(ShaderPlatform))
+	{
+		extern RENDERCORE_API TBitArray<TInlineAllocator<EShaderPlatform::SP_NumPlatforms / 8>> GRayTracingPlatformMask;
+		return (GRayTracingPlatformMask[(int)ShaderPlatform]);
+	}
+	else
+	{
+		return false;
+	}
+}
+
+FORCEINLINE bool ShouldCompileRayTracingShadersForProject(EShaderPlatform ShaderPlatform)
+{
+	if (RHISupportsRayTracingShaders(ShaderPlatform))
+	{
+		return IsRayTracingEnabledForProject(ShaderPlatform);
+	}
+	else
+	{
+		return false;
+	}
+}
+
+FORCEINLINE bool ShouldCompileRayTracingCallableShadersForProject(EShaderPlatform ShaderPlatform)
+{
+	return RHISupportsRayTracingCallableShaders(ShaderPlatform) && ShouldCompileRayTracingShadersForProject(ShaderPlatform);
+}
+
+// Returns `true` when running on RT-capable machine, RT support is enabled for the project and by game graphics options and RT is enabled with r.Raytracing.Enable
+// This function may only be called at runtime, never during cooking.
+extern RENDERCORE_API bool IsRayTracingEnabled();
+
+// Returns 'true' when running on RT-capable machine, RT support is enabled for the project and by game graphics options and ShaderPlatform supports RT and RT is enabled with r.Raytracing.Enable
+// This function may only be called at runtime, never during cooking.
+RENDERCORE_API bool IsRayTracingEnabled(EShaderPlatform ShaderPlatform);
+
+// Returns 'true' when running on RT-capable machine, RT support is enabled for the project and by game graphics options
+// This function may only be called at runtime, never during cooking.
+extern RENDERCORE_API bool IsRayTracingAllowed();
+
+// Returns the ray tracing mode if ray tracing is allowed.
+// This function may only be called at runtime, never during cooking.
+extern RENDERCORE_API ERayTracingMode GetRayTracingMode();
diff --git a/Engine/Source/Runtime/Renderer/Private/Tests/RayTracingTestbed.cpp b/Engine/Source/Runtime/Renderer/Private/Tests/RayTracingTestbed.cpp
--- a/Engine/Source/Runtime/Renderer/Private/Tests/RayTracingTestbed.cpp
+++ b/Engine/Source/Runtime/Renderer/Private/Tests/RayTracingTestbed.cpp
@@ -16,6 +16,7 @@ IMPLEMENT_SIMPLE_AUTOMATION_TEST(FRayTracingTestbed, "System.Renderer.RayTracing
 #include "RayTracingPayloadType.h"
 #include "GlobalShader.h"
 #include "Containers/DynamicRHIResourceArray.h"
+#include "RenderUtils.h"
 
 // HINT: Execute this test via console command in editor: Automation RunTest System.Renderer.RayTracing.BasicRayTracing
 bool RunRayTracingTestbed_RenderThread(const FString& Parameters)
