Moving all RenderResource implementations in RenderResource.h into the cpp file for future forward declaration.

diff --git a/Engine/Source/Runtime/RenderCore/Private/RenderResource.cpp b/Engine/Source/Runtime/RenderCore/Private/RenderResource.cpp
--- a/Engine/Source/Runtime/RenderCore/Private/RenderResource.cpp
+++ b/Engine/Source/Runtime/RenderCore/Private/RenderResource.cpp
@@ -307,6 +307,18 @@ void FRenderResource::UpdateRHI()
 	}
 }
 
+FRenderResource::FRenderResource()
+	: ListIndex(INDEX_NONE)
+	, FeatureLevel(ERHIFeatureLevel::Num)
+{
+}
+
+FRenderResource::FRenderResource(ERHIFeatureLevel::Type InFeatureLevel)
+	: ListIndex(INDEX_NONE)
+	, FeatureLevel(InFeatureLevel)
+{
+}
+
 FRenderResource::~FRenderResource()
 {
 	checkf(ResourceState == ERenderResourceState::Default, TEXT(" Invalid Resource State: %s"), ResourceState == ERenderResourceState::BatchReleased ? TEXT("BatchReleased") : TEXT("Deleted"));
@@ -318,6 +330,33 @@ FRenderResource::~FRenderResource()
 	}
 }
 
+FBufferRHIRef FRenderResource::CreateRHIBufferInternal(
+	const TCHAR* InDebugName,
+	uint32 ResourceCount,
+	EBufferUsageFlags InBufferUsageFlags,
+	FResourceArrayInterface* ResourceArray,
+	bool bRenderThread,
+	bool bWithoutNativeResource)
+{
+	const uint32 SizeInBytes = ResourceArray ? ResourceArray->GetResourceDataSize() : 0;
+	FRHIResourceCreateInfo CreateInfo(InDebugName, ResourceArray);
+	CreateInfo.bWithoutNativeResource = bWithoutNativeResource;
+
+	FBufferRHIRef Buffer;
+	if (bRenderThread)
+	{
+		Buffer = RHICreateVertexBuffer(SizeInBytes, InBufferUsageFlags, CreateInfo);
+	}
+	else
+	{
+		FRHIAsyncCommandList CommandList;
+		Buffer = CommandList->CreateBuffer(SizeInBytes, InBufferUsageFlags | EBufferUsageFlags::VertexBuffer, 0, ERHIAccess::SRVMask, CreateInfo);
+	}
+
+	Buffer->SetOwnerName(GetOwnerName());
+	return Buffer;
+}
+
 void FRenderResource::SetOwnerName(const FName& InOwnerName)
 {
 #if RHI_ENABLE_RESOURCE_INFO
@@ -438,6 +477,94 @@ void ReleaseResourceAndFlush(FRenderResource* Resource)
 	FlushRenderingCommands();
 }
 
+//////////////////////////////////////////////////////////////////////////////////////////////////////////////////////
+// FTextureSamplerStateCache
+
+class FTextureSamplerStateCache : public FRenderResource
+{
+public:
+	TMap<FSamplerStateInitializerRHI, FRHISamplerState*> Samplers;
+
+	virtual void ReleaseRHI() override
+	{
+		for (auto Pair : Samplers)
+		{
+			Pair.Value->Release();
+		}
+		Samplers.Empty();
+	}
+};
+
+TGlobalResource<FTextureSamplerStateCache> GTextureSamplerStateCache;
+
+//////////////////////////////////////////////////////////////////////////////////////////////////////////////////////
+// FTexture
+
+FTexture::FTexture() = default;
+FTexture::~FTexture() = default;
+
+uint32 FTexture::GetSizeX() const
+{
+	return 0;
+}
+
+uint32 FTexture::GetSizeY() const
+{
+	return 0;
+}
+
+uint32 FTexture::GetSizeZ() const
+{
+	return 0;
+}
+
+void FTexture::ReleaseRHI()
+{
+	TextureRHI.SafeRelease();
+	SamplerStateRHI.SafeRelease();
+	DeferredPassSamplerStateRHI.SafeRelease();
+}
+
+FString FTexture::GetFriendlyName() const
+{
+	return TEXT("FTexture");
+}
+
+FRHISamplerState* FTexture::GetOrCreateSamplerState(const FSamplerStateInitializerRHI& Initializer)
+{
+	// This sampler cache is supposed to be used only from RT
+	// Add a lock here if it's used from multiple threads
+	check(IsInRenderingThread());
+
+	FRHISamplerState** Found = GTextureSamplerStateCache.Samplers.Find(Initializer);
+	if (Found)
+	{
+		return *Found;
+	}
+
+	FSamplerStateRHIRef NewState = RHICreateSamplerState(Initializer);
+
+	// Add an extra reference so we don't have TRefCountPtr in the maps
+	NewState->AddRef();
+	GTextureSamplerStateCache.Samplers.Add(Initializer, NewState);
+	return NewState;
+}
+
+//////////////////////////////////////////////////////////////////////////////////////////////////////////////////////
+// FTextureWithSRV
+
+FTextureWithSRV::FTextureWithSRV() = default;
+FTextureWithSRV::~FTextureWithSRV() = default;
+
+void FTextureWithSRV::ReleaseRHI()
+{
+	ShaderResourceViewRHI.SafeRelease();
+	FTexture::ReleaseRHI();
+}
+
+//////////////////////////////////////////////////////////////////////////////////////////////////////////////////////
+// FTextureReference
+
 FTextureReference::FTextureReference()
 	: TextureReferenceRHI(NULL)
 {
@@ -499,12 +626,108 @@ FString FTextureReference::GetFriendlyName() const
 	return TEXT("FTextureReference");
 }
 
+//////////////////////////////////////////////////////////////////////////////////////////////////////////////////////
+// FVertexBuffer
+
+FVertexBuffer::FVertexBuffer() = default;
+FVertexBuffer::~FVertexBuffer() = default;
+
+void FVertexBuffer::ReleaseRHI()
+{
+	VertexBufferRHI.SafeRelease();
+}
+
+FString FVertexBuffer::GetFriendlyName() const
+{
+	return TEXT("FVertexBuffer");
+}
+
+//////////////////////////////////////////////////////////////////////////////////////////////////////////////////////
+// FVertexBufferWithSRV
+
+FVertexBufferWithSRV::FVertexBufferWithSRV() = default;
+FVertexBufferWithSRV::~FVertexBufferWithSRV() = default;
+
+void FVertexBufferWithSRV::ReleaseRHI()
+{
+	ShaderResourceViewRHI.SafeRelease();
+	UnorderedAccessViewRHI.SafeRelease();
+	FVertexBuffer::ReleaseRHI();
+}
+
+//////////////////////////////////////////////////////////////////////////////////////////////////////////////////////
+// FNullColorVertexBuffer
+
+FNullColorVertexBuffer::FNullColorVertexBuffer() = default;
+FNullColorVertexBuffer::~FNullColorVertexBuffer() = default;
+
+void FNullColorVertexBuffer::InitRHI()
+{
+	// create a static vertex buffer
+	FRHIResourceCreateInfo CreateInfo(TEXT("FNullColorVertexBuffer"));
+
+	VertexBufferRHI = RHICreateBuffer(sizeof(uint32) * 4, BUF_Static | BUF_VertexBuffer | BUF_ShaderResource, 0, ERHIAccess::VertexOrIndexBuffer | ERHIAccess::SRVMask, CreateInfo);
+	uint32* Vertices = (uint32*)RHILockBuffer(VertexBufferRHI, 0, sizeof(uint32) * 4, RLM_WriteOnly);
+	Vertices[0] = FColor(255, 255, 255, 255).DWColor();
+	Vertices[1] = FColor(255, 255, 255, 255).DWColor();
+	Vertices[2] = FColor(255, 255, 255, 255).DWColor();
+	Vertices[3] = FColor(255, 255, 255, 255).DWColor();
+	RHIUnlockBuffer(VertexBufferRHI);
+	VertexBufferSRV = RHICreateShaderResourceView(VertexBufferRHI, sizeof(FColor), PF_R8G8B8A8);
+}
+
+void FNullColorVertexBuffer::ReleaseRHI()
+{
+	VertexBufferSRV.SafeRelease();
+	FVertexBuffer::ReleaseRHI();
+}
+
 /** The global null color vertex buffer, which is set with a stride of 0 on meshes without a color component. */
 TGlobalResource<FNullColorVertexBuffer> GNullColorVertexBuffer;
 
+//////////////////////////////////////////////////////////////////////////////////////////////////////////////////////
+// FNullVertexBuffer
+
+FNullVertexBuffer::FNullVertexBuffer() = default;
+FNullVertexBuffer::~FNullVertexBuffer() = default;
+
+void FNullVertexBuffer::InitRHI()
+{
+	// create a static vertex buffer
+	FRHIResourceCreateInfo CreateInfo(TEXT("FNullVertexBuffer"));
+	VertexBufferRHI = RHICreateBuffer(sizeof(FVector3f), BUF_Static | BUF_VertexBuffer | BUF_ShaderResource, 0, ERHIAccess::VertexOrIndexBuffer | ERHIAccess::SRVMask, CreateInfo);
+	FVector3f* LockedData = (FVector3f*)RHILockBuffer(VertexBufferRHI, 0, sizeof(FVector3f), RLM_WriteOnly);
+	*LockedData = FVector3f(0.0f);
+	RHIUnlockBuffer(VertexBufferRHI);
+
+	VertexBufferSRV = RHICreateShaderResourceView(VertexBufferRHI, sizeof(FColor), PF_R8G8B8A8);
+}
+
+void FNullVertexBuffer::ReleaseRHI()
+{
+	VertexBufferSRV.SafeRelease();
+	FVertexBuffer::ReleaseRHI();
+}
+
 /** The global null vertex buffer, which is set with a stride of 0 on meshes */
 TGlobalResource<FNullVertexBuffer> GNullVertexBuffer;
 
+//////////////////////////////////////////////////////////////////////////////////////////////////////////////////////
+// FIndexBuffer
+
+FIndexBuffer::FIndexBuffer() = default;
+FIndexBuffer::~FIndexBuffer() = default;
+
+void FIndexBuffer::ReleaseRHI()
+{
+	IndexBufferRHI.SafeRelease();
+}
+
+FString FIndexBuffer::GetFriendlyName() const
+{
+	return TEXT("FIndexBuffer");
+}
+
 /*------------------------------------------------------------------------------
 	FGlobalDynamicVertexBuffer implementation.
 ------------------------------------------------------------------------------*/
@@ -994,40 +1217,3 @@ void FMipBiasFade::SetNewMipCount( float ActualMipCount, float TargetMipCount, d
 		}
 	}
 }
-
-class FTextureSamplerStateCache : public FRenderResource
-{
-public:
-	TMap<FSamplerStateInitializerRHI, FRHISamplerState*> Samplers;
-
-	virtual void ReleaseRHI() override
-	{
-		for (auto Pair : Samplers)
-		{
-			Pair.Value->Release();
-		}
-		Samplers.Empty();
-	}
-};
-
-TGlobalResource<FTextureSamplerStateCache> GTextureSamplerStateCache;
-
-FRHISamplerState* FTexture::GetOrCreateSamplerState(const FSamplerStateInitializerRHI& Initializer)
-{
-	// This sampler cache is supposed to be used only from RT
-	// Add a lock here if it's used from multiple threads
-	check(IsInRenderingThread());
-	
-	FRHISamplerState** Found = GTextureSamplerStateCache.Samplers.Find(Initializer);
-	if (Found)
-	{
-		return *Found;
-	}
-	
-	FSamplerStateRHIRef NewState = RHICreateSamplerState(Initializer);
-	
-	// Add an extra reference so we don't have TRefCountPtr in the maps
-	NewState->AddRef();
-	GTextureSamplerStateCache.Samplers.Add(Initializer, NewState);
-	return NewState;
-}
diff --git a/Engine/Source/Runtime/RenderCore/Public/RenderResource.h b/Engine/Source/Runtime/RenderCore/Public/RenderResource.h
--- a/Engine/Source/Runtime/RenderCore/Public/RenderResource.h
+++ b/Engine/Source/Runtime/RenderCore/Public/RenderResource.h
@@ -28,6 +28,8 @@
 #include "Serialization/MemoryLayout.h"
 #include "DataDrivenShaderPlatformInfo.h"
 
+class FRDGPooledBuffer;
+
 typedef TBitArray<TInlineAllocator<EShaderPlatform::SP_NumPlatforms / 8>>	ShaderPlatformMaskType;
 
 /** Number of frames after which unused global resource allocations will be discarded. */
@@ -75,16 +77,10 @@ public:
 	////////////////////////////////////////////////////////////////////////////////////
 
 	/** Default constructor. */
-	FRenderResource()
-		: ListIndex(INDEX_NONE)
-		, FeatureLevel(ERHIFeatureLevel::Num)
-	{}
+	FRenderResource();
 
 	/** Constructor when we know what feature level this resource should support */
-	FRenderResource(ERHIFeatureLevel::Type InFeatureLevel)
-		: ListIndex(INDEX_NONE)
-		, FeatureLevel(InFeatureLevel)
-	{}
+	FRenderResource(ERHIFeatureLevel::Type InFeatureLevel);
 
 	/** Destructor used to catch unreleased resources. */
 	virtual ~FRenderResource();
@@ -149,12 +145,6 @@ public:
 	void SetOwnerName(const FName& InOwnerName);
 	FName GetOwnerName() const;
 
-private:
-	int32 ListIndex;
-#if RHI_ENABLE_RESOURCE_INFO
-	FName OwnerName;
-#endif
-
 protected:
 	// This is used during mobile editor preview refactor, this will eventually be replaced with a parameter to InitRHI() etc..
 	void SetFeatureLevel(const FStaticFeatureLevel InFeatureLevel) { FeatureLevel = (ERHIFeatureLevel::Type)InFeatureLevel; }
@@ -169,20 +159,9 @@ protected:
 		FResourceArrayInterface* RESTRICT ResourceArray = InOutResourceObject ? InOutResourceObject->GetResourceArray() : nullptr;
 		if (ResourceCount != 0)
 		{
-			const uint32 SizeInBytes = ResourceArray ? ResourceArray->GetResourceDataSize() : 0;
-			FRHIResourceCreateInfo CreateInfo(InDebugName, ResourceArray);
-			CreateInfo.bWithoutNativeResource = !InOutResourceObject;
-
-			if (bRenderThread)
+			Buffer = CreateRHIBufferInternal(InDebugName, ResourceCount, InBufferUsageFlags, ResourceArray, bRenderThread, InOutResourceObject == nullptr);
+			if (!bRenderThread)
 			{
-				Buffer = RHICreateVertexBuffer(SizeInBytes, InBufferUsageFlags, CreateInfo);
-				Buffer->SetOwnerName(GetOwnerName());
-			}
-			else
-			{
-				FRHIAsyncCommandList CommandList;
-				Buffer = CommandList->CreateBuffer(SizeInBytes, InBufferUsageFlags | EBufferUsageFlags::VertexBuffer, 0, ERHIAccess::SRVMask, CreateInfo);
-				Buffer->SetOwnerName(GetOwnerName());
 				return Buffer;
 			}
 		}
@@ -198,6 +177,20 @@ protected:
 	}
 
 private:
+	static FBufferRHIRef CreateRHIBufferInternal(
+		const TCHAR* InDebugName,
+		uint32 ResourceCount,
+		EBufferUsageFlags InBufferUsageFlags,
+		FResourceArrayInterface* ResourceArray,
+		bool bRenderThread,
+		bool bWithoutNativeResource
+	);
+
+#if RHI_ENABLE_RESOURCE_INFO
+	FName OwnerName;
+#endif
+
+	int32 ListIndex;
 	TEnumAsByte<ERHIFeatureLevel::Type> FeatureLevel;
 
 public:
@@ -412,7 +405,7 @@ struct FMipBiasFade
 };
 
 /** A textures resource. */
-class FTexture : public FRenderResource
+class RENDERCORE_API FTexture : public FRenderResource
 {
 public:
 
@@ -426,7 +419,7 @@ public:
 	FSamplerStateRHIRef DeferredPassSamplerStateRHI;
 
 	/** The last time the texture has been bound */
-	mutable double		LastRenderTime;
+	mutable double		LastRenderTime = -FLT_MAX;
 
 	/** Base values for fading in/out mip-levels. */
 	FMipBiasFade		MipBiasFade;
@@ -435,79 +428,53 @@ public:
 	 *  this is set from CompressionSettings, not PixelFormat
 	 *  this is only used by Editor/Debug shaders, not real game materials, which use SamplerType from MaterialExpressions
 	 */
-	bool				bGreyScaleFormat;
+	bool				bGreyScaleFormat = false;
 
 	/**
 	 * true if the texture is in the same gamma space as the intended rendertarget (e.g. screenshots).
 	 * The texture will have sRGB==false and bIgnoreGammaConversions==true, causing a non-sRGB texture lookup
 	 * and no gamma-correction in the shader.
 	 */
-	bool				bIgnoreGammaConversions;
+	bool				bIgnoreGammaConversions = false;
 
 	/** 
 	 * Is the pixel data in this texture sRGB?
 	 **/
-	bool				bSRGB;
+	bool				bSRGB = false;
 
-	/** Default constructor. */
-	FTexture()
-	: TextureRHI(NULL)
-	, SamplerStateRHI(NULL)
-	, DeferredPassSamplerStateRHI(NULL)
-	, LastRenderTime(-FLT_MAX)
-	, bGreyScaleFormat(false)
-	, bIgnoreGammaConversions(false)
-	, bSRGB(false)
-	{}
-
-	// Destructor
-	virtual ~FTexture() {}
+	FTexture();
+	virtual ~FTexture();
 
-	FRHITexture* GetTextureRHI() { return TextureRHI; }
+	const FTextureRHIRef& GetTextureRHI() { return TextureRHI; }
 
 	/** Returns the width of the texture in pixels. */
-	virtual uint32 GetSizeX() const
-	{
-		return 0;
-	}
+	virtual uint32 GetSizeX() const;
+
 	/** Returns the height of the texture in pixels. */
-	virtual uint32 GetSizeY() const
-	{
-		return 0;
-	}
+	virtual uint32 GetSizeY() const;
+
 	/** Returns the depth of the texture in pixels. */
-	virtual uint32 GetSizeZ() const
-	{
-		return 0;
-	}
+	virtual uint32 GetSizeZ() const;
 
 	// FRenderResource interface.
-	virtual void ReleaseRHI() override
-	{
-		TextureRHI.SafeRelease();
-		SamplerStateRHI.SafeRelease();
-		DeferredPassSamplerStateRHI.SafeRelease();
-	}
-	virtual FString GetFriendlyName() const override { return TEXT("FTexture"); }
+	virtual void ReleaseRHI() override;
+	virtual FString GetFriendlyName() const override;
 
 protected:
-	RENDERCORE_API static FRHISamplerState* GetOrCreateSamplerState(const FSamplerStateInitializerRHI& Initializer);
+	static FRHISamplerState* GetOrCreateSamplerState(const FSamplerStateInitializerRHI& Initializer);
 };
 
 /** A textures resource that includes an SRV. */
-class FTextureWithSRV : public FTexture
+class RENDERCORE_API FTextureWithSRV : public FTexture
 {
 public:
-	/** SRV that views the entire texture */
-	FShaderResourceViewRHIRef ShaderResourceViewRHI;
+	FTextureWithSRV();
+	virtual ~FTextureWithSRV();
 
-	virtual ~FTextureWithSRV() {}
+	virtual void ReleaseRHI() override;
 
-	virtual void ReleaseRHI() override
-	{
-		ShaderResourceViewRHI.SafeRelease();
-		FTexture::ReleaseRHI();
-	}
+	/** SRV that views the entire texture */
+	FShaderResourceViewRHIRef ShaderResourceViewRHI;
 };
 
 /** A texture reference resource. */
@@ -554,67 +521,43 @@ public:
 class RENDERCORE_API FVertexBuffer : public FRenderResource
 {
 public:
-	FBufferRHIRef VertexBufferRHI;
-
-	/** Destructor. */
-	virtual ~FVertexBuffer() {}
+	FVertexBuffer();
+	virtual ~FVertexBuffer();
 
 	// FRenderResource interface.
-	virtual void ReleaseRHI() override
-	{
-		VertexBufferRHI.SafeRelease();
-	}
+	virtual void ReleaseRHI() override;
+	virtual FString GetFriendlyName() const override;
 
-	virtual FString GetFriendlyName() const override { return TEXT("FVertexBuffer"); }
+	FBufferRHIRef VertexBufferRHI;
 };
 
 class RENDERCORE_API FVertexBufferWithSRV : public FVertexBuffer
 {
 public:
+	FVertexBufferWithSRV();
+	~FVertexBufferWithSRV();
+
+	virtual void ReleaseRHI() override;
+
 	/** SRV that views the entire texture */
 	FShaderResourceViewRHIRef ShaderResourceViewRHI;
 
 	/** *optional* UAV that views the entire texture */
 	FUnorderedAccessViewRHIRef UnorderedAccessViewRHI;
-
-	virtual void ReleaseRHI() override
-	{
-		ShaderResourceViewRHI.SafeRelease();
-		UnorderedAccessViewRHI.SafeRelease();
-		FVertexBuffer::ReleaseRHI();
-	}
 };
 
 /**
 * A vertex buffer with a single color component.  This is used on meshes that don't have a color component
 * to keep from needing a separate vertex factory to handle this case.
 */
-class FNullColorVertexBuffer : public FVertexBuffer
+class RENDERCORE_API FNullColorVertexBuffer : public FVertexBuffer
 {
 public:
-	/** 
-	* Initialize the RHI for this rendering resource 
-	*/
-	virtual void InitRHI() override
-	{
-		// create a static vertex buffer
-		FRHIResourceCreateInfo CreateInfo(TEXT("FNullColorVertexBuffer"));
-
-		VertexBufferRHI = RHICreateBuffer(sizeof(uint32) * 4, BUF_Static | BUF_VertexBuffer | BUF_ShaderResource, 0, ERHIAccess::VertexOrIndexBuffer | ERHIAccess::SRVMask, CreateInfo);
-		uint32* Vertices = (uint32*)RHILockBuffer(VertexBufferRHI, 0, sizeof(uint32) * 4, RLM_WriteOnly);
-		Vertices[0] = FColor(255, 255, 255, 255).DWColor();
-		Vertices[1] = FColor(255, 255, 255, 255).DWColor();
-		Vertices[2] = FColor(255, 255, 255, 255).DWColor();
-		Vertices[3] = FColor(255, 255, 255, 255).DWColor();
-		RHIUnlockBuffer(VertexBufferRHI);
-		VertexBufferSRV = RHICreateShaderResourceView(VertexBufferRHI, sizeof(FColor), PF_R8G8B8A8);
-	}
+	FNullColorVertexBuffer();
+	~FNullColorVertexBuffer();
 
-	virtual void ReleaseRHI() override
-	{
-		VertexBufferSRV.SafeRelease();
-		FVertexBuffer::ReleaseRHI();
-	}
+	virtual void InitRHI() override;
+	virtual void ReleaseRHI() override;
 
 	FShaderResourceViewRHIRef VertexBufferSRV;
 };
@@ -625,29 +568,14 @@ extern RENDERCORE_API TGlobalResource<FNullColorVertexBuffer> GNullColorVertexBu
 /**
 * A vertex buffer with a single zero float3 component.
 */
-class FNullVertexBuffer : public FVertexBuffer
+class RENDERCORE_API FNullVertexBuffer : public FVertexBuffer
 {
 public:
-	/**
-	* Initialize the RHI for this rendering resource
-	*/
-	virtual void InitRHI() override
-	{
-		// create a static vertex buffer
-		FRHIResourceCreateInfo CreateInfo(TEXT("FNullVertexBuffer"));
-		VertexBufferRHI = RHICreateBuffer(sizeof(FVector3f), BUF_Static | BUF_VertexBuffer | BUF_ShaderResource, 0, ERHIAccess::VertexOrIndexBuffer | ERHIAccess::SRVMask, CreateInfo);
-		FVector3f* LockedData = (FVector3f*)RHILockBuffer(VertexBufferRHI, 0, sizeof(FVector3f), RLM_WriteOnly);
-		*LockedData = FVector3f(0.0f);
-		RHIUnlockBuffer(VertexBufferRHI);
-
-		VertexBufferSRV = RHICreateShaderResourceView(VertexBufferRHI, sizeof(FColor), PF_R8G8B8A8);
-	}
+	FNullVertexBuffer();
+	~FNullVertexBuffer();
 
-	virtual void ReleaseRHI() override
-	{
-		VertexBufferSRV.SafeRelease();
-		FVertexBuffer::ReleaseRHI();
-	}
+	virtual void InitRHI() override;
+	virtual void ReleaseRHI() override;
 
 	FShaderResourceViewRHIRef VertexBufferSRV;
 };
@@ -656,20 +584,17 @@ public:
 extern RENDERCORE_API TGlobalResource<FNullVertexBuffer> GNullVertexBuffer;
 
 /** An index buffer resource. */
-class FIndexBuffer : public FRenderResource
+class RENDERCORE_API FIndexBuffer : public FRenderResource
 {
 public:
-	FBufferRHIRef IndexBufferRHI;
-
-	/** Destructor. */
-	virtual ~FIndexBuffer() {}
+	FIndexBuffer();
+	virtual ~FIndexBuffer();
 
 	// FRenderResource interface.
-	virtual void ReleaseRHI() override
-	{
-		IndexBufferRHI.SafeRelease();
-	}
-	virtual FString GetFriendlyName() const override { return TEXT("FIndexBuffer"); }
+	virtual void ReleaseRHI() override;
+	virtual FString GetFriendlyName() const override;
+
+	FBufferRHIRef IndexBufferRHI;
 };
 
 /**
