Fixing CIS after manual merge.

diff --git a/Engine/Source/Runtime/RenderCore/Private/RenderResource.cpp b/Engine/Source/Runtime/RenderCore/Private/RenderResource.cpp
--- a/Engine/Source/Runtime/RenderCore/Private/RenderResource.cpp
+++ b/Engine/Source/Runtime/RenderCore/Private/RenderResource.cpp
@@ -332,6 +332,7 @@ FRenderResource::~FRenderResource()
 
 FBufferRHIRef FRenderResource::CreateRHIBufferInternal(
 	const TCHAR* InDebugName,
+	const FName& InOwnerName,
 	uint32 ResourceCount,
 	EBufferUsageFlags InBufferUsageFlags,
 	FResourceArrayInterface* ResourceArray,
@@ -353,7 +354,7 @@ FBufferRHIRef FRenderResource::CreateRHIBufferInternal(
 		Buffer = CommandList->CreateBuffer(SizeInBytes, InBufferUsageFlags | EBufferUsageFlags::VertexBuffer, 0, ERHIAccess::SRVMask, CreateInfo);
 	}
 
-	Buffer->SetOwnerName(GetOwnerName());
+	Buffer->SetOwnerName(InOwnerName);
 	return Buffer;
 }
 
diff --git a/Engine/Source/Runtime/RenderCore/Public/RenderResource.h b/Engine/Source/Runtime/RenderCore/Public/RenderResource.h
--- a/Engine/Source/Runtime/RenderCore/Public/RenderResource.h
+++ b/Engine/Source/Runtime/RenderCore/Public/RenderResource.h
@@ -159,7 +159,7 @@ protected:
 		FResourceArrayInterface* RESTRICT ResourceArray = InOutResourceObject ? InOutResourceObject->GetResourceArray() : nullptr;
 		if (ResourceCount != 0)
 		{
-			Buffer = CreateRHIBufferInternal(InDebugName, ResourceCount, InBufferUsageFlags, ResourceArray, bRenderThread, InOutResourceObject == nullptr);
+			Buffer = CreateRHIBufferInternal(InDebugName, GetOwnerName(), ResourceCount, InBufferUsageFlags, ResourceArray, bRenderThread, InOutResourceObject == nullptr);
 			if (!bRenderThread)
 			{
 				return Buffer;
@@ -179,6 +179,7 @@ protected:
 private:
 	static FBufferRHIRef CreateRHIBufferInternal(
 		const TCHAR* InDebugName,
+		const FName& InOwnerName,
 		uint32 ResourceCount,
 		EBufferUsageFlags InBufferUsageFlags,
 		FResourceArrayInterface* ResourceArray,
