Fix CIS and reduce numbers of shaders added (UE-171824).

- Only compile array copy/resolve shaders for platforms with MMV (mobile multiview).

diff --git a/Engine/Source/Runtime/Engine/Public/ScreenRendering.h b/Engine/Source/Runtime/Engine/Public/ScreenRendering.h
--- a/Engine/Source/Runtime/Engine/Public/ScreenRendering.h
+++ b/Engine/Source/Runtime/Engine/Public/ScreenRendering.h
@@ -14,6 +14,7 @@
 #include "GlobalShader.h"
 #include "ShaderParameterUtils.h"
 #include "SceneView.h"
+#include "StereoRenderUtils.h"
 
 struct FScreenVertex
 {
@@ -98,7 +99,15 @@ class FScreenFromSlice0PS : public FScreenPS
 	DECLARE_EXPORTED_SHADER_TYPE(FScreenFromSlice0PS, Global, ENGINE_API);
 public:
 
-	static bool ShouldCompilePermutation(const FGlobalShaderPermutationParameters& Parameters) { return true; }
+	static bool ShouldCompilePermutation(const FGlobalShaderPermutationParameters& Parameters)
+	{ 
+		if (FScreenPS::ShouldCompilePermutation(Parameters))
+		{
+			UE::StereoRenderUtils::FStereoShaderAspects Aspects(Parameters.Platform);
+			return Aspects.IsMobileMultiViewEnabled();
+		}
+		return false;
+	}
 
 	static void ModifyCompilationEnvironment(const FGlobalShaderPermutationParameters& Parameters, FShaderCompilerEnvironment& OutEnvironment)
 	{
diff --git a/Engine/Source/Runtime/RenderCore/Public/ResolveShader.h b/Engine/Source/Runtime/RenderCore/Public/ResolveShader.h
--- a/Engine/Source/Runtime/RenderCore/Public/ResolveShader.h
+++ b/Engine/Source/Runtime/RenderCore/Public/ResolveShader.h
@@ -12,6 +12,7 @@
 #include "ShaderCore.h"
 #include "ShaderParameters.h"
 #include "DataDrivenShaderPlatformInfo.h"
+#include "StereoRenderUtils.h"
 
 class FPointerTableBase;
 class FRHICommandList;
@@ -115,6 +116,16 @@ public:
 
 	typedef FDummyResolveParameter FParameter;
 
+	static bool ShouldCompilePermutation(const FGlobalShaderPermutationParameters& Parameters)
+	{
+		if (FResolveDepthPS::ShouldCompilePermutation(Parameters))
+		{
+			UE::StereoRenderUtils::FStereoShaderAspects Aspects(Parameters.Platform);
+			return Aspects.IsMobileMultiViewEnabled();
+		}
+		return false;
+	}
+
 	static void ModifyCompilationEnvironment(const FGlobalShaderPermutationParameters& Parameters, FShaderCompilerEnvironment& OutEnvironment)
 	{
 		FResolveDepthPS::ModifyCompilationEnvironment(Parameters, OutEnvironment);
@@ -135,6 +146,11 @@ public:
 
 	typedef FDummyResolveParameter FParameter;
 
+	static bool ShouldCompilePermutation(const FGlobalShaderPermutationParameters& Parameters)
+	{
+		return FResolveDepthArrayPS::ShouldCompilePermutation(Parameters);
+	}
+
 	static void ModifyCompilationEnvironment(const FGlobalShaderPermutationParameters& Parameters, FShaderCompilerEnvironment& OutEnvironment)
 	{
 		FResolveDepthArrayPS::ModifyCompilationEnvironment(Parameters, OutEnvironment);
@@ -155,6 +171,11 @@ public:
 
 	typedef FDummyResolveParameter FParameter;
 
+	static bool ShouldCompilePermutation(const FGlobalShaderPermutationParameters& Parameters)
+	{
+		return FResolveDepthArrayPS::ShouldCompilePermutation(Parameters);
+	}
+
 	static void ModifyCompilationEnvironment(const FGlobalShaderPermutationParameters& Parameters, FShaderCompilerEnvironment& OutEnvironment)
 	{
 		FResolveDepthArrayPS::ModifyCompilationEnvironment(Parameters, OutEnvironment);
@@ -178,6 +199,11 @@ public:
 
 	static bool ShouldCache(EShaderPlatform Platform) { return GetMaxSupportedFeatureLevel(Platform) >= ERHIFeatureLevel::SM5; }
 
+	static bool ShouldCompilePermutation(const FGlobalShaderPermutationParameters& Parameters)
+	{
+		return FResolveDepthArrayPS::ShouldCompilePermutation(Parameters);
+	}
+
 	static void ModifyCompilationEnvironment(const FGlobalShaderPermutationParameters& Parameters, FShaderCompilerEnvironment& OutEnvironment)
 	{
 		FResolveDepthArrayPS::ModifyCompilationEnvironment(Parameters, OutEnvironment);
@@ -246,7 +272,15 @@ class FResolveArrayVS : public FResolveVS
 	DECLARE_EXPORTED_SHADER_TYPE(FResolveArrayVS, Global, RENDERCORE_API);
 public:
 
-	static bool ShouldCompilePermutation(const FGlobalShaderPermutationParameters& Parameters) { return true; }
+	static bool ShouldCompilePermutation(const FGlobalShaderPermutationParameters& Parameters)
+	{
+		if (FResolveVS::ShouldCompilePermutation(Parameters))
+		{
+			UE::StereoRenderUtils::FStereoShaderAspects Aspects(Parameters.Platform);
+			return Aspects.IsMobileMultiViewEnabled();
+		}
+		return false;
+	}
 
 	static void ModifyCompilationEnvironment(const FGlobalShaderPermutationParameters& Parameters, FShaderCompilerEnvironment& OutEnvironment)
 	{
diff --git a/Engine/Source/Runtime/Renderer/Public/HdrCustomResolveShaders.h b/Engine/Source/Runtime/Renderer/Public/HdrCustomResolveShaders.h
--- a/Engine/Source/Runtime/Renderer/Public/HdrCustomResolveShaders.h
+++ b/Engine/Source/Runtime/Renderer/Public/HdrCustomResolveShaders.h
@@ -15,6 +15,7 @@
 #include "ShaderParameterUtils.h"
 #include "ShaderParameters.h"
 #include "DataDrivenShaderPlatformInfo.h"
+#include "StereoRenderUtils.h"
 
 class FPointerTableBase;
 
@@ -39,6 +40,16 @@ public:
 	{
 	}
 
+	static bool ShouldCompilePermutation(const FGlobalShaderPermutationParameters& Parameters)
+	{
+		if (FHdrCustomResolveVS::ShouldCompilePermutation(Parameters))
+		{
+			UE::StereoRenderUtils::FStereoShaderAspects Aspects(Parameters.Platform);
+			return Aspects.IsMobileMultiViewEnabled();
+		}
+		return false;
+	}
+
 	static void ModifyCompilationEnvironment(const FGlobalShaderPermutationParameters& Parameters, FShaderCompilerEnvironment& OutEnvironment)
 	{
 		FHdrCustomResolveVS::ModifyCompilationEnvironment(Parameters, OutEnvironment);
@@ -146,7 +157,12 @@ public:
 
 	static bool ShouldCompilePermutation(const FGlobalShaderPermutationParameters& Parameters)
 	{
-		return FHdrCustomResolve2xPS::ShouldCompilePermutation(Parameters);
+		if (FHdrCustomResolve2xPS::ShouldCompilePermutation(Parameters))
+		{
+			UE::StereoRenderUtils::FStereoShaderAspects Aspects(Parameters.Platform);
+			return Aspects.IsMobileMultiViewEnabled();
+		}
+		return false;
 	}
 
 	void SetParameters(FRHICommandList& RHICmdList, FRHITexture* Texture2DMS)
