Moving generic rendering math functions out of RenderUtils.h and into a new file, RenderMath.h

diff --git a/Engine/Source/Runtime/Engine/Private/GeometryCacheVertexFactory.cpp b/Engine/Source/Runtime/Engine/Private/GeometryCacheVertexFactory.cpp
--- a/Engine/Source/Runtime/Engine/Private/GeometryCacheVertexFactory.cpp
+++ b/Engine/Source/Runtime/Engine/Private/GeometryCacheVertexFactory.cpp
@@ -9,6 +9,7 @@
 #include "MeshBatch.h"
 #include "ShaderParameterUtils.h"
 #include "MeshMaterialShader.h"
+#include "PackedNormal.h"
 
 /*-----------------------------------------------------------------------------
 FGeometryCacheVertexFactoryShaderParameters
diff --git a/Engine/Source/Runtime/Engine/Public/DynamicMeshBuilder.h b/Engine/Source/Runtime/Engine/Public/DynamicMeshBuilder.h
--- a/Engine/Source/Runtime/Engine/Public/DynamicMeshBuilder.h
+++ b/Engine/Source/Runtime/Engine/Public/DynamicMeshBuilder.h
@@ -11,6 +11,7 @@
 #include "HitProxies.h"
 #include "RenderUtils.h"
 #include "LocalVertexFactory.h"
+#include "RenderMath.h"
 #if UE_ENABLE_INCLUDE_ORDER_DEPRECATED_IN_5_2
 #include "SceneManagement.h"
 #endif
diff --git a/Engine/Source/Runtime/Engine/Public/Rendering/SkeletalMeshVertexBuffer.h b/Engine/Source/Runtime/Engine/Public/Rendering/SkeletalMeshVertexBuffer.h
--- a/Engine/Source/Runtime/Engine/Public/Rendering/SkeletalMeshVertexBuffer.h
+++ b/Engine/Source/Runtime/Engine/Public/Rendering/SkeletalMeshVertexBuffer.h
@@ -7,6 +7,7 @@
 #include "Components.h"
 #include "Containers/DynamicRHIResourceArray.h"
 #include "Rendering/StaticMeshVertexBuffer.h"
+#include "RenderMath.h"
 
 struct FSoftSkinVertex;
 
diff --git a/Engine/Source/Runtime/Engine/Public/Rendering/StaticMeshVertexBuffer.h b/Engine/Source/Runtime/Engine/Public/Rendering/StaticMeshVertexBuffer.h
--- a/Engine/Source/Runtime/Engine/Public/Rendering/StaticMeshVertexBuffer.h
+++ b/Engine/Source/Runtime/Engine/Public/Rendering/StaticMeshVertexBuffer.h
@@ -8,6 +8,7 @@
 #include "PackedNormal.h"
 #include "Components.h"
 #include "Math/Vector2DHalf.h"
+#include "RenderMath.h"
 
 class FVertexFactory;
 class FStaticMeshVertexDataInterface;
diff --git a/Engine/Source/Runtime/Engine/Public/SkeletalRenderPublic.h b/Engine/Source/Runtime/Engine/Public/SkeletalRenderPublic.h
--- a/Engine/Source/Runtime/Engine/Public/SkeletalRenderPublic.h
+++ b/Engine/Source/Runtime/Engine/Public/SkeletalRenderPublic.h
@@ -15,6 +15,7 @@
 #include "Components/SkinnedMeshComponent.h"
 #include "MeshUVChannelInfo.h"
 #include "SkeletalMeshTypes.h"
+#include "RenderMath.h"
 
 class FPrimitiveDrawInterface;
 class FVertexFactory;
diff --git a/Engine/Source/Runtime/RenderCore/Private/RenderUtils.cpp b/Engine/Source/Runtime/RenderCore/Private/RenderUtils.cpp
--- a/Engine/Source/Runtime/RenderCore/Private/RenderUtils.cpp
+++ b/Engine/Source/Runtime/RenderCore/Private/RenderUtils.cpp
@@ -11,6 +11,7 @@
 #include "RenderGraphResourcePool.h"
 #include "Misc/DataDrivenPlatformInfoRegistry.h"
 #include "RenderCore.h"
+#include "PackedNormal.h"
 
 #if WITH_EDITOR
 #include "Misc/CoreMisc.h"
diff --git a/Engine/Source/Runtime/RenderCore/Public/RenderMath.h b/Engine/Source/Runtime/RenderCore/Public/RenderMath.h
new file mode 100644
--- /dev/null
+++ b/Engine/Source/Runtime/RenderCore/Public/RenderMath.h
@@ -0,0 +1,54 @@
+// Copyright Epic Games, Inc. All Rights Reserved.
+
+#pragma once
+
+#include "CoreMinimal.h"
+#include "PackedNormal.h"
+
+/**
+* Constructs a basis matrix for the axis vectors and returns the sign of the determinant
+*
+* @param XAxis - x axis (tangent)
+* @param YAxis - y axis (binormal)
+* @param ZAxis - z axis (normal)
+* @return sign of determinant either -1 or +1
+*/
+FORCEINLINE float GetBasisDeterminantSign(const FVector& XAxis, const FVector& YAxis, const FVector& ZAxis)
+{
+	FMatrix Basis(
+		FPlane(XAxis, 0),
+		FPlane(YAxis, 0),
+		FPlane(ZAxis, 0),
+		FPlane(0, 0, 0, 1)
+	);
+	return (Basis.Determinant() < 0) ? -1.0f : +1.0f;
+}
+
+/**
+* Constructs a basis matrix for the axis vectors and returns the sign of the determinant
+*
+* @param XAxis - x axis (tangent)
+* @param YAxis - y axis (binormal)
+* @param ZAxis - z axis (normal)
+* @return sign of determinant either -127 (-1) or +1 (127)
+*/
+FORCEINLINE int8 GetBasisDeterminantSignByte(const FPackedNormal& XAxis, const FPackedNormal& YAxis, const FPackedNormal& ZAxis)
+{
+	return GetBasisDeterminantSign(XAxis.ToFVector(), YAxis.ToFVector(), ZAxis.ToFVector()) < 0 ? -127 : 127;
+}
+
+/**
+ * Given 2 axes of a basis stored as a packed type, regenerates the y-axis tangent vector and scales by z.W
+ * @param XAxis - x axis (tangent)
+ * @param ZAxis - z axis (normal), the sign of the determinant is stored in ZAxis.W
+ * @return y axis (binormal)
+ */
+template<typename VectorType>
+FORCEINLINE FVector GenerateYAxis(const VectorType& XAxis, const VectorType& ZAxis)
+{
+	static_assert(ARE_TYPES_EQUAL(VectorType, FPackedNormal) ||
+		ARE_TYPES_EQUAL(VectorType, FPackedRGBA16N), "ERROR: Must be FPackedNormal or FPackedRGBA16N");
+	FVector  x = XAxis.ToFVector();
+	FVector4 z = ZAxis.ToFVector4();
+	return (FVector(z) ^ x) * z.W;
+}
diff --git a/Engine/Source/Runtime/RenderCore/Public/RenderUtils.h b/Engine/Source/Runtime/RenderCore/Public/RenderUtils.h
--- a/Engine/Source/Runtime/RenderCore/Public/RenderUtils.h
+++ b/Engine/Source/Runtime/RenderCore/Public/RenderUtils.h
@@ -1,13 +1,15 @@
 // Copyright Epic Games, Inc. All Rights Reserved.
 
-
 #pragma once
 
 #include "CoreMinimal.h"
 #include "RHI.h"
-#include "PackedNormal.h"
 #include "RenderResource.h"
 #include "RHIDefinitions.h"
+#if UE_ENABLE_INCLUDE_ORDER_DEPRECATED_IN_5_2
+#include "PackedNormal.h"
+#include "RenderMath.h"
+#endif
 
 class FBufferWithRDG;
 class FRDGBuffer;
@@ -31,54 +33,6 @@ public:
 
 extern RENDERCORE_API void RenderUtilsInit();
 
-/**
-* Constructs a basis matrix for the axis vectors and returns the sign of the determinant
-*
-* @param XAxis - x axis (tangent)
-* @param YAxis - y axis (binormal)
-* @param ZAxis - z axis (normal)
-* @return sign of determinant either -1 or +1 
-*/
-FORCEINLINE float GetBasisDeterminantSign( const FVector& XAxis, const FVector& YAxis, const FVector& ZAxis )
-{
-	FMatrix Basis(
-		FPlane(XAxis,0),
-		FPlane(YAxis,0),
-		FPlane(ZAxis,0),
-		FPlane(0,0,0,1)
-		);
-	return (Basis.Determinant() < 0) ? -1.0f : +1.0f;
-}
-
-/**
-* Constructs a basis matrix for the axis vectors and returns the sign of the determinant
-*
-* @param XAxis - x axis (tangent)
-* @param YAxis - y axis (binormal)
-* @param ZAxis - z axis (normal)
-* @return sign of determinant either -127 (-1) or +1 (127)
-*/
-FORCEINLINE int8 GetBasisDeterminantSignByte( const FPackedNormal& XAxis, const FPackedNormal& YAxis, const FPackedNormal& ZAxis )
-{
-	return GetBasisDeterminantSign(XAxis.ToFVector(),YAxis.ToFVector(),ZAxis.ToFVector()) < 0 ? -127 : 127;
-}
-
-/**
- * Given 2 axes of a basis stored as a packed type, regenerates the y-axis tangent vector and scales by z.W
- * @param XAxis - x axis (tangent)
- * @param ZAxis - z axis (normal), the sign of the determinant is stored in ZAxis.W
- * @return y axis (binormal)
- */
-template<typename VectorType>
-FORCEINLINE FVector GenerateYAxis(const VectorType& XAxis, const VectorType& ZAxis)
-{
-	static_assert(	ARE_TYPES_EQUAL(VectorType, FPackedNormal) ||
-					ARE_TYPES_EQUAL(VectorType, FPackedRGBA16N), "ERROR: Must be FPackedNormal or FPackedRGBA16N");
-	FVector  x = XAxis.ToFVector();
-	FVector4 z = ZAxis.ToFVector4();
-	return (FVector(z) ^ x) * z.W;
-}
-
 #define NUM_DEBUG_UTIL_COLORS (32)
 static const FColor DebugUtilColor[NUM_DEBUG_UTIL_COLORS] = 
 {
