RDG AddSetupTask now returns the task and has an optional bCondition parameter to control whether the task is actually launched.

diff --git a/Engine/Source/Runtime/RenderCore/Public/RenderGraphBuilder.h b/Engine/Source/Runtime/RenderCore/Public/RenderGraphBuilder.h
--- a/Engine/Source/Runtime/RenderCore/Public/RenderGraphBuilder.h
+++ b/Engine/Source/Runtime/RenderCore/Public/RenderGraphBuilder.h
@@ -225,11 +225,11 @@ public:
 
 	/** Launches a task that is synced prior to graph execution. If parallel execution is not enabled, the lambda is run immediately. */
 	template <typename TaskLambda>
-	void AddSetupTask(TaskLambda&& Task);
+	UE::Tasks::FTask AddSetupTask(TaskLambda&& Task, bool bCondition = true);
 
 	/** Launches a task that is synced prior to graph execution. If parallel execution is not enabled, the lambda is run immediately. */
 	template <typename TaskLambda>
-	void AddCommandListSetupTask(TaskLambda&& Task);
+	UE::Tasks::FTask AddCommandListSetupTask(TaskLambda&& Task, bool bCondition = true);
 
 	/** Tells the builder to delete unused RHI resources. The behavior of this method depends on whether RDG immediate mode is enabled:
 	 *   Deferred:  RHI resource flushes are performed prior to execution.
diff --git a/Engine/Source/Runtime/RenderCore/Public/RenderGraphBuilder.inl b/Engine/Source/Runtime/RenderCore/Public/RenderGraphBuilder.inl
--- a/Engine/Source/Runtime/RenderCore/Public/RenderGraphBuilder.inl
+++ b/Engine/Source/Runtime/RenderCore/Public/RenderGraphBuilder.inl
@@ -405,44 +405,55 @@ inline void FRDGBuilder::AddDispatchHint()
 	}
 }
 
-template <typename TaskLambda>
-void FRDGBuilder::AddSetupTask(TaskLambda&& Task)
+template <typename TaskLambdaType>
+UE::Tasks::FTask FRDGBuilder::AddSetupTask(TaskLambdaType&& TaskLambda, bool bCondition)
 {
-	if (bParallelExecuteEnabled)
+	UE::Tasks::FTask Task;
+
+	if (bParallelExecuteEnabled && bCondition)
 	{
-		ParallelSetupEvents.Emplace(UE::Tasks::Launch(TEXT("FRDGBuilder::AddSetupTask"), [Task = MoveTemp(Task)]
+		Task = UE::Tasks::Launch(TEXT("FRDGBuilder::AddSetupTask"), [TaskLambda = MoveTemp(TaskLambda)]
 		{
 			FTaskTagScope Scope(ETaskTag::EParallelRenderingThread);
-			Task();
-		}));
+			TaskLambda();
+		});
+
+		ParallelSetupEvents.Emplace(Task);
 	}
 	else
 	{
-		Task();
+		TaskLambda();
 	}
+
+	return Task;
 }
 
-template <typename TaskLambda>
-void FRDGBuilder::AddCommandListSetupTask(TaskLambda&& Task)
+template <typename TaskLambdaType>
+UE::Tasks::FTask FRDGBuilder::AddCommandListSetupTask(TaskLambdaType&& TaskLambda, bool bCondition)
 {
-	if (bParallelExecuteEnabled)
+	UE::Tasks::FTask Task;
+
+	if (bParallelExecuteEnabled && bCondition)
 	{
 		FRHICommandList* RHICmdListTask = new FRHICommandList(FRHIGPUMask::All());
 
-		ParallelSetupEvents.Emplace(UE::Tasks::Launch(TEXT("FRDGBuilder::AddCommandListSetupTask"), [Task = MoveTemp(Task), RHICmdListTask]
+		Task = UE::Tasks::Launch(TEXT("FRDGBuilder::AddCommandListSetupTask"), [TaskLambda = MoveTemp(TaskLambda), RHICmdListTask]
 		{
 			FTaskTagScope Scope(ETaskTag::EParallelRenderingThread);
 			RHICmdListTask->SwitchPipeline(ERHIPipeline::Graphics);
-			Task(*RHICmdListTask);
+			TaskLambda(*RHICmdListTask);
 			RHICmdListTask->FinishRecording();
-		}));
+		});
 
+		ParallelSetupEvents.Emplace(Task);
 		RHICmdList.QueueAsyncCommandListSubmit(RHICmdListTask);
 	}
 	else
 	{
-		Task(RHICmdList);
+		TaskLambda(RHICmdList);
 	}
+
+	return Task;
 }
 
 inline const TRefCountPtr<IPooledRenderTarget>& FRDGBuilder::GetPooledTexture(FRDGTextureRef Texture) const
