Remove game thread check in FRenderResource::SetOwnerName as it can be called from other threads on some platforms

diff --git a/Engine/Source/Runtime/RenderCore/Private/RenderResource.cpp b/Engine/Source/Runtime/RenderCore/Private/RenderResource.cpp
--- a/Engine/Source/Runtime/RenderCore/Private/RenderResource.cpp
+++ b/Engine/Source/Runtime/RenderCore/Private/RenderResource.cpp
@@ -361,7 +361,6 @@ FBufferRHIRef FRenderResource::CreateRHIBufferInternal(
 void FRenderResource::SetOwnerName(const FName& InOwnerName)
 {
 #if RHI_ENABLE_RESOURCE_INFO
-	check(IsInGameThread());
 	OwnerName = InOwnerName;
 #endif
 }
