Added debug print functionality to shader permutation templates, gated behind bLogPermutations constexpr in ShaderPermutations.h.

diff --git a/Engine/Source/Runtime/RenderCore/Public/ShaderPermutation.h b/Engine/Source/Runtime/RenderCore/Public/ShaderPermutation.h
--- a/Engine/Source/Runtime/RenderCore/Public/ShaderPermutation.h
+++ b/Engine/Source/Runtime/RenderCore/Public/ShaderPermutation.h
@@ -9,6 +9,9 @@
 #include "CoreMinimal.h"
 #include "ShaderCore.h"
 
+// Enable this to log the parameters of each compiled permutation
+constexpr static bool bLogPermutations = false;
+
 struct FShaderCompilerEnvironment;
 
 /** Defines at compile time a boolean permutation dimension. */
@@ -278,6 +281,12 @@ public:
 	template<typename TPermutationVector, typename TDimension>
 	static void ModifyCompilationEnvironment(const TPermutationVector& PermutationVector, FShaderCompilerEnvironment& OutEnvironment)
 	{
+
+		if constexpr (bLogPermutations)
+		{
+			UE_LOG(LogShaders, Verbose, TEXT("		%s = %d"), TDimension::DefineName, TDimension::ToDefineValue(PermutationVector.DimensionValue));
+		}
+
 		OutEnvironment.SetDefine(TDimension::DefineName, TDimension::ToDefineValue(PermutationVector.DimensionValue));
 		return PermutationVector.Tail.ModifyCompilationEnvironment(OutEnvironment);
 	}
@@ -474,5 +483,4 @@ using FShaderPermutationNone = TShaderPermutationDomain<>;
  * class FMyShaderDim : SHADER_PERMUTATION_ENUM_CLASS("MY_SHADER_DEFINE_NAME", EMyEnum);
  */
 #define SHADER_PERMUTATION_ENUM_CLASS(InDefineName, EnumName) \
-	DECLARE_SHADER_PERMUTATION_IMPL(InDefineName, TShaderPermutationInt, EnumName, static_cast<int32>(EnumName::MAX))
-	
\ No newline at end of file
+	DECLARE_SHADER_PERMUTATION_IMPL(InDefineName, TShaderPermutationInt, EnumName, static_cast<int32>(EnumName::MAX))
\ No newline at end of file
