Moving the bulk of global rendering resources into GlobalRenderResources.h.
Not fully deprecating includes yet.

diff --git a/Engine/Source/Runtime/RenderCore/Private/GlobalRenderResources.cpp b/Engine/Source/Runtime/RenderCore/Private/GlobalRenderResources.cpp
new file mode 100644
--- /dev/null
+++ b/Engine/Source/Runtime/RenderCore/Private/GlobalRenderResources.cpp
@@ -0,0 +1,1038 @@
+// Copyright Epic Games, Inc. All Rights Reserved.
+
+#include "GlobalRenderResources.h"
+//#include "RenderUtils.h"
+#include "RenderGraphUtils.h"
+
+// The maximum number of transient vertex buffer bytes to allocate before we start panic logging who is doing the allocations
+int32 GMaxVertexBytesAllocatedPerFrame = 32 * 1024 * 1024;
+
+FAutoConsoleVariableRef CVarMaxVertexBytesAllocatedPerFrame(
+	TEXT("r.MaxVertexBytesAllocatedPerFrame"),
+	GMaxVertexBytesAllocatedPerFrame,
+	TEXT("The maximum number of transient vertex buffer bytes to allocate before we start panic logging who is doing the allocations"));
+
+int32 GGlobalBufferNumFramesUnusedThresold = 30;
+FAutoConsoleVariableRef CVarReadBufferNumFramesUnusedThresold(
+	TEXT("r.NumFramesUnusedBeforeReleasingGlobalResourceBuffers"),
+	GGlobalBufferNumFramesUnusedThresold,
+	TEXT("Number of frames after which unused global resource allocations will be discarded. Set 0 to ignore. (default=30)"));
+
+//////////////////////////////////////////////////////////////////////////////////////////////////////////////////////
+
+
+// Bulk data interface for providing a single color used to initialize a volume texture.
+struct FColorBulkData : public FResourceBulkDataInterface
+{
+	FColorBulkData(uint8 Alpha) : Color(0, 0, 0, Alpha) { }
+	FColorBulkData(int32 R, int32 G, int32 B, int32 A) : Color(R, G, B, A) {}
+	FColorBulkData(FColor InColor) : Color(InColor) { }
+
+	virtual const void* GetResourceBulkData() const override { return &Color; }
+	virtual uint32 GetResourceBulkDataSize() const override { return sizeof(Color); }
+	virtual void Discard() override { }
+
+	FColor Color;
+};
+
+/**
+ * A solid-colored 1x1 texture.
+ */
+template <int32 R, int32 G, int32 B, int32 A>
+class FColoredTexture : public FTextureWithSRV
+{
+public:
+	// FResource interface.
+	virtual void InitRHI() override
+	{
+		// Create the texture RHI.  		
+		FColorBulkData BulkData(R, G, B, A);
+
+		// BGRA typed UAV is unsupported per D3D spec, use RGBA here.
+		const FRHITextureCreateDesc Desc =
+			FRHITextureCreateDesc::Create2D(TEXT("ColoredTexture"), 1, 1, PF_R8G8B8A8)
+			.SetFlags(ETextureCreateFlags::ShaderResource)
+			.SetBulkData(&BulkData);
+
+		TextureRHI = RHICreateTexture(Desc);
+
+		// Create the sampler state RHI resource.
+		FSamplerStateInitializerRHI SamplerStateInitializer(SF_Point, AM_Wrap, AM_Wrap, AM_Wrap);
+		SamplerStateRHI = GetOrCreateSamplerState(SamplerStateInitializer);
+
+		// Create a view of the texture
+		ShaderResourceViewRHI = RHICreateShaderResourceView(TextureRHI, 0u);
+	}
+
+	virtual uint32 GetSizeX() const override { return 1; }
+	virtual uint32 GetSizeY() const override { return 1; }
+};
+
+class FEmptyVertexBuffer : public FVertexBufferWithSRV
+{
+public:
+	virtual void InitRHI() override
+	{
+		// Create the texture RHI.  		
+		FRHIResourceCreateInfo CreateInfo(TEXT("EmptyVertexBuffer"));
+
+		VertexBufferRHI = RHICreateVertexBuffer(16u, BUF_Static | BUF_ShaderResource | BUF_UnorderedAccess, CreateInfo);
+
+		// Create a view of the buffer
+		ShaderResourceViewRHI = RHICreateShaderResourceView(VertexBufferRHI, 4u, PF_R32_UINT);
+		UnorderedAccessViewRHI = RHICreateUnorderedAccessView(VertexBufferRHI, PF_R32_UINT);
+	}
+};
+
+class FBlackTextureWithSRV : public FColoredTexture<0, 0, 0, 255>
+{
+public:
+	virtual void InitRHI() override
+	{
+		FColoredTexture::InitRHI();
+		FRHITextureReference::DefaultTexture = TextureRHI;
+	}
+
+	virtual void ReleaseRHI() override
+	{
+		FRHITextureReference::DefaultTexture.SafeRelease();
+		FColoredTexture::ReleaseRHI();
+	}
+};
+
+FTextureWithSRV* GWhiteTextureWithSRV = new TGlobalResource<FColoredTexture<255, 255, 255, 255> >;
+FTextureWithSRV* GBlackTextureWithSRV = new TGlobalResource<FBlackTextureWithSRV>();
+FTextureWithSRV* GTransparentBlackTextureWithSRV = new TGlobalResource<FColoredTexture<0, 0, 0, 0> >;
+FTexture* GWhiteTexture = GWhiteTextureWithSRV;
+FTexture* GBlackTexture = GBlackTextureWithSRV;
+FTexture* GTransparentBlackTexture = GTransparentBlackTextureWithSRV;
+
+FVertexBufferWithSRV* GEmptyVertexBufferWithUAV = new TGlobalResource<FEmptyVertexBuffer>;
+
+class FWhiteVertexBuffer : public FVertexBufferWithSRV
+{
+public:
+	virtual void InitRHI() override
+	{
+		// Create the texture RHI.  		
+		FRHIResourceCreateInfo CreateInfo(TEXT("WhiteVertexBuffer"));
+
+		VertexBufferRHI = RHICreateVertexBuffer(sizeof(FVector4f), BUF_Static | BUF_ShaderResource, CreateInfo);
+
+		FVector4f* BufferData = (FVector4f*)RHILockBuffer(VertexBufferRHI, 0, sizeof(FVector4f), RLM_WriteOnly);
+		*BufferData = FVector4f(1.0f, 1.0f, 1.0f, 1.0f);
+		RHIUnlockBuffer(VertexBufferRHI);
+
+		// Create a view of the buffer
+		ShaderResourceViewRHI = RHICreateShaderResourceView(VertexBufferRHI, sizeof(FVector4f), PF_A32B32G32R32F);
+	}
+};
+
+FVertexBufferWithSRV* GWhiteVertexBufferWithSRV = new TGlobalResource<FWhiteVertexBuffer>;
+
+class FWhiteVertexBufferWithRDG : public FBufferWithRDG
+{
+public:
+	virtual void InitRHI() override
+	{
+		if (!Buffer.IsValid())
+		{
+			Buffer = AllocatePooledBuffer(FRDGBufferDesc::CreateBufferDesc(sizeof(FVector4f), 1), TEXT("WhiteVertexBufferWithRDG"));
+
+			FRHICommandListImmediate& RHICmdList = FRHICommandListExecutor::GetImmediateCommandList();
+			FVector4f* BufferData = (FVector4f*)RHICmdList.LockBuffer(Buffer->GetRHI(), 0, sizeof(FVector4f), RLM_WriteOnly);
+			*BufferData = FVector4f(1.0f, 1.0f, 1.0f, 1.0f);
+			RHICmdList.UnlockBuffer(Buffer->GetRHI());
+		}
+	}
+};
+
+FBufferWithRDG* GWhiteVertexBufferWithRDG = new TGlobalResource<FWhiteVertexBufferWithRDG>();
+
+/**
+ * A class representing a 1x1x1 black volume texture.
+ */
+template <EPixelFormat PixelFormat, uint8 Alpha>
+class FBlackVolumeTexture : public FTexture
+{
+public:
+	virtual void InitRHI() override
+	{
+		if (GSupportsTexture3D)
+		{
+			// Create the texture.
+			FColorBulkData BulkData(Alpha);
+
+			const FRHITextureCreateDesc Desc =
+				FRHITextureCreateDesc::Create3D(TEXT("BlackVolumeTexture3D"), 1, 1, 1, PixelFormat)
+				.SetFlags(ETextureCreateFlags::ShaderResource)
+				.SetBulkData(&BulkData);
+
+			TextureRHI = RHICreateTexture(Desc);
+		}
+		else
+		{
+			// Create a texture, even though it's not a volume texture
+			FColorBulkData BulkData(Alpha);
+
+			const FRHITextureCreateDesc Desc =
+				FRHITextureCreateDesc::Create2D(TEXT("BlackVolumeTexture2D"), 1, 1, PixelFormat)
+				.SetFlags(ETextureCreateFlags::ShaderResource)
+				.SetBulkData(&BulkData);
+
+			TextureRHI = RHICreateTexture(Desc);
+		}
+
+		// Create the sampler state.
+		FSamplerStateInitializerRHI SamplerStateInitializer(SF_Point, AM_Wrap, AM_Wrap, AM_Wrap);
+		SamplerStateRHI = GetOrCreateSamplerState(SamplerStateInitializer);
+	}
+
+	virtual uint32 GetSizeX() const override { return 1; }
+	virtual uint32 GetSizeY() const override { return 1; }
+};
+
+/** Global black volume texture resource. */
+FTexture* GBlackVolumeTexture = new TGlobalResource<FBlackVolumeTexture<PF_B8G8R8A8, 0>>();
+FTexture* GBlackAlpha1VolumeTexture = new TGlobalResource<FBlackVolumeTexture<PF_B8G8R8A8, 255>>();
+
+/** Global black volume texture resource. */
+FTexture* GBlackUintVolumeTexture = new TGlobalResource<FBlackVolumeTexture<PF_R8G8B8A8_UINT, 0>>();
+
+class FBlackArrayTexture : public FTexture
+{
+public:
+	virtual void InitRHI() override
+	{
+		// Create the texture RHI.
+		FColorBulkData BulkData(0);
+
+		const FRHITextureCreateDesc Desc =
+			FRHITextureCreateDesc::Create2DArray(TEXT("BlackArrayTexture"), 1, 1, 1, PF_B8G8R8A8)
+			.SetFlags(ETextureCreateFlags::ShaderResource)
+			.SetBulkData(&BulkData);
+
+		TextureRHI = RHICreateTexture(Desc);
+
+		// Create the sampler state RHI resource.
+		FSamplerStateInitializerRHI SamplerStateInitializer(SF_Point, AM_Wrap, AM_Wrap, AM_Wrap);
+		SamplerStateRHI = GetOrCreateSamplerState(SamplerStateInitializer);
+	}
+
+	virtual uint32 GetSizeX() const override { return 1; }
+	virtual uint32 GetSizeY() const override { return 1; }
+};
+
+FTexture* GBlackArrayTexture = new TGlobalResource<FBlackArrayTexture>;
+
+//
+// FMipColorTexture implementation
+//
+
+/**
+ * A texture that has a different solid color in each mip-level
+ */
+class FMipColorTexture : public FTexture
+{
+public:
+	enum
+	{
+		NumMips = 12
+	};
+	static const FColor MipColors[NumMips];
+
+	// FResource interface.
+	virtual void InitRHI() override
+	{
+		// Create the texture RHI.
+		int32 TextureSize = 1 << (NumMips - 1);
+
+		const FRHITextureCreateDesc Desc =
+			FRHITextureCreateDesc::Create2D(TEXT("FMipColorTexture"), TextureSize, TextureSize, PF_B8G8R8A8)
+			.SetNumMips(NumMips)
+			.SetFlags(ETextureCreateFlags::ShaderResource);
+
+		TextureRHI = RHICreateTexture(Desc);
+
+		// Write the contents of the texture.
+		uint32 DestStride;
+		int32 Size = TextureSize;
+		for (int32 MipIndex = 0; MipIndex < NumMips; ++MipIndex)
+		{
+			FColor* DestBuffer = (FColor*)RHILockTexture2D(TextureRHI, MipIndex, RLM_WriteOnly, DestStride, false);
+			for (int32 Y = 0; Y < Size; ++Y)
+			{
+				for (int32 X = 0; X < Size; ++X)
+				{
+					DestBuffer[X] = MipColors[NumMips - 1 - MipIndex];
+				}
+				DestBuffer += DestStride / sizeof(FColor);
+			}
+			RHIUnlockTexture2D(TextureRHI, MipIndex, false);
+			Size >>= 1;
+		}
+
+		// Create the sampler state RHI resource.
+		FSamplerStateInitializerRHI SamplerStateInitializer(SF_Point, AM_Wrap, AM_Wrap, AM_Wrap);
+		SamplerStateRHI = GetOrCreateSamplerState(SamplerStateInitializer);
+	}
+
+	/** Returns the width of the texture in pixels. */
+	virtual uint32 GetSizeX() const override
+	{
+		int32 TextureSize = 1 << (NumMips - 1);
+		return TextureSize;
+	}
+
+	/** Returns the height of the texture in pixels. */
+	// PVS-Studio notices that the implementation of GetSizeX is identical to this one
+	// and warns us. In this case, it is intentional, so we disable the warning:
+	virtual uint32 GetSizeY() const override //-V524
+	{
+		int32 TextureSize = 1 << (NumMips - 1);
+		return TextureSize;
+	}
+};
+
+const FColor FMipColorTexture::MipColors[NumMips] =
+{
+	FColor(80,  80,  80, 0),		// Mip  0: 1x1			(dark grey)
+	FColor(200, 200, 200, 0),		// Mip  1: 2x2			(light grey)
+	FColor(200, 200,   0, 0),		// Mip  2: 4x4			(medium yellow)
+	FColor(255, 255,   0, 0),		// Mip  3: 8x8			(yellow)
+	FColor(160, 255,  40, 0),		// Mip  4: 16x16		(light green)
+	FColor(0, 255,   0, 0),		// Mip  5: 32x32		(green)
+	FColor(0, 255, 200, 0),		// Mip  6: 64x64		(cyan)
+	FColor(0, 170, 170, 0),		// Mip  7: 128x128		(light blue)
+	FColor(60,  60, 255, 0),		// Mip  8: 256x256		(dark blue)
+	FColor(255,   0, 255, 0),		// Mip  9: 512x512		(pink)
+	FColor(255,   0,   0, 0),		// Mip 10: 1024x1024	(red)
+	FColor(255, 130,   0, 0),		// Mip 11: 2048x2048	(orange)
+};
+
+FTexture* GMipColorTexture = new FMipColorTexture;
+int32 GMipColorTextureMipLevels = FMipColorTexture::NumMips;
+
+// 4: 8x8 cubemap resolution, shader needs to use the same value as preprocessing
+const uint32 GDiffuseConvolveMipLevel = 4;
+
+/** A solid color cube texture. */
+class FSolidColorTextureCube : public FTexture
+{
+public:
+	FSolidColorTextureCube(const FColor& InColor)
+		: bInitToZero(false)
+		, PixelFormat(PF_B8G8R8A8)
+		, ColorData(InColor.DWColor())
+	{}
+
+	FSolidColorTextureCube(EPixelFormat InPixelFormat)
+		: bInitToZero(true)
+		, PixelFormat(InPixelFormat)
+		, ColorData(0)
+	{}
+
+	// FRenderResource interface.
+	virtual void InitRHI() override
+	{
+		// Create the texture RHI.
+		const FRHITextureCreateDesc Desc =
+			FRHITextureCreateDesc::CreateCube(TEXT("SolidColorCube"), 1, PixelFormat)
+			.SetFlags(ETextureCreateFlags::ShaderResource);
+
+		FTextureCubeRHIRef TextureCube = RHICreateTexture(Desc);
+		TextureRHI = TextureCube;
+
+		// Write the contents of the texture.
+		for (uint32 FaceIndex = 0; FaceIndex < 6; FaceIndex++)
+		{
+			uint32 DestStride;
+			void* DestBuffer = RHILockTextureCubeFace(TextureCube, FaceIndex, 0, 0, RLM_WriteOnly, DestStride, false);
+			if (bInitToZero)
+			{
+				FMemory::Memzero(DestBuffer, GPixelFormats[PixelFormat].BlockBytes);
+			}
+			else
+			{
+				FMemory::Memcpy(DestBuffer, &ColorData, sizeof(ColorData));
+			}
+			RHIUnlockTextureCubeFace(TextureCube, FaceIndex, 0, 0, false);
+		}
+
+		// Create the sampler state RHI resource.
+		FSamplerStateInitializerRHI SamplerStateInitializer(SF_Point, AM_Wrap, AM_Wrap, AM_Wrap);
+		SamplerStateRHI = GetOrCreateSamplerState(SamplerStateInitializer);
+	}
+
+	virtual uint32 GetSizeX() const override { return 1; }
+	virtual uint32 GetSizeY() const override { return 1; }
+
+private:
+	const bool bInitToZero;
+	const EPixelFormat PixelFormat;
+	const uint32 ColorData;
+};
+
+/** A white cube texture. */
+FTexture* GWhiteTextureCube = new TGlobalResource<FSolidColorTextureCube>(FColor::White);
+
+/** A black cube texture. */
+FTexture* GBlackTextureCube = new TGlobalResource<FSolidColorTextureCube>(FColor::Black);
+
+/** A black cube texture. */
+FTexture* GBlackTextureDepthCube = new TGlobalResource<FSolidColorTextureCube>(PF_ShadowDepth);
+
+class FBlackCubeArrayTexture : public FTexture
+{
+public:
+	// FResource interface.
+	virtual void InitRHI() override
+	{
+		if (SupportsTextureCubeArray(GetFeatureLevel()))
+		{
+			const TCHAR* Name = TEXT("BlackCubeArray");
+
+			const FRHITextureCreateDesc Desc =
+				FRHITextureCreateDesc::CreateCubeArray(TEXT("BlackCubeArray"), 1, 1, PF_B8G8R8A8)
+				.SetFlags(ETextureCreateFlags::ShaderResource);
+
+			// Create the texture RHI.
+			TextureRHI = RHICreateTexture(Desc);
+
+			for (uint32 FaceIndex = 0; FaceIndex < 6; FaceIndex++)
+			{
+				uint32 DestStride;
+				FColor* DestBuffer = (FColor*)RHILockTextureCubeFace(TextureRHI, FaceIndex, 0, 0, RLM_WriteOnly, DestStride, false);
+				// Note: alpha is used by reflection environment to say how much of the foreground texture is visible, so 0 says it is completely invisible
+				*DestBuffer = FColor(0, 0, 0, 0);
+				RHIUnlockTextureCubeFace(TextureRHI, FaceIndex, 0, 0, false);
+			}
+
+			// Create the sampler state RHI resource.
+			FSamplerStateInitializerRHI SamplerStateInitializer(SF_Point, AM_Wrap, AM_Wrap, AM_Wrap);
+			SamplerStateRHI = GetOrCreateSamplerState(SamplerStateInitializer);
+		}
+	}
+
+	virtual uint32 GetSizeX() const override { return 1; }
+	virtual uint32 GetSizeY() const override { return 1; }
+};
+FTexture* GBlackCubeArrayTexture = new TGlobalResource<FBlackCubeArrayTexture>;
+
+/**
+ * A UINT 1x1 texture.
+ */
+template <EPixelFormat Format, uint32 R = 0, uint32 G = 0, uint32 B = 0, uint32 A = 0>
+class FUintTexture : public FTextureWithSRV
+{
+public:
+	// FResource interface.
+	virtual void InitRHI() override
+	{
+		const FRHITextureCreateDesc Desc =
+			FRHITextureCreateDesc::Create2D(TEXT("UintTexture"), 1, 1, Format)
+			.SetFlags(ETextureCreateFlags::ShaderResource);
+
+		TextureRHI = RHICreateTexture(Desc);
+
+		// Write the contents of the texture.
+		uint32 DestStride;
+		void* DestBuffer = RHILockTexture2D(TextureRHI, 0, RLM_WriteOnly, DestStride, false);
+		WriteData(DestBuffer);
+		RHIUnlockTexture2D(TextureRHI, 0, false);
+
+		// Create the sampler state RHI resource.
+		FSamplerStateInitializerRHI SamplerStateInitializer(SF_Point, AM_Wrap, AM_Wrap, AM_Wrap);
+		SamplerStateRHI = GetOrCreateSamplerState(SamplerStateInitializer);
+
+		// Create a view of the texture
+		ShaderResourceViewRHI = RHICreateShaderResourceView(TextureRHI, 0u);
+	}
+
+	virtual uint32 GetSizeX() const override { return 1; }
+	virtual uint32 GetSizeY() const override { return 1; }
+
+protected:
+	static int32 GetNumChannels()
+	{
+		return GPixelFormats[Format].NumComponents;
+	}
+
+	static int32 GetBytesPerChannel()
+	{
+		return GPixelFormats[Format].BlockBytes / GPixelFormats[Format].NumComponents;
+	}
+
+	template<typename T>
+	static void DoWriteData(T* DataPtr)
+	{
+		T Values[] = { R, G, B, A };
+		for (int32 i = 0; i < GetNumChannels(); ++i)
+		{
+			DataPtr[i] = Values[i];
+		}
+	}
+
+	static void WriteData(void* DataPtr)
+	{
+		switch (GetBytesPerChannel())
+		{
+		case 1:
+			DoWriteData((uint8*)DataPtr);
+			return;
+		case 2:
+			DoWriteData((uint16*)DataPtr);
+			return;
+		case 4:
+			DoWriteData((uint32*)DataPtr);
+			return;
+		}
+		// Unsupported format
+		check(0);
+	}
+};
+
+FTexture* GBlackUintTexture = new TGlobalResource< FUintTexture<PF_R32G32B32A32_UINT> >;
+
+
+//////////////////////////////////////////////////////////////////////////////////////////////////////////////////////
+// FNullColorVertexBuffer
+
+FNullColorVertexBuffer::FNullColorVertexBuffer() = default;
+FNullColorVertexBuffer::~FNullColorVertexBuffer() = default;
+
+void FNullColorVertexBuffer::InitRHI()
+{
+	// create a static vertex buffer
+	FRHIResourceCreateInfo CreateInfo(TEXT("FNullColorVertexBuffer"));
+
+	VertexBufferRHI = RHICreateBuffer(sizeof(uint32) * 4, BUF_Static | BUF_VertexBuffer | BUF_ShaderResource, 0, ERHIAccess::VertexOrIndexBuffer | ERHIAccess::SRVMask, CreateInfo);
+	uint32* Vertices = (uint32*)RHILockBuffer(VertexBufferRHI, 0, sizeof(uint32) * 4, RLM_WriteOnly);
+	Vertices[0] = FColor(255, 255, 255, 255).DWColor();
+	Vertices[1] = FColor(255, 255, 255, 255).DWColor();
+	Vertices[2] = FColor(255, 255, 255, 255).DWColor();
+	Vertices[3] = FColor(255, 255, 255, 255).DWColor();
+	RHIUnlockBuffer(VertexBufferRHI);
+	VertexBufferSRV = RHICreateShaderResourceView(VertexBufferRHI, sizeof(FColor), PF_R8G8B8A8);
+}
+
+void FNullColorVertexBuffer::ReleaseRHI()
+{
+	VertexBufferSRV.SafeRelease();
+	FVertexBuffer::ReleaseRHI();
+}
+
+/** The global null color vertex buffer, which is set with a stride of 0 on meshes without a color component. */
+TGlobalResource<FNullColorVertexBuffer> GNullColorVertexBuffer;
+
+//////////////////////////////////////////////////////////////////////////////////////////////////////////////////////
+// FNullVertexBuffer
+
+FNullVertexBuffer::FNullVertexBuffer() = default;
+FNullVertexBuffer::~FNullVertexBuffer() = default;
+
+void FNullVertexBuffer::InitRHI()
+{
+	// create a static vertex buffer
+	FRHIResourceCreateInfo CreateInfo(TEXT("FNullVertexBuffer"));
+	VertexBufferRHI = RHICreateBuffer(sizeof(FVector3f), BUF_Static | BUF_VertexBuffer | BUF_ShaderResource, 0, ERHIAccess::VertexOrIndexBuffer | ERHIAccess::SRVMask, CreateInfo);
+	FVector3f* LockedData = (FVector3f*)RHILockBuffer(VertexBufferRHI, 0, sizeof(FVector3f), RLM_WriteOnly);
+	*LockedData = FVector3f(0.0f);
+	RHIUnlockBuffer(VertexBufferRHI);
+
+	VertexBufferSRV = RHICreateShaderResourceView(VertexBufferRHI, sizeof(FColor), PF_R8G8B8A8);
+}
+
+void FNullVertexBuffer::ReleaseRHI()
+{
+	VertexBufferSRV.SafeRelease();
+	FVertexBuffer::ReleaseRHI();
+}
+
+/** The global null vertex buffer, which is set with a stride of 0 on meshes */
+TGlobalResource<FNullVertexBuffer> GNullVertexBuffer;
+
+//////////////////////////////////////////////////////////////////////////////////////////////////////////////////////
+// FScreenSpaceVertexBuffer
+
+void FScreenSpaceVertexBuffer::InitRHI()
+{
+	// create a static vertex buffer
+	FRHIResourceCreateInfo CreateInfo(TEXT("FScreenSpaceVertexBuffer"));
+	VertexBufferRHI = RHICreateVertexBuffer(sizeof(FVector2f) * 4, BUF_Static, CreateInfo);
+	void* VoidPtr = RHILockBuffer(VertexBufferRHI, 0, sizeof(FVector2f) * 4, RLM_WriteOnly);
+	static const FVector2f Vertices[4] =
+	{
+		FVector2f(-1,-1),
+		FVector2f(-1,+1),
+		FVector2f(+1,-1),
+		FVector2f(+1,+1),
+	};
+	FMemory::Memcpy(VoidPtr, Vertices, sizeof(FVector2f) * 4);
+	RHIUnlockBuffer(VertexBufferRHI);
+}
+
+TGlobalResource<FScreenSpaceVertexBuffer> GScreenSpaceVertexBuffer;
+
+//////////////////////////////////////////////////////////////////////////////////////////////////////////////////////
+// FTileVertexDeclaration
+
+FTileVertexDeclaration::FTileVertexDeclaration() = default;
+FTileVertexDeclaration::~FTileVertexDeclaration() = default;
+
+void FTileVertexDeclaration::InitRHI()
+{
+	FVertexDeclarationElementList Elements;
+	uint16 Stride = sizeof(FVector2f);
+	Elements.Add(FVertexElement(0, 0, VET_Float2, 0, Stride, false));
+	VertexDeclarationRHI = RHICreateVertexDeclaration(Elements);
+}
+
+void FTileVertexDeclaration::ReleaseRHI()
+{
+	VertexDeclarationRHI.SafeRelease();
+}
+
+TGlobalResource<FTileVertexDeclaration> GTileVertexDeclaration;
+
+//////////////////////////////////////////////////////////////////////////////////////////////////////////////////////
+// FCubeIndexBuffer
+
+void FCubeIndexBuffer::InitRHI()
+{
+	// create a static vertex buffer
+	FRHIResourceCreateInfo CreateInfo(TEXT("FCubeIndexBuffer"));
+	IndexBufferRHI = RHICreateIndexBuffer(sizeof(uint16), sizeof(uint16) * NUM_CUBE_VERTICES, BUF_Static, CreateInfo);
+	void* VoidPtr = RHILockBuffer(IndexBufferRHI, 0, sizeof(uint16) * NUM_CUBE_VERTICES, RLM_WriteOnly);
+	FMemory::Memcpy(VoidPtr, GCubeIndices, NUM_CUBE_VERTICES * sizeof(uint16));
+	RHIUnlockBuffer(IndexBufferRHI);
+}
+
+TGlobalResource<FCubeIndexBuffer> GCubeIndexBuffer;
+
+//////////////////////////////////////////////////////////////////////////////////////////////////////////////////////
+// FTwoTrianglesIndexBuffer
+
+void FTwoTrianglesIndexBuffer::InitRHI()
+{
+	// create a static vertex buffer
+	FRHIResourceCreateInfo CreateInfo(TEXT("FTwoTrianglesIndexBuffer"));
+	IndexBufferRHI = RHICreateIndexBuffer(sizeof(uint16), sizeof(uint16) * 6, BUF_Static, CreateInfo);
+	void* VoidPtr = RHILockBuffer(IndexBufferRHI, 0, sizeof(uint16) * 6, RLM_WriteOnly);
+	static const uint16 Indices[] = { 0, 1, 3, 0, 3, 2 };
+	FMemory::Memcpy(VoidPtr, Indices, 6 * sizeof(uint16));
+	RHIUnlockBuffer(IndexBufferRHI);
+}
+
+TGlobalResource<FTwoTrianglesIndexBuffer> GTwoTrianglesIndexBuffer;
+
+/*------------------------------------------------------------------------------
+	FGlobalDynamicVertexBuffer implementation.
+------------------------------------------------------------------------------*/
+
+/**
+ * An individual dynamic vertex buffer.
+ */
+class FDynamicVertexBuffer : public FVertexBuffer
+{
+public:
+	/** The aligned size of all dynamic vertex buffers. */
+	enum { ALIGNMENT = (1 << 16) }; // 64KB
+	/** Pointer to the vertex buffer mapped in main memory. */
+	uint8* MappedBuffer;
+	/** Size of the vertex buffer in bytes. */
+	uint32 BufferSize;
+	/** Number of bytes currently allocated from the buffer. */
+	uint32 AllocatedByteCount;
+	/** Number of successive frames for which AllocatedByteCount == 0. Used as a metric to decide when to free the allocation. */
+	int32 NumFramesUnused = 0;
+
+	/** Default constructor. */
+	explicit FDynamicVertexBuffer(uint32 InMinBufferSize)
+		: MappedBuffer(NULL)
+		, BufferSize(FMath::Max<uint32>(Align(InMinBufferSize, ALIGNMENT), ALIGNMENT))
+		, AllocatedByteCount(0)
+	{
+	}
+
+	/**
+	 * Locks the vertex buffer so it may be written to.
+	 */
+	void Lock()
+	{
+		check(MappedBuffer == NULL);
+		check(AllocatedByteCount == 0);
+		check(IsValidRef(VertexBufferRHI));
+		MappedBuffer = (uint8*)RHILockBuffer(VertexBufferRHI, 0, BufferSize, RLM_WriteOnly);
+	}
+
+	/**
+	 * Unocks the buffer so the GPU may read from it.
+	 */
+	void Unlock()
+	{
+		check(MappedBuffer != NULL);
+		check(IsValidRef(VertexBufferRHI));
+		RHIUnlockBuffer(VertexBufferRHI);
+		MappedBuffer = NULL;
+		AllocatedByteCount = 0;
+		NumFramesUnused = 0;
+	}
+
+	// FRenderResource interface.
+	virtual void InitRHI() override
+	{
+		check(!IsValidRef(VertexBufferRHI));
+		FRHIResourceCreateInfo CreateInfo(TEXT("FDynamicVertexBuffer"));
+		VertexBufferRHI = RHICreateVertexBuffer(BufferSize, BUF_Volatile, CreateInfo);
+		MappedBuffer = NULL;
+		AllocatedByteCount = 0;
+	}
+
+	virtual void ReleaseRHI() override
+	{
+		FVertexBuffer::ReleaseRHI();
+		MappedBuffer = NULL;
+		AllocatedByteCount = 0;
+	}
+
+	virtual FString GetFriendlyName() const override
+	{
+		return TEXT("FDynamicVertexBuffer");
+	}
+};
+
+/**
+ * A pool of dynamic vertex buffers.
+ */
+struct FDynamicVertexBufferPool
+{
+	/** List of vertex buffers. */
+	TIndirectArray<FDynamicVertexBuffer> VertexBuffers;
+	/** The current buffer from which allocations are being made. */
+	FDynamicVertexBuffer* CurrentVertexBuffer;
+
+	/** Default constructor. */
+	FDynamicVertexBufferPool()
+		: CurrentVertexBuffer(NULL)
+	{
+	}
+
+	/** Destructor. */
+	~FDynamicVertexBufferPool()
+	{
+		int32 NumVertexBuffers = VertexBuffers.Num();
+		for (int32 BufferIndex = 0; BufferIndex < NumVertexBuffers; ++BufferIndex)
+		{
+			VertexBuffers[BufferIndex].ReleaseResource();
+		}
+	}
+};
+
+FGlobalDynamicVertexBuffer::FGlobalDynamicVertexBuffer()
+	: TotalAllocatedSinceLastCommit(0)
+{
+	Pool = new FDynamicVertexBufferPool();
+}
+
+FGlobalDynamicVertexBuffer::~FGlobalDynamicVertexBuffer()
+{
+	delete Pool;
+	Pool = NULL;
+}
+
+FGlobalDynamicVertexBuffer::FAllocation FGlobalDynamicVertexBuffer::Allocate(uint32 SizeInBytes)
+{
+	FAllocation Allocation;
+
+	TotalAllocatedSinceLastCommit += SizeInBytes;
+	if (IsRenderAlarmLoggingEnabled())
+	{
+		UE_LOG(LogRendererCore, Warning, TEXT("FGlobalDynamicVertexBuffer::Allocate(%u), will have allocated %u total this frame"), SizeInBytes, TotalAllocatedSinceLastCommit);
+	}
+
+	FDynamicVertexBuffer* VertexBuffer = Pool->CurrentVertexBuffer;
+	if (VertexBuffer == NULL || VertexBuffer->AllocatedByteCount + SizeInBytes > VertexBuffer->BufferSize)
+	{
+		// Find a buffer in the pool big enough to service the request.
+		VertexBuffer = NULL;
+		for (int32 BufferIndex = 0, NumBuffers = Pool->VertexBuffers.Num(); BufferIndex < NumBuffers; ++BufferIndex)
+		{
+			FDynamicVertexBuffer& VertexBufferToCheck = Pool->VertexBuffers[BufferIndex];
+			if (VertexBufferToCheck.AllocatedByteCount + SizeInBytes <= VertexBufferToCheck.BufferSize)
+			{
+				VertexBuffer = &VertexBufferToCheck;
+				break;
+			}
+		}
+
+		// Create a new vertex buffer if needed.
+		if (VertexBuffer == NULL)
+		{
+			VertexBuffer = new FDynamicVertexBuffer(SizeInBytes);
+			Pool->VertexBuffers.Add(VertexBuffer);
+			VertexBuffer->InitResource();
+		}
+
+		// Lock the buffer if needed.
+		if (VertexBuffer->MappedBuffer == NULL)
+		{
+			VertexBuffer->Lock();
+		}
+
+		// Remember this buffer, we'll try to allocate out of it in the future.
+		Pool->CurrentVertexBuffer = VertexBuffer;
+	}
+
+	check(VertexBuffer != NULL);
+	checkf(VertexBuffer->AllocatedByteCount + SizeInBytes <= VertexBuffer->BufferSize, TEXT("Global vertex buffer allocation failed: BufferSize=%d AllocatedByteCount=%d SizeInBytes=%d"), VertexBuffer->BufferSize, VertexBuffer->AllocatedByteCount, SizeInBytes);
+	Allocation.Buffer = VertexBuffer->MappedBuffer + VertexBuffer->AllocatedByteCount;
+	Allocation.VertexBuffer = VertexBuffer;
+	Allocation.VertexOffset = VertexBuffer->AllocatedByteCount;
+	VertexBuffer->AllocatedByteCount += SizeInBytes;
+
+	return Allocation;
+}
+
+bool FGlobalDynamicVertexBuffer::IsRenderAlarmLoggingEnabled() const
+{
+	return GMaxVertexBytesAllocatedPerFrame > 0 && TotalAllocatedSinceLastCommit >= (size_t)GMaxVertexBytesAllocatedPerFrame;
+}
+
+void FGlobalDynamicVertexBuffer::Commit()
+{
+	for (int32 BufferIndex = 0, NumBuffers = Pool->VertexBuffers.Num(); BufferIndex < NumBuffers; ++BufferIndex)
+	{
+		FDynamicVertexBuffer& VertexBuffer = Pool->VertexBuffers[BufferIndex];
+		if (VertexBuffer.MappedBuffer != NULL)
+		{
+			VertexBuffer.Unlock();
+		}
+		else if (GGlobalBufferNumFramesUnusedThresold && !VertexBuffer.AllocatedByteCount)
+		{
+			++VertexBuffer.NumFramesUnused;
+			if (VertexBuffer.NumFramesUnused >= GGlobalBufferNumFramesUnusedThresold)
+			{
+				// Remove the buffer, assumes they are unordered.
+				VertexBuffer.ReleaseResource();
+				Pool->VertexBuffers.RemoveAtSwap(BufferIndex);
+				--BufferIndex;
+				--NumBuffers;
+			}
+		}
+	}
+	Pool->CurrentVertexBuffer = NULL;
+	TotalAllocatedSinceLastCommit = 0;
+}
+
+FGlobalDynamicVertexBuffer InitViewDynamicVertexBuffer;
+FGlobalDynamicVertexBuffer InitShadowViewDynamicVertexBuffer;
+
+/*------------------------------------------------------------------------------
+	FGlobalDynamicIndexBuffer implementation.
+------------------------------------------------------------------------------*/
+
+/**
+ * An individual dynamic index buffer.
+ */
+class FDynamicIndexBuffer : public FIndexBuffer
+{
+public:
+	/** The aligned size of all dynamic index buffers. */
+	enum { ALIGNMENT = (1 << 16) }; // 64KB
+	/** Pointer to the index buffer mapped in main memory. */
+	uint8* MappedBuffer;
+	/** Size of the index buffer in bytes. */
+	uint32 BufferSize;
+	/** Number of bytes currently allocated from the buffer. */
+	uint32 AllocatedByteCount;
+	/** Stride of the buffer in bytes. */
+	uint32 Stride;
+	/** Number of successive frames for which AllocatedByteCount == 0. Used as a metric to decide when to free the allocation. */
+	int32 NumFramesUnused = 0;
+
+	/** Initialization constructor. */
+	explicit FDynamicIndexBuffer(uint32 InMinBufferSize, uint32 InStride)
+		: MappedBuffer(NULL)
+		, BufferSize(FMath::Max<uint32>(Align(InMinBufferSize, ALIGNMENT), ALIGNMENT))
+		, AllocatedByteCount(0)
+		, Stride(InStride)
+	{
+	}
+
+	/**
+	 * Locks the vertex buffer so it may be written to.
+	 */
+	void Lock()
+	{
+		check(MappedBuffer == NULL);
+		check(AllocatedByteCount == 0);
+		check(IsValidRef(IndexBufferRHI));
+		MappedBuffer = (uint8*)RHILockBuffer(IndexBufferRHI, 0, BufferSize, RLM_WriteOnly);
+	}
+
+	/**
+	 * Unocks the buffer so the GPU may read from it.
+	 */
+	void Unlock()
+	{
+		check(MappedBuffer != NULL);
+		check(IsValidRef(IndexBufferRHI));
+		RHIUnlockBuffer(IndexBufferRHI);
+		MappedBuffer = NULL;
+		AllocatedByteCount = 0;
+		NumFramesUnused = 0;
+	}
+
+	// FRenderResource interface.
+	virtual void InitRHI() override
+	{
+		check(!IsValidRef(IndexBufferRHI));
+		FRHIResourceCreateInfo CreateInfo(TEXT("FDynamicIndexBuffer"));
+		IndexBufferRHI = RHICreateIndexBuffer(Stride, BufferSize, BUF_Volatile, CreateInfo);
+		MappedBuffer = NULL;
+		AllocatedByteCount = 0;
+	}
+
+	virtual void ReleaseRHI() override
+	{
+		FIndexBuffer::ReleaseRHI();
+		MappedBuffer = NULL;
+		AllocatedByteCount = 0;
+	}
+
+	virtual FString GetFriendlyName() const override
+	{
+		return TEXT("FDynamicIndexBuffer");
+	}
+};
+
+/**
+ * A pool of dynamic index buffers.
+ */
+struct FDynamicIndexBufferPool
+{
+	/** List of index buffers. */
+	TIndirectArray<FDynamicIndexBuffer> IndexBuffers;
+	/** The current buffer from which allocations are being made. */
+	FDynamicIndexBuffer* CurrentIndexBuffer;
+	/** Stride of buffers in this pool. */
+	uint32 BufferStride;
+
+	/** Initialization constructor. */
+	explicit FDynamicIndexBufferPool(uint32 InBufferStride)
+		: CurrentIndexBuffer(NULL)
+		, BufferStride(InBufferStride)
+	{
+	}
+
+	/** Destructor. */
+	~FDynamicIndexBufferPool()
+	{
+		int32 NumIndexBuffers = IndexBuffers.Num();
+		for (int32 BufferIndex = 0; BufferIndex < NumIndexBuffers; ++BufferIndex)
+		{
+			IndexBuffers[BufferIndex].ReleaseResource();
+		}
+	}
+};
+
+FGlobalDynamicIndexBuffer::FGlobalDynamicIndexBuffer()
+{
+	Pools[0] = new FDynamicIndexBufferPool(sizeof(uint16));
+	Pools[1] = new FDynamicIndexBufferPool(sizeof(uint32));
+}
+
+FGlobalDynamicIndexBuffer::~FGlobalDynamicIndexBuffer()
+{
+	for (int32 i = 0; i < 2; ++i)
+	{
+		delete Pools[i];
+		Pools[i] = NULL;
+	}
+}
+
+FGlobalDynamicIndexBuffer::FAllocation FGlobalDynamicIndexBuffer::Allocate(uint32 NumIndices, uint32 IndexStride)
+{
+	FAllocation Allocation;
+
+	if (IndexStride != 2 && IndexStride != 4)
+	{
+		return Allocation;
+	}
+
+	FDynamicIndexBufferPool* Pool = Pools[IndexStride >> 2]; // 2 -> 0, 4 -> 1
+
+	uint32 SizeInBytes = NumIndices * IndexStride;
+	FDynamicIndexBuffer* IndexBuffer = Pool->CurrentIndexBuffer;
+	if (IndexBuffer == NULL || IndexBuffer->AllocatedByteCount + SizeInBytes > IndexBuffer->BufferSize)
+	{
+		// Find a buffer in the pool big enough to service the request.
+		IndexBuffer = NULL;
+		for (int32 BufferIndex = 0, NumBuffers = Pool->IndexBuffers.Num(); BufferIndex < NumBuffers; ++BufferIndex)
+		{
+			FDynamicIndexBuffer& IndexBufferToCheck = Pool->IndexBuffers[BufferIndex];
+			if (IndexBufferToCheck.AllocatedByteCount + SizeInBytes <= IndexBufferToCheck.BufferSize)
+			{
+				IndexBuffer = &IndexBufferToCheck;
+				break;
+			}
+		}
+
+		// Create a new index buffer if needed.
+		if (IndexBuffer == NULL)
+		{
+			IndexBuffer = new FDynamicIndexBuffer(SizeInBytes, Pool->BufferStride);
+			Pool->IndexBuffers.Add(IndexBuffer);
+			IndexBuffer->InitResource();
+		}
+
+		// Lock the buffer if needed.
+		if (IndexBuffer->MappedBuffer == NULL)
+		{
+			IndexBuffer->Lock();
+		}
+		Pool->CurrentIndexBuffer = IndexBuffer;
+	}
+
+	check(IndexBuffer != NULL);
+	checkf(IndexBuffer->AllocatedByteCount + SizeInBytes <= IndexBuffer->BufferSize, TEXT("Global index buffer allocation failed: BufferSize=%d AllocatedByteCount=%d SizeInBytes=%d"), IndexBuffer->BufferSize, IndexBuffer->AllocatedByteCount, SizeInBytes);
+
+	Allocation.Buffer = IndexBuffer->MappedBuffer + IndexBuffer->AllocatedByteCount;
+	Allocation.IndexBuffer = IndexBuffer;
+	Allocation.FirstIndex = IndexBuffer->AllocatedByteCount / IndexStride;
+	IndexBuffer->AllocatedByteCount += SizeInBytes;
+
+	return Allocation;
+}
+
+void FGlobalDynamicIndexBuffer::Commit()
+{
+	for (int32 i = 0; i < 2; ++i)
+	{
+		FDynamicIndexBufferPool* Pool = Pools[i];
+
+		for (int32 BufferIndex = 0, NumBuffers = Pool->IndexBuffers.Num(); BufferIndex < NumBuffers; ++BufferIndex)
+		{
+			FDynamicIndexBuffer& IndexBuffer = Pool->IndexBuffers[BufferIndex];
+			if (IndexBuffer.MappedBuffer != NULL)
+			{
+				IndexBuffer.Unlock();
+			}
+			else if (GGlobalBufferNumFramesUnusedThresold && !IndexBuffer.AllocatedByteCount)
+			{
+				++IndexBuffer.NumFramesUnused;
+				if (IndexBuffer.NumFramesUnused >= GGlobalBufferNumFramesUnusedThresold)
+				{
+					// Remove the buffer, assumes they are unordered.
+					IndexBuffer.ReleaseResource();
+					Pool->IndexBuffers.RemoveAtSwap(BufferIndex);
+					--BufferIndex;
+					--NumBuffers;
+				}
+			}
+		}
+		Pool->CurrentIndexBuffer = NULL;
+	}
+}
+
diff --git a/Engine/Source/Runtime/RenderCore/Private/RenderResource.cpp b/Engine/Source/Runtime/RenderCore/Private/RenderResource.cpp
--- a/Engine/Source/Runtime/RenderCore/Private/RenderResource.cpp
+++ b/Engine/Source/Runtime/RenderCore/Private/RenderResource.cpp
@@ -11,24 +11,11 @@
 #include "ProfilingDebugging/LoadTimeTracker.h"
 #include "CoreGlobals.h"
 #include "RayTracingGeometryManager.h"
+#include "RenderGraphResources.h"
 
 /** Whether to enable mip-level fading or not: +1.0f if enabled, -1.0f if disabled. */
 float GEnableMipLevelFading = 1.0f;
 
-// The maximum number of transient vertex buffer bytes to allocate before we start panic logging who is doing the allocations
-int32 GMaxVertexBytesAllocatedPerFrame = 32 * 1024 * 1024;
-
-FAutoConsoleVariableRef CVarMaxVertexBytesAllocatedPerFrame(
-	TEXT("r.MaxVertexBytesAllocatedPerFrame"),
-	GMaxVertexBytesAllocatedPerFrame,
-	TEXT("The maximum number of transient vertex buffer bytes to allocate before we start panic logging who is doing the allocations"));
-
-int32 GGlobalBufferNumFramesUnusedThresold = 30;
-FAutoConsoleVariableRef CVarReadBufferNumFramesUnusedThresold(
-	TEXT("r.NumFramesUnusedBeforeReleasingGlobalResourceBuffers"),
-	GGlobalBufferNumFramesUnusedThresold,
-	TEXT("Number of frames after which unused global resource allocations will be discarded. Set 0 to ignore. (default=30)"));
-
 bool GFreeStructuresOnRHIBufferCreation = true;
 FAutoConsoleVariableRef CVarFreeStructuresOnRHIBufferCreation(
 	TEXT("r.FreeStructuresOnRHIBufferCreation"),
@@ -655,63 +642,6 @@ void FVertexBufferWithSRV::ReleaseRHI()
 	FVertexBuffer::ReleaseRHI();
 }
 
-//////////////////////////////////////////////////////////////////////////////////////////////////////////////////////
-// FNullColorVertexBuffer
-
-FNullColorVertexBuffer::FNullColorVertexBuffer() = default;
-FNullColorVertexBuffer::~FNullColorVertexBuffer() = default;
-
-void FNullColorVertexBuffer::InitRHI()
-{
-	// create a static vertex buffer
-	FRHIResourceCreateInfo CreateInfo(TEXT("FNullColorVertexBuffer"));
-
-	VertexBufferRHI = RHICreateBuffer(sizeof(uint32) * 4, BUF_Static | BUF_VertexBuffer | BUF_ShaderResource, 0, ERHIAccess::VertexOrIndexBuffer | ERHIAccess::SRVMask, CreateInfo);
-	uint32* Vertices = (uint32*)RHILockBuffer(VertexBufferRHI, 0, sizeof(uint32) * 4, RLM_WriteOnly);
-	Vertices[0] = FColor(255, 255, 255, 255).DWColor();
-	Vertices[1] = FColor(255, 255, 255, 255).DWColor();
-	Vertices[2] = FColor(255, 255, 255, 255).DWColor();
-	Vertices[3] = FColor(255, 255, 255, 255).DWColor();
-	RHIUnlockBuffer(VertexBufferRHI);
-	VertexBufferSRV = RHICreateShaderResourceView(VertexBufferRHI, sizeof(FColor), PF_R8G8B8A8);
-}
-
-void FNullColorVertexBuffer::ReleaseRHI()
-{
-	VertexBufferSRV.SafeRelease();
-	FVertexBuffer::ReleaseRHI();
-}
-
-/** The global null color vertex buffer, which is set with a stride of 0 on meshes without a color component. */
-TGlobalResource<FNullColorVertexBuffer> GNullColorVertexBuffer;
-
-//////////////////////////////////////////////////////////////////////////////////////////////////////////////////////
-// FNullVertexBuffer
-
-FNullVertexBuffer::FNullVertexBuffer() = default;
-FNullVertexBuffer::~FNullVertexBuffer() = default;
-
-void FNullVertexBuffer::InitRHI()
-{
-	// create a static vertex buffer
-	FRHIResourceCreateInfo CreateInfo(TEXT("FNullVertexBuffer"));
-	VertexBufferRHI = RHICreateBuffer(sizeof(FVector3f), BUF_Static | BUF_VertexBuffer | BUF_ShaderResource, 0, ERHIAccess::VertexOrIndexBuffer | ERHIAccess::SRVMask, CreateInfo);
-	FVector3f* LockedData = (FVector3f*)RHILockBuffer(VertexBufferRHI, 0, sizeof(FVector3f), RLM_WriteOnly);
-	*LockedData = FVector3f(0.0f);
-	RHIUnlockBuffer(VertexBufferRHI);
-
-	VertexBufferSRV = RHICreateShaderResourceView(VertexBufferRHI, sizeof(FColor), PF_R8G8B8A8);
-}
-
-void FNullVertexBuffer::ReleaseRHI()
-{
-	VertexBufferSRV.SafeRelease();
-	FVertexBuffer::ReleaseRHI();
-}
-
-/** The global null vertex buffer, which is set with a stride of 0 on meshes */
-TGlobalResource<FNullVertexBuffer> GNullVertexBuffer;
-
 //////////////////////////////////////////////////////////////////////////////////////////////////////////////////////
 // FIndexBuffer
 
@@ -728,414 +658,18 @@ FString FIndexBuffer::GetFriendlyName() const
 	return TEXT("FIndexBuffer");
 }
 
-/*------------------------------------------------------------------------------
-	FGlobalDynamicVertexBuffer implementation.
-------------------------------------------------------------------------------*/
-
-/**
- * An individual dynamic vertex buffer.
- */
-class FDynamicVertexBuffer : public FVertexBuffer
-{
-public:
-	/** The aligned size of all dynamic vertex buffers. */
-	enum { ALIGNMENT = (1 << 16) }; // 64KB
-	/** Pointer to the vertex buffer mapped in main memory. */
-	uint8* MappedBuffer;
-	/** Size of the vertex buffer in bytes. */
-	uint32 BufferSize;
-	/** Number of bytes currently allocated from the buffer. */
-	uint32 AllocatedByteCount;
-	/** Number of successive frames for which AllocatedByteCount == 0. Used as a metric to decide when to free the allocation. */
-	int32 NumFramesUnused = 0;
-
-	/** Default constructor. */
-	explicit FDynamicVertexBuffer(uint32 InMinBufferSize)
-		: MappedBuffer(NULL)
-		, BufferSize(FMath::Max<uint32>(Align(InMinBufferSize,ALIGNMENT),ALIGNMENT))
-		, AllocatedByteCount(0)
-	{
-	}
-
-	/**
-	 * Locks the vertex buffer so it may be written to.
-	 */
-	void Lock()
-	{
-		check(MappedBuffer == NULL);
-		check(AllocatedByteCount == 0);
-		check(IsValidRef(VertexBufferRHI));
-		MappedBuffer = (uint8*)RHILockBuffer(VertexBufferRHI, 0, BufferSize, RLM_WriteOnly);
-	}
-
-	/**
-	 * Unocks the buffer so the GPU may read from it.
-	 */
-	void Unlock()
-	{
-		check(MappedBuffer != NULL);
-		check(IsValidRef(VertexBufferRHI));
-		RHIUnlockBuffer(VertexBufferRHI);
-		MappedBuffer = NULL;
-		AllocatedByteCount = 0;
-		NumFramesUnused = 0;
-	}
-
-	// FRenderResource interface.
-	virtual void InitRHI() override
-	{
-		check(!IsValidRef(VertexBufferRHI));
-		FRHIResourceCreateInfo CreateInfo(TEXT("FDynamicVertexBuffer"));
-		VertexBufferRHI = RHICreateVertexBuffer(BufferSize, BUF_Volatile, CreateInfo);
-		MappedBuffer = NULL;
-		AllocatedByteCount = 0;
-	}
-
-	virtual void ReleaseRHI() override
-	{
-		FVertexBuffer::ReleaseRHI();
-		MappedBuffer = NULL;
-		AllocatedByteCount = 0;
-	}
-
-	virtual FString GetFriendlyName() const override
-	{
-		return TEXT("FDynamicVertexBuffer");
-	}
-};
-
-/**
- * A pool of dynamic vertex buffers.
- */
-struct FDynamicVertexBufferPool
-{
-	/** List of vertex buffers. */
-	TIndirectArray<FDynamicVertexBuffer> VertexBuffers;
-	/** The current buffer from which allocations are being made. */
-	FDynamicVertexBuffer* CurrentVertexBuffer;
-
-	/** Default constructor. */
-	FDynamicVertexBufferPool()
-		: CurrentVertexBuffer(NULL)
-	{
-	}
-
-	/** Destructor. */
-	~FDynamicVertexBufferPool()
-	{
-		int32 NumVertexBuffers = VertexBuffers.Num();
-		for (int32 BufferIndex = 0; BufferIndex < NumVertexBuffers; ++BufferIndex)
-		{
-			VertexBuffers[BufferIndex].ReleaseResource();
-		}
-	}
-};
-
-FGlobalDynamicVertexBuffer::FGlobalDynamicVertexBuffer()
-	: TotalAllocatedSinceLastCommit(0)
-{
-	Pool = new FDynamicVertexBufferPool();
-}
-
-FGlobalDynamicVertexBuffer::~FGlobalDynamicVertexBuffer()
-{
-	delete Pool;
-	Pool = NULL;
-}
-
-FGlobalDynamicVertexBuffer::FAllocation FGlobalDynamicVertexBuffer::Allocate(uint32 SizeInBytes)
-{
-	FAllocation Allocation;
-
-	TotalAllocatedSinceLastCommit += SizeInBytes;
-	if (IsRenderAlarmLoggingEnabled())
-	{
-		UE_LOG(LogRendererCore, Warning, TEXT("FGlobalDynamicVertexBuffer::Allocate(%u), will have allocated %u total this frame"), SizeInBytes, TotalAllocatedSinceLastCommit);
-	}
-
-	FDynamicVertexBuffer* VertexBuffer = Pool->CurrentVertexBuffer;
-	if (VertexBuffer == NULL || VertexBuffer->AllocatedByteCount + SizeInBytes > VertexBuffer->BufferSize)
-	{
-		// Find a buffer in the pool big enough to service the request.
-		VertexBuffer = NULL;
-		for (int32 BufferIndex = 0, NumBuffers = Pool->VertexBuffers.Num(); BufferIndex < NumBuffers; ++BufferIndex)
-		{
-			FDynamicVertexBuffer& VertexBufferToCheck = Pool->VertexBuffers[BufferIndex];
-			if (VertexBufferToCheck.AllocatedByteCount + SizeInBytes <= VertexBufferToCheck.BufferSize)
-			{
-				VertexBuffer = &VertexBufferToCheck;
-				break;
-			}
-		}
-
-		// Create a new vertex buffer if needed.
-		if (VertexBuffer == NULL)
-		{
-			VertexBuffer = new FDynamicVertexBuffer(SizeInBytes);
-			Pool->VertexBuffers.Add(VertexBuffer);
-			VertexBuffer->InitResource();
-		}
-
-		// Lock the buffer if needed.
-		if (VertexBuffer->MappedBuffer == NULL)
-		{
-			VertexBuffer->Lock();
-		}
-
-		// Remember this buffer, we'll try to allocate out of it in the future.
-		Pool->CurrentVertexBuffer = VertexBuffer;
-	}
-
-	check(VertexBuffer != NULL);
-	checkf(VertexBuffer->AllocatedByteCount + SizeInBytes <= VertexBuffer->BufferSize, TEXT("Global vertex buffer allocation failed: BufferSize=%d AllocatedByteCount=%d SizeInBytes=%d"), VertexBuffer->BufferSize, VertexBuffer->AllocatedByteCount, SizeInBytes);
-	Allocation.Buffer = VertexBuffer->MappedBuffer + VertexBuffer->AllocatedByteCount;
-	Allocation.VertexBuffer = VertexBuffer;
-	Allocation.VertexOffset = VertexBuffer->AllocatedByteCount;
-	VertexBuffer->AllocatedByteCount += SizeInBytes;
-
-	return Allocation;
-}
-
-bool FGlobalDynamicVertexBuffer::IsRenderAlarmLoggingEnabled() const
-{
-	return GMaxVertexBytesAllocatedPerFrame > 0 && TotalAllocatedSinceLastCommit >= (size_t)GMaxVertexBytesAllocatedPerFrame;
-}
-
-void FGlobalDynamicVertexBuffer::Commit()
-{
-	for (int32 BufferIndex = 0, NumBuffers = Pool->VertexBuffers.Num(); BufferIndex < NumBuffers; ++BufferIndex)
-	{
-		FDynamicVertexBuffer& VertexBuffer = Pool->VertexBuffers[BufferIndex];
-		if (VertexBuffer.MappedBuffer != NULL)
-		{
-			VertexBuffer.Unlock();
-		}
-		else if (GGlobalBufferNumFramesUnusedThresold && !VertexBuffer.AllocatedByteCount)
-		{
-			++VertexBuffer.NumFramesUnused;
-			if (VertexBuffer.NumFramesUnused >= GGlobalBufferNumFramesUnusedThresold)
-			{
-				// Remove the buffer, assumes they are unordered.
-				VertexBuffer.ReleaseResource();
-				Pool->VertexBuffers.RemoveAtSwap(BufferIndex);
-				--BufferIndex;
-				--NumBuffers;
-			}
-		}
-	}
-	Pool->CurrentVertexBuffer = NULL;
-	TotalAllocatedSinceLastCommit = 0;
-}
-
-FGlobalDynamicVertexBuffer InitViewDynamicVertexBuffer;
-FGlobalDynamicVertexBuffer InitShadowViewDynamicVertexBuffer;
-
-/*------------------------------------------------------------------------------
-	FGlobalDynamicIndexBuffer implementation.
-------------------------------------------------------------------------------*/
-
-/**
- * An individual dynamic index buffer.
- */
-class FDynamicIndexBuffer : public FIndexBuffer
-{
-public:
-	/** The aligned size of all dynamic index buffers. */
-	enum { ALIGNMENT = (1 << 16) }; // 64KB
-	/** Pointer to the index buffer mapped in main memory. */
-	uint8* MappedBuffer;
-	/** Size of the index buffer in bytes. */
-	uint32 BufferSize;
-	/** Number of bytes currently allocated from the buffer. */
-	uint32 AllocatedByteCount;
-	/** Stride of the buffer in bytes. */
-	uint32 Stride;
-	/** Number of successive frames for which AllocatedByteCount == 0. Used as a metric to decide when to free the allocation. */
-	int32 NumFramesUnused = 0;
-
-	/** Initialization constructor. */
-	explicit FDynamicIndexBuffer(uint32 InMinBufferSize, uint32 InStride)
-		: MappedBuffer(NULL)
-		, BufferSize(FMath::Max<uint32>(Align(InMinBufferSize,ALIGNMENT),ALIGNMENT))
-		, AllocatedByteCount(0)
-		, Stride(InStride)
-	{
-	}
-
-	/**
-	 * Locks the vertex buffer so it may be written to.
-	 */
-	void Lock()
-	{
-		check(MappedBuffer == NULL);
-		check(AllocatedByteCount == 0);
-		check(IsValidRef(IndexBufferRHI));
-		MappedBuffer = (uint8*)RHILockBuffer(IndexBufferRHI, 0, BufferSize, RLM_WriteOnly);
-	}
-
-	/**
-	 * Unocks the buffer so the GPU may read from it.
-	 */
-	void Unlock()
-	{
-		check(MappedBuffer != NULL);
-		check(IsValidRef(IndexBufferRHI));
-		RHIUnlockBuffer(IndexBufferRHI);
-		MappedBuffer = NULL;
-		AllocatedByteCount = 0;
-		NumFramesUnused = 0;
-	}
-
-	// FRenderResource interface.
-	virtual void InitRHI() override
-	{
-		check(!IsValidRef(IndexBufferRHI));
-		FRHIResourceCreateInfo CreateInfo(TEXT("FDynamicIndexBuffer"));
-		IndexBufferRHI = RHICreateIndexBuffer(Stride, BufferSize, BUF_Volatile, CreateInfo);
-		MappedBuffer = NULL;
-		AllocatedByteCount = 0;
-	}
-
-	virtual void ReleaseRHI() override
-	{
-		FIndexBuffer::ReleaseRHI();
-		MappedBuffer = NULL;
-		AllocatedByteCount = 0;
-	}
-
-	virtual FString GetFriendlyName() const override
-	{
-		return TEXT("FDynamicIndexBuffer");
-	}
-};
-
-/**
- * A pool of dynamic index buffers.
- */
-struct FDynamicIndexBufferPool
-{
-	/** List of index buffers. */
-	TIndirectArray<FDynamicIndexBuffer> IndexBuffers;
-	/** The current buffer from which allocations are being made. */
-	FDynamicIndexBuffer* CurrentIndexBuffer;
-	/** Stride of buffers in this pool. */
-	uint32 BufferStride;
-
-	/** Initialization constructor. */
-	explicit FDynamicIndexBufferPool(uint32 InBufferStride)
-		: CurrentIndexBuffer(NULL)
-		, BufferStride(InBufferStride)
-	{
-	}
-
-	/** Destructor. */
-	~FDynamicIndexBufferPool()
-	{
-		int32 NumIndexBuffers = IndexBuffers.Num();
-		for (int32 BufferIndex = 0; BufferIndex < NumIndexBuffers; ++BufferIndex)
-		{
-			IndexBuffers[BufferIndex].ReleaseResource();
-		}
-	}
-};
-
-FGlobalDynamicIndexBuffer::FGlobalDynamicIndexBuffer()
-{
-	Pools[0] = new FDynamicIndexBufferPool(sizeof(uint16));
-	Pools[1] = new FDynamicIndexBufferPool(sizeof(uint32));
-}
-
-FGlobalDynamicIndexBuffer::~FGlobalDynamicIndexBuffer()
-{
-	for (int32 i = 0; i < 2; ++i)
-	{
-		delete Pools[i];
-		Pools[i] = NULL;
-	}
-}
-
-FGlobalDynamicIndexBuffer::FAllocation FGlobalDynamicIndexBuffer::Allocate(uint32 NumIndices, uint32 IndexStride)
-{
-	FAllocation Allocation;
-
-	if (IndexStride != 2 && IndexStride != 4)
-	{
-		return Allocation;
-	}
-
-	FDynamicIndexBufferPool* Pool = Pools[IndexStride >> 2]; // 2 -> 0, 4 -> 1
-
-	uint32 SizeInBytes = NumIndices * IndexStride;
-	FDynamicIndexBuffer* IndexBuffer = Pool->CurrentIndexBuffer;
-	if (IndexBuffer == NULL || IndexBuffer->AllocatedByteCount + SizeInBytes > IndexBuffer->BufferSize)
-	{
-		// Find a buffer in the pool big enough to service the request.
-		IndexBuffer = NULL;
-		for (int32 BufferIndex = 0, NumBuffers = Pool->IndexBuffers.Num(); BufferIndex < NumBuffers; ++BufferIndex)
-		{
-			FDynamicIndexBuffer& IndexBufferToCheck = Pool->IndexBuffers[BufferIndex];
-			if (IndexBufferToCheck.AllocatedByteCount + SizeInBytes <= IndexBufferToCheck.BufferSize)
-			{
-				IndexBuffer = &IndexBufferToCheck;
-				break;
-			}
-		}
-
-		// Create a new index buffer if needed.
-		if (IndexBuffer == NULL)
-		{
-			IndexBuffer = new FDynamicIndexBuffer(SizeInBytes, Pool->BufferStride);
-			Pool->IndexBuffers.Add(IndexBuffer);
-			IndexBuffer->InitResource();
-		}
-
-		// Lock the buffer if needed.
-		if (IndexBuffer->MappedBuffer == NULL)
-		{
-			IndexBuffer->Lock();
-		}
-		Pool->CurrentIndexBuffer = IndexBuffer;
-	}
-
-	check(IndexBuffer != NULL);
-	checkf(IndexBuffer->AllocatedByteCount + SizeInBytes <= IndexBuffer->BufferSize, TEXT("Global index buffer allocation failed: BufferSize=%d AllocatedByteCount=%d SizeInBytes=%d"), IndexBuffer->BufferSize, IndexBuffer->AllocatedByteCount, SizeInBytes);
-
-	Allocation.Buffer = IndexBuffer->MappedBuffer + IndexBuffer->AllocatedByteCount;
-	Allocation.IndexBuffer = IndexBuffer;
-	Allocation.FirstIndex = IndexBuffer->AllocatedByteCount / IndexStride;
-	IndexBuffer->AllocatedByteCount += SizeInBytes;
+//////////////////////////////////////////////////////////////////////////////////////////////////////////////////////
+// FBufferWithRDG
 
-	return Allocation;
-}
+FBufferWithRDG::FBufferWithRDG() = default;
+FBufferWithRDG::FBufferWithRDG(const FBufferWithRDG& Other) = default;
+FBufferWithRDG& FBufferWithRDG::operator=(const FBufferWithRDG& Other) = default;
+FBufferWithRDG::~FBufferWithRDG() = default;
 
-void FGlobalDynamicIndexBuffer::Commit()
+void FBufferWithRDG::ReleaseRHI()
 {
-	for (int32 i = 0; i < 2; ++i)
-	{
-		FDynamicIndexBufferPool* Pool = Pools[i];
-
-		for (int32 BufferIndex = 0, NumBuffers = Pool->IndexBuffers.Num(); BufferIndex < NumBuffers; ++BufferIndex)
-		{
-			FDynamicIndexBuffer& IndexBuffer = Pool->IndexBuffers[BufferIndex];
-			if (IndexBuffer.MappedBuffer != NULL)
-			{
-				IndexBuffer.Unlock();
-			}
-			else if (GGlobalBufferNumFramesUnusedThresold && !IndexBuffer.AllocatedByteCount)
-			{
-				++IndexBuffer.NumFramesUnused;
-				if (IndexBuffer.NumFramesUnused >= GGlobalBufferNumFramesUnusedThresold)
-				{
-					// Remove the buffer, assumes they are unordered.
-					IndexBuffer.ReleaseResource();
-					Pool->IndexBuffers.RemoveAtSwap(BufferIndex);
-					--BufferIndex;
-					--NumBuffers;
-				}
-			}
-		}
-		Pool->CurrentIndexBuffer = NULL;
-	}
+	Buffer = nullptr;
+	FRenderResource::ReleaseRHI();
 }
 
 /*=============================================================================
diff --git a/Engine/Source/Runtime/RenderCore/Private/RenderUtils.cpp b/Engine/Source/Runtime/RenderCore/Private/RenderUtils.cpp
--- a/Engine/Source/Runtime/RenderCore/Private/RenderUtils.cpp
+++ b/Engine/Source/Runtime/RenderCore/Private/RenderUtils.cpp
@@ -5,13 +5,13 @@
 #include "Containers/DynamicRHIResourceArray.h"
 #include "RenderResource.h"
 #include "RHIStaticStates.h"
-#include "RenderGraphUtils.h"
 #include "PipelineStateCache.h"
 #include "Misc/ConfigCacheIni.h"
 #include "RenderGraphResourcePool.h"
 #include "Misc/DataDrivenPlatformInfoRegistry.h"
 #include "RenderCore.h"
 #include "PackedNormal.h"
+#include "Shader.h"
 
 #if WITH_EDITOR
 #include "Misc/CoreMisc.h"
@@ -73,18 +73,6 @@ static FAutoConsoleVariableRef CVarRayTracingRequireSM6(
 	ECVF_RenderThreadSafe | ECVF_ReadOnly
 );
 
-
-FBufferWithRDG::FBufferWithRDG() = default;
-FBufferWithRDG::FBufferWithRDG(const FBufferWithRDG & Other) = default;
-FBufferWithRDG& FBufferWithRDG::operator=(const FBufferWithRDG & Other) = default;
-FBufferWithRDG::~FBufferWithRDG() = default;
-
-void FBufferWithRDG::ReleaseRHI()
-{
-	Buffer = nullptr;
-	FRenderResource::ReleaseRHI();
-}
-
 const uint16 GCubeIndices[12*3] =
 {
 	0, 2, 3,
@@ -101,11 +89,6 @@ const uint16 GCubeIndices[12*3] =
 	1, 7, 5,
 };
 
-TGlobalResource<FCubeIndexBuffer> GCubeIndexBuffer;
-TGlobalResource<FTwoTrianglesIndexBuffer> GTwoTrianglesIndexBuffer;
-TGlobalResource<FScreenSpaceVertexBuffer> GScreenSpaceVertexBuffer;
-TGlobalResource<FTileVertexDeclaration> GTileVertexDeclaration;
-
 //
 // FPackedNormal serializer
 //
@@ -130,599 +113,6 @@ FArchive& operator<<(FArchive& Ar, FPackedRGBA16N& N)
 	return Ar;
 }
 
-/**
- * Bulk data interface for providing a single black color used to initialize a
- * volume texture.
- */
-class FBlackVolumeTextureResourceBulkDataInterface : public FResourceBulkDataInterface
-{
-public:
-
-	/** Default constructor. */
-	FBlackVolumeTextureResourceBulkDataInterface(uint8 Alpha)
-		: Color(0, 0, 0, Alpha)
-	{
-	}
-
-	/** Default constructor. */
-	FBlackVolumeTextureResourceBulkDataInterface(FColor InColor)
-		: Color(InColor)
-	{
-	}
-
-	/**
-	 * Returns a pointer to the bulk data.
-	 */
-	virtual const void* GetResourceBulkData() const override
-	{
-		return &Color;
-	}
-
-	/** 
-	 * @return size of resource memory
-	 */
-	virtual uint32 GetResourceBulkDataSize() const override
-	{
-		return sizeof(Color);
-	}
-
-	/**
-	 * Free memory after it has been used to initialize RHI resource 
-	 */
-	virtual void Discard() override
-	{
-	}
-
-private:
-
-	/** Storage for the color. */
-	FColor Color;
-};
-
-//
-// FWhiteTexture implementation
-//
-
-/**
- * A solid-colored 1x1 texture.
- */
-template <int32 R, int32 G, int32 B, int32 A>
-class FColoredTexture : public FTextureWithSRV
-{
-public:
-	// FResource interface.
-	virtual void InitRHI() override
-	{
-		// Create the texture RHI.  		
-		FBlackVolumeTextureResourceBulkDataInterface BlackTextureBulkData(FColor(R, G, B, A));
-
-		// BGRA typed UAV is unsupported per D3D spec, use RGBA here.
-		const FRHITextureCreateDesc Desc =
-			FRHITextureCreateDesc::Create2D(TEXT("ColoredTexture"), 1, 1, PF_R8G8B8A8)
-			.SetFlags(ETextureCreateFlags::ShaderResource)
-			.SetBulkData(&BlackTextureBulkData);
-
-		TextureRHI = RHICreateTexture(Desc);
-
-		// Create the sampler state RHI resource.
-		FSamplerStateInitializerRHI SamplerStateInitializer(SF_Point,AM_Wrap,AM_Wrap,AM_Wrap);
-		SamplerStateRHI = GetOrCreateSamplerState(SamplerStateInitializer);
-
-		// Create a view of the texture
-		ShaderResourceViewRHI = RHICreateShaderResourceView(TextureRHI, 0u);
-	}
-
-	/** Returns the width of the texture in pixels. */
-	virtual uint32 GetSizeX() const override
-	{
-		return 1;
-	}
-
-	/** Returns the height of the texture in pixels. */
-	virtual uint32 GetSizeY() const override
-	{
-		return 1;
-	}
-};
-
-class FEmptyVertexBuffer : public FVertexBufferWithSRV
-{
-public:
-	virtual void InitRHI() override
-	{
-		// Create the texture RHI.  		
-		FRHIResourceCreateInfo CreateInfo(TEXT("EmptyVertexBuffer"));
-		
-		VertexBufferRHI = RHICreateVertexBuffer(16u, BUF_Static | BUF_ShaderResource | BUF_UnorderedAccess, CreateInfo);
-
-		// Create a view of the buffer
-		ShaderResourceViewRHI = RHICreateShaderResourceView(VertexBufferRHI, 4u, PF_R32_UINT);
-		UnorderedAccessViewRHI = RHICreateUnorderedAccessView(VertexBufferRHI, PF_R32_UINT);
-	}
-};
-
-class FBlackTextureWithSRV : public FColoredTexture<0, 0, 0, 255>
-{
-	virtual void InitRHI() override
-	{
-		FColoredTexture::InitRHI();
-		FRHITextureReference::DefaultTexture = TextureRHI;
-	}
-
-	virtual void ReleaseRHI() override
-	{
-		FRHITextureReference::DefaultTexture.SafeRelease();
-		FColoredTexture::ReleaseRHI();
-	}
-};
-
-FTextureWithSRV* GWhiteTextureWithSRV = new TGlobalResource<FColoredTexture<255,255,255,255> >;
-FTextureWithSRV* GBlackTextureWithSRV = new TGlobalResource<FBlackTextureWithSRV>();
-FTextureWithSRV* GTransparentBlackTextureWithSRV = new TGlobalResource<FColoredTexture<0,0,0,0> >;
-FTexture* GWhiteTexture = GWhiteTextureWithSRV;
-FTexture* GBlackTexture = GBlackTextureWithSRV;
-FTexture* GTransparentBlackTexture = GTransparentBlackTextureWithSRV;
-
-FVertexBufferWithSRV* GEmptyVertexBufferWithUAV = new TGlobalResource<FEmptyVertexBuffer>;
-
-class FWhiteVertexBuffer : public FVertexBufferWithSRV
-{
-public:
-	virtual void InitRHI() override
-	{
-		// Create the texture RHI.  		
-		FRHIResourceCreateInfo CreateInfo(TEXT("WhiteVertexBuffer"));
-
-		VertexBufferRHI = RHICreateVertexBuffer(sizeof(FVector4f), BUF_Static | BUF_ShaderResource, CreateInfo);
-
-		FVector4f* BufferData = (FVector4f*)RHILockBuffer(VertexBufferRHI, 0, sizeof(FVector4f), RLM_WriteOnly);
-		*BufferData = FVector4f(1.0f, 1.0f, 1.0f, 1.0f);
-		RHIUnlockBuffer(VertexBufferRHI);
-
-		// Create a view of the buffer
-		ShaderResourceViewRHI = RHICreateShaderResourceView(VertexBufferRHI, sizeof(FVector4f), PF_A32B32G32R32F);
-	}
-};
-
-FVertexBufferWithSRV* GWhiteVertexBufferWithSRV = new TGlobalResource<FWhiteVertexBuffer>;
-
-class FWhiteVertexBufferWithRDG : public FBufferWithRDG
-{
-public:
-
-	/**
-	 * Initialize RHI resources.
-	 */
-	virtual void InitRHI() override
-	{
-		if (!Buffer.IsValid())
-		{
-			Buffer = AllocatePooledBuffer(FRDGBufferDesc::CreateBufferDesc(sizeof(FVector4f), 1), TEXT("WhiteVertexBufferWithRDG"));
-
-			FRHICommandListImmediate& RHICmdList = FRHICommandListExecutor::GetImmediateCommandList();
-			FVector4f* BufferData = (FVector4f*)RHICmdList.LockBuffer(Buffer->GetRHI(), 0, sizeof(FVector4f), RLM_WriteOnly);
-			*BufferData = FVector4f(1.0f, 1.0f, 1.0f, 1.0f);
-			RHICmdList.UnlockBuffer(Buffer->GetRHI());
-		}
-	}
-};
-
-FBufferWithRDG* GWhiteVertexBufferWithRDG = new TGlobalResource<FWhiteVertexBufferWithRDG>();
-
-/**
- * A class representing a 1x1x1 black volume texture.
- */
-template <EPixelFormat PixelFormat, uint8 Alpha>
-class FBlackVolumeTexture : public FTexture
-{
-public:
-	
-	/**
-	 * Initialize RHI resources.
-	 */
-	virtual void InitRHI() override
-	{
-		if (GSupportsTexture3D)
-		{
-			// Create the texture.
-			FBlackVolumeTextureResourceBulkDataInterface BlackTextureBulkData(Alpha);
-
-			const FRHITextureCreateDesc Desc =
-				FRHITextureCreateDesc::Create3D(TEXT("BlackVolumeTexture3D"), 1, 1, 1, PixelFormat)
-				.SetFlags(ETextureCreateFlags::ShaderResource)
-				.SetBulkData(&BlackTextureBulkData);
-
-			TextureRHI = RHICreateTexture(Desc);
-		}
-		else
-		{
-			// Create a texture, even though it's not a volume texture
-			FBlackVolumeTextureResourceBulkDataInterface BlackTextureBulkData(Alpha);
-
-			const FRHITextureCreateDesc Desc =
-				FRHITextureCreateDesc::Create2D(TEXT("BlackVolumeTexture2D"), 1, 1, PixelFormat)
-				.SetFlags(ETextureCreateFlags::ShaderResource)
-				.SetBulkData(&BlackTextureBulkData);
-
-			TextureRHI = RHICreateTexture(Desc);
-		}
-
-		// Create the sampler state.
-		FSamplerStateInitializerRHI SamplerStateInitializer(SF_Point, AM_Wrap, AM_Wrap, AM_Wrap);
-		SamplerStateRHI = GetOrCreateSamplerState(SamplerStateInitializer);
-	}
-
-	/**
-	 * Return the size of the texture in the X dimension.
-	 */
-	virtual uint32 GetSizeX() const override
-	{
-		return 1;
-	}
-
-	/**
-	 * Return the size of the texture in the Y dimension.
-	 */
-	virtual uint32 GetSizeY() const override
-	{
-		return 1;
-	}
-};
-
-/** Global black volume texture resource. */
-FTexture* GBlackVolumeTexture = new TGlobalResource<FBlackVolumeTexture<PF_B8G8R8A8, 0>>();
-FTexture* GBlackAlpha1VolumeTexture = new TGlobalResource<FBlackVolumeTexture<PF_B8G8R8A8, 255>>();
-
-/** Global black volume texture resource. */
-FTexture* GBlackUintVolumeTexture = new TGlobalResource<FBlackVolumeTexture<PF_R8G8B8A8_UINT, 0>>();
-
-class FBlackArrayTexture : public FTexture
-{
-public:
-	// FResource interface.
-	virtual void InitRHI() override
-	{
-		// Create the texture RHI.
-		FBlackVolumeTextureResourceBulkDataInterface BlackTextureBulkData(0);
-
-		const FRHITextureCreateDesc Desc =
-			FRHITextureCreateDesc::Create2DArray(TEXT("BlackArrayTexture"), 1, 1, 1, PF_B8G8R8A8)
-			.SetFlags(ETextureCreateFlags::ShaderResource)
-			.SetBulkData(&BlackTextureBulkData);
-
-		TextureRHI = RHICreateTexture(Desc);
-
-		// Create the sampler state RHI resource.
-		FSamplerStateInitializerRHI SamplerStateInitializer(SF_Point, AM_Wrap, AM_Wrap, AM_Wrap);
-		SamplerStateRHI = GetOrCreateSamplerState(SamplerStateInitializer);
-	}
-
-	/** Returns the width of the texture in pixels. */
-	virtual uint32 GetSizeX() const override
-	{
-		return 1;
-	}
-
-	/** Returns the height of the texture in pixels. */
-	virtual uint32 GetSizeY() const override
-	{
-		return 1;
-	}
-};
-
-FTexture* GBlackArrayTexture = new TGlobalResource<FBlackArrayTexture>;
-
-//
-// FMipColorTexture implementation
-//
-
-/**
- * A texture that has a different solid color in each mip-level
- */
-class FMipColorTexture : public FTexture
-{
-public:
-	enum
-	{
-		NumMips = 12
-	};
-	static const FColor MipColors[NumMips];
-
-	// FResource interface.
-	virtual void InitRHI() override
-	{
-		// Create the texture RHI.
-		int32 TextureSize = 1 << (NumMips - 1);
-
-		const FRHITextureCreateDesc Desc =
-			FRHITextureCreateDesc::Create2D(TEXT("FMipColorTexture"), TextureSize, TextureSize, PF_B8G8R8A8)
-			.SetNumMips(NumMips)
-			.SetFlags(ETextureCreateFlags::ShaderResource);
-
-		TextureRHI = RHICreateTexture(Desc);
-
-		// Write the contents of the texture.
-		uint32 DestStride;
-		int32 Size = TextureSize;
-		for ( int32 MipIndex=0; MipIndex < NumMips; ++MipIndex )
-		{
-			FColor* DestBuffer = (FColor*)RHILockTexture2D(TextureRHI, MipIndex, RLM_WriteOnly, DestStride, false);
-			for ( int32 Y=0; Y < Size; ++Y )
-			{
-				for ( int32 X=0; X < Size; ++X )
-				{
-					DestBuffer[X] = MipColors[NumMips - 1 - MipIndex];
-				}
-				DestBuffer += DestStride / sizeof(FColor);
-			}
-			RHIUnlockTexture2D(TextureRHI, MipIndex, false);
-			Size >>= 1;
-		}
-
-		// Create the sampler state RHI resource.
-		FSamplerStateInitializerRHI SamplerStateInitializer(SF_Point,AM_Wrap,AM_Wrap,AM_Wrap);
-		SamplerStateRHI = GetOrCreateSamplerState(SamplerStateInitializer);
-	}
-
-	/** Returns the width of the texture in pixels. */
-	virtual uint32 GetSizeX() const override
-	{
-		int32 TextureSize = 1 << (NumMips - 1);
-		return TextureSize;
-	}
-
-	/** Returns the height of the texture in pixels. */
-	// PVS-Studio notices that the implementation of GetSizeX is identical to this one
-	// and warns us. In this case, it is intentional, so we disable the warning:
-	virtual uint32 GetSizeY() const override //-V524
-	{
-		int32 TextureSize = 1 << (NumMips - 1);
-		return TextureSize;
-	}
-};
-
-const FColor FMipColorTexture::MipColors[NumMips] =
-{
-	FColor(  80,  80,  80, 0 ),		// Mip  0: 1x1			(dark grey)
-	FColor( 200, 200, 200, 0 ),		// Mip  1: 2x2			(light grey)
-	FColor( 200, 200,   0, 0 ),		// Mip  2: 4x4			(medium yellow)
-	FColor( 255, 255,   0, 0 ),		// Mip  3: 8x8			(yellow)
-	FColor( 160, 255,  40, 0 ),		// Mip  4: 16x16		(light green)
-	FColor(   0, 255,   0, 0 ),		// Mip  5: 32x32		(green)
-	FColor(   0, 255, 200, 0 ),		// Mip  6: 64x64		(cyan)
-	FColor(   0, 170, 170, 0 ),		// Mip  7: 128x128		(light blue)
-	FColor(  60,  60, 255, 0 ),		// Mip  8: 256x256		(dark blue)
-	FColor( 255,   0, 255, 0 ),		// Mip  9: 512x512		(pink)
-	FColor( 255,   0,   0, 0 ),		// Mip 10: 1024x1024	(red)
-	FColor( 255, 130,   0, 0 ),		// Mip 11: 2048x2048	(orange)
-};
-
-RENDERCORE_API FTexture* GMipColorTexture = new FMipColorTexture;
-RENDERCORE_API int32 GMipColorTextureMipLevels = FMipColorTexture::NumMips;
-
-// 4: 8x8 cubemap resolution, shader needs to use the same value as preprocessing
-RENDERCORE_API const uint32 GDiffuseConvolveMipLevel = 4;
-
-/** A solid color cube texture. */
-class FSolidColorTextureCube : public FTexture
-{
-public:
-	FSolidColorTextureCube(const FColor& InColor)
-		: bInitToZero(false)
-		, PixelFormat(PF_B8G8R8A8)
-		, ColorData(InColor.DWColor())
-	{}
-
-	FSolidColorTextureCube(EPixelFormat InPixelFormat)
-		: bInitToZero(true)
-		, PixelFormat(InPixelFormat)
-		, ColorData(0)
-	{}
-
-	// FRenderResource interface.
-	virtual void InitRHI() override
-	{
-		// Create the texture RHI.
-		const FRHITextureCreateDesc Desc =
-			FRHITextureCreateDesc::CreateCube(TEXT("SolidColorCube"), 1, PixelFormat)
-			.SetFlags(ETextureCreateFlags::ShaderResource);
-
-		FTextureCubeRHIRef TextureCube = RHICreateTexture(Desc);
-		TextureRHI = TextureCube;
-
-		// Write the contents of the texture.
-		for (uint32 FaceIndex = 0; FaceIndex < 6; FaceIndex++)
-		{
-			uint32 DestStride;
-			void* DestBuffer = RHILockTextureCubeFace(TextureCube, FaceIndex, 0, 0, RLM_WriteOnly, DestStride, false);
-			if (bInitToZero)
-			{
-				FMemory::Memzero(DestBuffer, GPixelFormats[PixelFormat].BlockBytes);
-			}
-			else
-			{
-				FMemory::Memcpy(DestBuffer, &ColorData, sizeof(ColorData));
-			}
-			RHIUnlockTextureCubeFace(TextureCube, FaceIndex, 0, 0, false);
-		}
-
-		// Create the sampler state RHI resource.
-		FSamplerStateInitializerRHI SamplerStateInitializer(SF_Point, AM_Wrap, AM_Wrap, AM_Wrap);
-		SamplerStateRHI = GetOrCreateSamplerState(SamplerStateInitializer);
-	}
-
-	/** Returns the width of the texture in pixels. */
-	virtual uint32 GetSizeX() const override
-	{
-		return 1;
-	}
-
-	/** Returns the height of the texture in pixels. */
-	virtual uint32 GetSizeY() const override
-	{
-		return 1;
-	}
-
-private:
-	const bool bInitToZero;
-	const EPixelFormat PixelFormat;
-	const uint32 ColorData;
-};
-
-/** A white cube texture. */
-class FWhiteTextureCube : public FSolidColorTextureCube
-{
-public:
-	FWhiteTextureCube() : FSolidColorTextureCube(FColor::White) {}
-};
-FTexture* GWhiteTextureCube = new TGlobalResource<FWhiteTextureCube>;
-
-/** A black cube texture. */
-class FBlackTextureCube : public FSolidColorTextureCube
-{
-public:
-	FBlackTextureCube() : FSolidColorTextureCube(FColor::Black) {}
-};
-FTexture* GBlackTextureCube = new TGlobalResource<FBlackTextureCube>;
-
-/** A black cube texture. */
-class FBlackTextureDepthCube : public FSolidColorTextureCube
-{
-public:
-	FBlackTextureDepthCube() : FSolidColorTextureCube(PF_ShadowDepth) {}
-};
-FTexture* GBlackTextureDepthCube = new TGlobalResource<FBlackTextureDepthCube>;
-
-class FBlackCubeArrayTexture : public FTexture
-{
-public:
-	// FResource interface.
-	virtual void InitRHI() override
-	{
-		if (SupportsTextureCubeArray(GetFeatureLevel() ))
-		{
-			const TCHAR* Name = TEXT("BlackCubeArray");
-
-			const FRHITextureCreateDesc Desc =
-				FRHITextureCreateDesc::CreateCubeArray(TEXT("BlackCubeArray"), 1, 1, PF_B8G8R8A8)
-				.SetFlags(ETextureCreateFlags::ShaderResource);
-
-			// Create the texture RHI.
-			TextureRHI = RHICreateTexture(Desc);
-
-			for(uint32 FaceIndex = 0;FaceIndex < 6;FaceIndex++)
-			{
-				uint32 DestStride;
-				FColor* DestBuffer = (FColor*)RHILockTextureCubeFace(TextureRHI, FaceIndex, 0, 0, RLM_WriteOnly, DestStride, false);
-				// Note: alpha is used by reflection environment to say how much of the foreground texture is visible, so 0 says it is completely invisible
-				*DestBuffer = FColor(0, 0, 0, 0);
-				RHIUnlockTextureCubeFace(TextureRHI, FaceIndex, 0, 0, false);
-			}
-
-			// Create the sampler state RHI resource.
-			FSamplerStateInitializerRHI SamplerStateInitializer(SF_Point,AM_Wrap,AM_Wrap,AM_Wrap);
-			SamplerStateRHI = GetOrCreateSamplerState(SamplerStateInitializer);
-		}
-	}
-
-	/** Returns the width of the texture in pixels. */
-	virtual uint32 GetSizeX() const override
-	{
-		return 1;
-	}
-
-	/** Returns the height of the texture in pixels. */
-	virtual uint32 GetSizeY() const override
-	{
-		return 1;
-	}
-};
-FTexture* GBlackCubeArrayTexture = new TGlobalResource<FBlackCubeArrayTexture>;
-
-/**
- * A UINT 1x1 texture.
- */
-template <EPixelFormat Format, uint32 R = 0, uint32 G = 0, uint32 B = 0, uint32 A = 0>
-class FUintTexture : public FTextureWithSRV
-{
-public:
-	// FResource interface.
-	virtual void InitRHI() override
-	{
-		const FRHITextureCreateDesc Desc =
-			FRHITextureCreateDesc::Create2D(TEXT("UintTexture"), 1, 1, Format)
-			.SetFlags(ETextureCreateFlags::ShaderResource);
-
-		TextureRHI = RHICreateTexture(Desc);
-
-		// Write the contents of the texture.
-		uint32 DestStride;
-		void* DestBuffer = RHILockTexture2D(TextureRHI, 0, RLM_WriteOnly, DestStride, false);
-		WriteData(DestBuffer);
-		RHIUnlockTexture2D(TextureRHI, 0, false);
-
-		// Create the sampler state RHI resource.
-		FSamplerStateInitializerRHI SamplerStateInitializer(SF_Point, AM_Wrap, AM_Wrap, AM_Wrap);
-		SamplerStateRHI = GetOrCreateSamplerState(SamplerStateInitializer);
-
-		// Create a view of the texture
-		ShaderResourceViewRHI = RHICreateShaderResourceView(TextureRHI, 0u);
-	}
-
-	/** Returns the width of the texture in pixels. */
-	virtual uint32 GetSizeX() const override
-	{
-		return 1;
-	}
-
-	/** Returns the height of the texture in pixels. */
-	virtual uint32 GetSizeY() const override
-	{
-		return 1;
-	}
-
-protected:
-	static int32 GetNumChannels()
-	{
-		return GPixelFormats[Format].NumComponents;
-	}
-
-	static int32 GetBytesPerChannel()
-	{
-		return GPixelFormats[Format].BlockBytes / GPixelFormats[Format].NumComponents;
-	}
-
-	template<typename T>
-	static void DoWriteData(T* DataPtr)
-	{
-		T Values[] = { R, G, B, A };
-		for (int32 i = 0; i < GetNumChannels(); ++i)
-		{
-			DataPtr[i] = Values[i];
-		}
-	}
-
-	static void WriteData(void* DataPtr)
-	{
-		switch (GetBytesPerChannel())
-		{
-		case 1: 
-			DoWriteData((uint8*)DataPtr);
-			return;
-		case 2:
-			DoWriteData((uint16*)DataPtr);
-			return;
-		case 4:
-			DoWriteData((uint32*)DataPtr);
-			return;
-		}
-		// Unsupported format
-		check(0);
-	}
-};
-
-FTexture* GBlackUintTexture = new TGlobalResource< FUintTexture<PF_R32G32B32A32_UINT> >;
-
 /*
 	3 XYZ packed in 4 bytes. (11:11:10 for X:Y:Z)
 */
diff --git a/Engine/Source/Runtime/RenderCore/Public/GlobalRenderResources.h b/Engine/Source/Runtime/RenderCore/Public/GlobalRenderResources.h
new file mode 100644
--- /dev/null
+++ b/Engine/Source/Runtime/RenderCore/Public/GlobalRenderResources.h
@@ -0,0 +1,354 @@
+// Copyright Epic Games, Inc. All Rights Reserved.
+
+#pragma once
+
+#include "RenderResource.h"
+
+// A global white texture.
+extern RENDERCORE_API FTexture* GWhiteTexture;
+
+// A global white texture with an SRV.
+extern RENDERCORE_API FTextureWithSRV* GWhiteTextureWithSRV;
+
+// A global black texture.
+extern RENDERCORE_API FTexture* GBlackTexture;
+
+// A global black texture with an SRV.
+extern RENDERCORE_API FTextureWithSRV* GBlackTextureWithSRV;
+
+// A global black transparent texture.
+extern RENDERCORE_API FTexture* GTransparentBlackTexture;
+
+// A global black transparent texture with an SRV
+extern RENDERCORE_API FTextureWithSRV* GTransparentBlackTextureWithSRV;
+
+// An empty vertex buffer with a UAV
+extern RENDERCORE_API FVertexBufferWithSRV* GEmptyVertexBufferWithUAV;
+
+// An empty vertex buffer with a SRV
+extern RENDERCORE_API FVertexBufferWithSRV* GWhiteVertexBufferWithSRV;
+
+// A white vertex buffer used with RenderGraph.
+extern RENDERCORE_API FBufferWithRDG* GWhiteVertexBufferWithRDG;
+
+// A global black array texture
+extern RENDERCORE_API FTexture* GBlackArrayTexture;
+
+// A global black volume texture.
+extern RENDERCORE_API FTexture* GBlackVolumeTexture;
+
+// A global black volume texture, with alpha=1.
+extern RENDERCORE_API FTexture* GBlackAlpha1VolumeTexture;
+
+// A global black texture<uint>
+extern RENDERCORE_API FTexture* GBlackUintTexture;
+
+// A global black volume texture<uint> 
+extern RENDERCORE_API FTexture* GBlackUintVolumeTexture;
+
+// A global white cube texture.
+extern RENDERCORE_API FTexture* GWhiteTextureCube;
+
+// A global black cube texture.
+extern RENDERCORE_API FTexture* GBlackTextureCube;
+
+// A global black cube depth texture.
+extern RENDERCORE_API FTexture* GBlackTextureDepthCube;
+
+// A global black cube array texture.
+extern RENDERCORE_API FTexture* GBlackCubeArrayTexture;
+
+// A global texture that has a different solid color in each mip-level.
+extern RENDERCORE_API FTexture* GMipColorTexture;
+
+/** Number of mip-levels in 'GMipColorTexture' */
+extern RENDERCORE_API int32 GMipColorTextureMipLevels;
+
+// 4: 8x8 cubemap resolution, shader needs to use the same value as preprocessing
+extern RENDERCORE_API const uint32 GDiffuseConvolveMipLevel;
+
+// TGlobalResource
+
+/** Used to declare a render resource that is initialized/released by static initialization/destruction. */
+template<class ResourceType>
+class TGlobalResource : public ResourceType
+{
+public:
+	/** Default constructor. */
+	TGlobalResource()
+	{
+		InitGlobalResource();
+	}
+
+	/** Initialization constructor: 1 parameter. */
+	template<typename... Args>
+	explicit TGlobalResource(Args... InArgs)
+		: ResourceType(InArgs...)
+	{
+		InitGlobalResource();
+	}
+
+	/** Destructor. */
+	virtual ~TGlobalResource()
+	{
+		ReleaseGlobalResource();
+	}
+
+private:
+
+	/**
+	 * Initialize the global resource.
+	 */
+	void InitGlobalResource()
+	{
+		if (IsInRenderingThread())
+		{
+			// If the resource is constructed in the rendering thread, directly initialize it.
+			((ResourceType*)this)->InitResource();
+		}
+		else
+		{
+			// If the resource is constructed outside of the rendering thread, enqueue a command to initialize it.
+			BeginInitResource((ResourceType*)this);
+		}
+	}
+
+	/**
+	 * Release the global resource.
+	 */
+	void ReleaseGlobalResource()
+	{
+		// This should be called in the rendering thread, or at shutdown when the rendering thread has exited.
+		// However, it may also be called at shutdown after an error, when the rendering thread is still running.
+		// To avoid a second error in that case we don't assert.
+#if 0
+		check(IsInRenderingThread());
+#endif
+
+		// Cleanup the resource.
+		((ResourceType*)this)->ReleaseResource();
+	}
+};
+
+/**
+* A vertex buffer with a single color component.  This is used on meshes that don't have a color component
+* to keep from needing a separate vertex factory to handle this case.
+*/
+class RENDERCORE_API FNullColorVertexBuffer : public FVertexBuffer
+{
+public:
+	FNullColorVertexBuffer();
+	~FNullColorVertexBuffer();
+
+	virtual void InitRHI() override;
+	virtual void ReleaseRHI() override;
+
+	FShaderResourceViewRHIRef VertexBufferSRV;
+};
+
+/** The global null color vertex buffer, which is set with a stride of 0 on meshes without a color component. */
+extern RENDERCORE_API TGlobalResource<FNullColorVertexBuffer> GNullColorVertexBuffer;
+
+/**
+* A vertex buffer with a single zero float3 component.
+*/
+class RENDERCORE_API FNullVertexBuffer : public FVertexBuffer
+{
+public:
+	FNullVertexBuffer();
+	~FNullVertexBuffer();
+
+	virtual void InitRHI() override;
+	virtual void ReleaseRHI() override;
+
+	FShaderResourceViewRHIRef VertexBufferSRV;
+};
+
+/** The global null vertex buffer, which is set with a stride of 0 on meshes */
+extern RENDERCORE_API TGlobalResource<FNullVertexBuffer> GNullVertexBuffer;
+
+class RENDERCORE_API FScreenSpaceVertexBuffer : public FVertexBuffer
+{
+public:
+	/**
+	* Initialize the RHI for this rendering resource
+	*/
+	virtual void InitRHI() override;
+};
+
+extern RENDERCORE_API TGlobalResource<FScreenSpaceVertexBuffer> GScreenSpaceVertexBuffer;
+
+class RENDERCORE_API FTileVertexDeclaration : public FRenderResource
+{
+public:
+	FTileVertexDeclaration();
+	virtual ~FTileVertexDeclaration();
+
+	virtual void InitRHI() override;
+	virtual void ReleaseRHI() override;
+
+	FVertexDeclarationRHIRef VertexDeclarationRHI;
+};
+
+extern RENDERCORE_API TGlobalResource<FTileVertexDeclaration> GTileVertexDeclaration;
+
+class FCubeIndexBuffer : public FIndexBuffer
+{
+public:
+	/**
+	* Initialize the RHI for this rendering resource
+	*/
+	virtual void InitRHI() override;
+};
+extern RENDERCORE_API TGlobalResource<FCubeIndexBuffer> GCubeIndexBuffer;
+
+class FTwoTrianglesIndexBuffer : public FIndexBuffer
+{
+public:
+	/**
+	* Initialize the RHI for this rendering resource
+	*/
+	virtual void InitRHI() override;
+};
+extern RENDERCORE_API TGlobalResource<FTwoTrianglesIndexBuffer> GTwoTrianglesIndexBuffer;
+
+
+//////////////////////////////////////////////////////////////////////////////////////////////////////////////////////
+// FGlobalDynamicVertexBuffer
+
+struct FGlobalDynamicVertexBufferAllocation
+{
+	/** The location of the buffer in main memory. */
+	uint8* Buffer = nullptr;
+
+	/** The vertex buffer to bind for draw calls. */
+	FVertexBuffer* VertexBuffer = nullptr;
+
+	/** The offset in to the vertex buffer. */
+	uint32 VertexOffset = 0;
+
+	/** Returns true if the allocation is valid. */
+	FORCEINLINE bool IsValid() const
+	{
+		return Buffer != nullptr;
+	}
+};
+
+/**
+ * A system for dynamically allocating GPU memory for vertices.
+ */
+class RENDERCORE_API FGlobalDynamicVertexBuffer
+{
+public:
+	using FAllocation = FGlobalDynamicVertexBufferAllocation;
+
+	/** Default constructor. */
+	FGlobalDynamicVertexBuffer();
+
+	/** Destructor. */
+	~FGlobalDynamicVertexBuffer();
+
+	/**
+	 * Allocates space in the global vertex buffer.
+	 * @param SizeInBytes - The amount of memory to allocate in bytes.
+	 * @returns An FAllocation with information regarding the allocated memory.
+	 */
+	FAllocation Allocate(uint32 SizeInBytes);
+
+	/**
+	 * Commits allocated memory to the GPU.
+	 *		WARNING: Once this buffer has been committed to the GPU, allocations
+	 *		remain valid only until the next call to Allocate!
+	 */
+	void Commit();
+
+	/** Returns true if log statements should be made because we exceeded GMaxVertexBytesAllocatedPerFrame */
+	bool IsRenderAlarmLoggingEnabled() const;
+
+private:
+	/** The pool of vertex buffers from which allocations are made. */
+	struct FDynamicVertexBufferPool* Pool;
+
+	/** A total of all allocations made since the last commit. Used to alert about spikes in memory usage. */
+	size_t TotalAllocatedSinceLastCommit;
+};
+
+struct FGlobalDynamicIndexBufferAllocation
+{
+	/** The location of the buffer in main memory. */
+	uint8* Buffer = nullptr;
+
+	/** The vertex buffer to bind for draw calls. */
+	FIndexBuffer* IndexBuffer = nullptr;
+
+	/** The offset in to the index buffer. */
+	uint32 FirstIndex = 0;
+
+	/** Returns true if the allocation is valid. */
+	FORCEINLINE bool IsValid() const
+	{
+		return Buffer != nullptr;
+	}
+};
+
+struct FGlobalDynamicIndexBufferAllocationEx : public FGlobalDynamicIndexBufferAllocation
+{
+	FGlobalDynamicIndexBufferAllocationEx(const FGlobalDynamicIndexBufferAllocation& InRef, uint32 InNumIndices, uint32 InIndexStride)
+		: FGlobalDynamicIndexBufferAllocation(InRef)
+		, NumIndices(InNumIndices)
+		, IndexStride(InIndexStride)
+	{}
+
+	/** The number of indices allocated. */
+	uint32 NumIndices = 0;
+	/** The allocation stride (2 or 4 bytes). */
+	uint32 IndexStride = 0;
+	/** The maximum value of the indices used. */
+	uint32 MaxUsedIndex = 0;
+};
+
+/**
+ * A system for dynamically allocating GPU memory for indices.
+ */
+class RENDERCORE_API FGlobalDynamicIndexBuffer
+{
+public:
+	using FAllocation = FGlobalDynamicIndexBufferAllocation;
+	using FAllocationEx = FGlobalDynamicIndexBufferAllocationEx;
+
+	/** Default constructor. */
+	FGlobalDynamicIndexBuffer();
+
+	/** Destructor. */
+	~FGlobalDynamicIndexBuffer();
+
+	/**
+	 * Allocates space in the global index buffer.
+	 * @param NumIndices - The number of indices to allocate.
+	 * @param IndexStride - The size of an index (2 or 4 bytes).
+	 * @returns An FAllocation with information regarding the allocated memory.
+	 */
+	FAllocation Allocate(uint32 NumIndices, uint32 IndexStride);
+
+	/**
+	 * Helper function to allocate.
+	 * @param NumIndices - The number of indices to allocate.
+	 * @returns an FAllocation with information regarding the allocated memory.
+	 */
+	template <typename IndexType>
+	FORCEINLINE FAllocationEx Allocate(uint32 NumIndices)
+	{
+		return FAllocationEx(Allocate(NumIndices, sizeof(IndexType)), NumIndices, sizeof(IndexType));
+	}
+
+	/**
+	 * Commits allocated memory to the GPU.
+	 *		WARNING: Once this buffer has been committed to the GPU, allocations
+	 *		remain valid only until the next call to Allocate!
+	 */
+	void Commit();
+
+private:
+	/** The pool of vertex buffers from which allocations are made. */
+	struct FDynamicIndexBufferPool* Pools[2];
+};
diff --git a/Engine/Source/Runtime/RenderCore/Public/RenderResource.h b/Engine/Source/Runtime/RenderCore/Public/RenderResource.h
--- a/Engine/Source/Runtime/RenderCore/Public/RenderResource.h
+++ b/Engine/Source/Runtime/RenderCore/Public/RenderResource.h
@@ -234,84 +234,6 @@ extern RENDERCORE_API void EndBatchedRelease();
  */
 extern RENDERCORE_API void ReleaseResourceAndFlush(FRenderResource* Resource);
 
-/** Used to declare a render resource that is initialized/released by static initialization/destruction. */
-template<class ResourceType>
-class TGlobalResource : public ResourceType
-{
-public:
-
-	/** Default constructor. */
-	TGlobalResource()
-	{
-		InitGlobalResource();
-	}
-
-	/** Initialization constructor: 1 parameter. */
-	template<typename T1>
-	explicit TGlobalResource(T1 Param1)
-		: ResourceType(Param1)
-	{
-		InitGlobalResource();
-	}
-
-	/** Initialization constructor: 2 parameters. */
-	template<typename T1, typename T2>
-	explicit TGlobalResource(T1 Param1, T2 Param2)
-		: ResourceType(Param1, Param2)
-	{
-		InitGlobalResource();
-	}
-
-	/** Initialization constructor: 3 parameters. */
-	template<typename T1, typename T2, typename T3>
-	explicit TGlobalResource(T1 Param1, T2 Param2, T3 Param3)
-		: ResourceType(Param1, Param2, Param3)
-	{
-		InitGlobalResource();
-	}
-
-	/** Destructor. */
-	virtual ~TGlobalResource()
-	{
-		ReleaseGlobalResource();
-	}
-
-private:
-
-	/**
-	 * Initialize the global resource.
-	 */
-	void InitGlobalResource()
-	{
-		if(IsInRenderingThread())
-		{
-			// If the resource is constructed in the rendering thread, directly initialize it.
-			((ResourceType*)this)->InitResource();
-		}
-		else
-		{
-			// If the resource is constructed outside of the rendering thread, enqueue a command to initialize it.
-			BeginInitResource((ResourceType*)this);
-		}
-	}
-
-	/**
-	 * Release the global resource.
-	 */
-	void ReleaseGlobalResource()
-	{
-		// This should be called in the rendering thread, or at shutdown when the rendering thread has exited.
-		// However, it may also be called at shutdown after an error, when the rendering thread is still running.
-		// To avoid a second error in that case we don't assert.
-#if 0
-		check(IsInRenderingThread());
-#endif
-
-		// Cleanup the resource.
-		((ResourceType*)this)->ReleaseResource();
-	}
-};
-
 enum EMipFadeSettings
 {
 	MipFade_Normal = 0,
@@ -485,7 +407,6 @@ public:
 	/** The texture reference's RHI resource. */
 	FTextureReferenceRHIRef	TextureReferenceRHI;
 
-
 private:
 	/** True if the texture reference has been initialized from the game thread. */
 	bool bInitialized_GameThread;
@@ -547,43 +468,6 @@ public:
 	FUnorderedAccessViewRHIRef UnorderedAccessViewRHI;
 };
 
-/**
-* A vertex buffer with a single color component.  This is used on meshes that don't have a color component
-* to keep from needing a separate vertex factory to handle this case.
-*/
-class RENDERCORE_API FNullColorVertexBuffer : public FVertexBuffer
-{
-public:
-	FNullColorVertexBuffer();
-	~FNullColorVertexBuffer();
-
-	virtual void InitRHI() override;
-	virtual void ReleaseRHI() override;
-
-	FShaderResourceViewRHIRef VertexBufferSRV;
-};
-
-/** The global null color vertex buffer, which is set with a stride of 0 on meshes without a color component. */
-extern RENDERCORE_API TGlobalResource<FNullColorVertexBuffer> GNullColorVertexBuffer;
-
-/**
-* A vertex buffer with a single zero float3 component.
-*/
-class RENDERCORE_API FNullVertexBuffer : public FVertexBuffer
-{
-public:
-	FNullVertexBuffer();
-	~FNullVertexBuffer();
-
-	virtual void InitRHI() override;
-	virtual void ReleaseRHI() override;
-
-	FShaderResourceViewRHIRef VertexBufferSRV;
-};
-
-/** The global null vertex buffer, which is set with a stride of 0 on meshes */
-extern RENDERCORE_API TGlobalResource<FNullVertexBuffer> GNullVertexBuffer;
-
 /** An index buffer resource. */
 class RENDERCORE_API FIndexBuffer : public FRenderResource
 {
@@ -598,152 +482,22 @@ public:
 	FBufferRHIRef IndexBufferRHI;
 };
 
-/**
- * A system for dynamically allocating GPU memory for vertices.
- */
-class RENDERCORE_API FGlobalDynamicVertexBuffer
+class RENDERCORE_API FBufferWithRDG : public FRenderResource
 {
 public:
-	/**
-	 * Information regarding an allocation from this buffer.
-	 */
-	struct FAllocation
-	{
-		/** The location of the buffer in main memory. */
-		uint8* Buffer;
-		/** The vertex buffer to bind for draw calls. */
-		FVertexBuffer* VertexBuffer;
-		/** The offset in to the vertex buffer. */
-		uint32 VertexOffset;
-
-		/** Default constructor. */
-		FAllocation()
-			: Buffer(NULL)
-			, VertexBuffer(NULL)
-			, VertexOffset(0)
-		{
-		}
-
-		/** Returns true if the allocation is valid. */
-		FORCEINLINE bool IsValid() const
-		{
-			return Buffer != NULL;
-		}
-	};
+	FBufferWithRDG();
+	FBufferWithRDG(const FBufferWithRDG& Other);
+	FBufferWithRDG& operator=(const FBufferWithRDG& Other);
+	~FBufferWithRDG() override;
 
-	/** Default constructor. */
-	FGlobalDynamicVertexBuffer();
-
-	/** Destructor. */
-	~FGlobalDynamicVertexBuffer();
-
-	/**
-	 * Allocates space in the global vertex buffer.
-	 * @param SizeInBytes - The amount of memory to allocate in bytes.
-	 * @returns An FAllocation with information regarding the allocated memory.
-	 */
-	FAllocation Allocate(uint32 SizeInBytes);
-
-	/**
-	 * Commits allocated memory to the GPU.
-	 *		WARNING: Once this buffer has been committed to the GPU, allocations
-	 *		remain valid only until the next call to Allocate!
-	 */
-	void Commit();
-
-	/** Returns true if log statements should be made because we exceeded GMaxVertexBytesAllocatedPerFrame */
-	bool IsRenderAlarmLoggingEnabled() const;
-
-private:
-	/** The pool of vertex buffers from which allocations are made. */
-	struct FDynamicVertexBufferPool* Pool;
+	void ReleaseRHI() override;
 
-	/** A total of all allocations made since the last commit. Used to alert about spikes in memory usage. */
-	size_t TotalAllocatedSinceLastCommit;
-};
-
-/**
- * A system for dynamically allocating GPU memory for indices.
- */
-class RENDERCORE_API FGlobalDynamicIndexBuffer
-{
-public:
-	/**
-	 * Information regarding an allocation from this buffer.
-	 */
-	struct FAllocation
-	{
-		/** The location of the buffer in main memory. */
-		uint8* Buffer = nullptr;
-		/** The vertex buffer to bind for draw calls. */
-		FIndexBuffer* IndexBuffer = nullptr;
-		/** The offset in to the index buffer. */
-		uint32 FirstIndex = 0;
-
-		/** Returns true if the allocation is valid. */
-		FORCEINLINE bool IsValid() const
-		{
-			return Buffer != NULL;
-		}
-	};
-
-	/** Information data with usage details to avoid passing around parameters. */
-	struct FAllocationEx : public FAllocation
-	{
-		FAllocationEx() = default;
-
-		FAllocationEx(const FAllocation& InRef, uint32 InNumIndices, uint32 InIndexStride) 
-			: FAllocation(InRef)
-			, NumIndices(InNumIndices)
-			, IndexStride(InIndexStride) 
-		{}
-
-		/** The number of indices allocated. */
-		uint32 NumIndices = 0;
-		/** The allocation stride (2 or 4 bytes). */
-		uint32 IndexStride = 0;
-		/** The maximum value of the indices used. */
-		uint32 MaxUsedIndex = 0;
-	};
-
-	/** Default constructor. */
-	FGlobalDynamicIndexBuffer();
-
-	/** Destructor. */
-	~FGlobalDynamicIndexBuffer();
-
-	/**
-	 * Allocates space in the global index buffer.
-	 * @param NumIndices - The number of indices to allocate.
-	 * @param IndexStride - The size of an index (2 or 4 bytes).
-	 * @returns An FAllocation with information regarding the allocated memory.
-	 */
-	FAllocation Allocate(uint32 NumIndices, uint32 IndexStride);
-
-	/**
-	 * Helper function to allocate.
-	 * @param NumIndices - The number of indices to allocate.
-	 * @returns an FAllocation with information regarding the allocated memory.
-	 */
-	template <typename IndexType>
-	FORCEINLINE FAllocationEx Allocate(uint32 NumIndices)
-	{
-		return FAllocationEx(Allocate(NumIndices, sizeof(IndexType)), NumIndices, sizeof(IndexType));
-	}
-
-	/**
-	 * Commits allocated memory to the GPU.
-	 *		WARNING: Once this buffer has been committed to the GPU, allocations
-	 *		remain valid only until the next call to Allocate!
-	 */
-	void Commit();
-
-private:
-	/** The pool of vertex buffers from which allocations are made. */
-	struct FDynamicIndexBufferPool* Pools[2];
+	TRefCountPtr<FRDGPooledBuffer> Buffer;
 };
 
 #if UE_ENABLE_INCLUDE_ORDER_DEPRECATED_IN_5_2
 #include "RayTracingGeometry.h"
 #include "RenderUtils.h"
 #endif
+
+#include "GlobalRenderResources.h"
diff --git a/Engine/Source/Runtime/RenderCore/Public/RenderUtils.h b/Engine/Source/Runtime/RenderCore/Public/RenderUtils.h
--- a/Engine/Source/Runtime/RenderCore/Public/RenderUtils.h
+++ b/Engine/Source/Runtime/RenderCore/Public/RenderUtils.h
@@ -9,28 +9,15 @@
 #if UE_ENABLE_INCLUDE_ORDER_DEPRECATED_IN_5_2
 #include "PackedNormal.h"
 #include "RenderMath.h"
+#include "GlobalRenderResources.h"
 #endif
 
 class FBufferWithRDG;
 class FRDGBuffer;
-class FRDGPooledBuffer;
 class FRDGTexture;
 class FRDGBuilder;
 struct IPooledRenderTarget;
 
-class RENDERCORE_API FBufferWithRDG : public FRenderResource
-{
-public:
-	FBufferWithRDG();
-	FBufferWithRDG(const FBufferWithRDG& Other);
-	FBufferWithRDG& operator=(const FBufferWithRDG& Other);
-	~FBufferWithRDG() override;
-
-	void ReleaseRHI() override;
-
-	TRefCountPtr<FRDGPooledBuffer> Buffer;
-};
-
 extern RENDERCORE_API void RenderUtilsInit();
 
 #define NUM_DEBUG_UTIL_COLORS (32)
@@ -70,150 +57,11 @@ static const FColor DebugUtilColor[NUM_DEBUG_UTIL_COLORS] =
 	FColor(19,25,126)
 };
 
-/** A global white texture. */
-extern RENDERCORE_API class FTexture* GWhiteTexture;
-extern RENDERCORE_API class FTextureWithSRV* GWhiteTextureWithSRV;
-
-/** A global black texture. */
-extern RENDERCORE_API class FTexture* GBlackTexture;
-extern RENDERCORE_API class FTextureWithSRV* GBlackTextureWithSRV;
-
-extern RENDERCORE_API class FTexture* GTransparentBlackTexture;
-extern RENDERCORE_API class FTextureWithSRV* GTransparentBlackTextureWithSRV;
-
-extern RENDERCORE_API class FVertexBufferWithSRV* GEmptyVertexBufferWithUAV;
-
-extern RENDERCORE_API class FVertexBufferWithSRV* GWhiteVertexBufferWithSRV;
-
-extern RENDERCORE_API class FBufferWithRDG* GWhiteVertexBufferWithRDG;
-
-/** A global black array texture. */
-extern RENDERCORE_API class FTexture* GBlackArrayTexture;
-
-/** A global black volume texture. */
-extern RENDERCORE_API class FTexture* GBlackVolumeTexture;
-
-/** A global black volume texture, with alpha=1. */
-extern RENDERCORE_API class FTexture* GBlackAlpha1VolumeTexture;
-
-/** A global black texture<uint> */
-extern RENDERCORE_API class FTexture* GBlackUintTexture;
-
-/** A global black volume texture<uint>  */
-extern RENDERCORE_API class FTexture* GBlackUintVolumeTexture;
-
-/** A global white cube texture. */
-extern RENDERCORE_API class FTexture* GWhiteTextureCube;
-
-/** A global black cube texture. */
-extern RENDERCORE_API class FTexture* GBlackTextureCube;
-
-/** A global black cube depth texture. */
-extern RENDERCORE_API class FTexture* GBlackTextureDepthCube;
-
-/** A global black cube array texture. */
-extern RENDERCORE_API class FTexture* GBlackCubeArrayTexture;
-
-/** A global texture that has a different solid color in each mip-level. */
-extern RENDERCORE_API class FTexture* GMipColorTexture;
-
-/** Number of mip-levels in 'GMipColorTexture' */
-extern RENDERCORE_API int32 GMipColorTextureMipLevels;
-
-// 4: 8x8 cubemap resolution, shader needs to use the same value as preprocessing
-extern RENDERCORE_API const uint32 GDiffuseConvolveMipLevel;
-
 #define NUM_CUBE_VERTICES 36
+
 /** The indices for drawing a cube. */
 extern RENDERCORE_API const uint16 GCubeIndices[36];
 
-class FCubeIndexBuffer : public FIndexBuffer
-{
-public:
-	/**
-	* Initialize the RHI for this rendering resource
-	*/
-	virtual void InitRHI() override
-	{
-		// create a static vertex buffer
-		FRHIResourceCreateInfo CreateInfo(TEXT("FCubeIndexBuffer"));
-		IndexBufferRHI = RHICreateIndexBuffer(sizeof(uint16), sizeof(uint16) * NUM_CUBE_VERTICES, BUF_Static, CreateInfo);
-		void* VoidPtr = RHILockBuffer(IndexBufferRHI, 0, sizeof(uint16) * NUM_CUBE_VERTICES, RLM_WriteOnly);
-		FMemory::Memcpy(VoidPtr, GCubeIndices, NUM_CUBE_VERTICES * sizeof(uint16));
-		RHIUnlockBuffer(IndexBufferRHI);
-	}
-};
-extern RENDERCORE_API TGlobalResource<FCubeIndexBuffer> GCubeIndexBuffer;
-
-class FTwoTrianglesIndexBuffer : public FIndexBuffer
-{
-public:
-	/**
-	* Initialize the RHI for this rendering resource
-	*/
-	virtual void InitRHI() override
-	{
-		// create a static vertex buffer
-		FRHIResourceCreateInfo CreateInfo(TEXT("FTwoTrianglesIndexBuffer"));
-		IndexBufferRHI = RHICreateIndexBuffer(sizeof(uint16), sizeof(uint16) * 6, BUF_Static, CreateInfo);
-		void* VoidPtr = RHILockBuffer(IndexBufferRHI, 0, sizeof(uint16) * 6, RLM_WriteOnly);
-		static const uint16 Indices[] = { 0, 1, 3, 0, 3, 2 };
-		FMemory::Memcpy(VoidPtr, Indices, 6 * sizeof(uint16));
-		RHIUnlockBuffer(IndexBufferRHI);
-	}
-};
-extern RENDERCORE_API TGlobalResource<FTwoTrianglesIndexBuffer> GTwoTrianglesIndexBuffer;
-
-class FScreenSpaceVertexBuffer : public FVertexBuffer
-{
-public:
-	/**
-	* Initialize the RHI for this rendering resource
-	*/
-	virtual void InitRHI() override
-	{
-		// create a static vertex buffer
-		FRHIResourceCreateInfo CreateInfo(TEXT("FScreenSpaceVertexBuffer"));
-		VertexBufferRHI = RHICreateVertexBuffer(sizeof(FVector2f) * 4, BUF_Static, CreateInfo);
-		void* VoidPtr = RHILockBuffer(VertexBufferRHI, 0, sizeof(FVector2f) * 4, RLM_WriteOnly);
-		static const FVector2f Vertices[4] =
-		{
-			FVector2f(-1,-1),
-			FVector2f(-1,+1),
-			FVector2f(+1,-1),
-			FVector2f(+1,+1),
-		};
-		FMemory::Memcpy(VoidPtr, Vertices, sizeof(FVector2f) * 4);
-		RHIUnlockBuffer(VertexBufferRHI);
-	}
-};
-
-extern RENDERCORE_API TGlobalResource<FScreenSpaceVertexBuffer> GScreenSpaceVertexBuffer;
-
-class FTileVertexDeclaration : public FRenderResource
-{
-public:
-	FVertexDeclarationRHIRef VertexDeclarationRHI;
-
-	/** Destructor. */
-	virtual ~FTileVertexDeclaration() {}
-
-	virtual void InitRHI()
-	{
-		FVertexDeclarationElementList Elements;
-		uint16 Stride = sizeof(FVector2f);
-		Elements.Add(FVertexElement(0, 0, VET_Float2, 0, Stride, false));
-		VertexDeclarationRHI = RHICreateVertexDeclaration(Elements);
-	}
-
-	virtual void ReleaseRHI()
-	{
-		VertexDeclarationRHI.SafeRelease();
-	}
-};
-
-extern RENDERCORE_API TGlobalResource<FTileVertexDeclaration> GTileVertexDeclaration;
-
 /**
  * Maps from an X,Y,Z cube vertex coordinate to the corresponding vertex index.
  */
