Moving RenderUtils helper functions into the cpp file.

diff --git a/Engine/Source/Runtime/RenderCore/Private/RenderUtils.cpp b/Engine/Source/Runtime/RenderCore/Private/RenderUtils.cpp
--- a/Engine/Source/Runtime/RenderCore/Private/RenderUtils.cpp
+++ b/Engine/Source/Runtime/RenderCore/Private/RenderUtils.cpp
@@ -648,6 +648,21 @@ RENDERCORE_API bool SupportsTSR(const FStaticShaderPlatform Platform)
 	return FDataDrivenShaderPlatformInfo::GetSupportsGen5TemporalAA(Platform);
 }
 
+EShaderPlatform GetEditorShaderPlatform(EShaderPlatform ShaderPlatform)
+{
+	EShaderPlatform ActualShaderPlatform = ShaderPlatform;
+
+	if (GIsEditor)
+	{
+		if (FDataDrivenShaderPlatformInfo::GetIsPreviewPlatform(ActualShaderPlatform))
+		{
+			ActualShaderPlatform = FDataDrivenShaderPlatformInfo::GetPreviewShaderPlatformParent(ActualShaderPlatform);
+		}
+	}
+
+	return ActualShaderPlatform;
+}
+
 RENDERCORE_API int32 GUseForwardShading = 0;
 static FAutoConsoleVariableRef CVarForwardShading(
 	TEXT("r.ForwardShading"),
@@ -1167,7 +1182,7 @@ RENDERCORE_API bool PlatformSupportsVelocityRendering(const FStaticShaderPlatfor
 	return true;
 }
 
-RENDERCORE_API bool DoesPlatformSupportNanite(EShaderPlatform Platform, bool bCheckForProjectSetting)
+bool DoesPlatformSupportNanite(EShaderPlatform Platform, bool bCheckForProjectSetting)
 {
 	// Nanite allowed for this project
 	if (bCheckForProjectSetting)
@@ -1192,10 +1207,105 @@ RENDERCORE_API bool DoesPlatformSupportNanite(EShaderPlatform Platform, bool bCh
 	return bFullCheck;
 }
 
-/** Returns whether DBuffer decals are enabled for a given shader platform */
+bool NaniteAtomicsSupported()
+{
+	// Are 64bit image atomics supported by the GPU/Driver/OS/API?
+	bool bAtomicsSupported = GRHISupportsAtomicUInt64;
+
+#if PLATFORM_WINDOWS
+	const ERHIInterfaceType RHIInterface = RHIGetInterfaceType();
+	const bool bIsDx11 = RHIInterface == ERHIInterfaceType::D3D11;
+	const bool bIsDx12 = RHIInterface == ERHIInterfaceType::D3D12;
+
+	static const auto NaniteRequireDX12CVar = IConsoleManager::Get().FindConsoleVariable(TEXT("r.Nanite.RequireDX12"));
+	static const uint32 NaniteRequireDX12 = (NaniteRequireDX12CVar != nullptr) ? NaniteRequireDX12CVar->GetInt() : 1;
+
+	if (bAtomicsSupported && NaniteRequireDX12 != 0)
+	{
+		// Only allow Vulkan or D3D12
+		bAtomicsSupported = !bIsDx11;
+
+		// Disable DX12 vendor extensions unless DX12 SM6.6 is supported
+		if (NaniteRequireDX12 == 1 && bIsDx12 && !GRHISupportsDX12AtomicUInt64)
+		{
+			// Vendor extensions currently support atomic64, but SM 6.6 and the DX12 Agility SDK are reporting that atomics are not supported.
+			// Likely due to a pre-1909 Windows 10 version, or outdated drivers without SM 6.6 support.
+			// See: https://devblogs.microsoft.com/directx/gettingstarted-dx12agility/
+			bAtomicsSupported = false;
+		}
+	}
+#endif
+
+	return bAtomicsSupported;
+}
+
+bool DoesRuntimeSupportNanite(EShaderPlatform ShaderPlatform, bool bCheckForAtomicSupport, bool bCheckForProjectSetting)
+{
+	// Does the platform support Nanite?
+	const bool bSupportedPlatform = DoesPlatformSupportNanite(ShaderPlatform, bCheckForProjectSetting);
+
+	// Nanite is not supported with forward shading at this time.
+	const bool bForwardShadingEnabled = IsForwardShadingEnabled(ShaderPlatform);
+
+	return bSupportedPlatform && (!bCheckForAtomicSupport || NaniteAtomicsSupported()) && !bForwardShadingEnabled;
+}
+
+bool UseNanite(EShaderPlatform ShaderPlatform, bool bCheckForAtomicSupport /*= true*/, bool bCheckForProjectSetting /*= true*/)
+{
+	static const auto EnableNaniteCVar = IConsoleManager::Get().FindConsoleVariable(TEXT("r.Nanite"));
+	const bool bNaniteEnabled = (EnableNaniteCVar != nullptr) ? (EnableNaniteCVar->GetInt() != 0) : true;
+	return bNaniteEnabled && DoesRuntimeSupportNanite(ShaderPlatform, bCheckForAtomicSupport, bCheckForProjectSetting);
+}
+
+bool UseVirtualShadowMaps(EShaderPlatform ShaderPlatform, const FStaticFeatureLevel FeatureLevel)
+{
+	static const auto EnableVirtualSMCVar = IConsoleManager::Get().FindConsoleVariable(TEXT("r.Shadow.Virtual.Enable"));
+	const bool bVirtualShadowMapsEnabled = EnableVirtualSMCVar ? (EnableVirtualSMCVar->GetInt() != 0) : false;
+	return bVirtualShadowMapsEnabled && DoesRuntimeSupportNanite(ShaderPlatform, true /* check for atomics */, false /* check project setting */);
+}
+
+bool DoesPlatformSupportVirtualShadowMaps(EShaderPlatform Platform)
+{
+	return DoesPlatformSupportNanite(Platform, false /* check project setting */);
+}
+
+bool DoesPlatformSupportNonNaniteVirtualShadowMaps(EShaderPlatform ShaderPlatform)
+{
+	static const auto EnableCVar = IConsoleManager::Get().FindConsoleVariable(TEXT("r.Shadow.Virtual.NonNaniteVSM"));
+	return EnableCVar->GetInt() != 0 && DoesPlatformSupportNanite(ShaderPlatform, false /* check project setting */);
+}
+
+bool UseNonNaniteVirtualShadowMaps(EShaderPlatform ShaderPlatform, const FStaticFeatureLevel FeatureLevel)
+{
+	static const auto EnableCVar = IConsoleManager::Get().FindConsoleVariable(TEXT("r.Shadow.Virtual.NonNaniteVSM"));
+	return EnableCVar->GetInt() != 0 && UseVirtualShadowMaps(ShaderPlatform, FeatureLevel);
+}
+
+bool IsWaterVirtualShadowMapFilteringEnabled(const FStaticShaderPlatform Platform)
+{
+	static const auto CVarWaterSingleLayerShaderSupportVSMFiltering = IConsoleManager::Get().FindConsoleVariable(TEXT("r.Water.SingleLayer.ShadersSupportVSMFiltering"));
+	const bool bWaterVSMFilteringSupported = CVarWaterSingleLayerShaderSupportVSMFiltering && (CVarWaterSingleLayerShaderSupportVSMFiltering->GetInt() > 0);
+
+	const bool bVirtualShadowMapsSupported = DoesPlatformSupportVirtualShadowMaps(Platform);
+
+	return !IsForwardShadingEnabled(Platform) && bWaterVSMFilteringSupported && bVirtualShadowMapsSupported;
+}
+
+bool IsSingleLayerWaterDepthPrepassEnabled(const FStaticShaderPlatform Platform, const FStaticFeatureLevel FeatureLevel)
+{
+	static const auto CVarWaterSingleLayerDepthPrepass = IConsoleManager::Get().FindConsoleVariable(TEXT("r.Water.SingleLayer.DepthPrepass"));
+	const bool bPrepassEnabled = CVarWaterSingleLayerDepthPrepass && CVarWaterSingleLayerDepthPrepass->GetInt() > 0;
+	// Currently VSM is the only feature dependent on the depth prepass which is why we only enable it if VSM could also be enabled.
+	// VSM can be toggled at runtime, but we need a compile time value here, so we fall back to DoesPlatformSupportVirtualShadowMaps() to check if
+	// VSM *could* be enabled.
+	const bool bVSMSupported = DoesPlatformSupportVirtualShadowMaps(Platform);
+
+	return bPrepassEnabled && bVSMSupported;
+
+}
+
 RENDERCORE_API bool IsUsingDBuffers(const FStaticShaderPlatform Platform)
 {
-	extern RENDERCORE_API ShaderPlatformMaskType GDBufferPlatformMask;
 	return (GDBufferPlatformMask[(int)Platform]);
 }
 
@@ -1213,6 +1323,25 @@ RENDERCORE_API bool DoesRuntimeSupportOnePassPointLightShadows(EShaderPlatform P
 		|| (CVar->GetValueOnAnyThread() != 0 && GRHISupportsArrayIndexFromAnyShader != 0);
 }
 
+bool IsForwardShadingEnabled(const FStaticShaderPlatform Platform)
+{
+	return (GForwardShadingPlatformMask[(int)Platform])
+		// Culling uses compute shader
+		&& GetMaxSupportedFeatureLevel(Platform) >= ERHIFeatureLevel::SM5;
+}
+
+bool IsUsingGBuffers(const FStaticShaderPlatform Platform)
+{
+	if (IsMobilePlatform(Platform))
+	{
+		return IsMobileDeferredShadingEnabled(Platform);
+	}
+	else
+	{
+		return !IsForwardShadingEnabled(Platform);
+	}
+}
+
 bool IsUsingBasePassVelocity(const FStaticShaderPlatform Platform)
 {
 	static FShaderPlatformCachedIniValue<int32> PerPlatformCVar(TEXT("r.VelocityOutputPass"));
@@ -1229,15 +1358,107 @@ bool IsUsingBasePassVelocity(const FStaticShaderPlatform Platform)
 	}
 }
 
-RENDERCORE_API bool AllowTranslucencyPerObjectShadows(const FStaticShaderPlatform& Platform)
+bool IsUsingSelectiveBasePassOutputs(const FStaticShaderPlatform Platform)
+{
+	return (GSelectiveBasePassOutputsPlatformMask[(int)Platform]);
+}
+
+bool IsUsingDistanceFields(const FStaticShaderPlatform Platform)
+{
+	return (GDistanceFieldsPlatformMask[(int)Platform]);
+}
+
+bool IsWaterDistanceFieldShadowEnabled(const FStaticShaderPlatform Platform)
+{
+	// Only deferred support such a feature. It is not possible to do that for water without a water depth pre-pass.
+	static const auto CVarWaterSingleLayerShaderSupportDistanceFieldShadow = IConsoleManager::Get().FindConsoleVariable(TEXT("r.Water.SingleLayer.ShadersSupportDistanceFieldShadow"));
+	const bool bWaterSingleLayerShaderSupportDistanceFieldShadow = CVarWaterSingleLayerShaderSupportDistanceFieldShadow && (CVarWaterSingleLayerShaderSupportDistanceFieldShadow->GetInt() > 0);
+	return !IsForwardShadingEnabled(Platform) && IsUsingDistanceFields(Platform) && bWaterSingleLayerShaderSupportDistanceFieldShadow;
+}
+
+bool UseGPUScene(const FStaticShaderPlatform Platform, const FStaticFeatureLevel FeatureLevel)
+{
+	if (FeatureLevel == ERHIFeatureLevel::ES3_1)
+	{
+		return MobileSupportsGPUScene();
+	}
+
+	// GPU Scene management uses compute shaders
+	return FeatureLevel >= ERHIFeatureLevel::SM5
+		//@todo - support GPU Scene management compute shaders on these platforms to get dynamic instancing speedups on the Rendering Thread and RHI Thread
+		&& !IsOpenGLPlatform(Platform)
+		&& !IsVulkanMobileSM5Platform(Platform)
+		&& !IsMetalMobileSM5Platform(Platform)
+		// we only check DDSPI for platforms that have been read in - IsValid() can go away once ALL platforms are converted over to this system
+		&& (!FDataDrivenShaderPlatformInfo::IsValid(Platform) || FDataDrivenShaderPlatformInfo::GetSupportsGPUScene(Platform));
+}
+
+bool UseGPUScene(const FStaticShaderPlatform Platform)
+{
+	return UseGPUScene(Platform, GetMaxSupportedFeatureLevel(Platform));
+}
+
+bool ForceSimpleSkyDiffuse(const FStaticShaderPlatform Platform)
+{
+	return (GSimpleSkyDiffusePlatformMask[(int)Platform]);
+}
+
+bool VelocityEncodeDepth(const FStaticShaderPlatform Platform)
+{
+	return (GVelocityEncodeDepthPlatformMask[(int)Platform]);
+}
+
+RENDERCORE_API bool AllowTranslucencyPerObjectShadows(const FStaticShaderPlatform Platform)
 {
 	return IsFeatureLevelSupported(Platform, ERHIFeatureLevel::SM5) && GAllowTranslucencyShadowsInProject != 0;
 }
 
+bool PlatformRequires128bitRT(EPixelFormat PixelFormat)
+{
+	switch (PixelFormat)
+	{
+	case PF_R32_FLOAT:
+	case PF_G32R32F:
+	case PF_A32B32G32R32F:
+		return FDataDrivenShaderPlatformInfo::GetRequiresExplicit128bitRT(GMaxRHIShaderPlatform);
+	default:
+		return false;
+	}
+}
+
+bool IsRayTracingEnabledForProject(EShaderPlatform ShaderPlatform)
+{
+	if (RHISupportsRayTracing(ShaderPlatform))
+	{
+		return (GRayTracingPlatformMask[(int)ShaderPlatform]);
+	}
+	else
+	{
+		return false;
+	}
+}
+
+bool ShouldCompileRayTracingShadersForProject(EShaderPlatform ShaderPlatform)
+{
+	if (RHISupportsRayTracingShaders(ShaderPlatform))
+	{
+		return IsRayTracingEnabledForProject(ShaderPlatform);
+	}
+	else
+	{
+		return false;
+	}
+}
+
+bool ShouldCompileRayTracingCallableShadersForProject(EShaderPlatform ShaderPlatform)
+{
+	return RHISupportsRayTracingCallableShaders(ShaderPlatform) && ShouldCompileRayTracingShadersForProject(ShaderPlatform);
+}
+
+
 bool IsRayTracingEnabled()
 {
 	bool bRayTracingEnabled = true;
-	extern RENDERCORE_API ERayTracingMode GRayTracingMode;
 	static const auto RayTracingEnableCVar = IConsoleManager::Get().FindTConsoleVariableDataInt(TEXT("r.Raytracing.Enable"));
 	if (GRayTracingMode == ERayTracingMode::Dynamic && RayTracingEnableCVar)
 	{
@@ -1261,7 +1482,6 @@ bool IsRayTracingAllowed()
 	}
 #endif // DO_CHECK && WITH_EDITOR
 
-	extern RENDERCORE_API ERayTracingMode GRayTracingMode;
 	return (int32)GRayTracingMode >= (int32)ERayTracingMode::Enabled;
 }
 
@@ -1272,6 +1492,5 @@ bool IsRayTracingEnabled(EShaderPlatform ShaderPlatform)
 
 ERayTracingMode GetRayTracingMode()
 {
-	extern RENDERCORE_API ERayTracingMode GRayTracingMode;
 	return IsRayTracingAllowed() ? GRayTracingMode : ERayTracingMode::Disabled;
 }
diff --git a/Engine/Source/Runtime/RenderCore/Public/RenderUtils.h b/Engine/Source/Runtime/RenderCore/Public/RenderUtils.h
--- a/Engine/Source/Runtime/RenderCore/Public/RenderUtils.h
+++ b/Engine/Source/Runtime/RenderCore/Public/RenderUtils.h
@@ -16,6 +16,7 @@ class FBufferWithRDG;
 class FRDGBuffer;
 class FRDGTexture;
 class FRDGBuilder;
+class ITargetPlatform;
 struct IPooledRenderTarget;
 
 extern RENDERCORE_API void RenderUtilsInit();
@@ -326,6 +327,8 @@ RENDERCORE_API bool PlatformSupportsVelocityRendering(const FStaticShaderPlatfor
 
 RENDERCORE_API bool IsUsingDBuffers(const FStaticShaderPlatform Platform);
 
+RENDERCORE_API EShaderPlatform GetEditorShaderPlatform(EShaderPlatform ShaderPlatform);
+
 template<typename Type>
 struct RENDERCORE_API FShaderPlatformCachedIniValue
 {
@@ -342,16 +345,9 @@ struct RENDERCORE_API FShaderPlatformCachedIniValue
 
 	Type Get(EShaderPlatform ShaderPlatform)
 	{
-		EShaderPlatform ActualShaderPlatform = ShaderPlatform;
 		Type Value{};
 
-		if (GIsEditor)
-		{
-			if (FDataDrivenShaderPlatformInfo::GetIsPreviewPlatform(ActualShaderPlatform))
-			{
-				ActualShaderPlatform = FDataDrivenShaderPlatformInfo::GetPreviewShaderPlatformParent(ActualShaderPlatform);
-			}
-		}
+		const EShaderPlatform ActualShaderPlatform = GetEditorShaderPlatform(ShaderPlatform);
 
 		FName IniPlatformName = ShaderPlatformToPlatformName(ActualShaderPlatform);
 		// find the cvar if needed
@@ -403,87 +399,30 @@ private:
 };
 
 /** Returns if ForwardShading is enabled. Only valid for the current platform (otherwise call ITargetPlatform::UsesForwardShading()). */
-inline bool IsForwardShadingEnabled(const FStaticShaderPlatform Platform)
-{
-	extern RENDERCORE_API ShaderPlatformMaskType GForwardShadingPlatformMask;
-	return (GForwardShadingPlatformMask[(int)Platform])
-		// Culling uses compute shader
-		&& GetMaxSupportedFeatureLevel(Platform) >= ERHIFeatureLevel::SM5;
-}
+RENDERCORE_API bool IsForwardShadingEnabled(const FStaticShaderPlatform Platform);
 
 /** Returns if the GBuffer is used. Only valid for the current platform. */
-inline bool IsUsingGBuffers(const FStaticShaderPlatform Platform)
-{
-	if (IsMobilePlatform(Platform))
-	{
-		return IsMobileDeferredShadingEnabled(Platform);
-	}
-	else
-	{
-		return !IsForwardShadingEnabled(Platform);
-	}
-}
+RENDERCORE_API bool IsUsingGBuffers(const FStaticShaderPlatform Platform);
 
 /** Returns whether the base pass should output to the velocity buffer is enabled for a given shader platform */
 RENDERCORE_API bool IsUsingBasePassVelocity(const FStaticShaderPlatform Platform);
 
 /** Returns whether the base pass should use selective outputs for a given shader platform */
-inline bool IsUsingSelectiveBasePassOutputs(const FStaticShaderPlatform Platform)
-{
-	extern RENDERCORE_API ShaderPlatformMaskType GSelectiveBasePassOutputsPlatformMask;
-	return (GSelectiveBasePassOutputsPlatformMask[(int)Platform]);
-}
+RENDERCORE_API bool IsUsingSelectiveBasePassOutputs(const FStaticShaderPlatform Platform);
 
 /** Returns whether distance fields are enabled for a given shader platform */
-inline bool IsUsingDistanceFields(const FStaticShaderPlatform Platform)
-{
-	extern RENDERCORE_API ShaderPlatformMaskType GDistanceFieldsPlatformMask;
-	return (GDistanceFieldsPlatformMask[(int)Platform]);
-}
+RENDERCORE_API bool IsUsingDistanceFields(const FStaticShaderPlatform Platform);
 
 /** Returns if water should render distance field shadow a second time for the water surface. This is for a platofrm so can be used at cook time. */
-inline bool IsWaterDistanceFieldShadowEnabled(const FStaticShaderPlatform Platform)
-{
-	// Only deferred support such a feature. It is not possible to do that for water without a water depth pre-pass.
-	static const auto CVarWaterSingleLayerShaderSupportDistanceFieldShadow = IConsoleManager::Get().FindConsoleVariable(TEXT("r.Water.SingleLayer.ShadersSupportDistanceFieldShadow"));
-	const bool bWaterSingleLayerShaderSupportDistanceFieldShadow = CVarWaterSingleLayerShaderSupportDistanceFieldShadow && (CVarWaterSingleLayerShaderSupportDistanceFieldShadow->GetInt() > 0);
-	return !IsForwardShadingEnabled(Platform) && IsUsingDistanceFields(Platform) && bWaterSingleLayerShaderSupportDistanceFieldShadow;
-}
+RENDERCORE_API bool IsWaterDistanceFieldShadowEnabled(const FStaticShaderPlatform Platform);
 
-inline bool UseGPUScene(const FStaticShaderPlatform Platform, const FStaticFeatureLevel FeatureLevel)
-{
-	if (FeatureLevel == ERHIFeatureLevel::ES3_1)
-	{
-		return MobileSupportsGPUScene();
-	}
-	
-	// GPU Scene management uses compute shaders
-	return FeatureLevel >= ERHIFeatureLevel::SM5 
-		//@todo - support GPU Scene management compute shaders on these platforms to get dynamic instancing speedups on the Rendering Thread and RHI Thread
-		&& !IsOpenGLPlatform(Platform)
-		&& !IsVulkanMobileSM5Platform(Platform)
-        && !IsMetalMobileSM5Platform(Platform)
-		// we only check DDSPI for platforms that have been read in - IsValid() can go away once ALL platforms are converted over to this system
-		&& (!FDataDrivenShaderPlatformInfo::IsValid(Platform) || FDataDrivenShaderPlatformInfo::GetSupportsGPUScene(Platform));
-}
-
-
-inline bool UseGPUScene(const FStaticShaderPlatform Platform)
-{
-	return UseGPUScene(Platform, GetMaxSupportedFeatureLevel(Platform));
-}
+RENDERCORE_API bool UseGPUScene(const FStaticShaderPlatform Platform, const FStaticFeatureLevel FeatureLevel);
 
-inline bool ForceSimpleSkyDiffuse(const FStaticShaderPlatform Platform)
-{
-	extern RENDERCORE_API ShaderPlatformMaskType GSimpleSkyDiffusePlatformMask;
-	return (GSimpleSkyDiffusePlatformMask[(int)Platform]);
-}
+RENDERCORE_API bool UseGPUScene(const FStaticShaderPlatform Platform);
 
-inline bool VelocityEncodeDepth(const FStaticShaderPlatform Platform)
-{
-	extern RENDERCORE_API ShaderPlatformMaskType GVelocityEncodeDepthPlatformMask;
-	return (GVelocityEncodeDepthPlatformMask[(int)Platform]);
-}
+RENDERCORE_API bool ForceSimpleSkyDiffuse(const FStaticShaderPlatform Platform);
+
+RENDERCORE_API bool VelocityEncodeDepth(const FStaticShaderPlatform Platform);
 
 /** Unit cube vertex buffer (VertexDeclarationFVector4) */
 RENDERCORE_API FBufferRHIRef& GetUnitCubeVertexBuffer();
@@ -507,133 +446,55 @@ RENDERCORE_API void QuantizeSceneBufferSize(const FIntPoint& InBufferSize, FIntP
 * Checks if virtual texturing enabled and supported
 * todo: Deprecate the version of the function that takes FStaticFeatureLevel
 */
-RENDERCORE_API bool UseVirtualTexturing(const EShaderPlatform InShaderPlatform, const class ITargetPlatform* TargetPlatform = nullptr);
-RENDERCORE_API bool UseVirtualTexturing(const FStaticFeatureLevel InFeatureLevel, const class ITargetPlatform* TargetPlatform = nullptr);
+RENDERCORE_API bool UseVirtualTexturing(const EShaderPlatform InShaderPlatform, const ITargetPlatform* TargetPlatform = nullptr);
+RENDERCORE_API bool UseVirtualTexturing(const FStaticFeatureLevel InFeatureLevel, const ITargetPlatform* TargetPlatform = nullptr);
 
 RENDERCORE_API bool DoesPlatformSupportNanite(EShaderPlatform Platform, bool bCheckForProjectSetting = true);
 
-inline bool NaniteAtomicsSupported()
-{
-	// Are 64bit image atomics supported by the GPU/Driver/OS/API?
-	bool bAtomicsSupported = GRHISupportsAtomicUInt64;
-
-#if PLATFORM_WINDOWS
-	const ERHIInterfaceType RHIInterface = RHIGetInterfaceType();
-	const bool bIsDx11 = RHIInterface == ERHIInterfaceType::D3D11;
-	const bool bIsDx12 = RHIInterface == ERHIInterfaceType::D3D12;
+RENDERCORE_API bool NaniteAtomicsSupported();
 
-	static const auto NaniteRequireDX12CVar = IConsoleManager::Get().FindConsoleVariable(TEXT("r.Nanite.RequireDX12"));
-	static const uint32 NaniteRequireDX12 = (NaniteRequireDX12CVar != nullptr) ? NaniteRequireDX12CVar->GetInt() : 1;
-	
-	if (bAtomicsSupported && NaniteRequireDX12 != 0)
-	{
-		// Only allow Vulkan or D3D12
-		bAtomicsSupported = !bIsDx11;
-
-		// Disable DX12 vendor extensions unless DX12 SM6.6 is supported
-		if (NaniteRequireDX12 == 1 && bIsDx12 && !GRHISupportsDX12AtomicUInt64)
-		{
-			// Vendor extensions currently support atomic64, but SM 6.6 and the DX12 Agility SDK are reporting that atomics are not supported.
-			// Likely due to a pre-1909 Windows 10 version, or outdated drivers without SM 6.6 support.
-			// See: https://devblogs.microsoft.com/directx/gettingstarted-dx12agility/
-			bAtomicsSupported = false;
-		}
-	}
-#endif
-
-	return bAtomicsSupported;
-}
-
-inline bool DoesRuntimeSupportNanite(EShaderPlatform ShaderPlatform, bool bCheckForAtomicSupport, bool bCheckForProjectSetting)
-{
-	// Does the platform support Nanite?
-	const bool bSupportedPlatform = DoesPlatformSupportNanite(ShaderPlatform, bCheckForProjectSetting);
-
-	// Nanite is not supported with forward shading at this time.
-	const bool bForwardShadingEnabled = IsForwardShadingEnabled(ShaderPlatform);
-
-	return bSupportedPlatform && (!bCheckForAtomicSupport || NaniteAtomicsSupported()) && !bForwardShadingEnabled;
-}
+RENDERCORE_API bool DoesRuntimeSupportNanite(EShaderPlatform ShaderPlatform, bool bCheckForAtomicSupport, bool bCheckForProjectSetting);
 
 /**
  * Returns true if Nanite rendering should be used for the given shader platform.
  */
-inline bool UseNanite(EShaderPlatform ShaderPlatform, bool bCheckForAtomicSupport = true, bool bCheckForProjectSetting = true)
-{
-	static const auto EnableNaniteCVar = IConsoleManager::Get().FindConsoleVariable(TEXT("r.Nanite"));
-	const bool bNaniteEnabled   = (EnableNaniteCVar != nullptr) ? (EnableNaniteCVar->GetInt() != 0) : true;
-	return bNaniteEnabled && DoesRuntimeSupportNanite(ShaderPlatform, bCheckForAtomicSupport, bCheckForProjectSetting);
-}
+RENDERCORE_API bool UseNanite(EShaderPlatform ShaderPlatform, bool bCheckForAtomicSupport = true, bool bCheckForProjectSetting = true);
 
 /**
  * Returns true if Virtual Shadow Maps should be used for the given shader platform.
  * Note: Virtual Shadow Maps require Nanite support.
  */
-inline bool UseVirtualShadowMaps(EShaderPlatform ShaderPlatform, const FStaticFeatureLevel FeatureLevel)
-{
-	static const auto EnableVirtualSMCVar = IConsoleManager::Get().FindConsoleVariable(TEXT("r.Shadow.Virtual.Enable"));
-	const bool bVirtualShadowMapsEnabled = EnableVirtualSMCVar ? (EnableVirtualSMCVar->GetInt() != 0) : false;
-	return bVirtualShadowMapsEnabled && DoesRuntimeSupportNanite(ShaderPlatform, true /* check for atomics */, false /* check project setting */);
-}
+RENDERCORE_API bool UseVirtualShadowMaps(EShaderPlatform ShaderPlatform, const FStaticFeatureLevel FeatureLevel);
 
 /**
 * Returns true if Virtual Shadow Mapsare supported for the given shader platform.
 * Note: Virtual Shadow Maps require Nanite platform support.
 */
-inline bool DoesPlatformSupportVirtualShadowMaps(EShaderPlatform Platform)
-{
-	return DoesPlatformSupportNanite(Platform, false /* check project setting */);
-}
+RENDERCORE_API bool DoesPlatformSupportVirtualShadowMaps(EShaderPlatform Platform);
 
 /**
 * Returns true if non-Nanite virtual shadow maps are enabled by CVar r.Shadow.Virtual.NonNaniteVSM
 * and the runtime supports Nanite/virtual shadow maps.
 */
-inline bool DoesPlatformSupportNonNaniteVirtualShadowMaps(EShaderPlatform ShaderPlatform)
-{
-	static const auto EnableCVar = IConsoleManager::Get().FindConsoleVariable(TEXT("r.Shadow.Virtual.NonNaniteVSM"));	
-	return EnableCVar->GetInt() != 0 && DoesPlatformSupportNanite(ShaderPlatform, false /* check project setting */);
-}
+RENDERCORE_API bool DoesPlatformSupportNonNaniteVirtualShadowMaps(EShaderPlatform ShaderPlatform);
 
 /**
 * Similar to DoesPlatformSupportNonNaniteVirtualShadowMaps, but checks if nanite and virtual shadow maps are enabled (at runtime).
 */
-inline bool UseNonNaniteVirtualShadowMaps(EShaderPlatform ShaderPlatform, const FStaticFeatureLevel FeatureLevel)
-{
-	static const auto EnableCVar = IConsoleManager::Get().FindConsoleVariable(TEXT("r.Shadow.Virtual.NonNaniteVSM"));
-	return EnableCVar->GetInt() != 0 && UseVirtualShadowMaps(ShaderPlatform, FeatureLevel);
-}
+RENDERCORE_API bool UseNonNaniteVirtualShadowMaps(EShaderPlatform ShaderPlatform, FStaticFeatureLevel FeatureLevel);
 
 /** Returns if water should evaluate virtual shadow maps a second time for the water surface. This is for a platform so can be used at cook time. */
-inline bool IsWaterVirtualShadowMapFilteringEnabled(const FStaticShaderPlatform Platform)
-{
-	static const auto CVarWaterSingleLayerShaderSupportVSMFiltering = IConsoleManager::Get().FindConsoleVariable(TEXT("r.Water.SingleLayer.ShadersSupportVSMFiltering"));
-	const bool bWaterVSMFilteringSupported = CVarWaterSingleLayerShaderSupportVSMFiltering && (CVarWaterSingleLayerShaderSupportVSMFiltering->GetInt() > 0);
-
-	const bool bVirtualShadowMapsSupported = DoesPlatformSupportVirtualShadowMaps(Platform);
-
-	return !IsForwardShadingEnabled(Platform) && bWaterVSMFilteringSupported && bVirtualShadowMapsSupported;
-}
+RENDERCORE_API bool IsWaterVirtualShadowMapFilteringEnabled(const FStaticShaderPlatform Platform);
 
 /**
 *	(Non-runtime) Checks if the depth prepass for single layer water is enabled. This also depends on virtual shadow maps to be supported on the platform.
 */
-inline bool IsSingleLayerWaterDepthPrepassEnabled(const FStaticShaderPlatform& Platform, const FStaticFeatureLevel& FeatureLevel)
-{
-	static const auto CVarWaterSingleLayerDepthPrepass = IConsoleManager::Get().FindConsoleVariable(TEXT("r.Water.SingleLayer.DepthPrepass"));
-	const bool bPrepassEnabled = CVarWaterSingleLayerDepthPrepass && CVarWaterSingleLayerDepthPrepass->GetInt() > 0;
-	// Currently VSM is the only feature dependent on the depth prepass which is why we only enable it if VSM could also be enabled.
-	// VSM can be toggled at runtime, but we need a compile time value here, so we fall back to DoesPlatformSupportVirtualShadowMaps() to check if
-	// VSM *could* be enabled.
-	const bool bVSMSupported = DoesPlatformSupportVirtualShadowMaps(Platform);
-
-	return bPrepassEnabled && bVSMSupported;
-}
+RENDERCORE_API bool IsSingleLayerWaterDepthPrepassEnabled(const FStaticShaderPlatform Platform, FStaticFeatureLevel FeatureLevel);
 
 /**
 *	Checks if virtual texturing lightmap enabled and supported
 */
-RENDERCORE_API bool UseVirtualTextureLightmap(const FStaticFeatureLevel InFeatureLevel, const class ITargetPlatform* TargetPlatform = nullptr);
+RENDERCORE_API bool UseVirtualTextureLightmap(const FStaticFeatureLevel InFeatureLevel, const ITargetPlatform* TargetPlatform = nullptr);
 
 /**
 *	Checks if platform uses a Nanite landscape mesh
@@ -650,7 +511,6 @@ RENDERCORE_API bool ExcludeNonPipelinedShaderTypes(EShaderPlatform ShaderPlatfor
  */
 RENDERCORE_API bool AreSkinCacheShadersEnabled(EShaderPlatform Platform);
 
-
 /*
  * Detect (at runtime) if the runtime supports rendering one-pass point light shadows (i.e., cube maps)
  */
@@ -659,54 +519,17 @@ RENDERCORE_API bool DoesRuntimeSupportOnePassPointLightShadows(EShaderPlatform P
 /**
  * Read-only switch to check if translucency per object shadows are enabled.
  */
-RENDERCORE_API bool AllowTranslucencyPerObjectShadows(const FStaticShaderPlatform &Platform);
+RENDERCORE_API bool AllowTranslucencyPerObjectShadows(const FStaticShaderPlatform Platform);
 
 
 /**Note, this should only be used when a platform requires special shader compilation for 32 bit pixel format render targets.
 Does not replace pixel format associations across the board**/
 
-FORCEINLINE bool PlatformRequires128bitRT(EPixelFormat PixelFormat)
-{
-	switch (PixelFormat)
-	{
-	case PF_R32_FLOAT:
-	case PF_G32R32F:
-	case PF_A32B32G32R32F:
-		return FDataDrivenShaderPlatformInfo::GetRequiresExplicit128bitRT(GMaxRHIShaderPlatform);
-	default:
-		return false;
-	}
-}
-
-FORCEINLINE bool IsRayTracingEnabledForProject(EShaderPlatform ShaderPlatform)
-{
-	if (RHISupportsRayTracing(ShaderPlatform))
-	{
-		extern RENDERCORE_API TBitArray<TInlineAllocator<EShaderPlatform::SP_NumPlatforms / 8>> GRayTracingPlatformMask;
-		return (GRayTracingPlatformMask[(int)ShaderPlatform]);
-	}
-	else
-	{
-		return false;
-	}
-}
+RENDERCORE_API bool PlatformRequires128bitRT(EPixelFormat PixelFormat);
 
-FORCEINLINE bool ShouldCompileRayTracingShadersForProject(EShaderPlatform ShaderPlatform)
-{
-	if (RHISupportsRayTracingShaders(ShaderPlatform))
-	{
-		return IsRayTracingEnabledForProject(ShaderPlatform);
-	}
-	else
-	{
-		return false;
-	}
-}
-
-FORCEINLINE bool ShouldCompileRayTracingCallableShadersForProject(EShaderPlatform ShaderPlatform)
-{
-	return RHISupportsRayTracingCallableShaders(ShaderPlatform) && ShouldCompileRayTracingShadersForProject(ShaderPlatform);
-}
+RENDERCORE_API bool IsRayTracingEnabledForProject(EShaderPlatform ShaderPlatform);
+RENDERCORE_API bool ShouldCompileRayTracingShadersForProject(EShaderPlatform ShaderPlatform);
+RENDERCORE_API bool ShouldCompileRayTracingCallableShadersForProject(EShaderPlatform ShaderPlatform);
 
 // Returns `true` when running on RT-capable machine, RT support is enabled for the project and by game graphics options and RT is enabled with r.Raytracing.Enable
 // This function may only be called at runtime, never during cooking.
