Moving FShaderPlatformCachedIniValue into its own file.
Preparing RenderUtils for include removal.

diff --git a/Engine/Plugins/Animation/DeformerGraph/Source/OptimusSettings/Private/OptimusSettingsModule.cpp b/Engine/Plugins/Animation/DeformerGraph/Source/OptimusSettings/Private/OptimusSettingsModule.cpp
--- a/Engine/Plugins/Animation/DeformerGraph/Source/OptimusSettings/Private/OptimusSettingsModule.cpp
+++ b/Engine/Plugins/Animation/DeformerGraph/Source/OptimusSettings/Private/OptimusSettingsModule.cpp
@@ -8,7 +8,7 @@
 #include "Misc/CoreDelegates.h"
 #include "Modules/ModuleManager.h"
 #include "OptimusSettings.h"
-#include "RenderUtils.h"
+#include "ShaderPlatformCachedIniValue.h"
 
 
 namespace Optimus
diff --git a/Engine/Source/Runtime/Engine/Private/GPUSkinCache.cpp b/Engine/Source/Runtime/Engine/Private/GPUSkinCache.cpp
--- a/Engine/Source/Runtime/Engine/Private/GPUSkinCache.cpp
+++ b/Engine/Source/Runtime/Engine/Private/GPUSkinCache.cpp
@@ -24,6 +24,7 @@ GPUSkinCache.cpp: Performs skinning on a compute shader into a buffer to avoid v
 #include "RayTracingSkinnedGeometry.h"
 #include "GPUSkinCacheVisualizationData.h"
 #include "Internationalization/Internationalization.h"
+#include "ShaderPlatformCachedIniValue.h"
 
 DEFINE_STAT(STAT_GPUSkinCache_TotalNumChunks);
 DEFINE_STAT(STAT_GPUSkinCache_TotalNumVertices);
diff --git a/Engine/Source/Runtime/Engine/Private/GPUSkinVertexFactory.cpp b/Engine/Source/Runtime/Engine/Private/GPUSkinVertexFactory.cpp
--- a/Engine/Source/Runtime/Engine/Private/GPUSkinVertexFactory.cpp
+++ b/Engine/Source/Runtime/Engine/Private/GPUSkinVertexFactory.cpp
@@ -17,8 +17,9 @@
 #include "Logging/LogMacros.h"
 #include "Misc/CoreMisc.h"
 #include "RenderGraphResources.h"
-
+#include "ShaderPlatformCachedIniValue.h"
 #include "Engine/RendererSettings.h"
+
 #if INTEL_ISPC
 #include "GPUSkinVertexFactory.ispc.generated.h"
 #endif
diff --git a/Engine/Source/Runtime/Engine/Private/ShaderCompiler/ShaderCompiler.cpp b/Engine/Source/Runtime/Engine/Private/ShaderCompiler/ShaderCompiler.cpp
--- a/Engine/Source/Runtime/Engine/Private/ShaderCompiler/ShaderCompiler.cpp
+++ b/Engine/Source/Runtime/Engine/Private/ShaderCompiler/ShaderCompiler.cpp
@@ -62,6 +62,7 @@
 #include "ShaderCodeLibrary.h"
 #include "ShaderCore.h"
 #include "ShaderParameterMetadata.h"
+#include "ShaderPlatformCachedIniValue.h"
 #include "StaticBoundShaderState.h"
 #include "StaticBoundShaderState.h"
 #include "StereoRenderUtils.h"
diff --git a/Engine/Source/Runtime/RHI/Public/RHIDefinitions.h b/Engine/Source/Runtime/RHI/Public/RHIDefinitions.h
--- a/Engine/Source/Runtime/RHI/Public/RHIDefinitions.h
+++ b/Engine/Source/Runtime/RHI/Public/RHIDefinitions.h
@@ -391,7 +391,7 @@ enum EVertexElementType
 static_assert(VET_MAX <= (1 << VET_NumBits), "VET_MAX will not fit on VET_NumBits");
 DECLARE_INTRINSIC_TYPE_LAYOUT(EVertexElementType);
 
-enum ECubeFace
+enum ECubeFace : uint32
 {
 	CubeFace_PosX = 0,
 	CubeFace_NegX,
diff --git a/Engine/Source/Runtime/RHI/Public/RHIFwd.h b/Engine/Source/Runtime/RHI/Public/RHIFwd.h
--- a/Engine/Source/Runtime/RHI/Public/RHIFwd.h
+++ b/Engine/Source/Runtime/RHI/Public/RHIFwd.h
@@ -8,6 +8,10 @@
 
 namespace ERHIFeatureLevel { enum Type : int; }
 enum EShaderPlatform : uint16;
+enum ECubeFace : uint32;
+
+enum EPixelFormat : uint8;
+enum class EPixelFormatChannelFlags : uint8;
 
 // Command Lists
 class FRHICommandListBase;
diff --git a/Engine/Source/Runtime/RenderCore/Private/RenderUtils.cpp b/Engine/Source/Runtime/RenderCore/Private/RenderUtils.cpp
--- a/Engine/Source/Runtime/RenderCore/Private/RenderUtils.cpp
+++ b/Engine/Source/Runtime/RenderCore/Private/RenderUtils.cpp
@@ -12,6 +12,7 @@
 #include "RenderCore.h"
 #include "PackedNormal.h"
 #include "Shader.h"
+#include "ShaderPlatformCachedIniValue.h"
 
 #if WITH_EDITOR
 #include "Misc/CoreMisc.h"
diff --git a/Engine/Source/Runtime/RenderCore/Private/Shader.cpp b/Engine/Source/Runtime/RenderCore/Private/Shader.cpp
--- a/Engine/Source/Runtime/RenderCore/Private/Shader.cpp
+++ b/Engine/Source/Runtime/RenderCore/Private/Shader.cpp
@@ -27,6 +27,7 @@
 #include "ProfilingDebugging/LoadTimeTracker.h"
 #include "Misc/LargeWorldRenderPosition.h"
 #include "DataDrivenShaderPlatformInfo.h"
+#include "ShaderPlatformCachedIniValue.h"
 
 #if WITH_EDITORONLY_DATA
 #include "Interfaces/IShaderFormat.h"
diff --git a/Engine/Source/Runtime/RenderCore/Private/StereoRenderUtils.cpp b/Engine/Source/Runtime/RenderCore/Private/StereoRenderUtils.cpp
--- a/Engine/Source/Runtime/RenderCore/Private/StereoRenderUtils.cpp
+++ b/Engine/Source/Runtime/RenderCore/Private/StereoRenderUtils.cpp
@@ -2,6 +2,7 @@
 
 #include "StereoRenderUtils.h"
 #include "RenderUtils.h"
+#include "ShaderPlatformCachedIniValue.h"
 
 // enable this to printf-debug stereo rendering aspects on a device
 #define UE_DEBUG_STEREO_ASPECTS			(!UE_BUILD_SHIPPING && !UE_BUILD_TEST && !UE_EDITOR)
diff --git a/Engine/Source/Runtime/RenderCore/Public/RenderUtils.h b/Engine/Source/Runtime/RenderCore/Public/RenderUtils.h
--- a/Engine/Source/Runtime/RenderCore/Public/RenderUtils.h
+++ b/Engine/Source/Runtime/RenderCore/Public/RenderUtils.h
@@ -3,6 +3,10 @@
 #pragma once
 
 #include "CoreMinimal.h"
+#include "RHIFwd.h"
+#include "RHIShaderPlatform.h"
+#include "RHIFeatureLevel.h"
+
 #include "RHI.h"
 #include "RenderResource.h"
 #include "RHIDefinitions.h"
@@ -10,14 +14,12 @@
 #include "PackedNormal.h"
 #include "RenderMath.h"
 #include "GlobalRenderResources.h"
+#include "ShaderPlatformCachedIniValue.h"
 #endif
 
-class FBufferWithRDG;
-class FRDGBuffer;
-class FRDGTexture;
-class FRDGBuilder;
 class ITargetPlatform;
 struct IPooledRenderTarget;
+enum class ERayTracingMode : uint8;
 
 extern RENDERCORE_API void RenderUtilsInit();
 
@@ -327,77 +329,6 @@ RENDERCORE_API bool PlatformSupportsVelocityRendering(const FStaticShaderPlatfor
 
 RENDERCORE_API bool IsUsingDBuffers(const FStaticShaderPlatform Platform);
 
-RENDERCORE_API EShaderPlatform GetEditorShaderPlatform(EShaderPlatform ShaderPlatform);
-
-template<typename Type>
-struct RENDERCORE_API FShaderPlatformCachedIniValue
-{
-	FShaderPlatformCachedIniValue(const TCHAR* InCVarName)
-		: CVarName(InCVarName)
-		, CVar(nullptr)
-	{
-	}
-
-	FShaderPlatformCachedIniValue(IConsoleVariable* InCVar)
-		: CVar(InCVar)
-	{
-	}
-
-	Type Get(EShaderPlatform ShaderPlatform)
-	{
-		Type Value{};
-
-		const EShaderPlatform ActualShaderPlatform = GetEditorShaderPlatform(ShaderPlatform);
-
-		FName IniPlatformName = ShaderPlatformToPlatformName(ActualShaderPlatform);
-		// find the cvar if needed
-		if (CVar == nullptr)
-		{
-			CVar = IConsoleManager::Get().FindConsoleVariable(*CVarName);
-		}
-
-		// if we are looking up our own platform, just use the current value, however
-		// ShaderPlatformToPlatformName can return the wrong platform than expected - for instance, Linux Vulkan will return Windows
-		// so instead of hitting an asser below, we detect that the request SP is the current SP, and use the CVar value that is set currently
-		if (IniPlatformName == FPlatformProperties::IniPlatformName() || ActualShaderPlatform == GMaxRHIShaderPlatform)
-		{
-			checkf(CVar != nullptr, TEXT("Failed to find CVar %s when getting current value for FShaderPlatformCachedIniValue"));
-
-			CVar->GetValue(Value);
-			return Value;
-		}
-
-#if ALLOW_OTHER_PLATFORM_CONFIG
-		// create a dummy cvar if needed
-		if (CVar == nullptr)
-		{
-			// this could be a cvar that only exists on the target platform so create a dummy one
-			CVar = IConsoleManager::Get().RegisterConsoleVariable(*CVarName, Type(), TEXT(""), ECVF_ReadOnly);
-		}
-
-		// now get the value from the platform that makes sense for this shader platform
-		TSharedPtr<IConsoleVariable> OtherPlatformVar = CVar->GetPlatformValueVariable(IniPlatformName);
-		ensureMsgf(OtherPlatformVar.IsValid(), TEXT("Failed to get another platform's version of a cvar (possible name: '%s'). It is probably an esoteric subclass that needs to implement GetPlatformValueVariable."), *CVarName);
-		if (OtherPlatformVar.IsValid())
-		{
-			OtherPlatformVar->GetValue(Value);
-		}
-		else
-		{
-			// get this platform's value, even tho it could be wrong
-			CVar->GetValue(Value);
-		}
-#else
-		checkf(IniPlatformName == FName(FPlatformProperties::IniPlatformName()), TEXT("FShaderPlatformCachedIniValue can only look up the current platform when ALLOW_OTHER_PLATFORM_CONFIG is false"));
-#endif
-		return Value;
-	}
-
-private:
-	FString CVarName;
-	IConsoleVariable* CVar;
-};
-
 /** Returns if ForwardShading is enabled. Only valid for the current platform (otherwise call ITargetPlatform::UsesForwardShading()). */
 RENDERCORE_API bool IsForwardShadingEnabled(const FStaticShaderPlatform Platform);
 
diff --git a/Engine/Source/Runtime/RenderCore/Public/ShaderPlatformCachedIniValue.h b/Engine/Source/Runtime/RenderCore/Public/ShaderPlatformCachedIniValue.h
new file mode 100644
--- /dev/null
+++ b/Engine/Source/Runtime/RenderCore/Public/ShaderPlatformCachedIniValue.h
@@ -0,0 +1,78 @@
+// Copyright Epic Games, Inc. All Rights Reserved.
+
+#pragma once
+
+#include "CoreMinimal.h"
+#include "HAL/IConsoleManager.h"
+#include "RHIShaderPlatform.h"
+
+RENDERCORE_API EShaderPlatform GetEditorShaderPlatform(EShaderPlatform ShaderPlatform);
+
+template<typename Type>
+struct RENDERCORE_API FShaderPlatformCachedIniValue
+{
+	FShaderPlatformCachedIniValue(const TCHAR* InCVarName)
+		: CVarName(InCVarName)
+		, CVar(nullptr)
+	{
+	}
+
+	FShaderPlatformCachedIniValue(IConsoleVariable* InCVar)
+		: CVar(InCVar)
+	{
+	}
+
+	Type Get(EShaderPlatform ShaderPlatform)
+	{
+		Type Value{};
+
+		const EShaderPlatform ActualShaderPlatform = GetEditorShaderPlatform(ShaderPlatform);
+
+		FName IniPlatformName = ShaderPlatformToPlatformName(ActualShaderPlatform);
+		// find the cvar if needed
+		if (CVar == nullptr)
+		{
+			CVar = IConsoleManager::Get().FindConsoleVariable(*CVarName);
+		}
+
+		// if we are looking up our own platform, just use the current value, however
+		// ShaderPlatformToPlatformName can return the wrong platform than expected - for instance, Linux Vulkan will return Windows
+		// so instead of hitting an asser below, we detect that the request SP is the current SP, and use the CVar value that is set currently
+		if (IniPlatformName == FPlatformProperties::IniPlatformName() || ActualShaderPlatform == GMaxRHIShaderPlatform)
+		{
+			checkf(CVar != nullptr, TEXT("Failed to find CVar %s when getting current value for FShaderPlatformCachedIniValue"));
+
+			CVar->GetValue(Value);
+			return Value;
+		}
+
+#if ALLOW_OTHER_PLATFORM_CONFIG
+		// create a dummy cvar if needed
+		if (CVar == nullptr)
+		{
+			// this could be a cvar that only exists on the target platform so create a dummy one
+			CVar = IConsoleManager::Get().RegisterConsoleVariable(*CVarName, Type(), TEXT(""), ECVF_ReadOnly);
+		}
+
+		// now get the value from the platform that makes sense for this shader platform
+		TSharedPtr<IConsoleVariable> OtherPlatformVar = CVar->GetPlatformValueVariable(IniPlatformName);
+		ensureMsgf(OtherPlatformVar.IsValid(), TEXT("Failed to get another platform's version of a cvar (possible name: '%s'). It is probably an esoteric subclass that needs to implement GetPlatformValueVariable."), *CVarName);
+		if (OtherPlatformVar.IsValid())
+		{
+			OtherPlatformVar->GetValue(Value);
+		}
+		else
+		{
+			// get this platform's value, even tho it could be wrong
+			CVar->GetValue(Value);
+		}
+#else
+		checkf(IniPlatformName == FName(FPlatformProperties::IniPlatformName()), TEXT("FShaderPlatformCachedIniValue can only look up the current platform when ALLOW_OTHER_PLATFORM_CONFIG is false"));
+#endif
+		return Value;
+	}
+
+private:
+	FString CVarName;
+	IConsoleVariable* CVar;
+};
diff --git a/Engine/Source/Runtime/Renderer/Private/PostProcess/TemporalAA.cpp b/Engine/Source/Runtime/Renderer/Private/PostProcess/TemporalAA.cpp
--- a/Engine/Source/Runtime/Renderer/Private/PostProcess/TemporalAA.cpp
+++ b/Engine/Source/Runtime/Renderer/Private/PostProcess/TemporalAA.cpp
@@ -10,6 +10,7 @@
 #include "PixelShaderUtils.h"
 #include "ScenePrivate.h"
 #include "RendererModule.h"
+#include "ShaderPlatformCachedIniValue.h"
 
 namespace
 {
diff --git a/Engine/Source/Runtime/Renderer/Private/ShadowDepthRendering.cpp b/Engine/Source/Runtime/Renderer/Private/ShadowDepthRendering.cpp
--- a/Engine/Source/Runtime/Renderer/Private/ShadowDepthRendering.cpp
+++ b/Engine/Source/Runtime/Renderer/Private/ShadowDepthRendering.cpp
@@ -48,6 +48,7 @@
 #include "Rendering/NaniteStreamingManager.h"
 #include "Shadows/ShadowSceneRenderer.h"
 #include "RenderCore.h"
+#include "ShaderPlatformCachedIniValue.h"
 
 DEFINE_GPU_DRAWCALL_STAT(ShadowDepths);
 
diff --git a/Engine/Source/Runtime/Renderer/Private/VT/RuntimeVirtualTextureRender.cpp b/Engine/Source/Runtime/Renderer/Private/VT/RuntimeVirtualTextureRender.cpp
--- a/Engine/Source/Runtime/Renderer/Private/VT/RuntimeVirtualTextureRender.cpp
+++ b/Engine/Source/Runtime/Renderer/Private/VT/RuntimeVirtualTextureRender.cpp
@@ -10,6 +10,7 @@
 #include "RenderGraphBuilder.h"
 #include "RenderGraphUtils.h"
 #include "RenderUtils.h"
+#include "ShaderPlatformCachedIniValue.h"
 #include "ScenePrivate.h"
 #include "PostProcess/SceneRenderTargets.h"
 #include "ShaderBaseClasses.h"
