Removing more includes from RenderUtils.h

diff --git a/Engine/Plugins/Enterprise/GLTFExporter/Source/GLTFExporter/Public/Utilities/GLTFTextureUtilities.h b/Engine/Plugins/Enterprise/GLTFExporter/Source/GLTFExporter/Public/Utilities/GLTFTextureUtilities.h
--- a/Engine/Plugins/Enterprise/GLTFExporter/Source/GLTFExporter/Public/Utilities/GLTFTextureUtilities.h
+++ b/Engine/Plugins/Enterprise/GLTFExporter/Source/GLTFExporter/Public/Utilities/GLTFTextureUtilities.h
@@ -3,7 +3,13 @@
 #pragma once
 
 #include "CoreMinimal.h"
+#if UE_ENABLE_INCLUDE_ORDER_DEPRECATED_IN_5_2
 #include "Engine/TextureRenderTarget2D.h"
+#endif
+
+class FTexture;
+class UTexture;
+class UTextureRenderTarget2D;
 
 class GLTFEXPORTER_API FGLTFTextureUtilities
 {
diff --git a/Engine/Source/Editor/UnrealEd/Private/Factories/TTFontImport.cpp b/Engine/Source/Editor/UnrealEd/Private/Factories/TTFontImport.cpp
--- a/Engine/Source/Editor/UnrealEd/Private/Factories/TTFontImport.cpp
+++ b/Engine/Source/Editor/UnrealEd/Private/Factories/TTFontImport.cpp
@@ -16,7 +16,7 @@
 #include "Engine/FontImportOptions.h"
 #include "Engine/Font.h"
 #include "Factories/TrueTypeFontFactory.h"
-#include "RenderUtils.h"
+#include "RHI.h"
 #include "Engine/Texture2D.h"
 #include "Interfaces/IMainFrameModule.h"
 #include "IDesktopPlatform.h"
diff --git a/Engine/Source/Runtime/Engine/Private/Debug/ReporterBase.cpp b/Engine/Source/Runtime/Engine/Private/Debug/ReporterBase.cpp
--- a/Engine/Source/Runtime/Engine/Private/Debug/ReporterBase.cpp
+++ b/Engine/Source/Runtime/Engine/Private/Debug/ReporterBase.cpp
@@ -3,7 +3,7 @@
 #include "Debug/ReporterBase.h"
 #include "CanvasItem.h"
 #include "Engine/Canvas.h"
-#include "RenderUtils.h"
+#include "GlobalRenderResources.h"
 
 #include UE_INLINE_GENERATED_CPP_BY_NAME(ReporterBase)
 
@@ -13,9 +13,6 @@ UReporterBase::UReporterBase(const FObjectInitializer& ObjectInitializer) :
 	Super(ObjectInitializer)
 {
 	bVisible = false;
-	
-
-
 }
 
 FVector2D UReporterBase::ToScreenSpace(const FVector2D& InVector, UCanvas* Canvas)
diff --git a/Engine/Source/Runtime/Engine/Private/Debug/ReporterGraph.cpp b/Engine/Source/Runtime/Engine/Private/Debug/ReporterGraph.cpp
--- a/Engine/Source/Runtime/Engine/Private/Debug/ReporterGraph.cpp
+++ b/Engine/Source/Runtime/Engine/Private/Debug/ReporterGraph.cpp
@@ -5,7 +5,7 @@
 #include "Engine/Engine.h"
 #include "CanvasItem.h"
 #include "Engine/Canvas.h"
-#include "RenderUtils.h"
+#include "GlobalRenderResources.h"
 
 #include UE_INLINE_GENERATED_CPP_BY_NAME(ReporterGraph)
 
diff --git a/Engine/Source/Runtime/Engine/Private/SceneUtils.cpp b/Engine/Source/Runtime/Engine/Private/SceneUtils.cpp
--- a/Engine/Source/Runtime/Engine/Private/SceneUtils.cpp
+++ b/Engine/Source/Runtime/Engine/Private/SceneUtils.cpp
@@ -4,6 +4,7 @@
 #include "RHI.h"
 #include "Engine/Scene.h"
 #include "RenderUtils.h"
+#include "GlobalRenderResources.h"
 
 DEFINE_LOG_CATEGORY(LogSceneUtils);
 
diff --git a/Engine/Source/Runtime/Engine/Private/UserInterface/PlayerInput.cpp b/Engine/Source/Runtime/Engine/Private/UserInterface/PlayerInput.cpp
--- a/Engine/Source/Runtime/Engine/Private/UserInterface/PlayerInput.cpp
+++ b/Engine/Source/Runtime/Engine/Private/UserInterface/PlayerInput.cpp
@@ -8,7 +8,7 @@
 #include "Misc/CommandLine.h"
 #include "Components/InputComponent.h"
 #include "Misc/App.h"
-#include "RenderUtils.h"
+#include "GlobalRenderResources.h"
 #include "UObject/UObjectHash.h"
 #include "UObject/UObjectIterator.h"
 #include "GameFramework/PlayerController.h"
diff --git a/Engine/Source/Runtime/Engine/Private/WorldPartition/WorldPartitionDebugHelper.cpp b/Engine/Source/Runtime/Engine/Private/WorldPartition/WorldPartitionDebugHelper.cpp
--- a/Engine/Source/Runtime/Engine/Private/WorldPartition/WorldPartitionDebugHelper.cpp
+++ b/Engine/Source/Runtime/Engine/Private/WorldPartition/WorldPartitionDebugHelper.cpp
@@ -2,7 +2,7 @@
 
 #include "WorldPartition/WorldPartitionDebugHelper.h"
 #include "Engine/Engine.h"
-#include "RenderUtils.h"
+#include "GlobalRenderResources.h"
 #include "WorldPartition/WorldPartitionRuntimeHash.h"
 #include "WorldPartition/WorldPartition.h"
 #include "WorldPartition/DataLayer/DataLayerInstance.h"
diff --git a/Engine/Source/Runtime/Engine/Private/WorldPartition/WorldPartitionRuntimeSpatialHash.cpp b/Engine/Source/Runtime/Engine/Private/WorldPartition/WorldPartitionRuntimeSpatialHash.cpp
--- a/Engine/Source/Runtime/Engine/Private/WorldPartition/WorldPartitionRuntimeSpatialHash.cpp
+++ b/Engine/Source/Runtime/Engine/Private/WorldPartition/WorldPartitionRuntimeSpatialHash.cpp
@@ -33,7 +33,7 @@
 #include "Engine/Canvas.h"
 #include "DrawDebugHelpers.h"
 #include "DisplayDebugHelpers.h"
-#include "RenderUtils.h"
+#include "GlobalRenderResources.h"
 #include "Algo/Transform.h"
 #include "Algo/Copy.h"
 #include "Algo/Accumulate.h"
diff --git a/Engine/Source/Runtime/RenderCore/Private/StereoRenderUtils.cpp b/Engine/Source/Runtime/RenderCore/Private/StereoRenderUtils.cpp
--- a/Engine/Source/Runtime/RenderCore/Private/StereoRenderUtils.cpp
+++ b/Engine/Source/Runtime/RenderCore/Private/StereoRenderUtils.cpp
@@ -1,8 +1,9 @@
 // Copyright Epic Games, Inc. All Rights Reserved.
 
 #include "StereoRenderUtils.h"
-#include "RenderUtils.h"
 #include "ShaderPlatformCachedIniValue.h"
+#include "DataDrivenShaderPlatformInfo.h"
+#include "RHI.h"
 
 // enable this to printf-debug stereo rendering aspects on a device
 #define UE_DEBUG_STEREO_ASPECTS			(!UE_BUILD_SHIPPING && !UE_BUILD_TEST && !UE_EDITOR)
diff --git a/Engine/Source/Runtime/RenderCore/Public/RenderUtils.h b/Engine/Source/Runtime/RenderCore/Public/RenderUtils.h
--- a/Engine/Source/Runtime/RenderCore/Public/RenderUtils.h
+++ b/Engine/Source/Runtime/RenderCore/Public/RenderUtils.h
@@ -7,10 +7,10 @@
 #include "RHIShaderPlatform.h"
 #include "RHIFeatureLevel.h"
 
+#if UE_ENABLE_INCLUDE_ORDER_DEPRECATED_IN_5_2
 #include "RHI.h"
 #include "RenderResource.h"
 #include "RHIDefinitions.h"
-#if UE_ENABLE_INCLUDE_ORDER_DEPRECATED_IN_5_2
 #include "PackedNormal.h"
 #include "RenderMath.h"
 #include "GlobalRenderResources.h"
@@ -361,10 +361,8 @@ RENDERCORE_API FBufferRHIRef& GetUnitCubeVertexBuffer();
 /** Unit cube index buffer */
 RENDERCORE_API FBufferRHIRef& GetUnitCubeIndexBuffer();
 
-#if RHI_RAYTRACING
 /** Unit cube AABB vertex buffer (useful to create procedural raytracing geometry) */
 RENDERCORE_API FBufferRHIRef& GetUnitCubeAABBVertexBuffer();
-#endif
 
 /**
 * Takes the requested buffer size and quantizes it to an appropriate size for the rest of the
