Check for both GameThread & ParallelGameThread in render resource set owner name function

diff --git a/Engine/Source/Runtime/RenderCore/Private/RenderResource.cpp b/Engine/Source/Runtime/RenderCore/Private/RenderResource.cpp
--- a/Engine/Source/Runtime/RenderCore/Private/RenderResource.cpp
+++ b/Engine/Source/Runtime/RenderCore/Private/RenderResource.cpp
@@ -348,6 +348,7 @@ FBufferRHIRef FRenderResource::CreateRHIBufferInternal(
 void FRenderResource::SetOwnerName(const FName& InOwnerName)
 {
 #if RHI_ENABLE_RESOURCE_INFO
+	check(IsInGameThread() || IsInParallelGameThread());
 	OwnerName = InOwnerName;
 #endif
 }
