PR #9871: Fix missing template argument (Contributed by [REDACTED_PERSON])

diff --git a/Engine/Source/Runtime/RenderCore/Public/RenderingThread.h b/Engine/Source/Runtime/RenderCore/Public/RenderingThread.h
--- a/Engine/Source/Runtime/RenderCore/Public/RenderingThread.h
+++ b/Engine/Source/Runtime/RenderCore/Public/RenderingThread.h
@@ -290,7 +290,7 @@ FORCEINLINE_DEBUGGABLE void EnqueueUniqueRenderCommand(LAMBDA&& Lambda)
 	}; \
 	EnqueueUniqueRenderCommand<Type##Name>
 
-template<typename LAMBDA>
+template<typename TSTR, typename LAMBDA>
 FORCEINLINE_DEBUGGABLE void EnqueueUniqueRenderCommand(LAMBDA& Lambda)
 {
 	static_assert(sizeof(LAMBDA) == 0, "EnqueueUniqueRenderCommand enforces use of rvalue and therefore move to avoid an extra copy of the Lambda");
@@ -540,4 +540,4 @@ using FStopRenderingThreadDelegate = FStopRenderingThread::FDelegate;
 
 extern RENDERCORE_API FDelegateHandle RegisterStopRenderingThreadDelegate(const FStopRenderingThreadDelegate& InDelegate);
 
-extern RENDERCORE_API void UnregisterStopRenderingThreadDelegate(FDelegateHandle InDelegateHandle);
\ No newline at end of file
+extern RENDERCORE_API void UnregisterStopRenderingThreadDelegate(FDelegateHandle InDelegateHandle);
