Moving TGlobalResource back into RenderResource.h as it makes more sense there.
GlobalRenderResources.h is no longer an implicit include from RenderResource.h

diff --git a/Engine/Plugins/2D/Paper2D/Source/Paper2DEditor/Private/PaperEditorViewportClient.cpp b/Engine/Plugins/2D/Paper2D/Source/Paper2DEditor/Private/PaperEditorViewportClient.cpp
--- a/Engine/Plugins/2D/Paper2D/Source/Paper2DEditor/Private/PaperEditorViewportClient.cpp
+++ b/Engine/Plugins/2D/Paper2D/Source/Paper2DEditor/Private/PaperEditorViewportClient.cpp
@@ -6,6 +6,7 @@
 #include "CanvasTypes.h"
 #include "Engine/Texture2D.h"
 #include "ImageUtils.h"
+#include "GlobalRenderResources.h"
 
 //////////////////////////////////////////////////////////////////////////
 // FPaperEditorViewportClient
diff --git a/Engine/Plugins/2D/Paper2D/Source/Paper2DEditor/Private/PaperTileSetThumbnailRenderer.cpp b/Engine/Plugins/2D/Paper2D/Source/Paper2DEditor/Private/PaperTileSetThumbnailRenderer.cpp
--- a/Engine/Plugins/2D/Paper2D/Source/Paper2DEditor/Private/PaperTileSetThumbnailRenderer.cpp
+++ b/Engine/Plugins/2D/Paper2D/Source/Paper2DEditor/Private/PaperTileSetThumbnailRenderer.cpp
@@ -7,6 +7,7 @@
 #include "PaperTileSet.h"
 #include "RenderUtils.h"
 #include "TextureResource.h"
+#include "GlobalRenderResources.h"
 
 #include UE_INLINE_GENERATED_CPP_BY_NAME(PaperTileSetThumbnailRenderer)
 
diff --git a/Engine/Plugins/2D/Paper2D/Source/Paper2DEditor/Private/SpriteEditor/SpriteEditorViewportClient.cpp b/Engine/Plugins/2D/Paper2D/Source/Paper2DEditor/Private/SpriteEditor/SpriteEditorViewportClient.cpp
--- a/Engine/Plugins/2D/Paper2D/Source/Paper2DEditor/Private/SpriteEditor/SpriteEditorViewportClient.cpp
+++ b/Engine/Plugins/2D/Paper2D/Source/Paper2DEditor/Private/SpriteEditor/SpriteEditorViewportClient.cpp
@@ -8,7 +8,7 @@
 #include "PaperSpriteComponent.h"
 #include "AssetEditorModeManager.h"
 #include "ScopedTransaction.h"
-
+#include "GlobalRenderResources.h"
 #include "ComponentReregisterContext.h"
 
 #include "IAssetTools.h"
diff --git a/Engine/Plugins/2D/Paper2D/Source/Paper2DEditor/Private/TileSetEditor/TileSetEditorViewportClient.cpp b/Engine/Plugins/2D/Paper2D/Source/Paper2DEditor/Private/TileSetEditor/TileSetEditorViewportClient.cpp
--- a/Engine/Plugins/2D/Paper2D/Source/Paper2DEditor/Private/TileSetEditor/TileSetEditorViewportClient.cpp
+++ b/Engine/Plugins/2D/Paper2D/Source/Paper2DEditor/Private/TileSetEditor/TileSetEditorViewportClient.cpp
@@ -9,6 +9,7 @@
 #include "EngineGlobals.h"
 #include "TextureCompiler.h"
 #include "TextureResource.h"
+#include "GlobalRenderResources.h"
 
 #define LOCTEXT_NAMESPACE "TileSetEditor"
 
diff --git a/Engine/Plugins/Experimental/VirtualHeightfieldMesh/Source/VirtualHeightfieldMesh/Private/VirtualHeightfieldMeshSceneProxy.cpp b/Engine/Plugins/Experimental/VirtualHeightfieldMesh/Source/VirtualHeightfieldMesh/Private/VirtualHeightfieldMeshSceneProxy.cpp
--- a/Engine/Plugins/Experimental/VirtualHeightfieldMesh/Source/VirtualHeightfieldMesh/Private/VirtualHeightfieldMeshSceneProxy.cpp
+++ b/Engine/Plugins/Experimental/VirtualHeightfieldMesh/Source/VirtualHeightfieldMesh/Private/VirtualHeightfieldMeshSceneProxy.cpp
@@ -6,6 +6,7 @@
 #include "EngineModule.h"
 #include "Engine/Engine.h"
 #include "Engine/Texture2D.h"
+#include "GlobalRenderResources.h"
 #include "GlobalShader.h"
 #include "HAL/IConsoleManager.h"
 #include "HeightfieldMinMaxTexture.h"
diff --git a/Engine/Plugins/FX/Niagara/Source/Niagara/Private/DataInterface/NiagaraDataInterfaceVirtualTexture.cpp b/Engine/Plugins/FX/Niagara/Source/Niagara/Private/DataInterface/NiagaraDataInterfaceVirtualTexture.cpp
--- a/Engine/Plugins/FX/Niagara/Source/Niagara/Private/DataInterface/NiagaraDataInterfaceVirtualTexture.cpp
+++ b/Engine/Plugins/FX/Niagara/Source/Niagara/Private/DataInterface/NiagaraDataInterfaceVirtualTexture.cpp
@@ -8,6 +8,7 @@
 #include "NiagaraShaderParametersBuilder.h"
 #include "NiagaraSystemInstance.h"
 
+#include "GlobalRenderResources.h"
 #include "Misc/LargeWorldRenderPosition.h"
 #include "RHIStaticStates.h"
 #include "ShaderCompilerCore.h"
diff --git a/Engine/Plugins/FX/Niagara/Source/Niagara/Private/NiagaraDataInterfaceVectorField.cpp b/Engine/Plugins/FX/Niagara/Source/Niagara/Private/NiagaraDataInterfaceVectorField.cpp
--- a/Engine/Plugins/FX/Niagara/Source/Niagara/Private/NiagaraDataInterfaceVectorField.cpp
+++ b/Engine/Plugins/FX/Niagara/Source/Niagara/Private/NiagaraDataInterfaceVectorField.cpp
@@ -11,6 +11,7 @@
 #include "NiagaraComponent.h"
 #include "NiagaraShaderParametersBuilder.h"
 #include "RHIStaticStates.h"
+#include "GlobalRenderResources.h"
 
 #include UE_INLINE_GENERATED_CPP_BY_NAME(NiagaraDataInterfaceVectorField)
 
diff --git a/Engine/Plugins/FX/Niagara/Source/Niagara/Private/NiagaraDataInterfaceVolumeCache.cpp b/Engine/Plugins/FX/Niagara/Source/Niagara/Private/NiagaraDataInterfaceVolumeCache.cpp
--- a/Engine/Plugins/FX/Niagara/Source/Niagara/Private/NiagaraDataInterfaceVolumeCache.cpp
+++ b/Engine/Plugins/FX/Niagara/Source/Niagara/Private/NiagaraDataInterfaceVolumeCache.cpp
@@ -13,6 +13,7 @@
 #include "NiagaraComputeExecutionContext.h"
 #include "NiagaraShaderParametersBuilder.h"
 #include "RHIStaticStates.h"
+#include "GlobalRenderResources.h"
 
 #include UE_INLINE_GENERATED_CPP_BY_NAME(NiagaraDataInterfaceVolumeCache)
 
diff --git a/Engine/Plugins/FX/Niagara/Source/Niagara/Private/NiagaraRendererRibbons.cpp b/Engine/Plugins/FX/Niagara/Source/Niagara/Private/NiagaraRendererRibbons.cpp
--- a/Engine/Plugins/FX/Niagara/Source/Niagara/Private/NiagaraRendererRibbons.cpp
+++ b/Engine/Plugins/FX/Niagara/Source/Niagara/Private/NiagaraRendererRibbons.cpp
@@ -2,6 +2,7 @@
 
 #include "NiagaraRendererRibbons.h"
 
+#include "GlobalRenderResources.h"
 #include "GPUSortManager.h"
 #include "ParticleResources.h"
 #include "NiagaraRibbonVertexFactory.h"
diff --git a/Engine/Plugins/FX/Niagara/Source/NiagaraShader/Private/NiagaraBatchedElements.cpp b/Engine/Plugins/FX/Niagara/Source/NiagaraShader/Private/NiagaraBatchedElements.cpp
--- a/Engine/Plugins/FX/Niagara/Source/NiagaraShader/Private/NiagaraBatchedElements.cpp
+++ b/Engine/Plugins/FX/Niagara/Source/NiagaraShader/Private/NiagaraBatchedElements.cpp
@@ -1,6 +1,7 @@
 // Copyright Epic Games, Inc. All Rights Reserved.
 
 #include "NiagaraBatchedElements.h"
+#include "GlobalRenderResources.h"
 #include "GlobalShader.h"
 #include "RHIStaticStates.h"
 #include "SimpleElementShaders.h"
diff --git a/Engine/Plugins/FX/Niagara/Source/NiagaraVertexFactories/Private/NiagaraMeshVertexFactory.cpp b/Engine/Plugins/FX/Niagara/Source/NiagaraVertexFactories/Private/NiagaraMeshVertexFactory.cpp
--- a/Engine/Plugins/FX/Niagara/Source/NiagaraVertexFactories/Private/NiagaraMeshVertexFactory.cpp
+++ b/Engine/Plugins/FX/Niagara/Source/NiagaraVertexFactories/Private/NiagaraMeshVertexFactory.cpp
@@ -9,6 +9,7 @@
 #include "ParticleResources.h"
 #include "ShaderParameterUtils.h"
 #include "MeshMaterialShader.h"
+#include "GlobalRenderResources.h"
 
 IMPLEMENT_GLOBAL_SHADER_PARAMETER_STRUCT(FNiagaraMeshUniformParameters, "NiagaraMeshVF");
 
diff --git a/Engine/Plugins/Runtime/HairStrands/Source/HairStrandsCore/Private/HairCardsVertexFactory.cpp b/Engine/Plugins/Runtime/HairStrands/Source/HairStrandsCore/Private/HairCardsVertexFactory.cpp
--- a/Engine/Plugins/Runtime/HairStrands/Source/HairStrandsCore/Private/HairCardsVertexFactory.cpp
+++ b/Engine/Plugins/Runtime/HairStrands/Source/HairStrandsCore/Private/HairCardsVertexFactory.cpp
@@ -14,6 +14,7 @@
 #include "HairStrandsInterface.h"
 #include "GroomInstance.h"
 #include "SystemTextures.h" 
+#include "GlobalRenderResources.h"
 
 template<typename T> inline void VFC_BindParam(FMeshDrawSingleShaderBindings& ShaderBindings, const FShaderResourceParameter& Param, T* Value) { if (Param.IsBound() && Value) ShaderBindings.Add(Param, Value); }
 template<typename T> inline void VFC_BindParam(FMeshDrawSingleShaderBindings& ShaderBindings, const FShaderParameter& Param, const T& Value) { if (Param.IsBound()) ShaderBindings.Add(Param, Value); }
diff --git a/Engine/Plugins/Runtime/nDisplay/Source/DisplayCluster/Private/Game/EngineClasses/Basics/DisplayClusterViewportClient.cpp b/Engine/Plugins/Runtime/nDisplay/Source/DisplayCluster/Private/Game/EngineClasses/Basics/DisplayClusterViewportClient.cpp
--- a/Engine/Plugins/Runtime/nDisplay/Source/DisplayCluster/Private/Game/EngineClasses/Basics/DisplayClusterViewportClient.cpp
+++ b/Engine/Plugins/Runtime/nDisplay/Source/DisplayCluster/Private/Game/EngineClasses/Basics/DisplayClusterViewportClient.cpp
@@ -23,6 +23,7 @@
 #include "Debug/DebugDrawService.h"
 #include "ContentStreaming.h"
 #include "EngineModule.h"
+#include "GlobalRenderResources.h"
 #include "FXSystem.h"
 #include "GameFramework/HUD.h"
 #include "SubtitleManager.h"
diff --git a/Engine/Source/Editor/UMGEditor/Private/WidgetBlueprintThumbnailRenderer.cpp b/Engine/Source/Editor/UMGEditor/Private/WidgetBlueprintThumbnailRenderer.cpp
--- a/Engine/Source/Editor/UMGEditor/Private/WidgetBlueprintThumbnailRenderer.cpp
+++ b/Engine/Source/Editor/UMGEditor/Private/WidgetBlueprintThumbnailRenderer.cpp
@@ -7,6 +7,7 @@
 #include "CanvasTypes.h"
 #include "Engine/Texture2D.h"
 #include "Engine/TextureRenderTarget2D.h"
+#include "GlobalRenderResources.h"
 #include "Input/HittestGrid.h"
 #include "Kismet2/KismetEditorUtilities.h"
 #include "Slate/WidgetRenderer.h"
diff --git a/Engine/Source/Editor/UnrealEd/Private/EditorViewportClient.cpp b/Engine/Source/Editor/UnrealEd/Private/EditorViewportClient.cpp
--- a/Engine/Source/Editor/UnrealEd/Private/EditorViewportClient.cpp
+++ b/Engine/Source/Editor/UnrealEd/Private/EditorViewportClient.cpp
@@ -67,6 +67,7 @@
 #include "GenericPlatform/GenericPlatformInputDeviceMapper.h"
 #include "IImageWrapperModule.h"
 #include "HDRHelper.h"
+#include "GlobalRenderResources.h"
 
 #define LOCTEXT_NAMESPACE "EditorViewportClient"
 
diff --git a/Engine/Source/Editor/UnrealEd/Private/Texture2DPreview.cpp b/Engine/Source/Editor/UnrealEd/Private/Texture2DPreview.cpp
--- a/Engine/Source/Editor/UnrealEd/Private/Texture2DPreview.cpp
+++ b/Engine/Source/Editor/UnrealEd/Private/Texture2DPreview.cpp
@@ -6,6 +6,7 @@
 
 #include "Texture2DPreview.h"
 #include "Shader.h"
+#include "GlobalRenderResources.h"
 #include "GlobalShader.h"
 #include "SimpleElementShaders.h"
 #include "ShaderParameterUtils.h"
diff --git a/Engine/Source/Editor/UnrealEd/Private/VolumeTexturePreview.cpp b/Engine/Source/Editor/UnrealEd/Private/VolumeTexturePreview.cpp
--- a/Engine/Source/Editor/UnrealEd/Private/VolumeTexturePreview.cpp
+++ b/Engine/Source/Editor/UnrealEd/Private/VolumeTexturePreview.cpp
@@ -15,6 +15,7 @@
 #include "PipelineStateCache.h"
 #include "Editor.h"
 #include "TextureResource.h"
+#include "GlobalRenderResources.h"
 
 UNREALED_API void GetBestFitForNumberOfTiles(int32 InSize, int32& OutNumTilesX, int32& OutNumTilesY)
 {
diff --git a/Engine/Source/Runtime/D3D12RHI/Private/D3D12Commands.cpp b/Engine/Source/Runtime/D3D12RHI/Private/D3D12Commands.cpp
--- a/Engine/Source/Runtime/D3D12RHI/Private/D3D12Commands.cpp
+++ b/Engine/Source/Runtime/D3D12RHI/Private/D3D12Commands.cpp
@@ -16,6 +16,7 @@ D3D12Commands.cpp: D3D RHI commands implementation.
 #include "ResolveShader.h"
 #include "SceneUtils.h"
 #include "RenderUtils.h"
+#include "GlobalRenderResources.h"
 #include "RHIShaderParametersShared.h"
 
 int32 AFRSyncTemporalResources = 1;
diff --git a/Engine/Source/Runtime/D3D12RHI/Private/D3D12RayTracing.cpp b/Engine/Source/Runtime/D3D12RHI/Private/D3D12RayTracing.cpp
--- a/Engine/Source/Runtime/D3D12RHI/Private/D3D12RayTracing.cpp
+++ b/Engine/Source/Runtime/D3D12RHI/Private/D3D12RayTracing.cpp
@@ -18,6 +18,7 @@
 #include "Async/ParallelFor.h"
 #include "Misc/BufferedOutputDevice.h"
 #include "String/LexFromString.h"
+#include "GlobalRenderResources.h"
 
 static int32 GRayTracingDebugForceBuildMode = 0;
 static FAutoConsoleVariableRef CVarRayTracingDebugForceFastTrace(
diff --git a/Engine/Source/Runtime/Engine/Private/BatchedElements.cpp b/Engine/Source/Runtime/Engine/Private/BatchedElements.cpp
--- a/Engine/Source/Runtime/Engine/Private/BatchedElements.cpp
+++ b/Engine/Source/Runtime/Engine/Private/BatchedElements.cpp
@@ -9,6 +9,7 @@
 #include "RHIStaticStates.h"
 #include "MeshPassProcessor.h"
 #include "SceneRelativeViewMatrices.h"
+#include "GlobalRenderResources.h"
 #include "HDRHelper.h"
 
 DEFINE_LOG_CATEGORY_STATIC(LogBatchedElements, Log, All);
diff --git a/Engine/Source/Runtime/Engine/Private/Components/ReflectionCaptureComponent.cpp b/Engine/Source/Runtime/Engine/Private/Components/ReflectionCaptureComponent.cpp
--- a/Engine/Source/Runtime/Engine/Private/Components/ReflectionCaptureComponent.cpp
+++ b/Engine/Source/Runtime/Engine/Private/Components/ReflectionCaptureComponent.cpp
@@ -42,6 +42,7 @@
 #include "ComponentRecreateRenderStateContext.h"
 #include "Engine/TextureCube.h"
 #include "Math/PackedVector.h"
+#include "GlobalRenderResources.h"
 
 #include UE_INLINE_GENERATED_CPP_BY_NAME(ReflectionCaptureComponent)
 
diff --git a/Engine/Source/Runtime/Engine/Private/GameViewportClient.cpp b/Engine/Source/Runtime/Engine/Private/GameViewportClient.cpp
--- a/Engine/Source/Runtime/Engine/Private/GameViewportClient.cpp
+++ b/Engine/Source/Runtime/Engine/Private/GameViewportClient.cpp
@@ -73,6 +73,7 @@
 #include "ObjectTrace.h"
 #include "DynamicResolutionState.h"
 #include "HDRHelper.h"
+#include "GlobalRenderResources.h"
 
 #if WITH_EDITOR
 #include "Settings/LevelEditorPlaySettings.h"
diff --git a/Engine/Source/Runtime/Engine/Private/HUD.cpp b/Engine/Source/Runtime/Engine/Private/HUD.cpp
--- a/Engine/Source/Runtime/Engine/Private/HUD.cpp
+++ b/Engine/Source/Runtime/Engine/Private/HUD.cpp
@@ -18,6 +18,7 @@
 #include "CanvasTypes.h"
 #include "TextureResource.h"
 #include "Engine/Texture.h"
+#include "GlobalRenderResources.h"
 #include "Engine/LocalPlayer.h"
 #include "GameFramework/GameModeBase.h"
 #include "EngineUtils.h"
diff --git a/Engine/Source/Runtime/Engine/Private/Materials/MaterialUniformExpressions.cpp b/Engine/Source/Runtime/Engine/Private/Materials/MaterialUniformExpressions.cpp
--- a/Engine/Source/Runtime/Engine/Private/Materials/MaterialUniformExpressions.cpp
+++ b/Engine/Source/Runtime/Engine/Private/Materials/MaterialUniformExpressions.cpp
@@ -16,6 +16,7 @@
 #include "Misc/UObjectToken.h"
 
 #include "RenderCore.h"
+#include "GlobalRenderResources.h"
 #include "VirtualTexturing.h"
 #include "VT/RuntimeVirtualTexture.h"
 #include "SparseVolumeTexture/SparseVolumeTexture.h"
diff --git a/Engine/Source/Runtime/Engine/Private/MeshParticleVertexFactory.cpp b/Engine/Source/Runtime/Engine/Private/MeshParticleVertexFactory.cpp
--- a/Engine/Source/Runtime/Engine/Private/MeshParticleVertexFactory.cpp
+++ b/Engine/Source/Runtime/Engine/Private/MeshParticleVertexFactory.cpp
@@ -8,6 +8,7 @@
 #include "ParticleHelper.h"
 #include "ShaderParameterUtils.h"
 #include "MeshMaterialShader.h"
+#include "GlobalRenderResources.h"
 
 class FMeshParticleVertexFactoryShaderParameters : public FVertexFactoryShaderParameters
 {
diff --git a/Engine/Source/Runtime/Engine/Private/Rendering/ColorVertexBuffer.cpp b/Engine/Source/Runtime/Engine/Private/Rendering/ColorVertexBuffer.cpp
--- a/Engine/Source/Runtime/Engine/Private/Rendering/ColorVertexBuffer.cpp
+++ b/Engine/Source/Runtime/Engine/Private/Rendering/ColorVertexBuffer.cpp
@@ -7,6 +7,7 @@
 #include "ProfilingDebugging/LoadTimeTracker.h"
 #include "VertexFactory.h"
 #include "DataDrivenShaderPlatformInfo.h"
+#include "GlobalRenderResources.h"
 #include "RHIResourceUpdates.h"
 
 /*-----------------------------------------------------------------------------
diff --git a/Engine/Source/Runtime/Engine/Private/StatsRender2.cpp b/Engine/Source/Runtime/Engine/Private/StatsRender2.cpp
--- a/Engine/Source/Runtime/Engine/Private/StatsRender2.cpp
+++ b/Engine/Source/Runtime/Engine/Private/StatsRender2.cpp
@@ -21,6 +21,7 @@
 
 #include "Stats/StatsData.h"
 #include "Performance/EnginePerformanceTargets.h"
+#include "GlobalRenderResources.h"
 
 TAutoConsoleVariable<int32> CVarNumStatsPerGroup(
 	TEXT("stats.MaxPerGroup"),
diff --git a/Engine/Source/Runtime/Engine/Private/TextureDerivedData.cpp b/Engine/Source/Runtime/Engine/Private/TextureDerivedData.cpp
--- a/Engine/Source/Runtime/Engine/Private/TextureDerivedData.cpp
+++ b/Engine/Source/Runtime/Engine/Private/TextureDerivedData.cpp
@@ -17,6 +17,7 @@
 #include "Serialization/MemoryReader.h"
 #include "UObject/Package.h"
 #include "RenderUtils.h"
+#include "GlobalRenderResources.h"
 #include "TextureResource.h"
 #include "Engine/Texture.h"
 #include "Engine/Texture2D.h"
diff --git a/Engine/Source/Runtime/Engine/Private/VectorField.cpp b/Engine/Source/Runtime/Engine/Private/VectorField.cpp
--- a/Engine/Source/Runtime/Engine/Private/VectorField.cpp
+++ b/Engine/Source/Runtime/Engine/Private/VectorField.cpp
@@ -30,6 +30,7 @@
 #include "Materials/Material.h"
 #include "Engine/Engine.h"
 #include "TextureResource.h"
+#include "GlobalRenderResources.h"
 
 #include UE_INLINE_GENERATED_CPP_BY_NAME(VectorField)
 
diff --git a/Engine/Source/Runtime/Engine/Public/LocalVertexFactory.h b/Engine/Source/Runtime/Engine/Public/LocalVertexFactory.h
--- a/Engine/Source/Runtime/Engine/Public/LocalVertexFactory.h
+++ b/Engine/Source/Runtime/Engine/Public/LocalVertexFactory.h
@@ -6,6 +6,7 @@
 #include "ShaderParameters.h"
 #include "Components.h"
 #include "VertexFactory.h"
+#include "GlobalRenderResources.h"
 
 class FMaterial;
 class FSceneView;
diff --git a/Engine/Source/Runtime/Engine/Public/ParticleHelper.h b/Engine/Source/Runtime/Engine/Public/ParticleHelper.h
--- a/Engine/Source/Runtime/Engine/Public/ParticleHelper.h
+++ b/Engine/Source/Runtime/Engine/Public/ParticleHelper.h
@@ -33,6 +33,8 @@ class UStaticMesh;
 struct FBaseParticle;
 struct FParticleMeshEmitterInstance;
 struct FStaticMeshLODResources;
+struct FGlobalDynamicIndexBufferAllocation;
+struct FGlobalDynamicVertexBufferAllocation;
 
 DECLARE_LOG_CATEGORY_EXTERN(LogParticles, Log, All);
 
@@ -1532,9 +1534,9 @@ struct FDynamicSpriteEmitterDataBase : public FDynamicEmitterDataBase
 		int32 InDynamicParameterVertexSize, 
 		FGlobalDynamicIndexBuffer& DynamicIndexBuffer,
 		FGlobalDynamicVertexBuffer& DynamicVertexBuffer,
-		FGlobalDynamicVertexBuffer::FAllocation& DynamicVertexAllocation,
-		FGlobalDynamicIndexBuffer::FAllocation& DynamicIndexAllocation,
-		FGlobalDynamicVertexBuffer::FAllocation* DynamicParameterAllocation,
+		FGlobalDynamicVertexBufferAllocation& DynamicVertexAllocation,
+		FGlobalDynamicIndexBufferAllocation& DynamicIndexAllocation,
+		FGlobalDynamicVertexBufferAllocation* DynamicParameterAllocation,
 		FAsyncBufferFillData& Data) const;
 
 	/** The material render proxy for this emitter */
diff --git a/Engine/Source/Runtime/Engine/Public/SceneManagement.h b/Engine/Source/Runtime/Engine/Public/SceneManagement.h
--- a/Engine/Source/Runtime/Engine/Public/SceneManagement.h
+++ b/Engine/Source/Runtime/Engine/Public/SceneManagement.h
@@ -33,6 +33,9 @@
 #include "Rendering/SkyAtmosphereCommonData.h"
 
 class FCanvas;
+class FGlobalDynamicIndexBuffer;
+class FGlobalDynamicReadBuffer;
+class FGlobalDynamicVertexBuffer;
 class FLightMap;
 class FLightmapResourceCluster;
 class FLightSceneInfo;
diff --git a/Engine/Source/Runtime/HeadMountedDisplay/Private/DefaultSpectatorScreenController.cpp b/Engine/Source/Runtime/HeadMountedDisplay/Private/DefaultSpectatorScreenController.cpp
--- a/Engine/Source/Runtime/HeadMountedDisplay/Private/DefaultSpectatorScreenController.cpp
+++ b/Engine/Source/Runtime/HeadMountedDisplay/Private/DefaultSpectatorScreenController.cpp
@@ -6,6 +6,7 @@
 #include "ClearQuad.h"
 #include "ScreenRendering.h"
 #include "TextureResource.h"
+#include "GlobalRenderResources.h"
 #include "Misc/CoreDelegates.h"
 #include "Engine/Texture.h"
 #include "HeadMountedDisplayBase.h"
diff --git a/Engine/Source/Runtime/RenderCore/Public/GlobalRenderResources.h b/Engine/Source/Runtime/RenderCore/Public/GlobalRenderResources.h
--- a/Engine/Source/Runtime/RenderCore/Public/GlobalRenderResources.h
+++ b/Engine/Source/Runtime/RenderCore/Public/GlobalRenderResources.h
@@ -67,69 +67,6 @@ extern RENDERCORE_API int32 GMipColorTextureMipLevels;
 // 4: 8x8 cubemap resolution, shader needs to use the same value as preprocessing
 extern RENDERCORE_API const uint32 GDiffuseConvolveMipLevel;
 
-// TGlobalResource
-
-/** Used to declare a render resource that is initialized/released by static initialization/destruction. */
-template<class ResourceType>
-class TGlobalResource : public ResourceType
-{
-public:
-	/** Default constructor. */
-	TGlobalResource()
-	{
-		InitGlobalResource();
-	}
-
-	/** Initialization constructor: 1 parameter. */
-	template<typename... Args>
-	explicit TGlobalResource(Args... InArgs)
-		: ResourceType(InArgs...)
-	{
-		InitGlobalResource();
-	}
-
-	/** Destructor. */
-	virtual ~TGlobalResource()
-	{
-		ReleaseGlobalResource();
-	}
-
-private:
-
-	/**
-	 * Initialize the global resource.
-	 */
-	void InitGlobalResource()
-	{
-		if (IsInRenderingThread())
-		{
-			// If the resource is constructed in the rendering thread, directly initialize it.
-			((ResourceType*)this)->InitResource();
-		}
-		else
-		{
-			// If the resource is constructed outside of the rendering thread, enqueue a command to initialize it.
-			BeginInitResource((ResourceType*)this);
-		}
-	}
-
-	/**
-	 * Release the global resource.
-	 */
-	void ReleaseGlobalResource()
-	{
-		// This should be called in the rendering thread, or at shutdown when the rendering thread has exited.
-		// However, it may also be called at shutdown after an error, when the rendering thread is still running.
-		// To avoid a second error in that case we don't assert.
-#if 0
-		check(IsInRenderingThread());
-#endif
-
-		// Cleanup the resource.
-		((ResourceType*)this)->ReleaseResource();
-	}
-};
-
 /**
 * A vertex buffer with a single color component.  This is used on meshes that don't have a color component
 * to keep from needing a separate vertex factory to handle this case.
diff --git a/Engine/Source/Runtime/RenderCore/Public/PixelShaderUtils.h b/Engine/Source/Runtime/RenderCore/Public/PixelShaderUtils.h
--- a/Engine/Source/Runtime/RenderCore/Public/PixelShaderUtils.h
+++ b/Engine/Source/Runtime/RenderCore/Public/PixelShaderUtils.h
@@ -7,6 +7,7 @@
 #pragma once
 
 #include "CommonRenderResources.h"
+#include "GlobalRenderResources.h"
 #include "GlobalShader.h"
 #include "Math/IntPoint.h"
 #include "Math/IntRect.h"
diff --git a/Engine/Source/Runtime/RenderCore/Public/RenderResource.h b/Engine/Source/Runtime/RenderCore/Public/RenderResource.h
--- a/Engine/Source/Runtime/RenderCore/Public/RenderResource.h
+++ b/Engine/Source/Runtime/RenderCore/Public/RenderResource.h
@@ -495,9 +495,69 @@ public:
 	TRefCountPtr<FRDGPooledBuffer> Buffer;
 };
 
+/** Used to declare a render resource that is initialized/released by static initialization/destruction. */
+template<class ResourceType>
+class TGlobalResource : public ResourceType
+{
+public:
+	/** Default constructor. */
+	TGlobalResource()
+	{
+		InitGlobalResource();
+	}
+
+	/** Initialization constructor: 1 parameter. */
+	template<typename... Args>
+	explicit TGlobalResource(Args... InArgs)
+		: ResourceType(InArgs...)
+	{
+		InitGlobalResource();
+	}
+
+	/** Destructor. */
+	virtual ~TGlobalResource()
+	{
+		ReleaseGlobalResource();
+	}
+
+private:
+
+	/**
+	 * Initialize the global resource.
+	 */
+	void InitGlobalResource()
+	{
+		if (IsInRenderingThread())
+		{
+			// If the resource is constructed in the rendering thread, directly initialize it.
+			((ResourceType*)this)->InitResource();
+		}
+		else
+		{
+			// If the resource is constructed outside of the rendering thread, enqueue a command to initialize it.
+			BeginInitResource((ResourceType*)this);
+		}
+	}
+
+	/**
+	 * Release the global resource.
+	 */
+	void ReleaseGlobalResource()
+	{
+		// This should be called in the rendering thread, or at shutdown when the rendering thread has exited.
+		// However, it may also be called at shutdown after an error, when the rendering thread is still running.
+		// To avoid a second error in that case we don't assert.
+#if 0
+		check(IsInRenderingThread());
+#endif
+
+		// Cleanup the resource.
+		((ResourceType*)this)->ReleaseResource();
+	}
+};
+
 #if UE_ENABLE_INCLUDE_ORDER_DEPRECATED_IN_5_2
 #include "RayTracingGeometry.h"
 #include "RenderUtils.h"
-#endif
-
 #include "GlobalRenderResources.h"
+#endif
diff --git a/Engine/Source/Runtime/Renderer/Private/CompositionLighting/PostProcessAmbient.cpp b/Engine/Source/Runtime/Renderer/Private/CompositionLighting/PostProcessAmbient.cpp
--- a/Engine/Source/Runtime/Renderer/Private/CompositionLighting/PostProcessAmbient.cpp
+++ b/Engine/Source/Runtime/Renderer/Private/CompositionLighting/PostProcessAmbient.cpp
@@ -2,7 +2,7 @@
 
 #include "AmbientCubemapParameters.h"
 #include "Engine/Texture.h"
-#include "RenderUtils.h"
+#include "GlobalRenderResources.h"
 #include "TextureResource.h"
 
 void SetupAmbientCubemapParameters(const FFinalPostProcessSettings::FCubemapEntry& Entry, FAmbientCubemapParameters* OutParameters)
diff --git a/Engine/Source/Runtime/Renderer/Private/SystemTextures.cpp b/Engine/Source/Runtime/Renderer/Private/SystemTextures.cpp
--- a/Engine/Source/Runtime/Renderer/Private/SystemTextures.cpp
+++ b/Engine/Source/Runtime/Renderer/Private/SystemTextures.cpp
@@ -12,6 +12,7 @@
 #include "ClearQuad.h"
 #include "LTC.h"
 #include "Math/PackedVector.h"
+#include "GlobalRenderResources.h"
 
 /*-----------------------------------------------------------------------------
 SystemTextures
diff --git a/Engine/Source/Runtime/Renderer/Private/VT/VirtualTextureFeedbackBuffer.cpp b/Engine/Source/Runtime/Renderer/Private/VT/VirtualTextureFeedbackBuffer.cpp
--- a/Engine/Source/Runtime/Renderer/Private/VT/VirtualTextureFeedbackBuffer.cpp
+++ b/Engine/Source/Runtime/Renderer/Private/VT/VirtualTextureFeedbackBuffer.cpp
@@ -5,6 +5,7 @@
 #include "RenderGraphUtils.h"
 #include "SceneView.h"
 #include "VT/VirtualTextureFeedback.h"
+#include "GlobalRenderResources.h"
 
 int32 GetVirtualTextureFeedbackScale(FSceneViewFamily const* InViewFamily)
 {
@@ -106,6 +107,11 @@ void FVirtualTextureFeedbackBuffer::End(FRDGBuilder& GraphBuilder)
 	});
 }
 
+FRHIUnorderedAccessView* FVirtualTextureFeedbackBuffer::GetUAV() const
+{
+	return UAV ? UAV : GEmptyVertexBufferWithUAV->UnorderedAccessViewRHI.GetReference();
+}
+
 void FVirtualTextureFeedbackBuffer::ReleaseRHI()
 {
 	PooledBuffer = nullptr;
diff --git a/Engine/Source/Runtime/Renderer/Private/VT/VirtualTextureSpace.cpp b/Engine/Source/Runtime/Renderer/Private/VT/VirtualTextureSpace.cpp
--- a/Engine/Source/Runtime/Renderer/Private/VT/VirtualTextureSpace.cpp
+++ b/Engine/Source/Runtime/Renderer/Private/VT/VirtualTextureSpace.cpp
@@ -9,6 +9,7 @@
 #include "VisualizeTexture.h"
 #include "CommonRenderResources.h"
 #include "GlobalShader.h"
+#include "GlobalRenderResources.h"
 #include "PipelineStateCache.h"
 #include "RHIStaticStates.h"
 #include "HAL/IConsoleManager.h"
diff --git a/Engine/Source/Runtime/Renderer/Public/GlobalDistanceFieldParameters.h b/Engine/Source/Runtime/Renderer/Public/GlobalDistanceFieldParameters.h
--- a/Engine/Source/Runtime/Renderer/Public/GlobalDistanceFieldParameters.h
+++ b/Engine/Source/Runtime/Renderer/Public/GlobalDistanceFieldParameters.h
@@ -4,6 +4,7 @@
 
 #include "CoreMinimal.h"
 #include "GlobalDistanceFieldConstants.h"
+#include "GlobalRenderResources.h"
 #include "ShaderParameterMacros.h"
 #include "ShaderParameters.h"
 #include "RenderUtils.h"
diff --git a/Engine/Source/Runtime/Renderer/Public/VT/VirtualTextureFeedbackBuffer.h b/Engine/Source/Runtime/Renderer/Public/VT/VirtualTextureFeedbackBuffer.h
--- a/Engine/Source/Runtime/Renderer/Public/VT/VirtualTextureFeedbackBuffer.h
+++ b/Engine/Source/Runtime/Renderer/Public/VT/VirtualTextureFeedbackBuffer.h
@@ -60,10 +60,7 @@ public:
 	// Submits the feedback buffer for readback.
 	void End(FRDGBuilder& GraphBuilder);
 
-	FRHIUnorderedAccessView* GetUAV() const
-	{
-		return UAV ? UAV : GEmptyVertexBufferWithUAV->UnorderedAccessViewRHI.GetReference();
-	}
+	FRHIUnorderedAccessView* GetUAV() const;
 
 	//~ Begin FRenderResource Interface
 	virtual void ReleaseRHI() override;
diff --git a/Engine/Source/Runtime/SlateRHIRenderer/Private/SlateRHIRenderingPolicy.cpp b/Engine/Source/Runtime/SlateRHIRenderer/Private/SlateRHIRenderingPolicy.cpp
--- a/Engine/Source/Runtime/SlateRHIRenderer/Private/SlateRHIRenderingPolicy.cpp
+++ b/Engine/Source/Runtime/SlateRHIRenderer/Private/SlateRHIRenderingPolicy.cpp
@@ -9,6 +9,7 @@
 #include "EngineGlobals.h"
 #include "RHIStaticStates.h"
 #include "RHIUtilities.h"
+#include "GlobalRenderResources.h"
 #include "SceneView.h"
 #include "SceneUtils.h"
 #include "Engine/Engine.h"
diff --git a/Engine/Source/Runtime/VulkanRHI/Private/VulkanCommands.cpp b/Engine/Source/Runtime/VulkanRHI/Private/VulkanCommands.cpp
--- a/Engine/Source/Runtime/VulkanRHI/Private/VulkanCommands.cpp
+++ b/Engine/Source/Runtime/VulkanRHI/Private/VulkanCommands.cpp
@@ -10,6 +10,7 @@
 #include "EngineGlobals.h"
 #include "VulkanLLM.h"
 #include "RenderUtils.h"
+#include "GlobalRenderResources.h"
 #include "RHIShaderParametersShared.h"
 
 static TAutoConsoleVariable<int32> GCVarSubmitOnDispatch(
