Reimplementing changes to ResolveShader. Moved type defintions to their cpp file to remove public dependencies.

diff --git a/Engine/Source/Runtime/RenderCore/Private/ResolveShader.cpp b/Engine/Source/Runtime/RenderCore/Private/ResolveShader.cpp
--- a/Engine/Source/Runtime/RenderCore/Private/ResolveShader.cpp
+++ b/Engine/Source/Runtime/RenderCore/Private/ResolveShader.cpp
@@ -1,8 +1,9 @@
 // Copyright Epic Games, Inc. All Rights Reserved.
 
-
 #include "ResolveShader.h"
 #include "ShaderParameterUtils.h"
+#include "DataDrivenShaderPlatformInfo.h"
+#include "StereoRenderUtils.h"
 
 IMPLEMENT_SHADER_TYPE(, FResolveDepthPS, TEXT("/Engine/Private/ResolvePixelShader.usf"), TEXT("MainDepth"), SF_Pixel);
 IMPLEMENT_SHADER_TYPE(, FResolveDepth2XPS, TEXT("/Engine/Private/ResolvePixelShader.usf"), TEXT("MainDepth"), SF_Pixel);
@@ -16,11 +17,192 @@ IMPLEMENT_SHADER_TYPE(, FResolveSingleSamplePS, TEXT("/Engine/Private/ResolvePix
 IMPLEMENT_SHADER_TYPE(, FResolveVS, TEXT("/Engine/Private/ResolveVertexShader.usf"), TEXT("Main"), SF_Vertex);
 IMPLEMENT_SHADER_TYPE(, FResolveArrayVS, TEXT("/Engine/Private/ResolveVertexShader.usf"), TEXT("Main"), SF_Vertex);
 
+// FResolveDepthPS
+
+FResolveDepthPS::FResolveDepthPS() = default;
+FResolveDepthPS::FResolveDepthPS(const ShaderMetaType::CompiledShaderInitializerType& Initializer)
+	: FGlobalShader(Initializer)
+{
+	UnresolvedSurface.Bind(Initializer.ParameterMap, TEXT("UnresolvedSurface"), SPF_Mandatory);
+}
+
+void FResolveDepthPS::ModifyCompilationEnvironment(const FGlobalShaderPermutationParameters& Parameters, FShaderCompilerEnvironment& OutEnvironment)
+{
+	FGlobalShader::ModifyCompilationEnvironment(Parameters, OutEnvironment);
+}
+
+void FResolveDepthPS::SetParameters(FRHICommandList& RHICmdList, FParameter)
+{
+}
+
+// FResolveDepth2XPS
+
+FResolveDepth2XPS::FResolveDepth2XPS() = default;
+FResolveDepth2XPS::FResolveDepth2XPS(const ShaderMetaType::CompiledShaderInitializerType& Initializer)
+	: FResolveDepthPS(Initializer)
+{
+}
+
+void FResolveDepth2XPS::ModifyCompilationEnvironment(const FGlobalShaderPermutationParameters& Parameters, FShaderCompilerEnvironment& OutEnvironment)
+{
+	FResolveDepthPS::ModifyCompilationEnvironment(Parameters, OutEnvironment);
+	OutEnvironment.SetDefine(TEXT("DEPTH_RESOLVE_NUM_SAMPLES"), 2);
+}
+
+// FResolveDepth4XPS
+
+FResolveDepth4XPS::FResolveDepth4XPS() = default;
+FResolveDepth4XPS::FResolveDepth4XPS(const ShaderMetaType::CompiledShaderInitializerType& Initializer)
+	: FResolveDepthPS(Initializer)
+{
+}
+
+void FResolveDepth4XPS::ModifyCompilationEnvironment(const FGlobalShaderPermutationParameters& Parameters, FShaderCompilerEnvironment& OutEnvironment)
+{
+	FResolveDepthPS::ModifyCompilationEnvironment(Parameters, OutEnvironment);
+	OutEnvironment.SetDefine(TEXT("DEPTH_RESOLVE_NUM_SAMPLES"), 4);
+}
+
+// FResolveDepth8XPS
+
+FResolveDepth8XPS::FResolveDepth8XPS() = default;
+FResolveDepth8XPS::FResolveDepth8XPS(const ShaderMetaType::CompiledShaderInitializerType& Initializer)
+	: FResolveDepthPS(Initializer)
+{
+}
+
+bool FResolveDepth8XPS::ShouldCache(EShaderPlatform Platform)
+{
+	return GetMaxSupportedFeatureLevel(Platform) >= ERHIFeatureLevel::SM5;
+}
+
+void FResolveDepth8XPS::ModifyCompilationEnvironment(const FGlobalShaderPermutationParameters& Parameters, FShaderCompilerEnvironment& OutEnvironment)
+{
+	FResolveDepthPS::ModifyCompilationEnvironment(Parameters, OutEnvironment);
+	OutEnvironment.SetDefine(TEXT("DEPTH_RESOLVE_NUM_SAMPLES"), 8);
+}
+
+// FResolveDepthArrayPS
+
+FResolveDepthArrayPS::FResolveDepthArrayPS() = default;
+FResolveDepthArrayPS::FResolveDepthArrayPS(const ShaderMetaType::CompiledShaderInitializerType& Initializer)
+	: FResolveDepthPS(Initializer)
+{
+}
+
+bool FResolveDepthArrayPS::ShouldCompilePermutation(const FGlobalShaderPermutationParameters& Parameters)
+{
+	if (FResolveDepthPS::ShouldCompilePermutation(Parameters))
+	{
+		UE::StereoRenderUtils::FStereoShaderAspects Aspects(Parameters.Platform);
+		return Aspects.IsMobileMultiViewEnabled();
+	}
+	return false;
+}
+
+void FResolveDepthArrayPS::ModifyCompilationEnvironment(const FGlobalShaderPermutationParameters& Parameters, FShaderCompilerEnvironment& OutEnvironment)
+{
+	FResolveDepthPS::ModifyCompilationEnvironment(Parameters, OutEnvironment);
+	OutEnvironment.SetDefine(TEXT("DEPTH_RESOLVE_TEXTUREARRAY"), 1);
+}
+
+// FResolveDepthArray2XPS
+
+FResolveDepthArray2XPS::FResolveDepthArray2XPS() = default;
+FResolveDepthArray2XPS::FResolveDepthArray2XPS(const ShaderMetaType::CompiledShaderInitializerType& Initializer)
+	: FResolveDepthArrayPS(Initializer)
+{
+}
+
+bool FResolveDepthArray2XPS::ShouldCompilePermutation(const FGlobalShaderPermutationParameters& Parameters)
+{
+	return FResolveDepthArrayPS::ShouldCompilePermutation(Parameters);
+}
+
+void FResolveDepthArray2XPS::ModifyCompilationEnvironment(const FGlobalShaderPermutationParameters& Parameters, FShaderCompilerEnvironment& OutEnvironment)
+{
+	FResolveDepthArrayPS::ModifyCompilationEnvironment(Parameters, OutEnvironment);
+	OutEnvironment.SetDefine(TEXT("DEPTH_RESOLVE_NUM_SAMPLES"), 2);
+}
+
+// FResolveDepthArray4XPS
+
+FResolveDepthArray4XPS::FResolveDepthArray4XPS() = default;
+FResolveDepthArray4XPS::FResolveDepthArray4XPS(const ShaderMetaType::CompiledShaderInitializerType& Initializer)
+	: FResolveDepthArrayPS(Initializer)
+{
+}
+
+bool FResolveDepthArray4XPS::ShouldCompilePermutation(const FGlobalShaderPermutationParameters& Parameters)
+{
+	return FResolveDepthArrayPS::ShouldCompilePermutation(Parameters);
+}
+
+void FResolveDepthArray4XPS::ModifyCompilationEnvironment(const FGlobalShaderPermutationParameters& Parameters, FShaderCompilerEnvironment& OutEnvironment)
+{
+	FResolveDepthArrayPS::ModifyCompilationEnvironment(Parameters, OutEnvironment);
+	OutEnvironment.SetDefine(TEXT("DEPTH_RESOLVE_NUM_SAMPLES"), 4);
+}
+
+// FResolveDepthArray8XPS
+
+FResolveDepthArray8XPS::FResolveDepthArray8XPS() = default;
+FResolveDepthArray8XPS::FResolveDepthArray8XPS(const ShaderMetaType::CompiledShaderInitializerType& Initializer)
+	: FResolveDepthArrayPS(Initializer)
+{
+}
+
+bool FResolveDepthArray8XPS::ShouldCache(EShaderPlatform Platform)
+{
+	return GetMaxSupportedFeatureLevel(Platform) >= ERHIFeatureLevel::SM5;
+}
+
+bool FResolveDepthArray8XPS::ShouldCompilePermutation(const FGlobalShaderPermutationParameters& Parameters)
+{
+	return FResolveDepthArrayPS::ShouldCompilePermutation(Parameters);
+}
+
+void FResolveDepthArray8XPS::ModifyCompilationEnvironment(const FGlobalShaderPermutationParameters& Parameters, FShaderCompilerEnvironment& OutEnvironment)
+{
+	FResolveDepthArrayPS::ModifyCompilationEnvironment(Parameters, OutEnvironment);
+	OutEnvironment.SetDefine(TEXT("DEPTH_RESOLVE_NUM_SAMPLES"), 8);
+}
+
+// FResolveSingleSamplePS
+
+FResolveSingleSamplePS::FResolveSingleSamplePS() = default;
+FResolveSingleSamplePS::FResolveSingleSamplePS(const ShaderMetaType::CompiledShaderInitializerType& Initializer)
+	: FGlobalShader(Initializer)
+{
+	UnresolvedSurface.Bind(Initializer.ParameterMap, TEXT("UnresolvedSurface"), SPF_Mandatory);
+	SingleSampleIndex.Bind(Initializer.ParameterMap, TEXT("SingleSampleIndex"), SPF_Mandatory);
+}
+
+bool FResolveSingleSamplePS::ShouldCompilePermutation(const FGlobalShaderPermutationParameters& Parameters)
+{
+	return FDataDrivenShaderPlatformInfo::GetIsLanguageD3D(Parameters.Platform);
+}
+
 void FResolveSingleSamplePS::SetParameters(FRHICommandList& RHICmdList, uint32 SingleSampleIndexValue)
 {
 	SetShaderValue(RHICmdList, RHICmdList.GetBoundPixelShader(),SingleSampleIndex,SingleSampleIndexValue);
 }
 
+// FResolveVS
+
+FResolveVS::FResolveVS() = default;
+FResolveVS::FResolveVS(const ShaderMetaType::CompiledShaderInitializerType& Initializer)
+	: FGlobalShader(Initializer)
+{
+	PositionMinMax.Bind(Initializer.ParameterMap, TEXT("PositionMinMax"), SPF_Mandatory);
+	UVMinMax.Bind(Initializer.ParameterMap, TEXT("UVMinMax"), SPF_Mandatory);
+}
+
+bool FResolveVS::ShouldCompilePermutation(const FGlobalShaderPermutationParameters& Parameters)
+{
+	return true;
+}
+
 void FResolveVS::SetParameters(FRHICommandList& RHICmdList, const FResolveRect& SrcBounds, const FResolveRect& DstBounds, uint32 DstSurfaceWidth, uint32 DstSurfaceHeight)
 {
 	// Generate the vertices used to copy from the source surface to the destination surface.
@@ -36,3 +218,32 @@ void FResolveVS::SetParameters(FRHICommandList& RHICmdList, const FResolveRect&
 	SetShaderValue(RHICmdList, RHICmdList.GetBoundVertexShader(), PositionMinMax, FVector4f(MinX, MinY, MaxX, MaxY));
 	SetShaderValue(RHICmdList, RHICmdList.GetBoundVertexShader(), UVMinMax, FVector4f(MinU, MinV, MaxU, MaxV));
 }
+
+// FResolveArrayVS
+
+FResolveArrayVS::FResolveArrayVS() = default;
+FResolveArrayVS::FResolveArrayVS(const ShaderMetaType::CompiledShaderInitializerType& Initializer)
+	: FResolveVS(Initializer)
+{
+}
+
+bool FResolveArrayVS::ShouldCompilePermutation(const FGlobalShaderPermutationParameters& Parameters)
+{
+	if (FResolveVS::ShouldCompilePermutation(Parameters))
+	{
+		UE::StereoRenderUtils::FStereoShaderAspects Aspects(Parameters.Platform);
+		return Aspects.IsMobileMultiViewEnabled();
+	}
+	return false;
+}
+
+void FResolveArrayVS::ModifyCompilationEnvironment(const FGlobalShaderPermutationParameters& Parameters, FShaderCompilerEnvironment& OutEnvironment)
+{
+	FResolveVS::ModifyCompilationEnvironment(Parameters, OutEnvironment);
+	OutEnvironment.SetDefine(TEXT("DEPTH_RESOLVE_TEXTUREARRAY"), 1);
+}
+
+void FResolveArrayVS::SetParameters(FRHICommandList& RHICmdList, const FResolveRect& SrcBounds, const FResolveRect& DstBounds, uint32 DstSurfaceWidth, uint32 DstSurfaceHeight)
+{
+	return FResolveVS::SetParameters(RHICmdList, SrcBounds, DstBounds, DstSurfaceWidth, DstSurfaceHeight);
+}
diff --git a/Engine/Source/Runtime/RenderCore/Public/ResolveShader.h b/Engine/Source/Runtime/RenderCore/Public/ResolveShader.h
--- a/Engine/Source/Runtime/RenderCore/Public/ResolveShader.h
+++ b/Engine/Source/Runtime/RenderCore/Public/ResolveShader.h
@@ -1,6 +1,5 @@
 // Copyright Epic Games, Inc. All Rights Reserved.
 
-
 #pragma once
 
 #include "CoreMinimal.h"
@@ -11,8 +10,6 @@
 #include "Shader.h"
 #include "ShaderCore.h"
 #include "ShaderParameters.h"
-#include "DataDrivenShaderPlatformInfo.h"
-#include "StereoRenderUtils.h"
 
 class FPointerTableBase;
 class FRHICommandList;
@@ -20,224 +17,122 @@ struct FResolveRect;
 
 struct FDummyResolveParameter {};
 
-class FResolveDepthPS : public FGlobalShader
+class RENDERCORE_API FResolveDepthPS : public FGlobalShader
 {
-	DECLARE_EXPORTED_SHADER_TYPE(FResolveDepthPS, Global, RENDERCORE_API);
+	DECLARE_GLOBAL_SHADER(FResolveDepthPS);
 public:
-	
 	typedef FDummyResolveParameter FParameter;
-
-	static void ModifyCompilationEnvironment(const FGlobalShaderPermutationParameters& Parameters, FShaderCompilerEnvironment& OutEnvironment)
-	{
-		FGlobalShader::ModifyCompilationEnvironment(Parameters, OutEnvironment);
-	}
 		
-	FResolveDepthPS(const ShaderMetaType::CompiledShaderInitializerType& Initializer):
-		FGlobalShader(Initializer)
-	{
-		UnresolvedSurface.Bind(Initializer.ParameterMap,TEXT("UnresolvedSurface"), SPF_Mandatory);
-	}
-	FResolveDepthPS() {}
-	
-	void SetParameters(FRHICommandList& RHICmdList, FParameter)
-	{
-	}
+	FResolveDepthPS();
+	FResolveDepthPS(const ShaderMetaType::CompiledShaderInitializerType& Initializer);
+
+	static void ModifyCompilationEnvironment(const FGlobalShaderPermutationParameters& Parameters, FShaderCompilerEnvironment& OutEnvironment);
+	void SetParameters(FRHICommandList& RHICmdList, FParameter);
 
 	LAYOUT_FIELD(FShaderResourceParameter, UnresolvedSurface);
 };
 
-class FResolveDepth2XPS : public FResolveDepthPS
+class RENDERCORE_API FResolveDepth2XPS : public FResolveDepthPS
 {
-	DECLARE_EXPORTED_SHADER_TYPE(FResolveDepth2XPS, Global, RENDERCORE_API);
+	DECLARE_GLOBAL_SHADER(FResolveDepth2XPS);
 public:
-
 	typedef FDummyResolveParameter FParameter;
 
-	static void ModifyCompilationEnvironment(const FGlobalShaderPermutationParameters& Parameters, FShaderCompilerEnvironment& OutEnvironment)
-	{
-		FResolveDepthPS::ModifyCompilationEnvironment(Parameters, OutEnvironment);
-		OutEnvironment.SetDefine(TEXT("DEPTH_RESOLVE_NUM_SAMPLES"), 2);
-	}
-
-	FResolveDepth2XPS(const ShaderMetaType::CompiledShaderInitializerType& Initializer) :
-		FResolveDepthPS(Initializer)
-	{
-	}
-	FResolveDepth2XPS() {}
+	FResolveDepth2XPS();
+	FResolveDepth2XPS(const ShaderMetaType::CompiledShaderInitializerType& Initializer);
+
+	static void ModifyCompilationEnvironment(const FGlobalShaderPermutationParameters& Parameters, FShaderCompilerEnvironment& OutEnvironment);
 };
 
-class FResolveDepth4XPS : public FResolveDepthPS
+class RENDERCORE_API FResolveDepth4XPS : public FResolveDepthPS
 {
-	DECLARE_EXPORTED_SHADER_TYPE(FResolveDepth4XPS, Global, RENDERCORE_API);
+	DECLARE_GLOBAL_SHADER(FResolveDepth4XPS);
 public:
-
 	typedef FDummyResolveParameter FParameter;
 
-	static void ModifyCompilationEnvironment(const FGlobalShaderPermutationParameters& Parameters, FShaderCompilerEnvironment& OutEnvironment)
-	{
-		FResolveDepthPS::ModifyCompilationEnvironment(Parameters, OutEnvironment);
-		OutEnvironment.SetDefine(TEXT("DEPTH_RESOLVE_NUM_SAMPLES"), 4);
-	}
-
-	FResolveDepth4XPS(const ShaderMetaType::CompiledShaderInitializerType& Initializer) :
-		FResolveDepthPS(Initializer)
-	{
-	}
-	FResolveDepth4XPS() {}
-};
+	FResolveDepth4XPS();
+	FResolveDepth4XPS(const ShaderMetaType::CompiledShaderInitializerType& Initializer);
 
+	static void ModifyCompilationEnvironment(const FGlobalShaderPermutationParameters& Parameters, FShaderCompilerEnvironment& OutEnvironment);
+};
 
-class FResolveDepth8XPS : public FResolveDepthPS
+class RENDERCORE_API FResolveDepth8XPS : public FResolveDepthPS
 {
-	DECLARE_EXPORTED_SHADER_TYPE(FResolveDepth8XPS, Global, RENDERCORE_API);
+	DECLARE_GLOBAL_SHADER(FResolveDepth8XPS);
 public:
-
 	typedef FDummyResolveParameter FParameter;
 
-	static bool ShouldCache(EShaderPlatform Platform) { return GetMaxSupportedFeatureLevel(Platform) >= ERHIFeatureLevel::SM5; }
-
-	static void ModifyCompilationEnvironment(const FGlobalShaderPermutationParameters& Parameters, FShaderCompilerEnvironment& OutEnvironment)
-	{
-		FResolveDepthPS::ModifyCompilationEnvironment(Parameters, OutEnvironment);
-		OutEnvironment.SetDefine(TEXT("DEPTH_RESOLVE_NUM_SAMPLES"), 8);
-	}
+	FResolveDepth8XPS();
+	FResolveDepth8XPS(const ShaderMetaType::CompiledShaderInitializerType& Initializer);
 
-	FResolveDepth8XPS(const ShaderMetaType::CompiledShaderInitializerType& Initializer) :
-		FResolveDepthPS(Initializer)
-	{
-	}
-	FResolveDepth8XPS() {}
+	static bool ShouldCache(EShaderPlatform Platform);
+	static void ModifyCompilationEnvironment(const FGlobalShaderPermutationParameters& Parameters, FShaderCompilerEnvironment& OutEnvironment);
 };
 
-class FResolveDepthArrayPS : public FResolveDepthPS
+class RENDERCORE_API FResolveDepthArrayPS : public FResolveDepthPS
 {
-	DECLARE_EXPORTED_SHADER_TYPE(FResolveDepthArrayPS, Global, RENDERCORE_API);
+	DECLARE_GLOBAL_SHADER(FResolveDepthArrayPS);
 public:
-
 	typedef FDummyResolveParameter FParameter;
 
-	static bool ShouldCompilePermutation(const FGlobalShaderPermutationParameters& Parameters)
-	{
-		if (FResolveDepthPS::ShouldCompilePermutation(Parameters))
-		{
-			UE::StereoRenderUtils::FStereoShaderAspects Aspects(Parameters.Platform);
-			return Aspects.IsMobileMultiViewEnabled();
-		}
-		return false;
-	}
-
-	static void ModifyCompilationEnvironment(const FGlobalShaderPermutationParameters& Parameters, FShaderCompilerEnvironment& OutEnvironment)
-	{
-		FResolveDepthPS::ModifyCompilationEnvironment(Parameters, OutEnvironment);
-		OutEnvironment.SetDefine(TEXT("DEPTH_RESOLVE_TEXTUREARRAY"), 1);
-	}
-
-	FResolveDepthArrayPS(const ShaderMetaType::CompiledShaderInitializerType& Initializer) :
-		FResolveDepthPS(Initializer)
-	{
-	}
-	FResolveDepthArrayPS() {}
+	FResolveDepthArrayPS();
+	FResolveDepthArrayPS(const ShaderMetaType::CompiledShaderInitializerType& Initializer);
+
+	static bool ShouldCompilePermutation(const FGlobalShaderPermutationParameters& Parameters);
+	static void ModifyCompilationEnvironment(const FGlobalShaderPermutationParameters& Parameters, FShaderCompilerEnvironment& OutEnvironment);
 };
 
-class FResolveDepthArray2XPS : public FResolveDepthArrayPS
+class RENDERCORE_API FResolveDepthArray2XPS : public FResolveDepthArrayPS
 {
-	DECLARE_EXPORTED_SHADER_TYPE(FResolveDepthArray2XPS, Global, RENDERCORE_API);
+	DECLARE_GLOBAL_SHADER(FResolveDepthArray2XPS);
 public:
-
 	typedef FDummyResolveParameter FParameter;
 
-	static bool ShouldCompilePermutation(const FGlobalShaderPermutationParameters& Parameters)
-	{
-		return FResolveDepthArrayPS::ShouldCompilePermutation(Parameters);
-	}
-
-	static void ModifyCompilationEnvironment(const FGlobalShaderPermutationParameters& Parameters, FShaderCompilerEnvironment& OutEnvironment)
-	{
-		FResolveDepthArrayPS::ModifyCompilationEnvironment(Parameters, OutEnvironment);
-		OutEnvironment.SetDefine(TEXT("DEPTH_RESOLVE_NUM_SAMPLES"), 2);
-	}
-
-	FResolveDepthArray2XPS(const ShaderMetaType::CompiledShaderInitializerType& Initializer) :
-		FResolveDepthArrayPS(Initializer)
-	{
-	}
-	FResolveDepthArray2XPS() {}
+	FResolveDepthArray2XPS();
+	FResolveDepthArray2XPS(const ShaderMetaType::CompiledShaderInitializerType& Initializer);
+
+	static bool ShouldCompilePermutation(const FGlobalShaderPermutationParameters& Parameters);
+	static void ModifyCompilationEnvironment(const FGlobalShaderPermutationParameters& Parameters, FShaderCompilerEnvironment& OutEnvironment);
 };
 
-class FResolveDepthArray4XPS : public FResolveDepthArrayPS
+class RENDERCORE_API FResolveDepthArray4XPS : public FResolveDepthArrayPS
 {
-	DECLARE_EXPORTED_SHADER_TYPE(FResolveDepthArray4XPS, Global, RENDERCORE_API);
+	DECLARE_GLOBAL_SHADER(FResolveDepthArray4XPS);
 public:
-
 	typedef FDummyResolveParameter FParameter;
 
-	static bool ShouldCompilePermutation(const FGlobalShaderPermutationParameters& Parameters)
-	{
-		return FResolveDepthArrayPS::ShouldCompilePermutation(Parameters);
-	}
-
-	static void ModifyCompilationEnvironment(const FGlobalShaderPermutationParameters& Parameters, FShaderCompilerEnvironment& OutEnvironment)
-	{
-		FResolveDepthArrayPS::ModifyCompilationEnvironment(Parameters, OutEnvironment);
-		OutEnvironment.SetDefine(TEXT("DEPTH_RESOLVE_NUM_SAMPLES"), 4);
-	}
-
-	FResolveDepthArray4XPS(const ShaderMetaType::CompiledShaderInitializerType& Initializer) :
-		FResolveDepthArrayPS(Initializer)
-	{
-	}
-	FResolveDepthArray4XPS() {}
-};
+	FResolveDepthArray4XPS();
+	FResolveDepthArray4XPS(const ShaderMetaType::CompiledShaderInitializerType& Initializer);
 
+	static bool ShouldCompilePermutation(const FGlobalShaderPermutationParameters& Parameters);
+	static void ModifyCompilationEnvironment(const FGlobalShaderPermutationParameters& Parameters, FShaderCompilerEnvironment& OutEnvironment);
+};
 
-class FResolveDepthArray8XPS : public FResolveDepthArrayPS
+class RENDERCORE_API FResolveDepthArray8XPS : public FResolveDepthArrayPS
 {
-	DECLARE_EXPORTED_SHADER_TYPE(FResolveDepthArray8XPS, Global, RENDERCORE_API);
+	DECLARE_GLOBAL_SHADER(FResolveDepthArray8XPS);
 public:
-
 	typedef FDummyResolveParameter FParameter;
 
-	static bool ShouldCache(EShaderPlatform Platform) { return GetMaxSupportedFeatureLevel(Platform) >= ERHIFeatureLevel::SM5; }
-
-	static bool ShouldCompilePermutation(const FGlobalShaderPermutationParameters& Parameters)
-	{
-		return FResolveDepthArrayPS::ShouldCompilePermutation(Parameters);
-	}
-
-	static void ModifyCompilationEnvironment(const FGlobalShaderPermutationParameters& Parameters, FShaderCompilerEnvironment& OutEnvironment)
-	{
-		FResolveDepthArrayPS::ModifyCompilationEnvironment(Parameters, OutEnvironment);
-		OutEnvironment.SetDefine(TEXT("DEPTH_RESOLVE_NUM_SAMPLES"), 8);
-	}
+	FResolveDepthArray8XPS();
+	FResolveDepthArray8XPS(const ShaderMetaType::CompiledShaderInitializerType& Initializer);
 
-	FResolveDepthArray8XPS(const ShaderMetaType::CompiledShaderInitializerType& Initializer) :
-		FResolveDepthArrayPS(Initializer)
-	{
-	}
-	FResolveDepthArray8XPS() {}
+	static bool ShouldCache(EShaderPlatform Platform);
+	static bool ShouldCompilePermutation(const FGlobalShaderPermutationParameters& Parameters);
+	static void ModifyCompilationEnvironment(const FGlobalShaderPermutationParameters& Parameters, FShaderCompilerEnvironment& OutEnvironment);
 };
 
-class FResolveSingleSamplePS : public FGlobalShader
+class RENDERCORE_API FResolveSingleSamplePS : public FGlobalShader
 {
-	DECLARE_EXPORTED_SHADER_TYPE(FResolveSingleSamplePS, Global, RENDERCORE_API);
+	DECLARE_GLOBAL_SHADER(FResolveSingleSamplePS);
 public:
-	
 	typedef uint32 FParameter;
+
+	FResolveSingleSamplePS();
+	FResolveSingleSamplePS(const ShaderMetaType::CompiledShaderInitializerType& Initializer);
 	
-	static bool ShouldCompilePermutation(const FGlobalShaderPermutationParameters& Parameters)
-	{
-		return FDataDrivenShaderPlatformInfo::GetIsLanguageD3D(Parameters.Platform);
-	}
-	
-	FResolveSingleSamplePS(const ShaderMetaType::CompiledShaderInitializerType& Initializer):
-		FGlobalShader(Initializer)
-	{
-		UnresolvedSurface.Bind(Initializer.ParameterMap,TEXT("UnresolvedSurface"), SPF_Mandatory);
-		SingleSampleIndex.Bind(Initializer.ParameterMap,TEXT("SingleSampleIndex"), SPF_Mandatory);
-	}
-	FResolveSingleSamplePS() {}
-	
-	RENDERCORE_API void SetParameters(FRHICommandList& RHICmdList, uint32 SingleSampleIndexValue);
+	static bool ShouldCompilePermutation(const FGlobalShaderPermutationParameters& Parameters);
+	void SetParameters(FRHICommandList& RHICmdList, uint32 SingleSampleIndexValue);
 	
 	LAYOUT_FIELD(FShaderResourceParameter, UnresolvedSurface);
 	LAYOUT_FIELD(FShaderParameter, SingleSampleIndex);
@@ -246,56 +141,29 @@ public:
 /**
  * A vertex shader for rendering a textured screen element.
  */
-class FResolveVS : public FGlobalShader
+class RENDERCORE_API FResolveVS : public FGlobalShader
 {
-	DECLARE_EXPORTED_SHADER_TYPE(FResolveVS, Global, RENDERCORE_API);
+	DECLARE_GLOBAL_SHADER(FResolveVS);
 public:
+	FResolveVS();
+	FResolveVS(const ShaderMetaType::CompiledShaderInitializerType& Initializer);
 	
-	static bool ShouldCompilePermutation(const FGlobalShaderPermutationParameters& Parameters) { return true; }
-	
-	FResolveVS(const ShaderMetaType::CompiledShaderInitializerType& Initializer)
-		: FGlobalShader(Initializer)
-	{
-		PositionMinMax.Bind(Initializer.ParameterMap, TEXT("PositionMinMax"), SPF_Mandatory);
-		UVMinMax.Bind(Initializer.ParameterMap, TEXT("UVMinMax"), SPF_Mandatory);
-	}
-	FResolveVS() {}
-
-	RENDERCORE_API void SetParameters(FRHICommandList& RHICmdList, const FResolveRect& SrcBounds, const FResolveRect& DstBounds, uint32 DstSurfaceWidth, uint32 DstSurfaceHeight);
+	static bool ShouldCompilePermutation(const FGlobalShaderPermutationParameters& Parameters);
+	void SetParameters(FRHICommandList& RHICmdList, const FResolveRect& SrcBounds, const FResolveRect& DstBounds, uint32 DstSurfaceWidth, uint32 DstSurfaceHeight);
 
 	LAYOUT_FIELD(FShaderParameter, PositionMinMax);
 	LAYOUT_FIELD(FShaderParameter, UVMinMax);
 };
 
-class FResolveArrayVS : public FResolveVS
+class RENDERCORE_API FResolveArrayVS : public FResolveVS
 {
-	DECLARE_EXPORTED_SHADER_TYPE(FResolveArrayVS, Global, RENDERCORE_API);
+	DECLARE_GLOBAL_SHADER(FResolveArrayVS);
 public:
+	FResolveArrayVS();
+	FResolveArrayVS(const ShaderMetaType::CompiledShaderInitializerType& Initializer);
+
+	static bool ShouldCompilePermutation(const FGlobalShaderPermutationParameters& Parameters);
+	static void ModifyCompilationEnvironment(const FGlobalShaderPermutationParameters& Parameters, FShaderCompilerEnvironment& OutEnvironment);
 
-	static bool ShouldCompilePermutation(const FGlobalShaderPermutationParameters& Parameters)
-	{
-		if (FResolveVS::ShouldCompilePermutation(Parameters))
-		{
-			UE::StereoRenderUtils::FStereoShaderAspects Aspects(Parameters.Platform);
-			return Aspects.IsMobileMultiViewEnabled();
-		}
-		return false;
-	}
-
-	static void ModifyCompilationEnvironment(const FGlobalShaderPermutationParameters& Parameters, FShaderCompilerEnvironment& OutEnvironment)
-	{
-		FResolveVS::ModifyCompilationEnvironment(Parameters, OutEnvironment);
-		OutEnvironment.SetDefine(TEXT("DEPTH_RESOLVE_TEXTUREARRAY"), 1);
-	}
-
-	FResolveArrayVS(const ShaderMetaType::CompiledShaderInitializerType& Initializer)
-		: FResolveVS(Initializer)
-	{
-	}
-	FResolveArrayVS() {}
-
-	RENDERCORE_API void SetParameters(FRHICommandList& RHICmdList, const FResolveRect& SrcBounds, const FResolveRect& DstBounds, uint32 DstSurfaceWidth, uint32 DstSurfaceHeight)
-	{
-		return FResolveVS::SetParameters(RHICmdList, SrcBounds, DstBounds, DstSurfaceWidth, DstSurfaceHeight);
-	}
+	void SetParameters(FRHICommandList& RHICmdList, const FResolveRect& SrcBounds, const FResolveRect& DstBounds, uint32 DstSurfaceWidth, uint32 DstSurfaceHeight);
 };
