Moving IsHDREnabled and IsHDRAllowed to HDRHelper.h
Removing the unused FVirtualTextureUniformData type.

diff --git a/Engine/Source/Runtime/D3D12RHI/Private/Windows/WindowsD3D12Viewport.cpp b/Engine/Source/Runtime/D3D12RHI/Private/Windows/WindowsD3D12Viewport.cpp
--- a/Engine/Source/Runtime/D3D12RHI/Private/Windows/WindowsD3D12Viewport.cpp
+++ b/Engine/Source/Runtime/D3D12RHI/Private/Windows/WindowsD3D12Viewport.cpp
@@ -6,6 +6,7 @@
 
 #include "D3D12RHIPrivate.h"
 #include "RenderCore.h"
+#include "HDRHelper.h"
 
 #include "Windows/AllowWindowsPlatformTypes.h"
 #include "Windows.h"
diff --git a/Engine/Source/Runtime/RenderCore/Private/RenderCore.cpp b/Engine/Source/Runtime/RenderCore/Private/RenderCore.cpp
--- a/Engine/Source/Runtime/RenderCore/Private/RenderCore.cpp
+++ b/Engine/Source/Runtime/RenderCore/Private/RenderCore.cpp
@@ -304,8 +304,6 @@ RENDERCORE_API int32 GetCVarForceLODShadow_AnyThread()
 	return Ret;
 }
 
-FMatrix44f FVirtualTextureUniformData::Invalid = FMatrix44f::Identity;
-
 // Note: Enables or disables HDR support for a project. Typically this would be set on a per-project/per-platform basis in defaultengine.ini
 TAutoConsoleVariable<int32> CVarAllowHDR(
 	TEXT("r.AllowHDR"),
diff --git a/Engine/Source/Runtime/RenderCore/Public/HDRHelper.h b/Engine/Source/Runtime/RenderCore/Public/HDRHelper.h
--- a/Engine/Source/Runtime/RenderCore/Public/HDRHelper.h
+++ b/Engine/Source/Runtime/RenderCore/Public/HDRHelper.h
@@ -4,6 +4,9 @@
 
 #include "CoreMinimal.h"
 
+RENDERCORE_API bool IsHDREnabled();
+RENDERCORE_API bool IsHDRAllowed();
+
 RENDERCORE_API void HDRGetMetaData(EDisplayOutputFormat& OutDisplayOutputFormat, EDisplayColorGamut& OutDisplayColorGamut, bool& OutbHDRSupported,
 								   const FVector2D& WindowTopLeft, const FVector2D& WindowBottomRight, void* OSWindow);
 
diff --git a/Engine/Source/Runtime/RenderCore/Public/RenderCore.h b/Engine/Source/Runtime/RenderCore/Public/RenderCore.h
--- a/Engine/Source/Runtime/RenderCore/Public/RenderCore.h
+++ b/Engine/Source/Runtime/RenderCore/Public/RenderCore.h
@@ -18,6 +18,10 @@
 #include "Stats/Stats2.h"
 #include "RenderTimer.h"
 
+#if UE_ENABLE_INCLUDE_ORDER_DEPRECATED_IN_5_2
+#include "HDRHelper.h"
+#endif
+
 DECLARE_LOG_CATEGORY_EXTERN(LogRendererCore, Log, All);
 
 /**
@@ -165,50 +169,3 @@ RENDERCORE_API int32 GetCVarForceLOD_AnyThread();
 
 RENDERCORE_API int32 GetCVarForceLODShadow();
 RENDERCORE_API int32 GetCVarForceLODShadow_AnyThread();
-
-RENDERCORE_API bool IsHDREnabled();
-
-RENDERCORE_API bool IsHDRAllowed();
-
-// struct containing all data the GPU needs to perform a lookup/feedback request
-struct RENDERCORE_API FVirtualTextureUniformData
-{
-	int SpaceID;
-	float PageTableSize;
-	float vPageSize;
-	float pPageBorder;
-	FVector2f pTextureSize; // The size of the cache texture
-	float MaxAnisotropic;
-	int MaxAssetLevel;
-	FVector4f Transform;
-	FVector2f vOneMinusOneOverTextureSize; // 1 - (1/TextureSize) this rather specific value is needed for clamp texturing.
-										   // We could store something more generically useful at the cost of a few shader ALU instructions?
-
-	FMatrix Pack() const
-	{
-		FMatrix Data(ForceInitToZero);
-
-		Data.M[0][0] = SpaceID;
-		Data.M[0][1] = PageTableSize;
-		Data.M[0][2] = vPageSize;
-		Data.M[0][3] = pPageBorder;
-
-		Data.M[1][0] = pTextureSize.X;
-		Data.M[1][1] = pTextureSize.Y;
-		Data.M[1][2] = FMath::Log2(MaxAnisotropic);
-		Data.M[1][3] = MaxAssetLevel;
-
-		Data.M[2][0] = Transform.X;
-		Data.M[2][1] = Transform.Y;
-		Data.M[2][2] = Transform.Z;
-		Data.M[2][3] = Transform.W;
-
-		Data.M[3][0] = vOneMinusOneOverTextureSize.X;
-		Data.M[3][1] = vOneMinusOneOverTextureSize.Y;
-
-		return Data;
-	}
-
-	static FMatrix44f Invalid;
-};
-static_assert(sizeof(FVirtualTextureUniformData) <= sizeof(FMatrix44f), "FVirtualTextureUniformData is unable to pack");
diff --git a/Engine/Source/Runtime/Windows/D3D11RHI/Private/Windows/WindowsD3D11Device.cpp b/Engine/Source/Runtime/Windows/D3D11RHI/Private/Windows/WindowsD3D11Device.cpp
--- a/Engine/Source/Runtime/Windows/D3D11RHI/Private/Windows/WindowsD3D11Device.cpp
+++ b/Engine/Source/Runtime/Windows/D3D11RHI/Private/Windows/WindowsD3D11Device.cpp
@@ -23,6 +23,7 @@
 #include "GenericPlatform/GenericPlatformCrashContext.h"
 #include "RHIValidation.h"
 #include "HAL/ExceptionHandling.h"
+#include "HDRHelper.h"
 
 #if NV_AFTERMATH
 bool GDX11NVAfterMathEnabled = false;
