Removing a lot of includes from RenderResource.h

diff --git a/Engine/Source/Runtime/Engine/Private/Components/LightComponent.cpp b/Engine/Source/Runtime/Engine/Private/Components/LightComponent.cpp
--- a/Engine/Source/Runtime/Engine/Private/Components/LightComponent.cpp
+++ b/Engine/Source/Runtime/Engine/Private/Components/LightComponent.cpp
@@ -8,6 +8,7 @@
 
 #include "Misc/App.h"
 #include "RenderingThread.h"
+#include "DataDrivenShaderPlatformInfo.h"
 #include "Engine/Level.h"
 #include "Engine/MapBuildDataRegistry.h"
 #include "Materials/Material.h"
diff --git a/Engine/Source/Runtime/RenderCore/Private/RenderResource.cpp b/Engine/Source/Runtime/RenderCore/Private/RenderResource.cpp
--- a/Engine/Source/Runtime/RenderCore/Private/RenderResource.cpp
+++ b/Engine/Source/Runtime/RenderCore/Private/RenderResource.cpp
@@ -12,6 +12,7 @@
 #include "CoreGlobals.h"
 #include "RayTracingGeometryManager.h"
 #include "RenderGraphResources.h"
+#include "Containers/ResourceArray.h"
 
 /** Whether to enable mip-level fading or not: +1.0f if enabled, -1.0f if disabled. */
 float GEnableMipLevelFading = 1.0f;
@@ -317,6 +318,11 @@ FRenderResource::~FRenderResource()
 	}
 }
 
+bool FRenderResource::ShouldFreeResourceObject(void* ResourceObject, FResourceArrayInterface* ResourceArray)
+{
+	return GFreeStructuresOnRHIBufferCreation && ResourceObject && (!ResourceArray || ResourceArray->GetResourceDataSize());
+}
+
 FBufferRHIRef FRenderResource::CreateRHIBufferInternal(
 	const TCHAR* InDebugName,
 	const FName& InOwnerName,
diff --git a/Engine/Source/Runtime/RenderCore/Public/RenderResource.h b/Engine/Source/Runtime/RenderCore/Public/RenderResource.h
--- a/Engine/Source/Runtime/RenderCore/Public/RenderResource.h
+++ b/Engine/Source/Runtime/RenderCore/Public/RenderResource.h
@@ -6,13 +6,21 @@
 
 #pragma once
 
+#include "RHIFwd.h"
+#include "RHIShaderPlatform.h"
+#include "RHIFeatureLevel.h"
+#include "RenderTimer.h"
+#include "CoreGlobals.h"
+#include "RenderCore.h"
+
+#if UE_ENABLE_INCLUDE_ORDER_DEPRECATED_IN_5_2
 #include "Containers/Array.h"
+#include "Containers/ResourceArray.h"
 #include "Containers/DynamicRHIResourceArray.h"
 #include "Containers/EnumAsByte.h"
 #include "Containers/List.h"
 #include "Containers/ResourceArray.h"
 #include "Containers/UnrealString.h"
-#include "CoreGlobals.h"
 #include "CoreMinimal.h"
 #include "HAL/CriticalSection.h"
 #include "Math/Color.h"
@@ -23,14 +31,13 @@
 #include "RHI.h"
 #include "RHICommandList.h"
 #include "RHIDefinitions.h"
-#include "RenderCore.h"
 #include "RenderingThread.h"
 #include "Serialization/MemoryLayout.h"
 #include "DataDrivenShaderPlatformInfo.h"
+#endif
 
 class FRDGPooledBuffer;
-
-typedef TBitArray<TInlineAllocator<EShaderPlatform::SP_NumPlatforms / 8>>	ShaderPlatformMaskType;
+class FResourceArrayInterface;
 
 /** Number of frames after which unused global resource allocations will be discarded. */
 extern int32 GGlobalBufferNumFramesUnusedThresold;
@@ -167,7 +174,7 @@ protected:
 		}
 
 		// If the buffer creation emptied the resource array, delete the containing structure as well
-		if (GFreeStructuresOnRHIBufferCreation && InOutResourceObject && (!ResourceArray || !ResourceArray->GetResourceDataSize()))
+		if (ShouldFreeResourceObject(InOutResourceObject, ResourceArray))
 		{
 			delete InOutResourceObject;
 			InOutResourceObject = nullptr;
@@ -177,6 +184,7 @@ protected:
 	}
 
 private:
+	static bool ShouldFreeResourceObject(void* ResourceObject, FResourceArrayInterface* ResourceArray);
 	static FBufferRHIRef CreateRHIBufferInternal(
 		const TCHAR* InDebugName,
 		const FName& InOwnerName,
