Removing redundant IsFeatureLevelSupported checks.

diff --git a/Engine/Source/Runtime/RenderCore/Public/MediaShaders.h b/Engine/Source/Runtime/RenderCore/Public/MediaShaders.h
--- a/Engine/Source/Runtime/RenderCore/Public/MediaShaders.h
+++ b/Engine/Source/Runtime/RenderCore/Public/MediaShaders.h
@@ -292,11 +292,6 @@ class FP010ConvertPS
 
 public:
 
-	static bool ShouldCompilePermutation(const FGlobalShaderPermutationParameters& Parameters)
-	{
-		return IsFeatureLevelSupported(Parameters.Platform, ERHIFeatureLevel::ES3_1);
-	}
-
 	FP010ConvertPS() { }
 
 	FP010ConvertPS(const ShaderMetaType::CompiledShaderInitializerType& Initializer)
@@ -316,11 +311,6 @@ class FP010ConvertAsUINT16sPS
 
 public:
 
-	static bool ShouldCompilePermutation(const FGlobalShaderPermutationParameters& Parameters)
-	{
-		return IsFeatureLevelSupported(Parameters.Platform, ERHIFeatureLevel::ES3_1);
-	}
-
 	FP010ConvertAsUINT16sPS() { }
 
 	FP010ConvertAsUINT16sPS(const ShaderMetaType::CompiledShaderInitializerType& Initializer)
@@ -340,11 +330,6 @@ class FP010_2101010ConvertPS
 
 public:
 
-	static bool ShouldCompilePermutation(const FGlobalShaderPermutationParameters& Parameters)
-	{
-		return IsFeatureLevelSupported(Parameters.Platform, ERHIFeatureLevel::ES3_1);
-	}
-
 	FP010_2101010ConvertPS() { }
 
 	FP010_2101010ConvertPS(const ShaderMetaType::CompiledShaderInitializerType& Initializer)
