Reimplementing and correcting a backed out change that removes the header dependency on FResourceArrayInterface.

diff --git a/Engine/Source/Runtime/RenderCore/Private/RenderResource.cpp b/Engine/Source/Runtime/RenderCore/Private/RenderResource.cpp
--- a/Engine/Source/Runtime/RenderCore/Private/RenderResource.cpp
+++ b/Engine/Source/Runtime/RenderCore/Private/RenderResource.cpp
@@ -12,6 +12,7 @@
 #include "CoreGlobals.h"
 #include "RayTracingGeometryManager.h"
 #include "RenderGraphResources.h"
+#include "Containers/ResourceArray.h"
 
 /** Whether to enable mip-level fading or not: +1.0f if enabled, -1.0f if disabled. */
 float GEnableMipLevelFading = 1.0f;
@@ -317,6 +318,11 @@ FRenderResource::~FRenderResource()
 	}
 }
 
+bool FRenderResource::ShouldFreeResourceObject(void* ResourceObject, FResourceArrayInterface* ResourceArray)
+{
+	return GFreeStructuresOnRHIBufferCreation && ResourceObject && (!ResourceArray || !ResourceArray->GetResourceDataSize());
+}
+
 FBufferRHIRef FRenderResource::CreateRHIBufferInternal(
 	const TCHAR* InDebugName,
 	const FName& InOwnerName,
diff --git a/Engine/Source/Runtime/RenderCore/Public/RenderResource.h b/Engine/Source/Runtime/RenderCore/Public/RenderResource.h
--- a/Engine/Source/Runtime/RenderCore/Public/RenderResource.h
+++ b/Engine/Source/Runtime/RenderCore/Public/RenderResource.h
@@ -167,7 +167,7 @@ protected:
 		}
 
 		// If the buffer creation emptied the resource array, delete the containing structure as well
-		if (GFreeStructuresOnRHIBufferCreation && InOutResourceObject && (!ResourceArray || !ResourceArray->GetResourceDataSize()))
+		if (ShouldFreeResourceObject(InOutResourceObject, ResourceArray))
 		{
 			delete InOutResourceObject;
 			InOutResourceObject = nullptr;
@@ -177,6 +177,7 @@ protected:
 	}
 
 private:
+	static bool ShouldFreeResourceObject(void* ResourceObject, FResourceArrayInterface* ResourceArray);
 	static FBufferRHIRef CreateRHIBufferInternal(
 		const TCHAR* InDebugName,
 		const FName& InOwnerName,
