Removing a lot of includes from RenderResource.h

diff --git a/Engine/Source/Runtime/Engine/Private/Components/LightComponent.cpp b/Engine/Source/Runtime/Engine/Private/Components/LightComponent.cpp
--- a/Engine/Source/Runtime/Engine/Private/Components/LightComponent.cpp
+++ b/Engine/Source/Runtime/Engine/Private/Components/LightComponent.cpp
@@ -8,6 +8,7 @@
 
 #include "Misc/App.h"
 #include "RenderingThread.h"
+#include "DataDrivenShaderPlatformInfo.h"
 #include "Engine/Level.h"
 #include "Engine/MapBuildDataRegistry.h"
 #include "Materials/Material.h"
diff --git a/Engine/Source/Runtime/RenderCore/Public/RenderResource.h b/Engine/Source/Runtime/RenderCore/Public/RenderResource.h
--- a/Engine/Source/Runtime/RenderCore/Public/RenderResource.h
+++ b/Engine/Source/Runtime/RenderCore/Public/RenderResource.h
@@ -6,13 +6,21 @@
 
 #pragma once
 
+#include "RHIFwd.h"
+#include "RHIShaderPlatform.h"
+#include "RHIFeatureLevel.h"
+#include "RenderTimer.h"
+#include "CoreGlobals.h"
+#include "RenderCore.h"
+
+#if UE_ENABLE_INCLUDE_ORDER_DEPRECATED_IN_5_2
 #include "Containers/Array.h"
+#include "Containers/ResourceArray.h"
 #include "Containers/DynamicRHIResourceArray.h"
 #include "Containers/EnumAsByte.h"
 #include "Containers/List.h"
 #include "Containers/ResourceArray.h"
 #include "Containers/UnrealString.h"
-#include "CoreGlobals.h"
 #include "CoreMinimal.h"
 #include "HAL/CriticalSection.h"
 #include "Math/Color.h"
@@ -23,14 +31,13 @@
 #include "RHI.h"
 #include "RHICommandList.h"
 #include "RHIDefinitions.h"
-#include "RenderCore.h"
 #include "RenderingThread.h"
 #include "Serialization/MemoryLayout.h"
 #include "DataDrivenShaderPlatformInfo.h"
+#endif
 
 class FRDGPooledBuffer;
-
-typedef TBitArray<TInlineAllocator<EShaderPlatform::SP_NumPlatforms / 8>>	ShaderPlatformMaskType;
+class FResourceArrayInterface;
 
 /** Number of frames after which unused global resource allocations will be discarded. */
 extern int32 GGlobalBufferNumFramesUnusedThresold;
