Moving FDeferredCleanupInterface and accompanying types to RenderDeferredCleanup.h to remove global dependencies on RenderThread.h

diff --git a/Engine/Plugins/Compositing/OpenColorIO/Source/OpenColorIO/Public/OpenColorIOShared.h b/Engine/Plugins/Compositing/OpenColorIO/Source/OpenColorIO/Public/OpenColorIOShared.h
--- a/Engine/Plugins/Compositing/OpenColorIO/Source/OpenColorIO/Public/OpenColorIOShared.h
+++ b/Engine/Plugins/Compositing/OpenColorIO/Source/OpenColorIO/Public/OpenColorIOShared.h
@@ -12,7 +12,10 @@
 #include "Misc/Optional.h"
 #include "Misc/SecureHash.h"
 #include "RenderResource.h"
+#if UE_ENABLE_INCLUDE_ORDER_DEPRECATED_IN_5_2
 #include "RenderingThread.h"
+#endif
+#include "RenderDeferredCleanup.h"
 #include "RHI.h"
 #include "SceneTypes.h"
 #include "Shader.h"
diff --git a/Engine/Plugins/FX/Niagara/Source/NiagaraShader/Public/NiagaraShared.h b/Engine/Plugins/FX/Niagara/Source/NiagaraShader/Public/NiagaraShared.h
--- a/Engine/Plugins/FX/Niagara/Source/NiagaraShader/Public/NiagaraShared.h
+++ b/Engine/Plugins/FX/Niagara/Source/NiagaraShader/Public/NiagaraShared.h
@@ -15,7 +15,10 @@
 #include "MaterialShared.h"
 #include "RHI.h"
 #include "RenderResource.h"
+#if UE_ENABLE_INCLUDE_ORDER_DEPRECATED_IN_5_2
 #include "RenderingThread.h"
+#endif
+#include "RenderDeferredCleanup.h"
 #include "UniformBuffer.h"
 #include "Shader.h"
 #include "VertexFactory.h"
diff --git a/Engine/Plugins/Media/BinkMedia/Source/BinkMediaPlayer/Private/Assets/BinkMediaTexture.cpp b/Engine/Plugins/Media/BinkMedia/Source/BinkMediaPlayer/Private/Assets/BinkMediaTexture.cpp
--- a/Engine/Plugins/Media/BinkMedia/Source/BinkMediaPlayer/Private/Assets/BinkMediaTexture.cpp
+++ b/Engine/Plugins/Media/BinkMedia/Source/BinkMediaPlayer/Private/Assets/BinkMediaTexture.cpp
@@ -6,6 +6,7 @@
 #include "BinkMediaPlayer.h"
 #include "BinkMediaPlayerPrivate.h"
 #include "BinkMediaTextureResource.h"
+#include "RenderingThread.h"
 
 UBinkMediaTexture::UBinkMediaTexture( const FObjectInitializer& ObjectInitializer )
 	: Super(ObjectInitializer)
diff --git a/Engine/Plugins/Media/BinkMedia/Source/BinkMediaPlayer/Private/Assets/BinkMediaTextureResource.cpp b/Engine/Plugins/Media/BinkMedia/Source/BinkMediaPlayer/Private/Assets/BinkMediaTextureResource.cpp
--- a/Engine/Plugins/Media/BinkMedia/Source/BinkMediaPlayer/Private/Assets/BinkMediaTextureResource.cpp
+++ b/Engine/Plugins/Media/BinkMedia/Source/BinkMediaPlayer/Private/Assets/BinkMediaTextureResource.cpp
@@ -9,6 +9,7 @@
 #include "Runtime/Launch/Resources/Version.h"
 #include "BinkMediaPlayer.h"
 #include "BinkMediaTexture.h"
+#include "RenderingThread.h"
 
 void FBinkMediaTextureResource::InitDynamicRHI() 
 {
diff --git a/Engine/Plugins/Media/BinkMedia/Source/BinkMediaPlayer/Private/BinkMediaPlayerModule.cpp b/Engine/Plugins/Media/BinkMedia/Source/BinkMediaPlayer/Private/BinkMediaPlayerModule.cpp
--- a/Engine/Plugins/Media/BinkMedia/Source/BinkMediaPlayer/Private/BinkMediaPlayerModule.cpp
+++ b/Engine/Plugins/Media/BinkMedia/Source/BinkMediaPlayer/Private/BinkMediaPlayerModule.cpp
@@ -8,6 +8,7 @@
 #include "Misc/Paths.h"
 #include "Interfaces/IPluginManager.h"
 #include "Misc/CoreDelegates.h"
+#include "RenderingThread.h"
 #include "UObject/UObjectIterator.h"
 
 DEFINE_LOG_CATEGORY(LogBink);
diff --git a/Engine/Plugins/Media/BinkMedia/Source/BinkMediaPlayerEditor/Private/Models/BinkMediaPlayerEditorTexture.h b/Engine/Plugins/Media/BinkMedia/Source/BinkMediaPlayerEditor/Private/Models/BinkMediaPlayerEditorTexture.h
--- a/Engine/Plugins/Media/BinkMedia/Source/BinkMediaPlayerEditor/Private/Models/BinkMediaPlayerEditorTexture.h
+++ b/Engine/Plugins/Media/BinkMedia/Source/BinkMediaPlayerEditor/Private/Models/BinkMediaPlayerEditorTexture.h
@@ -4,6 +4,7 @@
 #pragma once
 
 #include "BinkMediaPlayer.h"
+#include "RenderingThread.h"
 
 struct FBinkMediaPlayerEditorTexture : public FTickableGameObject 
 {
diff --git a/Engine/Plugins/Media/BinkMedia/Source/BinkMediaPlayerEditor/Private/SBinkWidgets.cpp b/Engine/Plugins/Media/BinkMedia/Source/BinkMediaPlayerEditor/Private/SBinkWidgets.cpp
--- a/Engine/Plugins/Media/BinkMedia/Source/BinkMediaPlayerEditor/Private/SBinkWidgets.cpp
+++ b/Engine/Plugins/Media/BinkMedia/Source/BinkMediaPlayerEditor/Private/SBinkWidgets.cpp
@@ -6,6 +6,7 @@
 #include "BinkMediaPlayerEditorPrivate.h"
 #include "Widgets/Input/SNumericEntryBox.h"
 #include "Widgets/SViewport.h"
+#include "RenderingThread.h"
 
 #define LOCTEXT_NAMESPACE "SBinkMediaPlayerWidgets"
 
diff --git a/Engine/Plugins/Media/MediaMovieStreamer/Source/MediaMovieStreamer/Private/MediaMovieStreamer.cpp b/Engine/Plugins/Media/MediaMovieStreamer/Source/MediaMovieStreamer/Private/MediaMovieStreamer.cpp
--- a/Engine/Plugins/Media/MediaMovieStreamer/Source/MediaMovieStreamer/Private/MediaMovieStreamer.cpp
+++ b/Engine/Plugins/Media/MediaMovieStreamer/Source/MediaMovieStreamer/Private/MediaMovieStreamer.cpp
@@ -12,6 +12,7 @@
 #include "MediaTexture.h"
 
 #include "Rendering/RenderingCommon.h"
+#include "RenderingThread.h"
 #include "Slate/SlateTextures.h"
 
 DEFINE_LOG_CATEGORY(LogMediaMovieStreamer);
diff --git a/Engine/Plugins/MovieScene/MovieRenderPipeline/Source/MovieRenderPipelineSettings/Private/MoviePipelineBurnInSetting.cpp b/Engine/Plugins/MovieScene/MovieRenderPipeline/Source/MovieRenderPipelineSettings/Private/MoviePipelineBurnInSetting.cpp
--- a/Engine/Plugins/MovieScene/MovieRenderPipeline/Source/MovieRenderPipelineSettings/Private/MoviePipelineBurnInSetting.cpp
+++ b/Engine/Plugins/MovieScene/MovieRenderPipeline/Source/MovieRenderPipelineSettings/Private/MoviePipelineBurnInSetting.cpp
@@ -16,6 +16,7 @@
 #include "MovieRenderPipelineCoreModule.h"
 #include "MoviePipelineQueue.h"
 #include "TextureResource.h"
+#include "RenderingThread.h"
 
 #include UE_INLINE_GENERATED_CPP_BY_NAME(MoviePipelineBurnInSetting)
 
diff --git a/Engine/Plugins/MovieScene/MovieRenderPipeline/Source/MovieRenderPipelineSettings/Private/MoviePipelineWidgetRenderSetting.cpp b/Engine/Plugins/MovieScene/MovieRenderPipeline/Source/MovieRenderPipelineSettings/Private/MoviePipelineWidgetRenderSetting.cpp
--- a/Engine/Plugins/MovieScene/MovieRenderPipeline/Source/MovieRenderPipelineSettings/Private/MoviePipelineWidgetRenderSetting.cpp
+++ b/Engine/Plugins/MovieScene/MovieRenderPipeline/Source/MovieRenderPipelineSettings/Private/MoviePipelineWidgetRenderSetting.cpp
@@ -18,6 +18,7 @@
 #include "MovieRenderPipelineCoreModule.h"
 #include "MoviePipelineQueue.h"
 #include "TextureResource.h"
+#include "RenderingThread.h"
 
 #include UE_INLINE_GENERATED_CPP_BY_NAME(MoviePipelineWidgetRenderSetting)
 
diff --git a/Engine/Plugins/Runtime/ComputeFramework/Source/ComputeFramework/Private/ComputeFramework/ComputeKernelShared.h b/Engine/Plugins/Runtime/ComputeFramework/Source/ComputeFramework/Private/ComputeFramework/ComputeKernelShared.h
--- a/Engine/Plugins/Runtime/ComputeFramework/Source/ComputeFramework/Private/ComputeFramework/ComputeKernelShared.h
+++ b/Engine/Plugins/Runtime/ComputeFramework/Source/ComputeFramework/Private/ComputeFramework/ComputeKernelShared.h
@@ -6,6 +6,7 @@
 #include "Engine/EngineTypes.h"
 #include "RenderResource.h"
 #include "RenderingThread.h"
+#include "RenderDeferredCleanup.h"
 #include "RHI.h"
 #include "SceneTypes.h"
 #include "Shader.h"
diff --git a/Engine/Plugins/Runtime/nDisplay/Source/DisplayCluster/Private/Game/EngineClasses/Scene/DisplayClusterWidgetComponent.cpp b/Engine/Plugins/Runtime/nDisplay/Source/DisplayCluster/Private/Game/EngineClasses/Scene/DisplayClusterWidgetComponent.cpp
--- a/Engine/Plugins/Runtime/nDisplay/Source/DisplayCluster/Private/Game/EngineClasses/Scene/DisplayClusterWidgetComponent.cpp
+++ b/Engine/Plugins/Runtime/nDisplay/Source/DisplayCluster/Private/Game/EngineClasses/Scene/DisplayClusterWidgetComponent.cpp
@@ -4,6 +4,9 @@
 
 #include "Engine/Engine.h"
 #include "Slate/WidgetRenderer.h"
+#if !UE_SERVER
+#include "RHI.h"
+#endif
 
 UDisplayClusterWidgetComponent::UDisplayClusterWidgetComponent()
 {
diff --git a/Engine/Source/Developer/FunctionalTesting/Private/FunctionalUIScreenshotTest.cpp b/Engine/Source/Developer/FunctionalTesting/Private/FunctionalUIScreenshotTest.cpp
--- a/Engine/Source/Developer/FunctionalTesting/Private/FunctionalUIScreenshotTest.cpp
+++ b/Engine/Source/Developer/FunctionalTesting/Private/FunctionalUIScreenshotTest.cpp
@@ -15,6 +15,7 @@
 #include "Slate/SceneViewport.h"
 #include "Slate/WidgetRenderer.h"
 #include "TextureResource.h"
+#include "RenderingThread.h"
 
 #include UE_INLINE_GENERATED_CPP_BY_NAME(FunctionalUIScreenshotTest)
 
diff --git a/Engine/Source/Editor/MovieSceneTools/Private/TrackEditors/AudioTrackEditor.cpp b/Engine/Source/Editor/MovieSceneTools/Private/TrackEditors/AudioTrackEditor.cpp
--- a/Engine/Source/Editor/MovieSceneTools/Private/TrackEditors/AudioTrackEditor.cpp
+++ b/Engine/Source/Editor/MovieSceneTools/Private/TrackEditors/AudioTrackEditor.cpp
@@ -7,6 +7,7 @@
 #include "Widgets/SBoxPanel.h"
 #include "Framework/MultiBox/MultiBoxBuilder.h"
 #include "RenderUtils.h"
+#include "RenderingThread.h"
 #include "Modules/ModuleManager.h"
 #include "Audio.h"
 #include "Sound/SoundBase.h"
diff --git a/Engine/Source/Editor/UnrealEd/Private/FeedbackContextEditor.cpp b/Engine/Source/Editor/UnrealEd/Private/FeedbackContextEditor.cpp
--- a/Engine/Source/Editor/UnrealEd/Private/FeedbackContextEditor.cpp
+++ b/Engine/Source/Editor/UnrealEd/Private/FeedbackContextEditor.cpp
@@ -7,6 +7,7 @@
 #include "Fonts/FontMeasure.h"
 #include "Framework/Application/SlateApplication.h"
 #include "RenderingThread.h"
+#include "RenderDeferredCleanup.h"
 #include "RHI.h"
 #include "Widgets/Layout/SBorder.h"
 #include "Widgets/Notifications/SProgressBar.h"
diff --git a/Engine/Source/Runtime/Engine/Classes/Components/SkyLightComponent.h b/Engine/Source/Runtime/Engine/Classes/Components/SkyLightComponent.h
--- a/Engine/Source/Runtime/Engine/Classes/Components/SkyLightComponent.h
+++ b/Engine/Source/Runtime/Engine/Classes/Components/SkyLightComponent.h
@@ -10,6 +10,7 @@
 #include "RenderingThread.h"
 #include "Components/LightComponentBase.h"
 #include "Math/SHMath.h"
+#include "RenderDeferredCleanup.h"
 #include "SkyLightComponent.generated.h"
 
 class FSkyLightSceneProxy;
diff --git a/Engine/Source/Runtime/Engine/Private/RendererSupport.cpp b/Engine/Source/Runtime/Engine/Private/RendererSupport.cpp
--- a/Engine/Source/Runtime/Engine/Private/RendererSupport.cpp
+++ b/Engine/Source/Runtime/Engine/Private/RendererSupport.cpp
@@ -11,7 +11,7 @@
 #include "UObject/UObjectHash.h"
 #include "Engine/Level.h"
 #include "EngineGlobals.h"
-#include "RenderingThread.h"
+#include "RenderDeferredCleanup.h"
 #include "Containers/List.h"
 #include "Shader.h"
 #include "VertexFactory.h"
diff --git a/Engine/Source/Runtime/Engine/Public/DistanceFieldAtlas.h b/Engine/Source/Runtime/Engine/Public/DistanceFieldAtlas.h
--- a/Engine/Source/Runtime/Engine/Public/DistanceFieldAtlas.h
+++ b/Engine/Source/Runtime/Engine/Public/DistanceFieldAtlas.h
@@ -12,7 +12,10 @@
 #include "Engine/EngineTypes.h"
 #include "UObject/GCObject.h"
 #include "RenderResource.h"
+#if UE_ENABLE_INCLUDE_ORDER_DEPRECATED_IN_5_2
 #include "RenderingThread.h"
+#endif
+#include "RenderDeferredCleanup.h"
 #include "TextureLayout3d.h"
 #include "AsyncCompilationHelpers.h"
 #include "AssetCompilingManager.h"
diff --git a/Engine/Source/Runtime/Engine/Public/LightMap.h b/Engine/Source/Runtime/Engine/Public/LightMap.h
--- a/Engine/Source/Runtime/Engine/Public/LightMap.h
+++ b/Engine/Source/Runtime/Engine/Public/LightMap.h
@@ -10,7 +10,10 @@
 #include "Containers/ChunkedArray.h"
 #include "Misc/Guid.h"
 #include "Engine/EngineTypes.h"
+#if UE_ENABLE_INCLUDE_ORDER_DEPRECATED_IN_5_2
 #include "RenderingThread.h"
+#endif
+#include "RenderDeferredCleanup.h"
 #include "Serialization/BulkData.h"
 #include "SceneManagement.h"
 #include "Engine/LightMapTexture2D.h"
diff --git a/Engine/Source/Runtime/Engine/Public/MaterialShared.h b/Engine/Source/Runtime/Engine/Public/MaterialShared.h
--- a/Engine/Source/Runtime/Engine/Public/MaterialShared.h
+++ b/Engine/Source/Runtime/Engine/Public/MaterialShared.h
@@ -16,7 +16,10 @@
 #include "Misc/SecureHash.h"
 #include "RHIDefinitions.h"
 #include "RenderResource.h"
+#if UE_ENABLE_INCLUDE_ORDER_DEPRECATED_IN_5_2
 #include "RenderingThread.h"
+#endif
+#include "RenderDeferredCleanup.h"
 #include "UniformBuffer.h"
 #include "Shader.h"
 #include "VertexFactory.h"
diff --git a/Engine/Source/Runtime/Engine/Public/MeshCardRepresentation.h b/Engine/Source/Runtime/Engine/Public/MeshCardRepresentation.h
--- a/Engine/Source/Runtime/Engine/Public/MeshCardRepresentation.h
+++ b/Engine/Source/Runtime/Engine/Public/MeshCardRepresentation.h
@@ -14,7 +14,10 @@
 #include "AsyncCompilationHelpers.h"
 #include "AssetCompilingManager.h"
 #include "RenderResource.h"
+#if UE_ENABLE_INCLUDE_ORDER_DEPRECATED_IN_5_2
 #include "RenderingThread.h"
+#endif
+#include "RenderDeferredCleanup.h"
 #include "Templates/UniquePtr.h"
 #include "DerivedMeshDataTaskUtils.h"
 #include "Async/AsyncWork.h"
diff --git a/Engine/Source/Runtime/Engine/Public/ShadowMap.h b/Engine/Source/Runtime/Engine/Public/ShadowMap.h
--- a/Engine/Source/Runtime/Engine/Public/ShadowMap.h
+++ b/Engine/Source/Runtime/Engine/Public/ShadowMap.h
@@ -4,7 +4,10 @@
 
 #include "CoreMinimal.h"
 #include "Misc/Guid.h"
+#if UE_ENABLE_INCLUDE_ORDER_DEPRECATED_IN_5_2
 #include "RenderingThread.h"
+#endif
+#include "RenderDeferredCleanup.h"
 #include "SceneManagement.h"
 
 class FShadowMap2D;
diff --git a/Engine/Source/Runtime/Engine/Public/SkeletalRenderPublic.h b/Engine/Source/Runtime/Engine/Public/SkeletalRenderPublic.h
--- a/Engine/Source/Runtime/Engine/Public/SkeletalRenderPublic.h
+++ b/Engine/Source/Runtime/Engine/Public/SkeletalRenderPublic.h
@@ -9,7 +9,10 @@
 #include "Stats/Stats.h"
 #include "ProfilingDebugging/ResourceSize.h"
 #include "PackedNormal.h"
+#if UE_ENABLE_INCLUDE_ORDER_DEPRECATED_IN_5_2
 #include "RenderingThread.h"
+#endif
+#include "RenderDeferredCleanup.h"
 #include "RenderUtils.h"
 #include "Engine/SkeletalMesh.h"
 #include "Components/SkinnedMeshComponent.h"
diff --git a/Engine/Source/Runtime/Engine/Public/Slate/DeferredCleanupSlateBrush.h b/Engine/Source/Runtime/Engine/Public/Slate/DeferredCleanupSlateBrush.h
--- a/Engine/Source/Runtime/Engine/Public/Slate/DeferredCleanupSlateBrush.h
+++ b/Engine/Source/Runtime/Engine/Public/Slate/DeferredCleanupSlateBrush.h
@@ -2,7 +2,10 @@
 
 #pragma once
 
+#if UE_ENABLE_INCLUDE_ORDER_DEPRECATED_IN_5_2
 #include "RenderingThread.h"
+#endif
+#include "RenderDeferredCleanup.h"
 #include "UObject/GCObject.h"
 #include "Styling/SlateBrush.h"
 
diff --git a/Engine/Source/Runtime/Engine/Public/Slate/SlateTextures.h b/Engine/Source/Runtime/Engine/Public/Slate/SlateTextures.h
--- a/Engine/Source/Runtime/Engine/Public/Slate/SlateTextures.h
+++ b/Engine/Source/Runtime/Engine/Public/Slate/SlateTextures.h
@@ -5,7 +5,10 @@
 #include "CoreMinimal.h"
 #include "RenderResource.h"
 #include "Textures/SlateTextureData.h"
+#if UE_ENABLE_INCLUDE_ORDER_DEPRECATED_IN_5_2
 #include "RenderingThread.h"
+#endif
+#include "RenderDeferredCleanup.h"
 #include "Textures/SlateShaderResource.h"
 #include "TextureResource.h"
 #include "Engine/Texture.h"
diff --git a/Engine/Source/Runtime/Launch/Private/LaunchEngineLoop.cpp b/Engine/Source/Runtime/Launch/Private/LaunchEngineLoop.cpp
--- a/Engine/Source/Runtime/Launch/Private/LaunchEngineLoop.cpp
+++ b/Engine/Source/Runtime/Launch/Private/LaunchEngineLoop.cpp
@@ -104,6 +104,7 @@
 	#include "Input/Reply.h"
 	#include "Styling/CoreStyle.h"
 	#include "RenderingThread.h"
+	#include "RenderDeferredCleanup.h"
 	#include "Editor/EditorEngine.h"
 	#include "UnrealEdMisc.h"
 	#include "UnrealEdGlobals.h"
diff --git a/Engine/Source/Runtime/MoviePlayer/Private/MoviePlayer.cpp b/Engine/Source/Runtime/MoviePlayer/Private/MoviePlayer.cpp
--- a/Engine/Source/Runtime/MoviePlayer/Private/MoviePlayer.cpp
+++ b/Engine/Source/Runtime/MoviePlayer/Private/MoviePlayer.cpp
@@ -11,6 +11,7 @@
 #include "NullMoviePlayer.h"
 #include "DefaultGameMoviePlayer.h"
 #include "Widgets/Images/SThrobber.h"
+#include "RenderingThread.h"
 
 IMPLEMENT_MODULE(FDefaultModuleImpl, MoviePlayer);
 
diff --git a/Engine/Source/Runtime/PreLoadScreen/Public/PreLoadSettingsContainer.h b/Engine/Source/Runtime/PreLoadScreen/Public/PreLoadSettingsContainer.h
--- a/Engine/Source/Runtime/PreLoadScreen/Public/PreLoadSettingsContainer.h
+++ b/Engine/Source/Runtime/PreLoadScreen/Public/PreLoadSettingsContainer.h
@@ -3,7 +3,10 @@
 #pragma once
 
 #include "CoreMinimal.h"
+#if UE_ENABLE_INCLUDE_ORDER_DEPRECATED_IN_5_2
 #include "RenderingThread.h"
+#endif
+#include "RenderDeferredCleanup.h"
 #include "UObject/GCObject.h"
 
 struct FSlateDynamicImageBrush;
diff --git a/Engine/Source/Runtime/RenderCore/Private/RenderingThread.cpp b/Engine/Source/Runtime/RenderCore/Private/RenderingThread.cpp
--- a/Engine/Source/Runtime/RenderCore/Private/RenderingThread.cpp
+++ b/Engine/Source/Runtime/RenderCore/Private/RenderingThread.cpp
@@ -17,6 +17,7 @@
 #include "RHI.h"
 #include "RenderCore.h"
 #include "RenderCommandFence.h"
+#include "RenderDeferredCleanup.h"
 #include "TickableObjectRenderThread.h"
 #include "Stats/StatsData.h"
 #include "HAL/ThreadHeartBeat.h"
diff --git a/Engine/Source/Runtime/RenderCore/Public/RenderDeferredCleanup.h b/Engine/Source/Runtime/RenderCore/Public/RenderDeferredCleanup.h
new file mode 100644
--- /dev/null
+++ b/Engine/Source/Runtime/RenderCore/Public/RenderDeferredCleanup.h
@@ -0,0 +1,38 @@
+// Copyright Epic Games, Inc. All Rights Reserved.
+
+#pragma once
+
+#include "Containers/Array.h"
+
+/**
+ * The base class of objects that need to defer deletion until the render command queue has been flushed.
+ */
+class RENDERCORE_API FDeferredCleanupInterface
+{
+public:
+	virtual ~FDeferredCleanupInterface() {}
+};
+
+/**
+ * A set of cleanup objects which are pending deletion.
+ */
+class FPendingCleanupObjects
+{
+	TArray<FDeferredCleanupInterface*> CleanupArray;
+
+public:
+	inline bool IsEmpty() const { return CleanupArray.IsEmpty(); }
+	FPendingCleanupObjects();
+	RENDERCORE_API ~FPendingCleanupObjects();
+};
+
+/**
+ * Adds the specified deferred cleanup object to the current set of pending cleanup objects.
+ */
+extern RENDERCORE_API void BeginCleanup(FDeferredCleanupInterface* CleanupObject);
+
+/**
+ * Transfers ownership of the current set of pending cleanup objects to the caller.  A new set is created for subsequent BeginCleanup calls.
+ * @return A pointer to the set of pending cleanup objects.  The called is responsible for deletion.
+ */
+extern RENDERCORE_API FPendingCleanupObjects* GetPendingCleanupObjects();
diff --git a/Engine/Source/Runtime/RenderCore/Public/RenderingThread.h b/Engine/Source/Runtime/RenderCore/Public/RenderingThread.h
--- a/Engine/Source/Runtime/RenderCore/Public/RenderingThread.h
+++ b/Engine/Source/Runtime/RenderCore/Public/RenderingThread.h
@@ -26,6 +26,9 @@
 #include "Templates/UnrealTemplate.h"
 #include "Templates/UnrealTypeTraits.h"
 #include "Trace/Trace.h"
+#if UE_ENABLE_INCLUDE_ORDER_DEPRECATED_IN_5_2
+#include "RenderDeferredCleanup.h"
+#endif
 
 namespace UE { namespace Trace { class FChannel; } }
 
@@ -296,42 +299,6 @@ FORCEINLINE_DEBUGGABLE void EnqueueUniqueRenderCommand(LAMBDA& Lambda)
 	static_assert(sizeof(LAMBDA) == 0, "EnqueueUniqueRenderCommand enforces use of rvalue and therefore move to avoid an extra copy of the Lambda");
 }
 
-////////////////////////////////////
-// Deferred cleanup
-////////////////////////////////////
-
-/**
- * The base class of objects that need to defer deletion until the render command queue has been flushed.
- */
-class RENDERCORE_API FDeferredCleanupInterface
-{
-public:
-	virtual ~FDeferredCleanupInterface() {}
-};
-
-/**
- * A set of cleanup objects which are pending deletion.
- */
-class FPendingCleanupObjects
-{
-	TArray<FDeferredCleanupInterface*> CleanupArray;
-public:
-	inline bool IsEmpty() const { return CleanupArray.IsEmpty(); }
-	FPendingCleanupObjects();
-	RENDERCORE_API ~FPendingCleanupObjects();
-};
-
-/**
- * Adds the specified deferred cleanup object to the current set of pending cleanup objects.
- */
-extern RENDERCORE_API void BeginCleanup(FDeferredCleanupInterface* CleanupObject);
-
-/**
- * Transfers ownership of the current set of pending cleanup objects to the caller.  A new set is created for subsequent BeginCleanup calls.
- * @return A pointer to the set of pending cleanup objects.  The called is responsible for deletion.
- */
-extern RENDERCORE_API FPendingCleanupObjects* GetPendingCleanupObjects();
-
 ////////////////////////////////////
 // RenderThread scoped work
 ////////////////////////////////////
@@ -540,4 +507,4 @@ using FStopRenderingThreadDelegate = FStopRenderingThread::FDelegate;
 
 extern RENDERCORE_API FDelegateHandle RegisterStopRenderingThreadDelegate(const FStopRenderingThreadDelegate& InDelegate);
 
-extern RENDERCORE_API void UnregisterStopRenderingThreadDelegate(FDelegateHandle InDelegateHandle);
+extern RENDERCORE_API void UnregisterStopRenderingThreadDelegate(FDelegateHandle InDelegateHandle);
\ No newline at end of file
diff --git a/Engine/Source/Runtime/RenderCore/Public/Shader.h b/Engine/Source/Runtime/RenderCore/Public/Shader.h
--- a/Engine/Source/Runtime/RenderCore/Public/Shader.h
+++ b/Engine/Source/Runtime/RenderCore/Public/Shader.h
@@ -32,7 +32,10 @@
 #include "RHI.h"
 #include "RHIDefinitions.h"
 #include "RenderResource.h"
+#if UE_ENABLE_INCLUDE_ORDER_DEPRECATED_IN_5_2
 #include "RenderingThread.h"
+#endif
+#include "RenderDeferredCleanup.h"
 #include "Serialization/Archive.h"
 #include "Serialization/ArchiveProxy.h"
 #include "Serialization/MemoryImage.h"
diff --git a/Engine/Source/Runtime/Renderer/Public/PrimitiveSceneInfo.h b/Engine/Source/Runtime/Renderer/Public/PrimitiveSceneInfo.h
--- a/Engine/Source/Runtime/Renderer/Public/PrimitiveSceneInfo.h
+++ b/Engine/Source/Runtime/Renderer/Public/PrimitiveSceneInfo.h
@@ -8,7 +8,10 @@
 
 #include "CoreMinimal.h"
 #include "Containers/IndirectArray.h"
+#if UE_ENABLE_INCLUDE_ORDER_DEPRECATED_IN_5_2
 #include "RenderingThread.h"
+#endif
+#include "RenderDeferredCleanup.h"
 #include "SceneTypes.h"
 #include "HitProxies.h"
 #include "Math/GenericOctreePublic.h"
diff --git a/Engine/Source/Runtime/SlateRHIRenderer/Private/Slate3DRenderer.h b/Engine/Source/Runtime/SlateRHIRenderer/Private/Slate3DRenderer.h
--- a/Engine/Source/Runtime/SlateRHIRenderer/Private/Slate3DRenderer.h
+++ b/Engine/Source/Runtime/SlateRHIRenderer/Private/Slate3DRenderer.h
@@ -7,7 +7,10 @@
 #include "Interfaces/ISlate3DRenderer.h"
 #include "SlateRHIResourceManager.h"
 #include "Rendering/SlateDrawBuffer.h"
+#if UE_ENABLE_INCLUDE_ORDER_DEPRECATED_IN_5_2
 #include "RenderingThread.h"
+#endif
+#include "RenderDeferredCleanup.h"
 #include "SlateRHIRenderingPolicy.h"
 
 class FSlateElementBatcher;
diff --git a/Engine/Source/Runtime/SlateRHIRenderer/Private/SlatePostProcessResource.cpp b/Engine/Source/Runtime/SlateRHIRenderer/Private/SlatePostProcessResource.cpp
--- a/Engine/Source/Runtime/SlateRHIRenderer/Private/SlatePostProcessResource.cpp
+++ b/Engine/Source/Runtime/SlateRHIRenderer/Private/SlatePostProcessResource.cpp
@@ -2,6 +2,7 @@
 
 #include "SlatePostProcessResource.h"
 #include "RenderUtils.h"
+#include "RHI.h"
 
 DECLARE_MEMORY_STAT(TEXT("PostProcess RenderTargets"), STAT_SLATEPPRenderTargetMem, STATGROUP_SlateMemory);
 
diff --git a/Engine/Source/Runtime/SlateRHIRenderer/Private/SlatePostProcessResource.h b/Engine/Source/Runtime/SlateRHIRenderer/Private/SlatePostProcessResource.h
--- a/Engine/Source/Runtime/SlateRHIRenderer/Private/SlatePostProcessResource.h
+++ b/Engine/Source/Runtime/SlateRHIRenderer/Private/SlatePostProcessResource.h
@@ -5,7 +5,10 @@
 #include "CoreMinimal.h"
 #include "Textures/SlateShaderResource.h"
 #include "RenderResource.h"
+#if UE_ENABLE_INCLUDE_ORDER_DEPRECATED_IN_5_2
 #include "RenderingThread.h"
+#endif
+#include "RenderDeferredCleanup.h"
 
 /**
  * Handle to resources used for post processing.  This should not be deleted manually because it implements FDeferredCleanupInterface
diff --git a/Engine/Source/Runtime/UMG/Private/Components/RichTextBlock.cpp b/Engine/Source/Runtime/UMG/Private/Components/RichTextBlock.cpp
--- a/Engine/Source/Runtime/UMG/Private/Components/RichTextBlock.cpp
+++ b/Engine/Source/Runtime/UMG/Private/Components/RichTextBlock.cpp
@@ -11,7 +11,7 @@
 #include "Framework/Text/RichTextMarkupProcessing.h"
 #include "Framework/Text/IRichTextMarkupParser.h"
 #include "Framework/Text/IRichTextMarkupWriter.h"
-#include "RenderingThread.h"
+#include "RenderDeferredCleanup.h"
 #include "Editor/WidgetCompilerLog.h"
 #include "Materials/MaterialInstanceDynamic.h"
 
diff --git a/Engine/Source/Runtime/UMG/Private/Slate/SRetainerWidget.cpp b/Engine/Source/Runtime/UMG/Private/Slate/SRetainerWidget.cpp
--- a/Engine/Source/Runtime/UMG/Private/Slate/SRetainerWidget.cpp
+++ b/Engine/Source/Runtime/UMG/Private/Slate/SRetainerWidget.cpp
@@ -5,6 +5,8 @@
 #include "UObject/Package.h"
 #include "Framework/Application/SlateApplication.h"
 #include "Materials/MaterialInstanceDynamic.h"
+#include "RenderDeferredCleanup.h"
+#include "RHI.h"
 #include "Engine/TextureRenderTarget2D.h"
 #include "Engine/World.h"
 #include "UMGPrivate.h"
diff --git a/Engine/Source/Runtime/UMG/Public/Slate/WidgetRenderer.h b/Engine/Source/Runtime/UMG/Public/Slate/WidgetRenderer.h
--- a/Engine/Source/Runtime/UMG/Public/Slate/WidgetRenderer.h
+++ b/Engine/Source/Runtime/UMG/Public/Slate/WidgetRenderer.h
@@ -10,7 +10,10 @@
 #include "Widgets/DeclarativeSyntaxSupport.h"
 #include "Rendering/DrawElements.h"
 #include "Engine/Texture.h"
+#if UE_ENABLE_INCLUDE_ORDER_DEPRECATED_IN_5_2
 #include "RenderingThread.h"
+#endif
+#include "RenderDeferredCleanup.h"
 #include "Widgets/SWindow.h"
 #include "Widgets/SVirtualWindow.h"
 
diff --git a/Engine/Source/Runtime/WebBrowser/Private/CEF/CEFWebBrowserWindowRHIHelper.cpp b/Engine/Source/Runtime/WebBrowser/Private/CEF/CEFWebBrowserWindowRHIHelper.cpp
--- a/Engine/Source/Runtime/WebBrowser/Private/CEF/CEFWebBrowserWindowRHIHelper.cpp
+++ b/Engine/Source/Runtime/WebBrowser/Private/CEF/CEFWebBrowserWindowRHIHelper.cpp
@@ -10,6 +10,7 @@
 #include "RHI.h"
 #if PLATFORM_WINDOWS
 #include "Slate/SlateTextures.h"
+#include "RenderingThread.h"
 #endif
 #endif
 
