Deprecate the old binding style ShaderParameterUtils functions that set specific resource array indices since they can't work with bindless.

diff --git a/Engine/Source/Runtime/RenderCore/Public/ShaderParameterUtils.h b/Engine/Source/Runtime/RenderCore/Public/ShaderParameterUtils.h
--- a/Engine/Source/Runtime/RenderCore/Public/ShaderParameterUtils.h
+++ b/Engine/Source/Runtime/RenderCore/Public/ShaderParameterUtils.h
@@ -157,13 +157,15 @@ GUARD_SETSHADERVALUE(FBox3)
  * Template'd on shader type (e.g. pixel shader or compute shader).
  */
 template<typename TRHIShader, typename TRHICmdList>
-FORCEINLINE void SetTextureParameter(TRHICmdList& RHICmdList, TRHIShader* Shader, const FShaderResourceParameter& Parameter, FRHITexture* TextureRHI, uint32 ElementIndex = 0)
+UE_DEPRECATED(5.2, "SetTextureParameter with an index can't be supported anymore. Your code should be changed to use shader parameter structs to utilize resource arrays.")
+FORCEINLINE void SetTextureParameter(TRHICmdList& RHICmdList, TRHIShader* Shader, const FShaderResourceParameter& Parameter, FRHITexture* TextureRHI, uint32 ElementIndex)
 {
 	if (Parameter.IsBound() && ElementIndex < Parameter.GetNumResources())
 	{
 #if PLATFORM_SUPPORTS_BINDLESS_RENDERING
 		if (Parameter.GetType() == EShaderParameterType::BindlessResourceIndex)
 		{
+			checkf(ElementIndex == 0, TEXT("Bindless resources don't support element offsets"));
 			RHICmdList.SetBindlessTexture(Shader, Parameter.GetBaseIndex(), TextureRHI);
 		}
 		else
@@ -174,17 +176,27 @@ FORCEINLINE void SetTextureParameter(TRHICmdList& RHICmdList, TRHIShader* Shader
 	}
 }
 
+template<typename TRHIShader, typename TRHICmdList>
+FORCEINLINE void SetTextureParameter(TRHICmdList& RHICmdList, TRHIShader* Shader, const FShaderResourceParameter& Parameter, FRHITexture* TextureRHI)
+{
+PRAGMA_DISABLE_DEPRECATION_WARNINGS
+	SetTextureParameter(RHICmdList, Shader, Parameter, TextureRHI, 0);
+PRAGMA_ENABLE_DEPRECATION_WARNINGS
+}
+
 /**
  * Sets the value of a shader sampler parameter. Template'd on shader type.
  */
 template<typename TRHIShader, typename TRHICmdList>
-FORCEINLINE void SetSamplerParameter(TRHICmdList& RHICmdList, TRHIShader* Shader, const FShaderResourceParameter& Parameter, FRHISamplerState* SamplerStateRHI, uint32 ElementIndex = 0)
+UE_DEPRECATED(5.2, "SetSamplerParameter with an index can't be supported anymore. Your code should be changed to use shader parameter structs to utilize resource arrays.")
+FORCEINLINE void SetSamplerParameter(TRHICmdList& RHICmdList, TRHIShader* Shader, const FShaderResourceParameter& Parameter, FRHISamplerState* SamplerStateRHI, uint32 ElementIndex)
 {
 	if (Parameter.IsBound() && ElementIndex < Parameter.GetNumResources())
 	{
 #if PLATFORM_SUPPORTS_BINDLESS_RENDERING
 		if (Parameter.GetType() == EShaderParameterType::BindlessSamplerIndex)
 		{
+			checkf(ElementIndex == 0, TEXT("Bindless resources don't support element offsets"));
 			RHICmdList.SetBindlessSampler(Shader, Parameter.GetBaseIndex(), SamplerStateRHI);
 		}
 		else
@@ -195,11 +207,19 @@ FORCEINLINE void SetSamplerParameter(TRHICmdList& RHICmdList, TRHIShader* Shader
 	}
 }
 
+template<typename TRHIShader, typename TRHICmdList>
+FORCEINLINE void SetSamplerParameter(TRHICmdList& RHICmdList, TRHIShader* Shader, const FShaderResourceParameter& Parameter, FRHISamplerState* SamplerStateRHI)
+{
+PRAGMA_DISABLE_DEPRECATION_WARNINGS
+	SetSamplerParameter(RHICmdList, Shader, Parameter, SamplerStateRHI, 0);
+PRAGMA_ENABLE_DEPRECATION_WARNINGS
+}
 
 /**
  * Sets the value of a shader texture parameter. Template'd on shader type.
  */
 template<typename TRHIShader, typename TRHICmdList>
+UE_DEPRECATED(5.2, "SetTextureParameter with an index can't be supported anymore. Your code should be changed to use shader parameter structs to utilize resource arrays.")
 FORCEINLINE void SetTextureParameter(
 	TRHICmdList& RHICmdList,
 	TRHIShader* Shader,
@@ -207,9 +227,10 @@ FORCEINLINE void SetTextureParameter(
 	const FShaderResourceParameter& SamplerParameter,
 	FRHISamplerState* SamplerStateRHI,
 	FRHITexture* TextureRHI,
-	uint32 ElementIndex = 0
+	uint32 ElementIndex
 	)
 {
+PRAGMA_DISABLE_DEPRECATION_WARNINGS
 	SetTextureParameter(RHICmdList, Shader, TextureParameter, TextureRHI, ElementIndex);
 	
 	// @todo UE samplerstate Should we maybe pass in two separate values? SamplerElement and TextureElement? Or never allow an array of samplers? Unsure best
@@ -220,19 +241,29 @@ FORCEINLINE void SetTextureParameter(
 	// This assumes that the all textures want to use the same sampler state
 
 	SetSamplerParameter(RHICmdList, Shader, SamplerParameter, SamplerStateRHI, ElementIndex);
+PRAGMA_ENABLE_DEPRECATION_WARNINGS
+}
+
+template<typename TRHIShader, typename TRHICmdList>
+FORCEINLINE void SetTextureParameter(TRHICmdList& RHICmdList, TRHIShader* Shader, const FShaderResourceParameter& TextureParameter, const FShaderResourceParameter& SamplerParameter, FRHISamplerState* SamplerStateRHI, FRHITexture* TextureRHI)
+{
+PRAGMA_DISABLE_DEPRECATION_WARNINGS
+	SetTextureParameter(RHICmdList, Shader, TextureParameter, SamplerParameter, SamplerStateRHI, TextureRHI, 0);
+PRAGMA_ENABLE_DEPRECATION_WARNINGS
 }
 
 /**
  * Sets the value of a shader texture parameter.  Template'd on shader type
  */
 template<typename TRHIShader, typename TRHICmdList>
+UE_DEPRECATED(5.2, "SetTextureParameter with an index can't be supported anymore. Your code should be changed to use shader parameter structs to utilize resource arrays.")
 FORCEINLINE void SetTextureParameter(
 	TRHICmdList& RHICmdList,
 	TRHIShader* Shader,
 	const FShaderResourceParameter& TextureParameter,
 	const FShaderResourceParameter& SamplerParameter,
 	const FTexture* Texture,
-	uint32 ElementIndex = 0
+	uint32 ElementIndex
 )
 {
 	if (TextureParameter.IsBound())
@@ -240,7 +271,18 @@ FORCEINLINE void SetTextureParameter(
 		Texture->LastRenderTime = FApp::GetCurrentTime();
 	}
 
-	SetTextureParameter(RHICmdList, Shader, TextureParameter, SamplerParameter, Texture->SamplerStateRHI, Texture->TextureRHI, ElementIndex);
+PRAGMA_DISABLE_DEPRECATION_WARNINGS
+	SetTextureParameter(RHICmdList, Shader, TextureParameter, Texture->TextureRHI, ElementIndex);
+	SetSamplerParameter(RHICmdList, Shader, SamplerParameter, Texture->SamplerStateRHI, ElementIndex);
+PRAGMA_ENABLE_DEPRECATION_WARNINGS
+}
+
+template<typename TRHIShader, typename TRHICmdList>
+FORCEINLINE void SetTextureParameter(TRHICmdList& RHICmdList, TRHIShader* Shader, const FShaderResourceParameter& TextureParameter, const FShaderResourceParameter& SamplerParameter, const FTexture* Texture)
+{
+PRAGMA_DISABLE_DEPRECATION_WARNINGS
+	SetTextureParameter(RHICmdList, Shader, TextureParameter, SamplerParameter, Texture, 0);
+PRAGMA_ENABLE_DEPRECATION_WARNINGS
 }
 
 /**
