Also collect number of outstanding PSO Precache requests when collecting total number of PSO precache jobs (used during level loading)

diff --git a/Engine/Source/Runtime/RenderCore/Private/ShaderPipelineCache.cpp b/Engine/Source/Runtime/RenderCore/Private/ShaderPipelineCache.cpp
--- a/Engine/Source/Runtime/RenderCore/Private/ShaderPipelineCache.cpp
+++ b/Engine/Source/Runtime/RenderCore/Private/ShaderPipelineCache.cpp
@@ -797,6 +797,9 @@ uint32 FShaderPipelineCache::NumPrecompilesRemaining()
 		}
 	}
 
+	// Add the number of active PSO precache requests as well
+	NumRemaining += PipelineStateCache::NumActivePrecacheRequests();
+
 	// NumRemaining is correct only for the current cache.
 	// This ensures we do not return 0 if there are more caches to process.
 	uint32 PendingCaches = IsPrecompiling() ? 1 : 0;
