Cleaning up MediaShaders.h dependencies.

diff --git a/Engine/Source/Runtime/RenderCore/Private/MediaShaders.cpp b/Engine/Source/Runtime/RenderCore/Private/MediaShaders.cpp
--- a/Engine/Source/Runtime/RenderCore/Private/MediaShaders.cpp
+++ b/Engine/Source/Runtime/RenderCore/Private/MediaShaders.cpp
@@ -4,9 +4,24 @@
 #include "RenderGraphBuilder.h"
 #include "RHIStaticStates.h"
 #include "ShaderParameterUtils.h"
+#include "PipelineStateCache.h"
 
+FMediaVertexDeclaration::FMediaVertexDeclaration() = default;
+FMediaVertexDeclaration::~FMediaVertexDeclaration() = default;
 
+void FMediaVertexDeclaration::InitRHI()
+{
+	FVertexDeclarationElementList Elements;
+	uint16 Stride = sizeof(FMediaElementVertex);
+	Elements.Add(FVertexElement(0, STRUCT_OFFSET(FMediaElementVertex, Position), VET_Float4, 0, Stride));
+	Elements.Add(FVertexElement(0, STRUCT_OFFSET(FMediaElementVertex, TextureCoordinate), VET_Float2, 1, Stride));
+	VertexDeclarationRHI = PipelineStateCache::GetOrCreateVertexDeclaration(Elements);
+}
 
+void FMediaVertexDeclaration::ReleaseRHI()
+{
+	VertexDeclarationRHI.SafeRelease();
+}
 
 TGlobalResource<FMediaVertexDeclaration> GMediaVertexDeclaration;
 
diff --git a/Engine/Source/Runtime/RenderCore/Public/MediaShaders.h b/Engine/Source/Runtime/RenderCore/Public/MediaShaders.h
--- a/Engine/Source/Runtime/RenderCore/Public/MediaShaders.h
+++ b/Engine/Source/Runtime/RenderCore/Public/MediaShaders.h
@@ -4,6 +4,9 @@
 
 #include "CoreMinimal.h"
 #include "GlobalShader.h"
+#include "ShaderParameterStruct.h"
+
+#if UE_ENABLE_INCLUDE_ORDER_DEPRECATED_IN_5_2
 #include "Math/Color.h"
 #include "Math/Matrix.h"
 #include "Math/UnrealMathSSE.h"
@@ -19,10 +22,10 @@
 #include "Serialization/MemoryLayout.h"
 #include "Shader.h"
 #include "ShaderParameterMacros.h"
-#include "ShaderParameterStruct.h"
 #include "ShaderPermutation.h"
 #include "Templates/RefCounting.h"
 #include "Templates/UnrealTemplate.h"
+#endif
 
 class FPointerTableBase;
 class FRDGBuilder;
@@ -122,31 +125,18 @@ inline FBufferRHIRef CreateTempMediaVertexBuffer(float ULeft = 0.0f, float URigh
 /**
  * The simple element vertex declaration resource type.
  */
-class FMediaVertexDeclaration
-	: public FRenderResource
+class RENDERCORE_API FMediaVertexDeclaration : public FRenderResource
 {
 public:
+	FMediaVertexDeclaration();
+	virtual ~FMediaVertexDeclaration();
 
-	FVertexDeclarationRHIRef VertexDeclarationRHI;
-
-	virtual ~FMediaVertexDeclaration() { }
-
-	virtual void InitRHI() override
-	{
-		FVertexDeclarationElementList Elements;
-		uint16 Stride = sizeof(FMediaElementVertex);
-		Elements.Add(FVertexElement(0, STRUCT_OFFSET(FMediaElementVertex, Position), VET_Float4, 0, Stride));
-		Elements.Add(FVertexElement(0, STRUCT_OFFSET(FMediaElementVertex, TextureCoordinate), VET_Float2, 1, Stride));
-		VertexDeclarationRHI = PipelineStateCache::GetOrCreateVertexDeclaration(Elements);
-	}
+	virtual void InitRHI() override;
+	virtual void ReleaseRHI() override;
 
-	virtual void ReleaseRHI() override
-	{
-		VertexDeclarationRHI.SafeRelease();
-	}
+	FVertexDeclarationRHIRef VertexDeclarationRHI;
 };
 
-
 RENDERCORE_API extern TGlobalResource<FMediaVertexDeclaration> GMediaVertexDeclaration;
 
 
