Make FShaderMapContent use FMemoryImageName instead of MemoryImageString
Make FShaderMapBase::Serialize use FNames instead of FString for the ShaderPlatform

diff --git a/Engine/Shaders/Public/ShaderVersion.ush b/Engine/Shaders/Public/ShaderVersion.ush
--- a/Engine/Shaders/Public/ShaderVersion.ush
+++ b/Engine/Shaders/Public/ShaderVersion.ush
@@ -2,4 +2,4 @@
 // Each time the console command is executed it generates a new GUID. As this file is included
 // in Platform.ush (which should be included in any shader) it allows to invalidate the shader DDC.
 // 
-// GUID = F9D4E1FF-01F2-4B5A-908F-4A70CDD19FDA
+// GUID = 71F1E87B-7930-4F25-9575-BCBEB3A9BB59
\ No newline at end of file
diff --git a/Engine/Source/Runtime/RenderCore/Private/ShaderMap.cpp b/Engine/Source/Runtime/RenderCore/Private/ShaderMap.cpp
--- a/Engine/Source/Runtime/RenderCore/Private/ShaderMap.cpp
+++ b/Engine/Source/Runtime/RenderCore/Private/ShaderMap.cpp
@@ -171,7 +171,7 @@ bool FShaderMapBase::Serialize(FArchive& Ar, bool bInlineShaderResources, bool b
 		// Serialize a copy of ShaderPlatform directly into the archive
 		// This will allow us to correctly deserialize the stream, even if we're not able to load the frozen content
 		const EShaderPlatform ShaderPlatform = GetShaderPlatform();
-		FString ShaderPlatformName = FDataDrivenShaderPlatformInfo::GetName(GetShaderPlatform()).ToString();
+		FName ShaderPlatformName = FDataDrivenShaderPlatformInfo::GetName(ShaderPlatform);
 		Ar << ShaderPlatformName;
 
 		if (Ar.IsCooking())
@@ -187,7 +187,7 @@ bool FShaderMapBase::Serialize(FArchive& Ar, bool bInlineShaderResources, bool b
 		{
 			FSHAHash ResourceHash = Code->ResourceHash;
 			Ar << ResourceHash;
-			FShaderLibraryCooker::AddShaderCode(GetShaderPlatform(), Code, GetAssociatedAssets());
+			FShaderLibraryCooker::AddShaderCode(ShaderPlatform, Code, GetAssociatedAssets());
 		}
 		else
 #endif // WITH_EDITOR
@@ -207,10 +207,10 @@ bool FShaderMapBase::Serialize(FArchive& Ar, bool bInlineShaderResources, bool b
 		bool bShareCode = false;
 		Ar << bShareCode;
 
-		FString ShaderPlatformName;
+		FName ShaderPlatformName;
 		Ar << ShaderPlatformName;
 		
-		const EShaderPlatform ShaderPlatform = FDataDrivenShaderPlatformInfo::GetShaderPlatformFromName(FName(ShaderPlatformName));
+		const EShaderPlatform ShaderPlatform = FDataDrivenShaderPlatformInfo::GetShaderPlatformFromName(ShaderPlatformName);
 
 		if (bShareCode)
 		{
@@ -300,7 +300,7 @@ static uint16 MakeShaderHash(const FHashedName& TypeName, int32 PermutationId)
 }
 
 FShaderMapContent::FShaderMapContent(EShaderPlatform InPlatform)
-	: ShaderHash(128u), ShaderPlatformName(FDataDrivenShaderPlatformInfo::GetName(InPlatform).ToString())
+	: ShaderHash(128u), ShaderPlatformName(FDataDrivenShaderPlatformInfo::GetName(InPlatform))
 {}
 
 FShaderMapContent::~FShaderMapContent()
@@ -310,7 +310,7 @@ FShaderMapContent::~FShaderMapContent()
 
 EShaderPlatform FShaderMapContent::GetShaderPlatform() const
 {
-	return FDataDrivenShaderPlatformInfo::GetShaderPlatformFromName(FName(ShaderPlatformName));
+	return FDataDrivenShaderPlatformInfo::GetShaderPlatformFromName(ShaderPlatformName);
 }
 
 FShader* FShaderMapContent::GetShader(const FHashedName& TypeName, int32 PermutationId) const
diff --git a/Engine/Source/Runtime/RenderCore/Public/Shader.h b/Engine/Source/Runtime/RenderCore/Public/Shader.h
--- a/Engine/Source/Runtime/RenderCore/Public/Shader.h
+++ b/Engine/Source/Runtime/RenderCore/Public/Shader.h
@@ -2201,7 +2201,7 @@ protected:
 	LAYOUT_FIELD(TMemoryImageArray<TMemoryImagePtr<FShader>>, Shaders);
 	LAYOUT_FIELD(TMemoryImageArray<TMemoryImagePtr<FShaderPipeline>>, ShaderPipelines);
 	/** The ShaderPlatform Name this shader map was compiled with */
-	LAYOUT_FIELD(FMemoryImageString, ShaderPlatformName);
+	LAYOUT_FIELD(FMemoryImageName, ShaderPlatformName);
 };
 
 class RENDERCORE_API FShaderMapBase
