Append instead of overwrite output array in FRayTracingShaderLibrary::GetShaders(...)
- required to support RTPSO with hitgroups with different payloads.

diff --git a/Engine/Source/Runtime/RenderCore/Private/ShaderResource.cpp b/Engine/Source/Runtime/RenderCore/Private/ShaderResource.cpp
--- a/Engine/Source/Runtime/RenderCore/Private/ShaderResource.cpp
+++ b/Engine/Source/Runtime/RenderCore/Private/ShaderResource.cpp
@@ -198,12 +198,15 @@ public:
 	void GetShaders(TArray<FRHIRayTracingShader*>& OutShaders, FRHIRayTracingShader* DefaultShader)
 	{
 		const int32 PayloadIndex = FMath::CountTrailingZeros(DefaultShader->RayTracingPayloadType);
+		const int32 BaseOutIndex = OutShaders.Num();
+
 		FScopeLock Lock(&CS);
-		OutShaders = Shaders[PayloadIndex];
+
+		OutShaders.Append(Shaders[PayloadIndex]);
 
 		for (uint32 Index : UnusedIndicies[PayloadIndex])
 		{
-			OutShaders[Index] = DefaultShader;
+			OutShaders[BaseOutIndex + Index] = DefaultShader;
 		}
 	}
 
