Adding RenderGraphFwd.h

diff --git a/Engine/Source/Programs/IncludeTool/IncludeTool/Rules.cs b/Engine/Source/Programs/IncludeTool/IncludeTool/Rules.cs
--- a/Engine/Source/Programs/IncludeTool/IncludeTool/Rules.cs
+++ b/Engine/Source/Programs/IncludeTool/IncludeTool/Rules.cs
@@ -444,6 +444,7 @@ namespace IncludeTool
 			"/Engine/Source/Runtime/Online/HTTP/Public/HttpFwd.h", // error: invalid forward declaration - 'typedef TSharedPtr<IHttpRequest, ESPMode::ThreadSafe> FHttpRequestPtr;'
 			"/Engine/Restricted/NotForLicensees/Plugins/Online/OnlineSubsystemMcp/Source/Public/OnlineSubsystemMcpFwd.h", // expected only include directives and text in forward declaration header
 			"/Engine/Source/Runtime/RHI/Public/RHIFwd.h", // invalid forward declaration - 'namespace ERHIFeatureLevel { enum Type : int; }'
+			"/Engine/Source/Runtime/RenderCore/Public/RenderGraphFwd.h",
 			"/Engine/Source/Runtime/Core/Public/Templates/SharedPointerFwd.h", // Has an enum as well
 
 		};
diff --git a/Engine/Source/Runtime/RenderCore/Public/RenderGraphDefinitions.h b/Engine/Source/Runtime/RenderCore/Public/RenderGraphDefinitions.h
--- a/Engine/Source/Runtime/RenderCore/Public/RenderGraphDefinitions.h
+++ b/Engine/Source/Runtime/RenderCore/Public/RenderGraphDefinitions.h
@@ -2,8 +2,9 @@
 
 #pragma once
 
-#include "RenderGraphAllocator.h"
 #include "ProfilingDebugging/RealtimeGPUProfiler.h"
+#include "RenderGraphAllocator.h"
+#include "RenderGraphFwd.h"
 
 /** DEFINES */
 
@@ -710,13 +711,9 @@ class FRDGBarrierBatch;
 class FRDGBarrierBatchBegin;
 class FRDGBarrierBatchEnd;
 class FRDGBarrierValidation;
-class FRDGBuilder;
 class FRDGEventName;
 class FRDGUserValidation;
 
-class FRDGResource;
-using FRDGResourceRef = FRDGResource*;
-
 class FRDGViewableResource;
 
 UE_DEPRECATED(5.1, "FRDGParentResource has been renamed to FRDGViewableResource.")
@@ -725,53 +722,23 @@ typedef FRDGViewableResource FRDGParentResource;
 UE_DEPRECATED(5.1, "FRDGParentResourceRef has been renamed to FRDGViewableResource*.")
 typedef FRDGViewableResource* FRDGParentResourceRef;
 
-class FRDGShaderResourceView;
-using FRDGShaderResourceViewRef = FRDGShaderResourceView*;
-
-class FRDGUnorderedAccessView;
-using FRDGUnorderedAccessViewRef = FRDGUnorderedAccessView*;
-
-class FRDGTextureSRV;
-using FRDGTextureSRVRef = FRDGTextureSRV*;
-
-class FRDGTextureUAV;
-using FRDGTextureUAVRef = FRDGTextureUAV*;
-
-class FRDGBufferSRV;
-using FRDGBufferSRVRef = FRDGBufferSRV*;
-
-class FRDGBufferUAV;
-using FRDGBufferUAVRef = FRDGBufferUAV*;
-
-class FRDGPass;
-using FRDGPassRef = FRDGPass*;
 using FRDGPassHandle = TRDGHandle<FRDGPass, uint16>;
 using FRDGPassRegistry = TRDGHandleRegistry<FRDGPassHandle>;
 using FRDGPassHandleArray = TArray<FRDGPassHandle, TInlineAllocator<4, FRDGArrayAllocator>>;
 using FRDGPassBitArray = TRDGHandleBitArray<FRDGPassHandle>;
 
-class FRDGUniformBuffer;
-using FRDGUniformBufferRef = FRDGUniformBuffer*;
 using FRDGUniformBufferHandle = TRDGHandle<FRDGUniformBuffer, uint16>;
 using FRDGUniformBufferRegistry = TRDGHandleRegistry<FRDGUniformBufferHandle>;
 using FRDGUniformBufferBitArray = TRDGHandleBitArray<FRDGUniformBufferHandle>;
 
-class FRDGView;
-using FRDGViewRef = FRDGView*;
 using FRDGViewHandle = TRDGHandle<FRDGView, uint16>;
 using FRDGViewRegistry = TRDGHandleRegistry<FRDGViewHandle, ERDGHandleRegistryDestructPolicy::Never>;
 using FRDGViewUniqueFilter = TRDGHandleUniqueFilter<FRDGViewHandle>;
 
-class FRDGTexture;
-using FRDGTextureRef = FRDGTexture*;
 using FRDGTextureHandle = TRDGHandle<FRDGTexture, uint16>;
 using FRDGTextureRegistry = TRDGHandleRegistry<FRDGTextureHandle, ERDGHandleRegistryDestructPolicy::Never>;
 using FRDGTextureBitArray = TRDGHandleBitArray<FRDGTextureHandle>;
 
-struct FRDGBufferDesc;
-
-class FRDGBuffer;
-using FRDGBufferRef = FRDGBuffer*;
 using FRDGBufferHandle = TRDGHandle<FRDGBuffer, uint16>;
 using FRDGBufferRegistry = TRDGHandleRegistry<FRDGBufferHandle, ERDGHandleRegistryDestructPolicy::Registry>;
 using FRDGBufferBitArray = TRDGHandleBitArray<FRDGBufferHandle>;
@@ -781,9 +748,6 @@ class FRDGPooledBuffer;
 class FRDGBufferPool;
 class FRDGTransientRenderTarget;
 
-template <typename TUniformStruct> class TRDGUniformBuffer;
-template <typename TUniformStruct> using TRDGUniformBufferRef = TRDGUniformBuffer<TUniformStruct>*;
-
 using FRDGPassHandlesByPipeline = TRHIPipelineArray<FRDGPassHandle>;
 using FRDGPassesByPipeline = TRHIPipelineArray<FRDGPass*>;
 
diff --git a/Engine/Source/Runtime/RenderCore/Public/RenderGraphFwd.h b/Engine/Source/Runtime/RenderCore/Public/RenderGraphFwd.h
new file mode 100644
--- /dev/null
+++ b/Engine/Source/Runtime/RenderCore/Public/RenderGraphFwd.h
@@ -0,0 +1,53 @@
+// Copyright Epic Games, Inc. All Rights Reserved.
+
+#pragma once
+
+#include "HAL/Platform.h"
+
+enum class ERDGBufferFlags : uint8;
+enum class ERDGPassFlags : uint16;
+enum class ERDGTextureFlags : uint8;
+enum class ERDGUnorderedAccessViewFlags : uint8;
+
+class FRDGBuffer;
+using FRDGBufferRef = FRDGBuffer*;
+
+struct FRDGBufferDesc;
+
+class FRDGBufferSRV;
+using FRDGBufferSRVRef = FRDGBufferSRV*;
+
+class FRDGBufferUAV;
+using FRDGBufferUAVRef = FRDGBufferUAV*;
+
+class FRDGBuilder;
+
+class FRDGPass;
+using FRDGPassRef = FRDGPass*;
+
+class FRDGResource;
+using FRDGResourceRef = FRDGResource*;
+
+class FRDGShaderResourceView;
+using FRDGShaderResourceViewRef = FRDGShaderResourceView*;
+
+class FRDGTexture;
+using FRDGTextureRef = FRDGTexture*;
+
+class FRDGTextureSRV;
+using FRDGTextureSRVRef = FRDGTextureSRV*;
+
+class FRDGTextureUAV;
+using FRDGTextureUAVRef = FRDGTextureUAV*;
+
+class FRDGUniformBuffer;
+using FRDGUniformBufferRef = FRDGUniformBuffer*;
+
+class FRDGUnorderedAccessView;
+using FRDGUnorderedAccessViewRef = FRDGUnorderedAccessView*;
+
+class FRDGView;
+using FRDGViewRef = FRDGView*;
+
+template <typename TUniformStruct> class TRDGUniformBuffer;
+template <typename TUniformStruct> using TRDGUniformBufferRef = TRDGUniformBuffer<TUniformStruct>*;
