Apply lock striping to shader cache to reduce CacheMeshDrawCommands contention
  - 9s to 6.8s for the heavy scene updates at beginning of PIE in CitySample

#rnx

diff --git a/Engine/Source/Runtime/RenderCore/Private/ShaderCore.cpp b/Engine/Source/Runtime/RenderCore/Private/ShaderCore.cpp
--- a/Engine/Source/Runtime/RenderCore/Private/ShaderCore.cpp
+++ b/Engine/Source/Runtime/RenderCore/Private/ShaderCore.cpp
@@ -102,11 +102,17 @@ DEFINE_STAT(STAT_Shaders_ShaderPreloadMemory);
 DEFINE_STAT(STAT_Shaders_NumShadersRegistered);
 DEFINE_STAT(STAT_Shaders_NumShadersDuplicated);
 
-/** Protects GShaderFileCache from simultaneous access by multiple threads. */
-FRWLock FileCacheRWLock;
+// Apply lock striping as we're mostly reader lock bound.
+constexpr int32 GSHADERFILECACHE_BUCKETS = 31; /* prime number for best distribution using modulo */
 
-/** The shader file cache, used to minimize shader file reads */
-TMap<FString, FString> GShaderFileCache;
+struct FShaderFileCache
+{
+	/** Protects Map from simultaneous access by multiple threads. */
+	FRWLock Lock;
+
+	/** The shader file cache, used to minimize shader file reads */
+	TMap<FString, FString> Map;
+} GShaderFileCache[GSHADERFILECACHE_BUCKETS];
 
 class FShaderHashCache
 {
@@ -995,10 +1001,11 @@ bool LoadShaderSourceFile(const TCHAR* InVirtualFilePath, EShaderPlatform Shader
 		FString* CachedFile = nullptr;
 
 		// First try a shared lock and only acquire exclusive access if element is not found in cache
-		
+		uint32 CurrentHash = GetTypeHash(VirtualFilePath);
+		FShaderFileCache& ShaderFileCache = GShaderFileCache[CurrentHash % GSHADERFILECACHE_BUCKETS];
 		{
-			FRWScopeLock ScopeLock(FileCacheRWLock, SLT_ReadOnly);
-			CachedFile = GShaderFileCache.Find(VirtualFilePath);
+			FRWScopeLock ScopeLock(ShaderFileCache.Lock, SLT_ReadOnly);
+			CachedFile = ShaderFileCache.Map.FindByHash(CurrentHash, VirtualFilePath);
 		}
 
 		if (CachedFile)
@@ -1011,10 +1018,10 @@ bool LoadShaderSourceFile(const TCHAR* InVirtualFilePath, EShaderPlatform Shader
 		}
 		else 
 		{
-			FRWScopeLock ScopeLock(FileCacheRWLock, SLT_Write);
+			FRWScopeLock ScopeLock(ShaderFileCache.Lock, SLT_Write);
 
 			// Double-check the cache while holding exclusive lock as another thread may have added the item we're looking for
-			CachedFile = GShaderFileCache.Find(VirtualFilePath);
+			CachedFile = ShaderFileCache.Map.FindByHash(CurrentHash, VirtualFilePath);
 
 			// if this file has already been loaded and cached, use that
 			if (CachedFile)
@@ -1034,11 +1041,11 @@ bool LoadShaderSourceFile(const TCHAR* InVirtualFilePath, EShaderPlatform Shader
 				if (!ShaderFilePath.IsEmpty() && FFileHelper::LoadFileToString(FileContents, *ShaderFilePath, FFileHelper::EHashOptions::EnableVerify|FFileHelper::EHashOptions::ErrorMissingHash) )
 				{
 					//update the shader file cache
-					GShaderFileCache.Add(VirtualFilePath, FileContents);
+					ShaderFileCache.Map.AddByHash(CurrentHash, VirtualFilePath, FileContents);
 
 					if (OutFileContents)
 					{
-						*OutFileContents = MoveTemp(FileContents);
+						*OutFileContents = FileContents;
 					}
 					bResult = true;
 				}
@@ -1568,8 +1575,11 @@ void FlushShaderFileCache()
 		GShaderHashCache.Empty();
 	}
 	{
-		FRWScopeLock ScopeLock(FileCacheRWLock, SLT_Write);
-		GShaderFileCache.Empty();
+		for (int32 Index = 0; Index < UE_ARRAY_COUNT(GShaderFileCache); ++Index)
+		{
+			FRWScopeLock ScopeLock(GShaderFileCache[Index].Lock, SLT_Write);
+			GShaderFileCache[Index].Map.Empty();
+		}
 	}
 
 	UE_LOG(LogShaders, Log, TEXT("FlushShaderFileCache() end"));
