Fixes to make ray tracing dynamic geometry and TLAS builds on async compute again

#rnx

diff --git a/Engine/Source/Runtime/RenderCore/Private/ProfilingDebugging/RealtimeGPUProfiler.cpp b/Engine/Source/Runtime/RenderCore/Private/ProfilingDebugging/RealtimeGPUProfiler.cpp
--- a/Engine/Source/Runtime/RenderCore/Private/ProfilingDebugging/RealtimeGPUProfiler.cpp
+++ b/Engine/Source/Runtime/RenderCore/Private/ProfilingDebugging/RealtimeGPUProfiler.cpp
@@ -1155,8 +1155,7 @@ void FScopedGPUStatEvent::Begin(FRHICommandListBase& InRHICmdList, const FName&
 		return;
 	}
 
-	// Non-immediate command lists are not supported (silently fail)
-	if (InRHICmdList.IsImmediate())
+	if (InRHICmdList.IsImmediate() && EnumHasAnyFlags(InRHICmdList.GetPipeline(), ERHIPipeline::Graphics))
 	{
 		Category = &InCategory;
 		RHICmdList = &InRHICmdList;
diff --git a/Engine/Source/Runtime/RenderCore/Private/RenderGraphBuilder.cpp b/Engine/Source/Runtime/RenderCore/Private/RenderGraphBuilder.cpp
--- a/Engine/Source/Runtime/RenderCore/Private/RenderGraphBuilder.cpp
+++ b/Engine/Source/Runtime/RenderCore/Private/RenderGraphBuilder.cpp
@@ -458,7 +458,7 @@ bool FRDGBuilder::IsImmediateMode()
 	return ::IsImmediateMode();
 }
 
-ERDGPassFlags FRDGBuilder::OverridePassFlags(const TCHAR* PassName, ERDGPassFlags PassFlags, bool bAsyncComputeSupported)
+ERDGPassFlags FRDGBuilder::OverridePassFlags(const TCHAR* PassName, ERDGPassFlags PassFlags)
 {
 	const bool bDebugAllowedForPass =
 #if RDG_ENABLE_DEBUG
@@ -467,7 +467,7 @@ ERDGPassFlags FRDGBuilder::OverridePassFlags(const TCHAR* PassName, ERDGPassFlag
 		true;
 #endif
 
-	if (IsAsyncComputeSupported() && bAsyncComputeSupported)
+	if (IsAsyncComputeSupported())
 	{
 		if (EnumHasAnyFlags(PassFlags, ERDGPassFlags::Compute) && GRDGAsyncCompute == RDG_ASYNC_COMPUTE_FORCE_ENABLED)
 		{
diff --git a/Engine/Source/Runtime/RenderCore/Public/RenderGraphBuilder.h b/Engine/Source/Runtime/RenderCore/Public/RenderGraphBuilder.h
--- a/Engine/Source/Runtime/RenderCore/Public/RenderGraphBuilder.h
+++ b/Engine/Source/Runtime/RenderCore/Public/RenderGraphBuilder.h
@@ -441,7 +441,7 @@ private:
 		ERDGPassFlags Flags,
 		ExecuteLambdaType&& ExecuteLambda);
 
-	static ERDGPassFlags OverridePassFlags(const TCHAR* PassName, ERDGPassFlags Flags, bool bAsyncComputeSupported);
+	static ERDGPassFlags OverridePassFlags(const TCHAR* PassName, ERDGPassFlags Flags);
 
 	void AddProloguePass();
 
diff --git a/Engine/Source/Runtime/RenderCore/Public/RenderGraphBuilder.inl b/Engine/Source/Runtime/RenderCore/Public/RenderGraphBuilder.inl
--- a/Engine/Source/Runtime/RenderCore/Public/RenderGraphBuilder.inl
+++ b/Engine/Source/Runtime/RenderCore/Public/RenderGraphBuilder.inl
@@ -260,7 +260,7 @@ FRDGPassRef FRDGBuilder::AddPassInternal(
 		MoveTemp(Name),
 		ParametersMetadata,
 		ParameterStruct,
-		OverridePassFlags(Name.GetTCHAR(), Flags, LambdaPassType::kSupportsAsyncCompute),
+		OverridePassFlags(Name.GetTCHAR(), Flags),
 		MoveTemp(ExecuteLambda));
 
 	IF_RDG_ENABLE_DEBUG(ClobberPassOutputs(Pass));
diff --git a/Engine/Source/Runtime/RenderCore/Public/RenderGraphPass.h b/Engine/Source/Runtime/RenderCore/Public/RenderGraphPass.h
--- a/Engine/Source/Runtime/RenderCore/Public/RenderGraphPass.h
+++ b/Engine/Source/Runtime/RenderCore/Public/RenderGraphPass.h
@@ -593,9 +593,7 @@ class TRDGLambdaPass
 	using TRDGPass = typename TLambdaTraits<ExecuteLambdaType>::TRDGPass;
 
 public:
-	static const bool kSupportsAsyncCompute = TIsSame<TRHICommandList, FRHIComputeCommandList>::Value;
-	static const bool kSupportsRaster = TIsDerivedFrom<TRHICommandList, FRHICommandList>::IsDerived;
-
+	
 	TRDGLambdaPass(
 		FRDGEventName&& InName,
 		const FShaderParametersMetadata* InParameterMetadata,
@@ -607,10 +605,7 @@ public:
 #if RDG_ENABLE_DEBUG
 		, DebugParameterStruct(InParameterStruct)
 #endif
-	{
-		checkf(kSupportsAsyncCompute || !EnumHasAnyFlags(InPassFlags, ERDGPassFlags::AsyncCompute),
-			TEXT("Pass %s is set to use 'AsyncCompute', but the pass lambda's first argument is not FRHIComputeCommandList&."), GetName());
-
+	{		
 		bParallelExecuteAllowed = !TIsSame<TRHICommandList, FRHICommandListImmediate>::Value && !EnumHasAnyFlags(InPassFlags, ERDGPassFlags::NeverParallel);
 	}
 
@@ -629,7 +624,6 @@ private:
 
 	void Execute(FRHIComputeCommandList& RHICmdList) override
 	{
-		check(!kSupportsRaster || RHICmdList.IsGraphics());
 		QUICK_SCOPE_CYCLE_COUNTER(STAT_FRDGPass_Execute);
 		RHICmdList.SetStaticUniformBuffers(ParameterStruct.GetStaticUniformBuffers());
 		ExecuteLambdaFunc<TRDGPass>(static_cast<TRHICommandList&>(RHICmdList));
