[pipeline cache] minor fixes for issues identified by VSL security assessment

diff --git a/Engine/Source/Runtime/RenderCore/Private/PipelineCacheUtilities.cpp b/Engine/Source/Runtime/RenderCore/Private/PipelineCacheUtilities.cpp
--- a/Engine/Source/Runtime/RenderCore/Private/PipelineCacheUtilities.cpp
+++ b/Engine/Source/Runtime/RenderCore/Private/PipelineCacheUtilities.cpp
@@ -363,6 +363,10 @@ bool UE::PipelineCacheUtilities::LoadStableKeysFile(const FStringView& Filename,
 		uint64 HashIdx = ReadVarUIntFromArchive(*Archive);
 		Item.PipelineHash = Hashes[static_cast<int32>(HashIdx)];
 		HashIdx = ReadVarUIntFromArchive(*Archive);
+		if (HashIdx >= Hashes.Num())
+		{
+			return false;
+		}
 		Item.OutputHash = Hashes[static_cast<int32>(HashIdx)];
 
 		// Standardize on all CompactNames being parsed from string. This is a temporary hack until the names are parsed from CSV when reading StablePC
@@ -716,7 +720,7 @@ bool UE::PipelineCacheUtilities::LoadStablePipelineCacheFile(const FString& File
 	TArray<FPipelineCacheFileFormatPSO> OriginalPSOs;
 	PermutationGroups.Reserve(Header.NumPermutationGroups);
 	OriginalPSOs.AddDefaulted(Header.NumPermutationGroups);	// need to avoid resizes and invalidating the pointeers
-	for (int32 PermutationGroupIdx = 0; PermutationGroupIdx < Header.NumPermutationGroups; ++PermutationGroupIdx)
+	for (int64 PermutationGroupIdx = 0; PermutationGroupIdx < Header.NumPermutationGroups; ++PermutationGroupIdx)
 	{
 		FPermsPerPSO Item;
 		UE::PipelineCacheUtilities::Private::LoadActiveSlots(MemReader, Item);
