[RenderCore]
* Changed RenderCore to use Stats.h instead of Stats2.h to fix some compile errors farm is getting in ue5 main

diff --git a/Engine/Source/Runtime/RenderCore/Public/RenderCore.h b/Engine/Source/Runtime/RenderCore/Public/RenderCore.h
--- a/Engine/Source/Runtime/RenderCore/Public/RenderCore.h
+++ b/Engine/Source/Runtime/RenderCore/Public/RenderCore.h
@@ -6,8 +6,7 @@
 
 #pragma once
 
-#include "HAL/Platform.h"
-#include "Stats/Stats2.h"
+#include "Stats/Stats.h" // IWYU pragma: keep
 
 #if UE_ENABLE_INCLUDE_ORDER_DEPRECATED_IN_5_2
 #include "CoreMinimal.h"
