- Add cvar to to only precache global shaders from manual PSO cache when runtime PSO precaching is enabled
- Add cvar to only wait for high priority runtime precache PSOs during loading

both disabled by default

#rnx

diff --git a/Engine/Source/Runtime/RHI/Private/PipelineStateCache.cpp b/Engine/Source/Runtime/RHI/Private/PipelineStateCache.cpp
--- a/Engine/Source/Runtime/RHI/Private/PipelineStateCache.cpp
+++ b/Engine/Source/Runtime/RHI/Private/PipelineStateCache.cpp
@@ -130,6 +130,15 @@ static FAutoConsoleVariableRef CVarPSOPrecaching(
 	ECVF_Default
 );
 
+int32 GPSOWaitForHighPriorityRequestsOnly = 0;
+static FAutoConsoleVariableRef CVarPSOPrecaching(
+	TEXT("r.PSOPrecaching.WaitForHighPriorityRequestsOnly"),
+	GPSOWaitForHighPriorityRequestsOnly,
+	TEXT("0 to wait for all pending PSO precache requests during loading (default)g\n")
+	TEXT("1 to only wait for the high priority PSO precache requests during loading\n"),
+	ECVF_Default
+);
+
 extern void DumpPipelineCacheStats();
 
 static FAutoConsoleCommand DumpPipelineCmd(
@@ -1129,11 +1138,14 @@ public:
 	uint32 NumActivePrecacheRequests()
 	{
 		FRWScopeLock ReadLock(PrecachePSOsRWLock, SLT_ReadOnly);
-		return ActiveCompileCount;
-
-		// Ideally we only want to wait for high priority compile requests, but this is still giving too much popping right now
-		// Boost of required PSOs needs to be moved earlier for this to work
-		//return HighPriorityCompileCount;
+		if (GPSOWaitForHighPriorityRequestsOnly)
+		{
+			return HighPriorityCompileCount;
+		}
+		else
+		{
+			return ActiveCompileCount;
+		}		
 	}
 
 	void PrecacheFinished(const TPrecachedPSOInitializer& Initializer, bool bValid)
diff --git a/Engine/Source/Runtime/RenderCore/Private/ShaderPipelineCache.cpp b/Engine/Source/Runtime/RenderCore/Private/ShaderPipelineCache.cpp
--- a/Engine/Source/Runtime/RenderCore/Private/ShaderPipelineCache.cpp
+++ b/Engine/Source/Runtime/RenderCore/Private/ShaderPipelineCache.cpp
@@ -157,6 +157,13 @@ static TAutoConsoleVariable<float> CVarPSOFileCacheMaxPrecompileTime(
 																ECVF_Default | ECVF_RenderThreadSafe
 																);
 
+static TAutoConsoleVariable<int32> CVarPSOGlobalShadersOnlyWhenPSOPrecaching(
+                                                                TEXT("r.ShaderPipelineCache.GlobalShadersOnlyWhenPSOPrecaching"),
+                                                                (int32)0,
+                                                                TEXT("Only compile PSOs from the GlobalShader cache when runtime PSOPrecaching is enabled (default disabled)"),
+                                                                ECVF_Default | ECVF_RenderThreadSafe
+                                                                );
+
 static bool GetShaderPipelineCacheSaveBoundPSOLog()
 {
 	static bool bOnce = false;
@@ -2035,13 +2042,23 @@ void FShaderPipelineCacheTask::BeginPrecompilingPipelineCache()
 		// Iterate over all the tasks we haven't yet begun to read data for - these are the 'waiting' tasks
 		int64 EligibleTaskCount = LocalPreFetchedTasks.Num();
 
+		// Only interested in global shaders when PSO precaching is enabled
+		FString LibraryName(TEXT("Global"));
+
 		int32 MissingShaders = 0;
-		int32 NewSize = Algo::StableRemoveIf(LocalPreFetchedTasks, [&MissingShaders](FPipelineCachePSOHeader& Task)
+		int32 NewSize = Algo::StableRemoveIf(LocalPreFetchedTasks, [&MissingShaders, LibraryName](FPipelineCachePSOHeader& Task)
 		{
 			bool bHasShaders = true;
 			for (FSHAHash const& Hash : Task.Shaders)
 			{
-				bHasShaders &= FShaderCodeLibrary::ContainsShaderCode(Hash);
+				if (PipelineStateCache::IsPSOPrecachingEnabled() && CVarPSOGlobalShadersOnlyWhenPSOPrecaching.GetValueOnAnyThread() > 0)
+				{
+					bHasShaders &= FShaderCodeLibrary::ContainsShaderCode(Hash, LibraryName);
+				}
+				else
+				{
+					bHasShaders &= FShaderCodeLibrary::ContainsShaderCode(Hash);
+				}
 			}
 			if(!bHasShaders)
 			{
