Add optional task tagging to RDG parallel setup tasks.

diff --git a/Engine/Source/Runtime/RenderCore/Public/RenderGraphBuilder.inl b/Engine/Source/Runtime/RenderCore/Public/RenderGraphBuilder.inl
--- a/Engine/Source/Runtime/RenderCore/Public/RenderGraphBuilder.inl
+++ b/Engine/Source/Runtime/RenderCore/Public/RenderGraphBuilder.inl
@@ -412,9 +412,9 @@ UE::Tasks::FTask FRDGBuilder::AddSetupTask(TaskLambdaType&& TaskLambda, bool bCo
 
 	if (bParallelExecuteEnabled && bCondition)
 	{
-		Task = UE::Tasks::Launch(TEXT("FRDGBuilder::AddSetupTask"), [TaskLambda = MoveTemp(TaskLambda)]
+		Task = UE::Tasks::Launch(TEXT("FRDGBuilder::AddSetupTask"), [TaskLambda = MoveTemp(TaskLambda)] () mutable
 		{
-			FTaskTagScope Scope(ETaskTag::EParallelRenderingThread);
+			FOptionalTaskTagScope Scope(ETaskTag::EParallelRenderingThread);
 			TaskLambda();
 		});
 
@@ -437,9 +437,9 @@ UE::Tasks::FTask FRDGBuilder::AddCommandListSetupTask(TaskLambdaType&& TaskLambd
 	{
 		FRHICommandList* RHICmdListTask = new FRHICommandList(FRHIGPUMask::All());
 
-		Task = UE::Tasks::Launch(TEXT("FRDGBuilder::AddCommandListSetupTask"), [TaskLambda = MoveTemp(TaskLambda), RHICmdListTask]
+		Task = UE::Tasks::Launch(TEXT("FRDGBuilder::AddCommandListSetupTask"), [TaskLambda = MoveTemp(TaskLambda), RHICmdListTask] () mutable
 		{
-			FTaskTagScope Scope(ETaskTag::EParallelRenderingThread);
+			FOptionalTaskTagScope Scope(ETaskTag::EParallelRenderingThread);
 			RHICmdListTask->SwitchPipeline(ERHIPipeline::Graphics);
 			TaskLambda(*RHICmdListTask);
 			RHICmdListTask->FinishRecording();
