Add ConvertUniformBufferToExternal to allow RDG UBs to be used in non-RDG contexts (e.g. RT pipeline)

diff --git a/Engine/Source/Runtime/RenderCore/Private/RenderGraphBuilder.cpp b/Engine/Source/Runtime/RenderCore/Private/RenderGraphBuilder.cpp
--- a/Engine/Source/Runtime/RenderCore/Private/RenderGraphBuilder.cpp
+++ b/Engine/Source/Runtime/RenderCore/Private/RenderGraphBuilder.cpp
@@ -631,6 +631,129 @@ const TRefCountPtr<IPooledRenderTarget>& FRDGBuilder::ConvertToExternalTexture(F
 	return GetPooledTexture(Texture);
 }
 
+FRHIUniformBuffer* FRDGBuilder::ConvertToExternalUniformBuffer(FRDGUniformBufferRef UniformBuffer)
+{
+	if (!UniformBuffer->bExternal)
+	{
+		UniformBuffer->GetParameters().Enumerate([this](const FRDGParameter& Param)
+		{
+			auto ConvertTexture = [](FRDGBuilder* Builder, FRDGTextureRef Texture)
+			{
+				if (Texture && !Texture->IsExternal())
+				{
+					Builder->ConvertToExternalTexture(Texture);
+				}
+			};
+			auto ConvertBuffer = [](FRDGBuilder* Builder, FRDGBufferRef Buffer)
+			{
+				if (Buffer && !Buffer->IsExternal())
+				{
+					Builder->ConvertToExternalBuffer(Buffer);
+				}
+			};
+
+			switch (Param.GetType())
+			{
+			case UBMT_RDG_TEXTURE:
+			{
+				ConvertTexture(this, Param.GetAsTexture());
+			}
+			break;
+			case UBMT_RDG_TEXTURE_ACCESS:
+			{
+				ConvertTexture(this, Param.GetAsTextureAccess().GetTexture());
+			}
+			break;
+			case UBMT_RDG_TEXTURE_ACCESS_ARRAY:
+			{
+				const FRDGTextureAccessArray& Array = Param.GetAsTextureAccessArray();
+				for (int Index = 0; Index < Array.Num(); ++Index)
+				{
+					ConvertTexture(this, Array[Index].GetTexture());
+				}
+			}
+			break;
+			case UBMT_RDG_TEXTURE_SRV:
+			{
+				ConvertTexture(this, Param.GetAsTextureSRV()->Desc.Texture);
+				InitRHI(Param.GetAsView());
+			}
+			break;
+			case UBMT_RDG_TEXTURE_UAV:
+			{
+				ConvertTexture(this, Param.GetAsTextureUAV()->Desc.Texture);
+				InitRHI(Param.GetAsView());
+			}
+			break;
+			case UBMT_RDG_BUFFER_ACCESS:
+			{
+				ConvertBuffer(this, Param.GetAsBufferAccess().GetBuffer());
+			}
+			break;
+			case UBMT_RDG_BUFFER_ACCESS_ARRAY:
+			{
+				const FRDGBufferAccessArray& Array = Param.GetAsBufferAccessArray();
+				for (int Index = 0; Index < Array.Num(); ++Index)
+				{
+					ConvertBuffer(this, Array[Index].GetBuffer());
+				}
+			}
+			break;
+			case UBMT_RDG_BUFFER_SRV:
+			{
+				ConvertBuffer(this, Param.GetAsBufferSRV()->Desc.Buffer);
+				InitRHI(Param.GetAsView());
+			}
+			break;
+			case UBMT_RDG_BUFFER_UAV:
+			{
+				ConvertBuffer(this, Param.GetAsBufferUAV()->Desc.Buffer);
+				InitRHI(Param.GetAsView());
+			}
+			break;
+			case UBMT_RDG_UNIFORM_BUFFER:
+			{
+				FRDGUniformBufferRef Buffer = Param.GetAsUniformBuffer().GetUniformBuffer();
+				if (Buffer)
+				{
+					ConvertToExternalUniformBuffer(Buffer);
+				}
+			}
+			break;
+
+			// Non-RDG cases
+			case UBMT_INT32:
+			case UBMT_UINT32:
+			case UBMT_FLOAT32:
+			case UBMT_TEXTURE:
+			case UBMT_SRV:
+			case UBMT_UAV:
+			case UBMT_SAMPLER:
+			case UBMT_NESTED_STRUCT:
+			case UBMT_INCLUDED_STRUCT:
+			case UBMT_REFERENCED_STRUCT:
+			case UBMT_RENDER_TARGET_BINDING_SLOTS:
+			break;
+
+			default:
+				check(0);
+			}
+		});
+	}
+	IF_RDG_ENABLE_DEBUG(UserValidation.ValidateConvertToExternalUniformBuffer(UniformBuffer));
+	if (!UniformBuffer->bExternal)
+	{
+		UniformBuffer->bExternal = true;
+		UniformBuffer->bQueuedForCreate = true;
+
+		// It's safe to reset the access to false because validation won't allow this call during execution.
+		IF_RDG_ENABLE_DEBUG(GRDGAllowRHIAccess = true);
+		UniformBuffer->InitRHI();
+		IF_RDG_ENABLE_DEBUG(GRDGAllowRHIAccess = false);
+	}
+	return UniformBuffer->GetRHIUnchecked();
+}
+
 ///////////////////////////////////////////////////////////////////////////////////////////////////
 
 BEGIN_SHADER_PARAMETER_STRUCT(FAccessModePassParameters, )
diff --git a/Engine/Source/Runtime/RenderCore/Private/RenderGraphValidation.cpp b/Engine/Source/Runtime/RenderCore/Private/RenderGraphValidation.cpp
--- a/Engine/Source/Runtime/RenderCore/Private/RenderGraphValidation.cpp
+++ b/Engine/Source/Runtime/RenderCore/Private/RenderGraphValidation.cpp
@@ -396,6 +396,14 @@ void FRDGUserValidation::ValidateConvertToExternalResource(FRDGViewableResource*
 		Resource->Name);
 }
 
+void FRDGUserValidation::ValidateConvertToExternalUniformBuffer(FRDGUniformBuffer* UniformBuffer)
+{
+	check(UniformBuffer);
+	checkf(!bHasExecuteBegun,
+		TEXT("Unable to convert uniform buffer %s to external because passes in the graph have already executed."),
+		UniformBuffer->GetLayoutName());
+}
+
 void FRDGUserValidation::RemoveUnusedWarning(FRDGViewableResource* Resource)
 {
 	check(Resource);
diff --git a/Engine/Source/Runtime/RenderCore/Public/RenderGraphBuilder.h b/Engine/Source/Runtime/RenderCore/Public/RenderGraphBuilder.h
--- a/Engine/Source/Runtime/RenderCore/Public/RenderGraphBuilder.h
+++ b/Engine/Source/Runtime/RenderCore/Public/RenderGraphBuilder.h
@@ -281,6 +281,12 @@ public:
 	 */
 	const TRefCountPtr<IPooledRenderTarget>& ConvertToExternalTexture(FRDGTextureRef Texture);
 	const TRefCountPtr<FRDGPooledBuffer>& ConvertToExternalBuffer(FRDGBufferRef Buffer);
+	/** For a graph-created uniform buffer, this forces immediate allocation of the underlying resource, effectively promoting it
+	 *  to an external resource. This will increase memory pressure, but allows access to the RHI resource.
+	 *  Graph resources that are referenced in the buffer will be converted to external.
+	 *  This is primarily used as an aid for porting code incrementally to RDG.
+	 */
+	FRHIUniformBuffer* ConvertToExternalUniformBuffer(FRDGUniformBufferRef UniformBuffer);
 
 	/** Performs an immediate query for the underlying pooled resource. This is only allowed for external or extracted resources. */
 	const TRefCountPtr<IPooledRenderTarget>& GetPooledTexture(FRDGTextureRef Texture) const;
diff --git a/Engine/Source/Runtime/RenderCore/Public/RenderGraphResources.h b/Engine/Source/Runtime/RenderCore/Public/RenderGraphResources.h
--- a/Engine/Source/Runtime/RenderCore/Public/RenderGraphResources.h
+++ b/Engine/Source/Runtime/RenderCore/Public/RenderGraphResources.h
@@ -193,6 +193,11 @@ public:
 		return ParameterStruct;
 	}
 
+	FORCEINLINE const TCHAR* GetLayoutName() const
+	{
+		return *ParameterStruct.GetLayout().GetDebugName();
+	}
+
 #if RDG_ENABLE_DEBUG
 	void MarkResourceAsUsed() override;
 #else
@@ -228,6 +233,7 @@ private:
 	TRefCountPtr<FRHIUniformBuffer> UniformBufferRHI;
 	FRDGUniformBufferHandle Handle;
 	bool bQueuedForCreate = false;
+	bool bExternal = false;
 
 	friend FRDGBuilder;
 	friend FRDGUniformBufferRegistry;
diff --git a/Engine/Source/Runtime/RenderCore/Public/RenderGraphValidation.h b/Engine/Source/Runtime/RenderCore/Public/RenderGraphValidation.h
--- a/Engine/Source/Runtime/RenderCore/Public/RenderGraphValidation.h
+++ b/Engine/Source/Runtime/RenderCore/Public/RenderGraphValidation.h
@@ -85,6 +85,7 @@ public:
 	void ValidateExtractBuffer(FRDGBufferRef Buffer, TRefCountPtr<FRDGPooledBuffer>* OutBufferPtr);
 
 	void ValidateConvertToExternalResource(FRDGViewableResource* Resource);
+	void ValidateConvertToExternalUniformBuffer(FRDGUniformBuffer* UniformBuffer);
 
 	/** Tracks and validates the addition of a new pass to the graph.
 	 *  @param bSkipPassAccessMarking Skips marking the pass as a producer or incrementing the pass access. Useful when
