Parallelize some parts of InitializeShaderTypes
  - 1.9s to 640ms

#rnx

diff --git a/Engine/Source/Runtime/RenderCore/Private/ShaderCore.cpp b/Engine/Source/Runtime/RenderCore/Private/ShaderCore.cpp
--- a/Engine/Source/Runtime/RenderCore/Private/ShaderCore.cpp
+++ b/Engine/Source/Runtime/RenderCore/Private/ShaderCore.cpp
@@ -18,6 +18,7 @@
 #include "Shader.h"
 #include "ShaderCompilerCore.h"
 #include "String/Find.h"
+#include "Tasks/Task.h"
 #include "VertexFactory.h"
 #include "Modules/ModuleManager.h"
 #include "Interfaces/IShaderFormat.h"
@@ -1455,6 +1456,8 @@ void BuildShaderFileToUniformBufferMap(TMap<FString, TArray<const TCHAR*> >& Sha
 {
 	if (!FPlatformProperties::RequiresCookedData())
 	{
+		TRACE_CPUPROFILER_EVENT_SCOPE(BuildShaderFileToUniformBufferMap);
+
 		TArray<FString> ShaderSourceFiles;
 		GetAllVirtualShaderSourcePaths(ShaderSourceFiles, GMaxRHIShaderPlatform);
 
@@ -1481,6 +1484,15 @@ void BuildShaderFileToUniformBufferMap(TMap<FString, TArray<const TCHAR*> >& Sha
 			SearchKeys.Add(FShaderVariable(StructIt->GetShaderVariableName()));
 		}
 
+		TArray<UE::Tasks::TTask<void>> Tasks;
+		Tasks.Reserve(ShaderSourceFiles.Num());
+
+		// Just make sure that all the TArray inside the map won't move while being used by async tasks
+		for (int32 FileIndex = 0; FileIndex < ShaderSourceFiles.Num(); FileIndex++)
+		{
+			ShaderFileToUniformBufferVariables.FindOrAdd(ShaderSourceFiles[FileIndex]);
+		}
+
 		// Find for each shader file which UBs it needs
 		for (int32 FileIndex = 0; FileIndex < ShaderSourceFiles.Num(); FileIndex++)
 		{
@@ -1489,21 +1501,30 @@ void BuildShaderFileToUniformBufferMap(TMap<FString, TArray<const TCHAR*> >& Sha
  			FString ShaderFileContents;
 			LoadShaderSourceFileChecked(*ShaderSourceFiles[FileIndex], GMaxRHIShaderPlatform, ShaderFileContents);
 
-			// To allow case sensitive search which is way faster on some platforms (no need to look up locale, etc)
-			ShaderFileContents.ToUpperInline();
+			Tasks.Emplace(
+				UE::Tasks::Launch(
+					TEXT("SearchKeysInShaderContent"),
+					[&SearchKeys, &ShaderSourceFiles, FileIndex, &ShaderFileToUniformBufferVariables, ShaderFileContents = MoveTemp(ShaderFileContents)]() mutable
+					{
+						// To allow case sensitive search which is way faster on some platforms (no need to look up locale, etc)
+						ShaderFileContents.ToUpperInline();
 
-			TArray<const TCHAR*>& ReferencedUniformBuffers = ShaderFileToUniformBufferVariables.FindOrAdd(ShaderSourceFiles[FileIndex]);
+						TArray<const TCHAR*>* ReferencedUniformBuffers = ShaderFileToUniformBufferVariables.Find(ShaderSourceFiles[FileIndex]);
 
-			for (int32 SearchKeyIndex = 0; SearchKeyIndex < SearchKeys.Num(); ++SearchKeyIndex)
-			{
-				// Searching for the uniform buffer shader variable being accessed with '.'
-				if (ShaderFileContents.Contains(SearchKeys[SearchKeyIndex].SearchKey, ESearchCase::CaseSensitive)
-					|| ShaderFileContents.Contains(SearchKeys[SearchKeyIndex].SearchKeyWithSpace, ESearchCase::CaseSensitive))
-				{
-					ReferencedUniformBuffers.AddUnique(SearchKeys[SearchKeyIndex].OriginalShaderVariable);
-				}
-			}
+						for (int32 SearchKeyIndex = 0; SearchKeyIndex < SearchKeys.Num(); ++SearchKeyIndex)
+						{
+							// Searching for the uniform buffer shader variable being accessed with '.'
+							if (ShaderFileContents.Contains(SearchKeys[SearchKeyIndex].SearchKey, ESearchCase::CaseSensitive)
+								|| ShaderFileContents.Contains(SearchKeys[SearchKeyIndex].SearchKeyWithSpace, ESearchCase::CaseSensitive))
+								{
+									ReferencedUniformBuffers->AddUnique(SearchKeys[SearchKeyIndex].OriginalShaderVariable);
+							}
+						}
+					}
+				)
+			);
 		}
+		UE::Tasks::Wait(Tasks);
 	}
 }
 #endif // WITH_EDITOR
