[shaders] remove functionality which allows setting cvars via the shader input environment, along with the only existing usage of it (which is obsolete)
this has much potential for misuse, and is generally incompatible with the goal of in-process multithreaded compilation

diff --git a/Engine/Source/Programs/ShaderCompileWorker/Private/ShaderCompileWorker.cpp b/Engine/Source/Programs/ShaderCompileWorker/Private/ShaderCompileWorker.cpp
--- a/Engine/Source/Programs/ShaderCompileWorker/Private/ShaderCompileWorker.cpp
+++ b/Engine/Source/Programs/ShaderCompileWorker/Private/ShaderCompileWorker.cpp
@@ -200,16 +200,6 @@ static FORCENOINLINE void ProcessCompilationJob(const FShaderCompilerInput& Inpu
 		ExitWithoutCrash(FSCWErrorCode::CantCompileForSpecificFormat, FString::Printf(TEXT("Can't compile shaders for format %s"), *Input.ShaderFormat.ToString()));
 	}
 
-	// Apply the console variable values from the input environment before calling the platform shader compiler
-	for (const auto& Pair : Input.Environment.ShaderFormatCVars)
-	{
-		IConsoleVariable* CVar = IConsoleManager::Get().FindConsoleVariable(*Pair.Key);
-		if (CVar)
-		{
-			CVar->Set(*Pair.Value, ECVF_SetByCode);
-		}
-	}
-
 #if PLATFORM_WINDOWS
 	FString ExceptionMsg;
 	FString ExceptionCallstack;
diff --git a/Engine/Source/Runtime/RenderCore/Public/ShaderCompilerCore.h b/Engine/Source/Runtime/RenderCore/Public/ShaderCompilerCore.h
--- a/Engine/Source/Runtime/RenderCore/Public/ShaderCompilerCore.h
+++ b/Engine/Source/Runtime/RenderCore/Public/ShaderCompilerCore.h
@@ -18,7 +18,7 @@
 class Error;
 
 // this is for the protocol, not the data, bump if FShaderCompilerInput or ProcessInputFromArchive changes.
-inline const int32 ShaderCompileWorkerInputVersion = 17;
+inline const int32 ShaderCompileWorkerInputVersion = 18;
 // this is for the protocol, not the data, bump if FShaderCompilerOutput or WriteToOutputArchive changes.
 inline const int32 ShaderCompileWorkerOutputVersion = 10;
 // this is for the protocol, not the data.
diff --git a/Engine/Source/Runtime/RenderCore/Public/ShaderCore.h b/Engine/Source/Runtime/RenderCore/Public/ShaderCore.h
--- a/Engine/Source/Runtime/RenderCore/Public/ShaderCore.h
+++ b/Engine/Source/Runtime/RenderCore/Public/ShaderCore.h
@@ -523,7 +523,6 @@ struct FShaderCompilerEnvironment
 	TMap<FString, FResourceTableEntry> ResourceTableMap;
 	TMap<FString, FUniformBufferEntry> UniformBufferMap;
 	TMap<FString, FString> RemoteServerData;
-	TMap<FString, FString> ShaderFormatCVars;
 
 	const ITargetPlatform* TargetPlatform = nullptr;
 
@@ -579,7 +578,6 @@ struct FShaderCompilerEnvironment
 		Ar << ResourceTableMap;
 		Ar << UniformBufferMap;
 		Ar << RemoteServerData;
-		Ar << ShaderFormatCVars;
 		Ar << FullPrecisionInPS;
 	}
 
@@ -620,7 +618,6 @@ struct FShaderCompilerEnvironment
 		Definitions.Merge(Other.Definitions);
 		RenderTargetOutputFormatsMap.Append(Other.RenderTargetOutputFormatsMap);
 		RemoteServerData.Append(Other.RemoteServerData);
-		ShaderFormatCVars.Append(Other.ShaderFormatCVars);
 		FullPrecisionInPS |= Other.FullPrecisionInPS;
 	}
 
