MGPU:  Crash when moving camera on CitySample startup map.  Crash was caused by FRHICommandListBase::SwitchPipeline calling ParameterBatcher.FlushAllParameters outside the SCOPED_GPU_MASK for the Pass, and thereby running on all GPUs as opposed to just the view's GPU.  Fix is to make it so FRHICommandListScopedPipeline is inside the GPU mask scope.

#rnx

diff --git a/Engine/Source/Runtime/RenderCore/Private/RenderGraphBuilder.cpp b/Engine/Source/Runtime/RenderCore/Private/RenderGraphBuilder.cpp
--- a/Engine/Source/Runtime/RenderCore/Private/RenderGraphBuilder.cpp
+++ b/Engine/Source/Runtime/RenderCore/Private/RenderGraphBuilder.cpp
@@ -2856,50 +2856,53 @@ void FRDGBuilder::ExecutePassEpilogue(FRHIComputeCommandList& RHICmdListPass, FR
 void FRDGBuilder::ExecutePass(FRDGPass* Pass, FRHIComputeCommandList& RHICmdListPass)
 {
 	{
-		FRHICommandListScopedPipeline Scope(RHICmdListPass, Pass->Pipeline);
-
-#if 0 // Disabled by default to reduce memory usage in Insights.
-		SCOPED_NAMED_EVENT_TCHAR(Pass->GetName(), FColor::Magenta);
-#endif
-
 		// Note that we must do this before doing anything with RHICmdList for the pass.
 		// For example, if this pass only executes on GPU 1 we want to avoid adding a
 		// 0-duration event for this pass on GPU 0's time line.
 		SCOPED_GPU_MASK(RHICmdListPass, Pass->GPUMask);
 
-#if RDG_CPU_SCOPES
-		if (!Pass->bParallelExecute)
+		// Extra scope here to ensure nested ordering of SCOPED_GPU_MASK and FRHICommandListScopedPipeline constructor/destructors
 		{
-			Pass->CPUScopeOps.Execute();
-		}
+			FRHICommandListScopedPipeline Scope(RHICmdListPass, Pass->Pipeline);
+
+#if 0 // Disabled by default to reduce memory usage in Insights.
+			SCOPED_NAMED_EVENT_TCHAR(Pass->GetName(), FColor::Magenta);
+#endif
+
+#if RDG_CPU_SCOPES
+			if (!Pass->bParallelExecute)
+			{
+				Pass->CPUScopeOps.Execute();
+			}
 #endif
 
-		IF_RDG_ENABLE_DEBUG(ConditionalDebugBreak(RDG_BREAKPOINT_PASS_EXECUTE, BuilderName.GetTCHAR(), Pass->GetName()));
+			IF_RDG_ENABLE_DEBUG(ConditionalDebugBreak(RDG_BREAKPOINT_PASS_EXECUTE, BuilderName.GetTCHAR(), Pass->GetName()));
 
 #if WITH_MGPU
-		if (Pass->bWaitForTemporalEffect)
-		{
-			static_cast<FRHICommandList&>(RHICmdListPass).WaitForTemporalEffect(NameForTemporalEffect);
-		}
+			if (Pass->bWaitForTemporalEffect)
+			{
+				static_cast<FRHICommandList&>(RHICmdListPass).WaitForTemporalEffect(NameForTemporalEffect);
+			}
 #endif
 
-		Pass->GPUScopeOpsPrologue.Execute(RHICmdListPass);
+			Pass->GPUScopeOpsPrologue.Execute(RHICmdListPass);
 
-		ExecutePassPrologue(RHICmdListPass, Pass);
+			ExecutePassPrologue(RHICmdListPass, Pass);
 
 #if RDG_DUMP_RESOURCES_AT_EACH_DRAW
-		BeginPassDump(Pass);
+			BeginPassDump(Pass);
 #endif
 
-		Pass->Execute(RHICmdListPass);
+			Pass->Execute(RHICmdListPass);
 
 #if RDG_DUMP_RESOURCES_AT_EACH_DRAW
-		EndPassDump(Pass);
+			EndPassDump(Pass);
 #endif
 
-		ExecutePassEpilogue(RHICmdListPass, Pass);
+			ExecutePassEpilogue(RHICmdListPass, Pass);
 
-		Pass->GPUScopeOpsEpilogue.Execute(RHICmdListPass);
+			Pass->GPUScopeOpsEpilogue.Execute(RHICmdListPass);
+		}
 	}
 
 	if (!Pass->bParallelExecute && Pass->bDispatchAfterExecute)
