MGPU:  ensure !bDrawingScene crash starting CitySample.  Need to use current GPU mask for FRHICommandList in FRDGBuilder::AddCommandListSetupTask, instead of FRHIGPUMask::All(), to avoid an issue where a SwitchPipeline call in a generated command list corrupts the GPU mask (where it calls RHISetGPUMask(CurrentGPUMask) with the wrong CurrentGPUMask), leading to the assert.

Looking at other places parallel command lists are generated, they don't need to be fixed up.  FParallelCommandListSet::AllocCommandList already uses the current GPU mask, and FRDGBuilder::DispatchParallelExecute runs passes, which each pass sets and resets the GPU mask, so it's correct that the initial state is FRHIGPUMask::All().

#rnx

diff --git a/Engine/Source/Runtime/RenderCore/Public/RenderGraphBuilder.inl b/Engine/Source/Runtime/RenderCore/Public/RenderGraphBuilder.inl
--- a/Engine/Source/Runtime/RenderCore/Public/RenderGraphBuilder.inl
+++ b/Engine/Source/Runtime/RenderCore/Public/RenderGraphBuilder.inl
@@ -435,7 +435,7 @@ UE::Tasks::FTask FRDGBuilder::AddCommandListSetupTask(TaskLambdaType&& TaskLambd
 
 	if (bParallelExecuteEnabled && bCondition)
 	{
-		FRHICommandList* RHICmdListTask = new FRHICommandList(FRHIGPUMask::All());
+		FRHICommandList* RHICmdListTask = new FRHICommandList(RHICmdList.GetGPUMask());
 
 		Task = UE::Tasks::Launch(TEXT("FRDGBuilder::AddCommandListSetupTask"), [TaskLambda = MoveTemp(TaskLambda), RHICmdListTask] () mutable
 		{
