Fix for memory leak in transient allocator due to RDG async compute producer -> graphics consumer fence where the consumer was culled. This caused the passes resources to not be traversed and EndResourcesRHI was not called, resulting in a leak of the underlying memory.
 - Fixed incorrect accumulation of reference counts when RDG culling is disabled but parallel setup is enabled.
 - Added assert to catch bad referencing counting in the future, which will be easier to debug than a leak.

diff --git a/Engine/Source/Runtime/RenderCore/Private/RenderGraphBuilder.cpp b/Engine/Source/Runtime/RenderCore/Private/RenderGraphBuilder.cpp
--- a/Engine/Source/Runtime/RenderCore/Private/RenderGraphBuilder.cpp
+++ b/Engine/Source/Runtime/RenderCore/Private/RenderGraphBuilder.cpp
@@ -1043,11 +1043,21 @@ void FRDGBuilder::AddPassDependency(FRDGPass* Producer, FRDGPass* Consumer)
 
 		if (Producer->Pipeline != Consumer->Pipeline)
 		{
-			// Finds the earliest consumer on the other pipeline for the producer.
-			if (Producer->CrossPipelineConsumer.IsNull() || Consumer->Handle < Producer->CrossPipelineConsumer)
+			const auto BinarySearchOrAdd = [](FRDGPassHandleArray& Range, FRDGPassHandle Handle)
 			{
-				Producer->CrossPipelineConsumer = Consumer->Handle;
-			}
+				const int32 LowerBoundIndex = Algo::LowerBound(Range, Handle);
+				if (LowerBoundIndex < Range.Num())
+				{
+					if (Range[LowerBoundIndex] == Handle)
+					{
+						return;
+					}
+				}
+				Range.Insert(Handle, LowerBoundIndex);
+			};
+
+			// Consumers could be culled, so we have to store all of them in a sorted list.
+			BinarySearchOrAdd(Producer->CrossPipelineConsumers, Consumer->Handle);
 
 			// Finds the latest producer on the other pipeline for the consumer.
 			if (Consumer->CrossPipelineProducer.IsNull() || Producer->Handle > Consumer->CrossPipelineProducer)
@@ -1283,8 +1293,6 @@ void FRDGBuilder::Compile()
 			{
 				AddLastProducersToCullStack(LastProducer);
 			}
-
-			Texture->ReferenceCount++;
 		}
 
 		for (const FExtractedBuffer& ExtractedBuffer : ExtractedBuffers)
@@ -1292,8 +1300,6 @@ void FRDGBuilder::Compile()
 			FRDGBufferRef Buffer = ExtractedBuffer.Buffer;
 
 			AddLastProducersToCullStack(Buffer->LastProducer);
-
-			Buffer->ReferenceCount++;
 		}
 
 		for (const auto& Pair : ExternalTextures)
@@ -1313,7 +1319,7 @@ void FRDGBuilder::Compile()
 			AddLastProducersToCullStack(Buffer->LastProducer);
 		}
 	}
-	else
+	else if (!bParallelSetupEnabled)
 	{
 		for (FRDGPassHandle PassHandle = ProloguePassHandle + 1; PassHandle < EpiloguePassHandle; ++PassHandle)
 		{
@@ -1333,6 +1339,16 @@ void FRDGBuilder::Compile()
 		}
 	}
 
+	for (const FExtractedTexture& ExtractedTexture : ExtractedTextures)
+	{
+		ExtractedTexture.Texture->ReferenceCount++;
+	}
+
+	for (const FExtractedBuffer& ExtractedBuffer : ExtractedBuffers)
+	{
+		ExtractedBuffer.Buffer->ReferenceCount++;
+	}
+
 	// All dependencies in the raw graph have been specified; if enabled, all passes are marked as culled and a
 	// depth first search is employed to find reachable regions of the graph. Roots of the search are those passes
 	// with outputs leaving the graph or those marked to never cull.
@@ -1580,7 +1596,21 @@ void FRDGBuilder::Compile()
 				continue;
 			}
 
-			FRDGPassHandle GraphicsJoinPassHandle = FRDGPassHandle::Min(AsyncComputePass->CrossPipelineConsumer, FRDGPassHandle::Min(CurrentGraphicsJoinPassHandle, EpiloguePassHandle));
+			FRDGPassHandle CrossPipelineConsumer;
+
+			// Cross pipeline consumers are sorted. Find the earliest consumer that isn't culled.
+			for (FRDGPassHandle ConsumerHandle : AsyncComputePass->CrossPipelineConsumers)
+			{
+				FRDGPass* Consumer = Passes[ConsumerHandle];
+
+				if (!Consumer->bCulled)
+				{
+					CrossPipelineConsumer = ConsumerHandle;
+					break;
+				}
+			}
+
+			FRDGPassHandle GraphicsJoinPassHandle = FRDGPassHandle::Min(CrossPipelineConsumer, FRDGPassHandle::Min(CurrentGraphicsJoinPassHandle, EpiloguePassHandle));
 			FRDGPass* GraphicsJoinPass = Passes[GraphicsJoinPassHandle];
 
 			AsyncComputePass->GraphicsJoinPass = GraphicsJoinPassHandle;
@@ -3005,6 +3035,8 @@ void FRDGBuilder::AddEpilogueTransition(FRDGTextureRef Texture)
 		return;
 	}
 
+	check(IsImmediateMode() || Texture->bExtracted || Texture->ReferenceCount == FRDGViewableResource::DeallocatedReferenceCount);
+
 	if (!EnumHasAnyFlags(Texture->Flags, ERDGTextureFlags::SkipTracking))
 	{
 		const FRDGPassHandle EpiloguePassHandle = GetEpiloguePassHandle();
@@ -3053,6 +3085,8 @@ void FRDGBuilder::AddEpilogueTransition(FRDGBufferRef Buffer)
 		return;
 	}
 
+	check(IsImmediateMode() || Buffer->bExtracted || Buffer->ReferenceCount == FRDGViewableResource::DeallocatedReferenceCount);
+
 	if (!EnumHasAnyFlags(Buffer->Flags, ERDGBufferFlags::SkipTracking))
 	{
 		const FRDGPassHandle EpiloguePassHandle = GetEpiloguePassHandle();
diff --git a/Engine/Source/Runtime/RenderCore/Public/RenderGraphPass.h b/Engine/Source/Runtime/RenderCore/Public/RenderGraphPass.h
--- a/Engine/Source/Runtime/RenderCore/Public/RenderGraphPass.h
+++ b/Engine/Source/Runtime/RenderCore/Public/RenderGraphPass.h
@@ -288,18 +288,6 @@ public:
 		return Producers;
 	}
 
-	/** Returns the producer pass on the other pipeline, if it exists. */
-	FRDGPassHandle GetCrossPipelineProducer() const
-	{
-		return CrossPipelineProducer;
-	}
-
-	/** Returns the consumer pass on the other pipeline, if it exists. */
-	FRDGPassHandle GetCrossPipelineConsumer() const
-	{
-		return CrossPipelineConsumer;
-	}
-
 	/** Returns the graphics pass responsible for forking the async interval this pass is in. */
 	FRDGPassHandle GetGraphicsForkPass() const
 	{
@@ -432,9 +420,8 @@ protected:
 		uint32 PacketBits2 = 0;
 	};
 
-	/** Handle of the latest cross-pipeline producer and earliest cross-pipeline consumer. */
+	/** Handle of the latest cross-pipeline producer. */
 	FRDGPassHandle CrossPipelineProducer;
-	FRDGPassHandle CrossPipelineConsumer;
 
 	/** (AsyncCompute only) Graphics passes which are the fork / join for async compute interval this pass is in. */
 	FRDGPassHandle GraphicsForkPass;
@@ -444,7 +431,8 @@ protected:
 	FRDGPassHandle PrologueBarrierPass;
 	FRDGPassHandle EpilogueBarrierPass;
 
-	/** Lists of producer passes. */
+	/** Lists of producer passes and the full list of cross-pipeline consumer passes. */
+	FRDGPassHandleArray CrossPipelineConsumers;
 	FRDGPassHandleArray Producers;
 
 	struct FTextureState
