Fixed typo in ValidateIndirectArgsBuffer, and added a variant of FRDGBufferDesc::CreateIndirectDesc that explicitly specifies the element count and the byte stride to support more elaborate strides.

diff --git a/Engine/Source/Runtime/RenderCore/Public/RenderGraphResources.h b/Engine/Source/Runtime/RenderCore/Public/RenderGraphResources.h
--- a/Engine/Source/Runtime/RenderCore/Public/RenderGraphResources.h
+++ b/Engine/Source/Runtime/RenderCore/Public/RenderGraphResources.h
@@ -887,6 +887,15 @@ PRAGMA_DISABLE_DEPRECATION_WARNINGS
 	FRDGBufferDesc& operator=(const FRDGBufferDesc&) = default;
 PRAGMA_ENABLE_DEPRECATION_WARNINGS
 
+	static FRDGBufferDesc CreateIndirectDesc(uint32 BytesPerElement, uint32 NumElements)
+	{
+		FRDGBufferDesc Desc;
+		Desc.Usage = EBufferUsageFlags::Static | EBufferUsageFlags::DrawIndirect | EBufferUsageFlags::UnorderedAccess | EBufferUsageFlags::ShaderResource | EBufferUsageFlags::VertexBuffer;
+		Desc.BytesPerElement = BytesPerElement;
+		Desc.NumElements = NumElements;
+		return Desc;
+	}
+
 	/** Create the descriptor for an indirect RHI call.
 	 *
 	 * Note, IndirectParameterStruct should be one of the:
@@ -897,20 +906,12 @@ PRAGMA_ENABLE_DEPRECATION_WARNINGS
 	template<typename IndirectParameterStruct>
 	static FRDGBufferDesc CreateIndirectDesc(uint32 NumElements = 1)
 	{
-		FRDGBufferDesc Desc;
-		Desc.Usage = EBufferUsageFlags::Static | EBufferUsageFlags::DrawIndirect | EBufferUsageFlags::UnorderedAccess | EBufferUsageFlags::ShaderResource | EBufferUsageFlags::VertexBuffer;
-		Desc.BytesPerElement = sizeof(IndirectParameterStruct);
-		Desc.NumElements = NumElements;
-		return Desc;
+		return CreateIndirectDesc(sizeof(IndirectParameterStruct), NumElements);
 	}
 
 	static FRDGBufferDesc CreateIndirectDesc(uint32 NumElements = 1)
 	{
-		FRDGBufferDesc Desc;
-		Desc.Usage = EBufferUsageFlags::Static | EBufferUsageFlags::DrawIndirect | EBufferUsageFlags::UnorderedAccess | EBufferUsageFlags::ShaderResource | EBufferUsageFlags::VertexBuffer;
-		Desc.BytesPerElement = 4;
-		Desc.NumElements = NumElements;
-		return Desc;
+		return CreateIndirectDesc(4u, NumElements);
 	}
 
 	static FRDGBufferDesc CreateStructuredDesc(uint32 BytesPerElement, uint32 NumElements)
diff --git a/Engine/Source/Runtime/RenderCore/Public/RenderGraphUtils.h b/Engine/Source/Runtime/RenderCore/Public/RenderGraphUtils.h
--- a/Engine/Source/Runtime/RenderCore/Public/RenderGraphUtils.h
+++ b/Engine/Source/Runtime/RenderCore/Public/RenderGraphUtils.h
@@ -637,7 +637,7 @@ struct RENDERCORE_API FComputeShaderUtils
 		checkf((IndirectArgOffset % 4) == 0, TEXT("IndirectArgOffset for compute shader indirect dispatch needs to be a multiple of 4."));
 		checkf(
 			(IndirectArgOffset + IndirectArgsStride) <= IndirectArgsBufferSize,
-			TEXT("Indirect parameters buffer for compute shader indirect dispatch at byte offset %d doesn't have anought room for one element."),
+			TEXT("Indirect parameters buffer for compute shader indirect dispatch at byte offset %d doesn't have enough room for one element."),
 			IndirectArgOffset);
 #if PLATFORM_DISPATCH_INDIRECT_ARGUMENT_BOUNDARY_SIZE != 0
 			checkf(IndirectArgOffset / PLATFORM_DISPATCH_INDIRECT_ARGUMENT_BOUNDARY_SIZE == (IndirectArgOffset + IndirectArgsStride - 1) / PLATFORM_DISPATCH_INDIRECT_ARGUMENT_BOUNDARY_SIZE, TEXT("Compute indirect dispatch arguments cannot cross %d byte boundary."), PLATFORM_DISPATCH_INDIRECT_ARGUMENT_BOUNDARY_SIZE);
