Moving RHIStaticStates.h to RenderCore to fix a dependency issue.

diff --git a/Engine/Source/Runtime/RHI/Public/RHIStaticStates.h b/Engine/Source/Runtime/RHI/Public/RHIStaticStates.h
deleted file mode 100644
diff --git a/Engine/Source/Runtime/RenderCore/Public/RHIStaticStates.h b/Engine/Source/Runtime/RenderCore/Public/RHIStaticStates.h
new file mode 100644
--- /dev/null
+++ b/Engine/Source/Runtime/RenderCore/Public/RHIStaticStates.h
@@ -0,0 +1,415 @@
+// Copyright Epic Games, Inc. All Rights Reserved.
+
+/*=============================================================================
+	RHIStaticStates.h: RHI static state template definition.
+=============================================================================*/
+
+#pragma once
+
+// HEADER_UNIT_SKIP - This can't be built by itself since it includes RenderResource.h which is not part of this module's additional include paths
+
+#include "CoreMinimal.h"
+#include "Stats/Stats.h"
+#include "Async/TaskGraphInterfaces.h"
+#include "Containers/StaticArray.h"
+#include "RHI.h"
+#include "RenderResource.h"
+#include "Misc/ScopedEvent.h"
+
+extern RHI_API FCriticalSection StaticStateRHICriticalSection;
+
+/**
+ * Helper task to initialize a static resource on the render thread.
+ */
+class FInitStaticResourceRenderThreadTask
+{
+	void (*DoConstruct)();
+	FScopedEvent& Event;
+public:
+
+	FInitStaticResourceRenderThreadTask(void (*InDoConstruct)(), FScopedEvent& InEvent)
+		: DoConstruct(InDoConstruct)
+		, Event(InEvent)
+	{
+	}
+
+	FORCEINLINE TStatId GetStatId() const
+	{
+		RETURN_QUICK_DECLARE_CYCLE_STAT(FInitStaticResourceRenderThreadTask, STATGROUP_TaskGraphTasks);
+	}
+
+	ENamedThreads::Type GetDesiredThread()
+	{
+		return ENamedThreads::GetRenderThread_Local();
+	}
+
+	static ESubsequentsMode::Type GetSubsequentsMode() { return ESubsequentsMode::FireAndForget; }
+
+	void DoTask(ENamedThreads::Type CurrentThread, const FGraphEventRef& MyCompletionGraphEvent)
+	{
+		DoConstruct();
+		Event.Trigger();
+	}
+};
+
+/**
+ * The base class of the static RHI state classes.
+ */
+template<typename InitializerType,typename RHIRefType,typename RHIParamRefType>
+class TStaticStateRHI
+{
+public:
+	static RHIParamRefType GetRHI()
+	{
+		checkSlow(StaticResource.StateRHI);
+		return StaticResource.StateRHI;
+	};
+
+private:
+	/** A resource which manages the RHI resource. */
+	class FStaticStateResource : public FRenderResource
+	{
+	public:
+		RHIRefType StateRHI;
+		FStaticStateResource()
+		{
+			if (GIsRHIInitialized && GRHISupportsRHIThread)
+			{
+				StateRHI = InitializerType::CreateRHI();
+			}
+			else
+			{
+				BeginInitResource(this);
+			}
+		}
+
+		// FRenderResource interface.
+		virtual void InitRHI() override
+		{
+			check(!GIsRHIInitialized || !GRHISupportsRHIThread);
+			StateRHI = InitializerType::CreateRHI();
+		}
+		virtual void ReleaseRHI() override final
+		{
+			check(!GIsRHIInitialized || !GRHISupportsRHIThread);
+			StateRHI.SafeRelease();
+		}
+		virtual void ReleaseResource() override final
+		{
+			FRenderResource::ReleaseResource();
+		}
+
+		~FStaticStateResource()
+		{
+			check(!GIsRHIInitialized || !GRHISupportsRHIThread);
+			ReleaseResource();
+		}
+	};
+
+	static FStaticStateResource StaticResource;
+};
+
+template<typename InitializerType, typename RHIRefType, typename RHIParamRefType>
+typename TStaticStateRHI<InitializerType, RHIRefType, RHIParamRefType>::FStaticStateResource TStaticStateRHI<InitializerType, RHIRefType, RHIParamRefType>::StaticResource = TStaticStateRHI<InitializerType, RHIRefType, RHIParamRefType>::FStaticStateResource();
+
+/**
+ * A static RHI sampler state resource.
+ * TStaticSamplerStateRHI<...>::GetStaticState() will return a FSamplerStateRHIRef to a sampler state with the desired settings.
+ * Should only be used from the rendering thread.
+ */
+template<ESamplerFilter Filter=SF_Point,
+	ESamplerAddressMode AddressU=AM_Clamp,
+	ESamplerAddressMode AddressV=AM_Clamp,
+	ESamplerAddressMode AddressW=AM_Clamp, 
+	int32 MipBias = 0,
+	// Note: setting to a different value than GSystemSettings.MaxAnisotropy is only supported in D3D11 (is that still true?)
+	// A value of 0 will use GSystemSettings.MaxAnisotropy
+	int32 MaxAnisotropy = 1,
+	uint32 BorderColor = 0,
+	/** Only supported in D3D11 */
+	ESamplerCompareFunction SamplerComparisonFunction=SCF_Never>
+class TStaticSamplerState : public TStaticStateRHI<TStaticSamplerState<Filter,AddressU,AddressV,AddressW,MipBias,MaxAnisotropy,BorderColor,SamplerComparisonFunction>,FSamplerStateRHIRef, FRHISamplerState*>
+{
+public:
+	static FSamplerStateRHIRef CreateRHI()
+	{
+		FSamplerStateInitializerRHI Initializer( Filter, AddressU, AddressV, AddressW, MipBias, MaxAnisotropy, 0, FLT_MAX, BorderColor, SamplerComparisonFunction );
+		return RHICreateSamplerState(Initializer);
+	}
+};
+
+/**
+ * A static RHI rasterizer state resource.
+ * TStaticRasterizerStateRHI<...>::GetStaticState() will return a FRasterizerStateRHIRef to a rasterizer state with the desired
+ * settings.
+ * Should only be used from the rendering thread.
+ */
+template<ERasterizerFillMode FillMode=FM_Solid,ERasterizerCullMode CullMode=CM_None,bool bEnableLineAA=false,bool bEnableMSAA=true, ERasterizerDepthClipMode DepthClipMode=ERasterizerDepthClipMode::DepthClip>
+class TStaticRasterizerState : public TStaticStateRHI<TStaticRasterizerState<FillMode,CullMode,bEnableLineAA>,FRasterizerStateRHIRef, FRHIRasterizerState*>
+{
+public:
+	FORCEINLINE_DEBUGGABLE static FRasterizerStateRHIRef CreateRHI()
+	{
+		const FRasterizerStateInitializerRHI Initializer(FillMode, CullMode, 0.0f, 0.0f, DepthClipMode, bEnableMSAA, bEnableLineAA);
+		return RHICreateRasterizerState(Initializer);
+	}
+};
+
+/** Given a fill and cull mode, returns a static rasterizer state. */
+template<bool bEnableMSAA>
+FORCEINLINE_DEBUGGABLE FRHIRasterizerState* GetStaticRasterizerState(ERasterizerFillMode FillMode,ERasterizerCullMode CullMode)
+{
+	switch(FillMode)
+	{
+	default:
+	case FM_Solid:
+		switch(CullMode)
+		{
+		default:
+		case CM_CW:   return TStaticRasterizerState<FM_Solid,CM_CW  ,false,bEnableMSAA>::GetRHI();
+		case CM_CCW:  return TStaticRasterizerState<FM_Solid,CM_CCW ,false,bEnableMSAA>::GetRHI();
+		case CM_None: return TStaticRasterizerState<FM_Solid,CM_None,false,bEnableMSAA>::GetRHI();
+		};
+		break;
+	case FM_Wireframe:
+		switch(CullMode)
+		{
+		default:
+		case CM_CW:   return TStaticRasterizerState<FM_Wireframe,CM_CW  ,false,bEnableMSAA>::GetRHI();
+		case CM_CCW:  return TStaticRasterizerState<FM_Wireframe,CM_CCW ,false,bEnableMSAA>::GetRHI();
+		case CM_None: return TStaticRasterizerState<FM_Wireframe,CM_None,false,bEnableMSAA>::GetRHI();
+		};
+		break;
+	case FM_Point:
+		switch(CullMode)
+		{
+		default:
+		case CM_CW:   return TStaticRasterizerState<FM_Point,CM_CW  ,false,bEnableMSAA>::GetRHI();
+		case CM_CCW:  return TStaticRasterizerState<FM_Point,CM_CCW ,false,bEnableMSAA>::GetRHI();
+		case CM_None: return TStaticRasterizerState<FM_Point,CM_None,false,bEnableMSAA>::GetRHI();
+		};
+		break;
+	}
+}
+
+/**
+ * A static RHI stencil state resource.
+ * TStaticStencilStateRHI<...>::GetStaticState() will return a FDepthStencilStateRHIRef to a stencil state with the desired
+ * settings.
+ * Should only be used from the rendering thread.
+ */
+template<
+	bool bEnableDepthWrite = true,
+	ECompareFunction DepthTest = CF_DepthNearOrEqual,
+	bool bEnableFrontFaceStencil = false,
+	ECompareFunction FrontFaceStencilTest = CF_Always,
+	EStencilOp FrontFaceStencilFailStencilOp = SO_Keep,
+	EStencilOp FrontFaceDepthFailStencilOp = SO_Keep,
+	EStencilOp FrontFacePassStencilOp = SO_Keep,
+	bool bEnableBackFaceStencil = false,
+	ECompareFunction BackFaceStencilTest = CF_Always,
+	EStencilOp BackFaceStencilFailStencilOp = SO_Keep,
+	EStencilOp BackFaceDepthFailStencilOp = SO_Keep,
+	EStencilOp BackFacePassStencilOp = SO_Keep,
+	uint8 StencilReadMask = 0xFF,
+	uint8 StencilWriteMask = 0xFF
+	>
+class TStaticDepthStencilState : public TStaticStateRHI<
+	TStaticDepthStencilState<
+		bEnableDepthWrite,
+		DepthTest,
+		bEnableFrontFaceStencil,
+		FrontFaceStencilTest,
+		FrontFaceStencilFailStencilOp,
+		FrontFaceDepthFailStencilOp,
+		FrontFacePassStencilOp,
+		bEnableBackFaceStencil,
+		BackFaceStencilTest,
+		BackFaceStencilFailStencilOp,
+		BackFaceDepthFailStencilOp,
+		BackFacePassStencilOp,
+		StencilReadMask,
+		StencilWriteMask
+		>,
+	FDepthStencilStateRHIRef,
+	FRHIDepthStencilState*
+	>
+{
+public:
+	static FDepthStencilStateRHIRef CreateRHI()
+	{
+		FDepthStencilStateInitializerRHI Initializer(
+			bEnableDepthWrite,
+			DepthTest,
+			bEnableFrontFaceStencil,
+			FrontFaceStencilTest,
+			FrontFaceStencilFailStencilOp,
+			FrontFaceDepthFailStencilOp,
+			FrontFacePassStencilOp,
+			bEnableBackFaceStencil,
+			BackFaceStencilTest,
+			BackFaceStencilFailStencilOp,
+			BackFaceDepthFailStencilOp,
+			BackFacePassStencilOp,
+			StencilReadMask,
+			StencilWriteMask);
+
+		return RHICreateDepthStencilState(Initializer);
+	}
+};
+
+/**
+ * A static RHI blend state resource.
+ * TStaticBlendStateRHI<...>::GetStaticState() will return a FBlendStateRHIRef to a blend state with the desired settings.
+ * Should only be used from the rendering thread.
+ * 
+ * Alpha blending happens on GPU's as:
+ * FinalColor.rgb = SourceColor * ColorSrcBlend (ColorBlendOp) DestColor * ColorDestBlend;
+ * if (BlendState->bSeparateAlphaBlendEnable)
+ *		FinalColor.a = SourceAlpha * AlphaSrcBlend (AlphaBlendOp) DestAlpha * AlphaDestBlend;
+ * else
+ *		Alpha blended the same way as rgb
+ * 
+ * Where source is the color coming from the pixel shader, and target is the color in the render target.
+ *
+ * So for example, TStaticBlendState<BO_Add,BF_SourceAlpha,BF_InverseSourceAlpha,BO_Add,BF_Zero,BF_One> produces:
+ * FinalColor.rgb = SourceColor * SourceAlpha + DestColor * (1 - SourceAlpha);
+ * FinalColor.a = SourceAlpha * 0 + DestAlpha * 1;
+ */
+template<
+	EColorWriteMask RT0ColorWriteMask = CW_RGBA,
+	EBlendOperation RT0ColorBlendOp = BO_Add,
+	EBlendFactor    RT0ColorSrcBlend = BF_One,
+	EBlendFactor    RT0ColorDestBlend = BF_Zero,
+	EBlendOperation RT0AlphaBlendOp = BO_Add,
+	EBlendFactor    RT0AlphaSrcBlend = BF_One,
+	EBlendFactor    RT0AlphaDestBlend = BF_Zero,
+	EColorWriteMask RT1ColorWriteMask = CW_RGBA,
+	EBlendOperation RT1ColorBlendOp = BO_Add,
+	EBlendFactor    RT1ColorSrcBlend = BF_One,
+	EBlendFactor    RT1ColorDestBlend = BF_Zero,
+	EBlendOperation RT1AlphaBlendOp = BO_Add,
+	EBlendFactor    RT1AlphaSrcBlend = BF_One,
+	EBlendFactor    RT1AlphaDestBlend = BF_Zero,
+	EColorWriteMask RT2ColorWriteMask = CW_RGBA,
+	EBlendOperation RT2ColorBlendOp = BO_Add,
+	EBlendFactor    RT2ColorSrcBlend = BF_One,
+	EBlendFactor    RT2ColorDestBlend = BF_Zero,
+	EBlendOperation RT2AlphaBlendOp = BO_Add,
+	EBlendFactor    RT2AlphaSrcBlend = BF_One,
+	EBlendFactor    RT2AlphaDestBlend = BF_Zero,
+	EColorWriteMask RT3ColorWriteMask = CW_RGBA,
+	EBlendOperation RT3ColorBlendOp = BO_Add,
+	EBlendFactor    RT3ColorSrcBlend = BF_One,
+	EBlendFactor    RT3ColorDestBlend = BF_Zero,
+	EBlendOperation RT3AlphaBlendOp = BO_Add,
+	EBlendFactor    RT3AlphaSrcBlend = BF_One,
+	EBlendFactor    RT3AlphaDestBlend = BF_Zero,
+	EColorWriteMask RT4ColorWriteMask = CW_RGBA,
+	EBlendOperation RT4ColorBlendOp = BO_Add,
+	EBlendFactor    RT4ColorSrcBlend = BF_One,
+	EBlendFactor    RT4ColorDestBlend = BF_Zero,
+	EBlendOperation RT4AlphaBlendOp = BO_Add,
+	EBlendFactor    RT4AlphaSrcBlend = BF_One,
+	EBlendFactor    RT4AlphaDestBlend = BF_Zero,
+	EColorWriteMask RT5ColorWriteMask = CW_RGBA,
+	EBlendOperation RT5ColorBlendOp = BO_Add,
+	EBlendFactor    RT5ColorSrcBlend = BF_One,
+	EBlendFactor    RT5ColorDestBlend = BF_Zero,
+	EBlendOperation RT5AlphaBlendOp = BO_Add,
+	EBlendFactor    RT5AlphaSrcBlend = BF_One,
+	EBlendFactor    RT5AlphaDestBlend = BF_Zero,
+	EColorWriteMask RT6ColorWriteMask = CW_RGBA,
+	EBlendOperation RT6ColorBlendOp = BO_Add,
+	EBlendFactor    RT6ColorSrcBlend = BF_One,
+	EBlendFactor    RT6ColorDestBlend = BF_Zero,
+	EBlendOperation RT6AlphaBlendOp = BO_Add,
+	EBlendFactor    RT6AlphaSrcBlend = BF_One,
+	EBlendFactor    RT6AlphaDestBlend = BF_Zero,
+	EColorWriteMask RT7ColorWriteMask = CW_RGBA,
+	EBlendOperation RT7ColorBlendOp = BO_Add,
+	EBlendFactor    RT7ColorSrcBlend = BF_One,
+	EBlendFactor    RT7ColorDestBlend = BF_Zero,
+	EBlendOperation RT7AlphaBlendOp = BO_Add,
+	EBlendFactor    RT7AlphaSrcBlend = BF_One,
+	EBlendFactor    RT7AlphaDestBlend = BF_Zero,
+	bool			bUseAlphaToCoverage = false
+	>
+class TStaticBlendState : public TStaticStateRHI<
+	TStaticBlendState<
+		RT0ColorWriteMask,RT0ColorBlendOp,RT0ColorSrcBlend,RT0ColorDestBlend,RT0AlphaBlendOp,RT0AlphaSrcBlend,RT0AlphaDestBlend,
+		RT1ColorWriteMask,RT1ColorBlendOp,RT1ColorSrcBlend,RT1ColorDestBlend,RT1AlphaBlendOp,RT1AlphaSrcBlend,RT1AlphaDestBlend,
+		RT2ColorWriteMask,RT2ColorBlendOp,RT2ColorSrcBlend,RT2ColorDestBlend,RT2AlphaBlendOp,RT2AlphaSrcBlend,RT2AlphaDestBlend,
+		RT3ColorWriteMask,RT3ColorBlendOp,RT3ColorSrcBlend,RT3ColorDestBlend,RT3AlphaBlendOp,RT3AlphaSrcBlend,RT3AlphaDestBlend,
+		RT4ColorWriteMask,RT4ColorBlendOp,RT4ColorSrcBlend,RT4ColorDestBlend,RT4AlphaBlendOp,RT4AlphaSrcBlend,RT4AlphaDestBlend,
+		RT5ColorWriteMask,RT5ColorBlendOp,RT5ColorSrcBlend,RT5ColorDestBlend,RT5AlphaBlendOp,RT5AlphaSrcBlend,RT5AlphaDestBlend,
+		RT6ColorWriteMask,RT6ColorBlendOp,RT6ColorSrcBlend,RT6ColorDestBlend,RT6AlphaBlendOp,RT6AlphaSrcBlend,RT6AlphaDestBlend,
+		RT7ColorWriteMask,RT7ColorBlendOp,RT7ColorSrcBlend,RT7ColorDestBlend,RT7AlphaBlendOp,RT7AlphaSrcBlend,RT7AlphaDestBlend,
+		bUseAlphaToCoverage
+		>,
+	FBlendStateRHIRef,
+	FRHIBlendState*
+	>
+{
+public:
+	static FBlendStateRHIRef CreateRHI()
+	{
+		TStaticArray<FBlendStateInitializerRHI::FRenderTarget,8> RenderTargetBlendStates;
+		RenderTargetBlendStates[0] = FBlendStateInitializerRHI::FRenderTarget(RT0ColorBlendOp,RT0ColorSrcBlend,RT0ColorDestBlend,RT0AlphaBlendOp,RT0AlphaSrcBlend,RT0AlphaDestBlend,RT0ColorWriteMask);
+		RenderTargetBlendStates[1] = FBlendStateInitializerRHI::FRenderTarget(RT1ColorBlendOp,RT1ColorSrcBlend,RT1ColorDestBlend,RT1AlphaBlendOp,RT1AlphaSrcBlend,RT1AlphaDestBlend,RT1ColorWriteMask);
+		RenderTargetBlendStates[2] = FBlendStateInitializerRHI::FRenderTarget(RT2ColorBlendOp,RT2ColorSrcBlend,RT2ColorDestBlend,RT2AlphaBlendOp,RT2AlphaSrcBlend,RT2AlphaDestBlend,RT2ColorWriteMask);
+		RenderTargetBlendStates[3] = FBlendStateInitializerRHI::FRenderTarget(RT3ColorBlendOp,RT3ColorSrcBlend,RT3ColorDestBlend,RT3AlphaBlendOp,RT3AlphaSrcBlend,RT3AlphaDestBlend,RT3ColorWriteMask);
+		RenderTargetBlendStates[4] = FBlendStateInitializerRHI::FRenderTarget(RT4ColorBlendOp,RT4ColorSrcBlend,RT4ColorDestBlend,RT4AlphaBlendOp,RT4AlphaSrcBlend,RT4AlphaDestBlend,RT4ColorWriteMask);
+		RenderTargetBlendStates[5] = FBlendStateInitializerRHI::FRenderTarget(RT5ColorBlendOp,RT5ColorSrcBlend,RT5ColorDestBlend,RT5AlphaBlendOp,RT5AlphaSrcBlend,RT5AlphaDestBlend,RT5ColorWriteMask);
+		RenderTargetBlendStates[6] = FBlendStateInitializerRHI::FRenderTarget(RT6ColorBlendOp,RT6ColorSrcBlend,RT6ColorDestBlend,RT6AlphaBlendOp,RT6AlphaSrcBlend,RT6AlphaDestBlend,RT6ColorWriteMask);
+		RenderTargetBlendStates[7] = FBlendStateInitializerRHI::FRenderTarget(RT7ColorBlendOp,RT7ColorSrcBlend,RT7ColorDestBlend,RT7AlphaBlendOp,RT7AlphaSrcBlend,RT7AlphaDestBlend,RT7ColorWriteMask);
+
+		return RHICreateBlendState(FBlendStateInitializerRHI(RenderTargetBlendStates, bUseAlphaToCoverage));
+	}
+};
+
+
+/**
+ * A static RHI blend state resource which only allows controlling MRT write masks, for use when only opaque blending is needed.
+ */
+template<
+	EColorWriteMask RT0ColorWriteMask = CW_RGBA,
+	EColorWriteMask RT1ColorWriteMask = CW_RGBA,
+	EColorWriteMask RT2ColorWriteMask = CW_RGBA,
+	EColorWriteMask RT3ColorWriteMask = CW_RGBA,
+	EColorWriteMask RT4ColorWriteMask = CW_RGBA,
+	EColorWriteMask RT5ColorWriteMask = CW_RGBA,
+	EColorWriteMask RT6ColorWriteMask = CW_RGBA,
+	EColorWriteMask RT7ColorWriteMask = CW_RGBA,
+	bool			bUseAlphaToCoverage = false
+	>
+class TStaticBlendStateWriteMask : public TStaticBlendState<
+	RT0ColorWriteMask,BO_Add,BF_One,BF_Zero,BO_Add,BF_One,BF_Zero,
+	RT1ColorWriteMask,BO_Add,BF_One,BF_Zero,BO_Add,BF_One,BF_Zero,
+	RT2ColorWriteMask,BO_Add,BF_One,BF_Zero,BO_Add,BF_One,BF_Zero,
+	RT3ColorWriteMask,BO_Add,BF_One,BF_Zero,BO_Add,BF_One,BF_Zero,
+	RT4ColorWriteMask,BO_Add,BF_One,BF_Zero,BO_Add,BF_One,BF_Zero,
+	RT5ColorWriteMask,BO_Add,BF_One,BF_Zero,BO_Add,BF_One,BF_Zero,
+	RT6ColorWriteMask,BO_Add,BF_One,BF_Zero,BO_Add,BF_One,BF_Zero,
+	RT7ColorWriteMask,BO_Add,BF_One,BF_Zero,BO_Add,BF_One,BF_Zero,
+	bUseAlphaToCoverage
+	>
+{
+public:
+	static FBlendStateRHIRef CreateRHI()
+	{
+		return TStaticBlendState<
+			RT0ColorWriteMask,BO_Add,BF_One,BF_Zero,BO_Add,BF_One,BF_Zero,
+			RT1ColorWriteMask,BO_Add,BF_One,BF_Zero,BO_Add,BF_One,BF_Zero,
+			RT2ColorWriteMask,BO_Add,BF_One,BF_Zero,BO_Add,BF_One,BF_Zero,
+			RT3ColorWriteMask,BO_Add,BF_One,BF_Zero,BO_Add,BF_One,BF_Zero,
+			RT4ColorWriteMask,BO_Add,BF_One,BF_Zero,BO_Add,BF_One,BF_Zero,
+			RT5ColorWriteMask,BO_Add,BF_One,BF_Zero,BO_Add,BF_One,BF_Zero,
+			RT6ColorWriteMask,BO_Add,BF_One,BF_Zero,BO_Add,BF_One,BF_Zero,
+			RT7ColorWriteMask,BO_Add,BF_One,BF_Zero,BO_Add,BF_One,BF_Zero,
+			bUseAlphaToCoverage
+			>
+			::CreateRHI();
+	}
+};
+
