Cleaning up header after move.

diff --git a/Engine/Source/Runtime/RenderCore/Public/RHIStaticStates.h b/Engine/Source/Runtime/RenderCore/Public/RHIStaticStates.h
--- a/Engine/Source/Runtime/RenderCore/Public/RHIStaticStates.h
+++ b/Engine/Source/Runtime/RenderCore/Public/RHIStaticStates.h
@@ -6,8 +6,6 @@
 
 #pragma once
 
-// HEADER_UNIT_SKIP - This can't be built by itself since it includes RenderResource.h which is not part of this module's additional include paths
-
 #include "CoreMinimal.h"
 #include "Stats/Stats.h"
 #include "Async/TaskGraphInterfaces.h"
@@ -16,8 +14,6 @@
 #include "RenderResource.h"
 #include "Misc/ScopedEvent.h"
 
-extern RHI_API FCriticalSection StaticStateRHICriticalSection;
-
 /**
  * Helper task to initialize a static resource on the render thread.
  */
