Add PSO cache name to the PSO precompile context.
Provides some useful tracking hints for PSO cache delegate users.

diff --git a/Engine/Source/Runtime/RenderCore/Private/ShaderPipelineCache.cpp b/Engine/Source/Runtime/RenderCore/Private/ShaderPipelineCache.cpp
--- a/Engine/Source/Runtime/RenderCore/Private/ShaderPipelineCache.cpp
+++ b/Engine/Source/Runtime/RenderCore/Private/ShaderPipelineCache.cpp
@@ -1997,6 +1997,7 @@ bool FShaderPipelineCacheTask::Open(const FString& Key, FString const& Name, ESh
 {
 	PSOCacheKey = Key;
 	ShaderPlatform = Platform;
+	PrecompileContext = FShaderPipelineCache::FShaderCachePrecompileContext(Name);
 
 	if (bIsUserCache)
 	{
diff --git a/Engine/Source/Runtime/RenderCore/Public/ShaderPipelineCache.h b/Engine/Source/Runtime/RenderCore/Public/ShaderPipelineCache.h
--- a/Engine/Source/Runtime/RenderCore/Public/ShaderPipelineCache.h
+++ b/Engine/Source/Runtime/RenderCore/Public/ShaderPipelineCache.h
@@ -157,9 +157,13 @@ public:
 	/** Allows handlers of FShaderCacheOpenedDelegate to send state to begin and complete pre-compilation delegates */
 	class FShaderCachePrecompileContext
 	{
-		bool bSlowPrecompileTask;
+		bool bSlowPrecompileTask = false;
+		FString CacheName;
 	public:
-		FShaderCachePrecompileContext() : bSlowPrecompileTask(false) {}
+		FShaderCachePrecompileContext() = default;
+		FShaderCachePrecompileContext(const FString& CacheNameIn) : CacheName(CacheNameIn) {}
+
+		const FString& GetCacheName() const { return CacheName; }
 		void SetPrecompilationIsSlowTask() { bSlowPrecompileTask = true; }
 		bool IsPrecompilationSlowTask() const { return bSlowPrecompileTask; }
 	};
