[shaders]
- move compile job types into RenderCore and use them in SCW
- using the above factor out a helper function for pipeline job compilation to further reduce code duplication
- make flags indicating job success/failure conditions part of the SCW output data rather than evaluating them when reading the output (so they can be evaluated once in the helper function mentioned above)
- modify debug worker input dumps to write a single file for combined-pipeline compilation rather than one file per stage

diff --git a/Engine/Source/Programs/ShaderCompileWorker/Private/ShaderCompileWorker.cpp b/Engine/Source/Programs/ShaderCompileWorker/Private/ShaderCompileWorker.cpp
--- a/Engine/Source/Programs/ShaderCompileWorker/Private/ShaderCompileWorker.cpp
+++ b/Engine/Source/Programs/ShaderCompileWorker/Private/ShaderCompileWorker.cpp
@@ -8,6 +8,7 @@
 #include "RequiredProgramMainCPPInclude.h"
 #include "ShaderCompilerCore.h"
 #include "ShaderCompilerCommon.h"
+#include "ShaderCompilerJobTypes.h"
 #include "ShaderCore.h"
 #include "HAL/ExceptionHandling.h"
 #include "Interfaces/IShaderFormat.h"
@@ -224,8 +225,9 @@ public:
 
 		while(true)
 		{
-			TArray<FJobResult> SingleJobResults;
-			TArray<FPipelineJobResult> PipelineJobResults;
+			TArray<FShaderCompileJob> SingleJobs;
+			TArray<FShaderPipelineCompileJob> PipelineJobs;
+			TArray<FString> PipelineJobNames;
 
 			// Read & Process Input
 			{
@@ -239,7 +241,7 @@ public:
 
 				UE_LOG(LogShaders, Log, TEXT("Processing shader"));
 
-				ProcessInputFromArchive(InputFilePtr, SingleJobResults, PipelineJobResults);
+				ProcessInputFromArchive(InputFilePtr, SingleJobs, PipelineJobs, PipelineJobNames);
 
 				LastCompileTime = FPlatformTime::Seconds();
 
@@ -253,13 +255,13 @@ public:
 				// write to in-memory bytestream instead for debugging purposes
 				TArray<uint8> MemBlock;
 				FMemoryWriter MemWriter(MemBlock);
-				WriteToOutputArchive(&MemWriter, SingleJobResults, PipelineJobResults);
+				WriteToOutputArchive(&MemWriter, SingleJobs, PipelineJobs, PipelineJobNames);
 			}
 			else
 			{
 				FArchive* OutputFilePtr = CreateOutputArchive();
 				check(OutputFilePtr);
-				WriteToOutputArchive(OutputFilePtr, SingleJobResults, PipelineJobResults);
+				WriteToOutputArchive(OutputFilePtr, SingleJobs, PipelineJobs, PipelineJobNames);
 				// Close the output file.
 				delete OutputFilePtr;
 
@@ -276,7 +278,7 @@ public:
 				break;
 			}
 
-			if (TimeToLive == 0 || AnyJobUsedHLSLccCompiler( SingleJobResults, PipelineJobResults ))
+			if (TimeToLive == 0 || AnyJobUsedHLSLccCompiler(SingleJobs, PipelineJobs))
 			{
 				UE_LOG(LogShaders, Log, TEXT("TimeToLive set to 0, or used HLSLcc compiler, exiting after single job"));
 				break;
@@ -295,17 +297,6 @@ public:
 	}
 
 private:
-	struct FJobResult
-	{
-		FShaderCompilerOutput CompilerOutput;
-	};
-
-	struct FPipelineJobResult
-	{
-		FString PipelineName;
-		TArray<FJobResult> SingleJobs;
-	};
-
 	const int32 ParentProcessId;
 	const FString WorkingDirectory;
 	const FString InputFilename;
@@ -352,7 +343,7 @@ private:
 		}
 	}
 
-	void ProcessInputFromArchive(FArchive* InputFilePtr, TArray<FJobResult>& OutSingleJobResults, TArray<FPipelineJobResult>& OutPipelineJobResults)
+	void ProcessInputFromArchive(FArchive* InputFilePtr, TArray<FShaderCompileJob>& OutSingleJobs, TArray<FShaderPipelineCompileJob>& OutPipelineJobs, TArray<FString>& OutPipelineNames)
 	{
 		TRACE_CPUPROFILER_EVENT_SCOPE(ProcessInputFromArchive);
 
@@ -607,29 +598,26 @@ private:
 			InputFile << NumBatches;
 
 			FlushShaderFileCache();
+
+			OutSingleJobs.Reserve(NumBatches);
 			
 			for (int32 BatchIndex = 0; BatchIndex < NumBatches; BatchIndex++)
 			{
+				FShaderCompileJob& Job = *new(OutSingleJobs) FShaderCompileJob;
 				// Deserialize the job's inputs.
-				FShaderCompilerInput CompilerInput;
-				InputFile << CompilerInput;
-				CompilerInput.DeserializeSharedInputs(InputFile, ExternalIncludes, SharedEnvironments, ParameterStructures);
+				InputFile << Job.Input;
+				Job.Input.DeserializeSharedInputs(InputFile, ExternalIncludes, SharedEnvironments, ParameterStructures);
 
 				// SCW doesn't run DDPI, GShaderHasCache Initialize is run  at start with no knowledge of the CustomPlatforms
 				// CustomPlatforms are known when we parse the WorkerInput so we populate the Directory here
-				if (IsCustomPlatform((EShaderPlatform)CompilerInput.Target.Platform))
+				if (IsCustomPlatform((EShaderPlatform)Job.Input.Target.Platform))
 				{
-					const EShaderPlatform ShaderPlatform = ShaderFormatNameToShaderPlatform(CompilerInput.ShaderFormat);
-					UpdateIncludeDirectoryForPreviewPlatform((EShaderPlatform)CompilerInput.Target.Platform, ShaderPlatform);
+					const EShaderPlatform ShaderPlatform = ShaderFormatNameToShaderPlatform(Job.Input.ShaderFormat);
+					UpdateIncludeDirectoryForPreviewPlatform((EShaderPlatform)Job.Input.Target.Platform, ShaderPlatform);
 				}
 
 				// Process the job.
-				FShaderCompilerOutput CompilerOutput;
-				CompileShader(GetShaderFormats(), CompilerInput, CompilerOutput, WorkingDirectory, &GNumProcessedJobs);
-
-				// Serialize the job's output.
-				FJobResult& JobResult = *new(OutSingleJobResults) FJobResult;
-				JobResult.CompilerOutput = CompilerOutput;
+				CompileShader(GetShaderFormats(), Job, WorkingDirectory, &GNumProcessedJobs); 
 			}
 		}
 
@@ -645,83 +633,37 @@ private:
 			int32 NumPipelines = 0;
 			InputFile << NumPipelines;
 
+			OutPipelineNames.Reserve(NumPipelines);
+			OutPipelineJobs.Reserve(NumPipelines);
+
 			for (int32 Index = 0; Index < NumPipelines; ++Index)
 			{
-				FPipelineJobResult& PipelineJob = *new(OutPipelineJobResults) FPipelineJobResult;
-
-				InputFile << PipelineJob.PipelineName;
-
+				FString& PipelineName = *new(OutPipelineNames) FString();
+				InputFile << PipelineName;
+				
 				int32 NumStages = 0;
 				InputFile << NumStages;
-
-				TArray<FShaderCompilerInput> CompilerInputs;
-				CompilerInputs.AddDefaulted(NumStages);
+				FShaderPipelineCompileJob* PipelineJob = new(OutPipelineJobs) FShaderPipelineCompileJob(NumStages);
 
 				for (int32 StageIndex = 0; StageIndex < NumStages; ++StageIndex)
 				{
 					// Deserialize the job's inputs.
-					InputFile << CompilerInputs[StageIndex];
-					CompilerInputs[StageIndex].DeserializeSharedInputs(InputFile, ExternalIncludes, SharedEnvironments, ParameterStructures);
+					FShaderCompileJob* Job = PipelineJob->StageJobs[StageIndex]->GetSingleShaderJob();
+					InputFile << Job->Input;
+					Job->Input.DeserializeSharedInputs(InputFile, ExternalIncludes, SharedEnvironments, ParameterStructures);
 
 					// SCW doesn't run DDPI, GShaderHasCache Initialize is run  at start with no knowledge of the CustomPlatforms
 					// CustomPlatforms are known when we parse the WorkerInput so we populate the Directory here
-					if (IsCustomPlatform((EShaderPlatform)CompilerInputs[StageIndex].Target.Platform))
+					if (IsCustomPlatform((EShaderPlatform)Job->Input.Target.Platform))
 					{
-						const EShaderPlatform ShaderPlatform = ShaderFormatNameToShaderPlatform(CompilerInputs[StageIndex].ShaderFormat);
-						UpdateIncludeDirectoryForPreviewPlatform((EShaderPlatform)CompilerInputs[StageIndex].Target.Platform, ShaderPlatform);
+						const EShaderPlatform ShaderPlatform = ShaderFormatNameToShaderPlatform(Job->Input.ShaderFormat);
+						UpdateIncludeDirectoryForPreviewPlatform((EShaderPlatform)Job->Input.Target.Platform, ShaderPlatform);
 					}
 				}
 
-				ProcessShaderPipelineCompilationJob(PipelineJob, CompilerInputs);
-			}
-		}
-	}
-
-	void ProcessShaderPipelineCompilationJob(FPipelineJobResult& PipelineJob, TArray<FShaderCompilerInput>& CompilerInputs)
-	{
-		checkf(CompilerInputs.Num() > 0, TEXT("Exiting due to Pipeline %s having zero jobs!"), *PipelineJob.PipelineName);
-
-		// Process the job.
-		FShaderCompilerOutput FirstCompilerOutput;
-		CompilerInputs[0].bCompilingForShaderPipeline = true;
-		CompilerInputs[0].bIncludeUsedOutputs = false;
-		CompileShader(GetShaderFormats(), CompilerInputs[0], FirstCompilerOutput, WorkingDirectory, &GNumProcessedJobs);
-
-		// Serialize the job's output.
-		{
-			FJobResult& JobResult = *new(PipelineJob.SingleJobs) FJobResult;
-			JobResult.CompilerOutput = FirstCompilerOutput;
-		}
-
-		bool bEnableRemovingUnused = true;
-
-		//#todo-rco: Only remove for pure VS & PS stages
-		for (int32 Index = 0; Index < CompilerInputs.Num(); ++Index)
-		{
-			auto Stage = CompilerInputs[Index].Target.Frequency;
-			if (Stage != SF_Vertex && Stage != SF_Pixel)
-			{
-				bEnableRemovingUnused = false;
-				break;
+				CompileShaderPipeline(GetShaderFormats(), PipelineJob, WorkingDirectory, &GNumProcessedJobs);
 			}
 		}
-
-		for (int32 Index = 1; Index < CompilerInputs.Num(); ++Index)
-		{
-			if (bEnableRemovingUnused && PipelineJob.SingleJobs.Last().CompilerOutput.bSupportsQueryingUsedAttributes)
-			{
-				CompilerInputs[Index].bIncludeUsedOutputs = true;
-				CompilerInputs[Index].bCompilingForShaderPipeline = true;
-				CompilerInputs[Index].UsedOutputs = PipelineJob.SingleJobs.Last().CompilerOutput.UsedAttributes;
-			}
-
-			FShaderCompilerOutput CompilerOutput;
-			CompileShader(GetShaderFormats(), CompilerInputs[Index], CompilerOutput, WorkingDirectory, &GNumProcessedJobs);
-
-			// Serialize the job's output.
-			FJobResult& JobResult = *new(PipelineJob.SingleJobs) FJobResult;
-			JobResult.CompilerOutput = CompilerOutput;
-		}
 	}
 
 	FArchive* CreateOutputArchive()
@@ -777,7 +719,7 @@ private:
 		return OutputFilePtr;
 	}
 
-	void WriteToOutputArchive(FArchive* OutputFilePtr, TArray<FJobResult>& SingleJobResults, TArray<FPipelineJobResult>& PipelineJobResults)
+	void WriteToOutputArchive(FArchive* OutputFilePtr, TArray<FShaderCompileJob>& SingleJobs, TArray<FShaderPipelineCompileJob>& PipelineJobs, TArray<FString>& PipelineNames)
 	{
 		TRACE_CPUPROFILER_EVENT_SCOPE(WriteToOutputArchive);
 
@@ -790,13 +732,13 @@ private:
 			int32 SingleJobHeader = ShaderCompileWorkerSingleJobHeader;
 			OutputFile << SingleJobHeader;
 
-			int32 NumBatches = SingleJobResults.Num();
+			int32 NumBatches = SingleJobs.Num();
 			OutputFile << NumBatches;
 
-			for (int32 ResultIndex = 0; ResultIndex < SingleJobResults.Num(); ResultIndex++)
+			for (int32 JobIndex = 0; JobIndex < SingleJobs.Num(); JobIndex++)
 			{
-				FJobResult& JobResult = SingleJobResults[ResultIndex];
-				OutputFile << JobResult.CompilerOutput;
+				FShaderCompileJob& Job = SingleJobs[JobIndex];
+				OutputFile << Job.Output;
 				UpdateFileSize(OutputFile, FileSizePosition);
 			}
 		}
@@ -804,19 +746,23 @@ private:
 		{
 			int32 PipelineJobHeader = ShaderCompileWorkerPipelineJobHeader;
 			OutputFile << PipelineJobHeader;
-			int32 NumBatches = PipelineJobResults.Num();
+			int32 NumBatches = PipelineJobs.Num();
 			OutputFile << NumBatches;
 
-			for (int32 ResultIndex = 0; ResultIndex < PipelineJobResults.Num(); ResultIndex++)
+			for (int32 JobIndex = 0; JobIndex < PipelineJobs.Num(); JobIndex++)
 			{
-				auto& PipelineJob = PipelineJobResults[ResultIndex];
-				OutputFile << PipelineJob.PipelineName;
-				int32 NumStageJobs = PipelineJob.SingleJobs.Num();
+				auto& PipelineJob = PipelineJobs[JobIndex];
+				OutputFile << PipelineNames[JobIndex];
+				bool bSucceeded = (bool)PipelineJob.bSucceeded;
+				OutputFile << bSucceeded;
+				OutputFile << PipelineJob.bFailedRemovingUnused;
+				int32 NumStageJobs = PipelineJob.StageJobs.Num();
 				OutputFile << NumStageJobs;
+
 				for (int32 Index = 0; Index < NumStageJobs; ++Index)
 				{
-					FJobResult& JobResult = PipelineJob.SingleJobs[Index];
-					OutputFile << JobResult.CompilerOutput;
+					FShaderCompileJob& JobResult = *PipelineJob.StageJobs[Index];
+					OutputFile << JobResult.Output;
 					UpdateFileSize(OutputFile, FileSizePosition);
 				}
 			}
@@ -896,26 +842,24 @@ private:
 #endif
 	}
 	
-	static bool AnyJobUsedHLSLccCompiler(TArray<FJobResult>& SingleJobResults, TArray<FPipelineJobResult>& PipelineJobResults)
+	static bool AnyJobUsedHLSLccCompiler(TArray<FShaderCompileJob>& SingleJobs, TArray<FShaderPipelineCompileJob>& PipelineJobs)
 	{
 		TRACE_CPUPROFILER_EVENT_SCOPE(AnyJobUsedHLSLccCompiler);
 
-		for (int32 ResultIndex = 0; ResultIndex < SingleJobResults.Num(); ResultIndex++)
+		for (int32 JobIndex = 0; JobIndex < SingleJobs.Num(); JobIndex++)
 		{
-			FJobResult& JobResult = SingleJobResults[ResultIndex];
-			if (JobResult.CompilerOutput.bUsedHLSLccCompiler)
+			if (SingleJobs[JobIndex].Output.bUsedHLSLccCompiler)
 			{
 				return true;
 			}
 		}
 
-		for (int32 ResultIndex = 0; ResultIndex < PipelineJobResults.Num(); ResultIndex++)
+		for (int32 JobIndex = 0; JobIndex < PipelineJobs.Num(); JobIndex++)
 		{
-			FPipelineJobResult& PipelineJob = PipelineJobResults[ResultIndex];
-			for (int32 Index = 0; Index < PipelineJob.SingleJobs.Num(); ++Index)
+			FShaderPipelineCompileJob& PipelineJob = PipelineJobs[JobIndex];
+			for (int32 Index = 0; Index < PipelineJob.StageJobs.Num(); ++Index)
 			{
-				FJobResult& JobResult = PipelineJob.SingleJobs[Index];
-				if (JobResult.CompilerOutput.bUsedHLSLccCompiler)
+				if (PipelineJob.StageJobs[Index]->Output.bUsedHLSLccCompiler)
 				{
 					return true;
 				}
@@ -1044,14 +988,14 @@ static void DirectCompile(const TArray<const class IShaderFormat*>& ShaderFormat
 
 	FString Dir = FPlatformProcess::UserTempDir();
 
-	FShaderCompilerInput Input;
-	Input.EntryPointName = Entry;
-	Input.ShaderFormat = FormatName;
-	Input.ShaderPlatformName = ShaderPlatformName;
-	Input.VirtualSourceFilePath = InputFile;
-	Input.Target.Platform =  ShaderFormatNameToShaderPlatform(FormatName);
-	Input.Target.Frequency = Frequency;
-	Input.bSkipPreprocessedCache = !bUseMCPP;
+	FShaderCompileJob Job;
+	Job.Input.EntryPointName = Entry;
+	Job.Input.ShaderFormat = FormatName;
+	Job.Input.ShaderPlatformName = ShaderPlatformName;
+	Job.Input.VirtualSourceFilePath = InputFile;
+	Job.Input.Target.Platform =  ShaderFormatNameToShaderPlatform(FormatName);
+	Job.Input.Target.Frequency = Frequency;
+	Job.Input.bSkipPreprocessedCache = !bUseMCPP;
 
 	uint32 ResourceIndex = 0;
 	auto AddResourceTableEntry = [&ResourceIndex](TMap<FString, FResourceTableEntry>& Map, const FString& Name, const FString& UBName, int32 Type)
@@ -1064,14 +1008,14 @@ static void DirectCompile(const TArray<const class IShaderFormat*>& ShaderFormat
 		++ResourceIndex;
 	};
 
-	Input.Environment.CompilerFlags = FShaderCompilerFlags(CFlags);
+	Job.Input.Environment.CompilerFlags = FShaderCompilerFlags(CFlags);
 
-	Input.bCompilingForShaderPipeline = bPipeline;
-	Input.bIncludeUsedOutputs = bIncludeUsedOutputs;
-	Input.UsedOutputs = UsedOutputs;
+	Job.Input.bCompilingForShaderPipeline = bPipeline;
+	Job.Input.bIncludeUsedOutputs = bIncludeUsedOutputs;
+	Job.Input.UsedOutputs = UsedOutputs;
 
 	FShaderCompilerOutput Output;
-	CompileShader(GetShaderFormats(), Input, Output, Dir, &GNumProcessedJobs);
+	CompileShader(GetShaderFormats(), Job, Dir, &GNumProcessedJobs);
 }
 
 
diff --git a/Engine/Source/Runtime/Engine/Private/ShaderCompiler/ShaderCompiler.cpp b/Engine/Source/Runtime/Engine/Private/ShaderCompiler/ShaderCompiler.cpp
--- a/Engine/Source/Runtime/Engine/Private/ShaderCompiler/ShaderCompiler.cpp
+++ b/Engine/Source/Runtime/Engine/Private/ShaderCompiler/ShaderCompiler.cpp
@@ -41,7 +41,6 @@
 #include "RenderUtils.h"
 #include "SceneInterface.h"
 #include "SceneManagement.h"
-#include "Serialization/MemoryHasher.h"
 #include "Serialization/NameAsStringProxyArchive.h"
 #include "ShaderCodeLibrary.h"
 #include "ShaderPlatformCachedIniValue.h"
@@ -134,14 +133,6 @@ bool IsShaderJobCacheDDCEnabled()
 	return false;
 }
 
-int32 GShaderCompilerDumpCompileJobInputs = 0;
-static FAutoConsoleVariableRef CVarShaderCompilerDumpCompileJobInputs(
-	TEXT("r.ShaderCompiler.DumpCompileJobInputs"),
-	GShaderCompilerDumpCompileJobInputs,
-	TEXT("if != 0, unpreprocessed input of the shader compiler jobs will be dumped into the debug directory for closer inspection. This is a debugging feature which is disabled by default."),
-	ECVF_Default
-);
-
 int32 GShaderCompilerCacheStatsPrintoutInterval = 180;
 static FAutoConsoleVariableRef CVarShaderCompilerCacheStatsPrintoutInterval(
 	TEXT("r.ShaderCompiler.CacheStatsPrintoutInterval"),
@@ -313,40 +304,6 @@ namespace ShaderCompiler
 	}
 }
 
-// The Id of 0 is reserved for global shaders
-FThreadSafeCounter FShaderCommonCompileJob::JobIdCounter(2);
-
-uint32 FShaderCommonCompileJob::GetNextJobId()
-{
-	uint32 Id = JobIdCounter.Increment();
-	if (Id == UINT_MAX)
-	{
-		JobIdCounter.Set(2);
-	}
-	return Id;
-}
-
-FShaderPipelineCompileJob::FShaderPipelineCompileJob(uint32 InHash, uint32 InId, EShaderCompileJobPriority InPriroity, const FShaderPipelineCompileJobKey& InKey) :
-	FShaderCommonCompileJob(Type, InHash, InId, InPriroity),
-	Key(InKey)
-{
-	const auto& Stages = InKey.ShaderPipeline->GetStages();
-	StageJobs.Empty(Stages.Num());
-	for (const FShaderType* ShaderType : Stages)
-	{
-		const FShaderCompileJobKey StageKey(ShaderType, InKey.VFType, InKey.PermutationId);
-		StageJobs.Add(new FShaderCompileJob(StageKey.MakeHash(InId), InId, InPriroity, StageKey));
-	}
-}
-
-FString FShaderCompileJobKey::ToString() const
-{
-	return FString::Printf(
-		TEXT("ShaderType:%s VertexFactoryType:%s PermutationId:%d"),
-		ShaderType ? ShaderType->GetName() : TEXT("None"),
-		VFType ? VFType->GetName() : TEXT("None"),
-		PermutationId);
-}
 
 FShaderCompileJobCollection::FShaderCompileJobCollection()
 {
@@ -1193,8 +1150,6 @@ static TAutoConsoleVariable<bool> CVarShadersUseLegacyPreprocessor(
 	TEXT("\tfalse: Disabled - preprocess with STB\n"),
 	ECVF_ReadOnly);
 
-extern bool CompileShaderPipeline(const TArray<const IShaderFormat*>& ShaderFormats, FName Format, FShaderPipelineCompileJob* PipelineJob, const FString& Dir);
-
 #if ENABLE_COOK_STATS
 namespace ShaderCompilerCookStats
 {
@@ -1652,12 +1607,11 @@ bool DoWriteTasksInner(const TArray<FShaderCommonCompileJobPtr>& QueuedJobs, FAr
 const TCHAR* DebugWorkerInputFileName = TEXT("DebugSCW.in");
 const TCHAR* DebugWorkerOutputFileName = TEXT("DebugSCW.out");
 
-FString CreateShaderCompilerWorkerDebugCommandLine(const FShaderCompileJob* Job)
+FString CreateShaderCompilerWorkerDebugCommandLine(FString DebugWorkerInputFilePath)
 {
 	// 0 is parent PID, pass zero TTL and KeepInput to make SCW process the single job then exit without deleting the input file
-	return FString::Printf(TEXT("\"%s\" 0 \"%s\" %s %s -TimeToLive=0.0f -KeepInput"),
-		*Job->Input.DumpDebugInfoPath, // working directory for SCW
-		*Job->Input.DebugGroupName, // console window title
+	return FString::Printf(TEXT("\"%s\" 0 \"DebugSCW\" %s %s -TimeToLive=0.0f -KeepInput"),
+		*DebugWorkerInputFilePath, // working directory for SCW
 		DebugWorkerInputFileName,
 		DebugWorkerOutputFileName);
 }
@@ -1669,40 +1623,36 @@ bool FShaderCompileUtilities::DoWriteTasks(const TArray<FShaderCommonCompileJobP
 	{
 		for (const FShaderCommonCompileJobPtr& CommonJob : QueuedJobs)
 		{
-			auto WriteSingleJobFile = [](FShaderCompileJob* Job)
-			{
-				if (Job->Input.DumpDebugInfoEnabled())
-				{
-					TArray<FShaderCommonCompileJobPtr> SingleJobArray;
-					SingleJobArray.Add(Job);
-
-					FString DebugWorkerInputFilePath = Job->Input.DumpDebugInfoPath / DebugWorkerInputFileName;
-					FArchive* DebugWorkerInputFileWriter = IFileManager::Get().CreateFileWriter(*DebugWorkerInputFilePath, FILEWRITE_NoFail);
-					DoWriteTasksInner(
-						SingleJobArray, 
-						*DebugWorkerInputFileWriter, 
-						nullptr,	// Don't pass a IDistributedBuildController, this is only used for conversion to relative paths which we do not want for debug files
-						false,		// As above, use absolute paths not relative
-						true);		// Always compress the debug files; they are rather large so this saves some disk space
-					DebugWorkerInputFileWriter->Close();
-					delete DebugWorkerInputFileWriter;
-
-					FFileHelper::SaveStringToFile(
-						CreateShaderCompilerWorkerDebugCommandLine(Job), 
-						*(Job->Input.DumpDebugInfoPath / TEXT("DebugCompileArgs.txt")));
-				}
-			};
 			FShaderPipelineCompileJob* PipelineJob = CommonJob->GetShaderPipelineJob();
+			FString DebugWorkerInputFilePath;
 			if (PipelineJob)
 			{
-				for (FShaderCompileJob* StageJob : PipelineJob->StageJobs)
-				{
-					WriteSingleJobFile(StageJob);
-				}
+				// for pipeline jobs, write out the worker input for the whole pipeline, but only for the first stage
+				// would be better to put in a parent folder probably...
+				DebugWorkerInputFilePath = PipelineJob->StageJobs[0]->Input.DumpDebugInfoPath;
 			}
 			else
 			{
-				WriteSingleJobFile(CommonJob->GetSingleShaderJob());
+				DebugWorkerInputFilePath = CommonJob->GetSingleShaderJob()->Input.DumpDebugInfoPath;
+			}
+			if (!DebugWorkerInputFilePath.IsEmpty())
+			{
+				TArray<FShaderCommonCompileJobPtr> SingleJobArray;
+				SingleJobArray.Add(CommonJob);
+
+				FArchive* DebugWorkerInputFileWriter = IFileManager::Get().CreateFileWriter(*(DebugWorkerInputFilePath / DebugWorkerInputFileName), FILEWRITE_NoFail);
+				DoWriteTasksInner(
+					SingleJobArray,
+					*DebugWorkerInputFileWriter,
+					nullptr,	// Don't pass a IDistributedBuildController, this is only used for conversion to relative paths which we do not want for debug files
+					false,		// As above, use absolute paths not relative
+					true);		// Always compress the debug files; they are rather large so this saves some disk space
+				DebugWorkerInputFileWriter->Close();
+				delete DebugWorkerInputFileWriter;
+
+				FFileHelper::SaveStringToFile(
+					CreateShaderCompilerWorkerDebugCommandLine(DebugWorkerInputFilePath),
+					*(DebugWorkerInputFilePath / TEXT("DebugCompileArgs.txt")));
 			}
 		}
 	}
@@ -2126,6 +2076,10 @@ void FShaderCompileUtilities::DoReadTaskResults(const TArray<FShaderCommonCompil
 
 				FString PipelineName;
 				OutputFile << PipelineName;
+				bool bSucceeded = false;
+				OutputFile << bSucceeded;
+				CurrentJob->bSucceeded = bSucceeded;
+				OutputFile << CurrentJob->bFailedRemovingUnused;
 				if (PipelineName != CurrentJob->Key.ShaderPipeline->GetName())
 				{
 					FString Text = FString::Printf(TEXT("Worker returned Pipeline %s, expected %s!"), *PipelineName, CurrentJob->Key.ShaderPipeline->GetName());
@@ -2134,7 +2088,6 @@ void FShaderCompileUtilities::DoReadTaskResults(const TArray<FShaderCommonCompil
 
 				check(!CurrentJob->bFinalized);
 				CurrentJob->bFinalized = true;
-				CurrentJob->bFailedRemovingUnused = false;
 
 				int32 NumStageJobs = -1;
 				OutputFile << NumStageJobs;
@@ -2146,14 +2099,11 @@ void FShaderCompileUtilities::DoReadTaskResults(const TArray<FShaderCommonCompil
 				}
 				else
 				{
-					CurrentJob->bSucceeded = true;
 					for (int32 Index = 0; Index < NumStageJobs; Index++)
 					{
 						FShaderCompileJob* SingleJob = CurrentJob->StageJobs[Index];
 						// cannot reissue a single stage of a pipeline job
 						ReadSingleJob(SingleJob, OutputFile);
-						CurrentJob->bFailedRemovingUnused = CurrentJob->bFailedRemovingUnused || SingleJob->Output.bFailedRemovingUnused;
-						CurrentJob->bSucceeded = CurrentJob->bSucceeded && SingleJob->bSucceeded;
 					}
 				}
 			}
@@ -2917,7 +2867,7 @@ void FShaderCompileUtilities::ExecuteShaderCompileJob(FShaderCommonCompileJob& J
 	if (SingleJob)
 	{
 		const FName Format = (SingleJob->Input.ShaderFormat != NAME_None) ? SingleJob->Input.ShaderFormat : LegacyShaderPlatformToShaderFormat(EShaderPlatform(SingleJob->Input.Target.Platform));
-		CompileShader(ShaderFormats, SingleJob->Input, SingleJob->Output, WorkingDir);
+		CompileShader(ShaderFormats, *SingleJob, WorkingDir);
 	}
 	else
 	{
@@ -2941,7 +2891,7 @@ void FShaderCompileUtilities::ExecuteShaderCompileJob(FShaderCommonCompileJob& J
 			}
 		}
 
-		CompileShaderPipeline(ShaderFormats, Format, PipelineJob, WorkingDir);
+		CompileShaderPipeline(ShaderFormats, PipelineJob, WorkingDir);
 	}
 
 	Job.bFinalized = true;
@@ -8414,206 +8364,6 @@ void LoadGlobalShadersForRemoteRecompile(FArchive& Ar, EShaderPlatform ShaderPla
 	}
 }
 
-FShaderCommonCompileJob::FInputHash FShaderCompileJob::GetInputHash()
-{
-	if (bInputHashSet)
-	{
-		return InputHash;
-	}
-
-	auto SerializeInputs = [this](FArchive& Archive)
-	{
-		checkf(Archive.IsSaving() && !Archive.IsLoading(), TEXT("A loading archive is passed to FShaderCompileJob::GetInputHash(), this is not supported as it may corrupt its data"));
-
-		Archive << Input;
-		Input.Environment.SerializeEverythingButFiles(Archive);
-
-		// hash the source file so changes to files during the development are picked up
-		const FSHAHash& SourceHash = GetShaderFileHash(*Input.VirtualSourceFilePath, Input.Target.GetPlatform());
-		Archive << const_cast<FSHAHash&>(SourceHash);
-
-		// unroll the included files for the parallel processing.
-		// These are temporary arrays that only exist for the ParallelFor
-		TArray<const TCHAR*> IncludeVirtualPaths;
-		TArray<const FString*> Contents;
-		TArray<bool> OnlyHashIncludes;
-		TArray<FBlake3Hash> Hashes;
-
-		// while the contents of this is already hashed (included in Environment's operator<<()), we still need to account for includes in the generated files and hash them, too
-		for (TMap<FString, FString>::TConstIterator It(Input.Environment.IncludeVirtualPathToContentsMap); It; ++It)
-		{
-			const FString& VirtualPath = It.Key();
-			IncludeVirtualPaths.Add(*VirtualPath);
-			Contents.Add(&It.Value());
-			OnlyHashIncludes.Add(true);	// not hashing contents of the file itself, as it was included in Environment's operator<<()
-			Hashes.AddDefaulted();
-		}
-
-		for (TMap<FString, FThreadSafeSharedStringPtr>::TConstIterator It(Input.Environment.IncludeVirtualPathToExternalContentsMap); It; ++It)
-		{
-			const FString& VirtualPath = It.Key();
-			IncludeVirtualPaths.Add(*VirtualPath);
-			check(It.Value());
-			Contents.Add(&(*It.Value()));
-			OnlyHashIncludes.Add(false);
-			Hashes.AddDefaulted();
-		}
-
-		if (Input.SharedEnvironment)
-		{
-			Input.SharedEnvironment->SerializeEverythingButFiles(Archive);
-
-			for (TMap<FString, FString>::TConstIterator It(Input.SharedEnvironment->IncludeVirtualPathToContentsMap); It; ++It)
-			{
-				const FString& VirtualPath = It.Key();
-				IncludeVirtualPaths.Add(*VirtualPath);
-				Contents.Add(&It.Value());
-				OnlyHashIncludes.Add(true);	// not hashing contents of the file itself, as it was included in Environment's operator<<()
-				Hashes.AddDefaulted();
-			}
-
-			for (TMap<FString, FThreadSafeSharedStringPtr>::TConstIterator It(Input.SharedEnvironment->IncludeVirtualPathToExternalContentsMap); It; ++It)
-			{
-				const FString& VirtualPath = It.Key();
-				IncludeVirtualPaths.Add(*VirtualPath);
-				check(It.Value());
-				Contents.Add(&(*It.Value()));
-				OnlyHashIncludes.Add(false);
-				Hashes.AddDefaulted();
-			}
-		}
-
-		check(IncludeVirtualPaths.Num() == Contents.Num());
-		check(Contents.Num() == OnlyHashIncludes.Num());
-		check(OnlyHashIncludes.Num() == Hashes.Num());
-
-		EShaderPlatform Platform = Input.Target.GetPlatform();
-		ParallelFor(Contents.Num(), [&IncludeVirtualPaths, &Contents, &OnlyHashIncludes, &Hashes, &Platform](int32 FileIndex)
-			{ 
-				FMemoryHasherBlake3 MemHasher;
-				HashShaderFileWithIncludes(MemHasher, IncludeVirtualPaths[FileIndex], *Contents[FileIndex], Platform, OnlyHashIncludes[FileIndex]);
-				Hashes[FileIndex] = MemHasher.Finalize();
-			},
-			EParallelForFlags::Unbalanced
-		);
-
-		// include the hashes in the main hash (consider sorting them if includes are found to have a random order)
-		for (int32 HashIndex = 0, NumHashes = Hashes.Num(); HashIndex < NumHashes; ++HashIndex)
-		{
-			Archive << Hashes[HashIndex];
-		}
-	};
-
-	// use faster hasher that doesn't allocate memory
-	FMemoryHasherBlake3 MemHasher;
-	SerializeInputs(MemHasher);
-	InputHash = MemHasher.Finalize();
-
-	if (GShaderCompilerDumpCompileJobInputs)
-	{
-		TArray<uint8> MemoryBlob;
-		FMemoryWriter MemWriter(MemoryBlob);
-
-		SerializeInputs(MemWriter);
-
-		FString IntermediateFormatPath = FPaths::ProjectSavedDir() / TEXT("ShaderJobInputs");
-#if UE_BUILD_DEBUG
-		FString TempPath = IntermediateFormatPath / TEXT("DebugEditor");
-#else
-		FString TempPath = IntermediateFormatPath / TEXT("DevelopmentEditor");
-#endif
-		IFileManager::Get().MakeDirectory(*TempPath, true);
-
-		static int32 InputHashID = 0;
-		FString FileName = Input.DebugGroupName.Replace(TEXT("/"), TEXT("_")).Replace(TEXT("<"), TEXT("_")).Replace(TEXT(">"), TEXT("_")).Replace(TEXT(":"), TEXT("_")).Replace(TEXT("|"), TEXT("_"))
-			+ TEXT("-") + Input.EntryPointName;
-		FString TempFile = TempPath / FString::Printf(TEXT("%s-%d.bin"), *FileName, InputHashID++);
-
-		TUniquePtr<FArchive> DumpAr(IFileManager::Get().CreateFileWriter(*TempFile));
-		DumpAr->Serialize(MemoryBlob.GetData(), MemoryBlob.Num());
-
-		// as an additional debugging feature, make sure that the hash is the same as calculated by the memhasher
-		FBlake3Hash Check = FBlake3::HashBuffer(MemoryBlob.GetData(), MemoryBlob.Num());
-		if (Check != InputHash)
-		{
-			UE_LOG(LogShaderCompilers, Error, TEXT("Job input hash disagrees between FMemoryHasherSHA1 (%s) and FMemoryWriter + FSHA1 (%s, which was dumped to disk)"), *LexToString(InputHash), *LexToString(Check));
-		}
-	}
-
-	bInputHashSet = true;
-	return InputHash;
-}
-
-void FShaderCompileJob::SerializeOutput(FArchive& Ar)
-{
-	double ActualCompileTime = 0.0;
-	double ActualPreprocessTime = 0.0;
-	if (Ar.IsSaving())
-	{
-		// Cached jobs won't have accurate results anyway, so reduce the storage requirements by setting those fields to a known value.
-		// This significantly reduces the memory needed to store the outputs (by more than a half)
-		ActualCompileTime = Output.CompileTime;
-		ActualPreprocessTime = Output.PreprocessTime;
-		Output.CompileTime = 0.0;
-		Output.PreprocessTime = 0.0;
-	}
-
-	Ar << Output;
-	// output hash is now serialized as part of the output, as the shader code is compressed in SCWs
-	checkf(!Output.bSucceeded || Output.OutputHash != FSHAHash(), TEXT("Successful compile job does not have an OutputHash generated."));
-
-	if (Ar.IsLoading())
-	{
-		bFinalized = true;
-		bSucceeded = Output.bSucceeded;
-	}
-	else
-	{
-		// restore the compile time for this jobs. Jobs that will be deserialized from the cache will have a compile time of 0.0
-		Output.CompileTime = ActualCompileTime;
-		Output.PreprocessTime = ActualPreprocessTime;
-	}
-}
-
-FShaderCommonCompileJob::FInputHash FShaderPipelineCompileJob::GetInputHash()
-{
-	if (bInputHashSet)
-	{
-		return InputHash;
-	}
-
-	FBlake3 Hasher;
-	for (int32 Index = 0; Index < StageJobs.Num(); ++Index)
-	{
-		if (StageJobs[Index])
-		{
-			FShaderCommonCompileJob::FInputHash StageHash = StageJobs[Index]->GetInputHash();
-			Hasher.Update(StageHash.GetBytes(), sizeof(decltype(StageHash.GetBytes())));
-		}
-	}
-
-	InputHash = Hasher.Finalize();
-
-	bInputHashSet = true;
-	return InputHash;
-}
-
-void FShaderPipelineCompileJob::SerializeOutput(FArchive& Ar)
-{
-	bool bAllStagesSucceeded = true;
-	for (int32 Index = 0, Num = StageJobs.Num(); Index < Num; ++Index)
-	{
-		StageJobs[Index]->SerializeOutput(Ar);
-		bAllStagesSucceeded = bAllStagesSucceeded && StageJobs[Index]->bSucceeded;
-	}
-
-	if (Ar.IsLoading())
-	{
-		bFinalized = true;
-		bSucceeded = bAllStagesSucceeded;
-	}
-}
-
 TRACE_DECLARE_INT_COUNTER(Shaders_JobCacheSearchAttempts, TEXT("Shaders/JobCache/SearchAttempts"));
 TRACE_DECLARE_INT_COUNTER(Shaders_JobCacheHits, TEXT("Shaders/JobCache/Hits"));
 
diff --git a/Engine/Source/Runtime/Engine/Private/ShaderCompiler/ShaderPipelineCompiler.cpp b/Engine/Source/Runtime/Engine/Private/ShaderCompiler/ShaderPipelineCompiler.cpp
deleted file mode 100644
diff --git a/Engine/Source/Runtime/Engine/Public/ShaderCompiler.h b/Engine/Source/Runtime/Engine/Public/ShaderCompiler.h
--- a/Engine/Source/Runtime/Engine/Public/ShaderCompiler.h
+++ b/Engine/Source/Runtime/Engine/Public/ShaderCompiler.h
@@ -11,6 +11,7 @@
 #include "HAL/PlatformProcess.h"
 #include "ShaderCore.h"
 #include "ShaderCompilerCore.h"
+#include "ShaderCompilerJobTypes.h"
 #include "Shader.h"
 #include "HAL/RunnableThread.h"
 #include "HAL/Runnable.h"
@@ -29,8 +30,6 @@
 #include "Hash/Blake3.h"
 #include "SceneTypes.h"
 
-class FShaderCompileJob;
-class FShaderPipelineCompileJob;
 class FVertexFactoryType;
 class IDistributedBuildController;
 class FMaterialShaderMap;
@@ -44,265 +43,6 @@ class FShaderPipelineCompileJob;
 
 extern ENGINE_API bool IsShaderJobCacheDDCEnabled();
 
-/** Results for a single compiled shader map. */
-struct FShaderMapCompileResults
-{
-	FShaderMapCompileResults() :
-		bAllJobsSucceeded(true),
-		bSkipResultProcessing(false),
-		TimeStarted(FPlatformTime::Seconds()),
-		bIsHung(false)
-	{}
-
-	void CheckIfHung();
-
-	TArray<TRefCountPtr<class FShaderCommonCompileJob>> FinishedJobs;
-	FThreadSafeCounter NumPendingJobs;
-	bool bAllJobsSucceeded;
-	bool bSkipResultProcessing;
-	double TimeStarted;
-	bool bIsHung;
-};
-
-struct FPendingShaderMapCompileResults
-	: public FShaderMapCompileResults
-	, public FRefCountBase
-{};
-using FPendingShaderMapCompileResultsPtr = TRefCountPtr<FPendingShaderMapCompileResults>;
-
-/** Stores all of the common information used to compile a shader or pipeline. */
-class FShaderCommonCompileJob : public TIntrusiveLinkedList<FShaderCommonCompileJob>
-{
-public:
-	using FInputHash = FBlake3Hash;
-
-	FPendingShaderMapCompileResultsPtr PendingShaderMap;
-
-	mutable FThreadSafeCounter NumRefs;
-	int32 JobIndex;
-	uint32 Hash;
-
-	/** Id of the shader map this shader belongs to. */
-	uint32 Id;
-
-	EShaderCompileJobType Type;
-	EShaderCompileJobPriority Priority;
-	EShaderCompileJobPriority PendingPriority;
-	EShaderCompilerWorkerType CurrentWorker;
-
-	/** true if the results of the shader compile have been processed. */
-	uint8 bFinalized : 1;
-	/** Output of the shader compile */
-	uint8 bSucceeded : 1;
-	uint8 bErrorsAreLikelyToBeCode : 1;
-	/** true if the results of the shader compile have been released from the FShaderCompilerManager.
-		After a job is bFinalized it will be bReleased when ReleaseJob() is invoked, which means that the shader compile thread
-		is no longer processing the job; which is useful for non standard job handling (Niagara as an example). */
-	uint8 bReleased : 1;
-	/** Whether we hashed the inputs */
-	uint8 bInputHashSet : 1;
-	/** Whether or not we are a default material. */
-	uint8 bIsDefaultMaterial : 1;
-	/** Whether or not we are a global shader. */
-	uint8 bIsGlobalShader : 1;
-	/** Hash of all the job inputs */
-	FInputHash InputHash;
-
-	/** In-engine timestamp of being added to a pending queue. Not set for jobs that are satisfied from the jobs cache */
-	double TimeAddedToPendingQueue = 0.0;
-	/** In-engine timestamp of being assigned to a worker. Not set for jobs that are satisfied from the jobs cache */
-	double TimeAssignedToExecution = 0.0;
-	/** In-engine timestamp of job being completed. Encompasses the compile time. Not set for jobs that are satisfied from the jobs cache */
-	double TimeExecutionCompleted = 0.0;
-
-	uint32 AddRef() const
-	{
-		return uint32(NumRefs.Increment());
-	}
-
-	uint32 Release() const
-	{
-		uint32 Refs = uint32(NumRefs.Decrement());
-		if (Refs == 0)
-		{
-			Destroy();
-		}
-		return Refs;
-	}
-	uint32 GetRefCount() const
-	{
-		return uint32(NumRefs.GetValue());
-	}
-
-	/** Returns hash of all inputs for this job (needed for caching). */
-	virtual FInputHash GetInputHash() { return FInputHash(); }
-
-	/** Serializes (and deserializes) the output for caching purposes. */
-	virtual void SerializeOutput(FArchive& Ar) {}
-
-	FShaderCompileJob* GetSingleShaderJob();
-	const FShaderCompileJob* GetSingleShaderJob() const;
-	FShaderPipelineCompileJob* GetShaderPipelineJob();
-	const FShaderPipelineCompileJob* GetShaderPipelineJob() const;
-
-	bool Equals(const FShaderCommonCompileJob& Rhs) const;
-	
-	/** This returns a unique id for a shader compiler job */
-	ENGINE_API static uint32 GetNextJobId();
-
-protected:
-	friend class FShaderCompilingManager;
-	friend class FShaderPipelineCompileJob;
-
-	FShaderCommonCompileJob(EShaderCompileJobType InType, uint32 InHash, uint32 InId, EShaderCompileJobPriority InPriority) :
-		NumRefs(0),
-		JobIndex(INDEX_NONE),
-		Hash(InHash),
-		Id(InId),
-		Type(InType),
-		Priority(InPriority),
-		PendingPriority(EShaderCompileJobPriority::None),
-		CurrentWorker(EShaderCompilerWorkerType::None),
-		bFinalized(false),
-		bSucceeded(false),
-		bErrorsAreLikelyToBeCode(false),
-		bReleased(false),
-		bInputHashSet(false),
-		bIsDefaultMaterial(false),
-		bIsGlobalShader(false)
-
-	{
-		check(InPriority != EShaderCompileJobPriority::None);
-	}
-
-	virtual ~FShaderCommonCompileJob() {}
-
-private:
-	/** Value counter for job ids. */
-	static FThreadSafeCounter JobIdCounter;
-
-	void Destroy() const;
-};
-using FShaderCommonCompileJobPtr = TRefCountPtr<FShaderCommonCompileJob>;
-
-struct FShaderCompileJobKey
-{
-	explicit FShaderCompileJobKey(const FShaderType* InType = nullptr, const FVertexFactoryType* InVFType = nullptr, int32 InPermutationId = 0)
-		: ShaderType(InType), VFType(InVFType), PermutationId(InPermutationId)
-	{}
-
-	uint32 MakeHash(uint32 Id) const { return HashCombine(HashCombine(HashCombine(GetTypeHash(Id), GetTypeHash(VFType)), GetTypeHash(ShaderType)), GetTypeHash(PermutationId)); }
-	ENGINE_API FString ToString() const;
-	const FShaderType* ShaderType;
-	const FVertexFactoryType* VFType;
-	int32 PermutationId;
-
-	friend inline bool operator==(const FShaderCompileJobKey& Lhs, const FShaderCompileJobKey& Rhs)
-	{
-		return Lhs.VFType == Rhs.VFType && Lhs.ShaderType == Rhs.ShaderType && Lhs.PermutationId == Rhs.PermutationId;
-	}
-	friend inline bool operator!=(const FShaderCompileJobKey& Lhs, const FShaderCompileJobKey& Rhs)
-	{
-		return !operator==(Lhs, Rhs);
-	}
-};
-
-/** Stores all of the input and output information used to compile a single shader. */
-class FShaderCompileJob : public FShaderCommonCompileJob
-{
-public:
-	static const EShaderCompileJobType Type = EShaderCompileJobType::Single;
-
-	FShaderCompileJobKey Key;
-
-	/** 
-	 * Additional parameters that can be supplied to the compile job such 
-	 * that it is available from the compilation begins to when the FShader is created.
-	 */
-	TSharedPtr<const FShaderType::FParameters, ESPMode::ThreadSafe> ShaderParameters;
-
-	/** Input for the shader compile */
-	FShaderCompilerInput Input;
-	FShaderCompilerOutput Output;
-
-	// List of pipelines that are sharing this job.
-	TMap<const FVertexFactoryType*, TArray<const FShaderPipelineType*>> SharingPipelines;
-
-	virtual ENGINE_API FInputHash GetInputHash() override;
-	virtual ENGINE_API void SerializeOutput(FArchive& Ar) override;
-
-	FShaderCompileJob(uint32 InHash, uint32 InId, EShaderCompileJobPriority InPriroity, const FShaderCompileJobKey& InKey) :
-		FShaderCommonCompileJob(Type, InHash, InId, InPriroity),
-		Key(InKey)
-	{}
-};
-
-struct FShaderPipelineCompileJobKey
-{
-	explicit FShaderPipelineCompileJobKey(const FShaderPipelineType* InType = nullptr, const FVertexFactoryType* InVFType = nullptr, int32 InPermutationId = 0)
-		: ShaderPipeline(InType), VFType(InVFType), PermutationId(InPermutationId)
-	{}
-
-	uint32 MakeHash(uint32 Id) const { return HashCombine(HashCombine(HashCombine(GetTypeHash(Id), GetTypeHash(ShaderPipeline)), GetTypeHash(VFType)), GetTypeHash(PermutationId)); }
-
-	const FShaderPipelineType* ShaderPipeline;
-	const FVertexFactoryType* VFType;
-	int32 PermutationId;
-
-	friend inline bool operator==(const FShaderPipelineCompileJobKey& Lhs, const FShaderPipelineCompileJobKey& Rhs)
-	{
-		return Lhs.ShaderPipeline == Rhs.ShaderPipeline && Lhs.VFType == Rhs.VFType && Lhs.PermutationId == Rhs.PermutationId;
-	}
-	friend inline bool operator!=(const FShaderPipelineCompileJobKey& Lhs, const FShaderPipelineCompileJobKey& Rhs)
-	{
-		return !operator==(Lhs, Rhs);
-	}
-};
-
-class FShaderPipelineCompileJob : public FShaderCommonCompileJob
-{
-public:
-	static const EShaderCompileJobType Type = EShaderCompileJobType::Pipeline;
-
-	FShaderPipelineCompileJobKey Key;
-	TArray<TRefCountPtr<FShaderCompileJob>> StageJobs;
-	bool bFailedRemovingUnused;
-
-	virtual ENGINE_API FInputHash GetInputHash() override;
-	virtual ENGINE_API void SerializeOutput(FArchive& Ar) override;
-
-	FShaderPipelineCompileJob(uint32 InHash, uint32 InId, EShaderCompileJobPriority InPriroity, const FShaderPipelineCompileJobKey& InKey);
-};
-
-inline FShaderCompileJob* FShaderCommonCompileJob::GetSingleShaderJob() { return Type == EShaderCompileJobType::Single ? static_cast<FShaderCompileJob*>(this) : nullptr; }
-inline const FShaderCompileJob* FShaderCommonCompileJob::GetSingleShaderJob() const { return Type == EShaderCompileJobType::Single ? static_cast<const FShaderCompileJob*>(this) : nullptr; }
-inline FShaderPipelineCompileJob* FShaderCommonCompileJob::GetShaderPipelineJob() { return Type == EShaderCompileJobType::Pipeline ? static_cast<FShaderPipelineCompileJob*>(this) : nullptr; }
-inline const FShaderPipelineCompileJob* FShaderCommonCompileJob::GetShaderPipelineJob() const { return Type == EShaderCompileJobType::Pipeline ? static_cast<const FShaderPipelineCompileJob*>(this) : nullptr; }
-
-inline bool FShaderCommonCompileJob::Equals(const FShaderCommonCompileJob& Rhs) const
-{
-	if (Type == Rhs.Type && Id == Rhs.Id)
-	{
-		switch (Type)
-		{
-		case EShaderCompileJobType::Single: return static_cast<const FShaderCompileJob*>(this)->Key == static_cast<const FShaderCompileJob&>(Rhs).Key;
-		case EShaderCompileJobType::Pipeline: return static_cast<const FShaderPipelineCompileJob*>(this)->Key == static_cast<const FShaderPipelineCompileJob&>(Rhs).Key;
-		default: checkNoEntry(); break;
-		}
-	}
-	return false;
-}
-
-inline void FShaderCommonCompileJob::Destroy() const
-{
-	switch (Type)
-	{
-	case EShaderCompileJobType::Single: delete static_cast<const FShaderCompileJob*>(this); break;
-	case EShaderCompileJobType::Pipeline: delete static_cast<const FShaderPipelineCompileJob*>(this); break;
-	default: checkNoEntry();
-	}
-}
-
 struct FShaderJobCacheStoredOutput;
 
 class FShaderJobCache
diff --git a/Engine/Source/Runtime/RenderCore/Private/ShaderCore.cpp b/Engine/Source/Runtime/RenderCore/Private/ShaderCore.cpp
--- a/Engine/Source/Runtime/RenderCore/Private/ShaderCore.cpp
+++ b/Engine/Source/Runtime/RenderCore/Private/ShaderCore.cpp
@@ -6,6 +6,7 @@
 
 #include "ShaderCore.h"
 #include "Algo/Find.h"
+#include "Async/ParallelFor.h"
 #include "HAL/FileManager.h"
 #include "HAL/PlatformStackWalk.h"
 #include "Misc/FileHelper.h"
@@ -16,8 +17,10 @@
 #include "Misc/ScopeRWLock.h"
 #include "Misc/ScopedSlowTask.h"
 #include "Misc/StringBuilder.h"
+#include "Serialization/MemoryHasher.h"
 #include "Shader.h"
 #include "ShaderCompilerCore.h"
+#include "ShaderCompilerJobTypes.h"
 #include "String/Find.h"
 #include "Tasks/Task.h"
 #include "VertexFactory.h"
@@ -296,6 +299,14 @@ static TAutoConsoleVariable<int32> CVarForceDebugViewModes(
 	TEXT("2: Forces debug view modes to be unavailable, even on editor builds.  Removes many shader permutations for faster shader iteration."),
 	ECVF_RenderThreadSafe | ECVF_ReadOnly);
 
+int32 GShaderCompilerDumpCompileJobInputs = 0;
+static FAutoConsoleVariableRef CVarShaderCompilerDumpCompileJobInputs(
+	TEXT("r.ShaderCompiler.DumpCompileJobInputs"),
+	GShaderCompilerDumpCompileJobInputs,
+	TEXT("if != 0, unpreprocessed input of the shader compiler jobs will be dumped into the debug directory for closer inspection. This is a debugging feature which is disabled by default."),
+	ECVF_Default
+);
+
 /** Returns true if debug viewmodes are allowed for the current platform. */
 bool AllowDebugViewmodes()
 {	
@@ -805,9 +816,9 @@ int HandleShaderCompileException(Windows::LPEXCEPTION_POINTERS Info, FString& Ou
 #endif
 
 #if PLATFORM_WINDOWS
-void CompileShaderInternal(const IShaderFormat* Compiler, const FShaderCompilerInput& Input, FShaderCompilerOutput& Output, const FString& WorkingDirectory, FString& OutExceptionMsg, FString& OutExceptionCallstack, int32* CompileCount)
+void CompileShaderInternal(const IShaderFormat* Compiler, FShaderCompileJob& Job, const FString& WorkingDirectory, FString& OutExceptionMsg, FString& OutExceptionCallstack, int32* CompileCount)
 #else
-void CompileShaderInternal(const IShaderFormat* Compiler, const FShaderCompilerInput& Input, FShaderCompilerOutput& Output, const FString& WorkingDirectory, int32* CompileCount)
+void CompileShaderInternal(const IShaderFormat* Compiler, FShaderCompileJob& Job, const FString& WorkingDirectory, int32* CompileCount)
 #endif
 {
 #if PLATFORM_WINDOWS
@@ -816,20 +827,20 @@ void CompileShaderInternal(const IShaderFormat* Compiler, const FShaderCompilerI
 	{
 		double TimeStart = FPlatformTime::Seconds();
 
-		Compiler->CompileShader(Input.ShaderFormat, Input, Output, WorkingDirectory);
-		if (Output.bSucceeded)
+		Compiler->CompileShader(Job.Input.ShaderFormat, Job.Input, Job.Output, WorkingDirectory);
+		if (Job.Output.bSucceeded)
 		{
-			Output.GenerateOutputHash();
-			if (Input.CompressionFormat != NAME_None)
+			Job.Output.GenerateOutputHash();
+			if (Job.Input.CompressionFormat != NAME_None)
 			{
-				Output.CompressOutput(Input.CompressionFormat, Input.OodleCompressor, Input.OodleLevel);
+				Job.Output.CompressOutput(Job.Input.CompressionFormat, Job.Input.OodleCompressor, Job.Input.OodleLevel);
 			}
 		}
-		Output.CompileTime = FPlatformTime::Seconds() - TimeStart;
+		Job.Output.CompileTime = FPlatformTime::Seconds() - TimeStart;
 
-		if (Compiler->UsesHLSLcc(Input))
+		if (Compiler->UsesHLSLcc(Job.Input))
 		{
-			Output.bUsedHLSLccCompiler = true;
+			Job.Output.bUsedHLSLccCompiler = true;
 		}
 
 		if (CompileCount)
@@ -840,44 +851,118 @@ void CompileShaderInternal(const IShaderFormat* Compiler, const FShaderCompilerI
 #if PLATFORM_WINDOWS
 	__except (HandleShaderCompileException(GetExceptionInformation(), OutExceptionMsg, OutExceptionCallstack))
 	{
-		Output.bSucceeded = false;
+		Job.Output.bSucceeded = false;
 	}
 #endif
 }
 
 void CompileShader(const TArray<const IShaderFormat*>& ShaderFormats, FShaderCompilerInput& Input, FShaderCompilerOutput& Output, const FString& WorkingDirectory, int32* CompileCount)
+{
+	FShaderCompileJob TempJob;
+	TempJob.Input = Input;
+	
+	CompileShader(ShaderFormats, TempJob, WorkingDirectory, CompileCount);
+	Output = TempJob.Output;
+}
+
+void CompileShader(const TArray<const IShaderFormat*>& ShaderFormats, FShaderCompileJob& Job, const FString& WorkingDirectory, int32* CompileCount)
 {
 	TRACE_CPUPROFILER_EVENT_SCOPE(CompileShader);
 
-	const IShaderFormat* Compiler = FindShaderFormat(Input.ShaderFormat, ShaderFormats);
+	const IShaderFormat* Compiler = FindShaderFormat(Job.Input.ShaderFormat, ShaderFormats);
 	if (!Compiler)
 	{
-		UE_LOG(LogShaders, Fatal, TEXT("Can't compile shaders for format %s, couldn't load compiler dll"), *Input.ShaderFormat.ToString());
+		UE_LOG(LogShaders, Fatal, TEXT("Can't compile shaders for format %s, couldn't load compiler dll"), *Job.Input.ShaderFormat.ToString());
 	}
 
-	if (IsValidRef(Input.SharedEnvironment))
+	if (IsValidRef(Job.Input.SharedEnvironment))
 	{
-		Input.Environment.Merge(*Input.SharedEnvironment);
+		Job.Input.Environment.Merge(*Job.Input.SharedEnvironment);
 	}
 
 #if PLATFORM_WINDOWS
 	FString ExceptionMsg;
 	FString ExceptionCallstack;
-	CompileShaderInternal(Compiler, Input, Output, WorkingDirectory, ExceptionMsg, ExceptionCallstack, CompileCount);
-	if (!Output.bSucceeded && (!ExceptionMsg.IsEmpty() || !ExceptionCallstack.IsEmpty()))
+	CompileShaderInternal(Compiler, Job, WorkingDirectory, ExceptionMsg, ExceptionCallstack, CompileCount);
+	if (!Job.Output.bSucceeded && (!ExceptionMsg.IsEmpty() || !ExceptionCallstack.IsEmpty()))
 	{
 		FShaderCompilerError Error;
 		Error.StrippedErrorMessage = FString::Printf(
 			TEXT("Exception encountered in platform compiler: %s\nException Callstack:\n%s"), 
 			*ExceptionMsg, 
 			*ExceptionCallstack);
-		Output.Errors.Add(Error);
+		Job.Output.Errors.Add(Error);
 	}
 #else
-	CompileShaderInternal(Compiler, Input, Output, WorkingDirectory, CompileCount);
+	CompileShaderInternal(Compiler, Job, WorkingDirectory, CompileCount);
 #endif
 }
 
+void CompileShaderPipeline(const TArray<const IShaderFormat*>& ShaderFormats, FShaderPipelineCompileJob* PipelineJob, const FString& Dir, int32* CompileCount)
+{
+	checkf(PipelineJob->StageJobs.Num() > 0, TEXT("Pipeline %s has zero jobs!"), PipelineJob->Key.ShaderPipeline->GetName());
+	FShaderCompileJob* CurrentJob = PipelineJob->StageJobs[0]->GetSingleShaderJob();
+
+	CurrentJob->Input.bCompilingForShaderPipeline = true;
+
+	// First job doesn't have to trim outputs
+	CurrentJob->Input.bIncludeUsedOutputs = false;
+
+	// Compile the shader directly through the platform dll (directly from the shader dir as the working directory)
+	CompileShader(ShaderFormats, *CurrentJob, Dir, CompileCount);
+
+	CurrentJob->bSucceeded = CurrentJob->Output.bSucceeded;
+	if (!CurrentJob->Output.bSucceeded)
+	{
+		// Can't carry on compiling the pipeline
+		return;
+	}
+
+	// This tells the shader compiler we do want to remove unused outputs
+	bool bEnableRemovingUnused = true;
+	// This tells us the hlsl parser failed at removing unused outputs
+	bool bJobFailedRemovingUnused = false;
+
+	//#todo-rco: Currently only removes for pure VS & PS stages
+	for (int32 Index = 0; Index < PipelineJob->StageJobs.Num(); ++Index)
+	{
+		auto Stage = PipelineJob->StageJobs[Index]->GetSingleShaderJob()->Input.Target.Frequency;
+		if (Stage != SF_Vertex && Stage != SF_Pixel)
+		{
+			bEnableRemovingUnused = false;
+			break;
+		}
+	}
+
+	for (int32 Index = 1; Index < PipelineJob->StageJobs.Num(); ++Index)
+	{
+		auto* PreviousJob = CurrentJob;
+		CurrentJob = PipelineJob->StageJobs[Index]->GetSingleShaderJob();
+		bEnableRemovingUnused = bEnableRemovingUnused && PreviousJob->Output.bSupportsQueryingUsedAttributes;
+		if (bEnableRemovingUnused)
+		{
+			CurrentJob->Input.bIncludeUsedOutputs = true;
+			CurrentJob->Input.bCompilingForShaderPipeline = true;
+			CurrentJob->Input.UsedOutputs = PreviousJob->Output.UsedAttributes;
+		}
+
+		// Compile the shader directly through the platform dll (directly from the shader dir as the working directory)
+		CompileShader(ShaderFormats, *CurrentJob, Dir, CompileCount);
+
+		CurrentJob->bSucceeded = CurrentJob->Output.bSucceeded;
+		if (!CurrentJob->Output.bSucceeded)
+		{
+			// Can't carry on compiling the pipeline
+			return;
+		}
+
+		bJobFailedRemovingUnused = CurrentJob->Output.bFailedRemovingUnused || bJobFailedRemovingUnused;
+	}
+
+	PipelineJob->bSucceeded = true;
+	PipelineJob->bFailedRemovingUnused = bJobFailedRemovingUnused;
+}
+
 /**
 * Add a new entry to the list of shader source files
 * Only unique entries which can be loaded are added as well as their #include files
@@ -2408,3 +2493,249 @@ FArchive& operator<<(FArchive& Ar, FShaderCompilerInput& Input)
 
 	return Ar;
 }
+
+FShaderCommonCompileJob::FInputHash FShaderPipelineCompileJob::GetInputHash()
+{
+	if (bInputHashSet)
+	{
+		return InputHash;
+	}
+
+	FBlake3 Hasher;
+	for (int32 Index = 0; Index < StageJobs.Num(); ++Index)
+	{
+		if (StageJobs[Index])
+		{
+			FShaderCommonCompileJob::FInputHash StageHash = StageJobs[Index]->GetInputHash();
+			Hasher.Update(StageHash.GetBytes(), sizeof(decltype(StageHash.GetBytes())));
+		}
+	}
+
+	InputHash = Hasher.Finalize();
+
+	bInputHashSet = true;
+	return InputHash;
+}
+
+// The Id of 0 is reserved for global shaders
+FThreadSafeCounter FShaderCommonCompileJob::JobIdCounter(2);
+
+uint32 FShaderCommonCompileJob::GetNextJobId()
+{
+	uint32 Id = JobIdCounter.Increment();
+	if (Id == UINT_MAX)
+	{
+		JobIdCounter.Set(2);
+	}
+	return Id;
+}
+
+FString FShaderCompileJobKey::ToString() const
+{
+	return FString::Printf(
+		TEXT("ShaderType:%s VertexFactoryType:%s PermutationId:%d"),
+		ShaderType ? ShaderType->GetName() : TEXT("None"),
+		VFType ? VFType->GetName() : TEXT("None"),
+		PermutationId);
+}
+
+FShaderCommonCompileJob::FInputHash FShaderCompileJob::GetInputHash()
+{
+	if (bInputHashSet)
+	{
+		return InputHash;
+	}
+
+	auto SerializeInputs = [this](FArchive& Archive)
+	{
+		checkf(Archive.IsSaving() && !Archive.IsLoading(), TEXT("A loading archive is passed to FShaderCompileJob::GetInputHash(), this is not supported as it may corrupt its data"));
+
+		Archive << Input;
+		Input.Environment.SerializeEverythingButFiles(Archive);
+
+		// hash the source file so changes to files during the development are picked up
+		const FSHAHash& SourceHash = GetShaderFileHash(*Input.VirtualSourceFilePath, Input.Target.GetPlatform());
+		Archive << const_cast<FSHAHash&>(SourceHash);
+
+		// unroll the included files for the parallel processing.
+		// These are temporary arrays that only exist for the ParallelFor
+		TArray<const TCHAR*> IncludeVirtualPaths;
+		TArray<const FString*> Contents;
+		TArray<bool> OnlyHashIncludes;
+		TArray<FBlake3Hash> Hashes;
+
+		// while the contents of this is already hashed (included in Environment's operator<<()), we still need to account for includes in the generated files and hash them, too
+		for (TMap<FString, FString>::TConstIterator It(Input.Environment.IncludeVirtualPathToContentsMap); It; ++It)
+		{
+			const FString& VirtualPath = It.Key();
+			IncludeVirtualPaths.Add(*VirtualPath);
+			Contents.Add(&It.Value());
+			OnlyHashIncludes.Add(true);	// not hashing contents of the file itself, as it was included in Environment's operator<<()
+			Hashes.AddDefaulted();
+		}
+
+		for (TMap<FString, FThreadSafeSharedStringPtr>::TConstIterator It(Input.Environment.IncludeVirtualPathToExternalContentsMap); It; ++It)
+		{
+			const FString& VirtualPath = It.Key();
+			IncludeVirtualPaths.Add(*VirtualPath);
+			check(It.Value());
+			Contents.Add(&(*It.Value()));
+			OnlyHashIncludes.Add(false);
+			Hashes.AddDefaulted();
+		}
+
+		if (Input.SharedEnvironment)
+		{
+			Input.SharedEnvironment->SerializeEverythingButFiles(Archive);
+
+			for (TMap<FString, FString>::TConstIterator It(Input.SharedEnvironment->IncludeVirtualPathToContentsMap); It; ++It)
+			{
+				const FString& VirtualPath = It.Key();
+				IncludeVirtualPaths.Add(*VirtualPath);
+				Contents.Add(&It.Value());
+				OnlyHashIncludes.Add(true);	// not hashing contents of the file itself, as it was included in Environment's operator<<()
+				Hashes.AddDefaulted();
+			}
+
+			for (TMap<FString, FThreadSafeSharedStringPtr>::TConstIterator It(Input.SharedEnvironment->IncludeVirtualPathToExternalContentsMap); It; ++It)
+			{
+				const FString& VirtualPath = It.Key();
+				IncludeVirtualPaths.Add(*VirtualPath);
+				check(It.Value());
+				Contents.Add(&(*It.Value()));
+				OnlyHashIncludes.Add(false);
+				Hashes.AddDefaulted();
+			}
+		}
+
+		check(IncludeVirtualPaths.Num() == Contents.Num());
+		check(Contents.Num() == OnlyHashIncludes.Num());
+		check(OnlyHashIncludes.Num() == Hashes.Num());
+
+		EShaderPlatform Platform = Input.Target.GetPlatform();
+		ParallelFor(Contents.Num(), [&IncludeVirtualPaths, &Contents, &OnlyHashIncludes, &Hashes, &Platform](int32 FileIndex)
+			{ 
+				FMemoryHasherBlake3 MemHasher;
+				HashShaderFileWithIncludes(MemHasher, IncludeVirtualPaths[FileIndex], *Contents[FileIndex], Platform, OnlyHashIncludes[FileIndex]);
+				Hashes[FileIndex] = MemHasher.Finalize();
+			},
+			EParallelForFlags::Unbalanced
+		);
+
+		// include the hashes in the main hash (consider sorting them if includes are found to have a random order)
+		for (int32 HashIndex = 0, NumHashes = Hashes.Num(); HashIndex < NumHashes; ++HashIndex)
+		{
+			Archive << Hashes[HashIndex];
+		}
+	};
+
+	// use faster hasher that doesn't allocate memory
+	FMemoryHasherBlake3 MemHasher;
+	SerializeInputs(MemHasher);
+	InputHash = MemHasher.Finalize();
+
+	if (GShaderCompilerDumpCompileJobInputs)
+	{
+		TArray<uint8> MemoryBlob;
+		FMemoryWriter MemWriter(MemoryBlob);
+
+		SerializeInputs(MemWriter);
+
+		FString IntermediateFormatPath = FPaths::ProjectSavedDir() / TEXT("ShaderJobInputs");
+#if UE_BUILD_DEBUG
+		FString TempPath = IntermediateFormatPath / TEXT("DebugEditor");
+#else
+		FString TempPath = IntermediateFormatPath / TEXT("DevelopmentEditor");
+#endif
+		IFileManager::Get().MakeDirectory(*TempPath, true);
+
+		static int32 InputHashID = 0;
+		FString FileName = Input.DebugGroupName.Replace(TEXT("/"), TEXT("_")).Replace(TEXT("<"), TEXT("_")).Replace(TEXT(">"), TEXT("_")).Replace(TEXT(":"), TEXT("_")).Replace(TEXT("|"), TEXT("_"))
+			+ TEXT("-") + Input.EntryPointName;
+		FString TempFile = TempPath / FString::Printf(TEXT("%s-%d.bin"), *FileName, InputHashID++);
+
+		TUniquePtr<FArchive> DumpAr(IFileManager::Get().CreateFileWriter(*TempFile));
+		DumpAr->Serialize(MemoryBlob.GetData(), MemoryBlob.Num());
+
+		// as an additional debugging feature, make sure that the hash is the same as calculated by the memhasher
+		FBlake3Hash Check = FBlake3::HashBuffer(MemoryBlob.GetData(), MemoryBlob.Num());
+		if (Check != InputHash)
+		{
+			UE_LOG(LogShaders, Error, TEXT("Job input hash disagrees between FMemoryHasherSHA1 (%s) and FMemoryWriter + FSHA1 (%s, which was dumped to disk)"), *LexToString(InputHash), *LexToString(Check));
+		}
+	}
+
+	bInputHashSet = true;
+	return InputHash;
+}
+
+void FShaderCompileJob::SerializeOutput(FArchive& Ar)
+{
+	double ActualCompileTime = 0.0;
+	double ActualPreprocessTime = 0.0;
+	if (Ar.IsSaving())
+	{
+		// Cached jobs won't have accurate results anyway, so reduce the storage requirements by setting those fields to a known value.
+		// This significantly reduces the memory needed to store the outputs (by more than a half)
+		ActualCompileTime = Output.CompileTime;
+		ActualPreprocessTime = Output.PreprocessTime;
+		Output.CompileTime = 0.0;
+		Output.PreprocessTime = 0.0;
+	}
+
+	Ar << Output;
+	// output hash is now serialized as part of the output, as the shader code is compressed in SCWs
+	checkf(!Output.bSucceeded || Output.OutputHash != FSHAHash(), TEXT("Successful compile job does not have an OutputHash generated."));
+
+	if (Ar.IsLoading())
+	{
+		bFinalized = true;
+		bSucceeded = Output.bSucceeded;
+	}
+	else
+	{
+		// restore the compile time for this jobs. Jobs that will be deserialized from the cache will have a compile time of 0.0
+		Output.CompileTime = ActualCompileTime;
+		Output.PreprocessTime = ActualPreprocessTime;
+	}
+}
+
+FShaderPipelineCompileJob::FShaderPipelineCompileJob(int32 NumStages)
+	: FShaderCommonCompileJob(Type, 0u, 0u, EShaderCompileJobPriority::Num)
+{
+	StageJobs.Empty(NumStages);
+	for (int32 StageIndex = 0; StageIndex < NumStages; ++StageIndex)
+	{
+		StageJobs.Add(new FShaderCompileJob());
+	}
+}
+
+FShaderPipelineCompileJob::FShaderPipelineCompileJob(uint32 InHash, uint32 InId, EShaderCompileJobPriority InPriroity, const FShaderPipelineCompileJobKey& InKey) :
+	FShaderCommonCompileJob(Type, InHash, InId, InPriroity),
+	Key(InKey)
+{
+	const auto& Stages = InKey.ShaderPipeline->GetStages();
+	StageJobs.Empty(Stages.Num());
+	for (const FShaderType* ShaderType : Stages)
+	{
+		const FShaderCompileJobKey StageKey(ShaderType, InKey.VFType, InKey.PermutationId);
+		StageJobs.Add(new FShaderCompileJob(StageKey.MakeHash(InId), InId, InPriroity, StageKey));
+	}
+}
+
+void FShaderPipelineCompileJob::SerializeOutput(FArchive& Ar)
+{
+	bool bAllStagesSucceeded = true;
+	for (int32 Index = 0, Num = StageJobs.Num(); Index < Num; ++Index)
+	{
+		StageJobs[Index]->SerializeOutput(Ar);
+		bAllStagesSucceeded = bAllStagesSucceeded && StageJobs[Index]->bSucceeded;
+	}
+
+	if (Ar.IsLoading())
+	{
+		bFinalized = true;
+		bSucceeded = bAllStagesSucceeded;
+	}
+}
+
diff --git a/Engine/Source/Runtime/RenderCore/Public/ShaderCompilerCore.h b/Engine/Source/Runtime/RenderCore/Public/ShaderCompilerCore.h
--- a/Engine/Source/Runtime/RenderCore/Public/ShaderCompilerCore.h
+++ b/Engine/Source/Runtime/RenderCore/Public/ShaderCompilerCore.h
@@ -17,11 +17,13 @@
 
 class Error;
 class IShaderFormat;
+class FShaderCompileJob;
+class FShaderPipelineCompileJob;
 
 // this is for the protocol, not the data, bump if FShaderCompilerInput or ProcessInputFromArchive changes.
 inline const int32 ShaderCompileWorkerInputVersion = 18;
 // this is for the protocol, not the data, bump if FShaderCompilerOutput or WriteToOutputArchive changes.
-inline const int32 ShaderCompileWorkerOutputVersion = 10;
+inline const int32 ShaderCompileWorkerOutputVersion = 11;
 // this is for the protocol, not the data.
 inline const int32 ShaderCompileWorkerSingleJobHeader = 'S';
 // this is for the protocol, not the data.
@@ -539,7 +541,10 @@ struct RENDERCORE_API FSCWErrorCode
 extern RENDERCORE_API int HandleShaderCompileException(Windows::LPEXCEPTION_POINTERS Info, FString& OutExMsg, FString& OutCallStack);
 #endif
 extern RENDERCORE_API const IShaderFormat* FindShaderFormat(FName Format, TArray<const IShaderFormat*> ShaderFormats);
+UE_DEPRECATED(5.3, "Use CompileShader overload which takes an FShaderCompileJob& rather than passing input/output directly.")
 extern RENDERCORE_API void CompileShader(const TArray<const IShaderFormat*>& ShaderFormats, FShaderCompilerInput& Input, FShaderCompilerOutput& Output, const FString& WorkingDirectory, int32* CompileCount = nullptr);
+extern RENDERCORE_API void CompileShader(const TArray<const IShaderFormat*>& ShaderFormats, FShaderCompileJob& Job, const FString& WorkingDirectory, int32* CompileCount = nullptr);
+extern RENDERCORE_API void CompileShaderPipeline(const TArray<const IShaderFormat*>& ShaderFormats, FShaderPipelineCompileJob* PipelineJob, const FString& WorkingDirectory, int32* CompileCount = nullptr);
 
 UE_DEPRECATED(5.2, "Functionality has moved to UE::ShaderCompilerCommon::ShouldUseStableConstantBuffer")
 inline bool ShouldUseStableConstantBuffer(const FShaderCompilerInput& Input)
@@ -609,3 +614,5 @@ inline const TCHAR* ShaderCompileJobPriorityToString(EShaderCompileJobPriority v
 	default: checkNoEntry(); return TEXT("");
 	}
 }
+
+
diff --git a/Engine/Source/Runtime/RenderCore/Public/ShaderCompilerJobTypes.h b/Engine/Source/Runtime/RenderCore/Public/ShaderCompilerJobTypes.h
new file mode 100644
--- /dev/null
+++ b/Engine/Source/Runtime/RenderCore/Public/ShaderCompilerJobTypes.h
@@ -0,0 +1,274 @@
+// Copyright Epic Games, Inc. All Rights Reserved.
+
+#include "HAL/PlatformTime.h"
+#include "Hash/Blake3.h"
+#include "Shader.h"
+#include "VertexFactory.h"
+
+class FShaderCommonCompileJob;
+class FShaderCompileJob;
+class FShaderPipelineCompileJob;
+
+/** Results for a single compiled shader map. */
+struct FShaderMapCompileResults
+{
+	FShaderMapCompileResults() :
+		bAllJobsSucceeded(true),
+		bSkipResultProcessing(false),
+		TimeStarted(FPlatformTime::Seconds()),
+		bIsHung(false)
+	{}
+
+	void CheckIfHung();
+
+	TArray<TRefCountPtr<class FShaderCommonCompileJob>> FinishedJobs;
+	FThreadSafeCounter NumPendingJobs;
+	bool bAllJobsSucceeded;
+	bool bSkipResultProcessing;
+	double TimeStarted;
+	bool bIsHung;
+};
+
+struct FPendingShaderMapCompileResults
+	: public FShaderMapCompileResults
+	, public FRefCountBase
+{};
+using FPendingShaderMapCompileResultsPtr = TRefCountPtr<FPendingShaderMapCompileResults>;
+
+
+/** Stores all of the common information used to compile a shader or pipeline. */
+class FShaderCommonCompileJob : public TIntrusiveLinkedList<FShaderCommonCompileJob>
+{
+public:
+	using FInputHash = FBlake3Hash;
+
+	FPendingShaderMapCompileResultsPtr PendingShaderMap;
+
+	mutable FThreadSafeCounter NumRefs;
+	int32 JobIndex;
+	uint32 Hash;
+
+	/** Id of the shader map this shader belongs to. */
+	uint32 Id;
+
+	EShaderCompileJobType Type;
+	EShaderCompileJobPriority Priority;
+	EShaderCompileJobPriority PendingPriority;
+	EShaderCompilerWorkerType CurrentWorker;
+
+	/** true if the results of the shader compile have been processed. */
+	uint8 bFinalized : 1;
+	/** Output of the shader compile */
+	uint8 bSucceeded : 1;
+	uint8 bErrorsAreLikelyToBeCode : 1;
+	/** true if the results of the shader compile have been released from the FShaderCompilerManager.
+		After a job is bFinalized it will be bReleased when ReleaseJob() is invoked, which means that the shader compile thread
+		is no longer processing the job; which is useful for non standard job handling (Niagara as an example). */
+	uint8 bReleased : 1;
+	/** Whether we hashed the inputs */
+	uint8 bInputHashSet : 1;
+	/** Whether or not we are a default material. */
+	uint8 bIsDefaultMaterial : 1;
+	/** Whether or not we are a global shader. */
+	uint8 bIsGlobalShader : 1;
+	/** Hash of all the job inputs */
+	FInputHash InputHash;
+
+	/** In-engine timestamp of being added to a pending queue. Not set for jobs that are satisfied from the jobs cache */
+	double TimeAddedToPendingQueue = 0.0;
+	/** In-engine timestamp of being assigned to a worker. Not set for jobs that are satisfied from the jobs cache */
+	double TimeAssignedToExecution = 0.0;
+	/** In-engine timestamp of job being completed. Encompasses the compile time. Not set for jobs that are satisfied from the jobs cache */
+	double TimeExecutionCompleted = 0.0;
+
+	uint32 AddRef() const
+	{
+		return uint32(NumRefs.Increment());
+	}
+
+	uint32 Release() const
+	{
+		uint32 Refs = uint32(NumRefs.Decrement());
+		if (Refs == 0)
+		{
+			Destroy();
+		}
+		return Refs;
+	}
+	uint32 GetRefCount() const
+	{
+		return uint32(NumRefs.GetValue());
+	}
+
+	/** Returns hash of all inputs for this job (needed for caching). */
+	virtual FInputHash GetInputHash() { return FInputHash(); }
+
+	/** Serializes (and deserializes) the output for caching purposes. */
+	virtual void SerializeOutput(FArchive& Ar) {}
+
+	FShaderCompileJob* GetSingleShaderJob();
+	const FShaderCompileJob* GetSingleShaderJob() const;
+	FShaderPipelineCompileJob* GetShaderPipelineJob();
+	const FShaderPipelineCompileJob* GetShaderPipelineJob() const;
+
+	bool Equals(const FShaderCommonCompileJob& Rhs) const;
+	
+	/** This returns a unique id for a shader compiler job */
+	RENDERCORE_API static uint32 GetNextJobId();
+
+protected:
+	friend class FShaderCompilingManager;
+	friend class FShaderPipelineCompileJob;
+
+	FShaderCommonCompileJob(EShaderCompileJobType InType, uint32 InHash, uint32 InId, EShaderCompileJobPriority InPriority) :
+		NumRefs(0),
+		JobIndex(INDEX_NONE),
+		Hash(InHash),
+		Id(InId),
+		Type(InType),
+		Priority(InPriority),
+		PendingPriority(EShaderCompileJobPriority::None),
+		CurrentWorker(EShaderCompilerWorkerType::None),
+		bFinalized(false),
+		bSucceeded(false),
+		bErrorsAreLikelyToBeCode(false),
+		bReleased(false),
+		bInputHashSet(false),
+		bIsDefaultMaterial(false),
+		bIsGlobalShader(false)
+
+	{
+		check(InPriority != EShaderCompileJobPriority::None);
+	}
+
+	virtual ~FShaderCommonCompileJob() {}
+
+private:
+	/** Value counter for job ids. */
+	static FThreadSafeCounter JobIdCounter;
+
+	void Destroy() const;
+};
+using FShaderCommonCompileJobPtr = TRefCountPtr<FShaderCommonCompileJob>;
+
+struct FShaderCompileJobKey
+{
+	explicit FShaderCompileJobKey(const FShaderType* InType = nullptr, const FVertexFactoryType* InVFType = nullptr, int32 InPermutationId = 0)
+		: ShaderType(InType), VFType(InVFType), PermutationId(InPermutationId)
+	{}
+
+	uint32 MakeHash(uint32 Id) const { return HashCombine(HashCombine(HashCombine(GetTypeHash(Id), GetTypeHash(VFType)), GetTypeHash(ShaderType)), GetTypeHash(PermutationId)); }
+	RENDERCORE_API FString ToString() const;
+	const FShaderType* ShaderType;
+	const FVertexFactoryType* VFType;
+	int32 PermutationId;
+
+	friend inline bool operator==(const FShaderCompileJobKey& Lhs, const FShaderCompileJobKey& Rhs)
+	{
+		return Lhs.VFType == Rhs.VFType && Lhs.ShaderType == Rhs.ShaderType && Lhs.PermutationId == Rhs.PermutationId;
+	}
+	friend inline bool operator!=(const FShaderCompileJobKey& Lhs, const FShaderCompileJobKey& Rhs)
+	{
+		return !operator==(Lhs, Rhs);
+	}
+};
+
+/** Stores all of the input and output information used to compile a single shader. */
+class FShaderCompileJob : public FShaderCommonCompileJob
+{
+public:
+	static const EShaderCompileJobType Type = EShaderCompileJobType::Single;
+
+	FShaderCompileJobKey Key;
+
+	/** 
+	 * Additional parameters that can be supplied to the compile job such 
+	 * that it is available from the compilation begins to when the FShader is created.
+	 */
+	TSharedPtr<const FShaderType::FParameters, ESPMode::ThreadSafe> ShaderParameters;
+
+	/** Input for the shader compile */
+	FShaderCompilerInput Input;
+	FShaderCompilerOutput Output;
+
+	// List of pipelines that are sharing this job.
+	TMap<const FVertexFactoryType*, TArray<const FShaderPipelineType*>> SharingPipelines;
+
+	virtual RENDERCORE_API FInputHash GetInputHash() override;
+	virtual RENDERCORE_API void SerializeOutput(FArchive& Ar) override;
+
+	RENDERCORE_API FShaderCompileJob() : FShaderCommonCompileJob(Type, 0u, 0u, EShaderCompileJobPriority::Num)
+	{}
+
+	RENDERCORE_API FShaderCompileJob(uint32 InHash, uint32 InId, EShaderCompileJobPriority InPriroity, const FShaderCompileJobKey& InKey) :
+		FShaderCommonCompileJob(Type, InHash, InId, InPriroity),
+		Key(InKey)
+	{}
+};
+
+struct FShaderPipelineCompileJobKey
+{
+	explicit FShaderPipelineCompileJobKey(const FShaderPipelineType* InType = nullptr, const FVertexFactoryType* InVFType = nullptr, int32 InPermutationId = 0)
+		: ShaderPipeline(InType), VFType(InVFType), PermutationId(InPermutationId)
+	{}
+
+	uint32 MakeHash(uint32 Id) const { return HashCombine(HashCombine(HashCombine(GetTypeHash(Id), GetTypeHash(ShaderPipeline)), GetTypeHash(VFType)), GetTypeHash(PermutationId)); }
+
+	const FShaderPipelineType* ShaderPipeline;
+	const FVertexFactoryType* VFType;
+	int32 PermutationId;
+
+	friend inline bool operator==(const FShaderPipelineCompileJobKey& Lhs, const FShaderPipelineCompileJobKey& Rhs)
+	{
+		return Lhs.ShaderPipeline == Rhs.ShaderPipeline && Lhs.VFType == Rhs.VFType && Lhs.PermutationId == Rhs.PermutationId;
+	}
+	friend inline bool operator!=(const FShaderPipelineCompileJobKey& Lhs, const FShaderPipelineCompileJobKey& Rhs)
+	{
+		return !operator==(Lhs, Rhs);
+	}
+};
+
+class FShaderPipelineCompileJob : public FShaderCommonCompileJob
+{
+public:
+	static const EShaderCompileJobType Type = EShaderCompileJobType::Pipeline;
+
+	FShaderPipelineCompileJobKey Key;
+	TArray<TRefCountPtr<FShaderCompileJob>> StageJobs;
+	bool bFailedRemovingUnused;
+
+	virtual RENDERCORE_API FInputHash GetInputHash() override;
+	virtual RENDERCORE_API void SerializeOutput(FArchive& Ar) override;
+
+	RENDERCORE_API FShaderPipelineCompileJob(int32 NumStages);
+	RENDERCORE_API FShaderPipelineCompileJob(uint32 InHash, uint32 InId, EShaderCompileJobPriority InPriroity, const FShaderPipelineCompileJobKey& InKey);
+};
+
+inline FShaderCompileJob* FShaderCommonCompileJob::GetSingleShaderJob() { return Type == EShaderCompileJobType::Single ? static_cast<FShaderCompileJob*>(this) : nullptr; }
+inline const FShaderCompileJob* FShaderCommonCompileJob::GetSingleShaderJob() const { return Type == EShaderCompileJobType::Single ? static_cast<const FShaderCompileJob*>(this) : nullptr; }
+inline FShaderPipelineCompileJob* FShaderCommonCompileJob::GetShaderPipelineJob() { return Type == EShaderCompileJobType::Pipeline ? static_cast<FShaderPipelineCompileJob*>(this) : nullptr; }
+inline const FShaderPipelineCompileJob* FShaderCommonCompileJob::GetShaderPipelineJob() const { return Type == EShaderCompileJobType::Pipeline ? static_cast<const FShaderPipelineCompileJob*>(this) : nullptr; }
+
+inline bool FShaderCommonCompileJob::Equals(const FShaderCommonCompileJob& Rhs) const
+{
+	if (Type == Rhs.Type && Id == Rhs.Id)
+	{
+		switch (Type)
+		{
+		case EShaderCompileJobType::Single: return static_cast<const FShaderCompileJob*>(this)->Key == static_cast<const FShaderCompileJob&>(Rhs).Key;
+		case EShaderCompileJobType::Pipeline: return static_cast<const FShaderPipelineCompileJob*>(this)->Key == static_cast<const FShaderPipelineCompileJob&>(Rhs).Key;
+		default: checkNoEntry(); break;
+		}
+	}
+	return false;
+}
+
+inline void FShaderCommonCompileJob::Destroy() const
+{
+	switch (Type)
+	{
+	case EShaderCompileJobType::Single: delete static_cast<const FShaderCompileJob*>(this); break;
+	case EShaderCompileJobType::Pipeline: delete static_cast<const FShaderPipelineCompileJob*>(this); break;
+	default: checkNoEntry();
+	}
+}
\ No newline at end of file
