Allow decoupling of UB parameter struct declaration and definition.
UB parameter structs can be forward declared, and so can be used by third parties without making the full definition public.

diff --git a/Engine/Source/Runtime/RenderCore/Public/ShaderParameterFwd.h b/Engine/Source/Runtime/RenderCore/Public/ShaderParameterFwd.h
new file mode 100644
--- /dev/null
+++ b/Engine/Source/Runtime/RenderCore/Public/ShaderParameterFwd.h
@@ -0,0 +1,19 @@
+#pragma once
+
+class FShaderParametersMetadata;
+
+#define DECLARE_UNIFORM_BUFFER_STRUCT(StructTypeName, PrefixKeywords) \
+	class StructTypeName; \
+	PrefixKeywords const FShaderParametersMetadata* GetForwardDeclaredShaderParametersStructMetadata(const StructTypeName* DummyPtr);
+
+
+/** Retrieves the metadata for a shader parameter struct type without having a definition of the type.
+ *  This is the default overload in case there is no exact non-template overload (prevents implicit casting).
+ */
+template<typename T>
+const FShaderParametersMetadata* GetForwardDeclaredShaderParametersStructMetadata(const T*)
+{
+	static_assert(!std::is_same<T, T>::value /* true */,
+		TEXT("Partial uniform buffer struct declaration. Use `DECLARE_UNIFORM_BUFFER_STRUCT()` instead of `class T;`.")
+	);
+}
diff --git a/Engine/Source/Runtime/RenderCore/Public/ShaderParameterMacros.h b/Engine/Source/Runtime/RenderCore/Public/ShaderParameterMacros.h
--- a/Engine/Source/Runtime/RenderCore/Public/ShaderParameterMacros.h
+++ b/Engine/Source/Runtime/RenderCore/Public/ShaderParameterMacros.h
@@ -7,6 +7,7 @@
 
 #pragma once
 
+#include "ShaderParameterFwd.h"
 #include "ShaderParameterMetadata.h"
 #include "RenderGraphAllocator.h"
 #include "Algo/Reverse.h"
@@ -1183,6 +1184,25 @@ struct TShaderResourceParameterTypeInfo<ShaderResourceType[InNumElements]>
 	static const FShaderParametersMetadata* GetStructMetadata() { return nullptr; }
 };
 
+template<class UniformBufferStructType, typename = void>
+struct TUniformBufferMetadataHelper
+{
+	static const FShaderParametersMetadata* GetStructMetadata()
+	{
+		// This uses ADL rather than templates, because template specializations can't be defined in a different namespace
+		return GetForwardDeclaredShaderParametersStructMetadata((UniformBufferStructType*)nullptr);
+	}
+};
+
+template<class UniformBufferStructType>
+struct TUniformBufferMetadataHelper<UniformBufferStructType, typename std::enable_if<!std::is_same<typename UniformBufferStructType::FTypeInfo, void>::value>::type>
+{
+	static const FShaderParametersMetadata* GetStructMetadata()
+	{
+		return UniformBufferStructType::FTypeInfo::GetStructMetadata();
+	}
+};
+
 template<class UniformBufferStructType>
 struct TShaderParameterTypeInfo<TUniformBufferRef<UniformBufferStructType>>
 {
@@ -1195,7 +1215,7 @@ struct TShaderParameterTypeInfo<TUniformBufferRef<UniformBufferStructType>>
 	using TAlignedType = TAlignedShaderParameterPtr<TUniformBufferRef<UniformBufferStructType>>;
 	using TInstancedType = UniformBufferStructType;
 
-	static const FShaderParametersMetadata* GetStructMetadata() { return &UniformBufferStructType::StaticStructMetadata; }
+	static const FShaderParametersMetadata* GetStructMetadata() { return TUniformBufferMetadataHelper<UniformBufferStructType>::GetStructMetadata(); }
 };
 
 template<class UniformBufferStructType>
@@ -1210,7 +1230,7 @@ struct TShaderParameterTypeInfo<TUniformBufferBinding<UniformBufferStructType>>
 	using TAlignedType = TUniformBufferBinding<UniformBufferStructType>;
 	using TInstancedType = UniformBufferStructType;
 
-	static const FShaderParametersMetadata* GetStructMetadata() { return &UniformBufferStructType::StaticStructMetadata; }
+	static const FShaderParametersMetadata* GetStructMetadata() { return TUniformBufferMetadataHelper<UniformBufferStructType>::GetStructMetadata(); }
 };
 
 template<class UniformBufferStructType>
@@ -1225,7 +1245,7 @@ struct TShaderParameterTypeInfo<TRDGUniformBufferBinding<UniformBufferStructType
 	using TAlignedType = TRDGUniformBufferBinding<UniformBufferStructType>;
 	using TInstancedType = UniformBufferStructType;
 
-	static const FShaderParametersMetadata* GetStructMetadata() { return &UniformBufferStructType::StaticStructMetadata; }
+	static const FShaderParametersMetadata* GetStructMetadata() { return TUniformBufferMetadataHelper<UniformBufferStructType>::GetStructMetadata(); }
 };
 
 template<typename StructType>
@@ -1390,9 +1410,11 @@ extern RENDERCORE_API FShaderParametersMetadata* FindUniformBufferStructByShader
  *	IMPLEMENT_UNIFORM_BUFFER_STRUCT(FMyParameterStruct, "MyShaderBindingName");
  */
 #define BEGIN_UNIFORM_BUFFER_STRUCT(StructTypeName, PrefixKeywords) \
+	DECLARE_UNIFORM_BUFFER_STRUCT(StructTypeName, PrefixKeywords) \
 	INTERNAL_SHADER_PARAMETER_STRUCT_BEGIN(StructTypeName,PrefixKeywords,{} INTERNAL_BEGIN_UNIFORM_BUFFER_STRUCT, INTERNAL_UNIFORM_BUFFER_STRUCT_GET_STRUCT_METADATA(StructTypeName), INTERNAL_UNIFORM_BUFFER_STRUCT_CREATE_UNIFORM_BUFFER)
 
 #define BEGIN_UNIFORM_BUFFER_STRUCT_WITH_CONSTRUCTOR(StructTypeName, PrefixKeywords) \
+	DECLARE_UNIFORM_BUFFER_STRUCT(StructTypeName, PrefixKeywords) \
 	INTERNAL_SHADER_PARAMETER_STRUCT_BEGIN(StructTypeName,PrefixKeywords,; INTERNAL_BEGIN_UNIFORM_BUFFER_STRUCT, INTERNAL_UNIFORM_BUFFER_STRUCT_GET_STRUCT_METADATA(StructTypeName), INTERNAL_UNIFORM_BUFFER_STRUCT_CREATE_UNIFORM_BUFFER)
 
 #define END_UNIFORM_BUFFER_STRUCT() \
@@ -1409,7 +1431,8 @@ extern RENDERCORE_API FShaderParametersMetadata* FindUniformBufferStructByShader
 	StructTypeName::FTypeInfo::FileName, \
 	StructTypeName::FTypeInfo::FileLine, \
 	sizeof(StructTypeName), \
-	StructTypeName::zzGetMembers())
+	StructTypeName::zzGetMembers()); \
+	const FShaderParametersMetadata* GetForwardDeclaredShaderParametersStructMetadata(const StructTypeName* DummyPtr) {  return StructTypeName::FTypeInfo::GetStructMetadata(); }
 
 #define IMPLEMENT_UNIFORM_BUFFER_ALIAS_STRUCT(StructTypeName, UniformBufferAlias) \
 	static const FShaderParametersMetadata UniformBufferAlias( \
@@ -1436,7 +1459,8 @@ extern RENDERCORE_API FShaderParametersMetadata* FindUniformBufferStructByShader
 		StructTypeName::FTypeInfo::FileName, \
 		StructTypeName::FTypeInfo::FileLine, \
 		sizeof(StructTypeName), \
-		StructTypeName::zzGetMembers())
+		StructTypeName::zzGetMembers()); \
+	const FShaderParametersMetadata* GetForwardDeclaredShaderParametersStructMetadata(const StructTypeName* DummyPtr) {  return StructTypeName::FTypeInfo::GetStructMetadata(); }
 
 #define IMPLEMENT_STATIC_UNIFORM_BUFFER_STRUCT_EX2(StructTypeName,ShaderVariableName,StaticSlotName,BindingFlagsEnum,UsageFlags) \
 	static_assert(EnumHasAnyFlags(EUniformBufferBindingFlags::BindingFlagsEnum, EUniformBufferBindingFlags::Static), "Shader bindings must include 'Static'."); \
@@ -1453,7 +1477,8 @@ extern RENDERCORE_API FShaderParametersMetadata* FindUniformBufferStructByShader
 		StructTypeName::zzGetMembers(), \
 		false, \
 		nullptr, \
-		(uint32)UsageFlags)
+		(uint32)UsageFlags); \
+	const FShaderParametersMetadata* GetForwardDeclaredShaderParametersStructMetadata(const StructTypeName* DummyPtr) {  return StructTypeName::FTypeInfo::GetStructMetadata(); }
 
 /** Implements the same contract as IMPLEMENT_STATIC_UNIFORM_BUFFER_STRUCT, with the addition of the 'StaticAndShader' binding
  *  flag. This means that the uniform buffer may be bound statically or through SetShaderParameters. The uniform buffer is NOT
