Add missing copyright notice

diff --git a/Engine/Source/Runtime/RenderCore/Public/ShaderParameterFwd.h b/Engine/Source/Runtime/RenderCore/Public/ShaderParameterFwd.h
--- a/Engine/Source/Runtime/RenderCore/Public/ShaderParameterFwd.h
+++ b/Engine/Source/Runtime/RenderCore/Public/ShaderParameterFwd.h
@@ -1,3 +1,5 @@
+// Copyright Epic Games, Inc. All Rights Reserved.
+
 #pragma once
 
 class FShaderParametersMetadata;
