Rename ShaderParameterFwd.h to ShaderParameterStructDeclaration.h

diff --git a/Engine/Source/Runtime/RenderCore/Public/ShaderParameterFwd.h b/Engine/Source/Runtime/RenderCore/Public/ShaderParameterFwd.h
deleted file mode 100644
diff --git a/Engine/Source/Runtime/RenderCore/Public/ShaderParameterMacros.h b/Engine/Source/Runtime/RenderCore/Public/ShaderParameterMacros.h
--- a/Engine/Source/Runtime/RenderCore/Public/ShaderParameterMacros.h
+++ b/Engine/Source/Runtime/RenderCore/Public/ShaderParameterMacros.h
@@ -7,7 +7,7 @@
 
 #pragma once
 
-#include "ShaderParameterFwd.h"
+#include "ShaderParameterStructDeclaration.h"
 #include "ShaderParameterMetadata.h"
 #include "RenderGraphAllocator.h"
 #include "Algo/Reverse.h"
diff --git a/Engine/Source/Runtime/RenderCore/Public/ShaderParameterStructDeclaration.h b/Engine/Source/Runtime/RenderCore/Public/ShaderParameterStructDeclaration.h
new file mode 100644
--- /dev/null
+++ b/Engine/Source/Runtime/RenderCore/Public/ShaderParameterStructDeclaration.h
@@ -0,0 +1,21 @@
+// Copyright Epic Games, Inc. All Rights Reserved.
+
+#pragma once
+
+class FShaderParametersMetadata;
+
+#define DECLARE_UNIFORM_BUFFER_STRUCT(StructTypeName, PrefixKeywords) \
+	class StructTypeName; \
+	PrefixKeywords const FShaderParametersMetadata* GetForwardDeclaredShaderParametersStructMetadata(const StructTypeName* DummyPtr);
+
+
+/** Retrieves the metadata for a shader parameter struct type without having a definition of the type.
+ *  This is the default overload in case there is no exact non-template overload (prevents implicit casting).
+ */
+template<typename T>
+const FShaderParametersMetadata* GetForwardDeclaredShaderParametersStructMetadata(const T*)
+{
+	static_assert(!std::is_same<T, T>::value /* true */,
+		TEXT("Partial uniform buffer struct declaration. Use `DECLARE_UNIFORM_BUFFER_STRUCT()` instead of `class T;`.")
+	);
+}
