Add missing includes for nopch build

diff --git a/Engine/Source/Runtime/RenderCore/Public/ShaderParameterStructDeclaration.h b/Engine/Source/Runtime/RenderCore/Public/ShaderParameterStructDeclaration.h
--- a/Engine/Source/Runtime/RenderCore/Public/ShaderParameterStructDeclaration.h
+++ b/Engine/Source/Runtime/RenderCore/Public/ShaderParameterStructDeclaration.h
@@ -2,6 +2,9 @@
 
 #pragma once
 
+#include <type_traits>
+#include "HAL/Platform.h"
+
 class FShaderParametersMetadata;
 
 #define DECLARE_UNIFORM_BUFFER_STRUCT(StructTypeName, PrefixKeywords) \
