Added RDG AddSetupTask variant that takes a pipe and prerequisite tasks.

diff --git a/Engine/Source/Runtime/RenderCore/Public/RenderGraphBuilder.h b/Engine/Source/Runtime/RenderCore/Public/RenderGraphBuilder.h
--- a/Engine/Source/Runtime/RenderCore/Public/RenderGraphBuilder.h
+++ b/Engine/Source/Runtime/RenderCore/Public/RenderGraphBuilder.h
@@ -220,6 +220,12 @@ public:
 	template <typename TaskLambda>
 	UE::Tasks::FTask AddSetupTask(TaskLambda&& Task, bool bCondition = true);
 
+	template <typename TaskLambda>
+	UE::Tasks::FTask AddSetupTask(TaskLambda&& Task, UE::Tasks::FPipe* Pipe, UE::Tasks::ETaskPriority Priority = UE::Tasks::ETaskPriority::Normal, bool bCondition = true);
+
+	template <typename TaskLambda, typename PrerequisitesCollectionType>
+	UE::Tasks::FTask AddSetupTask(TaskLambda&& Task, UE::Tasks::FPipe* Pipe, PrerequisitesCollectionType&& Prerequisites, UE::Tasks::ETaskPriority Priority = UE::Tasks::ETaskPriority::Normal, bool bCondition = true);
+
 	/** Launches a task that is synced prior to graph execution. If parallel execution is not enabled, the lambda is run immediately. */
 	template <typename TaskLambda>
 	UE::Tasks::FTask AddCommandListSetupTask(TaskLambda&& Task, bool bCondition = true);
diff --git a/Engine/Source/Runtime/RenderCore/Public/RenderGraphBuilder.inl b/Engine/Source/Runtime/RenderCore/Public/RenderGraphBuilder.inl
--- a/Engine/Source/Runtime/RenderCore/Public/RenderGraphBuilder.inl
+++ b/Engine/Source/Runtime/RenderCore/Public/RenderGraphBuilder.inl
@@ -406,17 +406,47 @@ inline void FRDGBuilder::AddDispatchHint()
 }
 
 template <typename TaskLambdaType>
-UE::Tasks::FTask FRDGBuilder::AddSetupTask(TaskLambdaType&& TaskLambda, bool bCondition)
+FORCEINLINE UE::Tasks::FTask FRDGBuilder::AddSetupTask(TaskLambdaType&& TaskLambda, bool bCondition)
+{
+	return AddSetupTask(MoveTemp(TaskLambda), nullptr, TArray<UE::Tasks::FTask>{}, UE::Tasks::ETaskPriority::Normal, bCondition);
+}
+
+template <typename TaskLambdaType>
+FORCEINLINE UE::Tasks::FTask FRDGBuilder::AddSetupTask(
+	TaskLambdaType&& TaskLambda,
+	UE::Tasks::FPipe* Pipe,
+	UE::Tasks::ETaskPriority Priority,
+	bool bCondition)
+{
+	return AddSetupTask(MoveTemp(TaskLambda), Pipe, TArray<UE::Tasks::FTask>{}, Priority, bCondition);
+}
+
+template <typename TaskLambdaType, typename PrerequisitesCollectionType>
+UE::Tasks::FTask FRDGBuilder::AddSetupTask(
+	TaskLambdaType&& TaskLambda,
+	UE::Tasks::FPipe* Pipe,
+	PrerequisitesCollectionType&& Prerequisites,
+	UE::Tasks::ETaskPriority Priority,
+	bool bCondition)
 {
 	UE::Tasks::FTask Task;
 
 	if (bParallelExecuteEnabled && bCondition)
 	{
-		Task = UE::Tasks::Launch(TEXT("FRDGBuilder::AddSetupTask"), [TaskLambda = MoveTemp(TaskLambda)] () mutable
+		auto OuterLambda = [TaskLambda = MoveTemp(TaskLambda)]() mutable
 		{
 			FOptionalTaskTagScope Scope(ETaskTag::EParallelRenderingThread);
 			TaskLambda();
-		});
+		};
+
+		if (Pipe)
+		{
+			Task = Pipe->Launch(TEXT("FRDGBuilder::AddSetupTask"), MoveTemp(OuterLambda), MoveTemp(Prerequisites), Priority);
+		}
+		else
+		{
+			Task = UE::Tasks::Launch(TEXT("FRDGBuilder::AddSetupTask"), MoveTemp(OuterLambda), MoveTemp(Prerequisites), Priority);
+		}
 
 		ParallelSetupEvents.Emplace(Task);
 	}
