Fixes compilation failure in RenderTimer.h

diff --git a/Engine/Source/Runtime/RenderCore/Public/RenderTimer.h b/Engine/Source/Runtime/RenderCore/Public/RenderTimer.h
--- a/Engine/Source/Runtime/RenderCore/Public/RenderTimer.h
+++ b/Engine/Source/Runtime/RenderCore/Public/RenderTimer.h
@@ -2,6 +2,7 @@
 
 #pragma once
 
+#include "CoreGlobals.h"
 #include "HAL/PlatformTime.h"
 
 /** Timer helper class. **/
