Add instanced fullscreen utility functions for multiviewport.

diff --git a/Engine/Shaders/Private/Tools/FullscreenVertexShader.usf b/Engine/Shaders/Private/Tools/FullscreenVertexShader.usf
--- a/Engine/Shaders/Private/Tools/FullscreenVertexShader.usf
+++ b/Engine/Shaders/Private/Tools/FullscreenVertexShader.usf
@@ -6,10 +6,27 @@
 
 #include "../Common.ush"
 
+#ifndef SCREEN_VS_FOR_INSTANCED_VIEWS
+	#define SCREEN_VS_FOR_INSTANCED_VIEWS 0
+#endif
+
 void MainVS(
 	float2 ViewportUV : ATTRIBUTE0,
 	float2 UV         : ATTRIBUTE1, // TODO: kill
-	out float4 OutPosition : SV_POSITION)
+#if !SCREEN_VS_FOR_INSTANCED_VIEWS
+	out float4 OutPosition : SV_POSITION
+#else
+	uint InstanceId	  : SV_InstanceID,
+	out float4 OutPosition : SV_POSITION,
+	out uint OutViewportId	: SV_ViewportArrayIndex,
+	out nointerpolation uint OutViewId : VIEW_ID
+#endif
+	)
 {
 	OutPosition = float4(ViewportUVToScreenPos(ViewportUV), 0, 1);
+
+#if SCREEN_VS_FOR_INSTANCED_VIEWS
+	OutViewportId = InstanceId;
+	OutViewId = InstanceId;
+#endif
 }
diff --git a/Engine/Source/Runtime/RenderCore/Private/CommonRenderResources.cpp b/Engine/Source/Runtime/RenderCore/Private/CommonRenderResources.cpp
--- a/Engine/Source/Runtime/RenderCore/Private/CommonRenderResources.cpp
+++ b/Engine/Source/Runtime/RenderCore/Private/CommonRenderResources.cpp
@@ -6,6 +6,7 @@
 
 #include "CommonRenderResources.h"
 #include "Containers/DynamicRHIResourceArray.h"
+#include "StereoRenderUtils.h"
 
 
 TGlobalResource<FFilterVertexDeclaration> GFilterVertexDeclaration;
@@ -15,8 +16,14 @@ TGlobalResource<FScreenRectangleVertexBuffer> GScreenRectangleVertexBuffer;
 TGlobalResource<FScreenRectangleIndexBuffer> GScreenRectangleIndexBuffer;
 
 IMPLEMENT_GLOBAL_SHADER(FScreenVertexShaderVS, "/Engine/Private/Tools/FullscreenVertexShader.usf", "MainVS", SF_Vertex);
+IMPLEMENT_GLOBAL_SHADER(FInstancedScreenVertexShaderVS, "/Engine/Private/Tools/FullscreenVertexShader.usf", "MainVS", SF_Vertex);
 IMPLEMENT_GLOBAL_SHADER(FCopyRectPS, "/Engine/Private/ScreenPass.usf", "CopyRectPS", SF_Pixel);
 
+bool FInstancedScreenVertexShaderVS::ShouldCompilePermutation(const FGlobalShaderPermutationParameters& Parameters)
+{
+	UE::StereoRenderUtils::FStereoShaderAspects Aspects(Parameters.Platform);
+	return Aspects.IsInstancedMultiViewportEnabled();
+}
 
 void FScreenRectangleVertexBuffer::InitRHI()
 {
diff --git a/Engine/Source/Runtime/RenderCore/Private/PixelShaderUtils.cpp b/Engine/Source/Runtime/RenderCore/Private/PixelShaderUtils.cpp
--- a/Engine/Source/Runtime/RenderCore/Private/PixelShaderUtils.cpp
+++ b/Engine/Source/Runtime/RenderCore/Private/PixelShaderUtils.cpp
@@ -65,3 +65,22 @@ void FPixelShaderUtils::InitFullscreenPipelineState(
 	GraphicsPSOInit.BoundShaderState.PixelShaderRHI = PixelShader.GetPixelShader();
 	GraphicsPSOInit.PrimitiveType = PT_TriangleList;
 }
+
+void FPixelShaderUtils::InitFullscreenMultiviewportPipelineState(
+	FRHICommandList& RHICmdList,
+	const FGlobalShaderMap* GlobalShaderMap,
+	const TShaderRef<FShader>& PixelShader,
+	FGraphicsPipelineStateInitializer& GraphicsPSOInit)
+{
+	TShaderMapRef<FInstancedScreenVertexShaderVS> VertexShader(GlobalShaderMap);
+
+	RHICmdList.ApplyCachedRenderTargets(GraphicsPSOInit);
+	GraphicsPSOInit.BlendState = TStaticBlendState<>::GetRHI();
+	GraphicsPSOInit.RasterizerState = TStaticRasterizerState<>::GetRHI();
+	GraphicsPSOInit.DepthStencilState = TStaticDepthStencilState<false, CF_Always>::GetRHI();
+
+	GraphicsPSOInit.BoundShaderState.VertexDeclarationRHI = GFilterVertexDeclaration.VertexDeclarationRHI;
+	GraphicsPSOInit.BoundShaderState.VertexShaderRHI = VertexShader.GetVertexShader();
+	GraphicsPSOInit.BoundShaderState.PixelShaderRHI = PixelShader.GetPixelShader();
+	GraphicsPSOInit.PrimitiveType = PT_TriangleList;
+}
diff --git a/Engine/Source/Runtime/RenderCore/Public/CommonRenderResources.h b/Engine/Source/Runtime/RenderCore/Public/CommonRenderResources.h
--- a/Engine/Source/Runtime/RenderCore/Public/CommonRenderResources.h
+++ b/Engine/Source/Runtime/RenderCore/Public/CommonRenderResources.h
@@ -119,6 +119,25 @@ class RENDERCORE_API FScreenVertexShaderVS : public FGlobalShader
 	END_SHADER_PARAMETER_STRUCT()
 };
 
+/** Vertex shader to draw an instanced quad covering all the viewports (SV_ViewportArrayIndex is output for each SV_InstanceID). Does not have any shader parameters.
+ * The pixel shader should just use SV_Position. */
+class RENDERCORE_API FInstancedScreenVertexShaderVS : public FScreenVertexShaderVS
+{
+	DECLARE_GLOBAL_SHADER(FInstancedScreenVertexShaderVS);
+	SHADER_USE_PARAMETER_STRUCT(FInstancedScreenVertexShaderVS, FScreenVertexShaderVS);
+
+	static bool ShouldCompilePermutation(const FGlobalShaderPermutationParameters& Parameters);
+
+	static void ModifyCompilationEnvironment(const FGlobalShaderPermutationParameters& Parameters, FShaderCompilerEnvironment& OutEnvironment)
+	{
+		FScreenVertexShaderVS::ModifyCompilationEnvironment(Parameters, OutEnvironment);
+		OutEnvironment.SetDefine(TEXT("SCREEN_VS_FOR_INSTANCED_VIEWS"), 1);
+	}
+
+	BEGIN_SHADER_PARAMETER_STRUCT(FParameters, )
+		END_SHADER_PARAMETER_STRUCT()
+};
+
 /** Pixel shader to copy pixels from src to dst performing a format change that works on all platforms. */
 class RENDERCORE_API FCopyRectPS : public FGlobalShader
 {
diff --git a/Engine/Source/Runtime/RenderCore/Public/PixelShaderUtils.h b/Engine/Source/Runtime/RenderCore/Public/PixelShaderUtils.h
--- a/Engine/Source/Runtime/RenderCore/Public/PixelShaderUtils.h
+++ b/Engine/Source/Runtime/RenderCore/Public/PixelShaderUtils.h
@@ -73,6 +73,13 @@ struct RENDERCORE_API FPixelShaderUtils
 		const TShaderRef<FShader>& PixelShader,
 		FGraphicsPipelineStateInitializer& GraphicsPSOInit);
 
+	/** Initialize a pipeline state object initializer with almost all the basics required to do a full multi-viewport pass. */
+	static void InitFullscreenMultiviewportPipelineState(
+		FRHICommandList& RHICmdList,
+		const FGlobalShaderMap* GlobalShaderMap,
+		const TShaderRef<FShader>& PixelShader,
+		FGraphicsPipelineStateInitializer& GraphicsPSOInit);
+
 	/** Dispatch a full screen pixel shader to rhi command list with its parameters. */
 	template<typename TShaderClass>
 	static inline void DrawFullscreenPixelShader(
@@ -102,6 +109,47 @@ struct RENDERCORE_API FPixelShaderUtils
 		DrawFullscreenTriangle(RHICmdList);
 	}
 
+	/** Dispatch a full screen pixel shader to rhi command list with its parameters, covering several views at once. */
+	template<typename TShaderClass>
+	static inline void DrawFullscreenInstancedMultiViewportPixelShader(
+		FRHICommandList& RHICmdList,
+		const FGlobalShaderMap* GlobalShaderMap,
+		const TShaderRef<TShaderClass>& PixelShader,
+		const typename TShaderClass::FParameters& Parameters,
+		TArrayView<FIntRect const> Viewports,
+		FRHIBlendState* BlendState = nullptr,
+		FRHIRasterizerState* RasterizerState = nullptr,
+		FRHIDepthStencilState* DepthStencilState = nullptr,
+		uint32 StencilRef = 0)
+	{
+		check(PixelShader.IsValid());
+		checkf(Viewports.Num() == 2, TEXT("Only two instanced viewports are currently supported"));
+		{
+			const FIntRect& LeftViewRect = Viewports[0];
+			const int32 LeftMinX = LeftViewRect.Min.X;
+			const int32 LeftMaxX = LeftViewRect.Max.X;
+			const int32 LeftMaxY = LeftViewRect.Max.Y;
+
+			const FIntRect& RightViewRect = Viewports[1];
+			const int32 RightMinX = RightViewRect.Min.X;
+			const int32 RightMaxX = RightViewRect.Max.X;
+			const int32 RightMaxY = RightViewRect.Max.Y;
+			RHICmdList.SetStereoViewport((float)LeftMinX, (float)RightMinX, 0.0f, 0.0f, 0.0f, (float)LeftMaxX, (float)RightMaxX, (float)LeftMaxY, (float)RightMaxY, 1.0f);
+		}
+
+		FGraphicsPipelineStateInitializer GraphicsPSOInit;
+		InitFullscreenMultiviewportPipelineState(RHICmdList, GlobalShaderMap, PixelShader, /* out */ GraphicsPSOInit);
+		GraphicsPSOInit.BlendState = BlendState ? BlendState : GraphicsPSOInit.BlendState;
+		GraphicsPSOInit.RasterizerState = RasterizerState ? RasterizerState : GraphicsPSOInit.RasterizerState;
+		GraphicsPSOInit.DepthStencilState = DepthStencilState ? DepthStencilState : GraphicsPSOInit.DepthStencilState;
+
+		SetGraphicsPipelineState(RHICmdList, GraphicsPSOInit, StencilRef);
+
+		SetShaderParameters(RHICmdList, PixelShader, PixelShader.GetPixelShader(), Parameters);
+
+		DrawFullscreenTriangle(RHICmdList, Viewports.Num());
+	}
+
 	/** Dispatch a pixel shader to render graph builder with its parameters. */
 	template<typename TShaderClass>
 	static inline void AddFullscreenPass(
@@ -130,6 +178,34 @@ struct RENDERCORE_API FPixelShaderUtils
 		});
 	}
 
+	/** Dispatch a pixel shader to render graph builder with its parameters. */
+	template<typename TShaderClass>
+	static inline void AddFullscreenInstancedMultiViewportPass(
+		FRDGBuilder& GraphBuilder,
+		const FGlobalShaderMap* GlobalShaderMap,
+		FRDGEventName&& PassName,
+		const TShaderRef<TShaderClass>& PixelShader,
+		typename TShaderClass::FParameters* Parameters,
+		TArray<FIntRect>&& Viewports,
+		FRHIBlendState* BlendState = nullptr,
+		FRHIRasterizerState* RasterizerState = nullptr,
+		FRHIDepthStencilState* DepthStencilState = nullptr,
+		uint32 StencilRef = 0)
+	{
+		check(PixelShader.IsValid());
+		ClearUnusedGraphResources(PixelShader, Parameters);
+
+		GraphBuilder.AddPass(
+			Forward<FRDGEventName>(PassName),
+			Parameters,
+			ERDGPassFlags::Raster,
+			[Parameters, GlobalShaderMap, PixelShader, Viewports, BlendState, RasterizerState, DepthStencilState, StencilRef](FRHICommandList& RHICmdList)
+		{
+			FPixelShaderUtils::DrawFullscreenInstancedMultiViewportPixelShader(RHICmdList, GlobalShaderMap, PixelShader, *Parameters, Viewports,
+				BlendState, RasterizerState, DepthStencilState, StencilRef);
+		});
+	}
+
 	/** Rect based pixel shader pass. */
 	template<typename TPixelShaderClass, typename TPassParameters>
 	static inline void AddRasterizeToRectsPass(
