Integrate RenderResourceViewer from //UE5/Main
Part3 - Integrate CL23347911 as edigrate
Add owner column to Render Resource Viewer, add owner path name tracking to skeletal & static mesh buffers
with conflict resolve modifications in RawIndexBuffer.cpp/h, RenderResource.cpp/h, RHIResources.h

diff --git a/Engine/Source/Runtime/Engine/Private/Rendering/SkinWeightVertexBuffer.cpp b/Engine/Source/Runtime/Engine/Private/Rendering/SkinWeightVertexBuffer.cpp
--- a/Engine/Source/Runtime/Engine/Private/Rendering/SkinWeightVertexBuffer.cpp
+++ b/Engine/Source/Runtime/Engine/Private/Rendering/SkinWeightVertexBuffer.cpp
@@ -803,6 +803,7 @@ void FSkinWeightVertexBuffer::ResetVertexBoneWeights(uint32 VertexIndex)
 
 void FSkinWeightVertexBuffer::SetOwnerName(const FName& OwnerName)
 {
+	check(IsInGameThread());
 	LookupVertexBuffer.SetOwnerName(OwnerName);
 	DataVertexBuffer.SetOwnerName(OwnerName);
 }
diff --git a/Engine/Source/Runtime/RenderCore/Private/RenderResource.cpp b/Engine/Source/Runtime/RenderCore/Private/RenderResource.cpp
--- a/Engine/Source/Runtime/RenderCore/Private/RenderResource.cpp
+++ b/Engine/Source/Runtime/RenderCore/Private/RenderResource.cpp
@@ -318,6 +318,23 @@ FRenderResource::~FRenderResource()
 	}
 }
 
+void FRenderResource::SetOwnerName(const FName& InOwnerName)
+{
+#if RHI_ENABLE_RESOURCE_INFO
+	check(IsInGameThread());
+	OwnerName = InOwnerName;
+#endif
+}
+
+FName FRenderResource::GetOwnerName() const
+{
+#if RHI_ENABLE_RESOURCE_INFO
+	return OwnerName;
+#else
+	return NAME_None;
+#endif
+}
+
 bool FRenderResource::ShouldFreeResourceObject(void* ResourceObject, FResourceArrayInterface* ResourceArray)
 {
 	return GFreeStructuresOnRHIBufferCreation && ResourceObject && (!ResourceArray || !ResourceArray->GetResourceDataSize());
diff --git a/Engine/Source/Runtime/RenderCore/Public/RenderResource.h b/Engine/Source/Runtime/RenderCore/Public/RenderResource.h
--- a/Engine/Source/Runtime/RenderCore/Public/RenderResource.h
+++ b/Engine/Source/Runtime/RenderCore/Public/RenderResource.h
@@ -152,6 +152,11 @@ public:
 	void SetOwnerName(const FName& InOwnerName);
 	FName GetOwnerName() const;
 
+private:
+#if RHI_ENABLE_RESOURCE_INFO
+	FName OwnerName;
+#endif
+
 protected:
 	// This is used during mobile editor preview refactor, this will eventually be replaced with a parameter to InitRHI() etc..
 	void SetFeatureLevel(const FStaticFeatureLevel InFeatureLevel) { FeatureLevel = (ERHIFeatureLevel::Type)InFeatureLevel; }
