Fix incorrect auto merge of CL23966262 causing duplicated OwnerName

diff --git a/Engine/Source/Runtime/RenderCore/Private/RenderResource.cpp b/Engine/Source/Runtime/RenderCore/Private/RenderResource.cpp
--- a/Engine/Source/Runtime/RenderCore/Private/RenderResource.cpp
+++ b/Engine/Source/Runtime/RenderCore/Private/RenderResource.cpp
@@ -318,23 +318,6 @@ FRenderResource::~FRenderResource()
 	}
 }
 
-void FRenderResource::SetOwnerName(const FName& InOwnerName)
-{
-#if RHI_ENABLE_RESOURCE_INFO
-	check(IsInGameThread());
-	OwnerName = InOwnerName;
-#endif
-}
-
-FName FRenderResource::GetOwnerName() const
-{
-#if RHI_ENABLE_RESOURCE_INFO
-	return OwnerName;
-#else
-	return NAME_None;
-#endif
-}
-
 bool FRenderResource::ShouldFreeResourceObject(void* ResourceObject, FResourceArrayInterface* ResourceArray)
 {
 	return GFreeStructuresOnRHIBufferCreation && ResourceObject && (!ResourceArray || !ResourceArray->GetResourceDataSize());
diff --git a/Engine/Source/Runtime/RenderCore/Public/RenderResource.h b/Engine/Source/Runtime/RenderCore/Public/RenderResource.h
--- a/Engine/Source/Runtime/RenderCore/Public/RenderResource.h
+++ b/Engine/Source/Runtime/RenderCore/Public/RenderResource.h
@@ -152,11 +152,6 @@ public:
 	void SetOwnerName(const FName& InOwnerName);
 	FName GetOwnerName() const;
 
-private:
-#if RHI_ENABLE_RESOURCE_INFO
-	FName OwnerName;
-#endif
-
 protected:
 	// This is used during mobile editor preview refactor, this will eventually be replaced with a parameter to InitRHI() etc..
 	void SetFeatureLevel(const FStaticFeatureLevel InFeatureLevel) { FeatureLevel = (ERHIFeatureLevel::Type)InFeatureLevel; }
