Fixed: Virtual texture shimmering on mobile platforms
Added manual VT filtering for mobile platforms that do not use TAA (contributed by [REDACTED_PERSON] at ARM)

diff --git a/Engine/Source/Runtime/RenderCore/Private/Shader.cpp b/Engine/Source/Runtime/RenderCore/Private/Shader.cpp
--- a/Engine/Source/Runtime/RenderCore/Private/Shader.cpp
+++ b/Engine/Source/Runtime/RenderCore/Private/Shader.cpp
@@ -2029,8 +2029,7 @@ void ShaderMapAppendKeyString(EShaderPlatform Platform, FString& KeyString)
 
 		const bool VTSupported = TargetPlatform != nullptr && TargetPlatform->SupportsFeature(ETargetPlatformFeatures::VirtualTextureStreaming);
 
-		auto tt = FString::Printf(TEXT("_VT-%d-%d-%d-%d"), VTLightmaps, VTTextures, VTSupported, VTFiltering);
- 		KeyString += tt;
+ 		KeyString += FString::Printf(TEXT("_VT-%d-%d-%d-%d"), VTLightmaps, VTTextures, VTSupported, VTFiltering);
 	}
 
 	{
