Fix for compiler error when dependent type is used with CreateStructuredBuffer.

diff --git a/Engine/Source/Runtime/RenderCore/Public/RenderGraphUtils.h b/Engine/Source/Runtime/RenderCore/Public/RenderGraphUtils.h
--- a/Engine/Source/Runtime/RenderCore/Public/RenderGraphUtils.h
+++ b/Engine/Source/Runtime/RenderCore/Public/RenderGraphUtils.h
@@ -914,7 +914,7 @@ FORCEINLINE FRDGBufferRef CreateStructuredBuffer_Impl(
 	const TCHAR* Name,
 	TRDGBufferArrayCallback<ArrayType>&& ArrayCallback)
 {
-	return CreateStructuredBuffer(GraphBuilder, Name, sizeof(std::remove_reference_t<ArrayType>::ElementType),
+	return CreateStructuredBuffer(GraphBuilder, Name, sizeof(typename std::remove_reference_t<ArrayType>::ElementType),
 		/*NumElementsCallback = */[ArrayCallback]() { return ArrayCallback().Num(); },
 		/*InitialDataCallback = */[ArrayCallback]() { return ArrayCallback().GetData(); },
 		/*InitialDataSizeCallback = */[ArrayCallback]() { return ArrayCallback().Num() * ArrayCallback().GetTypeSize(); });
