More shader parameter batching utilities
- Adding Uniform Buffer parameter batching functions to ShaderParameterUtils.h
- Adding batched Uniform Buffer binding helper to GlobalShader
- Adding batched Uniform Buffer bindings helpers to FMaterialShader
- Adding SetShaderParameters variant that uses batched parameters.
- Changing UnsetShaderUAVs and UnsetShaderSRVs to use the RHICmdList scratch parameters instead of their own arrays.

diff --git a/Engine/Source/Runtime/RenderCore/Private/ShaderParameterStruct.cpp b/Engine/Source/Runtime/RenderCore/Private/ShaderParameterStruct.cpp
--- a/Engine/Source/Runtime/RenderCore/Private/ShaderParameterStruct.cpp
+++ b/Engine/Source/Runtime/RenderCore/Private/ShaderParameterStruct.cpp
@@ -623,8 +623,6 @@ static void ExtractShaderParameters(
 	}
 }
 
-#define UE_SINGLE_SHADER_PARAMETERS_COMMAND_INCLUDE_UB (1)
-
 template<typename BindingParameterType>
 FRHIShaderParameterResource ExtractShaderParameterResource(FShaderParameterReader Reader, const BindingParameterType& Parameter)
 {
@@ -707,10 +705,7 @@ static void ExtractShaderParameterResources(
 	}
 #endif
 
-	int32 NumBindings = Bindings.ResourceParameters.Num();
-#if UE_SINGLE_SHADER_PARAMETERS_COMMAND_INCLUDE_UB
-	NumBindings += Bindings.GraphUniformBuffers.Num() + Bindings.ParameterReferences.Num();
-#endif
+	const int32 NumBindings = Bindings.ResourceParameters.Num() + Bindings.GraphUniformBuffers.Num() + Bindings.ParameterReferences.Num();
 
 	if (NumBindings)
 	{
@@ -722,7 +717,6 @@ static void ExtractShaderParameterResources(
 			OutResourceParameters.Emplace(ShaderParameterResource);
 		}
 
-#if UE_SINGLE_SHADER_PARAMETERS_COMMAND_INCLUDE_UB
 		for (const FShaderParameterBindings::FParameterStructReference& Parameter : Bindings.GraphUniformBuffers)
 		{
 			const FRDGUniformBufferBinding& UniformBufferBinding = Reader.Read<FRDGUniformBufferBinding>(Parameter);
@@ -742,10 +736,26 @@ static void ExtractShaderParameterResources(
 				OutResourceParameters.Emplace(UniformBufferBinding.GetUniformBuffer(), GetParameterIndex(Parameter));
 			}
 		}
-#endif // UE_SINGLE_SHADER_PARAMETERS_COMMAND_INCLUDE_UB
 	}
 }
 
+/** Set batched parameters from a parameters struct. */
+void SetShaderParameters(
+	FRHIBatchedShaderParameters& BatchedParameters,
+	const FShaderParameterBindings& Bindings,
+	const FShaderParametersMetadata* ParametersMetadata,
+	const void* InParametersData)
+{
+	TConstArrayView<uint8> FullParametersData((const uint8*)InParametersData, ParametersMetadata->GetSize());
+
+	for (const FShaderParameterBindings::FParameter& Parameter : Bindings.Parameters)
+	{
+		BatchedParameters.SetShaderParameter(Parameter.BufferIndex, Parameter.BaseIndex, Parameter.ByteSize, FullParametersData.GetData() + Parameter.ByteOffset);
+	}
+
+	ExtractShaderParameterResources(BatchedParameters.ResourceParameters, BatchedParameters.BindlessParameters, Bindings, ParametersMetadata, FullParametersData);
+}
+
 /** Set shader's parameters from its parameters struct. */
 template<typename TRHICmdList, typename TShaderRHI>
 inline void SetShaderParametersInternal(
@@ -770,31 +780,6 @@ inline void SetShaderParametersInternal(
 	ExtractShaderParameterResources(ResourceParameters, BindlessParameters, Bindings, ParametersMetadata, FullParametersData);
 
 	RHICmdList.SetShaderParameters(ShaderRHI, MinimalParametersData, Parameters, ResourceParameters, BindlessParameters);
-
-#if !UE_SINGLE_SHADER_PARAMETERS_COMMAND_INCLUDE_UB
-	const FShaderParameterReader Reader(FullParametersData);
-
-	// Graph Uniform Buffers
-	for (const FShaderParameterBindings::FParameterStructReference& Parameter : Bindings.GraphUniformBuffers)
-	{
-		const FRDGUniformBufferBinding& UniformBufferBinding = Reader.Read<FRDGUniformBufferBinding>(Parameter);
-		if (UniformBufferBinding.IsShader())
-		{
-			UniformBufferBinding->MarkResourceAsUsed();
-			RHICmdList.SetShaderUniformBuffer(ShaderRHI, Parameter.BufferIndex, UniformBufferBinding->GetRHI());
-		}
-	}
-
-	// Reference structures
-	for (const FShaderParameterBindings::FParameterStructReference& Parameter : Bindings.ParameterReferences)
-	{
-		const FUniformBufferBinding& UniformBufferBinding = Reader.Read<FUniformBufferBinding>(Parameter);
-		if (UniformBufferBinding.IsShader())
-		{
-			RHICmdList.SetShaderUniformBuffer(ShaderRHI, Parameter.BufferIndex, UniformBufferBinding.GetUniformBuffer());
-		}
-	}
-#endif
 }
 
 void SetShaderParameters(
diff --git a/Engine/Source/Runtime/RenderCore/Public/GlobalShader.h b/Engine/Source/Runtime/RenderCore/Public/GlobalShader.h
--- a/Engine/Source/Runtime/RenderCore/Public/GlobalShader.h
+++ b/Engine/Source/Runtime/RenderCore/Public/GlobalShader.h
@@ -19,6 +19,7 @@
 #include "Serialization/MemoryImage.h"
 #include "Serialization/MemoryLayout.h"
 #include "Shader.h"
+#include "ShaderParameterUtils.h"
 #include "Templates/Tuple.h"
 #include "Templates/UniquePtr.h"
 #include "UObject/NameTypes.h"
@@ -289,6 +290,13 @@ public:
 
 	RENDERCORE_API FGlobalShader(const ShaderMetaType::CompiledShaderInitializerType& Initializer);
 	
+	template<typename TViewUniformShaderParameters>
+	inline void SetParameters(FRHIBatchedShaderParameters& BatchedParameters, FRHIUniformBuffer* ViewUniformBuffer)
+	{
+		const auto& ViewUniformBufferParameter = static_cast<const FShaderUniformBufferParameter&>(GetUniformBufferParameter<TViewUniformShaderParameters>());
+		SetUniformBufferParameter(BatchedParameters, ViewUniformBufferParameter, ViewUniformBuffer);
+	}
+
 	template<typename TViewUniformShaderParameters, typename ShaderRHIParamRef, typename TRHICmdList>
 	inline void SetParameters(TRHICmdList& RHICmdList, const ShaderRHIParamRef ShaderRHI, FRHIUniformBuffer* ViewUniformBuffer)
 	{
diff --git a/Engine/Source/Runtime/RenderCore/Public/ShaderParameterStruct.h b/Engine/Source/Runtime/RenderCore/Public/ShaderParameterStruct.h
--- a/Engine/Source/Runtime/RenderCore/Public/ShaderParameterStruct.h
+++ b/Engine/Source/Runtime/RenderCore/Public/ShaderParameterStruct.h
@@ -129,7 +129,7 @@ FORCEINLINE void ValidateShaderParameters(const TShaderRef<TShaderClass>& Shader
 }
 
 template<typename TParameterType>
-inline void CollectUAVsToUnset(TArray<FRHIShaderParameterResource, TInlineAllocator<16>>& UAVsToUnset, const TMemoryImageArray<TParameterType>& InParameters)
+inline void CollectUAVsToUnset(TArray<FRHIShaderParameterResource>& UAVsToUnset, const TMemoryImageArray<TParameterType>& InParameters)
 {
 	FRHIUnorderedAccessView* NullUAV = nullptr;
 	for (const TParameterType& Parameter : InParameters)
@@ -144,7 +144,7 @@ inline void CollectUAVsToUnset(TArray<FRHIShaderParameterResource, TInlineAlloca
 }
 
 template<typename TParameterType>
-inline void CollectSRVsToUnset(TArray<FRHIShaderParameterResource, TInlineAllocator<16>>& SRVsToUnset, const TMemoryImageArray<TParameterType>& InParameters)
+inline void CollectSRVsToUnset(TArray<FRHIShaderParameterResource>& SRVsToUnset, const TMemoryImageArray<TParameterType>& InParameters)
 {
 	FRHIShaderResourceView* NullSRV = nullptr;
 	for (const TParameterType& Parameter : InParameters)
@@ -158,56 +158,89 @@ inline void CollectSRVsToUnset(TArray<FRHIShaderParameterResource, TInlineAlloca
 	}
 }
 
-/** Unset compute shader UAVs. */
-template<typename TRHICmdList, typename TShaderClass>
-inline void UnsetShaderUAVs(TRHICmdList& RHICmdList, const TShaderRef<TShaderClass>& Shader, FRHIComputeShader* ShaderRHI)
+inline void UnsetShaderUAVs(FRHIBatchedShaderParameters& BatchedParameters, const FShaderParameterBindings& Bindings)
+{
+	CollectUAVsToUnset(BatchedParameters.ResourceParameters, Bindings.ResourceParameters);
+	CollectUAVsToUnset(BatchedParameters.BindlessParameters, Bindings.BindlessResourceParameters);
+}
+
+inline void UnsetShaderSRVs(FRHIBatchedShaderParameters& BatchedParameters, const FShaderParameterBindings& Bindings)
+{
+	CollectSRVsToUnset(BatchedParameters.ResourceParameters, Bindings.ResourceParameters);
+	CollectSRVsToUnset(BatchedParameters.BindlessParameters, Bindings.BindlessResourceParameters);
+}
+
+template<typename TShaderClass>
+inline void UnsetShaderUAVs(FRHIBatchedShaderParameters& BatchedParameters, const TShaderRef<TShaderClass>& Shader)
 {
 	// TODO(RDG): Once all shader sets their parameter through this, can refactor RHI so all UAVs of a shader get unset through a single RHI function call.
 	const FShaderParameterBindings& Bindings = Shader->Bindings;
 
 	checkf(Bindings.RootParameterBufferIndex == FShaderParameterBindings::kInvalidBufferIndex, TEXT("Can't use UnsetShaderUAVs() for root parameter buffer index."));
 
-	TArray<FRHIShaderParameterResource, TInlineAllocator<16>> ResourceUAVsToUnset;
-	TArray<FRHIShaderParameterResource, TInlineAllocator<16>> BindlessUAVsToUnset;
-
-	CollectUAVsToUnset(ResourceUAVsToUnset, Bindings.ResourceParameters);
-	CollectUAVsToUnset(BindlessUAVsToUnset, Bindings.BindlessResourceParameters);
-
-	RHICmdList.SetShaderParameters(ShaderRHI, {}, {}, ResourceUAVsToUnset, BindlessUAVsToUnset);
+	UnsetShaderUAVs(BatchedParameters, Bindings);
 }
 
-/** Unset compute shader SRVs. */
-template<typename TRHICmdList, typename TShaderClass>
-inline void UnsetShaderSRVs(TRHICmdList& RHICmdList, const TShaderRef<TShaderClass>& Shader, FRHIComputeShader* ShaderRHI)
+template<typename TShaderClass>
+inline void UnsetShaderSRVs(FRHIBatchedShaderParameters& BatchedParameters, const TShaderRef<TShaderClass>& Shader)
 {
+	// TODO(RDG): Once all shader sets their parameter through this, can refactor RHI so all SRVs of a shader get unset through a single RHI function call.
 	const FShaderParameterBindings& Bindings = Shader->Bindings;
 
 	checkf(Bindings.RootParameterBufferIndex == FShaderParameterBindings::kInvalidBufferIndex, TEXT("Can't use UnsetShaderSRVs() for root parameter buffer index."));
 
-	TArray<FRHIShaderParameterResource, TInlineAllocator<16>> ResourceSRVsToUnset;
-	TArray<FRHIShaderParameterResource, TInlineAllocator<16>> BindlessSRVsToUnset;
+	UnsetShaderSRVs(BatchedParameters, Bindings);
+}
 
-	CollectSRVsToUnset(ResourceSRVsToUnset, Bindings.ResourceParameters);
-	CollectSRVsToUnset(BindlessSRVsToUnset, Bindings.BindlessResourceParameters);
+/** Unset compute shader UAVs. */
+template<typename TRHICmdList, typename TShaderClass>
+inline void UnsetShaderUAVs(TRHICmdList& RHICmdList, const TShaderRef<TShaderClass>& Shader, FRHIComputeShader* ShaderRHI)
+{
+	FRHIBatchedShaderParameters& BatchedParameters = RHICmdList.GetScratchShaderParameters();
+	UnsetShaderUAVs(BatchedParameters, Shader);
+	RHICmdList.SetBatchedShaderParameters(ShaderRHI, BatchedParameters);
+}
 
-	RHICmdList.SetShaderParameters(ShaderRHI, {}, {}, ResourceSRVsToUnset, BindlessSRVsToUnset);
+/** Unset compute shader SRVs. */
+template<typename TRHICmdList, typename TShaderClass>
+inline void UnsetShaderSRVs(TRHICmdList& RHICmdList, const TShaderRef<TShaderClass>& Shader, FRHIComputeShader* ShaderRHI)
+{
+	FRHIBatchedShaderParameters& BatchedParameters = RHICmdList.GetScratchShaderParameters();
+	UnsetShaderSRVs(BatchedParameters, Shader);
+	RHICmdList.SetBatchedShaderParameters(ShaderRHI, BatchedParameters);
 }
 
 /** Unset compute shader SRVs. */
 template<typename TRHICmdList, typename TShaderClass>
 inline void UnsetShaderSRVs(TRHICmdList& RHICmdList, const TShaderRef<TShaderClass>& Shader, FRHIGraphicsShader* ShaderRHI)
 {
-	const FShaderParameterBindings& Bindings = Shader->Bindings;
-
-	checkf(Bindings.RootParameterBufferIndex == FShaderParameterBindings::kInvalidBufferIndex, TEXT("Can't use UnsetShaderSRVs() for root parameter buffer index."));
+	FRHIBatchedShaderParameters& BatchedParameters = RHICmdList.GetScratchShaderParameters();
+	UnsetShaderSRVs(BatchedParameters, Shader);
+	RHICmdList.SetBatchedShaderParameters(ShaderRHI, BatchedParameters);
+}
 
-	TArray<FRHIShaderParameterResource, TInlineAllocator<16>> ResourceSRVsToUnset;
-	TArray<FRHIShaderParameterResource, TInlineAllocator<16>> BindlessSRVsToUnset;
+RENDERCORE_API void SetShaderParameters(
+	FRHIBatchedShaderParameters& BatchedParameters,
+	const FShaderParameterBindings& Bindings,
+	const FShaderParametersMetadata* ParametersMetadata,
+	const void* ParametersData);
 
-	CollectSRVsToUnset(ResourceSRVsToUnset, Bindings.ResourceParameters);
-	CollectSRVsToUnset(BindlessSRVsToUnset, Bindings.BindlessResourceParameters);
+template<typename TShaderClass>
+inline void SetShaderParameters(
+	FRHIBatchedShaderParameters& BatchedParameters,
+	const TShaderRef<TShaderClass>& Shader,
+	const FShaderParametersMetadata* ParametersMetadata,
+	const typename TShaderClass::FParameters& Parameters)
+{
+	ValidateShaderParameters(Shader, ParametersMetadata, &Parameters);
+	SetShaderParameters(BatchedParameters, Shader->Bindings, ParametersMetadata, &Parameters);
+}
 
-	RHICmdList.SetShaderParameters(ShaderRHI, {}, {}, ResourceSRVsToUnset, BindlessSRVsToUnset);
+template<typename TShaderClass>
+inline void SetShaderParameters(FRHIBatchedShaderParameters& BatchedParameters, const TShaderRef<TShaderClass>& Shader, const typename TShaderClass::FParameters& Parameters)
+{
+	const FShaderParametersMetadata* ParametersMetadata = TShaderClass::FParameters::FTypeInfo::GetStructMetadata();
+	SetShaderParameters(BatchedParameters, Shader, ParametersMetadata, Parameters);
 }
 
 RENDERCORE_API void SetShaderParameters(
diff --git a/Engine/Source/Runtime/RenderCore/Public/ShaderParameterUtils.h b/Engine/Source/Runtime/RenderCore/Public/ShaderParameterUtils.h
--- a/Engine/Source/Runtime/RenderCore/Public/ShaderParameterUtils.h
+++ b/Engine/Source/Runtime/RenderCore/Public/ShaderParameterUtils.h
@@ -170,6 +170,56 @@ inline void SetUAVParameter(FRHIBatchedShaderParameters& BatchedParameters, cons
 	}
 }
 
+inline void SetUniformBufferParameter(FRHIBatchedShaderParameters& BatchedParameters, const FShaderUniformBufferParameter& Parameter, FRHIUniformBuffer* UniformBufferRHI)
+{
+	// This will trigger if the parameter was not serialized
+	checkSlow(Parameter.IsInitialized());
+	// If it is bound, we must set it so something valid
+	checkSlow(!Parameter.IsBound() || UniformBufferRHI);
+	if (Parameter.IsBound())
+	{
+		BatchedParameters.SetShaderUniformBuffer(Parameter.GetBaseIndex(), UniformBufferRHI);
+	}
+}
+
+template<typename TBufferStruct>
+inline void SetUniformBufferParameter(FRHIBatchedShaderParameters& BatchedParameters, const TShaderUniformBufferParameter<TBufferStruct>& Parameter, const TUniformBufferRef<TBufferStruct>& UniformBufferRef)
+{
+	// This will trigger if the parameter was not serialized
+	checkSlow(Parameter.IsInitialized());
+	// If it is bound, we must set it so something valid
+	checkSlow(!Parameter.IsBound() || IsValidRef(UniformBufferRef));
+	if (Parameter.IsBound())
+	{
+		SetUniformBufferParameter(BatchedParameters, Parameter, UniformBufferRef.GetReference());
+	}
+}
+
+template<typename TBufferStruct>
+inline void SetUniformBufferParameter(FRHIBatchedShaderParameters& BatchedParameters, const TShaderUniformBufferParameter<TBufferStruct>& Parameter, const TUniformBuffer<TBufferStruct>& UniformBuffer)
+{
+	// This will trigger if the parameter was not serialized
+	checkSlow(Parameter.IsInitialized());
+	// If it is bound, we must set it so something valid
+	checkSlow(!Parameter.IsBound() || UniformBuffer.GetUniformBufferRHI());
+	if (Parameter.IsBound())
+	{
+		SetUniformBufferParameter(BatchedParameters, Parameter, UniformBuffer.GetUniformBufferRHI());
+	}
+}
+
+template<typename TBufferStruct>
+inline void SetUniformBufferParameterImmediate(FRHIBatchedShaderParameters& BatchedParameters, const TShaderUniformBufferParameter<TBufferStruct>& Parameter, const TBufferStruct& UniformBufferValue)
+{
+	// This will trigger if the parameter was not serialized
+	checkSlow(Parameter.IsInitialized());
+	if (Parameter.IsBound())
+	{
+		FUniformBufferRHIRef UniformBufferRef = RHICreateUniformBuffer(&UniformBufferValue, &TBufferStruct::StaticStructMetadata.GetLayout(), UniformBuffer_SingleDraw);
+		SetUniformBufferParameter(BatchedParameters, Parameter, UniformBufferRef.GetReference());
+	}
+}
+
 /// Utility to set all parameters for a Vertex shader. Requires the shader type to implement SetParameters(FRHIBatchedShaderParameters& BatchedParameters, ...)
 template<typename TCmdList, typename TShaderType, typename... TArguments>
 inline void SetAllShaderParametersVS(TCmdList& RHICmdList, const TShaderRef<TShaderType>& InShader, TArguments&&... InArguments)
diff --git a/Engine/Source/Runtime/Renderer/Private/ShaderBaseClasses.cpp b/Engine/Source/Runtime/Renderer/Private/ShaderBaseClasses.cpp
--- a/Engine/Source/Runtime/Renderer/Private/ShaderBaseClasses.cpp
+++ b/Engine/Source/Runtime/Renderer/Private/ShaderBaseClasses.cpp
@@ -219,10 +219,21 @@ void FMaterialShader::VerifyExpressionAndShaderMaps(const FMaterialRenderProxy*
 }
 #endif
 
-template<typename TRHIShader, typename TRHICommandList>
+void FMaterialShader::SetViewParameters(FRHIBatchedShaderParameters& BatchedParameters, const FSceneView& View, const TUniformBufferRef<FViewUniformShaderParameters>& ViewUniformBuffer)
+{
+	const auto& ViewUniformBufferParameter = GetUniformBufferParameter<FViewUniformShaderParameters>();
+	SetUniformBufferParameter(BatchedParameters, ViewUniformBufferParameter, ViewUniformBuffer);
+
+	if (View.bShouldBindInstancedViewUB)
+	{
+		// When drawing an instanced stereo scene, the instanced view UB should be taken from the same view where it will contains a copy of both left and eye values (see FViewInfo::CreateViewUniformBuffers).
+		const auto& InstancedViewUniformBufferParameter = GetUniformBufferParameter<FInstancedViewUniformShaderParameters>();
+		SetUniformBufferParameter(BatchedParameters, InstancedViewUniformBufferParameter, View.GetInstancedViewUniformBuffer());
+	}
+}
+
 void FMaterialShader::SetParameters(
-	TRHICommandList& RHICmdList,
-	TRHIShader* ShaderRHI,
+	FRHIBatchedShaderParameters& BatchedParameters,
 	const FMaterialRenderProxy* MaterialRenderProxy,
 	const FMaterial& Material,
 	const FSceneView& View)
@@ -267,7 +278,7 @@ void FMaterialShader::SetParameters(
 		MaterialRenderProxy->EvaluateUniformExpressions(*UniformExpressionCache, MaterialRenderContext);
 	}
 
-	SetUniformBufferParameter(RHICmdList, ShaderRHI, MaterialUniformBuffer, UniformExpressionCache->UniformBuffer);
+	SetUniformBufferParameter(BatchedParameters, MaterialUniformBuffer, UniformExpressionCache->UniformBuffer);
 
 #if !(UE_BUILD_TEST || UE_BUILD_SHIPPING || !WITH_EDITOR)
 	VerifyExpressionAndShaderMaps(MaterialRenderProxy, Material, UniformExpressionCache);
@@ -322,7 +333,7 @@ void FMaterialShader::SetParameters(
 					GDefaultMaterialParameterCollectionInstances.Num(), *InstancesString);
 			}
 
-			SetUniformBufferParameter(RHICmdList, ShaderRHI, ParameterCollectionUniformBuffers[CollectionIndex], UniformBuffer);			
+			SetUniformBufferParameter(BatchedParameters, ParameterCollectionUniformBuffers[CollectionIndex], UniformBuffer);
 		}
 	}
 
@@ -332,27 +343,6 @@ void FMaterialShader::SetParameters(
 	}
 }
 
-// Doxygen struggles to parse these explicit specializations. Just ignore them for now.
-#if !UE_BUILD_DOCS
-
-#define IMPLEMENT_MATERIAL_SHADER_SetParameters( TRHICommandList, TRHIShader ) \
-	template RENDERER_API void FMaterialShader::SetParameters< TRHIShader, TRHICommandList >( \
-		TRHICommandList& RHICmdList,				\
-		TRHIShader* ShaderRHI,							\
-		const FMaterialRenderProxy* MaterialRenderProxy,\
-		const FMaterial& Material,						\
-		const FSceneView& View							\
-	);
-
-IMPLEMENT_MATERIAL_SHADER_SetParameters(FRHICommandList, FRHIVertexShader);
-IMPLEMENT_MATERIAL_SHADER_SetParameters(FRHICommandList, FRHIGeometryShader);
-IMPLEMENT_MATERIAL_SHADER_SetParameters(FRHICommandList, FRHIPixelShader);
-IMPLEMENT_MATERIAL_SHADER_SetParameters(FRHICommandList, FRHIMeshShader);
-IMPLEMENT_MATERIAL_SHADER_SetParameters(FRHICommandList, FRHIComputeShader);
-IMPLEMENT_MATERIAL_SHADER_SetParameters(FRHIComputeCommandList, FRHIComputeShader);
-
-#endif
-
 void FMaterialShader::GetShaderBindings(
 	const FScene* Scene,
 	const ERHIFeatureLevel::Type FeatureLevel,
diff --git a/Engine/Source/Runtime/Renderer/Public/MaterialShader.h b/Engine/Source/Runtime/Renderer/Public/MaterialShader.h
--- a/Engine/Source/Runtime/Renderer/Public/MaterialShader.h
+++ b/Engine/Source/Runtime/Renderer/Public/MaterialShader.h
@@ -80,29 +80,37 @@ public:
 
 	FRHIUniformBuffer* GetParameterCollectionBuffer(const FGuid& Id, const FSceneInterface* SceneInterface) const;
 
+	void SetViewParameters(FRHIBatchedShaderParameters& BatchedParameters, const FSceneView& View, const TUniformBufferRef<FViewUniformShaderParameters>& ViewUniformBuffer);
+
 	template<typename ShaderRHIParamRef, typename TRHICommandList>
 	FORCEINLINE_DEBUGGABLE void SetViewParameters(TRHICommandList& RHICmdList, const ShaderRHIParamRef ShaderRHI, const FSceneView& View, const TUniformBufferRef<FViewUniformShaderParameters>& ViewUniformBuffer)
 	{
-		const auto& ViewUniformBufferParameter = GetUniformBufferParameter<FViewUniformShaderParameters>();
-		SetUniformBufferParameter(RHICmdList, ShaderRHI, ViewUniformBufferParameter, ViewUniformBuffer);
-
-		if (View.bShouldBindInstancedViewUB)
-		{
-			// When drawing an instanced stereo scene, the instanced view UB should be taken from the same view where it will contains a copy of both left and eye values (see FViewInfo::CreateViewUniformBuffers).
-			const auto& InstancedViewUniformBufferParameter = GetUniformBufferParameter<FInstancedViewUniformShaderParameters>();
-			SetUniformBufferParameter(RHICmdList, ShaderRHI, InstancedViewUniformBufferParameter, View.GetInstancedViewUniformBuffer());
-		}
+		FRHIBatchedShaderParameters& BatchedParameters = RHICmdList.GetScratchShaderParameters();
+		SetViewParameters(BatchedParameters, View, ViewUniformBuffer);
+		RHICmdList.SetBatchedShaderParameters(ShaderRHI, BatchedParameters);
 	}
 
 	/** Sets pixel parameters that are material specific but not FMeshBatch specific. */
-	template<typename TRHIShader, typename TRHICommandList>
 	void SetParameters(
-		TRHICommandList& RHICmdList,
-		TRHIShader* ShaderRHI,
+		FRHIBatchedShaderParameters& BatchedParameters,
 		const FMaterialRenderProxy* MaterialRenderProxy, 
 		const FMaterial& Material,
 		const FSceneView& View);
 
+	/** Sets pixel parameters that are material specific but not FMeshBatch specific. */
+	template<typename TRHIShader, typename TRHICommandList>
+	inline void SetParameters(
+		TRHICommandList& RHICmdList,
+		TRHIShader* ShaderRHI,
+		const FMaterialRenderProxy* MaterialRenderProxy,
+		const FMaterial& Material,
+		const FSceneView& View)
+	{
+		FRHIBatchedShaderParameters& BatchedParameters = RHICmdList.GetScratchShaderParameters();
+		SetParameters(BatchedParameters, MaterialRenderProxy, Material, View);
+		RHICmdList.SetBatchedShaderParameters(ShaderRHI, BatchedParameters);
+	}
+
 	void GetShaderBindings(
 		const FScene* Scene,
 		ERHIFeatureLevel::Type FeatureLevel,
