Change VT feedback to use StructuredBuffer instead of texel buffer, to avoid limitations on Apple A9 and A10 GPUs which can't write to texel buffers from a fragment shaders

diff --git a/Engine/Shaders/Private/VirtualTextureCommon.ush b/Engine/Shaders/Private/VirtualTextureCommon.ush
--- a/Engine/Shaders/Private/VirtualTextureCommon.ush
+++ b/Engine/Shaders/Private/VirtualTextureCommon.ush
@@ -51,7 +51,7 @@ void StoreVirtualTextureFeedback(in out FVirtualTextureFeedbackParams Params, ui
 }
 
 /** This should be called at the end of the pixel shader to write out the gathered VT feedback info to the OutputBuffer. */
-void FinalizeVirtualTextureFeedback(in FVirtualTextureFeedbackParams Params, float4 SvPosition, float Opacity, uint FrameNumber, RWBuffer<uint> OutputBuffer)
+void FinalizeVirtualTextureFeedback(in FVirtualTextureFeedbackParams Params, float4 SvPosition, float Opacity, uint FrameNumber, RWStructuredBuffer<uint> OutputBuffer)
 {
 	uint2 PixelTilePos = (uint2)SvPosition.xy & View.VirtualTextureFeedbackMask;
 	uint PixelTileIndex = (PixelTilePos.y << View.VirtualTextureFeedbackShift) + PixelTilePos.x;
diff --git a/Engine/Source/Runtime/Engine/Private/SceneManagement.cpp b/Engine/Source/Runtime/Engine/Private/SceneManagement.cpp
--- a/Engine/Source/Runtime/Engine/Private/SceneManagement.cpp
+++ b/Engine/Source/Runtime/Engine/Private/SceneManagement.cpp
@@ -1029,7 +1029,7 @@ FViewUniformShaderParameters::FViewUniformShaderParameters()
 	SSProfilesPreIntegratedTexture = GBlackArrayTexture->TextureRHI;
 	SSProfilesPreIntegratedSampler = TStaticSamplerState<SF_Bilinear, AM_Clamp, AM_Clamp, AM_Clamp>::GetRHI();
 
-	VTFeedbackBuffer = GEmptyVertexBufferWithUAV->UnorderedAccessViewRHI;
+	VTFeedbackBuffer = GEmptyStructuredBufferWithUAV->UnorderedAccessViewRHI;
 //#if WITH_EDITOR
 	EditorVisualizeLevelInstanceIds = GIdentityPrimitiveBuffer.EditorVisualizeLevelInstanceDataBufferSRV;
 	EditorSelectedHitProxyIds = GIdentityPrimitiveBuffer.EditorSelectedDataBufferSRV;
diff --git a/Engine/Source/Runtime/Engine/Public/SceneView.h b/Engine/Source/Runtime/Engine/Public/SceneView.h
--- a/Engine/Source/Runtime/Engine/Public/SceneView.h
+++ b/Engine/Source/Runtime/Engine/Public/SceneView.h
@@ -962,7 +962,7 @@ BEGIN_GLOBAL_SHADER_PARAMETER_STRUCT_WITH_CONSTRUCTOR(FViewUniformShaderParamete
 	SHADER_PARAMETER_SRV(Buffer<uint>, LandscapeIndirection)
 	SHADER_PARAMETER_SRV(Buffer<float>, LandscapePerComponentData)
 
-	SHADER_PARAMETER_UAV(RWBuffer<uint>, VTFeedbackBuffer)
+	SHADER_PARAMETER_UAV(RWStructuredBuffer<uint>, VTFeedbackBuffer)
 	SHADER_PARAMETER_SRV(Buffer<uint>, EditorVisualizeLevelInstanceIds)
 	SHADER_PARAMETER_SRV(Buffer<uint>, EditorSelectedHitProxyIds)
 
diff --git a/Engine/Source/Runtime/RenderCore/Private/GlobalRenderResources.cpp b/Engine/Source/Runtime/RenderCore/Private/GlobalRenderResources.cpp
--- a/Engine/Source/Runtime/RenderCore/Private/GlobalRenderResources.cpp
+++ b/Engine/Source/Runtime/RenderCore/Private/GlobalRenderResources.cpp
@@ -86,6 +86,23 @@ public:
 	}
 };
 
+class FEmptyStructuredBuffer : public FVertexBufferWithSRV
+{
+public:
+	virtual void InitRHI() override
+	{
+		// Create the buffer RHI.  		
+		FRHIResourceCreateInfo CreateInfo(TEXT("EmptyStructuredBuffer"));
+
+		const uint32 BufferSize = sizeof(uint32) * 4u;
+		VertexBufferRHI = RHICreateStructuredBuffer(sizeof(uint32), BufferSize, BUF_Static | BUF_ShaderResource | BUF_UnorderedAccess, CreateInfo);
+
+		// Create a view of the buffer
+		ShaderResourceViewRHI = RHICreateShaderResourceView(VertexBufferRHI);
+		UnorderedAccessViewRHI = RHICreateUnorderedAccessView(VertexBufferRHI, false, false);
+	}
+};
+
 class FBlackTextureWithSRV : public FColoredTexture<0, 0, 0, 255>
 {
 public:
@@ -110,6 +127,7 @@ FTexture* GBlackTexture = GBlackTextureWithSRV;
 FTexture* GTransparentBlackTexture = GTransparentBlackTextureWithSRV;
 
 FVertexBufferWithSRV* GEmptyVertexBufferWithUAV = new TGlobalResource<FEmptyVertexBuffer>;
+FVertexBufferWithSRV* GEmptyStructuredBufferWithUAV = new TGlobalResource<FEmptyStructuredBuffer>;
 
 class FWhiteVertexBuffer : public FVertexBufferWithSRV
 {
diff --git a/Engine/Source/Runtime/RenderCore/Public/GlobalRenderResources.h b/Engine/Source/Runtime/RenderCore/Public/GlobalRenderResources.h
--- a/Engine/Source/Runtime/RenderCore/Public/GlobalRenderResources.h
+++ b/Engine/Source/Runtime/RenderCore/Public/GlobalRenderResources.h
@@ -25,6 +25,9 @@ extern RENDERCORE_API FTextureWithSRV* GTransparentBlackTextureWithSRV;
 // An empty vertex buffer with a UAV
 extern RENDERCORE_API FVertexBufferWithSRV* GEmptyVertexBufferWithUAV;
 
+// An empty structured buffer with a UAV
+extern RENDERCORE_API FVertexBufferWithSRV* GEmptyStructuredBufferWithUAV;
+
 // An empty vertex buffer with a SRV
 extern RENDERCORE_API FVertexBufferWithSRV* GWhiteVertexBufferWithSRV;
 
diff --git a/Engine/Source/Runtime/Renderer/Private/VT/VirtualTextureFeedbackBuffer.cpp b/Engine/Source/Runtime/Renderer/Private/VT/VirtualTextureFeedbackBuffer.cpp
--- a/Engine/Source/Runtime/Renderer/Private/VT/VirtualTextureFeedbackBuffer.cpp
+++ b/Engine/Source/Runtime/Renderer/Private/VT/VirtualTextureFeedbackBuffer.cpp
@@ -77,13 +77,12 @@ void FVirtualTextureFeedbackBuffer::Begin(FRDGBuilder& GraphBuilder, const FVirt
 
 	Desc = InDesc;
 
-	FRDGBufferDesc BufferDesc(FRDGBufferDesc::CreateBufferDesc(sizeof(uint32), Desc.BufferSize.X * Desc.BufferSize.Y));
+	FRDGBufferDesc BufferDesc(FRDGBufferDesc::CreateStructuredDesc(sizeof(uint32), Desc.BufferSize.X * Desc.BufferSize.Y));
 	BufferDesc.Usage |= BUF_SourceCopy;
 
 	if (AllocatePooledBuffer(BufferDesc, PooledBuffer, TEXT("VirtualTextureFeedbackGPU")))
 	{
-		FRDGBufferUAVDesc UAVDesc;
-		UAVDesc.Format = PF_R32_UINT;
+		FRDGBufferUAVDesc UAVDesc{};
 		UAV = PooledBuffer->GetOrCreateUAV(UAVDesc);
 	}
 
@@ -109,7 +108,7 @@ void FVirtualTextureFeedbackBuffer::End(FRDGBuilder& GraphBuilder)
 
 FRHIUnorderedAccessView* FVirtualTextureFeedbackBuffer::GetUAV() const
 {
-	return UAV ? UAV : GEmptyVertexBufferWithUAV->UnorderedAccessViewRHI.GetReference();
+	return UAV ? UAV : GEmptyStructuredBufferWithUAV->UnorderedAccessViewRHI.GetReference();
 }
 
 void FVirtualTextureFeedbackBuffer::ReleaseRHI()
diff --git a/Engine/Source/Runtime/SlateRHIRenderer/Private/SlateRHIRenderingPolicy.cpp b/Engine/Source/Runtime/SlateRHIRenderer/Private/SlateRHIRenderingPolicy.cpp
--- a/Engine/Source/Runtime/SlateRHIRenderer/Private/SlateRHIRenderingPolicy.cpp
+++ b/Engine/Source/Runtime/SlateRHIRenderer/Private/SlateRHIRenderingPolicy.cpp
@@ -222,7 +222,7 @@ static FSceneView* CreateSceneView( FSceneViewFamilyContext* ViewFamilyContext,
 
 	// Create the view's uniform buffer.
 	FViewUniformShaderParameters ViewUniformShaderParameters;
-	ViewUniformShaderParameters.VTFeedbackBuffer = GEmptyVertexBufferWithUAV->UnorderedAccessViewRHI;
+	ViewUniformShaderParameters.VTFeedbackBuffer = GEmptyStructuredBufferWithUAV->UnorderedAccessViewRHI;
 
 	View->SetupCommonViewUniformBufferParameters(
 		ViewUniformShaderParameters,
