Log extra info when failing to decompress the shaders.

diff --git a/Engine/Source/Runtime/RenderCore/Private/ShaderCodeArchive.cpp b/Engine/Source/Runtime/RenderCore/Private/ShaderCodeArchive.cpp
--- a/Engine/Source/Runtime/RenderCore/Private/ShaderCodeArchive.cpp
+++ b/Engine/Source/Runtime/RenderCore/Private/ShaderCodeArchive.cpp
@@ -358,6 +358,25 @@ void ShaderCodeArchive::DecompressShaderWithOodle(uint8* OutDecompressedShader,
 	}
 }
 
+namespace
+{
+	void DecompressShadergroupWithOodleAndExtraLogging(const int32 GroupIndex, const FIoChunkId& GroupHash, FIoStoreShaderGroupEntry& Entry, const int32 ShaderIndex, const uint64 ShaderInGroupIndex, const FSHAHash& ShaderHash, uint8* OutDecompressedShader, int64 UncompressedSize, const uint8* CompressedShaderCode, int64 CompressedSize)
+	{
+		bool bSucceed = FCompression::UncompressMemory(NAME_Oodle, OutDecompressedShader, UncompressedSize, CompressedShaderCode, CompressedSize);
+		if (!bSucceed)
+		{
+			UE_LOG(LogShaderLibrary, Fatal, TEXT("DecompressShaderWithOodleAndExtraLogging(): Could not decompress shader group with Oodle. Group Index: %d Group IoStoreHash:%s Group NumShaders: %d Shader Index: %d Shader In-group Index: %d Shader Hash: %s"),
+				GroupIndex,
+				*LexToString(GroupHash),
+				Entry.NumShaders,
+				ShaderIndex,
+				ShaderInGroupIndex,
+				*LexToString(ShaderHash)
+				);
+		}
+	}
+}
+
 bool ShaderCodeArchive::CompressShaderWithOodle(uint8* OutCompressedShader, int64& OutCompressedSize, const uint8* InUncompressedShaderCode, int64 InUncompressedSize, FOodleDataCompression::ECompressor InOodleCompressor, FOodleDataCompression::ECompressionLevel InOodleLevel)
 {
 	if (OutCompressedShader)
@@ -2052,7 +2071,7 @@ TRefCountPtr<FRHIShader> FIoStoreShaderCodeArchive::CreateShader(int32 ShaderInd
 	if (GroupEntry.IsGroupCompressed())
 	{
 		uint8* UncompressedCode = reinterpret_cast<uint8*>(MemStack.Alloc(GroupEntry.UncompressedSize, 16));
-		ShaderCodeArchive::DecompressShaderWithOodle(UncompressedCode, GroupEntry.UncompressedSize, ShaderCode, GroupEntry.CompressedSize);
+		DecompressShadergroupWithOodleAndExtraLogging(GroupIndex, Header.ShaderGroupIoHashes[GroupIndex], GroupEntry, ShaderIndex, ShaderEntry.ShaderGroupIndex, Header.ShaderHashes[ShaderIndex], UncompressedCode, GroupEntry.UncompressedSize, ShaderCode, GroupEntry.CompressedSize);
 		ShaderCode = reinterpret_cast<uint8*>(UncompressedCode) + ShaderEntry.UncompressedOffsetInGroup;
 
 #if UE_SCA_VISUALIZE_SHADER_USAGE
