Fixes a crash in DumpGPU due to cvar using ECVF_SetByGameOverride

diff --git a/Engine/Source/Runtime/RenderCore/Private/DumpGPU.cpp b/Engine/Source/Runtime/RenderCore/Private/DumpGPU.cpp
--- a/Engine/Source/Runtime/RenderCore/Private/DumpGPU.cpp
+++ b/Engine/Source/Runtime/RenderCore/Private/DumpGPU.cpp
@@ -498,7 +498,7 @@ public:
 				return;
 			}
 
-			const TCHAR* SetByName = TEXT("");
+			const TCHAR* SetByName = TEXT("Unknown");
 			switch (CVarFlags & ECVF_SetByMask)
 			{
 			case ECVF_SetByConstructor:
@@ -531,6 +531,11 @@ public:
 				SetByName = TEXT("DeviceProfile");
 				break;
 			}
+			case ECVF_SetByGameOverride:
+			{
+				SetByName = TEXT("GameOverride");
+				break;
+			}
 			case ECVF_SetByConsoleVariablesIni:
 			{
 				SetByName = TEXT("ConsoleVariablesIni");
@@ -552,7 +557,7 @@ public:
 				break;
 			}
 			default:
-				unimplemented();
+				ensure(false);
 			}
 
 			FString Value = CVar->GetString();
