Properly cook out VT for mobile platforms if VT is not supported on mobile

diff --git a/Engine/Source/Runtime/RenderCore/Private/RenderUtils.cpp b/Engine/Source/Runtime/RenderCore/Private/RenderUtils.cpp
--- a/Engine/Source/Runtime/RenderCore/Private/RenderUtils.cpp
+++ b/Engine/Source/Runtime/RenderCore/Private/RenderUtils.cpp
@@ -1118,8 +1118,12 @@ bool UseVirtualTexturing(bool bIsMobilePlatform, const ITargetPlatform* TargetPl
 	if (CVarVirtualTexture->GetValueOnAnyThread() == 0)
 	{
 		return false;
-	}		
+	}
 
+#if WITH_EDITOR
+	// in case we are cooking for a platform that supports only mobile rendering
+	bIsMobilePlatform |= (TargetPlatform && !TargetPlatform->SupportsFeature(ETargetPlatformFeatures::DeferredRendering));
+#endif
 	// mobile needs an additional switch to enable VT		
 	static const auto CVarMobileVirtualTexture = IConsoleManager::Get().FindTConsoleVariableDataInt(TEXT("r.Mobile.VirtualTextures"));
 	if (bIsMobilePlatform && CVarMobileVirtualTexture->GetValueOnAnyThread() == 0)
