Convert shader type and parameter struct metadata from global variables to static locals to defer initialization. This allows us to create these types dynamically at startup, while keeping the rest of the infrastructure as is.

diff --git a/Engine/Plugins/Experimental/VirtualHeightfieldMesh/Source/VirtualHeightfieldMesh/Private/VirtualHeightfieldMeshVertexFactory.cpp b/Engine/Plugins/Experimental/VirtualHeightfieldMesh/Source/VirtualHeightfieldMesh/Private/VirtualHeightfieldMeshVertexFactory.cpp
--- a/Engine/Plugins/Experimental/VirtualHeightfieldMesh/Source/VirtualHeightfieldMesh/Private/VirtualHeightfieldMeshVertexFactory.cpp
+++ b/Engine/Plugins/Experimental/VirtualHeightfieldMesh/Source/VirtualHeightfieldMesh/Private/VirtualHeightfieldMeshVertexFactory.cpp
@@ -221,7 +221,7 @@ void FVirtualHeightfieldMeshVertexFactory::ValidateCompiledResult(const FVertexF
 #if 0
 	if (Type->SupportsPrimitiveIdStream()
 		&& UseGPUScene(Platform, GetMaxSupportedFeatureLevel(Platform))
-		&& ParameterMap.ContainsParameterAllocation(FPrimitiveUniformShaderParameters::StaticStructMetadata.GetShaderVariableName()))
+		&& ParameterMap.ContainsParameterAllocation(FPrimitiveUniformShaderParameters::FTypeInfo::GetStructMetadata()->GetShaderVariableName()))
 	{
 		OutErrors.AddUnique(*FString::Printf(TEXT("Shader attempted to bind the Primitive uniform buffer even though Vertex Factory %s computes a PrimitiveId per-instance.  This will break auto-instancing.  Shaders should use GetPrimitiveData(Parameters).Member instead of Primitive.Member."), Type->GetName()));
 	}
diff --git a/Engine/Plugins/Experimental/Water/Source/Runtime/Public/WaterVertexFactory.inl b/Engine/Plugins/Experimental/Water/Source/Runtime/Public/WaterVertexFactory.inl
--- a/Engine/Plugins/Experimental/Water/Source/Runtime/Public/WaterVertexFactory.inl
+++ b/Engine/Plugins/Experimental/Water/Source/Runtime/Public/WaterVertexFactory.inl
@@ -172,7 +172,7 @@ void TWaterVertexFactory<bWithWaterSelectionSupport>::ValidateCompiledResult(con
 #if 0
 	if (Type->SupportsPrimitiveIdStream()
 		&& UseGPUScene(Platform, GetMaxSupportedFeatureLevel(Platform))
-		&& ParameterMap.ContainsParameterAllocation(FPrimitiveUniformShaderParameters::StaticStructMetadata.GetShaderVariableName()))
+		&& ParameterMap.ContainsParameterAllocation(FPrimitiveUniformShaderParameters::FTypeInfo::GetStructMetadata()->GetShaderVariableName()))
 	{
 		OutErrors.AddUnique(*FString::Printf(TEXT("Shader attempted to bind the Primitive uniform buffer even though Vertex Factory %s computes a PrimitiveId per-instance.  This will break auto-instancing.  Shaders should use GetPrimitiveData(Parameters).Member instead of Primitive.Member."), Type->GetName()));
 	}
diff --git a/Engine/Plugins/FX/Niagara/Source/NiagaraShader/Private/NiagaraShared.cpp b/Engine/Plugins/FX/Niagara/Source/NiagaraShader/Private/NiagaraShared.cpp
--- a/Engine/Plugins/FX/Niagara/Source/NiagaraShader/Private/NiagaraShared.cpp
+++ b/Engine/Plugins/FX/Niagara/Source/NiagaraShader/Private/NiagaraShared.cpp
@@ -519,7 +519,7 @@ void FNiagaraShaderScript::UpdateCachedData_PostCompile(bool bCalledFromSerializ
 	{
 		for (int32 iPermutation = 0; iPermutation < CachedData.NumPermutations; ++iPermutation)
 		{
-			TNiagaraShaderRef<FShader> Shader = GameThreadShaderMap->GetShader(&FNiagaraShader::StaticType, iPermutation);
+			TNiagaraShaderRef<FShader> Shader = GameThreadShaderMap->GetShader(&FNiagaraShader::GetStaticType(), iPermutation);
 			if (!Shader.IsValid())
 			{
 				CachedData.bIsComplete = 0;
diff --git a/Engine/Plugins/Runtime/HairStrands/Source/HairStrandsCore/Private/HairCardsVertexFactory.cpp b/Engine/Plugins/Runtime/HairStrands/Source/HairStrandsCore/Private/HairCardsVertexFactory.cpp
--- a/Engine/Plugins/Runtime/HairStrands/Source/HairStrandsCore/Private/HairCardsVertexFactory.cpp
+++ b/Engine/Plugins/Runtime/HairStrands/Source/HairStrandsCore/Private/HairCardsVertexFactory.cpp
@@ -247,7 +247,7 @@ void FHairCardsVertexFactory::ValidateCompiledResult(const FVertexFactoryType* T
 	#if 0
 	if (Type->SupportsPrimitiveIdStream() 
 		&& UseGPUScene(Platform, GetMaxSupportedFeatureLevel(Platform)) 
-		&& ParameterMap.ContainsParameterAllocation(FPrimitiveUniformShaderParameters::StaticStructMetadata.GetShaderVariableName()))
+		&& ParameterMap.ContainsParameterAllocation(FPrimitiveUniformShaderParameters::FTypeInfo::GetStructMetadata()->GetShaderVariableName()))
 	{
 		OutErrors.AddUnique(*FString::Printf(TEXT("Shader attempted to bind the Primitive uniform buffer even though Vertex Factory %s computes a PrimitiveId per-instance.  This will break auto-instancing.  Shaders should use GetPrimitiveData(PrimitiveId).Member instead of Primitive.Member."), Type->GetName()));
 	}
diff --git a/Engine/Plugins/Runtime/HairStrands/Source/HairStrandsCore/Private/HairStrandsVertexFactory.cpp b/Engine/Plugins/Runtime/HairStrands/Source/HairStrandsCore/Private/HairStrandsVertexFactory.cpp
--- a/Engine/Plugins/Runtime/HairStrands/Source/HairStrandsCore/Private/HairStrandsVertexFactory.cpp
+++ b/Engine/Plugins/Runtime/HairStrands/Source/HairStrandsCore/Private/HairStrandsVertexFactory.cpp
@@ -195,7 +195,7 @@ void FHairStrandsVertexFactory::ValidateCompiledResult(const FVertexFactoryType*
 	if (Type->SupportsPrimitiveIdStream()
 		&& UseGPUScene(Platform, GetMaxSupportedFeatureLevel(Platform)) 
 		&& !IsMobilePlatform(Platform) // On mobile VS may use PrimtiveUB while GPUScene is enabled
-		&& ParameterMap.ContainsParameterAllocation(FPrimitiveUniformShaderParameters::StaticStructMetadata.GetShaderVariableName()))
+		&& ParameterMap.ContainsParameterAllocation(FPrimitiveUniformShaderParameters::FTypeInfo::GetStructMetadata()->GetShaderVariableName()))
 	{
 		OutErrors.AddUnique(*FString::Printf(TEXT("Shader attempted to bind the Primitive uniform buffer even though Vertex Factory %s computes a PrimitiveId per-instance.  This will break auto-instancing.  Shaders should use GetPrimitiveData(PrimitiveId).Member instead of Primitive.Member."), Type->GetName()));
 	}
diff --git a/Engine/Plugins/Runtime/RenderTrace/Source/Private/RenderTrace.cpp b/Engine/Plugins/Runtime/RenderTrace/Source/Private/RenderTrace.cpp
--- a/Engine/Plugins/Runtime/RenderTrace/Source/Private/RenderTrace.cpp
+++ b/Engine/Plugins/Runtime/RenderTrace/Source/Private/RenderTrace.cpp
@@ -137,7 +137,7 @@ public:
 	FPhysicalMaterialSamplerShader(const FMeshMaterialShaderType::CompiledShaderInitializerType& Initializer)
 		: FMeshMaterialShader(Initializer)
 	{
-		PassUniformBuffer.Bind(Initializer.ParameterMap, FSceneTextureUniformParameters::StaticStructMetadata.GetShaderVariableName());
+		PassUniformBuffer.Bind(Initializer.ParameterMap, FSceneTextureUniformParameters::FTypeInfo::GetStructMetadata()->GetShaderVariableName());
 	}
 
 	static bool ShouldCompilePermutation(const FMeshMaterialShaderPermutationParameters& Parameters)
diff --git a/Engine/Plugins/Runtime/nDisplay/Source/DisplayClusterLightCardEditorShaders/Private/DisplayClusterMeshProjectionRenderer.cpp b/Engine/Plugins/Runtime/nDisplay/Source/DisplayClusterLightCardEditorShaders/Private/DisplayClusterMeshProjectionRenderer.cpp
--- a/Engine/Plugins/Runtime/nDisplay/Source/DisplayClusterLightCardEditorShaders/Private/DisplayClusterMeshProjectionRenderer.cpp
+++ b/Engine/Plugins/Runtime/nDisplay/Source/DisplayClusterLightCardEditorShaders/Private/DisplayClusterMeshProjectionRenderer.cpp
@@ -182,7 +182,7 @@ public:
 	FMeshProjectionVS(const ShaderMetaType::CompiledShaderInitializerType& Initializer)
 		: FMeshMaterialShader(Initializer)
 	{
-		PassUniformBuffer.Bind(Initializer.ParameterMap, FSceneTextureUniformParameters::StaticStructMetadata.GetShaderVariableName());
+		PassUniformBuffer.Bind(Initializer.ParameterMap, FSceneTextureUniformParameters::FTypeInfo::GetStructMetadata()->GetShaderVariableName());
 		UVProjectionIndex.Bind(Initializer.ParameterMap, TEXT("ProjectionParameters_UVProjectionIndex"), SPF_Optional);
 		UVProjectionPlaneSize.Bind(Initializer.ParameterMap, TEXT("ProjectionParameters_UVProjectionPlaneSize"), SPF_Optional);
 		UVProjectionPlaneDistance.Bind(Initializer.ParameterMap, TEXT("ProjectionParameters_UVProjectionPlaneDistance"), SPF_Optional);
@@ -239,7 +239,7 @@ public:
 	FMeshProjectionPS(const ShaderMetaType::CompiledShaderInitializerType& Initializer)
 		: FMeshMaterialShader(Initializer)
 	{
-		PassUniformBuffer.Bind(Initializer.ParameterMap, FSceneTextureUniformParameters::StaticStructMetadata.GetShaderVariableName());
+		PassUniformBuffer.Bind(Initializer.ParameterMap, FSceneTextureUniformParameters::FTypeInfo::GetStructMetadata()->GetShaderVariableName());
 		NormalCorrectionMatrix.Bind(Initializer.ParameterMap, TEXT("NormalCorrectionMatrix"), SPF_Optional);
 	}
 
diff --git a/Engine/Plugins/Runtime/nDisplay/Source/DisplayClusterShaders/Private/Shaders/DisplayClusterShadersPreprocess_UVLightCards.cpp b/Engine/Plugins/Runtime/nDisplay/Source/DisplayClusterShaders/Private/Shaders/DisplayClusterShadersPreprocess_UVLightCards.cpp
--- a/Engine/Plugins/Runtime/nDisplay/Source/DisplayClusterShaders/Private/Shaders/DisplayClusterShadersPreprocess_UVLightCards.cpp
+++ b/Engine/Plugins/Runtime/nDisplay/Source/DisplayClusterShaders/Private/Shaders/DisplayClusterShadersPreprocess_UVLightCards.cpp
@@ -29,7 +29,7 @@ public:
 	FUVLightCardVS(const ShaderMetaType::CompiledShaderInitializerType& Initializer)
 		: FMeshMaterialShader(Initializer)
 	{
-		PassUniformBuffer.Bind(Initializer.ParameterMap, FSceneTextureUniformParameters::StaticStructMetadata.GetShaderVariableName());
+		PassUniformBuffer.Bind(Initializer.ParameterMap, FSceneTextureUniformParameters::FTypeInfo::GetStructMetadata()->GetShaderVariableName());
 	}
 	              
 	static bool ShouldCompilePermutation(const FMeshMaterialShaderPermutationParameters& Parameters)
@@ -51,7 +51,7 @@ public:
 	FUVLightCardPS(const ShaderMetaType::CompiledShaderInitializerType& Initializer)
 		: FMeshMaterialShader(Initializer)
 	{
-		PassUniformBuffer.Bind(Initializer.ParameterMap, FSceneTextureUniformParameters::StaticStructMetadata.GetShaderVariableName());
+		PassUniformBuffer.Bind(Initializer.ParameterMap, FSceneTextureUniformParameters::FTypeInfo::GetStructMetadata()->GetShaderVariableName());
 	}
 
 	static bool ShouldCompilePermutation(const FMeshMaterialShaderPermutationParameters& Parameters)
diff --git a/Engine/Source/Developer/LowLevelTestsRunner/Private/TestCommon/EngineUtilities.cpp b/Engine/Source/Developer/LowLevelTestsRunner/Private/TestCommon/EngineUtilities.cpp
--- a/Engine/Source/Developer/LowLevelTestsRunner/Private/TestCommon/EngineUtilities.cpp
+++ b/Engine/Source/Developer/LowLevelTestsRunner/Private/TestCommon/EngineUtilities.cpp
@@ -33,6 +33,8 @@ void InitAsyncQueues()
 
 void InitRendering()
 {
+	FShaderParametersMetadataRegistration::CommitAll();
+	FShaderTypeRegistration::CommitAll();
 	FShaderParametersMetadata::InitializeAllUniformBufferStructs();
 
 	{
diff --git a/Engine/Source/Runtime/Engine/Private/GPUSkinVertexFactory.cpp b/Engine/Source/Runtime/Engine/Private/GPUSkinVertexFactory.cpp
--- a/Engine/Source/Runtime/Engine/Private/GPUSkinVertexFactory.cpp
+++ b/Engine/Source/Runtime/Engine/Private/GPUSkinVertexFactory.cpp
@@ -370,7 +370,7 @@ bool FGPUBaseSkinVertexFactory::FShaderDataType::UpdateBoneData(FRHICommandListI
 	}
 	else
 	{
-		UniformBuffer = RHICreateUniformBuffer(&GBoneUniformStruct, FBoneMatricesUniformShaderParameters::StaticStructMetadata.GetLayoutPtr(), UniformBuffer_MultiFrame);
+		UniformBuffer = RHICreateUniformBuffer(&GBoneUniformStruct, FBoneMatricesUniformShaderParameters::FTypeInfo::GetStructMetadata()->GetLayoutPtr(), UniformBuffer_MultiFrame);
 	}
 	return false;
 }
diff --git a/Engine/Source/Runtime/Engine/Private/GeometryCollectionRendering.cpp b/Engine/Source/Runtime/Engine/Private/GeometryCollectionRendering.cpp
--- a/Engine/Source/Runtime/Engine/Private/GeometryCollectionRendering.cpp
+++ b/Engine/Source/Runtime/Engine/Private/GeometryCollectionRendering.cpp
@@ -78,7 +78,7 @@ void FGeometryCollectionVertexFactory::ValidateCompiledResult(const FVertexFacto
 	if (Type->SupportsPrimitiveIdStream()
 		&& UseGPUScene(Platform, GetMaxSupportedFeatureLevel(Platform))
 		&& !IsMobilePlatform(Platform) // On mobile VS may use PrimtiveUB while GPUScene is enabled
-		&& ParameterMap.ContainsParameterAllocation(FPrimitiveUniformShaderParameters::StaticStructMetadata.GetShaderVariableName()))
+		&& ParameterMap.ContainsParameterAllocation(FPrimitiveUniformShaderParameters::FTypeInfo::GetStructMetadata()->GetShaderVariableName()))
 	{
 		OutErrors.AddUnique(*FString::Printf(TEXT("Shader attempted to bind the Primitive uniform buffer even though Vertex Factory %s computes a PrimitiveId per-instance.  This will break auto-instancing.  Shaders should use GetPrimitiveData(Parameters).Member instead of Primitive.Member."), Type->GetName()));
 	}
diff --git a/Engine/Source/Runtime/Engine/Private/LocalVertexFactory.cpp b/Engine/Source/Runtime/Engine/Private/LocalVertexFactory.cpp
--- a/Engine/Source/Runtime/Engine/Private/LocalVertexFactory.cpp
+++ b/Engine/Source/Runtime/Engine/Private/LocalVertexFactory.cpp
@@ -322,7 +322,7 @@ void FLocalVertexFactory::ValidateCompiledResult(const FVertexFactoryType* Type,
 	if (Type->SupportsPrimitiveIdStream() 
 		&& UseGPUScene(Platform, GetMaxSupportedFeatureLevel(Platform)) 
 		&& !IsMobilePlatform(Platform) // On mobile VS may use PrimtiveUB while GPUScene is enabled
-		&& ParameterMap.ContainsParameterAllocation(FPrimitiveUniformShaderParameters::StaticStructMetadata.GetShaderVariableName()))
+		&& ParameterMap.ContainsParameterAllocation(FPrimitiveUniformShaderParameters::FTypeInfo::GetStructMetadata()->GetShaderVariableName()))
 	{
 		OutErrors.AddUnique(*FString::Printf(TEXT("Shader attempted to bind the Primitive uniform buffer even though Vertex Factory %s computes a PrimitiveId per-instance.  This will break auto-instancing.  Shaders should use GetPrimitiveData(Parameters).Member instead of Primitive.Member."), Type->GetName()));
 	}
diff --git a/Engine/Source/Runtime/Engine/Public/MaterialShaderType.h b/Engine/Source/Runtime/Engine/Public/MaterialShaderType.h
--- a/Engine/Source/Runtime/Engine/Public/MaterialShaderType.h
+++ b/Engine/Source/Runtime/Engine/Public/MaterialShaderType.h
@@ -219,7 +219,7 @@ struct FMaterialShaderTypes
 	template<typename ShaderType>
 	inline const FShaderType* AddShaderType(int32 InPermutationId = 0)
 	{
-		return AddShaderType(&ShaderType::StaticType, InPermutationId);
+		return AddShaderType(&ShaderType::GetStaticType(), InPermutationId);
 	}
 };
 
diff --git a/Engine/Source/Runtime/Engine/Public/MaterialShared.h b/Engine/Source/Runtime/Engine/Public/MaterialShared.h
--- a/Engine/Source/Runtime/Engine/Public/MaterialShared.h
+++ b/Engine/Source/Runtime/Engine/Public/MaterialShared.h
@@ -1211,11 +1211,11 @@ public:
 	}
 	template<typename ShaderType> TShaderRef<ShaderType> GetShader(int32 PermutationId = 0) const
 	{
-		return TShaderRef<ShaderType>::Cast(GetShader(&ShaderType::StaticType, PermutationId));
+		return TShaderRef<ShaderType>::Cast(GetShader(&ShaderType::GetStaticType(), PermutationId));
 	}
 	template<typename ShaderType> TShaderRef<ShaderType> GetShader(const typename ShaderType::FPermutationDomain& PermutationVector) const
 	{
-		return TShaderRef<ShaderType>::Cast(GetShader(&ShaderType::StaticType, PermutationVector.ToDimensionValueId()));
+		return TShaderRef<ShaderType>::Cast(GetShader(&ShaderType::GetStaticType(), PermutationVector.ToDimensionValueId()));
 	}
 
 	uint32 GetMaxNumInstructionsForShader(FShaderType* ShaderType) const { return GetContent()->GetMaxNumInstructionsForShader(*this, ShaderType); }
@@ -2145,7 +2145,7 @@ public:
 	template <typename ShaderType>
 	TShaderRef<ShaderType> GetShader(FVertexFactoryType* VertexFactoryType, int32 PermutationId = 0, bool bFatalIfMissing = true) const
 	{
-		return TShaderRef<ShaderType>::Cast(GetShader(&ShaderType::StaticType, VertexFactoryType, PermutationId, bFatalIfMissing));
+		return TShaderRef<ShaderType>::Cast(GetShader(&ShaderType::GetStaticType(), VertexFactoryType, PermutationId, bFatalIfMissing));
 	}
 
 	ENGINE_API FShaderPipelineRef GetShaderPipeline(class FShaderPipelineType* ShaderPipelineType, FVertexFactoryType* VertexFactoryType, bool bFatalIfNotFound = true) const;
diff --git a/Engine/Source/Runtime/Engine/Public/PrimitiveUniformShaderParameters.h b/Engine/Source/Runtime/Engine/Public/PrimitiveUniformShaderParameters.h
--- a/Engine/Source/Runtime/Engine/Public/PrimitiveUniformShaderParameters.h
+++ b/Engine/Source/Runtime/Engine/Public/PrimitiveUniformShaderParameters.h
@@ -557,9 +557,7 @@ inline FPrimitiveUniformShaderParameters GetIdentityPrimitiveParameters()
 class FIdentityPrimitiveUniformBuffer : public TUniformBuffer<FPrimitiveUniformShaderParameters>
 {
 public:
-
-	/** Default constructor. */
-	FIdentityPrimitiveUniformBuffer()
+	void InitContents()
 	{
 		SetContents(GetIdentityPrimitiveParameters());
 	}
diff --git a/Engine/Source/Runtime/Landscape/Private/LandscapeGrass.cpp b/Engine/Source/Runtime/Landscape/Private/LandscapeGrass.cpp
--- a/Engine/Source/Runtime/Landscape/Private/LandscapeGrass.cpp
+++ b/Engine/Source/Runtime/Landscape/Private/LandscapeGrass.cpp
@@ -312,7 +312,7 @@ protected:
 	: FMeshMaterialShader(Initializer)
 	{
 		RenderOffsetParameter.Bind(Initializer.ParameterMap, TEXT("RenderOffset"));
-		PassUniformBuffer.Bind(Initializer.ParameterMap, FSceneTextureUniformParameters::StaticStructMetadata.GetShaderVariableName());
+		PassUniformBuffer.Bind(Initializer.ParameterMap, FSceneTextureUniformParameters::FTypeInfo::GetStructMetadata()->GetShaderVariableName());
 	}
 
 public:
@@ -355,7 +355,7 @@ public:
 	: FMeshMaterialShader(Initializer)
 	{
 		OutputPassParameter.Bind(Initializer.ParameterMap, TEXT("OutputPass"));
-		PassUniformBuffer.Bind(Initializer.ParameterMap, FSceneTextureUniformParameters::StaticStructMetadata.GetShaderVariableName());
+		PassUniformBuffer.Bind(Initializer.ParameterMap, FSceneTextureUniformParameters::FTypeInfo::GetStructMetadata()->GetShaderVariableName());
 	}
 
 	FLandscapeGrassWeightPS()
diff --git a/Engine/Source/Runtime/Landscape/Private/LandscapePhysicalMaterial.cpp b/Engine/Source/Runtime/Landscape/Private/LandscapePhysicalMaterial.cpp
--- a/Engine/Source/Runtime/Landscape/Private/LandscapePhysicalMaterial.cpp
+++ b/Engine/Source/Runtime/Landscape/Private/LandscapePhysicalMaterial.cpp
@@ -42,7 +42,7 @@ public:
 	FLandscapePhysicalMaterial(const FMeshMaterialShaderType::CompiledShaderInitializerType& Initializer)
 		: FMeshMaterialShader(Initializer)
 	{
-		PassUniformBuffer.Bind(Initializer.ParameterMap, FSceneTextureUniformParameters::StaticStructMetadata.GetShaderVariableName());
+		PassUniformBuffer.Bind(Initializer.ParameterMap, FSceneTextureUniformParameters::FTypeInfo::GetStructMetadata()->GetShaderVariableName());
 	}
 
 	static bool ShouldCompilePermutation(const FMeshMaterialShaderPermutationParameters& Parameters)
diff --git a/Engine/Source/Runtime/Launch/Private/LaunchEngineLoop.cpp b/Engine/Source/Runtime/Launch/Private/LaunchEngineLoop.cpp
--- a/Engine/Source/Runtime/Launch/Private/LaunchEngineLoop.cpp
+++ b/Engine/Source/Runtime/Launch/Private/LaunchEngineLoop.cpp
@@ -3188,6 +3188,16 @@ PRAGMA_ENABLE_DEPRECATION_WARNINGS
 	});
 #endif	// USE_LOCALIZED_PACKAGE_CACHE
 
+	{
+		SCOPED_BOOT_TIMING("FShaderParametersMetadataRegistration::CommitAll()");
+		FShaderParametersMetadataRegistration::CommitAll();
+	}
+	
+	{
+		SCOPED_BOOT_TIMING("FShaderTypeRegistration::CommitAll()");
+		FShaderTypeRegistration::CommitAll();
+	}
+
 	{
 		SCOPED_BOOT_TIMING("FUniformBufferStruct::InitializeStructs()");
 		FShaderParametersMetadata::InitializeAllUniformBufferStructs();
diff --git a/Engine/Source/Runtime/RenderCore/Private/OneColorShader.cpp b/Engine/Source/Runtime/RenderCore/Private/OneColorShader.cpp
--- a/Engine/Source/Runtime/RenderCore/Private/OneColorShader.cpp
+++ b/Engine/Source/Runtime/RenderCore/Private/OneColorShader.cpp
@@ -5,7 +5,7 @@
 
 // #define avoids a lot of code duplication
 #define IMPLEMENT_ONECOLORVS(A,B) typedef TOneColorVS<A,B> TOneColorVS##A##B; \
-	IMPLEMENT_SHADER_TYPE2_WITH_TEMPLATE_PREFIX(template<> RENDERCORE_API, TOneColorVS##A##B, SF_Vertex);
+	IMPLEMENT_SHADER_TYPE4_WITH_TEMPLATE_PREFIX(template<>, RENDERCORE_API, TOneColorVS##A##B, SF_Vertex);
 
 IMPLEMENT_ONECOLORVS(false,false)
 IMPLEMENT_ONECOLORVS(false,true)
diff --git a/Engine/Source/Runtime/RenderCore/Private/Shader.cpp b/Engine/Source/Runtime/RenderCore/Private/Shader.cpp
--- a/Engine/Source/Runtime/RenderCore/Private/Shader.cpp
+++ b/Engine/Source/Runtime/RenderCore/Private/Shader.cpp
@@ -199,8 +199,8 @@ bool FShaderType::bInitializedSerializationHistory = false;
 
 static TArray<FShaderType*>& GetSortedShaderTypes(FShaderType::EShaderTypeForDynamicCast Type)
 {
-	static TArray<FShaderType*> SortedTypes[(uint32)FShaderType::EShaderTypeForDynamicCast::NumShaderTypes];
-	return SortedTypes[(uint32)Type];
+	static TArray<FShaderType*>* SortedTypesArray = new TArray<FShaderType*>[(uint32)FShaderType::EShaderTypeForDynamicCast::NumShaderTypes];
+	return SortedTypesArray[(uint32)Type];
 }
 
 
@@ -292,6 +292,17 @@ FShaderType::~FShaderType()
 	SortedTypes.RemoveAt(SortedIndex);
 }
 
+TArray<const FShaderTypeRegistration*> FShaderTypeRegistration::Instances;
+
+void FShaderTypeRegistration::CommitAll()
+{
+	for (const auto& Instance : Instances)
+	{
+		FShaderType& ShaderType = Instance->LazyShaderTypeAccessor(); // constructs and registers type
+	}
+	Instances.Empty();
+}
+
 TLinkedList<FShaderType*>*& FShaderType::GetTypeList()
 {
 	return GShaderTypeList;
diff --git a/Engine/Source/Runtime/RenderCore/Private/ShaderParameterMetadata.cpp b/Engine/Source/Runtime/RenderCore/Private/ShaderParameterMetadata.cpp
--- a/Engine/Source/Runtime/RenderCore/Private/ShaderParameterMetadata.cpp
+++ b/Engine/Source/Runtime/RenderCore/Private/ShaderParameterMetadata.cpp
@@ -10,6 +10,8 @@
 #include "ShaderParameters.h"
 #include "DataDrivenShaderPlatformInfo.h"
 #include "ShaderParameterMacros.h"
+#include "Interfaces/IPluginManager.h"
+#include "Misc/CommandLine.h"
 
 bool SupportShaderPrecisionModifier(EShaderPlatform Platform)
 {
@@ -121,6 +123,22 @@ FShaderParametersMetadata* FindUniformBufferStructByLayoutHash(uint32 Hash)
 	return GetLayoutHashStructMap().FindRef(Hash);
 }
 
+TArray<const FShaderParametersMetadataRegistration*> FShaderParametersMetadataRegistration::Instances;
+void FShaderParametersMetadataRegistration::CommitAll()
+{
+	for (const auto* Instance : Instances)
+	{
+		Instance->LazyShaderParametersMetadataAccessor();
+	}
+	Instances.Empty();
+}
+
+bool FShaderParametersMetadataRegistration::IsReadyForRegistration()
+{
+	return FCommandLine::IsInitialized() /* If cmd is not ready yet, then it's too early for the plugin manager */ 
+		&& IPluginManager::Get().GetLastCompletedLoadingPhase() >= ELoadingPhase::PostConfigInit;
+}
+
 const TCHAR* const kShaderParameterMacroNames[] = {
 	TEXT(""), // UBMT_INVALID,
 
diff --git a/Engine/Source/Runtime/RenderCore/Public/GlobalShader.h b/Engine/Source/Runtime/RenderCore/Public/GlobalShader.h
--- a/Engine/Source/Runtime/RenderCore/Public/GlobalShader.h
+++ b/Engine/Source/Runtime/RenderCore/Public/GlobalShader.h
@@ -200,15 +200,15 @@ public:
 	template<typename ShaderType>
 	TShaderRef<ShaderType> GetShader() const
 	{
-		ensureMsgf(ShaderType::StaticType.GetPermutationCount() == 1, TEXT("Failed to provide PermutationId for shader type %s with %u permutations"), ShaderType::StaticType.GetName(), ShaderType::StaticType.GetPermutationCount());
+		ensureMsgf(ShaderType::GetStaticType().GetPermutationCount() == 1, TEXT("Failed to provide PermutationId for shader type %s with %u permutations"), ShaderType::GetStaticType().GetName(), ShaderType::GetStaticType().GetPermutationCount());
 		return GetShader<ShaderType>(0);
 	}
 
 	template<typename ShaderType>
 	TShaderRef<ShaderType> GetShader(int32 PermutationId) const
 	{
-		TShaderRef<FShader> Shader = GetShader(&ShaderType::StaticType, PermutationId);
-		checkf(Shader.IsValid(), TEXT("Failed to find shader type %s in Platform %s"), ShaderType::StaticType.GetName(), *LegacyShaderPlatformToShaderFormat(Platform).ToString());
+		TShaderRef<FShader> Shader = GetShader(&ShaderType::GetStaticType(), PermutationId);
+		checkf(Shader.IsValid(), TEXT("Failed to find shader type %s in Platform %s"), ShaderType::GetStaticType().GetName(), *LegacyShaderPlatformToShaderFormat(Platform).ToString());
 		return TShaderRef<ShaderType>::Cast(Shader);
 	}
 
diff --git a/Engine/Source/Runtime/RenderCore/Public/RenderGraphBuilder.inl b/Engine/Source/Runtime/RenderCore/Public/RenderGraphBuilder.inl
--- a/Engine/Source/Runtime/RenderCore/Public/RenderGraphBuilder.inl
+++ b/Engine/Source/Runtime/RenderCore/Public/RenderGraphBuilder.inl
@@ -217,8 +217,8 @@ FORCEINLINE FRDGSubresourceState* FRDGBuilder::AllocSubresource(const FRDGSubres
 template <typename ParameterStructType>
 TRDGUniformBufferRef<ParameterStructType> FRDGBuilder::CreateUniformBuffer(const ParameterStructType* ParameterStruct)
 {
-	IF_RDG_ENABLE_DEBUG(UserValidation.ValidateCreateUniformBuffer(ParameterStruct, &ParameterStructType::StaticStructMetadata));
-	auto* UniformBuffer = UniformBuffers.Allocate<TRDGUniformBuffer<ParameterStructType>>(Allocator, ParameterStruct, ParameterStructType::StaticStructMetadata.GetShaderVariableName());
+	IF_RDG_ENABLE_DEBUG(UserValidation.ValidateCreateUniformBuffer(ParameterStruct, ParameterStructType::FTypeInfo::GetStructMetadata()));
+	auto* UniformBuffer = UniformBuffers.Allocate<TRDGUniformBuffer<ParameterStructType>>(Allocator, ParameterStruct, ParameterStructType::FTypeInfo::GetStructMetadata()->GetShaderVariableName());
 	IF_RDG_ENABLE_DEBUG(UserValidation.ValidateCreateUniformBuffer(UniformBuffer));
 	return UniformBuffer;
 }
diff --git a/Engine/Source/Runtime/RenderCore/Public/Shader.h b/Engine/Source/Runtime/RenderCore/Public/Shader.h
--- a/Engine/Source/Runtime/RenderCore/Public/Shader.h
+++ b/Engine/Source/Runtime/RenderCore/Public/Shader.h
@@ -911,7 +911,7 @@ public:
 	template<typename UniformBufferStructType>
 	FORCEINLINE_DEBUGGABLE const TShaderUniformBufferParameter<UniformBufferStructType>& GetUniformBufferParameter() const
 	{
-		const FShaderUniformBufferParameter& FoundParameter = GetUniformBufferParameter(&UniformBufferStructType::StaticStructMetadata);
+		const FShaderUniformBufferParameter& FoundParameter = GetUniformBufferParameter(UniformBufferStructType::FTypeInfo::GetStructMetadata());
 		return static_cast<const TShaderUniformBufferParameter<UniformBufferStructType>&>(FoundParameter);
 	}
 
@@ -1497,6 +1497,31 @@ protected:
 #endif // WITH_EDITOR
 };
 
+/**
+ * Registers a shader type in various systems. Should be created as a static field/global.
+ * 
+ * Each shader type is collected here, not as an instance but as an accessor, so the actual construction can be deferred.
+ * The collection happens during static init, the actual construction happens later during launch.
+ * The purpose of collecting the types is the CommitAll function, called in LaunchEngineLoop, to ensure all type instances are constructed and registered before other systems start iterating them.
+ */
+class RENDERCORE_API FShaderTypeRegistration
+{
+public:
+	FShaderTypeRegistration(TFunctionRef<FShaderType& ()> LazyShaderTypeAccessor)
+		: LazyShaderTypeAccessor(LazyShaderTypeAccessor)
+	{
+		Instances.Add(this);
+	}
+
+	// Actually register all the types and clear the array
+	static void CommitAll();
+
+private:
+	static TArray<const FShaderTypeRegistration*> Instances;
+
+	TFunctionRef<FShaderType& ()> LazyShaderTypeAccessor;
+};
+
 struct FShaderCompiledShaderInitializerType
 {
 	const FShaderType* Type;
@@ -1550,7 +1575,10 @@ struct FShaderCompiledShaderInitializerType
 	using ShaderMetaType = F##ShaderMetaTypeShortcut##ShaderType; \
 	using ShaderMapType = F##ShaderMetaTypeShortcut##ShaderMap; \
 	\
-	static RequiredAPI ShaderMetaType StaticType; \
+	static RequiredAPI ShaderMetaType& GetStaticType(); \
+	private: \
+	static FShaderTypeRegistration ShaderTypeRegistration; \
+	public: \
 	\
 	SHADER_DECLARE_VTABLE(ShaderClass)
 
@@ -1593,49 +1621,64 @@ struct FShaderCompiledShaderInitializerType
 #define IMPLEMENT_SHADER_TYPE(TemplatePrefix,ShaderClass,SourceFilename,FunctionName,Frequency) \
 	IMPLEMENT_UNREGISTERED_TEMPLATE_TYPE_LAYOUT(TemplatePrefix, ShaderClass); \
 	TemplatePrefix \
-	ShaderClass::ShaderMetaType ShaderClass::StaticType( \
-		ShaderClass::StaticGetTypeLayout(), \
-		TEXT(#ShaderClass), \
-		SourceFilename, \
-		FunctionName, \
-		Frequency, \
-		ShaderClass::FPermutationDomain::PermutationCount, \
-		SHADER_TYPE_VTABLE(ShaderClass), \
-		sizeof(ShaderClass), \
-		ShaderClass::GetRootParametersMetadata() \
-		);
+	ShaderClass::ShaderMetaType& ShaderClass::GetStaticType() \
+	{ \
+		static ShaderClass::ShaderMetaType StaticType( \
+			ShaderClass::StaticGetTypeLayout(), \
+			TEXT(#ShaderClass), \
+			SourceFilename, \
+			FunctionName, \
+			Frequency, \
+			ShaderClass::FPermutationDomain::PermutationCount, \
+			SHADER_TYPE_VTABLE(ShaderClass), \
+			sizeof(ShaderClass), \
+			ShaderClass::GetRootParametersMetadata() \
+		); \
+		return StaticType; \
+	} \
+	TemplatePrefix FShaderTypeRegistration ShaderClass::ShaderTypeRegistration{TFunctionRef<::FShaderType&()>{ShaderClass::GetStaticType}};
 
 /** A macro to implement a shader type. Shader name is got from GetDebugName(), which is helpful for templated shaders. */
 #define IMPLEMENT_SHADER_TYPE_WITH_DEBUG_NAME(TemplatePrefix,ShaderClass,SourceFilename,FunctionName,Frequency) \
 	IMPLEMENT_UNREGISTERED_TEMPLATE_TYPE_LAYOUT(TemplatePrefix, ShaderClass); \
 	TemplatePrefix \
-	typename ShaderClass::ShaderMetaType ShaderClass::StaticType( \
-		ShaderClass::StaticGetTypeLayout(), \
-		ShaderClass::GetDebugName(), \
-		SourceFilename, \
-		FunctionName, \
-		Frequency, \
-		ShaderClass::FPermutationDomain::PermutationCount, \
-		SHADER_TYPE_VTABLE(ShaderClass), \
-		sizeof(ShaderClass), \
-		ShaderClass::GetRootParametersMetadata() \
-		);
+	typename ShaderClass::ShaderMetaType& ShaderClass::GetStaticType() \
+	{ \
+		static typename ShaderClass::ShaderMetaType StaticType( \
+			ShaderClass::StaticGetTypeLayout(), \
+			ShaderClass::GetDebugName(), \
+			SourceFilename, \
+			FunctionName, \
+			Frequency, \
+			ShaderClass::FPermutationDomain::PermutationCount, \
+			SHADER_TYPE_VTABLE(ShaderClass), \
+			sizeof(ShaderClass), \
+			ShaderClass::GetRootParametersMetadata() \
+		); \
+		return StaticType; \
+	} \
+	TemplatePrefix FShaderTypeRegistration ShaderClass::ShaderTypeRegistration{TFunctionRef<::FShaderType&()>{ShaderClass::GetStaticType}};
 
 /** A macro to implement a templated shader type, the function name and the source filename comes from the class. */
 #define IMPLEMENT_SHADER_TYPE2_WITH_TEMPLATE_PREFIX(TemplatePrefix,ShaderClass,Frequency) \
 	IMPLEMENT_UNREGISTERED_TEMPLATE_TYPE_LAYOUT(TemplatePrefix, ShaderClass); \
 	TemplatePrefix \
-	ShaderClass::ShaderMetaType ShaderClass::StaticType( \
-	ShaderClass::StaticGetTypeLayout(), \
-	TEXT(#ShaderClass), \
-	ShaderClass::GetSourceFilename(), \
-	ShaderClass::GetFunctionName(), \
-	Frequency, \
-	ShaderClass::FPermutationDomain::PermutationCount, \
-	SHADER_TYPE_VTABLE(ShaderClass), \
-	sizeof(ShaderClass), \
-	ShaderClass::GetRootParametersMetadata() \
-	);
+	ShaderClass::ShaderMetaType& ShaderClass::GetStaticType() \
+	{ \
+		static ShaderClass::ShaderMetaType StaticType( \
+			ShaderClass::StaticGetTypeLayout(), \
+			TEXT(#ShaderClass), \
+			ShaderClass::GetSourceFilename(), \
+			ShaderClass::GetFunctionName(), \
+			Frequency, \
+			ShaderClass::FPermutationDomain::PermutationCount, \
+			SHADER_TYPE_VTABLE(ShaderClass), \
+			sizeof(ShaderClass), \
+			ShaderClass::GetRootParametersMetadata() \
+		); \
+		return StaticType; \
+	} \
+	TemplatePrefix FShaderTypeRegistration ShaderClass::ShaderTypeRegistration{TFunctionRef<::FShaderType&()>{ShaderClass::GetStaticType}};
 
 #define IMPLEMENT_SHADER_TYPE2(ShaderClass,Frequency) \
 	IMPLEMENT_SHADER_TYPE2_WITH_TEMPLATE_PREFIX(template<>, ShaderClass, Frequency)
@@ -1643,33 +1686,43 @@ struct FShaderCompiledShaderInitializerType
 /** todo: this should replace IMPLEMENT_SHADER_TYPE */
 #define IMPLEMENT_SHADER_TYPE3(ShaderClass,Frequency) \
 	IMPLEMENT_UNREGISTERED_TEMPLATE_TYPE_LAYOUT(,ShaderClass); \
-	ShaderClass::ShaderMetaType ShaderClass::StaticType( \
-	ShaderClass::StaticGetTypeLayout(), \
-	TEXT(#ShaderClass), \
-	ShaderClass::GetSourceFilename(), \
-	ShaderClass::GetFunctionName(), \
-	Frequency, \
-	ShaderClass::FPermutationDomain::PermutationCount, \
-	SHADER_TYPE_VTABLE(ShaderClass), \
-	sizeof(ShaderClass), \
-	ShaderClass::GetRootParametersMetadata() \
-	);
+	ShaderClass::ShaderMetaType& ShaderClass::GetStaticType() \
+	{ \
+		static ShaderClass::ShaderMetaType StaticType( \
+			ShaderClass::StaticGetTypeLayout(), \
+			TEXT(#ShaderClass), \
+			ShaderClass::GetSourceFilename(), \
+			ShaderClass::GetFunctionName(), \
+			Frequency, \
+			ShaderClass::FPermutationDomain::PermutationCount, \
+			SHADER_TYPE_VTABLE(ShaderClass), \
+			sizeof(ShaderClass), \
+			ShaderClass::GetRootParametersMetadata() \
+		); \
+		return StaticType; \
+	} \
+	FShaderTypeRegistration ShaderClass::ShaderTypeRegistration{TFunctionRef<::FShaderType&()>{ShaderClass::GetStaticType}};
 #endif // !UE_BUILD_DOCS
 
 #define IMPLEMENT_SHADER_TYPE4_WITH_TEMPLATE_PREFIX(TemplatePrefix,RequiredAPI,ShaderClass,Frequency) \
 	IMPLEMENT_UNREGISTERED_TEMPLATE_TYPE_LAYOUT(TemplatePrefix, ShaderClass); \
 	TemplatePrefix RequiredAPI \
-	ShaderClass::ShaderMetaType ShaderClass::StaticType( \
-	ShaderClass::StaticGetTypeLayout(), \
-	TEXT(#ShaderClass), \
-	ShaderClass::GetSourceFilename(), \
-	ShaderClass::GetFunctionName(), \
-	Frequency, \
-	ShaderClass::FPermutationDomain::PermutationCount, \
-	SHADER_TYPE_VTABLE(ShaderClass), \
-	sizeof(ShaderClass), \
-	ShaderClass::GetRootParametersMetadata() \
-	);
+	ShaderClass::ShaderMetaType& ShaderClass::GetStaticType() \
+	{ \
+		static ShaderClass::ShaderMetaType StaticType( \
+			ShaderClass::StaticGetTypeLayout(), \
+			TEXT(#ShaderClass), \
+			ShaderClass::GetSourceFilename(), \
+			ShaderClass::GetFunctionName(), \
+			Frequency, \
+			ShaderClass::FPermutationDomain::PermutationCount, \
+			SHADER_TYPE_VTABLE(ShaderClass), \
+			sizeof(ShaderClass), \
+			ShaderClass::GetRootParametersMetadata() \
+		); \
+		return StaticType; \
+	} \
+	TemplatePrefix FShaderTypeRegistration ShaderClass::ShaderTypeRegistration{TFunctionRef<::FShaderType&()>{ShaderClass::GetStaticType}};
 
 // Binding of a set of shader stages in a single pipeline
 class RENDERCORE_API FShaderPipelineType
@@ -1759,23 +1812,23 @@ protected:
 #if !UE_BUILD_DOCS
 // Vertex+Pixel
 #define IMPLEMENT_SHADERPIPELINE_TYPE_VSPS(PipelineName, VertexShaderType, PixelShaderType, bRemoveUnused)	\
-	static FShaderPipelineType PipelineName(TEXT(PREPROCESSOR_TO_STRING(PipelineName)), &VertexShaderType::StaticType, nullptr, &PixelShaderType::StaticType, false, bRemoveUnused);
+	static FShaderPipelineType PipelineName(TEXT(PREPROCESSOR_TO_STRING(PipelineName)), &VertexShaderType::GetStaticType(), nullptr, &PixelShaderType::GetStaticType(), false, bRemoveUnused);
 // Only VS
 #define IMPLEMENT_SHADERPIPELINE_TYPE_VS(PipelineName, VertexShaderType, bRemoveUnused)	\
-	static FShaderPipelineType PipelineName(TEXT(PREPROCESSOR_TO_STRING(PipelineName)), &VertexShaderType::StaticType, nullptr, nullptr, false, bRemoveUnused);
+	static FShaderPipelineType PipelineName(TEXT(PREPROCESSOR_TO_STRING(PipelineName)), &VertexShaderType::GetStaticType(), nullptr, nullptr, false, bRemoveUnused);
 // Vertex+Geometry+Pixel
 #define IMPLEMENT_SHADERPIPELINE_TYPE_VSGSPS(PipelineName, VertexShaderType, GeometryShaderType, PixelShaderType, bRemoveUnused)	\
-	static FShaderPipelineType PipelineName(TEXT(PREPROCESSOR_TO_STRING(PipelineName)), &VertexShaderType::StaticType, &GeometryShaderType::StaticType, &PixelShaderType::StaticType, false, bRemoveUnused);
+	static FShaderPipelineType PipelineName(TEXT(PREPROCESSOR_TO_STRING(PipelineName)), &VertexShaderType::GetStaticType(), &GeometryShaderType::GetStaticType(), &PixelShaderType::GetStaticType(), false, bRemoveUnused);
 // Vertex+Geometry
 #define IMPLEMENT_SHADERPIPELINE_TYPE_VSGS(PipelineName, VertexShaderType, GeometryShaderType, bRemoveUnused)	\
-	static FShaderPipelineType PipelineName(TEXT(PREPROCESSOR_TO_STRING(PipelineName)), &VertexShaderType::StaticType, &GeometryShaderType::StaticType, nullptr, false, bRemoveUnused);
+	static FShaderPipelineType PipelineName(TEXT(PREPROCESSOR_TO_STRING(PipelineName)), &VertexShaderType::GetStaticType(), &GeometryShaderType::GetStaticType(), nullptr, false, bRemoveUnused);
 
 // Mesh+Pixel
 #define IMPLEMENT_SHADERPIPELINE_TYPE_MSPS(PipelineName, MeshShaderType, PixelShaderType, bRemoveUnused)	\
-	static FShaderPipelineType PipelineName(TEXT(PREPROCESSOR_TO_STRING(PipelineName)), &MeshShaderType::StaticType, nullptr, &PixelShaderType::StaticType, true, bRemoveUnused);
+	static FShaderPipelineType PipelineName(TEXT(PREPROCESSOR_TO_STRING(PipelineName)), &MeshShaderType::GetStaticType(), nullptr, &PixelShaderType::GetStaticType(), true, bRemoveUnused);
 // Mesh+Amplification+Pixel
 #define IMPLEMENT_SHADERPIPELINE_TYPE_MSASPS(PipelineName, MeshShaderType, AmplificationShaderType, PixelShaderType, bRemoveUnused)	\
-	static FShaderPipelineType PipelineName(TEXT(PREPROCESSOR_TO_STRING(PipelineName)), &MeshShaderType::StaticType, &AmplificationShaderType::StaticType, &PixelShaderType::StaticType, true, bRemoveUnused);
+	static FShaderPipelineType PipelineName(TEXT(PREPROCESSOR_TO_STRING(PipelineName)), &MeshShaderType::GetStaticType(), &AmplificationShaderType::GetStaticType(), &PixelShaderType::GetStaticType(), true, bRemoveUnused);
 #endif
 
 /** Encapsulates a dependency on a shader type and saved state from that shader type. */
@@ -1949,7 +2002,7 @@ public:
 	template<typename ShaderType>
 	ShaderType* GetShader(const FShaderMapPointerTable& InPtrTable)
 	{
-		const FShaderType& Type = ShaderType::StaticType;
+		const FShaderType& Type = ShaderType::GetStaticType();
 		const EShaderFrequency Frequency = Type.GetFrequency();
 		if (Frequency < SF_NumGraphicsFrequencies && Shaders[Frequency].IsValid())
 		{
@@ -2073,8 +2126,8 @@ public:
 	template<typename ShaderType>
 	ShaderType* GetShader(int32 PermutationId = 0) const
 	{
-		FShader* Shader = GetShader(&ShaderType::StaticType, PermutationId);
-		checkf(Shader != nullptr, TEXT("Failed to find shader type %s in Platform %s"), ShaderType::StaticType.GetName(), *LegacyShaderPlatformToShaderFormat(GetShaderPlatform()).ToString());
+		FShader* Shader = GetShader(&ShaderType::GetStaticType(), PermutationId);
+		checkf(Shader != nullptr, TEXT("Failed to find shader type %s in Platform %s"), ShaderType::GetStaticType().GetName(), *LegacyShaderPlatformToShaderFormat(GetShaderPlatform()).ToString());
 		return static_cast<ShaderType*>(Shader);
 	}
 
@@ -2368,7 +2421,7 @@ class TOptionalShaderMapRef : public TShaderRef<ShaderType>
 {
 public:
 	TOptionalShaderMapRef(const typename ShaderType::ShaderMapType* ShaderIndex):
-		TShaderRef<ShaderType>(TShaderRef<ShaderType>::Cast(ShaderIndex->GetShader(&ShaderType::StaticType))) // gcc3 needs the template quantifier so it knows the < is not a less-than
+		TShaderRef<ShaderType>(TShaderRef<ShaderType>::Cast(ShaderIndex->GetShader(&ShaderType::GetStaticType()))) // gcc3 needs the template quantifier so it knows the < is not a less-than
 	{}
 };
 
diff --git a/Engine/Source/Runtime/RenderCore/Public/ShaderParameterMacros.h b/Engine/Source/Runtime/RenderCore/Public/ShaderParameterMacros.h
--- a/Engine/Source/Runtime/RenderCore/Public/ShaderParameterMacros.h
+++ b/Engine/Source/Runtime/RenderCore/Public/ShaderParameterMacros.h
@@ -13,6 +13,7 @@
 #include "Algo/Reverse.h"
 #include "Templates/IsArrayOrRefOfTypeByPredicate.h"
 #include "Traits/IsCharEncodingCompatibleWith.h"
+#include "Misc/AssertionMacros.h"
 
 class FRDGTexture;
 class FRDGTextureSRV;
@@ -838,7 +839,7 @@ struct TShaderParameterTypeInfo
 	/** Type that has a multiple of 4 components. */
 	using TInstancedType = TypeParameter;
 
-	static const FShaderParametersMetadata* GetStructMetadata() { return &TypeParameter::StaticStructMetadata; }
+	static const FShaderParametersMetadata* GetStructMetadata() { return TypeParameter::FTypeInfo::GetStructMetadata(); }
 };
 
 // Compile SHADER_PARAMETER(bool, MyBool), just to give good error message to programmer why they shouldn't do that.
@@ -1271,28 +1272,31 @@ struct TShaderParameterStructTypeInfo<StructType[InNumElements]>
 };
 
 #define INTERNAL_BEGIN_UNIFORM_BUFFER_STRUCT \
-	static FShaderParametersMetadata StaticStructMetadata;
+	static const FShaderParametersMetadata* GetStructMetadata(); 
 
 #define INTERNAL_UNIFORM_BUFFER_STRUCT_GET_STRUCT_METADATA(StructTypeName) \
-	return &StructTypeName::StaticStructMetadata;
+	{ return StructTypeName::GetStructMetadata(); }
 
 #define INTERNAL_SHADER_PARAMETER_GET_STRUCT_METADATA(StructTypeName) \
-	static FShaderParametersMetadata StaticStructMetadata(\
-		FShaderParametersMetadata::EUseCase::ShaderParameterStruct, \
-		EUniformBufferBindingFlags::Shader, \
-		TEXT(#StructTypeName), \
-		TEXT(#StructTypeName), \
-		nullptr, \
-		nullptr, \
-		FTypeInfo::FileName, \
-		FTypeInfo::FileLine, \
-		sizeof(StructTypeName), \
-		StructTypeName::zzGetMembers()); \
-	return &StaticStructMetadata;
+	{ \
+		static FShaderParametersMetadata StaticStructMetadata(\
+			FShaderParametersMetadata::EUseCase::ShaderParameterStruct, \
+			EUniformBufferBindingFlags::Shader, \
+			TEXT(#StructTypeName), \
+			TEXT(#StructTypeName), \
+			nullptr, \
+			nullptr, \
+			FTypeInfo::FileName, \
+			FTypeInfo::FileLine, \
+			sizeof(StructTypeName), \
+			StructTypeName::zzGetMembers()); \
+		return &StaticStructMetadata; \
+	}
+
 
 #define INTERNAL_SHADER_PARAMETER_STRUCT_CREATE_UNIFORM_BUFFER return nullptr;
 
-#define INTERNAL_UNIFORM_BUFFER_STRUCT_CREATE_UNIFORM_BUFFER return RHICreateUniformBuffer(&InContents, StaticStructMetadata.GetLayoutPtr(), InUsage);
+#define INTERNAL_UNIFORM_BUFFER_STRUCT_CREATE_UNIFORM_BUFFER return RHICreateUniformBuffer(&InContents, FTypeInfo::GetStructMetadata()->GetLayoutPtr(), InUsage);
 
 /** Begins a uniform buffer struct declaration. */
 #define INTERNAL_SHADER_PARAMETER_STRUCT_BEGIN(StructTypeName,PrefixKeywords,ConstructorSuffix,GetStructMetadataScope,CreateUniformBufferImpl) \
@@ -1309,7 +1313,7 @@ struct TShaderParameterStructTypeInfo<StructType[InNumElements]>
 			static constexpr const ANSICHAR* const FileName = UE_LOG_SOURCE_FILE(__FILE__); \
 			static constexpr int32 FileLine = __LINE__; \
 			using TAlignedType = StructTypeName; \
-			static inline const FShaderParametersMetadata* GetStructMetadata() { GetStructMetadataScope } \
+			static const FShaderParametersMetadata* GetStructMetadata() GetStructMetadataScope \
 		}; \
 		static FUniformBufferRHIRef CreateUniformBuffer(const StructTypeName& InContents, EUniformBufferUsage InUsage) \
 		{ \
@@ -1414,19 +1418,45 @@ extern RENDERCORE_API FShaderParametersMetadata* FindUniformBufferStructByShader
 #define END_UNIFORM_BUFFER_STRUCT() \
 	END_SHADER_PARAMETER_STRUCT()
 
+
+class RENDERCORE_API FShaderParametersMetadataRegistration
+{
+public:
+	FShaderParametersMetadataRegistration(TFunctionRef<const FShaderParametersMetadata* ()> LazyShaderParametersMetadataAccessor)
+		: LazyShaderParametersMetadataAccessor(LazyShaderParametersMetadataAccessor)
+	{
+		Instances.Add(this);
+	}
+
+	// Actually register all the instances and clear the array
+	static void CommitAll();
+	static bool IsReadyForRegistration();
+
+private:
+	static TArray<const FShaderParametersMetadataRegistration*> Instances;
+
+	TFunctionRef<const FShaderParametersMetadata* ()> LazyShaderParametersMetadataAccessor;
+};
+
 #define IMPLEMENT_UNIFORM_BUFFER_STRUCT(StructTypeName,ShaderVariableName) \
-	FShaderParametersMetadata StructTypeName::StaticStructMetadata( \
-	FShaderParametersMetadata::EUseCase::UniformBuffer, \
-	EUniformBufferBindingFlags::Shader, \
-	TEXT(#StructTypeName), \
-	TEXT(#StructTypeName), \
-	TEXT(ShaderVariableName), \
-	nullptr, \
-	StructTypeName::FTypeInfo::FileName, \
-	StructTypeName::FTypeInfo::FileLine, \
-	sizeof(StructTypeName), \
-	StructTypeName::zzGetMembers()); \
-	const FShaderParametersMetadata* GetForwardDeclaredShaderParametersStructMetadata(const StructTypeName* DummyPtr) {  return StructTypeName::FTypeInfo::GetStructMetadata(); }
+	const FShaderParametersMetadata* GetForwardDeclaredShaderParametersStructMetadata(const StructTypeName* DummyPtr) { return StructTypeName::FTypeInfo::GetStructMetadata(); } \
+	const FShaderParametersMetadata* StructTypeName::GetStructMetadata() \
+	{ \
+		ensureMsgf(FShaderParametersMetadataRegistration::IsReadyForRegistration(), TEXT("GetStructMetadata should not be called from a static initializer")); \
+		static const FShaderParametersMetadata StaticStructMetadata( \
+			FShaderParametersMetadata::EUseCase::UniformBuffer, \
+			EUniformBufferBindingFlags::Shader, \
+			TEXT(#StructTypeName), \
+			TEXT(#StructTypeName), \
+			TEXT(ShaderVariableName), \
+			nullptr, \
+			StructTypeName::FTypeInfo::FileName, \
+			StructTypeName::FTypeInfo::FileLine, \
+			sizeof(StructTypeName), \
+			StructTypeName::zzGetMembers()); \
+		return &StaticStructMetadata; \
+	} \
+	FShaderParametersMetadataRegistration StructTypeName##MetadataRegistration { TFunctionRef<const ::FShaderParametersMetadata* ()>{StructTypeName::GetStructMetadata} };
 
 #define IMPLEMENT_UNIFORM_BUFFER_ALIAS_STRUCT(StructTypeName, UniformBufferAlias) \
 	static const FShaderParametersMetadata UniformBufferAlias( \
@@ -1440,39 +1470,51 @@ extern RENDERCORE_API FShaderParametersMetadata* FindUniformBufferStructByShader
 	StructTypeName::FTypeInfo::FileLine, \
 	sizeof(StructTypeName), \
 	StructTypeName::zzGetMembers())
-
+ 
 #define IMPLEMENT_STATIC_UNIFORM_BUFFER_STRUCT_EX(StructTypeName,ShaderVariableName,StaticSlotName,BindingFlagsEnum) \
 	static_assert(EnumHasAnyFlags(EUniformBufferBindingFlags::BindingFlagsEnum, EUniformBufferBindingFlags::Static), "Shader bindings must include 'Static'."); \
-	FShaderParametersMetadata StructTypeName::StaticStructMetadata( \
-		FShaderParametersMetadata::EUseCase::UniformBuffer, \
-		EUniformBufferBindingFlags::BindingFlagsEnum, \
-		TEXT(#StructTypeName), \
-		TEXT(#StructTypeName), \
-		TEXT(ShaderVariableName), \
-		TEXT(#StaticSlotName), \
-		StructTypeName::FTypeInfo::FileName, \
-		StructTypeName::FTypeInfo::FileLine, \
-		sizeof(StructTypeName), \
-		StructTypeName::zzGetMembers()); \
-	const FShaderParametersMetadata* GetForwardDeclaredShaderParametersStructMetadata(const StructTypeName* DummyPtr) {  return StructTypeName::FTypeInfo::GetStructMetadata(); }
+	const FShaderParametersMetadata* GetForwardDeclaredShaderParametersStructMetadata(const StructTypeName* DummyPtr) { return StructTypeName::FTypeInfo::GetStructMetadata(); } \
+	const FShaderParametersMetadata* StructTypeName::GetStructMetadata() \
+	{ \
+		ensureMsgf(FShaderParametersMetadataRegistration::IsReadyForRegistration(), TEXT("GetStructMetadata should not be called from a static initializer")); \
+		static const FShaderParametersMetadata StaticStructMetadata( \
+			FShaderParametersMetadata::EUseCase::UniformBuffer, \
+			EUniformBufferBindingFlags::BindingFlagsEnum, \
+			TEXT(#StructTypeName), \
+			TEXT(#StructTypeName), \
+			TEXT(ShaderVariableName), \
+			TEXT(#StaticSlotName), \
+			StructTypeName::FTypeInfo::FileName, \
+			StructTypeName::FTypeInfo::FileLine, \
+			sizeof(StructTypeName), \
+			StructTypeName::zzGetMembers()); \
+		return &StaticStructMetadata; \
+	} \
+	FShaderParametersMetadataRegistration StructTypeName##MetadataRegistration { TFunctionRef<const ::FShaderParametersMetadata* ()>{StructTypeName::GetStructMetadata} };
 
 #define IMPLEMENT_STATIC_UNIFORM_BUFFER_STRUCT_EX2(StructTypeName,ShaderVariableName,StaticSlotName,BindingFlagsEnum,UsageFlags) \
 	static_assert(EnumHasAnyFlags(EUniformBufferBindingFlags::BindingFlagsEnum, EUniformBufferBindingFlags::Static), "Shader bindings must include 'Static'."); \
-	FShaderParametersMetadata StructTypeName::StaticStructMetadata( \
-		FShaderParametersMetadata::EUseCase::UniformBuffer, \
-		EUniformBufferBindingFlags::BindingFlagsEnum, \
-		TEXT(#StructTypeName), \
-		TEXT(#StructTypeName), \
-		TEXT(ShaderVariableName), \
-		TEXT(#StaticSlotName), \
-		StructTypeName::FTypeInfo::FileName, \
-		StructTypeName::FTypeInfo::FileLine, \
-		sizeof(StructTypeName), \
-		StructTypeName::zzGetMembers(), \
-		false, \
-		nullptr, \
-		(uint32)UsageFlags); \
-	const FShaderParametersMetadata* GetForwardDeclaredShaderParametersStructMetadata(const StructTypeName* DummyPtr) {  return StructTypeName::FTypeInfo::GetStructMetadata(); }
+	const FShaderParametersMetadata* GetForwardDeclaredShaderParametersStructMetadata(const StructTypeName* DummyPtr) {  return StructTypeName::FTypeInfo::GetStructMetadata(); } \
+	const FShaderParametersMetadata* StructTypeName::GetStructMetadata() \
+	{ \
+		ensureMsgf(FShaderParametersMetadataRegistration::IsReadyForRegistration(), TEXT("GetStructMetadata should not be called from a static initializer")); \
+		static const FShaderParametersMetadata StaticStructMetadata( \
+			FShaderParametersMetadata::EUseCase::UniformBuffer, \
+			EUniformBufferBindingFlags::BindingFlagsEnum, \
+			TEXT(#StructTypeName), \
+			TEXT(#StructTypeName), \
+			TEXT(ShaderVariableName), \
+			TEXT(#StaticSlotName), \
+			StructTypeName::FTypeInfo::FileName, \
+			StructTypeName::FTypeInfo::FileLine, \
+			sizeof(StructTypeName), \
+			StructTypeName::zzGetMembers(), \
+			false, \
+			nullptr, \
+			(uint32)UsageFlags); \
+		return &StaticStructMetadata; \
+	} \
+	FShaderParametersMetadataRegistration StructTypeName##MetadataRegistration { TFunctionRef<const ::FShaderParametersMetadata* ()>{StructTypeName::GetStructMetadata} };
 
 /** Implements the same contract as IMPLEMENT_STATIC_UNIFORM_BUFFER_STRUCT, with the addition of the 'StaticAndShader' binding
  *  flag. This means that the uniform buffer may be bound statically or through SetShaderParameters. The uniform buffer is NOT
diff --git a/Engine/Source/Runtime/RenderCore/Public/ShaderParameterUtils.h b/Engine/Source/Runtime/RenderCore/Public/ShaderParameterUtils.h
--- a/Engine/Source/Runtime/RenderCore/Public/ShaderParameterUtils.h
+++ b/Engine/Source/Runtime/RenderCore/Public/ShaderParameterUtils.h
@@ -215,7 +215,7 @@ inline void SetUniformBufferParameterImmediate(FRHIBatchedShaderParameters& Batc
 	checkSlow(Parameter.IsInitialized());
 	if (Parameter.IsBound())
 	{
-		FUniformBufferRHIRef UniformBufferRef = RHICreateUniformBuffer(&UniformBufferValue, &TBufferStruct::StaticStructMetadata.GetLayout(), UniformBuffer_SingleDraw);
+		FUniformBufferRHIRef UniformBufferRef = RHICreateUniformBuffer(&UniformBufferValue, &TBufferStruct::FTypeInfo::GetStructMetadata()->GetLayout(), UniformBuffer_SingleDraw);
 		SetUniformBufferParameter(BatchedParameters, Parameter, UniformBufferRef.GetReference());
 	}
 }
@@ -820,7 +820,7 @@ inline void SetUniformBufferParameterImmediate(
 		RHICmdList.SetShaderUniformBuffer(
 			Shader,
 			Parameter.GetBaseIndex(),
-			RHICreateUniformBuffer(&UniformBufferValue, &TBufferStruct::StaticStructMetadata.GetLayout(), UniformBuffer_SingleDraw)
+			RHICreateUniformBuffer(&UniformBufferValue, &TBufferStruct::FTypeInfo::GetStructMetadata()->GetLayout(), UniformBuffer_SingleDraw)
 		);
 	}
 }
@@ -841,7 +841,7 @@ inline void SetUniformBufferParameterImmediate(
 		RHICmdList.SetShaderUniformBuffer(
 			Shader,
 			Parameter.GetBaseIndex(),
-			RHICreateUniformBuffer(&UniformBufferValue, &TBufferStruct::StaticStructMetadata.GetLayout(), UniformBuffer_SingleDraw)
+			RHICreateUniformBuffer(&UniformBufferValue, &TBufferStruct::FTypeInfo::GetStructMetadata()->GetLayout(), UniformBuffer_SingleDraw)
 		);
 	}
 }
\ No newline at end of file
diff --git a/Engine/Source/Runtime/Renderer/Private/BasePassRendering.h b/Engine/Source/Runtime/Renderer/Private/BasePassRendering.h
--- a/Engine/Source/Runtime/Renderer/Private/BasePassRendering.h
+++ b/Engine/Source/Runtime/Renderer/Private/BasePassRendering.h
@@ -312,7 +312,7 @@ public:
 
 	static bool ValidateCompiledResult(EShaderPlatform Platform, const FShaderParameterMap& ParameterMap, TArray<FString>& OutError)
 	{
-		if (ParameterMap.ContainsParameterAllocation(FSceneTextureUniformParameters::StaticStructMetadata.GetShaderVariableName()))
+		if (ParameterMap.ContainsParameterAllocation(FSceneTextureUniformParameters::FTypeInfo::GetStructMetadata()->GetShaderVariableName()))
 		{
 			OutError.Add(TEXT("Base pass shaders cannot read from the SceneTexturesStruct."));
 			return false;
@@ -342,9 +342,9 @@ public:
 		Target7.Bind(Initializer.ParameterMap, TEXT("OutTarget7UAV"), SPF_Optional);
 
 		// These parameters should only be used nested in the base pass uniform buffer
-		check(!Initializer.ParameterMap.ContainsParameterAllocation(FFogUniformParameters::StaticStructMetadata.GetShaderVariableName()));
-		check(!Initializer.ParameterMap.ContainsParameterAllocation(FReflectionUniformParameters::StaticStructMetadata.GetShaderVariableName()));
-		check(!Initializer.ParameterMap.ContainsParameterAllocation(FPlanarReflectionUniformParameters::StaticStructMetadata.GetShaderVariableName()));
+		check(!Initializer.ParameterMap.ContainsParameterAllocation(FFogUniformParameters::FTypeInfo::GetStructMetadata()->GetShaderVariableName()));
+		check(!Initializer.ParameterMap.ContainsParameterAllocation(FReflectionUniformParameters::FTypeInfo::GetStructMetadata()->GetShaderVariableName()));
+		check(!Initializer.ParameterMap.ContainsParameterAllocation(FPlanarReflectionUniformParameters::FTypeInfo::GetStructMetadata()->GetShaderVariableName()));
 	}
 	TBasePassComputeShaderPolicyParamType() {}
 
@@ -482,7 +482,7 @@ public:
 
 	static bool ValidateCompiledResult(EShaderPlatform Platform, const FShaderParameterMap& ParameterMap, TArray<FString>& OutError)
 	{
-		if (ParameterMap.ContainsParameterAllocation(FSceneTextureUniformParameters::StaticStructMetadata.GetShaderVariableName()))
+		if (ParameterMap.ContainsParameterAllocation(FSceneTextureUniformParameters::FTypeInfo::GetStructMetadata()->GetShaderVariableName()))
 		{
 			OutError.Add(TEXT("Base pass shaders cannot read from the SceneTexturesStruct."));
 			return false;
@@ -499,9 +499,9 @@ public:
 		ReflectionCaptureBuffer.Bind(Initializer.ParameterMap, TEXT("ReflectionCapture"));
 
 		// These parameters should only be used nested in the base pass uniform buffer
-		check(!Initializer.ParameterMap.ContainsParameterAllocation(FFogUniformParameters::StaticStructMetadata.GetShaderVariableName()));
-		check(!Initializer.ParameterMap.ContainsParameterAllocation(FReflectionUniformParameters::StaticStructMetadata.GetShaderVariableName()));
-		check(!Initializer.ParameterMap.ContainsParameterAllocation(FPlanarReflectionUniformParameters::StaticStructMetadata.GetShaderVariableName()));
+		check(!Initializer.ParameterMap.ContainsParameterAllocation(FFogUniformParameters::FTypeInfo::GetStructMetadata()->GetShaderVariableName()));
+		check(!Initializer.ParameterMap.ContainsParameterAllocation(FReflectionUniformParameters::FTypeInfo::GetStructMetadata()->GetShaderVariableName()));
+		check(!Initializer.ParameterMap.ContainsParameterAllocation(FPlanarReflectionUniformParameters::FTypeInfo::GetStructMetadata()->GetShaderVariableName()));
 	}
 	TBasePassPixelShaderPolicyParamType() {}
 
diff --git a/Engine/Source/Runtime/Renderer/Private/DecalRenderingShared.cpp b/Engine/Source/Runtime/Renderer/Private/DecalRenderingShared.cpp
--- a/Engine/Source/Runtime/Renderer/Private/DecalRenderingShared.cpp
+++ b/Engine/Source/Runtime/Renderer/Private/DecalRenderingShared.cpp
@@ -83,7 +83,7 @@ public:
 		DecalOrientation.Bind(Initializer.ParameterMap,TEXT("DecalOrientation"));
 		DecalParams.Bind(Initializer.ParameterMap, TEXT("DecalParams"));
 		DecalColorParam.Bind(Initializer.ParameterMap, TEXT("DecalColorParam"));
-		MobileBasePassUniformBuffer.Bind(Initializer.ParameterMap, FMobileBasePassUniformParameters::StaticStructMetadata.GetShaderVariableName());
+		MobileBasePassUniformBuffer.Bind(Initializer.ParameterMap, FMobileBasePassUniformParameters::FTypeInfo::GetStructMetadata()->GetShaderVariableName());
 	}
 
 	void SetParameters(FRHICommandList& RHICmdList, const FViewInfo& View, const FDeferredDecalProxy& DecalProxy, const FMaterialRenderProxy* MaterialProxy, const FMaterial* MaterialResource, const float FadeAlphaValue = 1.0f)
diff --git a/Engine/Source/Runtime/Renderer/Private/HeterogeneousVolumes/HeterogeneousVolumesLiveShadingPipeline.cpp b/Engine/Source/Runtime/Renderer/Private/HeterogeneousVolumes/HeterogeneousVolumesLiveShadingPipeline.cpp
--- a/Engine/Source/Runtime/Renderer/Private/HeterogeneousVolumes/HeterogeneousVolumesLiveShadingPipeline.cpp
+++ b/Engine/Source/Runtime/Renderer/Private/HeterogeneousVolumes/HeterogeneousVolumesLiveShadingPipeline.cpp
@@ -286,8 +286,8 @@ void AddComputePass(
 			FMeshPassProcessorRenderState DrawRenderState;
 
 			FMeshMaterialShaderElementData ShaderElementData;
-			ShaderElementData.FadeUniformBuffer = GDistanceCullFadedInUniformBuffer.GetUniformBufferRHI();
-			ShaderElementData.DitherUniformBuffer = GDitherFadedInUniformBuffer.GetUniformBufferRHI();
+			ShaderElementData.FadeUniformBuffer = GGetDistanceCullFadedInUniformBuffer();
+			ShaderElementData.DitherUniformBuffer = GGetDitherFadedInUniformBuffer();
 
 			FMeshProcessorShaders PassShaders;
 			PassShaders.ComputeShader = ComputeShader;
diff --git a/Engine/Source/Runtime/Renderer/Private/HeterogeneousVolumes/HeterogeneousVolumesMaterialBakingPipeline.cpp b/Engine/Source/Runtime/Renderer/Private/HeterogeneousVolumes/HeterogeneousVolumesMaterialBakingPipeline.cpp
--- a/Engine/Source/Runtime/Renderer/Private/HeterogeneousVolumes/HeterogeneousVolumesMaterialBakingPipeline.cpp
+++ b/Engine/Source/Runtime/Renderer/Private/HeterogeneousVolumes/HeterogeneousVolumesMaterialBakingPipeline.cpp
@@ -183,8 +183,8 @@ void ComputeHeterogeneousVolumeBakeMaterial(
 				FMeshPassProcessorRenderState DrawRenderState;
 
 				FMeshMaterialShaderElementData ShaderElementData;
-				ShaderElementData.FadeUniformBuffer = GDistanceCullFadedInUniformBuffer.GetUniformBufferRHI();
-				ShaderElementData.DitherUniformBuffer = GDitherFadedInUniformBuffer.GetUniformBufferRHI();
+				ShaderElementData.FadeUniformBuffer = GGetDistanceCullFadedInUniformBuffer();
+				ShaderElementData.DitherUniformBuffer = GGetDitherFadedInUniformBuffer();
 
 				FMeshProcessorShaders PassShaders;
 				PassShaders.ComputeShader = ComputeShader;
diff --git a/Engine/Source/Runtime/Renderer/Private/MobileBasePassRendering.h b/Engine/Source/Runtime/Renderer/Private/MobileBasePassRendering.h
--- a/Engine/Source/Runtime/Renderer/Private/MobileBasePassRendering.h
+++ b/Engine/Source/Runtime/Renderer/Private/MobileBasePassRendering.h
@@ -127,7 +127,7 @@ protected:
 		FMeshMaterialShader(Initializer)
 	{
 		LightMapPolicyType::VertexParametersType::Bind(Initializer.ParameterMap);
-		PassUniformBuffer.Bind(Initializer.ParameterMap, FMobileBasePassUniformParameters::StaticStructMetadata.GetShaderVariableName());
+		PassUniformBuffer.Bind(Initializer.ParameterMap, FMobileBasePassUniformParameters::FTypeInfo::GetStructMetadata()->GetShaderVariableName());
 	}
 
 public:
@@ -236,10 +236,10 @@ public:
 		: FMeshMaterialShader(Initializer)
 	{
 		LightMapPolicyType::PixelParametersType::Bind(Initializer.ParameterMap);
-		PassUniformBuffer.Bind(Initializer.ParameterMap, FMobileBasePassUniformParameters::StaticStructMetadata.GetShaderVariableName());
+		PassUniformBuffer.Bind(Initializer.ParameterMap, FMobileBasePassUniformParameters::FTypeInfo::GetStructMetadata()->GetShaderVariableName());
 		
-		MobileDirectionLightBufferParam.Bind(Initializer.ParameterMap, FMobileDirectionalLightShaderParameters::StaticStructMetadata.GetShaderVariableName());
-		ReflectionParameter.Bind(Initializer.ParameterMap, FMobileReflectionCaptureShaderParameters::StaticStructMetadata.GetShaderVariableName());
+		MobileDirectionLightBufferParam.Bind(Initializer.ParameterMap, FMobileDirectionalLightShaderParameters::FTypeInfo::GetStructMetadata()->GetShaderVariableName());
+		ReflectionParameter.Bind(Initializer.ParameterMap, FMobileReflectionCaptureShaderParameters::FTypeInfo::GetStructMetadata()->GetShaderVariableName());
 				
 		UseCSMParameter.Bind(Initializer.ParameterMap, TEXT("UseCSM"));
 	}
diff --git a/Engine/Source/Runtime/Renderer/Private/PathTracing.cpp b/Engine/Source/Runtime/Renderer/Private/PathTracing.cpp
--- a/Engine/Source/Runtime/Renderer/Private/PathTracing.cpp
+++ b/Engine/Source/Runtime/Renderer/Private/PathTracing.cpp
@@ -1200,7 +1200,7 @@ public:
 
 	static bool ValidateCompiledResult(EShaderPlatform Platform, const FShaderParameterMap& ParameterMap, TArray<FString>& OutError)
 	{
-		if (ParameterMap.ContainsParameterAllocation(FSceneTextureUniformParameters::StaticStructMetadata.GetShaderVariableName()))
+		if (ParameterMap.ContainsParameterAllocation(FSceneTextureUniformParameters::FTypeInfo::GetStructMetadata()->GetShaderVariableName()))
 		{
 			OutError.Add(TEXT("Ray tracing closest hit shaders cannot read from the SceneTexturesStruct."));
 			return false;
diff --git a/Engine/Source/Runtime/Renderer/Private/RayTracing/RayTracingDecals.cpp b/Engine/Source/Runtime/Renderer/Private/RayTracing/RayTracingDecals.cpp
--- a/Engine/Source/Runtime/Renderer/Private/RayTracing/RayTracingDecals.cpp
+++ b/Engine/Source/Runtime/Renderer/Private/RayTracing/RayTracingDecals.cpp
@@ -175,7 +175,7 @@ public:
 
 	static bool ValidateCompiledResult(EShaderPlatform Platform, const FShaderParameterMap& ParameterMap, TArray<FString>& OutError)
 	{
-		if (ParameterMap.ContainsParameterAllocation(FSceneTextureUniformParameters::StaticStructMetadata.GetShaderVariableName()))
+		if (ParameterMap.ContainsParameterAllocation(FSceneTextureUniformParameters::FTypeInfo::GetStructMetadata()->GetShaderVariableName()))
 		{
 			OutError.Add(TEXT("Ray tracing callable shaders cannot read from the SceneTexturesStruct."));
 			return false;
@@ -280,7 +280,7 @@ public:
 
 	static bool ValidateCompiledResult(EShaderPlatform Platform, const FShaderParameterMap& ParameterMap, TArray<FString>& OutError)
 	{
-		if (ParameterMap.ContainsParameterAllocation(FSceneTextureUniformParameters::StaticStructMetadata.GetShaderVariableName()))
+		if (ParameterMap.ContainsParameterAllocation(FSceneTextureUniformParameters::FTypeInfo::GetStructMetadata()->GetShaderVariableName()))
 		{
 			OutError.Add(TEXT("Ray tracing closest hit shaders cannot read from the SceneTexturesStruct."));
 			return false;
@@ -316,11 +316,11 @@ FShaderType* GetRayTracingDecalMaterialShaderType(EBlendMode BlendMode)
 {
 	if (NeedsAnyHitShader(BlendMode))
 	{
-		return &FRayTracingDecalMaterialCHS_AHS::StaticType; 
+		return &FRayTracingDecalMaterialCHS_AHS::GetStaticType();
 	}
 	else
 	{
-		return &FRayTracingDecalMaterialCHS::StaticType;
+		return &FRayTracingDecalMaterialCHS::GetStaticType();
 	}
 }
 
diff --git a/Engine/Source/Runtime/Renderer/Private/RayTracing/RayTracingDynamicGeometry.cpp b/Engine/Source/Runtime/Renderer/Private/RayTracing/RayTracingDynamicGeometry.cpp
--- a/Engine/Source/Runtime/Renderer/Private/RayTracing/RayTracingDynamicGeometry.cpp
+++ b/Engine/Source/Runtime/Renderer/Private/RayTracing/RayTracingDynamicGeometry.cpp
@@ -40,7 +40,7 @@ public:
 	FRayTracingDynamicGeometryConverterCS(const FMeshMaterialShaderType::CompiledShaderInitializerType& Initializer)
 		: FMeshMaterialShader(Initializer)
 	{
-		PassUniformBuffer.Bind(Initializer.ParameterMap, FSceneTextureUniformParameters::StaticStructMetadata.GetShaderVariableName());
+		PassUniformBuffer.Bind(Initializer.ParameterMap, FSceneTextureUniformParameters::FTypeInfo::GetStructMetadata()->GetShaderVariableName());
 
 		RWVertexPositions.Bind(Initializer.ParameterMap, TEXT("VertexPositions"));
 		UsingIndirectDraw.Bind(Initializer.ParameterMap, TEXT("UsingIndirectDraw"));
diff --git a/Engine/Source/Runtime/Renderer/Private/RayTracing/RayTracingMaterialHitShaders.cpp b/Engine/Source/Runtime/Renderer/Private/RayTracing/RayTracingMaterialHitShaders.cpp
--- a/Engine/Source/Runtime/Renderer/Private/RayTracing/RayTracingMaterialHitShaders.cpp
+++ b/Engine/Source/Runtime/Renderer/Private/RayTracing/RayTracingMaterialHitShaders.cpp
@@ -72,7 +72,7 @@ public:
 	FMaterialCHS(const FMeshMaterialShaderType::CompiledShaderInitializerType& Initializer)
 		: FMeshMaterialShader(Initializer)
 	{
-		PassUniformBuffer.Bind(Initializer.ParameterMap, FSceneTextureUniformParameters::StaticStructMetadata.GetShaderVariableName());
+		PassUniformBuffer.Bind(Initializer.ParameterMap, FSceneTextureUniformParameters::FTypeInfo::GetStructMetadata()->GetShaderVariableName());
 		FUniformLightMapPolicyShaderParametersType::Bind(Initializer.ParameterMap);
 	}
 
@@ -164,7 +164,7 @@ public:
 
 	static bool ValidateCompiledResult(EShaderPlatform Platform, const FShaderParameterMap& ParameterMap, TArray<FString>& OutError)
 	{
-		if (ParameterMap.ContainsParameterAllocation(FSceneTextureUniformParameters::StaticStructMetadata.GetShaderVariableName()))
+		if (ParameterMap.ContainsParameterAllocation(FSceneTextureUniformParameters::FTypeInfo::GetStructMetadata()->GetShaderVariableName()))
 		{
 			OutError.Add(TEXT("Ray tracing closest hit shaders cannot read from the SceneTexturesStruct."));
 			return false;
diff --git a/Engine/Source/Runtime/Renderer/Private/Renderer.cpp b/Engine/Source/Runtime/Renderer/Private/Renderer.cpp
--- a/Engine/Source/Runtime/Renderer/Private/Renderer.cpp
+++ b/Engine/Source/Runtime/Renderer/Private/Renderer.cpp
@@ -74,6 +74,9 @@ void FRendererModule::StartupModule()
 			FSceneRenderer::CleanUp(RHICmdList);
 		});
 	}));
+
+	// Needs to run on startup, after static init. Ideally at some point replaced with static accessor similar to GGetDitherFadedInUniformBuffer().
+	GIdentityPrimitiveUniformBuffer.InitContents();
 }
 
 void FRendererModule::ShutdownModule()
diff --git a/Engine/Source/Runtime/Renderer/Private/RendererScene.cpp b/Engine/Source/Runtime/Renderer/Private/RendererScene.cpp
--- a/Engine/Source/Runtime/Renderer/Private/RendererScene.cpp
+++ b/Engine/Source/Runtime/Renderer/Private/RendererScene.cpp
@@ -143,11 +143,28 @@ IMPLEMENT_GLOBAL_SHADER_PARAMETER_STRUCT(FDistanceCullFadeUniformShaderParameter
 
 IMPLEMENT_GLOBAL_SHADER_PARAMETER_STRUCT(FDitherUniformShaderParameters, "PrimitiveDither");
 
-/** Global primitive uniform buffer resource containing distance cull faded in */
-TGlobalResource< FGlobalDistanceCullFadeUniformBuffer > GDistanceCullFadedInUniformBuffer;
+/** Global primitive distance cull fading uniform buffer containing fully faded in. */
+FRHIUniformBuffer* GGetDistanceCullFadedInUniformBuffer()
+{
+	FDistanceCullFadeUniformShaderParameters Parameters;
+	Parameters.FadeTimeScaleBias.X = 0.0f;
+	Parameters.FadeTimeScaleBias.Y = 1.0f;
+	static TUniformBufferRef<FDistanceCullFadeUniformShaderParameters> ResourceRef = 
+		TUniformBufferRef<FDistanceCullFadeUniformShaderParameters>::CreateUniformBufferImmediate(Parameters, UniformBuffer_MultiFrame, EUniformBufferValidation::None);
+
+	return ResourceRef.GetReference();
+}
 
-/** Global primitive uniform buffer resource containing dither faded in */
-TGlobalResource< FGlobalDitherUniformBuffer > GDitherFadedInUniformBuffer;
+/** Global primitive dither uniform buffer containing fully faded in. */
+FRHIUniformBuffer* GGetDitherFadedInUniformBuffer()
+{
+	FDitherUniformShaderParameters Parameters;
+	Parameters.LODFactor = 0.0f;
+	static TUniformBufferRef<FDitherUniformShaderParameters> ResourceRef =
+		TUniformBufferRef<FDitherUniformShaderParameters>::CreateUniformBufferImmediate(Parameters, UniformBuffer_MultiFrame, EUniformBufferValidation::None);
+
+	return ResourceRef.GetReference();
+}
 
 static FThreadSafeCounter FSceneViewState_UniqueID;
 
diff --git a/Engine/Source/Runtime/Renderer/Private/ScenePrivate.h b/Engine/Source/Runtime/Renderer/Private/ScenePrivate.h
--- a/Engine/Source/Runtime/Renderer/Private/ScenePrivate.h
+++ b/Engine/Source/Runtime/Renderer/Private/ScenePrivate.h
@@ -369,42 +369,11 @@ public:
 	}
 };
 
-/**
- * Distance cull fading uniform buffer containing fully faded in.
- */
-class FGlobalDistanceCullFadeUniformBuffer : public TUniformBuffer< FDistanceCullFadeUniformShaderParameters >
-{
-public:
-	/** Default constructor. */
-	FGlobalDistanceCullFadeUniformBuffer()
-	{
-		FDistanceCullFadeUniformShaderParameters Parameters;
-		Parameters.FadeTimeScaleBias.X = 0.0f;
-		Parameters.FadeTimeScaleBias.Y = 1.0f;
-		SetContents(Parameters);
-	}
-};
-
-/** Global primitive uniform buffer resource containing faded in */
-extern TGlobalResource< FGlobalDistanceCullFadeUniformBuffer > GDistanceCullFadedInUniformBuffer;
-
-/**
- * Dither uniform buffer containing fully faded in.
- */
-class FGlobalDitherUniformBuffer : public TUniformBuffer< FDitherUniformShaderParameters >
-{
-public:
-	/** Default constructor. */
-	FGlobalDitherUniformBuffer()
-	{
-		FDitherUniformShaderParameters Parameters;
-		Parameters.LODFactor = 0.0f;
-		SetContents(Parameters);
-	}
-};
+/** Global primitive distance cull fading uniform buffer containing fully faded in. */
+FRHIUniformBuffer* GGetDistanceCullFadedInUniformBuffer();
 
-/** Global primitive uniform buffer resource containing faded in */
-extern TGlobalResource< FGlobalDitherUniformBuffer > GDitherFadedInUniformBuffer;
+/** Global primitive dither uniform buffer containing fully faded in. */
+FRHIUniformBuffer* GGetDitherFadedInUniformBuffer();
 
 /**
  * Stores fading state for a single primitive in a single view
diff --git a/Engine/Source/Runtime/Renderer/Private/ShaderBaseClasses.cpp b/Engine/Source/Runtime/Renderer/Private/ShaderBaseClasses.cpp
--- a/Engine/Source/Runtime/Renderer/Private/ShaderBaseClasses.cpp
+++ b/Engine/Source/Runtime/Renderer/Private/ShaderBaseClasses.cpp
@@ -30,8 +30,8 @@ FAutoConsoleVariableRef FMaterialShader::CVarAllowCachedUniformExpressions(
 
 void FMeshMaterialShaderElementData::InitializeMeshMaterialData(const FSceneView* SceneView, const FPrimitiveSceneProxy* RESTRICT PrimitiveSceneProxy, int32 StaticMeshId, bool bDitheredLODTransition, bool bAllowStencilDither)
 {
-	FadeUniformBuffer = GDistanceCullFadedInUniformBuffer.GetUniformBufferRHI();
-	DitherUniformBuffer = GDitherFadedInUniformBuffer.GetUniformBufferRHI();
+	FadeUniformBuffer = GGetDistanceCullFadedInUniformBuffer();
+	DitherUniformBuffer = GGetDitherFadedInUniformBuffer();
 
 	if (SceneView)
 	{
diff --git a/Engine/Source/Runtime/Renderer/Private/ShadowDepthRendering.cpp b/Engine/Source/Runtime/Renderer/Private/ShadowDepthRendering.cpp
--- a/Engine/Source/Runtime/Renderer/Private/ShadowDepthRendering.cpp
+++ b/Engine/Source/Runtime/Renderer/Private/ShadowDepthRendering.cpp
@@ -398,12 +398,12 @@ public:
 
 		if (FSceneInterface::GetShadingPath(FeatureLevel) == EShadingPath::Deferred)
 		{
-			PassUniformBuffer.Bind(Initializer.ParameterMap, FShadowDepthPassUniformParameters::StaticStructMetadata.GetShaderVariableName());
+			PassUniformBuffer.Bind(Initializer.ParameterMap, FShadowDepthPassUniformParameters::FTypeInfo::GetStructMetadata()->GetShaderVariableName());
 		}
 
 		if (FSceneInterface::GetShadingPath(FeatureLevel) == EShadingPath::Mobile)
 		{
-			PassUniformBuffer.Bind(Initializer.ParameterMap, FMobileShadowDepthPassUniformParameters::StaticStructMetadata.GetShaderVariableName());
+			PassUniformBuffer.Bind(Initializer.ParameterMap, FMobileShadowDepthPassUniformParameters::FTypeInfo::GetStructMetadata()->GetShaderVariableName());
 		}
 	}
 
diff --git a/Engine/Source/Runtime/Renderer/Private/ShadowRendering.cpp b/Engine/Source/Runtime/Renderer/Private/ShadowRendering.cpp
--- a/Engine/Source/Runtime/Renderer/Private/ShadowRendering.cpp
+++ b/Engine/Source/Runtime/Renderer/Private/ShadowRendering.cpp
@@ -341,7 +341,7 @@ TModulatedShadowProjection<T>::TModulatedShadowProjection(const ShaderMetaType::
 TShadowProjectionPS<T, false, true>(Initializer)
 {
 	ModulatedShadowColorParameter.Bind(Initializer.ParameterMap, TEXT("ModulatedShadowColor"));
-	MobileBasePassUniformBuffer.Bind(Initializer.ParameterMap, FMobileBasePassUniformParameters::StaticStructMetadata.GetShaderVariableName());
+	MobileBasePassUniformBuffer.Bind(Initializer.ParameterMap, FMobileBasePassUniformParameters::FTypeInfo::GetStructMetadata()->GetShaderVariableName());
 } 
 
 /**
diff --git a/Engine/Source/Runtime/Renderer/Private/ShadowRendering.h b/Engine/Source/Runtime/Renderer/Private/ShadowRendering.h
--- a/Engine/Source/Runtime/Renderer/Private/ShadowRendering.h
+++ b/Engine/Source/Runtime/Renderer/Private/ShadowRendering.h
@@ -1500,8 +1500,8 @@ public:
 	TOnePassPointShadowProjectionPS(const ShaderMetaType::CompiledShaderInitializerType& Initializer):
 		FGlobalShader(Initializer)
 	{
-		HairStrandsParameters.Bind(Initializer.ParameterMap, FHairStrandsViewUniformParameters::StaticStructMetadata.GetShaderVariableName());
-		StrataGlobalParameters.Bind(Initializer.ParameterMap, FStrataGlobalUniformParameters::StaticStructMetadata.GetShaderVariableName());
+		HairStrandsParameters.Bind(Initializer.ParameterMap, FHairStrandsViewUniformParameters::FTypeInfo::GetStructMetadata()->GetShaderVariableName());
+		StrataGlobalParameters.Bind(Initializer.ParameterMap, FStrataGlobalUniformParameters::FTypeInfo::GetStructMetadata()->GetShaderVariableName());
 		OnePassShadowParameters.Bind(Initializer.ParameterMap);
 		ShadowDepthTextureSampler.Bind(Initializer.ParameterMap,TEXT("ShadowDepthTextureSampler"));
 		LightPosition.Bind(Initializer.ParameterMap,TEXT("LightPositionAndInvRadius"));
diff --git a/Engine/Source/Runtime/Renderer/Private/VolumetricCloudRendering.cpp b/Engine/Source/Runtime/Renderer/Private/VolumetricCloudRendering.cpp
--- a/Engine/Source/Runtime/Renderer/Private/VolumetricCloudRendering.cpp
+++ b/Engine/Source/Runtime/Renderer/Private/VolumetricCloudRendering.cpp
@@ -772,7 +772,7 @@ public:
 	FRenderVolumetricCloudRenderViewPs(const ShaderMetaType::CompiledShaderInitializerType& Initializer)
 		: FMeshMaterialShader(Initializer)
 	{
-		PassUniformBuffer.Bind(Initializer.ParameterMap, FRenderVolumetricCloudGlobalParameters::StaticStructMetadata.GetShaderVariableName());
+		PassUniformBuffer.Bind(Initializer.ParameterMap, FRenderVolumetricCloudGlobalParameters::FTypeInfo::GetStructMetadata()->GetShaderVariableName());
 	}
 
 	FRenderVolumetricCloudRenderViewPs() {}
@@ -2391,8 +2391,8 @@ void FSceneRenderer::RenderVolumetricCloudsInternal(FRDGBuilder& GraphBuilder, F
 					FMeshPassProcessorRenderState DrawRenderState;
 
 					FMeshMaterialShaderElementData ShaderElementData;
-					ShaderElementData.FadeUniformBuffer = GDistanceCullFadedInUniformBuffer.GetUniformBufferRHI();
-					ShaderElementData.DitherUniformBuffer = GDitherFadedInUniformBuffer.GetUniformBufferRHI();
+					ShaderElementData.FadeUniformBuffer = GGetDistanceCullFadedInUniformBuffer();
+					ShaderElementData.DitherUniformBuffer = GGetDitherFadedInUniformBuffer();
 
 					FMeshProcessorShaders PassShaders;
 					PassShaders.ComputeShader = ComputeShader;
@@ -2489,8 +2489,8 @@ void FSceneRenderer::RenderVolumetricCloudsInternal(FRDGBuilder& GraphBuilder, F
 				FMeshPassProcessorRenderState DrawRenderState;
 
 				FMeshMaterialShaderElementData ShaderElementData;
-				ShaderElementData.FadeUniformBuffer = GDistanceCullFadedInUniformBuffer.GetUniformBufferRHI();
-				ShaderElementData.DitherUniformBuffer = GDitherFadedInUniformBuffer.GetUniformBufferRHI();
+				ShaderElementData.FadeUniformBuffer = GGetDistanceCullFadedInUniformBuffer();
+				ShaderElementData.DitherUniformBuffer = GGetDitherFadedInUniformBuffer();
 
 				FMeshProcessorShaders PassShaders;
 				PassShaders.ComputeShader = ComputeShader;
diff --git a/Engine/Source/Runtime/Renderer/Public/MeshDrawShaderBindings.h b/Engine/Source/Runtime/Renderer/Public/MeshDrawShaderBindings.h
--- a/Engine/Source/Runtime/Renderer/Public/MeshDrawShaderBindings.h
+++ b/Engine/Source/Runtime/Renderer/Public/MeshDrawShaderBindings.h
@@ -116,8 +116,8 @@ public:
 
 		if (Parameter.IsBound())
 		{
-			checkf(Value.GetReference(), TEXT("Attempted to set null uniform buffer for type %s"), UniformBufferStructType::StaticStructMetadata.GetStructTypeName());
-			checkfSlow(Value.GetReference()->IsValid(), TEXT("Attempted to set already deleted uniform buffer for type %s"), UniformBufferStructType::StaticStructMetadata.GetStructTypeName());
+			checkf(Value.GetReference(), TEXT("Attempted to set null uniform buffer for type %s"), UniformBufferStructType::FTypeInfo::GetStructMetadata()->GetStructTypeName());
+			checkfSlow(Value.GetReference()->IsValid(), TEXT("Attempted to set already deleted uniform buffer for type %s"), UniformBufferStructType::FTypeInfo::GetStructMetadata()->GetStructTypeName());
 			WriteBindingUniformBuffer(Value.GetReference(), Parameter.GetBaseIndex());
 		}
 	}
@@ -129,8 +129,8 @@ public:
 
 		if (Parameter.IsBound())
 		{
-			checkf(Value.GetUniformBufferRHI(), TEXT("Attempted to set null uniform buffer for type %s"), UniformBufferStructType::StaticStructMetadata.GetStructTypeName());
-			checkfSlow(Value.GetUniformBufferRHI()->IsValid(), TEXT("Attempted to set already deleted uniform buffer for type %s"), UniformBufferStructType::StaticStructMetadata.GetStructTypeName());
+			checkf(Value.GetUniformBufferRHI(), TEXT("Attempted to set null uniform buffer for type %s"), UniformBufferStructType::FTypeInfo::GetStructMetadata()->GetStructTypeName());
+			checkfSlow(Value.GetUniformBufferRHI()->IsValid(), TEXT("Attempted to set already deleted uniform buffer for type %s"), UniformBufferStructType::FTypeInfo::GetStructMetadata()->GetStructTypeName());
 			WriteBindingUniformBuffer(Value.GetUniformBufferRHI(), Parameter.GetBaseIndex());
 		}
 	}
