[shaders] strip out deprecated shader map functions

diff --git a/Engine/Source/Runtime/RenderCore/Public/Shader.h b/Engine/Source/Runtime/RenderCore/Public/Shader.h
--- a/Engine/Source/Runtime/RenderCore/Public/Shader.h
+++ b/Engine/Source/Runtime/RenderCore/Public/Shader.h
@@ -515,11 +515,7 @@ public:
 
 	int32 FindShaderIndex(const FSHAHash& InHash) const;
 
-	UE_DEPRECATED(5.1, "Do not call AddShaderCode directly, AddShaderCompilerOutput manages this automatically.")
-	RENDERCORE_API void AddShaderCode(EShaderFrequency InFrequency, const FSHAHash& InHash, const FShaderCode& InCode) {}
 #if WITH_EDITORONLY_DATA
-	UE_DEPRECATED(5.1, "Do not call AddPlatformDebugData directly; AddShaderCompilerOutput manages this automatically.")
-	RENDERCORE_API void AddPlatformDebugData(TConstArrayView<uint8> InPlatformDebugData) {}
 	void AddEditorOnlyData(int32 Index, const FString& DebugName, TConstArrayView<uint8> InPlatformDebugData, TConstArrayView<FShaderCompilerError> InCompilerWarnings);
 	void AppendWarningsToEditorOnlyData(int32 Index, const FString& DebugName, TConstArrayView<FShaderCompilerError> InCompilerWarnings);
 	RENDERCORE_API void LogShaderCompilerWarnings();
