Speculative fix for inability to decompress the shaders.

- IoStore chunks must be getting recompressed along the way. Do not trust the header info and use the IoStore chunk info instead.

diff --git a/Engine/Source/Runtime/RenderCore/Private/ShaderCodeArchive.cpp b/Engine/Source/Runtime/RenderCore/Private/ShaderCodeArchive.cpp
--- a/Engine/Source/Runtime/RenderCore/Private/ShaderCodeArchive.cpp
+++ b/Engine/Source/Runtime/RenderCore/Private/ShaderCodeArchive.cpp
@@ -2067,11 +2067,11 @@ TRefCountPtr<FRHIShader> FIoStoreShaderCodeArchive::CreateShader(int32 ShaderInd
 	FMemStackBase& MemStack = FMemStack::Get();
 	FMemMark Mark(MemStack);
 	FIoStoreShaderGroupEntry& GroupEntry = Header.ShaderGroupEntries[GroupIndex];
-	ensureMsgf(GroupEntry.CompressedSize == PreloadEntryPtr->IoRequest.GetResultOrDie().DataSize(), TEXT("Shader archive header does not match the actual IoStore buffer size, decompression failure likely imminent."));
+	uint32 CompressedSize = PreloadEntryPtr->IoRequest.GetResultOrDie().DataSize();
 	if (GroupEntry.IsGroupCompressed())
 	{
 		uint8* UncompressedCode = reinterpret_cast<uint8*>(MemStack.Alloc(GroupEntry.UncompressedSize, 16));
-		DecompressShadergroupWithOodleAndExtraLogging(GroupIndex, Header.ShaderGroupIoHashes[GroupIndex], GroupEntry, ShaderIndex, ShaderEntry.ShaderGroupIndex, Header.ShaderHashes[ShaderIndex], UncompressedCode, GroupEntry.UncompressedSize, ShaderCode, GroupEntry.CompressedSize);
+		DecompressShadergroupWithOodleAndExtraLogging(GroupIndex, Header.ShaderGroupIoHashes[GroupIndex], GroupEntry, ShaderIndex, ShaderEntry.ShaderGroupIndex, Header.ShaderHashes[ShaderIndex], UncompressedCode, GroupEntry.UncompressedSize, ShaderCode, CompressedSize);
 		ShaderCode = reinterpret_cast<uint8*>(UncompressedCode) + ShaderEntry.UncompressedOffsetInGroup;
 
 #if UE_SCA_VISUALIZE_SHADER_USAGE
