Added support to separately enable OpenGL deferred

diff --git a/Engine/Source/Runtime/Core/Private/HAL/ConsoleManager.cpp b/Engine/Source/Runtime/Core/Private/HAL/ConsoleManager.cpp
--- a/Engine/Source/Runtime/Core/Private/HAL/ConsoleManager.cpp
+++ b/Engine/Source/Runtime/Core/Private/HAL/ConsoleManager.cpp
@@ -2628,6 +2628,13 @@ static TAutoConsoleVariable<int32> CVarMobileShadingPath(
 		 "1: Deferred shading"),
 	ECVF_RenderThreadSafe | ECVF_ReadOnly);
 
+static TAutoConsoleVariable<int32> CVarMobileAllowDeferredShadingOpenGL(
+	TEXT("r.Mobile.AllowDeferredShadingOpenGL"),
+	0,
+	TEXT("0: Do not Allow Deferred Shading on OpenGL (default)"
+		"1: Allow Deferred Shading on OpenGL"),
+	ECVF_RenderThreadSafe | ECVF_ReadOnly);
+
 static TAutoConsoleVariable<int32> CVarMobileEnableStaticAndCSMShadowReceivers(
 	TEXT("r.Mobile.EnableStaticAndCSMShadowReceivers"),
 	1,
diff --git a/Engine/Source/Runtime/Engine/Classes/Engine/RendererSettings.h b/Engine/Source/Runtime/Engine/Classes/Engine/RendererSettings.h
--- a/Engine/Source/Runtime/Engine/Classes/Engine/RendererSettings.h
+++ b/Engine/Source/Runtime/Engine/Classes/Engine/RendererSettings.h
@@ -294,6 +294,12 @@ class ENGINE_API URendererSettings : public UDeveloperSettings
 		ConfigRestartRequired = true))
 	TEnumAsByte<EMobileShadingPath::Type> MobileShadingPath;
 
+	UPROPERTY(config, EditAnywhere, Category = Mobile, meta = (
+		ConsoleVariable = "r.Mobile.AllowDeferredShadingOpenGL", DisplayName = "Allow Deferred Shading on OpenGL",
+		ToolTip = "Whether to allow Deferred Shading on OpenGL, requires the DXC shader compiler and Mobile Shading set to deferred",
+		ConfigRestartRequired = true))
+	uint32 bMobileSupportDeferredOnOpenGL : 1;
+
 	UPROPERTY(config, EditAnywhere, Category = Mobile, meta = (
 		ConsoleVariable = "r.Mobile.SupportGPUScene", DisplayName = "Enable GPUScene on Mobile",
 		ToolTip = "Whether to enable GPUScene on mobile. GPUScene is required for mesh auto-instancing. Changing this setting requires restarting the editor.",
diff --git a/Engine/Source/Runtime/Engine/Private/RendererSettings.cpp b/Engine/Source/Runtime/Engine/Private/RendererSettings.cpp
--- a/Engine/Source/Runtime/Engine/Private/RendererSettings.cpp
+++ b/Engine/Source/Runtime/Engine/Private/RendererSettings.cpp
@@ -251,6 +251,14 @@ void URendererSettings::PostEditChangeProperty(FPropertyChangedEvent& PropertyCh
 			}
 		}
 
+		if ((PropertyChangedEvent.Property->GetFName() == GET_MEMBER_NAME_CHECKED(URendererSettings, MobileShadingPath)))
+		{
+			if (MobileShadingPath.GetValue() == 0)
+			{
+				bMobileSupportDeferredOnOpenGL = 0;
+			}
+		}
+
 		if ((PropertyChangedEvent.Property->GetFName() == GET_MEMBER_NAME_CHECKED(URendererSettings, bEnableStrata)))
 		{
 			if (bEnableStrata)
@@ -347,6 +355,11 @@ bool URendererSettings::CanEditChange(const FProperty* InProperty) const
 		return !bForwardShading;
 	}
 
+	if ((InProperty->GetFName() == GET_MEMBER_NAME_CHECKED(URendererSettings, bMobileSupportDeferredOnOpenGL)))
+	{
+		return MobileShadingPath.GetValue() > 0;
+	}
+
 	PRAGMA_DISABLE_DEPRECATION_WARNINGS
 	// only allow changing ExtendDefaultLuminanceRange if it was disabled.
 	// we don't want new projects disabling this setting but still allow existing projects to enable it.
diff --git a/Engine/Source/Runtime/RenderCore/Private/RenderUtils.cpp b/Engine/Source/Runtime/RenderCore/Private/RenderUtils.cpp
--- a/Engine/Source/Runtime/RenderCore/Private/RenderUtils.cpp
+++ b/Engine/Source/Runtime/RenderCore/Private/RenderUtils.cpp
@@ -488,10 +488,15 @@ RENDERCORE_API bool MobileSupportsGPUScene()
 RENDERCORE_API bool IsMobileDeferredShadingEnabled(const FStaticShaderPlatform Platform)
 {
 	static FShaderPlatformCachedIniValue<bool> MobileShadingPathIniValue(TEXT("r.Mobile.ShadingPath"));
-	return 
-		MobileShadingPathIniValue.Get(Platform) == 1 && 
-		// OpenGL requires DXC for deferred shading
-		(!IsOpenGLPlatform(Platform) || IsDxcEnabledForPlatform(Platform));
+	static TConsoleVariableData<int32>* MobileAllowDeferredShadingOpenGL = IConsoleManager::Get().FindTConsoleVariableDataInt(TEXT("r.Mobile.AllowDeferredShadingOpenGL"));
+
+	// OpenGL requires DXC for deferred shading
+	bool bSupportedPlatform = !IsOpenGLPlatform(Platform) ||
+		(MobileAllowDeferredShadingOpenGL && MobileAllowDeferredShadingOpenGL->GetValueOnAnyThread() != 0);
+
+	return
+		MobileShadingPathIniValue.Get(Platform) == 1 &&
+		bSupportedPlatform && IsDxcEnabledForPlatform(Platform);
 }
 
 RENDERCORE_API bool MobileRequiresSceneDepthAux(const FStaticShaderPlatform Platform)
