Allow configuration of Bindless on a per-format basis in configs.
- The config section is the DDSPI ShaderFormat of the shader platform
- BindlessResources and BindlessSamplers can be set independently
- The cvars still exist and will take precedence over the config settings.
- Renamed RHIGetBindlessResourcesConfiguration/RHIGetBindlessSamplersConfiguration to have the word "Runtime" in them to reflect that they should only be used at runtime and not when deciding if it's enabled for a shader.

diff --git a/Engine/Config/Windows/BaseWindowsEngine.ini b/Engine/Config/Windows/BaseWindowsEngine.ini
--- a/Engine/Config/Windows/BaseWindowsEngine.ini
+++ b/Engine/Config/Windows/BaseWindowsEngine.ini
@@ -8,4 +8,16 @@ LocalPlatformName=WIN
 DeviceProfileSelectionModule="WindowsDeviceProfileSelector"
 
 [OnlineServices.EOS.Auth]
-DefaultExternalCredentialTypeStr=Steam
\ No newline at end of file
+DefaultExternalCredentialTypeStr=Steam
+
+[PCD3D_SM5]
+BindlessResources=Disabled
+BindlessSamplers=RayTracingOnly
+
+[PCD3D_SM6]
+BindlessResources=Disabled
+BindlessSamplers=RayTracingOnly
+
+[SF_VULKAN_SM5]
+BindlessResources=Disabled
+BindlessSamplers=Disabled
diff --git a/Engine/Source/Runtime/D3D12RHI/Private/D3D12Adapter.cpp b/Engine/Source/Runtime/D3D12RHI/Private/D3D12Adapter.cpp
--- a/Engine/Source/Runtime/D3D12RHI/Private/D3D12Adapter.cpp
+++ b/Engine/Source/Runtime/D3D12RHI/Private/D3D12Adapter.cpp
@@ -1131,8 +1131,8 @@ PRAGMA_ENABLE_DEPRECATION_WARNINGS
 #if PLATFORM_SUPPORTS_BINDLESS_RENDERING
 		if (Desc.MaxSupportedFeatureLevel >= D3D_FEATURE_LEVEL_12_0 && Desc.MaxSupportedShaderModel >= D3D_SHADER_MODEL_6_6 && Desc.ResourceBindingTier >= D3D12_RESOURCE_BINDING_TIER_3)
 		{
-			BindlessResourcesConfig = RHIGetBindlessResourcesConfiguration(GMaxRHIShaderPlatform);
-			BindlessSamplersConfig = RHIGetBindlessSamplersConfiguration(GMaxRHIShaderPlatform);
+			BindlessResourcesConfig = RHIGetRuntimeBindlessResourcesConfiguration(GMaxRHIShaderPlatform);
+			BindlessSamplersConfig = RHIGetRuntimeBindlessSamplersConfiguration(GMaxRHIShaderPlatform);
 
 			bBindlessResourcesAllowed = (BindlessResourcesConfig != ERHIBindlessConfiguration::Disabled);
 			bBindlessSamplersAllowed = (BindlessSamplersConfig != ERHIBindlessConfiguration::Disabled);
diff --git a/Engine/Source/Runtime/D3D12RHI/Private/D3D12Device.cpp b/Engine/Source/Runtime/D3D12RHI/Private/D3D12Device.cpp
--- a/Engine/Source/Runtime/D3D12RHI/Private/D3D12Device.cpp
+++ b/Engine/Source/Runtime/D3D12RHI/Private/D3D12Device.cpp
@@ -369,8 +369,8 @@ void FD3D12Device::SetupAfterDeviceCreation()
 	DescriptorHeapManager.Init(GGlobalResourceDescriptorHeapSize, GGlobalSamplerDescriptorHeapSize);
 
 #if PLATFORM_SUPPORTS_BINDLESS_RENDERING
-	const bool bBindlessResources = RHIGetBindlessResourcesConfiguration(GMaxRHIShaderPlatform) != ERHIBindlessConfiguration::Disabled;
-	const bool bBindlessSamplers = RHIGetBindlessSamplersConfiguration(GMaxRHIShaderPlatform) != ERHIBindlessConfiguration::Disabled;
+	const bool bBindlessResources = RHIGetRuntimeBindlessResourcesConfiguration(GMaxRHIShaderPlatform) != ERHIBindlessConfiguration::Disabled;
+	const bool bBindlessSamplers = RHIGetRuntimeBindlessSamplersConfiguration(GMaxRHIShaderPlatform) != ERHIBindlessConfiguration::Disabled;
 	if (bBindlessResources || bBindlessSamplers)
 	{
 		BindlessDescriptorManager.Init(bBindlessResources ? GBindlessResourceDescriptorHeapSize : 0, bBindlessSamplers ? GBindlessSamplerDescriptorHeapSize : 0);
diff --git a/Engine/Source/Runtime/D3D12RHI/Private/Windows/WindowsD3D12Device.cpp b/Engine/Source/Runtime/D3D12RHI/Private/Windows/WindowsD3D12Device.cpp
--- a/Engine/Source/Runtime/D3D12RHI/Private/Windows/WindowsD3D12Device.cpp
+++ b/Engine/Source/Runtime/D3D12RHI/Private/Windows/WindowsD3D12Device.cpp
@@ -420,8 +420,8 @@ static bool CheckDeviceForEmulatedAtomic64Support(IDXGIAdapter* Adapter, ID3D12D
 inline bool ShouldCheckBindlessSupport(EShaderPlatform ShaderPlatform)
 {
 	// Note: only checking against AllShaders allows the RayTracingOnly configuration to not raise the requirements for projects that aren't using RayTracing.
-	return RHIGetBindlessResourcesConfiguration(ShaderPlatform) == ERHIBindlessConfiguration::AllShaders
-		|| RHIGetBindlessSamplersConfiguration(ShaderPlatform) == ERHIBindlessConfiguration::AllShaders;
+	return RHIGetRuntimeBindlessResourcesConfiguration(ShaderPlatform) == ERHIBindlessConfiguration::AllShaders
+		|| RHIGetRuntimeBindlessSamplersConfiguration(ShaderPlatform) == ERHIBindlessConfiguration::AllShaders;
 }
 
 inline ERHIFeatureLevel::Type FindMaxRHIFeatureLevel(D3D_FEATURE_LEVEL InMaxFeatureLevel, D3D_SHADER_MODEL InMaxShaderModel, D3D12_RESOURCE_BINDING_TIER ResourceBindingTier, bool bSupportsWaveOps, bool bSupportsAtomic64)
diff --git a/Engine/Source/Runtime/Engine/Private/ShaderCompiler/ShaderCompiler.cpp b/Engine/Source/Runtime/Engine/Private/ShaderCompiler/ShaderCompiler.cpp
--- a/Engine/Source/Runtime/Engine/Private/ShaderCompiler/ShaderCompiler.cpp
+++ b/Engine/Source/Runtime/Engine/Private/ShaderCompiler/ShaderCompiler.cpp
@@ -5844,7 +5844,7 @@ void GlobalBeginCompileShader(
 	COOK_STAT(ShaderCompilerCookStats::GlobalBeginCompileShaderCalls++);
 	COOK_STAT(FScopedDurationTimer DurationTimer(ShaderCompilerCookStats::GlobalBeginCompileShaderTimeSec));
 
-	EShaderPlatform ShaderPlatform = EShaderPlatform(Target.Platform);
+	const EShaderPlatform ShaderPlatform = EShaderPlatform(Target.Platform);
 	const FName ShaderFormatName = LegacyShaderPlatformToShaderFormat(ShaderPlatform);
 
 	FShaderCompileUtilities::GenerateBrdfHeaders(ShaderPlatform);
@@ -6589,12 +6589,12 @@ void GlobalBeginCompileShader(
 	Input.Environment.SetDefine(TEXT("PLATFORM_SUPPORTS_BUFFER_LOAD_TYPE_CONVERSION"), RHISupportsBufferLoadTypeConversion(ShaderPlatform) ? 1 : 0);
 
 	bool bEnableBindlessMacro = false;
-	if (RHIGetBindlessSupport(EShaderPlatform(Target.Platform)) != ERHIBindlessSupport::Unsupported)
+	if (RHIGetBindlessSupport(ShaderPlatform) != ERHIBindlessSupport::Unsupported)
 	{
 		const bool bIsRaytracingShader = IsRayTracingShaderFrequency(Input.Target.GetFrequency());
 
-		const ERHIBindlessConfiguration ResourcesConfig = RHIGetBindlessResourcesConfiguration(EShaderPlatform(Target.Platform));
-		const ERHIBindlessConfiguration SamplersConfig = RHIGetBindlessSamplersConfiguration(EShaderPlatform(Target.Platform));
+		const ERHIBindlessConfiguration ResourcesConfig = UE::ShaderCompiler::GetBindlessResourcesConfiguration(ShaderFormatName);
+		const ERHIBindlessConfiguration SamplersConfig = UE::ShaderCompiler::GetBindlessSamplersConfiguration(ShaderFormatName);
 
 		if (ResourcesConfig == ERHIBindlessConfiguration::AllShaders || (ResourcesConfig == ERHIBindlessConfiguration::RayTracingShaders && bIsRaytracingShader))
 		{
diff --git a/Engine/Source/Runtime/RHI/Private/RHI.cpp b/Engine/Source/Runtime/RHI/Private/RHI.cpp
--- a/Engine/Source/Runtime/RHI/Private/RHI.cpp
+++ b/Engine/Source/Runtime/RHI/Private/RHI.cpp
@@ -1168,47 +1168,102 @@ static FAutoConsoleVariableRef CVarEnableAttachmentVariableRateShading(
 	ECVF_RenderThreadSafe);
 
 
-int32 GRHIBindlessResourceConfiguration = 0;
+FString GRHIBindlessResourceConfiguration = TEXT("Disabled");
 static FAutoConsoleVariableRef CVarEnableBindlessResources(
 	TEXT("rhi.Bindless.Resources"),
 	GRHIBindlessResourceConfiguration,
-	TEXT("Set to 1 to enable for all shader types. Set to 2 to restrict to Raytracing shaders."),
+	TEXT("Set to Enabled to enable for all shader types. Set to RayTracingOnly to restrict to Raytracing shaders."),
 	ECVF_ReadOnly
 );
 
-int32 GRHIBindlessSamplerConfiguration = 0;
+FString GRHIBindlessSamplerConfiguration = TEXT("Disabled");
 static FAutoConsoleVariableRef CVarEnableBindlessSamplers(
 	TEXT("rhi.Bindless.Samplers"),
 	GRHIBindlessSamplerConfiguration,
-	TEXT("Set to 1 to enable for all shader types. Set to 2 to restrict to Raytracing shaders."),
+	TEXT("Set to Enabled to enable for all shader types. Set to RayTracingOnly to restrict to Raytracing shaders."),
 	ECVF_ReadOnly
 );
 
-static ERHIBindlessConfiguration DetermineBindlessConfiguration(EShaderPlatform Platform, int32 BindlessConfigSetting)
+static ERHIBindlessConfiguration ParseConfigurationFromString(const FString& InSetting)
+{
+	if (InSetting.IsEmpty())
+	{
+		return ERHIBindlessConfiguration::Disabled;
+	}
+
+	if (FCString::Stricmp(*InSetting, TEXT("Disabled")) == 0)
+	{
+		return ERHIBindlessConfiguration::Disabled;
+	}
+
+	if (FCString::Stricmp(*InSetting, TEXT("Enabled")) == 0)
+	{
+		return ERHIBindlessConfiguration::AllShaders;
+	}
+
+	if (FCString::Stricmp(*InSetting, TEXT("RayTracingOnly")) == 0)
+	{
+		return ERHIBindlessConfiguration::RayTracingShaders;
+	}
+
+	return ERHIBindlessConfiguration::Disabled;
+}
+
+static bool GetBindlessConfigurationSetting(FString& OutSetting, EShaderPlatform Platform, const TCHAR* SettingName)
+{
+	const FString ShaderFormat = FDataDrivenShaderPlatformInfo::GetShaderFormat(Platform).ToString();
+	if (!ShaderFormat.IsEmpty())
+	{
+		return GConfig->GetString(*ShaderFormat, SettingName, OutSetting, GEngineIni);
+	}
+
+	return false;
+}
+
+ERHIBindlessConfiguration RHIParseBindlessConfiguration(EShaderPlatform Platform, const FString& ConfigSettingString, const FString& CVarSettingString)
 {
 	const ERHIBindlessSupport BindlessSupport = RHIGetBindlessSupport(Platform);
 
-	if (BindlessSupport == ERHIBindlessSupport::Unsupported || BindlessConfigSetting == 0)
+	const ERHIBindlessConfiguration ConfigSetting = ParseConfigurationFromString(ConfigSettingString);
+	const ERHIBindlessConfiguration CVarSetting = ParseConfigurationFromString(CVarSettingString);
+
+	if (BindlessSupport == ERHIBindlessSupport::Unsupported || (ConfigSetting == ERHIBindlessConfiguration::Disabled && CVarSetting == ERHIBindlessConfiguration::Disabled))
 	{
 		return ERHIBindlessConfiguration::Disabled;
 	}
 
+	// There's no choice here if the platform only supports RayTracing.
 	if (BindlessSupport == ERHIBindlessSupport::RayTracingOnly)
 	{
 		return ERHIBindlessConfiguration::RayTracingShaders;
 	}
 
-	return BindlessConfigSetting == 2 ? ERHIBindlessConfiguration::RayTracingShaders : ERHIBindlessConfiguration::AllShaders;
+	// CVar should always take precedence over the config setting
+	return CVarSetting != ERHIBindlessConfiguration::Disabled ? CVarSetting : ConfigSetting;
+}
+
+static ERHIBindlessConfiguration DetermineBindlessConfiguration(EShaderPlatform Platform, const TCHAR* ConfigName, const FString& CVarSetting)
+{
+	const ERHIBindlessSupport BindlessSupport = RHIGetBindlessSupport(Platform);
+	if (BindlessSupport == ERHIBindlessSupport::Unsupported)
+	{
+		return ERHIBindlessConfiguration::Disabled;
+	}
+
+	FString ConfigSetting;
+	GetBindlessConfigurationSetting(ConfigSetting, Platform, ConfigName);
+
+	return RHIParseBindlessConfiguration(Platform, ConfigSetting, CVarSetting);
 }
 
-ERHIBindlessConfiguration RHIGetBindlessResourcesConfiguration(EShaderPlatform Platform)
+ERHIBindlessConfiguration RHIGetRuntimeBindlessResourcesConfiguration(EShaderPlatform Platform)
 {
-	return DetermineBindlessConfiguration(Platform, GRHIBindlessResourceConfiguration);
+	return DetermineBindlessConfiguration(Platform, TEXT("BindlessResources"), GRHIBindlessResourceConfiguration);
 }
 
-ERHIBindlessConfiguration RHIGetBindlessSamplersConfiguration(EShaderPlatform Platform)
+ERHIBindlessConfiguration RHIGetRuntimeBindlessSamplersConfiguration(EShaderPlatform Platform)
 {
-	return DetermineBindlessConfiguration(Platform, GRHIBindlessSamplerConfiguration);
+	return DetermineBindlessConfiguration(Platform, TEXT("BindlessSamplers"), GRHIBindlessSamplerConfiguration);
 }
 
 namespace RHIConfig
diff --git a/Engine/Source/Runtime/RHI/Private/RHITextureReference.cpp b/Engine/Source/Runtime/RHI/Private/RHITextureReference.cpp
--- a/Engine/Source/Runtime/RHI/Private/RHITextureReference.cpp
+++ b/Engine/Source/Runtime/RHI/Private/RHITextureReference.cpp
@@ -70,7 +70,7 @@ const FRHITextureDesc& FRHITextureReference::GetDesc() const
 
 void FRHITextureReference::UpdateBindlessShaderResourceView()
 {
-	if (RHIGetBindlessResourcesConfiguration(GMaxRHIShaderPlatform) != ERHIBindlessConfiguration::Disabled)
+	if (RHIGetRuntimeBindlessResourcesConfiguration(GMaxRHIShaderPlatform) != ERHIBindlessConfiguration::Disabled)
 	{
 		BindlessShaderResourceViewRHI = RHICreateShaderResourceView(ReferencedTexture.GetReference(), 0u);
 	}
diff --git a/Engine/Source/Runtime/RHI/Public/RHI.h b/Engine/Source/Runtime/RHI/Public/RHI.h
--- a/Engine/Source/Runtime/RHI/Public/RHI.h
+++ b/Engine/Source/Runtime/RHI/Public/RHI.h
@@ -1215,8 +1215,20 @@ inline uint32 GetExpectedFeatureLevelMaxTextureSamplers(const FStaticFeatureLeve
 	return 16;
 }
 
-RHI_API ERHIBindlessConfiguration RHIGetBindlessResourcesConfiguration(EShaderPlatform Platform);
-RHI_API ERHIBindlessConfiguration RHIGetBindlessSamplersConfiguration(EShaderPlatform Platform);
+RHI_API ERHIBindlessConfiguration RHIParseBindlessConfiguration(EShaderPlatform Platform, const FString& ConfigSetting, const FString& CvarSetting);
+RHI_API ERHIBindlessConfiguration RHIGetRuntimeBindlessResourcesConfiguration(EShaderPlatform Platform);
+RHI_API ERHIBindlessConfiguration RHIGetRuntimeBindlessSamplersConfiguration(EShaderPlatform Platform);
+
+UE_DEPRECATED(5.3, "RHIGetRuntimeBindlessResourcesConfiguration should be used instead")
+inline ERHIBindlessConfiguration RHIGetBindlessResourcesConfiguration(EShaderPlatform Platform)
+{
+	return RHIGetRuntimeBindlessResourcesConfiguration(Platform);
+}
+UE_DEPRECATED(5.3, "RHIGetRuntimeBindlessSamplersConfiguration should be used instead")
+inline ERHIBindlessConfiguration RHIGetBindlessSamplersConfiguration(EShaderPlatform Platform)
+{
+	return RHIGetRuntimeBindlessSamplersConfiguration(Platform);
+}
 
 #if UE_ENABLE_INCLUDE_ORDER_DEPRECATED_IN_5_2
 #include "RHIStrings.h"
diff --git a/Engine/Source/Runtime/RenderCore/Private/Shader.cpp b/Engine/Source/Runtime/RenderCore/Private/Shader.cpp
--- a/Engine/Source/Runtime/RenderCore/Private/Shader.cpp
+++ b/Engine/Source/Runtime/RenderCore/Private/Shader.cpp
@@ -2095,8 +2095,8 @@ void ShaderMapAppendKeyString(EShaderPlatform Platform, FString& KeyString)
 
 	if (RHIGetBindlessSupport(Platform) != ERHIBindlessSupport::Unsupported)
 	{
-		const ERHIBindlessConfiguration ResourcesConfig = RHIGetBindlessResourcesConfiguration(Platform);
-		const ERHIBindlessConfiguration SamplersConfig = RHIGetBindlessSamplersConfiguration(Platform);
+		const ERHIBindlessConfiguration ResourcesConfig = UE::ShaderCompiler::GetBindlessResourcesConfiguration(ShaderFormatName);
+		const ERHIBindlessConfiguration SamplersConfig = UE::ShaderCompiler::GetBindlessSamplersConfiguration(ShaderFormatName);
 
 		if (ResourcesConfig != ERHIBindlessConfiguration::Disabled)
 		{
diff --git a/Engine/Source/Runtime/RenderCore/Private/ShaderCore.cpp b/Engine/Source/Runtime/RenderCore/Private/ShaderCore.cpp
--- a/Engine/Source/Runtime/RenderCore/Private/ShaderCore.cpp
+++ b/Engine/Source/Runtime/RenderCore/Private/ShaderCore.cpp
@@ -371,98 +371,177 @@ bool AllowDebugViewmodes(EShaderPlatform Platform)
 #endif
 }
 
-static bool IsShaderCompilerConfigEnabledForPlatform(const IConsoleVariable* GlobalCvar, const TCHAR* Key, FName ShaderFormat)
+struct FShaderSettingHelper
 {
-	bool bEnabled = false;
+	const IConsoleVariable* const SettingCVar;
+#if WITH_EDITOR
+	const TCHAR* const SettingSection;
+	const TCHAR* const SettingSectionBuildMachine;
+	const TCHAR* const SettingName;
+#endif
 
-	// First check the global cvar
-	if (GlobalCvar && GlobalCvar->GetInt())
+	FShaderSettingHelper() = delete;
+	FShaderSettingHelper(
+		const TCHAR* InSettingSection,
+		const TCHAR* InSettingSectionBuildMachine,
+		const TCHAR* InSettingName,
+		const IConsoleVariable* InSettingCVar)
+		: SettingCVar(InSettingCVar)
+#if WITH_EDITOR
+		, SettingSection(InSettingSection)
+		, SettingSectionBuildMachine(InSettingSectionBuildMachine)
+		, SettingName(InSettingName)
+#endif
 	{
-		bEnabled = true;
 	}
+
 #if WITH_EDITOR
-	else
+	static FConfigCacheIni* GetPlatformConfigForFormat(const FName ShaderFormat)
 	{
-		// Then check the per platform settings.
 		ITargetPlatform* TargetPlatform = GetTargetPlatformManager()->FindTargetPlatformWithSupport(TEXT("ShaderFormat"), ShaderFormat);
-		if (TargetPlatform)
+		return TargetPlatform ? TargetPlatform->GetConfigSystem() : nullptr;
+	}
+#endif
+
+	bool GetBoolForPlatform(FName ShaderFormat) const
+	{
+		bool bEnabled = false;
+
+		// First check the global cvar
+		if (SettingCVar && SettingCVar->GetInt())
+		{
+			bEnabled = true;
+		}
+#if WITH_EDITOR
+		// Then check the per platform settings.
+		else if (FConfigCacheIni* PlatformConfig = GetPlatformConfigForFormat(ShaderFormat))
 		{
-			if (FConfigCacheIni* PlatformConfig = TargetPlatform->GetConfigSystem())
+			// Override with a build machine specific setting, if present.
+			bool bFoundConfig = false;
+			if (GIsBuildMachine && SettingSectionBuildMachine)
 			{
-				// Override with a build machine specific setting, if present.
-				bool QueryBaseConfig = GIsBuildMachine ? !PlatformConfig->GetBool(TEXT("ShaderCompiler_BuildMachine"), Key, bEnabled, GEngineIni) : true;
-				if (QueryBaseConfig)
-				{
-					PlatformConfig->GetBool(TEXT("ShaderCompiler"), Key, bEnabled, GEngineIni);
-				}
+				bFoundConfig = PlatformConfig->GetValue(SettingSectionBuildMachine, SettingName, bEnabled, GEngineIni);
+			}
+			if (!bFoundConfig)
+			{
+				PlatformConfig->GetValue(SettingSection, SettingName, bEnabled, GEngineIni);
 			}
 		}
-	}
 #endif
 
-	return bEnabled;
-}
+		return bEnabled;
+	}
 
-static bool GetShaderCompilerStringForPlatform(FString& OutputString, const IConsoleVariable* GlobalCvar, const TCHAR* Key, FName ShaderFormat)
-{
-	// First check the global cvar
-	if (GlobalCvar)
+	int32 GetIntForPlatform(FName ShaderFormat) const
 	{
-		OutputString = GlobalCvar->GetString();
-	}
+		int32 Value = false;
 
+		// First check the global cvar
+		if (SettingCVar && SettingCVar->GetInt())
+		{
+			Value = SettingCVar->GetInt();
+		}
 #if WITH_EDITOR
-	if (OutputString.IsEmpty())
-	{
 		// Then check the per platform settings.
-		ITargetPlatform* TargetPlatform = GetTargetPlatformManager()->FindTargetPlatformWithSupport(TEXT("ShaderFormat"), ShaderFormat);
-		if (TargetPlatform)
+		else if (FConfigCacheIni* PlatformConfig = GetPlatformConfigForFormat(ShaderFormat))
+		{
+			// Override with a build machine specific setting, if present.
+			bool bFoundConfig = false;
+			if (GIsBuildMachine && SettingSectionBuildMachine)
+			{
+				bFoundConfig = PlatformConfig->GetValue(SettingSectionBuildMachine, SettingName, Value, GEngineIni);
+			}
+			if (!bFoundConfig)
+			{
+				PlatformConfig->GetValue(SettingSection, SettingName, Value, GEngineIni);
+			}
+		}
+#endif
+
+		return Value;
+	}
+
+	bool GetStringForPlatform(FString& OutputString, FName ShaderFormat) const
+	{
+		// First check the global cvar
+		if (SettingCVar)
+		{
+			OutputString = SettingCVar->GetString();
+		}
+
+#if WITH_EDITOR
+		if (OutputString.IsEmpty())
 		{
-			if (FConfigCacheIni* PlatformConfig = TargetPlatform->GetConfigSystem())
+			if (FConfigCacheIni* PlatformConfig = GetPlatformConfigForFormat(ShaderFormat))
 			{
 				// Override with a build machine specific setting, if present.
-				if (GIsBuildMachine)
+				if (GIsBuildMachine && SettingSectionBuildMachine)
 				{
-					OutputString = PlatformConfig->GetStr(TEXT("ShaderCompiler_BuildMachine"), Key, GEngineIni);
+					OutputString = PlatformConfig->GetStr(SettingSectionBuildMachine, SettingName, GEngineIni);
 				}
-
 				if (OutputString.IsEmpty())
 				{
-					OutputString = PlatformConfig->GetStr(TEXT("ShaderCompiler"), Key, GEngineIni);
+					OutputString = PlatformConfig->GetStr(SettingSection, SettingName, GEngineIni);
 				}
 			}
 		}
-	}
 #endif
 
-	return !OutputString.IsEmpty();
-}
+		return !OutputString.IsEmpty();
+	}
+};
 
-struct FShaderSymbolSettingHelper
+struct FShaderSymbolSettingHelper : public FShaderSettingHelper
 {
-	const TCHAR* const SettingName;
-	const IConsoleVariable* const SettingCVar;
-
-	const bool bEnableLegacySettings = true;
 public:
 	FShaderSymbolSettingHelper(const TCHAR* InSettingName, bool bPlatformOnly = false)
-		: SettingName(InSettingName)
-		, SettingCVar(!bPlatformOnly ? IConsoleManager::Get().FindConsoleVariable(InSettingName) : nullptr)
+		: FShaderSettingHelper(TEXT("ShaderCompiler"), TEXT("ShaderCompiler_BuildMachine"), InSettingName, !bPlatformOnly ? IConsoleManager::Get().FindConsoleVariable(InSettingName) : nullptr)
 	{
 		check(SettingCVar || bPlatformOnly);
 	}
 
 	bool IsEnabled(FName ShaderFormat) const
 	{
-		return IsShaderCompilerConfigEnabledForPlatform(SettingCVar, SettingName, ShaderFormat);
+		return GetBoolForPlatform(ShaderFormat);
 	}
 
 	bool GetString(FString& OutString, FName ShaderFormat) const
 	{
-		return GetShaderCompilerStringForPlatform(OutString, SettingCVar, SettingName, ShaderFormat);
+		return GetStringForPlatform(OutString, ShaderFormat);
 	}
 };
 
+inline ERHIBindlessConfiguration GetBindlessConfiguration(FName ShaderFormat, const TCHAR* SettingName, IConsoleVariable* CVar)
+{
+	const FString ShaderFormatStr = ShaderFormat.ToString();
+	const FShaderSettingHelper Helper(*ShaderFormatStr, nullptr, SettingName, CVar);
+
+	const EShaderPlatform ShaderPlatform = ShaderFormatToLegacyShaderPlatform(ShaderFormat);
+
+	FString SettingValue;
+	Helper.GetStringForPlatform(SettingValue, ShaderFormat);
+
+	FString CVarString;
+	if (CVar)
+	{
+		CVar->GetValue(CVarString);
+	}
+
+	return RHIParseBindlessConfiguration(ShaderPlatform, SettingValue, CVarString);
+}
+
+ERHIBindlessConfiguration UE::ShaderCompiler::GetBindlessResourcesConfiguration(FName ShaderFormat)
+{
+	static IConsoleVariable* const BindlessResourcesCVar = IConsoleManager::Get().FindConsoleVariable(TEXT("rhi.Bindless.Resources"));
+	return GetBindlessConfiguration(ShaderFormat, TEXT("BindlessResources"), BindlessResourcesCVar);
+}
+
+ERHIBindlessConfiguration UE::ShaderCompiler::GetBindlessSamplersConfiguration(FName ShaderFormat)
+{
+	static IConsoleVariable* const BindlessSamplersCVar = IConsoleManager::Get().FindConsoleVariable(TEXT("rhi.Bindless.Samplers"));
+	return GetBindlessConfiguration(ShaderFormat, TEXT("BindlessSamplers"), BindlessSamplersCVar);
+}
+
 bool ShouldGenerateShaderSymbols(FName ShaderFormat)
 {
 	static const FShaderSymbolSettingHelper Symbols(TEXT("r.Shaders.Symbols"));
diff --git a/Engine/Source/Runtime/RenderCore/Public/ShaderCompilerCore.h b/Engine/Source/Runtime/RenderCore/Public/ShaderCompilerCore.h
--- a/Engine/Source/Runtime/RenderCore/Public/ShaderCompilerCore.h
+++ b/Engine/Source/Runtime/RenderCore/Public/ShaderCompilerCore.h
@@ -29,6 +29,12 @@ inline const int32 ShaderCompileWorkerSingleJobHeader = 'S';
 // this is for the protocol, not the data.
 inline const int32 ShaderCompileWorkerPipelineJobHeader = 'P';
 
+namespace UE::ShaderCompiler
+{
+	RENDERCORE_API ERHIBindlessConfiguration GetBindlessResourcesConfiguration(FName ShaderFormat);
+	RENDERCORE_API ERHIBindlessConfiguration GetBindlessSamplersConfiguration(FName ShaderFormat);
+}
+
 /** Returns true if shader symbols should be kept for a given platform. */
 extern RENDERCORE_API bool ShouldGenerateShaderSymbols(FName ShaderFormat);
 
diff --git a/Engine/Source/Runtime/VulkanRHI/Private/VulkanDescriptorSets.cpp b/Engine/Source/Runtime/VulkanRHI/Private/VulkanDescriptorSets.cpp
--- a/Engine/Source/Runtime/VulkanRHI/Private/VulkanDescriptorSets.cpp
+++ b/Engine/Source/Runtime/VulkanRHI/Private/VulkanDescriptorSets.cpp
@@ -225,8 +225,8 @@ static inline VkMemoryPropertyFlags GetDescriptorBufferMemoryType(FVulkanDevice*
 bool FVulkanBindlessDescriptorManager::VerifySupport(FVulkanDevice* Device)
 {
 	const bool bFullyDisabled =
-		(RHIGetBindlessResourcesConfiguration(GMaxRHIShaderPlatform) == ERHIBindlessConfiguration::Disabled) &&
-		(RHIGetBindlessSamplersConfiguration(GMaxRHIShaderPlatform) == ERHIBindlessConfiguration::Disabled);
+		(RHIGetRuntimeBindlessResourcesConfiguration(GMaxRHIShaderPlatform) == ERHIBindlessConfiguration::Disabled) &&
+		(RHIGetRuntimeBindlessSamplersConfiguration(GMaxRHIShaderPlatform) == ERHIBindlessConfiguration::Disabled);
 
 	if ((GVulkanBindlessEnabled == 0) || bFullyDisabled)
 	{
@@ -248,8 +248,8 @@ bool FVulkanBindlessDescriptorManager::VerifySupport(FVulkanDevice* Device)
 	}
 
 	const bool bFullyEnabled =
-		(RHIGetBindlessResourcesConfiguration(GMaxRHIShaderPlatform) == ERHIBindlessConfiguration::AllShaders) &&
-		(RHIGetBindlessSamplersConfiguration(GMaxRHIShaderPlatform) == ERHIBindlessConfiguration::AllShaders);
+		(RHIGetRuntimeBindlessResourcesConfiguration(GMaxRHIShaderPlatform) == ERHIBindlessConfiguration::AllShaders) &&
+		(RHIGetRuntimeBindlessSamplersConfiguration(GMaxRHIShaderPlatform) == ERHIBindlessConfiguration::AllShaders);
 
 	if (bFullyEnabled)
 	{
diff --git a/Engine/Source/Runtime/VulkanRHI/Private/VulkanState.cpp b/Engine/Source/Runtime/VulkanRHI/Private/VulkanState.cpp
--- a/Engine/Source/Runtime/VulkanRHI/Private/VulkanState.cpp
+++ b/Engine/Source/Runtime/VulkanRHI/Private/VulkanState.cpp
@@ -233,7 +233,7 @@ FVulkanSamplerState::FVulkanSamplerState(const VkSamplerCreateInfo& InInfo, FVul
 		SamplerId = ++GVulkanSamplerHandleIdCounter;
 	}
 
-	if (InDevice.SupportsBindless() && (RHIGetBindlessSamplersConfiguration(GMaxRHIShaderPlatform) != ERHIBindlessConfiguration::Disabled))
+	if (InDevice.SupportsBindless() && (RHIGetRuntimeBindlessSamplersConfiguration(GMaxRHIShaderPlatform) != ERHIBindlessConfiguration::Disabled))
 	{
 		BindlessHandle = InDevice.GetBindlessDescriptorManager()->RegisterSampler(Sampler);
 	}
