Precreate FRDGPooledBuffer default SRV for structured or byte address buffers. This is so GetSRV() can be safely called off the render thread.

diff --git a/Engine/Source/Runtime/RenderCore/Public/RenderGraphResources.h b/Engine/Source/Runtime/RenderCore/Public/RenderGraphResources.h
--- a/Engine/Source/Runtime/RenderCore/Public/RenderGraphResources.h
+++ b/Engine/Source/Runtime/RenderCore/Public/RenderGraphResources.h
@@ -1158,7 +1158,14 @@ public:
 		, Buffer(MoveTemp(InBuffer))
 		, Name(InName)
 		, NumAllocatedElements(InNumAllocatedElements)
-	{}
+	{
+		if (EnumHasAnyFlags(InDesc.Usage, EBufferUsageFlags::StructuredBuffer | EBufferUsageFlags::ByteAddressBuffer))
+		{
+			FRHIBufferSRVCreateInfo SRVDesc;
+			SRVDesc.BytesPerElement = Desc.BytesPerElement;
+			CachedSRV = GetOrCreateSRV(SRVDesc);
+		}
+	}
 
 	const FRDGBufferDesc Desc;
 
@@ -1174,10 +1181,7 @@ public:
 	/** Returns the default SRV. */
 	FORCEINLINE FRHIShaderResourceView* GetSRV()
 	{
-		if (!CachedSRV)
-		{
-			CachedSRV = GetOrCreateSRV(FRHIBufferSRVCreateInfo());
-		}
+		checkf(CachedSRV, TEXT("Only byte address and structured buffers can use the default GetSRV call"));
 		return CachedSRV;
 	}
 
