Always dispatch RHI thread in RDG BeginFlushResourcesRHI. Avoids a CPU induced stall when flushing the RHI thread.

diff --git a/Engine/Source/Runtime/RenderCore/Private/RenderGraphBuilder.cpp b/Engine/Source/Runtime/RenderCore/Private/RenderGraphBuilder.cpp
--- a/Engine/Source/Runtime/RenderCore/Private/RenderGraphBuilder.cpp
+++ b/Engine/Source/Runtime/RenderCore/Private/RenderGraphBuilder.cpp
@@ -360,11 +360,7 @@ void FRDGBuilder::BeginFlushResourcesRHI()
 
 	CSV_SCOPED_TIMING_STAT_EXCLUSIVE(STAT_RDG_FlushResourcesRHI);
 	SCOPED_NAMED_EVENT(BeginFlushResourcesRHI, FColor::Emerald);
-
-	if (GDynamicRHI->RHIIncludeOptionalFlushes())
-	{
-		RHICmdList.ImmediateFlush(EImmediateFlushType::DispatchToRHIThread);
-	}
+	RHICmdList.ImmediateFlush(EImmediateFlushType::DispatchToRHIThread);
 }
 
 void FRDGBuilder::EndFlushResourcesRHI()
