Add some logging when precompile PSO cache batch rate changes.

diff --git a/Engine/Source/Runtime/RenderCore/Private/ShaderPipelineCache.cpp b/Engine/Source/Runtime/RenderCore/Private/ShaderPipelineCache.cpp
--- a/Engine/Source/Runtime/RenderCore/Private/ShaderPipelineCache.cpp
+++ b/Engine/Source/Runtime/RenderCore/Private/ShaderPipelineCache.cpp
@@ -743,6 +743,8 @@ void FShaderPipelineCache::SetBatchMode(BatchMode Mode)
 {
 	if (ShaderPipelineCache)
 	{
+		uint32 PreviousBatchSize = ShaderPipelineCache->BatchSize;
+		float PreviousBatchTime = ShaderPipelineCache->BatchTime;
 		switch (Mode)
 		{
 			case BatchMode::Precompile:
@@ -765,6 +767,8 @@ void FShaderPipelineCache::SetBatchMode(BatchMode Mode)
 				break;
 			}
 		}
+		UE_CLOG(PreviousBatchSize != ShaderPipelineCache->BatchSize || PreviousBatchTime != ShaderPipelineCache->BatchTime, LogRHI, Log,
+			TEXT("ShaderPipelineCache: Batch mode changed (%d): Size: %d -> %d, Time: %f -> %f"), (uint32)Mode, PreviousBatchSize, ShaderPipelineCache->BatchSize, PreviousBatchTime, ShaderPipelineCache->BatchTime);
 	}
 }
 
