Split ray tracing instance counter stats into explicit active and total parts

* Previously only active instances were reported, which is sometimes misleading
* The total number of instances affects the TLAS memory allocation, instance buffer building compute shader and TLAS build itself

diff --git a/Engine/Source/Runtime/RenderCore/Private/RenderCore.cpp b/Engine/Source/Runtime/RenderCore/Private/RenderCore.cpp
--- a/Engine/Source/Runtime/RenderCore/Private/RenderCore.cpp
+++ b/Engine/Source/Runtime/RenderCore/Private/RenderCore.cpp
@@ -121,7 +121,8 @@ DEFINE_STAT(STAT_PrimitiveCull);
 DEFINE_STAT(STAT_DecompressPrecomputedOcclusion);
 DEFINE_STAT(STAT_ViewVisibilityTime);
 
-DEFINE_STAT(STAT_RayTracingInstances);
+DEFINE_STAT(STAT_RayTracingTotalInstances);
+DEFINE_STAT(STAT_RayTracingActiveInstances);
 DEFINE_STAT(STAT_ProcessedPrimitives);
 DEFINE_STAT(STAT_CulledPrimitives);
 DEFINE_STAT(STAT_VisibleRayTracingPrimitives);
diff --git a/Engine/Source/Runtime/RenderCore/Public/RenderCore.h b/Engine/Source/Runtime/RenderCore/Public/RenderCore.h
--- a/Engine/Source/Runtime/RenderCore/Public/RenderCore.h
+++ b/Engine/Source/Runtime/RenderCore/Public/RenderCore.h
@@ -57,7 +57,10 @@ DECLARE_CYCLE_STAT_EXTERN(TEXT("Decal drawing"),STAT_DecalsDrawTime,STATGROUP_Sc
 DECLARE_CYCLE_STAT_EXTERN(TEXT("Cache Uniform Expressions"),STAT_CacheUniformExpressions,STATGROUP_SceneRendering, RENDERCORE_API);
 DECLARE_DWORD_COUNTER_STAT_EXTERN(TEXT("Lights using light shafts"), STAT_LightShaftsLights, STATGROUP_SceneRendering, RENDERCORE_API);
 DECLARE_CYCLE_STAT_EXTERN(TEXT("Bind ray tracing pipeline"), STAT_BindRayTracingPipeline, STATGROUP_SceneRendering, RENDERCORE_API);
-DECLARE_DWORD_COUNTER_STAT_EXTERN(TEXT("Ray tracing instances"), STAT_RayTracingInstances, STATGROUP_SceneRendering, RENDERCORE_API);
+// Number of instances sent to the TLAS build command, including active and inactive (culled) instances
+DECLARE_DWORD_COUNTER_STAT_EXTERN(TEXT("Ray tracing total instances"), STAT_RayTracingTotalInstances, STATGROUP_SceneRendering, RENDERCORE_API);
+// Number of valid instances, taking into account per-instance culling / activation mask
+DECLARE_DWORD_COUNTER_STAT_EXTERN(TEXT("Ray tracing active instances"), STAT_RayTracingActiveInstances, STATGROUP_SceneRendering, RENDERCORE_API);
 
 DECLARE_CYCLE_STAT_EXTERN(TEXT("View Visibility"),STAT_ViewVisibilityTime,STATGROUP_InitViews, RENDERCORE_API);
 DECLARE_CYCLE_STAT_EXTERN(TEXT("Decompress Occlusion"),STAT_DecompressPrecomputedOcclusion,STATGROUP_InitViews, RENDERCORE_API);
diff --git a/Engine/Source/Runtime/Renderer/Private/RayTracing/RayTracingInstanceBufferUtil.cpp b/Engine/Source/Runtime/Renderer/Private/RayTracing/RayTracingInstanceBufferUtil.cpp
--- a/Engine/Source/Runtime/Renderer/Private/RayTracing/RayTracingInstanceBufferUtil.cpp
+++ b/Engine/Source/Runtime/Renderer/Private/RayTracing/RayTracingInstanceBufferUtil.cpp
@@ -270,7 +270,8 @@ void FillRayTracingInstanceUploadBuffer(
 		TotalNumNativeInstances += NumNativeInstances;
 	}
 
-	SET_DWORD_STAT(STAT_RayTracingInstances, TotalNumNativeInstances - NumInactiveNativeInstances);
+	SET_DWORD_STAT(STAT_RayTracingTotalInstances, TotalNumNativeInstances);
+	SET_DWORD_STAT(STAT_RayTracingActiveInstances, TotalNumNativeInstances - NumInactiveNativeInstances);
 #endif
 }
 
