Fix PSO cache mismatch when packaging a local build using precompiled binaries.

- Reported on UDN: when using precompiled editor binaries to package a project, their version will lag behind the game's engine version - resulting in the bundled file being rejected due to GameVersion field mismatch.
   - Let's say PCBs (precompiled binaries) were compiled at CL X,  you have now synced to X + 100, and you're packaging a game locally.  Cooker will embed X as GameVersion of the bundled cache, but the game will check against its X+100 CL.

diff --git a/Engine/Source/Runtime/RenderCore/Private/ShaderPipelineCache.cpp b/Engine/Source/Runtime/RenderCore/Private/ShaderPipelineCache.cpp
--- a/Engine/Source/Runtime/RenderCore/Private/ShaderPipelineCache.cpp
+++ b/Engine/Source/Runtime/RenderCore/Private/ShaderPipelineCache.cpp
@@ -23,6 +23,7 @@
 #include "Async/AsyncFileHandle.h"
 #include "Misc/ScopeLock.h"
 #include <Algo/RemoveIf.h>
+#include "Modules/BuildVersion.h"
 
 DECLARE_DWORD_ACCUMULATOR_STAT(TEXT("Outstanding Tasks"), STAT_ShaderPipelineTaskCount, STATGROUP_PipelineStateCache );
 DECLARE_DWORD_ACCUMULATOR_STAT(TEXT("Waiting Tasks"), STAT_ShaderPipelineWaitingTaskCount, STATGROUP_PipelineStateCache );
@@ -634,6 +635,22 @@ static void PipelineStateCacheOnAppDeactivate()
 int32 FShaderPipelineCache::GetGameVersionForPSOFileCache()
 {
 	int32 GameVersion = (int32)FEngineVersion::Current().GetChangelist();
+#if WITH_EDITOR
+	// We might be using the editor built at a different (older) CL than the current sync (think precompiled binaries synced via UGS) when packaging.
+	// As such, don't use the current CL as it does not reflect the state of the content.
+	FBuildVersion BuildVersion;
+	if (FBuildVersion::TryRead(FBuildVersion::GetDefaultFileName(), BuildVersion))
+	{
+		// Double check that compatible changelist in the build info is the same (or older) as our binary.
+		// CompatibleChangelist should be populated at the time we compile the precompiled binaries.
+		// If it happens to be newer than the current CL encoded in the binary, it means that we synced the project in UGS
+		// without taking the new precompiled binaries (or without recompiling the engine manually)
+		if (FEngineVersion::Current().GetChangelist() >= (uint32)BuildVersion.CompatibleChangelist)
+		{
+			GameVersion = BuildVersion.Changelist;
+		}
+	}
+#endif
 	GConfig->GetInt(FShaderPipelineCacheConstants::SectionHeading, FShaderPipelineCacheConstants::GameVersionKey, GameVersion, *GGameIni);
 	return GameVersion;
 }
