Moving DumpGPU to Public so that DumpGPUServices can access it correctly.

diff --git a/Engine/Plugins/Developer/DumpGPUServices/Source/DumpGPUServices/DumpGPUServices.Build.cs b/Engine/Plugins/Developer/DumpGPUServices/Source/DumpGPUServices/DumpGPUServices.Build.cs
--- a/Engine/Plugins/Developer/DumpGPUServices/Source/DumpGPUServices/DumpGPUServices.Build.cs
+++ b/Engine/Plugins/Developer/DumpGPUServices/Source/DumpGPUServices/DumpGPUServices.Build.cs
@@ -8,12 +8,6 @@ namespace UnrealBuildTool.Rules
 	{
 		public DumpGPUServices(ReadOnlyTargetRules Target) : base(Target)
 		{
-			PrivateIncludePaths.AddRange(
-				new string[] {
-					EngineDirectory + "/Source/Runtime/RenderCore/Private",
-				}
-				);
-
 			PublicDependencyModuleNames.AddRange(
 				new string[]
 				{
diff --git a/Engine/Source/Runtime/RenderCore/Private/DumpGPU.h b/Engine/Source/Runtime/RenderCore/Private/DumpGPU.h
deleted file mode 100644
diff --git a/Engine/Source/Runtime/RenderCore/Public/DumpGPU.h b/Engine/Source/Runtime/RenderCore/Public/DumpGPU.h
new file mode 100644
--- /dev/null
+++ b/Engine/Source/Runtime/RenderCore/Public/DumpGPU.h
@@ -0,0 +1,32 @@
+// Copyright Epic Games, Inc. All Rights Reserved.
+
+#pragma once
+
+#include "Misc/WildcardString.h"
+#include "UObject/NameTypes.h"
+
+
+class RENDERCORE_API IDumpGPUUploadServiceProvider
+{
+public:
+
+	struct RENDERCORE_API FDumpParameters
+	{
+		static constexpr const TCHAR* kServiceFileName = TEXT("Base/DumpService.json");
+
+		FString Type;
+		FString LocalPath;
+		FString Time;
+
+		FName CompressionName;
+		FWildcardString CompressionFiles;
+
+		FString DumpServiceParametersFileContent() const;
+		bool DumpServiceParametersFile() const;
+	};
+
+	virtual void UploadDump(const FDumpParameters& Parameters) = 0;
+	virtual ~IDumpGPUUploadServiceProvider() = default;
+
+	static IDumpGPUUploadServiceProvider* GProvider;
+};
