Changing GetBindlessConfiguration to not use FShaderSettingHelper for now since RHIParseBindlessConfiguration handles the non-empty CVar override correctly for these settings.
Turning bindless samplers back on for PC D3D RayTracing shaders.

diff --git a/Engine/Config/Windows/BaseWindowsEngine.ini b/Engine/Config/Windows/BaseWindowsEngine.ini
--- a/Engine/Config/Windows/BaseWindowsEngine.ini
+++ b/Engine/Config/Windows/BaseWindowsEngine.ini
@@ -12,11 +12,11 @@ DefaultExternalCredentialTypeStr=Steam
 
 [PCD3D_SM5]
 BindlessResources=Disabled
-BindlessSamplers=Disabled
+BindlessSamplers=RayTracingOnly
 
 [PCD3D_SM6]
 BindlessResources=Disabled
-BindlessSamplers=Disabled
+BindlessSamplers=RayTracingOnly
 
 [SF_VULKAN_SM5]
 BindlessResources=Disabled
diff --git a/Engine/Source/Runtime/RenderCore/Private/ShaderCore.cpp b/Engine/Source/Runtime/RenderCore/Private/ShaderCore.cpp
--- a/Engine/Source/Runtime/RenderCore/Private/ShaderCore.cpp
+++ b/Engine/Source/Runtime/RenderCore/Private/ShaderCore.cpp
@@ -513,13 +513,16 @@ public:
 
 inline ERHIBindlessConfiguration GetBindlessConfiguration(FName ShaderFormat, const TCHAR* SettingName, IConsoleVariable* CVar)
 {
-	const FString ShaderFormatStr = ShaderFormat.ToString();
-	const FShaderSettingHelper Helper(*ShaderFormatStr, nullptr, SettingName, CVar);
-
 	const EShaderPlatform ShaderPlatform = ShaderFormatToLegacyShaderPlatform(ShaderFormat);
 
 	FString SettingValue;
-	Helper.GetStringForPlatform(SettingValue, ShaderFormat);
+#if WITH_EDITOR
+	if (FConfigCacheIni* PlatformConfig = FShaderSettingHelper::GetPlatformConfigForFormat(ShaderFormat))
+	{
+		const FString ShaderFormatStr = ShaderFormat.ToString();
+		SettingValue = PlatformConfig->GetStr(*ShaderFormatStr, SettingName, GEngineIni);
+	}
+#endif
 
 	FString CVarString;
 	if (CVar)
