Fix not having materials in DLC cooked against -fastcook builds.

- Fastcook disables chunking of the game content, so when a DLC is cooked against it and its content is split according to regular chunk IDs, the game can no longer find the shader libraries (that get chunk-specific names). To remedy this, chunking of the DLC shader libraries is disabled, especially given that currently the actual DLC pak file isn't split into chunks (if and when it becomes so, we will reconsider).
- Also improved logging to better understand how the libraries are getting opened, given the multitude of possibilities.

diff --git a/Engine/Source/Editor/UnrealEd/Private/ShaderLibraryChunkDataGenerator.cpp b/Engine/Source/Editor/UnrealEd/Private/ShaderLibraryChunkDataGenerator.cpp
--- a/Engine/Source/Editor/UnrealEd/Private/ShaderLibraryChunkDataGenerator.cpp
+++ b/Engine/Source/Editor/UnrealEd/Private/ShaderLibraryChunkDataGenerator.cpp
@@ -21,6 +21,13 @@ FShaderLibraryChunkDataGenerator::FShaderLibraryChunkDataGenerator(UCookOnTheFly
 	FConfigFile PlatformIniFile;
 	FConfigCacheIni::LoadLocalIniFile(PlatformIniFile, TEXT("Engine"), true, *PlatformNameUsedForIni);
 	PlatformIniFile.GetBool(TEXT("DevOptions.Shaders"), TEXT("bDoNotChunkShaderLib"), bOptedOut);
+
+	// Disable chunking for DLC - this causes problems as the main game can be optionally (for faster iteration) cooked with -fastcook. Fastcook disables chunking,
+	// so the game has no idea about ChunkIDs and cannot find DLC's chunked libs. If DLC lib is monolithic, both monolithic and chunked games will try to open it.
+	if (COTFS.IsCookingDLC())
+	{
+		bOptedOut = true;
+	}
 }
 
 void FShaderLibraryChunkDataGenerator::GenerateChunkDataFiles(const int32 InChunkId, const TSet<FName>& InPackagesInChunk,
diff --git a/Engine/Source/Runtime/RenderCore/Private/ShaderCodeLibrary.cpp b/Engine/Source/Runtime/RenderCore/Private/ShaderCodeLibrary.cpp
--- a/Engine/Source/Runtime/RenderCore/Private/ShaderCodeLibrary.cpp
+++ b/Engine/Source/Runtime/RenderCore/Private/ShaderCodeLibrary.cpp
@@ -2463,15 +2463,20 @@ public:
 					Library = new FNamedShaderLibrary(Name, ShaderPlatform, Directory);
 				}
 			}
+			check(Library);
 			// note, that since we're out of NamedLibrariesMutex locks now, other threads may arrive at the same point and acquire the same Library pointer
 			// (or create yet another new named library). In the latter case, the duplicate library will be deleted later, since we will re-check (under a lock)
 			// the presence of the same name in NamedLibrariesStack.  In the former case (two threads sharing the same Library pointer), we rely on FNamedShaderLibrary::OpenShaderCode
 			// implementation being thread-safe (which it is).
 
+			// more info for better logging
+			bool bOpenedAsMonolithic = false;
+
 			// if we're able to open the library by name, it's not chunked
 			if (Library->OpenShaderCode(Directory, Name))
 			{
 				bResult = true;
+				bOpenedAsMonolithic = true;
 
 				// Attempt to open the shared-cooked override code library if there is one.
 				// This is probably not ideal, but it should get shared-cooks working.
@@ -2489,7 +2494,6 @@ public:
 				}
 				NewComponentIDs = Library->PresentChunks.Difference(PrevComponentSet);
 				bResult = !NewComponentIDs.IsEmpty();
-				UE_LOG(LogShaderLibrary, Display, TEXT("Logical shader library '%s' component count %d, new components: %d"), *Name, Library->GetNumComponents(), NewComponentIDs.Num());
 
 #if UE_SHADERLIB_SUPPORT_CHUNK_DISCOVERY
 				if (!bResult)
@@ -2546,7 +2550,14 @@ public:
 					TUniquePtr<FNamedShaderLibrary>* LibraryPtr = NamedLibrariesStack.Find(Name);
 					if (LibraryPtr == nullptr)
 					{
-						UE_LOG(LogShaderLibrary, Display, TEXT("Logical shader library '%s' has been created, components %d"), *Name, Library->GetNumComponents());
+						if (bOpenedAsMonolithic)
+						{
+							UE_LOG(LogShaderLibrary, Display, TEXT("Logical shader library '%s' has been created as a monolithic library"), *Name, Library->GetNumComponents());
+						}
+						else
+						{
+							UE_LOG(LogShaderLibrary, Display, TEXT("Logical shader library '%s' has been created, components %d"), *Name, Library->GetNumComponents());
+						}
 						NamedLibrariesStack.Emplace(Name, Library);
 					}
 					else 
@@ -2556,6 +2567,10 @@ public:
 						Library = nullptr;
 					}
 				}
+				else
+				{
+					UE_LOG(LogShaderLibrary, Display, TEXT("Discovered new %d components for logical shader library '%s' (total number of components is now %d)"), NewComponentIDs.Num(), *Name, Library->GetNumComponents());
+				}
 
 				// Inform the pipeline cache that the state of loaded libraries has changed (unless we had to delete the duplicate)
 				if (Library != nullptr)
@@ -2567,6 +2582,21 @@ public:
 					}
 				}
 			}
+			else
+			{
+				if (bAddNewNamedLibrary)
+				{
+					UE_LOG(LogShaderLibrary, Display, TEXT("Tried to open shader library '%s', but could not find it neither as a monolithic library nor as a chunked one."), *Name);
+
+					check(Library->GetNumComponents() == 0);
+					delete Library;
+					Library = nullptr;
+				}
+				else
+				{
+					UE_LOG(LogShaderLibrary, Display, TEXT("Tried to open again shader library '%s', but could not find new components for it (existing components: %d)."), *Name, Library->GetNumComponents());
+				}
+			}
 		}
 
 #if WITH_EDITOR
