Fixing PSO deadlocked in Pre loading screen.

diff --git a/Engine/Source/Runtime/RenderCore/Private/ShaderPipelineCache.cpp b/Engine/Source/Runtime/RenderCore/Private/ShaderPipelineCache.cpp
--- a/Engine/Source/Runtime/RenderCore/Private/ShaderPipelineCache.cpp
+++ b/Engine/Source/Runtime/RenderCore/Private/ShaderPipelineCache.cpp
@@ -411,15 +411,6 @@ public:
 	void PollShutdownItems();
 	void Flush();
 	void Tick(float DeltaTime);
-
-	bool HasRHIFenceCompleted() 
-	{
-		if (LastPrecompileRHIFence.GetReference() && LastPrecompileRHIFence->IsComplete())
-		{
-			LastPrecompileRHIFence = nullptr;
-		}
-		return LastPrecompileRHIFence.GetReference()==nullptr;
-	}
 	bool OnShaderLibraryStateChanged(FShaderPipelineCache::ELibraryState State, EShaderPlatform Platform, FString const& Name, int32 ComponentID);
 private:
 	FString PSOCacheKey;
@@ -431,8 +422,6 @@ private:
 	TDoubleLinkedList<FPipelineCacheFileFormatPSORead*> FetchTasks;
 	FGuid CacheFileGuid;
 
-	FGraphEventRef LastPrecompileRHIFence;
-
 	TSet<uint64> CompletedMasks;
 	FShaderPipelineCache::FShaderCachePrecompileContext PrecompileContext;
 	bool bPreOptimizing = false;
@@ -1436,14 +1425,9 @@ bool FShaderPipelineCacheTask::ReadyForPrecompile()
 			++i;
 		}
 	}
-	if (LastPrecompileRHIFence.GetReference() && LastPrecompileRHIFence->IsComplete())
-	{
-		LastPrecompileRHIFence = nullptr;
-	}
 
 	return (CompileTasks.Num() != 0) 
-		&& PipelineStateCache::GetNumActivePipelinePrecompileTasks() < GMinPipelinePrecompileTasksInFlightThreshold
-		&& !LastPrecompileRHIFence.GetReference();
+		&& PipelineStateCache::GetNumActivePipelinePrecompileTasks() < GMinPipelinePrecompileTasksInFlightThreshold;
 }
 
 void FShaderPipelineCacheTask::PrecompilePipelineBatch()
@@ -1502,13 +1486,6 @@ void FShaderPipelineCacheTask::PrecompilePipelineBatch()
 	TotalPSOsCompiled += NumToPrecompile;
     TotalActiveTasks -= NumToPrecompile;
 
-#if PLATFORM_ANDROID
-	if (NumToPrecompile > 0 && IsRunningRHIInSeparateThread())
-	{
- 		LastPrecompileRHIFence = FRHICommandListExecutor::GetImmediateCommandList().RHIThreadFence(false);
-	}
-#endif
-
 	CompileTasks.RemoveAt(0, NumToPrecompile);
 }
 
@@ -1766,7 +1743,7 @@ void FShaderPipelineCache::Tick(float DeltaTime)
 
 	// HACK: note for the time being until the code further upstream properly uses the OnPrecompilationComplete delegate, we need to provide some delays to make sure it gets the messaging properly (the -10.0f and the -20.0f below).
 	const bool bHasCurrentCacheTaskCompleted = bHasActiveCacheTask &&
-		PipelineStateCache::GetNumActivePipelinePrecompileTasks() == 0 && CurrentCacheTask->HasRHIFenceCompleted() &&
+		PipelineStateCache::GetNumActivePipelinePrecompileTasks() == 0 &&
 				(
 				FMath::Max((int64_t)0, FShaderPipelineCacheTask::TotalWaitingTasks.load()) == 0 &&
 				FMath::Max((int64_t)0, FShaderPipelineCacheTask::TotalActiveTasks.load()) == 0 &&
@@ -1837,11 +1814,6 @@ void FShaderPipelineCache::Tick(float DeltaTime)
 
 void FShaderPipelineCacheTask::Tick(float DeltaTime)
 {
-	if (LastPrecompileRHIFence.GetReference() && LastPrecompileRHIFence->IsComplete())
-	{
-		LastPrecompileRHIFence = nullptr;
-	}
-
 	if (PrecompileStartTime > 0.0)
 	{
 		TotalPrecompileWallTime = float(FPlatformTime::Seconds() - PrecompileStartTime);
