Adding utility macro DECLARE_EXPORTED_GLOBAL_SHADER that acts like DECLARE_GLOBAL_SHADER with the ability to specify a dllexport macro.

diff --git a/Engine/Source/Runtime/RenderCore/Public/GlobalShader.h b/Engine/Source/Runtime/RenderCore/Public/GlobalShader.h
--- a/Engine/Source/Runtime/RenderCore/Public/GlobalShader.h
+++ b/Engine/Source/Runtime/RenderCore/Public/GlobalShader.h
@@ -444,6 +444,7 @@ inline FGlobalShaderMap* GetGlobalShaderMap(ERHIFeatureLevel::Type FeatureLevel)
 		*/
 
 #define DECLARE_GLOBAL_SHADER(ShaderClass) DECLARE_SHADER_TYPE(ShaderClass, Global)
+#define DECLARE_EXPORTED_GLOBAL_SHADER(ShaderClass, RequiredAPI) DECLARE_EXPORTED_SHADER_TYPE(ShaderClass, Global, RequiredAPI)
 #define IMPLEMENT_GLOBAL_SHADER(ShaderClass,SourceFilename,FunctionName,Frequency) IMPLEMENT_SHADER_TYPE(,ShaderClass,TEXT(SourceFilename),TEXT(FunctionName),Frequency)
 
 /**
