Adding missing export.

diff --git a/Engine/Source/Runtime/RenderCore/Public/ResolveShader.h b/Engine/Source/Runtime/RenderCore/Public/ResolveShader.h
--- a/Engine/Source/Runtime/RenderCore/Public/ResolveShader.h
+++ b/Engine/Source/Runtime/RenderCore/Public/ResolveShader.h
@@ -27,7 +27,7 @@ public:
 	FResolveDepthPS(const ShaderMetaType::CompiledShaderInitializerType& Initializer);
 
 	static void ModifyCompilationEnvironment(const FGlobalShaderPermutationParameters& Parameters, FShaderCompilerEnvironment& OutEnvironment);
-	void SetParameters(FRHICommandList& RHICmdList, FParameter);
+	RENDERCORE_API void SetParameters(FRHICommandList& RHICmdList, FParameter);
 
 	LAYOUT_FIELD(FShaderResourceParameter, UnresolvedSurface);
 };
