Changed logging category of shader hash calculation from LogTemp to LogShaders. This reduces the amount of LogTemp messages when setting it to VeryVerbose in a fresh project.

diff --git a/Engine/Source/Runtime/RenderCore/Private/ShaderCore.cpp b/Engine/Source/Runtime/RenderCore/Private/ShaderCore.cpp
--- a/Engine/Source/Runtime/RenderCore/Private/ShaderCore.cpp
+++ b/Engine/Source/Runtime/RenderCore/Private/ShaderCore.cpp
@@ -1661,9 +1661,9 @@ static void UpdateSingleShaderFilehash(FSHA1& InOutHashState, const TCHAR* Virtu
 	TArray<FString> IncludeVirtualFilePaths;
 	GetShaderIncludes(VirtualFilePath, VirtualFilePath, IncludeVirtualFilePaths, ShaderPlatform);
 #if WITH_EDITOR &&  !(UE_BUILD_SHIPPING || UE_BUILD_TEST)
-	if ( UE_LOG_ACTIVE(LogTemp, Verbose) )
+	if ( UE_LOG_ACTIVE(LogShaders, Verbose) )
 	{
-		UE_LOG(LogTemp, Verbose, TEXT("Generating hash of file %s, "), VirtualFilePath);
+		UE_LOG(LogShaders, Verbose, TEXT("Generating hash of file %s, "), VirtualFilePath);
 	}
 #endif
 	for (int32 IncludeIndex = 0; IncludeIndex < IncludeVirtualFilePaths.Num(); IncludeIndex++)
@@ -1673,12 +1673,12 @@ static void UpdateSingleShaderFilehash(FSHA1& InOutHashState, const TCHAR* Virtu
 		LoadShaderSourceFileChecked(*IncludeVirtualFilePaths[IncludeIndex], ShaderPlatform, IncludeFileContents);
 		InOutHashState.UpdateWithString(*IncludeFileContents, IncludeFileContents.Len());
 #if WITH_EDITOR &&  !(UE_BUILD_SHIPPING || UE_BUILD_TEST)
-		if (UE_LOG_ACTIVE(LogTemp, Verbose))
+		if (UE_LOG_ACTIVE(LogShaders, Verbose))
 		{
 			uint8 HashBytes[20];
 			FSHA1::HashBuffer(&InOutHashState, sizeof(FSHA1), HashBytes);
 			
-			UE_LOG(LogTemp, Verbose, TEXT("Processing include file for %s, %s, %s"), VirtualFilePath, *IncludeVirtualFilePaths[IncludeIndex], *BytesToHex(HashBytes, 20));
+			UE_LOG(LogShaders, Verbose, TEXT("Processing include file for %s, %s, %s"), VirtualFilePath, *IncludeVirtualFilePaths[IncludeIndex], *BytesToHex(HashBytes, 20));
 		}
 #endif
 	}
@@ -1741,7 +1741,7 @@ const FSHAHash& GetShaderFileHash(const TCHAR* VirtualFilePath, EShaderPlatform
 		HashState.GetHash(&NewHash.Hash[0]);
 
 #if WITH_EDITOR &&  !(UE_BUILD_SHIPPING || UE_BUILD_TEST)
-		UE_LOG(LogTemp, Verbose, TEXT("Final hash for file %s, %s"), VirtualFilePath,*BytesToHex(&NewHash.Hash[0], 20));
+		UE_LOG(LogShaders, Verbose, TEXT("Final hash for file %s, %s"), VirtualFilePath,*BytesToHex(&NewHash.Hash[0], 20));
 #endif
 		INC_FLOAT_STAT_BY(STAT_ShaderCompiling_HashingShaderFiles, (float)HashTime);
 		return NewHash;
