Use SceneRenderingAllocator for RDG allocated arrays.

diff --git a/Engine/Source/Runtime/RenderCore/Public/RenderGraphBuilder.h b/Engine/Source/Runtime/RenderCore/Public/RenderGraphBuilder.h
--- a/Engine/Source/Runtime/RenderCore/Public/RenderGraphBuilder.h
+++ b/Engine/Source/Runtime/RenderCore/Public/RenderGraphBuilder.h
@@ -149,7 +149,7 @@ public:
 
 	/** Allocates a C++ array where both the array and the data are tied to the lifetime of the graph. The array itself is safe to pass into an RDG lambda. */
 	template <typename ObjectType>
-	TArray<ObjectType, FRDGArrayAllocator>& AllocArray();
+	TArray<ObjectType, SceneRenderingAllocator>& AllocArray();
 
 	/** Allocates a parameter struct with a lifetime tied to graph execution. */
 	template <typename ParameterStructType>
diff --git a/Engine/Source/Runtime/RenderCore/Public/RenderGraphBuilder.inl b/Engine/Source/Runtime/RenderCore/Public/RenderGraphBuilder.inl
--- a/Engine/Source/Runtime/RenderCore/Public/RenderGraphBuilder.inl
+++ b/Engine/Source/Runtime/RenderCore/Public/RenderGraphBuilder.inl
@@ -150,9 +150,9 @@ FORCEINLINE ObjectType* FRDGBuilder::AllocObject(TArgs&&... Args)
 }
 
 template <typename ObjectType>
-FORCEINLINE TArray<ObjectType, FRDGArrayAllocator>& FRDGBuilder::AllocArray()
+FORCEINLINE TArray<ObjectType, SceneRenderingAllocator>& FRDGBuilder::AllocArray()
 {
-	return *Allocator.Alloc<TArray<ObjectType, FRDGArrayAllocator>>();
+	return *Allocator.Alloc<TArray<ObjectType, SceneRenderingAllocator>>();
 }
 
 template <typename ParameterStructType>
