Implements r.DumpGPU.FrameCount

diff --git a/Engine/Extras/GPUDumpViewer/GPUDumpViewer.html b/Engine/Extras/GPUDumpViewer/GPUDumpViewer.html
--- a/Engine/Extras/GPUDumpViewer/GPUDumpViewer.html
+++ b/Engine/Extras/GPUDumpViewer/GPUDumpViewer.html
@@ -1202,6 +1202,9 @@ const k_tips = [
 	// Dumping process
 	'Can use the CTRL+SHIFT+/ keyboard shortcut to summon the DumpGPU command.',
 	'Speed up your frame dump by only selecting the passes you need with r.DumpGPU.Root. For instance r.DumpGPU.Root="*PostProcessing*".',
+	'Uses r.DumpGPU.Delay to delay the dump of few seconds to have time to repro the issue with gameplay logic (for instance moving arround in the map).',
+	'Uses r.DumpGPU.FrameCount to dump more than one frame. This is useful when artifact may be produced only some frames but end up stucked in temporal histories.',
+	'Uses r.DumpGPU.DumpOnScreenshotTest=1 to automatically produce a GPU dump of the frame producing screenshot in the automated AScreenshotFunctionalTest',
 	'GPU dumps can be large and accumulate on your hard drive in your various projects\' Saved/ directories. Set r.DumpGPU.Directory="D:/tmp/DumpGPU/" in your console variables or UE-DumpGPUPath environment variable to dump them all at the same location on your machine.',
 	'Uses -cvarsini to override console variables with your own ConsoleVariables.ini file on a cooked build.',
 
diff --git a/Engine/Source/Runtime/RenderCore/Private/DumpGPU.cpp b/Engine/Source/Runtime/RenderCore/Private/DumpGPU.cpp
--- a/Engine/Source/Runtime/RenderCore/Private/DumpGPU.cpp
+++ b/Engine/Source/Runtime/RenderCore/Private/DumpGPU.cpp
@@ -90,6 +90,11 @@ static TAutoConsoleVariable<float> GDumpGPUDelay(
 	TEXT("Delay in seconds before dumping the frame."),
 	ECVF_RenderThreadSafe);
 
+static TAutoConsoleVariable<int32> GDumpGPUFrameCount(
+	TEXT("r.DumpGPU.FrameCount"), 1,
+	TEXT("Number of consecutive frames to dump (default=1)."),
+	ECVF_Default);
+
 static TAutoConsoleVariable<int32> GDumpGPUDraws(
 	TEXT("r.DumpGPU.Draws"), 0,
 	TEXT("Whether to dump resource after each individual draw call (disabled by default)."),
@@ -203,6 +208,7 @@ public:
 
 	bool bEnableDiskWrite = false;
 	bool bUpload = false;
+	int32 RemainingFrameCount = 1;
 	FName UploadResourceCompressionName;
 	FString DumpingDirectoryPath;
 	FDateTime Time;
@@ -1684,6 +1690,7 @@ FString FRDGBuilder::BeginResourceDump(const TCHAR* Cmd)
 			NewResourceDumpContext->DumpingDirectoryPath = DirectoryPath / FApp::GetProjectName() + TEXT("-") + FPlatformProperties::PlatformName() + TEXT("-") + NewResourceDumpContext->Time.ToString() + TEXT("/");
 		}
 		NewResourceDumpContext->bEnableDiskWrite = GDumpTestEnableDiskWrite.GetValueOnGameThread() != 0;
+		NewResourceDumpContext->RemainingFrameCount = FMath::Max(GDumpGPUFrameCount.GetValueOnGameThread(), 1);
 
 		if (Switches.Contains(TEXT("upload")))
 		{
@@ -1875,6 +1882,12 @@ void FRDGBuilder::EndResourceDump()
 		return;
 	}
 
+	GRDGResourceDumpContext->RemainingFrameCount--;
+	if (GRDGResourceDumpContext->RemainingFrameCount > 0)
+	{
+		return;
+	}
+
 	// Wait all rendering commands are completed to finish with GRDGResourceDumpContext.
 	{
 		UE_LOG(LogRendererCore, Display, TEXT("Stalling game thread until render thread finishes to dump resources"));
