Big push to convert more shader types to use batched parameters.
- Media Shaders
- Slate
- Particles

diff --git a/Engine/Plugins/Experimental/AVCodecs/AVCodecsCore/Source/AVCodecsCoreRHI/Private/Video/Resources/VideoResourceRHI.cpp b/Engine/Plugins/Experimental/AVCodecs/AVCodecsCore/Source/AVCodecsCoreRHI/Private/Video/Resources/VideoResourceRHI.cpp
--- a/Engine/Plugins/Experimental/AVCodecs/AVCodecsCore/Source/AVCodecsCoreRHI/Private/Video/Resources/VideoResourceRHI.cpp
+++ b/Engine/Plugins/Experimental/AVCodecs/AVCodecsCore/Source/AVCodecsCoreRHI/Private/Video/Resources/VideoResourceRHI.cpp
@@ -311,7 +311,7 @@ TSharedPtr<FVideoResourceRHI> FVideoResourceRHI::TransformResource(FVideoDescrip
 						TShaderMapRef<FNV12ConvertAsBytesPS> ConvertShader(ShaderMap);
 						GraphicsPSOInit.BoundShaderState.PixelShaderRHI = ConvertShader.GetPixelShader();
 						SetGraphicsPipelineState(RHICmdList, GraphicsPSOInit, 0);
-						ConvertShader->SetParameters(RHICmdList, Source, OutputDim, YUVToRGBMatrix, MediaShaders::YUVOffset8bits, true);
+						SetShaderParametersLegacyPS(RHICmdList, ConvertShader, Source, OutputDim, YUVToRGBMatrix, MediaShaders::YUVOffset8bits, true);
 
 						// draw full size quad into render target
 						FBufferRHIRef VertexBuffer = CreateTempMediaVertexBuffer();
@@ -377,7 +377,7 @@ void FVideoResourceRHI::TransformResourceTo(FRHICommandListImmediate& RHICmdList
 					TShaderMapRef<FNV12ConvertAsBytesPS> ConvertShader(ShaderMap);
 					GraphicsPSOInit.BoundShaderState.PixelShaderRHI = ConvertShader.GetPixelShader();
 					SetGraphicsPipelineState(RHICmdList, GraphicsPSOInit, 0);
-					ConvertShader->SetParameters(RHICmdList, Raw.Texture, OutputDim, YUVToRGBMatrix, MediaShaders::YUVOffset8bits, true);
+					SetShaderParametersLegacyPS(RHICmdList, ConvertShader, Raw.Texture, OutputDim, YUVToRGBMatrix, MediaShaders::YUVOffset8bits, true);
 
 					// draw full size quad into render target
 					FBufferRHIRef VertexBuffer = CreateTempMediaVertexBuffer();
diff --git a/Engine/Plugins/Media/AvfMedia/Source/AvfMedia/Private/Player/AvfMediaVideoSampler.cpp b/Engine/Plugins/Media/AvfMedia/Source/AvfMedia/Private/Player/AvfMediaVideoSampler.cpp
--- a/Engine/Plugins/Media/AvfMedia/Source/AvfMedia/Private/Player/AvfMediaVideoSampler.cpp
+++ b/Engine/Plugins/Media/AvfMedia/Source/AvfMedia/Private/Player/AvfMediaVideoSampler.cpp
@@ -268,7 +268,7 @@ void FAvfMediaVideoSampler::ProcessFrame(CVPixelBufferRef Frame, FTimespan Sampl
 
 					SetGraphicsPipelineState(RHICmdList, GraphicsPSOInit, 0);
 
-					PixelShader->SetParameters(RHICmdList, YTex, UVTex, *ColorTransform, MediaShaders::YUVOffset8bits, true);
+					SetShaderParametersLegacyPS(RHICmdList, PixelShader, YTex, UVTex, *ColorTransform, MediaShaders::YUVOffset8bits, true);
 
 					FBufferRHIRef VertexBuffer = CreateTempMediaVertexBuffer();
 					RHICmdList.SetStreamSource(0, VertexBuffer, 0);
diff --git a/Engine/Plugins/Media/ElectraUtil/Source/ElectraSamples/Private/ElectraTextureSampleAVF.cpp b/Engine/Plugins/Media/ElectraUtil/Source/ElectraSamples/Private/ElectraTextureSampleAVF.cpp
--- a/Engine/Plugins/Media/ElectraUtil/Source/ElectraSamples/Private/ElectraTextureSampleAVF.cpp
+++ b/Engine/Plugins/Media/ElectraUtil/Source/ElectraSamples/Private/ElectraTextureSampleAVF.cpp
@@ -248,7 +248,7 @@ void FElectraMediaTexConvApple::ConvertTexture(FTexture2DRHIRef & InDstTexture,
 						TShaderMapRef<FMediaShadersVS> VertexShader(GlobalShaderMap);
 						TShaderMapRef<FYCbCrConvertPS> PixelShader(GlobalShaderMap);
 
-						PixelShader->SetParameters(RHICmdList, YTex, UVTex, *ColorTransform, Off, true);
+						SetShaderParametersLegacyPS(RHICmdList, PixelShader, YTex, UVTex, *ColorTransform, Off, true);
 
 						GraphicsPSOInit.BoundShaderState.VertexShaderRHI = VertexShader.GetVertexShader();
 						GraphicsPSOInit.BoundShaderState.PixelShaderRHI = PixelShader.GetPixelShader();
@@ -273,7 +273,7 @@ void FElectraMediaTexConvApple::ConvertTexture(FTexture2DRHIRef & InDstTexture,
 						FShaderResourceViewRHIRef UV_SRV = RHICreateShaderResourceView(UVTex, 0, 1, PF_G16R16);
 
 						// Update shader uniform parameters.
-						PixelShader->SetParameters(RHICmdList, { YWidth, YHeight }, Y_SRV, UV_SRV, { YWidth, YHeight }, YUVMtx, ColorSpaceMtx, EncodingType == UE::Color::EEncoding::ST2084);
+						SetShaderParametersLegacyPS(RHICmdList, PixelShader, FIntPoint(YWidth, YHeight), Y_SRV, UV_SRV, FIntPoint(YWidth, YHeight), YUVMtx, ColorSpaceMtx, EncodingType == UE::Color::EEncoding::ST2084);
 
 						GraphicsPSOInit.BoundShaderState.VertexShaderRHI = VertexShader.GetVertexShader();
 						GraphicsPSOInit.BoundShaderState.PixelShaderRHI = PixelShader.GetPixelShader();
diff --git a/Engine/Plugins/Media/ElectraUtil/Source/ElectraSamples/Private/ElectraTextureSampleAndroid.cpp b/Engine/Plugins/Media/ElectraUtil/Source/ElectraSamples/Private/ElectraTextureSampleAndroid.cpp
--- a/Engine/Plugins/Media/ElectraUtil/Source/ElectraSamples/Private/ElectraTextureSampleAndroid.cpp
+++ b/Engine/Plugins/Media/ElectraUtil/Source/ElectraSamples/Private/ElectraTextureSampleAndroid.cpp
@@ -538,7 +538,7 @@ bool FElectraTextureSample::Convert(FTexture2DRHIRef& InDstTexture, const FConve
 	SetGraphicsPipelineState(RHICmdList, GraphicsPSOInit, 0);
 
 	// Update shader uniform parameters.
-	PixelShader->SetParameters(RHICmdList, InputTexture, OutputDim, GetEncodingType() == UE::Color::EEncoding::sRGB, GetEncodingType() == UE::Color::EEncoding::ST2084, ColorSpaceMtx);
+	SetShaderParametersLegacyPS(RHICmdList, PixelShader, InputTexture, OutputDim, GetEncodingType() == UE::Color::EEncoding::sRGB, GetEncodingType() == UE::Color::EEncoding::ST2084, ColorSpaceMtx);
 
 	RHICmdList.SetStreamSource(0, CreateTempMediaVertexBuffer(), 0);
 
diff --git a/Engine/Plugins/Media/WebMMedia/Source/WebMMedia/Private/Player/WebMVideoDecoder.cpp b/Engine/Plugins/Media/WebMMedia/Source/WebMMedia/Private/Player/WebMVideoDecoder.cpp
--- a/Engine/Plugins/Media/WebMMedia/Source/WebMMedia/Private/Player/WebMVideoDecoder.cpp
+++ b/Engine/Plugins/Media/WebMMedia/Source/WebMMedia/Private/Player/WebMVideoDecoder.cpp
@@ -283,7 +283,8 @@ void FWebMVideoDecoder::ConvertYUVToRGBAndSubmit(const FConvertParams& Params)
 			}
 
 			SetGraphicsPipelineState(CommandList, GraphicsPSOInit, 0);
-			PixelShader->SetParameters(CommandList, DecodedY->GetTexture2D(), DecodedU->GetTexture2D(), DecodedV->GetTexture2D(), FIntPoint(Image->d_w, Image->d_h), MediaShaders::YuvToRgbRec709Scaled, MediaShaders::YUVOffset8bits, true);
+
+			SetShaderParametersLegacyPS(CommandList, PixelShader, DecodedY->GetTexture2D(), DecodedU->GetTexture2D(), DecodedV->GetTexture2D(), FIntPoint(Image->d_w, Image->d_h), MediaShaders::YuvToRgbRec709Scaled, MediaShaders::YUVOffset8bits, true);
 
 			// draw full-size quad
 			CommandList.SetViewport(0, 0, 0.0f, Image->d_w, Image->d_h, 1.0f);
diff --git a/Engine/Plugins/Runtime/WindowsMoviePlayer/Source/WindowsMoviePlayer/Private/WindowsMovieStreamer.cpp b/Engine/Plugins/Runtime/WindowsMoviePlayer/Source/WindowsMoviePlayer/Private/WindowsMovieStreamer.cpp
--- a/Engine/Plugins/Runtime/WindowsMoviePlayer/Source/WindowsMoviePlayer/Private/WindowsMovieStreamer.cpp
+++ b/Engine/Plugins/Runtime/WindowsMoviePlayer/Source/WindowsMoviePlayer/Private/WindowsMovieStreamer.cpp
@@ -140,7 +140,7 @@ void FMediaFoundationMovieStreamer::ConvertSample()
 			TShaderMapRef<FBMPConvertPS> ConvertShader(ShaderMap);
 			GraphicsPSOInit.BoundShaderState.PixelShaderRHI = ConvertShader.GetPixelShader();
 			SetGraphicsPipelineState(CommandList, GraphicsPSOInit, 0);
-			ConvertShader->SetParameters(CommandList, InputTarget, OutputDim, bSampleIsOutputSrgb && !SrgbOutput);
+			SetShaderParametersLegacyPS(CommandList, ConvertShader, InputTarget, OutputDim, bSampleIsOutputSrgb && !SrgbOutput);
 		}
 		break;
 
@@ -149,7 +149,7 @@ void FMediaFoundationMovieStreamer::ConvertSample()
 			TShaderMapRef<FYUY2ConvertPS> ConvertShader(ShaderMap);
 			GraphicsPSOInit.BoundShaderState.PixelShaderRHI = ConvertShader.GetPixelShader();
 			SetGraphicsPipelineState(CommandList, GraphicsPSOInit, 0);
-			ConvertShader->SetParameters(CommandList, InputTarget, OutputDim, MediaShaders::YuvToRgbRec709Scaled, MediaShaders::YUVOffset8bits, bSampleIsOutputSrgb);
+			SetShaderParametersLegacyPS(CommandList, ConvertShader, InputTarget, OutputDim, MediaShaders::YuvToRgbRec709Scaled, MediaShaders::YUVOffset8bits, bSampleIsOutputSrgb);
 		}
 		break;
 
diff --git a/Engine/Source/Runtime/Engine/Private/BatchedElements.cpp b/Engine/Source/Runtime/Engine/Private/BatchedElements.cpp
--- a/Engine/Source/Runtime/Engine/Private/BatchedElements.cpp
+++ b/Engine/Source/Runtime/Engine/Private/BatchedElements.cpp
@@ -686,9 +686,8 @@ void FBatchedElements::PrepareShaders(
 					GraphicsPSOInit.BoundShaderState.PixelShaderRHI = MaskedPixelShader.GetPixelShader();
 					SetGraphicsPipelineState(RHICmdList, GraphicsPSOInit, StencilRef);
 
-					MaskedPixelShader->SetEditorCompositingParameters(RHICmdList, View);
 
-					SetShaderParametersLegacyPS(RHICmdList, MaskedPixelShader, Texture, Gamma, OpacityMaskRefVal, BlendMode);
+					SetShaderParametersLegacyPS(RHICmdList, MaskedPixelShader, View, Texture, Gamma, OpacityMaskRefVal, BlendMode);
 				}
 				else
 				{
@@ -697,8 +696,7 @@ void FBatchedElements::PrepareShaders(
 
 					SetGraphicsPipelineState(RHICmdList, GraphicsPSOInit, StencilRef);
 
-					MaskedPixelShader->SetEditorCompositingParameters(RHICmdList, View);
-					SetShaderParametersLegacyPS(RHICmdList, MaskedPixelShader, Texture, Gamma, OpacityMaskRefVal, BlendMode);
+					SetShaderParametersLegacyPS(RHICmdList, MaskedPixelShader, View, Texture, Gamma, OpacityMaskRefVal, BlendMode);
 				}
 			}
 			// render distance field elements
@@ -751,9 +749,6 @@ void FBatchedElements::PrepareShaders(
 					(GlowInfo != NULL) ? *GlowInfo : FDepthFieldGlowInfo(),
 					BlendMode
 					);
-
-				// This shader does not use editor compositing
-				DistanceFieldPixelShader->SetEditorCompositingParameters(RHICmdList, nullptr);
 			}
 			else if(BlendMode == SE_BLEND_TranslucentAlphaOnly || BlendMode == SE_BLEND_TranslucentAlphaOnlyWriteAlpha)
 			{
@@ -766,8 +761,7 @@ void FBatchedElements::PrepareShaders(
 
 					SetGraphicsPipelineState(RHICmdList, GraphicsPSOInit, StencilRef);
 
-					SetShaderParametersLegacyPS(RHICmdList, AlphaOnlyPixelShader, Texture);
-					AlphaOnlyPixelShader->SetEditorCompositingParameters(RHICmdList, View);
+					SetShaderParametersLegacyPS(RHICmdList, AlphaOnlyPixelShader, View, Texture);
 				}
 				else
 				{
@@ -776,8 +770,7 @@ void FBatchedElements::PrepareShaders(
 
 					SetGraphicsPipelineState(RHICmdList, GraphicsPSOInit, StencilRef);
 
-					SetShaderParametersLegacyPS(RHICmdList, GammaAlphaOnlyPixelShader, Texture, Gamma, BlendMode);
-					GammaAlphaOnlyPixelShader->SetEditorCompositingParameters(RHICmdList, View);
+					SetShaderParametersLegacyPS(RHICmdList, GammaAlphaOnlyPixelShader, View, Texture, Gamma, BlendMode);
 				}
 			}
 			else if(BlendMode >= SE_BLEND_RGBA_MASK_START && BlendMode <= SE_BLEND_RGBA_MASK_END)
@@ -800,8 +793,7 @@ void FBatchedElements::PrepareShaders(
 
 					SetGraphicsPipelineState(RHICmdList, GraphicsPSOInit, StencilRef);
 
-					SetShaderParametersLegacyPS(RHICmdList, PixelShader, Texture);
-					PixelShader->SetEditorCompositingParameters(RHICmdList, View);
+					SetShaderParametersLegacyPS(RHICmdList, PixelShader, View, Texture);
 				}
 				else
 				{
@@ -820,8 +812,7 @@ void FBatchedElements::PrepareShaders(
 					GraphicsPSOInit.BoundShaderState.PixelShaderRHI = BasePixelShader.GetPixelShader();
 					SetGraphicsPipelineState(RHICmdList, GraphicsPSOInit, StencilRef);
 
-					SetShaderParametersLegacyPS(RHICmdList, BasePixelShader, Texture, Gamma, BlendMode);
-					BasePixelShader->SetEditorCompositingParameters(RHICmdList, View);
+					SetShaderParametersLegacyPS(RHICmdList, BasePixelShader, View, Texture, Gamma, BlendMode);
 				}
 			}
 		}
diff --git a/Engine/Source/Runtime/Engine/Private/Particles/ParticleCurveTexture.cpp b/Engine/Source/Runtime/Engine/Private/Particles/ParticleCurveTexture.cpp
--- a/Engine/Source/Runtime/Engine/Private/Particles/ParticleCurveTexture.cpp
+++ b/Engine/Source/Runtime/Engine/Private/Particles/ParticleCurveTexture.cpp
@@ -73,15 +73,14 @@ public:
 	/**
 	* Sets parameters for particle injection.
 	*/
-	void SetParameters(FRHICommandList& RHICmdList, const FVector2D& CurveOffset)
+	void SetParameters(FRHIBatchedShaderParameters& BatchedParameters, const FVector2D& CurveOffset)
 	{
 		FParticleCurveInjectionParameters Parameters;
 		Parameters.PixelScale.X = 1.0f / GParticleCurveTextureSizeX;
 		Parameters.PixelScale.Y = 1.0f / GParticleCurveTextureSizeY;
 		Parameters.CurveOffset = FVector2f(CurveOffset);
 		FParticleCurveInjectionBufferRef UniformBuffer = FParticleCurveInjectionBufferRef::CreateUniformBufferImmediate(Parameters, UniformBuffer_SingleDraw);
-		FRHIVertexShader* VertexShader = RHICmdList.GetBoundVertexShader();
-		SetUniformBufferParameter(RHICmdList, VertexShader, GetUniformBufferParameter<FParticleCurveInjectionParameters>(), UniformBuffer);
+		SetUniformBufferParameter(BatchedParameters, GetUniformBufferParameter<FParticleCurveInjectionParameters>(), UniformBuffer);
 	}
 };
 
@@ -261,7 +260,7 @@ static void InjectCurves(
 			// Stream 0: TexCoord.
 			RHICmdList.SetStreamSource(1, GParticleTexCoordVertexBuffer.VertexBufferRHI, 0);
 
-			VertexShader->SetParameters(RHICmdList, CurveOffset);
+			SetShaderParametersLegacyVS(RHICmdList, VertexShader, CurveOffset);
 
 			// Inject particles.
 			RHICmdList.DrawIndexedPrimitive(
diff --git a/Engine/Source/Runtime/Engine/Private/Particles/ParticleGpuSimulation.cpp b/Engine/Source/Runtime/Engine/Private/Particles/ParticleGpuSimulation.cpp
--- a/Engine/Source/Runtime/Engine/Private/Particles/ParticleGpuSimulation.cpp
+++ b/Engine/Source/Runtime/Engine/Private/Particles/ParticleGpuSimulation.cpp
@@ -957,7 +957,7 @@ public:
  */
 class FParticleSimulationClearPS : public FGlobalShader
 {
-	DECLARE_SHADER_TYPE(FParticleSimulationClearPS,Global);
+	DECLARE_GLOBAL_SHADER(FParticleSimulationClearPS);
 
 public:
 
@@ -1450,7 +1450,7 @@ typedef TUniformBufferRef<FParticleInjectionParameters> FParticleInjectionBuffer
  */
 class FParticleInjectionVS : public FGlobalShader
 {
-	DECLARE_SHADER_TYPE(FParticleInjectionVS,Global);
+	DECLARE_GLOBAL_SHADER(FParticleInjectionVS);
 
 public:
 
@@ -1459,11 +1459,6 @@ public:
 		return SupportsGPUParticles(Parameters.Platform);
 	}
 
-	static void ModifyCompilationEnvironment(const FGlobalShaderPermutationParameters& Parameters, FShaderCompilerEnvironment& OutEnvironment)
-	{
-		FGlobalShader::ModifyCompilationEnvironment( Parameters, OutEnvironment );
-	}
-
 	/** Default constructor. */
 	FParticleInjectionVS()
 	{
@@ -1478,14 +1473,13 @@ public:
 	/**
 	 * Sets parameters for particle injection.
 	 */
-	void SetParameters(FRHICommandList& RHICmdList)
+	void SetParameters(FRHIBatchedShaderParameters& BatchedParameters)
 	{
 		FParticleInjectionParameters Parameters;
 		Parameters.PixelScale.X = 1.0f / GParticleSimulationTextureSizeX;
 		Parameters.PixelScale.Y = 1.0f / GParticleSimulationTextureSizeY;
 		FParticleInjectionBufferRef UniformBuffer = FParticleInjectionBufferRef::CreateUniformBufferImmediate( Parameters, UniformBuffer_SingleDraw );
-		FRHIVertexShader* VertexShader = RHICmdList.GetBoundVertexShader();
-		SetUniformBufferParameter(RHICmdList, VertexShader, GetUniformBufferParameter<FParticleInjectionParameters>(), UniformBuffer );
+		SetUniformBufferParameter(BatchedParameters, GetUniformBufferParameter<FParticleInjectionParameters>(), UniformBuffer );
 	}
 };
 
@@ -1495,7 +1489,7 @@ public:
 template <bool StaticPropertiesOnly>
 class TParticleInjectionPS : public FGlobalShader
 {
-	DECLARE_SHADER_TYPE(TParticleInjectionPS,Global);
+	DECLARE_GLOBAL_SHADER(TParticleInjectionPS);
 
 public:
 
@@ -1626,8 +1620,7 @@ void InjectNewParticles(FRHICommandList& RHICmdList, FGraphicsPipelineStateIniti
 
 		SetGraphicsPipelineState(RHICmdList, GraphicsPSOInit, 0);
 
-		
-		VertexShader->SetParameters(RHICmdList);
+		SetShaderParametersLegacyVS(RHICmdList, VertexShader);
 
 		// Stream 0: New particles.
 		RHICmdList.SetStreamSource(
diff --git a/Engine/Source/Runtime/Engine/Private/SimpleElementShaders.cpp b/Engine/Source/Runtime/Engine/Private/SimpleElementShaders.cpp
--- a/Engine/Source/Runtime/Engine/Private/SimpleElementShaders.cpp
+++ b/Engine/Source/Runtime/Engine/Private/SimpleElementShaders.cpp
@@ -69,6 +69,19 @@ FSimpleElementPS::FSimpleElementPS(const ShaderMetaType::CompiledShaderInitializ
 	TextureComponentReplicateAlpha.Bind(Initializer.ParameterMap,TEXT("TextureComponentReplicateAlpha"));
 }
 
+void FSimpleElementPS::SetEditorCompositingParameters(FRHIBatchedShaderParameters& BatchedParameters, const FSceneView* View)
+{
+	if (View)
+	{
+		FGlobalShader::SetParameters<FViewUniformShaderParameters>(BatchedParameters, View->ViewUniformBuffer);
+	}
+	else
+	{
+		// Unset the view uniform buffers since we don't have a view
+		SetUniformBufferParameter(BatchedParameters, GetUniformBufferParameter<FViewUniformShaderParameters>(), nullptr);
+	}
+}
+
 void FSimpleElementPS::SetParameters(FRHIBatchedShaderParameters& BatchedParameters, const FTexture* TextureValue)
 {
 	SetTextureParameter(BatchedParameters, InTexture, InTextureSampler, TextureValue);
@@ -77,17 +90,17 @@ void FSimpleElementPS::SetParameters(FRHIBatchedShaderParameters& BatchedParamet
 	SetShaderValue(BatchedParameters, TextureComponentReplicateAlpha, TextureValue->bGreyScaleFormat ? FLinearColor(1, 0, 0, 0) : FLinearColor(0, 0, 0, 1));
 }
 
+void FSimpleElementPS::SetParameters(FRHIBatchedShaderParameters& BatchedParameters, const FSceneView* View, const FTexture* TextureValue)
+{
+	SetEditorCompositingParameters(BatchedParameters, View);
+	SetParameters(BatchedParameters, TextureValue);
+}
+
 void FSimpleElementPS::SetEditorCompositingParameters(FRHICommandList& RHICmdList, const FSceneView* View)
 {
-	if( View )
-	{
-		FGlobalShader::SetParameters<FViewUniformShaderParameters>(RHICmdList, RHICmdList.GetBoundPixelShader(), View->ViewUniformBuffer );
-	}
-	else
-	{
-		// Unset the view uniform buffers since we don't have a view
-		SetUniformBufferParameter(RHICmdList, RHICmdList.GetBoundPixelShader(), GetUniformBufferParameter<FViewUniformShaderParameters>(), NULL);
-	}
+	FRHIBatchedShaderParameters& BatchedParameters = RHICmdList.GetScratchShaderParameters();
+	SetEditorCompositingParameters(BatchedParameters, View);
+	RHICmdList.SetBatchedShaderParameters(RHICmdList.GetBoundPixelShader(), BatchedParameters);
 }
 
 void FSimpleElementPS::SetParameters(FRHICommandList& RHICmdList, const FTexture* TextureValue)
@@ -115,6 +128,12 @@ void FSimpleElementGammaBasePS::SetParameters(FRHIBatchedShaderParameters& Batch
 	SetShaderValue(BatchedParameters, Gamma, GammaValue);
 }
 
+void FSimpleElementGammaBasePS::SetParameters(FRHIBatchedShaderParameters& BatchedParameters, const FSceneView* View, const FTexture* Texture, float GammaValue, ESimpleElementBlendMode BlendMode)
+{
+	SetEditorCompositingParameters(BatchedParameters, View);
+	SetParameters(BatchedParameters, Texture, GammaValue, BlendMode);
+}
+
 void FSimpleElementGammaBasePS::SetParameters(FRHICommandList& RHICmdList, const FTexture* Texture, float GammaValue, ESimpleElementBlendMode BlendMode)
 {
 	FRHIBatchedShaderParameters& BatchedParameters = RHICmdList.GetScratchShaderParameters();
@@ -134,6 +153,12 @@ void FSimpleElementMaskedGammaBasePS::SetParameters(FRHIBatchedShaderParameters&
 	SetShaderValue(BatchedParameters, ClipRef, ClipRefValue);
 }
 
+void FSimpleElementMaskedGammaBasePS::SetParameters(FRHIBatchedShaderParameters& BatchedParameters, const FSceneView* View, const FTexture* Texture, float InGamma, float ClipRefValue, ESimpleElementBlendMode BlendMode)
+{
+	SetEditorCompositingParameters(BatchedParameters, View);
+	SetParameters(BatchedParameters, Texture, InGamma, ClipRefValue, BlendMode);
+}
+
 void FSimpleElementMaskedGammaBasePS::SetParameters(FRHICommandList& RHICmdList, const FTexture* Texture, float InGamma, float ClipRefValue, ESimpleElementBlendMode BlendMode)
 {
 	FRHIBatchedShaderParameters& BatchedParameters = RHICmdList.GetScratchShaderParameters();
@@ -174,6 +199,9 @@ void FSimpleElementDistanceFieldGammaPS::SetParameters(
 	ESimpleElementBlendMode BlendMode
 )
 {
+	// This shader does not use editor compositing
+	SetEditorCompositingParameters(BatchedParameters, nullptr);
+
 	FSimpleElementMaskedGammaBasePS::SetParameters(BatchedParameters, Texture, InGamma, InClipRef, BlendMode);
 	SetShaderValue(BatchedParameters, SmoothWidth, SmoothWidthValue);
 	const uint32 bEnableShadowValueUInt = (bEnableShadowValue ? 1 : 0);
diff --git a/Engine/Source/Runtime/Engine/Public/ScreenRendering.h b/Engine/Source/Runtime/Engine/Public/ScreenRendering.h
--- a/Engine/Source/Runtime/Engine/Public/ScreenRendering.h
+++ b/Engine/Source/Runtime/Engine/Public/ScreenRendering.h
@@ -86,7 +86,7 @@ public:
 		SetTextureParameter(BatchedParameters,InTexture,InTextureSampler,SamplerStateRHI,TextureRHI);
 	}
 
-	//UE_DEPRECATED(5.3, "SetParameters with FRHIBatchedShaderParameters should be used.")
+	UE_DEPRECATED(5.3, "SetParameters with FRHIBatchedShaderParameters should be used.")
 	void SetParameters(FRHICommandList& RHICmdList, const FTexture* Texture)
 	{
 		FRHIBatchedShaderParameters& BatchedParameters = RHICmdList.GetScratchShaderParameters();
@@ -94,7 +94,7 @@ public:
 		RHICmdList.SetBatchedShaderParameters(RHICmdList.GetBoundPixelShader(), BatchedParameters);
 	}
 
-	//UE_DEPRECATED(5.3, "SetParameters with FRHIBatchedShaderParameters should be used.")
+	UE_DEPRECATED(5.3, "SetParameters with FRHIBatchedShaderParameters should be used.")
 	void SetParameters(FRHICommandList& RHICmdList, FRHISamplerState* SamplerStateRHI, FRHITexture* TextureRHI)
 	{
 		FRHIBatchedShaderParameters& BatchedParameters = RHICmdList.GetScratchShaderParameters();
@@ -167,7 +167,7 @@ public:
 		SetTextureParameter(BatchedParameters, InTexture, InTextureSampler, SamplerStateRHI, TextureRHI);
 	}
 
-	//UE_DEPRECATED(5.3, "SetParameters with FRHIBatchedShaderParameters should be used.")
+	UE_DEPRECATED(5.3, "SetParameters with FRHIBatchedShaderParameters should be used.")
 	void SetParameters(FRHICommandList& RHICmdList, const FTexture* Texture)
 	{
 		FRHIBatchedShaderParameters& BatchedParameters = RHICmdList.GetScratchShaderParameters();
@@ -175,7 +175,7 @@ public:
 		RHICmdList.SetBatchedShaderParameters(RHICmdList.GetBoundPixelShader(), BatchedParameters);
 	}
 
-	//UE_DEPRECATED(5.3, "SetParameters with FRHIBatchedShaderParameters should be used.")
+	UE_DEPRECATED(5.3, "SetParameters with FRHIBatchedShaderParameters should be used.")
 	void SetParameters(FRHICommandList& RHICmdList, FRHISamplerState* SamplerStateRHI, FRHITexture* TextureRHI)
 	{
 		FRHIBatchedShaderParameters& BatchedParameters = RHICmdList.GetScratchShaderParameters();
@@ -216,7 +216,7 @@ public:
 		SetTextureParameter(BatchedParameters, InTexture, InTextureSampler, SamplerStateRHI, TextureRHI);
 	}
 
-	//UE_DEPRECATED(5.3, "SetParameters with FRHIBatchedShaderParameters should be used.")
+	UE_DEPRECATED(5.3, "SetParameters with FRHIBatchedShaderParameters should be used.")
 	void SetParameters(FRHICommandList& RHICmdList, const FTexture* Texture)
 	{
 		FRHIBatchedShaderParameters& BatchedParameters = RHICmdList.GetScratchShaderParameters();
@@ -224,7 +224,7 @@ public:
 		RHICmdList.SetBatchedShaderParameters(RHICmdList.GetBoundPixelShader(), BatchedParameters);
 	}
 
-	//UE_DEPRECATED(5.3, "SetParameters with FRHIBatchedShaderParameters should be used.")
+	UE_DEPRECATED(5.3, "SetParameters with FRHIBatchedShaderParameters should be used.")
 	void SetParameters(FRHICommandList& RHICmdList, FRHISamplerState* SamplerStateRHI, FRHITexture* TextureRHI)
 	{
 		FRHIBatchedShaderParameters& BatchedParameters = RHICmdList.GetScratchShaderParameters();
@@ -268,7 +268,7 @@ public:
 		SetShaderValue(BatchedParameters, InMipLevelParameter, MipLevel);
 	}
 
-	//UE_DEPRECATED(5.3, "SetParameters with FRHIBatchedShaderParameters should be used.")
+	UE_DEPRECATED(5.3, "SetParameters with FRHIBatchedShaderParameters should be used.")
 	void SetParameters(FRHICommandList& RHICmdList, const FTexture* Texture, int MipLevel = 0)
 	{
 		FRHIBatchedShaderParameters& BatchedParameters = RHICmdList.GetScratchShaderParameters();
@@ -276,7 +276,7 @@ public:
 		RHICmdList.SetBatchedShaderParameters(RHICmdList.GetBoundPixelShader(), BatchedParameters);
 	}
 
-	//UE_DEPRECATED(5.3, "SetParameters with FRHIBatchedShaderParameters should be used.")
+	UE_DEPRECATED(5.3, "SetParameters with FRHIBatchedShaderParameters should be used.")
 	void SetParameters(FRHICommandList& RHICmdList, FRHISamplerState* SamplerStateRHI, FRHITexture* TextureRHI, int MipLevel = 0)
 	{
 		FRHIBatchedShaderParameters& BatchedParameters = RHICmdList.GetScratchShaderParameters();
@@ -321,7 +321,7 @@ public:
 		SetShaderValue(BatchedParameters, InMipLevelParameter, MipLevel);
 	}
 
-	//UE_DEPRECATED(5.3, "SetParameters with FRHIBatchedShaderParameters should be used.")
+	UE_DEPRECATED(5.3, "SetParameters with FRHIBatchedShaderParameters should be used.")
 	void SetParameters(FRHICommandList& RHICmdList, const FTexture* Texture, int MipLevel = 0)
 	{
 		FRHIBatchedShaderParameters& BatchedParameters = RHICmdList.GetScratchShaderParameters();
@@ -329,7 +329,7 @@ public:
 		RHICmdList.SetBatchedShaderParameters(RHICmdList.GetBoundPixelShader(), BatchedParameters);
 	}
 
-	//UE_DEPRECATED(5.3, "SetParameters with FRHIBatchedShaderParameters should be used.")
+	UE_DEPRECATED(5.3, "SetParameters with FRHIBatchedShaderParameters should be used.")
 	void SetParameters(FRHICommandList& RHICmdList, FRHISamplerState* SamplerStateRHI, FRHITexture* TextureRHI, int MipLevel = 0)
 	{
 		FRHIBatchedShaderParameters& BatchedParameters = RHICmdList.GetScratchShaderParameters();
@@ -369,7 +369,7 @@ public:
 		SetTextureParameter(BatchedParameters, InTexture, InTextureSampler, SamplerStateRHI, TextureRHI);
 	}
 
-	//UE_DEPRECATED(5.3, "SetParameters with FRHIBatchedShaderParameters should be used.")
+	UE_DEPRECATED(5.3, "SetParameters with FRHIBatchedShaderParameters should be used.")
 	void SetParameters(FRHICommandList& RHICmdList, const FTexture* Texture)
 	{
 		FRHIBatchedShaderParameters& BatchedParameters = RHICmdList.GetScratchShaderParameters();
@@ -377,7 +377,7 @@ public:
 		RHICmdList.SetBatchedShaderParameters(RHICmdList.GetBoundPixelShader(), BatchedParameters);
 	}
 
-	//UE_DEPRECATED(5.3, "SetParameters with FRHIBatchedShaderParameters should be used.")
+	UE_DEPRECATED(5.3, "SetParameters with FRHIBatchedShaderParameters should be used.")
 	void SetParameters(FRHICommandList& RHICmdList, FRHISamplerState* SamplerStateRHI, FRHITexture* TextureRHI)
 	{
 		FRHIBatchedShaderParameters& BatchedParameters = RHICmdList.GetScratchShaderParameters();
@@ -411,7 +411,7 @@ public:
 		FGlobalShader::SetParameters<FViewUniformShaderParameters>(BatchedParameters, ViewUniformBuffer);
 	}
 
-	//UE_DEPRECATED(5.3, "SetParameters with FRHIBatchedShaderParameters should be used.")
+	UE_DEPRECATED(5.3, "SetParameters with FRHIBatchedShaderParameters should be used.")
 	void SetParameters(FRHICommandList& RHICmdList, FRHIUniformBuffer* ViewUniformBuffer)
 	{
 		FGlobalShader::SetParameters<FViewUniformShaderParameters>(RHICmdList, RHICmdList.GetBoundVertexShader(), ViewUniformBuffer);
diff --git a/Engine/Source/Runtime/Engine/Public/SimpleElementShaders.h b/Engine/Source/Runtime/Engine/Public/SimpleElementShaders.h
--- a/Engine/Source/Runtime/Engine/Public/SimpleElementShaders.h
+++ b/Engine/Source/Runtime/Engine/Public/SimpleElementShaders.h
@@ -53,19 +53,17 @@ public:
 	FSimpleElementPS(const ShaderMetaType::CompiledShaderInitializerType& Initializer);
 	FSimpleElementPS() {}
 
-	/**
-	 * Sets parameters for compositing editor primitives
-	 *
-	 * @param View			SceneView for view constants when compositing
-	 */
-	void SetEditorCompositingParameters(FRHICommandList& RHICmdList, const FSceneView* View);
-
 	void SetParameters(FRHIBatchedShaderParameters& BatchedParameters, const FTexture* TextureValue);
+	/** Sets parameters for compositing editor primitives */
+	void SetEditorCompositingParameters(FRHIBatchedShaderParameters& BatchedParameters, const FSceneView* View);
+
+	void SetParameters(FRHIBatchedShaderParameters& BatchedParameters, const FSceneView* View, const FTexture* TextureValue);
 
 	UE_DEPRECATED(5.3, "SetParameters with FRHIBatchedShaderParameters should be used.")
-	void SetParameters(FRHICommandList& RHICmdList, const FTexture* TextureValue);
+	void SetEditorCompositingParameters(FRHICommandList& RHICmdList, const FSceneView* View);
 
-	//virtual bool Serialize(FArchive& Ar) override;
+	UE_DEPRECATED(5.3, "SetParameters with FRHIBatchedShaderParameters should be used.")
+	void SetParameters(FRHICommandList& RHICmdList, const FTexture* TextureValue);
 
 private:
 	LAYOUT_FIELD(FShaderResourceParameter, InTexture)
@@ -100,6 +98,7 @@ public:
 	FSimpleElementGammaBasePS() {}
 
 	void SetParameters(FRHIBatchedShaderParameters& BatchedParameters, const FTexture* Texture, float GammaValue, ESimpleElementBlendMode BlendMode);
+	void SetParameters(FRHIBatchedShaderParameters& BatchedParameters, const FSceneView* View, const FTexture* Texture, float GammaValue, ESimpleElementBlendMode BlendMode);
 
 	UE_DEPRECATED(5.3, "SetParameters with FRHIBatchedShaderParameters should be used.")
 	void SetParameters(FRHICommandList& RHICmdList, const FTexture* Texture,float GammaValue,ESimpleElementBlendMode BlendMode);
@@ -144,6 +143,7 @@ public:
 	FSimpleElementMaskedGammaBasePS() {}
 
 	void SetParameters(FRHIBatchedShaderParameters& BatchedParameters, const FTexture* Texture,float Gamma,float ClipRefValue,ESimpleElementBlendMode BlendMode);
+	void SetParameters(FRHIBatchedShaderParameters& BatchedParameters, const FSceneView* View, const FTexture* Texture,float Gamma,float ClipRefValue,ESimpleElementBlendMode BlendMode);
 
 	UE_DEPRECATED(5.3, "SetParameters with FRHIBatchedShaderParameters should be used.")
 	void SetParameters(FRHICommandList& RHICmdList, const FTexture* Texture,float Gamma,float ClipRefValue,ESimpleElementBlendMode BlendMode);
diff --git a/Engine/Source/Runtime/MediaAssets/Private/Misc/MediaTextureResource.cpp b/Engine/Source/Runtime/MediaAssets/Private/Misc/MediaTextureResource.cpp
--- a/Engine/Source/Runtime/MediaAssets/Private/Misc/MediaTextureResource.cpp
+++ b/Engine/Source/Runtime/MediaAssets/Private/Misc/MediaTextureResource.cpp
@@ -929,7 +929,7 @@ void FMediaTextureResource::ConvertSample(const TSharedPtr<IMediaTextureSample,
 					TShaderMapRef<FAYUVConvertPS> ConvertShader(ShaderMap);
 					GraphicsPSOInit.BoundShaderState.PixelShaderRHI = ConvertShader.GetPixelShader();
 					SetGraphicsPipelineState(CommandList, GraphicsPSOInit, 0);
-					ConvertShader->SetParameters(CommandList, InputTexture, YUVToRGBMatrix, YUVOffset, bIsSampleOutputSrgb);
+					SetShaderParametersLegacyPS(CommandList, ConvertShader, InputTexture, YUVToRGBMatrix, YUVOffset, bIsSampleOutputSrgb);
 				}
 				break;
 
@@ -939,7 +939,7 @@ void FMediaTextureResource::ConvertSample(const TSharedPtr<IMediaTextureSample,
 					TShaderMapRef<FBMPConvertPS> ConvertShader(ShaderMap);
 					GraphicsPSOInit.BoundShaderState.PixelShaderRHI = ConvertShader.GetPixelShader();
 					SetGraphicsPipelineState(CommandList, GraphicsPSOInit, 0);
-					ConvertShader->SetParameters(CommandList, InputTexture, OutputDim, false);
+					SetShaderParametersLegacyPS(CommandList, ConvertShader, InputTexture, OutputDim, false);
 				}
 				break;
 
@@ -953,14 +953,14 @@ void FMediaTextureResource::ConvertSample(const TSharedPtr<IMediaTextureSample,
 						FIntPoint TexDim = InputTexture->GetSizeXY();
 						TempSRV0 = RHICreateShaderResourceView(InputTexture, 0, 1, PF_G8); // note: the types of the views select the correct planes (offset) "magically"
 						TempSRV1 = RHICreateShaderResourceView(InputTexture, 0, 1, PF_R8G8);
-						ConvertShader->SetParameters(CommandList, TexDim, TempSRV0, TempSRV1, OutputDim, YUVToRGBMatrix, YUVOffset, bIsSampleOutputSrgb);
+						SetShaderParametersLegacyPS(CommandList, ConvertShader, TexDim, TempSRV0, TempSRV1, OutputDim, YUVToRGBMatrix, YUVOffset, bIsSampleOutputSrgb);
 					}
 					else
 					{
 						TShaderMapRef<FNV12ConvertAsBytesPS> ConvertShader(ShaderMap);
 						GraphicsPSOInit.BoundShaderState.PixelShaderRHI = ConvertShader.GetPixelShader();
 						SetGraphicsPipelineState(CommandList, GraphicsPSOInit, 0);
-						ConvertShader->SetParameters(CommandList, InputTexture, OutputDim, YUVToRGBMatrix, YUVOffset, bIsSampleOutputSrgb);
+						SetShaderParametersLegacyPS(CommandList, ConvertShader, InputTexture, OutputDim, YUVToRGBMatrix, YUVOffset, bIsSampleOutputSrgb);
 					}
 				}
 				break;
@@ -971,7 +971,7 @@ void FMediaTextureResource::ConvertSample(const TSharedPtr<IMediaTextureSample,
 					TShaderMapRef<FNV21ConvertPS> ConvertShader(ShaderMap);
 					GraphicsPSOInit.BoundShaderState.PixelShaderRHI = ConvertShader.GetPixelShader();
 					SetGraphicsPipelineState(CommandList, GraphicsPSOInit, 0);
-					ConvertShader->SetParameters(CommandList, InputTexture, OutputDim, YUVToRGBMatrix, YUVOffset, bIsSampleOutputSrgb);
+					SetShaderParametersLegacyPS(CommandList, ConvertShader, InputTexture, OutputDim, YUVToRGBMatrix, YUVOffset, bIsSampleOutputSrgb);
 				}
 				break;
 
@@ -980,7 +980,7 @@ void FMediaTextureResource::ConvertSample(const TSharedPtr<IMediaTextureSample,
 					TShaderMapRef<FUYVYConvertPS> ConvertShader(ShaderMap);
 					GraphicsPSOInit.BoundShaderState.PixelShaderRHI = ConvertShader.GetPixelShader();
 					SetGraphicsPipelineState(CommandList, GraphicsPSOInit, 0);
-					ConvertShader->SetParameters(CommandList, InputTexture, YUVToRGBMatrix, YUVOffset, bIsSampleOutputSrgb);
+					SetShaderParametersLegacyPS(CommandList, ConvertShader, InputTexture, YUVToRGBMatrix, YUVOffset, bIsSampleOutputSrgb);
 				}
 				break;
 
@@ -989,7 +989,7 @@ void FMediaTextureResource::ConvertSample(const TSharedPtr<IMediaTextureSample,
 					TShaderMapRef<FYUY2ConvertPS> ConvertShader(ShaderMap);
 					GraphicsPSOInit.BoundShaderState.PixelShaderRHI = ConvertShader.GetPixelShader();
 					SetGraphicsPipelineState(CommandList, GraphicsPSOInit, 0);
-					ConvertShader->SetParameters(CommandList, InputTexture, OutputDim, YUVToRGBMatrix, YUVOffset, bIsSampleOutputSrgb);
+					SetShaderParametersLegacyPS(CommandList, ConvertShader, InputTexture, OutputDim, YUVToRGBMatrix, YUVOffset, bIsSampleOutputSrgb);
 				}
 				break;
 
@@ -998,7 +998,7 @@ void FMediaTextureResource::ConvertSample(const TSharedPtr<IMediaTextureSample,
 					TShaderMapRef<FYVYUConvertPS> ConvertShader(ShaderMap);
 					GraphicsPSOInit.BoundShaderState.PixelShaderRHI = ConvertShader.GetPixelShader();
 					SetGraphicsPipelineState(CommandList, GraphicsPSOInit, 0);
-					ConvertShader->SetParameters(CommandList, InputTexture, YUVToRGBMatrix, YUVOffset, bIsSampleOutputSrgb);
+					SetShaderParametersLegacyPS(CommandList, ConvertShader, InputTexture, YUVToRGBMatrix, YUVOffset, bIsSampleOutputSrgb);
 				}
 				break;
 
@@ -1007,7 +1007,7 @@ void FMediaTextureResource::ConvertSample(const TSharedPtr<IMediaTextureSample,
 					TShaderMapRef<FYUVv210ConvertPS> ConvertShader(ShaderMap);
 					GraphicsPSOInit.BoundShaderState.PixelShaderRHI = ConvertShader.GetPixelShader();
 					SetGraphicsPipelineState(CommandList, GraphicsPSOInit, 0);
-					ConvertShader->SetParameters(CommandList, InputTexture, OutputDim, YUVToRGBMatrix, YUVOffset, bIsSampleOutputSrgb);
+					SetShaderParametersLegacyPS(CommandList, ConvertShader, InputTexture, OutputDim, YUVToRGBMatrix, YUVOffset, bIsSampleOutputSrgb);
 				}
 				break;
 
@@ -1018,7 +1018,7 @@ void FMediaTextureResource::ConvertSample(const TSharedPtr<IMediaTextureSample,
 					SetGraphicsPipelineState(CommandList, GraphicsPSOInit, 0);
 					FShaderResourceViewRHIRef SRV = RHICreateShaderResourceView(InputTexture, 0, 1, PF_A16B16G16R16);
 
-					ConvertShader->SetParameters(CommandList, SRV, YUVToRGBMatrix, YUVOffset, bIsSampleOutputSrgb);
+					SetShaderParametersLegacyPS(CommandList, ConvertShader, SRV, YUVToRGBMatrix, YUVOffset, bIsSampleOutputSrgb);
 				}
 				break;
 
@@ -1032,7 +1032,7 @@ void FMediaTextureResource::ConvertSample(const TSharedPtr<IMediaTextureSample,
 					TShaderMapRef<FRGBConvertPS> ConvertShader(ShaderMap);
 					GraphicsPSOInit.BoundShaderState.PixelShaderRHI = ConvertShader.GetPixelShader();
 					SetGraphicsPipelineState(CommandList, GraphicsPSOInit, 0);
-					ConvertShader->SetParameters(CommandList, InputTexture, OutputDim, false, false, FMatrix44f::Identity);
+					SetShaderParametersLegacyPS(CommandList, ConvertShader, InputTexture, OutputDim, false, false, FMatrix44f::Identity);
 				}
 				break;
 
@@ -1066,7 +1066,7 @@ void FMediaTextureResource::ConvertSample(const TSharedPtr<IMediaTextureSample,
 
 						FShaderResourceViewRHIRef Y_SRV = RHICreateShaderResourceView(InputTexture, 0, 1, PF_G16);
 						FShaderResourceViewRHIRef UV_SRV = RHICreateShaderResourceView(InputTexture, 0, 1, PF_G16R16);
-						ConvertShader->SetParameters(CommandList, TexDim, Y_SRV, UV_SRV, OutputDim, YUVMtx, ColorSpaceMtx, Sample->GetEncodingType() == UE::Color::EEncoding::ST2084);
+						SetShaderParametersLegacyPS(CommandList, ConvertShader, TexDim, Y_SRV, UV_SRV, OutputDim, YUVMtx, ColorSpaceMtx, Sample->GetEncodingType() == UE::Color::EEncoding::ST2084);
 					}
 					else
 					{
@@ -1074,7 +1074,7 @@ void FMediaTextureResource::ConvertSample(const TSharedPtr<IMediaTextureSample,
 						GraphicsPSOInit.BoundShaderState.PixelShaderRHI = ConvertShader.GetPixelShader();
 						SetGraphicsPipelineState(CommandList, GraphicsPSOInit, 0);
 
-						ConvertShader->SetParameters(CommandList, TexDim, InputTexture, OutputDim, YUVMtx, ColorSpaceMtx, Sample->GetEncodingType() == UE::Color::EEncoding::ST2084);
+						SetShaderParametersLegacyPS(CommandList, ConvertShader, TexDim, InputTexture, OutputDim, YUVMtx, ColorSpaceMtx, Sample->GetEncodingType() == UE::Color::EEncoding::ST2084);
 					}
 				}
 				break;
@@ -1236,7 +1236,7 @@ void FMediaTextureResource::CopyFromExternalTexture(const TSharedPtr <IMediaText
 			TShaderMapRef<FReadTextureExternalPS> CopyShader(ShaderMap);
 			GraphicsPSOInit.BoundShaderState.PixelShaderRHI = CopyShader.GetPixelShader();
 			SetGraphicsPipelineState(CommandList, GraphicsPSOInit, 0);
-			CopyShader->SetParameters(CommandList, SampleTexture, SamplerState, ScaleRotation, Offset);
+			SetShaderParametersLegacyPS(CommandList, CopyShader, SampleTexture, SamplerState, ScaleRotation, Offset);
 
 			// draw full size quad into render target
 			FBufferRHIRef VertexBuffer = CreateTempMediaVertexBuffer();
diff --git a/Engine/Source/Runtime/RenderCore/Private/MediaShaders.cpp b/Engine/Source/Runtime/RenderCore/Private/MediaShaders.cpp
--- a/Engine/Source/Runtime/RenderCore/Private/MediaShaders.cpp
+++ b/Engine/Source/Runtime/RenderCore/Private/MediaShaders.cpp
@@ -139,7 +139,7 @@ IMPLEMENT_GLOBAL_SHADER_PARAMETER_STRUCT(FAYUVConvertUB, "AYUVConvertUB");
 IMPLEMENT_SHADER_TYPE(, FAYUVConvertPS, TEXT("/Engine/Private/MediaShaders.usf"), TEXT("AYUVConvertPS"), SF_Pixel);
 
 
-void FAYUVConvertPS::SetParameters(FRHICommandList& CommandList, TRefCountPtr<FRHITexture2D> AYUVTexture, const FMatrix& ColorTransform, const FVector& YUVOffset, bool SrgbToLinear)
+void FAYUVConvertPS::SetParameters(FRHIBatchedShaderParameters& BatchedParameters, TRefCountPtr<FRHITexture2D> AYUVTexture, const FMatrix& ColorTransform, const FVector& YUVOffset, bool SrgbToLinear)
 {
 	FAYUVConvertUB UB;
 	{
@@ -150,7 +150,14 @@ void FAYUVConvertPS::SetParameters(FRHICommandList& CommandList, TRefCountPtr<FR
 	}
 
 	TUniformBufferRef<FAYUVConvertUB> Data = TUniformBufferRef<FAYUVConvertUB>::CreateUniformBufferImmediate(UB, UniformBuffer_SingleFrame);
-	SetUniformBufferParameter(CommandList, CommandList.GetBoundPixelShader(), GetUniformBufferParameter<FAYUVConvertUB>(), Data);
+	SetUniformBufferParameter(BatchedParameters, GetUniformBufferParameter<FAYUVConvertUB>(), Data);
+}
+
+void FAYUVConvertPS::SetParameters(FRHICommandList& CommandList, TRefCountPtr<FRHITexture2D> AYUVTexture, const FMatrix& ColorTransform, const FVector& YUVOffset, bool SrgbToLinear)
+{
+	FRHIBatchedShaderParameters& BatchedParameters = CommandList.GetScratchShaderParameters();
+	SetParameters(BatchedParameters, AYUVTexture, ColorTransform, YUVOffset, SrgbToLinear);
+	CommandList.SetBatchedShaderParameters(CommandList.GetBoundPixelShader(), BatchedParameters);
 }
 
 
@@ -168,7 +175,7 @@ IMPLEMENT_GLOBAL_SHADER_PARAMETER_STRUCT(FBMPConvertUB, "BMPConvertUB");
 IMPLEMENT_SHADER_TYPE(, FBMPConvertPS, TEXT("/Engine/Private/MediaShaders.usf"), TEXT("BMPConvertPS"), SF_Pixel);
 
 
-void FBMPConvertPS::SetParameters(FRHICommandList& CommandList, TRefCountPtr<FRHITexture2D> BMPTexture, const FIntPoint& OutputDimensions, bool SrgbToLinear)
+void FBMPConvertPS::SetParameters(FRHIBatchedShaderParameters& BatchedParameters, TRefCountPtr<FRHITexture2D> BMPTexture, const FIntPoint& OutputDimensions, bool SrgbToLinear)
 {
 	FBMPConvertUB UB;
 	{
@@ -179,7 +186,14 @@ void FBMPConvertPS::SetParameters(FRHICommandList& CommandList, TRefCountPtr<FRH
 	}
 
 	TUniformBufferRef<FBMPConvertUB> Data = TUniformBufferRef<FBMPConvertUB>::CreateUniformBufferImmediate(UB, UniformBuffer_SingleFrame);
-	SetUniformBufferParameter(CommandList, CommandList.GetBoundPixelShader(), GetUniformBufferParameter<FBMPConvertUB>(), Data);
+	SetUniformBufferParameter(BatchedParameters, GetUniformBufferParameter<FBMPConvertUB>(), Data);
+}
+
+void FBMPConvertPS::SetParameters(FRHICommandList& CommandList, TRefCountPtr<FRHITexture2D> BMPTexture, const FIntPoint& OutputDimensions, bool SrgbToLinear)
+{
+	FRHIBatchedShaderParameters& BatchedParameters = CommandList.GetScratchShaderParameters();
+	SetParameters(BatchedParameters, BMPTexture, OutputDimensions, SrgbToLinear);
+	CommandList.SetBatchedShaderParameters(CommandList.GetBoundPixelShader(), BatchedParameters);
 }
 
 
@@ -200,7 +214,7 @@ END_GLOBAL_SHADER_PARAMETER_STRUCT()
 IMPLEMENT_GLOBAL_SHADER_PARAMETER_STRUCT(FNV12ConvertUB, "NV12ConvertUB");
 IMPLEMENT_SHADER_TYPE(, FNV12ConvertPS, TEXT("/Engine/Private/MediaShaders.usf"), TEXT("NV12ConvertPS"), SF_Pixel);
 
-void FNV12ConvertPS::SetParameters(FRHICommandList& CommandList, const FIntPoint & TexDim, FShaderResourceViewRHIRef SRV_Y, FShaderResourceViewRHIRef SRV_UV, const FIntPoint& OutputDimensions, const FMatrix& ColorTransform, const FVector& YUVOffset, bool SrgbToLinear)
+void FNV12ConvertPS::SetParameters(FRHIBatchedShaderParameters& BatchedParameters, const FIntPoint & TexDim, FShaderResourceViewRHIRef SRV_Y, FShaderResourceViewRHIRef SRV_UV, const FIntPoint& OutputDimensions, const FMatrix& ColorTransform, const FVector& YUVOffset, bool SrgbToLinear)
 {
 	FNV12ConvertUB UB;
 	{
@@ -215,7 +229,14 @@ void FNV12ConvertPS::SetParameters(FRHICommandList& CommandList, const FIntPoint
 	}
 
 	TUniformBufferRef<FNV12ConvertUB> Data = TUniformBufferRef<FNV12ConvertUB>::CreateUniformBufferImmediate(UB, UniformBuffer_SingleFrame);
-	SetUniformBufferParameter(CommandList, CommandList.GetBoundPixelShader(), GetUniformBufferParameter<FNV12ConvertUB>(), Data);
+	SetUniformBufferParameter(BatchedParameters, GetUniformBufferParameter<FNV12ConvertUB>(), Data);
+}
+
+void FNV12ConvertPS::SetParameters(FRHICommandList& CommandList, const FIntPoint & TexDim, FShaderResourceViewRHIRef SRV_Y, FShaderResourceViewRHIRef SRV_UV, const FIntPoint& OutputDimensions, const FMatrix& ColorTransform, const FVector& YUVOffset, bool SrgbToLinear)
+{
+	FRHIBatchedShaderParameters& BatchedParameters = CommandList.GetScratchShaderParameters();
+	SetParameters(BatchedParameters, TexDim, SRV_Y, SRV_UV, OutputDimensions, ColorTransform, YUVOffset, SrgbToLinear);
+	CommandList.SetBatchedShaderParameters(CommandList.GetBoundPixelShader(), BatchedParameters);
 }
 
 /* FNV12ConvertAsBytesPS shader
@@ -234,7 +255,7 @@ END_GLOBAL_SHADER_PARAMETER_STRUCT()
 IMPLEMENT_GLOBAL_SHADER_PARAMETER_STRUCT(FNV12ConvertAsBytesUB, "NV12ConvertAsBytesUB");
 IMPLEMENT_SHADER_TYPE(, FNV12ConvertAsBytesPS, TEXT("/Engine/Private/MediaShaders.usf"), TEXT("NV12ConvertAsBytesPS"), SF_Pixel);
 
-void FNV12ConvertAsBytesPS::SetParameters(FRHICommandList& CommandList, TRefCountPtr<FRHITexture2D> NV12Texture, const FIntPoint& OutputDimensions, const FMatrix& ColorTransform, const FVector& YUVOffset, bool SrgbToLinear)
+void FNV12ConvertAsBytesPS::SetParameters(FRHIBatchedShaderParameters& BatchedParameters, TRefCountPtr<FRHITexture2D> NV12Texture, const FIntPoint& OutputDimensions, const FMatrix& ColorTransform, const FVector& YUVOffset, bool SrgbToLinear)
 {
 	FNV12ConvertAsBytesUB UB;
 	{
@@ -248,7 +269,14 @@ void FNV12ConvertAsBytesPS::SetParameters(FRHICommandList& CommandList, TRefCoun
 	}
 
 	TUniformBufferRef<FNV12ConvertAsBytesUB> Data = TUniformBufferRef<FNV12ConvertAsBytesUB>::CreateUniformBufferImmediate(UB, UniformBuffer_SingleFrame);
-	SetUniformBufferParameter(CommandList, CommandList.GetBoundPixelShader(), GetUniformBufferParameter<FNV12ConvertAsBytesUB>(), Data);
+	SetUniformBufferParameter(BatchedParameters, GetUniformBufferParameter<FNV12ConvertAsBytesUB>(), Data);
+}
+
+void FNV12ConvertAsBytesPS::SetParameters(FRHICommandList& CommandList, TRefCountPtr<FRHITexture2D> NV12Texture, const FIntPoint& OutputDimensions, const FMatrix& ColorTransform, const FVector& YUVOffset, bool SrgbToLinear)
+{
+	FRHIBatchedShaderParameters& BatchedParameters = CommandList.GetScratchShaderParameters();
+	SetParameters(BatchedParameters, NV12Texture, OutputDimensions, ColorTransform, YUVOffset, SrgbToLinear);
+	CommandList.SetBatchedShaderParameters(CommandList.GetBoundPixelShader(), BatchedParameters);
 }
 
 /* FNV21ConvertPS shader
@@ -268,7 +296,7 @@ IMPLEMENT_GLOBAL_SHADER_PARAMETER_STRUCT(FNV21ConvertUB, "NV21ConvertUB");
 IMPLEMENT_SHADER_TYPE(, FNV21ConvertPS, TEXT("/Engine/Private/MediaShaders.usf"), TEXT("NV21ConvertPS"), SF_Pixel);
 
 
-void FNV21ConvertPS::SetParameters(FRHICommandList& CommandList, TRefCountPtr<FRHITexture2D> NV21Texture, const FIntPoint& OutputDimensions, const FMatrix& ColorTransform, const FVector& YUVOffset, bool SrgbToLinear)
+void FNV21ConvertPS::SetParameters(FRHIBatchedShaderParameters& BatchedParameters, TRefCountPtr<FRHITexture2D> NV21Texture, const FIntPoint& OutputDimensions, const FMatrix& ColorTransform, const FVector& YUVOffset, bool SrgbToLinear)
 {
 	FNV21ConvertUB UB;
 	{
@@ -282,7 +310,14 @@ void FNV21ConvertPS::SetParameters(FRHICommandList& CommandList, TRefCountPtr<FR
 	}
 
 	TUniformBufferRef<FNV21ConvertUB> Data = TUniformBufferRef<FNV21ConvertUB>::CreateUniformBufferImmediate(UB, UniformBuffer_SingleFrame);
-	SetUniformBufferParameter(CommandList, CommandList.GetBoundPixelShader(), GetUniformBufferParameter<FNV21ConvertUB>(), Data);
+	SetUniformBufferParameter(BatchedParameters, GetUniformBufferParameter<FNV21ConvertUB>(), Data);
+}
+
+void FNV21ConvertPS::SetParameters(FRHICommandList& CommandList, TRefCountPtr<FRHITexture2D> NV21Texture, const FIntPoint& OutputDimensions, const FMatrix& ColorTransform, const FVector& YUVOffset, bool SrgbToLinear)
+{
+	FRHIBatchedShaderParameters& BatchedParameters = CommandList.GetScratchShaderParameters();
+	SetParameters(BatchedParameters, NV21Texture, OutputDimensions, ColorTransform, YUVOffset, SrgbToLinear);
+	CommandList.SetBatchedShaderParameters(CommandList.GetBoundPixelShader(), BatchedParameters);
 }
 
 /* FP010ConvertPS shader
@@ -302,7 +337,7 @@ END_GLOBAL_SHADER_PARAMETER_STRUCT()
 IMPLEMENT_GLOBAL_SHADER_PARAMETER_STRUCT(FP010ConvertUB, "P010ConvertUB");
 IMPLEMENT_SHADER_TYPE(, FP010ConvertPS, TEXT("/Engine/Private/MediaShaders.usf"), TEXT("P010ConvertPS"), SF_Pixel);
 
-void FP010ConvertPS::SetParameters(FRHICommandList& CommandList, const FIntPoint& TexDim, FShaderResourceViewRHIRef SRV_Y, FShaderResourceViewRHIRef SRV_UV, const FIntPoint& OutputDimensions, const FMatrix44f& ColorTransform, const FMatrix44f& CSTransform, bool bIsST2084)
+void FP010ConvertPS::SetParameters(FRHIBatchedShaderParameters& BatchedParameters, const FIntPoint& TexDim, FShaderResourceViewRHIRef SRV_Y, FShaderResourceViewRHIRef SRV_UV, const FIntPoint& OutputDimensions, const FMatrix44f& ColorTransform, const FMatrix44f& CSTransform, bool bIsST2084)
 {
 	FP010ConvertUB UB;
 	{
@@ -317,7 +352,14 @@ void FP010ConvertPS::SetParameters(FRHICommandList& CommandList, const FIntPoint
 	}
 
 	TUniformBufferRef<FP010ConvertUB> Data = TUniformBufferRef<FP010ConvertUB>::CreateUniformBufferImmediate(UB, UniformBuffer_SingleFrame);
-	SetUniformBufferParameter(CommandList, CommandList.GetBoundPixelShader(), GetUniformBufferParameter<FP010ConvertUB>(), Data);
+	SetUniformBufferParameter(BatchedParameters, GetUniformBufferParameter<FP010ConvertUB>(), Data);
+}
+
+void FP010ConvertPS::SetParameters(FRHICommandList& CommandList, const FIntPoint& TexDim, FShaderResourceViewRHIRef SRV_Y, FShaderResourceViewRHIRef SRV_UV, const FIntPoint& OutputDimensions, const FMatrix44f& ColorTransform, const FMatrix44f& CSTransform, bool bIsST2084)
+{
+	FRHIBatchedShaderParameters& BatchedParameters = CommandList.GetScratchShaderParameters();
+	SetParameters(BatchedParameters, TexDim, SRV_Y, SRV_UV, OutputDimensions, ColorTransform, CSTransform, bIsST2084);
+	CommandList.SetBatchedShaderParameters(CommandList.GetBoundPixelShader(), BatchedParameters);
 }
 
 /* FP010ConvertAsUINT16sPS shader (from G16 texture)
@@ -337,7 +379,7 @@ END_GLOBAL_SHADER_PARAMETER_STRUCT()
 IMPLEMENT_GLOBAL_SHADER_PARAMETER_STRUCT(FP010ConvertAsUINT16sUB, "P010ConvertAsUINT16sUB");
 IMPLEMENT_SHADER_TYPE(, FP010ConvertAsUINT16sPS, TEXT("/Engine/Private/MediaShaders.usf"), TEXT("P010ConvertAsUINT16sPS"), SF_Pixel);
 
-void FP010ConvertAsUINT16sPS::SetParameters(FRHICommandList& CommandList, const FIntPoint& TexDim, TRefCountPtr<FRHITexture2D> NV12Texture, const FIntPoint& OutputDimensions, const FMatrix44f& ColorTransform, const FMatrix44f& CSTransform, bool bIsST2084)
+void FP010ConvertAsUINT16sPS::SetParameters(FRHIBatchedShaderParameters& BatchedParameters, const FIntPoint& TexDim, TRefCountPtr<FRHITexture2D> NV12Texture, const FIntPoint& OutputDimensions, const FMatrix44f& ColorTransform, const FMatrix44f& CSTransform, bool bIsST2084)
 {
 	FP010ConvertAsUINT16sUB UB;
 	{
@@ -352,7 +394,14 @@ void FP010ConvertAsUINT16sPS::SetParameters(FRHICommandList& CommandList, const
 	}
 
 	TUniformBufferRef<FP010ConvertAsUINT16sUB> Data = TUniformBufferRef<FP010ConvertAsUINT16sUB>::CreateUniformBufferImmediate(UB, UniformBuffer_SingleFrame);
-	SetUniformBufferParameter(CommandList, CommandList.GetBoundPixelShader(), GetUniformBufferParameter<FP010ConvertAsUINT16sUB>(), Data);
+	SetUniformBufferParameter(BatchedParameters, GetUniformBufferParameter<FP010ConvertAsUINT16sUB>(), Data);
+}
+
+void FP010ConvertAsUINT16sPS::SetParameters(FRHICommandList& CommandList, const FIntPoint& TexDim, TRefCountPtr<FRHITexture2D> NV12Texture, const FIntPoint& OutputDimensions, const FMatrix44f& ColorTransform, const FMatrix44f& CSTransform, bool bIsST2084)
+{
+	FRHIBatchedShaderParameters& BatchedParameters = CommandList.GetScratchShaderParameters();
+	SetParameters(BatchedParameters, TexDim, NV12Texture, OutputDimensions, ColorTransform, CSTransform, bIsST2084);
+	CommandList.SetBatchedShaderParameters(CommandList.GetBoundPixelShader(), BatchedParameters);
 }
 
 /* FP010_2101010ConvertPS shader
@@ -375,7 +424,7 @@ END_GLOBAL_SHADER_PARAMETER_STRUCT()
 IMPLEMENT_GLOBAL_SHADER_PARAMETER_STRUCT(FP010_2101010ConvertUB, "P010_2101010ConvertUB");
 IMPLEMENT_SHADER_TYPE(, FP010_2101010ConvertPS, TEXT("/Engine/Private/MediaShaders.usf"), TEXT("P010_2101010ConvertPS"), SF_Pixel);
 
-void FP010_2101010ConvertPS::SetParameters(FRHICommandList& CommandList, const FIntPoint& TexDim, FShaderResourceViewRHIRef SRV_Y, FShaderResourceViewRHIRef SRV_U, FShaderResourceViewRHIRef SRV_V, const FIntPoint& OutputDimensions, const FMatrix44f& ColorTransform, const FMatrix44f& CSTransform, bool bIsST2084)
+void FP010_2101010ConvertPS::SetParameters(FRHIBatchedShaderParameters& BatchedParameters, const FIntPoint& TexDim, FShaderResourceViewRHIRef SRV_Y, FShaderResourceViewRHIRef SRV_U, FShaderResourceViewRHIRef SRV_V, const FIntPoint& OutputDimensions, const FMatrix44f& ColorTransform, const FMatrix44f& CSTransform, bool bIsST2084)
 {
 	FP010_2101010ConvertUB UB;
 	{
@@ -397,7 +446,14 @@ void FP010_2101010ConvertPS::SetParameters(FRHICommandList& CommandList, const F
 	}
 
 	TUniformBufferRef<FP010_2101010ConvertUB> Data = TUniformBufferRef<FP010_2101010ConvertUB>::CreateUniformBufferImmediate(UB, UniformBuffer_SingleFrame);
-	SetUniformBufferParameter(CommandList, CommandList.GetBoundPixelShader(), GetUniformBufferParameter<FP010_2101010ConvertUB>(), Data);
+	SetUniformBufferParameter(BatchedParameters, GetUniformBufferParameter<FP010_2101010ConvertUB>(), Data);
+}
+
+void FP010_2101010ConvertPS::SetParameters(FRHICommandList& CommandList, const FIntPoint& TexDim, FShaderResourceViewRHIRef SRV_Y, FShaderResourceViewRHIRef SRV_U, FShaderResourceViewRHIRef SRV_V, const FIntPoint& OutputDimensions, const FMatrix44f& ColorTransform, const FMatrix44f& CSTransform, bool bIsST2084)
+{
+	FRHIBatchedShaderParameters& BatchedParameters = CommandList.GetScratchShaderParameters();
+	SetParameters(BatchedParameters, TexDim, SRV_Y, SRV_U, SRV_V, OutputDimensions, ColorTransform, CSTransform, bIsST2084);
+	CommandList.SetBatchedShaderParameters(CommandList.GetBoundPixelShader(), BatchedParameters);
 }
 
 /* FRGBConvertPS shader
@@ -415,7 +471,7 @@ IMPLEMENT_GLOBAL_SHADER_PARAMETER_STRUCT(FRGBConvertUB, "RGBConvertUB");
 IMPLEMENT_SHADER_TYPE(, FRGBConvertPS, TEXT("/Engine/Private/MediaShaders.usf"), TEXT("RGBConvertPS"), SF_Pixel);
 
 
-void FRGBConvertPS::SetParameters(FRHICommandList& CommandList, TRefCountPtr<FRHITexture2D> RGBTexture, const FIntPoint& OutputDimensions, bool bSrgbToLinear, bool bIsST2084, const FMatrix44f& CSTransform)
+void FRGBConvertPS::SetParameters(FRHIBatchedShaderParameters& BatchedParameters, TRefCountPtr<FRHITexture2D> RGBTexture, const FIntPoint& OutputDimensions, bool bSrgbToLinear, bool bIsST2084, const FMatrix44f& CSTransform)
 {
 	FRGBConvertUB UB;
 	{
@@ -427,7 +483,14 @@ void FRGBConvertPS::SetParameters(FRHICommandList& CommandList, TRefCountPtr<FRH
 	}
 
 	TUniformBufferRef<FRGBConvertUB> Data = TUniformBufferRef<FRGBConvertUB>::CreateUniformBufferImmediate(UB, UniformBuffer_SingleFrame);
-	SetUniformBufferParameter(CommandList, CommandList.GetBoundPixelShader(), GetUniformBufferParameter<FRGBConvertUB>(), Data);
+	SetUniformBufferParameter(BatchedParameters, GetUniformBufferParameter<FRGBConvertUB>(), Data);
+}
+
+void FRGBConvertPS::SetParameters(FRHICommandList& CommandList, TRefCountPtr<FRHITexture2D> RGBTexture, const FIntPoint& OutputDimensions, bool bSrgbToLinear, bool bIsST2084, const FMatrix44f& CSTransform)
+{
+	FRHIBatchedShaderParameters& BatchedParameters = CommandList.GetScratchShaderParameters();
+	SetParameters(BatchedParameters, RGBTexture, OutputDimensions, bSrgbToLinear, bIsST2084, CSTransform);
+	CommandList.SetBatchedShaderParameters(CommandList.GetBoundPixelShader(), BatchedParameters);
 }
 
 
@@ -448,7 +511,7 @@ IMPLEMENT_SHADER_TYPE(, FYCbCrConvertPS, TEXT("/Engine/Private/MediaShaders.usf"
 IMPLEMENT_SHADER_TYPE(, FYCbCrConvertPS_4x4Matrix, TEXT("/Engine/Private/MediaShaders.usf"), TEXT("YCbCrConvertPS_4x4Matrix"), SF_Pixel);
 
 
-void FYCbCrConvertPS::SetParameters(FRHICommandList& CommandList, TRefCountPtr<FRHITexture2D> LumaTexture, TRefCountPtr<FRHITexture2D> CbCrTexture, const FMatrix& ColorTransform, const FVector& YUVOffset, bool SrgbToLinear)
+void FYCbCrConvertPS::SetParameters(FRHIBatchedShaderParameters& BatchedParameters, TRefCountPtr<FRHITexture2D> LumaTexture, TRefCountPtr<FRHITexture2D> CbCrTexture, const FMatrix& ColorTransform, const FVector& YUVOffset, bool SrgbToLinear)
 {
 	FYCbCrConvertUB UB;
 	{
@@ -464,7 +527,14 @@ void FYCbCrConvertPS::SetParameters(FRHICommandList& CommandList, TRefCountPtr<F
 	}
 
 	TUniformBufferRef<FYCbCrConvertUB> Data = TUniformBufferRef<FYCbCrConvertUB>::CreateUniformBufferImmediate(UB, UniformBuffer_SingleFrame);
-	SetUniformBufferParameter(CommandList, CommandList.GetBoundPixelShader(), GetUniformBufferParameter<FYCbCrConvertUB>(), Data);	
+	SetUniformBufferParameter(BatchedParameters, GetUniformBufferParameter<FYCbCrConvertUB>(), Data);
+}
+
+void FYCbCrConvertPS::SetParameters(FRHICommandList& CommandList, TRefCountPtr<FRHITexture2D> LumaTexture, TRefCountPtr<FRHITexture2D> CbCrTexture, const FMatrix& ColorTransform, const FVector& YUVOffset, bool SrgbToLinear)
+{
+	FRHIBatchedShaderParameters& BatchedParameters = CommandList.GetScratchShaderParameters();
+	SetParameters(BatchedParameters, LumaTexture, CbCrTexture, ColorTransform, YUVOffset, SrgbToLinear);
+	CommandList.SetBatchedShaderParameters(CommandList.GetBoundPixelShader(), BatchedParameters);
 }
 
 
@@ -484,7 +554,7 @@ IMPLEMENT_GLOBAL_SHADER_PARAMETER_STRUCT(FUYVYConvertUB, "UYVYConvertUB");
 IMPLEMENT_SHADER_TYPE(, FUYVYConvertPS, TEXT("/Engine/Private/MediaShaders.usf"), TEXT("UYVYConvertPS"), SF_Pixel);
 
 
-void FUYVYConvertPS::SetParameters(FRHICommandList& CommandList, TRefCountPtr<FRHITexture2D> UYVYTexture, const FMatrix& ColorTransform, const FVector& YUVOffset, bool SrgbToLinear)
+void FUYVYConvertPS::SetParameters(FRHIBatchedShaderParameters& BatchedParameters, TRefCountPtr<FRHITexture2D> UYVYTexture, const FMatrix& ColorTransform, const FVector& YUVOffset, bool SrgbToLinear)
 {
 	FUYVYConvertUB UB;
 	{
@@ -497,7 +567,14 @@ void FUYVYConvertPS::SetParameters(FRHICommandList& CommandList, TRefCountPtr<FR
 	}
 
 	TUniformBufferRef<FUYVYConvertUB> Data = TUniformBufferRef<FUYVYConvertUB>::CreateUniformBufferImmediate(UB, UniformBuffer_SingleFrame);
-	SetUniformBufferParameter(CommandList, CommandList.GetBoundPixelShader(), GetUniformBufferParameter<FUYVYConvertUB>(), Data);
+	SetUniformBufferParameter(BatchedParameters, GetUniformBufferParameter<FUYVYConvertUB>(), Data);
+}
+
+void FUYVYConvertPS::SetParameters(FRHICommandList& CommandList, TRefCountPtr<FRHITexture2D> UYVYTexture, const FMatrix& ColorTransform, const FVector& YUVOffset, bool SrgbToLinear)
+{
+	FRHIBatchedShaderParameters& BatchedParameters = CommandList.GetScratchShaderParameters();
+	SetParameters(BatchedParameters, UYVYTexture, ColorTransform, YUVOffset, SrgbToLinear);
+	CommandList.SetBatchedShaderParameters(CommandList.GetBoundPixelShader(), BatchedParameters);
 }
 
 
@@ -520,7 +597,7 @@ IMPLEMENT_GLOBAL_SHADER_PARAMETER_STRUCT(FYUVConvertUB, "YUVConvertUB");
 IMPLEMENT_SHADER_TYPE(, FYUVConvertPS, TEXT("/Engine/Private/MediaShaders.usf"), TEXT("YUVConvertPS"), SF_Pixel);
 
 
-void FYUVConvertPS::SetParameters(FRHICommandList& CommandList, TRefCountPtr<FRHITexture2D> YTexture, TRefCountPtr<FRHITexture2D> UTexture, TRefCountPtr<FRHITexture2D> VTexture, const FIntPoint& OutputDimensions, const FMatrix& ColorTransform, const FVector& YUVOffset, bool SrgbToLinear)
+void FYUVConvertPS::SetParameters(FRHIBatchedShaderParameters& BatchedParameters, TRefCountPtr<FRHITexture2D> YTexture, TRefCountPtr<FRHITexture2D> UTexture, TRefCountPtr<FRHITexture2D> VTexture, const FIntPoint& OutputDimensions, const FMatrix& ColorTransform, const FVector& YUVOffset, bool SrgbToLinear)
 {
 	FYUVConvertUB UB;
 	{
@@ -536,7 +613,14 @@ void FYUVConvertPS::SetParameters(FRHICommandList& CommandList, TRefCountPtr<FRH
 	}
 
 	TUniformBufferRef<FYUVConvertUB> Data = TUniformBufferRef<FYUVConvertUB>::CreateUniformBufferImmediate(UB, UniformBuffer_SingleFrame);
-	SetUniformBufferParameter(CommandList, CommandList.GetBoundPixelShader(), GetUniformBufferParameter<FYUVConvertUB>(), Data);
+	SetUniformBufferParameter(BatchedParameters, GetUniformBufferParameter<FYUVConvertUB>(), Data);
+}
+
+void FYUVConvertPS::SetParameters(FRHICommandList& CommandList, TRefCountPtr<FRHITexture2D> YTexture, TRefCountPtr<FRHITexture2D> UTexture, TRefCountPtr<FRHITexture2D> VTexture, const FIntPoint& OutputDimensions, const FMatrix& ColorTransform, const FVector& YUVOffset, bool SrgbToLinear)
+{
+	FRHIBatchedShaderParameters& BatchedParameters = CommandList.GetScratchShaderParameters();
+	SetParameters(BatchedParameters, YTexture, UTexture, VTexture, OutputDimensions, ColorTransform, YUVOffset, SrgbToLinear);
+	CommandList.SetBatchedShaderParameters(CommandList.GetBoundPixelShader(), BatchedParameters);
 }
 
 
@@ -555,7 +639,7 @@ IMPLEMENT_GLOBAL_SHADER_PARAMETER_STRUCT(FYUVv210ConvertUB, "YUVv210ConvertUB");
 IMPLEMENT_SHADER_TYPE(, FYUVv210ConvertPS, TEXT("/Engine/Private/MediaShaders.usf"), TEXT("YUVv210ConvertPS"), SF_Pixel);
 
 
-void FYUVv210ConvertPS::SetParameters(FRHICommandList& CommandList, TRefCountPtr<FRHITexture2D> YUVTexture, const FIntPoint& OutputDimensions, const FMatrix& ColorTransform, const FVector& YUVOffset, bool SrgbToLinear)
+void FYUVv210ConvertPS::SetParameters(FRHIBatchedShaderParameters& BatchedParameters, TRefCountPtr<FRHITexture2D> YUVTexture, const FIntPoint& OutputDimensions, const FMatrix& ColorTransform, const FVector& YUVOffset, bool SrgbToLinear)
 {
 	FYUVv210ConvertUB UB;
 	{
@@ -567,7 +651,14 @@ void FYUVv210ConvertPS::SetParameters(FRHICommandList& CommandList, TRefCountPtr
 	}
 
 	TUniformBufferRef<FYUVv210ConvertUB> Data = TUniformBufferRef<FYUVv210ConvertUB>::CreateUniformBufferImmediate(UB, UniformBuffer_SingleFrame);
-	SetUniformBufferParameter(CommandList, CommandList.GetBoundPixelShader(), GetUniformBufferParameter<FYUVv210ConvertUB>(), Data);
+	SetUniformBufferParameter(BatchedParameters, GetUniformBufferParameter<FYUVv210ConvertUB>(), Data);
+}
+
+void FYUVv210ConvertPS::SetParameters(FRHICommandList& CommandList, TRefCountPtr<FRHITexture2D> YUVTexture, const FIntPoint& OutputDimensions, const FMatrix& ColorTransform, const FVector& YUVOffset, bool SrgbToLinear)
+{
+	FRHIBatchedShaderParameters& BatchedParameters = CommandList.GetScratchShaderParameters();
+	SetParameters(BatchedParameters, YUVTexture, OutputDimensions, ColorTransform, YUVOffset, SrgbToLinear);
+	CommandList.SetBatchedShaderParameters(CommandList.GetBoundPixelShader(), BatchedParameters);
 }
 
 
@@ -585,7 +676,7 @@ IMPLEMENT_GLOBAL_SHADER_PARAMETER_STRUCT(FYUVY416ConvertUB, "YUVY416ConvertUB");
 IMPLEMENT_SHADER_TYPE(, FYUVY416ConvertPS, TEXT("/Engine/Private/MediaShaders.usf"), TEXT("YUVY416ConvertPS"), SF_Pixel);
 
 
-void FYUVY416ConvertPS::SetParameters(FRHICommandList& CommandList, FShaderResourceViewRHIRef SRV_Y, const FMatrix& ColorTransform, const FVector& YUVOffset, bool SrgbToLinear)
+void FYUVY416ConvertPS::SetParameters(FRHIBatchedShaderParameters& BatchedParameters, FShaderResourceViewRHIRef SRV_Y, const FMatrix& ColorTransform, const FVector& YUVOffset, bool SrgbToLinear)
 {
 	FYUVY416ConvertUB UB;
 	{
@@ -596,7 +687,14 @@ void FYUVY416ConvertPS::SetParameters(FRHICommandList& CommandList, FShaderResou
 	}
 
 	TUniformBufferRef<FYUVY416ConvertUB> Data = TUniformBufferRef<FYUVY416ConvertUB>::CreateUniformBufferImmediate(UB, UniformBuffer_SingleFrame);
-	SetUniformBufferParameter(CommandList, CommandList.GetBoundPixelShader(), GetUniformBufferParameter<FYUVY416ConvertUB>(), Data);
+	SetUniformBufferParameter(BatchedParameters, GetUniformBufferParameter<FYUVY416ConvertUB>(), Data);
+}
+
+void FYUVY416ConvertPS::SetParameters(FRHICommandList& CommandList, FShaderResourceViewRHIRef SRV_Y, const FMatrix& ColorTransform, const FVector& YUVOffset, bool SrgbToLinear)
+{
+	FRHIBatchedShaderParameters& BatchedParameters = CommandList.GetScratchShaderParameters();
+	SetParameters(BatchedParameters, SRV_Y, ColorTransform, YUVOffset, SrgbToLinear);
+	CommandList.SetBatchedShaderParameters(CommandList.GetBoundPixelShader(), BatchedParameters);
 }
 
 
@@ -617,7 +715,7 @@ IMPLEMENT_GLOBAL_SHADER_PARAMETER_STRUCT(FYUY2ConvertUB, "YUY2ConvertUB");
 IMPLEMENT_SHADER_TYPE(, FYUY2ConvertPS, TEXT("/Engine/Private/MediaShaders.usf"), TEXT("YUY2ConvertPS"), SF_Pixel);
 
 
-void FYUY2ConvertPS::SetParameters(FRHICommandList& CommandList, TRefCountPtr<FRHITexture2D> YUY2Texture, const FIntPoint& OutputDimensions, const FMatrix& ColorTransform, const FVector& YUVOffset, bool SrgbToLinear)
+void FYUY2ConvertPS::SetParameters(FRHIBatchedShaderParameters& BatchedParameters, TRefCountPtr<FRHITexture2D> YUY2Texture, const FIntPoint& OutputDimensions, const FMatrix& ColorTransform, const FVector& YUVOffset, bool SrgbToLinear)
 {
 	FYUY2ConvertUB UB;
 	{
@@ -631,7 +729,14 @@ void FYUY2ConvertPS::SetParameters(FRHICommandList& CommandList, TRefCountPtr<FR
 	}
 
 	TUniformBufferRef<FYUY2ConvertUB> Data = TUniformBufferRef<FYUY2ConvertUB>::CreateUniformBufferImmediate(UB, UniformBuffer_SingleFrame);
-	SetUniformBufferParameter(CommandList, CommandList.GetBoundPixelShader(), GetUniformBufferParameter<FYUY2ConvertUB>(), Data);
+	SetUniformBufferParameter(BatchedParameters, GetUniformBufferParameter<FYUY2ConvertUB>(), Data);
+}
+
+void FYUY2ConvertPS::SetParameters(FRHICommandList& CommandList, TRefCountPtr<FRHITexture2D> YUY2Texture, const FIntPoint& OutputDimensions, const FMatrix& ColorTransform, const FVector& YUVOffset, bool SrgbToLinear)
+{
+	FRHIBatchedShaderParameters& BatchedParameters = CommandList.GetScratchShaderParameters();
+	SetParameters(BatchedParameters, YUY2Texture, OutputDimensions, ColorTransform, YUVOffset, SrgbToLinear);
+	CommandList.SetBatchedShaderParameters(CommandList.GetBoundPixelShader(), BatchedParameters);
 }
 
 
@@ -651,7 +756,7 @@ IMPLEMENT_GLOBAL_SHADER_PARAMETER_STRUCT(FYVYUConvertUB, "YVYUConvertUB");
 IMPLEMENT_SHADER_TYPE(, FYVYUConvertPS, TEXT("/Engine/Private/MediaShaders.usf"), TEXT("YVYUConvertPS"), SF_Pixel);
 
 
-void FYVYUConvertPS::SetParameters(FRHICommandList& CommandList, TRefCountPtr<FRHITexture2D> YVYUTexture, const FMatrix& ColorTransform, const FVector& YUVOffset, bool SrgbToLinear)
+void FYVYUConvertPS::SetParameters(FRHIBatchedShaderParameters& BatchedParameters, TRefCountPtr<FRHITexture2D> YVYUTexture, const FMatrix& ColorTransform, const FVector& YUVOffset, bool SrgbToLinear)
 {
 	FYVYUConvertUB UB;
 	{
@@ -664,7 +769,14 @@ void FYVYUConvertPS::SetParameters(FRHICommandList& CommandList, TRefCountPtr<FR
 	}
 
 	TUniformBufferRef<FYVYUConvertUB> Data = TUniformBufferRef<FYVYUConvertUB>::CreateUniformBufferImmediate(UB, UniformBuffer_SingleFrame);
-	SetUniformBufferParameter(CommandList, CommandList.GetBoundPixelShader(), GetUniformBufferParameter<FYVYUConvertUB>(), Data);
+	SetUniformBufferParameter(BatchedParameters, GetUniformBufferParameter<FYVYUConvertUB>(), Data);
+}
+
+void FYVYUConvertPS::SetParameters(FRHICommandList& CommandList, TRefCountPtr<FRHITexture2D> YVYUTexture, const FMatrix& ColorTransform, const FVector& YUVOffset, bool SrgbToLinear)
+{
+	FRHIBatchedShaderParameters& BatchedParameters = CommandList.GetScratchShaderParameters();
+	SetParameters(BatchedParameters, YVYUTexture, ColorTransform, YUVOffset, SrgbToLinear);
+	CommandList.SetBatchedShaderParameters(CommandList.GetBoundPixelShader(), BatchedParameters);
 }
 
 
@@ -682,7 +794,7 @@ IMPLEMENT_GLOBAL_SHADER_PARAMETER_STRUCT(FRGB8toY8ConvertUB, "RGB8toY8ConvertUB"
 IMPLEMENT_SHADER_TYPE(, FRGB8toY8ConvertPS, TEXT("/Engine/Private/MediaShaders.usf"), TEXT("RGB8toY8ConvertPS"), SF_Pixel);
 
 
-void FRGB8toY8ConvertPS::SetParameters(FRHICommandList& CommandList, TRefCountPtr<FRHITexture2D> RGBATexture, const FVector4f& ColorTransform, bool LinearToSrgb)
+void FRGB8toY8ConvertPS::SetParameters(FRHIBatchedShaderParameters& BatchedParameters, TRefCountPtr<FRHITexture2D> RGBATexture, const FVector4f& ColorTransform, bool LinearToSrgb)
 {
 	FRGB8toY8ConvertUB UB;
 	{
@@ -693,7 +805,14 @@ void FRGB8toY8ConvertPS::SetParameters(FRHICommandList& CommandList, TRefCountPt
 	}
 
 	TUniformBufferRef<FRGB8toY8ConvertUB> Data = TUniformBufferRef<FRGB8toY8ConvertUB>::CreateUniformBufferImmediate(UB, UniformBuffer_SingleFrame);
-	SetUniformBufferParameter(CommandList, CommandList.GetBoundPixelShader(), GetUniformBufferParameter<FRGB8toY8ConvertUB>(), Data);
+	SetUniformBufferParameter(BatchedParameters, GetUniformBufferParameter<FRGB8toY8ConvertUB>(), Data);
+}
+
+void FRGB8toY8ConvertPS::SetParameters(FRHICommandList& CommandList, TRefCountPtr<FRHITexture2D> RGBATexture, const FVector4f& ColorTransform, bool LinearToSrgb)
+{
+	FRHIBatchedShaderParameters& BatchedParameters = CommandList.GetScratchShaderParameters();
+	SetParameters(BatchedParameters, RGBATexture, ColorTransform, LinearToSrgb);
+	CommandList.SetBatchedShaderParameters(CommandList.GetBoundPixelShader(), BatchedParameters);
 }
 
 
@@ -711,7 +830,7 @@ IMPLEMENT_GLOBAL_SHADER_PARAMETER_STRUCT(FReadTextureExternalUB, "ReadTextureExt
 IMPLEMENT_SHADER_TYPE(, FReadTextureExternalPS, TEXT("/Engine/Private/MediaShaders.usf"), TEXT("ReadTextureExternalPS"), SF_Pixel);
 
 
-void FReadTextureExternalPS::SetParameters(FRHICommandList& CommandList, FTextureRHIRef TextureExt, FSamplerStateRHIRef SamplerState, const FLinearColor& ScaleRotation, const FLinearColor& Offset)
+void FReadTextureExternalPS::SetParameters(FRHIBatchedShaderParameters& BatchedParameters, FTextureRHIRef TextureExt, FSamplerStateRHIRef SamplerState, const FLinearColor& ScaleRotation, const FLinearColor& Offset)
 {
 	FReadTextureExternalUB UB;
 	{
@@ -722,7 +841,14 @@ void FReadTextureExternalPS::SetParameters(FRHICommandList& CommandList, FTextur
 	}
 
 	TUniformBufferRef<FReadTextureExternalUB> Data = TUniformBufferRef<FReadTextureExternalUB>::CreateUniformBufferImmediate(UB, UniformBuffer_SingleFrame);
-	SetUniformBufferParameter(CommandList, CommandList.GetBoundPixelShader(), GetUniformBufferParameter<FReadTextureExternalUB>(), Data);
+	SetUniformBufferParameter(BatchedParameters, GetUniformBufferParameter<FReadTextureExternalUB>(), Data);
+}
+
+void FReadTextureExternalPS::SetParameters(FRHICommandList& CommandList, FTextureRHIRef TextureExt, FSamplerStateRHIRef SamplerState, const FLinearColor& ScaleRotation, const FLinearColor& Offset)
+{
+	FRHIBatchedShaderParameters& BatchedParameters = CommandList.GetScratchShaderParameters();
+	SetParameters(BatchedParameters, TextureExt, SamplerState, ScaleRotation, Offset);
+	CommandList.SetBatchedShaderParameters(CommandList.GetBoundPixelShader(), BatchedParameters);
 }
 
 /* FRGB8toUYVY8ConvertPS shader
diff --git a/Engine/Source/Runtime/RenderCore/Public/CopyTextureShaders.h b/Engine/Source/Runtime/RenderCore/Public/CopyTextureShaders.h
--- a/Engine/Source/Runtime/RenderCore/Public/CopyTextureShaders.h
+++ b/Engine/Source/Runtime/RenderCore/Public/CopyTextureShaders.h
@@ -79,10 +79,13 @@ public:
 		check(SrcOffset.GetMin() >= 0 && DstOffset.GetMin() >= 0 && Dimensions.GetMin() >= 0);
 		check(Context.DstType != ECopyTextureResourceType::Texture2D || Dimensions.Z <= 1);
 
-		FRHIComputeShader* ShaderRHI = RHICmdList.GetBoundComputeShader();
-		SetShaderValue(RHICmdList, ShaderRHI, SrcOffsetParam, SrcOffset);
-		SetShaderValue(RHICmdList, ShaderRHI, DstOffsetParam, DstOffset);
-		SetShaderValue(RHICmdList, ShaderRHI, DimensionsParam, Dimensions);
+		FRHIBatchedShaderParameters& BatchedParameters = RHICmdList.GetScratchShaderParameters();
+
+		SetShaderValue(BatchedParameters, SrcOffsetParam, SrcOffset);
+		SetShaderValue(BatchedParameters, DstOffsetParam, DstOffset);
+		SetShaderValue(BatchedParameters, DimensionsParam, Dimensions);
+
+		RHICmdList.SetBatchedShaderParameters(RHICmdList.GetBoundComputeShader(), BatchedParameters);
 
 		RHICmdList.DispatchComputeShader(
 			FMath::DivideAndRoundUp(uint32(Dimensions.X), Context.ThreadGroupSizeX),
diff --git a/Engine/Source/Runtime/RenderCore/Public/MediaShaders.h b/Engine/Source/Runtime/RenderCore/Public/MediaShaders.h
--- a/Engine/Source/Runtime/RenderCore/Public/MediaShaders.h
+++ b/Engine/Source/Runtime/RenderCore/Public/MediaShaders.h
@@ -179,6 +179,7 @@ public:
 	{ }
 
 	RENDERCORE_API void SetParameters(FRHIBatchedShaderParameters& BatchedParameters, TRefCountPtr<FRHITexture2D> AYUVTexture, const FMatrix& ColorTransform, const FVector& YUVOffset, bool SrgbToLinear);
+	UE_DEPRECATED(5.3, "SetParameters with FRHIBatchedShaderParameters should be used.")
 	RENDERCORE_API void SetParameters(FRHICommandList& RHICmdList, TRefCountPtr<FRHITexture2D> AYUVTexture, const FMatrix& ColorTransform, const FVector& YUVOffset, bool SrgbToLinear);
 };
 
@@ -201,6 +202,7 @@ public:
 	{ }
 
 	RENDERCORE_API void SetParameters(FRHIBatchedShaderParameters& BatchedParameters, TRefCountPtr<FRHITexture2D> BMPTexture, const FIntPoint& OutputDimensions, bool SrgbToLinear);
+	UE_DEPRECATED(5.3, "SetParameters with FRHIBatchedShaderParameters should be used.")
 	RENDERCORE_API void SetParameters(FRHICommandList& RHICmdList, TRefCountPtr<FRHITexture2D> BMPTexture, const FIntPoint& OutputDimensions, bool SrgbToLinear);
 };
 
@@ -225,6 +227,7 @@ public:
 	{ }
 
 	RENDERCORE_API void SetParameters(FRHIBatchedShaderParameters& BatchedParameters, const FIntPoint & TexDim, FShaderResourceViewRHIRef SRV_Y, FShaderResourceViewRHIRef SRV_UV, const FIntPoint& OutputDimensions, const FMatrix& ColorTransform, const FVector& YUVOffset, bool SrgbToLinear);
+	UE_DEPRECATED(5.3, "SetParameters with FRHIBatchedShaderParameters should be used.")
 	RENDERCORE_API void SetParameters(FRHICommandList& RHICmdList, const FIntPoint & TexDim, FShaderResourceViewRHIRef SRV_Y, FShaderResourceViewRHIRef SRV_UV, const FIntPoint& OutputDimensions, const FMatrix& ColorTransform, const FVector& YUVOffset, bool SrgbToLinear);
 };
 
@@ -250,6 +253,7 @@ public:
 	{ }
 
 	RENDERCORE_API void SetParameters(FRHIBatchedShaderParameters& BatchedParameters, TRefCountPtr<FRHITexture2D> NV12Texture, const FIntPoint& OutputDimensions, const FMatrix& ColorTransform, const FVector& YUVOffset, bool SrgbToLinear);
+	UE_DEPRECATED(5.3, "SetParameters with FRHIBatchedShaderParameters should be used.")
 	RENDERCORE_API void SetParameters(FRHICommandList& RHICmdList, TRefCountPtr<FRHITexture2D> NV12Texture, const FIntPoint& OutputDimensions, const FMatrix& ColorTransform, const FVector& YUVOffset, bool SrgbToLinear);
 };
 
@@ -274,6 +278,7 @@ public:
 	{ }
 
 	RENDERCORE_API void SetParameters(FRHIBatchedShaderParameters& BatchedParameters, TRefCountPtr<FRHITexture2D> NV21Texture, const FIntPoint& OutputDimensions, const FMatrix& ColorTransform, const FVector& YUVOffset, bool SrgbToLinear);
+	UE_DEPRECATED(5.3, "SetParameters with FRHIBatchedShaderParameters should be used.")
 	RENDERCORE_API void SetParameters(FRHICommandList& RHICmdList, TRefCountPtr<FRHITexture2D> NV21Texture, const FIntPoint& OutputDimensions, const FMatrix& ColorTransform, const FVector& YUVOffset, bool SrgbToLinear);
 };
 
@@ -294,6 +299,7 @@ public:
 	{ }
 
 	RENDERCORE_API void SetParameters(FRHIBatchedShaderParameters& BatchedParameters, const FIntPoint& TexDim, FShaderResourceViewRHIRef SRV_Y, FShaderResourceViewRHIRef SRV_UV, const FIntPoint& OutputDimensions, const FMatrix44f& ColorTransform, const FMatrix44f& CSTransform, bool bIsST2084);
+	UE_DEPRECATED(5.3, "SetParameters with FRHIBatchedShaderParameters should be used.")
 	RENDERCORE_API void SetParameters(FRHICommandList& RHICmdList, const FIntPoint& TexDim, FShaderResourceViewRHIRef SRV_Y, FShaderResourceViewRHIRef SRV_UV, const FIntPoint& OutputDimensions, const FMatrix44f& ColorTransform, const FMatrix44f& CSTransform, bool bIsST2084);
 };
 
@@ -314,6 +320,7 @@ public:
 	{ }
 
 	RENDERCORE_API void SetParameters(FRHIBatchedShaderParameters& BatchedParameters, const FIntPoint& TexDim, TRefCountPtr<FRHITexture2D> NV12Texture, const FIntPoint& OutputDimensions, const FMatrix44f& ColorTransform, const FMatrix44f& CSTransform, bool bIsST2084);
+	UE_DEPRECATED(5.3, "SetParameters with FRHIBatchedShaderParameters should be used.")
 	RENDERCORE_API void SetParameters(FRHICommandList& RHICmdList, const FIntPoint& TexDim, TRefCountPtr<FRHITexture2D> NV12Texture, const FIntPoint& OutputDimensions, const FMatrix44f& ColorTransform, const FMatrix44f& CSTransform, bool bIsST2084);
 };
 
@@ -334,6 +341,7 @@ public:
 	{ }
 
 	RENDERCORE_API void SetParameters(FRHIBatchedShaderParameters& BatchedParameters, const FIntPoint& TexDim, FShaderResourceViewRHIRef SRV_Y, FShaderResourceViewRHIRef SRV_U, FShaderResourceViewRHIRef SRV_V, const FIntPoint& OutputDimensions, const FMatrix44f& ColorTransform, const FMatrix44f& CSTransform, bool bIsST2084);
+	UE_DEPRECATED(5.3, "SetParameters with FRHIBatchedShaderParameters should be used.")
 	RENDERCORE_API void SetParameters(FRHICommandList& RHICmdList, const FIntPoint& TexDim, FShaderResourceViewRHIRef SRV_Y, FShaderResourceViewRHIRef SRV_U, FShaderResourceViewRHIRef SRV_V, const FIntPoint& OutputDimensions, const FMatrix44f& ColorTransform, const FMatrix44f& CSTransform, bool bIsST2084);
 };
 
@@ -357,6 +365,7 @@ public:
 
 
 	RENDERCORE_API void SetParameters(FRHIBatchedShaderParameters& BatchedParameters, TRefCountPtr<FRHITexture2D> RGBTexture, const FIntPoint& OutputDimensions, bool SrgbToLinear, bool bIsST2084, const FMatrix44f& CSTransform);
+	UE_DEPRECATED(5.3, "SetParameters with FRHIBatchedShaderParameters should be used.")
 	RENDERCORE_API void SetParameters(FRHICommandList& RHICmdList, TRefCountPtr<FRHITexture2D> RGBTexture, const FIntPoint& OutputDimensions, bool SrgbToLinear, bool bIsST2084, const FMatrix44f& CSTransform);
 };
 
@@ -383,6 +392,7 @@ public:
 	{ }
 
 	RENDERCORE_API void SetParameters(FRHIBatchedShaderParameters& BatchedParameters, TRefCountPtr<FRHITexture2D> LumaTexture, TRefCountPtr<FRHITexture2D> CbCrTexture, const FMatrix& ColorTransform, const FVector& YUVOffset, bool SrgbToLinear);
+	UE_DEPRECATED(5.3, "SetParameters with FRHIBatchedShaderParameters should be used.")
 	RENDERCORE_API void SetParameters(FRHICommandList& RHICmdList, TRefCountPtr<FRHITexture2D> LumaTexture, TRefCountPtr<FRHITexture2D> CbCrTexture, const FMatrix& ColorTransform, const FVector& YUVOffset, bool SrgbToLinear);
 };
 
@@ -420,6 +430,7 @@ public:
 	{ }
 
 	RENDERCORE_API void SetParameters(FRHIBatchedShaderParameters& BatchedParameters, TRefCountPtr<FRHITexture2D> UYVYTexture, const FMatrix& ColorTransform, const FVector& YUVOffset, bool SrgbToLinear);
+	UE_DEPRECATED(5.3, "SetParameters with FRHIBatchedShaderParameters should be used.")
 	RENDERCORE_API void SetParameters(FRHICommandList& RHICmdList, TRefCountPtr<FRHITexture2D> UYVYTexture, const FMatrix& ColorTransform, const FVector& YUVOffset, bool SrgbToLinear);
 };
 
@@ -444,6 +455,7 @@ public:
 
 
 	RENDERCORE_API void SetParameters(FRHIBatchedShaderParameters& BatchedParameters, TRefCountPtr<FRHITexture2D> YTexture, TRefCountPtr<FRHITexture2D> UTexture, TRefCountPtr<FRHITexture2D> VTexture, const FIntPoint& OutputDimensions, const FMatrix& ColorTransform, const FVector& YUVOffset, bool SrgbToLinear);
+	UE_DEPRECATED(5.3, "SetParameters with FRHIBatchedShaderParameters should be used.")
 	RENDERCORE_API void SetParameters(FRHICommandList& RHICmdList, TRefCountPtr<FRHITexture2D> YTexture, TRefCountPtr<FRHITexture2D> UTexture, TRefCountPtr<FRHITexture2D> VTexture, const FIntPoint& OutputDimensions, const FMatrix& ColorTransform, const FVector& YUVOffset, bool SrgbToLinear);
 };
 
@@ -466,6 +478,7 @@ public:
 	{ }
 
 	RENDERCORE_API void SetParameters(FRHIBatchedShaderParameters& BatchedParameters, TRefCountPtr<FRHITexture2D> YUVTexture, const FIntPoint& OutputDimensions, const FMatrix& ColorTransform, const FVector& YUVOffset, bool SrgbToLinear);
+	UE_DEPRECATED(5.3, "SetParameters with FRHIBatchedShaderParameters should be used.")
 	RENDERCORE_API void SetParameters(FRHICommandList& RHICmdList, TRefCountPtr<FRHITexture2D> YUVTexture, const FIntPoint& OutputDimensions, const FMatrix& ColorTransform, const FVector& YUVOffset, bool SrgbToLinear);
 };
 
@@ -488,6 +501,7 @@ public:
 	{ }
 
 	RENDERCORE_API void SetParameters(FRHIBatchedShaderParameters& BatchedParameters, FShaderResourceViewRHIRef SRV_Y, const FMatrix& ColorTransform, const FVector& YUVOffset, bool SrgbToLinear);
+	UE_DEPRECATED(5.3, "SetParameters with FRHIBatchedShaderParameters should be used.")
 	RENDERCORE_API void SetParameters(FRHICommandList& RHICmdList, FShaderResourceViewRHIRef SRV_Y, const FMatrix& ColorTransform, const FVector& YUVOffset, bool SrgbToLinear);
 };
 
@@ -513,6 +527,7 @@ public:
 	{ }
 
 	RENDERCORE_API void SetParameters(FRHIBatchedShaderParameters& BatchedParameters, TRefCountPtr<FRHITexture2D> YUY2Texture, const FIntPoint& OutputDimensions, const FMatrix& ColorTransform, const FVector& YUVOffset, bool SrgbToLinear);
+	UE_DEPRECATED(5.3, "SetParameters with FRHIBatchedShaderParameters should be used.")
 	RENDERCORE_API void SetParameters(FRHICommandList& RHICmdList, TRefCountPtr<FRHITexture2D> YUY2Texture, const FIntPoint& OutputDimensions, const FMatrix& ColorTransform, const FVector& YUVOffset, bool SrgbToLinear);
 };
 
@@ -538,6 +553,7 @@ public:
 	{ }
 
 	RENDERCORE_API void SetParameters(FRHIBatchedShaderParameters& BatchedParameters, TRefCountPtr<FRHITexture2D> YVYUTexture, const FMatrix& ColorTransform, const FVector& YUVOffset, bool SrgbToLinear);
+	UE_DEPRECATED(5.3, "SetParameters with FRHIBatchedShaderParameters should be used.")
 	RENDERCORE_API void SetParameters(FRHICommandList& RHICmdList, TRefCountPtr<FRHITexture2D> YVYUTexture, const FMatrix& ColorTransform, const FVector& YUVOffset, bool SrgbToLinear);
 };
 
@@ -560,6 +576,7 @@ public:
 	{ }
 
 	RENDERCORE_API void SetParameters(FRHIBatchedShaderParameters& BatchedParameters, TRefCountPtr<FRHITexture2D> RGBATexture, const FVector4f& ColorTransform, bool LinearToSrgb);
+	UE_DEPRECATED(5.3, "SetParameters with FRHIBatchedShaderParameters should be used.")
 	RENDERCORE_API void SetParameters(FRHICommandList& RHICmdList, TRefCountPtr<FRHITexture2D> RGBATexture, const FVector4f& ColorTransform, bool LinearToSrgb);
 };
 
@@ -580,6 +597,7 @@ public:
 	{ }
 
 	RENDERCORE_API void SetParameters(FRHIBatchedShaderParameters& BatchedParameters, FTextureRHIRef TextureExt, FSamplerStateRHIRef SamplerState, const FLinearColor & ScaleRotation, const FLinearColor & Offset);
+	UE_DEPRECATED(5.3, "SetParameters with FRHIBatchedShaderParameters should be used.")
 	RENDERCORE_API void SetParameters(FRHICommandList& RHICmdList, FTextureRHIRef TextureExt, FSamplerStateRHIRef SamplerState, const FLinearColor & ScaleRotation, const FLinearColor & Offset);
 };
 
diff --git a/Engine/Source/Runtime/Renderer/Private/MobileDeferredShadingPass.cpp b/Engine/Source/Runtime/Renderer/Private/MobileDeferredShadingPass.cpp
--- a/Engine/Source/Runtime/Renderer/Private/MobileDeferredShadingPass.cpp
+++ b/Engine/Source/Runtime/Renderer/Private/MobileDeferredShadingPass.cpp
@@ -148,12 +148,9 @@ class FMobileDirectionalLightFunctionPS : public FMaterialShader
 		return PermutationVector;
 	}
 
-	static void SetParameters(FRHICommandList& RHICmdList, const TShaderRef<FMobileDirectionalLightFunctionPS>& Shader, const FViewInfo& View, const FMaterialRenderProxy* Proxy, const FMaterial& Material, const FParameters& Parameters)
+	void SetParameters(FRHIBatchedShaderParameters& BatchedParameters, const FViewInfo& View, const FMaterialRenderProxy* Proxy, const FMaterial& Material)
 	{
-		FMaterialShader* MaterialShader = Shader.GetShader();
-		FRHIPixelShader* ShaderRHI = Shader.GetPixelShader();
-		MaterialShader->SetParameters(RHICmdList, ShaderRHI, Proxy, Material, View);
-		SetShaderParameters(RHICmdList, Shader, ShaderRHI, Parameters);
+		FMaterialShader::SetParameters(BatchedParameters, Proxy, Material, View);
 	}
 };
 
@@ -226,13 +223,10 @@ public:
 		OutEnvironment.SetDefine(TEXT("IS_MOBILE_DEFERREDSHADING_SUBPASS"), 1u);
 	}
 
-	static void SetParameters(FRHICommandList& RHICmdList, const TShaderRef<FMobileRadialLightFunctionPS>& Shader, const FViewInfo& View, const FMaterialRenderProxy* Proxy, const FMaterial& Material, const FParameters& Parameters)
+	void SetParameters(FRHIBatchedShaderParameters& BatchedParameters, const FViewInfo& View, const FMaterialRenderProxy* Proxy, const FMaterial& Material)
 	{
-		FMaterialShader* MaterialShader = Shader.GetShader();
-		FRHIPixelShader* ShaderRHI = Shader.GetPixelShader();
-		MaterialShader->SetViewParameters(RHICmdList, ShaderRHI, View, View.ViewUniformBuffer);
-		MaterialShader->SetParameters(RHICmdList, ShaderRHI, Proxy, Material, View);
-		SetShaderParameters(RHICmdList, Shader, ShaderRHI, Parameters);
+		FMaterialShader::SetViewParameters(BatchedParameters, View, View.ViewUniformBuffer);
+		FMaterialShader::SetParameters(BatchedParameters, Proxy, Material, View);
 	}
 };
 
@@ -542,7 +536,7 @@ static void RenderDirectionalLight(FRHICommandListImmediate& RHICmdList, const F
 		GraphicsPSOInit.PrimitiveType = PT_TriangleList;
 		SetGraphicsPipelineState(RHICmdList, GraphicsPSOInit, StencilRef);
 
-		FMobileDirectionalLightFunctionPS::SetParameters(RHICmdList, PixelShader, View, LightMaterial.MaterialProxy, *LightMaterial.Material, PassParameters);
+		SetShaderParametersMixedPS(RHICmdList, PixelShader, PassParameters, View, LightMaterial.MaterialProxy, *LightMaterial.Material);
 
 		const FIntPoint TargetSize = View.GetSceneTexturesConfig().Extent;
 
@@ -795,7 +789,8 @@ static void RenderLocalLight(
 		SetGraphicsPipelineState(RHICmdList, GraphicsPSOInit, StencilRef);
 
 		SetShaderParameters(RHICmdList, VertexShader, VertexShader.GetVertexShader(), ParametersVS);
-		FMobileRadialLightFunctionPS::SetParameters(RHICmdList, PixelShader, View, LightMaterial.MaterialProxy, *LightMaterial.Material, PassParameters);
+
+		SetShaderParametersMixedPS(RHICmdList, PixelShader, PassParameters, View, LightMaterial.MaterialProxy, *LightMaterial.Material);
 
 		if (LightType == LightType_Point)
 		{
@@ -910,7 +905,8 @@ static void RenderSimpleLights(
 		{
 			uint8 StencilRef = GET_STENCIL_MOBILE_SM_MASK(PassShadingModelStencilValue[PassIndex]);
 			SetGraphicsPipelineState(RHICmdList, GraphicsPSOLight[PassIndex], StencilRef);
-			FMobileRadialLightFunctionPS::SetParameters(RHICmdList, PassPixelShaders[PassIndex], View, DefaultMaterial.MaterialProxy, *DefaultMaterial.Material, PassParameters);
+
+			SetShaderParametersMixedPS(RHICmdList, PassPixelShaders[PassIndex], PassParameters, View, DefaultMaterial.MaterialProxy, *DefaultMaterial.Material);
 
 			// Apply the point or spot light with some approximately bounding geometry,
 			// So we can get speedups from depth testing and not processing pixels outside of the light's influence.
diff --git a/Engine/Source/Runtime/Renderer/Private/PostProcess/PostProcessMaterial.cpp b/Engine/Source/Runtime/Renderer/Private/PostProcess/PostProcessMaterial.cpp
--- a/Engine/Source/Runtime/Renderer/Private/PostProcess/PostProcessMaterial.cpp
+++ b/Engine/Source/Runtime/Renderer/Private/PostProcess/PostProcessMaterial.cpp
@@ -262,13 +262,9 @@ public:
 		OutEnvironment.SetDefine(TEXT("STRATA_DEFERRED_SHADING"), 1);
 	}
 
-protected:
-	template <typename TRHIShader>
-	static void SetParameters(FRHICommandList& RHICmdList, const TShaderRef<FPostProcessMaterialShader> & Shader, TRHIShader* ShaderRHI, const FViewInfo& View, const FMaterialRenderProxy* Proxy, const FMaterial& Material, const FParameters& Parameters)
+	void SetParameters(FRHIBatchedShaderParameters& BatchedParameters, const FViewInfo& View, const FMaterialRenderProxy* Proxy, const FMaterial& Material)
 	{
-		FMaterialShader* MaterialShader = Shader.GetShader();
-		MaterialShader->SetParameters(RHICmdList, ShaderRHI, Proxy, Material, View);
-		SetShaderParameters(RHICmdList, Shader, ShaderRHI, Parameters);
+		FMaterialShader::SetParameters(BatchedParameters, Proxy, Material, View);
 	}
 };
 
@@ -277,11 +273,6 @@ class FPostProcessMaterialVS : public FPostProcessMaterialShader
 public:
 	DECLARE_SHADER_TYPE(FPostProcessMaterialVS, Material);
 
-	static void SetParameters(FRHICommandList& RHICmdList, const TShaderRef<FPostProcessMaterialVS>& Shader, const FViewInfo& View, const FMaterialRenderProxy* Proxy, const FMaterial& Material, const FParameters& Parameters)
-	{
-		FPostProcessMaterialShader::SetParameters(RHICmdList, Shader, Shader.GetVertexShader(), View, Proxy, Material, Parameters);
-	}
-
 	FPostProcessMaterialVS() = default;
 	FPostProcessMaterialVS(const ShaderMetaType::CompiledShaderInitializerType& Initializer)
 		: FPostProcessMaterialShader(Initializer)
@@ -296,11 +287,6 @@ public:
 	class FManualStencilTestDim : SHADER_PERMUTATION_BOOL("MANUAL_STENCIL_TEST");
 	using FPermutationDomain = TShaderPermutationDomain<FManualStencilTestDim>;
 
-	static void SetParameters(FRHICommandList& RHICmdList, const TShaderRef<FPostProcessMaterialPS>& Shader, const FViewInfo& View, const FMaterialRenderProxy* Proxy, const FMaterial& Material, const FParameters& Parameters)
-	{
-		FPostProcessMaterialShader::SetParameters(RHICmdList, Shader, Shader.GetPixelShader(), View, Proxy, Material, Parameters);
-	}
-
 	FPostProcessMaterialPS() = default;
 	FPostProcessMaterialPS(const ShaderMetaType::CompiledShaderInitializerType& Initializer)
 		: FPostProcessMaterialShader(Initializer)
@@ -324,11 +310,6 @@ public:
 
 		return true;
 	}
-
-	static void ModifyCompilationEnvironment(const FMaterialShaderPermutationParameters& Parameters, FShaderCompilerEnvironment& OutEnvironment)
-	{
-		FPostProcessMaterialShader::ModifyCompilationEnvironment(Parameters, OutEnvironment);
-	}
 };
 
 IMPLEMENT_SHADER_TYPE(,FPostProcessMaterialVS, TEXT("/Engine/Private/PostProcessMaterialShaders.usf"), TEXT("MainVS"), SF_Vertex);
@@ -686,10 +667,10 @@ FScreenPassTexture AddPostProcessMaterialPass(
 			PostProcessMaterialParameters,
 			ScreenPassFlags,
 			[&View, VertexShader, PixelShader, MaterialRenderProxy, Material, PostProcessMaterialParameters](FRHICommandList& RHICmdList)
-		{
-			FPostProcessMaterialVS::SetParameters(RHICmdList, VertexShader, View, MaterialRenderProxy, *Material, *PostProcessMaterialParameters);
-			FPostProcessMaterialPS::SetParameters(RHICmdList, PixelShader, View, MaterialRenderProxy, *Material, *PostProcessMaterialParameters);
-		});
+			{
+				SetShaderParametersMixedVS(RHICmdList, VertexShader, *PostProcessMaterialParameters, View, MaterialRenderProxy, *Material);
+				SetShaderParametersMixedPS(RHICmdList, PixelShader, *PostProcessMaterialParameters, View, MaterialRenderProxy, *Material);
+			});
 
 		if (bForceIntermediateTarget && !bCompositeWithInputAndDecode)
 		{
diff --git a/Engine/Source/Runtime/Renderer/Private/TranslucentLighting.cpp b/Engine/Source/Runtime/Renderer/Private/TranslucentLighting.cpp
--- a/Engine/Source/Runtime/Renderer/Private/TranslucentLighting.cpp
+++ b/Engine/Source/Runtime/Renderer/Private/TranslucentLighting.cpp
@@ -727,13 +727,12 @@ public:
 	FTranslucentLightingInjectPS() {}
 
 	void SetParameters(
-		FRHICommandList& RHICmdList, 
+		FRHIBatchedShaderParameters& BatchedParameters,
 		const FViewInfo& View, 
 		const FMaterialRenderProxy* MaterialProxy)
 	{
-		FRHIPixelShader* ShaderRHI = RHICmdList.GetBoundPixelShader();
 		const FMaterial& Material = MaterialProxy->GetMaterialWithFallback(View.GetFeatureLevel(), MaterialProxy);
-		FMaterialShader::SetParameters(RHICmdList, ShaderRHI, MaterialProxy, Material, View);
+		FMaterialShader::SetParameters(BatchedParameters, MaterialProxy, Material, View);
 	}
 };
 
@@ -1223,7 +1222,7 @@ void InjectTranslucencyLightingVolume(
 						SetShaderParametersLegacyGS(RHICmdList, GeometryShader, VolumeBounds.MinZ);
 					}
 
-					PixelShader->SetParameters(RHICmdList, View, InjectionData.LightFunctionMaterialProxy);
+					SetShaderParametersLegacyPS(RHICmdList, PixelShader, View, InjectionData.LightFunctionMaterialProxy);
 					SetShaderParameters(RHICmdList, PixelShader, PixelShader.GetPixelShader(), PassParameters->PS);
 
 					RasterizeToVolumeTexture(RHICmdList, VolumeBounds);
diff --git a/Engine/Source/Runtime/Renderer/Private/VT/RuntimeVirtualTextureRender.cpp b/Engine/Source/Runtime/Renderer/Private/VT/RuntimeVirtualTextureRender.cpp
--- a/Engine/Source/Runtime/Renderer/Private/VT/RuntimeVirtualTextureRender.cpp
+++ b/Engine/Source/Runtime/Renderer/Private/VT/RuntimeVirtualTextureRender.cpp
@@ -152,19 +152,6 @@ namespace RuntimeVirtualTexture
 			: FMeshMaterialShader(Initializer)
 		{
 		}
-
-		template <typename TRHICmdList>
-		void SetParameters(TRHICmdList& RHICmdList, FSceneView const& View, FMaterialRenderProxy const& MaterialProxy, FMaterial const& Material)
-		{
-			FMeshMaterialShader::SetParameters(
-				RHICmdList,
-				RHICmdList.GetBoundPixelShader(),
-				&MaterialProxy,
-				Material,
-				View,
-				View.ViewUniformBuffer,
-				ESceneTextureSetupMode::All);
-		}
 	};
 
 
diff --git a/Engine/Source/Runtime/SlateRHIRenderer/Private/SlateMaterialShader.cpp b/Engine/Source/Runtime/SlateRHIRenderer/Private/SlateMaterialShader.cpp
--- a/Engine/Source/Runtime/SlateRHIRenderer/Private/SlateMaterialShader.cpp
+++ b/Engine/Source/Runtime/SlateRHIRenderer/Private/SlateMaterialShader.cpp
@@ -31,16 +31,15 @@ bool FSlateMaterialShaderVS::ShouldCompilePermutation(const FMaterialShaderPermu
 	return Parameters.MaterialParameters.MaterialDomain == MD_UI;
 }
 
-void FSlateMaterialShaderVS::SetViewProjection(FRHICommandList& RHICmdList, const FMatrix44f& InViewProjection )
+void FSlateMaterialShaderVS::SetViewProjection(FRHIBatchedShaderParameters& BatchedParameters, const FMatrix44f& InViewProjection )
 {
-	SetShaderValue(RHICmdList, RHICmdList.GetBoundVertexShader(), ViewProjection, InViewProjection );
+	SetShaderValue(BatchedParameters, ViewProjection, InViewProjection );
 }
 
-void FSlateMaterialShaderVS::SetMaterialShaderParameters(FRHICommandList& RHICmdList, const FSceneView& View, const FMaterialRenderProxy* MaterialRenderProxy, const FMaterial* Material)
+void FSlateMaterialShaderVS::SetMaterialShaderParameters(FRHIBatchedShaderParameters& BatchedParameters, const FSceneView& View, const FMaterialRenderProxy* MaterialRenderProxy, const FMaterial* Material)
 {
-	FRHIVertexShader* ShaderRHI = RHICmdList.GetBoundVertexShader();
-	SetViewParameters(RHICmdList, ShaderRHI, View, View.ViewUniformBuffer);
-	FMaterialShader::SetParameters<FRHIVertexShader>(RHICmdList, ShaderRHI, MaterialRenderProxy, *Material, View);
+	SetViewParameters(BatchedParameters, View, View.ViewUniformBuffer);
+	FMaterialShader::SetParameters(BatchedParameters, MaterialRenderProxy, *Material, View);
 }
 
 bool FSlateMaterialShaderPS::ShouldCompilePermutation(const FMaterialShaderPermutationParameters& Parameters)
@@ -105,34 +104,32 @@ void FSlateMaterialShaderPS::SetBlendState(FGraphicsPipelineStateInitializer& Gr
 	};
 }
 
-void FSlateMaterialShaderPS::SetParameters(FRHICommandList& RHICmdList, const FSceneView& View, const FMaterialRenderProxy* MaterialRenderProxy, const FMaterial* Material, const FShaderParams& InShaderParams)
+void FSlateMaterialShaderPS::SetParameters(FRHIBatchedShaderParameters& BatchedParameters, const FSceneView& View, const FMaterialRenderProxy* MaterialRenderProxy, const FMaterial* Material, const FShaderParams& InShaderParams)
 {
-	FRHIPixelShader* ShaderRHI = RHICmdList.GetBoundPixelShader();
+	SetShaderValue(BatchedParameters, ShaderParams, (FVector4f)InShaderParams.PixelParams);
+	SetShaderValue(BatchedParameters, ShaderParams2, (FVector4f)InShaderParams.PixelParams2);
 
-	SetShaderValue(RHICmdList, ShaderRHI, ShaderParams, (FVector4f)InShaderParams.PixelParams);
-	SetShaderValue(RHICmdList, ShaderRHI, ShaderParams2, (FVector4f)InShaderParams.PixelParams2);
-
-	SetViewParameters(RHICmdList, ShaderRHI, View, View.ViewUniformBuffer);
-	FMaterialShader::SetParameters<FRHIPixelShader>(RHICmdList, ShaderRHI, MaterialRenderProxy, *Material, View);
+	SetViewParameters(BatchedParameters, View, View.ViewUniformBuffer);
+	FMaterialShader::SetParameters(BatchedParameters, MaterialRenderProxy, *Material, View);
 }
 
-void FSlateMaterialShaderPS::SetAdditionalTexture( FRHICommandList& RHICmdList, FRHITexture* InTexture, const FSamplerStateRHIRef SamplerState )
+void FSlateMaterialShaderPS::SetAdditionalTexture(FRHIBatchedShaderParameters& BatchedParameters, FRHITexture* InTexture, const FSamplerStateRHIRef SamplerState )
 {
-	SetTextureParameter(RHICmdList, RHICmdList.GetBoundPixelShader(), AdditionalTextureParameter, TextureParameterSampler, SamplerState, InTexture );
+	SetTextureParameter(BatchedParameters, AdditionalTextureParameter, TextureParameterSampler, SamplerState, InTexture );
 }
 
-void FSlateMaterialShaderPS::SetDisplayGammaAndContrast(FRHICommandList& RHICmdList, float InDisplayGamma, float InContrast)
+void FSlateMaterialShaderPS::SetDisplayGammaAndContrast(FRHIBatchedShaderParameters& BatchedParameters, float InDisplayGamma, float InContrast)
 {
 	FVector4f InGammaValues(2.2f / InDisplayGamma, 1.0f / InDisplayGamma, 0.0f, InContrast);
 
-	SetShaderValue(RHICmdList, RHICmdList.GetBoundPixelShader(), GammaAndAlphaValues, InGammaValues);
+	SetShaderValue(BatchedParameters, GammaAndAlphaValues, InGammaValues);
 }
 
-void FSlateMaterialShaderPS::SetDrawFlags(FRHICommandList& RHICmdList, bool bDrawDisabledEffect)
+void FSlateMaterialShaderPS::SetDrawFlags(FRHIBatchedShaderParameters& BatchedParameters, bool bDrawDisabledEffect)
 {
 	FVector4f InDrawFlags((bDrawDisabledEffect ? 1.f : 0.f), 0.f, 0.f, 0.f);
 
-	SetShaderValue(RHICmdList, RHICmdList.GetBoundPixelShader(), DrawFlags, InDrawFlags);
+	SetShaderValue(BatchedParameters, DrawFlags, InDrawFlags);
 }
 
 
diff --git a/Engine/Source/Runtime/SlateRHIRenderer/Private/SlateMaterialShader.h b/Engine/Source/Runtime/SlateRHIRenderer/Private/SlateMaterialShader.h
--- a/Engine/Source/Runtime/SlateRHIRenderer/Private/SlateMaterialShader.h
+++ b/Engine/Source/Runtime/SlateRHIRenderer/Private/SlateMaterialShader.h
@@ -25,9 +25,9 @@ public:
 	 *
 	 * @param InViewProjection	The ViewProjection matrix to use when this shader is bound 
 	 */
-	void SetViewProjection(FRHICommandList& RHICmdList, const FMatrix44f& InViewProjection );
+	void SetViewProjection(FRHIBatchedShaderParameters& BatchedParameters, const FMatrix44f& InViewProjection );
 
-	void SetMaterialShaderParameters(FRHICommandList& RHICmdList, const FSceneView& View, const FMaterialRenderProxy* MaterialRenderProxy, const FMaterial* Material);
+	void SetMaterialShaderParameters(FRHIBatchedShaderParameters& BatchedParameters, const FSceneView& View, const FMaterialRenderProxy* MaterialRenderProxy, const FMaterial* Material);
 
 	/** Serializes the shader data */
 	//virtual bool Serialize( FArchive& Ar ) override;
@@ -53,13 +53,13 @@ public:
 
 	void SetBlendState(FGraphicsPipelineStateInitializer& GraphicsPSOInit, const FMaterial* Material);
 
-	void SetParameters(FRHICommandList& RHICmdList, const FSceneView& View, const FMaterialRenderProxy* MaterialRenderProxy, const FMaterial* Material, const FShaderParams& InShaderParams);
+	void SetParameters(FRHIBatchedShaderParameters& BatchedParameters, const FSceneView& View, const FMaterialRenderProxy* MaterialRenderProxy, const FMaterial* Material, const FShaderParams& InShaderParams);
 
-	void SetDisplayGammaAndContrast(FRHICommandList& RHICmdList, float InDisplayGamma, float InContrast);
+	void SetDisplayGammaAndContrast(FRHIBatchedShaderParameters& BatchedParameters, float InDisplayGamma, float InContrast);
 
-	void SetAdditionalTexture( FRHICommandList& RHICmdList, FRHITexture* InTexture, const FSamplerStateRHIRef SamplerState );
+	void SetAdditionalTexture(FRHIBatchedShaderParameters& BatchedParameters, FRHITexture* InTexture, const FSamplerStateRHIRef SamplerState );
 
-	void SetDrawFlags(FRHICommandList& RHICmdList, bool bDrawDisabledEffect);
+	void SetDrawFlags(FRHIBatchedShaderParameters& BatchedParameters, bool bDrawDisabledEffect);
 
 private:
 	LAYOUT_FIELD(FShaderParameter, GammaAndAlphaValues);
diff --git a/Engine/Source/Runtime/SlateRHIRenderer/Private/SlatePostProcessor.cpp b/Engine/Source/Runtime/SlateRHIRenderer/Private/SlatePostProcessor.cpp
--- a/Engine/Source/Runtime/SlateRHIRenderer/Private/SlatePostProcessor.cpp
+++ b/Engine/Source/Runtime/SlateRHIRenderer/Private/SlatePostProcessor.cpp
@@ -79,18 +79,18 @@ public:
 	}
 	FBlitUIToHDRPSBase() = default;
 
-	void SetParameters(FRHICommandList& RHICmdList, FRHITexture* UITextureRHI, FRHITexture* HDRTextureRHI, FRHITexture* UITextureWriteMaskRHI, const FVector2f& InUITextureSize)
+	void SetParameters(FRHIBatchedShaderParameters& BatchedParameters, FRHITexture* UITextureRHI, FRHITexture* HDRTextureRHI, FRHITexture* UITextureWriteMaskRHI, const FVector2f& InUITextureSize)
 	{
-		SetTextureParameter(RHICmdList, RHICmdList.GetBoundPixelShader(), HDRTexture, UISampler, TStaticSamplerState<SF_Bilinear>::GetRHI(), HDRTextureRHI);
-		SetTextureParameter(RHICmdList, RHICmdList.GetBoundPixelShader(), UITexture, UISampler, TStaticSamplerState<SF_Bilinear>::GetRHI(), UITextureRHI);
+		SetTextureParameter(BatchedParameters, HDRTexture, UISampler, TStaticSamplerState<SF_Bilinear>::GetRHI(), HDRTextureRHI);
+		SetTextureParameter(BatchedParameters, UITexture, UISampler, TStaticSamplerState<SF_Bilinear>::GetRHI(), UITextureRHI);
 		static auto CVarHDRUILevel = IConsoleManager::Get().FindConsoleVariable(TEXT("r.HDR.UI.Level"));
 		static auto CVarHDRUILuminance = IConsoleManager::Get().FindConsoleVariable(TEXT("r.HDR.UI.Luminance"));
-		SetShaderValue(RHICmdList, RHICmdList.GetBoundPixelShader(), UILevel, CVarHDRUILevel ? CVarHDRUILevel->GetFloat() : 1.0f);
-		SetShaderValue(RHICmdList, RHICmdList.GetBoundPixelShader(), UILuminance, CVarHDRUILuminance ? CVarHDRUILuminance->GetFloat() : 300.0f);
-		SetShaderValue(RHICmdList, RHICmdList.GetBoundPixelShader(), UITextureSize, InUITextureSize);
+		SetShaderValue(BatchedParameters, UILevel, CVarHDRUILevel ? CVarHDRUILevel->GetFloat() : 1.0f);
+		SetShaderValue(BatchedParameters, UILuminance, CVarHDRUILuminance ? CVarHDRUILuminance->GetFloat() : 300.0f);
+		SetShaderValue(BatchedParameters, UITextureSize, InUITextureSize);
 		if (UITextureWriteMaskRHI != nullptr)
 		{
-			SetTextureParameter(RHICmdList, RHICmdList.GetBoundPixelShader(), UIWriteMaskTexture, UITextureWriteMaskRHI);
+			SetTextureParameter(BatchedParameters, UIWriteMaskTexture, UITextureWriteMaskRHI);
 		}
 	}
 
@@ -223,7 +223,7 @@ static void BlitUIToHDRScene(FRHICommandListImmediate& RHICmdList, IRendererModu
 		GraphicsPSOInit.PrimitiveType = PT_TriangleList;
 		SetGraphicsPipelineState(RHICmdList, GraphicsPSOInit, 0);
 
-		PixelShader->SetParameters(RHICmdList, UITexture, RectParams.SourceTexture, UITargetRTMaskTexture, FVector2f(SrcTextureWidth, SrcTextureHeight));
+		SetShaderParametersLegacyPS(RHICmdList, PixelShader, UITexture, RectParams.SourceTexture, UITargetRTMaskTexture, FVector2f(SrcTextureWidth, SrcTextureHeight));
 
 		RendererModule.DrawRectangle(
 			RHICmdList,
@@ -340,13 +340,17 @@ void FSlatePostProcessor::BlurRect(FRHICommandListImmediate& RHICmdList, IRender
 				GraphicsPSOInit.PrimitiveType = PT_TriangleList;
 				SetGraphicsPipelineState(RHICmdList, GraphicsPSOInit, 0);
 
-				PixelShader->SetWeightsAndOffsets(RHICmdList, WeightsAndOffsets, SampleCount);
-				PixelShader->SetTexture(RHICmdList, SourceTexture, BilinearClamp);
+				FRHIBatchedShaderParameters& BatchedParameters = RHICmdList.GetScratchShaderParameters();
+
+				PixelShader->SetWeightsAndOffsets(BatchedParameters, WeightsAndOffsets, SampleCount);
+				PixelShader->SetTexture(BatchedParameters, SourceTexture, BilinearClamp);
 
 				if (bDownsample || bIsHDRSource)
 				{
-					PixelShader->SetUVBounds(RHICmdList, FVector4f(FVector2f::ZeroVector, FVector2f((float)DownsampleSize.X / DestTextureWidth, (float)DownsampleSize.Y / DestTextureHeight) - HalfTexelOffset));
-					PixelShader->SetBufferSizeAndDirection(RHICmdList, InvBufferSize, FVector2f(1.f, 0.f));
+					PixelShader->SetUVBounds(BatchedParameters, FVector4f(FVector2f::ZeroVector, FVector2f((float)DownsampleSize.X / DestTextureWidth, (float)DownsampleSize.Y / DestTextureHeight) - HalfTexelOffset));
+					PixelShader->SetBufferSizeAndDirection(BatchedParameters, InvBufferSize, FVector2f(1.f, 0.f));
+
+					RHICmdList.SetBatchedShaderParameters(PixelShader.GetPixelShader(), BatchedParameters);
 
 					RendererModule.DrawRectangle(
 						RHICmdList,
@@ -367,8 +371,10 @@ void FSlatePostProcessor::BlurRect(FRHICommandListImmediate& RHICmdList, IRender
 					const FVector2f UVEnd = FVector2f(DestRect.Right, DestRect.Bottom) * InvSrcTextureSize;
 					const FVector2f SizeUV = UVEnd - UVStart;
 
-					PixelShader->SetUVBounds(RHICmdList, FVector4f(UVStart, UVEnd));
-					PixelShader->SetBufferSizeAndDirection(RHICmdList, InvSrcTextureSize, FVector2f(1.f, 0.f));
+					PixelShader->SetUVBounds(BatchedParameters, FVector4f(UVStart, UVEnd));
+					PixelShader->SetBufferSizeAndDirection(BatchedParameters, InvSrcTextureSize, FVector2f(1.f, 0.f));
+
+					RHICmdList.SetBatchedShaderParameters(PixelShader.GetPixelShader(), BatchedParameters);
 
 					RendererModule.DrawRectangle(
 						RHICmdList,
@@ -403,10 +409,14 @@ void FSlatePostProcessor::BlurRect(FRHICommandListImmediate& RHICmdList, IRender
 				GraphicsPSOInit.PrimitiveType = PT_TriangleList;
 				SetGraphicsPipelineState(RHICmdList, GraphicsPSOInit, 0);
 
-				PixelShader->SetWeightsAndOffsets(RHICmdList, WeightsAndOffsets, SampleCount);
-				PixelShader->SetUVBounds(RHICmdList, FVector4f(FVector2f::ZeroVector, FVector2f((float)DownsampleSize.X / DestTextureWidth, (float)DownsampleSize.Y / DestTextureHeight) - HalfTexelOffset));
-				PixelShader->SetTexture(RHICmdList, SourceTexture, BilinearClamp);
-				PixelShader->SetBufferSizeAndDirection(RHICmdList, InvBufferSize, FVector2f(0.f, 1.f));
+				FRHIBatchedShaderParameters& BatchedParameters = RHICmdList.GetScratchShaderParameters();
+
+				PixelShader->SetWeightsAndOffsets(BatchedParameters, WeightsAndOffsets, SampleCount);
+				PixelShader->SetUVBounds(BatchedParameters, FVector4f(FVector2f::ZeroVector, FVector2f((float)DownsampleSize.X / DestTextureWidth, (float)DownsampleSize.Y / DestTextureHeight) - HalfTexelOffset));
+				PixelShader->SetTexture(BatchedParameters, SourceTexture, BilinearClamp);
+				PixelShader->SetBufferSizeAndDirection(BatchedParameters, InvBufferSize, FVector2f(0.f, 1.f));
+
+				RHICmdList.SetBatchedShaderParameters(PixelShader.GetPixelShader(), BatchedParameters);
 
 				RendererModule.DrawRectangle(
 					RHICmdList,
@@ -487,9 +497,13 @@ void FSlatePostProcessor::ColorDeficiency(FRHICommandListImmediate& RHICmdList,
 			GraphicsPSOInit.PrimitiveType = PT_TriangleList;
 			SetGraphicsPipelineState(RHICmdList, GraphicsPSOInit, 0);
 
-			PixelShader->SetColorRules(RHICmdList, GSlateColorDeficiencyCorrection, GSlateColorDeficiencyType, GSlateColorDeficiencySeverity);
-			PixelShader->SetShowCorrectionWithDeficiency(RHICmdList, GSlateShowColorDeficiencyCorrectionWithDeficiency);
-			PixelShader->SetTexture(RHICmdList, SourceTexture, PointClamp);
+			FRHIBatchedShaderParameters& BatchedParameters = RHICmdList.GetScratchShaderParameters();
+
+			PixelShader->SetColorRules(BatchedParameters, GSlateColorDeficiencyCorrection, GSlateColorDeficiencyType, GSlateColorDeficiencySeverity);
+			PixelShader->SetShowCorrectionWithDeficiency(BatchedParameters, GSlateShowColorDeficiencyCorrectionWithDeficiency);
+			PixelShader->SetTexture(BatchedParameters, SourceTexture, PointClamp);
+
+			RHICmdList.SetBatchedShaderParameters(PixelShader.GetPixelShader(), BatchedParameters);
 
 			RendererModule.DrawRectangle(
 				RHICmdList,
@@ -577,9 +591,13 @@ void FSlatePostProcessor::DownsampleRect(FRHICommandListImmediate& RHICmdList, I
 			GraphicsPSOInit.PrimitiveType = PT_TriangleList;
 			SetGraphicsPipelineState(RHICmdList, GraphicsPSOInit, 0);
 
-			PixelShader->SetShaderParams(RHICmdList, FShaderParams::MakePixelShaderParams(FVector4f(InvSrcTextureSize.X, InvSrcTextureSize.Y, 0, 0)));
-			PixelShader->SetUVBounds(RHICmdList, FVector4f(UVStart, UVEnd));
-			PixelShader->SetTexture(RHICmdList, Params.SourceTexture, BilinearClamp);
+			FRHIBatchedShaderParameters& BatchedParameters = RHICmdList.GetScratchShaderParameters();
+
+			PixelShader->SetShaderParams(BatchedParameters, FShaderParams::MakePixelShaderParams(FVector4f(InvSrcTextureSize.X, InvSrcTextureSize.Y, 0, 0)));
+			PixelShader->SetUVBounds(BatchedParameters, FVector4f(UVStart, UVEnd));
+			PixelShader->SetTexture(BatchedParameters, Params.SourceTexture, BilinearClamp);
+
+			RHICmdList.SetBatchedShaderParameters(PixelShader.GetPixelShader(), BatchedParameters);
 
 			RendererModule.DrawRectangle(
 				RHICmdList,
@@ -701,9 +719,10 @@ void FSlatePostProcessor::UpsampleRect(FRHICommandListImmediate& RHICmdList, IRe
 		const FVector2f Size(DestRect.Right - DestRect.Left, DestRect.Bottom - DestRect.Top);
 		FShaderParams ShaderParams = FShaderParams::MakePixelShaderParams(FVector4f(Size, SizeUV), Params.CornerRadius);
 
-
-		PixelShader->SetShaderParams(RHICmdList, ShaderParams);
-		PixelShader->SetTexture(RHICmdList, SrcTexture, Sampler);
+		FRHIBatchedShaderParameters& BatchedParameters = RHICmdList.GetScratchShaderParameters();
+		PixelShader->SetShaderParams(BatchedParameters, ShaderParams);
+		PixelShader->SetTexture(BatchedParameters, SrcTexture, Sampler);
+		RHICmdList.SetBatchedShaderParameters(PixelShader.GetPixelShader(), BatchedParameters);
 
 
 		RendererModule.DrawRectangle(RHICmdList,
diff --git a/Engine/Source/Runtime/SlateRHIRenderer/Private/SlateRHIRenderingPolicy.cpp b/Engine/Source/Runtime/SlateRHIRenderer/Private/SlateRHIRenderingPolicy.cpp
--- a/Engine/Source/Runtime/SlateRHIRenderer/Private/SlateRHIRenderingPolicy.cpp
+++ b/Engine/Source/Runtime/SlateRHIRenderer/Private/SlateRHIRenderingPolicy.cpp
@@ -432,15 +432,19 @@ static bool UpdateScissorRect(
 					WriteMaskPSOInit.BoundShaderState.PixelShaderRHI = PixelShader.GetPixelShader();
 					WriteMaskPSOInit.PrimitiveType = PT_TriangleStrip;
 
-					SetGraphicsPipelineState(RHICmdList, WriteMaskPSOInit, MaskingID + 1);
-
-					VertexShader->SetViewProjection(RHICmdList, FMatrix44f(ViewProjection));
-
 					// Draw the first stencil using SO_Replace, so that we stomp any pixel with a MaskingID + 1.
 					{
+
+						SetGraphicsPipelineState(RHICmdList, WriteMaskPSOInit, MaskingID + 1);
+
 						const FSlateClippingZone& MaskQuad = StencilQuads[0];
 
-						VertexShader->SetMaskRect(RHICmdList, MaskQuad.TopLeft, MaskQuad.TopRight, MaskQuad.BottomLeft, MaskQuad.BottomRight);
+						FRHIBatchedShaderParameters& BatchedParameters = RHICmdList.GetScratchShaderParameters();
+
+						VertexShader->SetViewProjection(BatchedParameters, FMatrix44f(ViewProjection));
+						VertexShader->SetMaskRect(BatchedParameters, MaskQuad.TopLeft, MaskQuad.TopRight, MaskQuad.BottomLeft, MaskQuad.BottomRight);
+
+						RHICmdList.SetBatchedShaderParameters(VertexShader.GetVertexShader(), BatchedParameters);
 
 						RHICmdList.SetStreamSource(0, StencilVertexBuffer.VertexBufferRHI, 0);
 						RHICmdList.DrawPrimitive(0, 2, 1);
@@ -470,8 +474,6 @@ static bool UpdateScissorRect(
 
 
 						SetGraphicsPipelineState(RHICmdList, WriteMaskPSOInit, 0);
-
-						VertexShader->SetViewProjection(RHICmdList, FMatrix44f(ViewProjection));
 					}
 				}
 
@@ -482,7 +484,12 @@ static bool UpdateScissorRect(
 				{
 					const FSlateClippingZone& MaskQuad = StencilQuads[MaskIndex];
 
-					VertexShader->SetMaskRect(RHICmdList, MaskQuad.TopLeft, MaskQuad.TopRight, MaskQuad.BottomLeft, MaskQuad.BottomRight);
+					FRHIBatchedShaderParameters& BatchedParameters = RHICmdList.GetScratchShaderParameters();
+
+					VertexShader->SetViewProjection(BatchedParameters, FMatrix44f(ViewProjection));
+					VertexShader->SetMaskRect(BatchedParameters, MaskQuad.TopLeft, MaskQuad.TopRight, MaskQuad.BottomLeft, MaskQuad.BottomRight);
+
+					RHICmdList.SetBatchedShaderParameters(VertexShader.GetVertexShader(), BatchedParameters);
 
 					RHICmdList.SetStreamSource(0, StencilVertexBuffer.VertexBufferRHI, 0);
 					RHICmdList.DrawPrimitive(0, 2, 1);
@@ -575,6 +582,98 @@ static bool UpdateScissorRect(
 		bForceStateChange);
 }
 
+static FRHISamplerState* GetSamplerState(ESlateBatchDrawFlag DrawFlags, ETextureSamplerFilter Filter)
+{
+	FRHISamplerState* SamplerState = nullptr;
+
+	if (EnumHasAllFlags(DrawFlags, (ESlateBatchDrawFlag::TileU | ESlateBatchDrawFlag::TileV)))
+	{
+		switch (Filter)
+		{
+		case ETextureSamplerFilter::Point:
+			SamplerState = TStaticSamplerState<SF_Point, AM_Wrap, AM_Wrap, AM_Wrap>::GetRHI();
+			break;
+		case ETextureSamplerFilter::AnisotropicPoint:
+			SamplerState = TStaticSamplerState<SF_AnisotropicPoint, AM_Wrap, AM_Wrap, AM_Wrap>::GetRHI();
+			break;
+		case ETextureSamplerFilter::Trilinear:
+			SamplerState = TStaticSamplerState<SF_Trilinear, AM_Wrap, AM_Wrap, AM_Wrap>::GetRHI();
+			break;
+		case ETextureSamplerFilter::AnisotropicLinear:
+			SamplerState = TStaticSamplerState<SF_AnisotropicLinear, AM_Wrap, AM_Wrap, AM_Wrap>::GetRHI();
+			break;
+		default:
+			SamplerState = TStaticSamplerState<SF_Bilinear, AM_Wrap, AM_Wrap, AM_Wrap>::GetRHI();
+			break;
+		}
+	}
+	else if (EnumHasAllFlags(DrawFlags, ESlateBatchDrawFlag::TileU))
+	{
+		switch (Filter)
+		{
+		case ETextureSamplerFilter::Point:
+			SamplerState = TStaticSamplerState<SF_Point, AM_Wrap, AM_Clamp, AM_Wrap>::GetRHI();
+			break;
+		case ETextureSamplerFilter::AnisotropicPoint:
+			SamplerState = TStaticSamplerState<SF_AnisotropicPoint, AM_Wrap, AM_Clamp, AM_Wrap>::GetRHI();
+			break;
+		case ETextureSamplerFilter::Trilinear:
+			SamplerState = TStaticSamplerState<SF_Trilinear, AM_Wrap, AM_Clamp, AM_Wrap>::GetRHI();
+			break;
+		case ETextureSamplerFilter::AnisotropicLinear:
+			SamplerState = TStaticSamplerState<SF_AnisotropicLinear, AM_Wrap, AM_Clamp, AM_Wrap>::GetRHI();
+			break;
+		default:
+			SamplerState = TStaticSamplerState<SF_Bilinear, AM_Wrap, AM_Clamp, AM_Wrap>::GetRHI();
+			break;
+		}
+	}
+	else if (EnumHasAllFlags(DrawFlags, ESlateBatchDrawFlag::TileV))
+	{
+		switch (Filter)
+		{
+		case ETextureSamplerFilter::Point:
+			SamplerState = TStaticSamplerState<SF_Point, AM_Clamp, AM_Wrap, AM_Wrap>::GetRHI();
+			break;
+		case ETextureSamplerFilter::AnisotropicPoint:
+			SamplerState = TStaticSamplerState<SF_AnisotropicPoint, AM_Clamp, AM_Wrap, AM_Wrap>::GetRHI();
+			break;
+		case ETextureSamplerFilter::Trilinear:
+			SamplerState = TStaticSamplerState<SF_Trilinear, AM_Clamp, AM_Wrap, AM_Wrap>::GetRHI();
+			break;
+		case ETextureSamplerFilter::AnisotropicLinear:
+			SamplerState = TStaticSamplerState<SF_AnisotropicLinear, AM_Clamp, AM_Wrap, AM_Wrap>::GetRHI();
+			break;
+		default:
+			SamplerState = TStaticSamplerState<SF_Bilinear, AM_Clamp, AM_Wrap, AM_Wrap>::GetRHI();
+			break;
+		}
+	}
+	else
+	{
+		switch (Filter)
+		{
+		case ETextureSamplerFilter::Point:
+			SamplerState = TStaticSamplerState<SF_Point, AM_Clamp, AM_Clamp, AM_Clamp>::GetRHI();
+			break;
+		case ETextureSamplerFilter::AnisotropicPoint:
+			SamplerState = TStaticSamplerState<SF_AnisotropicPoint, AM_Clamp, AM_Clamp, AM_Clamp>::GetRHI();
+			break;
+		case ETextureSamplerFilter::Trilinear:
+			SamplerState = TStaticSamplerState<SF_Trilinear, AM_Clamp, AM_Clamp, AM_Clamp>::GetRHI();
+			break;
+		case ETextureSamplerFilter::AnisotropicLinear:
+			SamplerState = TStaticSamplerState<SF_AnisotropicLinear, AM_Clamp, AM_Clamp, AM_Clamp>::GetRHI();
+			break;
+		default:
+			SamplerState = TStaticSamplerState<SF_Bilinear, AM_Clamp, AM_Clamp, AM_Clamp>::GetRHI();
+			break;
+		}
+	}
+
+	return SamplerState;
+}
+
 void FSlateRHIRenderingPolicy::DrawElements(
 	FRHICommandListImmediate& RHICmdList,
 	FSlateBackBuffer& BackBuffer,
@@ -882,10 +981,12 @@ void FSlateRHIRenderingPolicy::DrawElements(
 
 				SetGraphicsPipelineState(RHICmdList, GraphicsPSOInit, StencilRef);
 
+				FRHIBatchedShaderParameters& BatchedParameters = RHICmdList.GetScratchShaderParameters();
+
 #if WITH_SLATE_VISUALIZERS
 				if (CVarShowSlateBatching.GetValueOnRenderThread() != 0)
 				{
-					BatchingPixelShader->SetBatchColor(RHICmdList, BatchColor);
+					BatchingPixelShader->SetBatchColor(BatchedParameters, BatchColor);
 				}
 #endif
 
@@ -930,108 +1031,29 @@ void FSlateRHIRenderingPolicy::DrawElements(
 						TextureRHI = NativeTextureRHI ? NativeTextureRHI : (FRHITexture*)GWhiteTexture->TextureRHI;
 					}
 
-					if (EnumHasAllFlags(DrawFlags, (ESlateBatchDrawFlag::TileU | ESlateBatchDrawFlag::TileV)))
-					{
-						switch (Filter)
-						{
-						case ETextureSamplerFilter::Point:
-							SamplerState = TStaticSamplerState<SF_Point, AM_Wrap, AM_Wrap, AM_Wrap>::GetRHI();
-							break;
-						case ETextureSamplerFilter::AnisotropicPoint:
-							SamplerState = TStaticSamplerState<SF_AnisotropicPoint, AM_Wrap, AM_Wrap, AM_Wrap>::GetRHI();
-							break;
-						case ETextureSamplerFilter::Trilinear:
-							SamplerState = TStaticSamplerState<SF_Trilinear, AM_Wrap, AM_Wrap, AM_Wrap>::GetRHI();
-							break;
-						case ETextureSamplerFilter::AnisotropicLinear:
-							SamplerState = TStaticSamplerState<SF_AnisotropicLinear, AM_Wrap, AM_Wrap, AM_Wrap>::GetRHI();
-							break;
-						default:
-							SamplerState = TStaticSamplerState<SF_Bilinear, AM_Wrap, AM_Wrap, AM_Wrap>::GetRHI();
-							break;
-						}
-					}
-					else if (EnumHasAllFlags(DrawFlags, ESlateBatchDrawFlag::TileU))
-					{
-						switch (Filter)
-						{
-						case ETextureSamplerFilter::Point:
-							SamplerState = TStaticSamplerState<SF_Point, AM_Wrap, AM_Clamp, AM_Wrap>::GetRHI();
-							break;
-						case ETextureSamplerFilter::AnisotropicPoint:
-							SamplerState = TStaticSamplerState<SF_AnisotropicPoint, AM_Wrap, AM_Clamp, AM_Wrap>::GetRHI();
-							break;
-						case ETextureSamplerFilter::Trilinear:
-							SamplerState = TStaticSamplerState<SF_Trilinear, AM_Wrap, AM_Clamp, AM_Wrap>::GetRHI();
-							break;
-						case ETextureSamplerFilter::AnisotropicLinear:
-							SamplerState = TStaticSamplerState<SF_AnisotropicLinear, AM_Wrap, AM_Clamp, AM_Wrap>::GetRHI();
-							break;
-						default:
-							SamplerState = TStaticSamplerState<SF_Bilinear, AM_Wrap, AM_Clamp, AM_Wrap>::GetRHI();
-							break;
-						}
-					}
-					else if (EnumHasAllFlags(DrawFlags, ESlateBatchDrawFlag::TileV))
-					{
-						switch (Filter)
-						{
-						case ETextureSamplerFilter::Point:
-							SamplerState = TStaticSamplerState<SF_Point, AM_Clamp, AM_Wrap, AM_Wrap>::GetRHI();
-							break;
-						case ETextureSamplerFilter::AnisotropicPoint:
-							SamplerState = TStaticSamplerState<SF_AnisotropicPoint, AM_Clamp, AM_Wrap, AM_Wrap>::GetRHI();
-							break;
-						case ETextureSamplerFilter::Trilinear:
-							SamplerState = TStaticSamplerState<SF_Trilinear, AM_Clamp, AM_Wrap, AM_Wrap>::GetRHI();
-							break;
-						case ETextureSamplerFilter::AnisotropicLinear:
-							SamplerState = TStaticSamplerState<SF_AnisotropicLinear, AM_Clamp, AM_Wrap, AM_Wrap>::GetRHI();
-							break;
-						default:
-							SamplerState = TStaticSamplerState<SF_Bilinear, AM_Clamp, AM_Wrap, AM_Wrap>::GetRHI();
-							break;
-						}
-					}
-					else
-					{
-						switch (Filter)
-						{
-						case ETextureSamplerFilter::Point:
-							SamplerState = TStaticSamplerState<SF_Point, AM_Clamp, AM_Clamp, AM_Clamp>::GetRHI();
-							break;
-						case ETextureSamplerFilter::AnisotropicPoint:
-							SamplerState = TStaticSamplerState<SF_AnisotropicPoint, AM_Clamp, AM_Clamp, AM_Clamp>::GetRHI();
-							break;
-						case ETextureSamplerFilter::Trilinear:
-							SamplerState = TStaticSamplerState<SF_Trilinear, AM_Clamp, AM_Clamp, AM_Clamp>::GetRHI();
-							break;
-						case ETextureSamplerFilter::AnisotropicLinear:
-							SamplerState = TStaticSamplerState<SF_AnisotropicLinear, AM_Clamp, AM_Clamp, AM_Clamp>::GetRHI();
-							break;
-						default:
-							SamplerState = TStaticSamplerState<SF_Bilinear, AM_Clamp, AM_Clamp, AM_Clamp>::GetRHI();
-							break;
-						}
-					}
+					SamplerState = GetSamplerState(DrawFlags, Filter);
 				}
 
 				{
-					GlobalVertexShader->SetViewProjection(RHICmdList, FMatrix44f(ViewProjection));
-
 					if (bIsVirtualTexture && (TextureResource != nullptr))
 					{
-						PixelShader->SetVirtualTextureParameters(RHICmdList, static_cast<FVirtualTexture2DResource*>(TextureResource));
+						PixelShader->SetVirtualTextureParameters(BatchedParameters, static_cast<FVirtualTexture2DResource*>(TextureResource));
 					}
 					else
 					{
-						PixelShader->SetTexture(RHICmdList, TextureRHI, SamplerState);
+						PixelShader->SetTexture(BatchedParameters, TextureRHI, SamplerState);
 					}
 					
-					PixelShader->SetShaderParams(RHICmdList, ShaderParams);
+					PixelShader->SetShaderParams(BatchedParameters, ShaderParams);
 					const float FinalGamma = EnumHasAnyFlags(DrawFlags, ESlateBatchDrawFlag::ReverseGamma) ? (1.0f / EngineGamma) : EnumHasAnyFlags(DrawFlags, ESlateBatchDrawFlag::NoGamma) ? 1.0f : DisplayGamma;
 					const float FinalContrast = EnumHasAnyFlags(DrawFlags, ESlateBatchDrawFlag::NoGamma) ? 1 : DisplayContrast;
-					PixelShader->SetDisplayGammaAndInvertAlphaAndContrast(RHICmdList, FinalGamma, EnumHasAllFlags(DrawEffects, ESlateDrawEffect::InvertAlpha) ? 1.0f : 0.0f, FinalContrast);
+					PixelShader->SetDisplayGammaAndInvertAlphaAndContrast(BatchedParameters, FinalGamma, EnumHasAllFlags(DrawEffects, ESlateDrawEffect::InvertAlpha) ? 1.0f : 0.0f, FinalContrast);
+
+					RHICmdList.SetBatchedShaderParameters(PixelShader.GetPixelShader(), BatchedParameters);
+				}
+				{
+					GlobalVertexShader->SetViewProjection(BatchedParameters, FMatrix44f(ViewProjection));
+					RHICmdList.SetBatchedShaderParameters(PixelShader.GetPixelShader(), BatchedParameters);
 				}
 
 				{
@@ -1114,6 +1136,8 @@ void FSlateRHIRenderingPolicy::DrawElements(
 						const FUniformBufferStaticBindings StaticUniformBuffers(SceneTextureUniformBuffer);
 						SCOPED_UNIFORM_BUFFER_STATIC_BINDINGS(RHICmdList, StaticUniformBuffers);
 
+						FRHIBatchedShaderParameters& BatchedParameters = RHICmdList.GetScratchShaderParameters();
+
 #if WITH_SLATE_VISUALIZERS
 						if (CVarShowSlateBatching.GetValueOnRenderThread() != 0)
 						{
@@ -1122,9 +1146,9 @@ void FSlateRHIRenderingPolicy::DrawElements(
 							GraphicsPSOInit.BoundShaderState.VertexDeclarationRHI = bUseInstancing ? GSlateInstancedVertexDeclaration.VertexDeclarationRHI : GSlateVertexDeclaration.VertexDeclarationRHI;
 							GraphicsPSOInit.BoundShaderState.VertexShaderRHI = GlobalVertexShader.GetVertexShader();
 							GraphicsPSOInit.BoundShaderState.PixelShaderRHI = BatchingPixelShader.GetPixelShader();
-
-							BatchingPixelShader->SetBatchColor(RHICmdList, BatchColor);
 							GraphicsPSOInit.BlendState = TStaticBlendState<CW_RGB, BO_Add, BF_One, BF_One, BO_Add, BF_Zero, BF_InverseSourceAlpha>::GetRHI();
+
+							BatchingPixelShader->SetBatchColor(BatchedParameters, BatchColor);
 						}
 						else if (CVarShowSlateOverdraw.GetValueOnRenderThread() != 0)
 						{
@@ -1133,10 +1157,8 @@ void FSlateRHIRenderingPolicy::DrawElements(
 							GraphicsPSOInit.BoundShaderState.VertexDeclarationRHI = bUseInstancing ? GSlateInstancedVertexDeclaration.VertexDeclarationRHI : GSlateVertexDeclaration.VertexDeclarationRHI;
 							GraphicsPSOInit.BoundShaderState.VertexShaderRHI = GlobalVertexShader.GetVertexShader();
 							GraphicsPSOInit.BoundShaderState.PixelShaderRHI = OverdrawPixelShader.GetPixelShader();
-
 							GraphicsPSOInit.BlendState = TStaticBlendState<CW_RGB, BO_Add, BF_One, BF_One, BO_Add, BF_Zero, BF_InverseSourceAlpha>::GetRHI();
 						}
-						else
 #endif
 						{
 							PixelShader->SetBlendState(GraphicsPSOInit, EffectiveMaterial);
@@ -1155,23 +1177,28 @@ void FSlateRHIRenderingPolicy::DrawElements(
 							SetGraphicsPipelineState(RHICmdList, GraphicsPSOInit, StencilRef);
 
 							{
-								VertexShader->SetViewProjection(RHICmdList, FMatrix44f(ViewProjection));
-								VertexShader->SetMaterialShaderParameters(RHICmdList, ActiveSceneView, MaterialRenderProxy, EffectiveMaterial);
 
-								PixelShader->SetParameters(RHICmdList, ActiveSceneView, MaterialRenderProxy, EffectiveMaterial, ShaderParams);
+								PixelShader->SetParameters(BatchedParameters, ActiveSceneView, MaterialRenderProxy, EffectiveMaterial, ShaderParams);
 								const float FinalGamma = EnumHasAnyFlags(DrawFlags, ESlateBatchDrawFlag::ReverseGamma) ? 1.0f / EngineGamma : EnumHasAnyFlags(DrawFlags, ESlateBatchDrawFlag::NoGamma) ? 1.0f : DisplayGamma;
 								const float FinalContrast = EnumHasAnyFlags(DrawFlags, ESlateBatchDrawFlag::NoGamma) ? 1 : DisplayContrast;
-								PixelShader->SetDisplayGammaAndContrast(RHICmdList, FinalGamma, FinalContrast);
+								PixelShader->SetDisplayGammaAndContrast(BatchedParameters, FinalGamma, FinalContrast);
 								const bool bDrawDisabled = EnumHasAllFlags(DrawEffects, ESlateDrawEffect::DisabledEffect);
-								PixelShader->SetDrawFlags(RHICmdList, bDrawDisabled);
+								PixelShader->SetDrawFlags(BatchedParameters, bDrawDisabled);
 
 								if (MaskResource)
 								{
 									FTexture2DRHIRef TextureRHI;
 									TextureRHI = ((TSlateTexture<FTexture2DRHIRef>*)MaskResource)->GetTypedResource();
 
-									PixelShader->SetAdditionalTexture(RHICmdList, TextureRHI, BilinearClamp);
+									PixelShader->SetAdditionalTexture(BatchedParameters, TextureRHI, BilinearClamp);
 								}
+
+								RHICmdList.SetBatchedShaderParameters(PixelShader.GetPixelShader(), BatchedParameters);
+							}
+							{
+								VertexShader->SetViewProjection(BatchedParameters, FMatrix44f(ViewProjection));
+								VertexShader->SetMaterialShaderParameters(BatchedParameters, ActiveSceneView, MaterialRenderProxy, EffectiveMaterial);
+								RHICmdList.SetBatchedShaderParameters(VertexShader.GetVertexShader(), BatchedParameters);
 							}
 						}
 
diff --git a/Engine/Source/Runtime/SlateRHIRenderer/Private/SlateShaders.cpp b/Engine/Source/Runtime/SlateRHIRenderer/Private/SlateShaders.cpp
--- a/Engine/Source/Runtime/SlateRHIRenderer/Private/SlateShaders.cpp
+++ b/Engine/Source/Runtime/SlateRHIRenderer/Private/SlateShaders.cpp
@@ -161,14 +161,14 @@ FSlateElementVS::FSlateElementVS( const ShaderMetaType::CompiledShaderInitialize
 	VertexShaderParams.Bind( Initializer.ParameterMap, TEXT("VertexShaderParams"));
 }
 
-void FSlateElementVS::SetViewProjection(FRHICommandList& RHICmdList, const FMatrix44f& InViewProjection )
+void FSlateElementVS::SetViewProjection(FRHIBatchedShaderParameters& BatchedParameters, const FMatrix44f& InViewProjection )
 {
-	SetShaderValue(RHICmdList, RHICmdList.GetBoundVertexShader(), ViewProjection, InViewProjection );
+	SetShaderValue(BatchedParameters, ViewProjection, InViewProjection );
 }
 
-void FSlateElementVS::SetShaderParameters(FRHICommandList& RHICmdList, const FVector4f& ShaderParams )
+void FSlateElementVS::SetShaderParameters(FRHIBatchedShaderParameters& BatchedParameters, const FVector4f& ShaderParams )
 {
-	SetShaderValue(RHICmdList, RHICmdList.GetBoundVertexShader(), VertexShaderParams, ShaderParams );
+	SetShaderValue(BatchedParameters, VertexShaderParams, ShaderParams );
 }
 
 /** Serializes the shader data */
@@ -194,16 +194,16 @@ FSlateMaskingVS::FSlateMaskingVS(const ShaderMetaType::CompiledShaderInitializer
 	MaskRect.Bind(Initializer.ParameterMap, TEXT("MaskRectPacked"));
 }
 
-void FSlateMaskingVS::SetViewProjection(FRHICommandList& RHICmdList, const FMatrix44f& InViewProjection)
+void FSlateMaskingVS::SetViewProjection(FRHIBatchedShaderParameters& BatchedParameters, const FMatrix44f& InViewProjection)
 {
-	SetShaderValue(RHICmdList, RHICmdList.GetBoundVertexShader(), ViewProjection, InViewProjection);
+	SetShaderValue(BatchedParameters, ViewProjection, InViewProjection);
 }
 
-void FSlateMaskingVS::SetMaskRect(FRHICommandList& RHICmdList, const FVector2f TopLeft, const FVector2f TopRight, const FVector2f BotLeft, const FVector2f BotRight)
+void FSlateMaskingVS::SetMaskRect(FRHIBatchedShaderParameters& BatchedParameters, const FVector2f TopLeft, const FVector2f TopRight, const FVector2f BotLeft, const FVector2f BotRight)
 {
 	FVector4f MaskRectVal[2] = { FVector4f(TopLeft, TopRight), FVector4f(BotLeft, BotRight) };
 
-	SetShaderValue(RHICmdList, RHICmdList.GetBoundVertexShader(), MaskRect, MaskRectVal);
+	SetShaderValue(BatchedParameters, MaskRect, MaskRectVal);
 }
 
 /** Serializes the shader data */
diff --git a/Engine/Source/Runtime/SlateRHIRenderer/Private/SlateShaders.h b/Engine/Source/Runtime/SlateRHIRenderer/Private/SlateShaders.h
--- a/Engine/Source/Runtime/SlateRHIRenderer/Private/SlateShaders.h
+++ b/Engine/Source/Runtime/SlateRHIRenderer/Private/SlateShaders.h
@@ -80,14 +80,14 @@ public:
 	 *
 	 * @param InViewProjection	The ViewProjection matrix to use when this shader is bound 
 	 */
-	void SetViewProjection(FRHICommandList& RHICmdList, const FMatrix44f& InViewProjection );
+	void SetViewProjection(FRHIBatchedShaderParameters& BatchedParameters, const FMatrix44f& InViewProjection );
 
 	/** 
 	 * Sets shader parameters for use in this shader
 	 *
 	 * @param ShaderParams	The shader params to be used
 	 */
-	void SetShaderParameters(FRHICommandList& RHICmdList, const FVector4f& ShaderParams );
+	void SetShaderParameters(FRHIBatchedShaderParameters& BatchedParameters, const FVector4f& ShaderParams );
 
 	/** Serializes the shader data */
 	//virtual bool Serialize( FArchive& Ar ) override;
@@ -139,9 +139,9 @@ public:
 	 * @param Texture	Texture resource to use when this pixel shader is bound
 	 * @param SamplerState	Sampler state to use when sampling this texture
 	 */
-	void SetTexture(FRHICommandList& RHICmdList, FRHITexture* InTexture, const FSamplerStateRHIRef SamplerState )
+	void SetTexture(FRHIBatchedShaderParameters& BatchedParameters, FRHITexture* InTexture, const FSamplerStateRHIRef SamplerState )
 	{
-		SetTextureParameter(RHICmdList, RHICmdList.GetBoundPixelShader(), TextureParameter, TextureParameterSampler, SamplerState, InTexture );
+		SetTextureParameter(BatchedParameters, TextureParameter, TextureParameterSampler, SamplerState, InTexture );
 	}
 
 	/**
@@ -149,7 +149,7 @@ public:
 	 *
 	 * @param InVirtualTexture	Virtual Texture resource to use when this pixel shader is bound
 	 */
-	void SetVirtualTextureParameters(FRHICommandList& RHICmdList, FVirtualTexture2DResource* InVirtualTexture)
+	void SetVirtualTextureParameters(FRHIBatchedShaderParameters& BatchedParameters, FVirtualTexture2DResource* InVirtualTexture)
 	{
 		if (InVirtualTexture == nullptr)
 		{
@@ -161,9 +161,9 @@ public:
 
 		FRHIShaderResourceView* PhysicalView = AllocatedVT->GetPhysicalTextureSRV(LayerIndex, InVirtualTexture->bSRGB);
 		
-		SetSRVParameter(RHICmdList, RHICmdList.GetBoundPixelShader(), TextureParameter, PhysicalView);
-		SetSamplerParameter(RHICmdList, RHICmdList.GetBoundPixelShader(), TextureParameterSampler, InVirtualTexture->SamplerStateRHI);
-		SetTextureParameter(RHICmdList, RHICmdList.GetBoundPixelShader(), InPageTableTexture, AllocatedVT->GetPageTableTexture(0u));
+		SetSRVParameter(BatchedParameters, TextureParameter, PhysicalView);
+		SetSamplerParameter(BatchedParameters, TextureParameterSampler, InVirtualTexture->SamplerStateRHI);
+		SetTextureParameter(BatchedParameters, InPageTableTexture, AllocatedVT->GetPageTableTexture(0u));
 		
 		FUintVector4 PageTableUniform[2];
 		FUintVector4 Uniform;
@@ -173,9 +173,9 @@ public:
 		AllocatedVT->GetPackedPageTableUniform(PageTableUniform);
 		AllocatedVT->GetPackedUniform(&Uniform, LayerIndex);
 
-		SetShaderValueArray(RHICmdList, RHICmdList.GetBoundPixelShader(), VTPackedPageTableUniform, PageTableUniform, UE_ARRAY_COUNT(PageTableUniform));
-		SetShaderValue(RHICmdList, RHICmdList.GetBoundPixelShader(), VTPackedUniform, Uniform);
-		SetShaderValue(RHICmdList, RHICmdList.GetBoundPixelShader(), VTShaderParams, VTParams);
+		SetShaderValueArray(BatchedParameters, VTPackedPageTableUniform, PageTableUniform, UE_ARRAY_COUNT(PageTableUniform));
+		SetShaderValue(BatchedParameters, VTPackedUniform, Uniform);
+		SetShaderValue(BatchedParameters, VTShaderParams, VTParams);
 	}
 
 	/**
@@ -183,10 +183,10 @@ public:
 	 * 
 	 * @param InShaderParams Shader params to use
 	 */
-	void SetShaderParams(FRHICommandList& RHICmdList, const FShaderParams& InShaderParams)
+	void SetShaderParams(FRHIBatchedShaderParameters& BatchedParameters, const FShaderParams& InShaderParams)
 	{
-		SetShaderValue(RHICmdList, RHICmdList.GetBoundPixelShader(), ShaderParams, InShaderParams.PixelParams);
-		SetShaderValue(RHICmdList, RHICmdList.GetBoundPixelShader(), ShaderParams2, InShaderParams.PixelParams2);
+		SetShaderValue(BatchedParameters, ShaderParams, InShaderParams.PixelParams);
+		SetShaderValue(BatchedParameters, ShaderParams2, InShaderParams.PixelParams2);
 	}
 
 	/**
@@ -194,11 +194,11 @@ public:
 	 *
 	 * @param DisplayGamma The display gamma to use
 	 */
-	void SetDisplayGammaAndInvertAlphaAndContrast(FRHICommandList& RHICmdList, float InDisplayGamma, float bInvertAlpha, float InContrast)
+	void SetDisplayGammaAndInvertAlphaAndContrast(FRHIBatchedShaderParameters& BatchedParameters, float InDisplayGamma, float bInvertAlpha, float InContrast)
 	{
 		FVector4f Values( 2.2f / InDisplayGamma, 1.0f/InDisplayGamma, bInvertAlpha, InContrast);
 
-		SetShaderValue(RHICmdList, RHICmdList.GetBoundPixelShader(), GammaAndAlphaValues, Values);
+		SetShaderValue(BatchedParameters, GammaAndAlphaValues, Values);
 	}
 
 private:
@@ -305,9 +305,9 @@ public:
 	*
 	* @param InShaderParams Shader params to use
 	*/
-	void SetBatchColor(FRHICommandList& RHICmdList, const FLinearColor& InBatchColor)
+	void SetBatchColor(FRHIBatchedShaderParameters& BatchedParameters, const FLinearColor& InBatchColor)
 	{
-		SetShaderValue(RHICmdList, RHICmdList.GetBoundPixelShader(), BatchColor, InBatchColor);
+		SetShaderValue(BatchedParameters, BatchColor, InBatchColor);
 	}
 
 private:
@@ -340,21 +340,21 @@ public:
 		UVBounds.Bind(Initializer.ParameterMap, TEXT("UVBounds"));
 	}
 
-	void SetBufferSizeAndDirection(FRHICommandList& RHICmdList, const FVector2f InBufferSize, const FVector2f InDir)
+	void SetBufferSizeAndDirection(FRHIBatchedShaderParameters& BatchedParameters, const FVector2f InBufferSize, const FVector2f InDir)
 	{
-		SetShaderValue(RHICmdList, RHICmdList.GetBoundPixelShader(), BufferSizeAndDirection, FVector4f(InBufferSize, InDir));
+		SetShaderValue(BatchedParameters, BufferSizeAndDirection, FVector4f(InBufferSize, InDir));
 	}
 
-	void SetWeightsAndOffsets(FRHICommandList& RHICmdList, const TArray<FVector4f>& InWeightsAndOffsets, int32 NumSamples )
+	void SetWeightsAndOffsets(FRHIBatchedShaderParameters& BatchedParameters, const TArray<FVector4f>& InWeightsAndOffsets, int32 NumSamples )
 	{
 		check(InWeightsAndOffsets.Num() <= MAX_BLUR_SAMPLES);
-		SetShaderValueArray<FRHIPixelShader*, FVector4f>(RHICmdList, RHICmdList.GetBoundPixelShader(), WeightAndOffsets, InWeightsAndOffsets.GetData(), InWeightsAndOffsets.Num() );
-		SetShaderValue(RHICmdList, RHICmdList.GetBoundPixelShader(), SampleCount, NumSamples);
+		SetShaderValueArray<FVector4f>(BatchedParameters, WeightAndOffsets, InWeightsAndOffsets.GetData(), InWeightsAndOffsets.Num() );
+		SetShaderValue(BatchedParameters, SampleCount, NumSamples);
 	}
 
-	void SetUVBounds(FRHICommandList& RHICmdList, const FVector4f& InUVBounds)
+	void SetUVBounds(FRHIBatchedShaderParameters& BatchedParameters, const FVector4f& InUVBounds)
 	{
-		SetShaderValue(RHICmdList, RHICmdList.GetBoundPixelShader(), UVBounds, InUVBounds);
+		SetShaderValue(BatchedParameters, UVBounds, InUVBounds);
 	}
 
 private:
@@ -386,9 +386,9 @@ public:
 		UVBounds.Bind(Initializer.ParameterMap, TEXT("UVBounds"));
 	}
 
-	void SetUVBounds(FRHICommandList& RHICmdList, const FVector4f& InUVBounds)
+	void SetUVBounds(FRHIBatchedShaderParameters& BatchedParameters, const FVector4f& InUVBounds)
 	{
-		SetShaderValue(RHICmdList, RHICmdList.GetBoundPixelShader(), UVBounds, InUVBounds);
+		SetShaderValue(BatchedParameters, UVBounds, InUVBounds);
 	}
 
 private:
@@ -462,16 +462,16 @@ public:
 		bSimulateCorrectionWithDeficiency.Bind(Initializer.ParameterMap, TEXT("bSimulateCorrectionWithDeficiency"));
 	}
 
-	void SetColorRules(FRHICommandList& RHICmdList, bool bCorrect, EColorVisionDeficiency DeficiencyType, int32 Severity)
+	void SetColorRules(FRHIBatchedShaderParameters& BatchedParameters, bool bCorrect, EColorVisionDeficiency DeficiencyType, int32 Severity)
 	{
-		SetShaderValue(RHICmdList, RHICmdList.GetBoundPixelShader(), ColorVisionDeficiencyType, (float)DeficiencyType);
-		SetShaderValue(RHICmdList, RHICmdList.GetBoundPixelShader(), ColorVisionDeficiencySeverity, (float)Severity);
-		SetShaderValue(RHICmdList, RHICmdList.GetBoundPixelShader(), bCorrectDeficiency, bCorrect ? 1.0f : 0.0f);
+		SetShaderValue(BatchedParameters, ColorVisionDeficiencyType, (float)DeficiencyType);
+		SetShaderValue(BatchedParameters, ColorVisionDeficiencySeverity, (float)Severity);
+		SetShaderValue(BatchedParameters, bCorrectDeficiency, bCorrect ? 1.0f : 0.0f);
 	}
 
-	void SetShowCorrectionWithDeficiency(FRHICommandList& RHICmdList, bool bShowCorrectionWithDeficiency)
+	void SetShowCorrectionWithDeficiency(FRHIBatchedShaderParameters& BatchedParameters, bool bShowCorrectionWithDeficiency)
 	{
-		SetShaderValue(RHICmdList, RHICmdList.GetBoundPixelShader(), bSimulateCorrectionWithDeficiency, bShowCorrectionWithDeficiency ? 1.0f : 0.0f);
+		SetShaderValue(BatchedParameters, bSimulateCorrectionWithDeficiency, bShowCorrectionWithDeficiency ? 1.0f : 0.0f);
 	}
 
 private:
@@ -504,12 +504,12 @@ public:
 	*
 	* @param InViewProjection	The ViewProjection matrix to use when this shader is bound
 	*/
-	void SetViewProjection(FRHICommandList& RHICmdList, const FMatrix44f& InViewProjection);
+	void SetViewProjection(FRHIBatchedShaderParameters& BatchedParameters, const FMatrix44f& InViewProjection);
 
 	/**
 	 * Sets the mask rect positions
 	 */
-	void SetMaskRect(FRHICommandList& RHICmdList, const FVector2f TopLeft, const FVector2f TopRight, const FVector2f BotLeft, const FVector2f BotRight);
+	void SetMaskRect(FRHIBatchedShaderParameters& BatchedParameters, const FVector2f TopLeft, const FVector2f TopRight, const FVector2f BotLeft, const FVector2f BotRight);
 
 	//virtual bool Serialize(FArchive& Ar) override;
 
