Added tracing for the existing "Stats_GPU_*" float counters (when "stats" trace channel is enabled).

diff --git a/Engine/Source/Runtime/RenderCore/Private/ProfilingDebugging/RealtimeGPUProfiler.cpp b/Engine/Source/Runtime/RenderCore/Private/ProfilingDebugging/RealtimeGPUProfiler.cpp
--- a/Engine/Source/Runtime/RenderCore/Private/ProfilingDebugging/RealtimeGPUProfiler.cpp
+++ b/Engine/Source/Runtime/RenderCore/Private/ProfilingDebugging/RealtimeGPUProfiler.cpp
@@ -586,8 +586,17 @@ public:
 				TotalUs += IncExcTime.ExclusiveTimeUs;
 
 #if STATS
-				const EStatOperation::Type StatOp = bKnownStat ? EStatOperation::Add : EStatOperation::Set;
-				FThreadStats::AddMessage(Event.GetStatName(), StatOp, EventTimeUs / 1000.);
+				const double EventTimeMs = EventTimeUs / 1000.;
+				if (bKnownStat)
+				{
+					FThreadStats::AddMessage(Event.GetStatName(), EStatOperation::Add, EventTimeMs);
+					TRACE_STAT_ADD(Event.GetStatName(), EventTimeMs);
+				}
+				else
+				{
+					FThreadStats::AddMessage(Event.GetStatName(), EStatOperation::Set, EventTimeMs);
+					TRACE_STAT_SET(Event.GetStatName(), EventTimeMs);
+				}
 #endif
 
 #if CSV_PROFILER
@@ -619,8 +628,10 @@ PRAGMA_ENABLE_DEPRECATION_WARNINGS
 		}
 
 #if STATS
-		FThreadStats::AddMessage(GET_STATFNAME(Stat_GPU_Total), EStatOperation::Set, TotalUs / 1000.);
-#endif 
+		const double TotalMs = TotalUs / 1000.;
+		FThreadStats::AddMessage(GET_STATFNAME(Stat_GPU_Total), EStatOperation::Set, TotalMs);
+		TRACE_STAT_SET(GET_STATFNAME(Stat_GPU_Total), TotalMs);
+#endif
 
 #if CSV_PROFILER
 		if (CsvProfiler)
