Only separate PCD3D_SM5 RT shaders for compression.

- The reason for the separate grouping was that PCD3D_SM5 is shared between SM5 and SM6 platforms and as such compressing raytracing shaders together with non-RT ones would affect memory footprint for D3D11.
- For PCD3D_SM6 users this is less of a concern. While this change can regress memory footprint for users who disable raytracing in the settings, it will improve memory footprint for those who do not (presumably, the majority). Also, it saves more in the disk space for everyone.

#rnx

diff --git a/Engine/Source/Runtime/RenderCore/Private/ShaderCodeArchive.cpp b/Engine/Source/Runtime/RenderCore/Private/ShaderCodeArchive.cpp
--- a/Engine/Source/Runtime/RenderCore/Private/ShaderCodeArchive.cpp
+++ b/Engine/Source/Runtime/RenderCore/Private/ShaderCodeArchive.cpp
@@ -1328,6 +1328,7 @@ void FIoStoreShaderCodeArchive::CreateIoStoreShaderCodeArchiveHeader(const FName
 	// (to avoid too large groups that will take too much time to decompress - this is regulated by r.ShaderCodeLibrary.MaxShaderGroupSize). The results of that process (note, it can still be
 	// a single group) is added to the header.
 	// Each group's indices, like in case of shadermaps, are stored in ShaderIndices array. Before we append a new group's indices however, we look if we can find an existing range that we can reuse.
+	const bool bSeparateRaytracingShaders = (Format == FName("PCD3D_SM5"));
 
 	TArray<TPair<uint32, TArray<int32>>> ShaderToShadermapsArray;
 	{
@@ -1566,32 +1567,39 @@ void FIoStoreShaderCodeArchive::CreateIoStoreShaderCodeArchiveHeader(const FName
 	};
 
 	/** First stage of processing a streak of shaders all referenced by the same set of shadermaps. We begin with separating raytracing and non-raytracing shaders, so we can avoid preloading RTX in non-RT runs. */
-	auto ProcessShaderGroup_SplitRaytracing = [&OutHeader, &ProcessShaderGroup_SplitBySize](TArray<uint32>& CurrentShaderGroup)
+	auto ProcessShaderGroup_SplitRaytracing = [&bSeparateRaytracingShaders, &OutHeader, &ProcessShaderGroup_SplitBySize](TArray<uint32>& CurrentShaderGroup)
 	{
-		// The streak changed. Create the group, but first, determine if the group needs to be split in two because of the raytracing shaders.
-		// We want to isolate them into separate groups so their preload can be skipped if raytracing is off.
-		TArray<uint32> RaytracingShaders;
-		TArray<uint32> NonraytracingShaders;
-		for (int32 ShaderIndex : CurrentShaderGroup)
+		if (!bSeparateRaytracingShaders)
 		{
-			if (LIKELY(!IsRayTracingShaderFrequency(static_cast<EShaderFrequency>(OutHeader.ShaderEntries[ShaderIndex].Frequency))))
+			ProcessShaderGroup_SplitBySize(CurrentShaderGroup);
+		}
+		else
+		{
+			// The streak changed. Create the group, but first, determine if the group needs to be split in two because of the raytracing shaders.
+			// We want to isolate them into separate groups so their preload can be skipped if raytracing is off.
+			TArray<uint32> RaytracingShaders;
+			TArray<uint32> NonraytracingShaders;
+			for (int32 ShaderIndex : CurrentShaderGroup)
 			{
-				NonraytracingShaders.Add(ShaderIndex);
+				if (LIKELY(!IsRayTracingShaderFrequency(static_cast<EShaderFrequency>(OutHeader.ShaderEntries[ShaderIndex].Frequency))))
+				{
+					NonraytracingShaders.Add(ShaderIndex);
+				}
+				else
+				{
+					RaytracingShaders.Add(ShaderIndex);
+				}
 			}
-			else
+			check(CurrentShaderGroup.Num() == NonraytracingShaders.Num() + RaytracingShaders.Num());
+
+			if (LIKELY(!NonraytracingShaders.IsEmpty()))
 			{
-				RaytracingShaders.Add(ShaderIndex);
+				ProcessShaderGroup_SplitBySize(NonraytracingShaders);
+			}
+			if (UNLIKELY(!RaytracingShaders.IsEmpty()))
+			{
+				ProcessShaderGroup_SplitBySize(RaytracingShaders);
 			}
-		}
-		check(CurrentShaderGroup.Num() == NonraytracingShaders.Num() + RaytracingShaders.Num());
-
-		if (LIKELY(!NonraytracingShaders.IsEmpty()))
-		{
-			ProcessShaderGroup_SplitBySize(NonraytracingShaders);
-		}
-		if (UNLIKELY(!RaytracingShaders.IsEmpty()))
-		{
-			ProcessShaderGroup_SplitBySize(RaytracingShaders);
 		}
 	};
 
