RHI Parameter Batching
- OneColorShader now uses batched parameters.
- Clear replacement shaders now use batched parameters.
- Deprecating the RHICmdList functions in ShaderParameterUtils.h and forcing them to go through the batched parameter code paths.

diff --git a/Engine/Source/Runtime/RenderCore/Private/ClearQuad.cpp b/Engine/Source/Runtime/RenderCore/Private/ClearQuad.cpp
--- a/Engine/Source/Runtime/RenderCore/Private/ClearQuad.cpp
+++ b/Engine/Source/Runtime/RenderCore/Private/ClearQuad.cpp
@@ -76,8 +76,11 @@ static void ClearQuadSetup(FRHICommandList& RHICmdList, int32 NumClearColors, co
 
 	SetGraphicsPipelineState(RHICmdList, GraphicsPSOInit, Stencil);
 
-	VertexShader->SetDepthParameter(RHICmdList, Depth);
-	PixelShader->SetColors(RHICmdList, PixelShader, ClearColorArray, NumClearColors);
+	SetShaderParametersLegacyVS(RHICmdList, VertexShader, Depth);
+	
+	TOneColorPixelShaderMRT::FParameters PixelParameters;
+	PixelShader->FillParameters(PixelParameters, ClearColorArray, NumClearColors);
+	SetShaderParameters(RHICmdList, PixelShader, PixelShader.GetPixelShader(), PixelParameters);
 }
 
 void DrawClearQuadAlpha(FRHICommandList& RHICmdList, float Alpha)
diff --git a/Engine/Source/Runtime/RenderCore/Private/LongGPUTask.cpp b/Engine/Source/Runtime/RenderCore/Private/LongGPUTask.cpp
--- a/Engine/Source/Runtime/RenderCore/Private/LongGPUTask.cpp
+++ b/Engine/Source/Runtime/RenderCore/Private/LongGPUTask.cpp
@@ -47,7 +47,8 @@ void IssueScalableLongGPUTask(FRHICommandListImmediate& RHICmdList, int32 NumIte
 		GraphicsPSOInit.PrimitiveType = PT_TriangleStrip;
 
 		SetGraphicsPipelineState(RHICmdList, GraphicsPSOInit, 0);
-		VertexShader->SetDepthParameter(RHICmdList, 0.0f);
+
+		SetShaderParametersLegacyVS(RHICmdList, VertexShader, 0.0f);
 
 		RHICmdList.SetStreamSource(0, GClearVertexBuffer.VertexBufferRHI, 0);
 
diff --git a/Engine/Source/Runtime/RenderCore/Private/ShaderParameters.cpp b/Engine/Source/Runtime/RenderCore/Private/ShaderParameters.cpp
--- a/Engine/Source/Runtime/RenderCore/Private/ShaderParameters.cpp
+++ b/Engine/Source/Runtime/RenderCore/Private/ShaderParameters.cpp
@@ -15,7 +15,9 @@
 
 IMPLEMENT_TYPE_LAYOUT(FShaderParameter);
 IMPLEMENT_TYPE_LAYOUT(FShaderResourceParameter);
+PRAGMA_DISABLE_DEPRECATION_WARNINGS
 IMPLEMENT_TYPE_LAYOUT(FRWShaderParameter);
+PRAGMA_ENABLE_DEPRECATION_WARNINGS
 IMPLEMENT_TYPE_LAYOUT(FShaderUniformBufferParameter);
 
 void FShaderParameter::Bind(const FShaderParameterMap& ParameterMap,const TCHAR* ParameterName,EShaderParameterFlags Flags)
diff --git a/Engine/Source/Runtime/RenderCore/Public/ClearReplacementShaders.h b/Engine/Source/Runtime/RenderCore/Public/ClearReplacementShaders.h
--- a/Engine/Source/Runtime/RenderCore/Public/ClearReplacementShaders.h
+++ b/Engine/Source/Runtime/RenderCore/Public/ClearReplacementShaders.h
@@ -331,9 +331,13 @@ inline void ClearUAVShader_T(FRHIComputeCommandList& RHICmdList, FRHIUnorderedAc
 	FRHIComputeShader* ShaderRHI = ComputeShader.GetComputeShader();
 	SetComputePipelineState(RHICmdList, ShaderRHI);
 
-	SetShaderValue(RHICmdList, ShaderRHI, ComputeShader->GetClearValueParam(), ClearValues);
-	SetShaderValue(RHICmdList, ShaderRHI, ComputeShader->GetMinBoundsParam(), FUintVector4(0, 0, 0, 0));
-	SetShaderValue(RHICmdList, ShaderRHI, ComputeShader->GetMaxBoundsParam(), FUintVector4(SizeX, SizeY, SizeZ, 0));
+	FRHIBatchedShaderParameters& BatchedParameters = RHICmdList.GetScratchShaderParameters();
+
+	SetShaderValue(BatchedParameters, ComputeShader->GetClearValueParam(), ClearValues);
+	SetShaderValue(BatchedParameters, ComputeShader->GetMinBoundsParam(), FUintVector4(0, 0, 0, 0));
+	SetShaderValue(BatchedParameters, ComputeShader->GetMaxBoundsParam(), FUintVector4(SizeX, SizeY, SizeZ, 0));
+
+	RHICmdList.SetBatchedShaderParameters(ShaderRHI, BatchedParameters);
 
 	if (bBarriers)
 	{
@@ -360,12 +364,40 @@ inline void ClearUAVShader_T(FRHIComputeCommandList& RHICmdList, FRHIUnorderedAc
 template <EClearReplacementResourceType ResourceType, EClearReplacementValueType ValueType, uint32 NumChannels, bool bBarriers>
 inline void ClearUAVShader_T(FRHIComputeCommandList& RHICmdList, FRHIUnorderedAccessView* UAV, uint32 SizeX, uint32 SizeY, uint32 SizeZ, const typename TClearReplacementTypeSelector<ValueType>::Type(&ClearValues)[NumChannels])
 {
-	ClearUAVShader_T<ResourceType, ValueType, NumChannels, bBarriers>(RHICmdList, UAV, SizeX, SizeY, SizeZ, ClearValues, 
-		[&RHICmdList, UAV](FRHIComputeShader* ShaderRHI, const FShaderResourceParameter& Param, bool bSet)
-		{
-			SetUAVParameter(RHICmdList, ShaderRHI, Param, bSet ? UAV : nullptr);
-		}
+	typedef TClearReplacementCS<ResourceType, TClearReplacementBase<ValueType, NumChannels, false, true>> FClearShader;
+
+	TShaderMapRef<FClearShader> ComputeShader(GetGlobalShaderMap(GMaxRHIFeatureLevel));
+	FRHIComputeShader* ShaderRHI = ComputeShader.GetComputeShader();
+	SetComputePipelineState(RHICmdList, ShaderRHI);
+
+	FRHIBatchedShaderParameters& BatchedParameters = RHICmdList.GetScratchShaderParameters();
+
+	SetShaderValue(BatchedParameters, ComputeShader->GetClearValueParam(), ClearValues);
+	SetShaderValue(BatchedParameters, ComputeShader->GetMinBoundsParam(), FUintVector4(0, 0, 0, 0));
+	SetShaderValue(BatchedParameters, ComputeShader->GetMaxBoundsParam(), FUintVector4(SizeX, SizeY, SizeZ, 0));
+
+	if (bBarriers)
+	{
+		RHICmdList.Transition(FRHITransitionInfo(UAV, ERHIAccess::Unknown, ERHIAccess::UAVCompute));
+	}
+
+	SetUAVParameter(BatchedParameters, ComputeShader->GetClearResourceParam(), UAV);
+
+	RHICmdList.SetBatchedShaderParameters(ShaderRHI, BatchedParameters);
+
+	RHICmdList.DispatchComputeShader(
+		FMath::DivideAndRoundUp(SizeX, ComputeShader->ThreadGroupSizeX),
+		FMath::DivideAndRoundUp(SizeY, ComputeShader->ThreadGroupSizeY),
+		FMath::DivideAndRoundUp(SizeZ, ComputeShader->ThreadGroupSizeZ)
 	);
+
+	SetUAVParameter(BatchedParameters, ComputeShader->GetClearResourceParam(), nullptr);
+	RHICmdList.SetBatchedShaderParameters(ShaderRHI, BatchedParameters);
+
+	if (bBarriers)
+	{
+		RHICmdList.Transition(FRHITransitionInfo(UAV, ERHIAccess::UAVCompute, ERHIAccess::SRVCompute));
+	}
 }
 
 // Helper version of ClearUAVShader_T for determining float vs uint32 at runtime. Uses the above default implementation.
diff --git a/Engine/Source/Runtime/RenderCore/Public/OneColorShader.h b/Engine/Source/Runtime/RenderCore/Public/OneColorShader.h
--- a/Engine/Source/Runtime/RenderCore/Public/OneColorShader.h
+++ b/Engine/Source/Runtime/RenderCore/Public/OneColorShader.h
@@ -48,9 +48,17 @@ public:
 		OutEnvironment.SetDefine(TEXT("USING_LAYERS"), (uint32)(bUsingVertexLayers ? 1 : 0));
 	}
 
+	void SetParameters(FRHIBatchedShaderParameters& BatchedParameters, float Depth)
+	{
+		SetShaderValue(BatchedParameters, DepthParameter, Depth);
+	}
+
+	UE_DEPRECATED(5.3, "SetParameters with FRHIBatchedShaderParameters should be used.")
 	void SetDepthParameter(FRHICommandList& RHICmdList, float Depth)
 	{
-		SetShaderValue(RHICmdList, RHICmdList.GetBoundVertexShader(), DepthParameter, Depth);
+		FRHIBatchedShaderParameters& BatchedParameters = RHICmdList.GetScratchShaderParameters();
+		SetParameters(BatchedParameters, Depth);
+		RHICmdList.SetBatchedShaderParameters(RHICmdList.GetBoundVertexShader(), BatchedParameters);
 	}
 
 	static bool ShouldCompilePermutation(const FGlobalShaderPermutationParameters& Parameters)
@@ -86,7 +94,9 @@ public:
 		RENDER_TARGET_BINDING_SLOTS()
 	END_SHADER_PARAMETER_STRUCT()
 
-	void FillParameters(FParameters& Parameters, const FLinearColor* Colors, int32 NumColors);
+	RENDERCORE_API void FillParameters(FParameters& Parameters, const FLinearColor* Colors, int32 NumColors);
+
+	UE_DEPRECATED(5.3, "FParameters and FillParameters should be used instead of this helper.")
 	RENDERCORE_API static void SetColors(FRHICommandList& RHICmdList, const TShaderMapRef<FOneColorPS>& Shader, const FLinearColor* Colors, int32 NumColors);
 };
 
@@ -106,6 +116,8 @@ public:
 
 	static bool ShouldCompilePermutation(const FGlobalShaderPermutationParameters& Parameters);
 	static void ModifyCompilationEnvironment(const FGlobalShaderPermutationParameters& Parameters, FShaderCompilerEnvironment& OutEnvironment);
+
+	UE_DEPRECATED(5.3, "FParameters and FillParameters should be used instead of this helper.")
 	RENDERCORE_API static void SetColors(FRHICommandList& RHICmdList, const TShaderMapRef<TOneColorPixelShaderMRT>& Shader, const FLinearColor* Colors, int32 NumColors);
 };
 
diff --git a/Engine/Source/Runtime/RenderCore/Public/ShaderParameterUtils.h b/Engine/Source/Runtime/RenderCore/Public/ShaderParameterUtils.h
--- a/Engine/Source/Runtime/RenderCore/Public/ShaderParameterUtils.h
+++ b/Engine/Source/Runtime/RenderCore/Public/ShaderParameterUtils.h
@@ -220,6 +220,20 @@ inline void SetUniformBufferParameterImmediate(FRHIBatchedShaderParameters& Batc
 	}
 }
 
+// Utility to set a single shader value on a shader. Should only be used if a shader requires only a single value.
+template<typename TRHICmdList, typename TShaderTypeRHI, class ParameterType>
+void SetSingleShaderValue(
+	TRHICmdList& RHICmdList
+	, TShaderTypeRHI* InShaderRHI
+	, const FShaderParameter& Parameter
+	, const ParameterType& Value
+)
+{
+	FRHIBatchedShaderParameters& BatchedParameters = RHICmdList.GetScratchShaderParameters();
+	SetShaderValue(BatchedParameters, Parameter, Value);
+	RHICmdList.SetBatchedShaderParameters(InShaderRHI, BatchedParameters);
+}
+
 // Mixed mode binding utilities
 
 /// Utility to set all legacy and non-legacy parameters for a shader. Requires the shader type to implement SetParameters(FRHIBatchedShaderParameters& BatchedParameters, ...)
@@ -362,6 +376,7 @@ inline void UnsetShaderParametersLegacyCS(TRHICmdList& RHICmdList, const TShader
  * Otherwise AddRef/ReleaseRef will be called many times.
  */
 template<typename ShaderRHIParamRef, class ParameterType, typename TRHICmdList>
+UE_DEPRECATED(5.3, "SetShaderValue with FRHIBatchedShaderParameters should be used.")
 void SetShaderValue(
 	TRHICmdList& RHICmdList,
 	const ShaderRHIParamRef& Shader,
@@ -370,28 +385,14 @@ void SetShaderValue(
 	uint32 ElementIndex = 0
 	)
 {
-	static_assert(!TIsPointer<ParameterType>::Value, "Passing by value is not valid.");
-
-	const uint32 AlignedTypeSize = (uint32)Align(sizeof(ParameterType), SHADER_PARAMETER_ARRAY_ELEMENT_ALIGNMENT);
-	const int32 NumBytesToSet = FMath::Min<int32>(sizeof(ParameterType),Parameter.GetNumBytes() - ElementIndex * AlignedTypeSize);
-
-	// This will trigger if the parameter was not serialized
-	checkSlow(Parameter.IsInitialized());
-
-	if(NumBytesToSet > 0)
-	{
-		RHICmdList.SetShaderParameter(
-			Shader,
-			Parameter.GetBufferIndex(),
-			Parameter.GetBaseIndex() + ElementIndex * AlignedTypeSize,
-			(uint32)NumBytesToSet,
-			&Value
-			);
-	}
+	FRHIBatchedShaderParameters& BatchedParameters = RHICmdList.GetScratchShaderParameters();
+	SetShaderValue(BatchedParameters, Parameter, Value, ElementIndex);
+	RHICmdList.SetBatchedShaderParameters(Shader, BatchedParameters);
 }
 
 
 template<typename ShaderRHIParamRef, class ParameterType>
+UE_DEPRECATED(5.3, "SetShaderValue with FRHIBatchedShaderParameters should be used.")
 void SetShaderValueOnContext(
 	IRHICommandContext& RHICmdListContext,
 	const ShaderRHIParamRef& Shader,
@@ -400,24 +401,9 @@ void SetShaderValueOnContext(
 	uint32 ElementIndex = 0
 	)
 {
-	static_assert(!TIsPointer<ParameterType>::Value, "Passing by value is not valid.");
-
-	const uint32 AlignedTypeSize = Align(sizeof(ParameterType), SHADER_PARAMETER_ARRAY_ELEMENT_ALIGNMENT);
-	const int32 NumBytesToSet = FMath::Min<int32>(sizeof(ParameterType), Parameter.GetNumBytes() - ElementIndex * AlignedTypeSize);
-
-	// This will trigger if the parameter was not serialized
-	checkSlow(Parameter.IsInitialized());
-
-	if (NumBytesToSet > 0)
-	{
-		RHICmdListContext.RHISetShaderParameter(
-			Shader,
-			Parameter.GetBufferIndex(),
-			Parameter.GetBaseIndex() + ElementIndex * AlignedTypeSize,
-			(uint32)NumBytesToSet,
-			&Value
-			);
-	}
+	FRHIBatchedShaderParameters BatchedParameters;
+	SetShaderValue(BatchedParameters, Parameter, Value, ElementIndex);
+	RHICmdListContext.RHISetBatchedShaderParameters(Shader, BatchedParameters);
 }
 
 /**
@@ -427,6 +413,7 @@ void SetShaderValueOnContext(
  * Otherwise AddRef/ReleaseRef will be called many times.
  */
 template<typename ShaderRHIParamRef,class ParameterType, typename TRHICmdList>
+UE_DEPRECATED(5.3, "SetShaderValueArray with FRHIBatchedShaderParameters should be used.")
 void SetShaderValueArray(
 	TRHICmdList& RHICmdList,
 	const ShaderRHIParamRef& Shader,
@@ -436,26 +423,14 @@ void SetShaderValueArray(
 	uint32 BaseElementIndex = 0
 	)
 {
-	const uint32 AlignedTypeSize = (uint32)Align(sizeof(ParameterType), SHADER_PARAMETER_ARRAY_ELEMENT_ALIGNMENT);
-	const int32 NumBytesToSet = FMath::Min<int32>(NumElements * AlignedTypeSize,Parameter.GetNumBytes() - BaseElementIndex * AlignedTypeSize);
-
-	// This will trigger if the parameter was not serialized
-	checkSlow(Parameter.IsInitialized());
-
-	if(NumBytesToSet > 0)
-	{
-		RHICmdList.SetShaderParameter(
-			Shader,
-			Parameter.GetBufferIndex(),
-			Parameter.GetBaseIndex() + BaseElementIndex * AlignedTypeSize,
-			(uint32)NumBytesToSet,
-			Values
-			);
-	}
+	FRHIBatchedShaderParameters& BatchedParameters = RHICmdList.GetScratchShaderParameters();
+	SetShaderValueArray(BatchedParameters, Parameter, Values, NumElements, BaseElementIndex);
+	RHICmdList.SetBatchedShaderParameters(Shader, BatchedParameters);
 }
 
 /** Specialization of the above for C++ bool type. */
 template<typename ShaderRHIParamRef, typename TRHICmdList>
+UE_DEPRECATED(5.3, "SetShaderValueArray with FRHIBatchedShaderParameters should be used.")
 void SetShaderValueArray(
 	TRHICmdList& RHICmdList,
 	const ShaderRHIParamRef& Shader,
@@ -502,26 +477,29 @@ FORCEINLINE void SetTextureParameter(TRHICmdList& RHICmdList, TRHIShader* Shader
 {
 	if (Parameter.IsBound() && ElementIndex < Parameter.GetNumResources())
 	{
+		FRHIBatchedShaderParameters& BatchedParameters = RHICmdList.GetScratchShaderParameters();
 #if PLATFORM_SUPPORTS_BINDLESS_RENDERING
 		if (Parameter.GetType() == EShaderParameterType::BindlessResourceIndex)
 		{
 			checkf(ElementIndex == 0, TEXT("Bindless resources don't support element offsets"));
-			RHICmdList.SetBindlessTexture(Shader, Parameter.GetBaseIndex(), TextureRHI);
+			BatchedParameters.SetBindlessTexture(Parameter.GetBaseIndex(), TextureRHI);
 		}
 		else
 #endif
 		{
-			RHICmdList.SetShaderTexture(Shader, Parameter.GetBaseIndex() + ElementIndex, TextureRHI);
+			BatchedParameters.SetShaderTexture(Parameter.GetBaseIndex() + ElementIndex, TextureRHI);
 		}
+		RHICmdList.SetBatchedShaderParameters(Shader, BatchedParameters);
 	}
 }
 
 template<typename TRHIShader, typename TRHICmdList>
+UE_DEPRECATED(5.3, "SetTextureParameter with FRHIBatchedShaderParameters should be used.")
 FORCEINLINE void SetTextureParameter(TRHICmdList& RHICmdList, TRHIShader* Shader, const FShaderResourceParameter& Parameter, FRHITexture* TextureRHI)
 {
-PRAGMA_DISABLE_DEPRECATION_WARNINGS
-	SetTextureParameter(RHICmdList, Shader, Parameter, TextureRHI, 0);
-PRAGMA_ENABLE_DEPRECATION_WARNINGS
+	FRHIBatchedShaderParameters& BatchedParameters = RHICmdList.GetScratchShaderParameters();
+	SetTextureParameter(BatchedParameters, Parameter, TextureRHI);
+	RHICmdList.SetBatchedShaderParameters(Shader, BatchedParameters);
 }
 
 /**
@@ -533,26 +511,29 @@ FORCEINLINE void SetSamplerParameter(TRHICmdList& RHICmdList, TRHIShader* Shader
 {
 	if (Parameter.IsBound() && ElementIndex < Parameter.GetNumResources())
 	{
+		FRHIBatchedShaderParameters& BatchedParameters = RHICmdList.GetScratchShaderParameters();
 #if PLATFORM_SUPPORTS_BINDLESS_RENDERING
 		if (Parameter.GetType() == EShaderParameterType::BindlessSamplerIndex)
 		{
 			checkf(ElementIndex == 0, TEXT("Bindless resources don't support element offsets"));
-			RHICmdList.SetBindlessSampler(Shader, Parameter.GetBaseIndex(), SamplerStateRHI);
+			BatchedParameters.SetBindlessSampler(Parameter.GetBaseIndex(), SamplerStateRHI);
 		}
 		else
 #endif
 		{
-			RHICmdList.SetShaderSampler(Shader, Parameter.GetBaseIndex() + ElementIndex, SamplerStateRHI);
+			BatchedParameters.SetShaderSampler(Parameter.GetBaseIndex() + ElementIndex, SamplerStateRHI);
 		}
+		RHICmdList.SetBatchedShaderParameters(Shader, BatchedParameters);
 	}
 }
 
 template<typename TRHIShader, typename TRHICmdList>
+UE_DEPRECATED(5.3, "SetSamplerParameter with FRHIBatchedShaderParameters should be used.")
 FORCEINLINE void SetSamplerParameter(TRHICmdList& RHICmdList, TRHIShader* Shader, const FShaderResourceParameter& Parameter, FRHISamplerState* SamplerStateRHI)
 {
-PRAGMA_DISABLE_DEPRECATION_WARNINGS
-	SetSamplerParameter(RHICmdList, Shader, Parameter, SamplerStateRHI, 0);
-PRAGMA_ENABLE_DEPRECATION_WARNINGS
+	FRHIBatchedShaderParameters& BatchedParameters = RHICmdList.GetScratchShaderParameters();
+	SetSamplerParameter(BatchedParameters, Parameter, SamplerStateRHI);
+	RHICmdList.SetBatchedShaderParameters(Shader, BatchedParameters);
 }
 
 /**
@@ -587,9 +568,9 @@ PRAGMA_ENABLE_DEPRECATION_WARNINGS
 template<typename TRHIShader, typename TRHICmdList>
 FORCEINLINE void SetTextureParameter(TRHICmdList& RHICmdList, TRHIShader* Shader, const FShaderResourceParameter& TextureParameter, const FShaderResourceParameter& SamplerParameter, FRHISamplerState* SamplerStateRHI, FRHITexture* TextureRHI)
 {
-PRAGMA_DISABLE_DEPRECATION_WARNINGS
-	SetTextureParameter(RHICmdList, Shader, TextureParameter, SamplerParameter, SamplerStateRHI, TextureRHI, 0);
-PRAGMA_ENABLE_DEPRECATION_WARNINGS
+	FRHIBatchedShaderParameters& BatchedParameters = RHICmdList.GetScratchShaderParameters();
+	SetTextureParameter(BatchedParameters, TextureParameter, SamplerParameter, SamplerStateRHI, TextureRHI);
+	RHICmdList.SetBatchedShaderParameters(Shader, BatchedParameters);
 }
 
 /**
@@ -618,11 +599,12 @@ PRAGMA_ENABLE_DEPRECATION_WARNINGS
 }
 
 template<typename TRHIShader, typename TRHICmdList>
+UE_DEPRECATED(5.3, "SetTextureParameter with FRHIBatchedShaderParameters should be used.")
 FORCEINLINE void SetTextureParameter(TRHICmdList& RHICmdList, TRHIShader* Shader, const FShaderResourceParameter& TextureParameter, const FShaderResourceParameter& SamplerParameter, const FTexture* Texture)
 {
-PRAGMA_DISABLE_DEPRECATION_WARNINGS
-	SetTextureParameter(RHICmdList, Shader, TextureParameter, SamplerParameter, Texture, 0);
-PRAGMA_ENABLE_DEPRECATION_WARNINGS
+	FRHIBatchedShaderParameters& BatchedParameters = RHICmdList.GetScratchShaderParameters();
+	SetTextureParameter(BatchedParameters, TextureParameter, SamplerParameter, Texture);
+	RHICmdList.SetBatchedShaderParameters(Shader, BatchedParameters);
 }
 
 /**
@@ -630,84 +612,84 @@ PRAGMA_ENABLE_DEPRECATION_WARNINGS
  * Template'd on shader type (e.g. pixel shader or compute shader).
  */
 template<typename TRHIShader, typename TRHICmdList>
+UE_DEPRECATED(5.3, "SetSRVParameter with FRHIBatchedShaderParameters should be used.")
 FORCEINLINE void SetSRVParameter(TRHICmdList& RHICmdList, TRHIShader* Shader, const FShaderResourceParameter& Parameter, FRHIShaderResourceView* SRV)
 {
-	if (Parameter.IsBound())
-	{
-#if PLATFORM_SUPPORTS_BINDLESS_RENDERING
-		if (Parameter.GetType() == EShaderParameterType::BindlessResourceIndex)
-		{
-			RHICmdList.SetBindlessResourceView(Shader, Parameter.GetBaseIndex(), SRV);
-		}
-		else
-#endif
-		{
-			RHICmdList.SetShaderResourceViewParameter(Shader, Parameter.GetBaseIndex(), SRV);
-		}
-	}
+	FRHIBatchedShaderParameters& BatchedParameters = RHICmdList.GetScratchShaderParameters();
+	SetSRVParameter(BatchedParameters, Parameter, SRV);
+	RHICmdList.SetBatchedShaderParameters(Shader, BatchedParameters);
 }
 
 template<typename TRHIShader, typename TRHICmdList>
+UE_DEPRECATED(5.3, "SetSRVParameter with FRHIBatchedShaderParameters should be used.")
 FORCEINLINE void SetSRVParameter(TRHICmdList& RHICmdList, const TRefCountPtr<TRHIShader>& Shader, const FShaderResourceParameter& Parameter, FRHIShaderResourceView* SRV)
 {
-	SetSRVParameter(RHICmdList, Shader.GetReference(), Parameter, SRV);
+	FRHIBatchedShaderParameters& BatchedParameters = RHICmdList.GetScratchShaderParameters();
+	SetSRVParameter(BatchedParameters, Parameter, SRV);
+	RHICmdList.SetBatchedShaderParameters(Shader.GetReference(), BatchedParameters);
 }
 
 template<typename TRHIShader, typename TRHICmdList>
+UE_DEPRECATED(5.3, "SetUAVParameter with FRHIBatchedShaderParameters should be used.")
 FORCEINLINE void SetUAVParameterSafeShader(TRHICmdList& RHICmdList, TRHIShader* Shader, const FShaderResourceParameter& Parameter, FRHIUnorderedAccessView* UAV)
 {
-	if (Parameter.IsBound())
-	{
-#if PLATFORM_SUPPORTS_BINDLESS_RENDERING
-		if (Parameter.GetType() == EShaderParameterType::BindlessResourceIndex)
-		{
-			RHICmdList.SetBindlessUAV(Shader, Parameter.GetBaseIndex(), UAV);
-		}
-		else
-#endif
-		{
-			RHICmdList.SetUAVParameter(Shader, Parameter.GetBaseIndex(), UAV);
-		}
-	}
+	FRHIBatchedShaderParameters& BatchedParameters = RHICmdList.GetScratchShaderParameters();
+	SetUAVParameter(BatchedParameters, Parameter, UAV);
+	RHICmdList.SetBatchedShaderParameters(Shader, BatchedParameters);
 }
 
+UE_DEPRECATED(5.3, "SetUAVParameter with FRHIBatchedShaderParameters should be used.")
 FORCEINLINE void SetUAVParameter(FRHIComputeCommandList& RHICmdList, FRHIComputeShader* Shader, const FShaderResourceParameter& Parameter, FRHIUnorderedAccessView* UAV)
 {
-	SetUAVParameterSafeShader(RHICmdList, Shader, Parameter, UAV);
+	FRHIBatchedShaderParameters& BatchedParameters = RHICmdList.GetScratchShaderParameters();
+	SetUAVParameter(BatchedParameters, Parameter, UAV);
+	RHICmdList.SetBatchedShaderParameters(Shader, BatchedParameters);
 }
 
+UE_DEPRECATED(5.3, "SetUAVParameter with FRHIBatchedShaderParameters should be used.")
 FORCEINLINE void SetUAVParameter(FRHICommandList& RHICmdList, FRHIPixelShader* Shader, const FShaderResourceParameter& Parameter, FRHIUnorderedAccessView* UAV)
 {
-	SetUAVParameterSafeShader(RHICmdList, Shader, Parameter, UAV);
+	FRHIBatchedShaderParameters& BatchedParameters = RHICmdList.GetScratchShaderParameters();
+	SetUAVParameter(BatchedParameters, Parameter, UAV);
+	RHICmdList.SetBatchedShaderParameters(Shader, BatchedParameters);
 }
 
 
 template<typename TRHICmdList>
+UE_DEPRECATED(5.3, "SetUAVParameter with FRHIBatchedShaderParameters should be used.")
 inline bool SetUAVParameterIfCS(TRHICmdList& RHICmdList, FRHIVertexShader* Shader, const FShaderResourceParameter& UAVParameter, FRHIUnorderedAccessView* UAV)
 {
 	return false;
 }
 
 template<typename TRHICmdList>
+UE_DEPRECATED(5.3, "SetUAVParameter with FRHIBatchedShaderParameters should be used.")
 inline bool SetUAVParameterIfCS(TRHICmdList& RHICmdList, FRHIPixelShader* Shader, const FShaderResourceParameter& UAVParameter, FRHIUnorderedAccessView* UAV)
 {
+PRAGMA_DISABLE_DEPRECATION_WARNINGS
 	SetUAVParameter(RHICmdList, Shader, UAVParameter, UAV);
+PRAGMA_ENABLE_DEPRECATION_WARNINGS
 	return UAVParameter.IsBound();
 }
 
 template<typename TRHICmdList>
+UE_DEPRECATED(5.3, "SetUAVParameter with FRHIBatchedShaderParameters should be used.")
 inline bool SetUAVParameterIfCS(TRHICmdList& RHICmdList, FRHIGeometryShader* Shader, const FShaderResourceParameter& UAVParameter, FRHIUnorderedAccessView* UAV)
 {
 	return false;
 }
 
 template<typename TRHICmdList>
+UE_DEPRECATED(5.3, "SetUAVParameter with FRHIBatchedShaderParameters should be used.")
 inline bool SetUAVParameterIfCS(TRHICmdList& RHICmdList, FRHIComputeShader* Shader, const FShaderResourceParameter& UAVParameter, FRHIUnorderedAccessView* UAV)
 {
+PRAGMA_DISABLE_DEPRECATION_WARNINGS
 	SetUAVParameter(RHICmdList, Shader, UAVParameter, UAV);
+PRAGMA_ENABLE_DEPRECATION_WARNINGS
 	return UAVParameter.IsBound();
 }
 
+PRAGMA_DISABLE_DEPRECATION_WARNINGS
 template<typename TShaderRHIRef, typename TRHICmdList>
 inline void FRWShaderParameter::SetBuffer(TRHICmdList& RHICmdList, const TShaderRHIRef& Shader, const FRWBuffer& RWBuffer) const
 {
@@ -746,9 +728,11 @@ inline void FRWShaderParameter::UnsetUAV(TRHICmdList& RHICmdList, FRHIComputeSha
 {
 	SetUAVParameter(RHICmdList, ComputeShader,UAVParameter,FUnorderedAccessViewRHIRef());
 }
+PRAGMA_ENABLE_DEPRECATION_WARNINGS
 
 /** Sets the value of a shader uniform buffer parameter to a uniform buffer containing the struct. */
 template<typename TShaderRHIRef, typename TRHICmdList>
+UE_DEPRECATED(5.3, "SetUniformBufferParameter with FRHIBatchedShaderParameters should be used.")
 inline void SetUniformBufferParameter(
 	TRHICmdList& RHICmdList,
 	const TShaderRHIRef& Shader,
@@ -756,18 +740,14 @@ inline void SetUniformBufferParameter(
 	FRHIUniformBuffer* UniformBufferRHI
 	)
 {
-	// This will trigger if the parameter was not serialized
-	checkSlow(Parameter.IsInitialized());
-	// If it is bound, we must set it so something valid
-	checkSlow(!Parameter.IsBound() || UniformBufferRHI);
-	if(Parameter.IsBound())
-	{
-		RHICmdList.SetShaderUniformBuffer( Shader, Parameter.GetBaseIndex(), UniformBufferRHI);
-	}
+	FRHIBatchedShaderParameters& BatchedParameters = RHICmdList.GetScratchShaderParameters();
+	SetUniformBufferParameter(BatchedParameters, Parameter, UniformBufferRHI);
+	RHICmdList.SetBatchedShaderParameters(Shader, BatchedParameters);
 }
 
 /** Sets the value of a shader uniform buffer parameter to a uniform buffer containing the struct. */
 template<typename TShaderRHIRef, typename TBufferStruct, typename TRHICmdList>
+UE_DEPRECATED(5.3, "SetUniformBufferParameter with FRHIBatchedShaderParameters should be used.")
 inline void SetUniformBufferParameter(
 	TRHICmdList& RHICmdList,
 	const TShaderRHIRef& Shader,
@@ -775,18 +755,14 @@ inline void SetUniformBufferParameter(
 	const TUniformBufferRef<TBufferStruct>& UniformBufferRef
 	)
 {
-	// This will trigger if the parameter was not serialized
-	checkSlow(Parameter.IsInitialized());
-	// If it is bound, we must set it so something valid
-	checkSlow(!Parameter.IsBound() || IsValidRef(UniformBufferRef));
-	if(Parameter.IsBound())
-	{
-		RHICmdList.SetShaderUniformBuffer( Shader, Parameter.GetBaseIndex(), UniformBufferRef);
-	}
+	FRHIBatchedShaderParameters& BatchedParameters = RHICmdList.GetScratchShaderParameters();
+	SetUniformBufferParameter(BatchedParameters, Parameter, UniformBufferRef);
+	RHICmdList.SetBatchedShaderParameters(Shader, BatchedParameters);
 }
 
 /** Sets the value of a shader uniform buffer parameter to a uniform buffer containing the struct. */
 template<typename TShaderRHIRef, typename TBufferStruct, typename TRHICmdList>
+UE_DEPRECATED(5.3, "SetUniformBufferParameter with FRHIBatchedShaderParameters should be used.")
 inline void SetUniformBufferParameter(
 	TRHICmdList& RHICmdList,
 	const TShaderRHIRef& Shader,
@@ -794,18 +770,14 @@ inline void SetUniformBufferParameter(
 	const TUniformBuffer<TBufferStruct>& UniformBuffer
 	)
 {
-	// This will trigger if the parameter was not serialized
-	checkSlow(Parameter.IsInitialized());
-	// If it is bound, we must set it so something valid
-	checkSlow(!Parameter.IsBound() || UniformBuffer.GetUniformBufferRHI());
-	if (Parameter.IsBound())
-	{
-		RHICmdList.SetShaderUniformBuffer( Shader, Parameter.GetBaseIndex(), UniformBuffer.GetUniformBufferRHI());
-	}
+	FRHIBatchedShaderParameters& BatchedParameters = RHICmdList.GetScratchShaderParameters();
+	SetUniformBufferParameter(BatchedParameters, Parameter, UniformBuffer);
+	RHICmdList.SetBatchedShaderParameters(Shader, BatchedParameters);
 }
 
 /** Sets the value of a shader uniform buffer parameter to a value of the struct. */
 template<typename TShaderRHIRef,typename TBufferStruct>
+UE_DEPRECATED(5.3, "SetUniformBufferParameterImmediate with FRHIBatchedShaderParameters should be used.")
 inline void SetUniformBufferParameterImmediate(
 	FRHICommandList& RHICmdList,
 	const TShaderRHIRef& Shader,
@@ -813,20 +785,14 @@ inline void SetUniformBufferParameterImmediate(
 	const TBufferStruct& UniformBufferValue
 	)
 {
-	// This will trigger if the parameter was not serialized
-	checkSlow(Parameter.IsInitialized());
-	if(Parameter.IsBound())
-	{
-		RHICmdList.SetShaderUniformBuffer(
-			Shader,
-			Parameter.GetBaseIndex(),
-			RHICreateUniformBuffer(&UniformBufferValue, &TBufferStruct::FTypeInfo::GetStructMetadata()->GetLayout(), UniformBuffer_SingleDraw)
-		);
-	}
+	FRHIBatchedShaderParameters& BatchedParameters = RHICmdList.GetScratchShaderParameters();
+	SetUniformBufferParameterImmediate(BatchedParameters, Parameter, UniformBufferValue);
+	RHICmdList.SetBatchedShaderParameters(Shader, BatchedParameters);
 }
 
 /** Sets the value of a shader uniform buffer parameter to a value of the struct. */
 template<typename TShaderRHIRef,typename TBufferStruct, typename TRHICmdList>
+UE_DEPRECATED(5.3, "SetUniformBufferParameterImmediate with FRHIBatchedShaderParameters should be used.")
 inline void SetUniformBufferParameterImmediate(
 	TRHICmdList& RHICmdList,
 	const TShaderRHIRef& Shader,
@@ -834,14 +800,7 @@ inline void SetUniformBufferParameterImmediate(
 	const TBufferStruct& UniformBufferValue
 	)
 {
-	// This will trigger if the parameter was not serialized
-	checkSlow(Parameter.IsInitialized());
-	if(Parameter.IsBound())
-	{
-		RHICmdList.SetShaderUniformBuffer(
-			Shader,
-			Parameter.GetBaseIndex(),
-			RHICreateUniformBuffer(&UniformBufferValue, &TBufferStruct::FTypeInfo::GetStructMetadata()->GetLayout(), UniformBuffer_SingleDraw)
-		);
-	}
+	FRHIBatchedShaderParameters& BatchedParameters = RHICmdList.GetScratchShaderParameters();
+	SetUniformBufferParameterImmediate(BatchedParameters, Parameter, UniformBufferValue);
+	RHICmdList.SetBatchedShaderParameters(Shader, BatchedParameters);
 }
\ No newline at end of file
diff --git a/Engine/Source/Runtime/RenderCore/Public/ShaderParameters.h b/Engine/Source/Runtime/RenderCore/Public/ShaderParameters.h
--- a/Engine/Source/Runtime/RenderCore/Public/ShaderParameters.h
+++ b/Engine/Source/Runtime/RenderCore/Public/ShaderParameters.h
@@ -113,7 +113,7 @@ private:
 };
 
 /** A class that binds either a UAV or SRV of a resource. */
-class FRWShaderParameter
+class UE_DEPRECATED(5.3, "Individual shader parameters should be used instead of FRWShaderParameter.") FRWShaderParameter
 {
 	DECLARE_EXPORTED_TYPE_LAYOUT(FRWShaderParameter, RENDERCORE_API, NonVirtual);
 public:
diff --git a/Engine/Source/Runtime/Renderer/Private/BasePassRendering.cpp b/Engine/Source/Runtime/Renderer/Private/BasePassRendering.cpp
--- a/Engine/Source/Runtime/Renderer/Private/BasePassRendering.cpp
+++ b/Engine/Source/Runtime/Renderer/Private/BasePassRendering.cpp
@@ -843,10 +843,14 @@ static void ClearGBufferAtMaxZ(
 			GraphicsPSOInit.PrimitiveType = PT_TriangleStrip;
 
 			SetGraphicsPipelineState(RHICmdList, GraphicsPSOInit, 0);
-			VertexShader->SetDepthParameter(RHICmdList, float(ERHIZBuffer::FarPlane));
+
+			SetShaderParametersLegacyVS(RHICmdList, VertexShader, float(ERHIZBuffer::FarPlane));
 
 			RHICmdList.SetViewport(View.ViewRect.Min.X, View.ViewRect.Min.Y, 0, View.ViewRect.Max.X, View.ViewRect.Max.Y, 1);
-			PixelShader->SetColors(RHICmdList, PixelShader, ClearColors, ActiveTargetCount);
+
+			TOneColorPixelShaderMRT::FParameters PixelParameters;
+			PixelShader->FillParameters(PixelParameters, ClearColors, ActiveTargetCount);
+			SetShaderParameters(RHICmdList, PixelShader, PixelShader.GetPixelShader(), PixelParameters);
 
 			RHICmdList.SetStreamSource(0, GClearVertexBuffer.VertexBufferRHI, 0);
 			RHICmdList.DrawPrimitive(0, 2, 1);
diff --git a/Engine/Source/Runtime/Renderer/Private/DepthRendering.cpp b/Engine/Source/Runtime/Renderer/Private/DepthRendering.cpp
--- a/Engine/Source/Runtime/Renderer/Private/DepthRendering.cpp
+++ b/Engine/Source/Runtime/Renderer/Private/DepthRendering.cpp
@@ -682,7 +682,7 @@ void FDeferredShadingSceneRenderer::RenderPrePassHMD(FRDGBuilder& GraphBuilder,
 				RHICmdList.SetViewport(View.ViewRect.Min.X, View.ViewRect.Min.Y, 0.0f, View.ViewRect.Max.X, View.ViewRect.Max.Y, 1.0f);
 
 				SetGraphicsPipelineState(RHICmdList, GraphicsPSOInit, 0);
-				VertexShader->SetDepthParameter(RHICmdList, 1.0f);
+				SetShaderParametersLegacyVS(RHICmdList, VertexShader, 1.0f);
 				HMDDevice->DrawHiddenAreaMesh(RHICmdList, View.StereoViewIndex);
 			});
 		}
diff --git a/Engine/Source/Runtime/Renderer/Private/RayTracing/RayTracingDynamicGeometry.cpp b/Engine/Source/Runtime/Renderer/Private/RayTracing/RayTracingDynamicGeometry.cpp
--- a/Engine/Source/Runtime/Renderer/Private/RayTracing/RayTracingDynamicGeometry.cpp
+++ b/Engine/Source/Runtime/Renderer/Private/RayTracing/RayTracingDynamicGeometry.cpp
@@ -42,7 +42,7 @@ public:
 	{
 		PassUniformBuffer.Bind(Initializer.ParameterMap, FSceneTextureUniformParameters::FTypeInfo::GetStructMetadata()->GetShaderVariableName());
 
-		RWVertexPositions.Bind(Initializer.ParameterMap, TEXT("VertexPositions"));
+		RWVertexPositions.Bind(Initializer.ParameterMap, TEXT("RWVertexPositions"));
 		UsingIndirectDraw.Bind(Initializer.ParameterMap, TEXT("UsingIndirectDraw"));
 		NumVertices.Bind(Initializer.ParameterMap, TEXT("NumVertices"));
 		MinVertexIndex.Bind(Initializer.ParameterMap, TEXT("MinVertexIndex"));
@@ -96,7 +96,7 @@ public:
 		FMeshMaterialShader::GetElementShaderBindings(PointerTable, Scene, ViewIfDynamicMeshCommand, VertexFactory, InputStreamType, FeatureLevel, PrimitiveSceneProxy, MeshBatch, BatchElement, ShaderElementData, ShaderBindings, VertexStreams);
 	}
 
-	LAYOUT_FIELD(FRWShaderParameter, RWVertexPositions);
+	LAYOUT_FIELD(FShaderResourceParameter, RWVertexPositions);
 	LAYOUT_FIELD(FShaderParameter, UsingIndirectDraw);
 	LAYOUT_FIELD(FShaderParameter, NumVertices);
 	LAYOUT_FIELD(FShaderParameter, MinVertexIndex);
@@ -514,14 +514,19 @@ void FRayTracingDynamicGeometryCollection::DispatchUpdates(FRHICommandListImmedi
 						ShaderBindingState = FShaderBindingState();
 					}
 
+					FRHIBatchedShaderParameters& BatchedParameters = RHICmdList.GetScratchShaderParameters();
+
 					FRWBuffer* TargetBuffer = Cmd.TargetBuffer;
 					if (CurrentBuffer != TargetBuffer)
 					{
 						CurrentBuffer = TargetBuffer;
-						Shader->RWVertexPositions.SetBuffer(RHICmdList, CurrentShader, *Cmd.TargetBuffer);
+
+						SetUAVParameter(BatchedParameters, Shader->RWVertexPositions, Cmd.TargetBuffer->UAV);
 					}
 
-					Cmd.ShaderBindings.SetOnCommandList(RHICmdList, ComputeShader, &ShaderBindingState);
+					Cmd.ShaderBindings.SetParameters(BatchedParameters, ComputeShader, &ShaderBindingState);
+					RHICmdList.SetBatchedShaderParameters(CurrentShader, BatchedParameters);
+
 					RHICmdList.DispatchComputeShader(FMath::DivideAndRoundUp<uint32>(Cmd.NumMaxVertices, 64), 1, 1);
 				}
 
