Do not produce empty output stable shader key.

#rnx

diff --git a/Engine/Source/Runtime/RenderCore/Private/ShaderCodeLibrary.cpp b/Engine/Source/Runtime/RenderCore/Private/ShaderCodeLibrary.cpp
--- a/Engine/Source/Runtime/RenderCore/Private/ShaderCodeLibrary.cpp
+++ b/Engine/Source/Runtime/RenderCore/Private/ShaderCodeLibrary.cpp
@@ -2164,12 +2164,19 @@ struct FEditorShaderStableInfo
 		check(LibraryName.Len() > 0);
 		OutSCLCSVPath = FString();
 
+		if (StableMap.IsEmpty())
+		{
+			// Do not touch the disk if empty, but also don't consider this a failure.
+			// It is entirely possible that during the cook no assets with shaders were cooked.
+			return true;
+		}
+
 		bool bSuccess = IFileManager::Get().MakeDirectory(*OutputDir, true);
 
 		EShaderPlatform Platform = ShaderFormatToLegacyShaderPlatform(FormatName);
 
 		// Shader library
-		if (bSuccess && StableMap.Num() > 0)
+		if (bSuccess)
 		{
 			// Write to a intermediate file
 			FString IntermediateFormatPath = GetStableInfoArchiveFilename(FPaths::ProjectSavedDir() / TEXT("Shaders") / FormatName.ToString(), LibraryName, FormatName);
@@ -3089,7 +3096,12 @@ public:
 				// Stable shader info is not saved per-chunk (it is not needed at runtime), so save it always
 				FString SCLCSVPath;
 				bOk &= StableArchive->SaveToDisk(MetaOutputDir, SCLCSVPath);
-				OutSCLCSVPath.Add(SCLCSVPath);
+				
+				// Only add output files if they were actually written to disk (if there were no shaders in the library it is not a failure).
+				if (!SCLCSVPath.IsEmpty())
+				{
+					OutSCLCSVPath.Add(SCLCSVPath);
+				}
 			}
 		}
 
