Unconditionally disable velocity output in base pass for Forward Desktop (UE-169988).

#rnx

diff --git a/Engine/Source/Runtime/RHI/Private/RHI.cpp b/Engine/Source/Runtime/RHI/Private/RHI.cpp
--- a/Engine/Source/Runtime/RHI/Private/RHI.cpp
+++ b/Engine/Source/Runtime/RHI/Private/RHI.cpp
@@ -1555,7 +1555,7 @@ void FRHIRenderPassInfo::Validate() const
 			}
 			else
 			{
-				ensure(Entry.RenderTarget->GetNumSamples() == NumSamples);
+				ensureMsgf(Entry.RenderTarget->GetNumSamples() == NumSamples, TEXT("RenderTarget have inconsistent NumSamples: first %d, then %d"), NumSamples, Entry.RenderTarget->GetNumSamples());
 			}
 
 			ERenderTargetStoreAction Store = GetStoreAction(Entry.Action);
diff --git a/Engine/Source/Runtime/RenderCore/Private/RenderUtils.cpp b/Engine/Source/Runtime/RenderCore/Private/RenderUtils.cpp
--- a/Engine/Source/Runtime/RenderCore/Private/RenderUtils.cpp
+++ b/Engine/Source/Runtime/RenderCore/Private/RenderUtils.cpp
@@ -1363,10 +1363,9 @@ bool IsUsingGBuffers(const FStaticShaderPlatform Platform)
 bool IsUsingBasePassVelocity(const FStaticShaderPlatform Platform)
 {
 	static FShaderPlatformCachedIniValue<int32> PerPlatformCVar(TEXT("r.VelocityOutputPass"));
-	static const IConsoleVariable* AntiAliasingCVar = IConsoleManager::Get().FindConsoleVariable(TEXT("r.AntiAliasingMethod"));
-	// Writing velocity in base pass is disabled for desktop forward when using MSAA
-	const int32 AAM_MSAA = 3;	// see EAntiAliasingMethod in SceneUtils.h
-	if (IsMobilePlatform(Platform) || (AntiAliasingCVar && IsForwardShadingEnabled(Platform) && AntiAliasingCVar->GetInt() == AAM_MSAA))
+	// Writing velocity in base pass is disabled for desktop forward because it may use MSAA (runtime-settable setting) and Velocity isn't a multisample texture.
+	// TSR or vr.AllowMotionBlurInVR=1 case aren't recommended for Desktop Forward renderer, and if enabled, cause a separate Velocity pass.
+	if (IsMobilePlatform(Platform) || IsForwardShadingEnabled(Platform))
 	{
 		return false;
 	}
