Added CVar to RDG to control whether to enable expensive RDG validation in development builds (on by default).

diff --git a/Engine/Source/Runtime/RenderCore/Private/RenderGraphBuilder.cpp b/Engine/Source/Runtime/RenderCore/Private/RenderGraphBuilder.cpp
--- a/Engine/Source/Runtime/RenderCore/Private/RenderGraphBuilder.cpp
+++ b/Engine/Source/Runtime/RenderCore/Private/RenderGraphBuilder.cpp
@@ -2347,9 +2347,12 @@ void FRDGBuilder::SetupPassInternals(FRDGPass* Pass)
 #endif
 
 #if RDG_GPU_DEBUG_SCOPES && RDG_ENABLE_DEBUG
-	if (const FRDGEventScope* Scope = Pass->GPUScopes.Event)
+	if (GRDGValidation != 0)
 	{
-		Pass->FullPathIfDebug = Scope->GetPath(Pass->Name);
+		if (const FRDGEventScope* Scope = Pass->GPUScopes.Event)
+		{
+			Pass->FullPathIfDebug = Scope->GetPath(Pass->Name);
+		}
 	}
 #endif
 }
@@ -3912,7 +3915,7 @@ void FRDGBuilder::VisualizePassOutputs(const FRDGPass* Pass)
 
 void FRDGBuilder::ClobberPassOutputs(const FRDGPass* Pass)
 {
-	if (!GRDGClobberResources || !AuxiliaryPasses.IsClobberAllowed())
+	if (!GRDGValidation || !GRDGClobberResources || !AuxiliaryPasses.IsClobberAllowed())
 	{
 		return;
 	}
diff --git a/Engine/Source/Runtime/RenderCore/Private/RenderGraphEvent.cpp b/Engine/Source/Runtime/RenderCore/Private/RenderGraphEvent.cpp
--- a/Engine/Source/Runtime/RenderCore/Private/RenderGraphEvent.cpp
+++ b/Engine/Source/Runtime/RenderCore/Private/RenderGraphEvent.cpp
@@ -525,6 +525,7 @@ FRDGEventName::FRDGEventName(const TCHAR* InEventFormat, ...)
 {
 	check(InEventFormat);
 
+	if (GRDGValidation != 0)
 	{
 		va_list VAList;
 		va_start(VAList, InEventFormat);
diff --git a/Engine/Source/Runtime/RenderCore/Private/RenderGraphPrivate.cpp b/Engine/Source/Runtime/RenderCore/Private/RenderGraphPrivate.cpp
--- a/Engine/Source/Runtime/RenderCore/Private/RenderGraphPrivate.cpp
+++ b/Engine/Source/Runtime/RenderCore/Private/RenderGraphPrivate.cpp
@@ -16,6 +16,15 @@ FAutoConsoleVariableRef CVarImmediateMode(
 	TEXT("Executes passes as they get created. Useful to have a callstack of the wiring code when crashing in the pass' lambda."),
 	ECVF_RenderThreadSafe);
 
+int32 GRDGValidation = 1;
+FAutoConsoleVariableRef CVarRDGValidation(
+	TEXT("r.RDG.Validation"),
+	GRDGValidation,
+	TEXT("Enables validation of correctness in API calls and pass parameter dependencies.\n")
+	TEXT(" 0: disabled;\n")
+	TEXT(" 1: enabled (default);\n"),
+	ECVF_RenderThreadSafe);
+
 int32 GRDGDebug = 0;
 FAutoConsoleVariableRef CVarRDGDebug(
 	TEXT("r.RDG.Debug"),
@@ -513,6 +522,12 @@ void InitRenderGraph()
 		GRDGImmediateMode = 1;
 	}
 
+	int32 ValidationValue = 0;
+	if (FParse::Value(FCommandLine::Get(), TEXT("rdgvalidation="), ValidationValue))
+	{
+		GRDGValidation = ValidationValue;
+	}
+
 	if (FParse::Param(FCommandLine::Get(), TEXT("rdgdebug")))
 	{
 		GRDGDebug = 1;
diff --git a/Engine/Source/Runtime/RenderCore/Private/RenderGraphPrivate.h b/Engine/Source/Runtime/RenderCore/Private/RenderGraphPrivate.h
--- a/Engine/Source/Runtime/RenderCore/Private/RenderGraphPrivate.h
+++ b/Engine/Source/Runtime/RenderCore/Private/RenderGraphPrivate.h
@@ -31,6 +31,7 @@ DEFINE_LOG_CATEGORY_STATIC(LogRDG, Log, All);
 #if RDG_ENABLE_DEBUG
 extern int32 GRDGAsyncCompute;
 extern int32 GRDGClobberResources;
+extern int32 GRDGValidation;
 extern int32 GRDGDebug;
 extern int32 GRDGDebugFlushGPU;
 extern int32 GRDGDebugExtendResourceLifetimes;
@@ -93,6 +94,7 @@ void EmitRDGWarning(const FString& WarningMessage);
 #else // !RDG_ENABLE_DEBUG
 
 const int32 GRDGClobberResources = 0;
+const int32 GRDGValidation = 0;
 const int32 GRDGDebug = 0;
 const int32 GRDGDebugFlushGPU = 0;
 const int32 GRDGDebugExtendResourceLifetimes = 0;
diff --git a/Engine/Source/Runtime/RenderCore/Private/RenderGraphValidation.cpp b/Engine/Source/Runtime/RenderCore/Private/RenderGraphValidation.cpp
--- a/Engine/Source/Runtime/RenderCore/Private/RenderGraphValidation.cpp
+++ b/Engine/Source/Runtime/RenderCore/Private/RenderGraphValidation.cpp
@@ -71,6 +71,11 @@ struct FRDGResourceDebugData
 
 void FRDGResource::MarkResourceAsUsed()
 {
+	if (!GRDGValidation)
+	{
+		return;
+	}
+
 	ValidateRHIAccess();
 
 	GetDebugData().bIsActuallyUsedByPass = true;
@@ -78,6 +83,11 @@ void FRDGResource::MarkResourceAsUsed()
 
 void FRDGResource::ValidateRHIAccess() const
 {
+	if (!GRDGValidation)
+	{
+		return;
+	}
+
 	check(DebugData);
 	checkf(DebugData->bAllowRHIAccess || GRDGAllowRHIAccess,
 		TEXT("Accessing the RHI resource of %s at this time is not allowed. If you hit this check in pass, ")
@@ -87,6 +97,7 @@ void FRDGResource::ValidateRHIAccess() const
 
 FRDGResourceDebugData& FRDGResource::GetDebugData() const
 {
+	check(GRDGValidation != 0);
 	check(DebugData);
 	return *DebugData;
 }
@@ -105,6 +116,7 @@ struct FRDGViewableResourceDebugData
 
 FRDGViewableResourceDebugData& FRDGViewableResource::GetViewableDebugData() const
 {
+	check(GRDGValidation != 0);
 	check(ViewableDebugData);
 	return *ViewableDebugData;
 }
@@ -120,6 +132,7 @@ struct FRDGTextureDebugData
 
 FRDGTextureDebugData& FRDGTexture::GetTextureDebugData() const
 {
+	check(GRDGValidation != 0);
 	check(TextureDebugData);
 	return *TextureDebugData;
 }
@@ -132,12 +145,18 @@ struct FRDGBufferDebugData
 
 FRDGBufferDebugData& FRDGBuffer::GetBufferDebugData() const
 {
+	check(GRDGValidation != 0);
 	check(BufferDebugData);
 	return *BufferDebugData;
 }
 
 void FRDGUniformBuffer::MarkResourceAsUsed()
 {
+	if (!GRDGValidation)
+	{
+		return;
+	}
+
 	FRDGResource::MarkResourceAsUsed();
 
 	// Individual resources can't be culled from a uniform buffer, so we have to mark them all as used.
@@ -186,6 +205,11 @@ void FRDGUserValidation::ValidateCreateViewableResource(FRDGViewableResource* Re
 
 void FRDGUserValidation::ValidateCreateTexture(FRDGTextureRef Texture)
 {
+	if (!GRDGValidation)
+	{
+		return;
+	}
+
 	ValidateCreateViewableResource(Texture);
 	Texture->TextureDebugData = Allocator.Alloc<FRDGTextureDebugData>();
 	if (GRDGDebug)
@@ -196,6 +220,11 @@ void FRDGUserValidation::ValidateCreateTexture(FRDGTextureRef Texture)
 
 void FRDGUserValidation::ValidateCreateBuffer(FRDGBufferRef Buffer)
 {
+	if (!GRDGValidation)
+	{
+		return;
+	}
+
 	ValidateCreateViewableResource(Buffer);
 	Buffer->BufferDebugData = Allocator.Alloc<FRDGBufferDebugData>();
 	if (GRDGDebug)
@@ -206,26 +235,51 @@ void FRDGUserValidation::ValidateCreateBuffer(FRDGBufferRef Buffer)
 
 void FRDGUserValidation::ValidateCreateSRV(FRDGTextureSRVRef SRV)
 {
+	if (!GRDGValidation)
+	{
+		return;
+	}
+
 	ValidateCreateResource(SRV);
 }
 
 void FRDGUserValidation::ValidateCreateSRV(FRDGBufferSRVRef SRV)
 {
+	if (!GRDGValidation)
+	{
+		return;
+	}
+
 	ValidateCreateResource(SRV);
 }
 
 void FRDGUserValidation::ValidateCreateUAV(FRDGTextureUAVRef UAV)
 {
+	if (!GRDGValidation)
+	{
+		return;
+	}
+
 	ValidateCreateResource(UAV);
 }
 
 void FRDGUserValidation::ValidateCreateUAV(FRDGBufferUAVRef UAV)
 {
+	if (!GRDGValidation)
+	{
+		return;
+	}
+
 	ValidateCreateResource(UAV);
 }
 
 void FRDGUserValidation::ValidateCreateUniformBuffer(FRDGUniformBufferRef UniformBuffer)
 {
+	if (!GRDGValidation)
+	{
+		return;
+	}
+
 	ValidateCreateResource(UniformBuffer);
 }
 
@@ -234,6 +288,11 @@ void FRDGUserValidation::ValidateRegisterExternalTexture(
 	const TCHAR* Name,
 	ERDGTextureFlags Flags)
 {
+	if (!GRDGValidation)
+	{
+		return;
+	}
+
 	checkf(Name, TEXT("Attempted to register external texture with NULL name."));
 	checkf(ExternalPooledTexture.IsValid(), TEXT("Attempted to register NULL external texture."));
 	ExecuteGuard(TEXT("RegisterExternalTexture"), Name);
@@ -241,6 +300,11 @@ void FRDGUserValidation::ValidateRegisterExternalTexture(
 
 void FRDGUserValidation::ValidateRegisterExternalBuffer(const TRefCountPtr<FRDGPooledBuffer>& ExternalPooledBuffer, const TCHAR* Name, ERDGBufferFlags Flags)
 {
+	if (!GRDGValidation)
+	{
+		return;
+	}
+
 	checkf(Name, TEXT("Attempted to register external buffer with NULL name."));
 	checkf(ExternalPooledBuffer.IsValid(), TEXT("Attempted to register NULL external buffer."));
 	ExecuteGuard(TEXT("RegisterExternalBuffer"), Name);
@@ -248,16 +312,31 @@ void FRDGUserValidation::ValidateRegisterExternalBuffer(const TRefCountPtr<FRDGP
 
 void FRDGUserValidation::ValidateRegisterExternalTexture(FRDGTextureRef Texture)
 {
+	if (!GRDGValidation)
+	{
+		return;
+	}
+
 	ValidateCreateTexture(Texture);
 }
 
 void FRDGUserValidation::ValidateRegisterExternalBuffer(FRDGBufferRef Buffer)
 {
+	if (!GRDGValidation)
+	{
+		return;
+	}
+
 	ValidateCreateBuffer(Buffer);
 }
 
 void FRDGUserValidation::ValidateCreateTexture(const FRDGTextureDesc& Desc, const TCHAR* Name, ERDGTextureFlags Flags)
 {
+	if (!GRDGValidation)
+	{
+		return;
+	}
+
 	checkf(Name, TEXT("Creating a texture requires a valid debug name."));
 	ExecuteGuard(TEXT("CreateTexture"), Name);
 
@@ -279,6 +358,11 @@ void FRDGUserValidation::ValidateCreateTexture(const FRDGTextureDesc& Desc, cons
 
 void FRDGUserValidation::ValidateCreateBuffer(const FRDGBufferDesc& Desc, const TCHAR* Name, ERDGBufferFlags Flags)
 {
+	if (!GRDGValidation)
+	{
+		return;
+	}
+
 	checkf(Name, TEXT("Creating a buffer requires a valid debug name."));
 	ExecuteGuard(TEXT("CreateBuffer"), Name);
 
@@ -294,6 +378,11 @@ void FRDGUserValidation::ValidateCreateBuffer(const FRDGBufferDesc& Desc, const
 
 void FRDGUserValidation::ValidateCreateSRV(const FRDGTextureSRVDesc& Desc)
 {
+	if (!GRDGValidation)
+	{
+		return;
+	}
+
 	FRDGTextureRef Texture = Desc.Texture;
 	checkf(Texture, TEXT("Texture SRV created with a null texture."));
 	ExecuteGuard(TEXT("CreateSRV"), Texture->Name);
@@ -302,6 +391,11 @@ void FRDGUserValidation::ValidateCreateSRV(const FRDGTextureSRVDesc& Desc)
 
 void FRDGUserValidation::ValidateCreateSRV(const FRDGBufferSRVDesc& Desc)
 {
+	if (!GRDGValidation)
+	{
+		return;
+	}
+
 	FRDGBufferRef Buffer = Desc.Buffer;
 	checkf(Buffer, TEXT("Buffer SRV created with a null buffer."));
 	ExecuteGuard(TEXT("CreateSRV"), Buffer->Name);
@@ -309,6 +403,11 @@ void FRDGUserValidation::ValidateCreateSRV(const FRDGBufferSRVDesc& Desc)
 
 void FRDGUserValidation::ValidateCreateUAV(const FRDGTextureUAVDesc& Desc)
 {
+	if (!GRDGValidation)
+	{
+		return;
+	}
+
 	FRDGTextureRef Texture = Desc.Texture;
 
 	checkf(Texture, TEXT("Texture UAV created with a null texture."));
@@ -320,6 +419,11 @@ void FRDGUserValidation::ValidateCreateUAV(const FRDGTextureUAVDesc& Desc)
 
 void FRDGUserValidation::ValidateCreateUAV(const FRDGBufferUAVDesc& Desc)
 {
+	if (!GRDGValidation)
+	{
+		return;
+	}
+
 	FRDGBufferRef Buffer = Desc.Buffer;
 	checkf(Buffer, TEXT("Buffer UAV created with a null buffer."));
 	ExecuteGuard(TEXT("CreateUAV"), Buffer->Name);
@@ -327,6 +431,11 @@ void FRDGUserValidation::ValidateCreateUAV(const FRDGBufferUAVDesc& Desc)
 
 void FRDGUserValidation::ValidateCreateUniformBuffer(const void* ParameterStruct, const FShaderParametersMetadata* Metadata)
 {
+	if (!GRDGValidation)
+	{
+		return;
+	}
+
 	check(Metadata);
 	const TCHAR* Name = Metadata->GetShaderVariableName();
 	checkf(ParameterStruct, TEXT("Uniform buffer '%s' created with null parameters."), Name);
@@ -335,6 +444,11 @@ void FRDGUserValidation::ValidateCreateUniformBuffer(const void* ParameterStruct
 
 void FRDGUserValidation::ValidateUploadBuffer(FRDGBufferRef Buffer, const void* InitialData, uint64 InitialDataSize)
 {
+	if (!GRDGValidation)
+	{
+		return;
+	}
+
 	check(Buffer);
 	checkf(!Buffer->bQueuedForUpload, TEXT("Buffer %s already has an upload queued. Only one upload can be done for each graph."), Buffer->Name);
 	check(InitialData || InitialDataSize == 0);
@@ -342,6 +456,11 @@ void FRDGUserValidation::ValidateUploadBuffer(FRDGBufferRef Buffer, const void*
 
 void FRDGUserValidation::ValidateUploadBuffer(FRDGBufferRef Buffer, const void* InitialData, uint64 InitialDataSize, const FRDGBufferInitialDataFreeCallback& InitialDataFreeCallback)
 {
+	if (!GRDGValidation)
+	{
+		return;
+	}
+
 	check(Buffer);
 	checkf(!Buffer->bQueuedForUpload, TEXT("Buffer %s already has an upload queued. Only one upload can be done for each graph."), Buffer->Name);
 	check((InitialData || InitialDataSize == 0) && InitialDataFreeCallback);
@@ -349,6 +468,11 @@ void FRDGUserValidation::ValidateUploadBuffer(FRDGBufferRef Buffer, const void*
 
 void FRDGUserValidation::ValidateUploadBuffer(FRDGBufferRef Buffer, const FRDGBufferInitialDataCallback& InitialDataCallback, const FRDGBufferInitialDataSizeCallback& InitialDataSizeCallback)
 {
+	if (!GRDGValidation)
+	{
+		return;
+	}
+
 	check(Buffer);
 	checkf(!Buffer->bQueuedForUpload, TEXT("Buffer %s already has an upload queued. Only one upload can be done for each graph."), Buffer->Name);
 	check(InitialDataCallback && InitialDataSizeCallback);
@@ -356,6 +480,11 @@ void FRDGUserValidation::ValidateUploadBuffer(FRDGBufferRef Buffer, const FRDGBu
 
 void FRDGUserValidation::ValidateUploadBuffer(FRDGBufferRef Buffer, const FRDGBufferInitialDataCallback& InitialDataCallback, const FRDGBufferInitialDataSizeCallback& InitialDataSizeCallback, const FRDGBufferInitialDataFreeCallback& InitialDataFreeCallback)
 {
+	if (!GRDGValidation)
+	{
+		return;
+	}
+
 	check(Buffer);
 	checkf(!Buffer->bQueuedForUpload, TEXT("Buffer %s already has an upload queued. Only one upload can be done for each graph."), Buffer->Name);
 	check(InitialDataCallback && InitialDataSizeCallback && InitialDataFreeCallback);
@@ -363,12 +492,22 @@ void FRDGUserValidation::ValidateUploadBuffer(FRDGBufferRef Buffer, const FRDGBu
 
 void FRDGUserValidation::ValidateExtractTexture(FRDGTextureRef Texture, TRefCountPtr<IPooledRenderTarget>* OutTexturePtr)
 {
+	if (!GRDGValidation)
+	{
+		return;
+	}
+
 	ValidateExtractResource(Texture);
 	checkf(OutTexturePtr, TEXT("Texture %s was extracted, but the output texture pointer is null."), Texture->Name);
 }
 
 void FRDGUserValidation::ValidateExtractBuffer(FRDGBufferRef Buffer, TRefCountPtr<FRDGPooledBuffer>* OutBufferPtr)
 {
+	if (!GRDGValidation)
+	{
+		return;
+	}
+
 	ValidateExtractResource(Buffer);
 	checkf(OutBufferPtr, TEXT("Texture %s was extracted, but the output texture pointer is null."), Buffer->Name);
 }
@@ -390,6 +529,11 @@ void FRDGUserValidation::ValidateExtractResource(FRDGViewableResource* Resource)
 
 void FRDGUserValidation::ValidateConvertToExternalResource(FRDGViewableResource* Resource)
 {
+	if (!GRDGValidation)
+	{
+		return;
+	}
+
 	check(Resource);
 	checkf(!bHasExecuteBegun || !Resource->bTransient,
 		TEXT("Unable to convert resource %s to external because passes in the graph have already executed."),
@@ -398,6 +542,11 @@ void FRDGUserValidation::ValidateConvertToExternalResource(FRDGViewableResource*
 
 void FRDGUserValidation::ValidateConvertToExternalUniformBuffer(FRDGUniformBuffer* UniformBuffer)
 {
+	if (!GRDGValidation)
+	{
+		return;
+	}
+
 	check(UniformBuffer);
 	checkf(!bHasExecuteBegun,
 		TEXT("Unable to convert uniform buffer %s to external because passes in the graph have already executed."),
@@ -406,6 +555,11 @@ void FRDGUserValidation::ValidateConvertToExternalUniformBuffer(FRDGUniformBuffe
 
 void FRDGUserValidation::RemoveUnusedWarning(FRDGViewableResource* Resource)
 {
+	if (!GRDGValidation)
+	{
+		return;
+	}
+
 	check(Resource);
 	ExecuteGuard(TEXT("RemoveUnusedResourceWarning"), Resource->Name);
 
@@ -418,6 +572,11 @@ void FRDGUserValidation::RemoveUnusedWarning(FRDGViewableResource* Resource)
 
 bool FRDGUserValidation::TryMarkForClobber(FRDGViewableResource* Resource) const
 {
+	if (!GRDGValidation)
+	{
+		return false;
+	}
+
 	check(Resource);
 	FRDGViewableResourceDebugData& DebugData = Resource->GetViewableDebugData();
 
@@ -433,18 +592,33 @@ bool FRDGUserValidation::TryMarkForClobber(FRDGViewableResource* Resource) const
 
 void FRDGUserValidation::ValidateGetPooledTexture(FRDGTextureRef Texture) const
 {
+	if (!GRDGValidation)
+	{
+		return;
+	}
+
 	check(Texture);
 	checkf(Texture->bExternal, TEXT("GetPooledTexture called on texture %s, but it is not external. Call PreallocateTexture or register as an external texture instead."), Texture->Name);
 }
 
 void FRDGUserValidation::ValidateGetPooledBuffer(FRDGBufferRef Buffer) const
 {
+	if (!GRDGValidation)
+	{
+		return;
+	}
+
 	check(Buffer);
 	checkf(Buffer->bExternal, TEXT("GetPooledBuffer called on buffer %s, but it is not external. Call PreallocateBuffer or register as an external buffer instead."), Buffer->Name);
 }
 
 void FRDGUserValidation::ValidateSetAccessFinal(FRDGViewableResource* Resource, ERHIAccess AccessFinal)
 {
+	if (!GRDGValidation)
+	{
+		return;
+	}
+
 	check(Resource);
 	check(AccessFinal != ERHIAccess::Unknown && IsValidAccess(AccessFinal));
 	checkf(Resource->bExternal || Resource->bExtracted, TEXT("Cannot set final access on non-external resource '%s' unless it is first extracted or preallocated."), Resource->Name);
@@ -453,6 +627,11 @@ void FRDGUserValidation::ValidateSetAccessFinal(FRDGViewableResource* Resource,
 
 void FRDGUserValidation::ValidateUseExternalAccessMode(FRDGViewableResource* Resource, ERHIAccess ReadOnlyAccess, ERHIPipeline Pipelines)
 {
+	if (!GRDGValidation)
+	{
+		return;
+	}
+
 	check(Resource);
 	check(Pipelines != ERHIPipeline::None);
 
@@ -479,6 +658,11 @@ void FRDGUserValidation::ValidateUseExternalAccessMode(FRDGViewableResource* Res
 
 void FRDGUserValidation::ValidateUseInternalAccessMode(FRDGViewableResource* Resource)
 {
+	if (!GRDGValidation)
+	{
+		return;
+	}
+
 	check(Resource);
 
 	const auto& AccessModeState = Resource->AccessModeState;
@@ -487,6 +671,11 @@ void FRDGUserValidation::ValidateUseInternalAccessMode(FRDGViewableResource* Res
 
 void FRDGUserValidation::ValidateExternalAccess(FRDGViewableResource* Resource, ERHIAccess Access, const FRDGPass* Pass)
 {
+	if (!GRDGValidation)
+	{
+		return;
+	}
+
 	const auto& AccessModeState = Resource->AccessModeState;
 
 	ensureMsgf(EnumHasAnyFlags(AccessModeState.Access, Access),
@@ -500,6 +689,11 @@ void FRDGUserValidation::ValidateExternalAccess(FRDGViewableResource* Resource,
 
 void FRDGUserValidation::ValidateAddSubresourceAccess(FRDGViewableResource* Resource, const FRDGSubresourceState& Subresource, ERHIAccess Access)
 {
+	if (!GRDGValidation)
+	{
+		return;
+	}
+
 	const bool bOldAccessMergeable = EnumHasAnyFlags(Subresource.Access, GRHIMergeableAccessMask) || Subresource.Access == ERHIAccess::Unknown;
 	const bool bNewAccessMergeable = EnumHasAnyFlags(Access, GRHIMergeableAccessMask);
 
@@ -510,6 +704,11 @@ void FRDGUserValidation::ValidateAddSubresourceAccess(FRDGViewableResource* Reso
 
 void FRDGUserValidation::ValidateAddPass(const FRDGEventName& Name, ERDGPassFlags Flags)
 {
+	if (!GRDGValidation)
+	{
+		return;
+	}
+
 	ExecuteGuard(TEXT("AddPass"), Name.GetTCHAR());
 
 	checkf(!EnumHasAnyFlags(Flags, ERDGPassFlags::Copy | ERDGPassFlags::Compute | ERDGPassFlags::AsyncCompute | ERDGPassFlags::Raster),
@@ -518,6 +717,11 @@ void FRDGUserValidation::ValidateAddPass(const FRDGEventName& Name, ERDGPassFlag
 
 void FRDGUserValidation::ValidateAddPass(const void* ParameterStruct, const FShaderParametersMetadata* Metadata, const FRDGEventName& Name, ERDGPassFlags Flags)
 {
+	if (!GRDGValidation)
+	{
+		return;
+	}
+
 	checkf(ParameterStruct, TEXT("Pass '%s' created with null parameters."), Name.GetTCHAR());
 	ExecuteGuard(TEXT("AddPass"), Name.GetTCHAR());
 
@@ -539,6 +743,11 @@ void FRDGUserValidation::ValidateAddPass(const void* ParameterStruct, const FSha
 
 void FRDGUserValidation::ValidateAddPass(const FRDGPass* Pass, bool bSkipPassAccessMarking)
 {
+	if (!GRDGValidation)
+	{
+		return;
+	}
+
 	const FRenderTargetBindingSlots* RenderTargetBindingSlots = nullptr;
 
 	// Pass flags are validated as early as possible by the builder in AddPass.
@@ -854,7 +1063,7 @@ void FRDGUserValidation::ValidateExecuteEnd()
 	bHasExecuted = true;
 	GRDGBuilderActive = false;
 
-	if (GRDGDebug)
+	if (GRDGDebug && GRDGValidation)
 	{
 		auto ValidateResourceAtExecuteEnd = [](const FRDGViewableResource* Resource)
 		{
@@ -907,7 +1116,7 @@ void FRDGUserValidation::ValidateExecuteEnd()
 
 void FRDGUserValidation::ValidateExecutePassBegin(const FRDGPass* Pass)
 {
-	if (bParallelExecuteEnabled)
+	if (bParallelExecuteEnabled || !GRDGValidation)
 	{
 		return;
 	}
@@ -1007,7 +1216,7 @@ void FRDGUserValidation::ValidateExecutePassBegin(const FRDGPass* Pass)
 
 void FRDGUserValidation::ValidateExecutePassEnd(const FRDGPass* Pass)
 {
-	if (bParallelExecuteEnabled)
+	if (bParallelExecuteEnabled || !GRDGValidation)
 	{
 		return;
 	}
@@ -1071,6 +1280,11 @@ void FRDGUserValidation::ValidateExecutePassEnd(const FRDGPass* Pass)
 
 void FRDGUserValidation::SetAllowRHIAccess(const FRDGPass* Pass, bool bAllowAccess)
 {
+	if (!GRDGValidation)
+	{
+		return;
+	}
+
 	Pass->GetParameters().Enumerate([&](FRDGParameter Parameter)
 	{
 		if (Parameter.IsResource())
