sorted shader assinfo to make out deterministic
makes diffing cooked data easier

diff --git a/Engine/Source/Runtime/RenderCore/Private/ShaderCodeArchive.cpp b/Engine/Source/Runtime/RenderCore/Private/ShaderCodeArchive.cpp
--- a/Engine/Source/Runtime/RenderCore/Private/ShaderCodeArchive.cpp
+++ b/Engine/Source/Runtime/RenderCore/Private/ShaderCodeArchive.cpp
@@ -506,13 +506,26 @@ void FSerializedShaderArchive::SaveAssetInfo(FArchive& Ar)
 
 			Writer->WriteValue(TEXT("AssetInfoVersion"), static_cast<int32>(EAssetInfoVersion::CurrentVersion));
 
+			TArray<const TPair<FSHAHash, FShaderMapAssetPaths>*> SortedData;
+			SortedData.Reserve(ShaderCodeToAssets.Num());
+			for (TPair<FSHAHash, FShaderMapAssetPaths>& Pair : ShaderCodeToAssets)
+			{
+				SortedData.Add(&Pair);
+				Pair.Value.Sort(FNameLexicalLess());
+			}
+			Algo::Sort(SortedData, [](const TPair<FSHAHash, FShaderMapAssetPaths>* const a, const TPair<FSHAHash, FShaderMapAssetPaths>* const b)
+				{
+					return GetTypeHash(a->Key) < GetTypeHash(b->Key);
+				});
+
 			Writer->WriteArrayStart(TEXT("ShaderCodeToAssets"));
-			for (TMap<FSHAHash, FShaderMapAssetPaths>::TConstIterator Iter(ShaderCodeToAssets); Iter; ++Iter)
+			for (const TPair<FSHAHash, FShaderMapAssetPaths>* Pair : SortedData)
 			{
 				Writer->WriteObjectStart();
-				const FSHAHash& Hash = Iter.Key();
+				const FSHAHash& Hash = Pair->Key;
 				Writer->WriteValue(TEXT("ShaderMapHash"), Hash.ToString());
-				const FShaderMapAssetPaths& Assets = Iter.Value();
+				const FShaderMapAssetPaths& Assets = Pair->Value;
+
 				Writer->WriteArrayStart(TEXT("Assets"));
 				for (FShaderMapAssetPaths::TConstIterator AssetIter(Assets); AssetIter; ++AssetIter)
 				{
