made `FSharedShaderCodeRequest` thread-safe because it's broadcasted concurrently

diff --git a/Engine/Source/Runtime/RenderCore/Public/ShaderCodeLibrary.h b/Engine/Source/Runtime/RenderCore/Public/ShaderCodeLibrary.h
--- a/Engine/Source/Runtime/RenderCore/Public/ShaderCodeLibrary.h
+++ b/Engine/Source/Runtime/RenderCore/Public/ShaderCodeLibrary.h
@@ -178,8 +178,8 @@ RENDERCORE_API bool LoadFromCompactBinary(FCbFieldView Field, FStableShaderKeyAn
 DECLARE_DELEGATE_OneParam(FSharedShaderMapResourceExplicitRelease, const FShaderMapResource*);
 RENDERCORE_API extern FSharedShaderMapResourceExplicitRelease OnSharedShaderMapResourceExplicitRelease;
 
-DECLARE_MULTICAST_DELEGATE_TwoParams(FSharedShaderCodeRequest, const FSHAHash&, FArchive*);
-DECLARE_MULTICAST_DELEGATE_OneParam(FSharedShaderCodeRelease, const FSHAHash&);
+using FSharedShaderCodeRequest = TTSMulticastDelegate<void(const FSHAHash&, FArchive*)>; // thread-safe because it's "broadcasted" concurrently
+using FSharedShaderCodeRelease = TMulticastDelegate<void(const FSHAHash&)>;
 
 // Collection of unique shader code
 // Populated at cook time
