Remove obsolete temp code from 2016 to SuspendAsyncLoading when using FSuspendRenderingThread (a feature that in itself is only used on iOS).

#rnx

diff --git a/Engine/Source/Runtime/RenderCore/Private/RenderingThread.cpp b/Engine/Source/Runtime/RenderCore/Private/RenderingThread.cpp
--- a/Engine/Source/Runtime/RenderCore/Private/RenderingThread.cpp
+++ b/Engine/Source/Runtime/RenderCore/Private/RenderingThread.cpp
@@ -114,13 +114,6 @@ static void WaitAndResumeRendering()
  */
 FSuspendRenderingThread::FSuspendRenderingThread( bool bInRecreateThread )
 {
-	// Suspend async loading thread so that it doesn't start queueing render commands 
-	// while the render thread is suspended.
-	if (IsAsyncLoadingMultithreaded())
-	{
-		SuspendAsyncLoading();
-	}
-
 	// Pause asset streaming to prevent rendercommands from being enqueued.
 	SuspendTextureStreamingRenderTasks();
 
