Substrate - fixed InternalGetBytePerPixel to be aligned on uint because the read/write unit is uint32.

diff --git a/Engine/Source/Runtime/RenderCore/Private/RenderUtils.cpp b/Engine/Source/Runtime/RenderCore/Private/RenderUtils.cpp
--- a/Engine/Source/Runtime/RenderCore/Private/RenderUtils.cpp
+++ b/Engine/Source/Runtime/RenderCore/Private/RenderUtils.cpp
@@ -1590,10 +1590,20 @@ namespace Strata
 	}
 
 	static uint32 InternalGetBytePerPixel(uint32 InBytePerPixel)
-	{	
+	{
+		auto RoundUpValueToUInt = [](uint32 Value)
+		{
+			const uint32 Divisor = sizeof(uint32);
+			return  FMath::DivideAndRoundUp<uint32>(Value, Divisor) * Divisor;
+		};
+
 		// We enforce at least 20 bytes per pixel because this is the minimal Strata GBuffer footprint of the simplest material.
 		const uint32 MinStrataBytePerPixel = 20u;
 		const uint32 MaxStrataBytePerPixel = /*IsMobilePlatform(GMaxRHI InPlatform) ? 24u : */256u; // STRATA_TODO
+		// Align byte per pixel count to 4 because the read/write unit is uint32
+		check(MinStrataBytePerPixel == RoundUpValueToUInt(MinStrataBytePerPixel));
+		check(MaxStrataBytePerPixel == RoundUpValueToUInt(MaxStrataBytePerPixel));
+		InBytePerPixel = RoundUpValueToUInt(InBytePerPixel);
 		return FMath::Clamp(InBytePerPixel, MinStrataBytePerPixel, MaxStrataBytePerPixel);
 	}
 
