Substrate - updated / clarified ray tracing common payload byte count to really reflect what is shader dependent and what is Subsutrate dependent.

diff --git a/Engine/Shaders/Private/RayTracing/RayTracingCommon.ush b/Engine/Shaders/Private/RayTracing/RayTracingCommon.ush
--- a/Engine/Shaders/Private/RayTracing/RayTracingCommon.ush
+++ b/Engine/Shaders/Private/RayTracing/RayTracingCommon.ush
@@ -517,9 +517,7 @@ uint TranslateRayTracingInstanceFlags(uint InFlags)
 
 struct FMaterialClosestHitPayload : FMinimalPayload
 {
-	// Unpacked  Packed
-	// offset    bytes
-	// float FMinimalPayload::HitT
+	// float FMinimalPay	// float FMinimalPayload::HitT
 	FRayCone RayCone;
 	float3 Radiance;
 	float3 WorldNormal; // Top Layer Normal
@@ -560,9 +558,9 @@ struct FPackedMaterialClosestHitPayload : FMinimalPayload
 	uint Radiance;										 // 4  bytes
 	uint PackedGeometryNormal;							 // 4  bytes
 	uint SceneInstanceIndex;                             // 4  bytes - GPUSceneInstanceIndex
-														 //=28 bytes - must be in sync with HeaderPayload in InternalGetRayTracingMaterialPayloadSizeInBytes()
-	FStrataRaytracingPayload StrataData;				 // 44 bytes
-														 //=68 bytes total
+														 //=24 bytes - must be in sync with InternalGetRayTracingMaterialPayloadSizeInBytes()
+	FStrataRaytracingPayload StrataData;				 //  ? bytes - must be in sync with GetRaytracingMaterialPayloadSize (see FStrataRaytracingPayload)
+														 //=24+? bytes total
 
 	void SetMinimalPayloadMode()			{ Flags |= RAY_TRACING_PAYLOAD_INPUT_FLAG_MINIMAL_PAYLOAD; }
 	bool IsMinimalPayloadMode()				{ return (Flags & RAY_TRACING_PAYLOAD_INPUT_FLAG_MINIMAL_PAYLOAD) != 0; }
diff --git a/Engine/Shaders/Private/Strata/Strata.ush b/Engine/Shaders/Private/Strata/Strata.ush
--- a/Engine/Shaders/Private/Strata/Strata.ush
+++ b/Engine/Shaders/Private/Strata/Strata.ush
@@ -1656,6 +1656,7 @@ FStrataData GetStrataSingleLayerWaterBSDF(
 
 struct FStrataRaytracingPayload
 {
+	// !!! This structure content must be in sync with GetRaytracingMaterialPayloadSize !!!
 	STRATA_TOP_LAYER_TYPE PackedTopLayerData;
 	uint Data[STRATA_RT_PAYLOAD_NUM_UINTS];
 };
diff --git a/Engine/Source/Runtime/RenderCore/Private/BuiltInRayTracingShaders.cpp b/Engine/Source/Runtime/RenderCore/Private/BuiltInRayTracingShaders.cpp
--- a/Engine/Source/Runtime/RenderCore/Private/BuiltInRayTracingShaders.cpp
+++ b/Engine/Source/Runtime/RenderCore/Private/BuiltInRayTracingShaders.cpp
@@ -11,7 +11,18 @@
 
 uint32 GetRaytracingMaterialPayloadSize()
 {
-	return Strata::IsStrataEnabled() ? Strata::GetRayTracingMaterialPayloadSizeInBytes() : 64u;
+	if (Strata::IsStrataEnabled())
+	{
+		// All the data from FPackedMaterialClosestHitPayload except FStrataRaytracingPayload (see RayTracingCommon.ush)
+		uint32 PayloadSizeBytes = 6 * sizeof(uint32);
+
+		// The remaining data from FStrataRaytracingPayload.
+		PayloadSizeBytes += Strata::GetRayTracingMaterialPayloadSizeInBytes();
+
+		return PayloadSizeBytes;
+	}
+
+	return 64u;
 }
 
 IMPLEMENT_RT_PAYLOAD_TYPE(ERayTracingPayloadType::Default, 24);
diff --git a/Engine/Source/Runtime/RenderCore/Private/RenderUtils.cpp b/Engine/Source/Runtime/RenderCore/Private/RenderUtils.cpp
--- a/Engine/Source/Runtime/RenderCore/Private/RenderUtils.cpp
+++ b/Engine/Source/Runtime/RenderCore/Private/RenderUtils.cpp
@@ -1619,15 +1619,6 @@ namespace Strata
 		return InternalGetBytePerPixel(CVarBudget.Get(InPlatform));
 	}
 
-	static uint32 InternalGetRayTracingMaterialPayloadSizeInBytes(uint32 InBytePerPixels, uint32 InNormalQuality)
-	{
-		// If Strata is enabled, resize the payload accordingly to the byte per pixel request
-		const uint32 HeaderPayload = sizeof(uint32) * (6u); // See FPackedMaterialClosestHitPayload in RayTracingCommon.ush
-		const uint32 TopNormalPayload = sizeof(uint32) * (InNormalQuality == 1 ? 2u : 1u);
-
-		return HeaderPayload + TopNormalPayload + InBytePerPixels;
-	}
-
 	uint32 GetNormalQuality()
 	{
 		static const auto CVar = IConsoleManager::Get().FindTConsoleVariableDataInt(TEXT("r.GBufferFormat"));
@@ -1636,12 +1627,29 @@ namespace Strata
 
 	uint32 GetRayTracingMaterialPayloadSizeInBytes()
 	{
-		return InternalGetRayTracingMaterialPayloadSizeInBytes(GetBytePerPixel(), GetNormalQuality());
-	}
+		// The payload size represents FPackedMaterialClosestHitPayload containing FStrataRaytracingPayload composed of 
+		//  (1)- top layer data using a count of STRATA_TOP_LAYER_TYPE uints
+		//  (2)- packed data in a uint32 array of size STRATA_RT_PAYLOAD_NUM_UINTS
 
-	uint32 GetRayTracingMaterialPayloadSizeInBytes(EShaderPlatform InPlatform)
-	{
-		return InternalGetRayTracingMaterialPayloadSizeInBytes(GetBytePerPixel(InPlatform), GetNormalQuality());
+		uint32 PayloadSizeBytes = 0;
+
+		// Account for (1)
+		const uint32 StrataNormalQuality = Strata::GetNormalQuality();
+		check(StrataNormalQuality >= 0 && StrataNormalQuality <= 1);
+		switch (StrataNormalQuality)
+		{
+		case 0:
+			PayloadSizeBytes += 1 * sizeof(uint32);
+			break;
+		case 1:
+			PayloadSizeBytes += 2 * sizeof(uint32);	// with high quality normal
+			break;
+		}
+
+		// Account for (2)
+		PayloadSizeBytes += Strata::GetBytePerPixel();
+
+		return PayloadSizeBytes;
 	}
 
 	bool IsBackCompatibilityEnabled()
diff --git a/Engine/Source/Runtime/RenderCore/Public/RenderUtils.h b/Engine/Source/Runtime/RenderCore/Public/RenderUtils.h
--- a/Engine/Source/Runtime/RenderCore/Public/RenderUtils.h
+++ b/Engine/Source/Runtime/RenderCore/Public/RenderUtils.h
@@ -490,7 +490,6 @@ namespace Strata
 	RENDERCORE_API bool IsAccurateSRGBEnabled();
 
 	RENDERCORE_API uint32 GetRayTracingMaterialPayloadSizeInBytes();
-	RENDERCORE_API uint32 GetRayTracingMaterialPayloadSizeInBytes(EShaderPlatform InPlatform);
 
 	RENDERCORE_API uint32 GetBytePerPixel();
 	RENDERCORE_API uint32 GetBytePerPixel(EShaderPlatform InPlatform);
