Fix GPUScene uploader size limit check.

diff --git a/Engine/Source/Runtime/RenderCore/Private/UnifiedBuffer.cpp b/Engine/Source/Runtime/RenderCore/Private/UnifiedBuffer.cpp
--- a/Engine/Source/Runtime/RenderCore/Private/UnifiedBuffer.cpp
+++ b/Engine/Source/Runtime/RenderCore/Private/UnifiedBuffer.cpp
@@ -10,6 +10,12 @@
 #include "RenderGraphUtils.h"
 #include "DataDrivenShaderPlatformInfo.h"
 
+// Uploads use a storage buffers which are at least 128m elements
+static uint64 GetMaxUploadBufferElements()
+{
+	return (1 << 27);
+}
+
 enum class EByteBufferResourceType
 {
 	Float4_Buffer,
@@ -555,12 +561,12 @@ void FRDGScatterUploadBuffer::Init(FRDGBuilder& GraphBuilder, uint32 NumElements
 
 	const uint32 ScatterNumBytesPerElement = sizeof(uint32);
 	const uint32 ScatterBytes = NumElements * ScatterNumBytesPerElement;
-	const uint32 ScatterBufferSize = (uint32)FMath::Min((uint64)FMath::RoundUpToPowerOfTwo(ScatterBytes), GetMaxBufferDimension() * sizeof(uint32));
+	const uint32 ScatterBufferSize = (uint32)FMath::Min((uint64)FMath::RoundUpToPowerOfTwo(ScatterBytes), GetMaxUploadBufferElements() * sizeof(uint32));
 	check(ScatterBufferSize >= ScatterBytes);
 
 	const uint32 UploadNumBytesPerElement = TypeSize;
 	const uint32 UploadBytes = NumElements * NumBytesPerElement;
-	const uint32 UploadBufferSize = (uint32)FMath::Min((uint64)FMath::RoundUpToPowerOfTwo(UploadBytes), GetMaxBufferDimension() * TypeSize);
+	const uint32 UploadBufferSize = (uint32)FMath::Min((uint64)FMath::RoundUpToPowerOfTwo(UploadBytes), GetMaxUploadBufferElements() * TypeSize);
 	check(UploadBufferSize >= UploadBytes);
 
 	// Recreate buffers is they are already queued into RDG from a previous call.
@@ -714,12 +720,12 @@ FRDGScatterUploader* FRDGAsyncScatterUploadBuffer::Begin(FRDGBuilder& GraphBuild
 
 	const uint32 ScatterNumBytesPerElement = sizeof(uint32);
 	const uint32 ScatterBytes = NumElements * ScatterNumBytesPerElement;
-	const uint32 ScatterBufferSize = (uint32)FMath::Min((uint64)FMath::RoundUpToPowerOfTwo(ScatterBytes), GetMaxBufferDimension() * sizeof(uint32));
+	const uint32 ScatterBufferSize = (uint32)FMath::Min((uint64)FMath::RoundUpToPowerOfTwo(ScatterBytes), GetMaxUploadBufferElements() * sizeof(uint32));
 	check(ScatterBufferSize >= ScatterBytes);
 
 	const uint32 UploadNumBytesPerElement = TypeSize;
 	const uint32 UploadBytes = NumElements * NumBytesPerElement;
-	const uint32 UploadBufferSize = (uint32)FMath::Min((uint64)FMath::RoundUpToPowerOfTwo(UploadBytes), GetMaxBufferDimension() * TypeSize);
+	const uint32 UploadBufferSize = (uint32)FMath::Min((uint64)FMath::RoundUpToPowerOfTwo(UploadBytes), GetMaxUploadBufferElements() * TypeSize);
 	check(UploadBufferSize >= UploadBytes);
 
 	// Recreate buffers is they are already queued into RDG from a previous call.
