Substrata - ray tracing material shaders uses fully simplified material for smaller payload.
Next would be to have the fully simplified mode disable anisotropy to have an even lower payload size.

diff --git a/Engine/Shaders/Private/Strata/Strata.ush b/Engine/Shaders/Private/Strata/Strata.ush
--- a/Engine/Shaders/Private/Strata/Strata.ush
+++ b/Engine/Shaders/Private/Strata/Strata.ush
@@ -1649,9 +1649,22 @@ FStrataData GetStrataSingleLayerWaterBSDF(
 ///////////////////////////////////////////////////////////////////////////////
 // Raytracing material buffer
 
-// Data payload size is determined based on Strata.BytesPerPixel settings
-#ifndef STRATA_RT_PAYLOAD_NUM_UINTS
-#error STRATA_RT_PAYLOAD_NUM_UINTS needs to be defined
+#ifndef STRATA_MATERIAL_NUM_UINTS
+#error STRATA_MATERIAL_NUM_UINTS needs to be defined
+#endif
+
+/*
+Material payload drives whether or not a Strata material should be fully simplified.
+This is setup in a few places:
+- FShaderType::ModifyCompilationEnvironment automatically enables STRATA_USE_FULLYSIMPLIFIED_MATERIAL for some PayloadType (to automatically impact CHS, MS, RGS and tracing CS)
+- When a payload is declared and it needs to transport Strata data, the size can be computed using Strata::GetRayTracingMaterialPayloadSizeInBytes, e.g. GetRaytracingMaterialPayloadSizeFullySimplified.
+*/
+
+#if STRATA_USE_FULLYSIMPLIFIED_MATERIAL
+#define STRATA_RT_PAYLOAD_NUM_UINTS STRATA_FULLY_SIMPLIFIED_NUM_UINTS
+#else
+// STRATA_MATERIAL_NUM_UINTS is based on Strata.BytesPerPixel settings
+#define STRATA_RT_PAYLOAD_NUM_UINTS STRATA_MATERIAL_NUM_UINTS
 #endif
 
 struct FStrataRaytracingPayload
diff --git a/Engine/Shaders/Shared/StrataDefinitions.h b/Engine/Shaders/Shared/StrataDefinitions.h
--- a/Engine/Shaders/Shared/StrataDefinitions.h
+++ b/Engine/Shaders/Shared/StrataDefinitions.h
@@ -8,7 +8,7 @@
 #pragma once
 
 // Change this to force recompilation of all strata dependent shaders (use https://www.random.org/cgi-bin/randbyte?nbytes=4&format=h)
-#define STRATA_SHADER_VERSION 0x084bf26a2 
+#define STRATA_SHADER_VERSION 0xbd4a70e7
 
 
 
@@ -20,6 +20,9 @@
 
 #define STRATA_PACKED_SHAREDLOCALBASIS_STRIDE_BYTES	4
 
+// As of today, a fully simplified material is a slab with all features allowed. It can thus be complex if anisotropy is enabled and in this case eats up to 32bytes.
+// STRATA_TODO: fully simplified should remove all features but fuzz maybe. 
+#define STRATA_FULLY_SIMPLIFIED_NUM_UINTS	(32/4)
 
 #define STRATA_BSDF_TYPE_SLAB				0
 #define STRATA_BSDF_TYPE_VOLUMETRICFOGCLOUD	1
diff --git a/Engine/Source/Runtime/Engine/Private/ShaderCompiler/ShaderCompiler.cpp b/Engine/Source/Runtime/Engine/Private/ShaderCompiler/ShaderCompiler.cpp
--- a/Engine/Source/Runtime/Engine/Private/ShaderCompiler/ShaderCompiler.cpp
+++ b/Engine/Source/Runtime/Engine/Private/ShaderCompiler/ShaderCompiler.cpp
@@ -6560,7 +6560,7 @@ void GlobalBeginCompileShader(
 			Input.Environment.SetDefine(TEXT("STRATA_NORMAL_QUALITY"), StrataNormalQuality);
 
 			const uint32 StrataUintPerPixel = Strata::GetBytePerPixel(Target.GetPlatform()) / 4u;
-			Input.Environment.SetDefine(TEXT("STRATA_RT_PAYLOAD_NUM_UINTS"), StrataUintPerPixel);
+			Input.Environment.SetDefine(TEXT("STRATA_MATERIAL_NUM_UINTS"), StrataUintPerPixel);
 
 			const bool bTileCoord8Bits = Strata::Is8bitTileCoordEnabled();
 			Input.Environment.SetDefine(TEXT("USE_8BIT_TILE_COORD"), bTileCoord8Bits ? 1 : 0);
diff --git a/Engine/Source/Runtime/RenderCore/Private/BuiltInRayTracingShaders.cpp b/Engine/Source/Runtime/RenderCore/Private/BuiltInRayTracingShaders.cpp
--- a/Engine/Source/Runtime/RenderCore/Private/BuiltInRayTracingShaders.cpp
+++ b/Engine/Source/Runtime/RenderCore/Private/BuiltInRayTracingShaders.cpp
@@ -9,7 +9,7 @@
 
 #include "RayTracingPayloadType.h"
 
-uint32 GetRaytracingMaterialPayloadSize()
+uint32 GetRaytracingMaterialPayloadSizeFullySimplified()
 {
 	if (Strata::IsStrataEnabled())
 	{
@@ -17,7 +17,8 @@ uint32 GetRaytracingMaterialPayloadSize()
 		uint32 PayloadSizeBytes = 6 * sizeof(uint32);
 
 		// The remaining data from FStrataRaytracingPayload.
-		PayloadSizeBytes += Strata::GetRayTracingMaterialPayloadSizeInBytes();
+		const bool bFullySimplifiedMaterial = true;	// This is needed because ERayTracingPayloadType::RayTracingMaterial will be fully simplified, see FShaderType::ModifyCompilationEnvironment.
+		PayloadSizeBytes += Strata::GetRayTracingMaterialPayloadSizeInBytes(bFullySimplifiedMaterial);
 
 		return PayloadSizeBytes;
 	}
@@ -26,7 +27,7 @@ uint32 GetRaytracingMaterialPayloadSize()
 }
 
 IMPLEMENT_RT_PAYLOAD_TYPE(ERayTracingPayloadType::Default, 24);
-IMPLEMENT_RT_PAYLOAD_TYPE_FUNCTION(ERayTracingPayloadType::RayTracingMaterial, GetRaytracingMaterialPayloadSize);
+IMPLEMENT_RT_PAYLOAD_TYPE_FUNCTION(ERayTracingPayloadType::RayTracingMaterial, GetRaytracingMaterialPayloadSizeFullySimplified);
 
 IMPLEMENT_GLOBAL_SHADER( FDefaultPayloadMS, "/Engine/Private/RayTracing/RayTracingBuiltInShaders.usf", "DefaultPayloadMS", SF_RayMiss);
 IMPLEMENT_GLOBAL_SHADER( FPackedMaterialClosestHitPayloadMS, "/Engine/Private/RayTracing/RayTracingBuiltInShaders.usf", "PackedMaterialClosestHitPayloadMS", SF_RayMiss);
diff --git a/Engine/Source/Runtime/RenderCore/Private/RenderUtils.cpp b/Engine/Source/Runtime/RenderCore/Private/RenderUtils.cpp
--- a/Engine/Source/Runtime/RenderCore/Private/RenderUtils.cpp
+++ b/Engine/Source/Runtime/RenderCore/Private/RenderUtils.cpp
@@ -15,6 +15,7 @@
 #include "ShaderPlatformCachedIniValue.h"
 #include "DataDrivenShaderPlatformInfo.h"
 #include "RenderCore.h"
+#include "StrataDefinitions.h"
 
 #if WITH_EDITOR
 #include "Interfaces/ITargetPlatform.h"
@@ -1625,7 +1626,7 @@ namespace Strata
 		return CVar && CVar->GetValueOnAnyThread() > 1 ? 1 : 0;
 	}
 
-	uint32 GetRayTracingMaterialPayloadSizeInBytes()
+	uint32 GetRayTracingMaterialPayloadSizeInBytes(bool bFullySimplifiedMaterial)
 	{
 		// The payload size represents FPackedMaterialClosestHitPayload containing FStrataRaytracingPayload composed of 
 		//  (1)- top layer data using a count of STRATA_TOP_LAYER_TYPE uints
@@ -1647,7 +1648,14 @@ namespace Strata
 		}
 
 		// Account for (2)
-		PayloadSizeBytes += Strata::GetBytePerPixel();
+		if (bFullySimplifiedMaterial)
+		{
+			PayloadSizeBytes += STRATA_FULLY_SIMPLIFIED_NUM_UINTS * sizeof(uint32);
+		}
+		else
+		{
+			PayloadSizeBytes += Strata::GetBytePerPixel();
+		}
 
 		return PayloadSizeBytes;
 	}
diff --git a/Engine/Source/Runtime/RenderCore/Private/Shader.cpp b/Engine/Source/Runtime/RenderCore/Private/Shader.cpp
--- a/Engine/Source/Runtime/RenderCore/Private/Shader.cpp
+++ b/Engine/Source/Runtime/RenderCore/Private/Shader.cpp
@@ -450,6 +450,14 @@ void FShaderType::ModifyCompilationEnvironment(const FShaderPermutationParameter
 
 		OutEnvironment.SetDefine(TEXT("RT_PAYLOAD_TYPE"), static_cast<int32>(RayTracingPayloadType));
 		OutEnvironment.SetDefine(TEXT("RT_PAYLOAD_MAX_SIZE"), GetRayTracingPayloadTypeMaxSize(RayTracingPayloadType));
+
+		if (uint32(RayTracingPayloadType) & uint32(ERayTracingPayloadType::RayTracingMaterial))
+		{
+			// If any payload requires a fully simplified material, we force fully simplified material all the way.
+			// That is used to have material ray tracing shaders compressed to single slab.
+			// Smaller payload means faster performance and for some tracing this will be enough, e.g. reflected materials, lightmass diffuse interactions.
+			OutEnvironment.SetDefine(TEXT("STRATA_USE_FULLYSIMPLIFIED_MATERIAL"), 1);
+		}
 	}
 #endif
 }
diff --git a/Engine/Source/Runtime/RenderCore/Public/RenderUtils.h b/Engine/Source/Runtime/RenderCore/Public/RenderUtils.h
--- a/Engine/Source/Runtime/RenderCore/Public/RenderUtils.h
+++ b/Engine/Source/Runtime/RenderCore/Public/RenderUtils.h
@@ -489,7 +489,7 @@ namespace Strata
 	RENDERCORE_API bool Is8bitTileCoordEnabled();
 	RENDERCORE_API bool IsAccurateSRGBEnabled();
 
-	RENDERCORE_API uint32 GetRayTracingMaterialPayloadSizeInBytes();
+	RENDERCORE_API uint32 GetRayTracingMaterialPayloadSizeInBytes(bool bFullySimplifiedMaterial);
 
 	RENDERCORE_API uint32 GetBytePerPixel();
 	RENDERCORE_API uint32 GetBytePerPixel(EShaderPlatform InPlatform);
