Substrate - always precompile with DXC when substrate is enabled and if DXC is supported. No need for any ForceDXC cvar anymore.

diff --git a/Engine/Source/Runtime/RenderCore/Private/Shader.cpp b/Engine/Source/Runtime/RenderCore/Private/Shader.cpp
--- a/Engine/Source/Runtime/RenderCore/Private/Shader.cpp
+++ b/Engine/Source/Runtime/RenderCore/Private/Shader.cpp
@@ -1420,11 +1420,15 @@ void DispatchIndirectComputeShader(
 
 bool IsDxcEnabledForPlatform(EShaderPlatform Platform, bool bHlslVersion2021)
 {
+	// If Substrate is enabled, DXC must be used if possible to be able to compile with member functions and other language features.
+	static const IConsoleVariable* SubstrateCVar = IConsoleManager::Get().FindConsoleVariable(TEXT("r.substrate"));
+	const bool bSubstrateForcesDXC = (SubstrateCVar && SubstrateCVar->GetInt() != 0);
+
 	// Check the generic console variable first (if DXC is supported)
 	if (FDataDrivenShaderPlatformInfo::GetSupportsDxc(Platform))
 	{
 		static const IConsoleVariable* CVar = IConsoleManager::Get().FindConsoleVariable(TEXT("r.Shaders.ForceDXC"));
-		if (bHlslVersion2021 || (CVar && CVar->GetInt() != 0))
+		if (bHlslVersion2021 || (CVar && CVar->GetInt() != 0) || bSubstrateForcesDXC)
 		{
 			return true;
 		}
@@ -1434,12 +1438,12 @@ bool IsDxcEnabledForPlatform(EShaderPlatform Platform, bool bHlslVersion2021)
 	{
 		// D3D backend supports a precompile step for HLSL2021 which is separate from ForceDXC option
 		static const IConsoleVariable* CVar = IConsoleManager::Get().FindConsoleVariable(TEXT("r.D3D.ForceDXC"));
-		return (CVar && CVar->GetInt() != 0);
+		return (CVar && CVar->GetInt() != 0 || bSubstrateForcesDXC);
 	}
 	if (IsOpenGLPlatform(Platform))
 	{
 		static const IConsoleVariable* CVar = IConsoleManager::Get().FindConsoleVariable(TEXT("r.OpenGL.ForceDXC"));
-		return (bHlslVersion2021 || (CVar && CVar->GetInt() != 0));
+		return (bHlslVersion2021 || (CVar && CVar->GetInt() != 0) || bSubstrateForcesDXC);
 	}
 	// Hlslcc has been removed for Metal and Vulkan. There is only DXC now.
 	if (IsMetalPlatform(Platform) || IsVulkanPlatform(Platform))
