Compilation error fix for 24853509

diff --git a/Engine/Source/Runtime/RenderCore/Private/Shader.cpp b/Engine/Source/Runtime/RenderCore/Private/Shader.cpp
--- a/Engine/Source/Runtime/RenderCore/Private/Shader.cpp
+++ b/Engine/Source/Runtime/RenderCore/Private/Shader.cpp
@@ -1438,7 +1438,7 @@ bool IsDxcEnabledForPlatform(EShaderPlatform Platform, bool bHlslVersion2021)
 	{
 		// D3D backend supports a precompile step for HLSL2021 which is separate from ForceDXC option
 		static const IConsoleVariable* CVar = IConsoleManager::Get().FindConsoleVariable(TEXT("r.D3D.ForceDXC"));
-		return (CVar && CVar->GetInt() != 0 || bSubstrateForcesDXC);
+		return ((CVar && CVar->GetInt() != 0) || bSubstrateForcesDXC);
 	}
 	if (IsOpenGLPlatform(Platform))
 	{
