Fixes ERHIAccess::CopyDest not being in DumpGPU's output resources

diff --git a/Engine/Source/Runtime/RenderCore/Private/DumpGPU.cpp b/Engine/Source/Runtime/RenderCore/Private/DumpGPU.cpp
--- a/Engine/Source/Runtime/RenderCore/Private/DumpGPU.cpp
+++ b/Engine/Source/Runtime/RenderCore/Private/DumpGPU.cpp
@@ -2071,10 +2071,7 @@ void FRDGBuilder::DumpResourcePassOutputs(const FRDGPass* Pass)
 		{
 			if (FRDGTextureAccess TextureAccess = Parameter.GetAsTextureAccess())
 			{
-				bool bIsOutputResource = (
-					TextureAccess.GetAccess() == ERHIAccess::UAVCompute ||
-					TextureAccess.GetAccess() == ERHIAccess::UAVGraphics ||
-					TextureAccess.GetAccess() == ERHIAccess::RTV);
+				bool bIsOutputResource = IsWritableAccess(TextureAccess.GetAccess());
 
 				FRDGTextureSRVDesc TextureSubResource = FRDGTextureSRVDesc::Create(TextureAccess);
 				GRDGResourceDumpContext->AddDumpTexturePasses(*this, InputResourceNames, OutputResourceNames, Pass, TextureSubResource, bIsOutputResource);
@@ -2087,10 +2084,7 @@ void FRDGBuilder::DumpResourcePassOutputs(const FRDGPass* Pass)
 
 			for (FRDGTextureAccess TextureAccess : TextureAccessArray)
 			{
-				bool bIsOutputResource = (
-					TextureAccess.GetAccess() == ERHIAccess::UAVCompute ||
-					TextureAccess.GetAccess() == ERHIAccess::UAVGraphics ||
-					TextureAccess.GetAccess() == ERHIAccess::RTV);
+				bool bIsOutputResource = IsWritableAccess(TextureAccess.GetAccess());
 
 				FRDGTextureSRVDesc TextureSubResource = FRDGTextureSRVDesc::Create(TextureAccess);
 				GRDGResourceDumpContext->AddDumpTexturePasses(*this, InputResourceNames, OutputResourceNames, Pass, TextureSubResource, bIsOutputResource);
@@ -2120,9 +2114,7 @@ void FRDGBuilder::DumpResourcePassOutputs(const FRDGPass* Pass)
 		{
 			if (FRDGBufferAccess BufferAccess = Parameter.GetAsBufferAccess())
 			{
-				bool bIsOutputResource = (
-					BufferAccess.GetAccess() == ERHIAccess::UAVCompute ||
-					BufferAccess.GetAccess() == ERHIAccess::UAVGraphics);
+				bool bIsOutputResource = IsWritableAccess(BufferAccess.GetAccess());
 
 				GRDGResourceDumpContext->AddDumpBufferPass(*this, InputResourceNames, OutputResourceNames, Pass, BufferAccess, bIsOutputResource);
 			}
@@ -2134,9 +2126,7 @@ void FRDGBuilder::DumpResourcePassOutputs(const FRDGPass* Pass)
 
 			for (FRDGBufferAccess BufferAccess : BufferAccessArray)
 			{
-				bool bIsOutputResource = (
-					BufferAccess.GetAccess() == ERHIAccess::UAVCompute ||
-					BufferAccess.GetAccess() == ERHIAccess::UAVGraphics);
+				bool bIsOutputResource = IsWritableAccess(BufferAccess.GetAccess());
 
 				GRDGResourceDumpContext->AddDumpBufferPass(*this, InputResourceNames, OutputResourceNames, Pass, BufferAccess, bIsOutputResource);
 			}
