[shaders] fix uninitialized oodle compression fields on FShaderCode struct

diff --git a/Engine/Source/Runtime/RenderCore/Public/ShaderCore.h b/Engine/Source/Runtime/RenderCore/Public/ShaderCore.h
--- a/Engine/Source/Runtime/RenderCore/Public/ShaderCore.h
+++ b/Engine/Source/Runtime/RenderCore/Public/ShaderCore.h
@@ -6,6 +6,7 @@
 
 #pragma once
 
+#include "Compression/OodleDataCompression.h"
 #include "Containers/Array.h"
 #include "Containers/ArrayView.h"
 #include "Containers/ContainerAllocationPolicies.h"
@@ -137,12 +138,6 @@ extern RENDERCORE_API bool AllowDebugViewmodes(EShaderPlatform Platform);
 /** Returns the shader compression format. Oodle is used exclusively now. r.Shaders.SkipCompression configures Oodle to be uncompressed instead of returning NAME_None.*/
 extern RENDERCORE_API FName GetShaderCompressionFormat();
 
-namespace FOodleDataCompression
-{
-	enum class ECompressor : uint8;
-	enum class ECompressionLevel : int8;
-}
-
 /** Returns Oodle-specific shader compression format settings (passing ShaderFormat for future proofing, but as of now the setting is global for all formats). */
 extern RENDERCORE_API void GetShaderCompressionOodleSettings(FOodleDataCompression::ECompressor& OutCompressor, FOodleDataCompression::ECompressionLevel& OutLevel, const FName& ShaderFormat = NAME_None);
 
@@ -1007,6 +1002,8 @@ public:
 	: OptionalDataSize(0)
 	, UncompressedSize(0)
 	, CompressionFormat(NAME_None)
+	, OodleCompressor(FOodleDataCompression::ECompressor::NotSet)
+	, OodleLevel(FOodleDataCompression::ECompressionLevel::None)
 	, ShaderCodeSize(0)
 	{
 	}
