Fixed potential buffer overflows when parsing the GPU message buffer:
* use signed ints for sizes so that we don't have to check every subtraction for underflow
* make sure that we always have at least 2 ints available for reading for each message, to avoid overflow if a message is truncated before the payload size is written
* check that the payload size doesn't exceed the remaining buffer size

#rnx

diff --git a/Engine/Source/Runtime/RenderCore/Private/GPUMessaging.cpp b/Engine/Source/Runtime/RenderCore/Private/GPUMessaging.cpp
--- a/Engine/Source/Runtime/RenderCore/Private/GPUMessaging.cpp
+++ b/Engine/Source/Runtime/RenderCore/Private/GPUMessaging.cpp
@@ -81,7 +81,7 @@ public:
 			return;
 		}
 
-		const uint32 MaxBufferSize = GetMaxBufferSize();
+		const int32 MaxBufferSize = GetMaxBufferSize();
 
 		FRDGBufferDesc Desc = FRDGBufferDesc::CreateStructuredDesc(sizeof(uint32), MaxBufferSize);
 		Desc.Usage |= EBufferUsageFlags::SourceCopy;
@@ -136,32 +136,41 @@ private:
 
 			if (GPUBufferReadback->IsReady())
 			{
-				const uint32 MaxBufferSize = GetMaxBufferSize();
+				const int32 MaxBufferSize = GetMaxBufferSize();
 
 				const uint32* BufferPtr = (const uint32*)GPUBufferReadback->Lock(MaxBufferSize * sizeof(uint32));
-				uint32 ValidRangeEnd = BufferPtr[0] + 1U;
+				int32 ValidRangeEnd = (int32)BufferPtr[0] + 1;
 
-				FString LogMessage;
-
-				const uint32 OriginalValidRangeEnd = ValidRangeEnd;
+				const int32 OriginalValidRangeEnd = ValidRangeEnd;
 				ValidRangeEnd = FMath::Min(ValidRangeEnd, MaxBufferSize);
-				ensureMsgf(ValidRangeEnd == OriginalValidRangeEnd, TEXT("GPU messages size %u exceeded maximum size %u. Results have been truncated."), OriginalValidRangeEnd, MaxBufferSize);
+				ensureMsgf(ValidRangeEnd == OriginalValidRangeEnd, TEXT("GPU messages size %d exceeded maximum size %d. Results have been truncated."), OriginalValidRangeEnd, MaxBufferSize);
 
 				if (GLogAllMessages != 0)
 				{
+					FString LogMessage;
+
+					LogMessage.Appendf(TEXT("\nMessage Buffer[0] = %u;\n"), BufferPtr[0]);
+
 					if (ValidRangeEnd != OriginalValidRangeEnd)
 					{
-						LogMessage.Appendf(TEXT("GPU messages size %u exceeded maximum size %u. Results have been truncated!\n\n"), OriginalValidRangeEnd, MaxBufferSize);
+						LogMessage.Appendf(TEXT("GPU messages size %d exceeded maximum size %d. Results have been truncated.\n"), OriginalValidRangeEnd, MaxBufferSize);
 					}
 
-					LogMessage.Appendf(TEXT("\nMessage Buffer[0] = %d;\n"), BufferPtr[0]);
-
-					for (uint32 ReadOffset = 1U; ReadOffset < ValidRangeEnd; )
+					// Subtract 1 from ValidRangeEnd because we expect to have at least 2 ints available (message ID and payload size).
+					for (int32 ReadOffset = 1; ReadOffset < ValidRangeEnd - 1; )
 					{
 						const FMessageId MessageId(BufferPtr[ReadOffset]);
-						uint32 PayloadNumUints = BufferPtr[ReadOffset + 1];
+						int32 PayloadNumUints = (int32)BufferPtr[ReadOffset + 1];
+
+						const int32 OriginalPayloadNumUints = PayloadNumUints;
+						PayloadNumUints = FMath::Min(PayloadNumUints, ValidRangeEnd - ReadOffset - 2);
 
-						LogMessage.Appendf(TEXT("  Message[Offset: %d, ID: %d, PayloadSize: %d, "), ReadOffset, MessageId, PayloadNumUints);
+						LogMessage.Appendf(TEXT("  Message[Offset: %d, ID: %d, PayloadSize: %d, "), ReadOffset, MessageId.GetIndexUnchecked(), OriginalPayloadNumUints);
+
+						if (PayloadNumUints != OriginalPayloadNumUints)
+						{
+							LogMessage.Appendf(TEXT("!payload size %d exceeds buffer size (ReadOffset=%d, ValidRangeEnd=%d), truncated to %d!, "), OriginalPayloadNumUints, ReadOffset, ValidRangeEnd, PayloadNumUints);
+						}
 
 						if (auto Handler = MessageHandlers.Find(MessageId))
 						{
@@ -175,9 +184,9 @@ private:
 						if (PayloadNumUints > 0)
 						{
 							LogMessage.Append(TEXT(" Payload: { "));
-							for (uint32 PayloadIndex = 0; PayloadIndex < PayloadNumUints; ++PayloadIndex)
+							for (int32 PayloadIndex = 0; PayloadIndex < PayloadNumUints; ++PayloadIndex)
 							{
-								LogMessage.Appendf(TEXT("%d%s"), BufferPtr[ReadOffset + 2U + PayloadIndex], PayloadIndex + 1U == PayloadNumUints ? TEXT("") : TEXT(", "));
+								LogMessage.Appendf(TEXT("%u%s"), BufferPtr[ReadOffset + 2 + PayloadIndex], PayloadIndex + 1 == PayloadNumUints ? TEXT("") : TEXT(", "));
 							}
 							LogMessage.Append(TEXT(" }"));
 						}
@@ -185,26 +194,33 @@ private:
 						LogMessage.Append(TEXT("\n"));
 
 						// Step past the message to the next
-						ReadOffset += PayloadNumUints + 2U;
+						ReadOffset += PayloadNumUints + 2;
 					}
 
 					UE_LOG(LogGPUMessaging, Log, TEXT("%s"), *LogMessage);
 				}
 
-				for (uint32 ReadOffset = 1U; ReadOffset < ValidRangeEnd; )
+				int32 ReadOffset;
+				for (ReadOffset = 1; ReadOffset < ValidRangeEnd - 1; )
 				{
 					const FMessageId MessageId(BufferPtr[ReadOffset]);
-					uint32 PayloadNumUints = BufferPtr[ReadOffset + 1];
+					int32 PayloadNumUints = (int32)BufferPtr[ReadOffset + 1];
+
+					const int32 OriginalPayloadNumUints = PayloadNumUints;
+					PayloadNumUints = FMath::Min(PayloadNumUints, ValidRangeEnd - ReadOffset - 2);
+					ensureMsgf(PayloadNumUints == OriginalPayloadNumUints, TEXT("GPU message payload size %d exceeds buffer size (ReadOffset=%d, ValidRangeEnd=%d), truncated to %d."), OriginalPayloadNumUints, ReadOffset, ValidRangeEnd, PayloadNumUints);
 
 					if (auto Handler = MessageHandlers.Find(MessageId))
 					{
-						(*Handler)->Execute(FReader(MessageId, PayloadNumUints, &BufferPtr[ReadOffset + 2U]));
+						(*Handler)->Execute(FReader(MessageId, PayloadNumUints, &BufferPtr[ReadOffset + 2]));
 					}
 
 					// Step past the message to the next
-					ReadOffset += PayloadNumUints + 2U;
+					ReadOffset += PayloadNumUints + 2;
 				}
 
+				ensureMsgf(ReadOffset == ValidRangeEnd, TEXT("Garbage at the end of the GPU message buffer, ReadOffset=%d, ValidRangeEnd=%d"), ReadOffset, ValidRangeEnd);
+
 				GPUBufferReadback->Unlock();
 
 				MessageReadbackBuffersInUse.RemoveAt(Index);
@@ -236,7 +252,7 @@ private:
 		MessageReadbackBuffersReady.Empty();
 	}
 
-	uint32 GetMaxBufferSize() const { return GMaxBufferSize * 1024; }
+	int32 GetMaxBufferSize() const { return GMaxBufferSize * 1024; }
 
 	TMap<FMessageId, TSharedPtr<FHandler>> MessageHandlers;
 
