Prevents the editor from crashing when the user selects a Material whose ShaderMap contains null shaders.

This change only makes the editor more robust and prevents the crash. Shaders in the ShaderMap may be null after loading due to external issues.

diff --git a/Engine/Source/Runtime/RenderCore/Private/ShaderMap.cpp b/Engine/Source/Runtime/RenderCore/Private/ShaderMap.cpp
--- a/Engine/Source/Runtime/RenderCore/Private/ShaderMap.cpp
+++ b/Engine/Source/Runtime/RenderCore/Private/ShaderMap.cpp
@@ -473,8 +473,11 @@ void FShaderMapContent::GetShaderList(const FShaderMapBase& InShaderMap, TMap<FH
 {
 	for (int32 ShaderIndex = 0; ShaderIndex < Shaders.Num(); ++ShaderIndex)
 	{
-		FShader* Shader = Shaders[ShaderIndex].GetChecked();
-		OutShaders.Add(ShaderTypes[ShaderIndex], TShaderRef<FShader>(Shader, InShaderMap));
+		FShader* Shader = Shaders[ShaderIndex].Get();
+		if (ensure(Shader))
+		{
+			OutShaders.Add(ShaderTypes[ShaderIndex], TShaderRef<FShader>(Shader, InShaderMap));
+		}
 	}
 
 	for (const FShaderPipeline* ShaderPipeline : ShaderPipelines)
@@ -582,7 +585,10 @@ uint32 FShaderMapContent::GetMaxTextureSamplersShaderMap(const FShaderMapBase& I
 
 	for (FShader* Shader : Shaders)
 	{
-		MaxTextureSamplers = FMath::Max(MaxTextureSamplers, Shader->GetNumTextureSamplers());
+		if (ensure(Shader))
+		{
+			MaxTextureSamplers = FMath::Max(MaxTextureSamplers, Shader->GetNumTextureSamplers());
+		}
 	}
 
 	for (FShaderPipeline* Pipeline : ShaderPipelines)
