Removed duplicate assignment operator.

diff --git a/Engine/Source/Runtime/RenderCore/Public/RenderTransform.h b/Engine/Source/Runtime/RenderCore/Public/RenderTransform.h
--- a/Engine/Source/Runtime/RenderCore/Public/RenderTransform.h
+++ b/Engine/Source/Runtime/RenderCore/Public/RenderTransform.h
@@ -56,15 +56,6 @@ public:
 		Origin = FVector3f((float)M.M[3][0], (float)M.M[3][1], (float)M.M[3][2]);
 	}
 
-	FORCEINLINE FRenderTransform& operator=(const FRenderTransform& From)
-	{
-		TransformRows[0] = From.TransformRows[0];
-		TransformRows[1] = From.TransformRows[1];
-		TransformRows[2] = From.TransformRows[2];
-		Origin = From.Origin;
-		return *this;
-	}
-
 	FORCEINLINE FRenderTransform& operator=(const FMatrix44f& From)
 	{
 		TransformRows[0]	= FVector3f(From.M[0][0], From.M[0][1], From.M[0][2]);
