Fixes a bug on DumpGPU with Texture2DArray bound to FRenderTargetBinding with ArraySlice==-1

diff --git a/Engine/Source/Runtime/RenderCore/Private/DumpGPU.cpp b/Engine/Source/Runtime/RenderCore/Private/DumpGPU.cpp
--- a/Engine/Source/Runtime/RenderCore/Private/DumpGPU.cpp
+++ b/Engine/Source/Runtime/RenderCore/Private/DumpGPU.cpp
@@ -2142,7 +2142,7 @@ void FRDGBuilder::DumpResourcePassOutputs(const FRDGPass* Pass)
 				FRDGTextureRef Texture = RenderTarget.GetTexture();
 				FRDGTextureSRVDesc TextureSubResource = FRDGTextureSRVDesc::CreateForMipLevel(Texture, RenderTarget.GetMipIndex());
 
-				if (Texture->Desc.IsTextureArray())
+				if (Texture->Desc.IsTextureArray() && RenderTarget.GetArraySlice() != -1)
 				{
 					TextureSubResource.FirstArraySlice = RenderTarget.GetArraySlice();
 					TextureSubResource.NumArraySlices = 1;
