Compiles VisualizeTexture on test build

It made sens back in UE4 needing to do GVisualizeTexture.SetCheckPoint(), but now being fully transparent in RDG like DumpGPU it just doesn't make sens to not have it by default in test build like we do with DumpGPU

diff --git a/Engine/Source/Runtime/RenderCore/Public/VisualizeTexture.h b/Engine/Source/Runtime/RenderCore/Public/VisualizeTexture.h
--- a/Engine/Source/Runtime/RenderCore/Public/VisualizeTexture.h
+++ b/Engine/Source/Runtime/RenderCore/Public/VisualizeTexture.h
@@ -26,7 +26,7 @@ class FRDGBuilder;
 class FRHICommandListImmediate;
 class FWildcardString;
 
-#define SUPPORTS_VISUALIZE_TEXTURE (WITH_ENGINE && (!(UE_BUILD_SHIPPING || UE_BUILD_TEST) || WITH_EDITOR))
+#define SUPPORTS_VISUALIZE_TEXTURE (WITH_ENGINE && (!UE_BUILD_SHIPPING || WITH_EDITOR))
 
 class RENDERCORE_API FVisualizeTexture : public FRenderResource
 {
