Add owner name to the ensure reporting leaked SM resources as well.

diff --git a/Engine/Source/Runtime/RenderCore/Private/ShaderCodeLibrary.cpp b/Engine/Source/Runtime/RenderCore/Private/ShaderCodeLibrary.cpp
--- a/Engine/Source/Runtime/RenderCore/Private/ShaderCodeLibrary.cpp
+++ b/Engine/Source/Runtime/RenderCore/Private/ShaderCodeLibrary.cpp
@@ -1230,7 +1230,7 @@ void FShaderMapResource_SharedCode::ReleaseRHI()
 
 	if (GetNumRefs() > 0)
 	{
-		ensureMsgf(false, TEXT("FShaderMapResource_SharedCode::ReleaseRHI is still referenced (Num of references %d). Invoking OnSharedShaderMapResourceExplicitRelease delegate."), GetNumRefs());
+		ensureMsgf(false, TEXT("FShaderMapResource_SharedCode::ReleaseRHI is still referenced (Num of references %d, owner %s). Invoking OnSharedShaderMapResourceExplicitRelease delegate."), GetNumRefs(), *GetOwnerName().ToString());
 		UE_LOG(LogShaderLibrary, Warning, TEXT("FShaderMapResource_SharedCode::ReleaseRHI is still referenced (Num of references %d, owner %s). Invoking OnSharedShaderMapResourceExplicitRelease delegate."), 
 			GetNumRefs(), *GetOwnerName().ToString());
 		
