Fix support for ray tracing PSO pre-compilation when using LogPSO workflow

* Remove MaxPayloadSizeInBytes from FPipelineFileCacheRayTracingDesc
* Always use serialized RayTracingPayloadSize from FRHIRayTracingShader
* Maintain binary compatibility with old serialized data

diff --git a/Engine/Source/Editor/UnrealEd/Private/Commandlets/ShaderPipelineCacheToolsCommandlet.cpp b/Engine/Source/Editor/UnrealEd/Private/Commandlets/ShaderPipelineCacheToolsCommandlet.cpp
--- a/Engine/Source/Editor/UnrealEd/Private/Commandlets/ShaderPipelineCacheToolsCommandlet.cpp
+++ b/Engine/Source/Editor/UnrealEd/Private/Commandlets/ShaderPipelineCacheToolsCommandlet.cpp
@@ -934,15 +934,10 @@ bool SaveStablePipelineCacheDeprecated(const FString& OutputFilename, const TArr
 			}
 			else if (Item.PSO->Type == FPipelineCacheFileFormatPSO::DescriptorType::RayTracing)
 			{
-				// Serialize ray tracing PSO state description in backwards-compatible way, reusing graphics PSO fields.
-				// This is only required due to legacy.
-
 				FPipelineCacheFileFormatPSO::GraphicsDescriptor Desc;
 				FMemory::Memzero(Desc);
 
-				// Re-purpose graphics state fields to store RT PSO properties
-				// See corresponding parsing code in ParseStableCSV().
-				Desc.MSAASamples = Item.PSO->RayTracingDesc.MaxPayloadSizeInBytes;
+				// Serialize ray tracing PSO state description in backwards-compatible way, reusing graphics PSO fields. This is only required due to legacy.
 				Desc.DepthStencilFlags = Item.PSO->RayTracingDesc.bAllowHitGroupIndexing ? ETextureCreateFlags::SRGB : ETextureCreateFlags::None;
 
 				PSOLine += FString::Printf(TEXT("\"%s\""), *Desc.StateToString());
@@ -1582,7 +1577,6 @@ static TSet<FPipelineCacheFileFormatPSO> ParseStableCSV(const FString& FileName,
 				PSO.RayTracingDesc.ShaderHash = Match;
 				// See corresponding serialization code in ExpandPSOSC()
 				PSO.RayTracingDesc.Frequency = EShaderFrequency(AdjustedSlotIndex);
-				PSO.RayTracingDesc.MaxPayloadSizeInBytes = PSO.GraphicsDesc.MSAASamples;
 				PSO.RayTracingDesc.bAllowHitGroupIndexing = static_cast<uint64>(PSO.GraphicsDesc.DepthStencilFlags) != 0;
 				break;
 			default:
diff --git a/Engine/Source/Runtime/RHI/Private/PipelineFileCache.cpp b/Engine/Source/Runtime/RHI/Private/PipelineFileCache.cpp
--- a/Engine/Source/Runtime/RHI/Private/PipelineFileCache.cpp
+++ b/Engine/Source/Runtime/RHI/Private/PipelineFileCache.cpp
@@ -1329,7 +1329,9 @@ bool FPipelineCacheFileFormatPSO::Verify() const
 		case FPipelineCacheFileFormatPSO::DescriptorType::RayTracing:
 		{
 			Ar << Info.RayTracingDesc.ShaderHash;
-			Ar << Info.RayTracingDesc.MaxPayloadSizeInBytes;
+
+			uint32 DeprecatedMaxPayloadSizeInBytes = 0; // Not used, kept for binary format compatibility
+			Ar << DeprecatedMaxPayloadSizeInBytes;
 
 			uint32 Frequency = uint32(Info.RayTracingDesc.Frequency);
 			Ar << Frequency;
@@ -4176,7 +4178,6 @@ bool FPipelineFileCacheManager::MergePipelineFileCaches(FString const& PathA, FS
 
 FPipelineCacheFileFormatPSO::FPipelineFileCacheRayTracingDesc::FPipelineFileCacheRayTracingDesc(const FRayTracingPipelineStateInitializer& Initializer, const FRHIRayTracingShader* ShaderRHI)
 : ShaderHash(ShaderRHI->GetHash())
-, MaxPayloadSizeInBytes(Initializer.MaxPayloadSizeInBytes)
 , Frequency(ShaderRHI->GetFrequency())
 , bAllowHitGroupIndexing(Initializer.bAllowHitGroupIndexing)
 {
@@ -4184,14 +4185,15 @@ FPipelineCacheFileFormatPSO::FPipelineFileCacheRayTracingDesc::FPipelineFileCach
 
 FString FPipelineCacheFileFormatPSO::FPipelineFileCacheRayTracingDesc::HeaderLine() const
 {
-	return FString(TEXT("RayTracingShader,MaxPayloadSizeInBytes,Frequency,bAllowHitGroupIndexing"));
+	return FString(TEXT("RayTracingShader,DeprecatedMaxPayloadSizeInBytes,Frequency,bAllowHitGroupIndexing"));
 }
 
 FString FPipelineCacheFileFormatPSO::FPipelineFileCacheRayTracingDesc::ToString() const
 {
+	const uint32 DeprecatedMaxPayloadSizeInBytes = 0; // Not used, but kept for back-compatibility
 	return FString::Printf(TEXT("%s,%d,%d,%d")
 		, *ShaderHash.ToString()
-		, MaxPayloadSizeInBytes
+		, DeprecatedMaxPayloadSizeInBytes
 		, uint32(Frequency)
 		, uint32(bAllowHitGroupIndexing)
 	);
@@ -4216,8 +4218,6 @@ void FPipelineCacheFileFormatPSO::FPipelineFileCacheRayTracingDesc::AddToReadabl
 			break;
 	}
 	OutBuilder << ShaderHash.ToString();
-	OutBuilder << TEXT(" MPSIB ");
-	OutBuilder << MaxPayloadSizeInBytes;
 	OutBuilder << TEXT(" AHGI ");
 	OutBuilder << bAllowHitGroupIndexing;
 }
@@ -4229,7 +4229,8 @@ void FPipelineCacheFileFormatPSO::FPipelineFileCacheRayTracingDesc::FromString(c
 
 	ShaderHash.FromString(Parts[0]);
 
-	LexFromString(MaxPayloadSizeInBytes, Parts[1]);
+	uint32 DeprecatedMaxPayloadSizeInBytes = 0; // Not used, but kept for back-compatibility
+	LexFromString(DeprecatedMaxPayloadSizeInBytes, Parts[1]);
 
 	{
 		uint32 Temp = 0;
diff --git a/Engine/Source/Runtime/RHI/Public/PipelineFileCache.h b/Engine/Source/Runtime/RHI/Public/PipelineFileCache.h
--- a/Engine/Source/Runtime/RHI/Public/PipelineFileCache.h
+++ b/Engine/Source/Runtime/RHI/Public/PipelineFileCache.h
@@ -184,7 +184,6 @@ struct RHI_API FPipelineCacheFileFormatPSO
 	struct RHI_API FPipelineFileCacheRayTracingDesc
 	{
 		FSHAHash ShaderHash;
-		uint32 MaxPayloadSizeInBytes = 0;
 		EShaderFrequency Frequency = SF_RayGen;
 		bool bAllowHitGroupIndexing = true;
 
@@ -199,7 +198,6 @@ struct RHI_API FPipelineCacheFileFormatPSO
 		friend uint32 GetTypeHash(const FPipelineFileCacheRayTracingDesc& Desc)
 		{
 			return GetTypeHash(Desc.ShaderHash) ^
-				GetTypeHash(Desc.MaxPayloadSizeInBytes) ^
 				GetTypeHash(Desc.Frequency) ^
 				GetTypeHash(Desc.bAllowHitGroupIndexing);
 		}
@@ -207,7 +205,6 @@ struct RHI_API FPipelineCacheFileFormatPSO
 		bool operator == (const FPipelineFileCacheRayTracingDesc& Other) const
 		{
 			return ShaderHash == Other.ShaderHash &&
-				MaxPayloadSizeInBytes == Other.MaxPayloadSizeInBytes &&
 				Frequency == Other.Frequency &&
 				bAllowHitGroupIndexing == Other.bAllowHitGroupIndexing;
 		}
diff --git a/Engine/Source/Runtime/RenderCore/Private/PipelineCacheUtilities.cpp b/Engine/Source/Runtime/RenderCore/Private/PipelineCacheUtilities.cpp
--- a/Engine/Source/Runtime/RenderCore/Private/PipelineCacheUtilities.cpp
+++ b/Engine/Source/Runtime/RenderCore/Private/PipelineCacheUtilities.cpp
@@ -809,14 +809,12 @@ bool UE::PipelineCacheUtilities::LoadStablePipelineCacheFile(const FString& File
 		}
 		else if (NewPSO.Type == FPipelineCacheFileFormatPSO::DescriptorType::RayTracing)
 		{
-			// not yet supported
-			++OutPSOsRejected;
-
-			static bool bLogged = false;
-			if (!bLogged)
+			for (const UE::PipelineCacheUtilities::FPermutation& Perm : PermGroup.Permutations)
 			{
-				UE_LOG(LogPipelineCacheUtilities, Display, TEXT("Raytracing PSOs aren't yet supported in the PSO stable cache.\nFilename:%s PSO:%s\nNOTE: Only first rejected PSO is reported."), *Filename, *NewPSO.ToStringReadable());
-				bLogged = true;
+				const EShaderFrequency Frequency = NewPSO.RayTracingDesc.Frequency;
+				NewPSO.RayTracingDesc.ShaderHash = HashesForStableKeys[Perm.Slots[Frequency]];
+
+				AddNewPSO(NewPSO);
 			}
 		}
 	}
diff --git a/Engine/Source/Runtime/RenderCore/Private/ShaderPipelineCache.cpp b/Engine/Source/Runtime/RenderCore/Private/ShaderPipelineCache.cpp
--- a/Engine/Source/Runtime/RenderCore/Private/ShaderPipelineCache.cpp
+++ b/Engine/Source/Runtime/RenderCore/Private/ShaderPipelineCache.cpp
@@ -1091,44 +1091,47 @@ bool FShaderPipelineCacheTask::Precompile(FRHICommandListImmediate& RHICmdList,
 		}
 		else if (FPipelineCacheFileFormatPSO::DescriptorType::RayTracing == PSO.Type)
 		{
-		#if 0 // Workaround for UE-97607:
-			  // If ray tracing PSO file cache is generated using one payload size but later shaders were re-compiled with a different payload declaration 
-			  // it is possible for the wrong size to be used here, which leads to a D3D run-time error when attempting to create the PSO.
-			  // Ray tracing shader pre-compilation is disabled until a robust solution is found.
-			  // Consider setting r.RayTracing.NonBlockingPipelineCreation=1 meanwhile to avoid most of the RTPSO creation stalls during gameplay.
+			// If ray tracing PSO file cache is generated using one payload size but later shaders were re-compiled with a different payload declaration 
+			// it is possible for the wrong size to be used here, which leads to a D3D run-time error when attempting to create the PSO.
+			// Ray tracing shader pre-compilation is disabled until a robust solution is found.
+			// Consider setting r.RayTracing.NonBlockingPipelineCreation=1 meanwhile to avoid most of the RTPSO creation stalls during gameplay.
 			if (IsRayTracingEnabled())
 			{
-				FRayTracingPipelineStateInitializer Initializer;
-				Initializer.bPartial = true;
-				Initializer.MaxPayloadSizeInBytes = PSO.RayTracingDesc.MaxPayloadSizeInBytes;
-				Initializer.bAllowHitGroupIndexing = PSO.RayTracingDesc.bAllowHitGroupIndexing;
-
-				FRHIRayTracingShader* ShaderTable[] =
+				FRayTracingShaderRHIRef RayTracingShader = FShaderCodeLibrary::CreateRayTracingShader(Platform, PSO.RayTracingDesc.ShaderHash, PSO.RayTracingDesc.Frequency);
+				if (RayTracingShader.IsValid())
 				{
-					FShaderCodeLibrary::CreateRayTracingShader(Platform, PSO.RayTracingDesc.ShaderHash, PSO.RayTracingDesc.Frequency)
-				};
+					FRayTracingPipelineStateInitializer Initializer;
+						Initializer.bPartial = true; // Indicates that this RTPSO is used only as input for later RTPSO linking step (not for rendering)
+						Initializer.bAllowHitGroupIndexing = PSO.RayTracingDesc.bAllowHitGroupIndexing;
+						Initializer.MaxPayloadSizeInBytes = RayTracingShader->RayTracingPayloadSize;
 
-				switch (PSO.RayTracingDesc.Frequency)
-				{
-				case SF_RayGen:
-					Initializer.SetRayGenShaderTable(ShaderTable);
-					break;
-				case SF_RayMiss:
-					Initializer.SetMissShaderTable(ShaderTable);
-					break;
-				case SF_RayHitGroup:
-					Initializer.SetHitGroupTable(ShaderTable);
-					break;
-				case SF_RayCallable:
-					Initializer.SetCallableTable(ShaderTable);
-					break;
-				default:
-					checkNoEntry();
+						FRHIRayTracingShader* ShaderTable[] =
+						{
+							RayTracingShader.GetReference()
+						};
+
+						switch (PSO.RayTracingDesc.Frequency)
+						{
+						case SF_RayGen:
+							Initializer.SetRayGenShaderTable(ShaderTable);
+							break;
+						case SF_RayMiss:
+							Initializer.SetMissShaderTable(ShaderTable);
+							break;
+						case SF_RayHitGroup:
+							Initializer.SetHitGroupTable(ShaderTable);
+							break;
+						case SF_RayCallable:
+							Initializer.SetCallableTable(ShaderTable);
+							break;
+						default:
+							checkNoEntry();
+						}
+
+					FRayTracingPipelineState* RayTracingPipeline = PipelineStateCache::GetAndOrCreateRayTracingPipelineState(RHICmdList, Initializer);
+					bOk = RayTracingPipeline != nullptr;
 				}
-				FRayTracingPipelineState* RayTracingPipeline = PipelineStateCache::GetAndOrCreateRayTracingPipelineState(RHICmdList, Initializer);
-				bOk = RayTracingPipeline != nullptr;
 			}
-		#endif // Workaround for UE-97607
 		}
 		else
 		{
diff --git a/Engine/Source/Runtime/Renderer/Private/RayTracing/RayTracingMaterialHitShaders.cpp b/Engine/Source/Runtime/Renderer/Private/RayTracing/RayTracingMaterialHitShaders.cpp
--- a/Engine/Source/Runtime/Renderer/Private/RayTracing/RayTracingMaterialHitShaders.cpp
+++ b/Engine/Source/Runtime/Renderer/Private/RayTracing/RayTracingMaterialHitShaders.cpp
@@ -53,8 +53,8 @@ static FAutoConsoleVariableRef CVarRayTracingNonBlockingPipelineCreation(
 	GRayTracingNonBlockingPipelineCreation,
 	TEXT("Enable background ray tracing pipeline creation, without blocking RHI or Render thread.\n")
 	TEXT("Fallback opaque black material will be used for missing shaders meanwhile.\n")
-	TEXT(" 0: off (default, rendering will always use correct requested material)\n")
-	TEXT(" 1: on (non-blocking mode may sometimes use the fallback opaque black material)\n"),
+	TEXT(" 0: off (rendering will always use correct requested material)\n")
+	TEXT(" 1: on (default, non-blocking mode may sometimes use the fallback opaque black material outside of offline rendering scenarios)\n"),
 	ECVF_RenderThreadSafe);
 
 // CVar defined in DeferredShadingRenderer.cpp
