Make r.Mobile.EarlyZPassOnlyMaterialMasking be per shader platform

diff --git a/Engine/Source/Runtime/RenderCore/Private/RenderUtils.cpp b/Engine/Source/Runtime/RenderCore/Private/RenderUtils.cpp
--- a/Engine/Source/Runtime/RenderCore/Private/RenderUtils.cpp
+++ b/Engine/Source/Runtime/RenderCore/Private/RenderUtils.cpp
@@ -531,11 +531,11 @@ RENDERCORE_API bool SupportsTextureCubeArray(ERHIFeatureLevel::Type FeatureLevel
 
 RENDERCORE_API bool MaskedInEarlyPass(const FStaticShaderPlatform Platform)
 {
-	static IConsoleVariable* CVarMobileEarlyZPassOnlyMaterialMasking = IConsoleManager::Get().FindConsoleVariable(TEXT("r.Mobile.EarlyZPassOnlyMaterialMasking"));
+	static FShaderPlatformCachedIniValue<bool> CVarMobileEarlyZPassOnlyMaterialMasking(TEXT("r.Mobile.EarlyZPassOnlyMaterialMasking"));
 	static IConsoleVariable* CVarEarlyZPassOnlyMaterialMasking = IConsoleManager::Get().FindConsoleVariable(TEXT("r.EarlyZPassOnlyMaterialMasking"));
 	if (IsMobilePlatform(Platform))
 	{
-		return (CVarMobileEarlyZPassOnlyMaterialMasking && CVarMobileEarlyZPassOnlyMaterialMasking->GetInt() != 0);
+		return CVarMobileEarlyZPassOnlyMaterialMasking.Get(Platform) == 1;
 	}
 	else
 	{
diff --git a/Engine/Source/Runtime/Renderer/Private/RendererScene.cpp b/Engine/Source/Runtime/Renderer/Private/RendererScene.cpp
--- a/Engine/Source/Runtime/Renderer/Private/RendererScene.cpp
+++ b/Engine/Source/Runtime/Renderer/Private/RendererScene.cpp
@@ -4722,7 +4722,7 @@ void FScene::GetEarlyZPassMode(ERHIFeatureLevel::Type InFeatureLevel, EDepthDraw
 	{
 		OutZPassMode = DDM_None;
 				 
-		static const bool bMaskedInEarlyPass = MaskedInEarlyPass(ShaderPlatform);
+		const bool bMaskedInEarlyPass = MaskedInEarlyPass(ShaderPlatform);
 		if (bMaskedInEarlyPass)
 		{
 			OutZPassMode = DDM_MaskedOnly;
