[shaders] deprecated FShaderPipelineCompileJob::bFailedRemovingUnused field; this is only used to append some dubiously useful information to a log message (the underlying error causing the failure is always reported as any other shader preprocess/compile error). opting to remove rather than dragging this along with forthcoming preprocessing refactoring.

diff --git a/Engine/Source/Developer/Windows/ShaderFormatD3D/Private/D3DShaderCompiler.cpp b/Engine/Source/Developer/Windows/ShaderFormatD3D/Private/D3DShaderCompiler.cpp
--- a/Engine/Source/Developer/Windows/ShaderFormatD3D/Private/D3DShaderCompiler.cpp
+++ b/Engine/Source/Developer/Windows/ShaderFormatD3D/Private/D3DShaderCompiler.cpp
@@ -880,7 +880,6 @@ static bool CompileAndProcessD3DShaderFXCExt(
 								NewError.StrippedErrorMessage = FString::Printf(TEXT("Second pass had more used attributes (%d) than first pass (%d)"), Output.UsedAttributes.Num(), ShaderInputs.Num());
 								Output = OriginalOutput;
 								Output.Errors.Add(NewError);
-								Output.bFailedRemovingUnused = true;
 								break;
 							}
 
@@ -900,7 +899,6 @@ static bool CompileAndProcessD3DShaderFXCExt(
 									);
 								Output = OriginalOutput;
 								Output.Errors.Add(NewError);
-								Output.bFailedRemovingUnused = true;
 								break;
 							}
 
@@ -918,7 +916,6 @@ static bool CompileAndProcessD3DShaderFXCExt(
 								NewError.StrippedErrorMessage = ErrorMessage;
 								Output.Errors.Add(NewError);
 							}
-							Output.bFailedRemovingUnused = true;
 							break;
 						}
 					}
@@ -1134,7 +1131,6 @@ void CompileD3DShader(const FShaderCompilerInput& Input, FShaderCompilerOutput&
 
 	FString EntryPointName = Input.EntryPointName;
 
-	Output.bFailedRemovingUnused = false;
 	if (Input.Environment.CompilerFlags.Contains(CFLAG_ForceRemoveUnusedInterpolators) && Input.Target.Frequency == SF_Vertex && Input.bCompilingForShaderPipeline)
 	{
 		// Always add SV_Position
@@ -1180,7 +1176,6 @@ void CompileD3DShader(const FShaderCompilerInput& Input, FShaderCompilerOutput&
 				NewError.StrippedErrorMessage = ErrorReport;
 				Output.Errors.Add(NewError);
 			}
-			Output.bFailedRemovingUnused = true;
 		}
 	}
 
diff --git a/Engine/Source/Programs/ShaderCompileWorker/Private/ShaderCompileWorker.cpp b/Engine/Source/Programs/ShaderCompileWorker/Private/ShaderCompileWorker.cpp
--- a/Engine/Source/Programs/ShaderCompileWorker/Private/ShaderCompileWorker.cpp
+++ b/Engine/Source/Programs/ShaderCompileWorker/Private/ShaderCompileWorker.cpp
@@ -754,7 +754,6 @@ private:
 				OutputFile << PipelineNames[JobIndex];
 				bool bSucceeded = (bool)PipelineJob.bSucceeded;
 				OutputFile << bSucceeded;
-				OutputFile << PipelineJob.bFailedRemovingUnused;
 				int32 NumStageJobs = PipelineJob.StageJobs.Num();
 				OutputFile << NumStageJobs;
 
diff --git a/Engine/Source/Runtime/Engine/Private/ShaderCompiler/ShaderCompiler.cpp b/Engine/Source/Runtime/Engine/Private/ShaderCompiler/ShaderCompiler.cpp
--- a/Engine/Source/Runtime/Engine/Private/ShaderCompiler/ShaderCompiler.cpp
+++ b/Engine/Source/Runtime/Engine/Private/ShaderCompiler/ShaderCompiler.cpp
@@ -2114,7 +2114,6 @@ void FShaderCompileUtilities::DoReadTaskResults(const TArray<FShaderCommonCompil
 				bool bSucceeded = false;
 				OutputFile << bSucceeded;
 				CurrentJob->bSucceeded = bSucceeded;
-				OutputFile << CurrentJob->bFailedRemovingUnused;
 				if (PipelineName != CurrentJob->Key.ShaderPipeline->GetName())
 				{
 					OutputFileContext.ModalErrorOrLog(TEXT("Expecting pipeline job \"%s\", got \"%s\" instead!"), CurrentJob->Key.ShaderPipeline->GetName(), *PipelineName);
@@ -2521,10 +2520,6 @@ void FShaderCompileThreadRunnable::PushCompletedJobsToManager()
 					{
 						const FShaderPipelineCompileJob* PipelineJob = Job.GetShaderPipelineJob();
 						JobNames += FString(PipelineJob->Key.ShaderPipeline->GetName());
-						if (PipelineJob->bFailedRemovingUnused)
-						{
-							JobNames += FString(TEXT("(failed to optimize)"));
-						}
 					}
 					if (JobIndex < SortedJobs.Num() - 1)
 					{
diff --git a/Engine/Source/Runtime/RenderCore/Private/ShaderCore.cpp b/Engine/Source/Runtime/RenderCore/Private/ShaderCore.cpp
--- a/Engine/Source/Runtime/RenderCore/Private/ShaderCore.cpp
+++ b/Engine/Source/Runtime/RenderCore/Private/ShaderCore.cpp
@@ -1030,8 +1030,6 @@ void CompileShaderPipeline(const TArray<const IShaderFormat*>& ShaderFormats, FS
 
 	// This tells the shader compiler we do want to remove unused outputs
 	bool bEnableRemovingUnused = true;
-	// This tells us the hlsl parser failed at removing unused outputs
-	bool bJobFailedRemovingUnused = false;
 
 	//#todo-rco: Currently only removes for pure VS & PS stages
 	for (int32 Index = 0; Index < PipelineJob->StageJobs.Num(); ++Index)
@@ -1065,12 +1063,9 @@ void CompileShaderPipeline(const TArray<const IShaderFormat*>& ShaderFormats, FS
 			// Can't carry on compiling the pipeline
 			return;
 		}
-
-		bJobFailedRemovingUnused = CurrentJob->Output.bFailedRemovingUnused || bJobFailedRemovingUnused;
 	}
 
 	PipelineJob->bSucceeded = true;
-	PipelineJob->bFailedRemovingUnused = bJobFailedRemovingUnused;
 }
 
 /**
diff --git a/Engine/Source/Runtime/RenderCore/Public/ShaderCompilerCore.h b/Engine/Source/Runtime/RenderCore/Public/ShaderCompilerCore.h
--- a/Engine/Source/Runtime/RenderCore/Public/ShaderCompilerCore.h
+++ b/Engine/Source/Runtime/RenderCore/Public/ShaderCompilerCore.h
@@ -23,7 +23,7 @@ class FShaderPipelineCompileJob;
 // this is for the protocol, not the data, bump if FShaderCompilerInput or ProcessInputFromArchive changes.
 inline const int32 ShaderCompileWorkerInputVersion = 18;
 // this is for the protocol, not the data, bump if FShaderCompilerOutput or WriteToOutputArchive changes.
-inline const int32 ShaderCompileWorkerOutputVersion = 12;
+inline const int32 ShaderCompileWorkerOutputVersion = 13;
 // this is for the protocol, not the data.
 inline const int32 ShaderCompileWorkerSingleJobHeader = 'S';
 // this is for the protocol, not the data.
@@ -468,7 +468,6 @@ struct FShaderCompilerOutput
 	,	CompileTime(0.0)
 	,	PreprocessTime(0.0)
 	,	bSucceeded(false)
-	,	bFailedRemovingUnused(false)
 	,	bSupportsQueryingUsedAttributes(false)
 	,	bUsedHLSLccCompiler(false)
 	{
@@ -485,11 +484,22 @@ struct FShaderCompilerOutput
 	double CompileTime;
 	double PreprocessTime;
 	bool bSucceeded;
+	UE_DEPRECATED(5.3, "bFailedRemovingUnused field is no longer used")
 	bool bFailedRemovingUnused;
 	bool bSupportsQueryingUsedAttributes;
 	bool bUsedHLSLccCompiler;
 	TArray<FString> UsedAttributes;
 
+	PRAGMA_DISABLE_DEPRECATION_WARNINGS
+	// Explicitly-defaulted copy/move ctors & assignment operators are needed temporarily due to 
+	// deprecation of bFailedRemovingUnused field. These can be removed once the deprecation
+	// window for said field ends.
+	FShaderCompilerOutput(FShaderCompilerOutput&&) = default;
+	FShaderCompilerOutput(const FShaderCompilerOutput&) = default;
+	FShaderCompilerOutput& operator=(FShaderCompilerOutput&&) = default;
+	FShaderCompilerOutput& operator=(const FShaderCompilerOutput&) = default;
+	PRAGMA_ENABLE_DEPRECATION_WARNINGS
+
 	TArray<FShaderCodeValidationStride> ParametersStrideToValidate;
 
 	FString OptionalFinalShaderSource;
@@ -509,7 +519,7 @@ struct FShaderCompilerOutput
 	{
 		// Note: this serialize is used to pass between UE and the shader compile worker, recompile both when modifying
 		Ar << Output.ParameterMap << Output.Errors << Output.Target << Output.ShaderCode << Output.OutputHash << Output.NumInstructions << Output.NumTextureSamplers << Output.bSucceeded;
-		Ar << Output.bFailedRemovingUnused << Output.bSupportsQueryingUsedAttributes << Output.UsedAttributes;
+		Ar << Output.bSupportsQueryingUsedAttributes << Output.UsedAttributes;
 		Ar << Output.CompileTime;
 		Ar << Output.PreprocessTime;
 		Ar << Output.OptionalFinalShaderSource;
diff --git a/Engine/Source/Runtime/RenderCore/Public/ShaderCompilerJobTypes.h b/Engine/Source/Runtime/RenderCore/Public/ShaderCompilerJobTypes.h
--- a/Engine/Source/Runtime/RenderCore/Public/ShaderCompilerJobTypes.h
+++ b/Engine/Source/Runtime/RenderCore/Public/ShaderCompilerJobTypes.h
@@ -5,6 +5,7 @@
 #include "HAL/PlatformTime.h"
 #include "Hash/Blake3.h"
 #include "Shader.h"
+#include "ShaderCompilerCore.h"
 #include "VertexFactory.h"
 #include "Templates/Function.h"
 
@@ -244,6 +245,7 @@ public:
 
 	FShaderPipelineCompileJobKey Key;
 	TArray<TRefCountPtr<FShaderCompileJob>> StageJobs;
+	UE_DEPRECATED(5.3, "bFailedRemovingUnused field is no longer used")
 	bool bFailedRemovingUnused;
 
 	virtual RENDERCORE_API FInputHash GetInputHash() override;
