Strengthen assertions in CrossCompilerCommon.h for unpacking shader headers.

#rnx

diff --git a/Engine/Source/Runtime/RenderCore/Public/CrossCompilerCommon.h b/Engine/Source/Runtime/RenderCore/Public/CrossCompilerCommon.h
--- a/Engine/Source/Runtime/RenderCore/Public/CrossCompilerCommon.h
+++ b/Engine/Source/Runtime/RenderCore/Public/CrossCompilerCommon.h
@@ -84,7 +84,7 @@ namespace CrossCompiler
 		case SHADER_STAGE_COMPUTE:	return 'c';
 		default: break;
 		}
-		check(0);
+		checkf(0, TEXT("invalid value: ShaderStageIndexToTypeName(ShaderStage = '%c')"), (ANSICHAR)ShaderStage);
 		return 0;
 	}
 
@@ -101,7 +101,7 @@ namespace CrossCompiler
 		case PACKED_TYPEINDEX_IMAGE:	return PACKED_TYPENAME_IMAGE;
 		default: break;
 		}
-		check(0);
+		checkf(0, TEXT("invalid value: PackedTypeIndexToTypeName(ArrayType = %d)"), (int32)ArrayType);
 		return 0;
 	}
 
@@ -118,7 +118,7 @@ namespace CrossCompiler
 		case PACKED_TYPENAME_IMAGE:		return PACKED_TYPEINDEX_IMAGE;
 		default: break;
 		}
-		check(0);
+		checkf(0, TEXT("invalid value: PackedTypeNameToTypeIndex(ArrayName = %d)"), (int32)ArrayName);
 		return 0;
 	}
 
