ResizeStructuredBufferSOAIfNeeded() - Fixed string formatting, previously tried to print uints using %s

diff --git a/Engine/Source/Runtime/RenderCore/Private/UnifiedBuffer.cpp b/Engine/Source/Runtime/RenderCore/Private/UnifiedBuffer.cpp
--- a/Engine/Source/Runtime/RenderCore/Private/UnifiedBuffer.cpp
+++ b/Engine/Source/Runtime/RenderCore/Private/UnifiedBuffer.cpp
@@ -437,14 +437,14 @@ FRDGBuffer* ResizeStructuredBufferSOAIfNeeded(FRDGBuilder& GraphBuilder, TRefCou
 	const uint32 BytesPerElement = 16;
 	const uint32 ExternalBufferSize = TryGetSize(ExternalBuffer);
 
-	checkf(Params.NumBytes % BytesPerElement == 0, TEXT("NumBytes (%s) must be a multiple of BytesPerElement (%s)"), Params.NumBytes, BytesPerElement);
-	checkf(ExternalBufferSize % BytesPerElement == 0, TEXT("NumBytes (%s) must be a multiple of BytesPerElement (%s)"), ExternalBufferSize, BytesPerElement);
+	checkf(Params.NumBytes % BytesPerElement == 0, TEXT("NumBytes (%u) must be a multiple of BytesPerElement (%u)"), Params.NumBytes, BytesPerElement);
+	checkf(ExternalBufferSize % BytesPerElement == 0, TEXT("NumBytes (%u) must be a multiple of BytesPerElement (%u)"), ExternalBufferSize, BytesPerElement);
 
 	uint32 NumElements = Params.NumBytes / BytesPerElement;
 	uint32 NumElementsOld = ExternalBufferSize / BytesPerElement;
 
-	checkf(NumElements % Params.NumArrays == 0, TEXT("NumElements (%s) must be a multiple of NumArrays (%s)"), NumElements, Params.NumArrays);
-	checkf(NumElementsOld % Params.NumArrays == 0, TEXT("NumElements (%s) must be a multiple of NumArrays (%s)"), NumElementsOld, Params.NumArrays);
+	checkf(NumElements % Params.NumArrays == 0, TEXT("NumElements (%u) must be a multiple of NumArrays (%u)"), NumElements, Params.NumArrays);
+	checkf(NumElementsOld % Params.NumArrays == 0, TEXT("NumElements (%u) must be a multiple of NumArrays (%u)"), NumElementsOld, Params.NumArrays);
 
 	const FRDGBufferDesc BufferDesc = FRDGBufferDesc::CreateStructuredDesc(BytesPerElement, NumElements);
 
