Extend support for persistent buffer management to structured buffers of element size other than float4 & add utilities to manage typed persistent structured buffers.
- also added a hook to provide a execute-time callback for indirect compute shader passes (similar to what the group count callback facilitates for direct dispatches).

diff --git a/Engine/Shaders/Private/ByteBuffer.usf b/Engine/Shaders/Private/ByteBuffer.usf
--- a/Engine/Shaders/Private/ByteBuffer.usf
+++ b/Engine/Shaders/Private/ByteBuffer.usf
@@ -3,11 +3,16 @@
 #include "Common.ush"
 
 #define RESOURCE_TYPE_FLOAT4_BUFFER				(0)
-#define RESOURCE_TYPE_FLOAT4_STRUCTURED_BUFFER	(1)
+#define RESOURCE_TYPE_STRUCTURED_BUFFER			(1)
 #define RESOURCE_TYPE_UINT_BUFFER				(2)
 #define RESOURCE_TYPE_UINT4_ALIGNED_BUFFER		(3)
 #define RESOURCE_TYPE_FLOAT4_TEXTURE			(4)
 
+#define STRUCTURED_ELEMENT_SIZE_UINT1 (0)
+#define STRUCTURED_ELEMENT_SIZE_UINT2 (1)
+#define STRUCTURED_ELEMENT_SIZE_UINT4 (2)
+#define STRUCTURED_ELEMENT_SIZE_UINT8 (3)
+
 uint Value;
 uint Size;
 uint Float4sPerLine;
@@ -15,20 +20,65 @@ uint NumScatters;
 uint SrcOffset;
 uint DstOffset;
 
+struct FUint8
+{
+	uint Value0;
+	uint Value1;
+	uint Value2;
+	uint Value3;
+	uint Value4;
+	uint Value5;
+	uint Value6;
+	uint Value7;
+};
+
+#if STRUCTURED_ELEMENT_SIZE_UINT1 == STRUCTURED_ELEMENT_SIZE
+	#define STRUCTURED_BUFFER_ACCESS(BufferName, Index) BufferName##1x[(Index)]
+	#define STRUCTURED_VALUE Value
+#elif STRUCTURED_ELEMENT_SIZE_UINT2 == STRUCTURED_ELEMENT_SIZE
+	#define STRUCTURED_BUFFER_ACCESS(BufferName, Index) BufferName##2x[(Index)]
+	#define STRUCTURED_VALUE uint2(Value, Value)
+#elif STRUCTURED_ELEMENT_SIZE_UINT4 == STRUCTURED_ELEMENT_SIZE
+	#define STRUCTURED_BUFFER_ACCESS(BufferName, Index) BufferName##4x[(Index)]
+	#define STRUCTURED_VALUE uint4(Value,Value,Value,Value)
+#elif STRUCTURED_ELEMENT_SIZE_UINT8 == STRUCTURED_ELEMENT_SIZE
+	#define STRUCTURED_BUFFER_ACCESS(BufferName, Index) BufferName##8x[(Index)]
+	FUint8 MakeUint8(uint InValue)
+	{
+		const FUint8 Uint8Value = {InValue,InValue,InValue,InValue,InValue,InValue,InValue,InValue };
+		return Uint8Value;
+	}
+	#define STRUCTURED_VALUE MakeUint8(Value)
+#endif
+
 #if RESOURCE_TYPE == RESOURCE_TYPE_FLOAT4_BUFFER
 	StructuredBuffer<uint>			ScatterStructuredBuffer;
 	Buffer<float4>					SrcBuffer;
-	StructuredBuffer<float4>		UploadStructuredBuffer;
+	StructuredBuffer<float4>		UploadStructuredBuffer4x;
 	RWBuffer<float4>				DstBuffer;
-#elif RESOURCE_TYPE == RESOURCE_TYPE_FLOAT4_STRUCTURED_BUFFER
+#elif RESOURCE_TYPE == RESOURCE_TYPE_STRUCTURED_BUFFER
+
 	StructuredBuffer<uint>			ScatterStructuredBuffer;
-	StructuredBuffer<float4>		SrcStructuredBuffer;
-	StructuredBuffer<float4>		UploadStructuredBuffer;
-	RWStructuredBuffer<float4>		DstStructuredBuffer;
+
+	StructuredBuffer<uint>		UploadStructuredBuffer1x;
+	StructuredBuffer<uint2>		UploadStructuredBuffer2x;
+	StructuredBuffer<uint4>		UploadStructuredBuffer4x;
+	StructuredBuffer<FUint8>		UploadStructuredBuffer8x;
+
+	StructuredBuffer<uint>		SrcStructuredBuffer1x;
+	StructuredBuffer<uint2>		SrcStructuredBuffer2x;
+	StructuredBuffer<uint4>		SrcStructuredBuffer4x;
+	StructuredBuffer<FUint8>		SrcStructuredBuffer8x;
+
+	RWStructuredBuffer<uint>	DstStructuredBuffer1x;
+	RWStructuredBuffer<uint2>	DstStructuredBuffer2x;
+	RWStructuredBuffer<uint4>	DstStructuredBuffer4x;
+	RWStructuredBuffer<FUint8>	DstStructuredBuffer8x;
+	
 #elif RESOURCE_TYPE == RESOURCE_TYPE_FLOAT4_TEXTURE
 	StructuredBuffer<uint>			ScatterStructuredBuffer;
 	Texture2D<float4>				SrcTexture;
-	StructuredBuffer<float4>		UploadStructuredBuffer;
+	StructuredBuffer<float4>		UploadStructuredBuffer4x;
 	RWTexture2D<float4>				DstTexture;
 #elif RESOURCE_TYPE == RESOURCE_TYPE_UINT_BUFFER || RESOURCE_TYPE == RESOURCE_TYPE_UINT4_ALIGNED_BUFFER
 	ByteAddressBuffer				ScatterByteAddressBuffer;
@@ -46,11 +96,11 @@ void MemsetBufferCS( uint ThreadId : SV_DispatchThreadID )
 	{
 		DstBuffer[DstOffset + ThreadId] = asfloat(Value);
 	}
-#elif RESOURCE_TYPE == RESOURCE_TYPE_FLOAT4_STRUCTURED_BUFFER
+#elif RESOURCE_TYPE == RESOURCE_TYPE_STRUCTURED_BUFFER
 	// Size is in float4s
 	if( ThreadId < Size )
 	{
-		DstStructuredBuffer[ DstOffset + ThreadId ] = asfloat( Value );
+		STRUCTURED_BUFFER_ACCESS(DstStructuredBuffer, DstOffset + ThreadId) = STRUCTURED_VALUE;
 	}
 #elif RESOURCE_TYPE == RESOURCE_TYPE_UINT_BUFFER || RESOURCE_TYPE == RESOURCE_TYPE_UINT4_ALIGNED_BUFFER
 	// Size and offset is in dwords
@@ -95,11 +145,11 @@ void MemcpyCS( uint ThreadId : SV_DispatchThreadID )
 	{
 		DstBuffer[DstOffset + ThreadId] = SrcBuffer[SrcOffset + ThreadId];
 	}
-#elif RESOURCE_TYPE == RESOURCE_TYPE_FLOAT4_STRUCTURED_BUFFER
+#elif RESOURCE_TYPE == RESOURCE_TYPE_STRUCTURED_BUFFER
 	// Size is in float4s
 	if( ThreadId < Size )
 	{
-		DstStructuredBuffer[ DstOffset + ThreadId ] = SrcStructuredBuffer[ SrcOffset + ThreadId ];
+		STRUCTURED_BUFFER_ACCESS(DstStructuredBuffer, DstOffset + ThreadId) = STRUCTURED_BUFFER_ACCESS(SrcStructuredBuffer, SrcOffset + ThreadId);
 	}
 #elif RESOURCE_TYPE == RESOURCE_TYPE_FLOAT4_TEXTURE
 	uint2 IndexTexture;
@@ -154,11 +204,11 @@ void ScatterCopyCS( uint DispatchThreadId : SV_DispatchThreadID )
 #if RESOURCE_TYPE == RESOURCE_TYPE_FLOAT4_BUFFER
 		uint DstIndex = ScatterStructuredBuffer[ScatterIndex] * Size + ScatterOffset;
 		uint SrcIndex = ThreadId;
-		DstBuffer[DstIndex] = UploadStructuredBuffer[SrcIndex];
-#elif RESOURCE_TYPE == RESOURCE_TYPE_FLOAT4_STRUCTURED_BUFFER
+		DstBuffer[DstIndex] = UploadStructuredBuffer4x[SrcIndex];
+#elif RESOURCE_TYPE == RESOURCE_TYPE_STRUCTURED_BUFFER
 		uint DstIndex = ScatterStructuredBuffer[ ScatterIndex ] * Size + ScatterOffset;
 		uint SrcIndex = ThreadId;
-		DstStructuredBuffer[ DstIndex ] = UploadStructuredBuffer[ SrcIndex ];
+		STRUCTURED_BUFFER_ACCESS(DstStructuredBuffer, DstIndex) = STRUCTURED_BUFFER_ACCESS(UploadStructuredBuffer, SrcIndex);
 #elif RESOURCE_TYPE == RESOURCE_TYPE_UINT4_ALIGNED_BUFFER
 		uint DstIndex = ScatterByteAddressBuffer.Load( ScatterIndex * 4 ) * Size + ScatterOffset;
 		uint SrcIndex = ThreadId;
@@ -171,7 +221,7 @@ void ScatterCopyCS( uint DispatchThreadId : SV_DispatchThreadID )
 		IndexTexture.y = DstIndex / (Float4sPerLine);
 		IndexTexture.x = DstIndex % (Float4sPerLine);
 
-		float4 srvResourceVal = UploadStructuredBuffer[SrcIndex];
+		float4 srvResourceVal = asfloat(UploadStructuredBuffer4x[SrcIndex]);
 		DstTexture[IndexTexture.xy] = srvResourceVal;
 #elif RESOURCE_TYPE == RESOURCE_TYPE_UINT_BUFFER
 		uint DstIndex = ScatterByteAddressBuffer.Load( ScatterIndex * 4 ) * Size + ScatterOffset;
diff --git a/Engine/Source/Runtime/RenderCore/Private/UnifiedBuffer.cpp b/Engine/Source/Runtime/RenderCore/Private/UnifiedBuffer.cpp
--- a/Engine/Source/Runtime/RenderCore/Private/UnifiedBuffer.cpp
+++ b/Engine/Source/Runtime/RenderCore/Private/UnifiedBuffer.cpp
@@ -19,13 +19,28 @@ static uint64 GetMaxUploadBufferElements()
 enum class EByteBufferResourceType
 {
 	Float4_Buffer,
-	Float4_StructuredBuffer,
+	StructuredBuffer,
 	Uint_Buffer,
 	Uint4Aligned_Buffer,
 	Float4_Texture,
 	Count
 };
 
+enum class EByteBufferStructuredSize
+{
+	Uint1,
+	Uint2,
+	Uint4,
+	Uint8,
+	Count
+};
+
+// placeholder struct, not really used (on the host side)
+struct FUint8
+{
+	uint32 Values[8];
+};
+
 class FByteBufferShader : public FGlobalShader
 {
 	DECLARE_INLINE_TYPE_LAYOUT(FByteBufferShader, NonVirtual);
@@ -36,8 +51,9 @@ class FByteBufferShader : public FGlobalShader
 	{}
 
 	class ResourceTypeDim : SHADER_PERMUTATION_INT("RESOURCE_TYPE", (int)EByteBufferResourceType::Count);
+	class StructuredElementSizeDim : SHADER_PERMUTATION_INT("STRUCTURED_ELEMENT_SIZE", (int)EByteBufferStructuredSize::Count);
 
-	using FPermutationDomain = TShaderPermutationDomain<ResourceTypeDim>;
+	using FPermutationDomain = TShaderPermutationDomain<ResourceTypeDim, StructuredElementSizeDim>;
 
 	static bool ShouldCompilePermutation( const FGlobalShaderPermutationParameters& Parameters )
 	{
@@ -53,6 +69,12 @@ class FByteBufferShader : public FGlobalShader
 		{
 			return true;
 		}
+
+		// Don't compile structured buffer size variations unless we need them
+		if (ResourceType != EByteBufferResourceType::StructuredBuffer && static_cast<EByteBufferStructuredSize>(PermutationVector.Get<StructuredElementSizeDim>()) != EByteBufferStructuredSize::Uint4)
+		{
+			return false;
+		}
 	}
 
 	BEGIN_SHADER_PARAMETER_STRUCT( FParameters, )
@@ -62,7 +84,7 @@ class FByteBufferShader : public FGlobalShader
 		SHADER_PARAMETER(uint32, DstOffset)
 		SHADER_PARAMETER(uint32, Float4sPerLine)
 		SHADER_PARAMETER_UAV(RWBuffer<float4>, DstBuffer)
-		SHADER_PARAMETER_UAV(RWStructuredBuffer<float4>, DstStructuredBuffer)
+		SHADER_PARAMETER_UAV(RWStructuredBuffer<float4>, DstStructuredBuffer4x)
 		SHADER_PARAMETER_UAV(RWByteAddressBuffer, DstByteAddressBuffer)
 		SHADER_PARAMETER_UAV(RWTexture2D<float4>, DstTexture)
 	END_SHADER_PARAMETER_STRUCT()
@@ -72,6 +94,17 @@ class FMemsetBufferCS : public FByteBufferShader
 {
 	DECLARE_GLOBAL_SHADER( FMemsetBufferCS );
 	SHADER_USE_PARAMETER_STRUCT( FMemsetBufferCS, FByteBufferShader );
+
+	static bool ShouldCompilePermutation( const FGlobalShaderPermutationParameters& Parameters )
+	{
+		FPermutationDomain PermutationVector( Parameters.PermutationId );
+		// Don't compile structured buffer size variations
+		if (static_cast<EByteBufferStructuredSize>(PermutationVector.Get<StructuredElementSizeDim>()) != EByteBufferStructuredSize::Uint4)
+		{
+			return false;
+		}
+		return FByteBufferShader::ShouldCompilePermutation(Parameters);
+	}
 };
 IMPLEMENT_GLOBAL_SHADER( FMemsetBufferCS, "/Engine/Private/ByteBuffer.usf", "MemsetBufferCS", SF_Compute );
 
@@ -80,10 +113,21 @@ class FMemcpyCS : public FByteBufferShader
 	DECLARE_GLOBAL_SHADER( FMemcpyCS );
 	SHADER_USE_PARAMETER_STRUCT( FMemcpyCS, FByteBufferShader );
 
+	static bool ShouldCompilePermutation( const FGlobalShaderPermutationParameters& Parameters )
+	{
+		FPermutationDomain PermutationVector( Parameters.PermutationId );
+		// Don't compile structured buffer size variations
+		if (static_cast<EByteBufferStructuredSize>(PermutationVector.Get<StructuredElementSizeDim>()) != EByteBufferStructuredSize::Uint4)
+		{
+			return false;
+		}
+		return FByteBufferShader::ShouldCompilePermutation(Parameters);
+	}
+
 	BEGIN_SHADER_PARAMETER_STRUCT(FParameters, )
 		SHADER_PARAMETER_STRUCT_INCLUDE(FByteBufferShader::FParameters, Common)
 		SHADER_PARAMETER_SRV(Buffer<float4>, SrcBuffer)
-		SHADER_PARAMETER_SRV(StructuredBuffer<float4>, SrcStructuredBuffer)
+		SHADER_PARAMETER_SRV(StructuredBuffer<float4>, SrcStructuredBuffer4x)
 		SHADER_PARAMETER_SRV(ByteAddressBuffer, SrcByteAddressBuffer)
 		SHADER_PARAMETER_SRV(Texture2D<float4>, SrcTexture)
 	END_SHADER_PARAMETER_STRUCT()
@@ -95,11 +139,22 @@ class FScatterCopyCS : public FByteBufferShader
 	DECLARE_GLOBAL_SHADER( FScatterCopyCS );
 	SHADER_USE_PARAMETER_STRUCT( FScatterCopyCS, FByteBufferShader );
 
+	static bool ShouldCompilePermutation( const FGlobalShaderPermutationParameters& Parameters )
+	{
+		FPermutationDomain PermutationVector( Parameters.PermutationId );
+		// Don't compile structured buffer size variations
+		if (static_cast<EByteBufferStructuredSize>(PermutationVector.Get<StructuredElementSizeDim>()) != EByteBufferStructuredSize::Uint4)
+		{
+			return false;
+		}
+		return FByteBufferShader::ShouldCompilePermutation(Parameters);
+	}
+
 	BEGIN_SHADER_PARAMETER_STRUCT(FParameters, )
 		SHADER_PARAMETER_STRUCT_INCLUDE(FByteBufferShader::FParameters, Common)
 		SHADER_PARAMETER(uint32, NumScatters)
 		SHADER_PARAMETER_SRV(ByteAddressBuffer, UploadByteAddressBuffer)
-		SHADER_PARAMETER_SRV(StructuredBuffer<float4>, UploadStructuredBuffer)
+		SHADER_PARAMETER_SRV(StructuredBuffer<float4>, UploadStructuredBuffer4x)
 		SHADER_PARAMETER_SRV(ByteAddressBuffer, ScatterByteAddressBuffer)
 		SHADER_PARAMETER_SRV(StructuredBuffer<uint>, ScatterStructuredBuffer)
 	END_SHADER_PARAMETER_STRUCT()
@@ -158,8 +213,9 @@ class FRDGByteBufferShader : public FGlobalShader
 		: FGlobalShader(Initializer)
 	{}
 
-	class ResourceTypeDim : SHADER_PERMUTATION_INT("RESOURCE_TYPE", (int)EByteBufferResourceType::Count);
-	using FPermutationDomain = TShaderPermutationDomain<ResourceTypeDim>;
+	using ResourceTypeDim = FByteBufferShader::ResourceTypeDim;
+	using StructuredElementSizeDim = FByteBufferShader::StructuredElementSizeDim;
+	using FPermutationDomain = FByteBufferShader::FPermutationDomain;
 
 	static bool ShouldCompilePermutation(const FGlobalShaderPermutationParameters& Parameters)
 	{
@@ -173,7 +229,10 @@ class FRDGByteBufferShader : public FGlobalShader
 		SHADER_PARAMETER(uint32, DstOffset)
 		SHADER_PARAMETER(uint32, Float4sPerLine)
 		SHADER_PARAMETER_RDG_BUFFER_UAV(RWBuffer<float4>, DstBuffer)
-		SHADER_PARAMETER_RDG_BUFFER_UAV(RWStructuredBuffer<float4>, DstStructuredBuffer)
+		SHADER_PARAMETER_RDG_BUFFER_UAV(RWStructuredBuffer<FUint8>, DstStructuredBuffer8x)
+		SHADER_PARAMETER_RDG_BUFFER_UAV(RWStructuredBuffer<uint4>, DstStructuredBuffer4x)
+		SHADER_PARAMETER_RDG_BUFFER_UAV(RWStructuredBuffer<uint2>, DstStructuredBuffer2x)
+		SHADER_PARAMETER_RDG_BUFFER_UAV(RWStructuredBuffer<uint>, DstStructuredBuffer1x)
 		SHADER_PARAMETER_RDG_BUFFER_UAV(RWByteAddressBuffer, DstByteAddressBuffer)
 		SHADER_PARAMETER_RDG_TEXTURE_UAV(RWTexture2D<float4>, DstTexture)
 	END_SHADER_PARAMETER_STRUCT()
@@ -194,7 +253,10 @@ class FRDGMemcpyCS : public FRDGByteBufferShader
 	BEGIN_SHADER_PARAMETER_STRUCT(FParameters, )
 		SHADER_PARAMETER_STRUCT_INCLUDE(FRDGByteBufferShader::FParameters, Common)
 		SHADER_PARAMETER_RDG_BUFFER_SRV(Buffer<float4>, SrcBuffer)
-		SHADER_PARAMETER_RDG_BUFFER_SRV(StructuredBuffer<float4>, SrcStructuredBuffer)
+		SHADER_PARAMETER_RDG_BUFFER_SRV(StructuredBuffer<FUint8>, SrcStructuredBuffer8x)
+		SHADER_PARAMETER_RDG_BUFFER_SRV(StructuredBuffer<uint4>, SrcStructuredBuffer4x)
+		SHADER_PARAMETER_RDG_BUFFER_SRV(StructuredBuffer<uint2>, SrcStructuredBuffer2x)
+		SHADER_PARAMETER_RDG_BUFFER_SRV(StructuredBuffer<uint>, SrcStructuredBuffer1x)
 		SHADER_PARAMETER_RDG_BUFFER_SRV(ByteAddressBuffer, SrcByteAddressBuffer)
 		SHADER_PARAMETER_RDG_TEXTURE_SRV(Texture2D<float4>, SrcTexture)
 	END_SHADER_PARAMETER_STRUCT()
@@ -210,7 +272,10 @@ class FRDGScatterCopyCS : public FRDGByteBufferShader
 		SHADER_PARAMETER_STRUCT_INCLUDE(FRDGByteBufferShader::FParameters, Common)
 		SHADER_PARAMETER(uint32, NumScatters)
 		SHADER_PARAMETER_RDG_BUFFER_SRV(ByteAddressBuffer, UploadByteAddressBuffer)
-		SHADER_PARAMETER_RDG_BUFFER_SRV(StructuredBuffer<float4>, UploadStructuredBuffer)
+		SHADER_PARAMETER_RDG_BUFFER_SRV(StructuredBuffer<FUint8>, UploadStructuredBuffer8x)
+		SHADER_PARAMETER_RDG_BUFFER_SRV(StructuredBuffer<uint4>, UploadStructuredBuffer4x)
+		SHADER_PARAMETER_RDG_BUFFER_SRV(StructuredBuffer<uint2>, UploadStructuredBuffer2x)
+		SHADER_PARAMETER_RDG_BUFFER_SRV(StructuredBuffer<uint>, UploadStructuredBuffer1x)
 		SHADER_PARAMETER_RDG_BUFFER_SRV(ByteAddressBuffer, ScatterByteAddressBuffer)
 		SHADER_PARAMETER_RDG_BUFFER_SRV(StructuredBuffer<uint>, ScatterStructuredBuffer)
 	END_SHADER_PARAMETER_STRUCT()
@@ -259,12 +324,24 @@ void MemcpyResource(FRDGBuilder& GraphBuilder, FRDGBuffer* DstResource, FRDGBuff
 	MemcpyResource(GraphBuilder, GraphBuilder.CreateUAV(DstResource, ERDGUnorderedAccessViewFlags::SkipBarrier), GraphBuilder.CreateSRV(SrcResource), Params);
 }
 
+static EByteBufferStructuredSize GetStructuredBufferElementSize(FRDGBufferUAV * RDGBufferUAV)
+{
+	int32 BytesPerElement = RDGBufferUAV->Desc.Buffer->Desc.BytesPerElement;
+	uint32 Log2NumElements = FMath::FloorLog2(BytesPerElement / 4);
+	checkf((BytesPerElement % 4) == 0 && FMath::IsPowerOfTwo(BytesPerElement / 4) && Log2NumElements < uint32(EByteBufferStructuredSize::Count), TEXT("Unsupported structured buffer BytesPerElement size (%d) for buffer '%s' (supported sizes are 4,8,16,32)."), BytesPerElement, RDGBufferUAV->Name);
+
+	return static_cast<EByteBufferStructuredSize>(Log2NumElements);
+
+}
+
 void MemsetResource(FRDGBuilder& GraphBuilder, FRDGUnorderedAccessView* UAV, const FMemsetResourceParams& Params)
 {
 	check(UAV);
 	FRDGViewableResource* Resource = UAV->GetParent();
 
 	EByteBufferResourceType ResourceTypeEnum = EByteBufferResourceType::Count;
+	// This is only used for structured buffers, since that is where we must match the format specified when it is created / used or whatever the platform happens to care about.
+	EByteBufferStructuredSize ByteBufferStructuredSize = EByteBufferStructuredSize::Uint4;
 
 	auto* PassParameters = GraphBuilder.AllocParameters<FRDGMemsetBufferCS::FParameters>();
 	PassParameters->Value = Params.Value;
@@ -286,9 +363,29 @@ void MemsetResource(FRDGBuilder& GraphBuilder, FRDGUnorderedAccessView* UAV, con
 		PassParameters->DstBuffer = GetAs<FRDGBufferUAV>(UAV);
 		break;
 	case EResourceType::STRUCTURED_BUFFER:
-		ResourceTypeEnum = EByteBufferResourceType::Float4_StructuredBuffer;
-		PassParameters->DstStructuredBuffer = GetAs<FRDGBufferUAV>(UAV);
+	{
+		ResourceTypeEnum = EByteBufferResourceType::StructuredBuffer;
+		FRDGBufferUAV * RDGBufferUAV = GetAs<FRDGBufferUAV>(UAV);
+		ByteBufferStructuredSize = GetStructuredBufferElementSize(RDGBufferUAV);
+		switch(ByteBufferStructuredSize)
+		{
+		case EByteBufferStructuredSize::Uint1:
+			PassParameters->DstStructuredBuffer1x = RDGBufferUAV;
+			break;
+		case EByteBufferStructuredSize::Uint2:
+			PassParameters->DstStructuredBuffer2x = RDGBufferUAV;
+			break;
+		case EByteBufferStructuredSize::Uint4:
+			PassParameters->DstStructuredBuffer4x = RDGBufferUAV;
+			break;
+		case EByteBufferStructuredSize::Uint8:
+			PassParameters->DstStructuredBuffer8x = RDGBufferUAV;
+			break;
+		default:
+			break;
+		};
 		break;
+	}
 	case EResourceType::TEXTURE:
 		ResourceTypeEnum = EByteBufferResourceType::Float4_Texture;
 		PassParameters->DstTexture = GetAs<FRDGTextureUAV>(UAV);
@@ -300,6 +397,7 @@ void MemsetResource(FRDGBuilder& GraphBuilder, FRDGUnorderedAccessView* UAV, con
 
 	FRDGMemsetBufferCS::FPermutationDomain PermutationVector;
 	PermutationVector.Set<FRDGMemsetBufferCS::ResourceTypeDim >((int)ResourceTypeEnum);
+	PermutationVector.Set<FRDGMemsetBufferCS::StructuredElementSizeDim >((int)ByteBufferStructuredSize);
 	TShaderMapRef<FRDGMemsetBufferCS> ComputeShader(GetGlobalShaderMap(GMaxRHIFeatureLevel), PermutationVector);
 
 	FComputeShaderUtils::AddPass(
@@ -332,6 +430,7 @@ void MemcpyResource(FRDGBuilder& GraphBuilder, FRDGUnorderedAccessView* UAV, FRD
 		const uint32 NumElementsPerDispatch = FMath::Min(FMath::Max(NumWaves, 1u) * Divisor * 64, Params.Count - NumElementsProcessed);
 
 		EByteBufferResourceType ResourceTypeEnum = EByteBufferResourceType::Count;
+		EByteBufferStructuredSize ByteBufferStructuredSize = EByteBufferStructuredSize::Uint4;
 
 		auto* PassParameters = GraphBuilder.AllocParameters<FRDGMemcpyCS::FParameters>();
 		PassParameters->Common.Size = NumElementsPerDispatch;
@@ -346,10 +445,34 @@ void MemcpyResource(FRDGBuilder& GraphBuilder, FRDGUnorderedAccessView* UAV, FRD
 			PassParameters->Common.DstByteAddressBuffer    = GetAs<FRDGBufferUAV>(UAV);
 			break;
 		case EResourceType::STRUCTURED_BUFFER:
-			ResourceTypeEnum = EByteBufferResourceType::Float4_StructuredBuffer;
-			PassParameters->SrcStructuredBuffer            = GetAs<FRDGBufferSRV>(SRV);
-			PassParameters->Common.DstStructuredBuffer     = GetAs<FRDGBufferUAV>(UAV);
+		{
+			ResourceTypeEnum = EByteBufferResourceType::StructuredBuffer;
+
+			FRDGBufferUAV * RDGBufferUAV = GetAs<FRDGBufferUAV>(UAV);
+			ByteBufferStructuredSize = GetStructuredBufferElementSize(RDGBufferUAV);
+			switch(ByteBufferStructuredSize)
+			{
+			case EByteBufferStructuredSize::Uint1:
+				PassParameters->SrcStructuredBuffer1x = GetAs<FRDGBufferSRV>(SRV);
+				PassParameters->Common.DstStructuredBuffer1x = RDGBufferUAV;
+				break;
+			case EByteBufferStructuredSize::Uint2:
+				PassParameters->SrcStructuredBuffer2x = GetAs<FRDGBufferSRV>(SRV);
+				PassParameters->Common.DstStructuredBuffer2x = RDGBufferUAV;
+				break;
+			case EByteBufferStructuredSize::Uint4:
+				PassParameters->SrcStructuredBuffer4x = GetAs<FRDGBufferSRV>(SRV);
+				PassParameters->Common.DstStructuredBuffer4x = RDGBufferUAV;
+				break;
+			case EByteBufferStructuredSize::Uint8:
+				PassParameters->SrcStructuredBuffer8x = GetAs<FRDGBufferSRV>(SRV);
+				PassParameters->Common.DstStructuredBuffer8x = RDGBufferUAV;
+				break;
+			default:
+				break;
+			};
 			break;
+		}
 		case EResourceType::BUFFER:
 			ResourceTypeEnum = EByteBufferResourceType::Float4_Buffer;
 			PassParameters->SrcBuffer                      = GetAs<FRDGBufferSRV>(SRV);
@@ -367,6 +490,7 @@ void MemcpyResource(FRDGBuilder& GraphBuilder, FRDGUnorderedAccessView* UAV, FRD
 
 		FRDGMemcpyCS::FPermutationDomain PermutationVector;
 		PermutationVector.Set<FRDGMemcpyCS::ResourceTypeDim >((int)ResourceTypeEnum);
+		PermutationVector.Set<FRDGMemcpyCS::StructuredElementSizeDim >((int)ByteBufferStructuredSize);
 		TShaderMapRef<FRDGMemcpyCS> ComputeShader(GetGlobalShaderMap(GMaxRHIFeatureLevel), PermutationVector);
 
 		FComputeShaderUtils::AddPass(
@@ -533,20 +657,125 @@ struct FScatterUploadComputeConfig
 	uint32 NumLoops;
 };
 
-FScatterUploadComputeConfig GetScatterUploadComputeConfig(uint32 NumScatters, uint32 NumBytesPerElement)
+FScatterUploadComputeConfig GetScatterUploadComputeConfig(uint32 NumScatters, uint32 NumBytesPerElement, int32 NumElementsPerScatter = -1)
 {
 	constexpr uint32 ThreadGroupSize = 64u;
 
 	FScatterUploadComputeConfig Config;
 	Config.ThreadGroupSize = ThreadGroupSize;
-	Config.NumBytesPerThread = (NumBytesPerElement & 15) == 0 ? 16 : 4;
-	Config.NumThreadsPerScatter = NumBytesPerElement / Config.NumBytesPerThread;
+	if (NumElementsPerScatter != INDEX_NONE)
+	{
+		Config.NumBytesPerThread = NumBytesPerElement;
+		Config.NumThreadsPerScatter = NumElementsPerScatter; 
+	}
+	else
+	{
+		Config.NumBytesPerThread = (NumBytesPerElement & 15) == 0 ? 16 : 4;
+		Config.NumThreadsPerScatter = NumBytesPerElement / Config.NumBytesPerThread;
+	}
 	Config.NumThreads = NumScatters * Config.NumThreadsPerScatter;
 	Config.NumDispatches = FMath::DivideAndRoundUp(Config.NumThreads, Config.ThreadGroupSize);
 	Config.NumLoops = FMath::DivideAndRoundUp(Config.NumDispatches, (uint32)GMaxComputeDispatchDimension);
 	return Config;
 }
 
+
+void ScatterCopyResource(FRDGBuilder& GraphBuilder, FRDGViewableResource* DstResource, FRDGBufferSRV* ScatterBufferSRV, FRDGBufferSRV* UploadBufferSRV, const FScatterCopyParams &Params)
+{
+	const FScatterUploadComputeConfig ComputeConfig = GetScatterUploadComputeConfig(Params.NumScatters, Params.NumBytesPerElement, Params.NumElementsPerScatter);
+	const EResourceType DstResourceType = GetResourceType(DstResource);
+
+	EByteBufferResourceType ResourceTypeEnum = EByteBufferResourceType::Count;
+	EByteBufferStructuredSize ByteBufferStructuredSize = EByteBufferStructuredSize::Uint4;
+
+	FRDGScatterCopyCS::FParameters Parameters;
+	Parameters.Common.Size = ComputeConfig.NumThreadsPerScatter;
+	Parameters.NumScatters = Params.NumScatters;
+
+	if (DstResourceType == EResourceType::BYTEBUFFER)
+	{
+		if (ComputeConfig.NumBytesPerThread == 16)
+		{
+			ResourceTypeEnum = EByteBufferResourceType::Uint4Aligned_Buffer;
+		}
+		else
+		{
+			ResourceTypeEnum = EByteBufferResourceType::Uint_Buffer;
+		}
+		Parameters.UploadByteAddressBuffer = UploadBufferSRV;
+		Parameters.ScatterByteAddressBuffer = ScatterBufferSRV;
+		Parameters.Common.DstByteAddressBuffer = GraphBuilder.CreateUAV(GetAsBuffer(DstResource), ERDGUnorderedAccessViewFlags::SkipBarrier);
+	}
+	else if (DstResourceType == EResourceType::STRUCTURED_BUFFER)
+	{
+		ResourceTypeEnum = EByteBufferResourceType::StructuredBuffer;
+
+		FRDGBufferUAV * RDGBufferUAV = GraphBuilder.CreateUAV(GetAsBuffer(DstResource), ERDGUnorderedAccessViewFlags::SkipBarrier);
+		ByteBufferStructuredSize = GetStructuredBufferElementSize(RDGBufferUAV);
+		switch(ByteBufferStructuredSize)
+		{
+		case EByteBufferStructuredSize::Uint1:
+			Parameters.UploadStructuredBuffer1x = UploadBufferSRV;
+			Parameters.Common.DstStructuredBuffer1x = RDGBufferUAV;
+			break;
+		case EByteBufferStructuredSize::Uint2:
+			Parameters.UploadStructuredBuffer2x = UploadBufferSRV;
+			Parameters.Common.DstStructuredBuffer2x = RDGBufferUAV;
+			break;
+		case EByteBufferStructuredSize::Uint4:
+			Parameters.UploadStructuredBuffer4x = UploadBufferSRV;
+			Parameters.Common.DstStructuredBuffer4x = RDGBufferUAV;
+			break;
+		case EByteBufferStructuredSize::Uint8:
+			Parameters.UploadStructuredBuffer8x = UploadBufferSRV;
+			Parameters.Common.DstStructuredBuffer8x = RDGBufferUAV;
+			break;
+		default:
+			break;
+		};
+		Parameters.ScatterStructuredBuffer = ScatterBufferSRV;
+	}
+	else if (DstResourceType == EResourceType::BUFFER)
+	{
+		ResourceTypeEnum = EByteBufferResourceType::Float4_Buffer;
+
+		Parameters.UploadStructuredBuffer4x = UploadBufferSRV;
+		Parameters.ScatterStructuredBuffer = ScatterBufferSRV;
+		Parameters.Common.DstBuffer = GraphBuilder.CreateUAV(GetAsBuffer(DstResource), ERDGUnorderedAccessViewFlags::SkipBarrier);
+	}
+	else if (DstResourceType == EResourceType::TEXTURE)
+	{
+		ResourceTypeEnum = EByteBufferResourceType::Float4_Texture;
+
+		Parameters.UploadStructuredBuffer4x = UploadBufferSRV;
+		Parameters.ScatterStructuredBuffer = ScatterBufferSRV;
+		Parameters.Common.DstTexture = GraphBuilder.CreateUAV(GetAsTexture(DstResource), ERDGUnorderedAccessViewFlags::SkipBarrier);
+
+		Parameters.Common.Float4sPerLine = CalculateFloat4sPerLine();
+	}
+
+	FRDGScatterCopyCS::FPermutationDomain PermutationVector;
+	PermutationVector.Set<FRDGScatterCopyCS::ResourceTypeDim>((int)ResourceTypeEnum);
+	PermutationVector.Set<FRDGScatterCopyCS::StructuredElementSizeDim >((int)ByteBufferStructuredSize);
+	TShaderMapRef<FRDGScatterCopyCS> ComputeShader(GetGlobalShaderMap(GMaxRHIFeatureLevel), PermutationVector);
+
+	for (uint32 LoopIdx = 0; LoopIdx < ComputeConfig.NumLoops; ++LoopIdx)
+	{
+		Parameters.Common.SrcOffset = LoopIdx * (uint32)GMaxComputeDispatchDimension * ComputeConfig.ThreadGroupSize;
+
+		uint32 LoopNumDispatch = FMath::Min(ComputeConfig.NumDispatches - LoopIdx * (uint32)GMaxComputeDispatchDimension, (uint32)GMaxComputeDispatchDimension);
+
+		FComputeShaderUtils::AddPass(
+			GraphBuilder,
+			RDG_EVENT_NAME("ScatterUpload[%d] (Resource: %s, Offset: %u, GroupSize: %u)", LoopIdx, DstResource->Name, Parameters.Common.SrcOffset, LoopNumDispatch),
+			ComputeShader,
+			GraphBuilder.AllocParameters(&Parameters),
+			FIntVector(LoopNumDispatch, 1, 1));
+	}
+}
+
+
+
 void FRDGScatterUploadBuffer::Init(FRDGBuilder& GraphBuilder, uint32 NumElements, uint32 InNumBytesPerElement, bool bInFloat4Buffer, const TCHAR* Name)
 {
 	TRACE_CPUPROFILER_EVENT_SCOPE(FRDGScatterUploadBuffer::Init);
@@ -611,78 +840,12 @@ void FRDGScatterUploadBuffer::ResourceUploadTo(FRDGBuilder& GraphBuilder, FRDGVi
 		return;
 	}
 
-	const FScatterUploadComputeConfig ComputeConfig = GetScatterUploadComputeConfig(NumScatters, NumBytesPerElement);
-
 	const EResourceType DstResourceType = GetResourceType(DstResource);
 	check(bFloat4Buffer != (DstResourceType == EResourceType::BYTEBUFFER));
 
-	EByteBufferResourceType ResourceTypeEnum = EByteBufferResourceType::Count;
-
-	FRDGScatterCopyCS::FParameters Parameters;
-	Parameters.Common.Size = ComputeConfig.NumThreadsPerScatter;
-	Parameters.NumScatters = NumScatters;
-
 	FRDGBufferSRV* ScatterBufferSRV = GraphBuilder.CreateSRV(GraphBuilder.RegisterExternalBuffer(ScatterBuffer));
 	FRDGBufferSRV* UploadBufferSRV  = GraphBuilder.CreateSRV(GraphBuilder.RegisterExternalBuffer(UploadBuffer));
-
-	if (DstResourceType == EResourceType::BYTEBUFFER)
-	{
-		if (ComputeConfig.NumBytesPerThread == 16)
-		{
-			ResourceTypeEnum = EByteBufferResourceType::Uint4Aligned_Buffer;
-		}
-		else
-		{
-			ResourceTypeEnum = EByteBufferResourceType::Uint_Buffer;
-		}
-		Parameters.UploadByteAddressBuffer = UploadBufferSRV;
-		Parameters.ScatterByteAddressBuffer = ScatterBufferSRV;
-		Parameters.Common.DstByteAddressBuffer = GraphBuilder.CreateUAV(GetAsBuffer(DstResource), ERDGUnorderedAccessViewFlags::SkipBarrier);
-	}
-	else if (DstResourceType == EResourceType::STRUCTURED_BUFFER)
-	{
-		ResourceTypeEnum = EByteBufferResourceType::Float4_StructuredBuffer;
-
-		Parameters.UploadStructuredBuffer = UploadBufferSRV;
-		Parameters.ScatterStructuredBuffer = ScatterBufferSRV;
-		Parameters.Common.DstStructuredBuffer = GraphBuilder.CreateUAV(GetAsBuffer(DstResource), ERDGUnorderedAccessViewFlags::SkipBarrier);
-	}
-	else if (DstResourceType == EResourceType::BUFFER)
-	{
-		ResourceTypeEnum = EByteBufferResourceType::Float4_Buffer;
-
-		Parameters.UploadStructuredBuffer = UploadBufferSRV;
-		Parameters.ScatterStructuredBuffer = ScatterBufferSRV;
-		Parameters.Common.DstBuffer = GraphBuilder.CreateUAV(GetAsBuffer(DstResource), ERDGUnorderedAccessViewFlags::SkipBarrier);
-	}
-	else if (DstResourceType == EResourceType::TEXTURE)
-	{
-		ResourceTypeEnum = EByteBufferResourceType::Float4_Texture;
-
-		Parameters.UploadStructuredBuffer = UploadBufferSRV;
-		Parameters.ScatterStructuredBuffer = ScatterBufferSRV;
-		Parameters.Common.DstTexture = GraphBuilder.CreateUAV(GetAsTexture(DstResource), ERDGUnorderedAccessViewFlags::SkipBarrier);
-
-		Parameters.Common.Float4sPerLine = CalculateFloat4sPerLine();
-	}
-
-	FRDGByteBufferShader::FPermutationDomain PermutationVector;
-	PermutationVector.Set<FRDGByteBufferShader::ResourceTypeDim>((int)ResourceTypeEnum);
-	TShaderMapRef<FRDGScatterCopyCS> ComputeShader(GetGlobalShaderMap(GMaxRHIFeatureLevel), PermutationVector);
-
-	for (uint32 LoopIdx = 0; LoopIdx < ComputeConfig.NumLoops; ++LoopIdx)
-	{
-		Parameters.Common.SrcOffset = LoopIdx * (uint32)GMaxComputeDispatchDimension * ComputeConfig.ThreadGroupSize;
-
-		uint32 LoopNumDispatch = FMath::Min(ComputeConfig.NumDispatches - LoopIdx * (uint32)GMaxComputeDispatchDimension, (uint32)GMaxComputeDispatchDimension);
-
-		FComputeShaderUtils::AddPass(
-			GraphBuilder,
-			RDG_EVENT_NAME("ScatterUpload[%d] (Resource: %s, Offset: %u, GroupSize: %u)", LoopIdx, DstResource->Name, Parameters.Common.SrcOffset, LoopNumDispatch),
-			ComputeShader,
-			GraphBuilder.AllocParameters(&Parameters),
-			FIntVector(LoopNumDispatch, 1, 1));
-	}
+	ScatterCopyResource(GraphBuilder, DstResource, ScatterBufferSRV, UploadBufferSRV, FScatterCopyParams { NumScatters, NumBytesPerElement, INDEX_NONE});
 
 	Reset();
 }
@@ -790,6 +953,7 @@ void FRDGAsyncScatterUploadBuffer::End(FRDGBuilder& GraphBuilder, FRDGScatterUpl
 	const EResourceType DstResourceType = GetResourceType(DstResource);
 
 	EByteBufferResourceType ResourceTypeEnum = EByteBufferResourceType::Count;
+	EByteBufferStructuredSize ByteBufferStructuredSize = EByteBufferStructuredSize::Uint4;
 
 	if (DstResourceType == EResourceType::BYTEBUFFER)
 	{
@@ -807,17 +971,38 @@ void FRDGAsyncScatterUploadBuffer::End(FRDGBuilder& GraphBuilder, FRDGScatterUpl
 	}
 	else if (DstResourceType == EResourceType::STRUCTURED_BUFFER)
 	{
-		ResourceTypeEnum = EByteBufferResourceType::Float4_StructuredBuffer;
+		ResourceTypeEnum = EByteBufferResourceType::StructuredBuffer;
 
-		Parameters.UploadStructuredBuffer = UploadBufferSRV;
+		FRDGBufferUAV * RDGBufferUAV = GraphBuilder.CreateUAV(GetAsBuffer(DstResource), ERDGUnorderedAccessViewFlags::SkipBarrier);
+		ByteBufferStructuredSize = GetStructuredBufferElementSize(RDGBufferUAV);
+		switch(ByteBufferStructuredSize)
+		{
+		case EByteBufferStructuredSize::Uint1:
+			Parameters.UploadStructuredBuffer1x = UploadBufferSRV;
+			Parameters.Common.DstStructuredBuffer1x = RDGBufferUAV;
+			break;
+		case EByteBufferStructuredSize::Uint2:
+			Parameters.UploadStructuredBuffer2x = UploadBufferSRV;
+			Parameters.Common.DstStructuredBuffer2x = RDGBufferUAV;
+			break;
+		case EByteBufferStructuredSize::Uint4:
+			Parameters.UploadStructuredBuffer4x = UploadBufferSRV;
+			Parameters.Common.DstStructuredBuffer4x = RDGBufferUAV;
+			break;
+		case EByteBufferStructuredSize::Uint8:
+			Parameters.UploadStructuredBuffer8x = UploadBufferSRV;
+			Parameters.Common.DstStructuredBuffer8x = RDGBufferUAV;
+			break;
+		default:
+			break;
+		};
 		Parameters.ScatterStructuredBuffer = ScatterBufferSRV;
-		Parameters.Common.DstStructuredBuffer = GraphBuilder.CreateUAV(GetAsBuffer(DstResource), ERDGUnorderedAccessViewFlags::SkipBarrier);
 	}
 	else if (DstResourceType == EResourceType::BUFFER)
 	{
 		ResourceTypeEnum = EByteBufferResourceType::Float4_Buffer;
 
-		Parameters.UploadStructuredBuffer = UploadBufferSRV;
+		Parameters.UploadStructuredBuffer4x = UploadBufferSRV;
 		Parameters.ScatterStructuredBuffer = ScatterBufferSRV;
 		Parameters.Common.DstBuffer = GraphBuilder.CreateUAV(GetAsBuffer(DstResource), ERDGUnorderedAccessViewFlags::SkipBarrier);
 	}
@@ -825,15 +1010,16 @@ void FRDGAsyncScatterUploadBuffer::End(FRDGBuilder& GraphBuilder, FRDGScatterUpl
 	{
 		ResourceTypeEnum = EByteBufferResourceType::Float4_Texture;
 
-		Parameters.UploadStructuredBuffer = UploadBufferSRV;
+		Parameters.UploadStructuredBuffer4x = UploadBufferSRV;
 		Parameters.ScatterStructuredBuffer = ScatterBufferSRV;
 		Parameters.Common.DstTexture = GraphBuilder.CreateUAV(GetAsTexture(DstResource), ERDGUnorderedAccessViewFlags::SkipBarrier);
 
 		Parameters.Common.Float4sPerLine = CalculateFloat4sPerLine();
 	}
 
-	FRDGByteBufferShader::FPermutationDomain PermutationVector;
-	PermutationVector.Set<FRDGByteBufferShader::ResourceTypeDim>((int)ResourceTypeEnum);
+	FRDGScatterCopyCS::FPermutationDomain PermutationVector;
+	PermutationVector.Set<FRDGScatterCopyCS::ResourceTypeDim>((int)ResourceTypeEnum);
+	PermutationVector.Set<FRDGScatterCopyCS::StructuredElementSizeDim >(static_cast<int32>(ByteBufferStructuredSize));
 	TShaderMapRef<FRDGScatterCopyCS> ComputeShader(GetGlobalShaderMap(GMaxRHIFeatureLevel), PermutationVector);
 
 	for (uint32 LoopIdx = 0; LoopIdx < ComputeConfig.NumLoops; ++LoopIdx)
@@ -915,9 +1101,9 @@ void MemsetResource(FRHICommandList& RHICmdList, const ResourceType& DstBuffer,
 	}
 	else if (ResourceTypeTraits<ResourceType>::Type == EResourceType::STRUCTURED_BUFFER)
 	{
-		ResourceTypeEnum = EByteBufferResourceType::Float4_StructuredBuffer;
+		ResourceTypeEnum = EByteBufferResourceType::StructuredBuffer;
 
-		Parameters.DstStructuredBuffer = DstBuffer.UAV;
+		Parameters.DstStructuredBuffer4x = DstBuffer.UAV;
 	}
 	else if (ResourceTypeTraits<ResourceType>::Type == EResourceType::TEXTURE)
 	{
@@ -929,6 +1115,7 @@ void MemsetResource(FRHICommandList& RHICmdList, const ResourceType& DstBuffer,
 
 	FMemcpyCS::FPermutationDomain PermutationVector;
 	PermutationVector.Set<FMemcpyCS::ResourceTypeDim >((int)ResourceTypeEnum);
+	PermutationVector.Set<FMemcpyCS::StructuredElementSizeDim>((int32)EByteBufferStructuredSize::Uint4);
 
 	auto ShaderMap = GetGlobalShaderMap(GMaxRHIFeatureLevel);
 
@@ -972,10 +1159,10 @@ void MemcpyResource(FRHICommandList& RHICmdList, const ResourceType& DstBuffer,
 		}
 		else if (ResourceTypeTraits<ResourceType>::Type == EResourceType::STRUCTURED_BUFFER)
 		{
-			ResourceTypeEnum = EByteBufferResourceType::Float4_StructuredBuffer;
+			ResourceTypeEnum = EByteBufferResourceType::StructuredBuffer;
 
-			Parameters.SrcStructuredBuffer = SrcBuffer.SRV;
-			Parameters.Common.DstStructuredBuffer = DstBuffer.UAV;
+			Parameters.SrcStructuredBuffer4x = SrcBuffer.SRV;
+			Parameters.Common.DstStructuredBuffer4x = DstBuffer.UAV;
 		}
 		else if (ResourceTypeTraits<ResourceType>::Type == EResourceType::BUFFER)
 		{
@@ -999,6 +1186,7 @@ void MemcpyResource(FRHICommandList& RHICmdList, const ResourceType& DstBuffer,
 
 		FMemcpyCS::FPermutationDomain PermutationVector;
 		PermutationVector.Set<FMemcpyCS::ResourceTypeDim >((int)ResourceTypeEnum);
+		PermutationVector.Set<FMemcpyCS::StructuredElementSizeDim>((int32)EByteBufferStructuredSize::Uint4);
 
 		auto ComputeShader = GetGlobalShaderMap(GMaxRHIFeatureLevel)->GetShader<FMemcpyCS >(PermutationVector);
 
@@ -1495,7 +1683,9 @@ void FScatterUploadBuffer::ResourceUploadTo(FRHICommandList& RHICmdList, const R
 	}
 
 	if (NumScatters == 0)
+	{
 		return;
+	}
 
 	constexpr uint32 ThreadGroupSize = 64u;
 	uint32 NumBytesPerThread = (NumBytesPerElement & 15) == 0 ? 16 : 4;
@@ -1528,17 +1718,17 @@ void FScatterUploadBuffer::ResourceUploadTo(FRHICommandList& RHICmdList, const R
 	}
 	else if (ResourceTypeTraits<ResourceType>::Type == EResourceType::STRUCTURED_BUFFER)
 	{
-		ResourceTypeEnum = EByteBufferResourceType::Float4_StructuredBuffer;
+		ResourceTypeEnum = EByteBufferResourceType::StructuredBuffer;
 
-		Parameters.UploadStructuredBuffer = UploadBuffer.SRV;
+		Parameters.UploadStructuredBuffer4x = UploadBuffer.SRV;
 		Parameters.ScatterStructuredBuffer = ScatterBuffer.SRV;
-		Parameters.Common.DstStructuredBuffer = DstBuffer.UAV;
+		Parameters.Common.DstStructuredBuffer4x = DstBuffer.UAV;
 	}
 	else if (ResourceTypeTraits<ResourceType>::Type == EResourceType::BUFFER)
 	{
 		ResourceTypeEnum = EByteBufferResourceType::Float4_Buffer;
 
-		Parameters.UploadStructuredBuffer = UploadBuffer.SRV;
+		Parameters.UploadStructuredBuffer4x = UploadBuffer.SRV;
 		Parameters.ScatterStructuredBuffer = ScatterBuffer.SRV;
 		Parameters.Common.DstBuffer = DstBuffer.UAV;
 	}
@@ -1546,7 +1736,7 @@ void FScatterUploadBuffer::ResourceUploadTo(FRHICommandList& RHICmdList, const R
 	{
 		ResourceTypeEnum = EByteBufferResourceType::Float4_Texture;
 
-		Parameters.UploadStructuredBuffer = UploadBuffer.SRV;
+		Parameters.UploadStructuredBuffer4x = UploadBuffer.SRV;
 		Parameters.ScatterStructuredBuffer = ScatterBuffer.SRV;
 		Parameters.Common.DstTexture = DstBuffer.UAV;
 
@@ -1555,6 +1745,7 @@ void FScatterUploadBuffer::ResourceUploadTo(FRHICommandList& RHICmdList, const R
 
 	FByteBufferShader::FPermutationDomain PermutationVector;
 	PermutationVector.Set<FByteBufferShader::ResourceTypeDim>((int)ResourceTypeEnum);
+	PermutationVector.Set<FMemcpyCS::StructuredElementSizeDim>((int32)EByteBufferStructuredSize::Uint4);
 
 	auto ComputeShader = GetGlobalShaderMap(GMaxRHIFeatureLevel)->GetShader<FScatterCopyCS>(PermutationVector);
 
diff --git a/Engine/Source/Runtime/RenderCore/Public/RenderGraphUtils.h b/Engine/Source/Runtime/RenderCore/Public/RenderGraphUtils.h
--- a/Engine/Source/Runtime/RenderCore/Public/RenderGraphUtils.h
+++ b/Engine/Source/Runtime/RenderCore/Public/RenderGraphUtils.h
@@ -47,6 +47,9 @@ class FRHIGPUBufferReadback;
 class FRHIGPUTextureReadback;
 class FShaderParametersMetadata;
 
+// Callback to modify pass parameters just prior to dispatch, used for indirect dispatches where the group count callback cannot be used for this purpose.
+using FRDGDispatchLateParamCallback = TFunction<void()>;
+
 /** Returns whether the resource was produced by a prior pass. */
 inline bool HasBeenProduced(FRDGViewableResource* Resource)
 {
@@ -598,7 +601,8 @@ struct RENDERCORE_API FComputeShaderUtils
 		const TShaderRef<TShaderClass>& ComputeShader,
 		typename TShaderClass::FParameters* Parameters,
 		FRDGBufferRef IndirectArgsBuffer,
-		uint32 IndirectArgsOffset)
+		uint32 IndirectArgsOffset,
+		FRDGDispatchLateParamCallback&& DispatchLateParamCallback = FRDGDispatchLateParamCallback())
 	{
 		checkf(PassFlags == ERDGPassFlags::Compute || PassFlags == ERDGPassFlags::AsyncCompute, TEXT("AddPass only supports 'Compute' or 'AsyncCompute'."));
 		checkf(IndirectArgsBuffer->Desc.Usage & BUF_DrawIndirect, TEXT("The buffer %s was not flagged for indirect draw parameters"), IndirectArgsBuffer->Name);
@@ -610,12 +614,16 @@ struct RENDERCORE_API FComputeShaderUtils
 			Forward<FRDGEventName>(PassName),
 			Parameters,
 			PassFlags,
-			[Parameters, ComputeShader, IndirectArgsBuffer, IndirectArgsOffset](FRHIComputeCommandList& RHICmdList)
+			[Parameters, ComputeShader, IndirectArgsBuffer, IndirectArgsOffset, DispatchLateParamCallback = MoveTemp(DispatchLateParamCallback)](FRHIComputeCommandList& RHICmdList)
 		{			
 			// Marks the indirect draw parameter as used by the pass manually, given it can't be bound directly by any of the shader,
 			// meaning SetShaderParameters() won't be able to do it.
 			IndirectArgsBuffer->MarkResourceAsUsed();
 
+			if (DispatchLateParamCallback)
+			{
+				DispatchLateParamCallback();
+			}
 			FComputeShaderUtils::DispatchIndirect(RHICmdList, ComputeShader, *Parameters, IndirectArgsBuffer->GetIndirectRHICallBuffer(), IndirectArgsOffset);
 		});
 	}
@@ -627,9 +635,10 @@ struct RENDERCORE_API FComputeShaderUtils
 		const TShaderRef<TShaderClass>& ComputeShader,
 		typename TShaderClass::FParameters* Parameters,
 		FRDGBufferRef IndirectArgsBuffer,
-		uint32 IndirectArgsOffset)
+		uint32 IndirectArgsOffset,
+		FRDGDispatchLateParamCallback&& DispatchLateParamCallback = FRDGDispatchLateParamCallback())
 	{
-		return AddPass(GraphBuilder, Forward<FRDGEventName>(PassName), ERDGPassFlags::Compute, ComputeShader, Parameters, IndirectArgsBuffer, IndirectArgsOffset);
+		return AddPass(GraphBuilder, Forward<FRDGEventName>(PassName), ERDGPassFlags::Compute, ComputeShader, Parameters, IndirectArgsBuffer, IndirectArgsOffset, MoveTemp(DispatchLateParamCallback));
 	}
 
 	static void ClearUAV(FRDGBuilder& GraphBuilder, FGlobalShaderMap* ShaderMap, FRDGBufferUAVRef UAV, uint32 ClearValue);
@@ -992,6 +1001,23 @@ FORCEINLINE FRDGBufferRef CreateUploadBuffer(
 	return CreateUploadBuffer(GraphBuilder, Name, sizeof(ElementType), InitialData.Num(), InitialData.GetData(), sizeof(ElementType) * InitialData.Num(), InitialDataFlags);
 }
 
+/**
+ * Helper to create a structured upload buffer with initial data from a TArray.
+ * NOTE: does not provide a 1-size fallback for empty initial data.
+ */
+template <typename ElementType, typename AllocatorType>
+FORCEINLINE FRDGBufferRef CreateStructuredUploadBuffer(
+	FRDGBuilder& GraphBuilder,
+	const TCHAR* Name,
+	const TArray<ElementType, AllocatorType> &InitialData,
+	ERDGInitialDataFlags InitialDataFlags = ERDGInitialDataFlags::None)
+{
+	FRDGBufferRef Buffer = GraphBuilder.CreateBuffer(FRDGBufferDesc::CreateStructuredUploadDesc(sizeof(ElementType), InitialData.Num()), Name);
+	GraphBuilder.QueueBufferUpload(Buffer, TConstArrayView<ElementType>(InitialData), InitialDataFlags);
+	return Buffer;
+}
+
+
 /** Creates a vertex buffer with initial data by creating an upload pass. */
 RENDERCORE_API FRDGBufferRef CreateVertexBuffer(
 	FRDGBuilder& GraphBuilder,
diff --git a/Engine/Source/Runtime/RenderCore/Public/UnifiedBuffer.h b/Engine/Source/Runtime/RenderCore/Public/UnifiedBuffer.h
--- a/Engine/Source/Runtime/RenderCore/Public/UnifiedBuffer.h
+++ b/Engine/Source/Runtime/RenderCore/Public/UnifiedBuffer.h
@@ -196,6 +196,16 @@ extern RENDERCORE_API void MemcpyResource(FRDGBuilder& GraphBuilder, FRDGBuffer*
 extern RENDERCORE_API void MemsetResource(FRDGBuilder& GraphBuilder, FRDGUnorderedAccessView* DstResource, const FMemsetResourceParams& Params);
 extern RENDERCORE_API void MemcpyResource(FRDGBuilder& GraphBuilder, FRDGUnorderedAccessView* DstResource, FRDGShaderResourceView* SrcResource, const FMemcpyResourceParams& Params);
 
+struct FScatterCopyParams
+{
+	uint32 NumScatters = 0u;
+	uint32 NumBytesPerElement = 0u;
+	int32 NumElementsPerScatter = INDEX_NONE; // INDEX_NONE lets the setup figure it out, otherwise it will run NumScatters * NumElementsPerScatter threads to copy the source data.
+};
+
+void RENDERCORE_API ScatterCopyResource(FRDGBuilder& GraphBuilder, FRDGViewableResource* DstResource, FRDGBufferSRV* ScatterBufferSRV, FRDGBufferSRV* UploadBufferSRV, const FScatterCopyParams &Params);
+
+
 extern RENDERCORE_API FRDGBuffer* ResizeBufferIfNeeded(FRDGBuilder& GraphBuilder, TRefCountPtr<FRDGPooledBuffer>& ExternalBuffer, const FRDGBufferDesc& BufferDesc, const TCHAR* Name);
 extern RENDERCORE_API FRDGBuffer* ResizeBufferIfNeeded(FRDGBuilder& GraphBuilder, TRefCountPtr<FRDGPooledBuffer>& ExternalBuffer, EPixelFormat Format, uint32 NumElements, const TCHAR* Name);
 extern RENDERCORE_API FRDGBuffer* ResizeStructuredBufferIfNeeded(FRDGBuilder& GraphBuilder, TRefCountPtr<FRDGPooledBuffer>& ExternalBuffer, uint32 NumBytes, const TCHAR* Name);
diff --git a/Engine/Source/Runtime/Renderer/Private/RendererPrivateUtils.h b/Engine/Source/Runtime/Renderer/Private/RendererPrivateUtils.h
--- a/Engine/Source/Runtime/Renderer/Private/RendererPrivateUtils.h
+++ b/Engine/Source/Runtime/Renderer/Private/RendererPrivateUtils.h
@@ -5,6 +5,7 @@
 #include "RenderResource.h"
 #include "RHI.h"
 #include "RHICommandList.h"
+#include "RenderGraphUtils.h"
 
 class FTileTexCoordVertexBuffer : public FVertexBuffer
 {
@@ -62,4 +63,177 @@ public:
 	}
 
 	const uint32 NumTileQuadsInBuffer;
-};
\ No newline at end of file
+};
+
+template <typename ReadbackProcessingLambdaType>
+void AddBufferLockReadbackPass(FRDGBuilder& GraphBuilder, TRefCountPtr<FRDGPooledBuffer> SourceBuffer, uint32 NumBytes, ReadbackProcessingLambdaType &&ReadbackProcessingLambda)
+{
+	FRHIBuffer* SourceBufferRHI = SourceBuffer->GetRHI();
+	GraphBuilder.AddPass(
+		RDG_EVENT_NAME("BufferLockReadbackPass"),
+		ERDGPassFlags::None,
+		[ReadbackProcessingLambdaType=MoveTemp(ReadbackProcessingLambda), SourceBufferRHI, NumBytes](FRHICommandListImmediate& RHICmdList)
+	{
+		const void *GPUData = (uint32*)RHICmdList.LockBuffer(SourceBufferRHI, 0, NumBytes, RLM_ReadOnly);
+		ReadbackProcessingLambdaType(GPUData);
+		RHICmdList.UnlockBuffer(SourceBufferRHI);
+	});
+}
+
+
+/** 
+ * Helper class to manage a persistent structured buffer.
+ */
+class FPersistentStructuredBuffer
+{
+public:
+	FPersistentStructuredBuffer(int32 InMinimumNumElementsReserved, const TCHAR *InName, bool bInRoundUpToPOT = true);
+
+	FRDGBuffer* ResizeBufferIfNeeded(FRDGBuilder& GraphBuilder, int32 InNewMinNumElements, int32 BytesPerElement);
+	FRDGBuffer* Register(FRDGBuilder& GraphBuilder);
+
+protected:
+	int32 MinimumNumElementsReserved = 0;
+	const TCHAR *Name = nullptr;
+	bool bRoundUpToPOT = true;
+	TRefCountPtr<FRDGPooledBuffer> PooledBuffer;
+};
+
+/**
+ * Typed version of FPersistentStructuredBuffer 
+ */
+template <typename InValueType>
+class TPersistentStructuredBuffer : public FPersistentStructuredBuffer
+{
+public:
+	using ValueType = InValueType;
+	TPersistentStructuredBuffer(int32 InMinimumNumElementsReserved, const TCHAR *InName, bool bInRoundUpToPOT = true)
+	: FPersistentStructuredBuffer(InMinimumNumElementsReserved, InName, bInRoundUpToPOT)
+	{
+	}
+
+	FRDGBuffer* ResizeBufferIfNeeded(FRDGBuilder& GraphBuilder, int32 InNewMinNumElements)
+	{
+		return FPersistentStructuredBuffer::ResizeBufferIfNeeded(GraphBuilder, InNewMinNumElements, sizeof(ValueType));
+	}
+	
+
+	template <typename ValueCheckFuncType>
+	void ValidateGPUData(FRDGBuilder& GraphBuilder, TConstArrayView<ValueType> HostValues, ValueCheckFuncType &&ValueCheckFunc)
+	{
+#if DO_CHECK
+		check(HostValues.Num() == 0 || PooledBuffer.IsValid() && HostValues.Num() <= int32(PooledBuffer->Desc.NumElements));
+
+		if (!HostValues.IsEmpty())
+		{
+			// TODO: should perhaps copy the host values? A lot of overhead but very useful to be sure they are alive...
+			AddBufferLockReadbackPass(GraphBuilder, PooledBuffer, uint32(HostValues.GetTypeSize() * HostValues.Num()), [ValueCheckFunc=MoveTemp(ValueCheckFunc), HostValues](const void *LockedData)
+			{
+				const ValueType *GPUValuePtr = reinterpret_cast<const ValueType*>(LockedData);
+				for (int32 Index = 0; Index < HostValues.Num(); ++Index)
+				{
+					ValueCheckFunc(Index, HostValues[Index], GPUValuePtr[Index]);
+				}
+			});
+		}
+#endif
+	}
+
+	TRefCountPtr<FRDGPooledBuffer> &GetPooledBuffer() { return PooledBuffer; }
+};
+
+class FStructuredBufferScatterUploader
+{
+public:
+	void UploadTo(FRDGBuilder& GraphBuilder, FRDGBuffer *DestBuffer, FRDGBuffer *ScatterOffsets, FRDGBuffer *Values, uint32 NumScatters, uint32 NumBytesPerElement, uint32 NumValuesPerScatter);
+};
+
+template <typename InValueType, int32 InNumValuesPerScatter = 1>
+class TStructuredBufferScatterUploader : public FStructuredBufferScatterUploader
+{
+public:
+	using ValueType = InValueType;
+	static constexpr int32 BytesPerElement = sizeof(ValueType);
+	static constexpr int32 NumValuesPerScatter = InNumValuesPerScatter;
+
+	/**
+	 * Optionally reserve space for NumScatters items.
+	 */
+	TStructuredBufferScatterUploader(int32 NumScatters = 0)
+	{
+		Reserve(NumScatters);
+	}
+
+	/**
+	 * Pre-size the value and scatter arrays, allocates space for InNumValuesPerScatter * NumScatters values & NumScatters scatter offsets.
+	 */
+	void Reserve(int32 NumScatters)
+	{
+		check(UploadDataProxy == nullptr);
+
+		UploadData.Values.Reserve(NumScatters * NumValuesPerScatter);
+		UploadData.ScatterOffsets.Reserve(NumScatters);
+	}
+
+	/**
+	 * Add single value to scatter to the destination offset.
+	 */
+	void Add(const ValueType &Value, int32 ScatterOffset)
+	{
+		check(1 == NumValuesPerScatter);
+		check(UploadDataProxy == nullptr);
+
+		UploadData.Values.Add(Value);
+		UploadData.ScatterOffsets.Add(ScatterOffset);
+	}
+
+	/**
+	 * Add a number of values to scatter to a common destination start offset, 
+	 * NOTE: the destination start offset is ScatterOffset * InNumValuesPerScatter (not ScatterOffset)
+	 */
+	void Add(const TConstArrayView<ValueType> &InValues, int32 ScatterOffset)
+	{
+		check(InValues.Num() == NumValuesPerScatter);
+		check(UploadDataProxy == nullptr);
+
+		UploadData.Values.Append(InValues);
+		UploadData.ScatterOffsets.Add(ScatterOffset);
+	}
+
+	int32 GetNumScatters() const { return UploadDataProxy != nullptr ? UploadDataProxy->ScatterOffsets.Num() : UploadData.ScatterOffsets.Num(); }
+
+	/**
+	 * Resize the destination persistent buffer (if needed) and upload & scatter the collected data to it.
+	 * This locks the uploader to prevent accidental resize (and thus realloc) of the buffer by adding more elements.
+	 */
+	FRDGBuffer *ResizeAndUploadTo(FRDGBuilder& GraphBuilder, TPersistentStructuredBuffer<ValueType> &DestDataBuffer, int32 DestDataMinimumSize)
+	{
+		check(UploadDataProxy == nullptr);
+
+		FRDGBuffer *DestBufferRDG = DestDataBuffer.ResizeBufferIfNeeded(GraphBuilder, DestDataMinimumSize);
+
+		uint32 NumScatters = UploadData.ScatterOffsets.Num();
+		if (NumScatters != 0u)
+		{
+			// Move the data arrays to a proxy object owned by RDG to guarantee life-times for upload.
+			UploadDataProxy = GraphBuilder.AllocObject<FUploadData>(MoveTemp(UploadData));
+
+			// upload the values & offsets
+			FRDGBuffer *ScatterOffsetsRDG = CreateStructuredUploadBuffer(GraphBuilder, TEXT("ScatterUploader.Offsets"), UploadDataProxy->ScatterOffsets, ERDGInitialDataFlags::NoCopy);
+			FRDGBuffer *ValuesRDG = CreateStructuredUploadBuffer(GraphBuilder, TEXT("ScatterUploader.Values"), UploadDataProxy->Values, ERDGInitialDataFlags::NoCopy);
+
+			FStructuredBufferScatterUploader::UploadTo(GraphBuilder, DestBufferRDG, ScatterOffsetsRDG, ValuesRDG, NumScatters, sizeof(ValueType), NumValuesPerScatter);
+		}
+
+		return DestBufferRDG;
+	}
+private:
+	struct FUploadData
+	{
+		TArray<ValueType> Values;
+		TArray<uint32> ScatterOffsets;
+	};
+	FUploadData *UploadDataProxy = nullptr;
+	FUploadData UploadData;
+};
+
diff --git a/Engine/Source/Runtime/Renderer/Private/RendererUtils.cpp b/Engine/Source/Runtime/Renderer/Private/RendererUtils.cpp
--- a/Engine/Source/Runtime/Renderer/Private/RendererUtils.cpp
+++ b/Engine/Source/Runtime/Renderer/Private/RendererUtils.cpp
@@ -1,12 +1,14 @@
 // Copyright Epic Games, Inc. All Rights Reserved.
 
 #include "RendererUtils.h"
+#include "RendererPrivateUtils.h"
 #include "RenderTargetPool.h"
 #include "RHIDefinitions.h"
 #include "DataDrivenShaderPlatformInfo.h"
 #include "VisualizeTexture.h"
 #include "ScenePrivate.h"
 #include "SystemTextures.h"
+#include "UnifiedBuffer.h"
 
 class FRTWriteMaskDecodeCS : public FGlobalShader
 {
@@ -224,3 +226,28 @@ namespace Strata
 		return nullptr;
 	}
 }
+
+FPersistentStructuredBuffer::FPersistentStructuredBuffer(int32 InMinimumNumElementsReserved, const TCHAR *InName, bool bInRoundUpToPOT)
+	: MinimumNumElementsReserved(InMinimumNumElementsReserved)
+	, Name(InName)
+	, bRoundUpToPOT(bInRoundUpToPOT)
+{
+}
+
+FRDGBuffer* FPersistentStructuredBuffer::ResizeBufferIfNeeded(FRDGBuilder& GraphBuilder, int32 InNewMinNumElements, int32 BytesPerElement)
+{
+	int32 NewMinNumElements = FMath::Max(MinimumNumElementsReserved, bRoundUpToPOT ? int32(FMath::RoundUpToPowerOfTwo(InNewMinNumElements)) : InNewMinNumElements);
+	FRDGBuffer* BufferRDG = ::ResizeBufferIfNeeded(GraphBuilder, PooledBuffer, FRDGBufferDesc::CreateStructuredDesc(BytesPerElement, NewMinNumElements), Name);
+	return BufferRDG;
+}
+
+FRDGBuffer* FPersistentStructuredBuffer::Register(FRDGBuilder& GraphBuilder) 
+{ 
+	return GraphBuilder.RegisterExternalBuffer(PooledBuffer); 
+}
+
+void FStructuredBufferScatterUploader::UploadTo(FRDGBuilder& GraphBuilder, FRDGBuffer *DestBuffer, FRDGBuffer *ScatterOffsets, FRDGBuffer *Values, uint32 NumScatters, uint32 NumBytesPerElement, uint32 NumValuesPerScatter)
+{
+	FScatterCopyParams ScatterCopyParams { NumScatters, NumBytesPerElement, NumValuesPerScatter };
+	ScatterCopyResource(GraphBuilder, DestBuffer, GraphBuilder.CreateSRV(ScatterOffsets), GraphBuilder.CreateSRV(Values), ScatterCopyParams);
+}
\ No newline at end of file
