[shaders] move ShaderParameterParser & some related functions to RenderCore so it can be used in the core shader compilation code (forthcoming change requires this)

diff --git a/Engine/Source/Developer/ShaderCompilerCommon/Private/ShaderCompilerCommon.cpp b/Engine/Source/Developer/ShaderCompilerCommon/Private/ShaderCompilerCommon.cpp
--- a/Engine/Source/Developer/ShaderCompilerCommon/Private/ShaderCompilerCommon.cpp
+++ b/Engine/Source/Developer/ShaderCompilerCommon/Private/ShaderCompilerCommon.cpp
@@ -475,90 +475,12 @@ TCHAR* FindNextUniformBufferReference(TCHAR* SearchPtr, const TCHAR* SearchStrin
 	return nullptr;
 }
 
-bool UE::ShaderCompilerCommon::ShouldUseStableConstantBuffer(const FShaderCompilerInput& Input)
-{
-	// stable constant buffer is for the FShaderParameterBindings::BindForLegacyShaderParameters() code path.
-	// Ray tracing shaders use FShaderParameterBindings::BindForRootShaderParameters instead.
-	if (Input.IsRayTracingShader())
-	{
-		return false;
-	}
-
-	return Input.RootParametersStructure != nullptr;
-}
-
-static const TCHAR* const s_AllSRVTypes[] =
-{
-	TEXT("Texture1D"),
-	TEXT("Texture1DArray"),
-	TEXT("Texture2D"),
-	TEXT("Texture2DArray"),
-	TEXT("Texture2DMS"),
-	TEXT("Texture2DMSArray"),
-	TEXT("Texture3D"),
-	TEXT("TextureCube"),
-	TEXT("TextureCubeArray"),
-
-	TEXT("Buffer"),
-	TEXT("ByteAddressBuffer"),
-	TEXT("StructuredBuffer"),
-	TEXT("ConstantBuffer"),
-	TEXT("RaytracingAccelerationStructure"),
-};
-
-static const TCHAR* const s_AllUAVTypes[] =
-{
-	TEXT("AppendStructuredBuffer"),
-	TEXT("RWBuffer"),
-	TEXT("RWByteAddressBuffer"),
-	TEXT("RWStructuredBuffer"),
-	TEXT("RWTexture1D"),
-	TEXT("RWTexture1DArray"),
-	TEXT("RWTexture2D"),
-	TEXT("RWTexture2DArray"),
-	TEXT("RWTexture3D"),
-	TEXT("RasterizerOrderedTexture2D"),
-};
-
-static const TCHAR* const s_AllSamplerTypes[] =
-{
-	TEXT("SamplerState"),
-	TEXT("SamplerComparisonState"),
-};
-
 EShaderParameterType UE::ShaderCompilerCommon::ParseParameterType(
 	FStringView InType,
 	TArrayView<const TCHAR* const> InExtraSRVTypes,
 	TArrayView<const TCHAR* const> InExtraUAVTypes)
 {
-	TArrayView<const TCHAR* const> AllSamplerTypes(s_AllSamplerTypes);
-	TArrayView<const TCHAR* const> AllSRVTypes(s_AllSRVTypes);
-	TArrayView<const TCHAR* const> AllUAVTypes(s_AllUAVTypes);
-
-	if (AllSamplerTypes.Contains(InType))
-	{
-		return EShaderParameterType::Sampler;
-	}
-
-	FStringView UntemplatedType = InType;
-	if (int32 Index = InType.Find(TEXT("<")); Index != INDEX_NONE)
-	{
-		// Remove the template argument but don't forget to clean up the type name
-		const int32 NumChars = InType.Len() - Index;
-		UntemplatedType = InType.LeftChop(NumChars).TrimEnd();
-	}
-
-	if (AllSRVTypes.Contains(UntemplatedType) || InExtraSRVTypes.Contains(UntemplatedType))
-	{
-		return EShaderParameterType::SRV;
-	}
-
-	if (AllUAVTypes.Contains(UntemplatedType) || InExtraUAVTypes.Contains(UntemplatedType))
-	{
-		return EShaderParameterType::UAV;
-	}
-
-	return EShaderParameterType::LooseData;
+	return FShaderParameterParser::ParseParameterType(InType, InExtraSRVTypes, InExtraUAVTypes);
 }
 
 FStringView UE::ShaderCompilerCommon::RemoveConstantBufferPrefix(FStringView InName)
@@ -573,34 +495,17 @@ FString UE::ShaderCompilerCommon::RemoveConstantBufferPrefix(const FString& InNa
 
 EShaderParameterType UE::ShaderCompilerCommon::ParseAndRemoveBindlessParameterPrefix(FStringView& InName)
 {
-	const FStringView OriginalName = InName;
-
-	if (InName = UE::String::RemoveFromStart(InName, FStringView(UE::ShaderCompilerCommon::kBindlessResourcePrefix)); InName != OriginalName)
-	{
-		return EShaderParameterType::BindlessResourceIndex;
-	}
-
-	if (InName = UE::String::RemoveFromStart(InName, FStringView(UE::ShaderCompilerCommon::kBindlessSamplerPrefix)); InName != OriginalName)
-	{
-		return EShaderParameterType::BindlessSamplerIndex;
-	}
-
-	return EShaderParameterType::LooseData;
+	return FShaderParameterParser::ParseAndRemoveBindlessParameterPrefix(InName);
 }
 
 EShaderParameterType UE::ShaderCompilerCommon::ParseAndRemoveBindlessParameterPrefix(FString& InName)
 {
-	FStringView Name(InName);
-	const EShaderParameterType ParameterType = ParseAndRemoveBindlessParameterPrefix(Name);
-	InName = FString(Name);
-
-	return ParameterType;
+	return FShaderParameterParser::ParseAndRemoveBindlessParameterPrefix(InName);
 }
 
 bool UE::ShaderCompilerCommon::RemoveBindlessParameterPrefix(FString& InName)
 {
-	return InName.RemoveFromStart(UE::ShaderCompilerCommon::kBindlessResourcePrefix)
-		|| InName.RemoveFromStart(UE::ShaderCompilerCommon::kBindlessSamplerPrefix);
+	return FShaderParameterParser::RemoveBindlessParameterPrefix(InName);
 }
 
 bool UE::ShaderCompilerCommon::ValidatePackedResourceCounts(FShaderCompilerOutput& Output, const FShaderCodePackedResourceCounts& PackedResourceCounts)
@@ -732,7 +637,7 @@ void HandleReflectedGlobalConstantBufferMember(
 )
 {
 	FString MemberName = InMemberName;
-	const EShaderParameterType ParameterType = UE::ShaderCompilerCommon::ParseAndRemoveBindlessParameterPrefix(MemberName);
+	const EShaderParameterType ParameterType = FShaderParameterParser::ParseAndRemoveBindlessParameterPrefix(MemberName);
 
 	Output.ParameterMap.AddParameterAllocation(
 		*MemberName,
@@ -751,7 +656,7 @@ void HandleReflectedUniformBufferConstantBufferMember(
 )
 {
 	FString MemberName = InMemberName;
-	const EShaderParameterType ParameterType = UE::ShaderCompilerCommon::ParseAndRemoveBindlessParameterPrefix(MemberName);
+	const EShaderParameterType ParameterType = FShaderParameterParser::ParseAndRemoveBindlessParameterPrefix(MemberName);
 
 	if (ParameterType != EShaderParameterType::LooseData)
 	{
diff --git a/Engine/Source/Developer/ShaderCompilerCommon/Private/ShaderParameterParser.cpp b/Engine/Source/Developer/ShaderCompilerCommon/Private/ShaderParameterParser.cpp
deleted file mode 100644
diff --git a/Engine/Source/Developer/ShaderCompilerCommon/Public/ShaderCompilerCommon.h b/Engine/Source/Developer/ShaderCompilerCommon/Public/ShaderCompilerCommon.h
--- a/Engine/Source/Developer/ShaderCompilerCommon/Public/ShaderCompilerCommon.h
+++ b/Engine/Source/Developer/ShaderCompilerCommon/Public/ShaderCompilerCommon.h
@@ -18,7 +18,10 @@ class FShaderParameterParser;
 namespace UE::ShaderCompilerCommon
 {
 	static constexpr const TCHAR* kUniformBufferConstantBufferPrefix = TEXT("UniformBufferConstants_");
+
+	UE_DEPRECATED(5.3, "No longer used; use FShaderParameterParser::kBindlessResourcePrefix instead")
 	static constexpr const TCHAR* kBindlessResourcePrefix = TEXT("BindlessResource_");
+	UE_DEPRECATED(5.3, "No longer used; use FShaderParameterParser::kBindlessSamplerPrefix instead")
 	static constexpr const TCHAR* kBindlessSamplerPrefix = TEXT("BindlessSampler_");
 }
 
@@ -58,16 +61,25 @@ extern SHADERCOMPILERCOMMON_API int16 GetNumUniformBuffersUsed(const FShaderComp
 
 namespace UE::ShaderCompilerCommon
 {
-	extern SHADERCOMPILERCOMMON_API bool                 ShouldUseStableConstantBuffer(const FShaderCompilerInput& Input);
+	UE_DEPRECATED(5.3, "Deprecated; use FShaderCompilerInput::ShouldUseStableContantBuffer directly")
+	inline bool ShouldUseStableConstantBuffer(const FShaderCompilerInput& Input)
+	{
+		return Input.ShouldUseStableConstantBuffer();
+	}
 
+	UE_DEPRECATED(5.3, "Deprecated; use FShaderParameterParser::ParseParameterType instead")
 	extern SHADERCOMPILERCOMMON_API EShaderParameterType ParseParameterType(FStringView InType, TArrayView<const TCHAR* const> InExtraSRVTypes, TArrayView<const TCHAR* const> InExtraUAVTypes);
 
 	extern SHADERCOMPILERCOMMON_API FStringView          RemoveConstantBufferPrefix(FStringView InName);
 	extern SHADERCOMPILERCOMMON_API FString              RemoveConstantBufferPrefix(const FString& InName);
 
+	UE_DEPRECATED(5.3, "Deprecated; use FShaderParameterParser::ParseAndRemoveBindlessParameterPrefix instead")
 	extern SHADERCOMPILERCOMMON_API EShaderParameterType ParseAndRemoveBindlessParameterPrefix(FStringView& InName);
+	
+	UE_DEPRECATED(5.3, "Deprecated; use FShaderParameterParser::ParseAndRemoveBindlessParameterPrexix instead")
 	extern SHADERCOMPILERCOMMON_API EShaderParameterType ParseAndRemoveBindlessParameterPrefix(FString& InName);
 
+	UE_DEPRECATED(5.3, "Deprecated; use FShaderParameterParser::RemoveBindlessParameterPrefix instead")
 	extern SHADERCOMPILERCOMMON_API bool                 RemoveBindlessParameterPrefix(FString& InName);
 
 	extern SHADERCOMPILERCOMMON_API bool                 ValidatePackedResourceCounts(FShaderCompilerOutput& Output, const FShaderCodePackedResourceCounts& PackedResourceCounts);
diff --git a/Engine/Source/Developer/ShaderCompilerCommon/Public/ShaderParameterParser.h b/Engine/Source/Developer/ShaderCompilerCommon/Public/ShaderParameterParser.h
deleted file mode 100644
diff --git a/Engine/Source/Developer/VulkanShaderFormat/Private/VulkanShaderCompiler.cpp b/Engine/Source/Developer/VulkanShaderFormat/Private/VulkanShaderCompiler.cpp
--- a/Engine/Source/Developer/VulkanShaderFormat/Private/VulkanShaderCompiler.cpp
+++ b/Engine/Source/Developer/VulkanShaderFormat/Private/VulkanShaderCompiler.cpp
@@ -1105,7 +1105,7 @@ static void BuildShaderOutput(
 		);
 
 		FString ParamName = PackedGlobal.Name;
-		UE::ShaderCompilerCommon::RemoveBindlessParameterPrefix(ParamName);
+		FShaderParameterParser::RemoveBindlessParameterPrefix(ParamName);
 		NEWEntryTypes.Add(ParamName, FVulkanShaderHeader::PackedGlobal);
 
 		uint32& Size = PackedGlobalArraySize.FindOrAdd((CrossCompiler::EPackedTypeName)PackedGlobal.PackedType);
diff --git a/Engine/Source/Developer/Windows/ShaderFormatD3D/Private/D3DShaderCompiler.inl b/Engine/Source/Developer/Windows/ShaderFormatD3D/Private/D3DShaderCompiler.inl
--- a/Engine/Source/Developer/Windows/ShaderFormatD3D/Private/D3DShaderCompiler.inl
+++ b/Engine/Source/Developer/Windows/ShaderFormatD3D/Private/D3DShaderCompiler.inl
@@ -37,7 +37,7 @@ template <typename ID3D1xShaderReflection, typename D3D1x_SHADER_DESC, typename
 
 			if (bGlobalCB)
 			{
-				if (UE::ShaderCompilerCommon::ShouldUseStableConstantBuffer(Input))
+				if (Input.ShouldUseStableConstantBuffer())
 				{
 					// Each member found in the global constant buffer means it was not in RootParametersStructure or
 					// it would have been moved by ShaderParameterParser.ParseAndModify().
@@ -81,7 +81,7 @@ template <typename ID3D1xShaderReflection, typename D3D1x_SHADER_DESC, typename
 					}
 				}
 			}
-			else if (bIsRootCB && UE::ShaderCompilerCommon::ShouldUseStableConstantBuffer(Input))
+			else if (bIsRootCB && Input.ShouldUseStableConstantBuffer())
 			{
 				if (CBIndex == FShaderParametersMetadata::kRootCBufferBindingIndex)
 				{
diff --git a/Engine/Source/Runtime/RenderCore/Private/ShaderParameterParser.cpp b/Engine/Source/Runtime/RenderCore/Private/ShaderParameterParser.cpp
new file mode 100644
--- /dev/null
+++ b/Engine/Source/Runtime/RenderCore/Private/ShaderParameterParser.cpp
@@ -0,0 +1,1353 @@
+// Copyright Epic Games, Inc. All Rights Reserved.
+
+#include "ShaderParameterParser.h"
+#include "Containers/UnrealString.h"
+#include "ShaderCompilerCore.h"
+#include "String/RemoveFrom.h"
+#include "Misc/StringBuilder.h"
+
+template<typename TParameterFunction>
+static void IterateShaderParameterMembersInternal(
+	const FShaderParametersMetadata& ParametersMetadata,
+	uint16 ByteOffset,
+	TStringBuilder<1024>& ShaderBindingNameBuilder,
+	TParameterFunction Lambda)
+{
+	for (const FShaderParametersMetadata::FMember& Member : ParametersMetadata.GetMembers())
+	{
+		EUniformBufferBaseType BaseType = Member.GetBaseType();
+		const uint16 MemberOffset = ByteOffset + uint16(Member.GetOffset());
+		const uint32 NumElements = Member.GetNumElements();
+
+		int32 MemberNameLength = FCString::Strlen(Member.GetName());
+
+		if (BaseType == UBMT_INCLUDED_STRUCT)
+		{
+			check(NumElements == 0);
+			const FShaderParametersMetadata& NewParametersMetadata = *Member.GetStructMetadata();
+			IterateShaderParameterMembersInternal(NewParametersMetadata, MemberOffset, ShaderBindingNameBuilder, Lambda);
+		}
+		else if (BaseType == UBMT_NESTED_STRUCT && NumElements == 0)
+		{
+			ShaderBindingNameBuilder.Append(Member.GetName());
+			ShaderBindingNameBuilder.Append(TEXT("_"));
+
+			const FShaderParametersMetadata& NewParametersMetadata = *Member.GetStructMetadata();
+			IterateShaderParameterMembersInternal(NewParametersMetadata, MemberOffset, ShaderBindingNameBuilder, Lambda);
+
+			ShaderBindingNameBuilder.RemoveSuffix(MemberNameLength + 1);
+		}
+		else if (BaseType == UBMT_NESTED_STRUCT && NumElements > 0)
+		{
+			ShaderBindingNameBuilder.Append(Member.GetName());
+			ShaderBindingNameBuilder.Append(TEXT("_"));
+
+			const FShaderParametersMetadata& NewParametersMetadata = *Member.GetStructMetadata();
+			for (uint32 ArrayElementId = 0; ArrayElementId < NumElements; ArrayElementId++)
+			{
+				FString ArrayElementIdString;
+				ArrayElementIdString.AppendInt(ArrayElementId);
+				int32 ArrayElementIdLength = ArrayElementIdString.Len();
+
+				ShaderBindingNameBuilder.Append(ArrayElementIdString);
+				ShaderBindingNameBuilder.Append(TEXT("_"));
+
+				uint16 NewStructOffset = MemberOffset + ArrayElementId * NewParametersMetadata.GetSize();
+				IterateShaderParameterMembersInternal(NewParametersMetadata, NewStructOffset, ShaderBindingNameBuilder, Lambda);
+
+				ShaderBindingNameBuilder.RemoveSuffix(ArrayElementIdLength + 1);
+			}
+
+			ShaderBindingNameBuilder.RemoveSuffix(MemberNameLength + 1);
+		}
+		else
+		{
+			const bool bParametersAreExpanded =
+				NumElements > 0 &&
+				(BaseType == UBMT_TEXTURE ||
+					BaseType == UBMT_SRV ||
+					BaseType == UBMT_UAV ||
+					BaseType == UBMT_SAMPLER ||
+					IsRDGResourceReferenceShaderParameterType(BaseType));
+
+			if (bParametersAreExpanded)
+			{
+				const uint16 ElementSize = SHADER_PARAMETER_POINTER_ALIGNMENT;
+
+				for (uint32 Index = 0; Index < NumElements; Index++)
+				{
+					const FString RealBindingName = FString::Printf(TEXT("%s_%d"), Member.GetName(), Index);
+
+					ShaderBindingNameBuilder.Append(RealBindingName);
+					Lambda(ParametersMetadata, Member, ShaderBindingNameBuilder.ToString(), MemberOffset + Index * ElementSize);
+					ShaderBindingNameBuilder.RemoveSuffix(RealBindingName.Len());
+				}
+			}
+			else
+			{
+				ShaderBindingNameBuilder.Append(Member.GetName());
+				Lambda(ParametersMetadata, Member, ShaderBindingNameBuilder.ToString(), MemberOffset);
+				ShaderBindingNameBuilder.RemoveSuffix(MemberNameLength);
+			}
+		}
+	}
+}
+
+template<typename TParameterFunction>
+static void IterateShaderParameterMembers(const FShaderParametersMetadata& ShaderParametersMetadata, TParameterFunction Lambda)
+{
+	FShaderParametersMetadata::EUseCase UseCase = ShaderParametersMetadata.GetUseCase();
+
+	TStringBuilder<1024> ShaderBindingNameBuilder;
+	if (UseCase == FShaderParametersMetadata::EUseCase::UniformBuffer || UseCase == FShaderParametersMetadata::EUseCase::DataDrivenUniformBuffer)
+	{
+		ShaderBindingNameBuilder.Append(ShaderParametersMetadata.GetShaderVariableName());
+		ShaderBindingNameBuilder.Append(TEXT("_"));
+	}
+
+	IterateShaderParameterMembersInternal(
+		ShaderParametersMetadata, /* ByteOffset = */ 0, ShaderBindingNameBuilder, Lambda);
+}
+
+static void AddNoteToDisplayShaderParameterMemberOnCppSide(
+	const FShaderCompilerInput& CompilerInput,
+	const FShaderParameterParser::FParsedShaderParameter& ParsedParameter,
+	FShaderCompilerOutput& CompilerOutput)
+{
+	const FShaderParametersMetadata* MemberContainingStruct = nullptr;
+	const FShaderParametersMetadata::FMember* Member = nullptr;
+	{
+		int32 ArrayElementId = 0;
+		FString NamePrefix;
+		CompilerInput.RootParametersStructure->FindMemberFromOffset(ParsedParameter.ConstantBufferOffset, &MemberContainingStruct, &Member, &ArrayElementId, &NamePrefix);
+	}
+
+	FString CppCodeName = CompilerInput.RootParametersStructure->GetFullMemberCodeName(ParsedParameter.ConstantBufferOffset);
+
+	FShaderCompilerError Error;
+	Error.StrippedErrorMessage = FString::Printf(
+		TEXT("Note: Definition of %s"),
+		*CppCodeName);
+	Error.ErrorVirtualFilePath = ANSI_TO_TCHAR(MemberContainingStruct->GetFileName());
+	Error.ErrorLineString = FString::FromInt(Member->GetFileLine());
+
+	CompilerOutput.Errors.Add(Error);
+}
+
+FShaderParameterParser::FShaderParameterParser() = default;
+FShaderParameterParser::~FShaderParameterParser() = default;
+
+FShaderParameterParser::FShaderParameterParser(const TCHAR* InConstantBufferType)
+	: ConstantBufferType(InConstantBufferType)
+{
+}
+
+FShaderParameterParser::FShaderParameterParser(
+	const TCHAR* InConstantBufferType,
+	TArrayView<const TCHAR* const> InExtraSRVTypes,
+	TArrayView<const TCHAR* const> InExtraUAVTypes
+)
+	: ConstantBufferType(InConstantBufferType)
+	, ExtraSRVTypes(InExtraSRVTypes)
+	, ExtraUAVTypes(InExtraUAVTypes)
+{
+}
+
+static const TCHAR* const s_AllSRVTypes[] =
+{
+	TEXT("Texture1D"),
+	TEXT("Texture1DArray"),
+	TEXT("Texture2D"),
+	TEXT("Texture2DArray"),
+	TEXT("Texture2DMS"),
+	TEXT("Texture2DMSArray"),
+	TEXT("Texture3D"),
+	TEXT("TextureCube"),
+	TEXT("TextureCubeArray"),
+
+	TEXT("Buffer"),
+	TEXT("ByteAddressBuffer"),
+	TEXT("StructuredBuffer"),
+	TEXT("ConstantBuffer"),
+	TEXT("RaytracingAccelerationStructure"),
+};
+
+static const TCHAR* const s_AllUAVTypes[] =
+{
+	TEXT("AppendStructuredBuffer"),
+	TEXT("RWBuffer"),
+	TEXT("RWByteAddressBuffer"),
+	TEXT("RWStructuredBuffer"),
+	TEXT("RWTexture1D"),
+	TEXT("RWTexture1DArray"),
+	TEXT("RWTexture2D"),
+	TEXT("RWTexture2DArray"),
+	TEXT("RWTexture3D"),
+	TEXT("RasterizerOrderedTexture2D"),
+};
+
+static const TCHAR* const s_AllSamplerTypes[] =
+{
+	TEXT("SamplerState"),
+	TEXT("SamplerComparisonState"),
+};
+
+EShaderParameterType FShaderParameterParser::ParseParameterType(
+	FStringView InType,
+	TArrayView<const TCHAR* const> InExtraSRVTypes,
+	TArrayView<const TCHAR* const> InExtraUAVTypes)
+{
+	TArrayView<const TCHAR* const> AllSamplerTypes(s_AllSamplerTypes);
+	TArrayView<const TCHAR* const> AllSRVTypes(s_AllSRVTypes);
+	TArrayView<const TCHAR* const> AllUAVTypes(s_AllUAVTypes);
+
+	if (AllSamplerTypes.Contains(InType))
+	{
+		return EShaderParameterType::Sampler;
+	}
+
+	FStringView UntemplatedType = InType;
+	if (int32 Index = InType.Find(TEXT("<")); Index != INDEX_NONE)
+	{
+		// Remove the template argument but don't forget to clean up the type name
+		const int32 NumChars = InType.Len() - Index;
+		UntemplatedType = InType.LeftChop(NumChars).TrimEnd();
+	}
+
+	if (AllSRVTypes.Contains(UntemplatedType) || InExtraSRVTypes.Contains(UntemplatedType))
+	{
+		return EShaderParameterType::SRV;
+	}
+
+	if (AllUAVTypes.Contains(UntemplatedType) || InExtraUAVTypes.Contains(UntemplatedType))
+	{
+		return EShaderParameterType::UAV;
+	}
+
+	return EShaderParameterType::LooseData;
+}
+
+EShaderParameterType FShaderParameterParser::ParseAndRemoveBindlessParameterPrefix(FStringView& InName)
+{
+	const FStringView OriginalName = InName;
+
+	if (InName = UE::String::RemoveFromStart(InName, FStringView(kBindlessResourcePrefix)); InName != OriginalName)
+	{
+		return EShaderParameterType::BindlessResourceIndex;
+	}
+
+	if (InName = UE::String::RemoveFromStart(InName, FStringView(kBindlessSamplerPrefix)); InName != OriginalName)
+	{
+		return EShaderParameterType::BindlessSamplerIndex;
+	}
+
+	return EShaderParameterType::LooseData;
+}
+
+EShaderParameterType FShaderParameterParser::ParseAndRemoveBindlessParameterPrefix(FString& InName)
+{
+	FStringView Name(InName);
+	const EShaderParameterType ParameterType = ParseAndRemoveBindlessParameterPrefix(Name);
+	InName = FString(Name);
+
+	return ParameterType;
+}
+
+bool FShaderParameterParser::RemoveBindlessParameterPrefix(FString& InName)
+{
+	return InName.RemoveFromStart(kBindlessResourcePrefix)
+		|| InName.RemoveFromStart(kBindlessSamplerPrefix);
+}
+
+
+bool FShaderParameterParser::ParseParameters(
+	const FShaderCompilerInput& CompilerInput,
+	FShaderCompilerOutput& CompilerOutput)
+{
+	const FStringView ShaderSource(OriginalParsedShader);
+
+	if (CompilerInput.RootParametersStructure)
+	{
+		// Reserves the number of parameters up front.
+		ParsedParameters.Reserve(CompilerInput.RootParametersStructure->GetSize() / sizeof(int32));
+
+		IterateShaderParameterMembers(
+			*CompilerInput.RootParametersStructure,
+			[&](const FShaderParametersMetadata& ParametersMetadata,
+				const FShaderParametersMetadata::FMember& Member,
+				const TCHAR* ShaderBindingName,
+				uint16 ByteOffset)
+			{
+				FParsedShaderParameter ParsedParameter;
+				ParsedParameter.bIsBindable = true;
+				ParsedParameter.ConstantBufferOffset = ByteOffset;
+				ParsedParameter.BaseType = Member.GetBaseType();
+				ParsedParameter.PrecisionModifier = Member.GetPrecision();
+				ParsedParameter.NumRows = Member.GetNumRows();
+				ParsedParameter.NumColumns = Member.GetNumColumns();
+				ParsedParameter.MemberSize = Member.IsVariableNativeType() ? Member.GetMemberSize() : 0u;
+				ParsedParameters.Add(ShaderBindingName, ParsedParameter);
+			});
+	}
+
+	bool bSuccess = true;
+
+	// Browse the code for global shader parameter, Save their type and erase them white spaces.
+	{
+		enum class EState
+		{
+			// When to look for something to scan.
+			Scanning,
+
+			// When going to next ; in the global scope and reset.
+			GoToNextSemicolonAndReset,
+
+			// Parsing what might be a type of the parameter.
+			ParsingPotentialType,
+			ParsingPotentialTypeTemplateArguments,
+			FinishedPotentialType,
+
+			// Parsing what might be a name of the parameter.
+			ParsingPotentialName,
+			FinishedPotentialName,
+
+			// Parsing what looks like array of the parameter.
+			ParsingPotentialArraySize,
+			FinishedArraySize,
+
+			// Found a parameter, just finish to it's semi colon.
+			FoundParameter,
+		};
+
+		const int32 ShaderSourceLen = ShaderSource.Len();
+
+		int32 CurrentPragamLineoffset = -1;
+		int32 CurrentLineoffset = 0;
+
+		int32 TypeQualifierStartPos = -1;
+		int32 TypeStartPos = -1;
+		int32 TypeEndPos = -1;
+		int32 NameStartPos = -1;
+		int32 NameEndPos = -1;
+		int32 ArrayStartPos = -1;
+		int32 ArrayEndPos = -1;
+		int32 ScopeIndent = 0;
+		bool bGloballyCoherent = false;
+
+		EState State = EState::Scanning;
+		bool bGoToNextLine = false;
+
+		auto ResetState = [&]()
+		{
+			TypeQualifierStartPos = -1;
+			TypeStartPos = -1;
+			TypeEndPos = -1;
+			NameStartPos = -1;
+			NameEndPos = -1;
+			ArrayStartPos = -1;
+			ArrayEndPos = -1;
+			bGloballyCoherent = false;
+			State = EState::Scanning;
+		};
+
+		auto EmitError = [&](const FString& ErrorMessage)
+		{
+			FShaderCompilerError Error;
+			Error.StrippedErrorMessage = ErrorMessage;
+			ExtractFileAndLine(CurrentPragamLineoffset, CurrentLineoffset, Error.ErrorVirtualFilePath, Error.ErrorLineString);
+			CompilerOutput.Errors.Add(Error);
+			bSuccess = false;
+		};
+
+		auto EmitUnpextectedHLSLSyntaxError = [&]()
+		{
+			EmitError(TEXT("Unexpected syntax when parsing shader parameters from shader code."));
+			State = EState::GoToNextSemicolonAndReset;
+		};
+
+		for (int32 Cursor = 0; Cursor < ShaderSourceLen; Cursor++)
+		{
+			const TCHAR Char = ShaderSource[Cursor];
+
+			auto FoundShaderParameter = [&]()
+			{
+				check(Char == ';');
+				check(TypeStartPos != -1);
+				check(TypeEndPos != -1);
+				check(NameStartPos != -1);
+				check(NameEndPos != -1);
+
+				FStringView Type = ShaderSource.Mid(TypeStartPos, TypeEndPos - TypeStartPos + 1);
+				FStringView Name = ShaderSource.Mid(NameStartPos, NameEndPos - NameStartPos + 1);
+
+				FStringView Leftovers = ShaderSource.Mid(NameEndPos + 1, (Cursor - 1) - (NameEndPos + 1) + 1);
+
+				EShaderParameterType ParsedParameterType = ParseAndRemoveBindlessParameterPrefix(Name);
+				const bool bBindlessIndex = (ParsedParameterType == EShaderParameterType::BindlessResourceIndex || ParsedParameterType == EShaderParameterType::BindlessSamplerIndex);
+
+				EBindlessConversionType BindlessConversionType = EBindlessConversionType::None;
+				EShaderParameterType ParsedConstantBufferType = ParsedParameterType;
+
+				if ((bBindlessResources || bBindlessSamplers) && ParsedParameterType == EShaderParameterType::LooseData)
+				{
+					ParsedParameterType = ParseParameterType(Type, ExtraSRVTypes, ExtraUAVTypes);
+
+					if (bBindlessResources && (ParsedParameterType == EShaderParameterType::SRV || ParsedParameterType == EShaderParameterType::UAV))
+					{
+						BindlessConversionType = EBindlessConversionType::Resource;
+						ParsedConstantBufferType = EShaderParameterType::BindlessResourceIndex;
+					}
+					else if (bBindlessSamplers && ParsedParameterType == EShaderParameterType::Sampler)
+					{
+						BindlessConversionType = EBindlessConversionType::Sampler;
+						ParsedConstantBufferType = EShaderParameterType::BindlessSamplerIndex;
+					}
+
+					if (BindlessConversionType != EBindlessConversionType::None && Leftovers.Contains(TEXT("register")))
+					{
+						// avoid rewriting hardcoded register assignments
+						BindlessConversionType = EBindlessConversionType::None;
+						ParsedConstantBufferType = ParsedParameterType;
+					}
+				}
+
+				FParsedShaderParameter ParsedParameter;
+
+				EShaderParameterType ConstantBufferParameterType = EShaderParameterType::Num;
+				bool bMoveToRootConstantBuffer = false;
+				bool bUpdateParsedParameters = false;
+
+ 				const FString ParsedParameterKey(Name);
+				if (ParsedParameters.Contains(ParsedParameterKey))
+				{
+					if (ParsedParameters.FindChecked(ParsedParameterKey).IsFound())
+					{
+						// If it has already been found, it means it is duplicated. Do nothing and let the shader compiler throw the error.
+					}
+					else
+					{
+						// Update the parsed parameters
+						bUpdateParsedParameters = true;
+						ParsedParameter = ParsedParameters.FindChecked(ParsedParameterKey);
+
+						// Erase the parameter to move it into the root constant buffer.
+						if (bNeedToMoveToRootConstantBuffer && ParsedParameter.bIsBindable)
+						{
+							const EUniformBufferBaseType BaseType = ParsedParameter.BaseType;
+							bMoveToRootConstantBuffer =
+								BaseType == UBMT_INT32 ||
+								BaseType == UBMT_UINT32 ||
+								BaseType == UBMT_FLOAT32 ||
+								bBindlessIndex ||
+								(BindlessConversionType != EBindlessConversionType::None);
+
+							if (bMoveToRootConstantBuffer)
+							{
+								ConstantBufferParameterType = ParsedConstantBufferType;
+							}
+						}
+					}
+				}
+				else
+				{
+					// Update the parsed parameters to still have file and line number.
+					bUpdateParsedParameters = true;
+				}
+
+				// Update 
+				if (bUpdateParsedParameters)
+				{
+					ParsedParameter.ParsedName = Name;
+					ParsedParameter.ParsedType = Type;
+					ParsedParameter.ParsedPragmaLineoffset = CurrentPragamLineoffset;
+					ParsedParameter.ParsedLineOffset = CurrentLineoffset;
+					ParsedParameter.ParsedCharOffsetStart = TypeQualifierStartPos != -1 ? TypeQualifierStartPos : TypeStartPos;
+					ParsedParameter.ParsedCharOffsetEnd = Cursor;
+					ParsedParameter.BindlessConversionType = BindlessConversionType;
+					ParsedParameter.ConstantBufferParameterType = ConstantBufferParameterType;
+					ParsedParameter.bGloballyCoherent = bGloballyCoherent;
+
+					if (ArrayStartPos != -1 && ArrayEndPos != -1)
+					{
+						ParsedParameter.ParsedArraySize = ShaderSource.Mid(ArrayStartPos + 1, ArrayEndPos - ArrayStartPos - 1);
+					}
+
+					ParsedParameters.Add(ParsedParameterKey, ParsedParameter);
+				}
+
+				ResetState();
+			};
+
+			const bool bIsWhiteSpace = Char == ' ' || Char == '\t' || Char == '\r' || Char == '\n';
+			const bool bIsLetter = (Char >= 'a' && Char <= 'z') || (Char >= 'A' && Char <= 'Z');
+			const bool bIsNumber = Char >= '0' && Char <= '9';
+
+			const TCHAR* UpComing = ShaderSource.GetData() + Cursor;
+			const int32 RemainingSize = ShaderSourceLen - Cursor;
+
+			CurrentLineoffset += Char == '\n';
+
+			// Go to the next line if this is a preprocessor macro.
+			if (bGoToNextLine)
+			{
+				if (Char == '\n')
+				{
+					bGoToNextLine = false;
+				}
+				continue;
+			}
+			else if (Char == '#')
+			{
+				if (RemainingSize > 6 && FCString::Strncmp(UpComing, TEXT("#line "), 6) == 0)
+				{
+					CurrentPragamLineoffset = Cursor;
+					CurrentLineoffset = -1; // that will be incremented to 0 when reaching the \n at the end of the #line
+				}
+
+				bGoToNextLine = true;
+				continue;
+			}
+
+			// If within a scope, just carry on until outside the scope.
+			if (ScopeIndent > 0 || Char == '{')
+			{
+				if (Char == '{')
+				{
+					ScopeIndent++;
+				}
+				else if (Char == '}')
+				{
+					ScopeIndent--;
+					if (ScopeIndent == 0)
+					{
+						ResetState();
+					}
+				}
+				continue;
+			}
+
+			if (State == EState::Scanning)
+			{
+				if (bIsLetter)
+				{
+					static const TCHAR* KeywordTable[] =
+					{
+						TEXT("const"),
+						TEXT("globallycoherent")
+						TEXT("enum"),
+						TEXT("class"),
+						TEXT("struct"),
+						TEXT("static"),
+					};
+					static int32 KeywordTableSize[] = { 5, 16, 4, 5, 6, 6 };
+
+					int32 RecognisedKeywordId = -1;
+					for (int32 KeywordId = 0; KeywordId < UE_ARRAY_COUNT(KeywordTable); KeywordId++)
+					{
+						const TCHAR* Keyword = KeywordTable[KeywordId];
+						const int32 KeywordSize = KeywordTableSize[KeywordId];
+
+						if (RemainingSize > KeywordSize)
+						{
+							TCHAR KeywordEndTestChar = UpComing[KeywordSize];
+
+							if ((KeywordEndTestChar == ' ' || KeywordEndTestChar == '\r' || KeywordEndTestChar == '\n' || KeywordEndTestChar == '\t') &&
+								FCString::Strncmp(UpComing, Keyword, KeywordSize) == 0)
+							{
+								RecognisedKeywordId = KeywordId;
+								break;
+							}
+						}
+					}
+
+					if (RecognisedKeywordId == -1)
+					{
+						// Might have found beginning of the type of a parameter.
+						State = EState::ParsingPotentialType;
+						TypeStartPos = Cursor;
+					}
+					else if (RecognisedKeywordId == 0)
+					{
+						// Ignore the const keywords, but still parse given it might still be a shader parameter.
+						if (TypeQualifierStartPos == -1)
+						{
+							// If the parameter is erased, we also have to erase *all* 'const'-qualifiers, e.g. "const int Foo" or "const const int Foo".
+							TypeQualifierStartPos = Cursor;
+						}
+						Cursor += KeywordTableSize[RecognisedKeywordId];
+					}
+					else if (RecognisedKeywordId == 1)
+					{
+						// Mark that we got the globallycoherent keyword and keep moving to the next set of qualifiers
+						bGloballyCoherent = true;
+						if (TypeQualifierStartPos == -1)
+						{
+							// If the parameter is erased, we also have to erase *all* qualifiers, e.g. "const int Foo" or "const const int Foo".
+							TypeQualifierStartPos = Cursor;
+						}
+						Cursor += KeywordTableSize[RecognisedKeywordId];
+					}
+					else
+					{
+						// Purposefully ignore enum, class, struct, static
+						State = EState::GoToNextSemicolonAndReset;
+					}
+				}
+				else if (bIsWhiteSpace)
+				{
+					// Keep parsing void.
+				}
+				else if (Char == ';')
+				{
+					// Looks like redundant semicolon, just ignore and keep scanning.
+				}
+				else
+				{
+					// No idea what this is, just go to next semi colon.
+					State = EState::GoToNextSemicolonAndReset;
+				}
+			}
+			else if (State == EState::GoToNextSemicolonAndReset)
+			{
+				// If need to go to next global semicolon and reach it. Resume browsing.
+				if (Char == ';')
+				{
+					ResetState();
+				}
+			}
+			else if (State == EState::ParsingPotentialType)
+			{
+				// Found character legal for a type...
+				if (bIsLetter ||
+					bIsNumber ||
+					Char == '_')
+				{
+					// Keep browsing what might be type of the parameter.
+				}
+				else if (Char == '<')
+				{
+					// Found what looks like the beginning of template argument that is legal on resource types for Instance Texture2D< float >
+					State = EState::ParsingPotentialTypeTemplateArguments;
+				}
+				else if (bIsWhiteSpace)
+				{
+					// Might have found a type.
+					State = EState::FinishedPotentialType;
+					TypeEndPos = Cursor - 1;
+				}
+				else
+				{
+					// Found unexpected character in the type.
+					State = EState::GoToNextSemicolonAndReset;
+				}
+			}
+			else if (State == EState::ParsingPotentialTypeTemplateArguments)
+			{
+				// Found character legal for a template argument...
+				if (bIsLetter ||
+					bIsNumber ||
+					Char == '_')
+				{
+					// Keep browsing what might be type of the parameter.
+				}
+				else if (bIsWhiteSpace || Char == ',')
+				{
+					// Spaces and comas are legal within agrument of the template arguments.
+				}
+				else if (Char == '>')
+				{
+					// Might have found a type with template argument.
+					State = EState::FinishedPotentialType;
+					TypeEndPos = Cursor;
+				}
+				else
+				{
+					// Found unexpected character in the type.
+					State = EState::GoToNextSemicolonAndReset;
+				}
+			}
+			else if (State == EState::FinishedPotentialType)
+			{
+				if (bIsLetter)
+				{
+					// Might have found beginning of the name of a parameter.
+					State = EState::ParsingPotentialName;
+					NameStartPos = Cursor;
+				}
+				else if (Char == '<')
+				{
+					// Might have found beginning of a template argument for the type, that was separate by a whitespace from type. For instance Texture2D <float>
+					State = EState::ParsingPotentialTypeTemplateArguments;
+				}
+				else if (bIsWhiteSpace)
+				{
+					// Keep parsing void.
+				}
+				else
+				{
+					// No idea what this is, just go to next semi colon.
+					State = EState::GoToNextSemicolonAndReset;
+				}
+			}
+			else if (State == EState::ParsingPotentialName)
+			{
+				// Found character legal for a name...
+				if (bIsLetter ||
+					bIsNumber ||
+					Char == '_')
+				{
+					// Keep browsing what might be name of the parameter.
+				}
+				else if (Char == ':' || Char == '=')
+				{
+					// Found a parameter with syntax:
+					// uint MyParameter : <whatever>;
+					// uint MyParameter = <DefaultValue>;
+					NameEndPos = Cursor - 1;
+					State = EState::FoundParameter;
+				}
+				else if (Char == ';')
+				{
+					// Found a parameter with syntax:
+					// uint MyParameter;
+					NameEndPos = Cursor - 1;
+					FoundShaderParameter();
+				}
+				else if (Char == '[')
+				{
+					// Syntax:
+					//  uint MyArray[
+					NameEndPos = Cursor - 1;
+					ArrayStartPos = Cursor;
+					State = EState::ParsingPotentialArraySize;
+				}
+				else if (bIsWhiteSpace)
+				{
+					// Might have found a name.
+					// uint MyParameter <Still need to know what is after>;
+					NameEndPos = Cursor - 1;
+					State = EState::FinishedPotentialName;
+				}
+				else
+				{
+					// Found unexpected character in the name.
+					// syntax:
+					// uint MyFunction(<Don't care what is after>
+					State = EState::GoToNextSemicolonAndReset;
+				}
+			}
+			else if (State == EState::FinishedPotentialName || State == EState::FinishedArraySize)
+			{
+				if (Char == ';')
+				{
+					// Found a parameter with syntax:
+					// uint MyParameter <a bit of OK stuf>;
+					FoundShaderParameter();
+				}
+				else if (Char == ':')
+				{
+					// Found a parameter with syntax:
+					// uint MyParameter <a bit of OK stuf> : <Ignore all this crap>;
+					State = EState::FoundParameter;
+				}
+				else if (Char == '=')
+				{
+					// Found syntax that doesn't make any sens:
+					// uint MyParameter <a bit of OK stuf> = <Ignore all this crap>;
+					State = EState::FoundParameter;
+					// TDOO: should error out that this is useless.
+				}
+				else if (Char == '[')
+				{
+					if (State == EState::FinishedPotentialName)
+					{
+						// Syntax:
+						//  uint MyArray [
+						ArrayStartPos = Cursor;
+						State = EState::ParsingPotentialArraySize;
+					}
+					else
+					{
+						EmitError(TEXT("Shader parameters can only support one dimensional array"));
+					}
+				}
+				else if (bIsWhiteSpace)
+				{
+					// Keep parsing void.
+				}
+				else
+				{
+					// Found unexpected stuff.
+					State = EState::GoToNextSemicolonAndReset;
+				}
+			}
+			else if (State == EState::ParsingPotentialArraySize)
+			{
+				if (Char == ']')
+				{
+					ArrayEndPos = Cursor;
+					State = EState::FinishedArraySize;
+				}
+				else if (Char == ';')
+				{
+					EmitUnpextectedHLSLSyntaxError();
+				}
+				else
+				{
+					// Keep going through the array size that might be a complex expression.
+				}
+			}
+			else if (State == EState::FoundParameter)
+			{
+				if (Char == ';')
+				{
+					FoundShaderParameter();
+				}
+				else
+				{
+					// Cary on skipping all crap we don't care about shader parameter until we find it's semi colon.
+				}
+			}
+			else
+			{
+				unimplemented();
+			}
+		} // for (int32 Cursor = 0; Cursor < ShaderSourceLen; Cursor++)
+	}
+
+	return bSuccess;
+}
+
+void FShaderParameterParser::RemoveMovingParametersFromSource(FString& PreprocessedShaderSource)
+{
+	for (TPair<FString, FParsedShaderParameter>& Itr : ParsedParameters)
+	{
+		const FParsedShaderParameter& ParsedParameter = Itr.Value;
+
+		// If this parameter is going to be in the root constant buffer
+		if (ParsedParameter.ConstantBufferParameterType != EShaderParameterType::Num &&
+			// but it's not being converted to bindless
+			ParsedParameter.BindlessConversionType == EBindlessConversionType::None &&
+			ParsedParameter.ParsedCharOffsetStart != INDEX_NONE)
+		{
+			// then erase this shader parameter conserving the same line numbers.
+			for (int32 j = ParsedParameter.ParsedCharOffsetStart; j <= ParsedParameter.ParsedCharOffsetEnd; j++)
+			{
+				if (PreprocessedShaderSource[j] != '\r' && PreprocessedShaderSource[j] != '\n')
+				{
+					PreprocessedShaderSource[j] = ' ';
+				}
+			}
+		}
+	}
+}
+
+FString FShaderParameterParser::GenerateBindlessParameterDeclaration(const FParsedShaderParameter& ParsedParameter) const
+{
+	// NOTE: Macros AUTO_BINDLESS_SAMPLER_INDEX/VARIABLE and AUTO_BINDLESS_RESOURCE_INDEX/VARIABLE in BindlessResources.ush must be kept in sync with this function
+
+	const bool bIsSampler = (ParsedParameter.BindlessConversionType == EBindlessConversionType::Sampler);
+
+	const TCHAR* Kind = bIsSampler ? TEXT("Sampler") : TEXT("Resource");
+	const TCHAR* StorageClass = ParsedParameter.bGloballyCoherent ? TEXT("globallycoherent ") : TEXT("");
+
+	const FStringView Name = ParsedParameter.ParsedName;
+	const FStringView Type = ParsedParameter.ParsedType;
+
+	FString RewriteType(Type);
+
+	FString TypedefText = TEXT("");
+	if (ParsedParameter.BindlessConversionType == EBindlessConversionType::Resource)
+	{
+		RewriteType = FString::Printf(TEXT("SafeType%.*s"), Name.Len(), Name.GetData());
+		TypedefText = FString::Printf(TEXT("typedef %.*s %s;"), Type.Len(), Type.GetData(), *RewriteType);
+	}
+
+	TStringBuilder<512> Result;
+
+	// If we weren't going to be added to a root constant buffer, that means we need to declare our index before we declare our getter.
+	if (ParsedParameter.ConstantBufferParameterType == EShaderParameterType::Num)
+	{
+		// e.g. `uint BindlessResource_##Name;`
+		// or   `uint BindlessSampler_##Name;`
+		Result << TEXT("uint Bindless") << Kind << TEXT("_") << Name << TEXT(";");
+	}
+
+	Result << TypedefText;
+
+	// e.g. `Type GetBindlessResource##Name() { return GetResourceFromHeap(Type, BindlessResource_##Name); } static const Type Name = GetBindlessResource##Name()`
+	// or   `Type GetBindlessSampler##Name() { return GetSamplerFromHeap(Type, BindlessSampler_##Name); } static const Type Name = GetBindlessSampler##Name()`
+	Result << StorageClass << RewriteType << TEXT(" GetBindless") << Kind << Name << TEXT("()");
+	Result << TEXT("{ return Get") << Kind << TEXT("FromHeap(") << StorageClass << RewriteType << TEXT(", Bindless") << Kind << TEXT("_") << Name << TEXT("); } ");
+	Result << TEXT("static const ") << StorageClass << RewriteType << TEXT(" ") << Name << TEXT(" = GetBindless") << Kind << Name << TEXT("();");
+
+	return Result.ToString();
+}
+
+void FShaderParameterParser::ApplyBindlessModifications(FString& PreprocessedShaderSource)
+{
+	if (bBindlessResources || bBindlessSamplers)
+	{
+		// Array of modifications to do on PreprocessedShaderSource
+		struct FShaderCodeModifications
+		{
+			int32 CharOffsetStart;
+			int32 CharOffsetEnd;
+			FString Replace;
+		};
+		TArray<FShaderCodeModifications> Modifications;
+		Modifications.Reserve(ParsedParameters.Num());
+
+		for (TPair<FString, FParsedShaderParameter>& Itr : ParsedParameters)
+		{
+			FParsedShaderParameter& ParsedParameter = Itr.Value;
+
+			if (!ParsedParameter.IsFound())
+			{
+				continue;
+			}
+
+			if (ParsedParameter.BindlessConversionType != EBindlessConversionType::None)
+			{
+				FShaderCodeModifications Modif;
+				Modif.CharOffsetStart = ParsedParameter.ParsedCharOffsetStart;
+				Modif.CharOffsetEnd = ParsedParameter.ParsedCharOffsetEnd + 1;
+				Modif.Replace = GenerateBindlessParameterDeclaration(ParsedParameter);
+
+				Modifications.Add(Modif);
+			}
+		}
+
+		// Apply all modifications
+		if (Modifications.Num() > 0)
+		{
+			// Sort all the modifications in order
+			Modifications.Sort(
+				[](const FShaderCodeModifications& ModifA, const FShaderCodeModifications& ModifB)
+				{
+					return  ModifA.CharOffsetStart < ModifB.CharOffsetStart;
+				}
+			);
+
+			// Find out the size of the shader code after all modifications
+			int32 NewShaderCodeSize = PreprocessedShaderSource.Len();
+			for (const FShaderCodeModifications& Modif : Modifications)
+			{
+				// Count the number of line return \n in CharOffsetStart -> CharOffsetEnd to ensure the line number remain unchanged.
+				check(!Modif.Replace.Contains(TEXT("\n")));
+				//int32 ReplacedCarriagedReturn = 0;
+				for (int32 CharPos = Modif.CharOffsetStart; CharPos < Modif.CharOffsetEnd; CharPos++)
+				{
+					ensure(PreprocessedShaderSource[CharPos] != '\n');
+				}
+
+				NewShaderCodeSize += Modif.Replace.Len() - (Modif.CharOffsetEnd - Modif.CharOffsetStart); // + ReplacedCarriagedReturn;
+			}
+
+			// Splice all the code and modifications together
+			FString NewShaderCode;
+			NewShaderCode.Reserve(NewShaderCodeSize);
+
+			int32 CurrentCodePos = 0;
+			for (const FShaderCodeModifications& Modif : Modifications)
+			{
+				check(CurrentCodePos <= Modif.CharOffsetStart);
+				NewShaderCode += PreprocessedShaderSource.Mid(CurrentCodePos, Modif.CharOffsetStart - CurrentCodePos);
+				NewShaderCode += Modif.Replace;
+				CurrentCodePos = Modif.CharOffsetEnd;
+			}
+			check(CurrentCodePos <= PreprocessedShaderSource.Len());
+			NewShaderCode += PreprocessedShaderSource.Mid(CurrentCodePos, PreprocessedShaderSource.Len() - CurrentCodePos);
+
+			// Commit all modifications to caller
+			PreprocessedShaderSource = NewShaderCode;
+		}
+	}
+}
+
+static const TCHAR* GetConstantSwizzle(uint16 ByteOffset)
+{
+	switch (ByteOffset % 16)
+	{
+	default: unimplemented();
+	case 0:  return TEXT("");
+	case 4:  return TEXT(".y");
+	case 8:  return TEXT(".z");
+	case 12: return TEXT(".w");
+	}
+}
+
+bool FShaderParameterParser::MoveShaderParametersToRootConstantBuffer(
+	const FShaderCompilerInput& CompilerInput,
+	FShaderCompilerOutput& CompilerOutput,
+	FString& PreprocessedShaderSource)
+{
+	bool bSuccess = true;
+
+	// Generate the root cbuffer content.
+	if (CompilerInput.RootParametersStructure && bNeedToMoveToRootConstantBuffer)
+	{
+		FString RootCBufferContent;
+
+		IterateShaderParameterMembers(
+			*CompilerInput.RootParametersStructure,
+			[&](const FShaderParametersMetadata& ParametersMetadata,
+				const FShaderParametersMetadata::FMember& Member,
+				const TCHAR* ShaderBindingName,
+				uint16 ByteOffset)
+		{
+			FParsedShaderParameter* ParsedParameter = ParsedParameters.Find(ShaderBindingName);
+			if (ParsedParameter && ParsedParameter->IsFound() && ParsedParameter->ConstantBufferParameterType != EShaderParameterType::Num)
+			{
+				const uint32 ConstantRegister = ByteOffset / 16;
+				const TCHAR* ConstantSwizzle = GetConstantSwizzle(ByteOffset);
+
+#define SVARG(N) N.Len(), N.GetData()
+				if (ParsedParameter->ConstantBufferParameterType == EShaderParameterType::BindlessResourceIndex)
+				{
+					RootCBufferContent.Append(FString::Printf(
+						TEXT("uint BindlessResource_%.*s : packoffset(c%d%s);\r\n"),
+						SVARG(ParsedParameter->ParsedName),
+						ConstantRegister,
+						ConstantSwizzle));
+				}
+				else if (ParsedParameter->ConstantBufferParameterType == EShaderParameterType::BindlessSamplerIndex)
+				{
+					RootCBufferContent.Append(FString::Printf(
+						TEXT("uint BindlessSampler_%.*s : packoffset(c%d%s);\r\n"),
+						SVARG(ParsedParameter->ParsedName),
+						ConstantRegister,
+						ConstantSwizzle));
+				}
+				else if (ParsedParameter->ConstantBufferParameterType == EShaderParameterType::LooseData)
+				{
+					if (!ParsedParameter->ParsedArraySize.IsEmpty())
+					{
+						RootCBufferContent.Append(FString::Printf(
+							TEXT("%.*s %s[%.*s] : packoffset(c%d%s);\r\n"),
+							SVARG(ParsedParameter->ParsedType),
+							ShaderBindingName,
+							SVARG(ParsedParameter->ParsedArraySize),
+							ConstantRegister,
+							ConstantSwizzle));
+					}
+					else
+					{
+						RootCBufferContent.Append(FString::Printf(
+							TEXT("%.*s %s : packoffset(c%d%s);\r\n"),
+							SVARG(ParsedParameter->ParsedType),
+							ShaderBindingName,
+							ConstantRegister,
+							ConstantSwizzle));
+					}
+				}
+#undef SVARG
+
+			}
+		});
+
+		FString CBufferCodeBlock = FString::Printf(
+			TEXT("%s %s\r\n")
+			TEXT("{\r\n")
+			TEXT("%s")
+			TEXT("}\r\n\r\n"),
+			ConstantBufferType,
+			FShaderParametersMetadata::kRootUniformBufferBindingName,
+			*RootCBufferContent);
+
+		FString NewShaderCode = (
+			MakeInjectedShaderCodeBlock(TEXT("MoveShaderParametersToRootConstantBuffer"), CBufferCodeBlock) +
+			PreprocessedShaderSource);
+
+		PreprocessedShaderSource = MoveTemp(NewShaderCode);
+
+		bMovedLoosedParametersToRootConstantBuffer = true;
+	} // if (CompilerInput.RootParametersStructure && bNeedToMoveToRootConstantBuffer)
+
+	return bSuccess;
+}
+
+bool FShaderParameterParser::ParseAndModify(
+	const FShaderCompilerInput& CompilerInput,
+	FShaderCompilerOutput& CompilerOutput,
+	FString& PreprocessedShaderSource)
+{
+	bBindlessResources = CompilerInput.Environment.CompilerFlags.Contains(CFLAG_BindlessResources);
+	bBindlessSamplers = CompilerInput.Environment.CompilerFlags.Contains(CFLAG_BindlessSamplers);
+	const bool bHasRootParameters = (CompilerInput.RootParametersStructure != nullptr);
+
+	// The shader doesn't have any parameter binding through shader structure, therefore don't do anything.
+	if (!(bBindlessResources || bBindlessSamplers || bHasRootParameters))
+	{
+		return true;
+	}
+
+	bNeedToMoveToRootConstantBuffer = ConstantBufferType != nullptr && (CompilerInput.IsRayTracingShader() || CompilerInput.ShouldUseStableConstantBuffer());
+	OriginalParsedShader = PreprocessedShaderSource;
+
+	if (!ParseParameters(CompilerInput, CompilerOutput))
+	{
+		return false;
+	}
+
+	RemoveMovingParametersFromSource(PreprocessedShaderSource);
+	ApplyBindlessModifications(PreprocessedShaderSource);
+
+	if (bNeedToMoveToRootConstantBuffer)
+	{
+		return MoveShaderParametersToRootConstantBuffer(CompilerInput, CompilerOutput, PreprocessedShaderSource);
+	}
+	
+	return true;
+}
+
+void FShaderParameterParser::ValidateShaderParameterType(
+	const FShaderCompilerInput& CompilerInput,
+	const FString& ShaderBindingName,
+	int32 ReflectionOffset,
+	int32 ReflectionSize,
+	bool bPlatformSupportsPrecisionModifier,
+	FShaderCompilerOutput& CompilerOutput) const
+{
+	FString BindingName(ShaderBindingName);
+
+	const bool bBindlessHack = RemoveBindlessParameterPrefix(BindingName);
+
+	const FShaderParameterParser::FParsedShaderParameter& ParsedParameter = FindParameterInfos(BindingName);
+
+	check(ParsedParameter.IsFound());
+	check(CompilerInput.RootParametersStructure);
+
+	if (ReflectionSize > 0 && bMovedLoosedParametersToRootConstantBuffer)
+	{
+		// Verify the offset of the parameter coming from shader reflections honor the packoffset()
+		check(ReflectionOffset == ParsedParameter.ConstantBufferOffset);
+	}
+
+	// Validate the shader type.
+	if (!bBindlessHack)
+	{
+		FString ExpectedShaderType;
+		FShaderParametersMetadata::FMember::GenerateShaderParameterType(
+			ExpectedShaderType, 
+			bPlatformSupportsPrecisionModifier,
+			ParsedParameter.BaseType,
+			ParsedParameter.PrecisionModifier,
+			ParsedParameter.NumRows,
+			ParsedParameter.NumColumns);
+
+		const bool bShouldBeInt = ParsedParameter.BaseType == UBMT_INT32;
+		const bool bShouldBeUint = ParsedParameter.BaseType == UBMT_UINT32;
+
+		// Match parsed type with expected shader type
+		bool bIsTypeCorrect = ParsedParameter.ParsedType == ExpectedShaderType;
+
+		if (!bIsTypeCorrect)
+		{
+			auto CheckTypeCorrect = [&ParsedParameter, &ExpectedShaderType](int32 ParsedOffset, int32 ExpectedOffset) -> bool
+			{
+				const FStringView Parsed = ParsedParameter.ParsedType.RightChop(ParsedOffset);
+				const FStringView Expected = FStringView(ExpectedShaderType).RightChop(ExpectedOffset);
+
+				return Parsed.Compare(Expected, ESearchCase::CaseSensitive) == 0;
+			};
+
+			// Accept half-precision floats when single-precision was requested
+			if (ParsedParameter.ParsedType.StartsWith(TEXT("half")) && ParsedParameter.BaseType == UBMT_FLOAT32)
+			{
+				bIsTypeCorrect = CheckTypeCorrect(4, 5);
+			}
+			// Accept single-precision floats when half-precision was expected
+			else if (ParsedParameter.ParsedType.StartsWith(TEXT("float")) && ExpectedShaderType.StartsWith(TEXT("half")))
+			{
+				bIsTypeCorrect = CheckTypeCorrect(5, 4);
+			}
+			// support for min16float
+			else if (ParsedParameter.ParsedType.StartsWith(TEXT("min16float")) && ExpectedShaderType.StartsWith(TEXT("float")))
+			{
+				bIsTypeCorrect = CheckTypeCorrect(10, 5);
+			}
+			else if (ParsedParameter.ParsedType.StartsWith(TEXT("min16float")) && ExpectedShaderType.StartsWith(TEXT("half")))
+			{
+				bIsTypeCorrect = CheckTypeCorrect(10, 4);
+			}
+		}
+
+		// Allow silent casting between signed and unsigned on shader bindings.
+		if (!bIsTypeCorrect && (bShouldBeInt || bShouldBeUint))
+		{
+			FString NewExpectedShaderType;
+			if (bShouldBeInt)
+			{
+				// tries up with an uint.
+				NewExpectedShaderType = TEXT("u") + ExpectedShaderType;
+			}
+			else
+			{
+				// tries up with an int.
+				NewExpectedShaderType = ExpectedShaderType;
+				NewExpectedShaderType.RemoveAt(0);
+			}
+
+			bIsTypeCorrect = ParsedParameter.ParsedType == NewExpectedShaderType;
+		}
+
+		if (!bIsTypeCorrect)
+		{
+			FString CppCodeName = CompilerInput.RootParametersStructure->GetFullMemberCodeName(ParsedParameter.ConstantBufferOffset);
+
+			FShaderCompilerError Error;
+			Error.StrippedErrorMessage = FString::Printf(
+				TEXT("Error: Type %.*s of shader parameter %s in shader mismatch the shader parameter structure: %s expects a %s"),
+				ParsedParameter.ParsedType.Len(), ParsedParameter.ParsedType.GetData(),
+				*ShaderBindingName,
+				*CppCodeName,
+				*ExpectedShaderType);
+			GetParameterFileAndLine(ParsedParameter, Error.ErrorVirtualFilePath, Error.ErrorLineString);
+
+			CompilerOutput.Errors.Add(Error);
+			CompilerOutput.bSucceeded = false;
+
+			AddNoteToDisplayShaderParameterMemberOnCppSide(CompilerInput, ParsedParameter, CompilerOutput);
+		}
+	}
+
+	// Validate parameter size, in case this is an array.
+	if (!bBindlessHack && ReflectionSize > int32(ParsedParameter.MemberSize))
+	{
+		FString CppCodeName = CompilerInput.RootParametersStructure->GetFullMemberCodeName(ParsedParameter.ConstantBufferOffset);
+
+		FShaderCompilerError Error;
+		Error.StrippedErrorMessage = FString::Printf(
+			TEXT("Error: The size required to bind shader parameter %s is %i bytes, smaller than %s that is %i bytes in the parameter structure."),
+			*ShaderBindingName,
+			ReflectionSize,
+			*CppCodeName,
+			ParsedParameter.MemberSize);
+		GetParameterFileAndLine(ParsedParameter, Error.ErrorVirtualFilePath, Error.ErrorLineString);
+
+		CompilerOutput.Errors.Add(Error);
+		CompilerOutput.bSucceeded = false;
+
+		AddNoteToDisplayShaderParameterMemberOnCppSide(CompilerInput, ParsedParameter, CompilerOutput);
+	}
+}
+
+void FShaderParameterParser::ValidateShaderParameterTypes(
+	const FShaderCompilerInput& CompilerInput,
+	bool bPlatformSupportsPrecisionModifier,
+	FShaderCompilerOutput& CompilerOutput) const
+{
+	// The shader doesn't have any parameter binding through shader structure, therefore don't do anything.
+	if (!CompilerInput.RootParametersStructure)
+	{
+		return;
+	}
+
+	if (!CompilerOutput.bSucceeded)
+	{
+		return;
+	}
+
+	const TMap<FString, FParameterAllocation>& ParametersFoundByCompiler = CompilerOutput.ParameterMap.GetParameterMap();
+
+	IterateShaderParameterMembers(
+		*CompilerInput.RootParametersStructure,
+		[&](const FShaderParametersMetadata& ParametersMetadata,
+			const FShaderParametersMetadata::FMember& Member,
+			const TCHAR* ShaderBindingName,
+			uint16 ByteOffset)
+		{
+			if (
+				Member.GetBaseType() != UBMT_INT32 &&
+				Member.GetBaseType() != UBMT_UINT32 &&
+				Member.GetBaseType() != UBMT_FLOAT32)
+			{
+				return;
+			}
+
+			const FParsedShaderParameter& ParsedParameter = ParsedParameters[ShaderBindingName];
+
+			// Did not find shader parameter in code.
+			if (!ParsedParameter.IsFound())
+			{
+				// Verify the shader compiler also did not find this parameter to make sure there is no bug in the parser.
+				checkf(
+					!ParametersFoundByCompiler.Contains(ShaderBindingName),
+					TEXT("Looks like there is a bug in FShaderParameterParser ParameterName=%s DumpDebugInfoPath=%s"),
+					ShaderBindingName,
+					*CompilerInput.DumpDebugInfoPath);
+				return;
+			}
+
+			int32 BoundOffset = 0;
+			int32 BoundSize = 0;
+			if (const FParameterAllocation* ParameterAllocation = ParametersFoundByCompiler.Find(ShaderBindingName))
+			{
+				BoundOffset = ParameterAllocation->BaseIndex;
+				BoundSize = ParameterAllocation->Size;
+			}
+
+			ValidateShaderParameterType(CompilerInput, ShaderBindingName, BoundOffset, BoundSize, bPlatformSupportsPrecisionModifier, CompilerOutput);
+		});
+}
+
+void FShaderParameterParser::ExtractFileAndLine(int32 PragamLineoffset, int32 LineOffset, FString& OutFile, FString& OutLine) const
+{
+	if (PragamLineoffset == -1)
+	{
+		return;
+	}
+
+	check(FCString::Strncmp((*OriginalParsedShader) + PragamLineoffset, TEXT("#line "), 6) == 0);
+
+	const int32 ShaderSourceLen = OriginalParsedShader.Len();
+
+	int32 StartFilePos = -1;
+	int32 EndFilePos = -1;
+	int32 StartLinePos = PragamLineoffset + 6;
+	int32 EndLinePos = -1;
+
+	for (int32 Cursor = StartLinePos; Cursor < ShaderSourceLen; Cursor++)
+	{
+		const TCHAR Char = OriginalParsedShader[Cursor];
+
+		if (Char == '\n')
+		{
+			break;
+		}
+
+		if (EndLinePos == -1)
+		{
+			if (Char > '9' || Char < '0')
+			{
+				EndLinePos = Cursor - 1;
+			}
+		}
+		else if (StartFilePos == -1)
+		{
+			if (Char == '"')
+			{
+				StartFilePos = Cursor + 1;
+			}
+		}
+		else if (EndFilePos == -1)
+		{
+			if (Char == '"')
+			{
+				EndFilePos = Cursor - 1;
+				break;
+			}
+		}
+	}
+
+	check(StartFilePos != -1);
+	check(EndFilePos != -1);
+	check(EndLinePos != -1);
+
+	OutFile = OriginalParsedShader.Mid(StartFilePos, EndFilePos - StartFilePos + 1);
+	FString LineBasis = OriginalParsedShader.Mid(StartLinePos, EndLinePos - StartLinePos + 1);
+
+	int32 FinalLine = FCString::Atoi(*LineBasis) + LineOffset;
+	OutLine = FString::FromInt(FinalLine);
+}
diff --git a/Engine/Source/Runtime/RenderCore/Public/ShaderCompilerCore.h b/Engine/Source/Runtime/RenderCore/Public/ShaderCompilerCore.h
--- a/Engine/Source/Runtime/RenderCore/Public/ShaderCompilerCore.h
+++ b/Engine/Source/Runtime/RenderCore/Public/ShaderCompilerCore.h
@@ -371,6 +371,18 @@ struct FShaderCompilerInput
 		return IsRayTracingShaderFrequency(Target.GetFrequency());
 	}
 
+	bool ShouldUseStableConstantBuffer() const
+	{
+		// stable constant buffer is for the FShaderParameterBindings::BindForLegacyShaderParameters() code path.
+		// Ray tracing shaders use FShaderParameterBindings::BindForRootShaderParameters instead.
+		if (IsRayTracingShader())
+		{
+			return false;
+		}
+
+		return RootParametersStructure != nullptr;
+	}
+
 	/** Returns whether this shader input *can* be compiled with the legacy FXC compiler. */
 	bool CanCompileWithLegacyFxc() const
 	{
@@ -578,17 +590,10 @@ extern RENDERCORE_API void CompileShader(const TArray<const IShaderFormat*>& Sha
 extern RENDERCORE_API void CompileShader(const TArray<const IShaderFormat*>& ShaderFormats, FShaderCompileJob& Job, const FString& WorkingDirectory, int32* CompileCount = nullptr);
 extern RENDERCORE_API void CompileShaderPipeline(const TArray<const IShaderFormat*>& ShaderFormats, FShaderPipelineCompileJob* PipelineJob, const FString& WorkingDirectory, int32* CompileCount = nullptr);
 
-UE_DEPRECATED(5.2, "Functionality has moved to UE::ShaderCompilerCommon::ShouldUseStableConstantBuffer")
+UE_DEPRECATED(5.2, "Functionality has moved to FShaderCompilerInput::ShouldUseStableConstantBuffer")
 inline bool ShouldUseStableConstantBuffer(const FShaderCompilerInput& Input)
 {
-	// stable constant buffer is for the FShaderParameterBindings::BindForLegacyShaderParameters() code path.
-	// Ray tracing shaders use FShaderParameterBindings::BindForRootShaderParameters instead.
-	if (Input.IsRayTracingShader())
-	{
-		return false;
-	}
-
-	return Input.RootParametersStructure != nullptr;
+	return Input.ShouldUseStableConstantBuffer();
 }
 
 
diff --git a/Engine/Source/Runtime/RenderCore/Public/ShaderParameterParser.h b/Engine/Source/Runtime/RenderCore/Public/ShaderParameterParser.h
new file mode 100644
--- /dev/null
+++ b/Engine/Source/Runtime/RenderCore/Public/ShaderParameterParser.h
@@ -0,0 +1,198 @@
+// Copyright Epic Games, Inc. All Rights Reserved.
+
+#pragma once
+
+#include "CoreMinimal.h"
+#include "ShaderParameterMetadata.h"
+
+struct FShaderCompilerInput;
+struct FShaderCompilerOutput;
+
+enum class EShaderParameterType : uint8;
+
+enum class EBindlessConversionType : uint8
+{
+	None,
+	Resource,
+	Sampler
+};
+
+/** Validates and moves all the shader loose data parameter defined in the root scope of the shader into the root uniform buffer. */
+class RENDERCORE_API FShaderParameterParser
+{
+public:
+	struct FParsedShaderParameter
+	{
+	public:
+		/** Original information about the member. */
+		EUniformBufferBaseType BaseType;
+		EShaderPrecisionModifier::Type PrecisionModifier;
+		uint32 NumRows;
+		uint32 NumColumns;
+		uint32 MemberSize;
+
+		/** Information found about the member when parsing the preprocessed code. */
+		FStringView ParsedName; /** View into FShaderParameterParser::OriginalParsedShader */
+		FStringView ParsedType; /** View into FShaderParameterParser::OriginalParsedShader */
+		FStringView ParsedArraySize; /** View into FShaderParameterParser::OriginalParsedShader */
+
+		/** Offset the member should be in the constant buffer. */
+		int32 ConstantBufferOffset = 0;
+
+		/* Returns whether the shader parameter has been found when parsing. */
+		bool IsFound() const
+		{
+			return !ParsedType.IsEmpty();
+		}
+
+		int32 ParsedPragmaLineoffset = 0;
+		int32 ParsedLineOffset = 0;
+
+		/** Character position of the start and end of the parameter decelaration in FParsedShaderParameter::OriginalParsedShader */
+		int32 ParsedCharOffsetStart = INDEX_NONE;
+		int32 ParsedCharOffsetEnd = INDEX_NONE;
+
+		EBindlessConversionType BindlessConversionType{};
+		EShaderParameterType ConstantBufferParameterType{};
+
+		bool bGloballyCoherent = false;
+		bool bIsBindable = false;
+
+		friend class FShaderParameterParser;
+	};
+
+	FShaderParameterParser();
+	FShaderParameterParser(const TCHAR* InConstantBufferType);
+	FShaderParameterParser(const TCHAR* InConstantBufferType, TArrayView<const TCHAR* const> InExtraSRVTypes, TArrayView<const TCHAR* const> InExtraUAVTypes);
+	virtual ~FShaderParameterParser();
+
+	static constexpr const TCHAR* kBindlessResourcePrefix = TEXT("BindlessResource_");
+	static constexpr const TCHAR* kBindlessSamplerPrefix = TEXT("BindlessSampler_");
+
+	static EShaderParameterType ParseParameterType(FStringView InType, TArrayView<const TCHAR* const> InExtraSRVTypes, TArrayView<const TCHAR* const> InExtraUAVTypes);
+	static EShaderParameterType ParseAndRemoveBindlessParameterPrefix(FStringView& InName);
+	static EShaderParameterType ParseAndRemoveBindlessParameterPrefix(FString& InName);
+	static bool RemoveBindlessParameterPrefix(FString& InName);
+
+	/** Parses the preprocessed shader code and applies the necessary modifications to it. */
+	bool ParseAndModify(
+		const FShaderCompilerInput& CompilerInput,
+		FShaderCompilerOutput& CompilerOutput,
+		FString& PreprocessedShaderSource
+	);
+
+	UE_DEPRECATED(5.2, "ParseAndModify doesn't need ConstantBufferType anymore")
+	bool ParseAndModify(
+		const FShaderCompilerInput& CompilerInput,
+		FShaderCompilerOutput& CompilerOutput,
+		FString& PreprocessedShaderSource,
+		const TCHAR* InConstantBufferType
+	)
+	{
+		return ParseAndModify(CompilerInput, CompilerOutput, PreprocessedShaderSource);
+	}
+
+	/** Parses the preprocessed shader code and move the parameters into root constant buffer */
+	UE_DEPRECATED(5.1, "ParseAndModify should be called instead.")
+	bool ParseAndMoveShaderParametersToRootConstantBuffer(
+		const FShaderCompilerInput& CompilerInput,
+		FShaderCompilerOutput& CompilerOutput,
+		FString& PreprocessedShaderSource,
+		const TCHAR* InConstantBufferType)
+	{
+		return ParseAndModify(CompilerInput, CompilerOutput, PreprocessedShaderSource);
+	}
+
+	/** Gets parsing information from a parameter binding name. */
+	const FParsedShaderParameter& FindParameterInfos(const FString& ParameterName) const
+	{
+		return ParsedParameters.FindChecked(ParameterName);
+	}
+
+	/** Validates the shader parameter in code is compatible with the shader parameter structure. */
+	void ValidateShaderParameterType(
+		const FShaderCompilerInput& CompilerInput,
+		const FString& ShaderBindingName,
+		int32 ReflectionOffset,
+		int32 ReflectionSize,
+		bool bPlatformSupportsPrecisionModifier,
+		FShaderCompilerOutput& CompilerOutput) const;
+
+	void ValidateShaderParameterType(
+		const FShaderCompilerInput& CompilerInput,
+		const FString& ShaderBindingName,
+		int32 ReflectionOffset,
+		int32 ReflectionSize,
+		FShaderCompilerOutput& CompilerOutput) const
+	{
+		ValidateShaderParameterType(CompilerInput, ShaderBindingName, ReflectionOffset, ReflectionSize, false, CompilerOutput);
+	}
+
+	/** Validates shader parameter map is compatible with the shader parameter structure. */
+	void ValidateShaderParameterTypes(
+		const FShaderCompilerInput& CompilerInput,
+		bool bPlatformSupportsPrecisionModifier,
+		FShaderCompilerOutput& CompilerOutput) const;
+
+	void ValidateShaderParameterTypes(
+		const FShaderCompilerInput& CompilerInput,
+		FShaderCompilerOutput& CompilerOutput) const
+	{
+		ValidateShaderParameterTypes(CompilerInput, false, CompilerOutput);
+	}
+
+	/** Gets file and line of the parameter in the shader source code. */
+	void GetParameterFileAndLine(const FParsedShaderParameter& ParsedParameter, FString& OutFile, FString& OutLine) const
+	{
+		return ExtractFileAndLine(ParsedParameter.ParsedPragmaLineoffset, ParsedParameter.ParsedLineOffset, OutFile, OutLine);
+	}
+
+protected:
+	/** Parses the preprocessed shader code */
+	bool ParseParameters(
+		const FShaderCompilerInput& CompilerInput,
+		FShaderCompilerOutput& CompilerOutput
+	);
+
+	void RemoveMovingParametersFromSource(
+		FString& PreprocessedShaderSource
+	);
+
+	/** Converts parsed parameters into their bindless forms. */
+	void ApplyBindlessModifications(
+		FString& PreprocessedShaderSource
+	);
+
+	/** Moves parsed parameters into the root constant buffer. */
+	bool MoveShaderParametersToRootConstantBuffer(
+		const FShaderCompilerInput& CompilerInput,
+		FShaderCompilerOutput& CompilerOutput,
+		FString& PreprocessedShaderSource
+	);
+
+	void ExtractFileAndLine(int32 PragamLineoffset, int32 LineOffset, FString& OutFile, FString& OutLine) const;
+
+	/**
+	* Generates shader source code to declare a bindless resource or sampler (for automatic bindless conversion).
+	* May be overriden to allow custom implementations for different platforms.
+	*/
+	virtual FString GenerateBindlessParameterDeclaration(const FParsedShaderParameter& ParsedParameter) const;
+
+	const TCHAR* const ConstantBufferType = nullptr;
+
+	const TArray<const TCHAR*> ExtraSRVTypes;
+	const TArray<const TCHAR*> ExtraUAVTypes;
+
+	FString OriginalParsedShader;
+
+	TMap<FString, FParsedShaderParameter> ParsedParameters;
+
+	bool bBindlessResources = false;
+	bool bBindlessSamplers = false;
+
+	/** Indicates that parameters should be moved to the root cosntant buffer. */
+	bool bNeedToMoveToRootConstantBuffer = false;
+
+	/** Indicates that parameters were actually moved to the root constant buffer. */
+	bool bMovedLoosedParametersToRootConstantBuffer = false;
+};
