Optimize FRayTracingGeometryManager::Tick when Dynamic RayTracing mode is enabled
* Only init/release acceleration structures when r.RayTracing.Enable has changed.

diff --git a/Engine/Source/Runtime/RenderCore/Private/RayTracingGeometryManager.cpp b/Engine/Source/Runtime/RenderCore/Private/RayTracingGeometryManager.cpp
--- a/Engine/Source/Runtime/RenderCore/Private/RayTracingGeometryManager.cpp
+++ b/Engine/Source/Runtime/RenderCore/Private/RayTracingGeometryManager.cpp
@@ -109,7 +109,7 @@ void FRayTracingGeometryManager::ReleaseRayTracingGeometryHandle(RayTracingGeome
 	}	
 }
 
-void FRayTracingGeometryManager::Tick()
+void FRayTracingGeometryManager::Tick(bool bHasRayTracingEnableChanged)
 {
 	if (GetRayTracingMode() != ERayTracingMode::Dynamic)
 	{
@@ -119,6 +119,30 @@ void FRayTracingGeometryManager::Tick()
 	TRACE_CPUPROFILER_EVENT_SCOPE(FRayTracingGeometryManager::Tick);
 	QUICK_SCOPE_CYCLE_COUNTER(STAT_FRayTracingGeometryManager_Tick);
 
+	if (!bHasRayTracingEnableChanged)
+	{
+// If the code below triggers a check then dynamic ray tracing is not going to work as expected ie. not all memory will be released or we'll be missing geometry.
+#if DO_CHECK
+		if (IsRayTracingEnabled())
+		{
+			FScopeLock ScopeLock(&RequestCS);
+			for (FRayTracingGeometry* Geometry : RegisteredGeometries)
+			{
+				checkf(Geometry->RayTracingGeometryRHI != nullptr, TEXT("Ray tracing geometry should be valid at this point."));
+			}
+		}
+		else
+		{
+			FScopeLock ScopeLock(&RequestCS);
+			for (FRayTracingGeometry* Geometry : RegisteredGeometries)
+			{
+				checkf(Geometry->RayTracingGeometryRHI == nullptr, TEXT("Ray tracing geometry should not be valid at this point"));
+			}
+		}
+#endif
+		return;
+	}
+
 	if (IsRayTracingEnabled())
 	{
 		FScopeLock ScopeLock(&RequestCS);
diff --git a/Engine/Source/Runtime/RenderCore/Public/RayTracingGeometryManager.h b/Engine/Source/Runtime/RenderCore/Public/RayTracingGeometryManager.h
--- a/Engine/Source/Runtime/RenderCore/Public/RayTracingGeometryManager.h
+++ b/Engine/Source/Runtime/RenderCore/Public/RayTracingGeometryManager.h
@@ -36,7 +36,7 @@ public:
 	void ForceBuildIfPending(FRHIComputeCommandList& InCmdList, const TArrayView<const FRayTracingGeometry*> InGeometries);
 	void ProcessBuildRequests(FRHIComputeCommandList& InCmdList, bool bInBuildAll = false);
 
-	void Tick();
+	void Tick(bool bHasRayTracingEnableChanged);
 
 	RayTracingGeometryHandle RegisterRayTracingGeometry(FRayTracingGeometry* InGeometry);
 	void ReleaseRayTracingGeometryHandle(RayTracingGeometryHandle Handle);
diff --git a/Engine/Source/Runtime/Renderer/Private/DeferredShadingRenderer.cpp b/Engine/Source/Runtime/Renderer/Private/DeferredShadingRenderer.cpp
--- a/Engine/Source/Runtime/Renderer/Private/DeferredShadingRenderer.cpp
+++ b/Engine/Source/Runtime/Renderer/Private/DeferredShadingRenderer.cpp
@@ -2578,7 +2578,7 @@ void FDeferredShadingSceneRenderer::Render(FRDGBuilder& GraphBuilder)
 	// Initialize ray tracing flags, in case they weren't initialized in the CreateSceneRenderers code path
 	InitializeRayTracingFlags_RenderThread();
 
-	GRayTracingGeometryManager.Tick();
+	GRayTracingGeometryManager.Tick(bHasRayTracingEnableChanged);
 
 	if ((GetRayTracingMode() == ERayTracingMode::Dynamic) && bHasRayTracingEnableChanged)
 	{
