Modified RDG setup tasks to launch inline tasks, and added variants.

diff --git a/Engine/Source/Runtime/RenderCore/Public/RenderGraphBuilder.h b/Engine/Source/Runtime/RenderCore/Public/RenderGraphBuilder.h
--- a/Engine/Source/Runtime/RenderCore/Public/RenderGraphBuilder.h
+++ b/Engine/Source/Runtime/RenderCore/Public/RenderGraphBuilder.h
@@ -220,9 +220,15 @@ public:
 	template <typename TaskLambda>
 	UE::Tasks::FTask AddSetupTask(TaskLambda&& Task, bool bCondition = true);
 
+	template <typename TaskLambda>
+	UE::Tasks::FTask AddSetupTask(TaskLambda&& Task, UE::Tasks::ETaskPriority Priority, bool bCondition = true);
+
 	template <typename TaskLambda>
 	UE::Tasks::FTask AddSetupTask(TaskLambda&& Task, UE::Tasks::FPipe* Pipe, UE::Tasks::ETaskPriority Priority = UE::Tasks::ETaskPriority::Normal, bool bCondition = true);
 
+	template <typename TaskLambda, typename PrerequisitesCollectionType>
+	UE::Tasks::FTask AddSetupTask(TaskLambda&& Task, PrerequisitesCollectionType&& Prerequisites, UE::Tasks::ETaskPriority Priority = UE::Tasks::ETaskPriority::Normal, bool bCondition = true);
+
 	template <typename TaskLambda, typename PrerequisitesCollectionType>
 	UE::Tasks::FTask AddSetupTask(TaskLambda&& Task, UE::Tasks::FPipe* Pipe, PrerequisitesCollectionType&& Prerequisites, UE::Tasks::ETaskPriority Priority = UE::Tasks::ETaskPriority::Normal, bool bCondition = true);
 
@@ -230,6 +236,18 @@ public:
 	template <typename TaskLambda>
 	UE::Tasks::FTask AddCommandListSetupTask(TaskLambda&& Task, bool bCondition = true);
 
+	template <typename TaskLambda>
+	UE::Tasks::FTask AddCommandListSetupTask(TaskLambda&& Task, UE::Tasks::ETaskPriority Priority, bool bCondition = true);
+
+	template <typename TaskLambda>
+	UE::Tasks::FTask AddCommandListSetupTask(TaskLambda&& Task, UE::Tasks::FPipe* Pipe, UE::Tasks::ETaskPriority Priority = UE::Tasks::ETaskPriority::Normal, bool bCondition = true);
+
+	template <typename TaskLambda, typename PrerequisitesCollectionType>
+	UE::Tasks::FTask AddCommandListSetupTask(TaskLambda&& Task, PrerequisitesCollectionType&& Prerequisites, UE::Tasks::ETaskPriority Priority = UE::Tasks::ETaskPriority::Normal, bool bCondition = true);
+
+	template <typename TaskLambda, typename PrerequisitesCollectionType>
+	UE::Tasks::FTask AddCommandListSetupTask(TaskLambda&& Task, UE::Tasks::FPipe* Pipe, PrerequisitesCollectionType&& Prerequisites, UE::Tasks::ETaskPriority Priority = UE::Tasks::ETaskPriority::Normal, bool bCondition = true);
+
 	/** Tells the builder to delete unused RHI resources. The behavior of this method depends on whether RDG immediate mode is enabled:
 	 *   Deferred:  RHI resource flushes are performed prior to execution.
 	 *   Immediate: RHI resource flushes are performed immediately.
diff --git a/Engine/Source/Runtime/RenderCore/Public/RenderGraphBuilder.inl b/Engine/Source/Runtime/RenderCore/Public/RenderGraphBuilder.inl
--- a/Engine/Source/Runtime/RenderCore/Public/RenderGraphBuilder.inl
+++ b/Engine/Source/Runtime/RenderCore/Public/RenderGraphBuilder.inl
@@ -411,6 +411,12 @@ FORCEINLINE UE::Tasks::FTask FRDGBuilder::AddSetupTask(TaskLambdaType&& TaskLamb
 	return AddSetupTask(MoveTemp(TaskLambda), nullptr, TArray<UE::Tasks::FTask>{}, UE::Tasks::ETaskPriority::Normal, bCondition);
 }
 
+template <typename TaskLambdaType>
+FORCEINLINE UE::Tasks::FTask FRDGBuilder::AddSetupTask(TaskLambdaType&& TaskLambda, UE::Tasks::ETaskPriority Priority, bool bCondition)
+{
+	return AddSetupTask(MoveTemp(TaskLambda), nullptr, TArray<UE::Tasks::FTask>{}, Priority, bCondition);
+}
+
 template <typename TaskLambdaType>
 FORCEINLINE UE::Tasks::FTask FRDGBuilder::AddSetupTask(
 	TaskLambdaType&& TaskLambda,
@@ -421,6 +427,16 @@ FORCEINLINE UE::Tasks::FTask FRDGBuilder::AddSetupTask(
 	return AddSetupTask(MoveTemp(TaskLambda), Pipe, TArray<UE::Tasks::FTask>{}, Priority, bCondition);
 }
 
+template <typename TaskLambdaType, typename PrerequisitesCollectionType>
+FORCEINLINE UE::Tasks::FTask FRDGBuilder::AddSetupTask(
+	TaskLambdaType&& TaskLambda,
+	PrerequisitesCollectionType&& Prerequisites,
+	UE::Tasks::ETaskPriority Priority,
+	bool bCondition)
+{
+	return AddSetupTask(MoveTemp(TaskLambda), nullptr, Forward<PrerequisitesCollectionType&&>(Prerequisites), Priority, bCondition);
+}
+
 template <typename TaskLambdaType, typename PrerequisitesCollectionType>
 UE::Tasks::FTask FRDGBuilder::AddSetupTask(
 	TaskLambdaType&& TaskLambda,
@@ -431,56 +447,130 @@ UE::Tasks::FTask FRDGBuilder::AddSetupTask(
 {
 	UE::Tasks::FTask Task;
 
-	if (bParallelExecuteEnabled && bCondition)
+	bCondition &= bParallelExecuteEnabled;
+
+	auto OuterLambda = [TaskLambda = MoveTemp(TaskLambda)]() mutable
 	{
-		auto OuterLambda = [TaskLambda = MoveTemp(TaskLambda)]() mutable
-		{
-			FOptionalTaskTagScope Scope(ETaskTag::EParallelRenderingThread);
-			TaskLambda();
-		};
+		FOptionalTaskTagScope Scope(ETaskTag::EParallelRenderingThread);
+		TaskLambda();
+	};
 
-		if (Pipe)
-		{
-			Task = Pipe->Launch(TEXT("FRDGBuilder::AddSetupTask"), MoveTemp(OuterLambda), MoveTemp(Prerequisites), Priority);
-		}
-		else
-		{
-			Task = UE::Tasks::Launch(TEXT("FRDGBuilder::AddSetupTask"), MoveTemp(OuterLambda), MoveTemp(Prerequisites), Priority);
-		}
+	const UE::Tasks::EExtendedTaskPriority ExtendedTaskPriority = bCondition ? UE::Tasks::EExtendedTaskPriority::None : UE::Tasks::EExtendedTaskPriority::Inline;
 
-		ParallelSetupEvents.Emplace(Task);
+	if (IsImmediateMode())
+	{
+		UE::Tasks::Wait(Prerequisites);
+		OuterLambda();
+	}
+	else if (Pipe)
+	{
+		Task = Pipe->Launch(TEXT("FRDGBuilder::AddSetupTask"), MoveTemp(OuterLambda), Forward<PrerequisitesCollectionType&&>(Prerequisites), Priority, ExtendedTaskPriority);
 	}
 	else
 	{
-		TaskLambda();
+		Task = UE::Tasks::Launch(TEXT("FRDGBuilder::AddSetupTask"), MoveTemp(OuterLambda), Forward<PrerequisitesCollectionType&&>(Prerequisites), Priority, ExtendedTaskPriority);
+	}
+
+	if (bCondition)
+	{
+		ParallelSetupEvents.Emplace(Task);
 	}
 
 	return Task;
 }
 
 template <typename TaskLambdaType>
-UE::Tasks::FTask FRDGBuilder::AddCommandListSetupTask(TaskLambdaType&& TaskLambda, bool bCondition)
+FORCEINLINE UE::Tasks::FTask FRDGBuilder::AddCommandListSetupTask(TaskLambdaType&& TaskLambda, bool bCondition)
+{
+	return AddCommandListSetupTask(MoveTemp(TaskLambda), nullptr, TArray<UE::Tasks::FTask>{}, UE::Tasks::ETaskPriority::Normal, bCondition);
+}
+
+template <typename TaskLambdaType>
+FORCEINLINE UE::Tasks::FTask FRDGBuilder::AddCommandListSetupTask(TaskLambdaType&& TaskLambda, UE::Tasks::ETaskPriority TaskPriority, bool bCondition)
+{
+	return AddCommandListSetupTask(MoveTemp(TaskLambda), nullptr, TArray<UE::Tasks::FTask>{}, TaskPriority, bCondition);
+}
+
+template <typename TaskLambdaType>
+FORCEINLINE UE::Tasks::FTask FRDGBuilder::AddCommandListSetupTask(
+	TaskLambdaType&& TaskLambda,
+	UE::Tasks::FPipe* Pipe,
+	UE::Tasks::ETaskPriority Priority,
+	bool bCondition)
+{
+	return AddCommandListSetupTask(MoveTemp(TaskLambda), Pipe, TArray<UE::Tasks::FTask>{}, Priority, bCondition);
+}
+
+template <typename TaskLambdaType, typename PrerequisitesCollectionType>
+FORCEINLINE UE::Tasks::FTask FRDGBuilder::AddCommandListSetupTask(
+	TaskLambdaType&& TaskLambda,
+	PrerequisitesCollectionType&& Prerequisites,
+	UE::Tasks::ETaskPriority Priority,
+	bool bCondition)
+{
+	return AddCommandListSetupTask(MoveTemp(TaskLambda), nullptr, Forward<PrerequisitesCollectionType&&>(Prerequisites), Priority, bCondition);
+}
+
+template <typename TaskLambdaType, typename PrerequisitesCollectionType>
+UE::Tasks::FTask FRDGBuilder::AddCommandListSetupTask(
+	TaskLambdaType&& TaskLambda,
+	UE::Tasks::FPipe* Pipe,
+	PrerequisitesCollectionType&& Prerequisites,
+	UE::Tasks::ETaskPriority Priority,
+	bool bCondition)
 {
 	UE::Tasks::FTask Task;
 
-	if (bParallelExecuteEnabled && bCondition)
+	bCondition &= bParallelExecuteEnabled;
+
+	FRHICommandList* RHICmdListTask = nullptr;
+
+	if (bCondition)
+	{
+		RHICmdListTask = new FRHICommandList(RHICmdList.GetGPUMask());
+	}
+	else
+	{
+		RHICmdListTask = &RHICmdList;
+	}
+
+	auto OuterLambda = [TaskLambda = MoveTemp(TaskLambda), RHICmdListTask, bCondition]() mutable
 	{
-		FRHICommandList* RHICmdListTask = new FRHICommandList(RHICmdList.GetGPUMask());
+		FOptionalTaskTagScope Scope(ETaskTag::EParallelRenderingThread);
 
-		Task = UE::Tasks::Launch(TEXT("FRDGBuilder::AddCommandListSetupTask"), [TaskLambda = MoveTemp(TaskLambda), RHICmdListTask] () mutable
+		if (bCondition)
 		{
-			FOptionalTaskTagScope Scope(ETaskTag::EParallelRenderingThread);
 			RHICmdListTask->SwitchPipeline(ERHIPipeline::Graphics);
-			TaskLambda(*RHICmdListTask);
+		}
+
+		TaskLambda(*RHICmdListTask);
+
+		if (bCondition)
+		{
 			RHICmdListTask->FinishRecording();
-		});
+		}
+	};
 
-		ParallelSetupEvents.Emplace(Task);
-		RHICmdList.QueueAsyncCommandListSubmit(RHICmdListTask);
+	const UE::Tasks::EExtendedTaskPriority ExtendedTaskPriority = bCondition ? UE::Tasks::EExtendedTaskPriority::None : UE::Tasks::EExtendedTaskPriority::Inline;
+
+	if (IsImmediateMode())
+	{
+		UE::Tasks::Wait(Prerequisites);
+		OuterLambda();
+	}
+	else if (Pipe)
+	{
+		Task = Pipe->Launch(TEXT("FRDGBuilder::AddCommandListSetupTask"), MoveTemp(OuterLambda), Forward<PrerequisitesCollectionType&&>(Prerequisites), Priority, ExtendedTaskPriority);
 	}
 	else
 	{
-		TaskLambda(RHICmdList);
+		Task = UE::Tasks::Launch(TEXT("FRDGBuilder::AddCommandListSetupTask"), MoveTemp(OuterLambda), Forward<PrerequisitesCollectionType&&>(Prerequisites), Priority, ExtendedTaskPriority);
+	}
+
+	if (bCondition)
+	{
+		ParallelSetupEvents.Emplace(Task);
+		RHICmdList.QueueAsyncCommandListSubmit(RHICmdListTask);
 	}
 
 	return Task;
