Fixed mismatched copy constructor/assignment operator in FDepthStencilBinding.

diff --git a/Engine/Source/Runtime/RenderCore/Public/ShaderParameterMacros.h b/Engine/Source/Runtime/RenderCore/Public/ShaderParameterMacros.h
--- a/Engine/Source/Runtime/RenderCore/Public/ShaderParameterMacros.h
+++ b/Engine/Source/Runtime/RenderCore/Public/ShaderParameterMacros.h
@@ -599,8 +599,7 @@ private:
 struct FDepthStencilBinding
 {
 	FDepthStencilBinding() = default;
-	FDepthStencilBinding(const FDepthStencilBinding&) = default;
-	
+
 	/**
 	 * Creates a render target binding informations for a depth/stencil texture.
 	 *
