Fixed mismatched copy constructor/assignment operator in FRenderTargetBinding.

diff --git a/Engine/Source/Runtime/RenderCore/Public/ShaderParameterMacros.h b/Engine/Source/Runtime/RenderCore/Public/ShaderParameterMacros.h
--- a/Engine/Source/Runtime/RenderCore/Public/ShaderParameterMacros.h
+++ b/Engine/Source/Runtime/RenderCore/Public/ShaderParameterMacros.h
@@ -493,7 +493,6 @@ using FRDGTextureAccessArray = TRDGResourceAccessArray<FRDGTextureAccess>;
 struct FRenderTargetBinding
 {
 	FRenderTargetBinding() = default;
-	FRenderTargetBinding(const FRenderTargetBinding&) = default;
 
 	FRenderTargetBinding(FRDGTexture* InTexture, ERenderTargetLoadAction InLoadAction, uint8 InMipIndex = 0, int16 InArraySlice = -1)
 		: Texture(InTexture)
