Ensure GMainThreadBlockedOnRenderThread is properly initialized

diff --git a/Engine/Source/Runtime/RenderCore/Private/RenderingThread.cpp b/Engine/Source/Runtime/RenderCore/Private/RenderingThread.cpp
--- a/Engine/Source/Runtime/RenderCore/Private/RenderingThread.cpp
+++ b/Engine/Source/Runtime/RenderCore/Private/RenderingThread.cpp
@@ -44,7 +44,7 @@ RENDERCORE_API bool GIsThreadedRendering = false;
 RENDERCORE_API bool GUseThreadedRendering = false;
 
 #if !(UE_BUILD_SHIPPING || UE_BUILD_TEST)
-	RENDERCORE_API TAtomic<bool> GMainThreadBlockedOnRenderThread;
+	RENDERCORE_API TAtomic<bool> GMainThreadBlockedOnRenderThread(false);
 #endif // #if !(UE_BUILD_SHIPPING || UE_BUILD_TEST)
 
 static FRunnable* GRenderingThreadRunnable = NULL;
