Refactored the global dynamic vertex buffer to support parallel usage.

diff --git a/Engine/Source/Runtime/Engine/Public/UnrealClient.h b/Engine/Source/Runtime/Engine/Public/UnrealClient.h
--- a/Engine/Source/Runtime/Engine/Public/UnrealClient.h
+++ b/Engine/Source/Runtime/Engine/Public/UnrealClient.h
@@ -796,6 +796,11 @@ protected:
 	ENGINE_API virtual void ReleaseDynamicRHI() override;
 	ENGINE_API virtual void InitRHI() override;
 	ENGINE_API virtual void ReleaseRHI() override;
+
+	virtual void InitResource() override { FRenderResource::InitResource(); }
+	virtual void InitResource(FRHICommandList&) override { InitResource(); }
+	virtual void InitDynamicRHI(FRHICommandList&) override { InitDynamicRHI(); }
+	virtual void InitRHI(FRHICommandList&) override { InitRHI(); }
 };
 
 // Shortcuts for checking the state of both left&right variations of control keys.
diff --git a/Engine/Source/Runtime/RenderCore/Private/DynamicBufferAllocator.cpp b/Engine/Source/Runtime/RenderCore/Private/DynamicBufferAllocator.cpp
--- a/Engine/Source/Runtime/RenderCore/Private/DynamicBufferAllocator.cpp
+++ b/Engine/Source/Runtime/RenderCore/Private/DynamicBufferAllocator.cpp
@@ -208,6 +208,8 @@ bool FGlobalDynamicReadBuffer::IsRenderAlarmLoggingEnabled() const
 
 static void RemoveUnusedBuffers(FDynamicReadBufferPool* BufferPool)
 {
+	extern int32 GGlobalBufferNumFramesUnusedThreshold;
+
 	for (int32 BufferIndex = 0, NumBuffers = BufferPool->Buffers.Num(); BufferIndex < NumBuffers; ++BufferIndex)
 	{
 		FDynamicAllocReadBuffer& Buffer = BufferPool->Buffers[BufferIndex];
@@ -215,10 +217,10 @@ static void RemoveUnusedBuffers(FDynamicReadBufferPool* BufferPool)
 		{
 			Buffer.Unlock();
 		}
-		else if (GGlobalBufferNumFramesUnusedThresold && !Buffer.AllocatedByteCount)
+		else if (GGlobalBufferNumFramesUnusedThreshold && !Buffer.AllocatedByteCount)
 		{
 			++Buffer.NumFramesUnused;
-			if (Buffer.NumFramesUnused >= GGlobalBufferNumFramesUnusedThresold)
+			if (Buffer.NumFramesUnused >= GGlobalBufferNumFramesUnusedThreshold)
 			{
 				// Remove the buffer, assumes they are unordered.
 				Buffer.Release();
diff --git a/Engine/Source/Runtime/RenderCore/Private/GlobalRenderResources.cpp b/Engine/Source/Runtime/RenderCore/Private/GlobalRenderResources.cpp
--- a/Engine/Source/Runtime/RenderCore/Private/GlobalRenderResources.cpp
+++ b/Engine/Source/Runtime/RenderCore/Private/GlobalRenderResources.cpp
@@ -5,6 +5,7 @@
 #include "Containers/ResourceArray.h"
 #include "RenderCore.h"
 #include "RenderUtils.h"
+#include "Algo/Reverse.h"
 
 // The maximum number of transient vertex buffer bytes to allocate before we start panic logging who is doing the allocations
 int32 GMaxVertexBytesAllocatedPerFrame = 32 * 1024 * 1024;
@@ -14,10 +15,10 @@ FAutoConsoleVariableRef CVarMaxVertexBytesAllocatedPerFrame(
 	GMaxVertexBytesAllocatedPerFrame,
 	TEXT("The maximum number of transient vertex buffer bytes to allocate before we start panic logging who is doing the allocations"));
 
-int32 GGlobalBufferNumFramesUnusedThresold = 30;
-FAutoConsoleVariableRef CVarReadBufferNumFramesUnusedThresold(
+int32 GGlobalBufferNumFramesUnusedThreshold = 30;
+FAutoConsoleVariableRef CVarGlobalBufferNumFramesUnusedThreshold(
 	TEXT("r.NumFramesUnusedBeforeReleasingGlobalResourceBuffers"),
-	GGlobalBufferNumFramesUnusedThresold,
+	GGlobalBufferNumFramesUnusedThreshold,
 	TEXT("Number of frames after which unused global resource allocations will be discarded. Set 0 to ignore. (default=30)"));
 
 //////////////////////////////////////////////////////////////////////////////////////////////////////////////////////
@@ -660,7 +661,7 @@ TGlobalResource<FTwoTrianglesIndexBuffer> GTwoTrianglesIndexBuffer;
 /**
  * An individual dynamic vertex buffer.
  */
-class FDynamicVertexBuffer : public FVertexBuffer
+class FDynamicVertexBuffer final : public FVertexBuffer
 {
 public:
 	/** The aligned size of all dynamic vertex buffers. */
@@ -671,8 +672,8 @@ public:
 	uint32 BufferSize;
 	/** Number of bytes currently allocated from the buffer. */
 	uint32 AllocatedByteCount;
-	/** Number of successive frames for which AllocatedByteCount == 0. Used as a metric to decide when to free the allocation. */
-	int32 NumFramesUnused = 0;
+	/** Last render thread frame this resource was used in. */
+	uint64 LastUsedFrame = 0;
 
 	/** Default constructor. */
 	explicit FDynamicVertexBuffer(uint32 InMinBufferSize)
@@ -685,33 +686,32 @@ public:
 	/**
 	 * Locks the vertex buffer so it may be written to.
 	 */
-	void Lock()
+	void Lock(FRHICommandList& RHICmdList)
 	{
 		check(MappedBuffer == NULL);
 		check(AllocatedByteCount == 0);
 		check(IsValidRef(VertexBufferRHI));
-		MappedBuffer = (uint8*)RHILockBuffer(VertexBufferRHI, 0, BufferSize, RLM_WriteOnly);
+		MappedBuffer = (uint8*)RHICmdList.LockBuffer(VertexBufferRHI, 0, BufferSize, RLM_WriteOnly);
 	}
 
 	/**
 	 * Unocks the buffer so the GPU may read from it.
 	 */
-	void Unlock()
+	void Unlock(FRHICommandList& RHICmdList)
 	{
 		check(MappedBuffer != NULL);
 		check(IsValidRef(VertexBufferRHI));
-		RHIUnlockBuffer(VertexBufferRHI);
+		RHICmdList.UnlockBuffer(VertexBufferRHI);
 		MappedBuffer = NULL;
 		AllocatedByteCount = 0;
-		NumFramesUnused = 0;
 	}
 
 	// FRenderResource interface.
-	virtual void InitRHI() override
+	virtual void InitRHI(FRHICommandList& RHICmdList) override
 	{
 		check(!IsValidRef(VertexBufferRHI));
 		FRHIResourceCreateInfo CreateInfo(TEXT("FDynamicVertexBuffer"));
-		VertexBufferRHI = RHICreateVertexBuffer(BufferSize, BUF_Volatile, CreateInfo);
+		VertexBufferRHI = RHICmdList.CreateBuffer(BufferSize, EBufferUsageFlags::VertexBuffer, 0, ERHIAccess::VertexOrIndexBuffer, CreateInfo);
 		MappedBuffer = NULL;
 		AllocatedByteCount = 0;
 	}
@@ -732,128 +732,166 @@ public:
 /**
  * A pool of dynamic vertex buffers.
  */
-struct FDynamicVertexBufferPool
+struct FDynamicVertexBufferPool : public FRenderResource
 {
-	/** List of vertex buffers. */
-	TIndirectArray<FDynamicVertexBuffer> VertexBuffers;
-	/** The current buffer from which allocations are being made. */
-	FDynamicVertexBuffer* CurrentVertexBuffer;
-
-	/** Default constructor. */
-	FDynamicVertexBufferPool()
-		: CurrentVertexBuffer(NULL)
+	FCriticalSection CriticalSection;
+	TArray<FDynamicVertexBuffer*> LiveList;
+	TArray<FDynamicVertexBuffer*> FreeList;
+	TArray<FDynamicVertexBuffer*> LockList;
+	TArray<FDynamicVertexBuffer*> ReclaimList;
+	std::atomic_uint32_t TotalAllocatedMemory{0};
+	uint32 CurrentCycle = 0;
+
+	FDynamicVertexBuffer* Acquire(FRHICommandList& RHICmdList, uint32 SizeInBytes)
 	{
-	}
+		const uint32 MinimumBufferSize = 65536u;
+		SizeInBytes = FMath::Max(SizeInBytes, MinimumBufferSize);
+
+		FDynamicVertexBuffer* FoundBuffer = nullptr;
+		bool bInitializeBuffer = false;
 
-	/** Destructor. */
-	~FDynamicVertexBufferPool()
-	{
-		int32 NumVertexBuffers = VertexBuffers.Num();
-		for (int32 BufferIndex = 0; BufferIndex < NumVertexBuffers; ++BufferIndex)
 		{
-			VertexBuffers[BufferIndex].ReleaseResource();
+			FScopeLock Lock(&CriticalSection);
+
+			// Traverse the free list like a stack, starting from the top, so recently reclaimed items are allocated first.
+			for (int32 Index = FreeList.Num() - 1; Index >= 0; --Index)
+			{
+				FDynamicVertexBuffer* Buffer = FreeList[Index];
+
+				if (SizeInBytes <= Buffer->BufferSize)
+				{
+					FreeList.RemoveAt(Index, 1, false);
+					FoundBuffer = Buffer;
+					break;
+				}
+			}
+
+			if (!FoundBuffer)
+			{
+				FoundBuffer = new FDynamicVertexBuffer(SizeInBytes);
+				LiveList.Emplace(FoundBuffer);
+				bInitializeBuffer = true;
+			}
+
+			check(FoundBuffer);
 		}
-	}
-};
 
-FGlobalDynamicVertexBuffer::FGlobalDynamicVertexBuffer()
-	: TotalAllocatedSinceLastCommit(0)
-{
-	Pool = new FDynamicVertexBufferPool();
-}
+		if (IsRenderAlarmLoggingEnabled())
+		{
+			UE_LOG(LogRendererCore, Warning, TEXT("FGlobalDynamicVertexBuffer::Allocate(%u), will have allocated %u total this frame"), SizeInBytes, TotalAllocatedMemory.load());
+		}
 
-FGlobalDynamicVertexBuffer::~FGlobalDynamicVertexBuffer()
-{
-	delete Pool;
-	Pool = NULL;
-}
+		if (bInitializeBuffer)
+		{
+			FoundBuffer->InitResource(RHICmdList);
+			TotalAllocatedMemory += SizeInBytes;
+		}
 
-FGlobalDynamicVertexBuffer::FAllocation FGlobalDynamicVertexBuffer::Allocate(uint32 SizeInBytes)
-{
-	FAllocation Allocation;
+		FoundBuffer->Lock(RHICmdList);
+		FoundBuffer->LastUsedFrame = GFrameCounterRenderThread;
+
+		return FoundBuffer;
+	}
 
-	TotalAllocatedSinceLastCommit += SizeInBytes;
-	if (IsRenderAlarmLoggingEnabled())
+	void Forfeit(FRHICommandList& RHICmdList, TConstArrayView<FDynamicVertexBuffer*> BuffersToForfeit)
 	{
-		UE_LOG(LogRendererCore, Warning, TEXT("FGlobalDynamicVertexBuffer::Allocate(%u), will have allocated %u total this frame"), SizeInBytes, TotalAllocatedSinceLastCommit);
+		for (FDynamicVertexBuffer* Buffer : BuffersToForfeit)
+		{
+			Buffer->Unlock(RHICmdList);
+		}
+
+		FScopeLock Lock(&CriticalSection);
+		ReclaimList.Append(BuffersToForfeit);
 	}
 
-	FDynamicVertexBuffer* VertexBuffer = Pool->CurrentVertexBuffer;
-	if (VertexBuffer == NULL || VertexBuffer->AllocatedByteCount + SizeInBytes > VertexBuffer->BufferSize)
+	void GarbageCollect()
 	{
-		// Find a buffer in the pool big enough to service the request.
-		VertexBuffer = NULL;
-		for (int32 BufferIndex = 0, NumBuffers = Pool->VertexBuffers.Num(); BufferIndex < NumBuffers; ++BufferIndex)
+		FScopeLock Lock(&CriticalSection);
+		FreeList.Append(ReclaimList);
+		ReclaimList.Reset();
+
+		for (int32 Index = 0; Index < LiveList.Num(); ++Index)
 		{
-			FDynamicVertexBuffer& VertexBufferToCheck = Pool->VertexBuffers[BufferIndex];
-			if (VertexBufferToCheck.AllocatedByteCount + SizeInBytes <= VertexBufferToCheck.BufferSize)
+			FDynamicVertexBuffer* Buffer = LiveList[Index];
+
+			if (GGlobalBufferNumFramesUnusedThreshold > 0 && Buffer->LastUsedFrame + GGlobalBufferNumFramesUnusedThreshold <= GFrameCounterRenderThread)
 			{
-				VertexBuffer = &VertexBufferToCheck;
-				break;
+				Buffer->ReleaseResource();
+				LiveList.RemoveAt(Index, 1, false);
+				FreeList.Remove(Buffer);
+				delete Buffer;
 			}
 		}
+	}
 
-		// Create a new vertex buffer if needed.
-		if (VertexBuffer == NULL)
-		{
-			VertexBuffer = new FDynamicVertexBuffer(SizeInBytes);
-			Pool->VertexBuffers.Add(VertexBuffer);
-			VertexBuffer->InitResource();
-		}
+	bool IsRenderAlarmLoggingEnabled() const
+	{
+		return GMaxVertexBytesAllocatedPerFrame > 0 && TotalAllocatedMemory >= (size_t)GMaxVertexBytesAllocatedPerFrame;
+	}
 
-		// Lock the buffer if needed.
-		if (VertexBuffer->MappedBuffer == NULL)
+private:
+	void ReleaseDynamicRHI() override
+	{
+		check(LockList.IsEmpty());
+		check(FreeList.Num() == LiveList.Num());
+
+		for (FDynamicVertexBuffer* VertexBuffer : LiveList)
 		{
-			VertexBuffer->Lock();
+			TotalAllocatedMemory -= VertexBuffer->BufferSize;
+			VertexBuffer->ReleaseResource();
+			delete VertexBuffer;
 		}
+		LiveList.Empty();
+		FreeList.Empty();
+	}
+};
+
+static TGlobalResource<FDynamicVertexBufferPool> GDynamicVertexBufferPool;
+
+FGlobalDynamicVertexBuffer::FAllocation FGlobalDynamicVertexBuffer::Allocate(uint32 SizeInBytes)
+{
+	return Allocate(FRHICommandListExecutor::GetImmediateCommandList(), SizeInBytes);
+}
 
-		// Remember this buffer, we'll try to allocate out of it in the future.
-		Pool->CurrentVertexBuffer = VertexBuffer;
+FGlobalDynamicVertexBuffer::FAllocation FGlobalDynamicVertexBuffer::Allocate(FRHICommandList& RHICmdList, uint32 SizeInBytes)
+{
+	FAllocation Allocation;
+
+	if (VertexBuffers.IsEmpty() || VertexBuffers.Last()->AllocatedByteCount + SizeInBytes > VertexBuffers.Last()->BufferSize)
+	{
+		VertexBuffers.Emplace(GDynamicVertexBufferPool.Acquire(RHICmdList, SizeInBytes));
 	}
 
-	check(VertexBuffer != NULL);
+	FDynamicVertexBuffer* VertexBuffer = VertexBuffers.Last();
+
 	checkf(VertexBuffer->AllocatedByteCount + SizeInBytes <= VertexBuffer->BufferSize, TEXT("Global vertex buffer allocation failed: BufferSize=%d AllocatedByteCount=%d SizeInBytes=%d"), VertexBuffer->BufferSize, VertexBuffer->AllocatedByteCount, SizeInBytes);
 	Allocation.Buffer = VertexBuffer->MappedBuffer + VertexBuffer->AllocatedByteCount;
 	Allocation.VertexBuffer = VertexBuffer;
 	Allocation.VertexOffset = VertexBuffer->AllocatedByteCount;
 	VertexBuffer->AllocatedByteCount += SizeInBytes;
-
 	return Allocation;
 }
 
 bool FGlobalDynamicVertexBuffer::IsRenderAlarmLoggingEnabled() const
 {
-	return GMaxVertexBytesAllocatedPerFrame > 0 && TotalAllocatedSinceLastCommit >= (size_t)GMaxVertexBytesAllocatedPerFrame;
+	return GDynamicVertexBufferPool.IsRenderAlarmLoggingEnabled();
+}
+
+void FGlobalDynamicVertexBuffer::Commit(FRHICommandList& RHICmdList)
+{
+	GDynamicVertexBufferPool.Forfeit(RHICmdList, VertexBuffers);
+	VertexBuffers.Reset();
 }
 
 void FGlobalDynamicVertexBuffer::Commit()
 {
-	for (int32 BufferIndex = 0, NumBuffers = Pool->VertexBuffers.Num(); BufferIndex < NumBuffers; ++BufferIndex)
-	{
-		FDynamicVertexBuffer& VertexBuffer = Pool->VertexBuffers[BufferIndex];
-		if (VertexBuffer.MappedBuffer != NULL)
-		{
-			VertexBuffer.Unlock();
-		}
-		else if (GGlobalBufferNumFramesUnusedThresold && !VertexBuffer.AllocatedByteCount)
-		{
-			++VertexBuffer.NumFramesUnused;
-			if (VertexBuffer.NumFramesUnused >= GGlobalBufferNumFramesUnusedThresold)
-			{
-				// Remove the buffer, assumes they are unordered.
-				VertexBuffer.ReleaseResource();
-				Pool->VertexBuffers.RemoveAtSwap(BufferIndex);
-				--BufferIndex;
-				--NumBuffers;
-			}
-		}
-	}
-	Pool->CurrentVertexBuffer = NULL;
-	TotalAllocatedSinceLastCommit = 0;
+	Commit(FRHICommandListExecutor::GetImmediateCommandList());
 }
 
-FGlobalDynamicVertexBuffer InitViewDynamicVertexBuffer;
-FGlobalDynamicVertexBuffer InitShadowViewDynamicVertexBuffer;
+void FGlobalDynamicVertexBuffer::GarbageCollect()
+{
+	GDynamicVertexBufferPool.GarbageCollect();
+}
 
 /*------------------------------------------------------------------------------
 	FGlobalDynamicIndexBuffer implementation.
@@ -1046,10 +1084,10 @@ void FGlobalDynamicIndexBuffer::Commit()
 			{
 				IndexBuffer.Unlock();
 			}
-			else if (GGlobalBufferNumFramesUnusedThresold && !IndexBuffer.AllocatedByteCount)
+			else if (GGlobalBufferNumFramesUnusedThreshold && !IndexBuffer.AllocatedByteCount)
 			{
 				++IndexBuffer.NumFramesUnused;
-				if (IndexBuffer.NumFramesUnused >= GGlobalBufferNumFramesUnusedThresold)
+				if (IndexBuffer.NumFramesUnused >= GGlobalBufferNumFramesUnusedThreshold)
 				{
 					// Remove the buffer, assumes they are unordered.
 					IndexBuffer.ReleaseResource();
diff --git a/Engine/Source/Runtime/RenderCore/Private/RenderResource.cpp b/Engine/Source/Runtime/RenderCore/Private/RenderResource.cpp
--- a/Engine/Source/Runtime/RenderCore/Private/RenderResource.cpp
+++ b/Engine/Source/Runtime/RenderCore/Private/RenderResource.cpp
@@ -236,6 +236,7 @@ void FRenderResource::ChangeFeatureLevel(ERHIFeatureLevel::Type NewFeatureLevel)
 
 void FRenderResource::InitResource()
 {
+	check(IsInRenderingThread());
 	if (ListIndex == INDEX_NONE)
 	{
 		int32 LocalListIndex = INDEX_NONE;
@@ -254,8 +255,38 @@ void FRenderResource::InitResource()
 		if (GIsRHIInitialized)
 		{
 			CSV_SCOPED_TIMING_STAT_EXCLUSIVE(InitRenderResource);
-			InitDynamicRHI();
-			InitRHI();
+			FRHICommandListImmediate& RHICmdList = FRHICommandListExecutor::GetImmediateCommandList();
+			InitDynamicRHI(RHICmdList);
+			InitRHI(RHICmdList);
+		}
+
+		FPlatformMisc::MemoryBarrier(); // there are some multithreaded reads of ListIndex
+		ListIndex = LocalListIndex;
+	}
+}
+
+void FRenderResource::InitResource(FRHICommandList& RHICmdList)
+{
+	if (ListIndex == INDEX_NONE)
+	{
+		int32 LocalListIndex = INDEX_NONE;
+
+		if (PLATFORM_NEEDS_RHIRESOURCELIST || !GIsRHIInitialized)
+		{
+			LLM_SCOPE(ELLMTag::SceneRender);
+			LocalListIndex = FRenderResourceList::Get().Allocate(this);
+		}
+		else
+		{
+			// Mark this resource as initialized
+			LocalListIndex = 0;
+		}
+
+		if (GIsRHIInitialized)
+		{
+			CSV_SCOPED_TIMING_STAT_EXCLUSIVE(InitRenderResource);
+			InitDynamicRHI(RHICmdList);
+			InitRHI(RHICmdList);
 		}
 
 		FPlatformMisc::MemoryBarrier(); // there are some multithreaded reads of ListIndex
@@ -286,15 +317,19 @@ void FRenderResource::ReleaseResource()
 void FRenderResource::UpdateRHI()
 {
 	check(IsInRenderingThread());
-	if(IsInitialized() && GIsRHIInitialized)
+	UpdateRHI(FRHICommandListExecutor::GetImmediateCommandList());
+}
+
+void FRenderResource::UpdateRHI(FRHICommandList& RHICmdList)
+{
+	if (IsInitialized() && GIsRHIInitialized)
 	{
 		ReleaseRHI();
 		ReleaseDynamicRHI();
-		InitDynamicRHI();
-		InitRHI();
+		InitDynamicRHI(RHICmdList);
+		InitRHI(RHICmdList);
 	}
 }
-
 FRenderResource::FRenderResource()
 	: ListIndex(INDEX_NONE)
 	, FeatureLevel(ERHIFeatureLevel::Num)
diff --git a/Engine/Source/Runtime/RenderCore/Public/GlobalRenderResources.h b/Engine/Source/Runtime/RenderCore/Public/GlobalRenderResources.h
--- a/Engine/Source/Runtime/RenderCore/Public/GlobalRenderResources.h
+++ b/Engine/Source/Runtime/RenderCore/Public/GlobalRenderResources.h
@@ -182,17 +182,12 @@ class RENDERCORE_API FGlobalDynamicVertexBuffer
 public:
 	using FAllocation = FGlobalDynamicVertexBufferAllocation;
 
-	/** Default constructor. */
-	FGlobalDynamicVertexBuffer();
-
-	/** Destructor. */
-	~FGlobalDynamicVertexBuffer();
-
 	/**
 	 * Allocates space in the global vertex buffer.
 	 * @param SizeInBytes - The amount of memory to allocate in bytes.
 	 * @returns An FAllocation with information regarding the allocated memory.
 	 */
+	FAllocation Allocate(FRHICommandList& RHICmdList, uint32 SizeInBytes);
 	FAllocation Allocate(uint32 SizeInBytes);
 
 	/**
@@ -200,17 +195,16 @@ public:
 	 *		WARNING: Once this buffer has been committed to the GPU, allocations
 	 *		remain valid only until the next call to Allocate!
 	 */
+	void Commit(FRHICommandList& RHICmdList);
 	void Commit();
 
+	static void GarbageCollect();
+
 	/** Returns true if log statements should be made because we exceeded GMaxVertexBytesAllocatedPerFrame */
 	bool IsRenderAlarmLoggingEnabled() const;
 
 private:
-	/** The pool of vertex buffers from which allocations are made. */
-	struct FDynamicVertexBufferPool* Pool;
-
-	/** A total of all allocations made since the last commit. Used to alert about spikes in memory usage. */
-	size_t TotalAllocatedSinceLastCommit;
+	TArray<class FDynamicVertexBuffer*> VertexBuffers;
 };
 
 struct FGlobalDynamicIndexBufferAllocation
diff --git a/Engine/Source/Runtime/RenderCore/Public/RenderResource.h b/Engine/Source/Runtime/RenderCore/Public/RenderResource.h
--- a/Engine/Source/Runtime/RenderCore/Public/RenderResource.h
+++ b/Engine/Source/Runtime/RenderCore/Public/RenderResource.h
@@ -39,9 +39,6 @@
 class FRDGPooledBuffer;
 class FResourceArrayInterface;
 
-/** Number of frames after which unused global resource allocations will be discarded. */
-extern int32 GGlobalBufferNumFramesUnusedThresold;
-
 /** Experimental: whether we free helper structures after submitting to RHI. */
 extern RENDERCORE_API bool GFreeStructuresOnRHIBufferCreation;
 
@@ -99,6 +96,7 @@ public:
 	 * This is only called by the rendering thread.
 	 */
 	virtual void InitDynamicRHI() {}
+	virtual void InitDynamicRHI(FRHICommandList& RHICmdList) { InitDynamicRHI(); }
 
 	/**
 	 * Releases the dynamic RHI resource and/or RHI render target resources used by this resource.
@@ -114,6 +112,7 @@ public:
 	 * This is only called by the rendering thread.
 	 */
 	virtual void InitRHI() {}
+	virtual void InitRHI(FRHICommandList& RHICmdList) { InitRHI(); }
 
 	/**
 	 * Releases the RHI resources used by this resource.
@@ -127,6 +126,7 @@ public:
 	 * This is only called by the rendering thread.
 	 */
 	virtual void InitResource();
+	virtual void InitResource(FRHICommandList& RHICmdList);
 
 	/**
 	 * Prepares the resource for deletion.
@@ -139,6 +139,7 @@ public:
 	 * This is only called by the rendering thread.
 	 */
 	void UpdateRHI();
+	void UpdateRHI(FRHICommandList& RHICmdList);
 
 	/** @return The resource's friendly name.  Typically a UObject name. */
 	virtual FString GetFriendlyName() const { return TEXT("undefined"); }
diff --git a/Engine/Source/Runtime/Renderer/Private/SceneRendering.cpp b/Engine/Source/Runtime/Renderer/Private/SceneRendering.cpp
--- a/Engine/Source/Runtime/Renderer/Private/SceneRendering.cpp
+++ b/Engine/Source/Runtime/Renderer/Private/SceneRendering.cpp
@@ -4115,6 +4115,7 @@ static void FinishCleanUp(FRHICommandListImmediate& RHICmdList)
 	TRACE_CPUPROFILER_EVENT_SCOPE(FinishCleanUp);
 
 	// Can release only after all mesh pass tasks are finished.
+	FGlobalDynamicVertexBuffer::GarbageCollect();
 	GPrimitiveIdVertexBufferPool.DiscardAll();
 	FGraphicsMinimalPipelineStateId::ResetLocalPipelineIdTableSize();
 	FRenderResource::CoalesceResourceList();
