[shaders] further rework of parameter parser API to simplify the independent preprocessing use case, plus adding serialization support and some minor cleanup.

diff --git a/Engine/Source/Developer/Apple/MetalShaderFormat/Private/MetalShaderCompiler.cpp b/Engine/Source/Developer/Apple/MetalShaderFormat/Private/MetalShaderCompiler.cpp
--- a/Engine/Source/Developer/Apple/MetalShaderFormat/Private/MetalShaderCompiler.cpp
+++ b/Engine/Source/Developer/Apple/MetalShaderFormat/Private/MetalShaderCompiler.cpp
@@ -1007,8 +1007,8 @@ void CompileShader_Metal(const FShaderCompilerInput& _Input,FShaderCompilerOutpu
 		return;
 	}
 
-	FShaderParameterParser ShaderParameterParser;
-	if (!ShaderParameterParser.ParseAndModify(Input, Output, PreprocessedShader))
+	FShaderParameterParser ShaderParameterParser(Input.Environment.CompilerFlags);
+	if (!ShaderParameterParser.ParseAndModify(Input, Output.Errors, PreprocessedShader))
 	{
 		// The FShaderParameterParser will add any relevant errors.
 		return;
diff --git a/Engine/Source/Developer/ShaderFormatOpenGL/Private/OpenGLShaderCompiler.cpp b/Engine/Source/Developer/ShaderFormatOpenGL/Private/OpenGLShaderCompiler.cpp
--- a/Engine/Source/Developer/ShaderFormatOpenGL/Private/OpenGLShaderCompiler.cpp
+++ b/Engine/Source/Developer/ShaderFormatOpenGL/Private/OpenGLShaderCompiler.cpp
@@ -3366,8 +3366,8 @@ void FOpenGLFrontend::CompileShader(const FShaderCompilerInput& Input, FShaderCo
 		return;
 	}
 
-	FShaderParameterParser ShaderParameterParser;
-	if (!ShaderParameterParser.ParseAndModify(Input, Output, PreprocessedShader))
+	FShaderParameterParser ShaderParameterParser(Input.Environment.CompilerFlags);
+	if (!ShaderParameterParser.ParseAndModify(Input, Output.Errors, PreprocessedShader))
 	{
 		// The FShaderParameterParser will add any relevant errors.
 		return;
diff --git a/Engine/Source/Developer/VulkanShaderFormat/Private/VulkanShaderCompiler.cpp b/Engine/Source/Developer/VulkanShaderFormat/Private/VulkanShaderCompiler.cpp
--- a/Engine/Source/Developer/VulkanShaderFormat/Private/VulkanShaderCompiler.cpp
+++ b/Engine/Source/Developer/VulkanShaderFormat/Private/VulkanShaderCompiler.cpp
@@ -2485,14 +2485,14 @@ static bool CompileWithShaderConductor(
 class FVulkanShaderParameterParser : public FShaderParameterParser
 {
 public:
-	FVulkanShaderParameterParser(bool bHasBindlessSamplers, bool bHasBindlessResources)
-		: bEnabled(bHasBindlessSamplers || bHasBindlessResources)
+	FVulkanShaderParameterParser(FShaderCompilerFlags CompilerFlags)
+		: FShaderParameterParser(CompilerFlags)
 	{}
 
 protected:
 	FString GenerateBindlessParameterDeclaration(const FParsedShaderParameter& ParsedParameter) const override
 	{
-		if (bEnabled)
+		if (bBindlessResources || bBindlessSamplers)
 		{
 			const bool IsSampler = (ParsedParameter.BindlessConversionType == EBindlessConversionType::Sampler);
 			const TCHAR* IndexPrefix = IsSampler ? TEXT("BindlessSampler_") : TEXT("BindlessResource_");
@@ -2533,9 +2533,6 @@ protected:
 			return FShaderParameterParser::GenerateBindlessParameterDeclaration(ParsedParameter);
 		}
 	}
-
-private:
-	const bool bEnabled;
 };
 
 
@@ -2658,10 +2655,8 @@ void DoCompileVulkanShader(const FShaderCompilerInput& Input, FShaderCompilerOut
 		}
 	}
 
-	const bool bHasBindlessSamplers = Input.Environment.CompilerFlags.Contains(CFLAG_BindlessSamplers);
-	const bool bHasBindlessResources = Input.Environment.CompilerFlags.Contains(CFLAG_BindlessResources);
-	FVulkanShaderParameterParser ShaderParameterParser(bHasBindlessSamplers, bHasBindlessResources);
-	if (!ShaderParameterParser.ParseAndModify(Input, Output, PreprocessedShaderSource))
+	FVulkanShaderParameterParser ShaderParameterParser(Input.Environment.CompilerFlags);
+	if (!ShaderParameterParser.ParseAndModify(Input, Output.Errors, PreprocessedShaderSource))
 	{
 		// The FShaderParameterParser will add any relevant errors.
 		return;
diff --git a/Engine/Source/Developer/Windows/ShaderFormatD3D/Private/D3DShaderCompiler.cpp b/Engine/Source/Developer/Windows/ShaderFormatD3D/Private/D3DShaderCompiler.cpp
--- a/Engine/Source/Developer/Windows/ShaderFormatD3D/Private/D3DShaderCompiler.cpp
+++ b/Engine/Source/Developer/Windows/ShaderFormatD3D/Private/D3DShaderCompiler.cpp
@@ -1210,8 +1210,8 @@ void CompileD3DShader(const FShaderCompilerInput& Input, FShaderCompilerOutput&
 		}
 	}
 
-	FShaderParameterParser ShaderParameterParser(TEXT("cbuffer"));
-	if (!ShaderParameterParser.ParseAndModify(Input, Output, PreprocessedShaderSource))
+	FShaderParameterParser ShaderParameterParser(Input.Environment.CompilerFlags, TEXT("cbuffer"));
+	if (!ShaderParameterParser.ParseAndModify(Input, Output.Errors, PreprocessedShaderSource))
 	{
 		// The FShaderParameterParser will add any relevant errors.
 		return;
diff --git a/Engine/Source/Runtime/RenderCore/Private/ShaderParameterParser.cpp b/Engine/Source/Runtime/RenderCore/Private/ShaderParameterParser.cpp
--- a/Engine/Source/Runtime/RenderCore/Private/ShaderParameterParser.cpp
+++ b/Engine/Source/Runtime/RenderCore/Private/ShaderParameterParser.cpp
@@ -137,21 +137,29 @@ static void AddNoteToDisplayShaderParameterMemberOnCppSide(
 FShaderParameterParser::FShaderParameterParser() = default;
 FShaderParameterParser::~FShaderParameterParser() = default;
 
+PRAGMA_DISABLE_DEPRECATION_WARNINGS
 FShaderParameterParser::FShaderParameterParser(const TCHAR* InConstantBufferType)
-	: ConstantBufferType(InConstantBufferType)
-{
-}
+	: FShaderParameterParser(FShaderCompilerFlags(), InConstantBufferType)
+{}
+FShaderParameterParser::FShaderParameterParser(
+	const TCHAR* InConstantBufferType,
+	TConstArrayView<const TCHAR*> InExtraSRVTypes,
+	TConstArrayView<const TCHAR*> InExtraUAVTypes)
+	: FShaderParameterParser(FShaderCompilerFlags(), InConstantBufferType, InExtraSRVTypes, InExtraUAVTypes)
+{}
+PRAGMA_ENABLE_DEPRECATION_WARNINGS
 
 FShaderParameterParser::FShaderParameterParser(
+	FShaderCompilerFlags CompilerFlags,
 	const TCHAR* InConstantBufferType,
-	TArrayView<const TCHAR* const> InExtraSRVTypes,
-	TArrayView<const TCHAR* const> InExtraUAVTypes
-)
+	TConstArrayView<const TCHAR*> InExtraSRVTypes,
+	TConstArrayView<const TCHAR*> InExtraUAVTypes)
 	: ConstantBufferType(InConstantBufferType)
 	, ExtraSRVTypes(InExtraSRVTypes)
 	, ExtraUAVTypes(InExtraUAVTypes)
-{
-}
+	, bBindlessResources(CompilerFlags.Contains(CFLAG_BindlessResources))
+	, bBindlessSamplers(CompilerFlags.Contains(CFLAG_BindlessSamplers))
+{}
 
 static const TCHAR* const s_AllSRVTypes[] =
 {
@@ -194,12 +202,12 @@ static const TCHAR* const s_AllSamplerTypes[] =
 
 EShaderParameterType FShaderParameterParser::ParseParameterType(
 	FStringView InType,
-	TArrayView<const TCHAR* const> InExtraSRVTypes,
-	TArrayView<const TCHAR* const> InExtraUAVTypes)
+	TConstArrayView<const TCHAR*> InExtraSRVTypes,
+	TConstArrayView<const TCHAR*> InExtraUAVTypes)
 {
-	TArrayView<const TCHAR* const> AllSamplerTypes(s_AllSamplerTypes);
-	TArrayView<const TCHAR* const> AllSRVTypes(s_AllSRVTypes);
-	TArrayView<const TCHAR* const> AllUAVTypes(s_AllUAVTypes);
+	TConstArrayView<const TCHAR*> AllSamplerTypes(s_AllSamplerTypes);
+	TConstArrayView<const TCHAR*> AllSRVTypes(s_AllSRVTypes);
+	TConstArrayView<const TCHAR*> AllUAVTypes(s_AllUAVTypes);
 
 	if (AllSamplerTypes.Contains(InType))
 	{
@@ -261,18 +269,18 @@ bool FShaderParameterParser::RemoveBindlessParameterPrefix(FString& InName)
 
 
 bool FShaderParameterParser::ParseParameters(
-	const FShaderCompilerInput& CompilerInput,
-	FShaderCompilerOutput& CompilerOutput)
+	const FShaderParametersMetadata* RootParametersStructure,
+	TArray<FShaderCompilerError>& OutErrors)
 {
 	const FStringView ShaderSource(OriginalParsedShader);
 
-	if (CompilerInput.RootParametersStructure)
+	if (RootParametersStructure)
 	{
 		// Reserves the number of parameters up front.
-		ParsedParameters.Reserve(CompilerInput.RootParametersStructure->GetSize() / sizeof(int32));
+		ParsedParameters.Reserve(RootParametersStructure->GetSize() / sizeof(int32));
 
 		IterateShaderParameterMembers(
-			*CompilerInput.RootParametersStructure,
+			*RootParametersStructure,
 			[&](const FShaderParametersMetadata& ParametersMetadata,
 				const FShaderParametersMetadata::FMember& Member,
 				const TCHAR* ShaderBindingName,
@@ -321,8 +329,8 @@ bool FShaderParameterParser::ParseParameters(
 
 		const int32 ShaderSourceLen = ShaderSource.Len();
 
-		int32 CurrentPragamLineoffset = -1;
-		int32 CurrentLineoffset = 0;
+		int32 CurrentPragmaLineOffset = -1;
+		int32 CurrentLineOffset = 0;
 
 		int32 TypeQualifierStartPos = -1;
 		int32 TypeStartPos = -1;
@@ -354,8 +362,8 @@ bool FShaderParameterParser::ParseParameters(
 		{
 			FShaderCompilerError Error;
 			Error.StrippedErrorMessage = ErrorMessage;
-			ExtractFileAndLine(CurrentPragamLineoffset, CurrentLineoffset, Error.ErrorVirtualFilePath, Error.ErrorLineString);
-			CompilerOutput.Errors.Add(Error);
+			ExtractFileAndLine(CurrentPragmaLineOffset, CurrentLineOffset, Error.ErrorVirtualFilePath, Error.ErrorLineString);
+			OutErrors.Add(Error);
 			bSuccess = false;
 		};
 
@@ -459,8 +467,8 @@ bool FShaderParameterParser::ParseParameters(
 				{
 					ParsedParameter.ParsedName = Name;
 					ParsedParameter.ParsedType = Type;
-					ParsedParameter.ParsedPragmaLineoffset = CurrentPragamLineoffset;
-					ParsedParameter.ParsedLineOffset = CurrentLineoffset;
+					ParsedParameter.ParsedPragmaLineoffset = CurrentPragmaLineOffset;
+					ParsedParameter.ParsedLineOffset = CurrentLineOffset;
 					ParsedParameter.ParsedCharOffsetStart = TypeQualifierStartPos != -1 ? TypeQualifierStartPos : TypeStartPos;
 					ParsedParameter.ParsedCharOffsetEnd = Cursor;
 					ParsedParameter.BindlessConversionType = BindlessConversionType;
@@ -485,7 +493,7 @@ bool FShaderParameterParser::ParseParameters(
 			const TCHAR* UpComing = ShaderSource.GetData() + Cursor;
 			const int32 RemainingSize = ShaderSourceLen - Cursor;
 
-			CurrentLineoffset += Char == '\n';
+			CurrentLineOffset += Char == '\n';
 
 			// Go to the next line if this is a preprocessor macro.
 			if (bGoToNextLine)
@@ -500,8 +508,8 @@ bool FShaderParameterParser::ParseParameters(
 			{
 				if (RemainingSize > 6 && FCString::Strncmp(UpComing, TEXT("#line "), 6) == 0)
 				{
-					CurrentPragamLineoffset = Cursor;
-					CurrentLineoffset = -1; // that will be incremented to 0 when reaching the \n at the end of the #line
+					CurrentPragmaLineOffset = Cursor;
+					CurrentLineOffset = -1; // that will be incremented to 0 when reaching the \n at the end of the #line
 				}
 
 				bGoToNextLine = true;
@@ -978,19 +986,18 @@ static const TCHAR* GetConstantSwizzle(uint16 ByteOffset)
 }
 
 bool FShaderParameterParser::MoveShaderParametersToRootConstantBuffer(
-	const FShaderCompilerInput& CompilerInput,
-	FShaderCompilerOutput& CompilerOutput,
+	const FShaderParametersMetadata* RootParametersStructure,
 	FString& PreprocessedShaderSource)
 {
 	bool bSuccess = true;
 
 	// Generate the root cbuffer content.
-	if (CompilerInput.RootParametersStructure && bNeedToMoveToRootConstantBuffer)
+	if (RootParametersStructure && bNeedToMoveToRootConstantBuffer)
 	{
 		FString RootCBufferContent;
 
 		IterateShaderParameterMembers(
-			*CompilerInput.RootParametersStructure,
+			*RootParametersStructure,
 			[&](const FShaderParametersMetadata& ParametersMetadata,
 				const FShaderParametersMetadata::FMember& Member,
 				const TCHAR* ShaderBindingName,
@@ -1067,13 +1074,25 @@ bool FShaderParameterParser::MoveShaderParametersToRootConstantBuffer(
 	return bSuccess;
 }
 
+PRAGMA_DISABLE_DEPRECATION_WARNINGS
 bool FShaderParameterParser::ParseAndModify(
 	const FShaderCompilerInput& CompilerInput,
 	FShaderCompilerOutput& CompilerOutput,
 	FString& PreprocessedShaderSource)
 {
+	// assume that if this deprecated overload is still being called, that also the deprecated constructor overloads
+	// are also still being used (and so re-set these member variables according to the environment in the input struct)
 	bBindlessResources = CompilerInput.Environment.CompilerFlags.Contains(CFLAG_BindlessResources);
 	bBindlessSamplers = CompilerInput.Environment.CompilerFlags.Contains(CFLAG_BindlessSamplers);
+	return ParseAndModify(CompilerInput, CompilerOutput.Errors, PreprocessedShaderSource);
+}
+PRAGMA_ENABLE_DEPRECATION_WARNINGS
+
+bool FShaderParameterParser::ParseAndModify(
+	const FShaderCompilerInput& CompilerInput,
+	TArray<FShaderCompilerError>& OutErrors,
+	FString& PreprocessedShaderSource)
+{
 	const bool bHasRootParameters = (CompilerInput.RootParametersStructure != nullptr);
 
 	// The shader doesn't have any parameter binding through shader structure, therefore don't do anything.
@@ -1085,7 +1104,7 @@ bool FShaderParameterParser::ParseAndModify(
 	bNeedToMoveToRootConstantBuffer = ConstantBufferType != nullptr && (CompilerInput.IsRayTracingShader() || CompilerInput.ShouldUseStableConstantBuffer());
 	OriginalParsedShader = PreprocessedShaderSource;
 
-	if (!ParseParameters(CompilerInput, CompilerOutput))
+	if (!ParseParameters(CompilerInput.RootParametersStructure, OutErrors))
 	{
 		return false;
 	}
@@ -1095,7 +1114,7 @@ bool FShaderParameterParser::ParseAndModify(
 
 	if (bNeedToMoveToRootConstantBuffer)
 	{
-		return MoveShaderParametersToRootConstantBuffer(CompilerInput, CompilerOutput, PreprocessedShaderSource);
+		return MoveShaderParametersToRootConstantBuffer(CompilerInput.RootParametersStructure, PreprocessedShaderSource);
 	}
 	
 	return true;
@@ -1292,6 +1311,73 @@ void FShaderParameterParser::ValidateShaderParameterTypes(
 		});
 }
 
+void SerializeStringView(FArchive& Ar, FStringView& StringView, const TCHAR* BasePtr)
+{
+	if (Ar.IsSaving())
+	{
+		uint32 Len = StringView.Len();
+		uint64 Offset = StringView.GetData() - BasePtr;
+		Ar << Len;
+		Ar << Offset;
+	}
+	else
+	{
+		uint32 Len;
+		uint64 Offset;
+		Ar << Len;
+		Ar << Offset;
+		StringView = FStringView(BasePtr + Offset, Len);
+	}
+}
+void SerializeParam(FArchive& Ar, FShaderParameterParser::FParsedShaderParameter& Param, const TCHAR* StringViewBase)
+{
+	Ar << Param.BaseType;
+	Ar << Param.PrecisionModifier;
+	Ar << Param.NumRows;
+	Ar << Param.NumColumns;
+	Ar << Param.MemberSize;
+	SerializeStringView(Ar, Param.ParsedName, StringViewBase);
+	SerializeStringView(Ar, Param.ParsedType, StringViewBase);
+	SerializeStringView(Ar, Param.ParsedArraySize, StringViewBase);
+	Ar << Param.ConstantBufferOffset;
+	Ar << Param.ParsedPragmaLineoffset;
+	Ar << Param.ParsedLineOffset;
+	Ar << Param.ParsedCharOffsetStart;
+	Ar << Param.ParsedCharOffsetEnd;
+	Ar << Param.BindlessConversionType;
+	Ar << Param.ConstantBufferParameterType;
+	Ar << Param.bGloballyCoherent;
+	Ar << Param.bIsBindable;
+}
+FArchive& operator<<(FArchive& Ar, FShaderParameterParser& Parser)
+{
+	Ar << Parser.bMovedLoosedParametersToRootConstantBuffer;
+	Ar << Parser.OriginalParsedShader;
+	if (Ar.IsSaving())
+	{
+		int32 ParsedParamCount = Parser.ParsedParameters.Num();
+		Ar << ParsedParamCount;
+		for (TPair<FString, FShaderParameterParser::FParsedShaderParameter>& ParsedParam : Parser.ParsedParameters)
+		{
+			Ar << ParsedParam.Key;
+			SerializeParam(Ar, ParsedParam.Value, Parser.OriginalParsedShader.GetCharArray().GetData());
+		}
+	}
+	else
+	{
+		uint32 ParsedParameterCount;
+		Ar << ParsedParameterCount;
+		Parser.ParsedParameters.Reserve(ParsedParameterCount);
+		for (uint32 CurParam = 0; CurParam < ParsedParameterCount; ++CurParam)
+		{
+			FString Name;
+			Ar << Name;
+			FShaderParameterParser::FParsedShaderParameter& ParsedParam = Parser.ParsedParameters.Add(Name);
+			SerializeParam(Ar, ParsedParam, Parser.OriginalParsedShader.GetCharArray().GetData());
+		}
+	}
+	return Ar;
+}
 void FShaderParameterParser::ExtractFileAndLine(int32 PragamLineoffset, int32 LineOffset, FString& OutFile, FString& OutLine) const
 {
 	if (PragamLineoffset == -1)
diff --git a/Engine/Source/Runtime/RenderCore/Public/ShaderParameterMetadata.h b/Engine/Source/Runtime/RenderCore/Public/ShaderParameterMetadata.h
--- a/Engine/Source/Runtime/RenderCore/Public/ShaderParameterMetadata.h
+++ b/Engine/Source/Runtime/RenderCore/Public/ShaderParameterMetadata.h
@@ -33,7 +33,8 @@ namespace EShaderPrecisionModifier
 	{
 		Float,
 		Half,
-		Fixed
+		Fixed,
+		Invalid
 	};
 };
 
diff --git a/Engine/Source/Runtime/RenderCore/Public/ShaderParameterParser.h b/Engine/Source/Runtime/RenderCore/Public/ShaderParameterParser.h
--- a/Engine/Source/Runtime/RenderCore/Public/ShaderParameterParser.h
+++ b/Engine/Source/Runtime/RenderCore/Public/ShaderParameterParser.h
@@ -5,8 +5,11 @@
 #include "CoreMinimal.h"
 #include "ShaderParameterMetadata.h"
 
+class FShaderCompilerFlags;
 struct FShaderCompilerInput;
 struct FShaderCompilerOutput;
+struct FShaderCompilerEnvironment;
+struct FShaderCompilerError;
 
 enum class EShaderParameterType : uint8;
 
@@ -25,11 +28,11 @@ public:
 	{
 	public:
 		/** Original information about the member. */
-		EUniformBufferBaseType BaseType;
-		EShaderPrecisionModifier::Type PrecisionModifier;
-		uint32 NumRows;
-		uint32 NumColumns;
-		uint32 MemberSize;
+		TEnumAsByte<EUniformBufferBaseType> BaseType{ UBMT_INVALID };
+		TEnumAsByte<EShaderPrecisionModifier::Type> PrecisionModifier{ EShaderPrecisionModifier::Invalid };
+		uint32 NumRows = 0u;
+		uint32 NumColumns = 0u;
+		uint32 MemberSize = 0u;
 
 		/** Information found about the member when parsing the preprocessed code. */
 		FStringView ParsedName; /** View into FShaderParameterParser::OriginalParsedShader */
@@ -62,25 +65,44 @@ public:
 	};
 
 	FShaderParameterParser();
+
+	UE_DEPRECATED(5.3, "Use FShaderParameterParser constructor which accepts FShaderCompilerFlags")
 	FShaderParameterParser(const TCHAR* InConstantBufferType);
-	FShaderParameterParser(const TCHAR* InConstantBufferType, TArrayView<const TCHAR* const> InExtraSRVTypes, TArrayView<const TCHAR* const> InExtraUAVTypes);
+	UE_DEPRECATED(5.3, "Use FShaderParameterParser constructor which accepts FShaderCompilerFlags")
+	FShaderParameterParser(const TCHAR* InConstantBufferType, TConstArrayView<const TCHAR*> InExtraSRVTypes, TConstArrayView<const TCHAR*> InExtraUAVTypes);
+
+	FShaderParameterParser(
+		FShaderCompilerFlags CompilerFlags,
+		const TCHAR* InConstantBufferType = nullptr,
+		TConstArrayView<const TCHAR*> InExtraSRVTypes = {},
+		TConstArrayView<const TCHAR*> InExtraUAVTypes = {});
+
 	virtual ~FShaderParameterParser();
 
+	FShaderParameterParser& operator=(FShaderParameterParser&&) = default;
+
 	static constexpr const TCHAR* kBindlessResourcePrefix = TEXT("BindlessResource_");
 	static constexpr const TCHAR* kBindlessSamplerPrefix = TEXT("BindlessSampler_");
 
-	static EShaderParameterType ParseParameterType(FStringView InType, TArrayView<const TCHAR* const> InExtraSRVTypes, TArrayView<const TCHAR* const> InExtraUAVTypes);
+	static EShaderParameterType ParseParameterType(FStringView InType, TConstArrayView<const TCHAR*> InExtraSRVTypes, TConstArrayView<const TCHAR*> InExtraUAVTypes);
 	static EShaderParameterType ParseAndRemoveBindlessParameterPrefix(FStringView& InName);
 	static EShaderParameterType ParseAndRemoveBindlessParameterPrefix(FString& InName);
 	static bool RemoveBindlessParameterPrefix(FString& InName);
 
 	/** Parses the preprocessed shader code and applies the necessary modifications to it. */
+	UE_DEPRECATED(5.3, "Use ParseAndModify overload accepting array of FShaderCompilerError instead of passing FShaderCompilerOutput")
 	bool ParseAndModify(
 		const FShaderCompilerInput& CompilerInput,
 		FShaderCompilerOutput& CompilerOutput,
 		FString& PreprocessedShaderSource
 	);
 
+	bool ParseAndModify(
+		const FShaderCompilerInput& CompilerInput,
+		TArray<FShaderCompilerError>& OutErrors,
+		FString& PreprocessedShaderSource
+	);
+
 	UE_DEPRECATED(5.2, "ParseAndModify doesn't need ConstantBufferType anymore")
 	bool ParseAndModify(
 		const FShaderCompilerInput& CompilerInput,
@@ -89,18 +111,9 @@ public:
 		const TCHAR* InConstantBufferType
 	)
 	{
+		PRAGMA_DISABLE_DEPRECATION_WARNINGS
 		return ParseAndModify(CompilerInput, CompilerOutput, PreprocessedShaderSource);
-	}
-
-	/** Parses the preprocessed shader code and move the parameters into root constant buffer */
-	UE_DEPRECATED(5.1, "ParseAndModify should be called instead.")
-	bool ParseAndMoveShaderParametersToRootConstantBuffer(
-		const FShaderCompilerInput& CompilerInput,
-		FShaderCompilerOutput& CompilerOutput,
-		FString& PreprocessedShaderSource,
-		const TCHAR* InConstantBufferType)
-	{
-		return ParseAndModify(CompilerInput, CompilerOutput, PreprocessedShaderSource);
+		PRAGMA_ENABLE_DEPRECATION_WARNINGS
 	}
 
 	/** Gets parsing information from a parameter binding name. */
@@ -147,11 +160,13 @@ public:
 		return ExtractFileAndLine(ParsedParameter.ParsedPragmaLineoffset, ParsedParameter.ParsedLineOffset, OutFile, OutLine);
 	}
 
+	friend FArchive& operator<<(FArchive& Ar, FShaderParameterParser& Parser);
+
 protected:
 	/** Parses the preprocessed shader code */
 	bool ParseParameters(
-		const FShaderCompilerInput& CompilerInput,
-		FShaderCompilerOutput& CompilerOutput
+		const FShaderParametersMetadata* RootParametersStructure,
+		TArray<FShaderCompilerError>& OutErrors
 	);
 
 	void RemoveMovingParametersFromSource(
@@ -165,8 +180,7 @@ protected:
 
 	/** Moves parsed parameters into the root constant buffer. */
 	bool MoveShaderParametersToRootConstantBuffer(
-		const FShaderCompilerInput& CompilerInput,
-		FShaderCompilerOutput& CompilerOutput,
+		const FShaderParametersMetadata* RootParametersStructure,
 		FString& PreprocessedShaderSource
 	);
 
@@ -178,10 +192,10 @@ protected:
 	*/
 	virtual FString GenerateBindlessParameterDeclaration(const FParsedShaderParameter& ParsedParameter) const;
 
-	const TCHAR* const ConstantBufferType = nullptr;
+	const TCHAR* ConstantBufferType = nullptr;
 
-	const TArray<const TCHAR*> ExtraSRVTypes;
-	const TArray<const TCHAR*> ExtraUAVTypes;
+	TConstArrayView<const TCHAR*> ExtraSRVTypes;
+	TConstArrayView<const TCHAR*> ExtraUAVTypes;
 
 	FString OriginalParsedShader;
 
