[shaders] fix FShaderParameterParser::ExtractFileAndLine to not assume a single space between end of #line and start of line number + more cleanup

diff --git a/Engine/Source/Runtime/RenderCore/Private/ShaderParameterParser.cpp b/Engine/Source/Runtime/RenderCore/Private/ShaderParameterParser.cpp
--- a/Engine/Source/Runtime/RenderCore/Private/ShaderParameterParser.cpp
+++ b/Engine/Source/Runtime/RenderCore/Private/ShaderParameterParser.cpp
@@ -467,7 +467,7 @@ bool FShaderParameterParser::ParseParameters(
 				{
 					ParsedParameter.ParsedName = Name;
 					ParsedParameter.ParsedType = Type;
-					ParsedParameter.ParsedPragmaLineoffset = CurrentPragmaLineOffset;
+					ParsedParameter.ParsedPragmaLineOffset = CurrentPragmaLineOffset;
 					ParsedParameter.ParsedLineOffset = CurrentLineOffset;
 					ParsedParameter.ParsedCharOffsetStart = TypeQualifierStartPos != -1 ? TypeQualifierStartPos : TypeStartPos;
 					ParsedParameter.ParsedCharOffsetEnd = Cursor;
@@ -1340,7 +1340,7 @@ void SerializeParam(FArchive& Ar, FShaderParameterParser::FParsedShaderParameter
 	SerializeStringView(Ar, Param.ParsedType, StringViewBase);
 	SerializeStringView(Ar, Param.ParsedArraySize, StringViewBase);
 	Ar << Param.ConstantBufferOffset;
-	Ar << Param.ParsedPragmaLineoffset;
+	Ar << Param.ParsedPragmaLineOffset;
 	Ar << Param.ParsedLineOffset;
 	Ar << Param.ParsedCharOffsetStart;
 	Ar << Param.ParsedCharOffsetEnd;
@@ -1378,23 +1378,24 @@ FArchive& operator<<(FArchive& Ar, FShaderParameterParser& Parser)
 	}
 	return Ar;
 }
-void FShaderParameterParser::ExtractFileAndLine(int32 PragamLineoffset, int32 LineOffset, FString& OutFile, FString& OutLine) const
+
+void FShaderParameterParser::ExtractFileAndLine(int32 PragmaLineOffset, int32 LineOffset, FString& OutFile, FString& OutLine) const
 {
-	if (PragamLineoffset == -1)
+	if (PragmaLineOffset == -1)
 	{
 		return;
 	}
 
-	check(FCString::Strncmp((*OriginalParsedShader) + PragamLineoffset, TEXT("#line "), 6) == 0);
+	check(FCString::Strncmp((*OriginalParsedShader) + PragmaLineOffset, TEXT("#line"), 5) == 0);
 
 	const int32 ShaderSourceLen = OriginalParsedShader.Len();
 
 	int32 StartFilePos = -1;
 	int32 EndFilePos = -1;
-	int32 StartLinePos = PragamLineoffset + 6;
+	int32 StartLinePos = -1;
 	int32 EndLinePos = -1;
 
-	for (int32 Cursor = StartLinePos; Cursor < ShaderSourceLen; Cursor++)
+	for (int32 Cursor = PragmaLineOffset + 5; Cursor < ShaderSourceLen; Cursor++)
 	{
 		const TCHAR Char = OriginalParsedShader[Cursor];
 
@@ -1402,28 +1403,22 @@ void FShaderParameterParser::ExtractFileAndLine(int32 PragamLineoffset, int32 Li
 		{
 			break;
 		}
-
-		if (EndLinePos == -1)
+		else if (StartLinePos == -1 && FChar::IsDigit(Char))
 		{
-			if (Char > '9' || Char < '0')
-			{
-				EndLinePos = Cursor - 1;
-			}
+			StartLinePos = Cursor;
 		}
-		else if (StartFilePos == -1)
+		else if (StartLinePos != -1 && EndLinePos == -1 && !FChar::IsDigit(Char))
 		{
-			if (Char == '"')
-			{
-				StartFilePos = Cursor + 1;
-			}
+			EndLinePos = Cursor - 1;
 		}
-		else if (EndFilePos == -1)
+		else if (StartFilePos == -1 && Char == TEXT('"'))
 		{
-			if (Char == '"')
-			{
-				EndFilePos = Cursor - 1;
-				break;
-			}
+			StartFilePos = Cursor + 1;
+		}
+		else if (StartFilePos != -1 && EndFilePos == -1 && Char == TEXT('"'))
+		{
+			EndFilePos = Cursor - 1;
+			break;
 		}
 	}
 
diff --git a/Engine/Source/Runtime/RenderCore/Public/ShaderParameterParser.h b/Engine/Source/Runtime/RenderCore/Public/ShaderParameterParser.h
--- a/Engine/Source/Runtime/RenderCore/Public/ShaderParameterParser.h
+++ b/Engine/Source/Runtime/RenderCore/Public/ShaderParameterParser.h
@@ -48,7 +48,7 @@ public:
 			return !ParsedType.IsEmpty();
 		}
 
-		int32 ParsedPragmaLineoffset = 0;
+		int32 ParsedPragmaLineOffset = 0;
 		int32 ParsedLineOffset = 0;
 
 		/** Character position of the start and end of the parameter decelaration in FParsedShaderParameter::OriginalParsedShader */
@@ -157,7 +157,7 @@ public:
 	/** Gets file and line of the parameter in the shader source code. */
 	void GetParameterFileAndLine(const FParsedShaderParameter& ParsedParameter, FString& OutFile, FString& OutLine) const
 	{
-		return ExtractFileAndLine(ParsedParameter.ParsedPragmaLineoffset, ParsedParameter.ParsedLineOffset, OutFile, OutLine);
+		return ExtractFileAndLine(ParsedParameter.ParsedPragmaLineOffset, ParsedParameter.ParsedLineOffset, OutFile, OutLine);
 	}
 
 	friend FArchive& operator<<(FArchive& Ar, FShaderParameterParser& Parser);
