[shaders]
- add a debug info flags enum/field to indicate which shader debug outputs are desired
- change direct compile txt file mechanism to use this (more options will be added in the near future)
- rename the field on the debug info dump options which indicates that the direct compile file should be skipped to bSourceOnly (when addtional outputs are added these should be skipped as well)

diff --git a/Engine/Plugins/FX/Niagara/Source/NiagaraEditor/Private/NiagaraCompiler.cpp b/Engine/Plugins/FX/Niagara/Source/NiagaraEditor/Private/NiagaraCompiler.cpp
--- a/Engine/Plugins/FX/Niagara/Source/NiagaraEditor/Private/NiagaraCompiler.cpp
+++ b/Engine/Plugins/FX/Niagara/Source/NiagaraEditor/Private/NiagaraCompiler.cpp
@@ -1771,8 +1771,6 @@ int32 FHlslNiagaraCompiler::CompileScript(const FStringView GroupName, const FNi
 {
 	SCOPE_CYCLE_COUNTER(STAT_NiagaraEditor_HlslCompiler_CompileScript);
 
-	static const auto CVarDumpCommandLine = IConsoleManager::Get().FindConsoleVariable(TEXT("r.DumpShaderDebugWorkerCommandLine"));
-
 	CompileResults.Data = MakeShared<FNiagaraVMExecutableData>();
 
 	//TODO: This should probably be done via the same route that other shaders take through the shader compiler etc.
@@ -1795,7 +1793,7 @@ int32 FHlslNiagaraCompiler::CompileScript(const FStringView GroupName, const FNi
 	Input.Environment.SetDefine(TEXT("MESHSHADER"), 0);
 	Input.Environment.SetDefine(TEXT("AMPLIFICATIONSHADER"), 0);
 	Input.Environment.IncludeVirtualPathToContentsMap.Add(TEXT("/Engine/Generated/NiagaraEmitterInstance.ush"), TranslatedHLSL);
-	Input.bGenerateDirectCompileFile = CVarDumpCommandLine ? CVarDumpCommandLine->GetBool() : false;
+	Input.DebugInfoFlags = GShaderCompilingManager->GetDumpShaderDebugInfoFlags();
 	Input.DumpDebugInfoRootPath = GShaderCompilingManager->GetAbsoluteShaderDebugInfoDirectory() / TEXT("VM");
 	Input.DebugGroupName = GroupName;
 	Input.DebugExtension.Empty();
diff --git a/Engine/Source/Developer/ShaderCompilerCommon/Private/ShaderCompilerCommon.cpp b/Engine/Source/Developer/ShaderCompilerCommon/Private/ShaderCompilerCommon.cpp
--- a/Engine/Source/Developer/ShaderCompilerCommon/Private/ShaderCompilerCommon.cpp
+++ b/Engine/Source/Developer/ShaderCompilerCommon/Private/ShaderCompilerCommon.cpp
@@ -1400,7 +1400,7 @@ namespace UE::ShaderCompilerCommon
 			FString Contents = UE::ShaderCompilerCommon::GetDebugShaderContents(Input, PreprocessedSource, DebugDataOptions);
 			FFileHelper::SaveStringToFile(Contents, *(Input.DumpDebugInfoPath / BaseSourceFilename));
 
-			if (Input.bGenerateDirectCompileFile && !Options.bSkipDirectCompileTxt)
+			if (EnumHasAnyFlags(Input.DebugInfoFlags, EShaderDebugInfoFlags::DirectCompileCommandLine) && !Options.bSourceOnly)
 			{
 				FFileHelper::SaveStringToFile(CreateShaderCompilerWorkerDirectCommandLine(Input), *(Input.DumpDebugInfoPath / FString::Printf(TEXT("%sDirectCompile.txt"), *Prefix)));
 			}
diff --git a/Engine/Source/Developer/ShaderCompilerCommon/Public/ShaderCompilerCommon.h b/Engine/Source/Developer/ShaderCompilerCommon/Public/ShaderCompilerCommon.h
--- a/Engine/Source/Developer/ShaderCompilerCommon/Public/ShaderCompilerCommon.h
+++ b/Engine/Source/Developer/ShaderCompilerCommon/Public/ShaderCompilerCommon.h
@@ -112,7 +112,12 @@ namespace UE::ShaderCompilerCommon
 		const TCHAR* FilenamePrefix = nullptr;
 		TFunction<FString()> AppendPreSource{};
 		TFunction<FString()> AppendPostSource{};
-		bool bSkipDirectCompileTxt = false;
+		union
+		{
+			bool bSourceOnly = false; // if true, will only output source .usf as directed and skip all other debug data artifacts
+			UE_DEPRECATED(5.3, "bSkipDirectCompileTxt is deprecated, use bSourceOnly flag instead")
+			bool bSkipDirectCompileTxt;
+		};
 	};
 	/*
 	 * Dumps common debug information (preprocessed .usf as constructed by GetDebugShaderContents, and a directcompile.txt file 
diff --git a/Engine/Source/Runtime/Engine/Private/ShaderCompiler/ShaderCompiler.cpp b/Engine/Source/Runtime/Engine/Private/ShaderCompiler/ShaderCompiler.cpp
--- a/Engine/Source/Runtime/Engine/Private/ShaderCompiler/ShaderCompiler.cpp
+++ b/Engine/Source/Runtime/Engine/Private/ShaderCompiler/ShaderCompiler.cpp
@@ -4333,6 +4333,16 @@ FShaderCompilingManager::EDumpShaderDebugInfo FShaderCompilingManager::GetDumpSh
 	return static_cast<FShaderCompilingManager::EDumpShaderDebugInfo>(GDumpShaderDebugInfo);
 }
 
+EShaderDebugInfoFlags FShaderCompilingManager::GetDumpShaderDebugInfoFlags() const
+{
+	EShaderDebugInfoFlags Flags = EShaderDebugInfoFlags::Default;
+	if (GDumpShaderDebugInfoSCWCommandLine)
+	{
+		Flags |= EShaderDebugInfoFlags::DirectCompileCommandLine;
+	}
+	return Flags;
+}
+
 FString FShaderCompilingManager::CreateShaderDebugInfoPath(const FShaderCompilerInput& ShaderCompilerInput) const
 {
 	FString DumpDebugInfoPath = ShaderCompilerInput.DumpDebugInfoRootPath / ShaderCompilerInput.DebugGroupName + ShaderCompilerInput.DebugExtension;
@@ -5970,8 +5980,8 @@ void GlobalBeginCompileShader(
 	Input.EntryPointName = FunctionName;
 	Input.bCompilingForShaderPipeline = false;
 	Input.bIncludeUsedOutputs = false;
-	Input.bGenerateDirectCompileFile = (GDumpShaderDebugInfoSCWCommandLine != 0);
 	Input.DumpDebugInfoRootPath = GShaderCompilingManager->GetAbsoluteShaderDebugInfoDirectory() / Input.ShaderPlatformName.ToString();
+	Input.DebugInfoFlags = GShaderCompilingManager->GetDumpShaderDebugInfoFlags();
 	// asset material name or "Global"
 	Input.DebugGroupName = DebugGroupName;
 	Input.DebugDescription = DebugDescription;
diff --git a/Engine/Source/Runtime/Engine/Public/ShaderCompiler.h b/Engine/Source/Runtime/Engine/Public/ShaderCompiler.h
--- a/Engine/Source/Runtime/Engine/Public/ShaderCompiler.h
+++ b/Engine/Source/Runtime/Engine/Public/ShaderCompiler.h
@@ -834,6 +834,7 @@ public:
 	};
 
 	ENGINE_API EDumpShaderDebugInfo GetDumpShaderDebugInfo() const;
+	ENGINE_API EShaderDebugInfoFlags GetDumpShaderDebugInfoFlags() const;
 	ENGINE_API FString CreateShaderDebugInfoPath(const FShaderCompilerInput& ShaderCompilerInput) const;
 	ENGINE_API bool ShouldRecompileToDumpShaderDebugInfo(const FShaderCompileJob& Job) const;
 	ENGINE_API bool ShouldRecompileToDumpShaderDebugInfo(const FShaderCompilerInput& Input, const FShaderCompilerOutput& Output, bool bSucceeded) const;
diff --git a/Engine/Source/Runtime/RenderCore/Private/ShaderCore.cpp b/Engine/Source/Runtime/RenderCore/Private/ShaderCore.cpp
--- a/Engine/Source/Runtime/RenderCore/Private/ShaderCore.cpp
+++ b/Engine/Source/Runtime/RenderCore/Private/ShaderCore.cpp
@@ -2583,11 +2583,11 @@ FArchive& operator<<(FArchive& Ar, FShaderCompilerInput& Input)
 	Ar << Input.ShaderName;
 	Ar << Input.bSkipPreprocessedCache;
 	Ar << Input.bCompilingForShaderPipeline;
-	Ar << Input.bGenerateDirectCompileFile;
 	Ar << Input.bIncludeUsedOutputs;
 	Ar << Input.UsedOutputs;
 	Ar << Input.DumpDebugInfoRootPath;
 	Ar << Input.DumpDebugInfoPath;
+	Ar << Input.DebugInfoFlags;
 	Ar << Input.DebugExtension;
 	Ar << Input.DebugGroupName;
 	Ar << Input.DebugDescription;
diff --git a/Engine/Source/Runtime/RenderCore/Public/ShaderCompilerCore.h b/Engine/Source/Runtime/RenderCore/Public/ShaderCompilerCore.h
--- a/Engine/Source/Runtime/RenderCore/Public/ShaderCompilerCore.h
+++ b/Engine/Source/Runtime/RenderCore/Public/ShaderCompilerCore.h
@@ -10,6 +10,7 @@
 #include "HAL/FileManager.h"
 #include "Stats/Stats.h"
 #include "Templates/RefCounting.h"
+#include "Misc/EnumClassFlags.h"
 #include "Misc/SecureHash.h"
 #include "Misc/Paths.h"
 #include "Misc/CoreStats.h"
@@ -21,7 +22,7 @@ class FShaderCompileJob;
 class FShaderPipelineCompileJob;
 
 // this is for the protocol, not the data, bump if FShaderCompilerInput or ProcessInputFromArchive changes.
-inline const int32 ShaderCompileWorkerInputVersion = 18;
+inline const int32 ShaderCompileWorkerInputVersion = 19;
 // this is for the protocol, not the data, bump if FShaderCompilerOutput or WriteToOutputArchive changes.
 inline const int32 ShaderCompileWorkerOutputVersion = 13;
 // this is for the protocol, not the data.
@@ -188,6 +189,13 @@ namespace FOodleDataCompression
 	enum class ECompressionLevel : int8;
 }
 
+enum class EShaderDebugInfoFlags : uint8
+{
+	Default = 0,
+	DirectCompileCommandLine = 1 << 0,
+};
+ENUM_CLASS_FLAGS(EShaderDebugInfoFlags)
+
 /** Struct that gathers all readonly inputs needed for the compilation of a single shader. */
 struct FShaderCompilerInput
 {
@@ -204,8 +212,12 @@ struct FShaderCompilerInput
 	// Skips the preprocessor and instead loads the usf file directly
 	bool bSkipPreprocessedCache;
 
+	UE_DEPRECATED(5.3, "Use DebugInfoFlags field (EDebugInfoFlags::DirectCompileCommandLine)")
 	bool bGenerateDirectCompileFile;
 
+	// Indicates which additional debug outputs should be written for this compile job.
+	EShaderDebugInfoFlags DebugInfoFlags;
+
 	// Shader pipeline information
 	bool bCompilingForShaderPipeline;
 	bool bIncludeUsedOutputs;
@@ -245,12 +257,22 @@ struct FShaderCompilerInput
 	FShaderCompilerInput() :
 		Target(SF_NumFrequencies, SP_NumPlatforms),
 		bSkipPreprocessedCache(false),
-		bGenerateDirectCompileFile(false),
+		DebugInfoFlags(EShaderDebugInfoFlags::Default),
 		bCompilingForShaderPipeline(false),
 		bIncludeUsedOutputs(false)
 	{
 	}
 
+	PRAGMA_DISABLE_DEPRECATION_WARNINGS
+	// Explicitly-defaulted copy/move ctors & assignment operators are needed temporarily due to 
+	// deprecation of bGenerateDirectCompileFile field. These can be removed once the deprecation
+	// window for said field ends.
+	FShaderCompilerInput(FShaderCompilerInput&&) = default;
+	FShaderCompilerInput(const FShaderCompilerInput&) = default;
+	FShaderCompilerInput& operator=(FShaderCompilerInput&&) = default;
+	FShaderCompilerInput& operator=(const FShaderCompilerInput&) = default;
+	PRAGMA_ENABLE_DEPRECATION_WARNINGS
+
 	bool DumpDebugInfoEnabled() const 
 	{
 		return DumpDebugInfoPath != TEXT("") && IFileManager::Get().DirectoryExists(*DumpDebugInfoPath);
