Cleanup of some media shader parameters & added support for slog3 encoded data
#rnx

diff --git a/Engine/Plugins/Experimental/AVCodecs/AVCodecsCore/Source/AVCodecsCoreRHI/AVCodecsCoreRHI.Build.cs b/Engine/Plugins/Experimental/AVCodecs/AVCodecsCore/Source/AVCodecsCoreRHI/AVCodecsCoreRHI.Build.cs
--- a/Engine/Plugins/Experimental/AVCodecs/AVCodecsCore/Source/AVCodecsCoreRHI/AVCodecsCoreRHI.Build.cs
+++ b/Engine/Plugins/Experimental/AVCodecs/AVCodecsCore/Source/AVCodecsCoreRHI/AVCodecsCoreRHI.Build.cs
@@ -15,7 +15,8 @@ public class AVCodecsCoreRHI : ModuleRules
 			"Engine",
 			"AVCodecsCore",
 			"RHI",
-			"VulkanRHI"
+			"VulkanRHI",
+			"ColorManagement"
 		});
 
 		PublicDependencyModuleNames.AddRange(new string[] {
diff --git a/Engine/Plugins/Experimental/AVCodecs/AVCodecsCore/Source/AVCodecsCoreRHI/Private/Video/Resources/VideoResourceRHI.cpp b/Engine/Plugins/Experimental/AVCodecs/AVCodecsCore/Source/AVCodecsCoreRHI/Private/Video/Resources/VideoResourceRHI.cpp
--- a/Engine/Plugins/Experimental/AVCodecs/AVCodecsCore/Source/AVCodecsCoreRHI/Private/Video/Resources/VideoResourceRHI.cpp
+++ b/Engine/Plugins/Experimental/AVCodecs/AVCodecsCore/Source/AVCodecsCoreRHI/Private/Video/Resources/VideoResourceRHI.cpp
@@ -20,6 +20,8 @@
 #include "RHIStaticStates.h"
 #include "MediaShaders.h"
 #include "PipelineStateCache.h"
+#include "ColorManagementDefines.h"
+#include "ColorSpace.h"
 
 REGISTER_TYPEID(FVideoContextRHI);
 REGISTER_TYPEID(FVideoResourceRHI);
@@ -315,7 +317,7 @@ TSharedPtr<FVideoResourceRHI> FVideoResourceRHI::TransformResource(FVideoDescrip
 						TShaderMapRef<FNV12ConvertAsBytesPS> ConvertShader(ShaderMap);
 						GraphicsPSOInit.BoundShaderState.PixelShaderRHI = ConvertShader.GetPixelShader();
 						SetGraphicsPipelineState(RHICmdList, GraphicsPSOInit, 0);
-						SetShaderParametersLegacyPS(RHICmdList, ConvertShader, Source, OutputDim, YUVToRGBMatrix, true, FMatrix44f::Identity, false);
+						SetShaderParametersLegacyPS(RHICmdList, ConvertShader, Source, OutputDim, YUVToRGBMatrix, UE::Color::EEncoding::sRGB, FMatrix44f::Identity, false);
 
 						// draw full size quad into render target
 						FBufferRHIRef VertexBuffer = CreateTempMediaVertexBuffer();
@@ -385,7 +387,7 @@ void FVideoResourceRHI::TransformResourceTo(FRHICommandListImmediate& RHICmdList
 					TShaderMapRef<FNV12ConvertAsBytesPS> ConvertShader(ShaderMap);
 					GraphicsPSOInit.BoundShaderState.PixelShaderRHI = ConvertShader.GetPixelShader();
 					SetGraphicsPipelineState(RHICmdList, GraphicsPSOInit, 0);
-					SetShaderParametersLegacyPS(RHICmdList, ConvertShader, Raw.Texture, OutputDim, YUVToRGBMatrix, true, FMatrix44f::Identity, false);
+					SetShaderParametersLegacyPS(RHICmdList, ConvertShader, Raw.Texture, OutputDim, YUVToRGBMatrix, UE::Color::EEncoding::sRGB, FMatrix44f::Identity, false);
 
 					// draw full size quad into render target
 					FBufferRHIRef VertexBuffer = CreateTempMediaVertexBuffer();
diff --git a/Engine/Plugins/Media/ElectraUtil/Source/ElectraSamples/Private/ElectraTextureSampleAVF.cpp b/Engine/Plugins/Media/ElectraUtil/Source/ElectraSamples/Private/ElectraTextureSampleAVF.cpp
--- a/Engine/Plugins/Media/ElectraUtil/Source/ElectraSamples/Private/ElectraTextureSampleAVF.cpp
+++ b/Engine/Plugins/Media/ElectraUtil/Source/ElectraSamples/Private/ElectraTextureSampleAVF.cpp
@@ -255,7 +255,7 @@ void FElectraMediaTexConvApple::ConvertTexture(FTexture2DRHIRef & InDstTexture,
 						FShaderResourceViewRHIRef Y_SRV = RHICreateShaderResourceView(YTex, 0, 1, PF_G8);
 						FShaderResourceViewRHIRef UV_SRV = RHICreateShaderResourceView(UVTex, 0, 1, PF_R8G8);
 
-						SetShaderParametersLegacyPS(RHICmdList, PixelShader, FIntPoint(YWidth, YHeight), Y_SRV, UV_SRV, FIntPoint(YWidth, YHeight), YUVMtx, EncodingType == UE::Color::EEncoding::sRGB, ColorSpaceMtx, false);
+						SetShaderParametersLegacyPS(RHICmdList, PixelShader, FIntPoint(YWidth, YHeight), Y_SRV, UV_SRV, FIntPoint(YWidth, YHeight), YUVMtx, EncodingType, ColorSpaceMtx, false);
 					}
 					else
 					{
@@ -273,7 +273,7 @@ void FElectraMediaTexConvApple::ConvertTexture(FTexture2DRHIRef & InDstTexture,
 						FShaderResourceViewRHIRef UV_SRV = RHICreateShaderResourceView(UVTex, 0, 1, PF_G16R16);
 						
 						// Update shader uniform parameters.
-						SetShaderParametersLegacyPS(RHICmdList, PixelShader, FIntPoint(YWidth, YHeight), Y_SRV, UV_SRV, FIntPoint(YWidth, YHeight), YUVMtx, ColorSpaceMtx, EncodingType == UE::Color::EEncoding::sRGB, EncodingType == UE::Color::EEncoding::ST2084);
+						SetShaderParametersLegacyPS(RHICmdList, PixelShader, FIntPoint(YWidth, YHeight), Y_SRV, UV_SRV, FIntPoint(YWidth, YHeight), YUVMtx, ColorSpaceMtx, EncodingType);
 					}
 
 
diff --git a/Engine/Shaders/Private/GammaCorrectionCommon.ush b/Engine/Shaders/Private/GammaCorrectionCommon.ush
--- a/Engine/Shaders/Private/GammaCorrectionCommon.ush
+++ b/Engine/Shaders/Private/GammaCorrectionCommon.ush
@@ -176,4 +176,14 @@ LinearToST2084(float3 lin)
 	float3 P = pow( N, m2 );
 	
 	return P;
-}
\ No newline at end of file
+}
+
+//
+// Sony S-Log3
+//
+float3
+SLog3ToLinear(float3 Value)
+{
+	return (Value >= 171.2102946929f / 1023.0f) ? (pow(10.0f, (Value * 1023.0f - 420.f) / 261.5f)) * 0.19f - 0.01f
+												: (Value * 1023.0f - 95.0f) * 0.01125000f / (171.2102946929f - 95.0f);
+}
diff --git a/Engine/Shaders/Private/MediaShaders.usf b/Engine/Shaders/Private/MediaShaders.usf
--- a/Engine/Shaders/Private/MediaShaders.usf
+++ b/Engine/Shaders/Private/MediaShaders.usf
@@ -52,13 +52,32 @@ void BMPConvertPS(
 	OutColor.w = RGB.w;
 }
 
+#define UE_Color_EEncoding_Linear		1u
+#define UE_Color_EEncoding_sRGB			2u
+#define UE_Color_EEncoding_ST2084		3u
+#define UE_Color_EEncoding_Gamma22		4u
+#define UE_Color_EEncoding_BT1886		5u
+#define UE_Color_EEncoding_Gamma26		6u
+#define UE_Color_EEncoding_Cineon		7u
+#define UE_Color_EEncoding_REDLog		8u
+#define UE_Color_EEncoding_REDLog3G10	9u
+#define UE_Color_EEncoding_SLog1		10u
+#define UE_Color_EEncoding_SLog2		11u
+#define UE_Color_EEncoding_SLog3		12u
+#define UE_Color_EEncoding_AlexaV3LogC	13u
+#define UE_Color_EEncoding_CanonLog		14u
+#define UE_Color_EEncoding_ProTune		15u
+#define UE_Color_EEncoding_VLog			16u
 
 float3 ApplyColorProcessingRGB(float3 Color, float4x4 CSMtx, uint EOTF)
 {
 	// Go linear
-	if (EOTF != 0)
+	switch(EOTF)
 	{
-		Color = (EOTF == 1) ? sRGBToLinear(Color) : ST2084ToLinear(Color);
+	case UE_Color_EEncoding_sRGB:	Color = sRGBToLinear(Color); break;
+	case UE_Color_EEncoding_ST2084: Color = ST2084ToLinear(Color); break;
+	case UE_Color_EEncoding_SLog3:	Color = SLog3ToLinear(Color); break;
+	default: break;
 	}
 
 	// Adjust colorspace
@@ -100,7 +119,7 @@ void NV12ConvertPS(
 	YUV.yz = NV12ConvertUB.SwapChroma != 0u ? YUV.zy : YUV.yz;
 
 	// Output with solid alpha
-	OutColor = float4(ApplyColorProcessing(YUV, NV12ConvertUB.ColorTransform, NV12ConvertUB.CSTransform, (NV12ConvertUB.SrgbToLinear != 0u) ? 1 : 0), 1.0f); 
+	OutColor = float4(ApplyColorProcessing(YUV, NV12ConvertUB.ColorTransform, NV12ConvertUB.CSTransform, NV12ConvertUB.EOTF), 1.0f); 
 }
 
 
@@ -127,7 +146,7 @@ void NV12ConvertAsBytesPS(
 	YUV.z = NV12ConvertAsBytesUB.Texture.Sample(NV12ConvertAsBytesUB.SamplerP, float2(VPos, ScaledUV.y)).x;
 	YUV.yz = NV12ConvertAsBytesUB.SwapChroma != 0u ? YUV.zy : YUV.yz;
 
-	OutColor = float4(ApplyColorProcessing(YUV, NV12ConvertAsBytesUB.ColorTransform, NV12ConvertAsBytesUB.CSTransform, (NV12ConvertAsBytesUB.SrgbToLinear != 0u) ? 1 : 0), 1.0f); 
+	OutColor = float4(ApplyColorProcessing(YUV, NV12ConvertAsBytesUB.ColorTransform, NV12ConvertAsBytesUB.CSTransform, NV12ConvertAsBytesUB.EOTF), 1.0f); 
 }
 
 
diff --git a/Engine/Source/Runtime/MediaAssets/Private/Misc/MediaTextureResource.cpp b/Engine/Source/Runtime/MediaAssets/Private/Misc/MediaTextureResource.cpp
--- a/Engine/Source/Runtime/MediaAssets/Private/Misc/MediaTextureResource.cpp
+++ b/Engine/Source/Runtime/MediaAssets/Private/Misc/MediaTextureResource.cpp
@@ -1015,14 +1015,14 @@ void FMediaTextureResource::GetColorSpaceConversionMatrixForSample(const TShared
 						FIntPoint TexDim = InputTexture->GetSizeXY();
 						TempSRV0 = RHICreateShaderResourceView(InputTexture, 0, 1, PF_G8);		// note: the types of the views select the correct planes (offset) "magically"
 						TempSRV1 = RHICreateShaderResourceView(InputTexture, 0, 1, PF_R8G8);
-						SetShaderParametersLegacyPS(CommandList, ConvertShader, TexDim, TempSRV0, TempSRV1, OutputDim, YUVMtx, Sample->GetEncodingType() == UE::Color::EEncoding::sRGB, ColorSpaceMtx, SampleFormat == EMediaTextureSampleFormat::CharNV21);
+						SetShaderParametersLegacyPS(CommandList, ConvertShader, TexDim, TempSRV0, TempSRV1, OutputDim, YUVMtx, Sample->GetEncodingType(), ColorSpaceMtx, SampleFormat == EMediaTextureSampleFormat::CharNV21);
 					}
 					else
 					{
 						TShaderMapRef<FNV12ConvertAsBytesPS> ConvertShader(ShaderMap);
 						GraphicsPSOInit.BoundShaderState.PixelShaderRHI = ConvertShader.GetPixelShader();
 						SetGraphicsPipelineState(CommandList, GraphicsPSOInit, 0);
-						SetShaderParametersLegacyPS(CommandList, ConvertShader, InputTexture, OutputDim, YUVMtx, Sample->GetEncodingType() == UE::Color::EEncoding::sRGB, ColorSpaceMtx, SampleFormat == EMediaTextureSampleFormat::CharNV21);
+						SetShaderParametersLegacyPS(CommandList, ConvertShader, InputTexture, OutputDim, YUVMtx, Sample->GetEncodingType(), ColorSpaceMtx, SampleFormat == EMediaTextureSampleFormat::CharNV21);
 					}
 				}
 				break;
@@ -1043,7 +1043,7 @@ void FMediaTextureResource::GetColorSpaceConversionMatrixForSample(const TShared
 
 						FShaderResourceViewRHIRef Y_SRV = RHICreateShaderResourceView(InputTexture, 0, 1, PF_G16);		// note: the types of the views select the correct planes (offset) "magically"
 						FShaderResourceViewRHIRef UV_SRV = RHICreateShaderResourceView(InputTexture, 0, 1, PF_G16R16);
-						SetShaderParametersLegacyPS(CommandList, ConvertShader, TexDim, Y_SRV, UV_SRV, OutputDim, YUVMtx, ColorSpaceMtx, Sample->GetEncodingType() == UE::Color::EEncoding::sRGB, Sample->GetEncodingType() == UE::Color::EEncoding::ST2084);
+						SetShaderParametersLegacyPS(CommandList, ConvertShader, TexDim, Y_SRV, UV_SRV, OutputDim, YUVMtx, ColorSpaceMtx, Sample->GetEncodingType());
 					}
 					else
 					{
@@ -1051,7 +1051,7 @@ void FMediaTextureResource::GetColorSpaceConversionMatrixForSample(const TShared
 						GraphicsPSOInit.BoundShaderState.PixelShaderRHI = ConvertShader.GetPixelShader();
 						SetGraphicsPipelineState(CommandList, GraphicsPSOInit, 0);
 
-						SetShaderParametersLegacyPS(CommandList, ConvertShader, TexDim, InputTexture, OutputDim, YUVMtx, ColorSpaceMtx, Sample->GetEncodingType() == UE::Color::EEncoding::sRGB, Sample->GetEncodingType() == UE::Color::EEncoding::ST2084);
+						SetShaderParametersLegacyPS(CommandList, ConvertShader, TexDim, InputTexture, OutputDim, YUVMtx, ColorSpaceMtx, Sample->GetEncodingType());
 					}
 				}
 				break;
@@ -1070,7 +1070,7 @@ void FMediaTextureResource::GetColorSpaceConversionMatrixForSample(const TShared
 					TShaderMapRef<FYUVv216ConvertPS> ConvertShader(ShaderMap);
 					GraphicsPSOInit.BoundShaderState.PixelShaderRHI = ConvertShader.GetPixelShader();
 					SetGraphicsPipelineState(CommandList, GraphicsPSOInit, 0);
-					SetShaderParametersLegacyPS(CommandList, ConvertShader, InputTexture, OutputDim, YUVMtx, Sample->GetEncodingType() == UE::Color::EEncoding::sRGB, Sample->GetEncodingType() == UE::Color::EEncoding::ST2084, ColorSpaceMtx,
+					SetShaderParametersLegacyPS(CommandList, ConvertShader, InputTexture, OutputDim, YUVMtx, Sample->GetEncodingType(), ColorSpaceMtx,
 												SampleFormat!= EMediaTextureSampleFormat::CharYUY2 && SampleFormat != EMediaTextureSampleFormat::CharYVYU,	// Y or Cb first
 												InputTexture->GetFormat() == PF_B8G8R8A8,																	// ARGB vs. ABGR (memory order)
 												SampleFormat == EMediaTextureSampleFormat::CharYVYU															// Cb / Cr swap
@@ -1088,7 +1088,7 @@ void FMediaTextureResource::GetColorSpaceConversionMatrixForSample(const TShared
 					TShaderMapRef<FYUVv210ConvertPS> ConvertShader(ShaderMap);
 					GraphicsPSOInit.BoundShaderState.PixelShaderRHI = ConvertShader.GetPixelShader();
 					SetGraphicsPipelineState(CommandList, GraphicsPSOInit, 0);
-					SetShaderParametersLegacyPS(CommandList, ConvertShader, InputTexture, OutputDim, YUVMtx, Sample->GetEncodingType() == UE::Color::EEncoding::sRGB, Sample->GetEncodingType() == UE::Color::EEncoding::ST2084, ColorSpaceMtx,
+					SetShaderParametersLegacyPS(CommandList, ConvertShader, InputTexture, OutputDim, YUVMtx, Sample->GetEncodingType(), ColorSpaceMtx,
 												true);
 				}
 				break;
@@ -1107,7 +1107,8 @@ void FMediaTextureResource::GetColorSpaceConversionMatrixForSample(const TShared
 					SetGraphicsPipelineState(CommandList, GraphicsPSOInit, 0);
 					FShaderResourceViewRHIRef SRV = RHICreateShaderResourceView(InputTexture, 0, 1, (Sample->GetFormat() == EMediaTextureSampleFormat::Y416) ? PF_A16B16G16R16 : PF_A32B32G32R32F);
 
-					SetShaderParametersLegacyPS(CommandList, ConvertShader, SRV, YUVMtx, Sample->GetEncodingType() == UE::Color::EEncoding::sRGB, Sample->GetEncodingType() == UE::Color::EEncoding::ST2084, ColorSpaceMtx, InputTexture->GetFormat() == PF_A8R8G8B8);
+					SetShaderParametersLegacyPS(CommandList, ConvertShader, SRV, YUVMtx, Sample->GetEncodingType(), ColorSpaceMtx,
+												InputTexture->GetFormat() == PF_A8R8G8B8);
 				}
 				break;
 
@@ -1123,19 +1124,15 @@ void FMediaTextureResource::GetColorSpaceConversionMatrixForSample(const TShared
 					// We get here if we detect the need for an inverse EOTF application, CS conversion OR if we need to generate mips for RGBA samples
 
 					// Inverse EOTF?
-					bool bConvertFromSRGB;
-					bool bConvertFromST2084;
+					UE::Color::EEncoding Encoding;
 					if ((uint32)InputTexture->GetFlags() & (uint32)TexCreate_SRGB)
 					{
 						// We do not need one: sRGB will be handled in HW
-						bConvertFromSRGB = false;
-						bConvertFromST2084 = false;
+						Encoding = UE::Color::EEncoding::Linear;
 					}
 					else
 					{
-						// See what we need...
-						bConvertFromSRGB = (Sample->GetEncodingType() == UE::Color::EEncoding::sRGB);
-						bConvertFromST2084 = (Sample->GetEncodingType() == UE::Color::EEncoding::ST2084);
+						Encoding = Sample->GetEncodingType();
 					}
 
 					FMatrix44f ColorSpaceMtx;
@@ -1144,7 +1141,7 @@ void FMediaTextureResource::GetColorSpaceConversionMatrixForSample(const TShared
 					TShaderMapRef<FRGBConvertPS> ConvertShader(ShaderMap);
 					GraphicsPSOInit.BoundShaderState.PixelShaderRHI = ConvertShader.GetPixelShader();
 					SetGraphicsPipelineState(CommandList, GraphicsPSOInit, 0);
-					SetShaderParametersLegacyPS(CommandList, ConvertShader, InputTexture, OutputDim, bConvertFromSRGB, bConvertFromST2084, ColorSpaceMtx);
+					SetShaderParametersLegacyPS(CommandList, ConvertShader, InputTexture, OutputDim, Encoding, ColorSpaceMtx);
 				}
 				break;
 
@@ -1158,7 +1155,7 @@ void FMediaTextureResource::GetColorSpaceConversionMatrixForSample(const TShared
 					SetGraphicsPipelineState(CommandList, GraphicsPSOInit, 0);
 					FShaderResourceViewRHIRef SRV = RHICreateShaderResourceView(InputTexture, 0, 1, PF_R16G16B16A16_UINT);
 
-					SetShaderParametersLegacyPS(CommandList, ConvertShader, SRV, OutputDim, Sample->GetEncodingType() == UE::Color::EEncoding::sRGB, Sample->GetEncodingType() == UE::Color::EEncoding::ST2084, ColorSpaceMtx);
+					SetShaderParametersLegacyPS(CommandList, ConvertShader, SRV, OutputDim, Sample->GetEncodingType(), ColorSpaceMtx);
 				}
 				break;
 
@@ -1178,16 +1175,15 @@ void FMediaTextureResource::GetColorSpaceConversionMatrixForSample(const TShared
 					// Note: the shader used here will pass-through alpha data
 					
 					// Inverse EOTF?
-					bool bConvertFromSRGB;
+					UE::Color::EEncoding Encoding;
 					if ((uint32)InputTexture->GetFlags() & (uint32)TexCreate_SRGB)
 					{
 						// We do not need one: sRGB will be handled in HW
-						bConvertFromSRGB = false;
+						Encoding = UE::Color::EEncoding::Linear;
 					}
 					else
 					{
-						// See what we need...
-						bConvertFromSRGB = (Sample->GetEncodingType() == UE::Color::EEncoding::sRGB);
+						Encoding = Sample->GetEncodingType();
 					}
 
 					FMatrix44f ColorSpaceMtx;
@@ -1198,7 +1194,7 @@ void FMediaTextureResource::GetColorSpaceConversionMatrixForSample(const TShared
 					TShaderMapRef<FRGBConvertPS> ConvertShader(ShaderMap);
 					GraphicsPSOInit.BoundShaderState.PixelShaderRHI = ConvertShader.GetPixelShader();
 					SetGraphicsPipelineState(CommandList, GraphicsPSOInit, 0);
-					SetShaderParametersLegacyPS(CommandList, ConvertShader, InputTexture, OutputDim, bConvertFromSRGB, false, ColorSpaceMtx);
+					SetShaderParametersLegacyPS(CommandList, ConvertShader, InputTexture, OutputDim, Encoding, ColorSpaceMtx);
 				}
 				break;
 
@@ -1209,23 +1205,22 @@ void FMediaTextureResource::GetColorSpaceConversionMatrixForSample(const TShared
 					TShaderMapRef<FRGBConvertPS> ConvertShader(ShaderMap);
 					GraphicsPSOInit.BoundShaderState.PixelShaderRHI = ConvertShader.GetPixelShader();
 					SetGraphicsPipelineState(CommandList, GraphicsPSOInit, 0);
-					SetShaderParametersLegacyPS(CommandList, ConvertShader, InputTexture, OutputDim, false, false, SplatMtx);
+					SetShaderParametersLegacyPS(CommandList, ConvertShader, InputTexture, OutputDim, UE::Color::EEncoding::Linear, SplatMtx);
 				}
 				break;
 
 				case EMediaTextureSampleFormat::YCoCg_DXT5:
 				{
 					// Inverse EOTF?
-					bool bConvertFromSRGB;
+					UE::Color::EEncoding Encoding;
 					if ((uint32)InputTexture->GetFlags() & (uint32)TexCreate_SRGB)
 					{
 						// We do not need one: sRGB will be handled in HW
-						bConvertFromSRGB = false;
+						Encoding = UE::Color::EEncoding::Linear;
 					}
 					else
 					{
-						// See what we need...
-						bConvertFromSRGB = (Sample->GetEncodingType() == UE::Color::EEncoding::sRGB);
+						Encoding = Sample->GetEncodingType();
 					}
 
 					FMatrix44f ColorSpaceMtx;
@@ -1236,7 +1231,7 @@ void FMediaTextureResource::GetColorSpaceConversionMatrixForSample(const TShared
 					TShaderMapRef<FYCoCgConvertPS> ConvertShader(ShaderMap);
 					GraphicsPSOInit.BoundShaderState.PixelShaderRHI = ConvertShader.GetPixelShader();
 					SetGraphicsPipelineState(CommandList, GraphicsPSOInit, 0);
-					SetShaderParametersLegacyPS(CommandList, ConvertShader, InputTexture, OutputDim, bConvertFromSRGB, false, ColorSpaceMtx);
+					SetShaderParametersLegacyPS(CommandList, ConvertShader, InputTexture, OutputDim, Encoding, ColorSpaceMtx);
 				}
 				break;
 
diff --git a/Engine/Source/Runtime/RenderCore/Private/MediaShaders.cpp b/Engine/Source/Runtime/RenderCore/Private/MediaShaders.cpp
--- a/Engine/Source/Runtime/RenderCore/Private/MediaShaders.cpp
+++ b/Engine/Source/Runtime/RenderCore/Private/MediaShaders.cpp
@@ -5,6 +5,8 @@
 #include "RHIStaticStates.h"
 #include "ShaderParameterUtils.h"
 #include "PipelineStateCache.h"
+#include "ColorManagementDefines.h"
+#include "ColorSpace.h"
 
 FMediaVertexDeclaration::FMediaVertexDeclaration() = default;
 FMediaVertexDeclaration::~FMediaVertexDeclaration() = default;
@@ -193,7 +195,7 @@ BEGIN_GLOBAL_SHADER_PARAMETER_STRUCT(FNV12ConvertUB, )
 SHADER_PARAMETER(FMatrix44f, ColorTransform)
 SHADER_PARAMETER(FMatrix44f, CSTransform)
 SHADER_PARAMETER(uint32, OutputWidth)
-SHADER_PARAMETER(uint32, SrgbToLinear)
+SHADER_PARAMETER(uint32, EOTF)
 SHADER_PARAMETER(uint32, SwapChroma)
 SHADER_PARAMETER(FVector2f, UVScale)
 SHADER_PARAMETER_SRV(Texture2D, SRV_Y)
@@ -205,8 +207,26 @@ END_GLOBAL_SHADER_PARAMETER_STRUCT()
 IMPLEMENT_GLOBAL_SHADER_PARAMETER_STRUCT(FNV12ConvertUB, "NV12ConvertUB");
 IMPLEMENT_SHADER_TYPE(, FNV12ConvertPS, TEXT("/Engine/Private/MediaShaders.usf"), TEXT("NV12ConvertPS"), SF_Pixel);
 
-void FNV12ConvertPS::SetParameters(FRHIBatchedShaderParameters& BatchedParameters, const FIntPoint & TexDim, FShaderResourceViewRHIRef SRV_Y, FShaderResourceViewRHIRef SRV_UV, const FIntPoint& OutputDimensions, const FMatrix44f& ColorTransform, bool SrgbToLinear, const FMatrix44f& CSTransform, bool bSwapChroma)
+void FNV12ConvertPS::SetParameters(FRHIBatchedShaderParameters& BatchedParameters, const FIntPoint & TexDim, FShaderResourceViewRHIRef SRV_Y, FShaderResourceViewRHIRef SRV_UV, const FIntPoint& OutputDimensions, const FMatrix44f& ColorTransform, UE::Color::EEncoding Encoding, const FMatrix44f& CSTransform, bool bSwapChroma)
 {
+	// Ensure shader code assumptions about the layout of this enum are correct
+	static_assert(int(UE::Color::EEncoding::Linear) == 1, "Enum mismatch with use in shader code!");
+	static_assert(int(UE::Color::EEncoding::sRGB) == 2, "Enum mismatch with use in shader code!");
+	static_assert(int(UE::Color::EEncoding::ST2084) == 3, "Enum mismatch with use in shader code!");
+	static_assert(int(UE::Color::EEncoding::Gamma22) == 4, "Enum mismatch with use in shader code!");
+	static_assert(int(UE::Color::EEncoding::BT1886) == 5, "Enum mismatch with use in shader code!");
+	static_assert(int(UE::Color::EEncoding::Gamma26) == 6, "Enum mismatch with use in shader code!");
+	static_assert(int(UE::Color::EEncoding::Cineon) == 7, "Enum mismatch with use in shader code!");
+	static_assert(int(UE::Color::EEncoding::REDLog) == 8, "Enum mismatch with use in shader code!");
+	static_assert(int(UE::Color::EEncoding::REDLog3G10) == 9, "Enum mismatch with use in shader code!");
+	static_assert(int(UE::Color::EEncoding::SLog1) == 10, "Enum mismatch with use in shader code!");
+	static_assert(int(UE::Color::EEncoding::SLog2) == 11, "Enum mismatch with use in shader code!");
+	static_assert(int(UE::Color::EEncoding::SLog3) == 12, "Enum mismatch with use in shader code!");
+	static_assert(int(UE::Color::EEncoding::AlexaV3LogC) == 13, "Enum mismatch with use in shader code!");
+	static_assert(int(UE::Color::EEncoding::CanonLog) == 14, "Enum mismatch with use in shader code!");
+	static_assert(int(UE::Color::EEncoding::ProTune) == 15, "Enum mismatch with use in shader code!");
+	static_assert(int(UE::Color::EEncoding::VLog) == 16, "Enum mismatch with use in shader code!");
+
 	FNV12ConvertUB UB;
 	{
 		UB.ColorTransform = ColorTransform;
@@ -214,7 +234,7 @@ void FNV12ConvertPS::SetParameters(FRHIBatchedShaderParameters& BatchedParameter
 		UB.OutputWidth = OutputDimensions.X;
 		UB.SamplerB = TStaticSamplerState<SF_Bilinear>::GetRHI();
 		UB.SamplerP = TStaticSamplerState<SF_Point>::GetRHI();
-		UB.SrgbToLinear = SrgbToLinear;
+		UB.EOTF = int(Encoding);
 		UB.SRV_Y = SRV_Y;
 		UB.SRV_UV = SRV_UV;
 		UB.UVScale = FVector2f((float)OutputDimensions.X / (float)TexDim.X, (float)OutputDimensions.Y / (float)TexDim.Y);
@@ -233,7 +253,7 @@ BEGIN_GLOBAL_SHADER_PARAMETER_STRUCT(FNV12ConvertAsBytesUB, )
 SHADER_PARAMETER(FMatrix44f, ColorTransform)
 SHADER_PARAMETER(FMatrix44f, CSTransform)
 SHADER_PARAMETER(uint32, OutputWidth)
-SHADER_PARAMETER(uint32, SrgbToLinear)
+SHADER_PARAMETER(uint32, EOTF)
 SHADER_PARAMETER(uint32, SwapChroma)
 SHADER_PARAMETER(FVector2f, UVScale)
 SHADER_PARAMETER_TEXTURE(Texture2D, Texture)
@@ -244,7 +264,7 @@ END_GLOBAL_SHADER_PARAMETER_STRUCT()
 IMPLEMENT_GLOBAL_SHADER_PARAMETER_STRUCT(FNV12ConvertAsBytesUB, "NV12ConvertAsBytesUB");
 IMPLEMENT_SHADER_TYPE(, FNV12ConvertAsBytesPS, TEXT("/Engine/Private/MediaShaders.usf"), TEXT("NV12ConvertAsBytesPS"), SF_Pixel);
 
-void FNV12ConvertAsBytesPS::SetParameters(FRHIBatchedShaderParameters& BatchedParameters, TRefCountPtr<FRHITexture2D> NV12Texture, const FIntPoint& OutputDimensions, const FMatrix44f& ColorTransform, bool SrgbToLinear, const FMatrix44f& CSTransform, bool bSwapChroma)
+void FNV12ConvertAsBytesPS::SetParameters(FRHIBatchedShaderParameters& BatchedParameters, TRefCountPtr<FRHITexture2D> NV12Texture, const FIntPoint& OutputDimensions, const FMatrix44f& ColorTransform, UE::Color::EEncoding Encoding, const FMatrix44f& CSTransform, bool bSwapChroma)
 {
 	FNV12ConvertAsBytesUB UB;
 	{
@@ -253,7 +273,7 @@ void FNV12ConvertAsBytesPS::SetParameters(FRHIBatchedShaderParameters& BatchedPa
 		UB.OutputWidth = OutputDimensions.X;
 		UB.SamplerB = TStaticSamplerState<SF_Bilinear>::GetRHI();
 		UB.SamplerP = TStaticSamplerState<SF_Point>::GetRHI();
-		UB.SrgbToLinear = SrgbToLinear;
+		UB.EOTF = int(Encoding);
 		UB.Texture = NV12Texture;
 		UB.UVScale = FVector2f((float)OutputDimensions.X / (float)NV12Texture->GetSizeX(), (float)OutputDimensions.Y / (float)NV12Texture->GetSizeY());
 		UB.SwapChroma = bSwapChroma;
@@ -316,7 +336,7 @@ END_GLOBAL_SHADER_PARAMETER_STRUCT()
 IMPLEMENT_GLOBAL_SHADER_PARAMETER_STRUCT(FP010ConvertUB, "P010ConvertUB");
 IMPLEMENT_SHADER_TYPE(, FP010ConvertPS, TEXT("/Engine/Private/MediaShaders.usf"), TEXT("P010ConvertPS"), SF_Pixel);
 
-void FP010ConvertPS::SetParameters(FRHIBatchedShaderParameters& BatchedParameters, const FIntPoint& TexDim, FShaderResourceViewRHIRef SRV_Y, FShaderResourceViewRHIRef SRV_UV, const FIntPoint& OutputDimensions, const FMatrix44f& ColorTransform, const FMatrix44f& CSTransform, bool bSrgbToLinear, bool bIsST2084)
+void FP010ConvertPS::SetParameters(FRHIBatchedShaderParameters& BatchedParameters, const FIntPoint& TexDim, FShaderResourceViewRHIRef SRV_Y, FShaderResourceViewRHIRef SRV_UV, const FIntPoint& OutputDimensions, const FMatrix44f& ColorTransform, const FMatrix44f& CSTransform, UE::Color::EEncoding Encoding)
 {
 	FP010ConvertUB UB;
 	{
@@ -324,7 +344,7 @@ void FP010ConvertPS::SetParameters(FRHIBatchedShaderParameters& BatchedParameter
 		UB.CSTransform = CSTransform;
 		UB.SamplerB = TStaticSamplerState<SF_Bilinear>::GetRHI();
 		UB.SamplerP = TStaticSamplerState<SF_Point>::GetRHI();
-		UB.EOTF = bSrgbToLinear ? 1 : (bIsST2084 ? 2 : 0);
+		UB.EOTF = int(Encoding);
 		UB.SRV_Y = SRV_Y;
 		UB.SRV_UV = SRV_UV;
 		UB.UVScale = FVector2f((float)OutputDimensions.X / (float)TexDim.X, (float)OutputDimensions.Y / (float)TexDim.Y);
@@ -352,7 +372,7 @@ END_GLOBAL_SHADER_PARAMETER_STRUCT()
 IMPLEMENT_GLOBAL_SHADER_PARAMETER_STRUCT(FP010ConvertAsUINT16sUB, "P010ConvertAsUINT16sUB");
 IMPLEMENT_SHADER_TYPE(, FP010ConvertAsUINT16sPS, TEXT("/Engine/Private/MediaShaders.usf"), TEXT("P010ConvertAsUINT16sPS"), SF_Pixel);
 
-void FP010ConvertAsUINT16sPS::SetParameters(FRHIBatchedShaderParameters& BatchedParameters, const FIntPoint& TexDim, TRefCountPtr<FRHITexture2D> NV12Texture, const FIntPoint& OutputDimensions, const FMatrix44f& ColorTransform, const FMatrix44f& CSTransform, bool bSrgbToLinear, bool bIsST2084)
+void FP010ConvertAsUINT16sPS::SetParameters(FRHIBatchedShaderParameters& BatchedParameters, const FIntPoint& TexDim, TRefCountPtr<FRHITexture2D> NV12Texture, const FIntPoint& OutputDimensions, const FMatrix44f& ColorTransform, const FMatrix44f& CSTransform, UE::Color::EEncoding Encoding)
 {
 	FP010ConvertAsUINT16sUB UB;
 	{
@@ -361,7 +381,7 @@ void FP010ConvertAsUINT16sPS::SetParameters(FRHIBatchedShaderParameters& Batched
 		UB.OutputWidth = OutputDimensions.X;
 		UB.SamplerB = TStaticSamplerState<SF_Bilinear>::GetRHI();
 		UB.SamplerP = TStaticSamplerState<SF_Point>::GetRHI();
-		UB.EOTF = bSrgbToLinear ? 1 : (bIsST2084 ? 2 : 0);
+		UB.EOTF = int(Encoding);
 		UB.Texture = NV12Texture;
 		UB.UVScale = FVector2f((float)OutputDimensions.X / (float)TexDim.X, (float)OutputDimensions.Y / (float)TexDim.Y);
 	}
@@ -391,14 +411,14 @@ END_GLOBAL_SHADER_PARAMETER_STRUCT()
 IMPLEMENT_GLOBAL_SHADER_PARAMETER_STRUCT(FP010_2101010ConvertUB, "P010_2101010ConvertUB");
 IMPLEMENT_SHADER_TYPE(, FP010_2101010ConvertPS, TEXT("/Engine/Private/MediaShaders.usf"), TEXT("P010_2101010ConvertPS"), SF_Pixel);
 
-void FP010_2101010ConvertPS::SetParameters(FRHIBatchedShaderParameters& BatchedParameters, const FIntPoint& TexDim, FShaderResourceViewRHIRef SRV_Y, FShaderResourceViewRHIRef SRV_U, FShaderResourceViewRHIRef SRV_V, const FIntPoint& OutputDimensions, const FMatrix44f& ColorTransform, const FMatrix44f& CSTransform, bool bSrgbToLinear, bool bIsST2084)
+void FP010_2101010ConvertPS::SetParameters(FRHIBatchedShaderParameters& BatchedParameters, const FIntPoint& TexDim, FShaderResourceViewRHIRef SRV_Y, FShaderResourceViewRHIRef SRV_U, FShaderResourceViewRHIRef SRV_V, const FIntPoint& OutputDimensions, const FMatrix44f& ColorTransform, const FMatrix44f& CSTransform, UE::Color::EEncoding Encoding)
 {
 	FP010_2101010ConvertUB UB;
 	{
 		UB.ColorTransform = ColorTransform;
 		UB.CSTransform = CSTransform;
 		UB.SamplerP = TStaticSamplerState<SF_Point>::GetRHI();
-		UB.EOTF = bSrgbToLinear ? 1 : (bIsST2084 ? 2 : 0);
+		UB.EOTF = int(Encoding);
 		UB.SRV_Y = SRV_Y;
 		UB.SRV_U = SRV_U;
 		UB.SRV_V = SRV_V;
@@ -432,14 +452,14 @@ IMPLEMENT_GLOBAL_SHADER_PARAMETER_STRUCT(FRGBConvertUB, "RGBConvertUB");
 IMPLEMENT_SHADER_TYPE(, FRGBConvertPS, TEXT("/Engine/Private/MediaShaders.usf"), TEXT("RGBConvertPS"), SF_Pixel);
 
 
-void FRGBConvertPS::SetParameters(FRHIBatchedShaderParameters& BatchedParameters, TRefCountPtr<FRHITexture2D> RGBTexture, const FIntPoint& OutputDimensions, bool bSrgbToLinear, bool bIsST2084, const FMatrix44f& CSTransform)
+void FRGBConvertPS::SetParameters(FRHIBatchedShaderParameters& BatchedParameters, TRefCountPtr<FRHITexture2D> RGBTexture, const FIntPoint& OutputDimensions, UE::Color::EEncoding Encoding, const FMatrix44f& CSTransform)
 {
 	FRGBConvertUB UB;
 	{
 		UB.Sampler = TStaticSamplerState<SF_Bilinear>::GetRHI();
 		UB.Texture = RGBTexture;
 		UB.UVScale = FVector2f((float)OutputDimensions.X / (float)RGBTexture->GetSizeX(), (float)OutputDimensions.Y / (float)RGBTexture->GetSizeY());
-		UB.EOTF = bSrgbToLinear ? 1 : (bIsST2084 ? 2 : 0);
+		UB.EOTF = int(Encoding);
 		UB.CSTransform = CSTransform;
 	}
 
@@ -463,14 +483,14 @@ IMPLEMENT_GLOBAL_SHADER_PARAMETER_STRUCT(FYCoCgConvertUB, "YCoCgConvertUB");
 IMPLEMENT_SHADER_TYPE(, FYCoCgConvertPS, TEXT("/Engine/Private/MediaShaders.usf"), TEXT("YCoCgConvertPS"), SF_Pixel);
 
 
-void FYCoCgConvertPS::SetParameters(FRHIBatchedShaderParameters& BatchedParameters, TRefCountPtr<FRHITexture2D> RGBTexture, const FIntPoint& OutputDimensions, bool bSrgbToLinear, bool bIsST2084, const FMatrix44f& CSTransform)
+void FYCoCgConvertPS::SetParameters(FRHIBatchedShaderParameters& BatchedParameters, TRefCountPtr<FRHITexture2D> RGBTexture, const FIntPoint& OutputDimensions, UE::Color::EEncoding Encoding, const FMatrix44f& CSTransform)
 {
 	FYCoCgConvertUB UB;
 	{
 		UB.Sampler = TStaticSamplerState<SF_Bilinear>::GetRHI();
 		UB.Texture = RGBTexture;
 		UB.UVScale = FVector2f((float)OutputDimensions.X / (float)RGBTexture->GetSizeX(), (float)OutputDimensions.Y / (float)RGBTexture->GetSizeY());
-		UB.EOTF = bSrgbToLinear ? 1 : (bIsST2084 ? 2 : 0);
+		UB.EOTF = int(Encoding);
 		UB.CSTransform = CSTransform;
 	}
 
@@ -607,13 +627,13 @@ IMPLEMENT_GLOBAL_SHADER_PARAMETER_STRUCT(FYUVv216ConvertUB, "YUVv216ConvertUB");
 IMPLEMENT_SHADER_TYPE(, FYUVv216ConvertPS, TEXT("/Engine/Private/MediaShaders.usf"), TEXT("YUVv216ConvertPS"), SF_Pixel);
 
 
-void FYUVv216ConvertPS::SetParameters(FRHIBatchedShaderParameters& BatchedParameters, TRefCountPtr<FRHITexture2D> YUVTexture, const FIntPoint& OutputDimensions, const FMatrix44f& ColorTransform, bool bSrgbToLinear, bool bIsST2084, const FMatrix44f& CSTransform, bool bIsCbY0CrY1, bool bIsARGBFmt, bool bSwapChroma)
+void FYUVv216ConvertPS::SetParameters(FRHIBatchedShaderParameters& BatchedParameters, TRefCountPtr<FRHITexture2D> YUVTexture, const FIntPoint& OutputDimensions, const FMatrix44f& ColorTransform, UE::Color::EEncoding Encoding, const FMatrix44f& CSTransform, bool bIsCbY0CrY1, bool bIsARGBFmt, bool bSwapChroma)
 {
 	FYUVv216ConvertUB UB;
 	{
 		UB.ColorTransform = ColorTransform;
 		UB.CSTransform = CSTransform;
-		UB.EOTF = bSrgbToLinear ? 1 : (bIsST2084 ? 2 : 0);
+		UB.EOTF = int(Encoding);
 		UB.IsCbY0CrY1 = bIsCbY0CrY1;
 		UB.IsARGBFmt = bIsARGBFmt;
 		UB.SwapChroma = bSwapChroma;
@@ -644,13 +664,13 @@ IMPLEMENT_GLOBAL_SHADER_PARAMETER_STRUCT(FYUVv210ConvertUB, "YUVv210ConvertUB");
 IMPLEMENT_SHADER_TYPE(, FYUVv210ConvertPS, TEXT("/Engine/Private/MediaShaders.usf"), TEXT("YUVv210ConvertPS"), SF_Pixel);
 
 
-void FYUVv210ConvertPS::SetParameters(FRHIBatchedShaderParameters& BatchedParameters, TRefCountPtr<FRHITexture2D> YUVTexture, const FIntPoint& OutputDimensions, const FMatrix44f& ColorTransform, bool bSrgbToLinear, bool bIsST2084, const FMatrix44f& CSTransform, bool bIsCbY0CrY1)
+void FYUVv210ConvertPS::SetParameters(FRHIBatchedShaderParameters& BatchedParameters, TRefCountPtr<FRHITexture2D> YUVTexture, const FIntPoint& OutputDimensions, const FMatrix44f& ColorTransform, UE::Color::EEncoding Encoding, const FMatrix44f& CSTransform, bool bIsCbY0CrY1)
 {
 	FYUVv210ConvertUB UB;
 	{
 		UB.ColorTransform = ColorTransform;
 		UB.CSTransform = CSTransform;
-		UB.EOTF = bSrgbToLinear ? 1 : (bIsST2084 ? 2 : 0);
+		UB.EOTF = int(Encoding);
 		UB.IsCbY0CrY1 = bIsCbY0CrY1;
 		UB.OutputDimX = (float)OutputDimensions.X;
 		UB.OutputDimY = (float)OutputDimensions.Y;
@@ -678,7 +698,7 @@ IMPLEMENT_GLOBAL_SHADER_PARAMETER_STRUCT(FYUVY416ConvertUB, "YUVY416ConvertUB");
 IMPLEMENT_SHADER_TYPE(, FYUVY416ConvertPS, TEXT("/Engine/Private/MediaShaders.usf"), TEXT("YUVY416ConvertPS"), SF_Pixel);
 
 
-void FYUVY416ConvertPS::SetParameters(FRHIBatchedShaderParameters& BatchedParameters, FShaderResourceViewRHIRef SRV_Y, const FMatrix44f& ColorTransform, bool bSrgbToLinear, bool bIsST2084, const FMatrix44f& CSTransform, bool bIsARGBFmt)
+void FYUVY416ConvertPS::SetParameters(FRHIBatchedShaderParameters& BatchedParameters, FShaderResourceViewRHIRef SRV_Y, const FMatrix44f& ColorTransform, UE::Color::EEncoding Encoding, const FMatrix44f& CSTransform, bool bIsARGBFmt)
 {
 	FYUVY416ConvertUB UB;
 	{
@@ -686,7 +706,7 @@ void FYUVY416ConvertPS::SetParameters(FRHIBatchedShaderParameters& BatchedParame
 		UB.BilinearClampedSamplerUV = TStaticSamplerState<SF_Bilinear, AM_Clamp, AM_Clamp, AM_Clamp>::GetRHI();
 		UB.ColorTransform = ColorTransform;
 		UB.CSTransform = CSTransform;
-		UB.EOTF = bSrgbToLinear ? 1 : (bIsST2084 ? 2 : 0);
+		UB.EOTF = int(Encoding);
 		UB.bIsARGBFmt = bIsARGBFmt;
 	}
 
@@ -709,7 +729,7 @@ IMPLEMENT_GLOBAL_SHADER_PARAMETER_STRUCT(FARGB16BigConvertUB, "ARGB16BigConvertU
 IMPLEMENT_SHADER_TYPE(, FARGB16BigConvertPS, TEXT("/Engine/Private/MediaShaders.usf"), TEXT("ARGB16BigConvertPS"), SF_Pixel);
 
 
-void FARGB16BigConvertPS::SetParameters(FRHIBatchedShaderParameters& BatchedParameters, FShaderResourceViewRHIRef SRV, const FIntPoint& OutputDimensions, bool bSrgbToLinear, bool bIsST2084, const FMatrix44f& CSTransform)
+void FARGB16BigConvertPS::SetParameters(FRHIBatchedShaderParameters& BatchedParameters, FShaderResourceViewRHIRef SRV, const FIntPoint& OutputDimensions, UE::Color::EEncoding Encoding, const FMatrix44f& CSTransform)
 {
 	FARGB16BigConvertUB UB;
 	{
@@ -717,7 +737,7 @@ void FARGB16BigConvertPS::SetParameters(FRHIBatchedShaderParameters& BatchedPara
 		UB.OutputWidth = OutputDimensions.X;
 		UB.OutputHeight = OutputDimensions.Y;
 		UB.CSTransform = CSTransform;
-		UB.EOTF = bSrgbToLinear ? 1 : (bIsST2084 ? 2 : 0);
+		UB.EOTF = int(Encoding);
 	}
 
 	TUniformBufferRef<FARGB16BigConvertUB> Data = TUniformBufferRef<FARGB16BigConvertUB>::CreateUniformBufferImmediate(UB, UniformBuffer_SingleFrame);
diff --git a/Engine/Source/Runtime/RenderCore/Public/MediaShaders.h b/Engine/Source/Runtime/RenderCore/Public/MediaShaders.h
--- a/Engine/Source/Runtime/RenderCore/Public/MediaShaders.h
+++ b/Engine/Source/Runtime/RenderCore/Public/MediaShaders.h
@@ -23,6 +23,7 @@
 #include "Shader.h"
 #include "ShaderParameterMacros.h"
 #include "ShaderPermutation.h"
+#include "ShaderPermutation.h"
 #include "Templates/RefCounting.h"
 #include "Templates/UnrealTemplate.h"
 #endif
@@ -31,6 +32,11 @@ class FPointerTableBase;
 class FRDGBuilder;
 class FRDGTexture;
 
+namespace UE {
+	namespace Color {
+		enum class EEncoding : uint8;
+	}
+}
 
 /** MAX number of conversion operations. Reflects MAX in EMediaCaptureConversionOperation */
 #define NUM_MEDIA_SHADERS_CONVERSION_OP 5
@@ -228,7 +234,7 @@ public:
 		: FGlobalShader(Initializer)
 	{ }
 
-	RENDERCORE_API void SetParameters(FRHIBatchedShaderParameters& BatchedParameters, const FIntPoint & TexDim, FShaderResourceViewRHIRef SRV_Y, FShaderResourceViewRHIRef SRV_UV, const FIntPoint& OutputDimensions, const FMatrix44f& ColorTransform, bool SrgbToLinear, const FMatrix44f& CSTransform, bool bSwapChroma);
+	RENDERCORE_API void SetParameters(FRHIBatchedShaderParameters& BatchedParameters, const FIntPoint & TexDim, FShaderResourceViewRHIRef SRV_Y, FShaderResourceViewRHIRef SRV_UV, const FIntPoint& OutputDimensions, const FMatrix44f& ColorTransform, UE::Color::EEncoding Encoding, const FMatrix44f& CSTransform, bool bSwapChroma);
 };
 
 
@@ -252,7 +258,7 @@ public:
 		: FGlobalShader(Initializer)
 	{ }
 
-	RENDERCORE_API void SetParameters(FRHIBatchedShaderParameters& BatchedParameters, TRefCountPtr<FRHITexture2D> NV12Texture, const FIntPoint& OutputDimensions, const FMatrix44f& ColorTransform, bool SrgbToLinear, const FMatrix44f& CSTransform, bool bSwapChroma);
+	RENDERCORE_API void SetParameters(FRHIBatchedShaderParameters& BatchedParameters, TRefCountPtr<FRHITexture2D> NV12Texture, const FIntPoint& OutputDimensions, const FMatrix44f& ColorTransform, UE::Color::EEncoding Encoding, const FMatrix44f& CSTransform, bool bSwapChroma);
 };
 
 
@@ -294,7 +300,7 @@ public:
 		: FGlobalShader(Initializer)
 	{ }
 
-	RENDERCORE_API void SetParameters(FRHIBatchedShaderParameters& BatchedParameters, const FIntPoint& TexDim, FShaderResourceViewRHIRef SRV_Y, FShaderResourceViewRHIRef SRV_UV, const FIntPoint& OutputDimensions, const FMatrix44f& ColorTransform, const FMatrix44f& CSTransform, bool SrgbToLinear, bool bIsST2084);
+	RENDERCORE_API void SetParameters(FRHIBatchedShaderParameters& BatchedParameters, const FIntPoint& TexDim, FShaderResourceViewRHIRef SRV_Y, FShaderResourceViewRHIRef SRV_UV, const FIntPoint& OutputDimensions, const FMatrix44f& ColorTransform, const FMatrix44f& CSTransform, UE::Color::EEncoding Encoding);
 };
 
 /**
@@ -313,7 +319,7 @@ public:
 		: FGlobalShader(Initializer)
 	{ }
 
-	RENDERCORE_API void SetParameters(FRHIBatchedShaderParameters& BatchedParameters, const FIntPoint& TexDim, TRefCountPtr<FRHITexture2D> NV12Texture, const FIntPoint& OutputDimensions, const FMatrix44f& ColorTransform, const FMatrix44f& CSTransform, bool SrgbToLinear, bool bIsST2084);
+	RENDERCORE_API void SetParameters(FRHIBatchedShaderParameters& BatchedParameters, const FIntPoint& TexDim, TRefCountPtr<FRHITexture2D> NV12Texture, const FIntPoint& OutputDimensions, const FMatrix44f& ColorTransform, const FMatrix44f& CSTransform, UE::Color::EEncoding Encoding);
 };
 
 /**
@@ -332,7 +338,7 @@ public:
 		: FGlobalShader(Initializer)
 	{ }
 
-	RENDERCORE_API void SetParameters(FRHIBatchedShaderParameters& BatchedParameters, const FIntPoint& TexDim, FShaderResourceViewRHIRef SRV_Y, FShaderResourceViewRHIRef SRV_U, FShaderResourceViewRHIRef SRV_V, const FIntPoint& OutputDimensions, const FMatrix44f& ColorTransform, const FMatrix44f& CSTransform, bool SrgbToLinear, bool bIsST2084);
+	RENDERCORE_API void SetParameters(FRHIBatchedShaderParameters& BatchedParameters, const FIntPoint& TexDim, FShaderResourceViewRHIRef SRV_Y, FShaderResourceViewRHIRef SRV_U, FShaderResourceViewRHIRef SRV_V, const FIntPoint& OutputDimensions, const FMatrix44f& ColorTransform, const FMatrix44f& CSTransform, UE::Color::EEncoding Encoding);
 };
 
 /**
@@ -354,7 +360,7 @@ public:
 	{ }
 
 
-	RENDERCORE_API void SetParameters(FRHIBatchedShaderParameters& BatchedParameters, TRefCountPtr<FRHITexture2D> RGBTexture, const FIntPoint& OutputDimensions, bool SrgbToLinear, bool bIsST2084, const FMatrix44f& CSTransform);
+	RENDERCORE_API void SetParameters(FRHIBatchedShaderParameters& BatchedParameters, TRefCountPtr<FRHITexture2D> RGBTexture, const FIntPoint& OutputDimensions, UE::Color::EEncoding Encoding, const FMatrix44f& CSTransform);
 };
 
 /**
@@ -376,7 +382,7 @@ public:
 	{ }
 
 
-	RENDERCORE_API void SetParameters(FRHIBatchedShaderParameters& BatchedParameters, TRefCountPtr<FRHITexture2D> YCoCgTexture, const FIntPoint& OutputDimensions, bool SrgbToLinear, bool bIsST2084, const FMatrix44f& CSTransform);
+	RENDERCORE_API void SetParameters(FRHIBatchedShaderParameters& BatchedParameters, TRefCountPtr<FRHITexture2D> YCoCgTexture, const FIntPoint& OutputDimensions, UE::Color::EEncoding Encoding, const FMatrix44f& CSTransform);
 };
 
 
@@ -481,7 +487,7 @@ public:
 		: FGlobalShader(Initializer)
 	{ }
 
-	RENDERCORE_API void SetParameters(FRHIBatchedShaderParameters& BatchedParameters, TRefCountPtr<FRHITexture2D> YUVTexture, const FIntPoint& OutputDimensions, const FMatrix44f& ColorTransform, bool SrgbToLinear, bool bIsST2084, const FMatrix44f& CSTransform, bool bIsCbY0CrY1, bool bIsARGBFmt, bool bSwapChroma);
+	RENDERCORE_API void SetParameters(FRHIBatchedShaderParameters& BatchedParameters, TRefCountPtr<FRHITexture2D> YUVTexture, const FIntPoint& OutputDimensions, const FMatrix44f& ColorTransform, UE::Color::EEncoding Encoding, const FMatrix44f& CSTransform, bool bIsCbY0CrY1, bool bIsARGBFmt, bool bSwapChroma);
 };
 
 
@@ -503,7 +509,7 @@ public:
 		: FGlobalShader(Initializer)
 	{ }
 
-	RENDERCORE_API void SetParameters(FRHIBatchedShaderParameters& BatchedParameters, TRefCountPtr<FRHITexture2D> YUVTexture, const FIntPoint& OutputDimensions, const FMatrix44f& ColorTransform, bool SrgbToLinear, bool bIsST2084, const FMatrix44f& CSTransform, bool bIsCbY0CrY1);
+	RENDERCORE_API void SetParameters(FRHIBatchedShaderParameters& BatchedParameters, TRefCountPtr<FRHITexture2D> YUVTexture, const FIntPoint& OutputDimensions, const FMatrix44f& ColorTransform, UE::Color::EEncoding Encoding, const FMatrix44f& CSTransform, bool bIsCbY0CrY1);
 };
 
 
@@ -524,7 +530,7 @@ public:
 		: FGlobalShader(Initializer)
 	{ }
 
-	RENDERCORE_API void SetParameters(FRHIBatchedShaderParameters& BatchedParameters, FShaderResourceViewRHIRef SRV_Y, const FMatrix44f& ColorTransform, bool bSrgbToLinear, bool bIsST2084, const FMatrix44f& CSTransform, bool bIsARGBFmt);
+	RENDERCORE_API void SetParameters(FRHIBatchedShaderParameters& BatchedParameters, FShaderResourceViewRHIRef SRV_Y, const FMatrix44f& ColorTransform, UE::Color::EEncoding Encoding, const FMatrix44f& CSTransform, bool bIsARGBFmt);
 };
 
 /**
@@ -567,7 +573,7 @@ public:
 		: FGlobalShader(Initializer)
 	{ }
 
-	RENDERCORE_API void SetParameters(FRHIBatchedShaderParameters& BatchedParameters, FShaderResourceViewRHIRef SRV, const FIntPoint& OutputDimensions, bool bSrgbToLinear, bool bIsST2084, const FMatrix44f& CSTransform);
+	RENDERCORE_API void SetParameters(FRHIBatchedShaderParameters& BatchedParameters, FShaderResourceViewRHIRef SRV, const FIntPoint& OutputDimensions, UE::Color::EEncoding Encoding, const FMatrix44f& CSTransform);
 };
 
 
