Minor dependency cleanup in RenderGraphBuilder.

diff --git a/Engine/Source/Runtime/RenderCore/Private/RenderGraphBuilder.cpp b/Engine/Source/Runtime/RenderCore/Private/RenderGraphBuilder.cpp
--- a/Engine/Source/Runtime/RenderCore/Private/RenderGraphBuilder.cpp
+++ b/Engine/Source/Runtime/RenderCore/Private/RenderGraphBuilder.cpp
@@ -45,6 +45,18 @@ inline void GatherPassUAVsForOverlapValidation(const FRDGPass* Pass, TArray<FRHI
 
 #endif
 
+struct FParallelPassSet : public FRHICommandListImmediate::FQueuedCommandList
+{
+	FParallelPassSet() = default;
+
+	TArray<FRDGPass*, FRDGArrayAllocator> Passes;
+	IF_RHI_WANT_BREADCRUMB_EVENTS(FRDGBreadcrumbState* BreadcrumbStateBegin{});
+	IF_RHI_WANT_BREADCRUMB_EVENTS(FRDGBreadcrumbState* BreadcrumbStateEnd{});
+	int8 bInitialized = 0;
+	bool bDispatchAfterExecute = false;
+	bool bParallelTranslate = false;
+};
+
 inline void BeginUAVOverlap(const FRDGPass* Pass, FRHIComputeCommandList& RHICmdList)
 {
 #if ENABLE_RHI_VALIDATION
diff --git a/Engine/Source/Runtime/RenderCore/Public/RenderGraphBuilder.h b/Engine/Source/Runtime/RenderCore/Public/RenderGraphBuilder.h
--- a/Engine/Source/Runtime/RenderCore/Public/RenderGraphBuilder.h
+++ b/Engine/Source/Runtime/RenderCore/Public/RenderGraphBuilder.h
@@ -2,7 +2,6 @@
 
 #pragma once
 
-#include "Async/TaskGraphInterfaces.h"
 #include "Containers/Array.h"
 #include "Containers/ArrayView.h"
 #include "Containers/Map.h"
@@ -15,9 +14,7 @@
 #include "MultiGPU.h"
 #include "PixelFormat.h"
 #include "ProfilingDebugging/CsvProfiler.h"
-#include "RHI.h"
 #include "RHIBreadcrumbs.h"
-#include "RHICommandList.h"
 #include "RHIDefinitions.h"
 #include "RenderGraphAllocator.h"
 #include "RenderGraphBlackboard.h"
@@ -34,6 +31,13 @@
 #include "Templates/UnrealTemplate.h"
 #include "Tasks/Pipe.h"
 
+enum class ERenderTargetTexture : uint8;
+struct FParallelPassSet;
+struct FRHIRenderPassInfo;
+struct FRHITrackedAccessInfo;
+struct FRHITransientAliasingInfo;
+struct TStatId;
+
 /** Use the render graph builder to build up a graph of passes and then call Execute() to process them. Resource barriers
  *  and lifetimes are derived from _RDG_ parameters in the pass parameter struct provided to each AddPass call. The resulting
  *  graph is compiled, culled, and executed in Execute(). The builder should be created on the stack and executed prior to
@@ -644,18 +648,6 @@ private:
 
 	TArray<FUploadedBuffer, FRDGArrayAllocator> UploadedBuffers;
 
-	struct FParallelPassSet : public FRHICommandListImmediate::FQueuedCommandList
-	{
-		FParallelPassSet() = default;
-
-		TArray<FRDGPass*, FRDGArrayAllocator> Passes;
-		IF_RHI_WANT_BREADCRUMB_EVENTS(FRDGBreadcrumbState* BreadcrumbStateBegin{});
-		IF_RHI_WANT_BREADCRUMB_EVENTS(FRDGBreadcrumbState* BreadcrumbStateEnd{});
-		int8 bInitialized = 0;
-		bool bDispatchAfterExecute = false;
-		bool bParallelTranslate = false;
-	};
-
 	TArray<FParallelPassSet, FRDGArrayAllocator> ParallelPassSets;
 
 	/** Array of all active parallel execute tasks. */
@@ -879,3 +871,9 @@ class FRHITransientTexture;
 class FShaderParametersMetadata;
 class IRHICommandContext;
 class IRHITransientResourceAllocator;
+
+#if UE_ENABLE_INCLUDE_ORDER_DEPRECATED_IN_5_3
+#include "Async/TaskGraphInterfaces.h"
+#include "RHI.h"
+#include "RHICommandList.h"
+#endif
