Bindless: consider BindlessResourceParameters array when clearing unused graph resources for an array of shader bindings.

#rnx

diff --git a/Engine/Source/Runtime/RenderCore/Private/RenderGraphUtils.cpp b/Engine/Source/Runtime/RenderCore/Private/RenderGraphUtils.cpp
--- a/Engine/Source/Runtime/RenderCore/Private/RenderGraphUtils.cpp
+++ b/Engine/Source/Runtime/RenderCore/Private/RenderGraphUtils.cpp
@@ -100,8 +100,10 @@ void ClearUnusedGraphResourcesImpl(
 	const auto& GraphResources = ParametersMetadata->GetLayout().GraphResources;
 
 	TArray<int32, TInlineAllocator<SF_NumFrequencies>> ShaderResourceIds;
+	TArray<int32, TInlineAllocator<SF_NumFrequencies>> BindlessResourceIds;
 	TArray<int32, TInlineAllocator<SF_NumFrequencies>> GraphUniformBufferIds;
 	ShaderResourceIds.SetNumZeroed(ShaderBindingsList.Num());
+	BindlessResourceIds.SetNumZeroed(ShaderBindingsList.Num());
 	GraphUniformBufferIds.SetNumZeroed(ShaderBindingsList.Num());
 
 	auto Base = reinterpret_cast<uint8*>(InoutParameters);
@@ -120,12 +122,24 @@ void ClearUnusedGraphResourcesImpl(
 		{
 			for (int32 Index = 0; Index < ShaderBindingsList.Num(); ++Index)
 			{
-				const auto& ResourceParameters = ShaderBindingsList[Index]->ResourceParameters;
-				int32& ShaderResourceId = ShaderResourceIds[Index];
-				for (; ShaderResourceId < ResourceParameters.Num() && ResourceParameters[ShaderResourceId].ByteOffset < ByteOffset; ++ShaderResourceId)
 				{
+					const auto& ResourceParameters = ShaderBindingsList[Index]->ResourceParameters;
+					int32& ShaderResourceId = ShaderResourceIds[Index];
+					for (; ShaderResourceId < ResourceParameters.Num() && ResourceParameters[ShaderResourceId].ByteOffset < ByteOffset; ++ShaderResourceId)
+					{
+					}
+					bResourceIsUsed |= ShaderResourceId < ResourceParameters.Num() && ByteOffset == ResourceParameters[ShaderResourceId].ByteOffset;
+				}
+
+				if (!bResourceIsUsed)
+				{
+					const auto& BindlessResourceParameters = ShaderBindingsList[Index]->BindlessResourceParameters;
+					int32& BindlessResourceId = BindlessResourceIds[Index];
+					for (; BindlessResourceId < BindlessResourceParameters.Num() && BindlessResourceParameters[BindlessResourceId].ByteOffset < ByteOffset; ++BindlessResourceId)
+					{
+					}
+					bResourceIsUsed |= BindlessResourceId < BindlessResourceParameters.Num() && ByteOffset == BindlessResourceParameters[BindlessResourceId].ByteOffset;
 				}
-				bResourceIsUsed |= ShaderResourceId < ResourceParameters.Num() && ByteOffset == ResourceParameters[ShaderResourceId].ByteOffset;
 			}
 		}
 		else if (Type == UBMT_RDG_UNIFORM_BUFFER)
