Removing incorrect LLM tag from BeginInitResource() introduced in CL 22230608.

Forcing a tag here is incorrect, as it will override the tag that's supposed to be set by the owner of the resource. For example, this attributed StaticMesh memory to the generic SceneRender tag.

#rnx

diff --git a/Engine/Source/Runtime/RenderCore/Private/RenderResource.cpp b/Engine/Source/Runtime/RenderCore/Private/RenderResource.cpp
--- a/Engine/Source/Runtime/RenderCore/Private/RenderResource.cpp
+++ b/Engine/Source/Runtime/RenderCore/Private/RenderResource.cpp
@@ -406,7 +406,6 @@ FName FRenderResource::GetOwnerName() const
 
 void BeginInitResource(FRenderResource* Resource)
 {
-	LLM_SCOPE(ELLMTag::SceneRender);
 	ENQUEUE_RENDER_COMMAND(InitCommand)(
 		[Resource](FRHICommandListImmediate& RHICmdList)
 		{
