Fixed potential race condition when prerequisites are used for command list setup tasks.

diff --git a/Engine/Source/Runtime/RenderCore/Public/RenderGraphBuilder.inl b/Engine/Source/Runtime/RenderCore/Public/RenderGraphBuilder.inl
--- a/Engine/Source/Runtime/RenderCore/Public/RenderGraphBuilder.inl
+++ b/Engine/Source/Runtime/RenderCore/Public/RenderGraphBuilder.inl
@@ -447,7 +447,7 @@ UE::Tasks::FTask FRDGBuilder::AddSetupTask(
 {
 	UE::Tasks::FTask Task;
 
-	bCondition &= bParallelExecuteEnabled;
+	bCondition &= bParallelSetupEnabled;
 
 	auto OuterLambda = [TaskLambda = MoveTemp(TaskLambda)]() mutable
 	{
@@ -521,11 +521,24 @@ UE::Tasks::FTask FRDGBuilder::AddCommandListSetupTask(
 {
 	UE::Tasks::FTask Task;
 
-	bCondition &= bParallelExecuteEnabled;
+	bCondition &= bParallelSetupEnabled;
 
 	FRHICommandList* RHICmdListTask = nullptr;
 
-	if (bCondition)
+	bool bUseSeparateCommandList = bCondition;
+
+	for (const UE::Tasks::FTask& Prerequisite : Prerequisites)
+	{
+		// Always create a command list when a prerequisite task is present, as the inline task can be scheduled on any thread.
+		if (!Prerequisite.IsCompleted())
+		{
+			bUseSeparateCommandList = true;
+			break;
+		}
+	}
+
+	// When using prerequisites we have to 
+	if (bUseSeparateCommandList)
 	{
 		RHICmdListTask = new FRHICommandList(RHICmdList.GetGPUMask());
 	}
@@ -534,18 +547,18 @@ UE::Tasks::FTask FRDGBuilder::AddCommandListSetupTask(
 		RHICmdListTask = &RHICmdList;
 	}
 
-	auto OuterLambda = [TaskLambda = MoveTemp(TaskLambda), RHICmdListTask, bCondition]() mutable
+	auto OuterLambda = [TaskLambda = MoveTemp(TaskLambda), RHICmdListTask]() mutable
 	{
 		FOptionalTaskTagScope Scope(ETaskTag::EParallelRenderingThread);
 
-		if (bCondition)
+		if (!RHICmdListTask->IsImmediate())
 		{
 			RHICmdListTask->SwitchPipeline(ERHIPipeline::Graphics);
 		}
 
 		TaskLambda(*RHICmdListTask);
 
-		if (bCondition)
+		if (!RHICmdListTask->IsImmediate())
 		{
 			RHICmdListTask->FinishRecording();
 		}
@@ -567,7 +580,7 @@ UE::Tasks::FTask FRDGBuilder::AddCommandListSetupTask(
 		Task = UE::Tasks::Launch(TEXT("FRDGBuilder::AddCommandListSetupTask"), MoveTemp(OuterLambda), Forward<PrerequisitesCollectionType&&>(Prerequisites), Priority, ExtendedTaskPriority);
 	}
 
-	if (bCondition)
+	if (!RHICmdListTask->IsImmediate())
 	{
 		ParallelSetupEvents.Emplace(Task);
 		RHICmdList.QueueAsyncCommandListSubmit(RHICmdListTask);
