More GPU Scene data code clean-up:
- Move primitive data stride to SceneDefinitions header so we don't have to change it in 3 places in code.
- Move FPrimitiveSceneShaderData to its own header to (eventually) eliminate dependencies that adversely affect compile time.
- Deprecate scatter upload buffer support for texture targets and prepare to clean up hacks that were previously done for the sake of GPU Scene that are no longer needed.

diff --git a/Engine/Plugins/Experimental/GPULightmass/Source/GPULightmass/Private/LightmapRenderer.cpp b/Engine/Plugins/Experimental/GPULightmass/Source/GPULightmass/Private/LightmapRenderer.cpp
--- a/Engine/Plugins/Experimental/GPULightmass/Source/GPULightmass/Private/LightmapRenderer.cpp
+++ b/Engine/Plugins/Experimental/GPULightmass/Source/GPULightmass/Private/LightmapRenderer.cpp
@@ -34,6 +34,7 @@
 #include "PostProcess/DrawRectangle.h"
 #include "SceneUniformBuffer.h"
 #include "PrimitiveUniformShaderParametersBuilder.h"
+#include "PrimitiveSceneShaderData.h"
 
 RENDERER_API uint8 BlendModeToRayTracingInstanceMask(const EBlendMode BlendMode, ERayTracingViewMaskMode MaskMode);
 RENDERER_API uint8 ComputeRayTracingInstanceShadowMask(ERayTracingViewMaskMode MaskMode);
diff --git a/Engine/Shaders/Private/SceneData.ush b/Engine/Shaders/Private/SceneData.ush
--- a/Engine/Shaders/Private/SceneData.ush
+++ b/Engine/Shaders/Private/SceneData.ush
@@ -223,9 +223,6 @@ StructuredBuffer<float4> GPUScenePrimitiveSceneData;
 RWStructuredBuffer<float4> GPUScenePrimitiveSceneDataRW;
 #endif
 
-// Stride of a single primitive's data in float4's, must match C++ (FPrimitiveSceneShaderData::DataStrideInFloat4s and FScatterUploadBuffer::PrimitiveDataStrideInFloat4s)
-#define PRIMITIVE_SCENE_DATA_STRIDE 41
-
 float4 LoadPrimitivePrimitiveSceneDataElement(uint PrimitiveIndex, uint ItemIndex)
 {
 	uint TargetIdx = PrimitiveIndex + ItemIndex;
diff --git a/Engine/Shaders/Shared/SceneDefinitions.h b/Engine/Shaders/Shared/SceneDefinitions.h
--- a/Engine/Shaders/Shared/SceneDefinitions.h
+++ b/Engine/Shaders/Shared/SceneDefinitions.h
@@ -38,3 +38,6 @@
 #define PRIMITIVE_SCENE_DATA_FLAG_REVERSE_CULLING						0x20000000
 #define PRIMITIVE_SCENE_DATA_FLAG_HOLDOUT								0x40000000
 #define PRIMITIVE_SCENE_DATA_FLAG_VISIBLE_IN_LUMEN_SCENE				0x80000000
+
+// Stride of a single primitive's data in float4's
+#define PRIMITIVE_SCENE_DATA_STRIDE 41
\ No newline at end of file
diff --git a/Engine/Source/Runtime/Engine/Private/PrimitiveUniformShaderParameters.cpp b/Engine/Source/Runtime/Engine/Private/PrimitiveUniformShaderParameters.cpp
--- a/Engine/Source/Runtime/Engine/Private/PrimitiveUniformShaderParameters.cpp
+++ b/Engine/Source/Runtime/Engine/Private/PrimitiveUniformShaderParameters.cpp
@@ -8,6 +8,7 @@
 #include "UnrealEngine.h"
 #include "DataDrivenShaderPlatformInfo.h"
 #include "PrimitiveUniformShaderParametersBuilder.h"
+#include "PrimitiveSceneShaderData.h"
 
 FPrimitiveUniformShaderParametersBuilder& FPrimitiveUniformShaderParametersBuilder::InstanceDrawDistance(FVector2f DistanceMinMax)
 {
diff --git a/Engine/Source/Runtime/Engine/Private/SceneManagement.cpp b/Engine/Source/Runtime/Engine/Private/SceneManagement.cpp
--- a/Engine/Source/Runtime/Engine/Private/SceneManagement.cpp
+++ b/Engine/Source/Runtime/Engine/Private/SceneManagement.cpp
@@ -21,6 +21,7 @@
 #include "DataDrivenShaderPlatformInfo.h"
 #include "StaticMeshBatch.h"
 #include "PrimitiveUniformShaderParametersBuilder.h"
+#include "PrimitiveSceneShaderData.h"
 
 static TAutoConsoleVariable<float> CVarLODTemporalLag(
 	TEXT("lod.TemporalLag"),
diff --git a/Engine/Source/Runtime/Engine/Public/PrimitiveSceneShaderData.h b/Engine/Source/Runtime/Engine/Public/PrimitiveSceneShaderData.h
new file mode 100644
--- /dev/null
+++ b/Engine/Source/Runtime/Engine/Public/PrimitiveSceneShaderData.h
@@ -0,0 +1,122 @@
+// Copyright Epic Games, Inc. All Rights Reserved.
+
+#pragma once
+
+#include "CoreMinimal.h"
+#include "SceneDefinitions.h"
+#include "Containers/StaticArray.h"
+
+#if !UE_ENABLE_INCLUDE_ORDER_DEPRECATED_IN_5_3
+#include "PrimitiveUniformShaderParameters.h"
+#endif
+
+struct FPrimitiveSceneShaderData
+{
+	static const uint32 DataStrideInFloat4s = PRIMITIVE_SCENE_DATA_STRIDE;
+
+	TStaticArray<FVector4f, DataStrideInFloat4s> Data;
+
+	FPrimitiveSceneShaderData()
+		: Data(InPlace, NoInit)
+	{
+		Setup(GetIdentityPrimitiveParameters());
+	}
+
+	explicit FPrimitiveSceneShaderData(const FPrimitiveUniformShaderParameters& PrimitiveUniformShaderParameters)
+		: Data(InPlace, NoInit)
+	{
+		Setup(PrimitiveUniformShaderParameters);
+	}
+
+	ENGINE_API FPrimitiveSceneShaderData(const class FPrimitiveSceneProxy* RESTRICT Proxy);
+
+	ENGINE_API void Setup(const FPrimitiveUniformShaderParameters& PrimitiveUniformShaderParameters);
+};
+
+class FSinglePrimitiveStructured : public FRenderResource
+{
+public:
+
+	FSinglePrimitiveStructured()
+		: ShaderPlatform(SP_NumPlatforms)
+	{}
+
+	ENGINE_API virtual void InitRHI() override;
+
+	virtual void ReleaseRHI() override
+	{
+		PRAGMA_DISABLE_DEPRECATION_WARNINGS
+		PrimitiveSceneDataBufferRHI.SafeRelease();
+		PrimitiveSceneDataBufferSRV.SafeRelease();
+		SkyIrradianceEnvironmentMapRHI.SafeRelease();
+		SkyIrradianceEnvironmentMapSRV.SafeRelease();
+		InstanceSceneDataBufferRHI.SafeRelease();
+		InstanceSceneDataBufferSRV.SafeRelease();
+		InstancePayloadDataBufferRHI.SafeRelease();
+		InstancePayloadDataBufferSRV.SafeRelease();
+		PrimitiveSceneDataTextureRHI.SafeRelease();
+		PrimitiveSceneDataTextureSRV.SafeRelease();
+		LightmapSceneDataBufferRHI.SafeRelease();
+		LightmapSceneDataBufferSRV.SafeRelease();
+		EditorVisualizeLevelInstanceDataBufferRHI.SafeRelease();
+		EditorVisualizeLevelInstanceDataBufferSRV.SafeRelease();
+		EditorSelectedDataBufferRHI.SafeRelease();
+		EditorSelectedDataBufferSRV.SafeRelease();
+		PRAGMA_ENABLE_DEPRECATION_WARNINGS
+	}
+
+	ENGINE_API void UploadToGPU();
+
+	EShaderPlatform ShaderPlatform=SP_NumPlatforms;
+
+	UE_DEPRECATED(5.3, "Use FRendererModule::CreateSinglePrimitiveSceneUniformBuffer instead")
+	FPrimitiveSceneShaderData PrimitiveSceneData;
+	UE_DEPRECATED(5.3, "Use FRendererModule::CreateSinglePrimitiveSceneUniformBuffer instead")
+	FInstanceSceneShaderData InstanceSceneData;
+	UE_DEPRECATED(5.3, "Use FRendererModule::CreateSinglePrimitiveSceneUniformBuffer instead")
+	FLightmapSceneShaderData LightmapSceneData;
+
+	UE_DEPRECATED(5.3, "Use FRendererModule::CreateSinglePrimitiveSceneUniformBuffer instead")
+	FBufferRHIRef PrimitiveSceneDataBufferRHI;
+	UE_DEPRECATED(5.3, "Use FRendererModule::CreateSinglePrimitiveSceneUniformBuffer instead")
+	FShaderResourceViewRHIRef PrimitiveSceneDataBufferSRV;
+
+	FBufferRHIRef SkyIrradianceEnvironmentMapRHI;
+	FShaderResourceViewRHIRef SkyIrradianceEnvironmentMapSRV;
+
+	UE_DEPRECATED(5.3, "Use FRendererModule::CreateSinglePrimitiveSceneUniformBuffer instead")
+	FBufferRHIRef InstanceSceneDataBufferRHI;
+	UE_DEPRECATED(5.3, "Use FRendererModule::CreateSinglePrimitiveSceneUniformBuffer instead")
+	FShaderResourceViewRHIRef InstanceSceneDataBufferSRV;
+
+	UE_DEPRECATED(5.3, "Use FRendererModule::CreateSinglePrimitiveSceneUniformBuffer instead")
+	FBufferRHIRef InstancePayloadDataBufferRHI;
+	UE_DEPRECATED(5.3, "Use FRendererModule::CreateSinglePrimitiveSceneUniformBuffer instead")
+	FShaderResourceViewRHIRef InstancePayloadDataBufferSRV;
+
+	FTextureRHIRef PrimitiveSceneDataTextureRHI;
+	FShaderResourceViewRHIRef PrimitiveSceneDataTextureSRV;
+
+	UE_DEPRECATED(5.3, "Use FRendererModule::CreateSinglePrimitiveSceneUniformBuffer instead")
+	FBufferRHIRef LightmapSceneDataBufferRHI;
+	UE_DEPRECATED(5.3, "Use FRendererModule::CreateSinglePrimitiveSceneUniformBuffer instead")
+	FShaderResourceViewRHIRef LightmapSceneDataBufferSRV;
+
+	UE_DEPRECATED(5.3, "Use FRendererModule::CreateSinglePrimitiveSceneUniformBuffer instead")
+	FBufferRHIRef EditorVisualizeLevelInstanceDataBufferRHI;
+	UE_DEPRECATED(5.3, "Use FRendererModule::CreateSinglePrimitiveSceneUniformBuffer instead")
+	FShaderResourceViewRHIRef EditorVisualizeLevelInstanceDataBufferSRV;
+
+	UE_DEPRECATED(5.3, "Use FRendererModule::CreateSinglePrimitiveSceneUniformBuffer instead")
+	FBufferRHIRef EditorSelectedDataBufferRHI;
+	UE_DEPRECATED(5.3, "Use FRendererModule::CreateSinglePrimitiveSceneUniformBuffer instead")
+	FShaderResourceViewRHIRef EditorSelectedDataBufferSRV;
+};
+
+/**
+* Default Primitive data buffer.  
+* This is used when the VF is used for rendering outside normal mesh passes, where there is no valid scene.
+*/
+extern ENGINE_API TGlobalResource<FSinglePrimitiveStructured> GIdentityPrimitiveBuffer;
+UE_DEPRECATED(5.3, "Use FRendererModule::CreateSinglePrimitiveSceneUniformBuffer instead")
+extern ENGINE_API TGlobalResource<FSinglePrimitiveStructured> GTilePrimitiveBuffer;
diff --git a/Engine/Source/Runtime/Engine/Public/PrimitiveUniformShaderParameters.h b/Engine/Source/Runtime/Engine/Public/PrimitiveUniformShaderParameters.h
--- a/Engine/Source/Runtime/Engine/Public/PrimitiveUniformShaderParameters.h
+++ b/Engine/Source/Runtime/Engine/Public/PrimitiveUniformShaderParameters.h
@@ -11,14 +11,13 @@
 #include "UniformBuffer.h"
 #include "InstanceUniformShaderParameters.h"
 #include "LightmapUniformShaderParameters.h"
+#if UE_ENABLE_INCLUDE_ORDER_DEPRECATED_IN_5_3
 #include "UnifiedBuffer.h"
-#include "Containers/StaticArray.h"
+#endif
 #if UE_ENABLE_INCLUDE_ORDER_DEPRECATED_IN_5_2
 #include "UnrealEngine.h"
 #endif
 
-struct FPrimitiveUniformShaderParametersBuilder;
-
 /** 
  * The uniform shader parameters associated with a primitive. 
  * Note: Must match FPrimitiveSceneData in shaders.
@@ -96,120 +95,7 @@ public:
 /** Global primitive uniform buffer resource containing identity transformations. */
 extern ENGINE_API TGlobalResource<FIdentityPrimitiveUniformBuffer> GIdentityPrimitiveUniformBuffer;
 
-struct FPrimitiveSceneShaderData
-{
-	// Must match PRIMITIVE_SCENE_DATA_STRIDE in SceneData.ush
-	enum { DataStrideInFloat4s = 41 };
-
-	TStaticArray<FVector4f, DataStrideInFloat4s> Data;
-
-	FPrimitiveSceneShaderData()
-		: Data(InPlace, NoInit)
-	{
-		static_assert(FPrimitiveSceneShaderData::DataStrideInFloat4s == (int)FScatterUploadBuffer::PrimitiveDataStrideInFloat4s,"");
-		Setup(GetIdentityPrimitiveParameters());
-	}
-
-	explicit FPrimitiveSceneShaderData(const FPrimitiveUniformShaderParameters& PrimitiveUniformShaderParameters)
-		: Data(InPlace, NoInit)
-	{
-		Setup(PrimitiveUniformShaderParameters);
-	}
-
-	ENGINE_API FPrimitiveSceneShaderData(const class FPrimitiveSceneProxy* RESTRICT Proxy);
-
-	ENGINE_API void Setup(const FPrimitiveUniformShaderParameters& PrimitiveUniformShaderParameters);
-};
-
-class FSinglePrimitiveStructured : public FRenderResource
-{
-public:
-
-	FSinglePrimitiveStructured()
-		: ShaderPlatform(SP_NumPlatforms)
-	{}
-
-	ENGINE_API virtual void InitRHI() override;
-
-	virtual void ReleaseRHI() override
-	{
-		PRAGMA_DISABLE_DEPRECATION_WARNINGS
-		PrimitiveSceneDataBufferRHI.SafeRelease();
-		PrimitiveSceneDataBufferSRV.SafeRelease();
-		SkyIrradianceEnvironmentMapRHI.SafeRelease();
-		SkyIrradianceEnvironmentMapSRV.SafeRelease();
-		InstanceSceneDataBufferRHI.SafeRelease();
-		InstanceSceneDataBufferSRV.SafeRelease();
-		InstancePayloadDataBufferRHI.SafeRelease();
-		InstancePayloadDataBufferSRV.SafeRelease();
-		PrimitiveSceneDataTextureRHI.SafeRelease();
-		PrimitiveSceneDataTextureSRV.SafeRelease();
-		LightmapSceneDataBufferRHI.SafeRelease();
-		LightmapSceneDataBufferSRV.SafeRelease();
-		EditorVisualizeLevelInstanceDataBufferRHI.SafeRelease();
-		EditorVisualizeLevelInstanceDataBufferSRV.SafeRelease();
-		EditorSelectedDataBufferRHI.SafeRelease();
-		EditorSelectedDataBufferSRV.SafeRelease();
-		PRAGMA_ENABLE_DEPRECATION_WARNINGS
-	}
-
-	ENGINE_API void UploadToGPU();
-
-	EShaderPlatform ShaderPlatform=SP_NumPlatforms;
-
-	UE_DEPRECATED(5.3, "Use FRendererModule::CreateSinglePrimitiveSceneUniformBuffer instead")
-	FPrimitiveSceneShaderData PrimitiveSceneData;
-	UE_DEPRECATED(5.3, "Use FRendererModule::CreateSinglePrimitiveSceneUniformBuffer instead")
-	FInstanceSceneShaderData InstanceSceneData;
-	UE_DEPRECATED(5.3, "Use FRendererModule::CreateSinglePrimitiveSceneUniformBuffer instead")
-	FLightmapSceneShaderData LightmapSceneData;
-
-	UE_DEPRECATED(5.3, "Use FRendererModule::CreateSinglePrimitiveSceneUniformBuffer instead")
-	FBufferRHIRef PrimitiveSceneDataBufferRHI;
-	UE_DEPRECATED(5.3, "Use FRendererModule::CreateSinglePrimitiveSceneUniformBuffer instead")
-	FShaderResourceViewRHIRef PrimitiveSceneDataBufferSRV;
-
-	FBufferRHIRef SkyIrradianceEnvironmentMapRHI;
-	FShaderResourceViewRHIRef SkyIrradianceEnvironmentMapSRV;
-
-	UE_DEPRECATED(5.3, "Use FRendererModule::CreateSinglePrimitiveSceneUniformBuffer instead")
-	FBufferRHIRef InstanceSceneDataBufferRHI;
-	UE_DEPRECATED(5.3, "Use FRendererModule::CreateSinglePrimitiveSceneUniformBuffer instead")
-	FShaderResourceViewRHIRef InstanceSceneDataBufferSRV;
-
-	UE_DEPRECATED(5.3, "Use FRendererModule::CreateSinglePrimitiveSceneUniformBuffer instead")
-	FBufferRHIRef InstancePayloadDataBufferRHI;
-	UE_DEPRECATED(5.3, "Use FRendererModule::CreateSinglePrimitiveSceneUniformBuffer instead")
-	FShaderResourceViewRHIRef InstancePayloadDataBufferSRV;
-
-	FTextureRHIRef PrimitiveSceneDataTextureRHI;
-	FShaderResourceViewRHIRef PrimitiveSceneDataTextureSRV;
-
-	UE_DEPRECATED(5.3, "Use FRendererModule::CreateSinglePrimitiveSceneUniformBuffer instead")
-	FBufferRHIRef LightmapSceneDataBufferRHI;
-	UE_DEPRECATED(5.3, "Use FRendererModule::CreateSinglePrimitiveSceneUniformBuffer instead")
-	FShaderResourceViewRHIRef LightmapSceneDataBufferSRV;
-
-	UE_DEPRECATED(5.3, "Use FRendererModule::CreateSinglePrimitiveSceneUniformBuffer instead")
-	FBufferRHIRef EditorVisualizeLevelInstanceDataBufferRHI;
-	UE_DEPRECATED(5.3, "Use FRendererModule::CreateSinglePrimitiveSceneUniformBuffer instead")
-	FShaderResourceViewRHIRef EditorVisualizeLevelInstanceDataBufferSRV;
-
-	UE_DEPRECATED(5.3, "Use FRendererModule::CreateSinglePrimitiveSceneUniformBuffer instead")
-	FBufferRHIRef EditorSelectedDataBufferRHI;
-	UE_DEPRECATED(5.3, "Use FRendererModule::CreateSinglePrimitiveSceneUniformBuffer instead")
-	FShaderResourceViewRHIRef EditorSelectedDataBufferSRV;
-
-};
-
-/**
-* Default Primitive data buffer.  
-* This is used when the VF is used for rendering outside normal mesh passes, where there is no valid scene.
-*/
-extern ENGINE_API TGlobalResource<FSinglePrimitiveStructured> GIdentityPrimitiveBuffer;
-UE_DEPRECATED(5.3, "Use FRendererModule::CreateSinglePrimitiveSceneUniformBuffer instead")
-extern ENGINE_API TGlobalResource<FSinglePrimitiveStructured> GTilePrimitiveBuffer;
-
 #if UE_ENABLE_INCLUDE_ORDER_DEPRECATED_IN_5_3
 #include "PrimitiveUniformShaderParametersBuilder.h"
+#include "PrimitiveSceneShaderData.h"
 #endif
diff --git a/Engine/Source/Runtime/RenderCore/Private/UnifiedBuffer.cpp b/Engine/Source/Runtime/RenderCore/Private/UnifiedBuffer.cpp
--- a/Engine/Source/Runtime/RenderCore/Private/UnifiedBuffer.cpp
+++ b/Engine/Source/Runtime/RenderCore/Private/UnifiedBuffer.cpp
@@ -166,7 +166,7 @@ enum class EResourceType
 	BUFFER,
 	STRUCTURED_BUFFER,
 	BYTEBUFFER,
-	TEXTURE
+	TEXTURE // NOTE: Deprecated
 };
 
 template<typename ResourceType>
@@ -184,6 +184,7 @@ struct ResourceTypeTraits<FRWBufferStructured>
 	static const EResourceType Type = EResourceType::STRUCTURED_BUFFER;
 };
 
+// NOTE: Deprecated
 template<>
 struct ResourceTypeTraits<FTextureRWBuffer>
 {
@@ -198,8 +199,12 @@ struct ResourceTypeTraits<FRWByteAddressBuffer>
 
 static uint32 CalculateFloat4sPerLine()
 {
+	// This is a remnant from when primitive scene data used textures to scatter upload, and should be removed once
+	// texture uploads have been fully deprecated
+	PRAGMA_DISABLE_DEPRECATION_WARNINGS
 	uint16 PrimitivesPerTextureLine = (uint16)FMath::Min((int32)MAX_uint16, (int32)GMaxTextureDimensions) / FScatterUploadBuffer::PrimitiveDataStrideInFloat4s;
 	return PrimitivesPerTextureLine * FScatterUploadBuffer::PrimitiveDataStrideInFloat4s;
+	PRAGMA_ENABLE_DEPRECATION_WARNINGS
 }
 
 ///////////////////////////////////////////////////////////////////////////////////////////////////
@@ -408,6 +413,20 @@ void MemsetResource(FRDGBuilder& GraphBuilder, FRDGUnorderedAccessView* UAV, con
 		FIntVector(FMath::DivideAndRoundUp(Params.Count / Divisor, 64u), 1, 1));
 }
 
+void MemsetResource(FRDGBuilder& GraphBuilder, FRDGBufferUAV* UAV, const FMemsetResourceParams& Params)
+{
+	PRAGMA_DISABLE_DEPRECATION_WARNINGS
+	MemsetResource(GraphBuilder, static_cast<FRDGUnorderedAccessView*>(UAV), Params);
+	PRAGMA_ENABLE_DEPRECATION_WARNINGS
+}
+
+void MemsetResource(FRDGBuilder& GraphBuilder, FRDGTextureUAV* UAV, const FMemsetResourceParams& Params)
+{
+	PRAGMA_DISABLE_DEPRECATION_WARNINGS
+	MemsetResource(GraphBuilder, static_cast<FRDGUnorderedAccessView*>(UAV), Params);
+	PRAGMA_ENABLE_DEPRECATION_WARNINGS
+}
+
 void MemcpyResource(FRDGBuilder& GraphBuilder, FRDGUnorderedAccessView* UAV, FRDGShaderResourceView* SRV, const FMemcpyResourceParams& Params)
 {
 	check(UAV && SRV);
@@ -504,6 +523,20 @@ void MemcpyResource(FRDGBuilder& GraphBuilder, FRDGUnorderedAccessView* UAV, FRD
 	}
 }
 
+void MemcpyResource(FRDGBuilder& GraphBuilder, FRDGBufferUAV* UAV, FRDGBufferSRV* SRV, const FMemcpyResourceParams& Params)
+{
+	PRAGMA_DISABLE_DEPRECATION_WARNINGS
+	MemcpyResource(GraphBuilder, static_cast<FRDGUnorderedAccessView*>(UAV), static_cast<FRDGShaderResourceView*>(SRV), Params);
+	PRAGMA_ENABLE_DEPRECATION_WARNINGS
+}
+
+void MemcpyResource(FRDGBuilder& GraphBuilder, FRDGTextureUAV* UAV, FRDGTextureSRV* SRV, const FMemcpyResourceParams& Params)
+{
+	PRAGMA_DISABLE_DEPRECATION_WARNINGS
+	MemcpyResource(GraphBuilder, static_cast<FRDGUnorderedAccessView*>(UAV), static_cast<FRDGShaderResourceView*>(SRV), Params);
+	PRAGMA_ENABLE_DEPRECATION_WARNINGS
+}
+
 FRDGBuffer* ResizeBufferIfNeeded(FRDGBuilder& GraphBuilder, TRefCountPtr<FRDGPooledBuffer>& ExternalBuffer, const FRDGBufferDesc& BufferDesc, const TCHAR* Name)
 {
 	FRDGBuffer* InternalBufferNew = nullptr;
@@ -774,6 +807,19 @@ void ScatterCopyResource(FRDGBuilder& GraphBuilder, FRDGViewableResource* DstRes
 	}
 }
 
+void ScatterCopyResource(FRDGBuilder& GraphBuilder, FRDGBuffer* DstResource, FRDGBufferSRV* ScatterBufferSRV, FRDGBufferSRV* UploadBufferSRV, const FScatterCopyParams &Params)
+{
+	PRAGMA_DISABLE_DEPRECATION_WARNINGS
+	ScatterCopyResource(GraphBuilder, static_cast<FRDGViewableResource*>(DstResource), ScatterBufferSRV, UploadBufferSRV, Params);
+	PRAGMA_ENABLE_DEPRECATION_WARNINGS
+}
+void ScatterCopyResource(FRDGBuilder& GraphBuilder, FRDGTexture* DstResource, FRDGBufferSRV* ScatterBufferSRV, FRDGBufferSRV* UploadBufferSRV, const FScatterCopyParams &Params)
+{
+	PRAGMA_DISABLE_DEPRECATION_WARNINGS
+	ScatterCopyResource(GraphBuilder, static_cast<FRDGViewableResource*>(DstResource), ScatterBufferSRV, UploadBufferSRV, Params);
+	PRAGMA_ENABLE_DEPRECATION_WARNINGS
+}
+
 
 
 void FRDGScatterUploadBuffer::Init(FRDGBuilder& GraphBuilder, uint32 NumElements, uint32 InNumBytesPerElement, bool bInFloat4Buffer, const TCHAR* Name)
@@ -825,7 +871,7 @@ void FRDGScatterUploadBuffer::Init(FRDGBuilder& GraphBuilder, uint32 NumElements
 	UploadData = (uint8*)RHILockBuffer(UploadBuffer->GetRHI(), 0, UploadBytes, RLM_WriteOnly);
 }
 
-void FRDGScatterUploadBuffer::ResourceUploadTo(FRDGBuilder& GraphBuilder, FRDGViewableResource* DstResource)
+void FRDGScatterUploadBuffer::ResourceUploadToInternal(FRDGBuilder& GraphBuilder, FRDGViewableResource* DstResource)
 {
 	TRACE_CPUPROFILER_EVENT_SCOPE(FRDGScatterUploadBuffer::ResourceUploadTo);
 
@@ -845,7 +891,10 @@ void FRDGScatterUploadBuffer::ResourceUploadTo(FRDGBuilder& GraphBuilder, FRDGVi
 
 	FRDGBufferSRV* ScatterBufferSRV = GraphBuilder.CreateSRV(GraphBuilder.RegisterExternalBuffer(ScatterBuffer));
 	FRDGBufferSRV* UploadBufferSRV  = GraphBuilder.CreateSRV(GraphBuilder.RegisterExternalBuffer(UploadBuffer));
+	
+	PRAGMA_DISABLE_DEPRECATION_WARNINGS
 	ScatterCopyResource(GraphBuilder, DstResource, ScatterBufferSRV, UploadBufferSRV, FScatterCopyParams { NumScatters, NumBytesPerElement, INDEX_NONE});
+	PRAGMA_ENABLE_DEPRECATION_WARNINGS
 
 	Reset();
 }
@@ -1199,8 +1248,8 @@ void MemcpyResource(FRHICommandList& RHICmdList, const ResourceType& DstBuffer,
 		RHICmdList.EndUAVOverlap(DstBuffer.UAV);
 }
 
-template<>
-RENDERCORE_API bool ResizeResourceIfNeeded<FTextureRWBuffer>(FRHICommandList& RHICmdList, FTextureRWBuffer& Texture, uint32 NumBytes, const TCHAR* DebugName)
+// NOTE: Deprecated
+RENDERCORE_API bool ResizeResourceIfNeeded(FRHICommandList& RHICmdList, FTextureRWBuffer& Texture, uint32 NumBytes, const TCHAR* DebugName)
 {
 	check((NumBytes & 15) == 0);
 	uint32 Float4sPerLine = CalculateFloat4sPerLine();
@@ -1225,7 +1274,9 @@ RENDERCORE_API bool ResizeResourceIfNeeded<FTextureRWBuffer>(FRHICommandList& RH
 		Params.Count = NumBytes / BytesPerElement;
 		Params.SrcOffset = 0;
 		Params.DstOffset = 0;
+		PRAGMA_DISABLE_DEPRECATION_WARNINGS
 		MemcpyResource(RHICmdList, NewTexture, Texture, Params);
+		PRAGMA_ENABLE_DEPRECATION_WARNINGS
 		Texture = NewTexture;
 		return true;
 	}
@@ -1780,7 +1831,9 @@ template RENDERCORE_API void MemcpyResource<FRWBuffer>(FRHICommandList& RHICmdLi
 template RENDERCORE_API void MemcpyResource<FRWBufferStructured>(FRHICommandList& RHICmdList, const FRWBufferStructured& DstBuffer, const FRWBufferStructured& SrcBuffer, const FMemcpyResourceParams& Params, bool bAlreadyInUAVOverlap);
 template RENDERCORE_API void MemcpyResource<FRWByteAddressBuffer>(FRHICommandList& RHICmdList, const FRWByteAddressBuffer& DstBuffer, const FRWByteAddressBuffer& SrcBuffer, const FMemcpyResourceParams& Params, bool bAlreadyInUAVOverlap);
 
-template RENDERCORE_API void FScatterUploadBuffer::ResourceUploadTo<FTextureRWBuffer>(FRHICommandList& RHICmdList, const FTextureRWBuffer& DstBuffer, bool bFlush);
 template RENDERCORE_API void FScatterUploadBuffer::ResourceUploadTo<FRWBuffer>(FRHICommandList& RHICmdList, const FRWBuffer& DstBuffer, bool bFlush);
 template RENDERCORE_API void FScatterUploadBuffer::ResourceUploadTo<FRWBufferStructured>(FRHICommandList& RHICmdList, const FRWBufferStructured& DstBuffer, bool bFlush);
 template RENDERCORE_API void FScatterUploadBuffer::ResourceUploadTo<FRWByteAddressBuffer>(FRHICommandList& RHICmdList, const FRWByteAddressBuffer& DstBuffer, bool bFlush);
+
+// NOTE: Deprecated
+template RENDERCORE_API void FScatterUploadBuffer::ResourceUploadTo<FTextureRWBuffer>(FRHICommandList& RHICmdList, const FTextureRWBuffer& DstBuffer, bool bFlush);
\ No newline at end of file
diff --git a/Engine/Source/Runtime/RenderCore/Public/UnifiedBuffer.h b/Engine/Source/Runtime/RenderCore/Public/UnifiedBuffer.h
--- a/Engine/Source/Runtime/RenderCore/Public/UnifiedBuffer.h
+++ b/Engine/Source/Runtime/RenderCore/Public/UnifiedBuffer.h
@@ -50,12 +50,26 @@ extern RENDERCORE_API void MemsetResource(FRHICommandList& RHICmdList, const Res
 template<typename ResourceType>
 extern RENDERCORE_API void MemcpyResource(FRHICommandList& RHICmdList, const ResourceType& DstBuffer, const ResourceType& SrcBuffer, const FMemcpyResourceParams& Params, bool bAlreadyInUAVOverlap = false);
 
+UE_DEPRECATED(5.3, "MemsetResource will no longer support FTextureRWBuffer in a future release. Either use a buffer target instead, or clear the contents manually.")
+inline void MemsetResource(FRHICommandList& RHICmdList, const FTextureRWBuffer& DstBuffer, const FMemsetResourceParams& Params)
+{
+	MemsetResource<FTextureRWBuffer>(RHICmdList, DstBuffer, Params);
+}
+
+UE_DEPRECATED(5.3, "MemcpyResource will no longer support FTextureRWBuffer in a future release. Either use a buffer target instead, or copy the contents manually.")
+inline void MemcpyResource(FRHICommandList& RHICmdList, const FTextureRWBuffer& DstBuffer, const FTextureRWBuffer& SrcBuffer, const FMemcpyResourceParams& Params, bool bAlreadyInUAVOverlap = false)
+{
+	MemcpyResource<FTextureRWBuffer>(RHICmdList, DstBuffer, SrcBuffer, Params, bAlreadyInUAVOverlap);
+}
+
 template<typename ResourceType>
 extern RENDERCORE_API bool ResizeResourceSOAIfNeeded(FRHICommandList& RHICmdList, ResourceType& Texture, const FResizeResourceSOAParams& Params, const TCHAR* DebugName);
 template<typename ResourceType>
-extern RENDERCORE_API bool ResizeResourceIfNeeded(FRHICommandList& RHICmdList, ResourceType& Texture, uint32 NumBytes, const TCHAR* DebugName);
+extern RENDERCORE_API bool ResizeResourceIfNeeded(FRHICommandList& RHICmdList, ResourceType& Buffer, uint32 NumBytes, const TCHAR* DebugName);
 RENDERCORE_API bool ResizeResourceIfNeeded(FRHICommandList& RHICmdList, FRWBuffer& Buffer, EPixelFormat Format, uint32 NumElements, const TCHAR* DebugName);
 
+UE_DEPRECATED(5.3, "ResizeResourceIfNeeded will no longer support FTextureRWBuffer in a future release. Either use a buffer target instead, or create a new texture and copy the contents manually.")
+RENDERCORE_API bool ResizeResourceIfNeeded(FRHICommandList& RHICmdList, FTextureRWBuffer& Texture, uint32 NumBytes, const TCHAR* DebugName);
 
 /**
  * This version will resize/allocate the buffer at once and add a RDG pass to perform the copy on the RDG time-line if there was previous data).
@@ -68,7 +82,13 @@ RENDERCORE_API bool ResizeResourceIfNeeded(FRDGBuilder& GraphBuilder, FRWBuffer&
 class FScatterUploadBuffer
 {
 public:
-	enum { PrimitiveDataStrideInFloat4s = 41 };
+	// This is an old dependency upon the primitive data that is no longer used
+	
+	enum
+	UE_DEPRECATED(5.3, "Use PRIMITIVE_SCENE_DATA_STRIDE in SceneDefinitions.h instead.")
+	{
+		PrimitiveDataStrideInFloat4s = 41
+	};
 
 	FByteAddressBuffer ScatterBuffer;
 	FByteAddressBuffer UploadBuffer;
@@ -94,6 +114,12 @@ public:
 
 	template<typename ResourceType>
 	RENDERCORE_API void ResourceUploadTo(FRHICommandList& RHICmdList, const ResourceType& DstBuffer, bool bFlush = false);
+	
+	UE_DEPRECATED(5.3, "Scattered uploads to FTextureRWBuffer will no longer supported in a future release. Use a supported destination buffer type instead.")
+	inline void ResourceUploadTo(FRHICommandList& RHICmdList, const FTextureRWBuffer& DstBuffer, bool bFlush = false)
+	{
+		ResourceUploadTo<FTextureRWBuffer>(RHICmdList, DstBuffer, bFlush);
+	}
 
 	void Add( uint32 Index, const void* Data, uint32 Num = 1 )
 	{
@@ -193,8 +219,17 @@ public:
 extern RENDERCORE_API void MemsetResource(FRDGBuilder& GraphBuilder, FRDGBuffer* DstResource, const FMemsetResourceParams& Params);
 extern RENDERCORE_API void MemcpyResource(FRDGBuilder& GraphBuilder, FRDGBuffer* DstResource, FRDGBuffer* SrcResource, const FMemcpyResourceParams& Params);
 
+extern RENDERCORE_API void MemsetResource(FRDGBuilder& GraphBuilder, FRDGBufferUAV* DstResource, const FMemsetResourceParams& Params);
+extern RENDERCORE_API void MemcpyResource(FRDGBuilder& GraphBuilder, FRDGBufferUAV* DstResource, FRDGBufferSRV* SrcResource, const FMemcpyResourceParams& Params);
+
+UE_DEPRECATED(5.3, "Use the overloads that take FRDGBuffer/FRDGBufferUAV instead.")
 extern RENDERCORE_API void MemsetResource(FRDGBuilder& GraphBuilder, FRDGUnorderedAccessView* DstResource, const FMemsetResourceParams& Params);
+UE_DEPRECATED(5.3, "MemsetResource will no longer support FRDGTextureUAV in a future release. Either use a buffer target instead, or clear the contents manually.")
+extern RENDERCORE_API void MemsetResource(FRDGBuilder& GraphBuilder, FRDGTextureUAV* DstResource, const FMemsetResourceParams& Params);
+UE_DEPRECATED(5.3, "Use the overloads that take FRDGBuffer/FRDGBufferUAV instead.")
 extern RENDERCORE_API void MemcpyResource(FRDGBuilder& GraphBuilder, FRDGUnorderedAccessView* DstResource, FRDGShaderResourceView* SrcResource, const FMemcpyResourceParams& Params);
+UE_DEPRECATED(5.3, "MemcpyResource will no longer support FRDGTextureUAV in a future release. Either use a buffer target instead, or copy the contents manually.")
+extern RENDERCORE_API void MemcpyResource(FRDGBuilder& GraphBuilder, FRDGTextureUAV* DstResource, FRDGTextureSRV* SrcResource, const FMemcpyResourceParams& Params);
 
 struct FScatterCopyParams
 {
@@ -203,7 +238,12 @@ struct FScatterCopyParams
 	int32 NumElementsPerScatter = INDEX_NONE; // INDEX_NONE lets the setup figure it out, otherwise it will run NumScatters * NumElementsPerScatter threads to copy the source data.
 };
 
+void RENDERCORE_API ScatterCopyResource(FRDGBuilder& GraphBuilder, FRDGBuffer* DstResource, FRDGBufferSRV* ScatterBufferSRV, FRDGBufferSRV* UploadBufferSRV, const FScatterCopyParams &Params);
+
+UE_DEPRECATED(5.3, "Use the overload that takes FRDGBuffer instead.")
 void RENDERCORE_API ScatterCopyResource(FRDGBuilder& GraphBuilder, FRDGViewableResource* DstResource, FRDGBufferSRV* ScatterBufferSRV, FRDGBufferSRV* UploadBufferSRV, const FScatterCopyParams &Params);
+UE_DEPRECATED(5.3, "ScatterCopyResource will no longer support FRDGTextureUAV in a future release. Either use a buffer target instead, or copy the contents manually.")
+void RENDERCORE_API ScatterCopyResource(FRDGBuilder& GraphBuilder, FRDGTexture* DstResource, FRDGBufferSRV* ScatterBufferSRV, FRDGBufferSRV* UploadBufferSRV, const FScatterCopyParams &Params);
 
 
 extern RENDERCORE_API FRDGBuffer* ResizeBufferIfNeeded(FRDGBuilder& GraphBuilder, TRefCountPtr<FRDGPooledBuffer>& ExternalBuffer, const FRDGBufferDesc& BufferDesc, const TCHAR* Name);
@@ -361,7 +401,14 @@ class RENDERCORE_API FRDGScatterUploadBuffer
 	: public FRDGScatterUploadBase
 {
 public:
-	enum { PrimitiveDataStrideInFloat4s = FScatterUploadBuffer::PrimitiveDataStrideInFloat4s };
+	
+	enum
+	UE_DEPRECATED(5.3, "Use PRIMITIVE_SCENE_DATA_STRIDE in SceneDefinitions.h instead.")
+	{
+		PRAGMA_DISABLE_DEPRECATION_WARNINGS
+		PrimitiveDataStrideInFloat4s = FScatterUploadBuffer::PrimitiveDataStrideInFloat4s
+		PRAGMA_ENABLE_DEPRECATION_WARNINGS
+	};
 
 	/**
 	 * Init with presized num scatters, expecting each to be set at a later point. Requires the user to keep track of the offsets to use.
@@ -375,13 +422,23 @@ public:
 
 	void Init(FRDGBuilder& GraphBuilder, uint32 NumElements, uint32 NumBytesPerElement, bool bInFloat4Buffer, const TCHAR* Name);
 
-	void ResourceUploadTo(FRDGBuilder& GraphBuilder, FRDGViewableResource* DstResource);
+	inline void ResourceUploadTo(FRDGBuilder& GraphBuilder, FRDGBuffer* DstResource)
+	{
+		ResourceUploadToInternal(GraphBuilder, DstResource);
+	}
+
+	UE_DEPRECATED(5.3, "Scattered uploads to FTextureRWBuffer will no longer supported in a future release. Use a supported destination buffer type instead.")
+	inline void ResourceUploadTo(FRDGBuilder& GraphBuilder, FRDGTexture* DstResource)
+	{
+		ResourceUploadToInternal(GraphBuilder, DstResource);
+	}
 
 	void Release();
 
 	uint32 GetNumBytes() const;
 
 private:
+	void ResourceUploadToInternal(FRDGBuilder& GraphBuilder, FRDGViewableResource* DstResource);
 	void Reset();
 
 	TRefCountPtr<FRDGPooledBuffer> ScatterBuffer;
diff --git a/Engine/Source/Runtime/Renderer/Private/GPUScene.cpp b/Engine/Source/Runtime/Renderer/Private/GPUScene.cpp
--- a/Engine/Source/Runtime/Renderer/Private/GPUScene.cpp
+++ b/Engine/Source/Runtime/Renderer/Private/GPUScene.cpp
@@ -10,7 +10,6 @@
 #include "RHI.h"
 #include "SceneUtils.h"
 #include "ScenePrivate.h"
-#include "UnifiedBuffer.h"
 #include "SpriteIndexBuffer.h"
 #include "PostProcess/SceneFilterRendering.h"
 #include "ClearQuad.h"
@@ -28,6 +27,7 @@
 #include "LightSceneProxy.h"
 #include "SystemTextures.h"
 #include "SceneDefinitions.h"
+#include "PrimitiveSceneShaderData.h"
 
 #define LOG_INSTANCE_ALLOCATIONS 0
 
diff --git a/Engine/Source/Runtime/Renderer/Private/GPUScene.h b/Engine/Source/Runtime/Renderer/Private/GPUScene.h
--- a/Engine/Source/Runtime/Renderer/Private/GPUScene.h
+++ b/Engine/Source/Runtime/Renderer/Private/GPUScene.h
@@ -11,6 +11,7 @@
 #include "MeshBatch.h"
 #include "LightSceneData.h"
 #include "SceneUniformBuffer.h"
+#include "UnifiedBuffer.h"
 
 class FRDGExternalAccessQueue;
 class FRHICommandList;
diff --git a/Engine/Source/Runtime/Renderer/Private/Renderer.cpp b/Engine/Source/Runtime/Renderer/Private/Renderer.cpp
--- a/Engine/Source/Runtime/Renderer/Private/Renderer.cpp
+++ b/Engine/Source/Runtime/Renderer/Private/Renderer.cpp
@@ -42,6 +42,7 @@
 #include "SceneRenderTargetParameters.h"
 #include "EngineModule.h"
 #include "RendererInterface.h"
+#include "PrimitiveSceneShaderData.h"
 
 DEFINE_LOG_CATEGORY(LogRenderer);
 
diff --git a/Engine/Source/Runtime/Renderer/Private/SceneRendering.cpp b/Engine/Source/Runtime/Renderer/Private/SceneRendering.cpp
--- a/Engine/Source/Runtime/Renderer/Private/SceneRendering.cpp
+++ b/Engine/Source/Runtime/Renderer/Private/SceneRendering.cpp
@@ -88,6 +88,7 @@
 #include "SkyAtmosphereRendering.h"
 #include "VolumetricCloudRendering.h"
 #include "VolumetricFog.h"
+#include "PrimitiveSceneShaderData.h"
 
 /*-----------------------------------------------------------------------------
 	Globals
