Fixed: Additional errors in Cook ShooterGame Android

diff --git a/Engine/Source/Runtime/Engine/Private/ShaderCompiler/ShaderGenerationUtil.cpp b/Engine/Source/Runtime/Engine/Private/ShaderCompiler/ShaderGenerationUtil.cpp
--- a/Engine/Source/Runtime/Engine/Private/ShaderCompiler/ShaderGenerationUtil.cpp
+++ b/Engine/Source/Runtime/Engine/Private/ShaderCompiler/ShaderGenerationUtil.cpp
@@ -649,6 +649,8 @@ static int32 GetBufferNumBits(EGBufferType Format, int32 Channel)
 	case GBT_Float_16_16_16_16:
 		Ret = 16;
 		break;
+	case GBT_Float_32:
+		Ret = 32;
 	default:
 		check(0);
 		break;
diff --git a/Engine/Source/Runtime/RenderCore/Private/GBufferInfo.cpp b/Engine/Source/Runtime/RenderCore/Private/GBufferInfo.cpp
--- a/Engine/Source/Runtime/RenderCore/Private/GBufferInfo.cpp
+++ b/Engine/Source/Runtime/RenderCore/Private/GBufferInfo.cpp
@@ -580,7 +580,7 @@ FGBufferInfo RENDERCORE_API FetchMobileGBufferInfo(const FGBufferParams& Params)
 FGBufferInfo RENDERCORE_API FetchFullGBufferInfo(const FGBufferParams& Params)
 {
 	// For now, we are only doing legacy. But next, we will have a switch between the old and new formats.
-	if (IsMobilePlatform(Params.ShaderPlatform))
+	if (IsMobilePlatform(Params.ShaderPlatform) && IsMobileDeferredShadingEnabled(Params.ShaderPlatform))
 	{
 		return FetchMobileGBufferInfo(Params);
 	}
