Add entry point name to material shader's debug info for better validation messaging.

- This gets printed in the logs by RHI validation so the entrypoint helps figuring it out without the debugger attached.

diff --git a/Engine/Source/Runtime/RenderCore/Public/ShaderCompilerCore.h b/Engine/Source/Runtime/RenderCore/Public/ShaderCompilerCore.h
--- a/Engine/Source/Runtime/RenderCore/Public/ShaderCompilerCore.h
+++ b/Engine/Source/Runtime/RenderCore/Public/ShaderCompilerCore.h
@@ -289,8 +289,7 @@ struct FShaderCompilerInput
 		}
 		else
 		{
-			// we skip EntryPointName as it's usually not useful
-			Name = DebugGroupName + TEXT(":") + VirtualSourceFilePath;
+			Name = DebugGroupName + TEXT(":") + VirtualSourceFilePath + TEXT("|") + EntryPointName;
 		}
 
 		return Name;
