Fixed case where mobile deferred shading was not active in Editor mobile preview

diff --git a/Engine/Source/Runtime/RenderCore/Private/RenderUtils.cpp b/Engine/Source/Runtime/RenderCore/Private/RenderUtils.cpp
--- a/Engine/Source/Runtime/RenderCore/Private/RenderUtils.cpp
+++ b/Engine/Source/Runtime/RenderCore/Private/RenderUtils.cpp
@@ -481,11 +481,9 @@ RENDERCORE_API bool IsMobileDeferredShadingEnabled(const FStaticShaderPlatform P
 
 	// OpenGL requires DXC for deferred shading
 	bool bSupportedPlatform = !IsOpenGLPlatform(Platform) ||
-		(MobileAllowDeferredShadingOpenGL && MobileAllowDeferredShadingOpenGL->GetValueOnAnyThread() != 0);
+		(MobileAllowDeferredShadingOpenGL && MobileAllowDeferredShadingOpenGL->GetValueOnAnyThread() != 0 && IsDxcEnabledForPlatform(Platform));
 
-	return
-		MobileShadingPathIniValue.Get(Platform) == 1 &&
-		bSupportedPlatform && IsDxcEnabledForPlatform(Platform);
+	return MobileShadingPathIniValue.Get(Platform) == 1 && bSupportedPlatform;
 }
 
 RENDERCORE_API bool MobileRequiresSceneDepthAux(const FStaticShaderPlatform Platform)
