[shaders]
- add new IShaderFormat API for separate preprocessing and compilation; backends can implement one or the other depending on the return value of SupportsIndependentPreprocessing
- add support for executing preprocessing in the cook process prior to job submission and constructing job input hashes based on preprocessed source (and a subset of the environment used as compile inputs). controlled by a cvar for now and disabled by default
- add a BaseShaderFormat class in ShaderCompilerCommon which implements common behaviour for output of debug data - note this function is only called for formats which support independent preprocessing, so is expected to be used only by formats which have been converted to use this API
- add new cvars for output of some additional shader debug data - 1. a txt file containing the input hash a.k.a. job cache key 2. a text file containing all diagnostic messages (errors and warnings) for the job
- minor change to how input hashes are constructed for pipeline jobs - sum hashes as 256-bit ints instead of adding to a buffer and re-hashing. faster and simpler, and also more collision resistant (sum of two well distributed hashes equally well distributed)

diff --git a/Engine/Source/Developer/ShaderCompilerCommon/Private/ShaderCompilerCommon.cpp b/Engine/Source/Developer/ShaderCompilerCommon/Private/ShaderCompilerCommon.cpp
--- a/Engine/Source/Developer/ShaderCompilerCommon/Private/ShaderCompilerCommon.cpp
+++ b/Engine/Source/Developer/ShaderCompilerCommon/Private/ShaderCompilerCommon.cpp
@@ -9,6 +9,7 @@
 #include "HlslccDefinitions.h"
 #include "HAL/FileManager.h"
 #include "String/RemoveFrom.h"
+#include "ShaderPreprocessTypes.h"
 #include "ShaderSymbolExport.h"
 #include "ShaderMinifier.h"
 
@@ -1388,6 +1389,32 @@ void CompileOfflineMali(const FShaderCompilerInput& Input, FShaderCompilerOutput
 
 namespace UE::ShaderCompilerCommon
 {
+	void BaseShaderFormat::OutputDebugData(const FString& InputHash, const FShaderCompilerInput& Input, const FShaderPreprocessOutput& PreprocessOutput, const FShaderCompilerOutput& Output) const
+	{
+		DumpDebugShaderData(Input, PreprocessOutput.GetOriginalSource());
+		FString OutputHashFileName = FPaths::Combine(Input.DumpDebugInfoPath, TEXT("OutputHash.txt"));
+		FFileHelper::SaveStringToFile(Output.OutputHash.ToString(), *OutputHashFileName, FFileHelper::EEncodingOptions::ForceAnsi);
+
+		if (EnumHasAnyFlags(Input.DebugInfoFlags, EShaderDebugInfoFlags::Diagnostics))
+		{
+			FString Merged;
+			for (const FShaderCompilerError& Diag : Output.Errors)
+			{
+				Merged += Diag.GetErrorStringWithLineMarker() + "\n";
+			}
+			if (!Merged.IsEmpty())
+			{
+				FFileHelper::SaveStringToFile(Merged, *(Input.DumpDebugInfoPath / TEXT("Diagnostics.txt")));
+			}
+		}
+
+		if (EnumHasAnyFlags(Input.DebugInfoFlags, EShaderDebugInfoFlags::InputHash))
+		{
+			FString InputHashFileName = FPaths::Combine(Input.DumpDebugInfoPath, TEXT("InputHash.txt"));
+			FFileHelper::SaveStringToFile(InputHash, *InputHashFileName, FFileHelper::EEncodingOptions::ForceAnsi);
+		}
+	}
+
 	void DumpDebugShaderData(const FShaderCompilerInput& Input, const FString& PreprocessedSource, const FDebugShaderDataOptions& Options)
 	{
 		if (Input.DumpDebugInfoEnabled())
diff --git a/Engine/Source/Developer/ShaderCompilerCommon/Public/ShaderCompilerCommon.h b/Engine/Source/Developer/ShaderCompilerCommon/Public/ShaderCompilerCommon.h
--- a/Engine/Source/Developer/ShaderCompilerCommon/Public/ShaderCompilerCommon.h
+++ b/Engine/Source/Developer/ShaderCompilerCommon/Public/ShaderCompilerCommon.h
@@ -11,7 +11,7 @@
 #include "ShaderParameterParser.h"
 #endif
 #include "Templates/Function.h"
-
+#include "Interfaces/IShaderFormat.h"
 
 class FShaderParameterParser;
 
@@ -61,6 +61,12 @@ extern SHADERCOMPILERCOMMON_API int16 GetNumUniformBuffersUsed(const FShaderComp
 
 namespace UE::ShaderCompilerCommon
 {
+	class FBaseShaderFormat : public IShaderFormat
+	{
+	public:
+		virtual SHADERCOMPILERCOMMON_API void OutputDebugData(const FString& InputHash, const FShaderCompilerInput& Input, const FShaderPreprocessOutput& PreprocessOutput, const FShaderCompilerOutput& Output) const override;
+	};
+
 	UE_DEPRECATED(5.3, "Deprecated; use FShaderCompilerInput::ShouldUseStableContantBuffer directly")
 	inline bool ShouldUseStableConstantBuffer(const FShaderCompilerInput& Input)
 	{
diff --git a/Engine/Source/Developer/ShaderCompilerCommon/ShaderCompilerCommon.Build.cs b/Engine/Source/Developer/ShaderCompilerCommon/ShaderCompilerCommon.Build.cs
--- a/Engine/Source/Developer/ShaderCompilerCommon/ShaderCompilerCommon.Build.cs
+++ b/Engine/Source/Developer/ShaderCompilerCommon/ShaderCompilerCommon.Build.cs
@@ -11,6 +11,7 @@ public class ShaderCompilerCommon : ModuleRules
 			new string[] {
 				"Core",
 				"RenderCore",
+				"TargetPlatform",
 			}
 			);
 		if (Target.bCompileAgainstEngine)
diff --git a/Engine/Source/Developer/ShaderPreprocessor/Private/PreprocessorPrivate.h b/Engine/Source/Developer/ShaderPreprocessor/Private/PreprocessorPrivate.h
--- a/Engine/Source/Developer/ShaderPreprocessor/Private/PreprocessorPrivate.h
+++ b/Engine/Source/Developer/ShaderPreprocessor/Private/PreprocessorPrivate.h
@@ -6,6 +6,7 @@
 #include "ShaderCore.h"
 #include "mcpp.h"
 #include "ShaderCompilerCore.h"
+#include "ShaderPreprocessTypes.h"
 
 enum class EMessageType
 {
@@ -61,7 +62,7 @@ static void ExtractDirective(FString& OutString, FString WarningString)
  * @param ShaderOutput - Shader output to which to add errors.
  * @param McppErrors - MCPP error output.
  */
-static bool ParseMcppErrors(TArray<FShaderCompilerError>& OutErrors, TArray<FString>& OutStrings, const FString& McppErrors)
+static bool ParseMcppErrors(FShaderPreprocessOutput& Output, const FString& McppErrors)
 {
 	bool bSuccess = true;
 	if (McppErrors.Len() > 0)
@@ -70,7 +71,7 @@ static bool ParseMcppErrors(TArray<FShaderCompilerError>& OutErrors, TArray<FStr
 		McppErrors.ParseIntoArray(Lines, TEXT("\n"), true);
 		for (int32 LineIndex = 0; LineIndex < Lines.Num(); ++LineIndex)
 		{
-			const FString& Line = Lines[LineIndex];
+			FString& Line = Lines[LineIndex];
 			int32 SepIndex1 = Line.Find(TEXT(":"), ESearchCase::CaseSensitive, ESearchDir::FromStart, 2);
 			int32 SepIndex2 = Line.Find(TEXT(":"), ESearchCase::CaseSensitive, ESearchDir::FromStart, SepIndex1 + 1);
 			if (SepIndex1 != INDEX_NONE && SepIndex2 != INDEX_NONE && SepIndex1 < SepIndex2)
@@ -92,10 +93,7 @@ static bool ParseMcppErrors(TArray<FShaderCompilerError>& OutErrors, TArray<FStr
 					{
 						case EMessageType::Error:
 						{
-							FShaderCompilerError* CompilerError = new(OutErrors) FShaderCompilerError;
-							CompilerError->ErrorVirtualFilePath = Filename;
-							CompilerError->ErrorLineString = LineNumStr;
-							CompilerError->StrippedErrorMessage = Message;
+							Output.LogError(MoveTemp(Filename), MoveTemp(Message), MoveTemp(LineNumStr));
 							bSuccess = false;
 						}
 							break;
@@ -108,7 +106,7 @@ static bool ParseMcppErrors(TArray<FShaderCompilerError>& OutErrors, TArray<FStr
 						{
 							FString Directive;
 							ExtractDirective(Directive, Message);
-							OutStrings.Add(Directive);
+							Output.AddDirective(MoveTemp(Directive));
 						}
 							break;
 						default:
@@ -118,16 +116,14 @@ static bool ParseMcppErrors(TArray<FShaderCompilerError>& OutErrors, TArray<FStr
 				else
 				{
 					// Presume message is an error
-					FShaderCompilerError* CompilerError = new(OutErrors) FShaderCompilerError;
-					CompilerError->StrippedErrorMessage = Line;
+					Output.LogError(MoveTemp(Line));
 					bSuccess = false;
 				}
 			}
 			else
 			{
 				// Presume message is an error
-				FShaderCompilerError* CompilerError = new(OutErrors) FShaderCompilerError;
-				CompilerError->StrippedErrorMessage = Line;
+				Output.LogError(MoveTemp(Line));
 				bSuccess = false;
 			}
 		}
diff --git a/Engine/Source/Developer/ShaderPreprocessor/Private/ShaderPreprocessor.cpp b/Engine/Source/Developer/ShaderPreprocessor/Private/ShaderPreprocessor.cpp
--- a/Engine/Source/Developer/ShaderPreprocessor/Private/ShaderPreprocessor.cpp
+++ b/Engine/Source/Developer/ShaderPreprocessor/Private/ShaderPreprocessor.cpp
@@ -15,14 +15,11 @@ namespace
 {
 	const FString PlatformHeader = TEXT("/Engine/Public/Platform.ush");
 	const FString PlatformHeaderLowerCase = PlatformHeader.ToLower();
-	void LogMandatoryHeaderError(const FShaderCompilerInput& Input, FShaderCompilerOutput& Output)
+	void LogMandatoryHeaderError(const FShaderCompilerInput& Input, FShaderPreprocessOutput& Output)
 	{
-		FShaderCompilerError Error;
-		Error.ErrorVirtualFilePath = Input.VirtualSourceFilePath;
-		Error.ErrorLineString = TEXT("1");
-		Error.StrippedErrorMessage = FString::Printf(TEXT("Error: Shader is required to include %s"), *PlatformHeader);
-
-		Output.Errors.Add(Error);
+		FString Path = Input.VirtualSourceFilePath;
+		FString Message = FString::Printf(TEXT("Error: Shader is required to include %s"), *PlatformHeader);
+		Output.LogError(MoveTemp(Path), MoveTemp(Message), 1);
 	}
 }
 
@@ -48,9 +45,9 @@ class FMcppFileLoader
 {
 public:
 	/** Initialization constructor. */
-	explicit FMcppFileLoader(const FShaderCompilerInput& InShaderInput, FShaderCompilerOutput& InShaderOutput)
-		: ShaderInput(InShaderInput)
-		, ShaderOutput(InShaderOutput)
+	explicit FMcppFileLoader(const FShaderCompilerInput& InShaderInput, FShaderPreprocessOutput& InShaderOutput)
+		: Input(InShaderInput)
+		, Output(InShaderOutput)
 	{
 		FString InputShaderSource;
 		if (LoadShaderSourceFile(*InShaderInput.VirtualSourceFilePath, InShaderInput.Target.GetPlatform(),  &InputShaderSource, nullptr, &InShaderInput.ShaderPlatformName))
@@ -87,10 +84,10 @@ private:
 		FString VirtualFilePath = UTF8Converter.Get();
 
 		// Substitute virtual platform path here to make sure that #line directives refer to the platform-specific file.
-		ReplaceVirtualFilePathForShaderPlatform(VirtualFilePath, This->ShaderInput.Target.GetPlatform());
+		ReplaceVirtualFilePathForShaderPlatform(VirtualFilePath, This->Input.Target.GetPlatform());
 
 		// Fixup autogen file
-		ReplaceVirtualFilePathForShaderAutogen(VirtualFilePath, This->ShaderInput.Target.GetPlatform(), &This->ShaderInput.ShaderPlatformName);
+		ReplaceVirtualFilePathForShaderAutogen(VirtualFilePath, This->Input.Target.GetPlatform(), &This->Input.ShaderPlatformName);
 
 		// Collapse any relative directories to allow #include "../MyFile.ush"
 		FPaths::CollapseRelativeDirectories(VirtualFilePath);
@@ -100,19 +97,20 @@ private:
 		{
 			FString FileContents;
 
-			if (This->ShaderInput.Environment.IncludeVirtualPathToContentsMap.Contains(VirtualFilePath))
+			if (This->Input.Environment.IncludeVirtualPathToContentsMap.Contains(VirtualFilePath))
 			{
-				FileContents = This->ShaderInput.Environment.IncludeVirtualPathToContentsMap.FindRef(VirtualFilePath);
+				FileContents = This->Input.Environment.IncludeVirtualPathToContentsMap.FindRef(VirtualFilePath);
 			}
-			else if (This->ShaderInput.Environment.IncludeVirtualPathToExternalContentsMap.Contains(VirtualFilePath))
+			else if (This->Input.Environment.IncludeVirtualPathToExternalContentsMap.Contains(VirtualFilePath))
 			{
-				FileContents = *This->ShaderInput.Environment.IncludeVirtualPathToExternalContentsMap.FindRef(VirtualFilePath);
+				FileContents = *This->Input.Environment.IncludeVirtualPathToExternalContentsMap.FindRef(VirtualFilePath);
 			}
 			else
 			{
-				CheckShaderHashCacheInclude(VirtualFilePath, This->ShaderInput.Target.GetPlatform(), This->ShaderInput.ShaderFormat.ToString());
+				CheckShaderHashCacheInclude(VirtualFilePath, This->Input.Target.GetPlatform(), This->Input.ShaderFormat.ToString());
 
-				LoadShaderSourceFile(*VirtualFilePath, This->ShaderInput.Target.GetPlatform(), &FileContents, &This->ShaderOutput.Errors, &This->ShaderInput.ShaderPlatformName);
+				TArray<FShaderCompilerError>& OutErrors = This->Output.EditErrors();
+				LoadShaderSourceFile(*VirtualFilePath, This->Input.Target.GetPlatform(), &FileContents, &OutErrors, &This->Input.ShaderPlatformName);
 			}
 
 			if (FileContents.Len() > 0)
@@ -138,9 +136,9 @@ private:
 	}
 
 	/** Shader input data. */
-	const FShaderCompilerInput& ShaderInput;
+	const FShaderCompilerInput& Input;
 	/** Shader output data. */
-	FShaderCompilerOutput& ShaderOutput;
+	FShaderPreprocessOutput& Output;
 	/** File contents are cached as needed. */
 	TMap<FString,FShaderContents> CachedFileContents;
 };
@@ -215,10 +213,9 @@ static void AddStbDefines(stb_arena* MacroArena, macro_definition**& StbDefines,
 class FShaderPreprocessorUtilities
 {
 public:
-	static void DumpShaderDefinesAsCommentedCode(const FShaderCompilerInput& ShaderInput, FString* OutDefines)
+	static void DumpShaderDefinesAsCommentedCode(const FShaderCompilerEnvironment& Environment, FString* OutDefines)
 	{
-		const TMap<FString, FString>& Definitions = ShaderInput.Environment.Definitions.GetDefinitionMap();
-
+		const TMap<FString, FString>& Definitions = Environment.Definitions.GetDefinitionMap();
 		TArray<FString> Keys;
 		Definitions.GetKeys(/* out */ Keys);
 		Keys.Sort();
@@ -229,7 +226,7 @@ public:
 			Defines += FString::Printf(TEXT("// #define %s %s\n"), *Key, *Definitions[Key]);
 		}
 
-		*OutDefines = MakeInjectedShaderCodeBlock(TEXT("DumpShaderDefinesAsCommentedCode"), Defines);
+		*OutDefines += MakeInjectedShaderCodeBlock(TEXT("DumpShaderDefinesAsCommentedCode"), Defines);
 	}
 
 	static void PopulateDefinesMcpp(const FShaderCompilerInput& Input, const FShaderCompilerDefinitions& AdditionalDefines, TArray<TArray<ANSICHAR>>& OutDefines)
@@ -238,9 +235,9 @@ public:
 		AddMcppDefines(OutDefines, AdditionalDefines.GetDefinitionMap());
 	}
 
-	static void PopulateDefinesStb(const FShaderCompilerInput& Input, const FShaderCompilerDefinitions& AdditionalDefines, stb_arena* MacroArena, macro_definition**& OutDefines)
+	static void PopulateDefinesStb(const FShaderCompilerEnvironment& Environment, const FShaderCompilerDefinitions& AdditionalDefines, stb_arena* MacroArena, macro_definition**& OutDefines)
 	{
-		AddStbDefines(MacroArena, OutDefines, Input.Environment.Definitions.GetDefinitionMap());
+		AddStbDefines(MacroArena, OutDefines, Environment.Definitions.GetDefinitionMap());
 		AddStbDefines(MacroArena, OutDefines, AdditionalDefines.GetDefinitionMap());
 		AddStbDefine(MacroArena, OutDefines, TEXT("_STB_PREPROCESS"), TEXT("1"));
 	}
@@ -249,9 +246,8 @@ public:
 //////////////////////////////////////////////////////////////////////////
 
 bool InnerPreprocessShaderMcpp(
-	FString& OutPreprocessedShader,
-	FShaderCompilerOutput& ShaderOutput,
-	const FShaderCompilerInput& ShaderInput,
+	FShaderPreprocessOutput& Output,
+	const FShaderCompilerInput& Input,
 	const FShaderCompilerDefinitions& AdditionalDefines)
 {
 	int32 McppResult = 0;
@@ -261,10 +257,10 @@ bool InnerPreprocessShaderMcpp(
 
 	bool bHasIncludedMandatoryHeaders = false;
 	{
-		FMcppFileLoader FileLoader(ShaderInput, ShaderOutput);
+		FMcppFileLoader FileLoader(Input, Output);
 
 		TArray<TArray<ANSICHAR>> McppOptions;
-		FShaderPreprocessorUtilities::PopulateDefinesMcpp(ShaderInput, AdditionalDefines, McppOptions);
+		FShaderPreprocessorUtilities::PopulateDefinesMcpp(Input, AdditionalDefines, McppOptions);
 
 		// MCPP is not threadsafe.
 
@@ -296,7 +292,7 @@ bool InnerPreprocessShaderMcpp(
 			McppResult = mcpp_run(
 				McppOptionsANSI.GetData(),
 				McppOptionsANSI.Num(),
-				TCHAR_TO_ANSI(*ShaderInput.VirtualSourceFilePath),
+				TCHAR_TO_ANSI(*Input.VirtualSourceFilePath),
 				&McppOutAnsi,
 				&McppErrAnsi,
 				FileLoader.GetMcppInterface()
@@ -309,7 +305,7 @@ bool InnerPreprocessShaderMcpp(
 		bHasIncludedMandatoryHeaders = FileLoader.HasIncludedMandatoryHeaders();
 	}
 
-	if (!ParseMcppErrors(ShaderOutput.Errors, ShaderOutput.PragmaDirectives, McppErrors))
+	if (!ParseMcppErrors(Output, McppErrors))
 	{
 		return false;
 	}
@@ -317,20 +313,19 @@ bool InnerPreprocessShaderMcpp(
 	// Report unhandled mcpp failure that didn't generate any errors
 	if (McppResult != 0)
 	{
-		FShaderCompilerError* CompilerError = new(ShaderOutput.Errors) FShaderCompilerError;
-		CompilerError->ErrorVirtualFilePath = ShaderInput.VirtualSourceFilePath;
-		CompilerError->ErrorLineString = TEXT("0");
-		CompilerError->StrippedErrorMessage = FString::Printf(TEXT("PreprocessShader mcpp_run failed with error code %d"), McppResult);
+		FString Path = Input.VirtualSourceFilePath;
+		FString Message = FString::Printf(TEXT("PreprocessShader mcpp_run failed with error code %d"), McppResult);
+		Output.LogError(MoveTemp(Path), MoveTemp(Message), 0);
 		return false;
 	}
 
 	if (!bHasIncludedMandatoryHeaders)
 	{
-		LogMandatoryHeaderError(ShaderInput, ShaderOutput);
+		LogMandatoryHeaderError(Input, Output);
 		return false;
 	}
 
-	OutPreprocessedShader += McppOutput;
+	Output.EditSource() += McppOutput;
 
 	return true;
 }
@@ -370,6 +365,7 @@ extern "C"
 struct FStbPreprocessContext
 {
 	const FShaderCompilerInput& ShaderInput;
+	const FShaderCompilerEnvironment& Environment;
 	TMap<FString, TArray<ANSICHAR>> LoadedIncludesCache;
 	TMap<FString, TUniquePtr<ANSICHAR[]>> SeenPathsLowerCase;
 
@@ -467,6 +463,17 @@ void ConvertAndStripComments(const FString& ShaderSource, TArray<ANSICHAR>& OutS
 	OutStripped.SetNum(CurrentOut - OutStripped.GetData(), /* bAllowShrinking */false);
 }
 
+const FString* FindInMemorySource(const FShaderCompilerEnvironment& Environment, const FString& FilenameConverted)
+{
+	const FString* InMemorySource = Environment.IncludeVirtualPathToContentsMap.Find(FilenameConverted);
+	if (!InMemorySource)
+	{
+		const FThreadSafeSharedStringPtr* SharedPtr = Environment.IncludeVirtualPathToExternalContentsMap.Find(FilenameConverted);
+		InMemorySource = SharedPtr ? SharedPtr->Get() : nullptr;
+	}
+	return InMemorySource;
+}
+
 static const ANSICHAR* StbLoadFile(const ANSICHAR* Filename, void* RawContext, size_t* OutLength)
 {
 	FStbPreprocessContext& Context = *reinterpret_cast<FStbPreprocessContext*>(RawContext);
@@ -474,24 +481,21 @@ static const ANSICHAR* StbLoadFile(const ANSICHAR* Filename, void* RawContext, s
 	TArray<ANSICHAR>* ContentsCached = Context.LoadedIncludesCache.Find(FilenameConverted);
 	if (!ContentsCached)
 	{
-		FString ShaderSource;
+		ContentsCached = &Context.LoadedIncludesCache.Add(FilenameConverted);
+		// Local FString used for the LoadShaderSourceFile path; we should consider retrieving source from the shader file cache as a reference
+		// (avoid an extra alloc+copy)
+		FString SourceCopy;
 
-		if (Context.ShaderInput.Environment.IncludeVirtualPathToContentsMap.Contains(FilenameConverted))
-		{
-			ShaderSource = Context.ShaderInput.Environment.IncludeVirtualPathToContentsMap.FindRef(FilenameConverted);
-		}
-		else if (Context.ShaderInput.Environment.IncludeVirtualPathToExternalContentsMap.Contains(FilenameConverted))
-		{
-			ShaderSource = *Context.ShaderInput.Environment.IncludeVirtualPathToExternalContentsMap.FindRef(FilenameConverted);
-		}
-		else
+		const FString* InMemorySource = FindInMemorySource(Context.Environment, FilenameConverted);
+
+		if (!InMemorySource)
 		{
 			CheckShaderHashCacheInclude(FilenameConverted, Context.ShaderInput.Target.GetPlatform(), Context.ShaderInput.ShaderFormat.ToString());
-			LoadShaderSourceFile(*FilenameConverted, Context.ShaderInput.Target.GetPlatform(), &ShaderSource, nullptr);
+			LoadShaderSourceFile(*FilenameConverted, Context.ShaderInput.Target.GetPlatform(), &SourceCopy, nullptr);
+			InMemorySource = &SourceCopy;
 		}
-		check(!ShaderSource.IsEmpty());
-		ContentsCached = &Context.LoadedIncludesCache.Add(FilenameConverted);
-		ConvertAndStripComments(ShaderSource, *ContentsCached);
+		check(InMemorySource && !InMemorySource->IsEmpty());
+		ConvertAndStripComments(*InMemorySource, *ContentsCached);
 	}
 	check(ContentsCached);
 	*OutLength = ContentsCached->Num();
@@ -533,8 +537,8 @@ static const ANSICHAR* StbResolveInclude(const ANSICHAR* PathInSource, uint32 Pa
 	}
 
 	bool bExists =
-		Context.ShaderInput.Environment.IncludeVirtualPathToContentsMap.Contains(PathModified) ||
-		Context.ShaderInput.Environment.IncludeVirtualPathToExternalContentsMap.Contains(PathModified) ||
+		Context.Environment.IncludeVirtualPathToContentsMap.Contains(PathModified) ||
+		Context.Environment.IncludeVirtualPathToExternalContentsMap.Contains(PathModified) ||
 		// LoadShaderSourceFile will load the file if it exists, but then cache it internally, so the next call in StbLoadFile will be cheap
 		// (and hence this is not overly wasteful)
 		LoadShaderSourceFile(*PathModified, Context.ShaderInput.Target.GetPlatform(), nullptr, nullptr);
@@ -580,19 +584,19 @@ static void AddStbDefines(stb_arena* MacroArena, macro_definition**& StbDefines,
 }
 
 bool InnerPreprocessShaderStb(
-	FString& OutPreprocessedShader,
-	FShaderCompilerOutput& ShaderOutput,
-	const FShaderCompilerInput& ShaderInput,
+	FShaderPreprocessOutput& Output,
+	const FShaderCompilerInput& Input,
+	const FShaderCompilerEnvironment& Environment,
 	const FShaderCompilerDefinitions& AdditionalDefines
 )
 {
 	stb_arena MacroArena = { 0 };
 	macro_definition** StbDefines = nullptr;
-	FShaderPreprocessorUtilities::PopulateDefinesStb(ShaderInput, AdditionalDefines, &MacroArena, StbDefines);
+	FShaderPreprocessorUtilities::PopulateDefinesStb(Environment, AdditionalDefines, &MacroArena, StbDefines);
 
-	FStbPreprocessContext Context{ ShaderInput };
+	FStbPreprocessContext Context{ Input, Environment };
 
-	auto InFilename = StringCast<ANSICHAR>(*ShaderInput.VirtualSourceFilePath);
+	auto InFilename = StringCast<ANSICHAR>(*Input.VirtualSourceFilePath);
 	int NumDiagnostics = 0;
 	pp_diagnostic* Diagnostics = nullptr;
 	
@@ -609,10 +613,8 @@ bool InnerPreprocessShaderStb(
 			// as we do with MCPP, we are ignoring warnings (for now?)
 			if (Diagnostic->error_level == PP_RESULT_MODE_error)
 			{
-				FShaderCompilerError* CompilerError = new(ShaderOutput.Errors) FShaderCompilerError;
-				CompilerError->ErrorVirtualFilePath = Diagnostic->where->filename;
-				CompilerError->ErrorLineString = FString::Printf(TEXT("%d"), Diagnostic->where->line_number);
-				CompilerError->StrippedErrorMessage = Message;
+				FString Filename = Diagnostic->where->filename;
+				Output.LogError(MoveTemp(Filename), MoveTemp(Message), Diagnostic->where->line_number);
 			}
 			else
 			{
@@ -621,17 +623,20 @@ bool InnerPreprocessShaderStb(
 				{
 					FString Directive;
 					ExtractDirective(Directive, Message);
-					ShaderOutput.PragmaDirectives.Add(Directive);
+					Output.AddDirective(MoveTemp(Directive));
 				}
 			}
 		}
 	}
 
-	OutPreprocessedShader = StringCast<TCHAR>(OutPreprocessedAnsi).Get();
+	if (!HasError)
+	{
+		Output.EditSource().Append(OutPreprocessedAnsi);
+	}
 
 	if (!HasError && !Context.HasIncludedMandatoryHeaders())
 	{
-		LogMandatoryHeaderError(ShaderInput, ShaderOutput);
+		LogMandatoryHeaderError(Input, Output);
 		HasError = true;
 	}
 
@@ -642,6 +647,28 @@ bool InnerPreprocessShaderStb(
 	return !HasError;
 }
 
+bool PreprocessShader(
+	FString& OutPreprocessedShader,
+	FShaderCompilerOutput& ShaderOutput,
+	const FShaderCompilerInput& ShaderInput,
+	const FShaderCompilerDefinitions& AdditionalDefines,
+	EDumpShaderDefines DefinesPolicy)
+{
+	FShaderPreprocessOutput Output;
+	// when called via this overload, environment is assumed to be already merged in input struct
+	const FShaderCompilerEnvironment& Environment = ShaderInput.Environment;
+	bool bSucceeded = PreprocessShader(Output, ShaderInput, Environment, AdditionalDefines, DefinesPolicy);
+
+	OutPreprocessedShader = MoveTemp(Output.EditSource());
+
+	Output.MoveDirectives(ShaderOutput.PragmaDirectives);
+	for (FShaderCompilerError& Error : Output.EditErrors())
+	{
+		ShaderOutput.Errors.Add(MoveTemp(Error));
+	}
+	return bSucceeded;
+}
+
 /**
  * Preprocess a shader.
  * @param OutPreprocessedShader - Upon return contains the preprocessed source code.
@@ -652,9 +679,9 @@ bool InnerPreprocessShaderStb(
  * @returns true if the shader is preprocessed without error.
  */
 bool PreprocessShader(
-	FString& OutPreprocessedShader,
-	FShaderCompilerOutput& ShaderOutput,
-	const FShaderCompilerInput& ShaderInput,
+	FShaderPreprocessOutput& Output,
+	const FShaderCompilerInput& Input,
+	const FShaderCompilerEnvironment& Environment,
 	const FShaderCompilerDefinitions& AdditionalDefines,
 	EDumpShaderDefines DefinesPolicy
 )
@@ -662,33 +689,34 @@ bool PreprocessShader(
 	TRACE_CPUPROFILER_EVENT_SCOPE(PreprocessShader);
 
 	// Skip the cache system and directly load the file path (used for debugging)
-	if (ShaderInput.bSkipPreprocessedCache)
+	if (Input.bSkipPreprocessedCache)
 	{
-		return FFileHelper::LoadFileToString(OutPreprocessedShader, *ShaderInput.VirtualSourceFilePath);
+		return FFileHelper::LoadFileToString(Output.EditSource(), *Input.VirtualSourceFilePath);
 	}
 	else
 	{
-		check(CheckVirtualShaderFilePath(ShaderInput.VirtualSourceFilePath));
+		check(CheckVirtualShaderFilePath(Input.VirtualSourceFilePath));
+	}
+
+	// List the defines used for compilation in the preprocessed shaders, especially to know which permutation vector this shader is.
+	if (DefinesPolicy == EDumpShaderDefines::AlwaysIncludeDefines || (DefinesPolicy == EDumpShaderDefines::DontCare && Input.DumpDebugInfoPath.Len() > 0))
+	{
+		FShaderPreprocessorUtilities::DumpShaderDefinesAsCommentedCode(Environment, &Output.EditSource());
 	}
 
 	bool bResult = false;
-	bool bLegacyPreprocess = ShaderInput.Environment.CompilerFlags.Contains(CFLAG_UseLegacyPreprocessor);
-	FString PreprocessorOutput;
+	bool bLegacyPreprocess = Environment.CompilerFlags.Contains(CFLAG_UseLegacyPreprocessor);
 	if (!bLegacyPreprocess)
 	{
-		bResult |= InnerPreprocessShaderStb(PreprocessorOutput, ShaderOutput, ShaderInput, AdditionalDefines);
+		bResult |= InnerPreprocessShaderStb(Output, Input, Environment, AdditionalDefines);
 	}
 	else
 	{
-		bResult |= InnerPreprocessShaderMcpp(PreprocessorOutput, ShaderOutput, ShaderInput, AdditionalDefines);
-	}
-
-	// List the defines used for compilation in the preprocessed shaders, especially to know witch permutation vector this shader is.
-	if (DefinesPolicy == EDumpShaderDefines::AlwaysIncludeDefines || (DefinesPolicy == EDumpShaderDefines::DontCare && ShaderInput.DumpDebugInfoPath.Len() > 0))
-	{
-		FShaderPreprocessorUtilities::DumpShaderDefinesAsCommentedCode(ShaderInput, &OutPreprocessedShader);
+		// note: ignoring explicit Environment parameter for the MCPP case; it's enforced upstream that
+		// the preprocessed job cache cannot be enabled at the same time as the legacy preprocessor, and
+		// so this environment will always match the one in FShaderCompilerInput
+		bResult |= InnerPreprocessShaderMcpp(Output, Input, AdditionalDefines);
 	}
 
-	OutPreprocessedShader += PreprocessorOutput;
 	return bResult;
 }
diff --git a/Engine/Source/Developer/ShaderPreprocessor/Public/ShaderPreprocessor.h b/Engine/Source/Developer/ShaderPreprocessor/Public/ShaderPreprocessor.h
--- a/Engine/Source/Developer/ShaderPreprocessor/Public/ShaderPreprocessor.h
+++ b/Engine/Source/Developer/ShaderPreprocessor/Public/ShaderPreprocessor.h
@@ -7,6 +7,7 @@
 #include "ShaderCore.h"
 
 class FShaderCompilerDefinitions;
+class FShaderPreprocessOutput;
 class FString;
 struct FShaderCompilerInput;
 struct FShaderCompilerOutput;
@@ -39,3 +40,21 @@ extern SHADERPREPROCESSOR_API bool PreprocessShader(
 	const FShaderCompilerInput& ShaderInput,
 	const FShaderCompilerDefinitions& AdditionalDefines,
 	EDumpShaderDefines DefinesPolicy = EDumpShaderDefines::DontCare);
+
+/**
+ * Preprocess a shader.
+ * @param Output - Preprocess output struct. Source, directives and possibly errors will be populated.
+ * @param Input - The shader compiler input.
+ * @param MergedEnvironment - The result of merging the Environment and SharedEnvironment from the FShaderCompilerInput
+ * (it is assumed this overload is called outside of the worker process which merges this in-place, so this merge step must be
+ * performed by the caller)
+ * @param AdditionalDefines - Additional defines with which to preprocess the shader.
+ * @param bShaderDumpDefinesAsCommentedCode - Whether to add shader definitions as comments.
+ * @returns true if the shader is preprocessed without error.
+ */
+extern SHADERPREPROCESSOR_API bool PreprocessShader(
+	FShaderPreprocessOutput& Output,
+	const FShaderCompilerInput& Input,
+	const FShaderCompilerEnvironment& MergedEnvironment,
+	const FShaderCompilerDefinitions& AdditionalDefines,
+	EDumpShaderDefines DefinesPolicy = EDumpShaderDefines::DontCare);
diff --git a/Engine/Source/Developer/TargetPlatform/Public/Interfaces/IShaderFormat.h b/Engine/Source/Developer/TargetPlatform/Public/Interfaces/IShaderFormat.h
--- a/Engine/Source/Developer/TargetPlatform/Public/Interfaces/IShaderFormat.h
+++ b/Engine/Source/Developer/TargetPlatform/Public/Interfaces/IShaderFormat.h
@@ -17,14 +17,14 @@ class IShaderFormat
 public:
 
 	/** 
-	 * Compile the specified shader.
+	 * Compile the specified shader. Only called if SupportsIndependentPreprocessing returns false.
 	 *
 	 * @param Format The desired format
 	 * @param Input The input to the shader compiler.
 	 * @param Output The output from shader compiler.
 	 * @param WorkingDirectory The working directory.
 	 */
-	virtual void CompileShader( FName Format, const struct FShaderCompilerInput& Input, struct FShaderCompilerOutput& Output,const FString& WorkingDirectory ) const = 0;
+	virtual void CompileShader(FName Format, const struct FShaderCompilerInput& Input, struct FShaderCompilerOutput& Output, const FString& WorkingDirectory) const {};
 
 	/**
 	 * Gets the current version of the specified shader format.
@@ -128,6 +128,33 @@ public:
 	 */
 	virtual bool UsesHLSLcc(const struct FShaderCompilerInput& Input) const { return false; }
 
+	/**
+	 * Executes all shader preprocessing steps; only called if SupportsIndependentPreprocessing returns true
+	 */
+	virtual bool PreprocessShader(const struct FShaderCompilerInput& Input, const struct FShaderCompilerEnvironment& Environment, class FShaderPreprocessOutput& PreprocessOutput) const { return false; };
+
+	/**
+	 * Compile the specified preprocessed shader; only called if SupportsIndependentPreprocessing returns true
+	 */
+	virtual void CompilePreprocessedShader(const struct FShaderCompilerInput& Input, const class FShaderPreprocessOutput& PreprocessOutput, struct FShaderCompilerOutput& Output, const FString& WorkingDirectory) const {}
+
+	/* 
+	 * Output debug info for a single compile job; only called if SupportsIndependentPreprocessing returns true.
+	 * This will be called for all jobs (including those found in the job cache) but only if debug info is enabled for the job.
+	 * Note that any debug info output in CompilePreprocessedShader will only be done for the job that actually executes the
+	 * compile step, as such any debug outputs that are desirable for all jobs should be written by this function.
+	 * A BaseShaderFormat implementation is provided in ShaderCompilerCommon which dumps preprocessed and stripped USFs along
+	 * with the hash of the shader code in an OutputHash.txt file; if no additional custom debug output is required the shader 
+	 * format can inherit from BaseShaderFormat instead of IShaderFormat.
+	 */
+	virtual void OutputDebugData(const FString& InputHash, const struct FShaderCompilerInput& Input, const class FShaderPreprocessOutput& PreprocessOutput, const struct FShaderCompilerOutput& Output) const {};
+
+	/**
+	 * Return true if preprocessing for this format can be executed independent of compilation (i.e. the format has an
+	 * implementation of the PreprocessShader function). 
+	 */
+	virtual bool SupportsIndependentPreprocessing() const { return false; }
+
 public:
 	/** Virtual destructor. */
 	virtual ~IShaderFormat() { }
diff --git a/Engine/Source/Programs/ShaderCompileWorker/Private/ShaderCompileWorker.cpp b/Engine/Source/Programs/ShaderCompileWorker/Private/ShaderCompileWorker.cpp
--- a/Engine/Source/Programs/ShaderCompileWorker/Private/ShaderCompileWorker.cpp
+++ b/Engine/Source/Programs/ShaderCompileWorker/Private/ShaderCompileWorker.cpp
@@ -605,7 +605,7 @@ private:
 			{
 				FShaderCompileJob& Job = *new(OutSingleJobs) FShaderCompileJob;
 				// Deserialize the job's inputs.
-				InputFile << Job.Input;
+				Job.SerializeWorkerInput(InputFile);
 				Job.Input.DeserializeSharedInputs(InputFile, ExternalIncludes, SharedEnvironments, ParameterStructures);
 
 				// SCW doesn't run DDPI, GShaderHasCache Initialize is run  at start with no knowledge of the CustomPlatforms
@@ -649,7 +649,7 @@ private:
 				{
 					// Deserialize the job's inputs.
 					FShaderCompileJob* Job = PipelineJob->StageJobs[StageIndex]->GetSingleShaderJob();
-					InputFile << Job->Input;
+					Job->SerializeWorkerInput(InputFile);
 					Job->Input.DeserializeSharedInputs(InputFile, ExternalIncludes, SharedEnvironments, ParameterStructures);
 
 					// SCW doesn't run DDPI, GShaderHasCache Initialize is run  at start with no knowledge of the CustomPlatforms
diff --git a/Engine/Source/Runtime/Engine/Private/ShaderCompiler/ShaderCompiler.cpp b/Engine/Source/Runtime/Engine/Private/ShaderCompiler/ShaderCompiler.cpp
--- a/Engine/Source/Runtime/Engine/Private/ShaderCompiler/ShaderCompiler.cpp
+++ b/Engine/Source/Runtime/Engine/Private/ShaderCompiler/ShaderCompiler.cpp
@@ -111,6 +111,13 @@ static FAutoConsoleVariableRef CVarShaderCompilerMaxJobCacheMemoryPercent(
 	ECVF_Default
 );
 
+static TAutoConsoleVariable<bool> CVarPreprocessedJobCache(
+	TEXT("r.ShaderCompiler.PreprocessedJobCache"),
+	false,
+	TEXT("EXPERIMENTAL: if enabled will preprocess jobs up-front when the job is queued and generate job input hashes based on preprocessed source."),
+	ECVF_Default
+);
+
 static TAutoConsoleVariable<bool> CVarJobCacheDDC(
 	TEXT("r.ShaderCompiler.JobCacheDDC"),
 	true,
@@ -129,6 +136,18 @@ static TAutoConsoleVariable<bool> CVarDebugDumpWorkerInputs(
 	TEXT("If true, worker input files will be saved for each individual compile job alongside other debug data (note that r.DumpShaderDebugInfo must also be enabled for this to function)"),
 	ECVF_ReadOnly);
 
+static TAutoConsoleVariable<bool> CVarDebugDumpJobInputHashes(
+	TEXT("r.ShaderCompiler.DebugDumpJobInputHashes"),
+	false,
+	TEXT("If true, the job input hash will be dumped alongside other debug data (in InputHash.txt)"),
+	ECVF_ReadOnly);
+
+static TAutoConsoleVariable<bool> CVarDebugDumpJobDiagnostics(
+	TEXT("r.ShaderCompiler.DebugDumpJobDiagnostics"),
+	false,
+	TEXT("If true, all diagnostic messages (errors and warnings) for each shader job will be dumped alongside other debug data (in Diagnostics.txt)"),
+	ECVF_ReadOnly);
+
 static TAutoConsoleVariable<bool> CVarCompileParallelInProcess(
 	TEXT("r.ShaderCompiler.ParallelInProcess"),
 	false,
@@ -582,6 +601,8 @@ void FShaderCompileJobCollection::SubmitJobs(const TArray<FShaderCommonCompileJo
 				check(Job->JobIndex != INDEX_NONE);
 				check(Job->Priority != EShaderCompileJobPriority::None);
 				check(Job->PendingPriority == EShaderCompileJobPriority::None);
+				
+				ConditionalPreprocessShader(Job);
 
 				const int32 PriorityIndex = (int32)Job->Priority;
 				bool bNewJob = true;
@@ -692,6 +713,8 @@ void FShaderCompileJobCollection::HandlePrintStats()
 TRACE_DECLARE_INT_COUNTER(Shaders_Compiled, TEXT("Shaders/Compiled"));
 void FShaderCompileJobCollection::ProcessFinishedJob(FShaderCommonCompileJob* FinishedJob, bool bWasCached)
 {
+	FinishedJob->OnComplete();
+
 	if (!bWasCached)
 	{
 		GShaderCompilerStats->RegisterFinishedJob(*FinishedJob);
@@ -1586,7 +1609,7 @@ bool DoWriteTasksInner(const TArray<FShaderCommonCompileJobPtr>& QueuedJobs, FAr
 		// Serialize all the batched jobs
 		for (int32 JobIndex = 0; JobIndex < QueuedSingleJobs.Num(); JobIndex++)
 		{
-			TransferFile << QueuedSingleJobs[JobIndex]->Input;
+			QueuedSingleJobs[JobIndex]->SerializeWorkerInput(TransferFile);
 			QueuedSingleJobs[JobIndex]->Input.SerializeSharedInputs(TransferFile, SharedEnvironments, AllShaderParameterStructures);
 		}
 	}
@@ -1607,7 +1630,7 @@ bool DoWriteTasksInner(const TArray<FShaderCommonCompileJobPtr>& QueuedJobs, FAr
 			TransferFile << NumStageJobs;
 			for (int32 Index = 0; Index < NumStageJobs; Index++)
 			{
-				TransferFile << PipelineJob->StageJobs[Index]->GetSingleShaderJob()->Input;
+				PipelineJob->StageJobs[Index]->SerializeWorkerInput(TransferFile);
 				PipelineJob->StageJobs[Index]->Input.SerializeSharedInputs(TransferFile, SharedEnvironments, AllShaderParameterStructures);
 			}
 		}
@@ -4327,6 +4350,17 @@ EShaderDebugInfoFlags FShaderCompilingManager::GetDumpShaderDebugInfoFlags() con
 	{
 		Flags |= EShaderDebugInfoFlags::DirectCompileCommandLine;
 	}
+	
+	if (CVarDebugDumpJobInputHashes.GetValueOnAnyThread())
+	{
+		Flags |= EShaderDebugInfoFlags::InputHash;
+	}
+
+	if (CVarDebugDumpJobDiagnostics.GetValueOnAnyThread())
+	{
+		Flags |= EShaderDebugInfoFlags::Diagnostics;
+	}
+
 	return Flags;
 }
 
@@ -6726,6 +6760,10 @@ void GlobalBeginCompileShader(
 
 	if (CVarShadersUseLegacyPreprocessor.GetValueOnAnyThread())
 	{
+		if (CVarPreprocessedJobCache.GetValueOnAnyThread())
+		{
+			UE_LOG(LogShaderCompilers, Fatal, TEXT("r.Shaders.UseLegacyPreprocessor cannot be used in conjunction with r.ShaderCompiler.PreprocessedJobCache (legacy preprocessor is not thread-safe)."));
+		}
 		Input.Environment.CompilerFlags.Add(CFLAG_UseLegacyPreprocessor);
 	}
 
@@ -6780,6 +6818,11 @@ void GlobalBeginCompileShader(
 	checkf(Format, TEXT("Shader format %s cannot be found"), *ShaderFormatName.ToString());
 	Format->ModifyShaderCompilerInput(Input);
 
+	if (Format->SupportsIndependentPreprocessing() && CVarPreprocessedJobCache.GetValueOnAnyThread())
+	{
+		Input.bCachePreprocessed = true;
+	}
+
 	// Allow the GBuffer and other shader defines to cause dependend environment changes, but minimizing the #ifdef magic in the shaders, which
 	// is nearly impossible to debug when it goes wrong.
 	FShaderCompileUtilities::ApplyDerivedDefines(Input.Environment, Input.SharedEnvironment, (EShaderPlatform)Target.Platform);
diff --git a/Engine/Source/Runtime/RenderCore/Private/ShaderCore.cpp b/Engine/Source/Runtime/RenderCore/Private/ShaderCore.cpp
--- a/Engine/Source/Runtime/RenderCore/Private/ShaderCore.cpp
+++ b/Engine/Source/Runtime/RenderCore/Private/ShaderCore.cpp
@@ -9,6 +9,7 @@
 #include "Async/ParallelFor.h"
 #include "HAL/FileManager.h"
 #include "HAL/PlatformStackWalk.h"
+#include "Math/BigInt.h"
 #include "Misc/CommandLine.h"
 #include "Misc/FileHelper.h"
 #include "Misc/Compression.h"
@@ -33,10 +34,10 @@
 #include "DataDrivenShaderPlatformInfo.h"
 #include "Compression/OodleDataCompression.h"
 #include "Serialization/MemoryWriter.h"
-#if WITH_EDITOR
 #include "Misc/CoreMisc.h"
 #include "Interfaces/ITargetPlatform.h"
 #include "Interfaces/ITargetPlatformManagerModule.h"
+#if WITH_EDITOR
 #include "Misc/ConfigCacheIni.h"
 #include "Misc/StringBuilder.h"
 #endif
@@ -915,45 +916,131 @@ int HandleShaderCompileException(Windows::LPEXCEPTION_POINTERS Info, FString& Ou
 }
 #endif
 
+
+class FInternalShaderCompilerFunctions
+{
+public:
+	static bool PreprocessShaderInternal(const IShaderFormat* Backend, FShaderCompileJob& Job)
+	{
+		const double StartPreprocessTime = FPlatformTime::Seconds();
+		if (IsValidRef(Job.Input.SharedEnvironment))
+		{
+			// only create new environment & merge if necessary, save some allocs
+			// (need a copy here as we don't want to merge the environment in place like we do in the compile path
+			// and affect what is passed to the workers)
+			FShaderCompilerEnvironment MergedEnvironment = Job.Input.Environment;
+			MergedEnvironment.Merge(*Job.Input.SharedEnvironment);
+			Job.PreprocessOutput.bSucceeded = Backend->PreprocessShader(Job.Input, MergedEnvironment, Job.PreprocessOutput);
+		}
+		else
+		{
+			Job.PreprocessOutput.bSucceeded = Backend->PreprocessShader(Job.Input, Job.Input.Environment, Job.PreprocessOutput);
+		}
+
+		if (Job.PreprocessOutput.bSucceeded && Job.Input.bCachePreprocessed)
+		{
+			// if the preprocessed job cache is enabled we need to strip the preprocessed code, this removes comments, line directives
+			// and blank lines to improve deduplication (and populates data required to remap diagnostic messages to correct line numbers)
+			Job.PreprocessOutput.StripCode();
+		}
+
+		Job.PreprocessOutput.ElapsedTime = FPlatformTime::Seconds() - StartPreprocessTime;
+		Job.Output.PreprocessTime = Job.PreprocessOutput.ElapsedTime;
+		return Job.PreprocessOutput.bSucceeded;
+	}
+
 #if PLATFORM_WINDOWS
-void CompileShaderInternal(const IShaderFormat* Compiler, FShaderCompileJob& Job, const FString& WorkingDirectory, FString& OutExceptionMsg, FString& OutExceptionCallstack, int32* CompileCount)
+	static void CompileShaderInternal(const IShaderFormat* Compiler, FShaderCompileJob& Job, const FString& WorkingDirectory, FString& OutExceptionMsg, FString& OutExceptionCallstack, int32* CompileCount)
 #else
-void CompileShaderInternal(const IShaderFormat* Compiler, FShaderCompileJob& Job, const FString& WorkingDirectory, int32* CompileCount)
+	static void CompileShaderInternal(const IShaderFormat* Compiler, FShaderCompileJob& Job, const FString& WorkingDirectory, int32* CompileCount)
 #endif
-{
+	{
 #if PLATFORM_WINDOWS
-	__try
+		__try
 #endif
-	{
-		double TimeStart = FPlatformTime::Seconds();
-
-		Compiler->CompileShader(Job.Input.ShaderFormat, Job.Input, Job.Output, WorkingDirectory);
-		if (Job.Output.bSucceeded)
 		{
-			Job.Output.GenerateOutputHash();
-			if (Job.Input.CompressionFormat != NAME_None)
+			double TimeStart = FPlatformTime::Seconds();
+
+			if (Compiler->SupportsIndependentPreprocessing())
 			{
-				Job.Output.CompressOutput(Job.Input.CompressionFormat, Job.Input.OodleCompressor, Job.Input.OodleLevel);
+				if (!Job.Input.bCachePreprocessed)
+				{
+					PreprocessShaderInternal(Compiler, Job);
+				}
+				if (Job.PreprocessOutput.bSucceeded)
+				{
+					Compiler->CompilePreprocessedShader(Job.Input, Job.PreprocessOutput, Job.Output, WorkingDirectory);
+				}
+				else
+				{
+					Job.Output.Errors.Append(Job.PreprocessOutput.Errors);
+					Job.Output.bSucceeded = false;
+				}
 			}
-		}
-		Job.Output.CompileTime = FPlatformTime::Seconds() - TimeStart;
+			else
+			{
+				Compiler->CompileShader(Job.Input.ShaderFormat, Job.Input, Job.Output, WorkingDirectory);
+			}
+
+			if (Job.Output.bSucceeded)
+			{
+				Job.Output.GenerateOutputHash();
+				if (Job.Input.CompressionFormat != NAME_None)
+				{
+					Job.Output.CompressOutput(Job.Input.CompressionFormat, Job.Input.OodleCompressor, Job.Input.OodleLevel);
+				}
+			}
+			Job.Output.CompileTime = FPlatformTime::Seconds() - TimeStart;
 
-		if (Compiler->UsesHLSLcc(Job.Input))
+			if (Compiler->UsesHLSLcc(Job.Input))
+			{
+				Job.Output.bUsedHLSLccCompiler = true;
+			}
+
+			if (CompileCount)
+			{
+				++(*CompileCount);
+			}
+		}
+#if PLATFORM_WINDOWS
+		__except (HandleShaderCompileException(GetExceptionInformation(), OutExceptionMsg, OutExceptionCallstack))
 		{
-			Job.Output.bUsedHLSLccCompiler = true;
+			Job.Output.bSucceeded = false;
 		}
+#endif
+	}
+};
 
-		if (CompileCount)
+void ConditionalPreprocessShader(FShaderCommonCompileJob* Job)
+{
+	static ITargetPlatformManagerModule& TargetPlatformManager = GetTargetPlatformManagerRef();
+	if (FShaderCompileJob* SingleJob = Job->GetSingleShaderJob())
+	{
+		if (SingleJob->Input.bCachePreprocessed)
 		{
-			++(*CompileCount);
+			const IShaderFormat* ShaderFormat = TargetPlatformManager.FindShaderFormat(SingleJob->Input.ShaderFormat);
+			FInternalShaderCompilerFunctions::PreprocessShaderInternal(ShaderFormat, *SingleJob);
 		}
 	}
-#if PLATFORM_WINDOWS
-	__except (HandleShaderCompileException(GetExceptionInformation(), OutExceptionMsg, OutExceptionCallstack))
+	else if (FShaderPipelineCompileJob* PipelineJob = Job->GetShaderPipelineJob())
 	{
-		Job.Output.bSucceeded = false;
+		for (FShaderCompileJob* StageJob : PipelineJob->StageJobs)
+		{
+			if (StageJob->Input.bCachePreprocessed)
+			{
+				const IShaderFormat* ShaderFormat = TargetPlatformManager.FindShaderFormat(StageJob->Input.ShaderFormat);
+				if (!FInternalShaderCompilerFunctions::PreprocessShaderInternal(ShaderFormat, *StageJob))
+				{
+					// early out if preprocessing failed on one stage, no point in continuing
+					return;
+				}
+			}
+		}
+	}
+	else
+	{
+		checkf(0, TEXT("Unknown shader compile job type or bad job pointer"));
 	}
-#endif
 }
 
 void CompileShader(const TArray<const IShaderFormat*>& ShaderFormats, FShaderCompilerInput& Input, FShaderCompilerOutput& Output, const FString& WorkingDirectory, int32* CompileCount)
@@ -983,7 +1070,7 @@ void CompileShader(const TArray<const IShaderFormat*>& ShaderFormats, FShaderCom
 #if PLATFORM_WINDOWS
 	FString ExceptionMsg;
 	FString ExceptionCallstack;
-	CompileShaderInternal(Compiler, Job, WorkingDirectory, ExceptionMsg, ExceptionCallstack, CompileCount);
+	FInternalShaderCompilerFunctions::CompileShaderInternal(Compiler, Job, WorkingDirectory, ExceptionMsg, ExceptionCallstack, CompileCount);
 	if (!Job.Output.bSucceeded && (!ExceptionMsg.IsEmpty() || !ExceptionCallstack.IsEmpty()))
 	{
 		FShaderCompilerError Error;
@@ -994,12 +1081,13 @@ void CompileShader(const TArray<const IShaderFormat*>& ShaderFormats, FShaderCom
 		Job.Output.Errors.Add(Error);
 	}
 #else
-	CompileShaderInternal(Compiler, Job, WorkingDirectory, CompileCount);
+	FInternalShaderCompilerFunctions::CompileShaderInternal(Compiler, Job, WorkingDirectory, CompileCount);
 #endif
 
 	Job.bSucceeded = Job.Output.bSucceeded;
 
-	if (Job.bSucceeded && Job.Input.DumpDebugInfoPath.Len() > 0)
+	// skip this if the shader format supports independent preprocessing; this will be output as part of job completion callback for such jobs
+	if (!Compiler->SupportsIndependentPreprocessing() && Job.bSucceeded && Job.Input.DumpDebugInfoPath.Len() > 0)
 	{
 		// write down the output hash as a file
 		FString HashFileName = FPaths::Combine(Job.Input.DumpDebugInfoPath, TEXT("OutputHash.txt"));
@@ -2584,6 +2672,7 @@ FArchive& operator<<(FArchive& Ar, FShaderCompilerInput& Input)
 	Ar << Input.bSkipPreprocessedCache;
 	Ar << Input.bCompilingForShaderPipeline;
 	Ar << Input.bIncludeUsedOutputs;
+	Ar << Input.bCachePreprocessed;
 	Ar << Input.UsedOutputs;
 	Ar << Input.DumpDebugInfoRootPath;
 	Ar << Input.DumpDebugInfoPath;
@@ -2591,7 +2680,14 @@ FArchive& operator<<(FArchive& Ar, FShaderCompilerInput& Input)
 	Ar << Input.DebugExtension;
 	Ar << Input.DebugGroupName;
 	Ar << Input.DebugDescription;
-	Ar << Input.Environment;
+	if (Input.bCachePreprocessed)
+	{
+		Input.Environment.SerializeCompilationDependencies(Ar);
+	}
+	else
+	{
+		Ar << Input.Environment;
+	}
 	Ar << Input.ExtraSettings;
 	Ar << reinterpret_cast<uint8&>(Input.OodleCompressor);
 	Ar << reinterpret_cast<uint8&>(Input.OodleLevel);
@@ -2607,19 +2703,18 @@ FShaderCommonCompileJob::FInputHash FShaderPipelineCompileJob::GetInputHash()
 	{
 		return InputHash;
 	}
-
-	FBlake3 Hasher;
+	static_assert(sizeof(FShaderCommonCompileJob::FInputHash) == 32);
+	int256 CombinedHash = 0u;
 	for (int32 Index = 0; Index < StageJobs.Num(); ++Index)
 	{
 		if (StageJobs[Index])
 		{
 			FShaderCommonCompileJob::FInputHash StageHash = StageJobs[Index]->GetInputHash();
-			Hasher.Update(StageHash.GetBytes(), sizeof(decltype(StageHash.GetBytes())));
+			CombinedHash += int256(StageHash.GetBytes(), sizeof(StageHash.GetBytes()));
 		}
 	}
 
-	InputHash = Hasher.Finalize();
-
+	InputHash = FShaderCommonCompileJob::FInputHash(reinterpret_cast<FShaderCommonCompileJob::FInputHash::ByteArray&>(*CombinedHash.GetBits()));
 	bInputHashSet = true;
 	return InputHash;
 }
@@ -2653,52 +2748,47 @@ FShaderCommonCompileJob::FInputHash FShaderCompileJob::GetInputHash()
 		return InputHash;
 	}
 
-	auto SerializeInputs = [this](FArchive& Archive)
+	if (Input.bCachePreprocessed)
 	{
-		checkf(Archive.IsSaving() && !Archive.IsLoading(), TEXT("A loading archive is passed to FShaderCompileJob::GetInputHash(), this is not supported as it may corrupt its data"));
-
-		// Don't include debug group name in the hashing; this drastically worsens our cache hit rate
-		FString DebugGroupNameTmp(MoveTemp(Input.DebugGroupName));
-		Archive << Input;
-		Input.DebugGroupName = MoveTemp(DebugGroupNameTmp);
-		Input.Environment.SerializeEverythingButFiles(Archive);
-
-		// hash the source file so changes to files during the development are picked up
-		const FSHAHash& SourceHash = GetShaderFileHash(*Input.VirtualSourceFilePath, Input.Target.GetPlatform());
-		Archive << const_cast<FSHAHash&>(SourceHash);
-
-		// unroll the included files for the parallel processing.
-		// These are temporary arrays that only exist for the ParallelFor
-		TArray<const TCHAR*> IncludeVirtualPaths;
-		TArray<const FString*> Contents;
-		TArray<bool> OnlyHashIncludes;
-		TArray<FBlake3Hash> Hashes;
-
-		// while the contents of this is already hashed (included in Environment's operator<<()), we still need to account for includes in the generated files and hash them, too
-		for (TMap<FString, FString>::TConstIterator It(Input.Environment.IncludeVirtualPathToContentsMap); It; ++It)
+		FMemoryHasherBlake3 Hasher;
+		FShaderTarget Target = Input.Target;
+		Hasher << Target;
+		FShaderCompilerEnvironment MergedEnvironment = Input.Environment;
+		if (Input.SharedEnvironment)
 		{
-			const FString& VirtualPath = It.Key();
-			IncludeVirtualPaths.Add(*VirtualPath);
-			Contents.Add(&It.Value());
-			OnlyHashIncludes.Add(true);	// not hashing contents of the file itself, as it was included in Environment's operator<<()
-			Hashes.AddDefaulted();
+			MergedEnvironment.Merge(*Input.SharedEnvironment);
 		}
+		MergedEnvironment.SerializeCompilationDependencies(Hasher);
 
-		for (TMap<FString, FThreadSafeSharedStringPtr>::TConstIterator It(Input.Environment.IncludeVirtualPathToExternalContentsMap); It; ++It)
+		// const_cast due to serialization API requiring non-const. better than not having const correctness in the API.
+		Hasher << const_cast<FString&>(PreprocessOutput.GetSource());
+		InputHash = Hasher.Finalize();
+	}
+	else
+	{
+		auto SerializeInputs = [this](FArchive& Archive)
 		{
-			const FString& VirtualPath = It.Key();
-			IncludeVirtualPaths.Add(*VirtualPath);
-			check(It.Value());
-			Contents.Add(&(*It.Value()));
-			OnlyHashIncludes.Add(false);
-			Hashes.AddDefaulted();
-		}
+			checkf(Archive.IsSaving() && !Archive.IsLoading(), TEXT("A loading archive is passed to FShaderCompileJob::GetInputHash(), this is not supported as it may corrupt its data"));
 
-		if (Input.SharedEnvironment)
-		{
-			Input.SharedEnvironment->SerializeEverythingButFiles(Archive);
+			// Don't include debug group name in the hashing; this drastically worsens our cache hit rate
+			FString DebugGroupNameTmp(MoveTemp(Input.DebugGroupName));
+			Archive << Input;
+			Input.DebugGroupName = MoveTemp(DebugGroupNameTmp);
+			Input.Environment.SerializeEverythingButFiles(Archive);
+
+			// hash the source file so changes to files during the development are picked up
+			const FSHAHash& SourceHash = GetShaderFileHash(*Input.VirtualSourceFilePath, Input.Target.GetPlatform());
+			Archive << const_cast<FSHAHash&>(SourceHash);
 
-			for (TMap<FString, FString>::TConstIterator It(Input.SharedEnvironment->IncludeVirtualPathToContentsMap); It; ++It)
+			// unroll the included files for the parallel processing.
+			// These are temporary arrays that only exist for the ParallelFor
+			TArray<const TCHAR*> IncludeVirtualPaths;
+			TArray<const FString*> Contents;
+			TArray<bool> OnlyHashIncludes;
+			TArray<FBlake3Hash> Hashes;
+
+			// while the contents of this is already hashed (included in Environment's operator<<()), we still need to account for includes in the generated files and hash them, too
+			for (TMap<FString, FString>::TConstIterator It(Input.Environment.IncludeVirtualPathToContentsMap); It; ++It)
 			{
 				const FString& VirtualPath = It.Key();
 				IncludeVirtualPaths.Add(*VirtualPath);
@@ -2707,7 +2797,7 @@ FShaderCommonCompileJob::FInputHash FShaderCompileJob::GetInputHash()
 				Hashes.AddDefaulted();
 			}
 
-			for (TMap<FString, FThreadSafeSharedStringPtr>::TConstIterator It(Input.SharedEnvironment->IncludeVirtualPathToExternalContentsMap); It; ++It)
+			for (TMap<FString, FThreadSafeSharedStringPtr>::TConstIterator It(Input.Environment.IncludeVirtualPathToExternalContentsMap); It; ++It)
 			{
 				const FString& VirtualPath = It.Key();
 				IncludeVirtualPaths.Add(*VirtualPath);
@@ -2716,62 +2806,86 @@ FShaderCommonCompileJob::FInputHash FShaderCompileJob::GetInputHash()
 				OnlyHashIncludes.Add(false);
 				Hashes.AddDefaulted();
 			}
-		}
 
-		check(IncludeVirtualPaths.Num() == Contents.Num());
-		check(Contents.Num() == OnlyHashIncludes.Num());
-		check(OnlyHashIncludes.Num() == Hashes.Num());
+			if (Input.SharedEnvironment)
+			{
+				Input.SharedEnvironment->SerializeEverythingButFiles(Archive);
 
-		EShaderPlatform Platform = Input.Target.GetPlatform();
-		ParallelFor(Contents.Num(), [&IncludeVirtualPaths, &Contents, &OnlyHashIncludes, &Hashes, &Platform](int32 FileIndex)
-			{ 
-				FMemoryHasherBlake3 MemHasher;
-				HashShaderFileWithIncludes(MemHasher, IncludeVirtualPaths[FileIndex], *Contents[FileIndex], Platform, OnlyHashIncludes[FileIndex]);
-				Hashes[FileIndex] = MemHasher.Finalize();
-			},
-			EParallelForFlags::Unbalanced
-		);
+				for (TMap<FString, FString>::TConstIterator It(Input.SharedEnvironment->IncludeVirtualPathToContentsMap); It; ++It)
+				{
+					const FString& VirtualPath = It.Key();
+					IncludeVirtualPaths.Add(*VirtualPath);
+					Contents.Add(&It.Value());
+					OnlyHashIncludes.Add(true);	// not hashing contents of the file itself, as it was included in Environment's operator<<()
+					Hashes.AddDefaulted();
+				}
 
-		// include the hashes in the main hash (consider sorting them if includes are found to have a random order)
-		for (int32 HashIndex = 0, NumHashes = Hashes.Num(); HashIndex < NumHashes; ++HashIndex)
-		{
-			Archive << Hashes[HashIndex];
-		}
-	};
+				for (TMap<FString, FThreadSafeSharedStringPtr>::TConstIterator It(Input.SharedEnvironment->IncludeVirtualPathToExternalContentsMap); It; ++It)
+				{
+					const FString& VirtualPath = It.Key();
+					IncludeVirtualPaths.Add(*VirtualPath);
+					check(It.Value());
+					Contents.Add(&(*It.Value()));
+					OnlyHashIncludes.Add(false);
+					Hashes.AddDefaulted();
+				}
+			}
 
-	// use faster hasher that doesn't allocate memory
-	FMemoryHasherBlake3 MemHasher;
-	SerializeInputs(MemHasher);
-	InputHash = MemHasher.Finalize();
+			check(IncludeVirtualPaths.Num() == Contents.Num());
+			check(Contents.Num() == OnlyHashIncludes.Num());
+			check(OnlyHashIncludes.Num() == Hashes.Num());
 
-	if (GShaderCompilerDumpCompileJobInputs)
-	{
-		TArray<uint8> MemoryBlob;
-		FMemoryWriter MemWriter(MemoryBlob);
+			EShaderPlatform Platform = Input.Target.GetPlatform();
+			ParallelFor(Contents.Num(), [&IncludeVirtualPaths, &Contents, &OnlyHashIncludes, &Hashes, &Platform](int32 FileIndex)
+				{
+					FMemoryHasherBlake3 MemHasher;
+					HashShaderFileWithIncludes(MemHasher, IncludeVirtualPaths[FileIndex], *Contents[FileIndex], Platform, OnlyHashIncludes[FileIndex]);
+					Hashes[FileIndex] = MemHasher.Finalize();
+				},
+				EParallelForFlags::Unbalanced
+					);
+
+			// include the hashes in the main hash (consider sorting them if includes are found to have a random order)
+			for (int32 HashIndex = 0, NumHashes = Hashes.Num(); HashIndex < NumHashes; ++HashIndex)
+			{
+				Archive << Hashes[HashIndex];
+			}
+		};
+
+		// use faster hasher that doesn't allocate memory
+		FMemoryHasherBlake3 MemHasher;
+		SerializeInputs(MemHasher);
+		InputHash = MemHasher.Finalize();
 
-		SerializeInputs(MemWriter);
+		if (GShaderCompilerDumpCompileJobInputs)
+		{
+			TArray<uint8> MemoryBlob;
+			FMemoryWriter MemWriter(MemoryBlob);
+
+			SerializeInputs(MemWriter);
 
-		FString IntermediateFormatPath = FPaths::ProjectSavedDir() / TEXT("ShaderJobInputs");
+			FString IntermediateFormatPath = FPaths::ProjectSavedDir() / TEXT("ShaderJobInputs");
 #if UE_BUILD_DEBUG
-		FString TempPath = IntermediateFormatPath / TEXT("DebugEditor");
+			FString TempPath = IntermediateFormatPath / TEXT("DebugEditor");
 #else
-		FString TempPath = IntermediateFormatPath / TEXT("DevelopmentEditor");
+			FString TempPath = IntermediateFormatPath / TEXT("DevelopmentEditor");
 #endif
-		IFileManager::Get().MakeDirectory(*TempPath, true);
+			IFileManager::Get().MakeDirectory(*TempPath, true);
 
-		static int32 InputHashID = 0;
-		FString FileName = Input.DebugGroupName.Replace(TEXT("/"), TEXT("_")).Replace(TEXT("<"), TEXT("_")).Replace(TEXT(">"), TEXT("_")).Replace(TEXT(":"), TEXT("_")).Replace(TEXT("|"), TEXT("_"))
-			+ TEXT("-") + Input.EntryPointName;
-		FString TempFile = TempPath / FString::Printf(TEXT("%s-%d.bin"), *FileName, InputHashID++);
+			static int32 InputHashID = 0;
+			FString FileName = Input.DebugGroupName.Replace(TEXT("/"), TEXT("_")).Replace(TEXT("<"), TEXT("_")).Replace(TEXT(">"), TEXT("_")).Replace(TEXT(":"), TEXT("_")).Replace(TEXT("|"), TEXT("_"))
+				+ TEXT("-") + Input.EntryPointName;
+			FString TempFile = TempPath / FString::Printf(TEXT("%s-%d.bin"), *FileName, InputHashID++);
 
-		TUniquePtr<FArchive> DumpAr(IFileManager::Get().CreateFileWriter(*TempFile));
-		DumpAr->Serialize(MemoryBlob.GetData(), MemoryBlob.Num());
+			TUniquePtr<FArchive> DumpAr(IFileManager::Get().CreateFileWriter(*TempFile));
+			DumpAr->Serialize(MemoryBlob.GetData(), MemoryBlob.Num());
 
-		// as an additional debugging feature, make sure that the hash is the same as calculated by the memhasher
-		FBlake3Hash Check = FBlake3::HashBuffer(MemoryBlob.GetData(), MemoryBlob.Num());
-		if (Check != InputHash)
-		{
-			UE_LOG(LogShaders, Error, TEXT("Job input hash disagrees between FMemoryHasherSHA1 (%s) and FMemoryWriter + FSHA1 (%s, which was dumped to disk)"), *LexToString(InputHash), *LexToString(Check));
+			// as an additional debugging feature, make sure that the hash is the same as calculated by the memhasher
+			FBlake3Hash Check = FBlake3::HashBuffer(MemoryBlob.GetData(), MemoryBlob.Num());
+			if (Check != InputHash)
+			{
+				UE_LOG(LogShaders, Error, TEXT("Job input hash disagrees between FMemoryHasherSHA1 (%s) and FMemoryWriter + FSHA1 (%s, which was dumped to disk)"), *LexToString(InputHash), *LexToString(Check));
+			}
 		}
 	}
 
@@ -2810,6 +2924,27 @@ void FShaderCompileJob::SerializeOutput(FArchive& Ar)
 	}
 }
 
+void FShaderCompileJob::OnComplete()
+{
+	const IShaderFormat* ShaderFormat = GetTargetPlatformManagerRef().FindShaderFormat(Input.ShaderFormat);
+	if (ShaderFormat->SupportsIndependentPreprocessing())
+	{
+		// For jobs using the preprocessed cache, we need to remap error messages whether or not the job was actually the one that ran
+		// the compilation step. In addition since we always run preprocessing we set the total preprocess time accordingly.
+		if (Input.bCachePreprocessed)
+		{
+			PreprocessOutput.RemapErrors(Output);
+			Output.PreprocessTime = PreprocessOutput.ElapsedTime;
+		}
+		// dump debug info for the job whether or not the preprocessed cache was enabled, shader formats which support independent preprocessing
+		// will expect this to be done automatically.
+		if (Input.DumpDebugInfoEnabled())
+		{
+			ShaderFormat->OutputDebugData(LexToString(GetInputHash()), Input, PreprocessOutput, Output);
+		}
+	}
+}
+
 void FShaderCompileJob::SerializeWorkerOutput(FArchive& Ar)
 {
 	Ar << Output;
@@ -2818,6 +2953,16 @@ void FShaderCompileJob::SerializeWorkerOutput(FArchive& Ar)
 	bSucceeded = bSucceededTemp;
 }
 
+void FShaderCompileJob::SerializeWorkerInput(FArchive& Ar)
+{
+	Ar << Input;
+
+	if (Input.bCachePreprocessed)
+	{
+		Ar << PreprocessOutput;
+	}
+}
+
 FShaderPipelineCompileJob::FShaderPipelineCompileJob(int32 NumStages)
 	: FShaderCommonCompileJob(Type, 0u, 0u, EShaderCompileJobPriority::Num)
 {
@@ -2856,3 +3001,11 @@ void FShaderPipelineCompileJob::SerializeOutput(FArchive& Ar)
 		bSucceeded = bAllStagesSucceeded;
 	}
 }
+
+void FShaderPipelineCompileJob::OnComplete()
+{
+	for (int32 Index = 0, Num = StageJobs.Num(); Index < Num; ++Index)
+	{
+		StageJobs[Index]->OnComplete();
+	}
+}
diff --git a/Engine/Source/Runtime/RenderCore/Private/ShaderPreprocessTypes.cpp b/Engine/Source/Runtime/RenderCore/Private/ShaderPreprocessTypes.cpp
new file mode 100644
--- /dev/null
+++ b/Engine/Source/Runtime/RenderCore/Private/ShaderPreprocessTypes.cpp
@@ -0,0 +1,291 @@
+#include "ShaderPreprocessTypes.h"
+
+#include "ShaderCompilerCore.h"
+
+const TCHAR* FilenameSentinel = TEXT("__UE_FILENAME_SENTINEL__");
+static const int FilenameSentinelLen = FCString::Strlen(FilenameSentinel);
+static const FString LineDirectiveSentinel = FString::Printf(TEXT("#line 1 \"%s\"\n"), FilenameSentinel);
+
+void FShaderDiagnosticRemapper::Remap(FShaderCompilerError& Diagnostic) const
+{
+	if (!Diagnostic.ErrorLineString.IsEmpty() && !Diagnostic.ErrorVirtualFilePath.IsEmpty())
+	{
+		if (Diagnostic.ErrorVirtualFilePath != FilenameSentinel)
+		{
+			// only attempt remapping if the filename sentinel matches; this will bypass remapping in the case where compilation was retried on FXC
+			// with a DXC precompile (this compilation will report error messages relative to an intermediate hlsl file, which is not a change in behaviour)
+			return;
+		}
+
+		// Some backends give us already-parsed FShaderCompilerError objects which simplifies the remapping process somewhat
+		const FString& ErrorLineStr = Diagnostic.ErrorLineString;
+		check(FChar::IsDigit(ErrorLineStr[0]));
+		int32 StrippedLineNum = 0, LineNumberEnd = 0;
+		while (FChar::IsDigit(ErrorLineStr[LineNumberEnd]))
+		{
+			StrippedLineNum = StrippedLineNum * 10 + ErrorLineStr[LineNumberEnd++] - TEXT('0');
+		}
+
+		FRemapData RemapData = GetRemapData(StrippedLineNum);
+
+		FString NewErrorLineString;
+		NewErrorLineString.Reserve(ErrorLineStr.Len());
+		NewErrorLineString.AppendInt(RemapData.LineNumber);
+		NewErrorLineString.Append(ErrorLineStr.GetCharArray().GetData() + LineNumberEnd, ErrorLineStr.Len() - LineNumberEnd);
+
+		Diagnostic.ErrorLineString = MoveTemp(NewErrorLineString);
+		Diagnostic.ErrorVirtualFilePath = RemapData.Filename;
+	}
+	// need to parse/remap/reconstruct the full error message
+	else
+	{
+		const FString& Message = Diagnostic.StrippedErrorMessage;
+		int32 FilenameIndex = Message.Find(FilenameSentinel);
+		if (FilenameIndex < 0)
+		{
+			// only attempt remapping if we find our sentinel filename in the diagnostic message; otherwise it's likely
+			// a non-line-specific error that doesn't need remapping
+			return;
+		}
+
+		// Parse the line number from the message and record start and end indices
+		// Code assumes that the next integral number found in the message string after the filename is always line number
+		int32 LineNumberStart = FilenameIndex + FilenameSentinelLen;
+		while (!FChar::IsDigit(Message[LineNumberStart]))
+		{
+			++LineNumberStart;
+		}
+		int32 LineNumberEnd = LineNumberStart;
+		int32 StrippedLineNum = 0;
+		while (FChar::IsDigit(Message[LineNumberEnd]) && LineNumberEnd < Message.Len())
+		{
+			StrippedLineNum = StrippedLineNum * 10 + (Message[LineNumberEnd++] - TEXT('0'));
+		}
+
+		FRemapData RemapData = GetRemapData(StrippedLineNum);
+
+		// Assume line number is the same number of digits for simplicity in reserve allocation size;
+		// it's an upper bound but in most practical cases worst case we're allocating an extra byte
+		int32 RemappedLen = Message.Len() - FilenameSentinelLen + RemapData.Filename.Len();
+		FString RemappedMessage;
+		RemappedMessage.Reserve(RemappedLen);
+		RemappedMessage.AppendChars(Message.GetCharArray().GetData(), FilenameIndex);
+		RemappedMessage.Append(RemapData.Filename);
+		int32 SeparatorStart = FilenameIndex + FilenameSentinelLen;
+		RemappedMessage.Append(Message.GetCharArray().GetData() + SeparatorStart, LineNumberStart - SeparatorStart);
+		RemappedMessage.AppendInt(RemapData.LineNumber);
+		RemappedMessage.Append(Message.GetCharArray().GetData() + LineNumberEnd, Message.Len() - LineNumberEnd);
+
+		Diagnostic.StrippedErrorMessage = MoveTemp(RemappedMessage);
+	}
+}
+
+void FShaderDiagnosticRemapper::AddSourceBlock(int32 OriginalLineNum, int32 StrippedLineNum, FString && OriginalPath)
+{
+	if (Blocks.Num() && Blocks.Last().StrippedLineNum == StrippedLineNum)
+	{
+		// if no additional stripped source has been added since the last block, then the new block
+		// can replace it (not possible to hit an error/warning in that block, it must be entirely
+		// comments or whitespace)
+		Blocks.Last().OriginalLineNum = OriginalLineNum;
+		// If we weren't given a path, set path to the same as the last block (handles filename-less line directive)
+		Blocks.Last().OriginalPath = OriginalPath.IsEmpty() ? Blocks.Last().OriginalPath : MoveTemp(OriginalPath);
+	}
+	else
+	{
+		// sanity check that ranges are always added in order of stripped line numbers
+		check(Blocks.Num() == 0 || Blocks.Last().StrippedLineNum < StrippedLineNum);
+		// sanity check to ensure if we're not given a path we must have a previous block from which to retrieve the path
+		check(!OriginalPath.IsEmpty() || Blocks.Num() > 0);
+		Blocks.Add(FSourceBlock{ StrippedLineNum, OriginalLineNum, OriginalPath.IsEmpty() ? Blocks.Last().OriginalPath : MoveTemp(OriginalPath) });
+	}
+}
+
+void FShaderDiagnosticRemapper::AddStrippedLine(int32 StrippedLineNum, int32 Offset)
+{
+	check(StrippedLineOffsets.Num() == (StrippedLineNum - 1));
+	StrippedLineOffsets.Add(Offset);
+}
+
+FShaderDiagnosticRemapper::FRemapData FShaderDiagnosticRemapper::GetRemapData(int32 StrippedLineNum) const
+{
+	// Find the index which should contain a block with starting line number greater than the one to which the message points; the preceding
+	// block will be the one containing the stripped line in the unstripped source (in the case UpperBoundBy returns Blocks.Num() this means
+	// the error must be in the last recorded block)
+	int32 FoundIndex = Algo::UpperBoundBy(Blocks, StrippedLineNum, [](const FSourceBlock& Block) { return Block.StrippedLineNum; }) - 1;
+
+	// Sanity check - we can't possibly have an error which occurs before the first recorded block unless the stripping didn't record blocks
+	// properly (it only skips recording a block if it's empty in the stripped source, i.e. all comments or whitespace)
+	check(FoundIndex >= 0);
+
+	const FSourceBlock& FoundBlock = Blocks[FoundIndex];
+
+	// -1 when indexing stripped line offsets since line numbers are 1-based
+	return FRemapData{ FoundBlock.OriginalPath, FoundBlock.OriginalLineNum + StrippedLineOffsets[StrippedLineNum - 1] };
+}
+
+inline bool IsEndOfLine(TCHAR C)
+{
+	return C == TEXT('\r') || C == TEXT('\n');
+}
+
+inline bool StripNeedsHandling(TCHAR C)
+{
+	return IsEndOfLine(C) || C == TEXT('/') || C == 0 || C == TEXT('#');
+}
+
+inline void SkipNewLine(const TCHAR*& Current, const TCHAR* End)
+{
+	TCHAR First = Current < End ? Current[0] : 0;
+	TCHAR Second = Current + 1 < End ? Current[1] : 0;
+	Current += ((First + Second) == TEXT('\r') + TEXT('\n')) ? 2 : 1;
+}
+
+void FShaderPreprocessOutput::StripCode()
+{
+	// Reserve worst case slack (i.e. assuming there is nothing to strip) to avoid reallocation
+	FString PreprocessedSourceStripped(LineDirectiveSentinel, PreprocessedSource.Len() + 1);
+
+	const TCHAR* Begin = PreprocessedSource.GetCharArray().GetData(), *Current = Begin;
+	const TCHAR* End = Current + PreprocessedSource.Len();
+	int32 CurrentBlockUnstrippedLineOffset = 0;
+	int32 CurrentStrippedLineNum = 1;
+	while (Current < End)
+	{
+		while (!StripNeedsHandling(*Current))
+		{
+			PreprocessedSourceStripped.AppendChar(*Current++);
+		}
+
+		if (IsEndOfLine(*Current))
+		{
+			// only emit \n if it wasn't preceded immediately by another linebreak 
+			// (i.e. skip empty lines)
+			if (PreprocessedSourceStripped.Len() && !IsEndOfLine(PreprocessedSourceStripped[PreprocessedSourceStripped.Len() - 1]))
+			{
+				// Record the offset from the start of the block given by the last line directive for each line
+				// output in the stripped code. 
+				Remapper.AddStrippedLine(CurrentStrippedLineNum, CurrentBlockUnstrippedLineOffset);
+
+				// normalize line endings
+				PreprocessedSourceStripped.AppendChar(TEXT('\n'));
+				CurrentStrippedLineNum++;
+			}
+			CurrentBlockUnstrippedLineOffset++;
+			SkipNewLine(Current, End);
+		}
+		else if (Current[0] == TEXT('/'))
+		{
+			if (Current[1] == TEXT('/'))
+			{
+				while (!IsEndOfLine(*Current) && Current < End)
+				{
+					++Current;
+				}
+			}
+			else if (Current[1] == TEXT('*'))
+			{
+				Current += 2;
+				while (Current < End && !(Current[0] == TEXT('*') && Current[1] == TEXT('/')))
+				{
+					++Current;
+				}
+				Current += 2;
+			}
+			else
+			{
+				PreprocessedSourceStripped.AppendChar(*Current++);
+			}
+		}
+		else if (Current[0] == TEXT('#'))
+		{
+			if (Current + 4 < End
+				&& Current[1] == TEXT('l')
+				&& Current[2] == TEXT('i')
+				&& Current[3] == TEXT('n')
+				&& Current[4] == TEXT('e'))
+			{
+				Current += 4;
+
+				// scan past any whitespace until we find the first digit of the line number
+				while (!FChar::IsDigit(*Current) && Current < End)
+				{
+					++Current;
+				}
+				int32 DirectiveLineNumber = 0;
+				while (FChar::IsDigit(*Current) && Current < End)
+				{
+					DirectiveLineNumber = 10 * DirectiveLineNumber + (*Current - TEXT('0'));
+					++Current;
+				}
+
+				// scan past open quote
+				while (*Current != '\"' && Current < End && !IsEndOfLine(*Current))
+				{
+					++Current;
+				}
+				
+				if (*Current == '\"')
+				{
+					++Current;
+
+					const TCHAR* DirectiveFileNameStart = Current;
+					int32 DirectiveFileNameLen = 0;
+
+					// count chars in directive filename
+					while (*Current != '\"' && Current < End)
+					{
+						++DirectiveFileNameLen;
+						++Current;
+					}
+
+					FString DirectiveFileName(DirectiveFileNameLen, DirectiveFileNameStart);
+
+					// scan to end-of-line and skip past the newline; this would be handled by the newline case above as well,
+					// but we don't want the newline at the end of the line directive to count in our calculated offsets for
+					// emitted stripped code
+					while (!IsEndOfLine(*Current) && Current < End)
+					{
+						++Current;
+					}
+					Remapper.AddSourceBlock(DirectiveLineNumber, CurrentStrippedLineNum, MoveTemp(DirectiveFileName));
+				}
+				else
+				{
+					// if we hit EOL before finding an open quote, assume this is a filename-less line directive
+					Remapper.AddSourceBlock(DirectiveLineNumber, CurrentStrippedLineNum);
+				}
+				SkipNewLine(Current, End);
+
+				CurrentBlockUnstrippedLineOffset = 0;
+			}
+			else
+			{
+				PreprocessedSourceStripped.AppendChar(*Current++);
+			}
+		}
+	}
+	// Null terminate after stripped copy
+	PreprocessedSourceStripped.AppendChar(0);
+
+	OriginalPreprocessedSource = MoveTemp(PreprocessedSource);
+	PreprocessedSource = MoveTemp(PreprocessedSourceStripped);
+}
+
+void FShaderPreprocessOutput::RemapErrors(FShaderCompilerOutput& Output) const
+{
+	for (FShaderCompilerError& Error : Output.Errors)
+	{
+		Remapper.Remap(Error);
+	}
+}
+
+FArchive& operator<<(FArchive& Ar, FShaderPreprocessOutput& PreprocessOutput)
+{
+	Ar << PreprocessOutput.bSucceeded;
+	Ar << PreprocessOutput.Errors;
+	Ar << PreprocessOutput.PragmaDirectives;
+	Ar << PreprocessOutput.PreprocessedSource;
+	Ar << PreprocessOutput.ParameterParser;
+	return Ar;
+}
diff --git a/Engine/Source/Runtime/RenderCore/Public/ShaderCompilerCore.h b/Engine/Source/Runtime/RenderCore/Public/ShaderCompilerCore.h
--- a/Engine/Source/Runtime/RenderCore/Public/ShaderCompilerCore.h
+++ b/Engine/Source/Runtime/RenderCore/Public/ShaderCompilerCore.h
@@ -15,9 +15,12 @@
 #include "Misc/Paths.h"
 #include "Misc/CoreStats.h"
 #include "ShaderCore.h"
+#include "ShaderParameterMetadata.h"
+#include "ShaderParameterParser.h"
 
 class Error;
 class IShaderFormat;
+class FShaderCommonCompileJob;
 class FShaderCompileJob;
 class FShaderPipelineCompileJob;
 
@@ -193,6 +196,8 @@ enum class EShaderDebugInfoFlags : uint8
 {
 	Default = 0,
 	DirectCompileCommandLine = 1 << 0,
+	InputHash = 1 << 1,
+	Diagnostics = 1 << 2,
 };
 ENUM_CLASS_FLAGS(EShaderDebugInfoFlags)
 
@@ -218,6 +223,12 @@ struct FShaderCompilerInput
 	// Indicates which additional debug outputs should be written for this compile job.
 	EShaderDebugInfoFlags DebugInfoFlags;
 
+	// True if the cache key for this job should be based on preprocessed source. If so,
+	// preprocessing will be executed in the cook process independent of compilation (and
+	// as such this will only ever be set for jobs whose shader format supports independent
+	// preprocessing)
+	bool bCachePreprocessed;
+	
 	// Shader pipeline information
 	bool bCompilingForShaderPipeline;
 	bool bIncludeUsedOutputs;
@@ -258,6 +269,7 @@ struct FShaderCompilerInput
 		Target(SF_NumFrequencies, SP_NumPlatforms),
 		bSkipPreprocessedCache(false),
 		DebugInfoFlags(EShaderDebugInfoFlags::Default),
+		bCachePreprocessed(false),
 		bCompilingForShaderPipeline(false),
 		bIncludeUsedOutputs(false)
 	{
@@ -307,13 +319,17 @@ struct FShaderCompilerInput
 	{
 		check(!SharedEnvironment || SharedEnvironment->IncludeVirtualPathToExternalContentsMap.Num() == 0);
 
-		for (const auto& It : Environment.IncludeVirtualPathToExternalContentsMap)
+		// If the input is already preprocessed we don't need to serialize includes when writing worker input files
+		if (!bCachePreprocessed)
 		{
-			FString* FoundEntry = ExternalIncludes.Find(It.Key);
-
-			if (!FoundEntry)
+			for (const auto& It : Environment.IncludeVirtualPathToExternalContentsMap)
 			{
-				ExternalIncludes.Add(It.Key, *It.Value);
+				FString* FoundEntry = ExternalIncludes.Find(It.Key);
+
+				if (!FoundEntry)
+				{
+					ExternalIncludes.Add(It.Key, *It.Value);
+				}
 			}
 		}
 
@@ -332,15 +348,18 @@ struct FShaderCompilerInput
 	{
 		check(Ar.IsSaving());
 
-		TArray<FString> ReferencedExternalIncludes;
-		ReferencedExternalIncludes.Empty(Environment.IncludeVirtualPathToExternalContentsMap.Num());
-
-		for (const auto& It : Environment.IncludeVirtualPathToExternalContentsMap)
+		if (!bCachePreprocessed)
 		{
-			ReferencedExternalIncludes.Add(It.Key);
-		}
+			TArray<FString> ReferencedExternalIncludes;
+			ReferencedExternalIncludes.Empty(Environment.IncludeVirtualPathToExternalContentsMap.Num());
+
+			for (const auto& It : Environment.IncludeVirtualPathToExternalContentsMap)
+			{
+				ReferencedExternalIncludes.Add(It.Key);
+			}
 
-		Ar << ReferencedExternalIncludes;
+			Ar << ReferencedExternalIncludes;
+		}
 
 		int32 SharedEnvironmentIndex = SharedEnvironments.Find(SharedEnvironment);
 		Ar << SharedEnvironmentIndex;
@@ -362,14 +381,17 @@ struct FShaderCompilerInput
 	{
 		check(Ar.IsLoading());
 
-		TArray<FString> ReferencedExternalIncludes;
-		Ar << ReferencedExternalIncludes;
+		if (!bCachePreprocessed)
+		{
+			TArray<FString> ReferencedExternalIncludes;
+			Ar << ReferencedExternalIncludes;
 
-		Environment.IncludeVirtualPathToExternalContentsMap.Reserve(ReferencedExternalIncludes.Num());
+			Environment.IncludeVirtualPathToExternalContentsMap.Reserve(ReferencedExternalIncludes.Num());
 
-		for (int32 i = 0; i < ReferencedExternalIncludes.Num(); i++)
-		{
-			Environment.IncludeVirtualPathToExternalContentsMap.Add(ReferencedExternalIncludes[i], ExternalIncludes.FindChecked(ReferencedExternalIncludes[i]));
+			for (int32 i = 0; i < ReferencedExternalIncludes.Num(); i++)
+			{
+				Environment.IncludeVirtualPathToExternalContentsMap.Add(ReferencedExternalIncludes[i], ExternalIncludes.FindChecked(ReferencedExternalIncludes[i]));
+			}
 		}
 
 		int32 SharedEnvironmentIndex = 0;
@@ -609,6 +631,9 @@ struct RENDERCORE_API FSCWErrorCode
 extern RENDERCORE_API int HandleShaderCompileException(Windows::LPEXCEPTION_POINTERS Info, FString& OutExMsg, FString& OutCallStack);
 #endif
 extern RENDERCORE_API const IShaderFormat* FindShaderFormat(FName Format, TArray<const IShaderFormat*> ShaderFormats);
+
+// Executes preprocessing for the given job, if the job is marked to be preprocessed independently prior to compilation.
+extern RENDERCORE_API void ConditionalPreprocessShader(FShaderCommonCompileJob* Job);
 UE_DEPRECATED(5.3, "Use CompileShader overload which takes an FShaderCompileJob& rather than passing input/output directly.")
 extern RENDERCORE_API void CompileShader(const TArray<const IShaderFormat*>& ShaderFormats, FShaderCompilerInput& Input, FShaderCompilerOutput& Output, const FString& WorkingDirectory, int32* CompileCount = nullptr);
 extern RENDERCORE_API void CompileShader(const TArray<const IShaderFormat*>& ShaderFormats, FShaderCompileJob& Job, const FString& WorkingDirectory, int32* CompileCount = nullptr);
diff --git a/Engine/Source/Runtime/RenderCore/Public/ShaderCompilerJobTypes.h b/Engine/Source/Runtime/RenderCore/Public/ShaderCompilerJobTypes.h
--- a/Engine/Source/Runtime/RenderCore/Public/ShaderCompilerJobTypes.h
+++ b/Engine/Source/Runtime/RenderCore/Public/ShaderCompilerJobTypes.h
@@ -6,6 +6,7 @@
 #include "Hash/Blake3.h"
 #include "Shader.h"
 #include "ShaderCompilerCore.h"
+#include "ShaderPreprocessTypes.h"
 #include "VertexFactory.h"
 #include "Templates/Function.h"
 
@@ -115,6 +116,9 @@ public:
 	FShaderPipelineCompileJob* GetShaderPipelineJob();
 	const FShaderPipelineCompileJob* GetShaderPipelineJob() const;
 
+	// Executed for all jobs (including those read from cache) on completion.
+	virtual void OnComplete() = 0;
+
 	bool Equals(const FShaderCommonCompileJob& Rhs) const;
 
 	/** Calls the specified predicate for each single compile job, i.e. FShaderCompileJob and each stage of FShaderPipelineCompileJob. */
@@ -196,6 +200,7 @@ public:
 
 	/** Input for the shader compile */
 	FShaderCompilerInput Input;
+	FShaderPreprocessOutput PreprocessOutput;
 	FShaderCompilerOutput Output;
 
 	// List of pipelines that are sharing this job.
@@ -204,8 +209,13 @@ public:
 	virtual RENDERCORE_API FInputHash GetInputHash() override;
 	virtual RENDERCORE_API void SerializeOutput(FArchive& Ar) override;
 
+	virtual RENDERCORE_API void OnComplete() override;
+
 	// Serializes only the subset of data written by SCW/read back from ShaderCompiler when using worker processes.
 	RENDERCORE_API void SerializeWorkerOutput(FArchive& Ar);
+	
+	// Serializes only the subset of data written by ShaderCompiler and read from SCW when using worker processes.
+	RENDERCORE_API void SerializeWorkerInput(FArchive& Ar);
 
 	RENDERCORE_API FShaderCompileJob() : FShaderCommonCompileJob(Type, 0u, 0u, EShaderCompileJobPriority::Num)
 	{}
@@ -250,6 +260,7 @@ public:
 
 	virtual RENDERCORE_API FInputHash GetInputHash() override;
 	virtual RENDERCORE_API void SerializeOutput(FArchive& Ar) override;
+	virtual RENDERCORE_API void OnComplete() override;
 
 	RENDERCORE_API FShaderPipelineCompileJob(int32 NumStages);
 	RENDERCORE_API FShaderPipelineCompileJob(uint32 InHash, uint32 InId, EShaderCompileJobPriority InPriroity, const FShaderPipelineCompileJobKey& InKey);
diff --git a/Engine/Source/Runtime/RenderCore/Public/ShaderCore.h b/Engine/Source/Runtime/RenderCore/Public/ShaderCore.h
--- a/Engine/Source/Runtime/RenderCore/Public/ShaderCore.h
+++ b/Engine/Source/Runtime/RenderCore/Public/ShaderCore.h
@@ -643,6 +643,16 @@ struct FShaderCompilerEnvironment
 		Ar << FullPrecisionInPS;
 	}
 
+	// Serializes the portions of the environment that are used as input to the backend compilation process (i.e. after all preprocessing)
+	inline void SerializeCompilationDependencies(FArchive& Ar)
+	{
+		Ar << CompileArgs;
+		Ar << CompilerFlags;
+		Ar << ResourceTableMap;
+		Ar << UniformBufferMap;
+		Ar << FullPrecisionInPS;
+	}
+
 	friend FArchive& operator<<(FArchive& Ar,FShaderCompilerEnvironment& Environment)
 	{
 		// Note: this serialize is used to pass between UE and the shader compile worker, recompile both when modifying
diff --git a/Engine/Source/Runtime/RenderCore/Public/ShaderPreprocessTypes.h b/Engine/Source/Runtime/RenderCore/Public/ShaderPreprocessTypes.h
new file mode 100644
--- /dev/null
+++ b/Engine/Source/Runtime/RenderCore/Public/ShaderPreprocessTypes.h
@@ -0,0 +1,179 @@
+// Copyright Epic Games, Inc. All Rights Reserved.
+
+#pragma once
+
+#include "Containers/Array.h"
+#include "Containers/UnrealString.h"
+#include "ShaderCompilerCore.h"
+#include "ShaderCore.h"
+#include "ShaderParameterParser.h"
+
+struct FShaderCompilerEnvironment;
+struct FShaderCompilerError;
+struct FShaderCompilerInput;
+
+/*
+* Helper class used to remap compiler diagnostic messages from stripped preprocessed source (i.e. source with all whitespace normalized
+* and comments and line directives removed) back to line numbers/locations from the original source. 
+*/
+struct FShaderDiagnosticRemapper
+{
+private:
+	void Remap(FShaderCompilerError& Diagnostic) const;
+
+	friend class FShaderPreprocessOutput;
+	void AddSourceBlock(int32 OriginalLineNum, int32 StrippedLineNum)
+	{
+		AddSourceBlock(OriginalLineNum, StrippedLineNum, FString());
+	}
+	void AddSourceBlock(int32 OriginalLineNum, int32 StrippedLineNum, FString&& OriginalPath);
+	void AddStrippedLine(int32 StrippedLineNum, int32 Offset);
+
+	struct FRemapData
+	{
+		const FString& Filename;
+		int32 LineNumber;
+	};
+	FRemapData GetRemapData(int32 StrippedLineNum) const;
+
+	struct FSourceBlock
+	{
+		// Line number of the first line of code in the stripped preprocessed source for this block
+		int32 StrippedLineNum;
+		
+		// Associated line number of where the first line of code in the block occurred in the unstripped source
+		int32 OriginalLineNum;
+
+		// Full path associated with this block in the original unstripped source (as given by line directive)
+		FString OriginalPath;
+	};
+	TArray<FSourceBlock> Blocks;
+	TArray<int32> StrippedLineOffsets;
+};
+
+struct FShaderCompilerOutput;
+
+class FShaderPreprocessOutput
+{
+public:
+	FShaderPreprocessOutput()
+	{
+	}
+	const FString& GetSource() const
+	{
+		return PreprocessedSource;
+	}
+
+	const FString& GetOriginalSource() const
+	{
+		return OriginalPreprocessedSource;
+	}
+
+	FString& EditSource()
+	{
+		return PreprocessedSource;
+	}
+
+	inline bool HasDirective(const FString& Directive)
+	{
+		int32 NumberOfDirectives = PragmaDirectives.Num();
+		for (int32 i = 0; i < NumberOfDirectives; i++)
+		{
+			FString CurrentDirective = PragmaDirectives[i];
+			if (CurrentDirective.Equals(Directive))
+			{
+				return true;
+			}
+		}
+		return false;
+	}
+
+	inline void AddDirective(FString&& Directive)
+	{
+		PragmaDirectives.Add(Directive);
+	}
+
+	// Temporary helper for preprocessor wrapper function. Can be deprecated when
+	// all backends move to independent preprocessing.
+	inline void MoveDirectives(TArray<FString>& OutDirectives)
+	{
+		for (FString& Directive : PragmaDirectives)
+		{
+			OutDirectives.Add(MoveTemp(Directive));
+		}
+	}
+
+	inline void LogError(FString&& Message)
+	{
+		FShaderCompilerError* CompilerError = new(Errors) FShaderCompilerError;
+		CompilerError->StrippedErrorMessage = Message;
+	}
+
+	inline void LogError(FString&& FilePath, FString&& Message, FString&& LineNumberStr)
+	{
+		FShaderCompilerError* CompilerError = new(Errors) FShaderCompilerError;
+		CompilerError->ErrorVirtualFilePath = FilePath;
+		CompilerError->ErrorLineString = LineNumberStr;
+		CompilerError->StrippedErrorMessage = Message;
+	}
+
+	inline void LogError(FString&& FilePath, FString&& Message, int32 LineNumber)
+	{
+		FString LineNumberStr = LexToString(LineNumber);
+		LogError(MoveTemp(FilePath), MoveTemp(Message), MoveTemp(LineNumberStr));
+	}
+
+	inline TArray<FShaderCompilerError>& EditErrors()
+	{
+		return Errors;
+	}
+
+	inline bool ParseAndModify(const FShaderCompilerInput& Input, const FShaderCompilerEnvironment& Environment, const TCHAR* ConstantBufferType, const TArrayView<const TCHAR* const> ExtraSRVTypes = {}, const TArrayView<const TCHAR* const> ExtraUAVTypes = {})
+	{
+		FShaderParameterParser Tmp(Environment.CompilerFlags, ConstantBufferType, ExtraSRVTypes, ExtraUAVTypes);
+		ParameterParser = MoveTemp(Tmp);
+		return ParameterParser.ParseAndModify(Input, Errors, PreprocessedSource);
+	}
+
+	inline const FShaderParameterParser& GetParameterParser() const
+	{
+		return ParameterParser;
+	}
+
+	double GetElapsedTime() const
+	{
+		return ElapsedTime;
+	}
+
+	friend FArchive& operator<<(FArchive& Ar, FShaderPreprocessOutput& PreprocessOutput);
+
+private:
+
+	friend class FInternalShaderCompilerFunctions;
+	friend class FShaderCompileJob;
+
+	// Strips comments/whitespace/line directives from the preprocessed source, replacing the contents of PreprocessedSource
+	// and saving the original source in the OriginalPreprocessedSource member
+	void StripCode();
+
+	void RemapErrors(FShaderCompilerOutput& Output) const;
+
+	// Output of preprocessing; should be set by IShaderFormat::PreprocessShader
+	FString PreprocessedSource;
+
+	// Set by Finalize; original preprocessed source as set by IShaderFormat::PreprocessShader
+	FString OriginalPreprocessedSource;
+
+	// Array of errors encountered in preprocessing; should be populated by IShaderFormat::PreprocessShader
+	TArray<FShaderCompilerError> Errors;
+
+	// Array of "UESHADERMETADATA" pragmas encountered by preprocessing; set automatically by core preprocessing
+	// and expected to be queried by IShaderFormat 
+	TArray<FString> PragmaDirectives;
+	FShaderDiagnosticRemapper Remapper;
+
+	double ElapsedTime = 0.0;
+	bool bSucceeded = false;
+
+	FShaderParameterParser ParameterParser;
+};
