MPCook ShaderLibrary replication: add a maximum size for shader data in CopyToCompactBinaryAndClear to avoid overflowing the maximum size per message.
#rnx

diff --git a/Engine/Source/Editor/UnrealEd/Private/CookOnTheFlyServer.cpp b/Engine/Source/Editor/UnrealEd/Private/CookOnTheFlyServer.cpp
--- a/Engine/Source/Editor/UnrealEd/Private/CookOnTheFlyServer.cpp
+++ b/Engine/Source/Editor/UnrealEd/Private/CookOnTheFlyServer.cpp
@@ -8542,26 +8542,45 @@ public:
 	virtual void ServerReceiveMessage(FMPCollectorServerMessageContext& Context, FCbObjectView Message) override;
 
 	static FGuid MessageType;
+	static constexpr ANSICHAR RootObjectID[2]{ "S" };
 };
 FGuid FShaderLibraryCollector::MessageType(TEXT("4DF3B36BBA2F4E04A846E894E24EB2C4"));
 
 void FShaderLibraryCollector::ClientTick(FMPCollectorClientTickContext& Context)
 {
-	FCbWriter Writer;
-	bool bHasData;
-	Writer.BeginObject();
-	Writer.SetName("S");
-	FShaderLibraryCooker::CopyToCompactBinaryAndClear(Writer, bHasData);
-	if (bHasData)
+	constexpr int32 NumLoopsPerWarning = 100;
+	for (int32 NumMessages = 0; ; ++NumMessages)
 	{
-		Writer.EndObject();
-		Context.AddMessage(Writer.Save().AsObject());
+		// Maximum size for a message is 1GB. Caller will crash if we go over that.
+		// Provide a maximum shader size of half that because CopyToCompactBinaryAndClear adds on additional
+		// small amounts of data beyond the shader size limit.
+		const int64 MaximumSize = UE::CompactBinaryTCP::MaxOSPacketSize / 2;
+
+		bool bHasData;
+		bool bRanOutOfRoom;
+		FCbWriter Writer;
+		Writer.BeginObject();
+		Writer.SetName(RootObjectID);
+		FShaderLibraryCooker::CopyToCompactBinaryAndClear(Writer, bHasData, bRanOutOfRoom, MaximumSize);
+		if (bHasData)
+		{
+			Writer.EndObject();
+			Context.AddMessage(Writer.Save().AsObject());
+		}
+		if (!bRanOutOfRoom)
+		{
+			break;
+		}
+		if (NumMessages > 0 && NumMessages % NumLoopsPerWarning == 0)
+		{
+			UE_LOG(LogCook, Warning, TEXT("FShaderLibraryCollector::ClientTick has an unexpectedly high number of loops. Infinite loop?"))
+		}
 	}
 }
 
 void FShaderLibraryCollector::ServerReceiveMessage(FMPCollectorServerMessageContext& Context, FCbObjectView Message)
 {
-	bool bSuccessful = FShaderLibraryCooker::AppendFromCompactBinary(Message["S"]);
+	bool bSuccessful = FShaderLibraryCooker::AppendFromCompactBinary(Message[RootObjectID]);
 	UE_CLOG(!bSuccessful, LogCook, Error,
 		TEXT("Corrupt message received from CookWorker when replicating ShaderLibrary. Shaders will be missing from the cook."));
 }
diff --git a/Engine/Source/Editor/UnrealEd/Private/Cooker/CompactBinaryTCP.cpp b/Engine/Source/Editor/UnrealEd/Private/Cooker/CompactBinaryTCP.cpp
--- a/Engine/Source/Editor/UnrealEd/Private/Cooker/CompactBinaryTCP.cpp
+++ b/Engine/Source/Editor/UnrealEd/Private/Cooker/CompactBinaryTCP.cpp
@@ -259,10 +259,6 @@ EConnectionStatus FCompactBinaryTCPImpl::TryWritePacket(FSocket* Socket, FSendBu
 			Buffer.PendingMessages.Add(MoveTemp(NewMessage));
 		}
 	}
-	// Our Socket API restricts us to int32 Count field (Max == 1 << 31 - 1), but we further
-	// restrict that by a factor of two to prevent impact from any OS limitations less than that
-	// and to allow room for the 16 byte PacketHeader.
-	constexpr uint64 MaxOSPacketSize = (1 << 30) - 1;
 	MaxPacketSize = MaxPacketSize == 0 ? MaxOSPacketSize : FMath::Min(MaxPacketSize, MaxOSPacketSize);
 
 	for (;;)
diff --git a/Engine/Source/Editor/UnrealEd/Private/Cooker/CompactBinaryTCP.h b/Engine/Source/Editor/UnrealEd/Private/Cooker/CompactBinaryTCP.h
--- a/Engine/Source/Editor/UnrealEd/Private/Cooker/CompactBinaryTCP.h
+++ b/Engine/Source/Editor/UnrealEd/Private/Cooker/CompactBinaryTCP.h
@@ -21,6 +21,11 @@ class FCompactBinaryTCPImpl;
 struct FReceiveBuffer;
 struct FSendBuffer;
 
+// Our Socket API restricts us to int32 Count field (Max == 1 << 31 - 1), but we further
+// restrict that by a factor of two to prevent impact from any OS limitations less than that
+// and to allow room for the 16 byte PacketHeader.
+constexpr uint64 MaxOSPacketSize = (1 << 30) - 1;
+
 enum EConnectionStatus
 {
 	/** Connection is still okay or operation succeeded. */
diff --git a/Engine/Source/Runtime/RenderCore/Private/ShaderCodeArchive.cpp b/Engine/Source/Runtime/RenderCore/Private/ShaderCodeArchive.cpp
--- a/Engine/Source/Runtime/RenderCore/Private/ShaderCodeArchive.cpp
+++ b/Engine/Source/Runtime/RenderCore/Private/ShaderCodeArchive.cpp
@@ -162,6 +162,16 @@ bool FSerializedShaderArchive::FindOrAddShader(const FSHAHash& Hash, int32& OutI
 	return false;
 }
 
+void FSerializedShaderArchive::RemoveLastAddedShader()
+{
+	check(!ShaderEntries.IsEmpty() && ShaderEntries.Num() == ShaderHashes.Num());
+	int32 ShaderIndex = ShaderEntries.Num() - 1;
+	const uint32 Key = GetTypeHash(ShaderHashes[ShaderIndex]);
+	ShaderHashTable.Remove(Key, ShaderIndex);
+	ShaderHashes.RemoveAt(ShaderIndex);
+	ShaderEntries.RemoveAt(ShaderIndex);
+}
+
 #if WITH_EDITOR
 FCbWriter& operator<<(FCbWriter& Writer, const FSerializedShaderArchive& Archive)
 {
diff --git a/Engine/Source/Runtime/RenderCore/Private/ShaderCodeLibrary.cpp b/Engine/Source/Runtime/RenderCore/Private/ShaderCodeLibrary.cpp
--- a/Engine/Source/Runtime/RenderCore/Private/ShaderCodeLibrary.cpp
+++ b/Engine/Source/Runtime/RenderCore/Private/ShaderCodeLibrary.cpp
@@ -1382,12 +1382,8 @@ struct FEditorShaderCodeArchive
 		return bHasCopiedAndCleared ? !ShaderMapsToCopy.IsEmpty() : !SerializedShaders.ShaderMapHashes.IsEmpty();
 	}
 
-	void CopyToCompactBinaryAndClear(FCbWriter& Writer)
+	void CopyToCompactBinary(FCbWriter& Writer, FSerializedShaderArchive& TransferArchive, TArray<uint8>& TransferCode)
 	{
-		FSerializedShaderArchive TransferArchive;
-		TArray<uint8> TransferCode;
-		CopyToArchiveAndClear(TransferArchive, TransferCode);
-
 		Writer.BeginObject();
 		Writer << "SerializedShaders" << TransferArchive;
 		Writer << "ShaderCode";
@@ -1408,8 +1404,10 @@ struct FEditorShaderCodeArchive
 		return AppendFromArchive(MoveTemp(TransferArchive), TransferCode, CodeStats);
 	}
 
-	void CopyToArchiveAndClear(FSerializedShaderArchive& TargetArchive, TArray<uint8>& TargetFlatShaderCode)
+	void CopyToArchiveAndClear(FSerializedShaderArchive& TargetArchive, TArray<uint8>& TargetFlatShaderCode,
+		bool& bOutRanOutOfRoom, int64 MaxShaderSize, int64 MaxShaderCount)
 	{
+		bOutRanOutOfRoom = false;
 		if (!bHasCopiedAndCleared)
 		{
 			// First CopyAndClear adds all shadermaps to ShaderMapsToCopy, because as an optimization we do not write to
@@ -1422,6 +1420,7 @@ struct FEditorShaderCodeArchive
 
 		TArray<int32> LocalShaderMapsToCopy = ShaderMapsToCopy.Array();
 		LocalShaderMapsToCopy.Sort(); // Maintain the same order that the shadermaps were added in
+		int32 NumShadersSentWithCode = 0;
 
 		const FSerializedShaderArchive& SourceArchive = this->SerializedShaders;
 		TArray<TArray<uint8>>& SourceShaderCodes = this->ShaderCode;
@@ -1453,6 +1452,10 @@ struct FEditorShaderCodeArchive
 				TargetArchive.ShaderIndices[TargetEntry.ShaderIndicesOffset + ShaderIndexIndex] = TargetShaderIndex;
 				if (bIsNewShader)
 				{
+					// We rely on the index of the newly added shader being at the end of the list of ShaderEntries,
+					// so we can pop off the added index if we overflow below
+					check(TargetShaderIndex == TargetArchive.ShaderEntries.Num() - 1);
+
 					const FShaderCodeEntry& SourceShaderEntry = SourceArchive.ShaderEntries[SourceShaderIndex];
 					TArray<uint8>& SourceShaderCode = SourceShaderCodes[SourceShaderIndex];
 					FShaderCodeEntry& TargetShaderEntry = TargetArchive.ShaderEntries[TargetShaderIndex];
@@ -1462,6 +1465,23 @@ struct FEditorShaderCodeArchive
 					{
 						TargetShaderEntry.Offset = TargetFlatShaderCode.Num();
 						check(SourceShaderEntry.Size == SourceShaderCode.Num());
+						if ((MaxShaderSize > 0 && TargetFlatShaderCode.Num() + SourceShaderCode.Num() > MaxShaderSize) ||
+							(MaxShaderCount > 0 && NumShadersSentWithCode > MaxShaderCount))
+						{
+							// We have to stop here to avoid overflowing the shader limit. Send the shaders we have accumulated
+							// but do not send&clear any other data.
+							// Remove all ShaderMap data
+							TargetArchive.EmptyShaderMaps();
+							// Remove the ShaderEntry we just added; we are not adding its code so we have to remove it from the
+							// list of shaders contained by the targetarchive
+							check(TargetShaderIndex == TargetArchive.ShaderHashes.Num() - 1);
+							TargetArchive.RemoveLastAddedShader();
+							bOutRanOutOfRoom = true;
+							// Keep any shaders we added before the one we just added
+							return;
+						}
+						++NumShadersSentWithCode;
+
 						TargetFlatShaderCode.Append(SourceShaderCode);
 
 						// Empty the ShaderCode to save memory in the local process. The consumer of the TargetArchive and
@@ -1472,7 +1492,6 @@ struct FEditorShaderCodeArchive
 					{
 						// The shadercode was already copied in an earlier CopyAndClear operation; we just need to note that
 						// ShaderMaps in this call to CopyAndClear reference it.
-						check(bHasCopiedAndCleared);
 						TargetShaderEntry.Offset = INDEX_NONE;
 					}
 				}
@@ -1490,6 +1509,54 @@ struct FEditorShaderCodeArchive
 		FSerializedShaderArchive& TargetArchive = this->SerializedShaders;
 		TArray<TArray<uint8>>& TargetShaderCodes = this->ShaderCode;
 
+		// Add all the shaders; we can sometimes get messages that send the shaders in advance without sending the shadermaps that use them
+		for (int32 SourceShaderIndex = 0; SourceShaderIndex < SourceArchive.ShaderHashes.Num(); ++SourceShaderIndex)
+		{
+			const FSHAHash& SourceShaderHash = SourceArchive.ShaderHashes[SourceShaderIndex];
+			int32 TargetShaderIndex = INDEX_NONE;
+			const bool bShaderIsNew = TargetArchive.FindOrAddShader(SourceShaderHash, TargetShaderIndex);
+			if (!bShaderIsNew)
+			{
+				continue;
+			}
+			check(TargetShaderIndex == TargetArchive.ShaderEntries.Num() - 1 &&
+				TargetShaderCodes.Num() == TargetArchive.ShaderEntries.Num() - 1);
+			TArray<uint8>& TargetShaderCode = TargetShaderCodes.Emplace_GetRef();
+
+			const FShaderCodeEntry& SourceShaderEntry = SourceArchive.ShaderEntries[SourceShaderIndex];
+			FShaderCodeEntry& TargetShaderEntry = TargetArchive.ShaderEntries[TargetShaderIndex];
+			TargetShaderEntry = SourceShaderEntry;
+			TargetShaderEntry.Offset = 0;
+
+			CodeStats.NumUniqueShaders++;
+			CodeStats.ShadersUniqueSize += SourceShaderEntry.Size;
+
+			if (SourceShaderEntry.Offset == INDEX_NONE)
+			{
+				UE_LOG(LogShaderLibrary, Error, TEXT("ShaderMapLibrary transfer received from a remote machine has incomplete record for Shader %s. ")
+					TEXT("The remote machine thought this machine already had the shader and did not send the ShaderCode for the Shader, but the shader is not found. ")
+					TEXT("The ShaderMaps using the shader will be corrupt."),
+					*SourceShaderHash.ToString());
+				bOk = false;
+				TargetShaderEntry.Size = 0;
+			}
+			else if (SourceShaderEntry.Offset + SourceShaderEntry.Size > static_cast<uint64>(SourceFlatShaderCode.Num()))
+			{
+				UE_LOG(LogShaderLibrary, Error, TEXT("ShaderMapLibrary transfer received from a remote machine has corrupt record for Shader %s. ")
+					TEXT("The (Offset, Size) specified by the ShaderEntry does not fit in the TransferCode (an array of bytes that should contain the shader code for all transferred shaders. ")
+					TEXT("The ShaderMaps using the shader will be corrupt."),
+					*SourceShaderHash.ToString());
+				bOk = false;
+				TargetShaderEntry.Size = 0;
+			}
+			else
+			{
+				const TConstArrayView<uint8> SourceShaderCode(SourceFlatShaderCode.GetData() + SourceShaderEntry.Offset, SourceShaderEntry.Size);
+				TargetShaderCode = SourceShaderCode; // Copy from source's flat list to the target's separate TArray<uint8> for each shader
+			}
+		}
+
+		// Add all the shadermaps
 		for (int32 SourceShaderMapIndex = 0; SourceShaderMapIndex < SourceArchive.ShaderMapHashes.Num(); ++SourceShaderMapIndex)
 		{
 			const FSHAHash& SourceShaderMapHash = SourceArchive.ShaderMapHashes[SourceShaderMapIndex];
@@ -1517,46 +1584,8 @@ struct FEditorShaderCodeArchive
 				int32 TargetShaderIndex;
 				const bool bShaderIsNew = TargetArchive.FindOrAddShader(SourceShaderHash, TargetShaderIndex);
 				TargetArchive.ShaderIndices[TargetEntry.ShaderIndicesOffset + ShaderIndexIndex] = TargetShaderIndex;
-				if (!bShaderIsNew)
-				{
-					// Shader has already been loaded for another ShaderMap
-					continue;
-				}
-				check(TargetShaderIndex == TargetArchive.ShaderEntries.Num() - 1 &&
-					TargetShaderCodes.Num() == TargetArchive.ShaderEntries.Num() - 1);
-				TArray<uint8>& TargetShaderCode = TargetShaderCodes.Emplace_GetRef();
-
-				const FShaderCodeEntry& SourceShaderEntry = SourceArchive.ShaderEntries[SourceShaderIndex];
-				FShaderCodeEntry& TargetShaderEntry = TargetArchive.ShaderEntries[TargetShaderIndex];
-				TargetShaderEntry = SourceShaderEntry;
-				TargetShaderEntry.Offset = 0;
-
-				CodeStats.NumUniqueShaders++;
-				CodeStats.ShadersUniqueSize += SourceShaderEntry.Size;
-
-				if (SourceShaderEntry.Offset == INDEX_NONE)
-				{
-					UE_LOG(LogShaderLibrary, Error, TEXT("ShaderMapLibrary transfer received from a remote machine has incomplete record for ShaderMap %s, Shader %s. ")
-						TEXT("The remote machine thought this machine already had the shader and did not send the ShaderCode for the Shader, but the shader is not found. ")
-						TEXT("The ShaderMap will be corrupt."),
-						*SourceShaderMapHash.ToString(), *SourceShaderHash.ToString());
-					bOk = false;
-					TargetShaderEntry.Size = 0;
-				}
-				else if (SourceShaderEntry.Offset + SourceShaderEntry.Size > static_cast<uint64>(SourceFlatShaderCode.Num()))
-				{
-					UE_LOG(LogShaderLibrary, Error, TEXT("ShaderMapLibrary transfer received from a remote machine has corrupt record for ShaderMap %s, Shader %s. ")
-						TEXT("The (Offset, Size) specified by the ShaderEntry does not fit in the TransferCode (an array of bytes that should contain the shader code for all transferred shaders. ")
-						TEXT("The ShaderMap will be corrupt."),
-						*SourceShaderMapHash.ToString(), *SourceShaderHash.ToString());
-					bOk = false;
-					TargetShaderEntry.Size = 0;
-				}
-				else
-				{
-					const TConstArrayView<uint8> SourceShaderCode(SourceFlatShaderCode.GetData() + SourceShaderEntry.Offset, SourceShaderEntry.Size);
-					TargetShaderCode = SourceShaderCode; // Copy from source's flat list to the target's separate TArray<uint8> for each shader
-				}
+				// Every shader in the SourceArchive should have already been added by the loop above over SourceArchive.ShaderHashes
+				check(!bShaderIsNew);
 			}
 
 			CodeStats.NumShaders += TargetEntry.NumShaders; // Sum of shader counts used by each ShaderMap, without removing duplicates
@@ -2919,8 +2948,10 @@ public:
 	}
 
 #if WITH_EDITOR
-	void CopyToCompactBinaryAndClear(FCbWriter& Writer, bool& bOutHasData)
+	void CopyToCompactBinaryAndClear(FCbWriter& Writer, bool& bOutHasData, bool& bOutRanOutOfRoom, int64 MaxShaderSize)
 	{
+		bOutRanOutOfRoom = false;
+
 		TArray<EShaderPlatform, TInlineAllocator<10>> PlatformsToCopy;
 		FScopeLock ScopeLock(&ShaderCodeCS);
 		for (EShaderPlatform Platform = (EShaderPlatform)0; Platform < (EShaderPlatform)UE_ARRAY_COUNT(EditorShaderCodeStats);
@@ -2940,9 +2971,14 @@ public:
 			return;
 		}
 		bOutHasData = true;
+		int64 RemainingSize = MaxShaderSize;
 		Writer.BeginArray();
 		for (EShaderPlatform Platform : PlatformsToCopy)
 		{
+			if (bOutRanOutOfRoom)
+			{
+				break;
+			}
 			Writer.BeginObject();
 			{
 				Writer << "Platform" << (uint32)Platform;
@@ -2951,10 +2987,32 @@ public:
 
 				if ((CodeArchive && CodeArchive->HasDataToCopy()))
 				{
-					Writer.SetName("EditorShaderCodeArchive");
-					CodeArchive->CopyToCompactBinaryAndClear(Writer);
+					FSerializedShaderArchive TransferArchive;
+					TArray<uint8> TransferCode;
+					int64 MaxShaderSizeThisCall = RemainingSize;
+					int64 MaxShaderCount = -1;
+					bool bRanOutOfRoom;
+					CodeArchive->CopyToArchiveAndClear(TransferArchive, TransferCode, bRanOutOfRoom, MaxShaderSizeThisCall, MaxShaderCount);
+					bOutRanOutOfRoom |= bRanOutOfRoom;
+					if (bRanOutOfRoom && TransferCode.IsEmpty() && RemainingSize == MaxShaderSize)
+					{
+						UE_LOG(LogShaderLibrary, Error,
+							TEXT("MaxShaderSize %" INT64_FMT " is too small to read even a single shader. We will ignore it and allow uncapped size, which will possibly cause an overflow in the caller."),
+							MaxShaderSize);
+						MaxShaderSizeThisCall = -1;
+						MaxShaderCount = 1;
+						TransferArchive.Empty();
+						CodeArchive->CopyToArchiveAndClear(TransferArchive, TransferCode, bRanOutOfRoom, MaxShaderSizeThisCall, MaxShaderCount);
+					}
+					if (!TransferArchive.IsEmpty())
+					{
+						RemainingSize -= TransferCode.Num();
+						Writer.SetName("EditorShaderCodeArchive");
+						CodeArchive->CopyToCompactBinary(Writer, TransferArchive, TransferCode);
+					}
+
 				}
-				if (StableInfo && StableInfo->HasDataToCopy())
+				if (!bOutRanOutOfRoom && StableInfo && StableInfo->HasDataToCopy())
 				{
 					Writer.SetName("EditorShaderStableInfo");
 					StableInfo->CopyToCompactBinary(Writer);
@@ -3630,11 +3688,13 @@ bool FShaderLibraryCooker::AddShaderCode(EShaderPlatform ShaderPlatform, const F
 }
 
 #if WITH_EDITOR
-void FShaderLibraryCooker::CopyToCompactBinaryAndClear(FCbWriter& Writer, bool& bOutHasData)
+void FShaderLibraryCooker::CopyToCompactBinaryAndClear(FCbWriter& Writer, bool& bOutHasData,
+	bool& bOutRanOutOfRoom, int64 MaxShaderSize)
 {
 	if (FShaderLibrariesCollection::Impl)
 	{
-		FShaderLibrariesCollection::Impl->CopyToCompactBinaryAndClear(Writer, bOutHasData);
+		FShaderLibrariesCollection::Impl->CopyToCompactBinaryAndClear(Writer, bOutHasData,
+			bOutRanOutOfRoom, MaxShaderSize);
 	}
 }
 
diff --git a/Engine/Source/Runtime/RenderCore/Public/ShaderCodeArchive.h b/Engine/Source/Runtime/RenderCore/Public/ShaderCodeArchive.h
--- a/Engine/Source/Runtime/RenderCore/Public/ShaderCodeArchive.h
+++ b/Engine/Source/Runtime/RenderCore/Public/ShaderCodeArchive.h
@@ -156,17 +156,23 @@ public:
 
 	void Empty()
 	{
+		EmptyShaderMaps();
+
 		ShaderHashes.Empty();
 		ShaderEntries.Empty();
+		ShaderHashTable.Clear();
+	}
+
+	void EmptyShaderMaps()
+	{
 		ShaderMapHashes.Empty();
 		ShaderMapEntries.Empty();
 		PreloadEntries.Empty();
 		ShaderIndices.Empty();
 		ShaderMapHashTable.Clear();
-		ShaderHashTable.Clear();
 #if WITH_EDITOR
 		ShaderCodeToAssets.Empty();
-#endif // WITH_EDITOR
+#endif
 	}
 
 	int32 GetNumShaderMaps() const
@@ -179,6 +185,15 @@ public:
 		return ShaderEntries.Num();
 	}
 
+	bool IsEmpty() const
+	{
+		return ShaderMapEntries.IsEmpty() && ShaderEntries.IsEmpty() && PreloadEntries.IsEmpty()
+#if WITH_EDITOR
+			&& ShaderCodeToAssets.IsEmpty()
+#endif
+			;
+	}
+
 	int32 FindShaderMapWithKey(const FSHAHash& Hash, uint32 Key) const;
 	int32 FindShaderMap(const FSHAHash& Hash) const;
 	bool FindOrAddShaderMap(const FSHAHash& Hash, int32& OutIndex, const FShaderMapAssetPaths* AssociatedAssets);
@@ -186,6 +201,7 @@ public:
 	int32 FindShaderWithKey(const FSHAHash& Hash, uint32 Key) const;
 	int32 FindShader(const FSHAHash& Hash) const;
 	bool FindOrAddShader(const FSHAHash& Hash, int32& OutIndex);
+	void RemoveLastAddedShader();
 
 	void DecompressShader(int32 Index, const TArray<TArray<uint8>>& ShaderCode, TArray<uint8>& OutDecompressedShader) const;
 
diff --git a/Engine/Source/Runtime/RenderCore/Public/ShaderCodeLibrary.h b/Engine/Source/Runtime/RenderCore/Public/ShaderCodeLibrary.h
--- a/Engine/Source/Runtime/RenderCore/Public/ShaderCodeLibrary.h
+++ b/Engine/Source/Runtime/RenderCore/Public/ShaderCodeLibrary.h
@@ -299,7 +299,7 @@ struct RENDERCORE_API FShaderLibraryCooker
 
 #if WITH_EDITOR
 	/** Called from a CookWorker to send all contents of the ShaderLibrary to the CookDirector */
-	static void CopyToCompactBinaryAndClear(FCbWriter& Writer, bool& bOutHasData);
+	static void CopyToCompactBinaryAndClear(FCbWriter& Writer, bool& bOutHasData, bool& bOutRanOutOfRoom, int64 MaxShaderSize);
 
 	/** Called On the CookDirector to receive ShaderLibrary contents from a CookWorker */
 	static bool AppendFromCompactBinary(FCbFieldView Field);
