removed excessive FGraphEvent instance from `FRenderCommandFence::BeginFence()`.  The change uses FGraphEvent returned by the task instead of allocating a new one and manually "completing" it inside a task.
an excessive FGraphEvent instance is still used in `bSyncToRHIAndGPU` branch, and can be similarly improved. though this would require massive changes and won't bring any measurable perf improvements as this happens once a frame.

diff --git a/Engine/Source/Runtime/RenderCore/Private/RenderingThread.cpp b/Engine/Source/Runtime/RenderCore/Private/RenderingThread.cpp
--- a/Engine/Source/Runtime/RenderCore/Private/RenderingThread.cpp
+++ b/Engine/Source/Runtime/RenderCore/Private/RenderingThread.cpp
@@ -1125,8 +1125,7 @@ void FRenderCommandFence::BeginFence(bool bSyncToRHIAndGPU)
 			STAT_FNullGraphTask_FenceRenderCommand,
 				STATGROUP_TaskGraphTasks);
 
-			CompletionEvent = FGraphEvent::CreateGraphEvent();
-			FFunctionGraphTask::CreateAndDispatchWhenReady([CompletionEvent = CompletionEvent] { CompletionEvent->DispatchSubsequents(); }, GET_STATID(STAT_FNullGraphTask_FenceRenderCommand), nullptr, ENamedThreads::GetRenderThread());
+			CompletionEvent = FFunctionGraphTask::CreateAndDispatchWhenReady([] {}, GET_STATID(STAT_FNullGraphTask_FenceRenderCommand), nullptr, ENamedThreads::GetRenderThread());
 		}
 	}
 }
