Fix RHI validation error for RayTracingSceneBuffer: add missing SRVMask flag.

diff --git a/Engine/Source/Runtime/RenderCore/Private/RenderGraphBuilder.cpp b/Engine/Source/Runtime/RenderCore/Private/RenderGraphBuilder.cpp
--- a/Engine/Source/Runtime/RenderCore/Private/RenderGraphBuilder.cpp
+++ b/Engine/Source/Runtime/RenderCore/Private/RenderGraphBuilder.cpp
@@ -284,7 +284,7 @@ void EnumerateBufferAccess(FRDGParameterStruct PassParameters, ERDGPassFlags Pas
 
 				if (EnumHasAnyFlags(Buffer->Desc.Usage, BUF_AccelerationStructure))
 				{
-					BufferAccess = ERHIAccess::BVHRead;
+					BufferAccess = ERHIAccess::BVHRead | ERHIAccess::SRVMask;
 				}
 
 				AccessFunction(SRV, Buffer, BufferAccess);
