VulkanRHI: Use the DimensionOverride in DumpGPU when dumping slices of arrays individually.

#rnx

diff --git a/Engine/Source/Runtime/RenderCore/Private/DumpGPU.cpp b/Engine/Source/Runtime/RenderCore/Private/DumpGPU.cpp
--- a/Engine/Source/Runtime/RenderCore/Private/DumpGPU.cpp
+++ b/Engine/Source/Runtime/RenderCore/Private/DumpGPU.cpp
@@ -1252,6 +1252,7 @@ public:
 				FRDGTextureSRVDesc SubresourceDesc = SubresourceRangeDesc;
 				SubresourceDesc.FirstArraySlice = ArraySlice;
 				SubresourceDesc.NumArraySlices = 1;
+				SubresourceDesc.DimensionOverride = ETextureDimension::Texture2D;
 
 				AddDumpTextureAllSubResourcesPasses(
 					GraphBuilder,
