Fixes an error in DumpGPU viewer when the log's name is not ProjectName.log

diff --git a/Engine/Extras/GPUDumpViewer/GPUDumpViewer.html b/Engine/Extras/GPUDumpViewer/GPUDumpViewer.html
--- a/Engine/Extras/GPUDumpViewer/GPUDumpViewer.html
+++ b/Engine/Extras/GPUDumpViewer/GPUDumpViewer.html
@@ -1392,7 +1392,7 @@ function display_cvars()
 
 function get_log_path()
 {
-	return `Base/${g_infos['Project']}.log`;
+	return `Base/${g_infos['LogFilename']}`;
 }
 
 class LogView extends IView
diff --git a/Engine/Source/Runtime/RenderCore/Private/DumpGPU.cpp b/Engine/Source/Runtime/RenderCore/Private/DumpGPU.cpp
--- a/Engine/Source/Runtime/RenderCore/Private/DumpGPU.cpp
+++ b/Engine/Source/Runtime/RenderCore/Private/DumpGPU.cpp
@@ -3,6 +3,7 @@
 #include "DumpGPU.h"
 #include "HAL/PlatformFileManager.h"
 #include "HAL/PlatformMisc.h"
+#include "HAL/PlatformOutputDevices.h"
 #include "Misc/App.h"
 #include "Misc/FileHelper.h"
 #include "Misc/WildcardString.h"
@@ -1812,6 +1813,11 @@ FString FRDGBuilder::BeginResourceDump(const TCHAR* Cmd)
 		JsonObject->SetStringField(TEXT("RHI"), GDynamicRHI->GetName());
 		JsonObject->SetStringField(TEXT("RHIMaxFeatureLevel"), LexToString(GMaxRHIFeatureLevel));
 		JsonObject->SetStringField(TEXT("DumpTime"), NewResourceDumpContext->Time.ToString());
+		{
+			const FString LogSrcAbsolute = FPlatformOutputDevices::GetAbsoluteLogFilename();
+			FString LogFilename = FPaths::GetCleanFilename(LogSrcAbsolute);
+			JsonObject->SetStringField(TEXT("LogFilename"), LogFilename);
+		}
 
 		NewResourceDumpContext->DumpJsonToFile(JsonObject, FString(FRDGResourceDumpContext::kBaseDir) / TEXT("Infos.json"));
 	}
