[shaders] fix debug info output when preprocessed job cache is disabled

diff --git a/Engine/Source/Developer/ShaderCompilerCommon/Private/ShaderCompilerCommon.cpp b/Engine/Source/Developer/ShaderCompilerCommon/Private/ShaderCompilerCommon.cpp
--- a/Engine/Source/Developer/ShaderCompilerCommon/Private/ShaderCompilerCommon.cpp
+++ b/Engine/Source/Developer/ShaderCompilerCommon/Private/ShaderCompilerCommon.cpp
@@ -1391,7 +1391,7 @@ namespace UE::ShaderCompilerCommon
 {
 	void FBaseShaderFormat::OutputDebugData(const FString& InputHash, const FShaderCompilerInput& Input, const FShaderPreprocessOutput& PreprocessOutput, const FShaderCompilerOutput& Output) const
 	{
-		DumpDebugShaderData(Input, PreprocessOutput.GetOriginalSource());
+		DumpDebugShaderData(Input, Input.bCachePreprocessed ? PreprocessOutput.GetOriginalSource() : PreprocessOutput.GetSource());
 		FString OutputHashFileName = FPaths::Combine(Input.DumpDebugInfoPath, TEXT("OutputHash.txt"));
 		FFileHelper::SaveStringToFile(Output.OutputHash.ToString(), *OutputHashFileName, FFileHelper::EEncodingOptions::ForceAnsi);
 
@@ -1408,7 +1408,7 @@ namespace UE::ShaderCompilerCommon
 			}
 		}
 
-		if (EnumHasAnyFlags(Input.DebugInfoFlags, EShaderDebugInfoFlags::InputHash))
+		if (EnumHasAnyFlags(Input.DebugInfoFlags, EShaderDebugInfoFlags::InputHash) && !InputHash.IsEmpty())
 		{
 			FString InputHashFileName = FPaths::Combine(Input.DumpDebugInfoPath, TEXT("InputHash.txt"));
 			FFileHelper::SaveStringToFile(InputHash, *InputHashFileName, FFileHelper::EEncodingOptions::ForceAnsi);
diff --git a/Engine/Source/Runtime/RenderCore/Private/ShaderCore.cpp b/Engine/Source/Runtime/RenderCore/Private/ShaderCore.cpp
--- a/Engine/Source/Runtime/RenderCore/Private/ShaderCore.cpp
+++ b/Engine/Source/Runtime/RenderCore/Private/ShaderCore.cpp
@@ -1085,15 +1085,23 @@ void CompileShader(const TArray<const IShaderFormat*>& ShaderFormats, FShaderCom
 #endif
 
 	Job.bSucceeded = Job.Output.bSucceeded;
-
-	// skip this if the shader format supports independent preprocessing; this will be output as part of job completion callback for such jobs
-	if (!Compiler->SupportsIndependentPreprocessing() && Job.bSucceeded && Job.Input.DumpDebugInfoPath.Len() > 0)
+	if (Job.bSucceeded && Job.Input.DumpDebugInfoEnabled())
 	{
-		// write down the output hash as a file
-		FString HashFileName = FPaths::Combine(Job.Input.DumpDebugInfoPath, TEXT("OutputHash.txt"));
-		FFileHelper::SaveStringToFile(Job.Output.OutputHash.ToString(), *HashFileName, FFileHelper::EEncodingOptions::ForceAnsi);
+		// if the preprocessed cache is disabled, dump debug output here, since we don't serialize preprocess output 
+		// back to the cooker in this case (if enabled this will occur in the job OnComplete callback)
+		// pass an empty string for input hash; this is not available here but also not important if the preprocessed
+		// cache is disabled
+		if (!Job.Input.bCachePreprocessed)
+		{
+			Compiler->OutputDebugData(FString(), Job.Input, Job.PreprocessOutput, Job.Output);
+		}
+		else
+		{
+			// write down the output hash as a file
+			FString HashFileName = FPaths::Combine(Job.Input.DumpDebugInfoPath, TEXT("OutputHash.txt"));
+			FFileHelper::SaveStringToFile(Job.Output.OutputHash.ToString(), *HashFileName, FFileHelper::EEncodingOptions::ForceAnsi);
+		}
 	}
-
 }
 
 void CompileShaderPipeline(const TArray<const IShaderFormat*>& ShaderFormats, FShaderPipelineCompileJob* PipelineJob, const FString& Dir, int32* CompileCount)
@@ -2936,9 +2944,9 @@ void FShaderCompileJob::OnComplete()
 			PreprocessOutput.RemapErrors(Output);
 			Output.PreprocessTime = PreprocessOutput.ElapsedTime;
 		}
-		// dump debug info for the job whether or not the preprocessed cache was enabled, shader formats which support independent preprocessing
-		// will expect this to be done automatically.
-		if (Input.DumpDebugInfoEnabled())
+		// dump debug info for the job at this point if the preprocessed cache is enabled
+		// this is executed in the cooker process and we know the input hash is computed at this point
+		if (Input.bCachePreprocessed && Input.DumpDebugInfoEnabled())
 		{
 			ShaderFormat->OutputDebugData(LexToString(GetInputHash()), Input, PreprocessOutput, Output);
 		}
