Implements r.DumpGPU.FixedTickRate to tick the engine at fixed delta time when dumping many frames

diff --git a/Engine/Extras/GPUDumpViewer/GPUDumpViewer.html b/Engine/Extras/GPUDumpViewer/GPUDumpViewer.html
--- a/Engine/Extras/GPUDumpViewer/GPUDumpViewer.html
+++ b/Engine/Extras/GPUDumpViewer/GPUDumpViewer.html
@@ -1204,6 +1204,7 @@ const k_tips = [
 	'Speed up your frame dump by only selecting the passes you need with r.DumpGPU.Root. For instance r.DumpGPU.Root="*PostProcessing*".',
 	'Uses r.DumpGPU.Delay to delay the dump of few seconds to have time to repro the issue with gameplay logic (for instance moving arround in the map).',
 	'Uses r.DumpGPU.FrameCount to dump more than one frame. This is useful when artifact may be produced only some frames but end up stucked in temporal histories.',
+	'Uses r.DumpGPU.FixedTickRate to automatically override the engine\'s tick rate to fixed value when dumping multiple frames.',
 	'Uses r.DumpGPU.DumpOnScreenshotTest=1 to automatically produce a GPU dump of the frame producing screenshot in the automated AScreenshotFunctionalTest',
 	'GPU dumps can be large and accumulate on your hard drive in your various projects\' Saved/ directories. Set r.DumpGPU.Directory="D:/tmp/DumpGPU/" in your console variables or UE-DumpGPUPath environment variable to dump them all at the same location on your machine.',
 	'Uses -cvarsini to override console variables with your own ConsoleVariables.ini file on a cooked build.',
diff --git a/Engine/Source/Editor/UnrealEd/Private/SEditorViewport.cpp b/Engine/Source/Editor/UnrealEd/Private/SEditorViewport.cpp
--- a/Engine/Source/Editor/UnrealEd/Private/SEditorViewport.cpp
+++ b/Engine/Source/Editor/UnrealEd/Private/SEditorViewport.cpp
@@ -26,6 +26,9 @@
 #include "Widgets/Colors/SComplexGradient.h"
 #include "Modules/ModuleManager.h"
 #include "ISettingsModule.h"
+#if WITH_DUMPGPU
+	#include "RenderGraph.h"
+#endif
 
 #define LOCTEXT_NAMESPACE "EditorViewport"
 
@@ -598,10 +601,6 @@ bool SEditorViewport::IsRealtime() const
 
 bool SEditorViewport::IsVisible() const
 {
-#if WITH_DUMPGPU
-	extern ENGINE_API uint32 GDumpGPU_FrameNumber;
-#endif
-
 	const float VisibilityTimeThreshold = .25f;
 	// The viewport is visible if we don't have a parent layout (likely a floating window) or this viewport is visible in the parent layout.
 	// Also, always render the viewport if DumpGPU is active, regardless of tick time threshold -- otherwise these don't show up due to lag
@@ -610,7 +609,7 @@ bool SEditorViewport::IsVisible() const
 		LastTickTime == 0.0	||	// Never been ticked
 		FPlatformTime::Seconds() - LastTickTime <= VisibilityTimeThreshold	// Ticked recently
 #if WITH_DUMPGPU
-		|| GDumpGPU_FrameNumber == GFrameNumber	// GPU dump in progress
+		|| FRDGBuilder::IsDumpingFrame()	// GPU dump in progress
 #endif		
 		;
 }
diff --git a/Engine/Source/Runtime/Engine/Private/UnrealEngine.cpp b/Engine/Source/Runtime/Engine/Private/UnrealEngine.cpp
--- a/Engine/Source/Runtime/Engine/Private/UnrealEngine.cpp
+++ b/Engine/Source/Runtime/Engine/Private/UnrealEngine.cpp
@@ -5804,7 +5804,6 @@ bool UEngine::HandleProfileGPUCommand( const TCHAR* Cmd, FOutputDevice& Ar )
 #endif // WITH_PROFILEGPU
 
 #if WITH_DUMPGPU
-ENGINE_API uint32 GDumpGPU_FrameNumber = INDEX_NONE;
 
 bool UEngine::HandleDumpGPUCommand(const TCHAR* Cmd, FOutputDevice& Ar)
 {
@@ -5814,7 +5813,6 @@ bool UEngine::HandleDumpGPUCommand(const TCHAR* Cmd, FOutputDevice& Ar)
 	}
 
 	check(IsInGameThread());
-	GDumpGPU_FrameNumber = GFrameNumber;
 
 	TArray<FString> Args;
 	FString ResourceDumpDirectory = FRDGBuilder::BeginResourceDump(Cmd);
diff --git a/Engine/Source/Runtime/RenderCore/Private/DumpGPU.cpp b/Engine/Source/Runtime/RenderCore/Private/DumpGPU.cpp
--- a/Engine/Source/Runtime/RenderCore/Private/DumpGPU.cpp
+++ b/Engine/Source/Runtime/RenderCore/Private/DumpGPU.cpp
@@ -96,6 +96,11 @@ static TAutoConsoleVariable<int32> GDumpGPUFrameCount(
 	TEXT("Number of consecutive frames to dump (default=1)."),
 	ECVF_Default);
 
+static TAutoConsoleVariable<float> GDumpGPUFixedTickRate(
+	TEXT("r.DumpGPU.FixedTickRate"), 0.0f,
+	TEXT("Override the engine's tick rate to be fixed for every dumped frames (default=0)."),
+	ECVF_Default);
+
 static TAutoConsoleVariable<int32> GDumpGPUDraws(
 	TEXT("r.DumpGPU.Draws"), 0,
 	TEXT("Whether to dump resource after each individual draw call (disabled by default)."),
@@ -221,6 +226,9 @@ public:
 	TMap<const FRDGResource*, const FRDGPass*> LastResourceVersion;
 	TSet<const void*> IsDumpedToDisk;
 
+	bool bOverrideFixedDeltaTime = false;
+	double PreviousFixedDeltaTime = 0.0f;
+
 	// Pass being dumping individual draws
 	const FRDGPass* DrawDumpingPass = nullptr;
 	int32 DrawDumpCount = 0;
@@ -1767,6 +1775,15 @@ FString FRDGBuilder::BeginResourceDump(const TCHAR* Cmd)
 		NewResourceDumpContext->GetDumpParameters().DumpServiceParametersFile();
 	}
 
+	if (GDumpGPUFixedTickRate.GetValueOnGameThread() > 0.0f && !FApp::UseFixedTimeStep())
+	{
+		NewResourceDumpContext->bOverrideFixedDeltaTime = true;
+		NewResourceDumpContext->PreviousFixedDeltaTime = FApp::GetFixedDeltaTime();
+
+		FApp::SetFixedDeltaTime(1.0f / GDumpGPUFixedTickRate.GetValueOnGameThread());
+		FApp::SetUseFixedTimeStep(true);
+	}
+
 	// Output informations
 	{
 		const TCHAR* BranchName = BuildSettings::GetBranchName();
@@ -1978,6 +1995,13 @@ void FRDGBuilder::EndResourceDump()
 	}
 	#endif
 
+	// Restore the engine tick rate to what it was
+	if (GRDGResourceDumpContext->bOverrideFixedDeltaTime)
+	{
+		FApp::SetFixedDeltaTime(GRDGResourceDumpContext->PreviousFixedDeltaTime);
+		FApp::SetUseFixedTimeStep(false);
+	}
+
 	delete GRDGResourceDumpContext;
 	GRDGResourceDumpContext = nullptr;
 	DumpingFrameCounter_GameThread = 0;
