[shaders] fix condition preventing dump of outputhash.txt file to correctly account for backends which do not implement independent preprocessing API

diff --git a/Engine/Source/Runtime/RenderCore/Private/ShaderCore.cpp b/Engine/Source/Runtime/RenderCore/Private/ShaderCore.cpp
--- a/Engine/Source/Runtime/RenderCore/Private/ShaderCore.cpp
+++ b/Engine/Source/Runtime/RenderCore/Private/ShaderCore.cpp
@@ -1091,7 +1091,7 @@ void CompileShader(const TArray<const IShaderFormat*>& ShaderFormats, FShaderCom
 		// back to the cooker in this case (if enabled this will occur in the job OnComplete callback)
 		// pass an empty string for input hash; this is not available here but also not important if the preprocessed
 		// cache is disabled
-		if (!Job.Input.bCachePreprocessed)
+		if (Compiler->SupportsIndependentPreprocessing() && !Job.Input.bCachePreprocessed)
 		{
 			Compiler->OutputDebugData(FString(), Job.Input, Job.PreprocessOutput, Job.Output);
 		}
