Disable DCC for gbuffer if Nanite CS materials are enabled

diff --git a/Engine/Source/Runtime/RenderCore/Private/GBufferInfo.cpp b/Engine/Source/Runtime/RenderCore/Private/GBufferInfo.cpp
--- a/Engine/Source/Runtime/RenderCore/Private/GBufferInfo.cpp
+++ b/Engine/Source/Runtime/RenderCore/Private/GBufferInfo.cpp
@@ -185,7 +185,7 @@ FGBufferBinding FindGBufferBindingByName(const FGBufferInfo& GBufferInfo, const
 
 		if (NaniteComputeMaterialsSupported())
 		{
-			Binding.Flags |= TexCreate_UAV;
+			Binding.Flags |= TexCreate_UAV | TexCreate_DisableDCC;
 		}
 	}
 
