Moves DumpGPU's texture post processing in its own PostProcessTexture()

This notiably make the difference between a PreprocessedPixelFormat and PostprocessedPixelFormat in
FTextureSubresourceDumpDesc to handle the special logic for GL not able to write uav with less than 4byte
directly in TranslateSubresourceDumpDesc().

Motivation of this change is to have a very modular PostProcessTexture() easily reusable from asynchronous
stream GPU resource capture to disk.

diff --git a/Engine/Source/Runtime/RenderCore/Private/DumpGPU.cpp b/Engine/Source/Runtime/RenderCore/Private/DumpGPU.cpp
--- a/Engine/Source/Runtime/RenderCore/Private/DumpGPU.cpp
+++ b/Engine/Source/Runtime/RenderCore/Private/DumpGPU.cpp
@@ -801,6 +801,7 @@ public:
 		bool bPreprocessForStaging = false;
 		FDumpTextureCS::ETextureType DumpTextureType = FDumpTextureCS::ETextureType::MAX;
 		EPixelFormat PreprocessedPixelFormat;
+		EPixelFormat PostprocessedPixelFormat;
 
 		bool IsDumpSupported() const
 		{
@@ -816,7 +817,7 @@ public:
 		const FRDGTextureDesc& Desc = SubresourceDesc.Texture->Desc;
 
 		FTextureSubresourceDumpDesc SubresourceDumpDesc;
-		SubresourceDumpDesc.PreprocessedPixelFormat = Desc.Format;
+		SubresourceDumpDesc.PostprocessedPixelFormat = Desc.Format;
 
 		bool bIsUnsupported = false;
 
@@ -851,47 +852,45 @@ public:
 
 			if (SubresourceDesc.Format == PF_X24_G8)
 			{
-				SubresourceDumpDesc.PreprocessedPixelFormat = PF_R8_UINT;
+				SubresourceDumpDesc.PostprocessedPixelFormat = PF_R8_UINT;
 				SubresourceDumpDesc.DumpTextureType = FDumpTextureCS::ETextureType::Texture2DDepthStencilNoMSAA;
 			}
 			else if (Desc.Format == PF_DepthStencil)
 			{
-				SubresourceDumpDesc.PreprocessedPixelFormat = PF_R32_FLOAT;
+				SubresourceDumpDesc.PostprocessedPixelFormat = PF_R32_FLOAT;
 				SubresourceDumpDesc.DumpTextureType = FDumpTextureCS::ETextureType::Texture2DFloatNoMSAA;
 			}
 			else if (Desc.Format == PF_ShadowDepth)
 			{
-				SubresourceDumpDesc.PreprocessedPixelFormat = PF_R32_FLOAT;
+				SubresourceDumpDesc.PostprocessedPixelFormat = PF_R32_FLOAT;
 				SubresourceDumpDesc.DumpTextureType = FDumpTextureCS::ETextureType::Texture2DFloatNoMSAA;
 			}
 			else if (Desc.Format == PF_D24)
 			{
-				SubresourceDumpDesc.PreprocessedPixelFormat = PF_R32_FLOAT;
+				SubresourceDumpDesc.PostprocessedPixelFormat = PF_R32_FLOAT;
 				SubresourceDumpDesc.DumpTextureType = FDumpTextureCS::ETextureType::Texture2DFloatNoMSAA;
 			}
 			else if (Desc.Format == PF_BC4)
 			{
-				SubresourceDumpDesc.PreprocessedPixelFormat = PF_R32_FLOAT;
+				SubresourceDumpDesc.PostprocessedPixelFormat = PF_R32_FLOAT;
 				SubresourceDumpDesc.DumpTextureType = FDumpTextureCS::ETextureType::Texture2DFloatNoMSAA;
 			}
 			else if (Desc.Format == PF_BC5)
 			{
-				SubresourceDumpDesc.PreprocessedPixelFormat = PF_G16R16F;
+				SubresourceDumpDesc.PostprocessedPixelFormat = PF_G16R16F;
 				SubresourceDumpDesc.DumpTextureType = FDumpTextureCS::ETextureType::Texture2DFloatNoMSAA;
 			}
 			else if ((Desc.Format == PF_BC6H) || (Desc.Format == PF_BC7))
 			{
-				SubresourceDumpDesc.PreprocessedPixelFormat = PF_FloatRGBA;
+				SubresourceDumpDesc.PostprocessedPixelFormat = PF_FloatRGBA;
 				SubresourceDumpDesc.DumpTextureType = FDumpTextureCS::ETextureType::Texture2DFloatNoMSAA;
 			}
 		}
 
-		SubresourceDumpDesc.ByteSize = SIZE_T(SubresourceDumpDesc.SubResourceExtent.X) * SIZE_T(SubresourceDumpDesc.SubResourceExtent.Y) * SIZE_T(GPixelFormats[SubresourceDumpDesc.PreprocessedPixelFormat].BlockBytes);
-
 		// Whether the subresource need preprocessing pass before copy into staging.
 		{
 			// If need a pixel format conversion, use a pixel shader to do it.
-			SubresourceDumpDesc.bPreprocessForStaging |= SubresourceDumpDesc.PreprocessedPixelFormat != Desc.Format;
+			SubresourceDumpDesc.bPreprocessForStaging |= SubresourceDumpDesc.PostprocessedPixelFormat != Desc.Format;
 
 			// If the texture has a mip chain, use pixel shader to correctly copy the right mip given RHI doesn't support copy from mip levels. Also on Mip 0 to avoid bugs on D3D11
 			SubresourceDumpDesc.bPreprocessForStaging |= SubresourceDesc.Texture->Desc.NumMips > 1;
@@ -900,9 +899,72 @@ public:
 			SubresourceDumpDesc.bPreprocessForStaging |= SubresourceDesc.Texture->Desc.IsTextureArray();
 		}
 
+		// Some RHIs (GL) only support 32Bit single channel images as CS output
+		SubresourceDumpDesc.PreprocessedPixelFormat = SubresourceDumpDesc.PostprocessedPixelFormat;
+		if (SubresourceDumpDesc.bPreprocessForStaging)
+		{
+			if (IsOpenGLPlatform(GMaxRHIShaderPlatform) &&
+				GPixelFormats[SubresourceDumpDesc.PostprocessedPixelFormat].NumComponents == 1 &&
+				GPixelFormats[SubresourceDumpDesc.PostprocessedPixelFormat].BlockBytes < 4)
+			{
+				SubresourceDumpDesc.PreprocessedPixelFormat = PF_R32_UINT;
+			}
+		}
+
+		SubresourceDumpDesc.ByteSize = SIZE_T(SubresourceDumpDesc.SubResourceExtent.X) * SIZE_T(SubresourceDumpDesc.SubResourceExtent.Y) * SIZE_T(GPixelFormats[SubresourceDumpDesc.PreprocessedPixelFormat].BlockBytes);
+
 		return SubresourceDumpDesc;
 	}
 
+	void PostProcessTexture(
+		const FTextureSubresourceDumpDesc& SubresourceDumpDesc,
+		void* Content,
+		int32 RowPitchInPixels,
+		int32 ColumnPitchInPixels,
+		TArray64<uint8>& Array)
+	{
+		FTimeBucketMeasure TimeBucketMeasure(this, ETimingBucket::CPUPostProcessing);
+
+		Array.SetNumUninitialized(SubresourceDumpDesc.ByteSize);
+
+		SIZE_T BytePerPixel = SIZE_T(GPixelFormats[SubresourceDumpDesc.PreprocessedPixelFormat].BlockBytes);
+
+		const uint8* SrcData = static_cast<const uint8*>(Content);
+
+		for (int32 y = 0; y < SubresourceDumpDesc.SubResourceExtent.Y; y++)
+		{
+			// Flip the data to be bottom left corner for the WebGL viewer.
+			const uint8* SrcPos = SrcData + SIZE_T(SubresourceDumpDesc.SubResourceExtent.Y - 1 - y) * SIZE_T(RowPitchInPixels) * BytePerPixel;
+			uint8* DstPos = (&Array[0]) + SIZE_T(y) * SIZE_T(SubresourceDumpDesc.SubResourceExtent.X) * BytePerPixel;
+
+			FPlatformMemory::Memmove(DstPos, SrcPos, SIZE_T(SubresourceDumpDesc.SubResourceExtent.X) * BytePerPixel);
+		}
+
+		if (SubresourceDumpDesc.PreprocessedPixelFormat != SubresourceDumpDesc.PostprocessedPixelFormat)
+		{
+			// Convert 32Bit values back to 16 or 8bit
+			const int32 DstPixelNumBytes = GPixelFormats[SubresourceDumpDesc.PostprocessedPixelFormat].BlockBytes;
+			const uint32* SrcData32 = (const uint32*)Array.GetData();
+			uint8* DstData8 = Array.GetData();
+			uint16* DstData16 = (uint16*)Array.GetData();
+
+			for (int32 Index = 0; Index < Array.Num() / 4; Index++)
+			{
+				uint32 Value32 = SrcData32[Index];
+				if (DstPixelNumBytes == 2)
+				{
+					DstData16[Index] = (uint16)Value32;
+				}
+				else
+				{
+					DstData8[Index] = (uint8)Value32;
+				}
+			}
+
+			Array.SetNum(Array.Num() / (4 / DstPixelNumBytes));
+		}
+	}
+
 	void DumpTextureSubResource(
 		FRHICommandListImmediate& RHICmdList,
 		const TCHAR* TextureDebugName,
@@ -915,26 +977,15 @@ public:
 
 		// Preprocess
 		FTextureRHIRef StagingSrcTexture;
-		EPixelFormat PreprocessedPixelFormat = SubresourceDumpDesc.PreprocessedPixelFormat;
-		SIZE_T SubresourceByteSize = SubresourceDumpDesc.ByteSize;
 		FTextureRHIRef StagingTexture;
 
 		{
 			FTimeBucketMeasure TimeBucketMeasure(this, ETimingBucket::RHIReadbackCommands);
 			if (SubresourceDumpDesc.bPreprocessForStaging)
 			{
-				// Some RHIs (GL) only support 32Bit single channel images as CS output
-				if (IsOpenGLPlatform(GMaxRHIShaderPlatform) &&
-					GPixelFormats[PreprocessedPixelFormat].NumComponents == 1 && 
-					GPixelFormats[PreprocessedPixelFormat].BlockBytes < 4)
-				{
-					SubresourceByteSize*= (4 / GPixelFormats[PreprocessedPixelFormat].BlockBytes);
-					PreprocessedPixelFormat = PF_R32_UINT;
-				}
-
 				{
 					const FRHITextureCreateDesc Desc =
-						FRHITextureCreateDesc::Create2D(TEXT("DumpGPU.PreprocessTexture"), SubresourceDumpDesc.SubResourceExtent, PreprocessedPixelFormat)
+						FRHITextureCreateDesc::Create2D(TEXT("DumpGPU.PreprocessTexture"), SubresourceDumpDesc.SubResourceExtent, SubresourceDumpDesc.PreprocessedPixelFormat)
 						.SetFlags(ETextureCreateFlags::UAV | ETextureCreateFlags::ShaderResource | ETextureCreateFlags::HideInVisualizeTexture);
 
 					StagingSrcTexture = RHICreateTexture(Desc);
@@ -968,7 +1019,7 @@ public:
 			// Copy the texture for CPU readback
 			{
 				const FRHITextureCreateDesc Desc =
-					FRHITextureCreateDesc::Create2D(TEXT("DumpGPU.StagingTexture"), SubresourceDumpDesc.SubResourceExtent, PreprocessedPixelFormat)
+					FRHITextureCreateDesc::Create2D(TEXT("DumpGPU.StagingTexture"), SubresourceDumpDesc.SubResourceExtent, SubresourceDumpDesc.PreprocessedPixelFormat)
 					.SetFlags(ETextureCreateFlags::CPUReadback | ETextureCreateFlags::HideInVisualizeTexture);
 
 				StagingTexture = RHICreateTexture(Desc);
@@ -1011,50 +1062,12 @@ public:
 		if (Content)
 		{
 			TArray64<uint8> Array;
-			{
-				FTimeBucketMeasure TimeBucketMeasure(this, ETimingBucket::CPUPostProcessing);
-
-				Array.SetNumUninitialized(SubresourceByteSize);
-
-				SIZE_T BytePerPixel = SIZE_T(GPixelFormats[PreprocessedPixelFormat].BlockBytes);
-
-				const uint8* SrcData = static_cast<const uint8*>(Content);
-
-				for (int32 y = 0; y < SubresourceDumpDesc.SubResourceExtent.Y; y++)
-				{
-					// Flip the data to be bottom left corner for the WebGL viewer.
-					const uint8* SrcPos = SrcData + SIZE_T(SubresourceDumpDesc.SubResourceExtent.Y - 1 - y) * SIZE_T(RowPitchInPixels) * BytePerPixel;
-					uint8* DstPos = (&Array[0]) + SIZE_T(y) * SIZE_T(SubresourceDumpDesc.SubResourceExtent.X) * BytePerPixel;
-
-					FPlatformMemory::Memmove(DstPos, SrcPos, SIZE_T(SubresourceDumpDesc.SubResourceExtent.X) * BytePerPixel);
-				}
-
-				RHICmdList.UnmapStagingSurface(StagingTexture, GPUIndex);
-
-				if (PreprocessedPixelFormat != SubresourceDumpDesc.PreprocessedPixelFormat)
-				{
-					// Convert 32Bit values back to 16 or 8bit
-					const int32 DstPixelNumBytes = GPixelFormats[SubresourceDumpDesc.PreprocessedPixelFormat].BlockBytes;
-					const uint32* SrcData32 = (const uint32*)Array.GetData();
-					uint8* DstData8 = Array.GetData();
-					uint16* DstData16 = (uint16*)Array.GetData();
-											
-					for (int32 Index = 0; Index < Array.Num()/4; Index++)
-					{
-						uint32 Value32 = SrcData32[Index];
-						if (DstPixelNumBytes == 2)
-						{
-							DstData16[Index] = (uint16)Value32;
-						}
-						else
-						{
-							DstData8[Index] = (uint8)Value32;
-						}
-					}
-
-					Array.SetNum(Array.Num() / (4 / DstPixelNumBytes));
-				}
-			}
+			PostProcessTexture(
+				SubresourceDumpDesc,
+				Content,
+				RowPitchInPixels,
+				ColumnPitchInPixels,
+				/* out */ Array);
 
 			DumpResourceBinaryToFile(Array.GetData(), Array.Num(), DumpFilePath);
 		}
