Disable AsyncCompute and ParallelExecute when FlushGPU is enabled.
- Previously there was a callback associated with r.RDG.Debug.FlushGPU that tried to disable AsyncCompute however it didn't work because it got reenabled by CVarRDGAsyncComputeSink.

diff --git a/Engine/Source/Runtime/RenderCore/Private/RenderGraphBuilder.cpp b/Engine/Source/Runtime/RenderCore/Private/RenderGraphBuilder.cpp
--- a/Engine/Source/Runtime/RenderCore/Private/RenderGraphBuilder.cpp
+++ b/Engine/Source/Runtime/RenderCore/Private/RenderGraphBuilder.cpp
@@ -2931,13 +2931,11 @@ void FRDGBuilder::ExecutePass(FRDGPass* Pass, FRHIComputeCommandList& RHICmdList
 		}
 	}
 
-	if (!bParallelExecuteEnabled)
+	if (GRDGDebugFlushGPU)
 	{
-		if (GRDGDebugFlushGPU && !GRDGAsyncCompute)
-		{
-			RHICmdList.SubmitCommandsAndFlushGPU();
-			RHICmdList.BlockUntilGPUIdle();
-		}
+		check(!GRDGAsyncCompute && !bParallelExecuteEnabled);
+		RHICmdList.SubmitCommandsAndFlushGPU();
+		RHICmdList.BlockUntilGPUIdle();
 	}
 }
 
diff --git a/Engine/Source/Runtime/RenderCore/Private/RenderGraphPrivate.cpp b/Engine/Source/Runtime/RenderCore/Private/RenderGraphPrivate.cpp
--- a/Engine/Source/Runtime/RenderCore/Private/RenderGraphPrivate.cpp
+++ b/Engine/Source/Runtime/RenderCore/Private/RenderGraphPrivate.cpp
@@ -39,16 +39,9 @@ int32 GRDGDebugFlushGPU = 0;
 FAutoConsoleVariableRef CVarRDGDebugFlushGPU(
 	TEXT("r.RDG.Debug.FlushGPU"),
 	GRDGDebugFlushGPU,
-	TEXT("Enables flushing the GPU after every pass. Disables async compute when set (r.RDG.AsyncCompute=0).\n")
+	TEXT("Enables flushing the GPU after every pass. Disables async compute (r.RDG.AsyncCompute=0) and parallel execute (r.RDG.ParallelExecute=0) when set.\n")
 	TEXT(" 0: disabled (default);\n")
-	TEXT(" 1: enabled (default)."),
-	FConsoleVariableDelegate::CreateLambda([](IConsoleVariable*)
-	{
-		if (GRDGDebugFlushGPU)
-		{
-			GRDGAsyncCompute = 0;
-		}
-	}),
+	TEXT(" 1: enabled."),
 	ECVF_RenderThreadSafe);
 
 int32 GRDGDebugExtendResourceLifetimes = 0;
@@ -280,6 +273,11 @@ FAutoConsoleVariableSink CVarRDGAsyncComputeSink(FConsoleCommandDelegate::Create
 {
 	GRDGAsyncCompute = CVarRDGAsyncCompute.GetValueOnGameThread();
 
+	if (GRDGDebugFlushGPU)
+	{
+		GRDGAsyncCompute = 0;
+	}
+
 	if (!IsAsyncComputeSupported())
 	{
 		GRDGAsyncCompute = 0;
@@ -633,6 +631,7 @@ bool IsParallelExecuteEnabled()
 		&& !GRHICommandList.Bypass()
 		&& !IsImmediateMode()
 		&& !GRDGDebug
+		&& !GRDGDebugFlushGPU
 		&& !GRDGTransitionLog
 		&& !IsMobilePlatform(GMaxRHIShaderPlatform)
 		&& !IsOpenGLPlatform(GMaxRHIShaderPlatform)
