[shaders]
- expose a helper function for shader file path fixup
- fix an uinintentional copy in the FindShaderFormat API

diff --git a/Engine/Source/Runtime/RenderCore/Private/ShaderCore.cpp b/Engine/Source/Runtime/RenderCore/Private/ShaderCore.cpp
--- a/Engine/Source/Runtime/RenderCore/Private/ShaderCore.cpp
+++ b/Engine/Source/Runtime/RenderCore/Private/ShaderCore.cpp
@@ -818,7 +818,7 @@ bool CheckVirtualShaderFilePath(FStringView VirtualFilePath, TArray<FShaderCompi
 	return bSuccess;
 }
 
-const IShaderFormat* FindShaderFormat(FName Format, TArray<const IShaderFormat*> ShaderFormats)
+const IShaderFormat* FindShaderFormat(FName Format, const TArray<const IShaderFormat*>& ShaderFormats)
 {
 	for (int32 Index = 0; Index < ShaderFormats.Num(); Index++)
 	{
@@ -1440,6 +1440,15 @@ bool ReplaceVirtualFilePathForShaderAutogen(FString& InOutVirtualFilePath, EShad
 	return false;
 }
 
+void FixupShaderFilePath(FString& VirtualFilePath, EShaderPlatform ShaderPlatform, const FName* ShaderPlatformName)
+{
+	// Always substitute virtual platform path before accessing GShaderFileCache to get platform-specific file.
+	ReplaceVirtualFilePathForShaderPlatform(VirtualFilePath, ShaderPlatform);
+
+	// Fixup autogen file
+	ReplaceVirtualFilePathForShaderAutogen(VirtualFilePath, ShaderPlatform, ShaderPlatformName);
+}
+
 bool LoadShaderSourceFile(const TCHAR* InVirtualFilePath, EShaderPlatform ShaderPlatform, FString* OutFileContents, TArray<FShaderCompilerError>* OutCompileErrors, const FName* ShaderPlatformName) // TODO: const FString&
 {
 #if WITH_EDITORONLY_DATA
@@ -1457,12 +1466,7 @@ bool LoadShaderSourceFile(const TCHAR* InVirtualFilePath, EShaderPlatform Shader
 		SCOPE_SECONDS_COUNTER(ShaderFileLoadingTime);
 
 		FString VirtualFilePath(InVirtualFilePath);
-
-		// Always substitute virtual platform path before accessing GShaderFileCache to get platform-specific file.
-		ReplaceVirtualFilePathForShaderPlatform(VirtualFilePath, ShaderPlatform);
-
-		// Fixup autogen file
-		ReplaceVirtualFilePathForShaderAutogen(VirtualFilePath, ShaderPlatform, ShaderPlatformName);
+		FixupShaderFilePath(VirtualFilePath, ShaderPlatform, ShaderPlatformName);
 
 		FString* CachedFile = nullptr;
 
diff --git a/Engine/Source/Runtime/RenderCore/Public/ShaderCompilerCore.h b/Engine/Source/Runtime/RenderCore/Public/ShaderCompilerCore.h
--- a/Engine/Source/Runtime/RenderCore/Public/ShaderCompilerCore.h
+++ b/Engine/Source/Runtime/RenderCore/Public/ShaderCompilerCore.h
@@ -630,7 +630,7 @@ struct RENDERCORE_API FSCWErrorCode
 #if PLATFORM_WINDOWS
 extern RENDERCORE_API int HandleShaderCompileException(Windows::LPEXCEPTION_POINTERS Info, FString& OutExMsg, FString& OutCallStack);
 #endif
-extern RENDERCORE_API const IShaderFormat* FindShaderFormat(FName Format, TArray<const IShaderFormat*> ShaderFormats);
+extern RENDERCORE_API const IShaderFormat* FindShaderFormat(FName Format, const TArray<const IShaderFormat*>& ShaderFormats);
 
 // Executes preprocessing for the given job, if the job is marked to be preprocessed independently prior to compilation.
 extern RENDERCORE_API void ConditionalPreprocessShader(FShaderCommonCompileJob* Job);
@@ -653,6 +653,11 @@ inline bool ShouldUseStableConstantBuffer(const FShaderCompilerInput& Input)
  */
 extern RENDERCORE_API bool CheckVirtualShaderFilePath(FStringView VirtualPath, TArray<FShaderCompilerError>* CompileErrors = nullptr);
 
+/**
+ * Fixes up the given virtual file path (substituting virtual platform path/autogen path for the given platform)
+ */
+extern RENDERCORE_API void FixupShaderFilePath(FString& VirtualFilePath, EShaderPlatform ShaderPlatform, const FName* ShaderPlatformName);
+
 /**
  * Loads the shader file with the given name.
  * @param VirtualFilePath - The virtual path of shader file to load.
