GlobalBeginCompileShader optimizations.  3.2x improvement in performance of the function on Lyra cook:
* Early out when adding uniform buffer resource entries, by checking if uniform buffer had been processed, before processing individual entries.
* Changed resource entry list from a map to an array, saving the cost of map operations adding items to it.
* Resource entries for global uniform buffers cached once at startup.
* Resource entry member names stored in an external buffer, to avoid per-entry string memory allocation overhead.
* Miscellaneous smaller optimizations -- for example, using ByHash to avoid redundant hashing and FString construction.

#rnx

diff --git a/Engine/Source/Developer/Apple/MetalShaderFormat/Private/MetalDerivedData.cpp b/Engine/Source/Developer/Apple/MetalShaderFormat/Private/MetalDerivedData.cpp
--- a/Engine/Source/Developer/Apple/MetalShaderFormat/Private/MetalDerivedData.cpp
+++ b/Engine/Source/Developer/Apple/MetalShaderFormat/Private/MetalDerivedData.cpp
@@ -229,7 +229,7 @@ bool DoCompileMetalShader(
 	uint32 SourceLen = 0;
 	int32 Result = 0;
 	
-	struct FMetalResourceTableEntry : FResourceTableEntry
+	struct FMetalResourceTableEntry : FUniformResourceEntry
 	{
 		FString Name;
 		uint32 Size;
@@ -365,19 +365,19 @@ bool DoCompileMetalShader(
 					TableNames.Add(*Pair.Key);
 				}
 				
-				for (auto Pair : Input.Environment.ResourceTableMap)
+				for (const FUniformResourceEntry& Entry : Input.Environment.ResourceTableMap.Resources)
 				{
-					const FResourceTableEntry& Entry = Pair.Value;
-					TArray<FMetalResourceTableEntry>& Resources = IABs.FindOrAdd(Entry.UniformBufferName);
+					TArray<FMetalResourceTableEntry>& Resources = IABs.FindOrAdd(FString(Entry.GetUniformBufferName()));
 					if ((uint32)Resources.Num() <= Entry.ResourceIndex)
 					{
 						Resources.SetNum(Entry.ResourceIndex + 1);
 					}
 					FMetalResourceTableEntry NewEntry;
-					NewEntry.UniformBufferName = Entry.UniformBufferName;
+					NewEntry.UniformBufferMemberName = Entry.UniformBufferMemberName;
+					NewEntry.UniformBufferNameLength = Entry.UniformBufferNameLength;
 					NewEntry.Type = Entry.Type;
 					NewEntry.ResourceIndex = Entry.ResourceIndex;
-					NewEntry.Name = Pair.Key;
+					NewEntry.Name = Entry.UniformBufferMemberName;
 					NewEntry.Size = 1;
 					NewEntry.bUsed = false;
 					Resources[Entry.ResourceIndex] = NewEntry;
@@ -417,7 +417,7 @@ bool DoCompileMetalShader(
 						}
 						for (uint32 j = 0; j < (uint32)TableNames.Num(); j++)
 						{
-							if (Entry.UniformBufferName == TableNames[j])
+							if (Entry.GetUniformBufferName() == TableNames[j])
 							{
 								Entry.SetIndex = j;
 								break;
@@ -474,7 +474,7 @@ bool DoCompileMetalShader(
 						ResourceBindings.Add(Binding);
 						
 						FMetalResourceTableEntry Entry = ResourceTable.FindRef(UTF8_TO_TCHAR(Binding->name));
-						UsedSets.Add(Entry.UniformBufferName);
+						UsedSets.Add(FString(Entry.GetUniformBufferName()));
 						
 						continue;
 					}
@@ -524,7 +524,8 @@ bool DoCompileMetalShader(
 						}
 						
 						FMetalResourceTableEntry Entry;
-						Entry.UniformBufferName = LastResource.UniformBufferName;
+						Entry.UniformBufferMemberName = LastResource.UniformBufferMemberName;
+						Entry.UniformBufferNameLength = LastResource.UniformBufferNameLength;
 						Entry.Name = Name;
 						Entry.ResourceIndex = ResIndex;
 						Entry.SetIndex = SetIndex;
@@ -608,10 +609,11 @@ bool DoCompileMetalShader(
 					for (auto const& Binding : ResourceBindings)
 					{
 						FMetalResourceTableEntry* Entry = ResourceTable.Find(UTF8_TO_TCHAR(Binding->name));
-						auto* ResourceArray = IABs.Find(Entry->UniformBufferName);
-						if (!IABTier1Index.Contains(Entry->UniformBufferName))
+						FString EntryUniformBufferName(Entry->GetUniformBufferName());
+						auto* ResourceArray = IABs.Find(EntryUniformBufferName);
+						if (!IABTier1Index.Contains(EntryUniformBufferName))
 						{
-							IABTier1Index.Add(Entry->UniformBufferName, 0);
+							IABTier1Index.Add(EntryUniformBufferName, 0);
 						}
 						if (Binding->descriptor_type == SPV_REFLECT_DESCRIPTOR_TYPE_STORAGE_BUFFER)
 						{
@@ -627,7 +629,8 @@ bool DoCompileMetalShader(
 							if (!bFoundBufferSizes)
 							{
 								FMetalResourceTableEntry BufferSizes;
-								BufferSizes.UniformBufferName = Entry->UniformBufferName;
+								BufferSizes.UniformBufferMemberName = Entry->UniformBufferMemberName;
+								BufferSizes.UniformBufferNameLength = Entry->UniformBufferNameLength;
 								BufferSizes.Name = TEXT("BufferSizes");
 								BufferSizes.Type = UBMT_SRV;
 								BufferSizes.ResourceIndex = 65535;
@@ -635,7 +638,7 @@ bool DoCompileMetalShader(
 								BufferSizes.Size = 1;
 								BufferSizes.bUsed = true;
 								ResourceArray->Insert(BufferSizes, 0);
-								IABTier1Index[Entry->UniformBufferName] = 1;
+								IABTier1Index[EntryUniformBufferName] = 1;
 							}
 						}
 					}
@@ -644,10 +647,11 @@ bool DoCompileMetalShader(
 				for (auto const& Binding : ResourceBindings)
 				{
 					FMetalResourceTableEntry* Entry = ResourceTable.Find(UTF8_TO_TCHAR(Binding->name));
+					FString EntryUniformBufferName(Entry->GetUniformBufferName());
 					
 					for (uint32 j = 0; j < (uint32)TableNames.Num(); j++)
 					{
-						if (Entry->UniformBufferName == TableNames[j])
+						if (EntryUniformBufferName == TableNames[j])
 						{
 							Entry->SetIndex = j;
 							BufferIndices &= ~(1ull << ((uint64)j + IABOffsetIndex));
@@ -657,7 +661,7 @@ bool DoCompileMetalShader(
 					}
 					Entry->bUsed = true;
 					
-					auto* ResourceArray = IABs.Find(Entry->UniformBufferName);
+					auto* ResourceArray = IABs.Find(EntryUniformBufferName);
 					uint32 ResourceIndex = Entry->ResourceIndex;
 					if (IABTier == 1)
 					{
@@ -666,7 +670,7 @@ bool DoCompileMetalShader(
 							Resource.SetIndex = Entry->SetIndex;
 							if (Resource.ResourceIndex == Entry->ResourceIndex)
 							{
-								uint32& Tier1Index = IABTier1Index.FindChecked(Entry->UniformBufferName);
+								uint32& Tier1Index = IABTier1Index.FindChecked(EntryUniformBufferName);
 								ResourceIndex = Tier1Index++;
 								Resource.bUsed = true;
 								break;
diff --git a/Engine/Source/Developer/ShaderCompilerCommon/Private/ShaderCompilerCommon.cpp b/Engine/Source/Developer/ShaderCompilerCommon/Private/ShaderCompilerCommon.cpp
--- a/Engine/Source/Developer/ShaderCompilerCommon/Private/ShaderCompilerCommon.cpp
+++ b/Engine/Source/Developer/ShaderCompilerCommon/Private/ShaderCompilerCommon.cpp
@@ -78,7 +78,7 @@ void BuildResourceTableTokenStream(const TArray<uint32>& InResourceMap, int32 Ma
 
 
 bool BuildResourceTableMapping(
-	const TMap<FString, FResourceTableEntry>& ResourceTableMap,
+	const FShaderResourceTableMap& ResourceTableMap,
 	const TMap<FString, FUniformBufferEntry>& UniformBufferMap,
 	TBitArray<>& UsedUniformBufferSlots,
 	FShaderParameterMap& ParameterMap,
@@ -91,10 +91,9 @@ bool BuildResourceTableMapping(
 	int32 MaxBoundResourceTable = -1;
 
 	// Go through ALL the members of ALL the UB resources
-	for (const TPair<FString, FResourceTableEntry>& ResourceEntry : ResourceTableMap)
+	for (const FUniformResourceEntry& Entry : ResourceTableMap.Resources)
 	{
-		const FString& Name = ResourceEntry.Key;
-		const FResourceTableEntry& Entry = ResourceEntry.Value;
+		const FString& Name = Entry.UniformBufferMemberName;
 
 		// If the shaders uses this member (eg View_PerlinNoise3DTexture)...
 		if (TOptional<FParameterAllocation> Allocation = ParameterMap.FindParameterAllocation(Name))
@@ -111,10 +110,11 @@ bool BuildResourceTableMapping(
 			uint16 UBBaseIndex, UBSize;
 
 			// Add the UB itself as a parameter if not there
-			if (!ParameterMap.FindParameterAllocation(*Entry.UniformBufferName, UniformBufferIndex, UBBaseIndex, UBSize))
+			FString UniformBufferName(Entry.GetUniformBufferName());
+			if (!ParameterMap.FindParameterAllocation(*UniformBufferName, UniformBufferIndex, UBBaseIndex, UBSize))
 			{
 				UniformBufferIndex = UsedUniformBufferSlots.FindAndSetFirstZeroBit();
-				ParameterMap.AddParameterAllocation(*Entry.UniformBufferName,UniformBufferIndex,0,0,EShaderParameterType::UniformBuffer);
+				ParameterMap.AddParameterAllocation(*UniformBufferName, UniformBufferIndex,0,0,EShaderParameterType::UniformBuffer);
 			}
 
 			// Mark used UB index
@@ -179,6 +179,20 @@ bool BuildResourceTableMapping(
 	return true;
 }
 
+PRAGMA_DISABLE_DEPRECATION_WARNINGS
+// Deprecated version of function
+bool BuildResourceTableMapping(
+	const TMap<FString, FResourceTableEntry>& ResourceTableMap,
+	const TMap<FString, FUniformBufferEntry>& UniformBufferMap,
+	TBitArray<>& UsedUniformBufferSlots,
+	FShaderParameterMap& ParameterMap,
+	FShaderCompilerResourceTable& OutSRT)
+{
+	UE_LOG(LogShaders, Error, TEXT("Using unimplemented deprecated version of BuildResourceTableMapping -- use version that accepts FShaderResourceTableMap instead."));
+	return false;
+}
+PRAGMA_ENABLE_DEPRECATION_WARNINGS
+
 void CullGlobalUniformBuffers(const TMap<FString, FUniformBufferEntry>& UniformBufferMap, FShaderParameterMap& ParameterMap)
 {
 	TArray<FString> ParameterNames;
@@ -1538,10 +1552,9 @@ namespace CrossCompiler
 			Line += FString::Printf(TEXT("%s, %d\n"), *Pair.Key, Pair.Value.LayoutHash);
 		}
 		Line += TEXT("NULL, 0\n");
-		for (auto Pair : Environment.ResourceTableMap)
+		for (const FUniformResourceEntry& Entry : Environment.ResourceTableMap.Resources)
 		{
-			const FResourceTableEntry& Entry = Pair.Value;
-			Line += FString::Printf(TEXT("%s, %s, %d, %d\n"), *Pair.Key, *Entry.UniformBufferName, Entry.Type, Entry.ResourceIndex);
+			Line += FString::Printf(TEXT("%s, %s, %d, %d\n"), Entry.UniformBufferMemberName, *FString(Entry.GetUniformBufferName()), Entry.Type, Entry.ResourceIndex);
 		}
 		Line += TEXT("NULL, NULL, 0, 0\n");
 
@@ -1600,8 +1613,31 @@ namespace CrossCompiler
 
 			FUniformBufferEntry& UniformBufferEntry = OutEnvironment.UniformBufferMap.FindOrAdd(UB);
 			UniformBufferEntry.LayoutHash = (uint32)Hash;
+
+			if (!UniformBufferEntry.MemberNameBuffer)
+			{
+				TArray<TCHAR>* MemberNameBuffer = new TArray<TCHAR>();
+				UniformBufferEntry.MemberNameBuffer = MakeShareable(MemberNameBuffer);
+			}
 		}
 
+		// Need to iterate through Uniform Buffer Map to add strings to correct MemberNameBuffer storage
+		auto UniformBufferMapIt = OutEnvironment.UniformBufferMap.begin();
+		
+		// If we exit parse early due to error, we still want to fixup the string names for the members we found,
+		// so the partial data isn't corrupt.
+		struct FFixupOnExit
+		{
+			FFixupOnExit(FShaderCompilerEnvironment& OutEnvironment) : Environment(OutEnvironment) {}
+			~FFixupOnExit()
+			{
+				Environment.ResourceTableMap.FixupOnLoad(Environment.UniformBufferMap);
+			}
+			
+			FShaderCompilerEnvironment& Environment;
+		};
+		FFixupOnExit FixupOnExit(OutEnvironment);
+
 		while (Ptr < PtrEnd)
 		{
 			FString Name;
@@ -1647,10 +1683,30 @@ namespace CrossCompiler
 			{
 				break;
 			}
-			FResourceTableEntry& Entry = OutEnvironment.ResourceTableMap.FindOrAdd(Name);
-			Entry.UniformBufferName = UB;
-			Entry.Type = Type;
-			Entry.ResourceIndex = ResourceIndex;
+
+			// Advance the uniform buffer map if this is a different UB name
+			while (UniformBufferMapIt.Key() != UB)
+			{
+				++UniformBufferMapIt;
+				if (UniformBufferMapIt == OutEnvironment.UniformBufferMap.end())
+				{
+					return;
+				}
+			}
+
+			// Append the Name we parsed to the member name buffer
+			TArray<TCHAR>& Buffer = *UniformBufferMapIt.Value().MemberNameBuffer.Get();
+			uint32 MemberNameLength = Name.Len();
+			
+			Buffer.Append(*Name, MemberNameLength + 1);
+
+			// The member name field of the entries is initialized at the end of parsing by the FixupOnLoad call from FFixupOnExit, so we can set it to nullptr here
+			OutEnvironment.ResourceTableMap.Resources.Add({
+				nullptr,
+				(uint8)UB.Len(),
+				(uint8)Type,
+				(uint16)ResourceIndex
+			});
 		}
 	}
 
diff --git a/Engine/Source/Developer/ShaderCompilerCommon/Public/ShaderCompilerCommon.h b/Engine/Source/Developer/ShaderCompilerCommon/Public/ShaderCompilerCommon.h
--- a/Engine/Source/Developer/ShaderCompilerCommon/Public/ShaderCompilerCommon.h
+++ b/Engine/Source/Developer/ShaderCompilerCommon/Public/ShaderCompilerCommon.h
@@ -32,6 +32,16 @@ namespace UE::ShaderCompilerCommon
  * buffer we allocate a slot for it from UsedUniformBufferSlots.
  * Returns false if there's any internal error.
  */
+extern SHADERCOMPILERCOMMON_API bool BuildResourceTableMapping(
+		const FShaderResourceTableMap& ResourceTableMap,
+		const TMap<FString,FUniformBufferEntry>& UniformBufferMap,
+		TBitArray<>& UsedUniformBufferSlots,
+		FShaderParameterMap& ParameterMap,
+		FShaderCompilerResourceTable& OutSRT
+	);
+
+PRAGMA_DISABLE_DEPRECATION_WARNINGS
+UE_DEPRECATED(5.3, "No longer supported; use version of function that accepts a FShaderResourceTableMap instead")
 extern SHADERCOMPILERCOMMON_API bool BuildResourceTableMapping(
 		const TMap<FString,FResourceTableEntry>& ResourceTableMap,
 		const TMap<FString,FUniformBufferEntry>& UniformBufferMap,
@@ -39,6 +49,7 @@ extern SHADERCOMPILERCOMMON_API bool BuildResourceTableMapping(
 		FShaderParameterMap& ParameterMap,
 		FShaderCompilerResourceTable& OutSRT
 	);
+PRAGMA_ENABLE_DEPRECATION_WARNINGS
 
 /** Culls global uniform buffer entries from the parameter map. */
 extern SHADERCOMPILERCOMMON_API void CullGlobalUniformBuffers(const TMap<FString, FUniformBufferEntry>& UniformBufferMap, FShaderParameterMap& ParameterMap);
diff --git a/Engine/Source/Developer/VulkanShaderFormat/Private/VulkanShaderCompiler.cpp b/Engine/Source/Developer/VulkanShaderFormat/Private/VulkanShaderCompiler.cpp
--- a/Engine/Source/Developer/VulkanShaderFormat/Private/VulkanShaderCompiler.cpp
+++ b/Engine/Source/Developer/VulkanShaderFormat/Private/VulkanShaderCompiler.cpp
@@ -153,19 +153,18 @@ static bool ContainsBinding(const FVulkanBindingTable& BindingTable, const FStri
 	return false;
 }
 
-static FString GetResourceEntryFromUBMember(const TMap<FString, FResourceTableEntry>& ResourceTableMap, const FString& UBName, uint16 ResourceIndex, FResourceTableEntry& OutEntry)
+static void GetResourceEntryFromUBMember(const FShaderResourceTableMap& ResourceTableMap, const FString& UBName, uint16 ResourceIndex, FUniformResourceEntry& OutEntry)
 {
-	for (const auto& Pair : ResourceTableMap)
+	for (const FUniformResourceEntry& Entry : ResourceTableMap.Resources)
 	{
-		if (Pair.Value.UniformBufferName == UBName && Pair.Value.ResourceIndex == ResourceIndex)
+		if (Entry.GetUniformBufferName() == UBName && Entry.ResourceIndex == ResourceIndex)
 		{
-			OutEntry = Pair.Value;
-			return Pair.Key;
+			OutEntry = Entry;
+			return;
 		}
 	}
 
 	check(0);
-	return "";
 }
 
 static FString FindTextureNameForSamplerState(const CrossCompiler::FHlslccHeader& CCHeader, const FString& InSamplerName)
@@ -374,7 +373,7 @@ static int32 AddGlobalForUBEntry(FOLDVulkanCodeHeader& OLDHeader,
 
 static void AddUBResources(FOLDVulkanCodeHeader& OLDHeader,
 							const FString& UBName,
-							const TMap<FString, FResourceTableEntry>& ResourceTableMap,
+							const FShaderResourceTableMap& ResourceTableMap,
 							uint32 BufferIndex,
 							const TArray<uint32>& BindingArray,
 							const FVulkanBindingTable& BindingTable,
@@ -405,13 +404,13 @@ static void AddUBResources(FOLDVulkanCodeHeader& OLDHeader,
 				// Extract index of the resource stored in the resource table from ResourceInfo
 				const uint16 ResourceIndex = FRHIResourceTableEntry::GetResourceIndex(ResourceInfo);
 
-				FResourceTableEntry ResourceTableEntry;
-				FString MemberName = GetResourceEntryFromUBMember(ResourceTableMap, UBName, ResourceIndex, ResourceTableEntry);
+				FUniformResourceEntry ResourceTableEntry;
+				GetResourceEntryFromUBMember(ResourceTableMap, UBName, ResourceIndex, ResourceTableEntry);
 
 				int32 HeaderUBResourceInfoIndex = OutUBInfo.ResourceEntries.AddZeroed();
 				FVulkanShaderHeader::FUBResourceInfo& UBResourceInfo = OutUBInfo.ResourceEntries[HeaderUBResourceInfoIndex];
 
-				int32 HeaderGlobalIndex = AddGlobalForUBEntry(OLDHeader, BindingTable, CCHeader, MemberName, BindingIndex, Spirv, GlobalNames, (EUniformBufferBaseType)ResourceTableEntry.Type, OutTypePatch, OutHeader);
+				int32 HeaderGlobalIndex = AddGlobalForUBEntry(OLDHeader, BindingTable, CCHeader, ResourceTableEntry.UniformBufferMemberName, BindingIndex, Spirv, GlobalNames, (EUniformBufferBaseType)ResourceTableEntry.Type, OutTypePatch, OutHeader);
 				UBResourceInfo.SourceUBResourceIndex = ResourceIndex;
 				UBResourceInfo.OriginalBindingIndex = BindingIndex;
 				UBResourceInfo.GlobalIndex = HeaderGlobalIndex;
@@ -535,7 +534,7 @@ struct FVulkanHlslccHeader : public CrossCompiler::FHlslccHeader
 	TArray<FString> ExternalTextures;
 };
 
-static void PrepareUBResourceEntryGlobals(const FVulkanHlslccHeader& CCHeader, const TArray<uint32>& BindingArray, const TMap<FString, FResourceTableEntry>& ResourceTableMap,
+static void PrepareUBResourceEntryGlobals(const FVulkanHlslccHeader& CCHeader, const TArray<uint32>& BindingArray, const FShaderResourceTableMap& ResourceTableMap,
 	int32 BufferIndex, const FString& UBName, TArray<FString>& OutGlobalNames, FVulkanShaderHeader& OutHeader)
 {
 	if (BindingArray.Num() > 0)
@@ -557,11 +556,11 @@ static void PrepareUBResourceEntryGlobals(const FVulkanHlslccHeader& CCHeader, c
 				// Extract index of the resource stored in the resource table from ResourceInfo
 				const uint16 ResourceIndex = FRHIResourceTableEntry::GetResourceIndex(ResourceInfo);
 
-				FResourceTableEntry ResourceTableEntry;
-				FString MemberName = GetResourceEntryFromUBMember(ResourceTableMap, UBName, ResourceIndex, ResourceTableEntry);
+				FUniformResourceEntry ResourceTableEntry;
+				GetResourceEntryFromUBMember(ResourceTableMap, UBName, ResourceIndex, ResourceTableEntry);
 
-				int32 GlobalIndex = DoAddGlobal(MemberName, OutHeader, OutGlobalNames);
-				if (CCHeader.ExternalTextures.Contains(MemberName))
+				int32 GlobalIndex = DoAddGlobal(ResourceTableEntry.UniformBufferMemberName, OutHeader, OutGlobalNames);
+				if (CCHeader.ExternalTextures.Contains(ResourceTableEntry.UniformBufferMemberName))
 				{
 					AddImmutable(OutHeader, GlobalIndex);
 				}
diff --git a/Engine/Source/Programs/ShaderCompileWorker/Private/ShaderCompileWorker.cpp b/Engine/Source/Programs/ShaderCompileWorker/Private/ShaderCompileWorker.cpp
--- a/Engine/Source/Programs/ShaderCompileWorker/Private/ShaderCompileWorker.cpp
+++ b/Engine/Source/Programs/ShaderCompileWorker/Private/ShaderCompileWorker.cpp
@@ -989,17 +989,6 @@ static void DirectCompile(const TArray<const class IShaderFormat*>& ShaderFormat
 	Job.Input.Target.Frequency = Frequency;
 	Job.Input.bSkipPreprocessedCache = true;
 
-	uint32 ResourceIndex = 0;
-	auto AddResourceTableEntry = [&ResourceIndex](TMap<FString, FResourceTableEntry>& Map, const FString& Name, const FString& UBName, int32 Type)
-	{
-		FResourceTableEntry LambdaEntry;
-		LambdaEntry.UniformBufferName = UBName;
-		LambdaEntry.Type = Type;
-		LambdaEntry.ResourceIndex = ResourceIndex;
-		Map.Add(Name, LambdaEntry);
-		++ResourceIndex;
-	};
-
 	Job.Input.Environment.CompilerFlags = FShaderCompilerFlags(CFlags);
 
 	Job.Input.bCompilingForShaderPipeline = bPipeline;
diff --git a/Engine/Source/Runtime/RenderCore/Private/ShaderCore.cpp b/Engine/Source/Runtime/RenderCore/Private/ShaderCore.cpp
--- a/Engine/Source/Runtime/RenderCore/Private/ShaderCore.cpp
+++ b/Engine/Source/Runtime/RenderCore/Private/ShaderCore.cpp
@@ -701,6 +701,61 @@ TArray<FString> FShaderParameterMap::GetAllParameterNamesOfType(EShaderParameter
 	return Result;
 }
 
+void FShaderResourceTableMap::Append(const FShaderResourceTableMap& Other)
+{
+	// Get the set of uniform buffers used by the target resource table map
+	TSet<FString> UniformBufferNames;
+	FStringView PreviousUniformBufferName;
+
+	for (const FUniformResourceEntry& Resource : Resources)
+	{
+		// Cheaper to check if consecutive array elements are from the same uniform buffer (which is common) before adding to set,
+		// which involves a more expensive hash lookup versus a string comparison.
+		if (!PreviousUniformBufferName.Equals(Resource.GetUniformBufferName(), ESearchCase::CaseSensitive))
+		{
+			PreviousUniformBufferName = Resource.GetUniformBufferName();
+			UniformBufferNames.Add(FString(PreviousUniformBufferName));
+		}
+	}
+
+	// Then add any entries from "Other" that aren't from a uniform buffer we already include.
+	PreviousUniformBufferName = FStringView();
+	bool PreviousUniformBufferFound = false;
+	for (const FUniformResourceEntry& OtherResource : Other.Resources)
+	{
+		if (!PreviousUniformBufferName.Equals(OtherResource.GetUniformBufferName(), ESearchCase::CaseSensitive))
+		{
+			PreviousUniformBufferName = OtherResource.GetUniformBufferName();
+			PreviousUniformBufferFound = UniformBufferNames.Find(FString(PreviousUniformBufferName)) != nullptr;
+		}
+
+		if (!PreviousUniformBufferFound)
+		{
+			Resources.Add(OtherResource);
+		}
+	}
+}
+
+void FShaderResourceTableMap::FixupOnLoad(const TMap<FString, FUniformBufferEntry>& UniformBufferMap)
+{
+	// Need to fix up UniformBufferMemberName string pointers to point into the MemberNameBuffer storage in UniformBufferMap
+	uint16 ResourceIndex = 0;
+	for (const auto& Pair : UniformBufferMap)
+	{
+		const TArray<TCHAR>* MemberNameBuffer = Pair.Value.MemberNameBuffer.Get();
+		if (MemberNameBuffer && MemberNameBuffer->Num())
+		{
+			const TCHAR* MemberNameCurrent = MemberNameBuffer->GetData();
+			const TCHAR* MemberNameEnd = MemberNameCurrent + MemberNameBuffer->Num();
+
+			for (; MemberNameCurrent < MemberNameEnd; MemberNameCurrent += FCString::Strlen(MemberNameCurrent) + 1)
+			{
+				Resources[ResourceIndex++].UniformBufferMemberName = MemberNameCurrent;
+			}
+		}
+	}
+}
+
 void FShaderCompilerDefinitions::SetFloatDefine(const TCHAR* Name, float Value)
 {
 	// Make sure the printed value perfectly matches the given number
diff --git a/Engine/Source/Runtime/RenderCore/Private/ShaderParameterMetadata.cpp b/Engine/Source/Runtime/RenderCore/Private/ShaderParameterMetadata.cpp
--- a/Engine/Source/Runtime/RenderCore/Private/ShaderParameterMetadata.cpp
+++ b/Engine/Source/Runtime/RenderCore/Private/ShaderParameterMetadata.cpp
@@ -826,18 +826,6 @@ void FShaderParametersMetadata::InitializeLayout(FRHIUniformBufferLayoutInitiali
 	Layout = RHICreateUniformBufferLayout(LayoutInitializer);
 }
 
-#if WITH_EDITOR
-void FShaderParametersMetadata::InitializeUniformBufferDeclaration()
-{
-	if (UseCase != EUseCase::ShaderParameterStruct)
-	{
-		FString* NewDeclaration = new FString(UE::ShaderParameters::CreateUniformBufferShaderDeclaration(ShaderVariableName, *this));
-		check(!NewDeclaration->IsEmpty());
-
-		UniformBufferDeclaration = MakeShareable(NewDeclaration);
-	}
-}
-#endif // WITH_EDITOR
 
 void FShaderParametersMetadata::GetNestedStructs(TArray<const FShaderParametersMetadata*>& OutNestedStructs) const
 {
@@ -857,8 +845,19 @@ void FShaderParametersMetadata::GetNestedStructs(TArray<const FShaderParametersM
 
 #if WITH_EDITOR
 
-static void AddResourceTableEntriesRecursive(const TArray<FShaderParametersMetadata::FMember>& Members, const TCHAR* UniformBufferName, const TCHAR* Prefix, uint16& ResourceIndex, TMap<FString, FResourceTableEntry>& ResourceTableMap)
+// Actual longest token is 78 characters -- these are a few human written variable names concatenated together, and unlikely to go much higher.
+typedef TStringBuilder<2048> FResourceTableEntryString;
+
+// Function returns count to allocate for MemberNameBuffer and ResourceIndex specifies the number of Entries to allocate, if MemberNameBuffer is nullptr
+static uint32 CacheResourceTableEntriesRecursive(
+	const TArray<FShaderParametersMetadata::FMember>& Members, const TCHAR* UniformBufferName, FResourceTableEntryString& Prefix, uint16& ResourceIndex, TCHAR*& MemberNameBuffer, TArray<FUniformResourceEntry>* Entries)
 {
+	uint32 MemberNameBufferCount = 0;
+	uint32 OriginalPrefixLength = Prefix.Len();
+
+	uint32 UniformBufferNameLength = FCString::Strlen(UniformBufferName);
+	check(UniformBufferNameLength <= UINT8_MAX);
+
 	for (const FShaderParametersMetadata::FMember& Member : Members)
 	{
 		const EUniformBufferBaseType BaseType = Member.GetBaseType();
@@ -866,51 +865,128 @@ static void AddResourceTableEntriesRecursive(const TArray<FShaderParametersMetad
 
 		if (IsShaderParameterTypeForUniformBufferLayout(BaseType))
 		{
-			FResourceTableEntry& Entry = ResourceTableMap.FindOrAdd(FString::Printf(TEXT("%s%s"), Prefix, Member.GetName()));
-			if (Entry.UniformBufferName.IsEmpty())
+			// Format:  "%s%s"
+			Prefix.Append(Member.GetName());
+			uint32 StringStorageCount = Prefix.Len() + 1;
+			if (MemberNameBuffer)
 			{
-				Entry.UniformBufferName = UniformBufferName;
-				Entry.Type = BaseType;
-				Entry.ResourceIndex = ResourceIndex++;
+				FCString::Strcpy(MemberNameBuffer, StringStorageCount, Prefix.ToString());
+
+				check(Entries);
+				Entries->Add({
+					MemberNameBuffer,
+					(uint8)UniformBufferNameLength,
+					BaseType,
+					ResourceIndex
+					});
+
+				MemberNameBuffer += StringStorageCount;
 			}
+
+			MemberNameBufferCount += StringStorageCount;
+			ResourceIndex++;
+			Prefix.RemoveSuffix(Prefix.Len() - OriginalPrefixLength);
 		}
 		else if (BaseType == UBMT_NESTED_STRUCT && NumElements == 0)
 		{
 			check(Member.GetStructMetadata());
-			FString MemberPrefix = FString::Printf(TEXT("%s%s_"), Prefix, Member.GetName());
-			AddResourceTableEntriesRecursive(Member.GetStructMetadata()->GetMembers(), UniformBufferName, *MemberPrefix, ResourceIndex, ResourceTableMap);
+
+			// Format:  "%s%s_"
+			Prefix.Append(Member.GetName());
+			Prefix.AppendChar('_');
+			MemberNameBufferCount += CacheResourceTableEntriesRecursive(Member.GetStructMetadata()->GetMembers(), UniformBufferName, Prefix, ResourceIndex, MemberNameBuffer, Entries);
+			Prefix.RemoveSuffix(Prefix.Len() - OriginalPrefixLength);
 		}
 		else if (BaseType == UBMT_NESTED_STRUCT && NumElements > 0)
 		{
+			// Format:  %s%s_%u_
+			Prefix.Append(Member.GetName());
+			Prefix.AppendChar('_');
+			uint32 PrefixLengthMinusArrayIndex = Prefix.Len();
+
 			for (uint32 ArrayElementId = 0; ArrayElementId < NumElements; ArrayElementId++)
 			{
 				check(Member.GetStructMetadata());
-				FString MemberPrefix = FString::Printf(TEXT("%s%s_%u_"), Prefix, Member.GetName(), ArrayElementId);
-				AddResourceTableEntriesRecursive(Member.GetStructMetadata()->GetMembers(), UniformBufferName, *MemberPrefix, ResourceIndex, ResourceTableMap);
+				Prefix.Appendf(TEXT("%u"), ArrayElementId);
+				Prefix.AppendChar('_');
+				MemberNameBufferCount += CacheResourceTableEntriesRecursive(Member.GetStructMetadata()->GetMembers(), UniformBufferName, Prefix, ResourceIndex, MemberNameBuffer, Entries);
+
+				Prefix.RemoveSuffix(Prefix.Len() - PrefixLengthMinusArrayIndex);
 			}
+
+			Prefix.RemoveSuffix(Prefix.Len() - OriginalPrefixLength);
 		}
 		else if (BaseType == UBMT_INCLUDED_STRUCT)
 		{
 			check(Member.GetStructMetadata());
 			check(NumElements == 0);
-			AddResourceTableEntriesRecursive(Member.GetStructMetadata()->GetMembers(), UniformBufferName, Prefix, ResourceIndex, ResourceTableMap);
+			MemberNameBufferCount += CacheResourceTableEntriesRecursive(Member.GetStructMetadata()->GetMembers(), UniformBufferName, Prefix, ResourceIndex, MemberNameBuffer, Entries);
 		}
 	}
+
+	return MemberNameBufferCount;
+}
+
+void FShaderParametersMetadata::InitializeUniformBufferDeclaration()
+{
+	if (UseCase != EUseCase::ShaderParameterStruct)
+	{
+		FString* NewDeclaration = new FString(UE::ShaderParameters::CreateUniformBufferShaderDeclaration(ShaderVariableName, *this));
+		check(!NewDeclaration->IsEmpty());
+
+		UniformBufferDeclaration = MakeShareable(NewDeclaration);
+
+		// Generate ResourceTableCache and MemberNameBuffer
+		FResourceTableEntryString Prefix;
+		Prefix.Append(ShaderVariableName);
+		Prefix.AppendChar('_');
+
+		// First pass to determine number of entries and name buffer count, by passing nullptr for Buffer
+		uint16 ResourceIndex = 0;
+		TCHAR* Buffer = nullptr;
+		uint32 BufferCount = CacheResourceTableEntriesRecursive(Members, ShaderVariableName, Prefix, ResourceIndex, Buffer, nullptr);
+
+		// Then generate the entries and name buffer
+		TArray<TCHAR>* MemberNameBufferAllocation = new TArray<TCHAR>();
+		MemberNameBufferAllocation->SetNumZeroed(BufferCount);
+		ResourceTableCache.Reserve(ResourceIndex);
+
+		ResourceIndex = 0;
+		Buffer = MemberNameBufferAllocation->GetData();
+		CacheResourceTableEntriesRecursive(Members, ShaderVariableName, Prefix, ResourceIndex, Buffer, &ResourceTableCache);
+
+		MemberNameBuffer = MakeShareable(MemberNameBufferAllocation);
+
+		// Cache strings for uniform buffer generated path and include
+		UniformBufferPath = FString::Printf(TEXT("/Engine/Generated/UniformBuffers/%s.ush"), ShaderVariableName);
+		UniformBufferInclude = FString::Printf(TEXT("#include \"/Engine/Generated/UniformBuffers/%s.ush\"") LINE_TERMINATOR, ShaderVariableName);
+
+		// Cache some frequently used hashes
+		UniformBufferPathHash = GetTypeHash(UniformBufferPath);
+		ShaderVariableNameHash = GetTypeHash(FStringView(ShaderVariableName));
+	}
 }
 
-void FShaderParametersMetadata::AddResourceTableEntries(TMap<FString, FResourceTableEntry>& ResourceTableMap, TMap<FString, FUniformBufferEntry>& UniformBufferMap) const
+void FShaderParametersMetadata::AddResourceTableEntries(FShaderResourceTableMap& ResourceTableMap, TMap<FString, FUniformBufferEntry>& UniformBufferMap) const
 {
-	uint16 ResourceIndex = 0;
-	const FString Prefix = FString::Printf(TEXT("%s_"), ShaderVariableName);
-	AddResourceTableEntriesRecursive(GetMembers(), ShaderVariableName, *Prefix, ResourceIndex, ResourceTableMap);
+	ResourceTableMap.Resources.Append(ResourceTableCache);
 	
 	FUniformBufferEntry UniformBufferEntry;
 	UniformBufferEntry.StaticSlotName = StaticSlotName;
 	UniformBufferEntry.LayoutHash = IsLayoutInitialized() ? GetLayout().GetHash() : 0;
 	UniformBufferEntry.BindingFlags = BindingFlags;
 	UniformBufferEntry.bNoEmulatedUniformBuffer = UsageFlags & (uint32)EUsageFlags::NoEmulatedUniformBuffer;
-	UniformBufferMap.Add(ShaderVariableName, UniformBufferEntry);
+	UniformBufferEntry.MemberNameBuffer = MemberNameBuffer;
+	UniformBufferMap.AddByHash(ShaderVariableNameHash, ShaderVariableName, UniformBufferEntry);
+}
+
+PRAGMA_DISABLE_DEPRECATION_WARNINGS
+// Deprecated version of function
+void FShaderParametersMetadata::AddResourceTableEntries(TMap<FString, FResourceTableEntry>& ResourceTableMap, TMap<FString, FUniformBufferEntry>& UniformBufferMap) const
+{
+	UE_LOG(LogShaders, Error, TEXT("FShaderParametersMetadata::AddResourceTableEntries call that accepts a TMap has been deprecated.  Use FShaderResourceTableMap structure instead."));
 }
+PRAGMA_ENABLE_DEPRECATION_WARNINGS
 
 #endif // WITH_EDITOR
 
diff --git a/Engine/Source/Runtime/RenderCore/Private/ShaderParameters.cpp b/Engine/Source/Runtime/RenderCore/Private/ShaderParameters.cpp
--- a/Engine/Source/Runtime/RenderCore/Private/ShaderParameters.cpp
+++ b/Engine/Source/Runtime/RenderCore/Private/ShaderParameters.cpp
@@ -398,20 +398,22 @@ void UE::ShaderParameters::AddUniformBufferIncludesToEnvironment(FShaderCompiler
 
 	for (const TCHAR* UniformBufferName : InUniformBufferNames)
 	{
-		if (const FShaderParametersMetadata* Metadata = FindShaderParametersMetadataWithVariableName(UniformBufferName))
+		FStringView UniformBufferNameView(UniformBufferName);
+		if (!OutEnvironment.UniformBufferMap.FindByHash(GetTypeHash(UniformBufferName), UniformBufferName))
 		{
-			const FThreadSafeSharedStringPtr UniformBufferDeclaration = Metadata->GetUniformBufferDeclarationPtr();
-
-			check(UniformBufferDeclaration.Get() != NULL);
-			check(!UniformBufferDeclaration.Get()->IsEmpty());
+			if (const FShaderParametersMetadata* Metadata = FindShaderParametersMetadataWithVariableName(UniformBufferName))
+			{
+				const FThreadSafeSharedStringPtr UniformBufferDeclaration = Metadata->GetUniformBufferDeclarationPtr();
 
-			const FString UniformBufferPath = FString::Printf(TEXT("/Engine/Generated/UniformBuffers/%s.ush"), UniformBufferName);
+				check(UniformBufferDeclaration.Get() != NULL);
+				check(!UniformBufferDeclaration.Get()->IsEmpty());
 
-			UniformBufferIncludes += FString::Printf(TEXT("#include \"%s\"") LINE_TERMINATOR, *UniformBufferPath);
+				UniformBufferIncludes += Metadata->GetUniformBufferInclude();
 
-			OutEnvironment.IncludeVirtualPathToExternalContentsMap.Add(UniformBufferPath, UniformBufferDeclaration);
+				OutEnvironment.IncludeVirtualPathToExternalContentsMap.AddByHash(Metadata->GetUniformBufferPathHash(), Metadata->GetUniformBufferPath(), UniformBufferDeclaration);
 
-			Metadata->AddResourceTableEntries(OutEnvironment.ResourceTableMap, OutEnvironment.UniformBufferMap);
+				Metadata->AddResourceTableEntries(OutEnvironment.ResourceTableMap, OutEnvironment.UniformBufferMap);
+			}
 		}
 	}
 
diff --git a/Engine/Source/Runtime/RenderCore/Public/ShaderCore.h b/Engine/Source/Runtime/RenderCore/Public/ShaderCore.h
--- a/Engine/Source/Runtime/RenderCore/Public/ShaderCore.h
+++ b/Engine/Source/Runtime/RenderCore/Public/ShaderCore.h
@@ -432,9 +432,14 @@ inline FArchive& operator<<(FArchive& Ar, FShaderResourceTable& SRT)
 	return Ar;
 }
 
-inline FArchive& operator<<(FArchive& Ar, FResourceTableEntry& Entry)
+inline FArchive& operator<<(FArchive& Ar, FUniformResourceEntry& Entry)
 {
-	Ar << Entry.UniformBufferName;
+	if (Ar.IsLoading())
+	{
+		// Filled in later in FShaderResourceTableMap::FixupOnLoad
+		Entry.UniformBufferMemberName = nullptr;
+	}
+	Ar << Entry.UniformBufferNameLength;
 	Ar << Entry.Type;
 	Ar << Entry.ResourceIndex;
 	return Ar;
@@ -446,10 +451,16 @@ inline FArchive& operator<<(FArchive& Ar, FUniformBufferEntry& Entry)
 	Ar << Entry.LayoutHash;
 	Ar << Entry.BindingFlags;
 	Ar << Entry.bNoEmulatedUniformBuffer;
+	if (Ar.IsLoading())
+	{
+		Entry.MemberNameBuffer = MakeShareable(new TArray<TCHAR>());
+	}
+	Ar << *Entry.MemberNameBuffer.Get();
 	return Ar;
 }
 
 using FThreadSafeSharedStringPtr = TSharedPtr<FString, ESPMode::ThreadSafe>;
+using FThreadSafeNameBufferPtr = TSharedPtr<TArray<TCHAR>, ESPMode::ThreadSafe>;
 
 // Simple wrapper for a uint64 bitfield; doesn't use TBitArray as it is fixed size and doesn't need dynamic memory allocations
 class FShaderCompilerFlags
@@ -507,6 +518,14 @@ private:
 	uint64 Data;
 };
 
+struct FShaderResourceTableMap
+{
+	TArray<FUniformResourceEntry> Resources;
+
+	RENDERCORE_API void Append(const FShaderResourceTableMap& Other);
+	RENDERCORE_API void FixupOnLoad(const TMap<FString, FUniformBufferEntry>& UniformBufferMap);
+};
+
 /** The environment used to compile a shader. */
 struct FShaderCompilerEnvironment
 {
@@ -518,7 +537,7 @@ struct FShaderCompilerEnvironment
 
 	FShaderCompilerFlags CompilerFlags;
 	TMap<uint32,uint8> RenderTargetOutputFormatsMap;
-	TMap<FString, FResourceTableEntry> ResourceTableMap;
+	FShaderResourceTableMap ResourceTableMap;
 	TMap<FString, FUniformBufferEntry> UniformBufferMap;
 
 	UE_DEPRECATED(5.3, "RemoteServerData field is deprecated (no longer used in compilation backends).")
@@ -638,9 +657,13 @@ struct FShaderCompilerEnvironment
 		Ar << CompileArgs;
 		Ar << CompilerFlags;
 		Ar << RenderTargetOutputFormatsMap;
-		Ar << ResourceTableMap;
+		Ar << ResourceTableMap.Resources;
 		Ar << UniformBufferMap;
 		Ar << FullPrecisionInPS;
+		if (Ar.IsLoading())
+		{
+			ResourceTableMap.FixupOnLoad(UniformBufferMap);
+		}
 	}
 
 	// Serializes the portions of the environment that are used as input to the backend compilation process (i.e. after all preprocessing)
@@ -648,9 +671,13 @@ struct FShaderCompilerEnvironment
 	{
 		Ar << CompileArgs;
 		Ar << CompilerFlags;
-		Ar << ResourceTableMap;
+		Ar << ResourceTableMap.Resources;
 		Ar << UniformBufferMap;
 		Ar << FullPrecisionInPS;
+		if (Ar.IsLoading())
+		{
+			ResourceTableMap.FixupOnLoad(UniformBufferMap);
+		}
 	}
 
 	friend FArchive& operator<<(FArchive& Ar,FShaderCompilerEnvironment& Environment)
diff --git a/Engine/Source/Runtime/RenderCore/Public/ShaderParameterMetadata.h b/Engine/Source/Runtime/RenderCore/Public/ShaderParameterMetadata.h
--- a/Engine/Source/Runtime/RenderCore/Public/ShaderParameterMetadata.h
+++ b/Engine/Source/Runtime/RenderCore/Public/ShaderParameterMetadata.h
@@ -26,6 +26,8 @@
 #include "Templates/SharedPointer.h"
 
 using FThreadSafeSharedStringPtr = TSharedPtr<FString, ESPMode::ThreadSafe>;
+using FThreadSafeNameBufferPtr = TSharedPtr<TArray<TCHAR>, ESPMode::ThreadSafe>;
+struct FShaderResourceTableMap;
 
 namespace EShaderPrecisionModifier
 {
@@ -42,14 +44,25 @@ namespace EShaderPrecisionModifier
 bool SupportShaderPrecisionModifier(EShaderPlatform Platform);
 
 /** Each entry in a resource table is provided to the shader compiler for creating mappings. */
-struct FResourceTableEntry
+struct FUniformResourceEntry
 {
-	/** The name of the uniform buffer in which this resource exists. */
-	FString UniformBufferName;
+	/** The name of the uniform buffer member for this resource. */
+	const TCHAR* UniformBufferMemberName;
+	/** The Uniform Buffer's name is a prefix of this length at the start of UniformBufferMemberName. */
+	uint8 UniformBufferNameLength{};
 	/** The type of the resource (EUniformBufferBaseType). */
-	uint16 Type{};
+	uint8 Type{};
 	/** The index of the resource in the table. */
 	uint16 ResourceIndex{};
+
+	FORCEINLINE FStringView GetUniformBufferName() const { return FStringView(UniformBufferMemberName, UniformBufferNameLength); }
+};
+
+struct UE_DEPRECATED(5.3, "Deprecated structure -- replaced with FUniformResourceEntry.") FResourceTableEntry
+{
+	FString UniformBufferName;
+	uint16 Type{};
+	uint16 ResourceIndex{};
 };
 
 /** Minimal information about each uniform buffer entry fed to the shader compiler. */
@@ -63,6 +76,8 @@ struct FUniformBufferEntry
 	EUniformBufferBindingFlags BindingFlags{ EUniformBufferBindingFlags::Shader };
 	/** Whether to force a real uniform buffer when using emulated uniform buffers */
 	bool bNoEmulatedUniformBuffer;
+	/** Storage for member names for this uniform buffer (pointed to by FUniformResourceEntry::UniformBufferMemberName)  */
+	FThreadSafeNameBufferPtr MemberNameBuffer;
 };
 
 /** Parse the shader resource binding from the binding type used in shader code. */
@@ -287,7 +302,12 @@ public:
 	void GetNestedStructs(TArray<const FShaderParametersMetadata*>& OutNestedStructs) const;
 
 #if WITH_EDITOR
+	void AddResourceTableEntries(FShaderResourceTableMap& ResourceTableMap, TMap<FString, FUniformBufferEntry>& UniformBufferMap) const;
+
+	PRAGMA_DISABLE_DEPRECATION_WARNINGS
+	UE_DEPRECATED(5.3, "Resource table entries are now stored in FShaderResourceTableMap, rather than a TMap.")
 	void AddResourceTableEntries(TMap<FString, FResourceTableEntry>& ResourceTableMap, TMap<FString, FUniformBufferEntry>& UniformBufferMap) const;
+	PRAGMA_ENABLE_DEPRECATION_WARNINGS
 #endif
 
 	const TCHAR* GetStructTypeName() const { return StructTypeName; }
@@ -334,6 +354,9 @@ public:
 	inline bool IsUniformBufferDeclarationInitialized() const { return UniformBufferDeclaration.IsValid(); }
 	FThreadSafeSharedStringPtr GetUniformBufferDeclarationPtr() const { return UniformBufferDeclaration; }
 	const FString& GetUniformBufferDeclaration() const { return *UniformBufferDeclaration; }
+	FORCEINLINE const FString& GetUniformBufferPath() const { return UniformBufferPath; }
+	FORCEINLINE const FString& GetUniformBufferInclude() const { return UniformBufferInclude; }
+	FORCEINLINE uint32 GetUniformBufferPathHash() const { return UniformBufferPathHash; }
 #endif // WITH_EDITOR
 
 	/** Find a member for a given offset. */
@@ -416,6 +439,18 @@ private:
 #if WITH_EDITOR
 	/** Uniform buffer declaration, created once */
 	FThreadSafeSharedStringPtr UniformBufferDeclaration;
+
+	/** Cache of uniform buffer resource table, and storage for member names used by the table, created once */
+	TArray<FUniformResourceEntry> ResourceTableCache;
+	FThreadSafeNameBufferPtr MemberNameBuffer;
+
+	/** Strings for uniform buffer generated path and include, created once */
+	FString UniformBufferPath;		// Format:  "/Engine/Generated/UniformBuffers/%s.ush"
+	FString UniformBufferInclude;	// Format:  "#include \"/Engine/Generated/UniformBuffers/%s.ush\"" LINE_TERMINATOR
+
+	/** Hashes for frequently used strings */
+	uint32 UniformBufferPathHash;
+	uint32 ShaderVariableNameHash;
 #endif
 
 	/** Shackle elements in global link list of globally named shader parameters. */
