Add ability to trace shader group preload entries.

diff --git a/Engine/Source/Runtime/RenderCore/Private/ShaderCodeArchive.cpp b/Engine/Source/Runtime/RenderCore/Private/ShaderCodeArchive.cpp
--- a/Engine/Source/Runtime/RenderCore/Private/ShaderCodeArchive.cpp
+++ b/Engine/Source/Runtime/RenderCore/Private/ShaderCodeArchive.cpp
@@ -1829,7 +1829,13 @@ void FIoStoreShaderCodeArchive::Teardown()
 	{
 		FShaderGroupPreloadEntry* PreloadEntry = Iter.Value();
 
-		checkf(PreloadEntry->NumRefs == 0, TEXT("Group %d has still %d references on deletion"), Iter.Key(), PreloadEntry->NumRefs);
+#if UE_SCA_DEBUG_PRELOADING
+		checkf(PreloadEntry->NumRefs == 0, TEXT("Group %d has still %d references on deletion. Group extended debug info: \n%s"), Iter.Key(), PreloadEntry->NumRefs,
+			*PreloadEntry->DebugInfo);
+#else
+		checkf(PreloadEntry->NumRefs == 0, TEXT("Group %d has still %d references on deletion. Group extended debug info: \n%s"), Iter.Key(), PreloadEntry->NumRefs,
+			TEXT("Not compiled in (set UE_SCA_DEBUG_PRELOADING to 1 in ShaderCodeArchive.h and recompile the game binary)"));
+#endif
 
 		const FIoStoreShaderGroupEntry& GroupEntry = Header.ShaderGroupEntries[Iter.Key()];
 		DEC_DWORD_STAT_BY(STAT_Shaders_ShaderPreloadMemory, (GroupEntry.CompressedSize + sizeof(FShaderGroupPreloadEntry)));
@@ -1854,7 +1860,11 @@ void FIoStoreShaderCodeArchive::SetupPreloadEntryForLoading(int32 ShaderGroupInd
 #endif // UE_SCA_VISUALIZE_SHADER_USAGE
 }
 
-bool FIoStoreShaderCodeArchive::PreloadShaderGroup(int32 ShaderGroupIndex, FGraphEventArray& OutCompletionEvents, FCoreDelegates::FAttachShaderReadRequestFunc* AttachShaderReadRequestFuncPtr)
+bool FIoStoreShaderCodeArchive::PreloadShaderGroup(int32 ShaderGroupIndex, FGraphEventArray& OutCompletionEvents, 
+#if UE_SCA_DEBUG_PRELOADING
+	const FString& CallsiteInfo,
+#endif
+	FCoreDelegates::FAttachShaderReadRequestFunc* AttachShaderReadRequestFuncPtr)
 {
 	// should be called within LLMTag::Shaders scope
 	FWriteScopeLock Lock(PreloadedShaderGroupsLock);
@@ -1862,6 +1872,12 @@ bool FIoStoreShaderCodeArchive::PreloadShaderGroup(int32 ShaderGroupIndex, FGrap
 	checkf(!PreloadEntry.bNeverToBePreloaded, TEXT("We are preloading a shader group (index=%d) that shouldn't be preloaded in this run (e.g. raytracing shaders on D3D11)."), ShaderGroupIndex);
 
 	const uint32 NumRefs = PreloadEntry.NumRefs++;
+#if UE_SCA_DEBUG_PRELOADING
+	FString AppendInfo = FString::Printf(TEXT("PreloadShaderGroup: NumRefs %d -> %d    CallsiteInfo: %s\n"),
+		NumRefs, PreloadEntry.NumRefs, *CallsiteInfo
+	);
+	PreloadEntry.DebugInfo.Append(AppendInfo);
+#endif
 	if (NumRefs == 0u)
 	{
 		SetupPreloadEntryForLoading(ShaderGroupIndex, PreloadEntry);
@@ -1888,12 +1904,22 @@ bool FIoStoreShaderCodeArchive::PreloadShaderGroup(int32 ShaderGroupIndex, FGrap
 	return true;
 }
 
-void FIoStoreShaderCodeArchive::MarkPreloadEntrySkipped(int32 ShaderGroupIndex)
+void FIoStoreShaderCodeArchive::MarkPreloadEntrySkipped(int32 ShaderGroupIndex
+#if UE_SCA_DEBUG_PRELOADING
+	, const FString& CallsiteInfo
+#endif
+)
 {
 	// should be called within LLMTag::Shaders scope
 	FWriteScopeLock Lock(PreloadedShaderGroupsLock);
 	FShaderGroupPreloadEntry& PreloadEntry = *FindOrAddPreloadEntry(ShaderGroupIndex);
 	const uint32 NumRefs = PreloadEntry.NumRefs++;
+#if UE_SCA_DEBUG_PRELOADING
+	FString AppendInfo = FString::Printf(TEXT("MarkPreloadEntrySkipped: NumRefs %d -> %d    CallsiteInfo: %s\n"),
+		NumRefs, PreloadEntry.NumRefs, *CallsiteInfo
+	);
+	PreloadEntry.DebugInfo.Append(AppendInfo);
+#endif
 	if (NumRefs == 0u)
 	{
 		PreloadEntry.bNeverToBePreloaded = 1;
@@ -1905,7 +1931,11 @@ bool FIoStoreShaderCodeArchive::PreloadShader(int32 ShaderIndex, FGraphEventArra
 {
 	LLM_SCOPE(ELLMTag::Shaders);
 	DebugVisualizer.MarkExplicitlyPreloadedForVisualization(ShaderIndex);
-	return PreloadShaderGroup(GetGroupIndexForShader(ShaderIndex), OutCompletionEvents);
+	return PreloadShaderGroup(GetGroupIndexForShader(ShaderIndex), OutCompletionEvents
+#if UE_SCA_DEBUG_PRELOADING
+		, FString::Printf(TEXT("PreloadShader %d"), ShaderIndex)
+#endif
+	);
 }
 
 bool FIoStoreShaderCodeArchive::GroupOnlyContainsRaytracingShaders(int32 ShaderGroupIndex)
@@ -1928,6 +1958,9 @@ bool FIoStoreShaderCodeArchive::PreloadShaderMap(int32 ShaderMapIndex, FGraphEve
 	LLM_SCOPE(ELLMTag::Shaders);
 	const FIoStoreShaderMapEntry& ShaderMapEntry = Header.ShaderMapEntries[ShaderMapIndex];
 
+#if UE_SCA_DEBUG_PRELOADING
+	FString Callsite = FString::Printf(TEXT("PreloadShaderMap %d"), ShaderMapIndex);
+#endif
 	for (uint32 i = 0u; i < ShaderMapEntry.NumShaders; ++i)
 	{
 		const int32 ShaderIndex = Header.ShaderIndices[ShaderMapEntry.ShaderIndicesOffset + i];
@@ -1936,14 +1969,22 @@ bool FIoStoreShaderCodeArchive::PreloadShaderMap(int32 ShaderMapIndex, FGraphEve
 #if RHI_RAYTRACING
 		if (!IsRayTracingAllowed() && !IsCreateShadersOnLoadEnabled() && GroupOnlyContainsRaytracingShaders(ShaderGroupIndex))
 		{
-			MarkPreloadEntrySkipped(ShaderGroupIndex);
+			MarkPreloadEntrySkipped(ShaderGroupIndex
+#if UE_SCA_DEBUG_PRELOADING
+				, Callsite
+#endif
+			);
 			continue;
 		}
 #endif
 
 		// only shaders we actually want to preload should be marked as such, not just everything in the group
 		DebugVisualizer.MarkExplicitlyPreloadedForVisualization(ShaderIndex);
-		PreloadShaderGroup(ShaderGroupIndex, OutCompletionEvents);
+		PreloadShaderGroup(ShaderGroupIndex, OutCompletionEvents
+#if UE_SCA_DEBUG_PRELOADING
+			, Callsite
+#endif
+		);
 	}
 
 	return true;
@@ -1955,6 +1996,9 @@ bool FIoStoreShaderCodeArchive::PreloadShaderMap(int32 ShaderMapIndex, FCoreDele
 	const FIoStoreShaderMapEntry& ShaderMapEntry = Header.ShaderMapEntries[ShaderMapIndex];
 
 	FGraphEventArray Dummy;
+#if UE_SCA_DEBUG_PRELOADING
+	FString Callsite = FString::Printf(TEXT("PreloadShaderMap(AttachShaderReadRequestFunc) %d"), ShaderMapIndex);
+#endif
 	for (uint32 i = 0u; i < ShaderMapEntry.NumShaders; ++i)
 	{
 		const int32 ShaderIndex = Header.ShaderIndices[ShaderMapEntry.ShaderIndicesOffset + i];
@@ -1963,20 +2007,32 @@ bool FIoStoreShaderCodeArchive::PreloadShaderMap(int32 ShaderMapIndex, FCoreDele
 #if RHI_RAYTRACING
 		if (!IsRayTracingAllowed() && !IsCreateShadersOnLoadEnabled() && GroupOnlyContainsRaytracingShaders(ShaderGroupIndex))
 		{
-			MarkPreloadEntrySkipped(ShaderGroupIndex);
+			MarkPreloadEntrySkipped(ShaderGroupIndex
+#if UE_SCA_DEBUG_PRELOADING
+				, Callsite
+#endif
+			);
 			continue;
 		}
 #endif
 
 		// only shaders we actually want to preload should be marked as such, not just everything in the group
 		DebugVisualizer.MarkExplicitlyPreloadedForVisualization(ShaderIndex);
-		PreloadShaderGroup(ShaderGroupIndex, Dummy, &AttachShaderReadRequestFunc);
+		PreloadShaderGroup(ShaderGroupIndex, Dummy, 
+#if UE_SCA_DEBUG_PRELOADING
+			Callsite,
+#endif
+			&AttachShaderReadRequestFunc);
 	}
 
 	return true;
 }
 
-void FIoStoreShaderCodeArchive::ReleasePreloadEntry(int32 ShaderGroupIndex)
+void FIoStoreShaderCodeArchive::ReleasePreloadEntry(int32 ShaderGroupIndex
+#if UE_SCA_DEBUG_PRELOADING
+	, const FString& CallsiteInfo
+#endif
+)
 {
 	FWriteScopeLock Lock(PreloadedShaderGroupsLock);
 	FShaderGroupPreloadEntry** ExistingEntry = PreloadedShaderGroups.Find(ShaderGroupIndex);
@@ -1987,6 +2043,14 @@ void FIoStoreShaderCodeArchive::ReleasePreloadEntry(int32 ShaderGroupIndex)
 
 		const uint32 ShaderNumRefs = PreloadEntry->NumRefs--;
 		check(ShaderNumRefs > 0u);
+
+#if UE_SCA_DEBUG_PRELOADING
+		FString AppendInfo = FString::Printf(TEXT("ReleasePreloadEntry: NumRefs %d -> %d    CallsiteInfo: %s\n"),
+			ShaderNumRefs, PreloadEntry->NumRefs, *CallsiteInfo
+			);
+		PreloadEntry->DebugInfo.Append(AppendInfo);
+#endif
+
 		if (ShaderNumRefs == 1u)
 		{
 			if (!PreloadEntry->bNeverToBePreloaded)
@@ -2001,6 +2065,14 @@ void FIoStoreShaderCodeArchive::ReleasePreloadEntry(int32 ShaderGroupIndex)
 				DEC_DWORD_STAT_BY(STAT_Shaders_ShaderPreloadMemory, sizeof(FShaderGroupPreloadEntry));
 			}
 
+#if UE_SCA_DEBUG_PRELOADING
+			if (0)	// use this if you need comparison with all other groups
+			{
+				UE_LOG(LogInit, Log, TEXT("Group %d has still %d references on deletion. Group extended debug info: \n%s"), ShaderGroupIndex, PreloadEntry->NumRefs,
+					*PreloadEntry->DebugInfo);
+			}
+#endif
+
 			delete PreloadEntry;
 			PreloadedShaderGroups.Remove(ShaderGroupIndex);
 		}
@@ -2009,7 +2081,11 @@ void FIoStoreShaderCodeArchive::ReleasePreloadEntry(int32 ShaderGroupIndex)
 
 void FIoStoreShaderCodeArchive::ReleasePreloadedShader(int32 ShaderIndex)
 {
-	ReleasePreloadEntry(GetGroupIndexForShader(ShaderIndex));
+	ReleasePreloadEntry(GetGroupIndexForShader(ShaderIndex)
+#if UE_SCA_DEBUG_PRELOADING
+		, FString::Printf(TEXT("ReleasePreloadedShader %d"), ShaderIndex)
+#endif
+	);
 }
 
 int32 FIoStoreShaderCodeArchive::FindShaderMapIndex(const FSHAHash& Hash)
@@ -2062,7 +2138,14 @@ TRefCountPtr<FRHIShader> FIoStoreShaderCodeArchive::CreateShader(int32 ShaderInd
 
 	// Preload shader group if it wasn't yet. This will also addref it so we can be sure it will exist.
 	FGraphEventArray Dummy;
-	PreloadShaderGroup(GroupIndex, Dummy);
+#if UE_SCA_DEBUG_PRELOADING
+	FString Callsite = FString::Printf(TEXT("CreateShader %d"), ShaderIndex);
+#endif
+	PreloadShaderGroup(GroupIndex, Dummy
+#if UE_SCA_DEBUG_PRELOADING
+		, Callsite
+#endif
+	);
 
 	FShaderGroupPreloadEntry* PreloadEntryPtr;
 	{
@@ -2140,7 +2223,11 @@ TRefCountPtr<FRHIShader> FIoStoreShaderCodeArchive::CreateShader(int32 ShaderInd
 	DebugVisualizer.MarkCreatedForVisualization(ShaderIndex);
 
 	PreloadEntryPtr = nullptr;
-	ReleasePreloadEntry(GroupIndex);
+	ReleasePreloadEntry(GroupIndex
+#if UE_SCA_DEBUG_PRELOADING
+		, Callsite
+#endif
+	);
 
 	if (Shader)
 	{
diff --git a/Engine/Source/Runtime/RenderCore/Public/ShaderCodeArchive.h b/Engine/Source/Runtime/RenderCore/Public/ShaderCodeArchive.h
--- a/Engine/Source/Runtime/RenderCore/Public/ShaderCodeArchive.h
+++ b/Engine/Source/Runtime/RenderCore/Public/ShaderCodeArchive.h
@@ -33,6 +33,9 @@ class FCbWriter;
 // enable visualization in the desktop Development builds only as it has a memory hit and writes files
 #define UE_SCA_VISUALIZE_SHADER_USAGE			(!WITH_EDITOR && UE_BUILD_DEVELOPMENT && PLATFORM_DESKTOP)
 
+// enable deep, manual debugging of leaked preload groups. This level of information slows the engine down and is only needed when chasing tricky bugs
+#define UE_SCA_DEBUG_PRELOADING					(0)
+
 struct FShaderMapEntry
 {
 	uint32 ShaderIndicesOffset = 0u;
@@ -610,6 +613,10 @@ private:
 		uint32 NumRefs : 31;
 		uint32 bNeverToBePreloaded : 1;
 
+#if UE_SCA_DEBUG_PRELOADING
+		FString DebugInfo;
+#endif
+
 		FShaderGroupPreloadEntry()
 			: NumRefs(0)
 			, bNeverToBePreloaded(0)
@@ -622,16 +629,28 @@ private:
 	FIoDispatcher& IoDispatcher;
 
 	/** Preloads a given shader group. */
-	bool PreloadShaderGroup(int32 ShaderGroupIndex, FGraphEventArray& OutCompletionEvents, FCoreDelegates::FAttachShaderReadRequestFunc* AttachShaderReadRequestFuncPtr = nullptr);
+	bool PreloadShaderGroup(int32 ShaderGroupIndex, FGraphEventArray& OutCompletionEvents, 
+#if UE_SCA_DEBUG_PRELOADING
+		const FString& CallsiteInfo,
+#endif
+		FCoreDelegates::FAttachShaderReadRequestFunc* AttachShaderReadRequestFuncPtr = nullptr);
 
 	/** Sets up a new preload entry for preload.*/
 	void SetupPreloadEntryForLoading(int32 ShaderGroupIndex, FShaderGroupPreloadEntry& PreloadEntry);
 
 	/** Sets up a preload entry for groups that shouldn't be preloaded.*/
-	void MarkPreloadEntrySkipped(int32 ShaderGroupIndex);
+	void MarkPreloadEntrySkipped(int32 ShaderGroupIndex
+#if UE_SCA_DEBUG_PRELOADING
+		, const FString& CallsiteInfo
+#endif
+	);
 
 	/** Releases a reference to a preloaded shader group, potentially deleting it. */
-	void ReleasePreloadEntry(int32 ShaderGroupIndex);
+	void ReleasePreloadEntry(int32 ShaderGroupIndex
+#if UE_SCA_DEBUG_PRELOADING
+		, const FString& CallsiteInfo
+#endif
+	);
 
 	/** Returns the index of shader group that a given shader belongs to. */
 	inline int32 GetGroupIndexForShader(int32 ShaderIndex) const
