Add ensure to report aligned resource overflow & set back to unaligned size.

diff --git a/Engine/Source/Runtime/RenderCore/Private/RenderGraphResourcePool.cpp b/Engine/Source/Runtime/RenderCore/Private/RenderGraphResourcePool.cpp
--- a/Engine/Source/Runtime/RenderCore/Private/RenderGraphResourcePool.cpp
+++ b/Engine/Source/Runtime/RenderCore/Private/RenderGraphResourcePool.cpp
@@ -71,6 +71,12 @@ TRefCountPtr<FRDGPooledBuffer> FRDGBufferPool::FindFreeBuffer(const FRDGBufferDe
 		AlignedDesc.NumElements = Align(AlignedDesc.BytesPerElement * AlignedDesc.NumElements, BufferPageSize) / AlignedDesc.BytesPerElement;
 	}
 
+	if (!ensureMsgf(AlignedDesc.NumElements >= Desc.NumElements, TEXT("Alignment caused buffer size overflow for buffer '%s' (AlignedDesc.NumElements: %d < Desc.NumElements: %d)"), InDebugName, AlignedDesc.NumElements, Desc.NumElements))
+	{
+		// Use the unaligned desc since we apparently overflowed when rounding up.
+		AlignedDesc = Desc;
+	}
+
 	const uint32 BufferHash = GetTypeHash(AlignedDesc);
 	
 	// First find if available.
