Use FMath::RoundUpToPowerOfTwo64 in FRDGScatterUploadBuffer::Init
- the 32 bit version could return 0 when inputs caused an overflow which caused a bad result due to the usage of FMath::Min afterwards.

diff --git a/Engine/Source/Runtime/RenderCore/Private/UnifiedBuffer.cpp b/Engine/Source/Runtime/RenderCore/Private/UnifiedBuffer.cpp
--- a/Engine/Source/Runtime/RenderCore/Private/UnifiedBuffer.cpp
+++ b/Engine/Source/Runtime/RenderCore/Private/UnifiedBuffer.cpp
@@ -836,12 +836,12 @@ void FRDGScatterUploadBuffer::Init(FRDGBuilder& GraphBuilder, uint32 NumElements
 
 	const uint32 ScatterNumBytesPerElement = sizeof(uint32);
 	const uint32 ScatterBytes = NumElements * ScatterNumBytesPerElement;
-	const uint32 ScatterBufferSize = (uint32)FMath::Min((uint64)FMath::RoundUpToPowerOfTwo(ScatterBytes), GetMaxUploadBufferElements() * sizeof(uint32));
+	const uint32 ScatterBufferSize = (uint32)FMath::Min(FMath::RoundUpToPowerOfTwo64((uint64)ScatterBytes), GetMaxUploadBufferElements() * sizeof(uint32));
 	check(ScatterBufferSize >= ScatterBytes);
 
 	const uint32 UploadNumBytesPerElement = TypeSize;
 	const uint32 UploadBytes = NumElements * NumBytesPerElement;
-	const uint32 UploadBufferSize = (uint32)FMath::Min((uint64)FMath::RoundUpToPowerOfTwo(UploadBytes), GetMaxUploadBufferElements() * TypeSize);
+	const uint32 UploadBufferSize = (uint32)FMath::Min(FMath::RoundUpToPowerOfTwo64((uint64)UploadBytes), GetMaxUploadBufferElements() * TypeSize);
 	check(UploadBufferSize >= UploadBytes);
 
 	// Recreate buffers is they are already queued into RDG from a previous call.
