Adding a fast path for shader parameter unsets, or "Unbinds"
- FRHIBatchedShaderUnbinds acts like FRHIBatchedShaderParameters but doesn't need to collect pointers.
- The context code can also avoid any attempts to access the resource since we know it's always null.
- This change will let us easily remove Unbinds for RHIs that won't need them in the future, aka after the parameter batching is done once per draw.
- This addresses the bulk of unbind performance regressions but there is still some a minor regression in the single unbind case.

diff --git a/Engine/Source/Runtime/Apple/MetalRHI/Private/MetalCommands.cpp b/Engine/Source/Runtime/Apple/MetalRHI/Private/MetalCommands.cpp
--- a/Engine/Source/Runtime/Apple/MetalRHI/Private/MetalCommands.cpp
+++ b/Engine/Source/Runtime/Apple/MetalRHI/Private/MetalCommands.cpp
@@ -390,6 +390,16 @@ void FMetalRHICommandContext::RHISetShaderParameters(FRHIComputeShader* Shader,
 	);
 }
 
+void FMetalRHICommandContext::RHISetShaderUnbinds(FRHIComputeShader* Shader, TConstArrayView<FRHIShaderParameterUnbind> InUnbinds)
+{
+	UE::RHICore::RHISetShaderUnbindsShared(*this, Shader, InUnbinds);
+}
+
+void FMetalRHICommandContext::RHISetShaderUnbinds(FRHIGraphicsShader* Shader, TConstArrayView<FRHIShaderParameterUnbind> InUnbinds)
+{
+	UE::RHICore::RHISetShaderUnbindsShared(*this, Shader, InUnbinds);
+}
+
 void FMetalRHICommandContext::RHISetStencilRef(uint32 StencilRef)
 {
 	Context->GetCurrentState().SetStencilRef(StencilRef);
diff --git a/Engine/Source/Runtime/Apple/MetalRHI/Public/MetalRHIContext.h b/Engine/Source/Runtime/Apple/MetalRHI/Public/MetalRHIContext.h
--- a/Engine/Source/Runtime/Apple/MetalRHI/Public/MetalRHIContext.h
+++ b/Engine/Source/Runtime/Apple/MetalRHI/Public/MetalRHIContext.h
@@ -144,6 +144,8 @@ public:
 	
 	virtual void RHISetShaderParameters(FRHIGraphicsShader* Shader, TConstArrayView<uint8> InParametersData, TConstArrayView<FRHIShaderParameter> InParameters, TConstArrayView<FRHIShaderParameterResource> InResourceParameters, TConstArrayView<FRHIShaderParameterResource> InBindlessParameters) final override;
 	virtual void RHISetShaderParameters(FRHIComputeShader* Shader, TConstArrayView<uint8> InParametersData, TConstArrayView<FRHIShaderParameter> InParameters, TConstArrayView<FRHIShaderParameterResource> InResourceParameters, TConstArrayView<FRHIShaderParameterResource> InBindlessParameters) final override;
+	virtual void RHISetShaderUnbinds(FRHIComputeShader* Shader, TConstArrayView<FRHIShaderParameterUnbind> InUnbinds) final override;
+	virtual void RHISetShaderUnbinds(FRHIGraphicsShader* Shader, TConstArrayView<FRHIShaderParameterUnbind> InUnbinds) final override;
 
 	virtual void RHISetStencilRef(uint32 StencilRef) final override;
 
diff --git a/Engine/Source/Runtime/D3D12RHI/Private/D3D12CommandContext.h b/Engine/Source/Runtime/D3D12RHI/Private/D3D12CommandContext.h
--- a/Engine/Source/Runtime/D3D12RHI/Private/D3D12CommandContext.h
+++ b/Engine/Source/Runtime/D3D12RHI/Private/D3D12CommandContext.h
@@ -588,6 +588,8 @@ public:
 	virtual void RHISetShaderUniformBuffer(FRHIComputeShader* ComputeShader, uint32 BufferIndex, FRHIUniformBuffer* Buffer) final override;
 	virtual void RHISetShaderParameters(FRHIComputeShader* Shader, TConstArrayView<uint8> InParametersData, TConstArrayView<FRHIShaderParameter> InParameters, TConstArrayView<FRHIShaderParameterResource> InResourceParameters, TConstArrayView<FRHIShaderParameterResource> InBindlessParameters) final override;
 	virtual void RHISetShaderParameter(FRHIComputeShader* ComputeShader, uint32 BufferIndex, uint32 BaseIndex, uint32 NumBytes, const void* NewValue) final override;
+	virtual void RHISetShaderUnbinds(FRHIComputeShader* Shader, TConstArrayView<FRHIShaderParameterUnbind> InUnbinds) final override;
+	virtual void RHISetShaderUnbinds(FRHIGraphicsShader* Shader, TConstArrayView<FRHIShaderParameterUnbind> InUnbinds) final override;
 	virtual void RHIPushEvent(const TCHAR* Name, FColor Color) final override;
 	virtual void RHIPopEvent() final override;
 	virtual void RHISubmitCommandsHint() final override;
@@ -844,6 +846,10 @@ public:
 	{
 		ContextRedirect(RHISetShaderParameters(Shader, InParametersData, InParameters, InResourceParameters, InBindlessParameters));
 	}
+	FORCEINLINE virtual void RHISetShaderUnbinds(FRHIComputeShader* Shader, TConstArrayView<FRHIShaderParameterUnbind> InUnbinds) final override
+	{
+		ContextRedirect(RHISetShaderUnbinds(Shader, InUnbinds));
+	}
 	FORCEINLINE virtual void RHIPushEvent(const TCHAR* Name, FColor Color) final override
 	{
 		ContextRedirect(RHIPushEvent(Name, Color));
@@ -947,6 +953,10 @@ public:
 	{
 		ContextRedirect(RHISetShaderParameters(Shader, InParametersData, InParameters, InResourceParameters, InBindlessParameters));
 	}
+	FORCEINLINE virtual void RHISetShaderUnbinds(FRHIGraphicsShader* Shader, TConstArrayView<FRHIShaderParameterUnbind> InUnbinds) final override
+	{
+		ContextRedirect(RHISetShaderUnbinds(Shader, InUnbinds));
+	}
 	FORCEINLINE virtual void RHISetStencilRef(uint32 StencilRef) final override
 	{
 		ContextRedirect(RHISetStencilRef(StencilRef));
diff --git a/Engine/Source/Runtime/D3D12RHI/Private/D3D12Commands.cpp b/Engine/Source/Runtime/D3D12RHI/Private/D3D12Commands.cpp
--- a/Engine/Source/Runtime/D3D12RHI/Private/D3D12Commands.cpp
+++ b/Engine/Source/Runtime/D3D12RHI/Private/D3D12Commands.cpp
@@ -1196,6 +1196,58 @@ void FD3D12CommandContext::RHISetShaderParameters(FRHIComputeShader* Shader, TAr
 	);
 }
 
+static void SetShaderUnbindsOnContext(
+	FD3D12CommandContext& Context
+	, FRHIShader* Shader
+	, EShaderFrequency ShaderFrequency
+	, TConstArrayView<FRHIShaderParameterUnbind> InUnbinds)
+{
+	for (const FRHIShaderParameterUnbind& Unbind : InUnbinds)
+	{
+		switch (Unbind.Type)
+		{
+		case FRHIShaderParameterUnbind::EType::ResourceView:
+			Context.StateCache.SetShaderResourceView(ShaderFrequency, nullptr, Unbind.Index);
+			break;
+		case FRHIShaderParameterUnbind::EType::UnorderedAccessView:
+			if (ShaderFrequency == SF_Pixel || ShaderFrequency == SF_Compute)
+			{
+				Context.StateCache.SetUAV(ShaderFrequency, Unbind.Index, nullptr);
+			}
+			else
+			{
+				checkf(false, TEXT("TShaderRHI Can't have compute shader to be set. UAVs are not supported on vertex, tessellation and geometry shaders."));
+			}
+			break;
+		default:
+			checkf(false, TEXT("Unhandled resource type?"));
+			break;
+		}
+	}
+}
+
+void FD3D12CommandContext::RHISetShaderUnbinds(FRHIGraphicsShader* Shader, TConstArrayView<FRHIShaderParameterUnbind> InUnbinds)
+{
+	const EShaderFrequency ShaderFrequency = Shader->GetFrequency();
+	if (IsValidGraphicsFrequency(ShaderFrequency))
+	{
+		ValidateBoundShader(StateCache, Shader);
+
+		SetShaderUnbindsOnContext(*this, Shader, ShaderFrequency, InUnbinds);
+	}
+	else
+	{
+		checkf(0, TEXT("Unsupported FRHIGraphicsShader Type '%s'!"), GetShaderFrequencyString(ShaderFrequency, false));
+	}
+}
+
+void FD3D12CommandContext::RHISetShaderUnbinds(FRHIComputeShader* Shader, TConstArrayView<FRHIShaderParameterUnbind> InUnbinds)
+{
+	//ValidateBoundShader(StateCache, Shader);
+
+	SetShaderUnbindsOnContext(*this, Shader, SF_Compute, InUnbinds);
+}
+
 void FD3D12CommandContext::RHISetShaderParameter(FRHIGraphicsShader* ShaderRHI, uint32 BufferIndex, uint32 Offset, uint32 NumBytes, const void* NewValue)
 {
 	checkSlow(BufferIndex == 0);
diff --git a/Engine/Source/Runtime/Engine/Private/GPUSkinCache.cpp b/Engine/Source/Runtime/Engine/Private/GPUSkinCache.cpp
--- a/Engine/Source/Runtime/Engine/Private/GPUSkinCache.cpp
+++ b/Engine/Source/Runtime/Engine/Private/GPUSkinCache.cpp
@@ -711,10 +711,10 @@ public:
 		SetShaderValue(BatchedParameters, SkinCacheDebug, CVarGPUSkinCacheDebug.GetValueOnRenderThread());
 	}
 
-	void UnsetParameters(FRHIBatchedShaderParameters& BatchedParameters)
+	void UnsetParameters(FRHIBatchedShaderUnbinds& BatchedUnbinds)
 	{
-		SetUAVParameter(BatchedParameters, PositionBufferUAV, 0);
-		SetUAVParameter(BatchedParameters, TangentBufferUAV, 0);
+		UnsetUAVParameter(BatchedUnbinds, PositionBufferUAV);
+		UnsetUAVParameter(BatchedUnbinds, TangentBufferUAV);
 	}
 
 private:
@@ -940,9 +940,9 @@ public:
         }
 	}
 
-	void UnsetParameters(FRHIBatchedShaderParameters& BatchedParameters)
+	void UnsetParameters(FRHIBatchedShaderUnbinds& BatchedUnbinds)
 	{
-		SetUAVParameter(BatchedParameters, IntermediateAccumBufferUAV, 0);
+		UnsetUAVParameter(BatchedUnbinds, IntermediateAccumBufferUAV);
 	}
 
 	LAYOUT_FIELD(FShaderResourceParameter, IntermediateAccumBufferUAV);
@@ -1055,10 +1055,10 @@ public:
 		SetSRVParameter(BatchedParameters, ColorInputBuffer, DispatchData.ColorBufferSRV);
 	}
 
-	void UnsetParameters(FRHIBatchedShaderParameters& BatchedParameters)
+	void UnsetParameters(FRHIBatchedShaderUnbinds& BatchedUnbinds)
 	{
-		SetUAVParameter(BatchedParameters, TangentBufferUAV, 0);
-		SetUAVParameter(BatchedParameters, IntermediateAccumBufferUAV, 0);
+		UnsetUAVParameter(BatchedUnbinds, TangentBufferUAV);
+		UnsetUAVParameter(BatchedUnbinds, IntermediateAccumBufferUAV);
 	}
 };
 
diff --git a/Engine/Source/Runtime/Engine/Private/GPUSort.cpp b/Engine/Source/Runtime/Engine/Private/GPUSort.cpp
--- a/Engine/Source/Runtime/Engine/Private/GPUSort.cpp
+++ b/Engine/Source/Runtime/Engine/Private/GPUSort.cpp
@@ -240,9 +240,9 @@ public:
 		SetUAVParameter(BatchedParameters, OutOffsets, OutOffsetsUAV);
 	}
 
-	void UnsetParameters(FRHIBatchedShaderParameters& BatchedParameters)
+	void UnsetParameters(FRHIBatchedShaderUnbinds& BatchedUnbinds)
 	{
-		SetUAVParameter(BatchedParameters, OutOffsets, nullptr);
+		UnsetUAVParameter(BatchedUnbinds, OutOffsets);
 	}
 
 private:
@@ -303,11 +303,11 @@ public:
 		SetSRVParameter(BatchedParameters, RadixSortParameterBuffer, RadixSortParameterBufferSRV);
 	}
 
-	void UnsetParameters(FRHIBatchedShaderParameters& BatchedParameters)
+	void UnsetParameters(FRHIBatchedShaderUnbinds& BatchedUnbinds)
 	{
-		SetSRVParameter(BatchedParameters, RadixSortParameterBuffer, nullptr);
-		SetSRVParameter(BatchedParameters, InKeys, nullptr);
-		SetUAVParameter(BatchedParameters, OutOffsets, nullptr);
+		UnsetSRVParameter(BatchedUnbinds, RadixSortParameterBuffer);
+		UnsetSRVParameter(BatchedUnbinds, InKeys);
+		UnsetUAVParameter(BatchedUnbinds, OutOffsets);
 	}
 
 private:
@@ -358,10 +358,10 @@ public:
 		SetSRVParameter(BatchedParameters, InOffsets, InOffsetsSRV);
 	}
 
-	void UnsetParameters(FRHIBatchedShaderParameters& BatchedParameters)
+	void UnsetParameters(FRHIBatchedShaderUnbinds& BatchedUnbinds)
 	{
-		SetSRVParameter(BatchedParameters, InOffsets, nullptr);
-		SetUAVParameter(BatchedParameters, OutOffsets, nullptr);
+		UnsetSRVParameter(BatchedUnbinds, InOffsets);
+		UnsetUAVParameter(BatchedUnbinds, OutOffsets);
 	}
 
 private:
@@ -441,14 +441,14 @@ public:
 		SetSRVParameter(BatchedParameters, InOffsets, InOffsetsSRV);
 	}
 
-	void UnsetParameters(FRHIBatchedShaderParameters& BatchedParameters)
+	void UnsetParameters(FRHIBatchedShaderUnbinds& BatchedUnbinds)
 	{
-		SetSRVParameter(BatchedParameters, RadixSortParameterBuffer, nullptr);
-		SetSRVParameter(BatchedParameters, InKeys, nullptr);
-		SetSRVParameter(BatchedParameters, InValues, nullptr);
-		SetSRVParameter(BatchedParameters, InOffsets, nullptr);
-		SetUAVParameter(BatchedParameters, OutKeys, nullptr);
-		SetUAVParameter(BatchedParameters, OutValues, nullptr);
+		UnsetSRVParameter(BatchedUnbinds, RadixSortParameterBuffer);
+		UnsetSRVParameter(BatchedUnbinds, InKeys);
+		UnsetSRVParameter(BatchedUnbinds, InValues);
+		UnsetSRVParameter(BatchedUnbinds, InOffsets);
+		UnsetUAVParameter(BatchedUnbinds, OutKeys);
+		UnsetUAVParameter(BatchedUnbinds, OutValues);
 	}
 
 private:
diff --git a/Engine/Source/Runtime/Engine/Private/GPUSortManager.cpp b/Engine/Source/Runtime/Engine/Private/GPUSortManager.cpp
--- a/Engine/Source/Runtime/Engine/Private/GPUSortManager.cpp
+++ b/Engine/Source/Runtime/Engine/Private/GPUSortManager.cpp
@@ -114,7 +114,7 @@ public:
 		int32 StartingIndex,
 		int32 DestCount);
 
-	void UnsetParameters(FRHIBatchedShaderParameters& BatchedParameters);
+	void UnsetParameters(FRHIBatchedShaderUnbinds& BatchedUnbinds);
 
 	void Begin(FRHICommandList& RHICmdList);
 	void End(FRHICommandList& RHICmdList);
@@ -171,12 +171,12 @@ void FCopyUIntBufferCS::SetParameters(
 	SetShaderValue(BatchedParameters, CopyParams, CopyParamsValue);
 }
 
-void FCopyUIntBufferCS::UnsetParameters(FRHIBatchedShaderParameters& BatchedParameters)
+void FCopyUIntBufferCS::UnsetParameters(FRHIBatchedShaderUnbinds& BatchedUnbinds)
 {
-	SetSRVParameter(BatchedParameters, SourceData, nullptr);
+	UnsetSRVParameter(BatchedUnbinds, SourceData);
 	for (int32 Index = 0; Index < COPYUINTCS_BUFFER_COUNT; ++Index)
 	{
-		SetUAVParameter(BatchedParameters, DestData[Index], nullptr);
+		UnsetUAVParameter(BatchedUnbinds, DestData[Index]);
 	}
 }
 
diff --git a/Engine/Source/Runtime/Engine/Private/Particles/ParticleSortingGPU.cpp b/Engine/Source/Runtime/Engine/Private/Particles/ParticleSortingGPU.cpp
--- a/Engine/Source/Runtime/Engine/Private/Particles/ParticleSortingGPU.cpp
+++ b/Engine/Source/Runtime/Engine/Private/Particles/ParticleSortingGPU.cpp
@@ -94,11 +94,11 @@ public:
 		SetSRVParameter(BatchedParameters, InParticleIndices, InIndicesSRV);
 	}
 
-	void UnsetParameters(FRHIBatchedShaderParameters& BatchedParameters)
+	void UnsetParameters(FRHIBatchedShaderUnbinds& BatchedUnbinds)
 	{
-		SetSRVParameter(BatchedParameters, InParticleIndices, nullptr);
-		SetUAVParameter(BatchedParameters, OutKeys, nullptr);
-		SetUAVParameter(BatchedParameters, OutParticleIndices, nullptr);
+		UnsetSRVParameter(BatchedUnbinds, InParticleIndices);
+		UnsetUAVParameter(BatchedUnbinds, OutKeys);
+		UnsetUAVParameter(BatchedUnbinds, OutParticleIndices);
 	}
 
 private:
diff --git a/Engine/Source/Runtime/Engine/Private/PhysicsField/PhysicsFieldComponent.cpp b/Engine/Source/Runtime/Engine/Private/PhysicsField/PhysicsFieldComponent.cpp
--- a/Engine/Source/Runtime/Engine/Private/PhysicsField/PhysicsFieldComponent.cpp
+++ b/Engine/Source/Runtime/Engine/Private/PhysicsField/PhysicsFieldComponent.cpp
@@ -255,9 +255,9 @@ public:
 		}
 	}
 
-	void UnsetParameters(FRHIBatchedShaderParameters& BatchedParameters)
+	void UnsetParameters(FRHIBatchedShaderUnbinds& BatchedUnbinds)
 	{
-		SetUAVParameter(BatchedParameters, FieldClipmap, nullptr);
+		UnsetUAVParameter(BatchedUnbinds, FieldClipmap);
 	}
 
 private:
@@ -355,9 +355,9 @@ public:
 		}
 	}
 
-	void UnsetParameters(FRHIBatchedShaderParameters& BatchedParameters)
+	void UnsetParameters(FRHIBatchedShaderUnbinds& BatchedUnbinds)
 	{
-		SetUAVParameter(BatchedParameters, FieldClipmap, nullptr);
+		UnsetUAVParameter(BatchedUnbinds, FieldClipmap);
 	}
 
 private:
diff --git a/Engine/Source/Runtime/Engine/Private/SkeletalRenderGPUSkin.cpp b/Engine/Source/Runtime/Engine/Private/SkeletalRenderGPUSkin.cpp
--- a/Engine/Source/Runtime/Engine/Private/SkeletalRenderGPUSkin.cpp
+++ b/Engine/Source/Runtime/Engine/Private/SkeletalRenderGPUSkin.cpp
@@ -1156,9 +1156,9 @@ void FGPUMorphUpdateCS::Dispatch(FRHICommandList& RHICmdList, uint32 Size)
 	RHICmdList.DispatchComputeShader(DispatchSize.X, DispatchSize.Y, DispatchSize.Z);
 }
 
-void FGPUMorphUpdateCS::UnsetParameters(FRHIBatchedShaderParameters& BatchedParameters)
+void FGPUMorphUpdateCS::UnsetParameters(FRHIBatchedShaderUnbinds& BatchedUnbinds)
 {
-	SetUAVParameter(BatchedParameters, MorphVertexBufferParameter, nullptr);
+	UnsetUAVParameter(BatchedUnbinds, MorphVertexBufferParameter);
 }
 
 bool FGPUMorphUpdateCS::ShouldCompilePermutation(const FGlobalShaderPermutationParameters& Parameters)
@@ -1196,9 +1196,9 @@ void FGPUMorphNormalizeCS::Dispatch(FRHICommandList& RHICmdList, uint32 NumVerti
 	RHICmdList.DispatchComputeShader(DispatchSize.X, DispatchSize.Y, DispatchSize.Z);
 }
 
-void FGPUMorphNormalizeCS::UnsetParameters(FRHIBatchedShaderParameters& BatchedParameters)
+void FGPUMorphNormalizeCS::UnsetParameters(FRHIBatchedShaderUnbinds& BatchedUnbinds)
 {
-	SetUAVParameter(BatchedParameters, MorphVertexBufferParameter, nullptr);
+	UnsetUAVParameter(BatchedUnbinds, MorphVertexBufferParameter);
 }
 
 IMPLEMENT_SHADER_TYPE(, FGPUMorphNormalizeCS, TEXT("/Engine/Private/MorphTargets.usf"), TEXT("GPUMorphNormalizeCS"), SF_Compute);
diff --git a/Engine/Source/Runtime/Engine/Private/SkeletalRenderGPUSkin.h b/Engine/Source/Runtime/Engine/Private/SkeletalRenderGPUSkin.h
--- a/Engine/Source/Runtime/Engine/Private/SkeletalRenderGPUSkin.h
+++ b/Engine/Source/Runtime/Engine/Private/SkeletalRenderGPUSkin.h
@@ -723,7 +723,7 @@ public:
 		float Weights[MorphTargetDispatchBatchSize]);
 
 	void Dispatch(FRHICommandList& RHICmdList, uint32 Size);
-	void UnsetParameters(FRHIBatchedShaderParameters& BatchedParameters);
+	void UnsetParameters(FRHIBatchedShaderUnbinds& BatchedUnbinds);
 
 	static bool ShouldCompilePermutation(const FGlobalShaderPermutationParameters& Parameters);
 
@@ -754,7 +754,7 @@ public:
 	void SetParameters(FRHIBatchedShaderParameters& BatchedParameters, const FVector4& LocalScale, const FMorphTargetVertexInfoBuffers& MorphTargetVertexInfoBuffers, FMorphVertexBuffer& MorphVertexBuffer, uint32 NumVertices);
 
 	void Dispatch(FRHICommandList& RHICmdList, uint32 NumVertices);
-	void UnsetParameters(FRHIBatchedShaderParameters& BatchedParameters);
+	void UnsetParameters(FRHIBatchedShaderUnbinds& BatchedUnbinds);
 
 protected:
 	LAYOUT_FIELD(FShaderResourceParameter, MorphVertexBufferParameter);
diff --git a/Engine/Source/Runtime/Engine/Private/VectorField.cpp b/Engine/Source/Runtime/Engine/Private/VectorField.cpp
--- a/Engine/Source/Runtime/Engine/Private/VectorField.cpp
+++ b/Engine/Source/Runtime/Engine/Private/VectorField.cpp
@@ -1002,9 +1002,9 @@ public:
 		SetTextureParameter(BatchedParameters, NoiseVolumeTexture, NoiseVolumeTextureSampler, SamplerStateLinear, NoiseVolumeTextureRHI);
 	}
 
-	void UnsetParameters(FRHIBatchedShaderParameters& BatchedParameters)
+	void UnsetParameters(FRHIBatchedShaderUnbinds& BatchedUnbinds)
 	{
-		SetUAVParameter(BatchedParameters, OutVolumeTexture, nullptr);
+		UnsetUAVParameter(BatchedUnbinds, OutVolumeTexture);
 	}
 
 private:
diff --git a/Engine/Source/Runtime/Landscape/Private/LandscapeEditLayers.cpp b/Engine/Source/Runtime/Landscape/Private/LandscapeEditLayers.cpp
--- a/Engine/Source/Runtime/Landscape/Private/LandscapeEditLayers.cpp
+++ b/Engine/Source/Runtime/Landscape/Private/LandscapeEditLayers.cpp
@@ -927,9 +927,9 @@ public:
 		SetShaderValue(BatchedParameters, ComponentSizeParam, InParams.ComponentSize);
 	}
 
-	void UnsetParameters(FRHIBatchedShaderParameters& BatchedParameters)
+	void UnsetParameters(FRHIBatchedShaderUnbinds& BatchedUnbinds)
 	{
-		SetUAVParameter(BatchedParameters, AtlasPaintListsParam, nullptr);
+		UnsetUAVParameter(BatchedUnbinds, AtlasPaintListsParam);
 	}
 
 private:
@@ -1119,9 +1119,9 @@ public:
 		SetSRVParameter(BatchedParameters, WeightmapTextureOutputOffsetParam, InParams.ComputeShaderResource->WeightmapTextureOutputOffsetSRV);
 	}
 
-	void UnsetParameters(FRHIBatchedShaderParameters& BatchedParameters)
+	void UnsetParameters(FRHIBatchedShaderUnbinds& BatchedUnbinds)
 	{
-		SetUAVParameter(BatchedParameters, ComponentWeightmapParam, nullptr);
+		UnsetUAVParameter(BatchedUnbinds, ComponentWeightmapParam);
 	}
 
 private:
diff --git a/Engine/Source/Runtime/NullDrv/Public/NullRHI.h b/Engine/Source/Runtime/NullDrv/Public/NullRHI.h
--- a/Engine/Source/Runtime/NullDrv/Public/NullRHI.h
+++ b/Engine/Source/Runtime/NullDrv/Public/NullRHI.h
@@ -507,6 +507,14 @@ public:
 	{
 	}
 
+	virtual void RHISetShaderUnbinds(FRHIGraphicsShader* ComputeShader, TConstArrayView<FRHIShaderParameterUnbind> InUnbinds)
+	{
+	}
+
+	virtual void RHISetShaderUnbinds(FRHIComputeShader* ComputeShader, TConstArrayView<FRHIShaderParameterUnbind> InUnbinds)
+	{
+	}
+
 	virtual void RHISetDepthStencilState(FRHIDepthStencilState* NewState, uint32 StencilRef) final override
 	{
 
diff --git a/Engine/Source/Runtime/OpenGLDrv/Private/OpenGLCommands.cpp b/Engine/Source/Runtime/OpenGLDrv/Private/OpenGLCommands.cpp
--- a/Engine/Source/Runtime/OpenGLDrv/Private/OpenGLCommands.cpp
+++ b/Engine/Source/Runtime/OpenGLDrv/Private/OpenGLCommands.cpp
@@ -1178,6 +1178,16 @@ void FOpenGLDynamicRHI::RHISetShaderParameters(FRHIComputeShader* Shader, TConst
 	);
 }
 
+void FOpenGLDynamicRHI::RHISetShaderUnbinds(FRHIComputeShader* Shader, TConstArrayView<FRHIShaderParameterUnbind> InUnbinds)
+{
+	UE::RHICore::RHISetShaderUnbindsShared(*this, Shader, InUnbinds);
+}
+
+void FOpenGLDynamicRHI::RHISetShaderUnbinds(FRHIGraphicsShader* Shader, TConstArrayView<FRHIShaderParameterUnbind> InUnbinds)
+{
+	UE::RHICore::RHISetShaderUnbindsShared(*this, Shader, InUnbinds);
+}
+
 void FOpenGLDynamicRHI::RHISetDepthStencilState(FRHIDepthStencilState* NewStateRHI,uint32 StencilRef)
 {
 	VERIFY_GL_SCOPE();
diff --git a/Engine/Source/Runtime/OpenGLDrv/Public/OpenGLDrv.h b/Engine/Source/Runtime/OpenGLDrv/Public/OpenGLDrv.h
--- a/Engine/Source/Runtime/OpenGLDrv/Public/OpenGLDrv.h
+++ b/Engine/Source/Runtime/OpenGLDrv/Public/OpenGLDrv.h
@@ -458,6 +458,8 @@ public:
 	virtual void RHISetShaderParameter(FRHIComputeShader* ComputeShader, uint32 BufferIndex, uint32 BaseIndex, uint32 NumBytes, const void* NewValue) final override;
 	virtual void RHISetShaderParameters(FRHIGraphicsShader* Shader, TConstArrayView<uint8> InParametersData, TConstArrayView<FRHIShaderParameter> InParameters, TConstArrayView<FRHIShaderParameterResource> InResourceParameters, TConstArrayView<FRHIShaderParameterResource> InBindlessParameters) final override;
 	virtual void RHISetShaderParameters(FRHIComputeShader* Shader, TConstArrayView<uint8> InParametersData, TConstArrayView<FRHIShaderParameter> InParameters, TConstArrayView<FRHIShaderParameterResource> InResourceParameters, TConstArrayView<FRHIShaderParameterResource> InBindlessParameters) final override;
+	virtual void RHISetShaderUnbinds(FRHIComputeShader* Shader, TConstArrayView<FRHIShaderParameterUnbind> InUnbinds) final override;
+	virtual void RHISetShaderUnbinds(FRHIGraphicsShader* Shader, TConstArrayView<FRHIShaderParameterUnbind> InUnbinds) final override;
 	virtual void RHISetDepthStencilState(FRHIDepthStencilState* NewState, uint32 StencilRef) final override;
 	virtual void RHISetStencilRef(uint32 StencilRef) final override;
 	virtual void RHISetBlendState(FRHIBlendState* NewState, const FLinearColor& BlendFactor) final override;
diff --git a/Engine/Source/Runtime/RHI/Public/RHICommandList.h b/Engine/Source/Runtime/RHI/Public/RHICommandList.h
--- a/Engine/Source/Runtime/RHI/Public/RHICommandList.h
+++ b/Engine/Source/Runtime/RHI/Public/RHICommandList.h
@@ -750,6 +750,15 @@ public:
 		return ScratchShaderParameters;
 	}
 
+	inline FRHIBatchedShaderUnbinds& GetScratchShaderUnbinds()
+	{
+		if (!ensureMsgf(!ScratchShaderUnbinds.HasParameters(), TEXT("Scratch shader parameters left without committed parameters")))
+		{
+			ScratchShaderUnbinds.Reset();
+		}
+		return ScratchShaderUnbinds;
+	}
+
 protected:
 	FMemStackBase& GetAllocator() { return MemManager; }
 
@@ -821,6 +830,7 @@ protected:
 	TRHIPipelineArray<IRHIComputeContext*> Contexts = {};
 
 	FRHIBatchedShaderParameters ScratchShaderParameters;
+	FRHIBatchedShaderUnbinds ScratchShaderUnbinds;
 
 #if RHI_COUNT_COMMANDS
 	uint32 NumCommands = 0;
@@ -1206,6 +1216,19 @@ FRHICOMMAND_MACRO_TPL(TRHIShader, FRHICommandSetShaderParameters)
 	RHI_API void Execute(FRHICommandListBase& CmdList);
 };
 
+FRHICOMMAND_MACRO_TPL(TRHIShader, FRHICommandSetShaderUnbinds)
+{
+	TRHIShader* Shader;
+	TConstArrayView<FRHIShaderParameterUnbind> Unbinds;
+
+	FORCEINLINE_DEBUGGABLE FRHICommandSetShaderUnbinds(TRHIShader * InShader, TConstArrayView<FRHIShaderParameterUnbind> InUnbinds)
+		: Shader(InShader)
+		, Unbinds(InUnbinds)
+	{
+	}
+	RHI_API void Execute(FRHICommandListBase & CmdList);
+};
+
 FRHICOMMAND_MACRO(FRHICommandDrawPrimitive)
 {
 	uint32 BaseVertexIndex;
@@ -2077,6 +2100,7 @@ FRHICOMMAND_MACRO(FRHICommandSetRayTracingBindings)
 #endif // RHI_RAYTRACING
 
 template<> RHI_API void FRHICommandSetShaderParameters           <FRHIComputeShader>::Execute(FRHICommandListBase& CmdList);
+template<> RHI_API void FRHICommandSetShaderUnbinds              <FRHIComputeShader>::Execute(FRHICommandListBase& CmdList);
 
 extern RHI_API FRHIComputePipelineState* ExecuteSetComputePipelineState(FComputePipelineState* ComputePipelineState);
 extern RHI_API FRHIGraphicsPipelineState* ExecuteSetGraphicsPipelineState(class FGraphicsPipelineState* GraphicsPipelineState);
@@ -2193,6 +2217,29 @@ PRAGMA_ENABLE_DEPRECATION_WARNINGS
 		}
 	}
 
+	FORCEINLINE_DEBUGGABLE void SetShaderUnbinds(FRHIComputeShader* InShader, TConstArrayView<FRHIShaderParameterUnbind> InUnbinds)
+	{
+		ValidateBoundShader(InShader);
+
+		if (Bypass())
+		{
+			GetComputeContext().RHISetShaderUnbinds(InShader, InUnbinds);
+			return;
+		}
+
+		ALLOC_COMMAND(FRHICommandSetShaderUnbinds<FRHIComputeShader>)(InShader, AllocArray(InUnbinds));
+	}
+
+	FORCEINLINE_DEBUGGABLE void SetBatchedShaderUnbinds(FRHIComputeShader* InShader, FRHIBatchedShaderUnbinds& InBatchedUnbinds)
+	{
+		if (InBatchedUnbinds.HasParameters())
+		{
+			SetShaderUnbinds(InShader, InBatchedUnbinds.Unbinds);
+
+			InBatchedUnbinds.Reset();
+		}
+	}
+
 	UE_DEPRECATED(5.3, "FRHIBatchedShaderParameters and SetBatchedShaderParameters should be used instead of setting individual parameters.")
 	FORCEINLINE_DEBUGGABLE void SetShaderParameter(FRHIComputeShader* Shader, uint32 BufferIndex, uint32 BaseIndex, uint32 NumBytes, const void* NewValue)
 	{
@@ -2769,6 +2816,7 @@ PRAGMA_ENABLE_DEPRECATION_WARNINGS
 };
 
 template<> RHI_API void FRHICommandSetShaderParameters           <FRHIGraphicsShader>::Execute(FRHICommandListBase& CmdList);
+template<> RHI_API void FRHICommandSetShaderUnbinds              <FRHIGraphicsShader>::Execute(FRHICommandListBase& CmdList);
 
 class FRHICommandList : public FRHIComputeCommandList
 {
@@ -2890,6 +2938,33 @@ PRAGMA_ENABLE_DEPRECATION_WARNINGS
 		}
 	}
 
+	using FRHIComputeCommandList::SetShaderUnbinds;
+
+	FORCEINLINE_DEBUGGABLE void SetShaderUnbinds(FRHIGraphicsShader* InShader, TConstArrayView<FRHIShaderParameterUnbind> InUnbinds)
+	{
+		ValidateBoundShader(InShader);
+
+		if (Bypass())
+		{
+			GetContext().RHISetShaderUnbinds(InShader, InUnbinds);
+			return;
+		}
+
+		ALLOC_COMMAND(FRHICommandSetShaderUnbinds<FRHIGraphicsShader>)(InShader, AllocArray(InUnbinds));
+	}
+
+	using FRHIComputeCommandList::SetBatchedShaderUnbinds;
+
+	FORCEINLINE_DEBUGGABLE void SetBatchedShaderUnbinds(FRHIGraphicsShader* InShader, FRHIBatchedShaderUnbinds& InBatchedUnbinds)
+	{
+		if (InBatchedUnbinds.HasParameters())
+		{
+			SetShaderUnbinds(InShader, InBatchedUnbinds.Unbinds);
+
+			InBatchedUnbinds.Reset();
+		}
+	}
+
 	using FRHIComputeCommandList::SetShaderParameter;
 
 	UE_DEPRECATED(5.3, "FRHIBatchedShaderParameters and SetBatchedShaderParameters should be used instead of setting individual parameters.")
diff --git a/Engine/Source/Runtime/RHI/Public/RHICommandListCommandExecutes.inl b/Engine/Source/Runtime/RHI/Public/RHICommandListCommandExecutes.inl
--- a/Engine/Source/Runtime/RHI/Public/RHICommandListCommandExecutes.inl
+++ b/Engine/Source/Runtime/RHI/Public/RHICommandListCommandExecutes.inl
@@ -133,6 +133,18 @@ template<> RHI_API void FRHICommandSetShaderParameters<FRHIGraphicsShader>::Exec
 	INTERNAL_DECORATOR(RHISetShaderParameters)(Shader, ParametersData, Parameters, ResourceParameters, BindlessParameters);
 }
 
+template<> RHI_API void FRHICommandSetShaderUnbinds<FRHIComputeShader>::Execute(FRHICommandListBase& CmdList)
+{
+	RHISTAT(SetShaderUnbinds);
+	INTERNAL_DECORATOR_COMPUTE(RHISetShaderUnbinds)(Shader, Unbinds);
+}
+
+template<> RHI_API void FRHICommandSetShaderUnbinds<FRHIGraphicsShader>::Execute(FRHICommandListBase& CmdList)
+{
+	RHISTAT(SetShaderUnbinds);
+	INTERNAL_DECORATOR(RHISetShaderUnbinds)(Shader, Unbinds);
+}
+
 void FRHICommandDrawPrimitive::Execute(FRHICommandListBase& CmdList)
 {
 	RHISTAT(DrawPrimitive);
diff --git a/Engine/Source/Runtime/RHI/Public/RHIContext.h b/Engine/Source/Runtime/RHI/Public/RHIContext.h
--- a/Engine/Source/Runtime/RHI/Public/RHIContext.h
+++ b/Engine/Source/Runtime/RHI/Public/RHIContext.h
@@ -336,6 +336,8 @@ public:
 		InBatchedParameters.Reset();
 	}
 
+	virtual void RHISetShaderUnbinds(FRHIComputeShader* ComputeShader, TConstArrayView<FRHIShaderParameterUnbind> InUnbinds) = 0;
+
 	virtual void RHISetShaderParameter(FRHIComputeShader* ComputeShader, uint32 BufferIndex, uint32 BaseIndex, uint32 NumBytes, const void* NewValue) = 0;
 
 	virtual void RHISetStaticUniformBuffers(const FUniformBufferStaticBindings& InUniformBuffers)
@@ -695,6 +697,7 @@ public:
 	using IRHIComputeContext::RHISetShaderUniformBuffer;
 	using IRHIComputeContext::RHISetShaderParameters;
 	using IRHIComputeContext::RHISetBatchedShaderParameters;
+	using IRHIComputeContext::RHISetShaderUnbinds;
 	using IRHIComputeContext::RHISetShaderParameter;
 
 	/** Set the shader resource view of a surface. */
@@ -735,6 +738,8 @@ public:
 		InBatchedParameters.Reset();
 	}
 
+	virtual void RHISetShaderUnbinds(FRHIGraphicsShader* Shader, TConstArrayView<FRHIShaderParameterUnbind> InUnbinds) = 0;
+
 	virtual void RHISetShaderParameter(FRHIGraphicsShader* Shader, uint32 BufferIndex, uint32 BaseIndex, uint32 NumBytes, const void* NewValue) = 0;
 
 	virtual void RHISetStencilRef(uint32 StencilRef) {}
diff --git a/Engine/Source/Runtime/RHI/Public/RHIShaderParameters.h b/Engine/Source/Runtime/RHI/Public/RHIShaderParameters.h
--- a/Engine/Source/Runtime/RHI/Public/RHIShaderParameters.h
+++ b/Engine/Source/Runtime/RHI/Public/RHIShaderParameters.h
@@ -136,3 +136,49 @@ struct FRHIBatchedShaderParameters
 		BindlessParameters.Emplace(UAV, (uint16)Index);
 	}
 };
+
+/** Compact representation of a resource parameter unbind, limited to  SRVs and UAVs */
+struct FRHIShaderParameterUnbind
+{
+	enum class EType : uint8
+	{
+		ResourceView,
+		UnorderedAccessView,
+	};
+
+	FRHIShaderParameterUnbind() = default;
+	FRHIShaderParameterUnbind(EType InType, uint16 InIndex)
+		: Index(InIndex)
+		, Type(InType)
+	{
+	}
+
+	uint16  Index = 0;
+	EType   Type = EType::ResourceView;
+};
+
+/** Collection of parameters to unbind in the RHI. These unbinds aren't tied to any specific shader until SetBatchedShaderUnbinds is called. */
+struct FRHIBatchedShaderUnbinds
+{
+	TArray<FRHIShaderParameterUnbind> Unbinds;
+
+	bool HasParameters() const
+	{
+		return Unbinds.Num() > 0;
+	}
+
+	void Reset()
+	{
+		Unbinds.Reset();
+	}
+
+	void UnsetSRV(uint32 Index)
+	{
+		Unbinds.Emplace(FRHIShaderParameterUnbind::EType::ResourceView, (uint16)Index);
+	}
+	void UnsetUAV(uint32 Index)
+	{
+		Unbinds.Emplace(FRHIShaderParameterUnbind::EType::UnorderedAccessView, (uint16)Index);
+	}
+};
+
diff --git a/Engine/Source/Runtime/RHI/Public/RHIValidationContext.h b/Engine/Source/Runtime/RHI/Public/RHIValidationContext.h
--- a/Engine/Source/Runtime/RHI/Public/RHIValidationContext.h
+++ b/Engine/Source/Runtime/RHI/Public/RHIValidationContext.h
@@ -281,6 +281,12 @@ public:
 		RHIContext->RHISetShaderParameters(Shader, InParametersData, InParameters, InResourceParameters, InBindlessParameters);
 	}
 
+	virtual void RHISetShaderUnbinds(FRHIComputeShader* Shader, TConstArrayView<FRHIShaderParameterUnbind> InUnbinds) final override
+	{
+		checkf(State.bComputePSOSet, TEXT("A Compute PSO has to be set to set resources into a shader!"));
+		RHIContext->RHISetShaderUnbinds(Shader, InUnbinds);
+	}
+
 	virtual void RHISetStaticUniformBuffers(const FUniformBufferStaticBindings& InUniformBuffers) override final
 	{
 		InUniformBuffers.Bind(State.StaticUniformBuffers.Bindings);
@@ -812,6 +818,18 @@ public:
 		RHIContext->RHISetShaderParameters(Shader, InParametersData, InParameters, InResourceParameters, InBindlessParameters);
 	}
 
+	virtual void RHISetShaderUnbinds(FRHIGraphicsShader* Shader, TConstArrayView<FRHIShaderParameterUnbind> InUnbinds) override final
+	{
+		checkf(State.bGfxPSOSet, TEXT("A Graphics PSO has to be set to set resources into a shader!"));
+		RHIContext->RHISetShaderUnbinds(Shader, InUnbinds);
+	}
+
+	virtual void RHISetShaderUnbinds(FRHIComputeShader* Shader, TConstArrayView<FRHIShaderParameterUnbind> InUnbinds) override final
+	{
+		checkf(State.bComputePSOSet, TEXT("A Compute PSO has to be set to set resources into a shader!"));
+		RHIContext->RHISetShaderUnbinds(Shader, InUnbinds);
+	}
+
 	virtual void RHISetStaticUniformBuffers(const FUniformBufferStaticBindings& InUniformBuffers) override final
 	{
 		InUniformBuffers.Bind(State.StaticUniformBuffers.Bindings);
diff --git a/Engine/Source/Runtime/RHICore/Internal/RHIShaderParametersShared.h b/Engine/Source/Runtime/RHICore/Internal/RHIShaderParametersShared.h
--- a/Engine/Source/Runtime/RHICore/Internal/RHIShaderParametersShared.h
+++ b/Engine/Source/Runtime/RHICore/Internal/RHIShaderParametersShared.h
@@ -114,4 +114,24 @@ namespace UE::RHICore
 			}
 		}
 	}
+
+	template<typename FContextRHI, typename TShaderRHI>
+	inline void RHISetShaderUnbindsShared(FContextRHI& Context, TShaderRHI* ShaderRHI, TConstArrayView<FRHIShaderParameterUnbind> InUnbinds)
+	{
+		for (const FRHIShaderParameterUnbind& Unbind : InUnbinds)
+		{
+			switch (Unbind.Type)
+			{
+			case FRHIShaderParameterUnbind::EType::ResourceView:
+				Context.RHISetShaderResourceViewParameter(ShaderRHI, Unbind.Index, nullptr);
+				break;
+			case FRHIShaderParameterUnbind::EType::UnorderedAccessView:
+				SetShaderUAV(Context, ShaderRHI, Unbind.Index, nullptr);
+				break;
+			default:
+				checkf(false, TEXT("Unhandled resource type?"));
+				break;
+			}
+		}
+	}
 }
diff --git a/Engine/Source/Runtime/RenderCore/Private/BuiltInRayTracingShaders.cpp b/Engine/Source/Runtime/RenderCore/Private/BuiltInRayTracingShaders.cpp
--- a/Engine/Source/Runtime/RenderCore/Private/BuiltInRayTracingShaders.cpp
+++ b/Engine/Source/Runtime/RenderCore/Private/BuiltInRayTracingShaders.cpp
@@ -81,11 +81,11 @@ void FRayTracingDispatchDescCS::Dispatch(FRHICommandList& RHICmdList,
 
 	RHICmdList.DispatchComputeShader(1, 1, 1);
 
-	BatchedParameters.Reset();
+	FRHIBatchedShaderUnbinds& BatchedUnbinds = RHICmdList.GetScratchShaderUnbinds();
 
-	SetSRVParameter(BatchedParameters, ComputeShader->DispatchDimensionsParam, nullptr);
-	SetUAVParameter(BatchedParameters, ComputeShader->DispatchDescOutputParam, nullptr);
-	RHICmdList.SetBatchedShaderParameters(ComputeShader.GetComputeShader(), BatchedParameters);
+	UnsetSRVParameter(BatchedUnbinds, ComputeShader->DispatchDimensionsParam);
+	UnsetUAVParameter(BatchedUnbinds, ComputeShader->DispatchDescOutputParam);
+	RHICmdList.SetBatchedShaderUnbinds(ComputeShader.GetComputeShader(), BatchedUnbinds);
 }
 
 #endif // RHI_RAYTRACING
diff --git a/Engine/Source/Runtime/RenderCore/Private/RayTracingValidationShaders.cpp b/Engine/Source/Runtime/RenderCore/Private/RayTracingValidationShaders.cpp
--- a/Engine/Source/Runtime/RenderCore/Private/RayTracingValidationShaders.cpp
+++ b/Engine/Source/Runtime/RenderCore/Private/RayTracingValidationShaders.cpp
@@ -74,12 +74,12 @@ void FRayTracingValidateGeometryBuildParamsCS::Dispatch(FRHICommandList& RHICmdL
 		RHICmdList.DispatchComputeShader(NumGroupsX, 1, 1);
 	}
 
-	FRHIBatchedShaderParameters& BatchedParameters = RHICmdList.GetScratchShaderParameters();
+	FRHIBatchedShaderUnbinds& BatchedUnbinds = RHICmdList.GetScratchShaderUnbinds();
 
-	SetSRVParameter(BatchedParameters, ComputeShader->VertexBufferParam, nullptr);
-	SetSRVParameter(BatchedParameters, ComputeShader->IndexBufferParam, nullptr);
+	UnsetSRVParameter(BatchedUnbinds, ComputeShader->VertexBufferParam);
+	UnsetSRVParameter(BatchedUnbinds, ComputeShader->IndexBufferParam);
 
-	RHICmdList.SetBatchedShaderParameters(ShaderRHI, BatchedParameters);
+	RHICmdList.SetBatchedShaderUnbinds(ShaderRHI, BatchedUnbinds);
 
 	RHICmdList.PopEvent();
 }
@@ -118,9 +118,9 @@ void FRayTracingValidateSceneBuildParamsCS::Dispatch(FRHICommandList& RHICmdList
 		RHICmdList.DispatchComputeShader(NumGroupsX, 1, 1);
 	}
 
-	FRHIBatchedShaderParameters& BatchedParameters = RHICmdList.GetScratchShaderParameters();
-	SetSRVParameter(BatchedParameters, ComputeShader->InstanceBufferParam, nullptr);
-	RHICmdList.SetBatchedShaderParameters(ShaderRHI, BatchedParameters);
+	FRHIBatchedShaderUnbinds& BatchedUnbinds = RHICmdList.GetScratchShaderUnbinds();
+	UnsetSRVParameter(BatchedUnbinds, ComputeShader->InstanceBufferParam);
+	RHICmdList.SetBatchedShaderUnbinds(ShaderRHI, BatchedUnbinds);
 
 	RHICmdList.PopEvent();
 }
diff --git a/Engine/Source/Runtime/RenderCore/Public/ClearReplacementShaders.h b/Engine/Source/Runtime/RenderCore/Public/ClearReplacementShaders.h
--- a/Engine/Source/Runtime/RenderCore/Public/ClearReplacementShaders.h
+++ b/Engine/Source/Runtime/RenderCore/Public/ClearReplacementShaders.h
@@ -391,8 +391,9 @@ inline void ClearUAVShader_T(FRHIComputeCommandList& RHICmdList, FRHIUnorderedAc
 		FMath::DivideAndRoundUp(SizeZ, ComputeShader->ThreadGroupSizeZ)
 	);
 
-	SetUAVParameter(BatchedParameters, ComputeShader->GetClearResourceParam(), nullptr);
-	RHICmdList.SetBatchedShaderParameters(ShaderRHI, BatchedParameters);
+	FRHIBatchedShaderUnbinds& BatchedUnbinds = RHICmdList.GetScratchShaderUnbinds();
+	UnsetUAVParameter(BatchedUnbinds, ComputeShader->GetClearResourceParam());
+	RHICmdList.SetBatchedShaderUnbinds(ShaderRHI, BatchedUnbinds);
 
 	if (bBarriers)
 	{
diff --git a/Engine/Source/Runtime/RenderCore/Public/ShaderParameterStruct.h b/Engine/Source/Runtime/RenderCore/Public/ShaderParameterStruct.h
--- a/Engine/Source/Runtime/RenderCore/Public/ShaderParameterStruct.h
+++ b/Engine/Source/Runtime/RenderCore/Public/ShaderParameterStruct.h
@@ -129,94 +129,90 @@ FORCEINLINE void ValidateShaderParameters(const TShaderRef<TShaderClass>& Shader
 }
 
 template<typename TParameterType>
-inline void CollectUAVsToUnset(TArray<FRHIShaderParameterResource>& UAVsToUnset, const TMemoryImageArray<TParameterType>& InParameters)
+inline void CollectUAVsToUnset(FRHIBatchedShaderUnbinds& BatchedUnbinds, const TMemoryImageArray<TParameterType>& InParameters)
 {
-	FRHIUnorderedAccessView* NullUAV = nullptr;
 	for (const TParameterType& Parameter : InParameters)
 	{
 		if (Parameter.BaseType == UBMT_UAV ||
 			Parameter.BaseType == UBMT_RDG_TEXTURE_UAV ||
 			Parameter.BaseType == UBMT_RDG_BUFFER_UAV)
 		{
-			UAVsToUnset.Emplace(NullUAV, GetParameterIndex(Parameter));
+			BatchedUnbinds.UnsetUAV(GetParameterIndex(Parameter));
 		}
 	}
 }
 
 template<typename TParameterType>
-inline void CollectSRVsToUnset(TArray<FRHIShaderParameterResource>& SRVsToUnset, const TMemoryImageArray<TParameterType>& InParameters)
+inline void CollectSRVsToUnset(FRHIBatchedShaderUnbinds& BatchedUnbinds, const TMemoryImageArray<TParameterType>& InParameters)
 {
-	FRHIShaderResourceView* NullSRV = nullptr;
 	for (const TParameterType& Parameter : InParameters)
 	{
 		if (Parameter.BaseType == UBMT_SRV ||
 			Parameter.BaseType == UBMT_RDG_TEXTURE_SRV ||
 			Parameter.BaseType == UBMT_RDG_BUFFER_SRV)
 		{
-			SRVsToUnset.Emplace(NullSRV, GetParameterIndex(Parameter));
+			BatchedUnbinds.UnsetSRV(GetParameterIndex(Parameter));
 		}
 	}
 }
 
-inline void UnsetShaderUAVs(FRHIBatchedShaderParameters& BatchedParameters, const FShaderParameterBindings& Bindings)
+inline void UnsetShaderUAVs(FRHIBatchedShaderUnbinds& BatchedUnbinds, const FShaderParameterBindings& Bindings)
 {
-	CollectUAVsToUnset(BatchedParameters.ResourceParameters, Bindings.ResourceParameters);
-	CollectUAVsToUnset(BatchedParameters.BindlessParameters, Bindings.BindlessResourceParameters);
+	CollectUAVsToUnset(BatchedUnbinds, Bindings.ResourceParameters);
 }
 
-inline void UnsetShaderSRVs(FRHIBatchedShaderParameters& BatchedParameters, const FShaderParameterBindings& Bindings)
+inline void UnsetShaderSRVs(FRHIBatchedShaderUnbinds& BatchedUnbinds, const FShaderParameterBindings& Bindings)
 {
-	CollectSRVsToUnset(BatchedParameters.ResourceParameters, Bindings.ResourceParameters);
-	CollectSRVsToUnset(BatchedParameters.BindlessParameters, Bindings.BindlessResourceParameters);
+	CollectSRVsToUnset(BatchedUnbinds, Bindings.ResourceParameters);
 }
 
 template<typename TShaderClass>
-inline void UnsetShaderUAVs(FRHIBatchedShaderParameters& BatchedParameters, const TShaderRef<TShaderClass>& Shader)
+inline void UnsetShaderUAVs(FRHIBatchedShaderUnbinds& BatchedUnbinds, const TShaderRef<TShaderClass>& Shader)
 {
 	// TODO(RDG): Once all shader sets their parameter through this, can refactor RHI so all UAVs of a shader get unset through a single RHI function call.
 	const FShaderParameterBindings& Bindings = Shader->Bindings;
 
 	checkf(Bindings.RootParameterBufferIndex == FShaderParameterBindings::kInvalidBufferIndex, TEXT("Can't use UnsetShaderUAVs() for root parameter buffer index."));
 
-	UnsetShaderUAVs(BatchedParameters, Bindings);
+	UnsetShaderUAVs(BatchedUnbinds, Bindings);
 }
 
 template<typename TShaderClass>
-inline void UnsetShaderSRVs(FRHIBatchedShaderParameters& BatchedParameters, const TShaderRef<TShaderClass>& Shader)
+inline void UnsetShaderSRVs(FRHIBatchedShaderUnbinds& BatchedUnbinds, const TShaderRef<TShaderClass>& Shader)
 {
 	// TODO(RDG): Once all shader sets their parameter through this, can refactor RHI so all SRVs of a shader get unset through a single RHI function call.
 	const FShaderParameterBindings& Bindings = Shader->Bindings;
 
 	checkf(Bindings.RootParameterBufferIndex == FShaderParameterBindings::kInvalidBufferIndex, TEXT("Can't use UnsetShaderSRVs() for root parameter buffer index."));
 
-	UnsetShaderSRVs(BatchedParameters, Bindings);
+	UnsetShaderSRVs(BatchedUnbinds, Bindings);
 }
 
 /** Unset compute shader UAVs. */
 template<typename TRHICmdList, typename TShaderClass>
 inline void UnsetShaderUAVs(TRHICmdList& RHICmdList, const TShaderRef<TShaderClass>& Shader, FRHIComputeShader* ShaderRHI)
 {
-	FRHIBatchedShaderParameters& BatchedParameters = RHICmdList.GetScratchShaderParameters();
-	UnsetShaderUAVs(BatchedParameters, Shader);
-	RHICmdList.SetBatchedShaderParameters(ShaderRHI, BatchedParameters);
+	FRHIBatchedShaderUnbinds& BatchedUnbinds = RHICmdList.GetScratchShaderUnbinds();
+	UnsetShaderUAVs(BatchedUnbinds, Shader);
+	RHICmdList.SetBatchedShaderUnbinds(ShaderRHI, BatchedUnbinds);
 }
 
 /** Unset compute shader SRVs. */
 template<typename TRHICmdList, typename TShaderClass>
 inline void UnsetShaderSRVs(TRHICmdList& RHICmdList, const TShaderRef<TShaderClass>& Shader, FRHIComputeShader* ShaderRHI)
 {
-	FRHIBatchedShaderParameters& BatchedParameters = RHICmdList.GetScratchShaderParameters();
-	UnsetShaderSRVs(BatchedParameters, Shader);
-	RHICmdList.SetBatchedShaderParameters(ShaderRHI, BatchedParameters);
+	FRHIBatchedShaderUnbinds& BatchedUnbinds = RHICmdList.GetScratchShaderUnbinds();
+	UnsetShaderSRVs(BatchedUnbinds, Shader);
+	RHICmdList.SetBatchedShaderUnbinds(ShaderRHI, BatchedUnbinds);
 }
 
 /** Unset compute shader SRVs. */
 template<typename TRHICmdList, typename TShaderClass>
 inline void UnsetShaderSRVs(TRHICmdList& RHICmdList, const TShaderRef<TShaderClass>& Shader, FRHIGraphicsShader* ShaderRHI)
 {
-	FRHIBatchedShaderParameters& BatchedParameters = RHICmdList.GetScratchShaderParameters();
-	UnsetShaderSRVs(BatchedParameters, Shader);
-	RHICmdList.SetBatchedShaderParameters(ShaderRHI, BatchedParameters);
+	FRHIBatchedShaderUnbinds& BatchedUnbinds = RHICmdList.GetScratchShaderUnbinds();
+	UnsetShaderSRVs(BatchedUnbinds, Shader);
+	RHICmdList.SetBatchedShaderUnbinds(ShaderRHI, BatchedUnbinds);
 }
 
 RENDERCORE_API void SetShaderParameters(
diff --git a/Engine/Source/Runtime/RenderCore/Public/ShaderParameterUtils.h b/Engine/Source/Runtime/RenderCore/Public/ShaderParameterUtils.h
--- a/Engine/Source/Runtime/RenderCore/Public/ShaderParameterUtils.h
+++ b/Engine/Source/Runtime/RenderCore/Public/ShaderParameterUtils.h
@@ -170,6 +170,40 @@ inline void SetUAVParameter(FRHIBatchedShaderParameters& BatchedParameters, cons
 	}
 }
 
+inline void UnsetSRVParameter(FRHIBatchedShaderUnbinds& BatchedUnbinds, const FShaderResourceParameter& Parameter)
+{
+	if (Parameter.IsBound())
+	{
+#if PLATFORM_SUPPORTS_BINDLESS_RENDERING
+		if (Parameter.GetType() == EShaderParameterType::BindlessResourceIndex)
+		{
+			// We don't need to clear Bindless views
+		}
+		else
+#endif
+		{
+			BatchedUnbinds.UnsetSRV(Parameter.GetBaseIndex());
+		}
+	}
+}
+
+inline void UnsetUAVParameter(FRHIBatchedShaderUnbinds& BatchedUnbinds, const FShaderResourceParameter& Parameter)
+{
+	if (Parameter.IsBound())
+	{
+#if PLATFORM_SUPPORTS_BINDLESS_RENDERING
+		if (Parameter.GetType() == EShaderParameterType::BindlessResourceIndex)
+		{
+			// We don't need to clear Bindless views
+		}
+		else
+#endif
+		{
+			BatchedUnbinds.UnsetUAV(Parameter.GetBaseIndex());
+		}
+	}
+}
+
 inline void SetUniformBufferParameter(FRHIBatchedShaderParameters& BatchedParameters, const FShaderUniformBufferParameter& Parameter, FRHIUniformBuffer* UniformBufferRHI)
 {
 	// This will trigger if the parameter was not serialized
@@ -351,22 +385,22 @@ inline void SetShaderParametersLegacyCS(TRHICmdList& RHICmdList, const TShaderRe
 	SetShaderParametersLegacy(RHICmdList, InShader, InShader.GetComputeShader(), Forward<TArguments>(InArguments)...);
 }
 
-/// Utility to unset all legacy parameters for a Pixel shader. Requires the shader type to implement UnsetParameters(FRHIBatchedShaderParameters& BatchedParameters)
+/// Utility to unset all legacy parameters for a Pixel shader. Requires the shader type to implement UnsetParameters(FRHIBatchedShaderUnbinds& BatchedUnbinds)
 template<typename TRHICmdList, typename TShaderType>
 inline void UnsetShaderParametersLegacyPS(TRHICmdList& RHICmdList, const TShaderRef<TShaderType>& InShader)
 {
-	FRHIBatchedShaderParameters& BatchedParameters = RHICmdList.GetScratchShaderParameters();
-	InShader->UnsetParameters(BatchedParameters);
-	RHICmdList.SetBatchedShaderParameters(InShader.GetPixelShader(), BatchedParameters);
+	FRHIBatchedShaderUnbinds& BatchedUnbinds = RHICmdList.GetScratchShaderUnbinds();
+	InShader->UnsetParameters(BatchedUnbinds);
+	RHICmdList.SetBatchedShaderUnbinds(InShader.GetPixelShader(), BatchedUnbinds);
 }
 
-/// Utility to unset all legacy parameters for a Compute shader. Requires the shader type to implement UnsetParameters(FRHIBatchedShaderParameters& BatchedParameters)
+/// Utility to unset all legacy parameters for a Compute shader. Requires the shader type to implement UnsetParameters(FRHIBatchedShaderUnbinds& BatchedUnbinds)
 template<typename TRHICmdList, typename TShaderType>
 inline void UnsetShaderParametersLegacyCS(TRHICmdList& RHICmdList, const TShaderRef<TShaderType>& InShader)
 {
-	FRHIBatchedShaderParameters& BatchedParameters = RHICmdList.GetScratchShaderParameters();
-	InShader->UnsetParameters(BatchedParameters);
-	RHICmdList.SetBatchedShaderParameters(InShader.GetComputeShader(), BatchedParameters);
+	FRHIBatchedShaderUnbinds& BatchedUnbinds = RHICmdList.GetScratchShaderUnbinds();
+	InShader->UnsetParameters(BatchedUnbinds);
+	RHICmdList.SetBatchedShaderUnbinds(InShader.GetComputeShader(), BatchedUnbinds);
 }
 
 /**
diff --git a/Engine/Source/Runtime/VulkanRHI/Private/VulkanCommands.cpp b/Engine/Source/Runtime/VulkanRHI/Private/VulkanCommands.cpp
--- a/Engine/Source/Runtime/VulkanRHI/Private/VulkanCommands.cpp
+++ b/Engine/Source/Runtime/VulkanRHI/Private/VulkanCommands.cpp
@@ -434,6 +434,16 @@ void FVulkanCommandListContext::RHISetShaderParameters(FRHIComputeShader* Shader
 	);
 }
 
+void FVulkanCommandListContext::RHISetShaderUnbinds(FRHIComputeShader* Shader, TConstArrayView<FRHIShaderParameterUnbind> InUnbinds)
+{
+	UE::RHICore::RHISetShaderUnbindsShared(*this, Shader, InUnbinds);
+}
+
+void FVulkanCommandListContext::RHISetShaderUnbinds(FRHIGraphicsShader* Shader, TConstArrayView<FRHIShaderParameterUnbind> InUnbinds)
+{
+	UE::RHICore::RHISetShaderUnbindsShared(*this, Shader, InUnbinds);
+}
+
 void FVulkanCommandListContext::RHISetStaticUniformBuffers(const FUniformBufferStaticBindings& InUniformBuffers)
 {
 	FMemory::Memzero(GlobalUniformBuffers.GetData(), GlobalUniformBuffers.Num() * sizeof(FRHIUniformBuffer*));
diff --git a/Engine/Source/Runtime/VulkanRHI/Private/VulkanContext.h b/Engine/Source/Runtime/VulkanRHI/Private/VulkanContext.h
--- a/Engine/Source/Runtime/VulkanRHI/Private/VulkanContext.h
+++ b/Engine/Source/Runtime/VulkanRHI/Private/VulkanContext.h
@@ -60,6 +60,8 @@ public:
 	virtual void RHISetShaderParameter(FRHIComputeShader* ComputeShader, uint32 BufferIndex, uint32 BaseIndex, uint32 NumBytes, const void* NewValue) final override;
 	virtual void RHISetShaderParameters(FRHIGraphicsShader* Shader, TConstArrayView<uint8> InParametersData, TConstArrayView<FRHIShaderParameter> InParameters, TConstArrayView<FRHIShaderParameterResource> InResourceParameters, TConstArrayView<FRHIShaderParameterResource> InBindlessParameters) final override;
 	virtual void RHISetShaderParameters(FRHIComputeShader* Shader, TConstArrayView<uint8> InParametersData, TConstArrayView<FRHIShaderParameter> InParameters, TConstArrayView<FRHIShaderParameterResource> InResourceParameters, TConstArrayView<FRHIShaderParameterResource> InBindlessParameters) final override;
+	virtual void RHISetShaderUnbinds(FRHIComputeShader* Shader, TConstArrayView<FRHIShaderParameterUnbind> InUnbinds) final override;
+	virtual void RHISetShaderUnbinds(FRHIGraphicsShader* Shader, TConstArrayView<FRHIShaderParameterUnbind> InUnbinds) final override;
 	virtual void RHISetStencilRef(uint32 StencilRef) final override;
 	virtual void RHIDrawPrimitive(uint32 BaseVertexIndex, uint32 NumPrimitives, uint32 NumInstances) final override;
 	virtual void RHIDrawPrimitiveIndirect(FRHIBuffer* ArgumentBuffer, uint32 ArgumentOffset) final override;
diff --git a/Engine/Source/Runtime/Windows/D3D11RHI/Private/D3D11Commands.cpp b/Engine/Source/Runtime/Windows/D3D11RHI/Private/D3D11Commands.cpp
--- a/Engine/Source/Runtime/Windows/D3D11RHI/Private/D3D11Commands.cpp
+++ b/Engine/Source/Runtime/Windows/D3D11RHI/Private/D3D11Commands.cpp
@@ -631,6 +631,16 @@ void FD3D11DynamicRHI::RHISetShaderParameters(FRHIGraphicsShader* Shader, TConst
 	);
 }
 
+void FD3D11DynamicRHI::RHISetShaderUnbinds(FRHIComputeShader* Shader, TConstArrayView<FRHIShaderParameterUnbind> InUnbinds)
+{
+	UE::RHICore::RHISetShaderUnbindsShared(*this, Shader, InUnbinds);
+}
+
+void FD3D11DynamicRHI::RHISetShaderUnbinds(FRHIGraphicsShader* Shader, TConstArrayView<FRHIShaderParameterUnbind> InUnbinds)
+{
+	UE::RHICore::RHISetShaderUnbindsShared(*this, Shader, InUnbinds);
+}
+
 void FD3D11DynamicRHI::ValidateExclusiveDepthStencilAccess(FExclusiveDepthStencil RequestedAccess) const
 {
 	const bool bSrcDepthWrite = RequestedAccess.IsDepthWrite();
diff --git a/Engine/Source/Runtime/Windows/D3D11RHI/Private/D3D11RHIPrivate.h b/Engine/Source/Runtime/Windows/D3D11RHI/Private/D3D11RHIPrivate.h
--- a/Engine/Source/Runtime/Windows/D3D11RHI/Private/D3D11RHIPrivate.h
+++ b/Engine/Source/Runtime/Windows/D3D11RHI/Private/D3D11RHIPrivate.h
@@ -600,6 +600,8 @@ public:
 	virtual void RHISetShaderParameter(FRHIComputeShader* ComputeShader, uint32 BufferIndex, uint32 BaseIndex, uint32 NumBytes, const void* NewValue) final override;
 	virtual void RHISetShaderParameters(FRHIComputeShader* Shader, TConstArrayView<uint8> InParametersData, TConstArrayView<FRHIShaderParameter> InParameters, TConstArrayView<FRHIShaderParameterResource> InResourceParameters, TConstArrayView<FRHIShaderParameterResource> InBindlessParameters) final override;
 	virtual void RHISetShaderParameters(FRHIGraphicsShader* Shader, TConstArrayView<uint8> InParametersData, TConstArrayView<FRHIShaderParameter> InParameters, TConstArrayView<FRHIShaderParameterResource> InResourceParameters, TConstArrayView<FRHIShaderParameterResource> InBindlessParameters) final override;
+	virtual void RHISetShaderUnbinds(FRHIComputeShader* Shader, TConstArrayView<FRHIShaderParameterUnbind> InUnbinds) final override;
+	virtual void RHISetShaderUnbinds(FRHIGraphicsShader* Shader, TConstArrayView<FRHIShaderParameterUnbind> InUnbinds) final override;
 	virtual void RHISetDepthStencilState(FRHIDepthStencilState* NewState, uint32 StencilRef) final override;
 	virtual void RHISetStencilRef(uint32 StencilRef) final override;
 	virtual void RHISetBlendState(FRHIBlendState* NewState, const FLinearColor& BlendFactor) final override;
