Use a simple recursive mutex for render resource list to fix race condition when adding resources concurrently.

diff --git a/Engine/Source/Runtime/RenderCore/Private/RenderResource.cpp b/Engine/Source/Runtime/RenderCore/Private/RenderResource.cpp
--- a/Engine/Source/Runtime/RenderCore/Private/RenderResource.cpp
+++ b/Engine/Source/Runtime/RenderCore/Private/RenderResource.cpp
@@ -13,6 +13,7 @@
 #include "RenderGraphResources.h"
 #include "Containers/ResourceArray.h"
 #include "RenderCore.h"
+#include "Async/RecursiveMutex.h"
 
 /** Whether to enable mip-level fading or not: +1.0f if enabled, -1.0f if disabled. */
 float GEnableMipLevelFading = 1.0f;
@@ -36,93 +37,44 @@ public:
 		return Instance;
 	}
 
-	~FRenderResourceList()
-	{
-		for (FFreeList* FreeList : LocalFreeLists)
-		{
-			delete FreeList;
-		}
-		FPlatformTLS::FreeTlsSlot(TLSSlot);
-	}
-
 	int32 Allocate(FRenderResource* Resource)
 	{
-		if (bIsIterating)
+		Mutex.Lock();
+		int32 Index;
+		if (FreeIndexList.IsEmpty())
 		{
-			// This part is not thread safe. Iteration requires that no adds / removals are happening concurrently. The only
-			// supported case is recursive adds on the same thread (i.e. a parent resource initializes a child resource). In
-			// this case, we need to add the resource to the end so that it gets iterated as well.
-			check(IsInRenderingThread());
-			return ResourceList.AddElement(Resource);
+			Index = ResourceList.Num();
+			ResourceList.Emplace(Resource);
 		}
-
-		FFreeList& LocalFreeList = GetLocalFreeList();
-
-		if (LocalFreeList.IsEmpty())
+		else
 		{
-			FScopeLock Lock(&CS);
-
-			// Try to allocate free slots from the global free list.
-			const int32 OldFreeListSize = GlobalFreeList.Num();
-			const int32 NewFreeListSize = FMath::Max(OldFreeListSize - ChunkSize, 0);
-			int32 NumElements = OldFreeListSize - NewFreeListSize;
-
-			if (NumElements > 0)
-			{
-				LocalFreeList.Append(GlobalFreeList.GetData() + NewFreeListSize, NumElements);
-				GlobalFreeList.SetNum(NewFreeListSize, false);
-			}
-
-			// Allocate more if we didn't get a full chunk from the global list.
-			while (NumElements < ChunkSize)
-			{
-				LocalFreeList.Emplace(ResourceList.AddElement(nullptr));
-				NumElements++;
-			}
+			Index = FreeIndexList.Pop(false);
+			ResourceList[Index] = Resource;
 		}
-
-		int32 Index = LocalFreeList.Pop(false);
-		ResourceList[Index] = Resource;
+		Mutex.Unlock();
 		return Index;
 	}
 
 	void Deallocate(int32 Index)
 	{
-		GetLocalFreeList().Emplace(Index);
+		Mutex.Lock();
+		FreeIndexList.Emplace(Index);
 		ResourceList[Index] = nullptr;
+		Mutex.Unlock();
 	}
 
-	//////////////////////////////////////////////////////////////////////////////
-	// These methods must be called at sync points where allocations / deallocations can't occur from another thread.
-
 	void Clear()
 	{
-		check(IsInRenderingThread());
-
-		for (FFreeList* FreeList : LocalFreeLists)
-		{
-			FreeList->Empty();
-		}
+		Mutex.Lock();
+		FreeIndexList.Empty();
 		ResourceList.Empty();
-	}
-
-	void Coalesce()
-	{
-		check(IsInRenderingThread());
-
-		for (FFreeList* FreeList : LocalFreeLists)
-		{
-			GlobalFreeList.Append(*FreeList);
-			FreeList->Empty();
-		}
+		Mutex.Unlock();
 	}
 
 	template<typename FunctionType>
 	void ForEach(const FunctionType& Function)
 	{
-		check(IsInRenderingThread());
-		check(!bIsIterating);
-		bIsIterating = true;
+		Mutex.Lock();
 		for (int32 Index = 0; Index < ResourceList.Num(); ++Index)
 		{
 			FRenderResource* Resource = ResourceList[Index];
@@ -132,15 +84,13 @@ public:
 				Function(Resource);
 			}
 		}
-		bIsIterating = false;
+		Mutex.Unlock();
 	}
 
 	template<typename FunctionType>
 	void ForEachReverse(const FunctionType& Function)
 	{
-		check(IsInRenderingThread());
-		check(!bIsIterating);
-		bIsIterating = true;
+		Mutex.Lock();
 		for (int32 Index = ResourceList.Num() - 1; Index >= 0; --Index)
 		{
 			FRenderResource* Resource = ResourceList[Index];
@@ -150,46 +100,17 @@ public:
 				Function(Resource);
 			}
 		}
-		bIsIterating = false;
+		Mutex.Unlock();
 	}
 
-	//////////////////////////////////////////////////////////////////////////////
-
 private:
-	FRenderResourceList()
-	{
-		TLSSlot = FPlatformTLS::AllocTlsSlot();
-	}
-
-	const int32 ChunkSize = 1024;
-
-	using FFreeList = TArray<int32>;
-
-	FFreeList& GetLocalFreeList()
-	{
-		void* TLSValue = FPlatformTLS::GetTlsValue(TLSSlot);
-		if (TLSValue == nullptr)
-		{
-			FFreeList* TLSCache = new FFreeList();
-			FPlatformTLS::SetTlsValue(TLSSlot, (void*)(TLSCache));
-			FScopeLock S(&CS);
-			LocalFreeLists.Add(TLSCache);
-			return *TLSCache;
-		}
-		return *((FFreeList*)TLSValue);
-	}
-
-	uint32 TLSSlot;
-	FCriticalSection CS;
-	TArray<FFreeList*> LocalFreeLists;
-	FFreeList GlobalFreeList;
-	TChunkedArray<FRenderResource*> ResourceList;
-	bool bIsIterating = false;
+	UE::FRecursiveMutex Mutex;
+	TArray<int32> FreeIndexList;
+	TArray<FRenderResource*> ResourceList;
 };
 
 void FRenderResource::CoalesceResourceList()
 {
-	FRenderResourceList::Get().Coalesce();
 }
 
 void FRenderResource::ReleaseRHIForAllResources()
