Allows DumpGPU to captures GSystemTextures that have SkipTracking

diff --git a/Engine/Source/Runtime/RenderCore/Private/DumpGPU.cpp b/Engine/Source/Runtime/RenderCore/Private/DumpGPU.cpp
--- a/Engine/Source/Runtime/RenderCore/Private/DumpGPU.cpp
+++ b/Engine/Source/Runtime/RenderCore/Private/DumpGPU.cpp
@@ -897,6 +897,9 @@ public:
 
 			// If the texture is an array, use pixel shader to correctly copy the right slice given RHI doesn't support copy from slices
 			SubresourceDumpDesc.bPreprocessForStaging |= SubresourceDesc.Texture->Desc.IsTextureArray();
+
+			// Reads the texture using a shader if it has SkipTracking to avoid transitioning it in anyway from its SRV read state.
+			SubresourceDumpDesc.bPreprocessForStaging |= EnumHasAnyFlags(SubresourceDesc.Texture->Flags, ERDGTextureFlags::SkipTracking);
 		}
 
 		// Some RHIs (GL) only support 32Bit single channel images as CS output
@@ -1186,16 +1189,10 @@ public:
 			return;
 		}
 
-		// Verify the texture is able to do resource transitions.
-		if (EnumHasAnyFlags(SubresourceDesc.Texture->Flags, ERDGTextureFlags::SkipTracking))
-		{
-			UE_LOG(LogRendererCore, Warning, TEXT("Not dumping %s because has ERDGTextureFlags::SkipTracking."), SubresourceDesc.Texture->Name);
-			return;
-		}
-
 		const FRDGViewableResource::FAccessModeState AccessModeState = SubresourceDesc.Texture->AccessModeState;
 
-		if (AccessModeState.Mode == FRDGViewableResource::EAccessMode::External)
+		const bool bIsExternalMode = AccessModeState.Mode == FRDGViewableResource::EAccessMode::External && !EnumHasAnyFlags(SubresourceDesc.Texture->Flags, ERDGTextureFlags::SkipTracking);
+		if (bIsExternalMode)
 		{
 			GraphBuilder.UseInternalAccessMode(SubresourceDesc.Texture);
 		}
@@ -1246,7 +1243,7 @@ public:
 			ResourcesDumpPasses++;
 		}
 
-		if (AccessModeState.Mode == FRDGViewableResource::EAccessMode::External)
+		if (bIsExternalMode)
 		{
 			GraphBuilder.UseExternalAccessMode(SubresourceDesc.Texture, AccessModeState.Access, AccessModeState.Pipelines);
 		}
