Fixes missing RHIUnmapStagingSurface() call in DumpGPU caused by 25447955

diff --git a/Engine/Source/Runtime/RenderCore/Private/DumpGPU.cpp b/Engine/Source/Runtime/RenderCore/Private/DumpGPU.cpp
--- a/Engine/Source/Runtime/RenderCore/Private/DumpGPU.cpp
+++ b/Engine/Source/Runtime/RenderCore/Private/DumpGPU.cpp
@@ -1072,6 +1072,8 @@ public:
 				ColumnPitchInPixels,
 				/* out */ Array);
 
+			RHICmdList.UnmapStagingSurface(StagingTexture, GPUIndex);
+
 			DumpResourceBinaryToFile(Array.GetData(), Array.Num(), DumpFilePath);
 		}
 		else
