Implements DumpGPU streaming to better diagnose temporal problems on isolated feature like TSR without tanking frame rate

r.DumpGPU.Stream=1 instead allocate and reuse staging resources from its own pool, only issuing copy to staging and gpu fence rhi commands.
Then every frame it polls whether the resource staging is lockable with the GPU fence on render thread.
When ready, the render thread kicks of a background task that cakes care of resource CPU post processing and disk writing.
Once disk write is complete, the render thread polls everyframe a FEvent to know when complete and unlock the staging resource.
The staging resource is then release to staging resource pool ready to be reused for dumping another resource of the current frame.

diff --git a/Engine/Source/Runtime/RenderCore/Private/DumpGPU.cpp b/Engine/Source/Runtime/RenderCore/Private/DumpGPU.cpp
--- a/Engine/Source/Runtime/RenderCore/Private/DumpGPU.cpp
+++ b/Engine/Source/Runtime/RenderCore/Private/DumpGPU.cpp
@@ -4,6 +4,7 @@
 #include "HAL/PlatformFileManager.h"
 #include "HAL/PlatformMisc.h"
 #include "HAL/PlatformOutputDevices.h"
+#include "Async/AsyncWork.h"
 #include "Misc/App.h"
 #include "Misc/FileHelper.h"
 #include "Misc/WildcardString.h"
@@ -101,6 +102,15 @@ static TAutoConsoleVariable<float> GDumpGPUFixedTickRate(
 	TEXT("Override the engine's tick rate to be fixed for every dumped frames (default=0)."),
 	ECVF_Default);
 
+static TAutoConsoleVariable<int32> GDumpGPUStream(
+	TEXT("r.DumpGPU.Stream"), 0,
+	TEXT("Asynchronously readback from GPU to disk.\n")
+	TEXT(" 0: Synchronously copy from GPU to disk with extra carefulness to avoid OOM (default);\n")
+	TEXT(" 1: Asynchronously copy from GPU to disk with dedicated staging resources pool. May run OOM. ")
+	TEXT("Please consider using r.DumpGPU.Root to minimise amount of passes to stream and r.Test.SecondaryUpscaleOverride ")
+	TEXT("to reduce resource size to minimise OOM and disk bandwidth bottleneck per frame."),
+	ECVF_Default);
+
 static TAutoConsoleVariable<int32> GDumpGPUDraws(
 	TEXT("r.DumpGPU.Draws"), 0,
 	TEXT("Whether to dump resource after each individual draw call (disabled by default)."),
@@ -201,6 +211,9 @@ BEGIN_SHADER_PARAMETER_STRUCT(FDumpBufferPass, )
 	RDG_BUFFER_ACCESS(Buffer, ERHIAccess::CopySrc)
 END_SHADER_PARAMETER_STRUCT()
 
+const FName GDumpGPUTextureFenceName(TEXT("DumpGPU.TextureFence"));
+const FName GDumpGPUBufferFenceName(TEXT("DumpGPU.BufferFence"));
+
 }
 
 class FRDGResourceDumpContext
@@ -214,6 +227,8 @@ public:
 
 	bool bEnableDiskWrite = false;
 	bool bUpload = false;
+	bool bStream = false;
+	float DeltaTime = 0.0f;
 	int32 RemainingFrameCount = 1;
 	FName UploadResourceCompressionName;
 	FString DumpingDirectoryPath;
@@ -233,10 +248,12 @@ public:
 	const FRDGPass* DrawDumpingPass = nullptr;
 	int32 DrawDumpCount = 0;
 
-	int32 MetadataFilesOpened = 0;
-	int64 MetadataFilesWriteBytes = 0;
-	int32 ResourceBinaryFilesOpened = 0;
-	int64 ResourceBinaryWriteBytes = 0;
+	TAtomic<int32> MetadataFilesOpened = 0;
+	TAtomic<int64> MetadataFilesWriteBytes = 0;
+	TAtomic<int32> ParametersFilesOpened = 0;
+	TAtomic<int64> ParametersFilesWriteBytes = 0;
+	TAtomic<int32> ResourceBinaryFilesOpened = 0;
+	TAtomic<int64> ResourceBinaryWriteBytes = 0;
 
 	enum class ETimingBucket : uint8
 	{
@@ -245,10 +262,11 @@ public:
 		GPUWait,
 		CPUPostProcessing,
 		MetadataFileWrite,
+		ParametersFileWrite,
 		ResourceBinaryFileWrite,
 		MAX
 	};
-	mutable double TimingBucket[int32(ETimingBucket::MAX)];
+	mutable TAtomic<double> TimingBucket[int32(ETimingBucket::MAX)];
 
 	class FTimeBucketMeasure
 	{
@@ -262,7 +280,13 @@ public:
 		~FTimeBucketMeasure()
 		{
 			uint64 End = FPlatformTime::Cycles64();
-			DumpCtx->TimingBucket[int32(Bucket)] += FPlatformTime::ToSeconds64(FMath::Max(End - Start, uint64(0)));
+
+			double Add = FPlatformTime::ToSeconds64(FMath::Max(End - Start, uint64(0)));
+			double Old;
+			do
+			{
+				Old = DumpCtx->TimingBucket[int32(Bucket)].Load();
+			} while (!DumpCtx->TimingBucket[int32(Bucket)].CompareExchange(Old, Old + Add));
 		}
 
 	protected:
@@ -284,6 +308,11 @@ public:
 				DumpCtx->MetadataFilesOpened += FilesOpened;
 				DumpCtx->MetadataFilesWriteBytes += WriteSizeBytes;
 			}
+			else if (InBucket == ETimingBucket::ParametersFileWrite)
+			{
+				DumpCtx->ParametersFilesOpened += FilesOpened;
+				DumpCtx->ParametersFilesWriteBytes += WriteSizeBytes;
+			}
 			else if (InBucket == ETimingBucket::ResourceBinaryFileWrite)
 			{
 				DumpCtx->ResourceBinaryFilesOpened += FilesOpened;
@@ -306,6 +335,40 @@ public:
 		}
 	}
 
+	using FTaskCallback = TFunction<void()>;
+
+	class FDumpGPUTask : public FNonAbandonableTask
+	{
+	public:
+		FTaskCallback TaskCallback;
+
+		FDumpGPUTask(FTaskCallback&& InTaskCallback)
+			: TaskCallback(InTaskCallback)
+		{ }
+
+		void DoWork()
+		{
+			TaskCallback();
+		}
+
+		FORCEINLINE TStatId GetStatId() const
+		{
+			RETURN_QUICK_DECLARE_CYCLE_STAT(FDumpGPUTask, STATGROUP_ThreadPoolAsyncTasks);
+		}
+	};
+
+	void KickOffAsyncTask(FTaskCallback&& Callback)
+	{
+		if (bStream)
+		{
+			(new FAutoDeleteAsyncTask<FDumpGPUTask>(MoveTemp(Callback)))->StartBackgroundTask();
+		}
+		else
+		{
+			Callback();
+		}
+	}
+
 	IDumpGPUUploadServiceProvider::FDumpParameters GetDumpParameters() const
 	{
 		IDumpGPUUploadServiceProvider::FDumpParameters DumpServiceParameters;
@@ -339,6 +402,7 @@ public:
 
 	bool DumpStatusToFile(FStringView StatusString)
 	{
+		UE_LOG(LogRendererCore, Display, TEXT("DumpGPU status = %s"), StatusString.GetData());
 		return DumpStringToFile(StatusString, FString(FRDGResourceDumpContext::kBaseDir) / TEXT("Status.txt"));
 	}
 
@@ -351,7 +415,7 @@ public:
 		return DumpStringToFile(OutputString, FileName, WriteFlags);
 	}
 
-	bool DumpBinaryToFile(TArrayView<const uint8> ArrayView, const FString& FileName)
+	bool DumpBinaryToFile(TArrayView<const uint8> ArrayView, const FString& FileName, ETimingBucket Bucket)
 	{
 		// Make it has if the write happened and was successful.
 		if (!bEnableDiskWrite)
@@ -360,11 +424,11 @@ public:
 		}
 
 		FString FullPath = GetDumpFullPath(FileName);
-		FFileWriteCtx WriteCtx(this, ETimingBucket::ResourceBinaryFileWrite, ArrayView.Num());
+		FFileWriteCtx WriteCtx(this, Bucket, ArrayView.Num());
 		return FFileHelper::SaveArrayToFile(ArrayView, *FullPath);
 	}
 
-	bool DumpBinaryToFile(const uint8* Data, int64 DataByteSize, const FString& FileName)
+	bool DumpBinaryToFile(const uint8* Data, int64 DataByteSize, const FString& FileName, ETimingBucket Bucket)
 	{
 		// Make it has if the write happened and was successful.
 		if (!bEnableDiskWrite)
@@ -373,7 +437,7 @@ public:
 		}
 
 		FString FullPath = GetDumpFullPath(FileName);
-		FFileWriteCtx WriteCtx(this, ETimingBucket::ResourceBinaryFileWrite, DataByteSize);
+		FFileWriteCtx WriteCtx(this, Bucket, DataByteSize);
 
 		TUniquePtr<FArchive> Ar = TUniquePtr<FArchive>(IFileManager::Get().CreateFileWriter(*FullPath, /* WriteFlags = */ 0));
 		if (!Ar)
@@ -390,7 +454,7 @@ public:
 
 	bool DumpResourceBinaryToFile(const uint8* UncompressedData, int64 UncompressedSize, const FString& FileName)
 	{
-		return DumpBinaryToFile(UncompressedData, UncompressedSize, FileName);
+		return DumpBinaryToFile(UncompressedData, UncompressedSize, FileName, ETimingBucket::ResourceBinaryFileWrite);
 	}
 
 	bool IsUnsafeToDumpResource(SIZE_T ResourceByteSize, float DumpMemoryMultiplier) const
@@ -594,7 +658,7 @@ public:
 				Ar->Serialize((ANSICHAR*)Src.Get(), Src.Length() * sizeof(ANSICHAR));
 			}
 
-			if (GDumpRenderingConsoleVariablesCVar.GetValueOnGameThread() != 0)
+			if (GDumpRenderingConsoleVariablesCVar.GetValueOnAnyThread() != 0)
 			{
 				IConsoleManager::Get().ForEachConsoleObjectThatStartsWith(FConsoleObjectVisitor::CreateLambda(OnConsoleVariable), TEXT(""));
 			}
@@ -968,6 +1032,153 @@ public:
 		}
 	}
 
+	enum class FStagingResourceStatus
+	{
+		Unused,
+		RenderingIntermediaryOnly,
+		Rendering,
+		WritingToDisk,
+		WritingToDiskComplete,
+	};
+
+	struct FStagingPoolEntry
+	{
+		FStagingResourceStatus Status = FStagingResourceStatus::Unused;
+
+		FString DumpFilePath;
+		int32 GPUIndex = 0;
+		FGPUFenceRHIRef WriteToResourceComplete;
+		FEvent* WriteToDiskComplete = nullptr;
+	};
+
+	struct FStagingTexturePoolEntry : public FStagingPoolEntry
+	{
+		FTextureSubresourceDumpDesc SubresourceDumpDesc;
+		FRHITextureCreateDesc Desc;
+		FTextureRHIRef Texture;
+		FUnorderedAccessViewRHIRef UAV;
+	};
+
+	struct FStagingBufferPoolEntry : public FStagingPoolEntry
+	{
+		int32 ByteSize = 0;
+		FStagingBufferRHIRef StagingBuffer;
+	};
+
+	TArray<TUniquePtr<FStagingTexturePoolEntry>> StagingTexturePool;
+	TArray<TUniquePtr<FStagingBufferPoolEntry>> StagingBufferPool;
+
+	void CreateStagingTexture(const FRHITextureCreateDesc& Desc, FTextureRHIRef* Texture)
+	{
+		check(IsInRenderingThread());
+		(*Texture) = RHICreateTexture(Desc);
+		(*Texture)->DisableLifetimeExtension();
+	}
+
+	FStagingTexturePoolEntry& ReuseStagingTexture(
+		FRHICommandListImmediate& RHICmdList,
+		const FRHITextureCreateDesc& Desc,
+		ERHIAccess Access,
+		FTextureRHIRef* Texture,
+		FUnorderedAccessViewRHIRef* UAV)
+	{
+		check(IsInRenderingThread());
+		check(bStream);
+		check(Access == ERHIAccess::CopyDest || Access == ERHIAccess::UAVCompute);
+
+		for (TUniquePtr<FStagingTexturePoolEntry>& Entry : StagingTexturePool)
+		{
+			if (Entry->Desc != Desc)
+			{
+				continue;
+			}
+
+			if (Access == ERHIAccess::CopyDest)
+			{
+				if (Entry->Status == FStagingResourceStatus::Unused)
+				{
+					check(Entry->Texture.GetRefCount() == 1);
+					*Texture = Entry->Texture;
+					check(Entry->Texture.GetRefCount() == 2);
+					Entry->Status = FStagingResourceStatus::Rendering;
+					Entry->WriteToResourceComplete->Clear();
+					return *Entry;
+				}
+			}
+			else
+			{
+				if (Entry->Status == FStagingResourceStatus::RenderingIntermediaryOnly && Entry->Texture.GetRefCount() == 1)
+				{
+					*Texture = Entry->Texture;
+					check(Entry->Texture.GetRefCount() == 2);
+					check(!Entry->WriteToResourceComplete);
+					return *Entry;
+				}
+			}
+		}
+
+		CreateStagingTexture(Desc, /* out */ Texture);
+
+		FStagingTexturePoolEntry* NewEntry = new FStagingTexturePoolEntry;
+		NewEntry->Status = Access == ERHIAccess::CopyDest ? FStagingResourceStatus::Rendering : FStagingResourceStatus::RenderingIntermediaryOnly;
+		NewEntry->Desc = Desc;
+		NewEntry->Texture = *Texture;
+
+		StagingTexturePool.Add(TUniquePtr<FStagingTexturePoolEntry>(NewEntry));
+
+		RHICmdList.Transition(FRHITransitionInfo(*Texture, ERHIAccess::Unknown, Access));
+
+		if (UAV)
+		{
+			check(Access == ERHIAccess::UAVCompute);
+			NewEntry->UAV = RHICreateUnorderedAccessView(*Texture, /* MipLevel = */ 0);
+			*UAV = NewEntry->UAV;
+		}
+		else
+		{
+			check(Access == ERHIAccess::CopyDest);
+			NewEntry->WriteToResourceComplete = RHICreateGPUFence(GDumpGPUTextureFenceName);
+			NewEntry->WriteToResourceComplete->Clear();
+		}
+
+		return *NewEntry;
+	}
+
+	FStagingBufferPoolEntry& ReuseStagingBuffer(
+		FRHICommandListImmediate& RHICmdList,
+		int32 ByteSize)
+	{
+		check(IsInRenderingThread());
+		check(bStream);
+
+		for (TUniquePtr<FStagingBufferPoolEntry>& Entry : StagingBufferPool)
+		{
+			if (Entry->ByteSize != ByteSize)
+			{
+				continue;
+			}
+
+			if (Entry->Status == FStagingResourceStatus::Unused)
+			{
+				Entry->Status = FStagingResourceStatus::Rendering;
+				Entry->WriteToResourceComplete->Clear();
+				return *Entry;
+			}
+		}
+
+		FStagingBufferPoolEntry* NewEntry = new FStagingBufferPoolEntry;
+		NewEntry->Status = FStagingResourceStatus::Rendering;
+		NewEntry->ByteSize = ByteSize;
+		NewEntry->StagingBuffer = RHICreateStagingBuffer();
+		NewEntry->StagingBuffer->DisableLifetimeExtension();
+		NewEntry->WriteToResourceComplete = RHICreateGPUFence(GDumpGPUBufferFenceName);
+		NewEntry->WriteToResourceComplete->Clear();
+
+		StagingBufferPool.Add(TUniquePtr<FStagingBufferPoolEntry>(NewEntry));
+
+		return *NewEntry;
+	}
+
 	void DumpTextureSubResource(
 		FRHICommandListImmediate& RHICmdList,
 		const TCHAR* TextureDebugName,
@@ -979,25 +1190,34 @@ public:
 		check(IsInRenderingThread());
 
 		// Preprocess
-		FTextureRHIRef StagingSrcTexture;
+		FTextureRHIRef StagingSrcTextureRef;
+		FRHITexture* StagingSrcTexture;
 		FTextureRHIRef StagingTexture;
+		FStagingTexturePoolEntry* StagingTexturePoolEntry = nullptr;
 
 		{
 			FTimeBucketMeasure TimeBucketMeasure(this, ETimingBucket::RHIReadbackCommands);
 			if (SubresourceDumpDesc.bPreprocessForStaging)
 			{
+				FUnorderedAccessViewRHIRef StagingOutput;
 				{
 					const FRHITextureCreateDesc Desc =
 						FRHITextureCreateDesc::Create2D(TEXT("DumpGPU.PreprocessTexture"), SubresourceDumpDesc.SubResourceExtent, SubresourceDumpDesc.PreprocessedPixelFormat)
 						.SetFlags(ETextureCreateFlags::UAV | ETextureCreateFlags::ShaderResource | ETextureCreateFlags::HideInVisualizeTexture);
 
-					StagingSrcTexture = RHICreateTexture(Desc);
-					StagingSrcTexture->DisableLifetimeExtension();
-				}
-
-				FUnorderedAccessViewRHIRef StagingOutput = RHICreateUnorderedAccessView(StagingSrcTexture, /* MipLevel = */ 0);
+					if (bStream)
+					{
+						ReuseStagingTexture(RHICmdList, Desc, ERHIAccess::UAVCompute, /* out */ &StagingSrcTextureRef, /* out */ &StagingOutput);
+					}
+					else
+					{
+						CreateStagingTexture(Desc, /* out */ &StagingSrcTextureRef);
 
-				RHICmdList.Transition(FRHITransitionInfo(StagingSrcTexture, ERHIAccess::Unknown, ERHIAccess::UAVCompute));
+						RHICmdList.Transition(FRHITransitionInfo(StagingSrcTextureRef, ERHIAccess::Unknown, ERHIAccess::UAVCompute));
+						StagingOutput = RHICreateUnorderedAccessView(StagingSrcTextureRef, /* MipLevel = */ 0);
+					}
+					StagingSrcTexture = StagingSrcTextureRef;
+				}
 
 				FDumpTextureCS::FPermutationDomain PermutationVector;
 				PermutationVector.Set<FDumpTextureCS::FTextureTypeDim>(SubresourceDumpDesc.DumpTextureType);
@@ -1025,34 +1245,29 @@ public:
 					FRHITextureCreateDesc::Create2D(TEXT("DumpGPU.StagingTexture"), SubresourceDumpDesc.SubResourceExtent, SubresourceDumpDesc.PreprocessedPixelFormat)
 					.SetFlags(ETextureCreateFlags::CPUReadback | ETextureCreateFlags::HideInVisualizeTexture);
 
-				StagingTexture = RHICreateTexture(Desc);
-				StagingTexture->DisableLifetimeExtension();
-
-				RHICmdList.Transition(FRHITransitionInfo(StagingTexture, ERHIAccess::Unknown, ERHIAccess::CopyDest));
+				if (bStream)
+				{
+					StagingTexturePoolEntry = &ReuseStagingTexture(RHICmdList, Desc, ERHIAccess::CopyDest, /* out */ &StagingTexture, /* UAV = */ nullptr);
+				}
+				else
+				{
+					CreateStagingTexture(Desc, /* out */ &StagingTexture);
+					RHICmdList.Transition(FRHITransitionInfo(StagingSrcTexture, ERHIAccess::Unknown, ERHIAccess::CopyDest));
+				}
 
 				// Transfer memory GPU -> CPU
 				RHICmdList.CopyTexture(StagingSrcTexture, StagingTexture, {});
 
 				RHICmdList.Transition(FRHITransitionInfo(StagingTexture, ERHIAccess::CopyDest, ERHIAccess::CPURead));
 			}
-		}
-
-		// Submit to GPU and wait for completion.
-		static const FName FenceName(TEXT("DumpGPU.TextureFence"));
-		FGPUFenceRHIRef Fence = RHICreateGPUFence(FenceName);
-		{
-			FTimeBucketMeasure TimeBucketMeasure(this, ETimingBucket::GPUWait);
 
-			Fence->Clear();
-			RHICmdList.WriteGPUFence(Fence);
-			RHICmdList.SubmitCommandsAndFlushGPU();
-			RHICmdList.BlockUntilGPUIdle();
+			// Transition back the intermediary texture to UAVCompute for next ReuseStagingTexture().
+			if (StagingSrcTextureRef && bStream)
+			{
+				RHICmdList.Transition(FRHITransitionInfo(StagingSrcTextureRef, ERHIAccess::CopySrc, ERHIAccess::UAVCompute));
+			}
 		}
 
-		void* Content = nullptr;
-		int32 RowPitchInPixels = 0;
-		int32 ColumnPitchInPixels = 0;
-
 		// jhoerner_todo 12/9/2021:  pick arbitrary GPU out of mask to avoid assert.  Eventually want to dump results for all GPUs, but
 		// I need to understand how to modify the dumping logic, and this works for now (usually when debugging, the bugs happen on
 		// secondary GPUs, so I figure the last index is most useful if we need to pick one).  I also would like the dump to include
@@ -1060,25 +1275,52 @@ public:
 		// modified by the graph (especially useful for MGPU, where we are concerned about cross-view or cross-frame state).
 		uint32 GPUIndex = RHICmdList.GetGPUMask().GetLastIndex();
 
-		RHICmdList.MapStagingSurface(StagingTexture, Fence.GetReference(), Content, RowPitchInPixels, ColumnPitchInPixels, GPUIndex);
-
-		if (Content)
+		if (bStream)
 		{
-			TArray64<uint8> Array;
-			PostProcessTexture(
-				SubresourceDumpDesc,
-				Content,
-				RowPitchInPixels,
-				ColumnPitchInPixels,
-				/* out */ Array);
-
-			RHICmdList.UnmapStagingSurface(StagingTexture, GPUIndex);
-
-			DumpResourceBinaryToFile(Array.GetData(), Array.Num(), DumpFilePath);
+			check(StagingTexturePoolEntry);
+			// Write a fence to know when to transition to 
+			RHICmdList.WriteGPUFence(StagingTexturePoolEntry->WriteToResourceComplete);
+			StagingTexturePoolEntry->SubresourceDumpDesc = SubresourceDumpDesc;
+			StagingTexturePoolEntry->DumpFilePath = DumpFilePath;
+			StagingTexturePoolEntry->GPUIndex = GPUIndex;
 		}
 		else
 		{
-			UE_LOG(LogRendererCore, Warning, TEXT("RHICmdList.MapStagingSurface() to dump texture %s failed."), TextureDebugName);
+			// Submit to GPU and wait for completion.
+			FGPUFenceRHIRef Fence = RHICreateGPUFence(GDumpGPUTextureFenceName);
+			{
+				FTimeBucketMeasure TimeBucketMeasure(this, ETimingBucket::GPUWait);
+
+				Fence->Clear();
+				RHICmdList.WriteGPUFence(Fence);
+				RHICmdList.SubmitCommandsAndFlushGPU();
+				RHICmdList.BlockUntilGPUIdle();
+			}
+
+			void* Content = nullptr;
+			int32 RowPitchInPixels = 0;
+			int32 ColumnPitchInPixels = 0;
+
+			RHICmdList.MapStagingSurface(StagingTexture, Fence.GetReference(), Content, RowPitchInPixels, ColumnPitchInPixels, GPUIndex);
+
+			if (Content)
+			{
+				TArray64<uint8> Array;
+				PostProcessTexture(
+					SubresourceDumpDesc,
+					Content,
+					RowPitchInPixels,
+					ColumnPitchInPixels,
+					/* out */ Array);
+
+				RHICmdList.UnmapStagingSurface(StagingTexture, GPUIndex);
+
+				DumpResourceBinaryToFile(Array.GetData(), Array.Num(), DumpFilePath);
+			}
+			else
+			{
+				UE_LOG(LogRendererCore, Warning, TEXT("RHICmdList.MapStagingSurface() to dump texture %s failed."), TextureDebugName);
+			}
 		}
 	}
 
@@ -1137,7 +1379,10 @@ public:
 		}
 
 		SubResourceSRV = nullptr;
-		ReleaseRHIResources(RHICmdListImmediate);
+		if (!bStream)
+		{
+			ReleaseRHIResources(RHICmdListImmediate);
+		}
 	}
 
 	void AddDumpTextureSubResourcePass(
@@ -1238,7 +1483,10 @@ public:
 					SubresourceDumpDesc.bPreprocessForStaging ? PassParameters->Texture->GetRHI() : nullptr,
 					SubresourceDumpDesc,
 					DumpFilePath);
-				this->ReleaseRHIResources(RHICmdList);
+				if (!bStream)
+				{
+					this->ReleaseRHIResources(RHICmdList);
+				}
 				this->UpdatePassProgress();
 			});
 
@@ -1466,7 +1714,7 @@ public:
 		// Dump the resource's binary to a .bin file.
 		if (bDumpResourceBinary && DumpBufferMode == 2)
 		{
-			int32 StagingResourceByteSize = FMath::Min(ByteSize, GDumpMaxStagingSize.GetValueOnRenderThread() * 1024 * 1024);
+			const int32 StagingResourceByteSize = bStream ? ByteSize : FMath::Min(ByteSize, GDumpMaxStagingSize.GetValueOnRenderThread() * 1024 * 1024);
 
 			FString DumpFilePath = kResourcesDir / FString::Printf(TEXT("%s.v%016x.bin"), *UniqueResourceName, PtrToUint(bIsOutputResource ? Pass : nullptr));
 
@@ -1500,75 +1748,91 @@ public:
 				[this, DumpFilePath, Buffer, ByteSize, StagingResourceByteSize](FRHICommandListImmediate& RHICmdList)
 			{
 				check(IsInRenderingThread());
-
-				FStagingBufferRHIRef StagingBuffer;
-				FGPUFenceRHIRef Fence;
-				{
-					FTimeBucketMeasure TimeBucketMeasure(this, ETimingBucket::RHIReadbackCommands);
-					StagingBuffer = RHICreateStagingBuffer();
-					StagingBuffer->DisableLifetimeExtension();
-
-					static const FName FenceName(TEXT("DumpGPU.BufferFence"));
-					Fence = RHICreateGPUFence(FenceName);
-				}
-				
-				TUniquePtr<FArchive> Ar;
-				if (bEnableDiskWrite)
+				if (bStream)
 				{
-					FString FullPath = GetDumpFullPath(DumpFilePath);
-					FFileWriteCtx WriteCtx(this, ETimingBucket::ResourceBinaryFileWrite, /* ByteSize = */ 0);
-
-					Ar = TUniquePtr<FArchive>(IFileManager::Get().CreateFileWriter(*FullPath, /* WriteFlags = */ 0));
-				}
+					check(ByteSize == StagingResourceByteSize);
 
-				for (int32 Offset = 0; Offset < ByteSize; Offset += StagingResourceByteSize)
-				{
-					int32 CopyByteSize = FMath::Min(StagingResourceByteSize, ByteSize - Offset);
+					FStagingBufferPoolEntry& StagingBufferPoolEntry = ReuseStagingBuffer(RHICmdList, StagingResourceByteSize);
+					StagingBufferPoolEntry.DumpFilePath = DumpFilePath;
 
 					// Transfer memory GPU -> CPU
 					{
 						FTimeBucketMeasure TimeBucketMeasure(this, ETimingBucket::RHIReadbackCommands);
-						RHICmdList.CopyToStagingBuffer(Buffer->GetRHI(), StagingBuffer, Offset, CopyByteSize);
+						RHICmdList.CopyToStagingBuffer(Buffer->GetRHI(), StagingBufferPoolEntry.StagingBuffer, /* Offset = */ 0, ByteSize);
+						RHICmdList.WriteGPUFence(StagingBufferPoolEntry.WriteToResourceComplete);
 					}
+				}
+				else // if (!bStream)
+				{
+					FStagingBufferRHIRef StagingBuffer;
+					FGPUFenceRHIRef Fence;
+					{
+						FTimeBucketMeasure TimeBucketMeasure(this, ETimingBucket::RHIReadbackCommands);
+						StagingBuffer = RHICreateStagingBuffer();
+						StagingBuffer->DisableLifetimeExtension();
 
-					// Submit to GPU and wait for completion.
+						Fence = RHICreateGPUFence(GDumpGPUBufferFenceName);
+					}
+				
+					TUniquePtr<FArchive> Ar;
+					if (bEnableDiskWrite)
 					{
-						FTimeBucketMeasure TimeBucketMeasure(this, ETimingBucket::GPUWait);
-						Fence->Clear();
-						RHICmdList.WriteGPUFence(Fence);
-						RHICmdList.SubmitCommandsAndFlushGPU();
-						RHICmdList.BlockUntilGPUIdle();
+						FString FullPath = GetDumpFullPath(DumpFilePath);
+						FFileWriteCtx WriteCtx(this, ETimingBucket::ResourceBinaryFileWrite, /* ByteSize = */ 0);
+
+						Ar = TUniquePtr<FArchive>(IFileManager::Get().CreateFileWriter(*FullPath, /* WriteFlags = */ 0));
 					}
 
-					void* Content = RHICmdList.LockStagingBuffer(StagingBuffer, Fence.GetReference(), 0, CopyByteSize);
-					if (Content)
+					for (int32 Offset = 0; Offset < ByteSize; Offset += StagingResourceByteSize)
 					{
-						if (Ar)
+						int32 CopyByteSize = FMath::Min(StagingResourceByteSize, ByteSize - Offset);
+
+						// Transfer memory GPU -> CPU
+						{
+							FTimeBucketMeasure TimeBucketMeasure(this, ETimingBucket::RHIReadbackCommands);
+							RHICmdList.CopyToStagingBuffer(Buffer->GetRHI(), StagingBuffer, Offset, CopyByteSize);
+						}
+
+						// Submit to GPU and wait for completion.
+						{
+							FTimeBucketMeasure TimeBucketMeasure(this, ETimingBucket::GPUWait);
+							Fence->Clear();
+							RHICmdList.WriteGPUFence(Fence);
+							RHICmdList.SubmitCommandsAndFlushGPU();
+							RHICmdList.BlockUntilGPUIdle();
+						}
+
+						void* Content = RHICmdList.LockStagingBuffer(StagingBuffer, Fence.GetReference(), 0, CopyByteSize);
+						if (Content)
+						{
+							if (Ar)
+							{
+								FFileWriteCtx WriteCtx(this, ETimingBucket::ResourceBinaryFileWrite, CopyByteSize, /* OpenedFiles = */ 0);
+								Ar->Serialize((void*)Content, CopyByteSize);
+								Ar->Flush();
+							}
+							RHICmdList.UnlockStagingBuffer(StagingBuffer);
+						}
+						else
 						{
-							FFileWriteCtx WriteCtx(this, ETimingBucket::ResourceBinaryFileWrite, CopyByteSize, /* OpenedFiles = */ 0);
-							Ar->Serialize((void*)Content, CopyByteSize);
-							Ar->Flush();
+							UE_LOG(LogRendererCore, Warning, TEXT("RHICmdList.LockStagingBuffer() to dump buffer %s failed."), Buffer->Name);
+							break;
 						}
-						RHICmdList.UnlockStagingBuffer(StagingBuffer);
 					}
-					else
+
+					if (Ar)
 					{
-						UE_LOG(LogRendererCore, Warning, TEXT("RHICmdList.LockStagingBuffer() to dump buffer %s failed."), Buffer->Name);
-						break;
+						FFileWriteCtx WriteCtx(this, ETimingBucket::ResourceBinaryFileWrite, /* ByteSize = */ 0, /* OpenedFiles = */ 0);
+						// Always explicitly close to catch errors from flush/close
+						Ar->Close();
+						Ar = nullptr;
 					}
-				}
 
-				if (Ar)
-				{
-					FFileWriteCtx WriteCtx(this, ETimingBucket::ResourceBinaryFileWrite, /* ByteSize = */ 0, /* OpenedFiles = */ 0);
-					// Always explicitly close to catch errors from flush/close
-					Ar->Close();
-					Ar = nullptr;
-				}
+					StagingBuffer = nullptr;
+					Fence = nullptr;
+					this->ReleaseRHIResources(RHICmdList);
+				} // if (!bStream)
 
-				StagingBuffer = nullptr;
-				Fence = nullptr;
-				this->ReleaseRHIResources(RHICmdList);
 				this->UpdatePassProgress();
 			});
 
@@ -1581,6 +1845,159 @@ public:
 		}
 	}
 
+	static void WaitDiskWriteAndReset(FStagingPoolEntry* Entry)
+	{
+		check(Entry->Status == FStagingResourceStatus::WritingToDisk || Entry->Status == FStagingResourceStatus::WritingToDiskComplete);
+		check(Entry->WriteToDiskComplete);
+
+		Entry->WriteToDiskComplete->Wait();
+		FPlatformProcess::ReturnSynchEventToPool(Entry->WriteToDiskComplete);
+		Entry->WriteToDiskComplete = nullptr;
+
+		check(Entry->Status == FStagingResourceStatus::WritingToDiskComplete);
+		Entry->Status = FStagingResourceStatus::Unused;
+	}
+
+	void LandCompletedResources(FRHICommandListImmediate& RHICmdList)
+	{
+		check(IsInRenderingThread());
+		check(bStream);
+		for (TUniquePtr<FStagingTexturePoolEntry>& Entry : StagingTexturePool)
+		{
+			if (Entry->Status == FStagingResourceStatus::Rendering)
+			{
+				// Try lock the staging surface.
+
+				void* Content = nullptr;
+				int32 RowPitchInPixels = 0;
+				int32 ColumnPitchInPixels = 0;
+
+				RHICmdList.MapStagingSurface(Entry->Texture, Entry->WriteToResourceComplete.GetReference(), Content, RowPitchInPixels, ColumnPitchInPixels, Entry->GPUIndex);
+
+				if (!Content)
+				{
+					// NOP
+				}
+				else
+				{
+					check(Entry->WriteToDiskComplete == nullptr);
+					Entry->WriteToDiskComplete = FPlatformProcess::GetSynchEventFromPool(true);
+					Entry->Status = FStagingResourceStatus::WritingToDisk;
+
+					FStagingTexturePoolEntry* Staging = Entry.Get();
+
+					KickOffAsyncTask([this, Staging, Content, RowPitchInPixels, ColumnPitchInPixels](){
+						check(Staging->WriteToDiskComplete);
+						check(Staging->Status == FStagingResourceStatus::WritingToDisk);
+
+						TArray64<uint8> Array;
+						this->PostProcessTexture(
+							Staging->SubresourceDumpDesc,
+							Content,
+							RowPitchInPixels,
+							ColumnPitchInPixels,
+							/* out */ Array);
+
+						this->DumpResourceBinaryToFile(Array.GetData(), Array.Num(), Staging->DumpFilePath);
+
+						Staging->Status = FStagingResourceStatus::WritingToDiskComplete;
+						Staging->WriteToDiskComplete->Trigger();
+					});
+				}
+			}
+			else if (Entry->Status == FStagingResourceStatus::WritingToDiskComplete)
+			{
+				WaitDiskWriteAndReset(Entry.Get());
+
+				RHICmdList.UnmapStagingSurface(Entry->Texture, Entry->GPUIndex);
+				RHICmdList.Transition(FRHITransitionInfo(Entry->Texture, ERHIAccess::CPURead, ERHIAccess::CopyDest));
+			}
+		}
+
+		for (TUniquePtr<FStagingBufferPoolEntry>& Entry : StagingBufferPool)
+		{
+			if (Entry->Status == FStagingResourceStatus::Rendering)
+			{
+				// Try lock the staging buffer.
+				void* Content = RHICmdList.LockStagingBuffer(Entry->StagingBuffer, Entry->WriteToResourceComplete.GetReference(), 0, Entry->ByteSize);
+				if (!Content)
+				{
+					// NOP
+				}
+				else
+				{
+					check(Entry->WriteToDiskComplete == nullptr);
+					Entry->WriteToDiskComplete = FPlatformProcess::GetSynchEventFromPool(true);
+					Entry->Status = FStagingResourceStatus::WritingToDisk;
+
+					FStagingBufferPoolEntry* Staging = Entry.Get();
+
+					KickOffAsyncTask([this, Staging, Content]() {
+						check(Staging->WriteToDiskComplete);
+						check(Staging->Status == FStagingResourceStatus::WritingToDisk);
+
+						this->DumpResourceBinaryToFile(static_cast<const uint8*>(Content), Staging->ByteSize, Staging->DumpFilePath);
+
+						Staging->Status = FStagingResourceStatus::WritingToDiskComplete;
+						Staging->WriteToDiskComplete->Trigger();
+					});
+				}
+			}
+			else if (Entry->Status == FStagingResourceStatus::WritingToDiskComplete)
+			{
+				WaitDiskWriteAndReset(Entry.Get());
+
+				RHICmdList.UnlockStagingBuffer(Entry->StagingBuffer);
+				Entry->Status = FStagingResourceStatus::Unused;
+			}
+		}
+	}
+
+	void WaitAndReleaseStagingResources(FRHICommandListImmediate& RHICmdList)
+	{
+		check(IsInRenderingThread());
+		check(bStream);
+
+		RHICmdList.SubmitCommandsAndFlushGPU();
+		RHICmdList.BlockUntilGPUIdle();
+
+		LandCompletedResources(RHICmdList);
+
+		for (TUniquePtr<FStagingTexturePoolEntry>& Entry : StagingTexturePool)
+		{
+			if (Entry->Status == FStagingResourceStatus::WritingToDisk || Entry->Status == FStagingResourceStatus::WritingToDiskComplete)
+			{
+				WaitDiskWriteAndReset(Entry.Get());
+
+				RHICmdList.UnmapStagingSurface(Entry->Texture, Entry->GPUIndex);
+			}
+
+			check(Entry->WriteToDiskComplete == nullptr);
+			check(Entry->Status == FStagingResourceStatus::Unused || Entry->Status == FStagingResourceStatus::RenderingIntermediaryOnly);
+		}
+
+		for (TUniquePtr<FStagingBufferPoolEntry>& Entry : StagingBufferPool)
+		{
+			if (Entry->Status == FStagingResourceStatus::WritingToDisk || Entry->Status == FStagingResourceStatus::WritingToDiskComplete)
+			{
+				WaitDiskWriteAndReset(Entry.Get());
+
+				RHICmdList.UnlockStagingBuffer(Entry->StagingBuffer);
+			}
+
+			check(Entry->WriteToDiskComplete == nullptr);
+			check(Entry->Status == FStagingResourceStatus::Unused);
+		}
+
+		RHICmdList.SubmitCommandsAndFlushGPU();
+		RHICmdList.BlockUntilGPUIdle();
+
+		// Releases all staging resources
+		StagingTexturePool.Reset();
+		StagingBufferPool.Reset();
+		ReleaseRHIResources(RHICmdList);
+	}
+
 	// Look whether the pass matches matches r.DumpGPU.Root
 	bool IsDumpingPass(const FRDGPass* Pass)
 	{
@@ -1635,10 +2052,24 @@ bool IsDumpingRDGResources()
 
 void FRDGBuilder::InitResourceDump()
 {
+	check(IsInGameThread());
+
 	if(DumpingFrameCounter_GameThread == MAX_uint64)
 	{
 		DumpingFrameCounter_GameThread = GFrameCounter;
 	}
+
+	if (GRDGResourceDumpContext)
+	{
+		float DeltaTime = FApp::GetDeltaTime();
+
+		ENQUEUE_RENDER_COMMAND(FDumpGPUUpdateDeltaTime)(
+			[DeltaTime](FRHICommandListImmediate& ImmediateRHICmdList)
+		{
+			GRDGResourceDumpContext->DeltaTime = DeltaTime;
+			UE_LOG(LogRendererCore, Display, TEXT("Remaining frames %d"), GRDGResourceDumpContext->RemainingFrameCount);
+		});
+	}
 }
 
 FString FRDGBuilder::BeginResourceDump(const TCHAR* Cmd)
@@ -1710,7 +2141,9 @@ FString FRDGBuilder::BeginResourceDump(const TCHAR* Cmd)
 			}
 			NewResourceDumpContext->DumpingDirectoryPath = DirectoryPath / FApp::GetProjectName() + TEXT("-") + FPlatformProperties::PlatformName() + TEXT("-") + NewResourceDumpContext->Time.ToString() + TEXT("/");
 		}
+		NewResourceDumpContext->bStream = GDumpGPUStream.GetValueOnGameThread() != 0;
 		NewResourceDumpContext->bEnableDiskWrite = GDumpTestEnableDiskWrite.GetValueOnGameThread() != 0;
+		NewResourceDumpContext->DeltaTime = FApp::GetDeltaTime();
 		NewResourceDumpContext->RemainingFrameCount = FMath::Max(GDumpGPUFrameCount.GetValueOnGameThread(), 1);
 
 		if (Switches.Contains(TEXT("upload")))
@@ -1747,6 +2180,7 @@ FString FRDGBuilder::BeginResourceDump(const TCHAR* Cmd)
 		{
 			GNextDumpingRemainingTime = GDumpGPUDelay.GetValueOnGameThread();
 			GNextRDGResourceDumpContext = NewResourceDumpContext;
+			UE_LOG(LogRendererCore, Display, TEXT("DumpGPU to %s armed for %f seconds."), *NewResourceDumpContext->DumpingDirectoryPath, GNextDumpingRemainingTime);
 			return NewResourceDumpContext->DumpingDirectoryPath;
 		}
 	}
@@ -1755,7 +2189,8 @@ FString FRDGBuilder::BeginResourceDump(const TCHAR* Cmd)
 		NewResourceDumpContext = GNextRDGResourceDumpContext;
 		GNextRDGResourceDumpContext = nullptr;
 	}
-	
+
+	UE_LOG(LogRendererCore, Display, TEXT("DumpGPU to %s starting this frame"), *NewResourceDumpContext->DumpingDirectoryPath);
 	NewResourceDumpContext->MemoryConstants = FPlatformMemory::GetConstants();
 	NewResourceDumpContext->MemoryStats = FPlatformMemory::GetStats();
 
@@ -1793,6 +2228,8 @@ FString FRDGBuilder::BeginResourceDump(const TCHAR* Cmd)
 
 		FApp::SetFixedDeltaTime(1.0f / GDumpGPUFixedTickRate.GetValueOnGameThread());
 		FApp::SetUseFixedTimeStep(true);
+
+		UE_LOG(LogRendererCore, Display, TEXT("DumpGPU overriding tick rate to %fs"), float(FApp::GetFixedDeltaTime()));
 	}
 
 	// Output informations
@@ -1853,21 +2290,25 @@ FString FRDGBuilder::BeginResourceDump(const TCHAR* Cmd)
 	// Dump the rendering cvars
 	if (NewResourceDumpContext->bEnableDiskWrite)
 	{
-		NewResourceDumpContext->DumpRenderingCVarsToCSV();
+		NewResourceDumpContext->KickOffAsyncTask([NewResourceDumpContext]() {
+			NewResourceDumpContext->DumpRenderingCVarsToCSV();
+		});
 	}
 
 	// Copy the viewer
 	if (NewResourceDumpContext->bEnableDiskWrite)
 	{
-		const TCHAR* OpenGPUDumpViewerBatName = TEXT("OpenGPUDumpViewer.bat");
-		const TCHAR* OpenGPUDumpViewerShName = TEXT("OpenGPUDumpViewer.sh");
+		NewResourceDumpContext->KickOffAsyncTask([NewResourceDumpContext, &PlatformFile]() {
+			const TCHAR* OpenGPUDumpViewerBatName = TEXT("OpenGPUDumpViewer.bat");
+			const TCHAR* OpenGPUDumpViewerShName = TEXT("OpenGPUDumpViewer.sh");
 
-		const TCHAR* ViewerHTML = TEXT("GPUDumpViewer.html");
-		FString DumpGPUViewerSourcePath = FPaths::EngineDir() + FString(TEXT("Extras")) / TEXT("GPUDumpViewer");
+			const TCHAR* ViewerHTML = TEXT("GPUDumpViewer.html");
+			FString DumpGPUViewerSourcePath = FPaths::EngineDir() + FString(TEXT("Extras")) / TEXT("GPUDumpViewer");
 
-		PlatformFile.CopyFile(*(NewResourceDumpContext->DumpingDirectoryPath / ViewerHTML), *(DumpGPUViewerSourcePath / ViewerHTML));
-		PlatformFile.CopyFile(*(NewResourceDumpContext->DumpingDirectoryPath / OpenGPUDumpViewerBatName), *(DumpGPUViewerSourcePath / OpenGPUDumpViewerBatName));
-		PlatformFile.CopyFile(*(NewResourceDumpContext->DumpingDirectoryPath / OpenGPUDumpViewerShName), *(DumpGPUViewerSourcePath / OpenGPUDumpViewerShName));
+			PlatformFile.CopyFile(*(NewResourceDumpContext->DumpingDirectoryPath / ViewerHTML), *(DumpGPUViewerSourcePath / ViewerHTML));
+			PlatformFile.CopyFile(*(NewResourceDumpContext->DumpingDirectoryPath / OpenGPUDumpViewerBatName), *(DumpGPUViewerSourcePath / OpenGPUDumpViewerBatName));
+			PlatformFile.CopyFile(*(NewResourceDumpContext->DumpingDirectoryPath / OpenGPUDumpViewerShName), *(DumpGPUViewerSourcePath / OpenGPUDumpViewerShName));
+		});
 	}
 
 	ENQUEUE_RENDER_COMMAND(FStartGPUDump)(
@@ -1920,6 +2361,15 @@ void FRDGBuilder::EndResourceDump()
 	GRDGResourceDumpContext->RemainingFrameCount--;
 	if (GRDGResourceDumpContext->RemainingFrameCount > 0)
 	{
+		if (GRDGResourceDumpContext->bStream)
+		{
+			ENQUEUE_RENDER_COMMAND(FEndGPUDump)(
+				[](FRHICommandListImmediate& ImmediateRHICmdList)
+			{
+				GRDGResourceDumpContext->LandCompletedResources(ImmediateRHICmdList);
+			});
+		}
+
 		return;
 	}
 
@@ -1930,6 +2380,11 @@ void FRDGBuilder::EndResourceDump()
 		ENQUEUE_RENDER_COMMAND(FEndGPUDump)(
 			[](FRHICommandListImmediate& ImmediateRHICmdList)
 		{
+			if (GRDGResourceDumpContext->bStream)
+			{
+				GRDGResourceDumpContext->WaitAndReleaseStagingResources(ImmediateRHICmdList);
+			}
+
 			ImmediateRHICmdList.SubmitCommandsAndFlushGPU();
 			#if ENABLE_RHI_VALIDATION
 				GRHIValidateBufferSourceCopy = true;
@@ -1948,6 +2403,7 @@ void FRDGBuilder::EndResourceDump()
 		double GPUWaitSeconds = GRDGResourceDumpContext->TimingBucket[int32(FRDGResourceDumpContext::ETimingBucket::GPUWait)];
 		double CPUPostProcessingSeconds = GRDGResourceDumpContext->TimingBucket[int32(FRDGResourceDumpContext::ETimingBucket::CPUPostProcessing)];
 		double MetadataFileWriteSeconds = GRDGResourceDumpContext->TimingBucket[int32(FRDGResourceDumpContext::ETimingBucket::MetadataFileWrite)];
+		double ParametersFileWriteSeconds = GRDGResourceDumpContext->TimingBucket[int32(FRDGResourceDumpContext::ETimingBucket::ParametersFileWrite)];
 		double ResourceBinaryFileWriteSeconds = GRDGResourceDumpContext->TimingBucket[int32(FRDGResourceDumpContext::ETimingBucket::ResourceBinaryFileWrite)];
 		double RHIReleaseResourcesTimeSeconds = GRDGResourceDumpContext->TimingBucket[int32(FRDGResourceDumpContext::ETimingBucket::RHIReleaseResources)];
 
@@ -1957,12 +2413,17 @@ void FRDGBuilder::EndResourceDump()
 		UE_LOG(LogRendererCore, Display, TEXT("Dumped CPU resource binary post processing: %.3f s"), float(CPUPostProcessingSeconds));
 		UE_LOG(LogRendererCore, Display, TEXT("Dumped metadata: %.3f MB in %d files under %.3f s at %.3f MB/s"),
 			float(GRDGResourceDumpContext->MetadataFilesWriteBytes) / float(1024 * 1024),
-			GRDGResourceDumpContext->MetadataFilesOpened,
+			int32(GRDGResourceDumpContext->MetadataFilesOpened),
 			float(MetadataFileWriteSeconds),
 			float(GRDGResourceDumpContext->MetadataFilesWriteBytes) / (float(1024 * 1024) * float(MetadataFileWriteSeconds)));
+		UE_LOG(LogRendererCore, Display, TEXT("Dumped parameters: %.3f MB in %d files under %.3f s at %.3f MB/s"),
+			float(GRDGResourceDumpContext->ParametersFilesWriteBytes) / float(1024 * 1024),
+			int32(GRDGResourceDumpContext->ParametersFilesOpened),
+			float(ParametersFileWriteSeconds),
+			float(GRDGResourceDumpContext->ParametersFilesWriteBytes) / (float(1024 * 1024) * float(ParametersFileWriteSeconds)));
 		UE_LOG(LogRendererCore, Display, TEXT("Dumped resource binary: %.3f MB in %d files under %.3f s at %.3f MB/s"),
 			float(GRDGResourceDumpContext->ResourceBinaryWriteBytes) / float(1024 * 1024),
-			GRDGResourceDumpContext->ResourceBinaryFilesOpened,
+			int32(GRDGResourceDumpContext->ResourceBinaryFilesOpened),
 			float(ResourceBinaryFileWriteSeconds),
 			float(GRDGResourceDumpContext->ResourceBinaryWriteBytes) / (float(1024 * 1024) * float(ResourceBinaryFileWriteSeconds)));
 		UE_LOG(LogRendererCore, Display, TEXT("Dumped GPU readback resource release: %.3f s"), float(RHIReleaseResourcesTimeSeconds));
@@ -2247,7 +2708,7 @@ void FRDGBuilder::DumpResourcePassOutputs(const FRDGPass* Pass)
 		}
 		#endif
 		{
-			ParentEventScopeNames.Add(MakeShareable(new FJsonValueString(FString::Printf(TEXT("Frame %llu"), GFrameCounterRenderThread))));
+			ParentEventScopeNames.Add(MakeShareable(new FJsonValueString(FString::Printf(TEXT("Frame %llu (Delta=%fs)"), GFrameCounterRenderThread, GRDGResourceDumpContext->DeltaTime))));
 		}
 
 		FString EventNameStorage;
@@ -2289,7 +2750,7 @@ void FRDGBuilder::DumpResourcePassOutputs(const FRDGPass* Pass)
 		{
 			TArrayView<const uint8> ArrayView(PassParametersContent, PassParametersByteSize);
 			FString DumpFilePath = FRDGResourceDumpContext::kStructuresDir / FRDGResourceDumpContext::PtrToString(PassParametersContent) + TEXT(".bin");
-			GRDGResourceDumpContext->DumpBinaryToFile(ArrayView, DumpFilePath);
+			GRDGResourceDumpContext->DumpBinaryToFile(ArrayView, DumpFilePath, FRDGResourceDumpContext::ETimingBucket::ParametersFileWrite);
 			GRDGResourceDumpContext->SetDumped(PassParametersContent);
 		}
 	}
