Adds a LogDumpGPU category for improved log search

diff --git a/Engine/Source/Runtime/RenderCore/Private/DumpGPU.cpp b/Engine/Source/Runtime/RenderCore/Private/DumpGPU.cpp
--- a/Engine/Source/Runtime/RenderCore/Private/DumpGPU.cpp
+++ b/Engine/Source/Runtime/RenderCore/Private/DumpGPU.cpp
@@ -23,7 +23,8 @@
 #include "RHIUtilities.h"
 #include "RHIValidation.h"
 #include "RenderGraphPrivate.h"
-#include "RenderCore.h"
+
+DEFINE_LOG_CATEGORY_STATIC(LogDumpGPU, Log, All);
 
 IDumpGPUUploadServiceProvider* IDumpGPUUploadServiceProvider::GProvider = nullptr;
 
@@ -402,7 +403,7 @@ public:
 
 	bool DumpStatusToFile(FStringView StatusString)
 	{
-		UE_LOG(LogRendererCore, Display, TEXT("DumpGPU status = %s"), StatusString.GetData());
+		UE_LOG(LogDumpGPU, Display, TEXT("DumpGPU status = %s"), StatusString.GetData());
 		return DumpStringToFile(StatusString, FString(FRDGResourceDumpContext::kBaseDir) / TEXT("Status.txt"));
 	}
 
@@ -530,7 +531,7 @@ public:
 
 		if (ResourcesDumpExecutedPasses % 10 == 0)
 		{
-			UE_LOG(LogRendererCore, Display, TEXT("Dumped %d / %d resources"), ResourcesDumpExecutedPasses, ResourcesDumpPasses);
+			UE_LOG(LogDumpGPU, Display, TEXT("Dumped %d / %d resources"), ResourcesDumpExecutedPasses, ResourcesDumpPasses);
 		}
 	}
 
@@ -674,11 +675,11 @@ public:
 
 		if (bSuccess)
 		{
-			UE_LOG(LogRendererCore, Display, TEXT("DumpGPU dumped rendering cvars to %s."), *FileName);
+			UE_LOG(LogDumpGPU, Display, TEXT("DumpGPU dumped rendering cvars to %s."), *FileName);
 		}
 		else
 		{
-			UE_LOG(LogRendererCore, Error, TEXT("DumpGPU had a file error when dumping rendering cvars to %s."), *FileName);
+			UE_LOG(LogDumpGPU, Error, TEXT("DumpGPU had a file error when dumping rendering cvars to %s."), *FileName);
 		}
 	}
 
@@ -1319,7 +1320,7 @@ public:
 			}
 			else
 			{
-				UE_LOG(LogRendererCore, Warning, TEXT("RHICmdList.MapStagingSurface() to dump texture %s failed."), TextureDebugName);
+				UE_LOG(LogDumpGPU, Warning, TEXT("RHICmdList.MapStagingSurface() to dump texture %s failed."), TextureDebugName);
 			}
 		}
 	}
@@ -1432,7 +1433,7 @@ public:
 		// Verify there is enough available memory to dump the resource.
 		if (IsUnsafeToDumpResource(SubresourceDumpDesc.ByteSize, 2.2f + (SubresourceDumpDesc.bPreprocessForStaging ? 1.0f : 0.0f)))
 		{
-			UE_LOG(LogRendererCore, Warning, TEXT("Not dumping %s because of insuficient memory available for staging texture."), SubresourceDesc.Texture->Name);
+			UE_LOG(LogDumpGPU, Warning, TEXT("Not dumping %s because of insuficient memory available for staging texture."), SubresourceDesc.Texture->Name);
 			return;
 		}
 
@@ -1453,7 +1454,7 @@ public:
 			{
 				if (!(SubresourceDesc.Texture->Desc.Flags & TexCreate_ShaderResource))
 				{
-					UE_LOG(LogRendererCore, Warning, TEXT("Not dumping %s because requires copy to staging texture using compute, but is missing TexCreate_ShaderResource."), *UniqueResourceSubResourceName);
+					UE_LOG(LogDumpGPU, Warning, TEXT("Not dumping %s because requires copy to staging texture using compute, but is missing TexCreate_ShaderResource."), *UniqueResourceSubResourceName);
 					return;
 				}
 
@@ -1720,14 +1721,14 @@ public:
 
 			if (IsUnsafeToDumpResource(StagingResourceByteSize, 1.2f))
 			{
-				UE_LOG(LogRendererCore, Warning, TEXT("Not dumping %s because of insuficient memory available for staging buffer."), *DumpFilePath);
+				UE_LOG(LogDumpGPU, Warning, TEXT("Not dumping %s because of insuficient memory available for staging buffer."), *DumpFilePath);
 				return;
 			}
 
 			// Verify the texture is able to do resource transitions.
 			if (EnumHasAnyFlags(Buffer->Flags, ERDGBufferFlags::SkipTracking))
 			{
-				UE_LOG(LogRendererCore, Warning, TEXT("Not dumping %s because has ERDGBufferFlags::SkipTracking."), Buffer->Name);
+				UE_LOG(LogDumpGPU, Warning, TEXT("Not dumping %s because has ERDGBufferFlags::SkipTracking."), Buffer->Name);
 				return;
 			}
 
@@ -1815,7 +1816,7 @@ public:
 						}
 						else
 						{
-							UE_LOG(LogRendererCore, Warning, TEXT("RHICmdList.LockStagingBuffer() to dump buffer %s failed."), Buffer->Name);
+							UE_LOG(LogDumpGPU, Warning, TEXT("RHICmdList.LockStagingBuffer() to dump buffer %s failed."), Buffer->Name);
 							break;
 						}
 					}
@@ -2067,7 +2068,7 @@ void FRDGBuilder::InitResourceDump()
 			[DeltaTime](FRHICommandListImmediate& ImmediateRHICmdList)
 		{
 			GRDGResourceDumpContext->DeltaTime = DeltaTime;
-			UE_LOG(LogRendererCore, Display, TEXT("Remaining frames %d"), GRDGResourceDumpContext->RemainingFrameCount);
+			UE_LOG(LogDumpGPU, Display, TEXT("Remaining frames %d"), GRDGResourceDumpContext->RemainingFrameCount);
 		});
 	}
 }
@@ -2150,11 +2151,11 @@ FString FRDGBuilder::BeginResourceDump(const TCHAR* Cmd)
 		{
 			if (!IDumpGPUUploadServiceProvider::GProvider || !NewResourceDumpContext->bEnableDiskWrite)
 			{
-				UE_LOG(LogRendererCore, Warning, TEXT("DumpGPU upload services are not set up."));
+				UE_LOG(LogDumpGPU, Warning, TEXT("DumpGPU upload services are not set up."));
 			}
 			else if (GDumpGPUUploadCVar.GetValueOnGameThread() == 0)
 			{
-				UE_LOG(LogRendererCore, Warning, TEXT("DumpGPU upload services are not available because r.DumpGPU.Upload=0."));
+				UE_LOG(LogDumpGPU, Warning, TEXT("DumpGPU upload services are not available because r.DumpGPU.Upload=0."));
 			}
 			else
 			{
@@ -2180,7 +2181,7 @@ FString FRDGBuilder::BeginResourceDump(const TCHAR* Cmd)
 		{
 			GNextDumpingRemainingTime = GDumpGPUDelay.GetValueOnGameThread();
 			GNextRDGResourceDumpContext = NewResourceDumpContext;
-			UE_LOG(LogRendererCore, Display, TEXT("DumpGPU to %s armed for %f seconds."), *NewResourceDumpContext->DumpingDirectoryPath, GNextDumpingRemainingTime);
+			UE_LOG(LogDumpGPU, Display, TEXT("DumpGPU to %s armed for %f seconds."), *NewResourceDumpContext->DumpingDirectoryPath, GNextDumpingRemainingTime);
 			return NewResourceDumpContext->DumpingDirectoryPath;
 		}
 	}
@@ -2190,7 +2191,7 @@ FString FRDGBuilder::BeginResourceDump(const TCHAR* Cmd)
 		GNextRDGResourceDumpContext = nullptr;
 	}
 
-	UE_LOG(LogRendererCore, Display, TEXT("DumpGPU to %s starting this frame"), *NewResourceDumpContext->DumpingDirectoryPath);
+	UE_LOG(LogDumpGPU, Display, TEXT("DumpGPU to %s starting this frame"), *NewResourceDumpContext->DumpingDirectoryPath);
 	NewResourceDumpContext->MemoryConstants = FPlatformMemory::GetConstants();
 	NewResourceDumpContext->MemoryStats = FPlatformMemory::GetStats();
 
@@ -2229,7 +2230,7 @@ FString FRDGBuilder::BeginResourceDump(const TCHAR* Cmd)
 		FApp::SetFixedDeltaTime(1.0f / GDumpGPUFixedTickRate.GetValueOnGameThread());
 		FApp::SetUseFixedTimeStep(true);
 
-		UE_LOG(LogRendererCore, Display, TEXT("DumpGPU overriding tick rate to %fs"), float(FApp::GetFixedDeltaTime()));
+		UE_LOG(LogDumpGPU, Display, TEXT("DumpGPU overriding tick rate to %fs"), float(FApp::GetFixedDeltaTime()));
 	}
 
 	// Output informations
@@ -2375,7 +2376,7 @@ void FRDGBuilder::EndResourceDump()
 
 	// Wait all rendering commands are completed to finish with GRDGResourceDumpContext.
 	{
-		UE_LOG(LogRendererCore, Display, TEXT("Stalling game thread until render thread finishes to dump resources"));
+		UE_LOG(LogDumpGPU, Display, TEXT("Stalling game thread until render thread finishes to dump resources"));
 
 		ENQUEUE_RENDER_COMMAND(FEndGPUDump)(
 			[](FRHICommandListImmediate& ImmediateRHICmdList)
@@ -2407,26 +2408,26 @@ void FRDGBuilder::EndResourceDump()
 		double ResourceBinaryFileWriteSeconds = GRDGResourceDumpContext->TimingBucket[int32(FRDGResourceDumpContext::ETimingBucket::ResourceBinaryFileWrite)];
 		double RHIReleaseResourcesTimeSeconds = GRDGResourceDumpContext->TimingBucket[int32(FRDGResourceDumpContext::ETimingBucket::RHIReleaseResources)];
 
-		UE_LOG(LogRendererCore, Display, TEXT("Dumped %d resources in %.3f s to %s"), GRDGResourceDumpContext->ResourcesDumpPasses, float(TotalDumpSeconds), *AbsDumpingDirectoryPath);
-		UE_LOG(LogRendererCore, Display, TEXT("Dumped GPU readback commands: %.3f s"), float(RHIReadbackCommandsSeconds));
-		UE_LOG(LogRendererCore, Display, TEXT("Dumped GPU wait: %.3f s"), float(GPUWaitSeconds));
-		UE_LOG(LogRendererCore, Display, TEXT("Dumped CPU resource binary post processing: %.3f s"), float(CPUPostProcessingSeconds));
-		UE_LOG(LogRendererCore, Display, TEXT("Dumped metadata: %.3f MB in %d files under %.3f s at %.3f MB/s"),
+		UE_LOG(LogDumpGPU, Display, TEXT("Dumped %d resources in %.3f s to %s"), GRDGResourceDumpContext->ResourcesDumpPasses, float(TotalDumpSeconds), *AbsDumpingDirectoryPath);
+		UE_LOG(LogDumpGPU, Display, TEXT("Dumped GPU readback commands: %.3f s"), float(RHIReadbackCommandsSeconds));
+		UE_LOG(LogDumpGPU, Display, TEXT("Dumped GPU wait: %.3f s"), float(GPUWaitSeconds));
+		UE_LOG(LogDumpGPU, Display, TEXT("Dumped CPU resource binary post processing: %.3f s"), float(CPUPostProcessingSeconds));
+		UE_LOG(LogDumpGPU, Display, TEXT("Dumped metadata: %.3f MB in %d files under %.3f s at %.3f MB/s"),
 			float(GRDGResourceDumpContext->MetadataFilesWriteBytes) / float(1024 * 1024),
 			int32(GRDGResourceDumpContext->MetadataFilesOpened),
 			float(MetadataFileWriteSeconds),
 			float(GRDGResourceDumpContext->MetadataFilesWriteBytes) / (float(1024 * 1024) * float(MetadataFileWriteSeconds)));
-		UE_LOG(LogRendererCore, Display, TEXT("Dumped parameters: %.3f MB in %d files under %.3f s at %.3f MB/s"),
+		UE_LOG(LogDumpGPU, Display, TEXT("Dumped parameters: %.3f MB in %d files under %.3f s at %.3f MB/s"),
 			float(GRDGResourceDumpContext->ParametersFilesWriteBytes) / float(1024 * 1024),
 			int32(GRDGResourceDumpContext->ParametersFilesOpened),
 			float(ParametersFileWriteSeconds),
 			float(GRDGResourceDumpContext->ParametersFilesWriteBytes) / (float(1024 * 1024) * float(ParametersFileWriteSeconds)));
-		UE_LOG(LogRendererCore, Display, TEXT("Dumped resource binary: %.3f MB in %d files under %.3f s at %.3f MB/s"),
+		UE_LOG(LogDumpGPU, Display, TEXT("Dumped resource binary: %.3f MB in %d files under %.3f s at %.3f MB/s"),
 			float(GRDGResourceDumpContext->ResourceBinaryWriteBytes) / float(1024 * 1024),
 			int32(GRDGResourceDumpContext->ResourceBinaryFilesOpened),
 			float(ResourceBinaryFileWriteSeconds),
 			float(GRDGResourceDumpContext->ResourceBinaryWriteBytes) / (float(1024 * 1024) * float(ResourceBinaryFileWriteSeconds)));
-		UE_LOG(LogRendererCore, Display, TEXT("Dumped GPU readback resource release: %.3f s"), float(RHIReleaseResourcesTimeSeconds));
+		UE_LOG(LogDumpGPU, Display, TEXT("Dumped GPU readback resource release: %.3f s"), float(RHIReleaseResourcesTimeSeconds));
 	}
 
 	// Dump the log into the dump directory.
@@ -2542,7 +2543,7 @@ void FRDGBuilder::DumpResourcePassOutputs(const FRDGPass* Pass)
 				}
 				else
 				{
-					UE_LOG(LogRendererCore, Warning, TEXT("Dumping texture %s's meta data unsupported"), SRV->Desc.Texture->Name);
+					UE_LOG(LogDumpGPU, Warning, TEXT("Dumping texture %s's meta data unsupported"), SRV->Desc.Texture->Name);
 				}
 			}
 		}
@@ -2565,7 +2566,7 @@ void FRDGBuilder::DumpResourcePassOutputs(const FRDGPass* Pass)
 				}
 				else
 				{
-					UE_LOG(LogRendererCore, Warning, TEXT("Dumping texture %s's meta data unsupported"), UAV->Desc.Texture->Name);
+					UE_LOG(LogDumpGPU, Warning, TEXT("Dumping texture %s's meta data unsupported"), UAV->Desc.Texture->Name);
 				}
 			}
 		}
@@ -2779,7 +2780,7 @@ void FRDGBuilder::BeginPassDump(const FRDGPass* Pass)
 
 	if (!IsInRenderingThread())
 	{
-		UE_LOG(LogRendererCore, Warning, TEXT("Couldn't start dumping draw's resources for pass %s because not in the rendering thread"), Pass->GetEventName().GetTCHAR());
+		UE_LOG(LogDumpGPU, Warning, TEXT("Couldn't start dumping draw's resources for pass %s because not in the rendering thread"), Pass->GetEventName().GetTCHAR());
 		return;
 	}
 
@@ -2802,7 +2803,7 @@ void FRDGBuilder::DumpDraw(const FRDGEventName& DrawEventName)
 
 	if (!IsInRenderingThread())
 	{
-		UE_LOG(LogRendererCore, Warning, TEXT("Couldn't dump draw because not in the rendering thread"));
+		UE_LOG(LogDumpGPU, Warning, TEXT("Couldn't dump draw because not in the rendering thread"));
 		return;
 	}
 
@@ -2885,7 +2886,7 @@ void FRDGBuilder::DumpDraw(const FRDGEventName& DrawEventName)
 
 	if (GRDGResourceDumpContext->DrawDumpCount % 10 == 0)
 	{
-		UE_LOG(LogRendererCore, Display, TEXT("Dumped %d draws' resources"), GRDGResourceDumpContext->DrawDumpCount);
+		UE_LOG(LogDumpGPU, Display, TEXT("Dumped %d draws' resources"), GRDGResourceDumpContext->DrawDumpCount);
 		return;
 	}
 }
@@ -2921,7 +2922,7 @@ void FRDGBuilder::EndPassDump(const FRDGPass* Pass)
 
 		GRDGResourceDumpContext->DumpJsonToFile(JsonObject, FString(FRDGResourceDumpContext::kBaseDir) / TEXT("PassDrawCounts.json"), FILEWRITE_Append);
 
-		UE_LOG(LogRendererCore, Display, TEXT("Completed dump of %d draws for pass: %s"), GRDGResourceDumpContext->DrawDumpCount, Pass->GetEventName().GetTCHAR());
+		UE_LOG(LogDumpGPU, Display, TEXT("Completed dump of %d draws for pass: %s"), GRDGResourceDumpContext->DrawDumpCount, Pass->GetEventName().GetTCHAR());
 	}
 
 	GRDGResourceDumpContext->DrawDumpingPass = nullptr;
