Fixes resource transition bug in DumpGPU

diff --git a/Engine/Source/Runtime/RenderCore/Private/DumpGPU.cpp b/Engine/Source/Runtime/RenderCore/Private/DumpGPU.cpp
--- a/Engine/Source/Runtime/RenderCore/Private/DumpGPU.cpp
+++ b/Engine/Source/Runtime/RenderCore/Private/DumpGPU.cpp
@@ -1253,7 +1253,7 @@ public:
 				else
 				{
 					CreateStagingTexture(Desc, /* out */ &StagingTexture);
-					RHICmdList.Transition(FRHITransitionInfo(StagingSrcTexture, ERHIAccess::Unknown, ERHIAccess::CopyDest));
+					RHICmdList.Transition(FRHITransitionInfo(StagingTexture, ERHIAccess::Unknown, ERHIAccess::CopyDest));
 				}
 
 				// Transfer memory GPU -> CPU
