Fixed non-unity build break in RayTracingGeometry.cpp.

diff --git a/Engine/Source/Runtime/RenderCore/Private/RayTracingGeometry.cpp b/Engine/Source/Runtime/RenderCore/Private/RayTracingGeometry.cpp
--- a/Engine/Source/Runtime/RenderCore/Private/RayTracingGeometry.cpp
+++ b/Engine/Source/Runtime/RenderCore/Private/RayTracingGeometry.cpp
@@ -1,7 +1,7 @@
 // Copyright Epic Games, Inc. All Rights Reserved.
 
 #include "RayTracingGeometry.h"
-#include "DynamicRHI.h"
+#include "RHICommandList.h"
 #include "HAL/IConsoleManager.h"
 #include "RayTracingGeometryManager.h"
 #include "RenderUtils.h"
