fixed memory accounting error in FDynamicVertexBufferPool

diff --git a/Engine/Source/Runtime/RenderCore/Private/GlobalRenderResources.cpp b/Engine/Source/Runtime/RenderCore/Private/GlobalRenderResources.cpp
--- a/Engine/Source/Runtime/RenderCore/Private/GlobalRenderResources.cpp
+++ b/Engine/Source/Runtime/RenderCore/Private/GlobalRenderResources.cpp
@@ -816,6 +816,7 @@ struct FDynamicVertexBufferPool : public FRenderResource
 
 			if (GGlobalBufferNumFramesUnusedThreshold > 0 && Buffer->LastUsedFrame + GGlobalBufferNumFramesUnusedThreshold <= GFrameCounterRenderThread)
 			{
+				TotalAllocatedMemory -= Buffer->BufferSize;
 				Buffer->ReleaseResource();
 				LiveList.RemoveAt(Index, 1, false);
 				FreeList.Remove(Buffer);
