RenderCommandFence and its bundling:
on shutdown, AppPreExit() is covered by StartRenderCommandFenceBundler()/StoptRenderCommandFenceBundler() to minimise the number of RT tasks as there're lots of render fences in this region. waiting on any of fences stops bundling, e.g. by FlushRenderingCommands(). As this happens early in AppPreExit(), this effectively disables bundling and we get a ton of RT tasks that drives mem usage up.
This CL adds FlushRenderCommandFenceBundler() and uses it in FRenderCommandFence::Wait() instead of StopRenderCommandFenceBundler() to keep bundling enabled

diff --git a/Engine/Source/Runtime/RenderCore/Private/RenderingThread.cpp b/Engine/Source/Runtime/RenderCore/Private/RenderingThread.cpp
--- a/Engine/Source/Runtime/RenderCore/Private/RenderingThread.cpp
+++ b/Engine/Source/Runtime/RenderCore/Private/RenderingThread.cpp
@@ -970,40 +970,95 @@ bool IsRenderingThreadHealthy()
 
 static FGraphEventRef BundledCompletionEvent;
 static FGraphEventRef BundledCompletionEventPrereq; // We fire this when we are done, which queues the actual fence
+static uint32 BundlingRecursionDepth = 0; // accessed only by GT
+
+DECLARE_CYCLE_STAT(TEXT("FNullGraphTask.FenceRenderCommandBundled"),
+	STAT_FNullGraphTask_FenceRenderCommandBundled,
+	STATGROUP_TaskGraphTasks);
+
+void CheckRenderCommandFenceBundler()
+{
+	check(BundlingRecursionDepth > 0);
+	check(BundledCompletionEvent.GetReference());
+	check(!BundledCompletionEvent->IsComplete());
+	check(BundledCompletionEventPrereq.GetReference());
+	check(!BundledCompletionEventPrereq->IsComplete());
+}
 
 void StartRenderCommandFenceBundler()
 {
+	check(IsInGameThread());
+
 	if (!GIsThreadedRendering)
 	{
 		return;
 	}
 
-	check(IsInGameThread() && !BundledCompletionEvent.GetReference() && !BundledCompletionEventPrereq.GetReference()); // can't use this in a nested fashion
+	++BundlingRecursionDepth;
+
+	if (BundlingRecursionDepth > 1)
+	{
+		CheckRenderCommandFenceBundler();
+		return;
+	}
+
 	BundledCompletionEventPrereq = FGraphEvent::CreateGraphEvent();
 
 	FGraphEventArray Prereqs;
 	Prereqs.Add(BundledCompletionEventPrereq);
 
-	DECLARE_CYCLE_STAT(TEXT("FNullGraphTask.FenceRenderCommandBundled"),
-	STAT_FNullGraphTask_FenceRenderCommandBundled,
-		STATGROUP_TaskGraphTasks);
-
 	BundledCompletionEvent = TGraphTask<FNullGraphTask>::CreateTask(&Prereqs, ENamedThreads::GameThread).ConstructAndDispatchWhenReady(
 		GET_STATID(STAT_FNullGraphTask_FenceRenderCommandBundled), ENamedThreads::GetRenderThread());
 
 	StartBatchedRelease();
 }
 
+void FlushRenderCommandFenceBundler()
+{
+	check(IsInGameThread());
+
+	if (BundledCompletionEvent.GetReference())
+	{
+		CheckRenderCommandFenceBundler();
+
+		// prepare the new bundle
+		FGraphEventRef LocalBundledCompletionEventPrereq = BundledCompletionEventPrereq;
+		BundledCompletionEventPrereq = FGraphEvent::CreateGraphEvent();
+
+		FGraphEventArray Prereqs;
+		Prereqs.Add(BundledCompletionEventPrereq);
+
+		BundledCompletionEvent = TGraphTask<FNullGraphTask>::CreateTask(&Prereqs, ENamedThreads::GameThread).ConstructAndDispatchWhenReady(
+			GET_STATID(STAT_FNullGraphTask_FenceRenderCommandBundled), ENamedThreads::GetRenderThread());
+
+		// release the old bundle
+		LocalBundledCompletionEventPrereq->DispatchSubsequents();
+	}
+}
+
 void StopRenderCommandFenceBundler()
 {
-	if (!GIsThreadedRendering || !BundledCompletionEvent.GetReference())
+	check(IsInGameThread());
+
+	if (!GIsThreadedRendering)
 	{
 		return;
 	}
 
+	CheckRenderCommandFenceBundler();
+
+	--BundlingRecursionDepth;
+
+	if (BundlingRecursionDepth > 0)
+	{
+		FlushRenderCommandFenceBundler();
+		return;
+	}
+
 	EndBatchedRelease();
-	checkf(IsInGameThread() && BundledCompletionEvent.GetReference() && !BundledCompletionEvent->IsComplete() && BundledCompletionEventPrereq.GetReference() && !BundledCompletionEventPrereq->IsComplete(), TEXT("IsInGameThread: %d, BundledCompletionEvent is completed: %d, BundledCompletionEventPrereq is completed: %d"), IsInGameThread(), BundledCompletionEvent->IsComplete(), BundledCompletionEventPrereq->IsComplete()); // can't use this in a nested fashion
+
 	BundledCompletionEventPrereq->DispatchSubsequents();
+
 	BundledCompletionEventPrereq = nullptr;
 	BundledCompletionEvent = nullptr;
 }
@@ -1292,7 +1347,7 @@ void FRenderCommandFence::Wait(bool bProcessGameThreadTasks) const
 {
 	if (!IsFenceComplete())
 	{
-		StopRenderCommandFenceBundler();
+		FlushRenderCommandFenceBundler();
 #if 0
 		// on most platforms this is a better solution because it doesn't spin
 		// windows needs to pump messages
