Add FString support to the compileArg array so that we can add GlobalDefines from the engine .ini file. This allows to access the compile arguments at the compiler level to possibly enable compile flags/options.

diff --git a/Engine/Source/Runtime/RenderCore/Private/GlobalShader.cpp b/Engine/Source/Runtime/RenderCore/Private/GlobalShader.cpp
--- a/Engine/Source/Runtime/RenderCore/Private/GlobalShader.cpp
+++ b/Engine/Source/Runtime/RenderCore/Private/GlobalShader.cpp
@@ -39,7 +39,7 @@ public:
 				{
 					if (Define.IsRelevant(ShaderFormat, Parameters.PermutationId))
 					{
-						OutEnvironment.SetDefine(*Define.Name, Define.Value);
+						OutEnvironment.SetDefineAndCompileArgument(*Define.Name, Define.Value);
 					}
 				}
 			}
diff --git a/Engine/Source/Runtime/RenderCore/Public/ShaderCore.h b/Engine/Source/Runtime/RenderCore/Public/ShaderCore.h
--- a/Engine/Source/Runtime/RenderCore/Public/ShaderCore.h
+++ b/Engine/Source/Runtime/RenderCore/Public/ShaderCore.h
@@ -602,7 +602,7 @@ struct FShaderCompilerEnvironment
 	// behaviour. Intended to replace any usage of definitions after shader preprocessing.
 	template <typename ValueType> void SetCompileArgument(const TCHAR* Name, ValueType Value)
 	{
-		CompileArgs.Add(Name, TVariant<bool, float, int32, uint32>(TInPlaceType<ValueType>(), Value));
+		CompileArgs.Add(Name, TVariant<bool, float, int32, uint32, FString>(TInPlaceType<ValueType>(), Value));
 	}
 
 	// Helper to set both a define and a compile argument to the same value. Useful for various parameters which
@@ -613,11 +613,21 @@ struct FShaderCompilerEnvironment
 		SetCompileArgument(Name, Value);
 	}
 
+	// If a compile argument with the given name exists, returns true. 
+	bool HasCompileArgument(const TCHAR* Name) const
+	{
+		if (CompileArgs.Contains(Name))
+		{
+			return true;
+		}
+		return false;
+	}
+
 	// If a compile argument with the given name exists and is of the specified type, returns its value. Otherwise, 
 	// either the named argument doesn't exist or the type does not match, and the default value will be returned.
 	template <typename ValueType> ValueType GetCompileArgument(const TCHAR* Name, const ValueType& DefaultValue) const
 	{
-		const TVariant<bool, float, int32, uint32>* StoredValue = CompileArgs.Find(Name);
+		const TVariant<bool, float, int32, uint32, FString>* StoredValue = CompileArgs.Find(Name);
 		if (StoredValue && StoredValue->IsType<ValueType>())
 		{
 			return StoredValue->Get<ValueType>();
@@ -630,7 +640,7 @@ struct FShaderCompilerEnvironment
 	// OutValue will be left unmodified and the function will return false.
 	template <typename ValueType> bool GetCompileArgument(const TCHAR* Name, ValueType& OutValue) const
 	{
-		const TVariant<bool, float, int32, uint32>* StoredValue = CompileArgs.Find(Name);
+		const TVariant<bool, float, int32, uint32, FString>* StoredValue = CompileArgs.Find(Name);
 		if (StoredValue && StoredValue->IsType<ValueType>())
 		{
 			OutValue = StoredValue->Get<ValueType>();
@@ -726,7 +736,7 @@ private:
 	friend class FShaderPreprocessorUtilities;
 
 	FShaderCompilerDefinitions Definitions;
-	TMap<FString, TVariant<bool, float, int32, uint32>> CompileArgs;
+	TMap<FString, TVariant<bool, float, int32, uint32, FString>> CompileArgs;
 };
 
 struct FSharedShaderCompilerEnvironment final : public FShaderCompilerEnvironment, public FRefCountBase
