[shaders] add Hash field to FShaderCompilerInput so debug output of InputHash.txt works correctly whether or not the preprocessed cache is enabled, and the IShaderFormat API for debug output of independent-preprocessing-enabled backends is cleaner

diff --git a/Engine/Source/Developer/ShaderCompilerCommon/Private/ShaderCompilerCommon.cpp b/Engine/Source/Developer/ShaderCompilerCommon/Private/ShaderCompilerCommon.cpp
--- a/Engine/Source/Developer/ShaderCompilerCommon/Private/ShaderCompilerCommon.cpp
+++ b/Engine/Source/Developer/ShaderCompilerCommon/Private/ShaderCompilerCommon.cpp
@@ -1401,7 +1401,7 @@ void CompileOfflineMali(const FShaderCompilerInput& Input, FShaderCompilerOutput
 
 namespace UE::ShaderCompilerCommon
 {
-	void FBaseShaderFormat::OutputDebugData(const FString& InputHash, const FShaderCompilerInput& Input, const FShaderPreprocessOutput& PreprocessOutput, const FShaderCompilerOutput& Output) const
+	void FBaseShaderFormat::OutputDebugData(const FShaderCompilerInput& Input, const FShaderPreprocessOutput& PreprocessOutput, const FShaderCompilerOutput& Output) const
 	{
 		const FString& SourceToDump = Output.ModifiedShaderSource.IsEmpty() ? PreprocessOutput.GetUnstrippedSource() : Output.ModifiedShaderSource;
 		DumpDebugShaderData(Input, SourceToDump);
@@ -1422,10 +1422,10 @@ namespace UE::ShaderCompilerCommon
 			}
 		}
 
-		if (EnumHasAnyFlags(Input.DebugInfoFlags, EShaderDebugInfoFlags::InputHash) && !InputHash.IsEmpty())
+		if (EnumHasAnyFlags(Input.DebugInfoFlags, EShaderDebugInfoFlags::InputHash))
 		{
 			FString InputHashFileName = FPaths::Combine(Input.DumpDebugInfoPath, TEXT("InputHash.txt"));
-			FFileHelper::SaveStringToFile(InputHash, *InputHashFileName, FFileHelper::EEncodingOptions::ForceAnsi);
+			FFileHelper::SaveStringToFile(LexToString(Input.Hash), *InputHashFileName, FFileHelper::EEncodingOptions::ForceAnsi);
 		}
 	}
 
diff --git a/Engine/Source/Developer/ShaderCompilerCommon/Public/ShaderCompilerCommon.h b/Engine/Source/Developer/ShaderCompilerCommon/Public/ShaderCompilerCommon.h
--- a/Engine/Source/Developer/ShaderCompilerCommon/Public/ShaderCompilerCommon.h
+++ b/Engine/Source/Developer/ShaderCompilerCommon/Public/ShaderCompilerCommon.h
@@ -75,7 +75,7 @@ namespace UE::ShaderCompilerCommon
 	class FBaseShaderFormat : public IShaderFormat
 	{
 	public:
-		virtual SHADERCOMPILERCOMMON_API void OutputDebugData(const FString& InputHash, const FShaderCompilerInput& Input, const FShaderPreprocessOutput& PreprocessOutput, const FShaderCompilerOutput& Output) const override;
+		virtual SHADERCOMPILERCOMMON_API void OutputDebugData(const FShaderCompilerInput& Input, const FShaderPreprocessOutput& PreprocessOutput, const FShaderCompilerOutput& Output) const override;
 	};
 
 	UE_DEPRECATED(5.3, "Deprecated; use FShaderCompilerInput::ShouldUseStableContantBuffer directly")
diff --git a/Engine/Source/Developer/TargetPlatform/Public/Interfaces/IShaderFormat.h b/Engine/Source/Developer/TargetPlatform/Public/Interfaces/IShaderFormat.h
--- a/Engine/Source/Developer/TargetPlatform/Public/Interfaces/IShaderFormat.h
+++ b/Engine/Source/Developer/TargetPlatform/Public/Interfaces/IShaderFormat.h
@@ -147,7 +147,7 @@ public:
 	 * with the hash of the shader code in an OutputHash.txt file; if no additional custom debug output is required the shader 
 	 * format can inherit from BaseShaderFormat instead of IShaderFormat.
 	 */
-	virtual void OutputDebugData(const FString& InputHash, const struct FShaderCompilerInput& Input, const class FShaderPreprocessOutput& PreprocessOutput, const struct FShaderCompilerOutput& Output) const {};
+	virtual void OutputDebugData(const struct FShaderCompilerInput& Input, const class FShaderPreprocessOutput& PreprocessOutput, const struct FShaderCompilerOutput& Output) const {};
 
 	/**
 	 * Return true if preprocessing for this format can be executed independent of compilation (i.e. the format has an
diff --git a/Engine/Source/Runtime/RenderCore/Private/ShaderCore.cpp b/Engine/Source/Runtime/RenderCore/Private/ShaderCore.cpp
--- a/Engine/Source/Runtime/RenderCore/Private/ShaderCore.cpp
+++ b/Engine/Source/Runtime/RenderCore/Private/ShaderCore.cpp
@@ -1148,7 +1148,7 @@ void CompileShader(const TArray<const IShaderFormat*>& ShaderFormats, FShaderCom
 		// cache is disabled
 		if (Compiler->SupportsIndependentPreprocessing() && !Job.Input.bCachePreprocessed)
 		{
-			Compiler->OutputDebugData(FString(), Job.Input, Job.PreprocessOutput, Job.Output);
+			Compiler->OutputDebugData(Job.Input, Job.PreprocessOutput, Job.Output);
 		}
 		else
 		{
@@ -2748,6 +2748,7 @@ FArchive& operator<<(FArchive& Ar, FShaderCompilerInput& Input)
 	Ar << Input.DebugExtension;
 	Ar << Input.DebugGroupName;
 	Ar << Input.DebugDescription;
+	Ar << Input.Hash;
 	if (Input.bCachePreprocessed)
 	{
 		Input.Environment.SerializeCompilationDependencies(Ar);
@@ -2958,6 +2959,7 @@ FShaderCommonCompileJob::FInputHash FShaderCompileJob::GetInputHash()
 	}
 
 	bInputHashSet = true;
+	Input.Hash = InputHash;
 	return InputHash;
 }
 
@@ -3005,10 +3007,11 @@ void FShaderCompileJob::OnComplete()
 			Output.PreprocessTime = PreprocessOutput.ElapsedTime;
 		}
 		// dump debug info for the job at this point if the preprocessed cache is enabled
-		// this is executed in the cooker process and we know the input hash is computed at this point
+		// this ensures we get debug output for all jobs, including those that were found in the job cache,
+		// or matched another in-flight job's hash and so could share its results
 		if (Input.bCachePreprocessed && Input.DumpDebugInfoEnabled())
 		{
-			ShaderFormat->OutputDebugData(LexToString(GetInputHash()), Input, PreprocessOutput, Output);
+			ShaderFormat->OutputDebugData(Input, PreprocessOutput, Output);
 		}
 	}
 }
diff --git a/Engine/Source/Runtime/RenderCore/Public/ShaderCompilerCore.h b/Engine/Source/Runtime/RenderCore/Public/ShaderCompilerCore.h
--- a/Engine/Source/Runtime/RenderCore/Public/ShaderCompilerCore.h
+++ b/Engine/Source/Runtime/RenderCore/Public/ShaderCompilerCore.h
@@ -8,6 +8,7 @@
 
 #include "CoreMinimal.h"
 #include "HAL/FileManager.h"
+#include "Hash/Blake3.h"
 #include "Stats/Stats.h"
 #include "Templates/RefCounting.h"
 #include "Misc/EnumClassFlags.h"
@@ -25,7 +26,7 @@ class FShaderCompileJob;
 class FShaderPipelineCompileJob;
 
 // this is for the protocol, not the data, bump if FShaderCompilerInput or ProcessInputFromArchive changes.
-inline const int32 ShaderCompileWorkerInputVersion = 19;
+inline const int32 ShaderCompileWorkerInputVersion = 20;
 // this is for the protocol, not the data, bump if FShaderCompilerOutput or WriteToOutputArchive changes.
 inline const int32 ShaderCompileWorkerOutputVersion = 14;
 // this is for the protocol, not the data.
@@ -248,6 +249,9 @@ struct FShaderCompilerInput
 	// Description of the configuration used when compiling. 
 	FString DebugDescription;
 
+	// Hash of this input (used as the key for the shader job cache)
+	FBlake3Hash Hash;
+
 	// Compilation Environment
 	FShaderCompilerEnvironment Environment;
 	TRefCountPtr<FSharedShaderCompilerEnvironment> SharedEnvironment;
