Deprecated {Init, Release}DynamicRHI in favor of just {Init, Release}RHI.

diff --git a/Engine/Plugins/FX/Niagara/Source/NiagaraVertexFactories/Private/NiagaraRibbonVertexFactory.cpp b/Engine/Plugins/FX/Niagara/Source/NiagaraVertexFactories/Private/NiagaraRibbonVertexFactory.cpp
--- a/Engine/Plugins/FX/Niagara/Source/NiagaraVertexFactories/Private/NiagaraRibbonVertexFactory.cpp
+++ b/Engine/Plugins/FX/Niagara/Source/NiagaraVertexFactories/Private/NiagaraRibbonVertexFactory.cpp
@@ -95,19 +95,17 @@ public:
 	{
 	}
 
-	virtual void InitDynamicRHI()
+	virtual void InitRHI()
 	{
 		FVertexDeclarationElementList Elements;
 		int32	Offset = 0;
 		FillDeclElements(Elements, Offset);
 
 		// Create the vertex declaration for rendering the factory normally.
-		// This is done in InitDynamicRHI instead of InitRHI to allow FNiagaraRibbonVertexFactory::InitRHI
-		// to rely on it being initialized, since InitDynamicRHI is called before InitRHI.
 		VertexDeclarationRHI = PipelineStateCache::GetOrCreateVertexDeclaration(Elements);
 	}
 
-	virtual void ReleaseDynamicRHI()
+	virtual void ReleaseRHI()
 	{
 		VertexDeclarationRHI.SafeRelease();
 	}
diff --git a/Engine/Plugins/FX/Niagara/Source/NiagaraVertexFactories/Private/NiagaraSpriteVertexFactory.cpp b/Engine/Plugins/FX/Niagara/Source/NiagaraVertexFactories/Private/NiagaraSpriteVertexFactory.cpp
--- a/Engine/Plugins/FX/Niagara/Source/NiagaraVertexFactories/Private/NiagaraSpriteVertexFactory.cpp
+++ b/Engine/Plugins/FX/Niagara/Source/NiagaraVertexFactories/Private/NiagaraSpriteVertexFactory.cpp
@@ -137,20 +137,18 @@ public:
 		Elements.Add(FVertexElement(0, Offset, VET_Float2, 0, InitialStride, false));
 	}
 
-	virtual void InitDynamicRHI()
+	virtual void InitRHI()
 	{
 		FVertexDeclarationElementList Elements;
 		int32	Offset = 0;
 
 		FillDeclElements(Elements, Offset);
 
-		// Create the vertex declaration for rendering the factory normally.
-		// This is done in InitDynamicRHI instead of InitRHI to allow FParticleSpriteVertexFactory::InitRHI
-		// to rely on it being initialized, since InitDynamicRHI is called before InitRHI.
+		// Create the vertex declaration for rendering the factory normally
 		VertexDeclarationRHI = PipelineStateCache::GetOrCreateVertexDeclaration(Elements);
 	}
 
-	virtual void ReleaseDynamicRHI()
+	virtual void ReleaseRHI()
 	{
 		VertexDeclarationRHI.SafeRelease();
 	}
diff --git a/Engine/Plugins/Media/BinkMedia/Source/BinkMediaPlayer/Private/Assets/BinkMediaTextureResource.cpp b/Engine/Plugins/Media/BinkMedia/Source/BinkMediaPlayer/Private/Assets/BinkMediaTextureResource.cpp
--- a/Engine/Plugins/Media/BinkMedia/Source/BinkMediaPlayer/Private/Assets/BinkMediaTextureResource.cpp
+++ b/Engine/Plugins/Media/BinkMedia/Source/BinkMediaPlayer/Private/Assets/BinkMediaTextureResource.cpp
@@ -9,7 +9,7 @@
 #include "RenderingThread.h"
 #include "Rendering/SlateRenderer.h"
 
-void FBinkMediaTextureResource::InitDynamicRHI() 
+void FBinkMediaTextureResource::InitRHI() 
 {
 	int w = Owner->GetSurfaceWidth() > 0 ? Owner->GetSurfaceWidth() : 1;
 	int h = Owner->GetSurfaceHeight() > 0 ? Owner->GetSurfaceHeight() : 1;
@@ -77,9 +77,9 @@ void FBinkMediaTextureResource::InitDynamicRHI()
 	}
 }
 
-void FBinkMediaTextureResource::ReleaseDynamicRHI() 
+void FBinkMediaTextureResource::ReleaseRHI() 
 {
-	ReleaseRHI();
+	FTexture::ReleaseRHI();
 	RenderTargetTextureRHI.SafeRelease();
 	RemoveFromDeferredUpdateList();
 }
diff --git a/Engine/Plugins/Media/BinkMedia/Source/BinkMediaPlayer/Private/Assets/BinkMediaTextureResource.h b/Engine/Plugins/Media/BinkMedia/Source/BinkMediaPlayer/Private/Assets/BinkMediaTextureResource.h
--- a/Engine/Plugins/Media/BinkMedia/Source/BinkMediaPlayer/Private/Assets/BinkMediaTextureResource.h
+++ b/Engine/Plugins/Media/BinkMedia/Source/BinkMediaPlayer/Private/Assets/BinkMediaTextureResource.h
@@ -13,8 +13,8 @@ struct FBinkMediaTextureResource : FTextureResource, FRenderTarget, FDeferredUpd
 	{ 
 	}
 
-	virtual void InitDynamicRHI() override;
-	virtual void ReleaseDynamicRHI() override;
+	virtual void InitRHI() override;
+	virtual void ReleaseRHI() override;
 	virtual FIntPoint GetSizeXY() const override { return Owner->CachedDimensions; }
 	virtual void UpdateDeferredResource(FRHICommandListImmediate& RHICmdList, bool bClearRenderTarget = true) override;
 
diff --git a/Engine/Plugins/Runtime/AppleMoviePlayer/Source/AppleMoviePlayer/Private/AppleMovieStreamer.cpp b/Engine/Plugins/Runtime/AppleMoviePlayer/Source/AppleMoviePlayer/Private/AppleMovieStreamer.cpp
--- a/Engine/Plugins/Runtime/AppleMoviePlayer/Source/AppleMoviePlayer/Private/AppleMovieStreamer.cpp
+++ b/Engine/Plugins/Runtime/AppleMoviePlayer/Source/AppleMoviePlayer/Private/AppleMovieStreamer.cpp
@@ -313,7 +313,7 @@ void FAVPlayerMovieStreamer::Cleanup()
 		(
 			[SlateTextureCleanUp](FRHICommandListImmediate& RHICmdList)
 			{
-				SlateTextureCleanUp->ReleaseDynamicRHI();
+				SlateTextureCleanUp->ReleaseRHI();
 			}
 		);
 		
@@ -604,7 +604,7 @@ bool FAVPlayerMovieStreamer::CheckForNextFrameAndCopy()
 		{
 			if(SlateTexture->IsValid())
 			{
-				SlateTexture->ReleaseDynamicRHI();
+				SlateTexture->ReleaseRHI();
 			}
 			SlateTexture->SetRHIRef(MediaWrappedTexture, Width, Height);
 		}
diff --git a/Engine/Plugins/Runtime/OpenXR/Source/OpenXRHMD/Private/OpenXRHMD.cpp b/Engine/Plugins/Runtime/OpenXR/Source/OpenXRHMD/Private/OpenXRHMD.cpp
--- a/Engine/Plugins/Runtime/OpenXR/Source/OpenXRHMD/Private/OpenXRHMD.cpp
+++ b/Engine/Plugins/Runtime/OpenXR/Source/OpenXRHMD/Private/OpenXRHMD.cpp
@@ -3453,7 +3453,7 @@ void FOpenXRHMD::SetPixelDensity(const float NewDensity)
 	check(IsInGameThread());
 	PipelinedFrameStateGame.PixelDensity = FMath::Min(NewDensity, RuntimePixelDensityMax);
 
-	// We have to update the RT state because the new swapchain will be allocated (FSceneViewport::InitDynamicRHI + AllocateRenderTargetTexture)
+	// We have to update the RT state because the new swapchain will be allocated (FSceneViewport::InitRHI + AllocateRenderTargetTexture)
 	// before we call OnBeginRendering_GameThread.
 	ENQUEUE_RENDER_COMMAND(UpdatePixelDensity)(
 		[this, PixelDensity = PipelinedFrameStateGame.PixelDensity](FRHICommandListImmediate&)
diff --git a/Engine/Plugins/Runtime/WebMMoviePlayer/Source/WebMMoviePlayer/Private/WebMMovieStreamer.cpp b/Engine/Plugins/Runtime/WebMMoviePlayer/Source/WebMMoviePlayer/Private/WebMMovieStreamer.cpp
--- a/Engine/Plugins/Runtime/WebMMoviePlayer/Source/WebMMoviePlayer/Private/WebMMovieStreamer.cpp
+++ b/Engine/Plugins/Runtime/WebMMoviePlayer/Source/WebMMoviePlayer/Private/WebMMovieStreamer.cpp
@@ -255,7 +255,7 @@ bool FWebMMovieStreamer::DisplayFrames(float InDeltaTime)
 
 		if (SlateVideoTexture->IsValid())
 		{
-			SlateVideoTexture->ReleaseDynamicRHI();
+			SlateVideoTexture->ReleaseRHI();
 		}
 
 		SlateVideoTexture->SetRHIRef(WebMSample->GetTextureRef(), WebMSample->GetDim().X, WebMSample->GetDim().Y);
diff --git a/Engine/Plugins/Runtime/nDisplay/Source/DisplayCluster/Private/Render/Viewport/LightCard/DisplayClusterViewportLightCardResource.cpp b/Engine/Plugins/Runtime/nDisplay/Source/DisplayCluster/Private/Render/Viewport/LightCard/DisplayClusterViewportLightCardResource.cpp
--- a/Engine/Plugins/Runtime/nDisplay/Source/DisplayCluster/Private/Render/Viewport/LightCard/DisplayClusterViewportLightCardResource.cpp
+++ b/Engine/Plugins/Runtime/nDisplay/Source/DisplayCluster/Private/Render/Viewport/LightCard/DisplayClusterViewportLightCardResource.cpp
@@ -5,7 +5,7 @@
 ///////////////////////////////////////////////////////////////////////////////////////////////
 // FDisplayClusterViewportLightCardResource
 ///////////////////////////////////////////////////////////////////////////////////////////////
-void FDisplayClusterViewportLightCardResource::InitDynamicRHI()
+void FDisplayClusterViewportLightCardResource::InitRHI()
 {
 	ETextureCreateFlags CreateFlags = TexCreate_Dynamic;
 	CreateFlags |= TexCreate_MultiGPUGraphIgnore;
diff --git a/Engine/Plugins/Runtime/nDisplay/Source/DisplayCluster/Private/Render/Viewport/LightCard/DisplayClusterViewportLightCardResource.h b/Engine/Plugins/Runtime/nDisplay/Source/DisplayCluster/Private/Render/Viewport/LightCard/DisplayClusterViewportLightCardResource.h
--- a/Engine/Plugins/Runtime/nDisplay/Source/DisplayCluster/Private/Render/Viewport/LightCard/DisplayClusterViewportLightCardResource.h
+++ b/Engine/Plugins/Runtime/nDisplay/Source/DisplayCluster/Private/Render/Viewport/LightCard/DisplayClusterViewportLightCardResource.h
@@ -24,7 +24,7 @@ public:
 	virtual uint32 GetSizeY() const override { return LightCardResourceSize.Y; }
 	virtual FIntPoint GetSizeXY() const override { return LightCardResourceSize; }
 
-	virtual void InitDynamicRHI() override;
+	virtual void InitRHI() override;
 
 	virtual FString GetFriendlyName() const override { return TEXT("DisplayClusterViewportLightCardResource"); }
 
diff --git a/Engine/Plugins/Runtime/nDisplay/Source/DisplayCluster/Private/Render/Viewport/RenderTarget/DisplayClusterRenderTargetResource.cpp b/Engine/Plugins/Runtime/nDisplay/Source/DisplayCluster/Private/Render/Viewport/RenderTarget/DisplayClusterRenderTargetResource.cpp
--- a/Engine/Plugins/Runtime/nDisplay/Source/DisplayCluster/Private/Render/Viewport/RenderTarget/DisplayClusterRenderTargetResource.cpp
+++ b/Engine/Plugins/Runtime/nDisplay/Source/DisplayCluster/Private/Render/Viewport/RenderTarget/DisplayClusterRenderTargetResource.cpp
@@ -80,7 +80,7 @@ void FDisplayClusterViewportResource::ImplInitDynamicRHI_TextureResource2D(FText
 //------------------------------------------------------------------------------------------------
 // FDisplayClusterViewportTextureResource
 //------------------------------------------------------------------------------------------------
-void FDisplayClusterViewportTextureResource::InitDynamicRHI()
+void FDisplayClusterViewportTextureResource::InitRHI()
 {
 	FTexture2DRHIRef NewTextureRHI;
 
@@ -100,7 +100,7 @@ void FDisplayClusterViewportTextureResource::InitDynamicRHI()
 //------------------------------------------------------------------------------------------------
 // FDisplayClusterViewportRenderTargetResource
 //------------------------------------------------------------------------------------------------
-void FDisplayClusterViewportRenderTargetResource::InitDynamicRHI()
+void FDisplayClusterViewportRenderTargetResource::InitRHI()
 {
 	// Create RTT and shader resources
 	FTexture2DRHIRef NewTextureRHI;
diff --git a/Engine/Plugins/Runtime/nDisplay/Source/DisplayCluster/Private/Render/Viewport/RenderTarget/DisplayClusterRenderTargetResource.h b/Engine/Plugins/Runtime/nDisplay/Source/DisplayCluster/Private/Render/Viewport/RenderTarget/DisplayClusterRenderTargetResource.h
--- a/Engine/Plugins/Runtime/nDisplay/Source/DisplayCluster/Private/Render/Viewport/RenderTarget/DisplayClusterRenderTargetResource.h
+++ b/Engine/Plugins/Runtime/nDisplay/Source/DisplayCluster/Private/Render/Viewport/RenderTarget/DisplayClusterRenderTargetResource.h
@@ -157,7 +157,7 @@ public:
 	{ }
 
 public:
-	virtual void InitDynamicRHI() override;
+	virtual void InitRHI() override;
 
 	virtual FString GetFriendlyName() const override
 	{
@@ -183,7 +183,7 @@ public:
 	{ }
 
 public:
-	virtual void InitDynamicRHI() override;
+	virtual void InitRHI() override;
 
 	virtual FIntPoint GetSizeXY() const override
 	{
diff --git a/Engine/Plugins/VirtualProduction/TextureShare/Source/TextureShare/Private/Resources/TextureShareResource.cpp b/Engine/Plugins/VirtualProduction/TextureShare/Source/TextureShare/Private/Resources/TextureShareResource.cpp
--- a/Engine/Plugins/VirtualProduction/TextureShare/Source/TextureShare/Private/Resources/TextureShareResource.cpp
+++ b/Engine/Plugins/VirtualProduction/TextureShare/Source/TextureShare/Private/Resources/TextureShareResource.cpp
@@ -165,7 +165,7 @@ bool FTextureShareResource::ReleaseTextureShareHandle_RenderThread()
 }
 
 //////////////////////////////////////////////////////////////////////////////////////////////
-void FTextureShareResource::InitDynamicRHI()
+void FTextureShareResource::InitRHI()
 {
 	FTexture2DRHIRef NewTextureRHI;
 	switch (CoreObject->GetObjectDesc_RenderThread().ProcessDesc.DeviceType)
diff --git a/Engine/Plugins/VirtualProduction/TextureShare/Source/TextureShare/Private/Resources/TextureShareResource.h b/Engine/Plugins/VirtualProduction/TextureShare/Source/TextureShare/Private/Resources/TextureShareResource.h
--- a/Engine/Plugins/VirtualProduction/TextureShare/Source/TextureShare/Private/Resources/TextureShareResource.h
+++ b/Engine/Plugins/VirtualProduction/TextureShare/Source/TextureShare/Private/Resources/TextureShareResource.h
@@ -69,7 +69,7 @@ public:
 	virtual void HandleFrameEnd_RenderThread();
 
 public:
-	virtual void InitDynamicRHI() override;
+	virtual void InitRHI() override;
 
 	// Returns the width of the texture in pixels.
 	virtual uint32 GetSizeX() const override
diff --git a/Engine/Source/Editor/UnrealEd/Private/AssetThumbnail.cpp b/Engine/Source/Editor/UnrealEd/Private/AssetThumbnail.cpp
--- a/Engine/Source/Editor/UnrealEd/Private/AssetThumbnail.cpp
+++ b/Engine/Source/Editor/UnrealEd/Private/AssetThumbnail.cpp
@@ -1236,7 +1236,7 @@ bool FAssetThumbnailPool::LoadThumbnail(TSharedRef<FThumbnailInfo> ThumbnailInfo
 								if (ThumbInfo.ThumbnailTexture->GetTypedResource() != ThumbInfo.ThumbnailRenderTarget->GetTextureRHI())
 								{
 									ThumbInfo.ThumbnailTexture->ClearTextureData();
-									ThumbInfo.ThumbnailTexture->ReleaseDynamicRHI();
+									ThumbInfo.ThumbnailTexture->ReleaseRHI();
 									ThumbInfo.ThumbnailTexture->SetRHIRef(ThumbInfo.ThumbnailRenderTarget->GetTextureRHI(), ThumbInfo.Width, ThumbInfo.Height);
 								}
 							});
diff --git a/Engine/Source/Runtime/Engine/Classes/Engine/SubsurfaceProfile.h b/Engine/Source/Runtime/Engine/Classes/Engine/SubsurfaceProfile.h
--- a/Engine/Source/Runtime/Engine/Classes/Engine/SubsurfaceProfile.h
+++ b/Engine/Source/Runtime/Engine/Classes/Engine/SubsurfaceProfile.h
@@ -217,7 +217,7 @@ public:
 	/**
 	* Release textures when device is lost/destroyed.
 	*/
-	virtual void ReleaseDynamicRHI() override;
+	virtual void ReleaseRHI() override;
 
 	// for debugging, can be removed
 	void Dump();
diff --git a/Engine/Source/Runtime/Engine/Private/ExternalTexture.cpp b/Engine/Source/Runtime/Engine/Private/ExternalTexture.cpp
--- a/Engine/Source/Runtime/Engine/Private/ExternalTexture.cpp
+++ b/Engine/Source/Runtime/Engine/Private/ExternalTexture.cpp
@@ -75,7 +75,7 @@ bool FExternalTextureRegistry::GetExternalTexture(const FMaterialRenderProxy* Ma
 	{
 		ReferencingMaterialRenderProxies.Add(MaterialRenderProxy);
 
-		// Note: FMaterialRenderProxy::ReleaseDynamicRHI()
+		// Note: FMaterialRenderProxy::ReleaseRHI()
 		// is responsible for removing the material proxy
 	}
 
diff --git a/Engine/Source/Runtime/Engine/Private/GPUSkinVertexFactory.cpp b/Engine/Source/Runtime/Engine/Private/GPUSkinVertexFactory.cpp
--- a/Engine/Source/Runtime/Engine/Private/GPUSkinVertexFactory.cpp
+++ b/Engine/Source/Runtime/Engine/Private/GPUSkinVertexFactory.cpp
@@ -783,16 +783,9 @@ void TGPUSkinVertexFactory<BoneInfluenceType>::InitRHI()
 }
 
 template <GPUSkinBoneInfluenceType BoneInfluenceType>
-void TGPUSkinVertexFactory<BoneInfluenceType>::InitDynamicRHI()
+void TGPUSkinVertexFactory<BoneInfluenceType>::ReleaseRHI()
 {
-	FVertexFactory::InitDynamicRHI();
-	//ShaderData.UpdateBoneData(GetFeatureLevel());
-}
-
-template <GPUSkinBoneInfluenceType BoneInfluenceType>
-void TGPUSkinVertexFactory<BoneInfluenceType>::ReleaseDynamicRHI()
-{
-	FVertexFactory::ReleaseDynamicRHI();
+	FVertexFactory::ReleaseRHI();
 	ShaderData.ReleaseBoneData();
 }
 
@@ -824,9 +817,9 @@ TGPUSkinAPEXClothVertexFactory
 -----------------------------------------------------------------------------*/
 
 template <GPUSkinBoneInfluenceType BoneInfluenceType>
-void TGPUSkinAPEXClothVertexFactory<BoneInfluenceType>::ReleaseDynamicRHI()
+void TGPUSkinAPEXClothVertexFactory<BoneInfluenceType>::ReleaseRHI()
 {
-	Super::ReleaseDynamicRHI();
+	Super::ReleaseRHI();
 	ClothShaderData.ReleaseClothSimulData();
 
 	// Release the RHIResource reference held in FGPUSkinAPEXClothDataType
diff --git a/Engine/Source/Runtime/Engine/Private/LocalVertexFactory.cpp b/Engine/Source/Runtime/Engine/Private/LocalVertexFactory.cpp
--- a/Engine/Source/Runtime/Engine/Private/LocalVertexFactory.cpp
+++ b/Engine/Source/Runtime/Engine/Private/LocalVertexFactory.cpp
@@ -29,16 +29,16 @@ class FSpeedTreeWindNullUniformBuffer : public TUniformBuffer<FSpeedTreeUniformP
 {
 	typedef TUniformBuffer< FSpeedTreeUniformParameters > Super;
 public:
-	virtual void InitDynamicRHI() override;
+	virtual void InitRHI() override;
 };
 
-void FSpeedTreeWindNullUniformBuffer::InitDynamicRHI()
+void FSpeedTreeWindNullUniformBuffer::InitRHI()
 {
 	FSpeedTreeUniformParameters Parameters;
 	FMemory::Memzero(Parameters);
 	SetContentsNoUpdate(Parameters);
 	
-	Super::InitDynamicRHI();
+	Super::InitRHI();
 }
 
 static TGlobalResource< FSpeedTreeWindNullUniformBuffer > GSpeedTreeWindNullUniformBuffer;
diff --git a/Engine/Source/Runtime/Engine/Private/Materials/MaterialRenderProxy.cpp b/Engine/Source/Runtime/Engine/Private/Materials/MaterialRenderProxy.cpp
--- a/Engine/Source/Runtime/Engine/Private/Materials/MaterialRenderProxy.cpp
+++ b/Engine/Source/Runtime/Engine/Private/Materials/MaterialRenderProxy.cpp
@@ -558,7 +558,7 @@ FMaterialRenderProxy::~FMaterialRenderProxy()
 	DeletedFlag = -1;
 }
 
-void FMaterialRenderProxy::InitDynamicRHI()
+void FMaterialRenderProxy::InitRHI()
 {
 #if WITH_EDITOR
 	// MaterialRenderProxyMap is only used by shader compiling
@@ -575,7 +575,7 @@ void FMaterialRenderProxy::CancelCacheUniformExpressions()
 	DeferredUniformExpressionCacheRequests.Remove(this);
 }
 
-void FMaterialRenderProxy::ReleaseDynamicRHI()
+void FMaterialRenderProxy::ReleaseRHI()
 {
 #if WITH_EDITOR
 	if (!FPlatformProperties::RequiresCookedData())
diff --git a/Engine/Source/Runtime/Engine/Private/Particles/ParticleBeamTrailVertexFactory.cpp b/Engine/Source/Runtime/Engine/Private/Particles/ParticleBeamTrailVertexFactory.cpp
--- a/Engine/Source/Runtime/Engine/Private/Particles/ParticleBeamTrailVertexFactory.cpp
+++ b/Engine/Source/Runtime/Engine/Private/Particles/ParticleBeamTrailVertexFactory.cpp
@@ -76,19 +76,19 @@ public:
 		Elements.Add(FVertexElement(1, 0, VET_Float4, 5, bUsesDynamicParameter ? sizeof(FVector4f) : 0));
 	}
 
-	virtual void InitDynamicRHI()
+	virtual void InitRHI()
 	{
 		FVertexDeclarationElementList Elements;
 		int32	Offset = 0;
 		FillDeclElements(Elements, Offset);
 
 		// Create the vertex declaration for rendering the factory normally.
-		// This is done in InitDynamicRHI instead of InitRHI to allow FParticleBeamTrailVertexFactory::InitRHI
-		// to rely on it being initialized, since InitDynamicRHI is called before InitRHI.
+		// This is done in InitRHI instead of InitRHI to allow FParticleBeamTrailVertexFactory::InitRHI
+		// to rely on it being initialized, since InitRHI is called before InitRHI.
 		VertexDeclarationRHI = PipelineStateCache::GetOrCreateVertexDeclaration(Elements);
 	}
 
-	virtual void ReleaseDynamicRHI()
+	virtual void ReleaseRHI()
 	{
 		VertexDeclarationRHI.SafeRelease();
 	}
diff --git a/Engine/Source/Runtime/Engine/Private/Particles/ParticleVertexFactory.cpp b/Engine/Source/Runtime/Engine/Private/Particles/ParticleVertexFactory.cpp
--- a/Engine/Source/Runtime/Engine/Private/Particles/ParticleVertexFactory.cpp
+++ b/Engine/Source/Runtime/Engine/Private/Particles/ParticleVertexFactory.cpp
@@ -178,7 +178,7 @@ public:
 		Offset += sizeof(float) * 4;
 	}
 
-	virtual void InitDynamicRHI()
+	virtual void InitRHI()
 	{
 		FVertexDeclarationElementList Elements;
 		int32	Offset = 0;
@@ -186,12 +186,10 @@ public:
 		FillDeclElements(Elements, Offset);
 
 		// Create the vertex declaration for rendering the factory normally.
-		// This is done in InitDynamicRHI instead of InitRHI to allow FParticleSpriteVertexFactory::InitRHI
-		// to rely on it being initialized, since InitDynamicRHI is called before InitRHI.
 		VertexDeclarationRHI = PipelineStateCache::GetOrCreateVertexDeclaration(Elements);
 	}
 
-	virtual void ReleaseDynamicRHI()
+	virtual void ReleaseRHI()
 	{
 		VertexDeclarationRHI.SafeRelease();
 	}
diff --git a/Engine/Source/Runtime/Engine/Private/Rendering/SpecularProfile.cpp b/Engine/Source/Runtime/Engine/Private/Rendering/SpecularProfile.cpp
--- a/Engine/Source/Runtime/Engine/Private/Rendering/SpecularProfile.cpp
+++ b/Engine/Source/Runtime/Engine/Private/Rendering/SpecularProfile.cpp
@@ -120,7 +120,7 @@ public:
 	/**
 		* Release textures when device is lost/destroyed.
 		*/
-	virtual void ReleaseDynamicRHI() override;
+	virtual void ReleaseRHI() override;
 
 	// @param InProfile must not be 0, game thread pointer, do not dereference, only for comparison
 	// @return INDEX_NONE if not found
@@ -388,7 +388,7 @@ IPooledRenderTarget* FSpecularProfileTextureManager::GetAtlasTexture(FRDGBuilder
 	return GSpecularProfileTextureAtlas;
 }
 
-void FSpecularProfileTextureManager::ReleaseDynamicRHI()
+void FSpecularProfileTextureManager::ReleaseRHI()
 {
 	GSpecularProfileTextureAtlas.SafeRelease();
 }
diff --git a/Engine/Source/Runtime/Engine/Private/Rendering/SubsurfaceProfile.cpp b/Engine/Source/Runtime/Engine/Private/Rendering/SubsurfaceProfile.cpp
--- a/Engine/Source/Runtime/Engine/Private/Rendering/SubsurfaceProfile.cpp
+++ b/Engine/Source/Runtime/Engine/Private/Rendering/SubsurfaceProfile.cpp
@@ -303,7 +303,7 @@ IPooledRenderTarget* FSubsurfaceProfileTexture::GetSSProfilesPreIntegratedTextur
 	return GSSProfilesPreIntegratedTexture;
 }
 
-void FSubsurfaceProfileTexture::ReleaseDynamicRHI()
+void FSubsurfaceProfileTexture::ReleaseRHI()
 {
 	GSSProfiles.SafeRelease();
 	GSSProfilesPreIntegratedTexture.SafeRelease();
diff --git a/Engine/Source/Runtime/Engine/Private/SceneManagement.cpp b/Engine/Source/Runtime/Engine/Private/SceneManagement.cpp
--- a/Engine/Source/Runtime/Engine/Private/SceneManagement.cpp
+++ b/Engine/Source/Runtime/Engine/Private/SceneManagement.cpp
@@ -1221,12 +1221,12 @@ void GetLightmapClusterResourceParameters(
 	}
 }
 
-void FDefaultLightmapResourceClusterUniformBuffer::InitDynamicRHI()
+void FDefaultLightmapResourceClusterUniformBuffer::InitRHI()
 {
 	FLightmapResourceClusterShaderParameters Parameters;
 	GetLightmapClusterResourceParameters(GMaxRHIFeatureLevel, FLightmapClusterResourceInput(), nullptr, Parameters);
 	SetContents(Parameters);
-	Super::InitDynamicRHI();
+	Super::InitRHI();
 }
 
 /** Global uniform buffer containing the default precomputed lighting data. */
@@ -1501,14 +1501,14 @@ bool FMeshBatch::Validate(const FPrimitiveSceneProxy* PrimitiveSceneProxy, ERHIF
 
 IMPLEMENT_GLOBAL_SHADER_PARAMETER_STRUCT(FMobileReflectionCaptureShaderParameters, "MobileReflectionCapture");
 
-void FDefaultMobileReflectionCaptureUniformBuffer::InitDynamicRHI()
+void FDefaultMobileReflectionCaptureUniformBuffer::InitRHI()
 {
 	FMobileReflectionCaptureShaderParameters Parameters;
 	Parameters.Params = FVector4f(1.f, 0.f, 0.f, 0.f);
 	Parameters.Texture = GBlackTextureCube->TextureRHI;
 	Parameters.TextureSampler = GBlackTextureCube->SamplerStateRHI;
 	SetContents(Parameters);
-	Super::InitDynamicRHI();
+	Super::InitRHI();
 }
 
 /** Global uniform buffer containing the default reflection data used in mobile renderer. */
diff --git a/Engine/Source/Runtime/Engine/Private/SkeletalRenderGPUSkin.cpp b/Engine/Source/Runtime/Engine/Private/SkeletalRenderGPUSkin.cpp
--- a/Engine/Source/Runtime/Engine/Private/SkeletalRenderGPUSkin.cpp
+++ b/Engine/Source/Runtime/Engine/Private/SkeletalRenderGPUSkin.cpp
@@ -84,7 +84,7 @@ static int32 GetRayTracingSkeletalMeshGlobalLODBias()
 FMorphVertexBuffer
 -----------------------------------------------------------------------------*/
 
-void FMorphVertexBuffer::InitDynamicRHI()
+void FMorphVertexBuffer::InitRHI()
 {
 	// LOD of the skel mesh is used to find number of vertices in buffer
 	FSkeletalMeshLODRenderData& LodData = SkelMeshRenderData->LODRenderData[LODIdx];
@@ -126,7 +126,7 @@ void FMorphVertexBuffer::InitDynamicRHI()
 	bHasBeenUpdated = false;
 }
 
-void FMorphVertexBuffer::ReleaseDynamicRHI()
+void FMorphVertexBuffer::ReleaseRHI()
 {
 	UAVValue.SafeRelease();
 	VertexBufferRHI.SafeRelease();
diff --git a/Engine/Source/Runtime/Engine/Private/SkeletalRenderGPUSkin.h b/Engine/Source/Runtime/Engine/Private/SkeletalRenderGPUSkin.h
--- a/Engine/Source/Runtime/Engine/Private/SkeletalRenderGPUSkin.h
+++ b/Engine/Source/Runtime/Engine/Private/SkeletalRenderGPUSkin.h
@@ -210,12 +210,12 @@ public:
 	/** 
 	* Initialize the dynamic RHI for this rendering resource 
 	*/
-	virtual void InitDynamicRHI();
+	virtual void InitRHI();
 
 	/** 
 	* Release the dynamic RHI for this rendering resource 
 	*/
-	virtual void ReleaseDynamicRHI();
+	virtual void ReleaseRHI();
 
 	inline void RecreateResourcesIfRequired(bool bInUsesComputeShader)
 	{
@@ -231,7 +231,7 @@ public:
 	virtual FString GetFriendlyName() const { return TEXT("Morph target mesh vertices"); }
 
 	/**
-	 * Get Resource Size : mostly copied from InitDynamicRHI - how much they allocate when initialize
+	 * Get Resource Size : mostly copied from InitRHI - how much they allocate when initialize
 	 */
 	SIZE_T GetResourceSize() const
 	{
@@ -260,7 +260,7 @@ public:
 	/** Has been updated or not by UpdateMorphVertexBuffer**/
 	bool bHasBeenUpdated;
 
-	/** DX12 cannot clear the buffer in InitDynamicRHI with UAV flag enables, we should really have a Zero initzialized flag instead**/
+	/** DX12 cannot clear the buffer in InitRHI with UAV flag enables, we should really have a Zero initzialized flag instead**/
 	bool bNeedsInitialClear;
 
 	// @param guaranteed only to be valid if the vertex buffer is valid
diff --git a/Engine/Source/Runtime/Engine/Private/Slate/SceneViewport.cpp b/Engine/Source/Runtime/Engine/Private/Slate/SceneViewport.cpp
--- a/Engine/Source/Runtime/Engine/Private/Slate/SceneViewport.cpp
+++ b/Engine/Source/Runtime/Engine/Private/Slate/SceneViewport.cpp
@@ -2055,7 +2055,7 @@ void FSceneViewport::OnPostResizeWindowBackbuffer(void* Backbuffer)
 	}
 }
 
-void FSceneViewport::InitDynamicRHI()
+void FSceneViewport::InitRHI()
 {
 	if(bRequiresHitProxyStorage)
 	{
@@ -2191,9 +2191,9 @@ void FSceneViewport::InitDynamicRHI()
 	}
 }
 
-void FSceneViewport::ReleaseDynamicRHI()
+void FSceneViewport::ReleaseRHI()
 {
-	FViewport::ReleaseDynamicRHI();
+	FViewport::ReleaseRHI();
 
 	ViewportRHI.SafeRelease();
 
@@ -2203,12 +2203,12 @@ void FSceneViewport::ReleaseDynamicRHI()
 	{
 		if (BufferedSlateHandles[i])
 		{
-			BufferedSlateHandles[i]->ReleaseDynamicRHI();
+			BufferedSlateHandles[i]->ReleaseRHI();
 		}
 	}
 	if (RenderThreadSlateTexture)
 	{
-		RenderThreadSlateTexture->ReleaseDynamicRHI();
+		RenderThreadSlateTexture->ReleaseRHI();
 	}
 }
 
diff --git a/Engine/Source/Runtime/Engine/Private/Slate/SlateTextures.cpp b/Engine/Source/Runtime/Engine/Private/Slate/SlateTextures.cpp
--- a/Engine/Source/Runtime/Engine/Private/Slate/SlateTextures.cpp
+++ b/Engine/Source/Runtime/Engine/Private/Slate/SlateTextures.cpp
@@ -39,7 +39,7 @@ void FSlateTexture2DRHIRef::Cleanup()
 	BeginCleanup(this);
 }
 
-void FSlateTexture2DRHIRef::InitDynamicRHI()
+void FSlateTexture2DRHIRef::InitRHI()
 {
 	SCOPED_LOADTIMER(FSlateTexture2DRHIRef_InitDynamicRHI);
 
@@ -97,7 +97,7 @@ void FSlateTexture2DRHIRef::InitDynamicRHI()
 	}
 }
 
-void FSlateTexture2DRHIRef::ReleaseDynamicRHI()
+void FSlateTexture2DRHIRef::ReleaseRHI()
 {
 	check( IsInRenderingThread() );
 
@@ -315,7 +315,7 @@ void FSlateTextureRenderTarget2DResource::ClampSize(int32 MaxSizeX,int32 MaxSize
 	}	
 }
 
-void FSlateTextureRenderTarget2DResource::InitDynamicRHI()
+void FSlateTextureRenderTarget2DResource::InitRHI()
 {
 	SCOPED_LOADTIMER(FSlateTextureRenderTarget2DResource_InitDynamicRHI);
 
@@ -349,12 +349,12 @@ void FSlateTextureRenderTarget2DResource::InitDynamicRHI()
 	SamplerStateRHI = GetOrCreateSamplerState( SamplerStateInitializer );
 }
 
-void FSlateTextureRenderTarget2DResource::ReleaseDynamicRHI()
+void FSlateTextureRenderTarget2DResource::ReleaseRHI()
 {
 	check(IsInRenderingThread());
 
 	// Release the FTexture RHI resources here as well
-	ReleaseRHI();
+	FTexture::ReleaseRHI();
 
 	RenderTargetTextureRHI.SafeRelease();
 
diff --git a/Engine/Source/Runtime/Engine/Private/TextureRenderTarget2D.cpp b/Engine/Source/Runtime/Engine/Private/TextureRenderTarget2D.cpp
--- a/Engine/Source/Runtime/Engine/Private/TextureRenderTarget2D.cpp
+++ b/Engine/Source/Runtime/Engine/Private/TextureRenderTarget2D.cpp
@@ -622,7 +622,7 @@ ETextureCreateFlags FTextureRenderTarget2DResource::GetCreateFlags()
  * Called when the resource is initialized, or when reseting all RHI resources.
  * This is only called by the rendering thread.
  */
-void FTextureRenderTarget2DResource::InitDynamicRHI()
+void FTextureRenderTarget2DResource::InitRHI()
 {
 	LLM_SCOPED_TAG_WITH_OBJECT_IN_SET(Owner->GetOutermost(), ELLMTagSet::Assets);
 
@@ -685,10 +685,10 @@ void FTextureRenderTarget2DResource::InitDynamicRHI()
  * Called when the resource is released, or when reseting all RHI resources.
  * This is only called by the rendering thread.
  */
-void FTextureRenderTarget2DResource::ReleaseDynamicRHI()
+void FTextureRenderTarget2DResource::ReleaseRHI()
 {
 	// release the FTexture RHI resources here as well
-	ReleaseRHI();
+	FTexture::ReleaseRHI();
 
 	RHIUpdateTextureReference(Owner->TextureReference.TextureReferenceRHI, nullptr);
 	RenderTargetTextureRHI.SafeRelease();
diff --git a/Engine/Source/Runtime/Engine/Private/TextureRenderTarget2DArray.cpp b/Engine/Source/Runtime/Engine/Private/TextureRenderTarget2DArray.cpp
--- a/Engine/Source/Runtime/Engine/Private/TextureRenderTarget2DArray.cpp
+++ b/Engine/Source/Runtime/Engine/Private/TextureRenderTarget2DArray.cpp
@@ -221,7 +221,7 @@ UTexture2DArray* UTextureRenderTarget2DArray::ConstructTexture2DArray(UObject* O
  * Called when the resource is initialized, or when reseting all RHI resources.
  * This is only called by the rendering thread.
  */
-void FTextureRenderTarget2DArrayResource::InitDynamicRHI()
+void FTextureRenderTarget2DArrayResource::InitRHI()
 {
 	LLM_SCOPED_TAG_WITH_OBJECT_IN_SET(Owner->GetOutermost(), ELLMTagSet::Assets);
 
@@ -286,10 +286,10 @@ void FTextureRenderTarget2DArrayResource::InitDynamicRHI()
  * Called when the resource is released, or when reseting all RHI resources.
  * This is only called by the rendering thread.
  */
-void FTextureRenderTarget2DArrayResource::ReleaseDynamicRHI()
+void FTextureRenderTarget2DArrayResource::ReleaseRHI()
 {
 	// release the FTexture RHI resources here as well
-	ReleaseRHI();
+	FTexture::ReleaseRHI();
 
 	RHIUpdateTextureReference(Owner->TextureReference.TextureReferenceRHI, nullptr);
 	RenderTargetTextureRHI.SafeRelease();
diff --git a/Engine/Source/Runtime/Engine/Private/TextureRenderTargetCube.cpp b/Engine/Source/Runtime/Engine/Private/TextureRenderTargetCube.cpp
--- a/Engine/Source/Runtime/Engine/Private/TextureRenderTargetCube.cpp
+++ b/Engine/Source/Runtime/Engine/Private/TextureRenderTargetCube.cpp
@@ -250,7 +250,7 @@ UTextureCube* UTextureRenderTargetCube::ConstructTextureCube(
  * Called when the resource is initialized, or when reseting all RHI resources.
  * This is only called by the rendering thread.
  */
-void FTextureRenderTargetCubeResource::InitDynamicRHI()
+void FTextureRenderTargetCubeResource::InitRHI()
 {
 	LLM_SCOPED_TAG_WITH_OBJECT_IN_SET(Owner->GetOutermost(), ELLMTagSet::Assets);
 
@@ -325,10 +325,10 @@ void FTextureRenderTargetCubeResource::InitDynamicRHI()
  * Called when the resource is released, or when reseting all RHI resources.
  * This is only called by the rendering thread.
  */
-void FTextureRenderTargetCubeResource::ReleaseDynamicRHI()
+void FTextureRenderTargetCubeResource::ReleaseRHI()
 {
 	// release the FTexture RHI resources here as well
-	ReleaseRHI();
+	FTexture::ReleaseRHI();
 
 	RHIUpdateTextureReference(Owner->TextureReference.TextureReferenceRHI, nullptr);
 	RenderTargetTextureRHI.SafeRelease();
diff --git a/Engine/Source/Runtime/Engine/Private/TextureRenderTargetVolume.cpp b/Engine/Source/Runtime/Engine/Private/TextureRenderTargetVolume.cpp
--- a/Engine/Source/Runtime/Engine/Private/TextureRenderTargetVolume.cpp
+++ b/Engine/Source/Runtime/Engine/Private/TextureRenderTargetVolume.cpp
@@ -227,7 +227,7 @@ UVolumeTexture* UTextureRenderTargetVolume::ConstructTextureVolume(UObject* ObjO
  * Called when the resource is initialized, or when reseting all RHI resources.
  * This is only called by the rendering thread.
  */
-void FTextureRenderTargetVolumeResource::InitDynamicRHI()
+void FTextureRenderTargetVolumeResource::InitRHI()
 {
 	LLM_SCOPED_TAG_WITH_OBJECT_IN_SET(Owner->GetOutermost(), ELLMTagSet::Assets);
 
@@ -291,10 +291,10 @@ void FTextureRenderTargetVolumeResource::InitDynamicRHI()
  * Called when the resource is released, or when reseting all RHI resources.
  * This is only called by the rendering thread.
  */
-void FTextureRenderTargetVolumeResource::ReleaseDynamicRHI()
+void FTextureRenderTargetVolumeResource::ReleaseRHI()
 {
 	// release the FTexture RHI resources here as well
-	ReleaseRHI();
+	FTexture::ReleaseRHI();
 
 	RHIUpdateTextureReference(Owner->TextureReference.TextureReferenceRHI, nullptr);
 	RenderTargetTextureRHI.SafeRelease();
diff --git a/Engine/Source/Runtime/Engine/Private/UnrealClient.cpp b/Engine/Source/Runtime/Engine/Private/UnrealClient.cpp
--- a/Engine/Source/Runtime/Engine/Private/UnrealClient.cpp
+++ b/Engine/Source/Runtime/Engine/Private/UnrealClient.cpp
@@ -2280,42 +2280,19 @@ void FViewport::SetViewportClient( FViewportClient* InViewportClient )
 	ViewportClient = InViewportClient;
 }
 
-void FViewport::InitDynamicRHI()
-{
-	UpdateRenderTargetSurfaceRHIToCurrentBackBuffer();
-
-	if(bRequiresHitProxyStorage)
-	{
-		// Initialize the hit proxy map.
-		HitProxyMap.Init(SizeX,SizeY);
-	}
-}
-
-void FViewport::ReleaseDynamicRHI()
+void FViewport::ReleaseRHI()
 {
 	HitProxyMap.Release();
 	RenderTargetTextureRHI.SafeRelease();
-}
-
-void FViewport::ReleaseRHI()
-{
-	FlushRenderingCommands();
 	ViewportRHI.SafeRelease();
 }
 
 void FViewport::InitRHI()
 {
-	FlushRenderingCommands();
-	if(!IsValidRef(ViewportRHI))
+	if(bRequiresHitProxyStorage)
 	{
-		ViewportRHI = RHICreateViewport(
-			GetWindow(),
-			SizeX,
-			SizeY,
-			IsFullscreen(),
-			EPixelFormat::PF_Unknown
-			);
-		UpdateRenderTargetSurfaceRHIToCurrentBackBuffer();
+		// Initialize the hit proxy map.
+		HitProxyMap.Init(SizeX,SizeY);
 	}
 }
 
@@ -2677,7 +2654,7 @@ FDummyViewport::~FDummyViewport()
 	}
 }
 
-void FDummyViewport::InitDynamicRHI()
+void FDummyViewport::InitRHI()
 {
 	EPixelFormat DummyViewportFormat = bSceneHDREnabled ? GRHIHDRDisplayOutputFormat : PF_A2B10G10R10;
 	const FRHITextureCreateDesc Desc =
diff --git a/Engine/Source/Runtime/Engine/Public/DummyViewport.h b/Engine/Source/Runtime/Engine/Public/DummyViewport.h
--- a/Engine/Source/Runtime/Engine/Public/DummyViewport.h
+++ b/Engine/Source/Runtime/Engine/Public/DummyViewport.h
@@ -56,10 +56,9 @@ public:
 	//~ End FViewport Interface
 
 	//~ Begin FRenderResource Interface
-	virtual void InitDynamicRHI() override;
+	virtual void InitRHI() override;
 
 	// @todo DLL: Without these functions we get unresolved linker errors with FRenderResource
-	virtual void InitRHI() override{}
 	virtual void ReleaseRHI() override{}
 	virtual void InitResource() override{ FViewport::InitResource(); }
 	virtual void ReleaseResource() override { FViewport::ReleaseResource(); }
diff --git a/Engine/Source/Runtime/Engine/Public/GPUSkinVertexFactory.h b/Engine/Source/Runtime/Engine/Public/GPUSkinVertexFactory.h
--- a/Engine/Source/Runtime/Engine/Public/GPUSkinVertexFactory.h
+++ b/Engine/Source/Runtime/Engine/Public/GPUSkinVertexFactory.h
@@ -473,8 +473,7 @@ public:
 
 	// FRenderResource interface.
 	virtual void InitRHI() override;
-	virtual void InitDynamicRHI() override;
-	virtual void ReleaseDynamicRHI() override;
+	virtual void ReleaseRHI() override;
 
 protected:
 	/**
@@ -681,7 +680,7 @@ public:
 	 */
 	virtual ~TGPUSkinAPEXClothVertexFactory() override
 	{
-		checkf(!ClothDataPtr->ClothBuffer.IsValid(), TEXT("ClothBuffer RHI resource should have been released in ReleaseDynamicRHI"));
+		checkf(!ClothDataPtr->ClothBuffer.IsValid(), TEXT("ClothBuffer RHI resource should have been released in ReleaseRHI"));
 		delete ClothDataPtr;
 		(void)this->Data.Release();
 	}
@@ -723,7 +722,7 @@ public:
 	* initializes the device resource
 	*/
 	virtual void InitRHI() override;
-	virtual void ReleaseDynamicRHI() override;
+	virtual void ReleaseRHI() override;
 
 protected:
 	/** Alias pointer to TUniquePtr<FGPUSkinDataType> Data of FGPUBaseSkinVertexFactory. Note memory isn't managed through this pointer. */
diff --git a/Engine/Source/Runtime/Engine/Public/Materials/MaterialRenderProxy.h b/Engine/Source/Runtime/Engine/Public/Materials/MaterialRenderProxy.h
--- a/Engine/Source/Runtime/Engine/Public/Materials/MaterialRenderProxy.h
+++ b/Engine/Source/Runtime/Engine/Public/Materials/MaterialRenderProxy.h
@@ -187,8 +187,8 @@ public:
 	}
 
 	// FRenderResource interface.
-	virtual void InitDynamicRHI() override;
-	virtual void ReleaseDynamicRHI() override;
+	virtual void InitRHI() override;
+	virtual void ReleaseRHI() override;
 	virtual void ReleaseResource() override;
 
 #if WITH_EDITOR
diff --git a/Engine/Source/Runtime/Engine/Public/PlanarReflectionSceneProxy.h b/Engine/Source/Runtime/Engine/Public/PlanarReflectionSceneProxy.h
--- a/Engine/Source/Runtime/Engine/Public/PlanarReflectionSceneProxy.h
+++ b/Engine/Source/Runtime/Engine/Public/PlanarReflectionSceneProxy.h
@@ -26,7 +26,7 @@ public:
 		Size(InSize)
 	{}
 
-	virtual void InitDynamicRHI()
+	virtual void InitRHI()
 	{
 		// Create the sampler state RHI resource.
 		FSamplerStateInitializerRHI SamplerStateInitializer
diff --git a/Engine/Source/Runtime/Engine/Public/SceneManagement.h b/Engine/Source/Runtime/Engine/Public/SceneManagement.h
--- a/Engine/Source/Runtime/Engine/Public/SceneManagement.h
+++ b/Engine/Source/Runtime/Engine/Public/SceneManagement.h
@@ -815,7 +815,7 @@ class FDefaultLightmapResourceClusterUniformBuffer : public TUniformBuffer< FLig
 {
 	typedef TUniformBuffer< FLightmapResourceClusterShaderParameters > Super;
 public:
-	virtual void InitDynamicRHI() override;
+	virtual void InitRHI() override;
 };
 
 ENGINE_API extern TGlobalResource< FDefaultLightmapResourceClusterUniformBuffer > GDefaultLightmapResourceClusterUniformBuffer;
@@ -1110,7 +1110,7 @@ class FDefaultMobileReflectionCaptureUniformBuffer : public TUniformBuffer<FMobi
 {
 	typedef TUniformBuffer<FMobileReflectionCaptureShaderParameters> Super;
 public:
-	virtual void InitDynamicRHI() override;
+	virtual void InitRHI() override;
 };
 
 ENGINE_API extern TGlobalResource<FDefaultMobileReflectionCaptureUniformBuffer> GDefaultMobileReflectionCaptureUniformBuffer;
diff --git a/Engine/Source/Runtime/Engine/Public/Slate/SceneViewport.h b/Engine/Source/Runtime/Engine/Public/Slate/SceneViewport.h
--- a/Engine/Source/Runtime/Engine/Public/Slate/SceneViewport.h
+++ b/Engine/Source/Runtime/Engine/Public/Slate/SceneViewport.h
@@ -284,12 +284,10 @@ private:
 	void Destroy() override;
 
 	// FRenderResource interface.
-	virtual void InitDynamicRHI() override;
-	virtual void ReleaseDynamicRHI() override;
+	virtual void InitRHI() override;
+	virtual void ReleaseRHI() override;
 
 	// @todo DLL: Without these functions we get unresolved linker errors with FRenderResource
-	virtual void InitRHI() override {}
-	virtual void ReleaseRHI() override {}
 	virtual void InitResource() override { FViewport::InitResource(); }
 	virtual void ReleaseResource() override { FViewport::ReleaseResource(); }
 	virtual FString GetFriendlyName() const override { return FString(TEXT("FSlateSceneViewport"));}
diff --git a/Engine/Source/Runtime/Engine/Public/Slate/SlateTextures.h b/Engine/Source/Runtime/Engine/Public/Slate/SlateTextures.h
--- a/Engine/Source/Runtime/Engine/Public/Slate/SlateTextures.h
+++ b/Engine/Source/Runtime/Engine/Public/Slate/SlateTextures.h
@@ -31,10 +31,10 @@ public:
 	virtual uint32 GetHeight() const override { return Height; }
 
 	/** FRenderResource Interface.  Called when render resources need to be initialized */
-	virtual void InitDynamicRHI() override;
+	virtual void InitRHI() override;
 
 	/** FRenderResource Interface.  Called when render resources need to be released */
-	virtual void ReleaseDynamicRHI() override;
+	virtual void ReleaseRHI() override;
 
 	/**
 	 * Resize the texture.  Can only be called on the render thread
@@ -99,7 +99,7 @@ private:
 
 	/** Texture creation flags for if this texture needs to be recreated dynamically */
 	ETextureCreateFlags TexCreateFlags;
-	/** Data used between ReleaseDynamicRHI and InitDynamicRHI.  May be null if the data is not used */
+	/** Data used between ReleaseRHI and InitRHI.  May be null if the data is not used */
 	TSharedPtr<FSlateTextureData, ESPMode::ThreadSafe> TextureData;
 	/** Pixel format of the texture */
 	EPixelFormat PixelFormat;
@@ -120,12 +120,12 @@ public:
 	{
 	}
 
-	virtual void InitDynamicRHI() override {}
+	virtual void InitRHI() override {}
 
 	/** 
 	 * Releases all dynamic RHI data
 	 */
-	virtual void ReleaseDynamicRHI() override
+	virtual void ReleaseRHI() override
 	{
 		ShaderResource.SafeRelease();
 	}
@@ -161,8 +161,8 @@ public:
 	virtual void ClampSize(int32 SizeX,int32 SizeY) override;
 
 	// FRenderResource implementation
-	virtual void InitDynamicRHI() override;
-	virtual void ReleaseDynamicRHI() override;
+	virtual void InitRHI() override;
+	virtual void ReleaseRHI() override;
 
 	// FRenderTarget interface
 	virtual uint32 GetSizeX() const override;
diff --git a/Engine/Source/Runtime/Engine/Public/TextureRenderTarget2DArrayResource.h b/Engine/Source/Runtime/Engine/Public/TextureRenderTarget2DArrayResource.h
--- a/Engine/Source/Runtime/Engine/Public/TextureRenderTarget2DArrayResource.h
+++ b/Engine/Source/Runtime/Engine/Public/TextureRenderTarget2DArrayResource.h
@@ -36,7 +36,7 @@ public:
 	 * Resources that need to initialize after a D3D device reset must implement this function.
 	 * This is only called by the rendering thread.
 	 */
-	virtual void InitDynamicRHI() override;
+	virtual void InitRHI() override;
 
 	/**
 	 * Releases the dynamic RHI resource and/or RHI render target resources used by this resource.
@@ -44,7 +44,7 @@ public:
 	 * Resources that need to release before a D3D device reset must implement this function.
 	 * This is only called by the rendering thread.
 	 */
-	virtual void ReleaseDynamicRHI() override;
+	virtual void ReleaseRHI() override;
 
 	// FRenderTarget interface.
 
diff --git a/Engine/Source/Runtime/Engine/Public/TextureRenderTargetVolumeResource.h b/Engine/Source/Runtime/Engine/Public/TextureRenderTargetVolumeResource.h
--- a/Engine/Source/Runtime/Engine/Public/TextureRenderTargetVolumeResource.h
+++ b/Engine/Source/Runtime/Engine/Public/TextureRenderTargetVolumeResource.h
@@ -36,7 +36,7 @@ public:
 	 * Resources that need to initialize after a D3D device reset must implement this function.
 	 * This is only called by the rendering thread.
 	 */
-	virtual void InitDynamicRHI() override;
+	virtual void InitRHI() override;
 
 	/**
 	 * Releases the dynamic RHI resource and/or RHI render target resources used by this resource.
@@ -44,7 +44,7 @@ public:
 	 * Resources that need to release before a D3D device reset must implement this function.
 	 * This is only called by the rendering thread.
 	 */
-	virtual void ReleaseDynamicRHI() override;
+	virtual void ReleaseRHI() override;
 
 	// FRenderTarget interface.
 
diff --git a/Engine/Source/Runtime/Engine/Public/TextureResource.h b/Engine/Source/Runtime/Engine/Public/TextureResource.h
--- a/Engine/Source/Runtime/Engine/Public/TextureResource.h
+++ b/Engine/Source/Runtime/Engine/Public/TextureResource.h
@@ -445,7 +445,7 @@ public:
 	 * Resources that need to initialize after a D3D device reset must implement this function.
 	 * This is only called by the rendering thread.
 	 */
-	virtual void InitDynamicRHI() override;
+	virtual void InitRHI() override;
 
 	/**
 	 * Releases the dynamic RHI resource and/or RHI render target resources used by this resource.
@@ -453,7 +453,7 @@ public:
 	 * Resources that need to release before a D3D device reset must implement this function.
 	 * This is only called by the rendering thread.
 	 */
-	virtual void ReleaseDynamicRHI() override;
+	virtual void ReleaseRHI() override;
 
 	// FDeferredClearResource interface
 
@@ -540,7 +540,7 @@ public:
 	 * Resources that need to initialize after a D3D device reset must implement this function.
 	 * This is only called by the rendering thread.
 	 */
-	virtual void InitDynamicRHI() override;
+	virtual void InitRHI() override;
 
 	/**
 	 * Releases the dynamic RHI resource and/or RHI render target resources used by this resource.
@@ -548,7 +548,7 @@ public:
 	 * Resources that need to release before a D3D device reset must implement this function.
 	 * This is only called by the rendering thread.
 	 */
-	virtual void ReleaseDynamicRHI() override;
+	virtual void ReleaseRHI() override;
 
 	// FRenderTarget interface.
 
diff --git a/Engine/Source/Runtime/Engine/Public/UnrealClient.h b/Engine/Source/Runtime/Engine/Public/UnrealClient.h
--- a/Engine/Source/Runtime/Engine/Public/UnrealClient.h
+++ b/Engine/Source/Runtime/Engine/Public/UnrealClient.h
@@ -791,16 +791,10 @@ protected:
 
 	/** Triggers the taking of a high res screen shot for this viewport. */
 	bool bTakeHighResScreenShot;
+
 	//~ Begin FRenderResource Interface.
-	ENGINE_API virtual void InitDynamicRHI() override;
-	ENGINE_API virtual void ReleaseDynamicRHI() override;
-	ENGINE_API virtual void InitRHI() override;
 	ENGINE_API virtual void ReleaseRHI() override;
-
-	virtual void InitResource() override { FRenderResource::InitResource(); }
-	virtual void InitResource(FRHICommandList&) override { InitResource(); }
-	virtual void InitDynamicRHI(FRHICommandList&) override { InitDynamicRHI(); }
-	virtual void InitRHI(FRHICommandList&) override { InitRHI(); }
+	ENGINE_API virtual void InitRHI() override;
 };
 
 // Shortcuts for checking the state of both left&right variations of control keys.
diff --git a/Engine/Source/Runtime/Landscape/Private/LandscapeEditLayers.cpp b/Engine/Source/Runtime/Landscape/Private/LandscapeEditLayers.cpp
--- a/Engine/Source/Runtime/Landscape/Private/LandscapeEditLayers.cpp
+++ b/Engine/Source/Runtime/Landscape/Private/LandscapeEditLayers.cpp
@@ -842,7 +842,7 @@ public:
 	}
 
 	/** Called when the resource is initialized. This is only called by the rendering thread. */
-	virtual void InitDynamicRHI() override
+	virtual void InitRHI() override
 	{
 		FRHIResourceCreateInfo CreateInfo(TEXT("FLandscapeLayerWeightmapExtractMaterialLayersComputeShaderResource"));
 		ComponentsData = RHICreateStructuredBuffer(sizeof(FLandscapeLayerWeightmapExtractMaterialLayersComponentData), OriginalComponentsData.Num() * sizeof(FLandscapeLayerWeightmapExtractMaterialLayersComponentData), BUF_ShaderResource | BUF_Volatile, CreateInfo);
@@ -853,7 +853,7 @@ public:
 		RHIUnlockBuffer(ComponentsData);
 	}
 
-	virtual void ReleaseDynamicRHI() override
+	virtual void ReleaseRHI() override
 	{
 		ComponentsData.SafeRelease();
 		ComponentsDataSRV.SafeRelease();
@@ -1001,7 +1001,7 @@ public:
 	}
 
 	/** Called when the resource is initialized. This is only called by the rendering thread. */
-	virtual void InitDynamicRHI() override
+	virtual void InitRHI() override
 	{
 		FRHIResourceCreateInfo CreateInfo(TEXT("ComponentsData"));
 		uint32 ComponentsDataMemSize = OriginalComponentsData.Num() * sizeof(FLandscapeLayerWeightmapPackMaterialLayersComponentData);
@@ -1031,7 +1031,7 @@ public:
 		RHIUnlockBuffer(WeightmapTextureOutputOffset);
 	}
 
-	virtual void ReleaseDynamicRHI() override
+	virtual void ReleaseRHI() override
 	{
 		ComponentsData.SafeRelease();
 		ComponentsDataSRV.SafeRelease();
diff --git a/Engine/Source/Runtime/MediaAssets/Private/Misc/MediaTextureResource.cpp b/Engine/Source/Runtime/MediaAssets/Private/Misc/MediaTextureResource.cpp
--- a/Engine/Source/Runtime/MediaAssets/Private/Misc/MediaTextureResource.cpp
+++ b/Engine/Source/Runtime/MediaAssets/Private/Misc/MediaTextureResource.cpp
@@ -682,7 +682,7 @@ void FMediaTextureResource::SetupSampler()
 }
 
 
-void FMediaTextureResource::InitDynamicRHI()
+void FMediaTextureResource::InitRHI()
 {
 	SetupSampler();
 	
@@ -705,7 +705,7 @@ void FMediaTextureResource::InitDynamicRHI()
 }
 
 
-void FMediaTextureResource::ReleaseDynamicRHI()
+void FMediaTextureResource::ReleaseRHI()
 {
 	Cleared = false;
 
diff --git a/Engine/Source/Runtime/MediaAssets/Private/Misc/MediaTextureResource.h b/Engine/Source/Runtime/MediaAssets/Private/Misc/MediaTextureResource.h
--- a/Engine/Source/Runtime/MediaAssets/Private/Misc/MediaTextureResource.h
+++ b/Engine/Source/Runtime/MediaAssets/Private/Misc/MediaTextureResource.h
@@ -127,8 +127,8 @@ public:
 	virtual FString GetFriendlyName() const override;
 	virtual uint32 GetSizeX() const override;
 	virtual uint32 GetSizeY() const override;
-	virtual void InitDynamicRHI() override;
-	virtual void ReleaseDynamicRHI() override;
+	virtual void InitRHI() override;
+	virtual void ReleaseRHI() override;
 
 protected:
 
diff --git a/Engine/Source/Runtime/RenderCore/Private/GPUMessaging.cpp b/Engine/Source/Runtime/RenderCore/Private/GPUMessaging.cpp
--- a/Engine/Source/Runtime/RenderCore/Private/GPUMessaging.cpp
+++ b/Engine/Source/Runtime/RenderCore/Private/GPUMessaging.cpp
@@ -233,7 +233,7 @@ private:
 		}
 	}
 
-	void ReleaseDynamicRHI() override
+	void ReleaseRHI() override
 	{
 		FScopeLock Lock(&AccessCS);
 
diff --git a/Engine/Source/Runtime/RenderCore/Private/GlobalRenderResources.cpp b/Engine/Source/Runtime/RenderCore/Private/GlobalRenderResources.cpp
--- a/Engine/Source/Runtime/RenderCore/Private/GlobalRenderResources.cpp
+++ b/Engine/Source/Runtime/RenderCore/Private/GlobalRenderResources.cpp
@@ -831,7 +831,7 @@ struct FDynamicVertexBufferPool : public FRenderResource
 	}
 
 private:
-	void ReleaseDynamicRHI() override
+	void ReleaseRHI() override
 	{
 		check(LockList.IsEmpty());
 		check(FreeList.Num() == LiveList.Num());
diff --git a/Engine/Source/Runtime/RenderCore/Private/RenderGraphEvent.cpp b/Engine/Source/Runtime/RenderCore/Private/RenderGraphEvent.cpp
--- a/Engine/Source/Runtime/RenderCore/Private/RenderGraphEvent.cpp
+++ b/Engine/Source/Runtime/RenderCore/Private/RenderGraphEvent.cpp
@@ -14,14 +14,14 @@ public:
 	// Destructor
 	virtual ~FRDGTimingPool() = default;
 
-	virtual void InitDynamicRHI() override
+	virtual void InitRHI() override
 	{
 		check(IsInRenderingThread());
 		bIsBudgetRecordingEnabled.SetAll(false);
 		LastTimings.SetAll(uint64(0));
 	}
 
-	virtual void ReleaseDynamicRHI() override
+	virtual void ReleaseRHI() override
 	{
 		check(IsInRenderingThread());
 
diff --git a/Engine/Source/Runtime/RenderCore/Private/RenderGraphResourcePool.cpp b/Engine/Source/Runtime/RenderCore/Private/RenderGraphResourcePool.cpp
--- a/Engine/Source/Runtime/RenderCore/Private/RenderGraphResourcePool.cpp
+++ b/Engine/Source/Runtime/RenderCore/Private/RenderGraphResourcePool.cpp
@@ -144,7 +144,7 @@ TRefCountPtr<FRDGPooledBuffer> FRDGBufferPool::FindFreeBuffer(const FRDGBufferDe
 	}
 }
 
-void FRDGBufferPool::ReleaseDynamicRHI()
+void FRDGBufferPool::ReleaseRHI()
 {
 	AllocatedBuffers.Empty();
 	AllocatedBufferHashes.Empty();
@@ -215,12 +215,12 @@ uint32 FRDGTransientRenderTarget::Release()
 	return Refs;
 }
 
-void FRDGTransientResourceAllocator::InitDynamicRHI()
+void FRDGTransientResourceAllocator::InitRHI()
 {
 	Allocator = RHICreateTransientResourceAllocator();
 }
 
-void FRDGTransientResourceAllocator::ReleaseDynamicRHI()
+void FRDGTransientResourceAllocator::ReleaseRHI()
 {
 	if (Allocator)
 	{
diff --git a/Engine/Source/Runtime/RenderCore/Private/RenderGraphResourcePool.h b/Engine/Source/Runtime/RenderCore/Private/RenderGraphResourcePool.h
--- a/Engine/Source/Runtime/RenderCore/Private/RenderGraphResourcePool.h
+++ b/Engine/Source/Runtime/RenderCore/Private/RenderGraphResourcePool.h
@@ -23,7 +23,7 @@ public:
 	void DumpMemoryUsage(FOutputDevice& OutputDevice);
 
 private:
-	void ReleaseDynamicRHI() override;
+	void ReleaseRHI() override;
 
 	/** Elements can be 0, we compact the buffer later. */
 	TArray<TRefCountPtr<FRDGPooledBuffer>> AllocatedBuffers;
@@ -95,8 +95,8 @@ public:
 	bool IsValid() const { return Allocator != nullptr; }
 
 private:
-	void InitDynamicRHI() override;
-	void ReleaseDynamicRHI() override;
+	void InitRHI() override;
+	void ReleaseRHI() override;
 
 	void AddPendingDeallocation(FRDGTransientRenderTarget* RenderTarget);
 
diff --git a/Engine/Source/Runtime/RenderCore/Private/RenderResource.cpp b/Engine/Source/Runtime/RenderCore/Private/RenderResource.cpp
--- a/Engine/Source/Runtime/RenderCore/Private/RenderResource.cpp
+++ b/Engine/Source/Runtime/RenderCore/Private/RenderResource.cpp
@@ -24,10 +24,6 @@ FAutoConsoleVariableRef CVarFreeStructuresOnRHIBufferCreation(
 	GFreeStructuresOnRHIBufferCreation,
 	TEXT("Toggles experimental method for freeing helper structures that own the resource arrays after submitting to RHI instead of in the callback sink."));
 
-/** Tracks render resources in a list. The implementation is optimized to allow fast allocation / deallocation from any thread,
- *  at the cost of period coalescing of thread-local data at a sync point each frame. Furthermore, iteration is not thread safe
- *  and must be performed at sync points.
- */
 class FRenderResourceList
 {
 public:
@@ -117,7 +113,6 @@ void FRenderResource::ReleaseRHIForAllResources()
 {
 	FRenderResourceList& ResourceList = FRenderResourceList::Get();
 	ResourceList.ForEachReverse([](FRenderResource* Resource) { check(Resource->IsInitialized()); Resource->ReleaseRHI(); });
-	ResourceList.ForEachReverse([](FRenderResource* Resource) { Resource->ReleaseDynamicRHI(); });
 }
 
 /** Initialize all resources initialized before the RHI was initialized */
@@ -127,8 +122,6 @@ void FRenderResource::InitPreRHIResources()
 
 	// Notify all initialized FRenderResources that there's a valid RHI device to create their RHI resources for now.
 	ResourceList.ForEach([](FRenderResource* Resource) { Resource->InitRHI(); });
-	// Dynamic resources can have dependencies on static resources (with uniform buffers) and must initialized last!
-	ResourceList.ForEach([](FRenderResource* Resource) { Resource->InitDynamicRHI(); });
 
 #if !PLATFORM_NEEDS_RHIRESOURCELIST
 	ResourceList.Clear();
@@ -146,9 +139,7 @@ void FRenderResource::ChangeFeatureLevel(ERHIFeatureLevel::Type NewFeatureLevel)
 			if (Resource->HasValidFeatureLevel() && (Resource->FeatureLevel != NewFeatureLevel))
 			{
 				Resource->ReleaseRHI();
-				Resource->ReleaseDynamicRHI();
 				Resource->FeatureLevel = NewFeatureLevel;
-				Resource->InitDynamicRHI();
 				Resource->InitRHI();
 			}
 		});
@@ -177,7 +168,6 @@ void FRenderResource::InitResource()
 		{
 			CSV_SCOPED_TIMING_STAT_EXCLUSIVE(InitRenderResource);
 			FRHICommandListImmediate& RHICmdList = FRHICommandListExecutor::GetImmediateCommandList();
-			InitDynamicRHI(RHICmdList);
 			InitRHI(RHICmdList);
 		}
 
@@ -206,7 +196,6 @@ void FRenderResource::InitResource(FRHICommandList& RHICmdList)
 		if (GIsRHIInitialized)
 		{
 			CSV_SCOPED_TIMING_STAT_EXCLUSIVE(InitRenderResource);
-			InitDynamicRHI(RHICmdList);
 			InitRHI(RHICmdList);
 		}
 
@@ -224,7 +213,6 @@ void FRenderResource::ReleaseResource()
 			if(GIsRHIInitialized)
 			{
 				ReleaseRHI();
-				ReleaseDynamicRHI();
 			}
 
 #if PLATFORM_NEEDS_RHIRESOURCELIST
@@ -246,8 +234,6 @@ void FRenderResource::UpdateRHI(FRHICommandList& RHICmdList)
 	if (IsInitialized() && GIsRHIInitialized)
 	{
 		ReleaseRHI();
-		ReleaseDynamicRHI();
-		InitDynamicRHI(RHICmdList);
 		InitRHI(RHICmdList);
 	}
 }
diff --git a/Engine/Source/Runtime/RenderCore/Private/RenderTargetPool.cpp b/Engine/Source/Runtime/RenderCore/Private/RenderTargetPool.cpp
--- a/Engine/Source/Runtime/RenderCore/Private/RenderTargetPool.cpp
+++ b/Engine/Source/Runtime/RenderCore/Private/RenderTargetPool.cpp
@@ -502,7 +502,7 @@ void FRenderTargetPool::DumpMemoryUsage(FOutputDevice& OutputDevice)
 	OutputDevice.Logf(TEXT("%.3fMB Deferred total"), DeferredTotal / 1024.f);
 }
 
-void FRenderTargetPool::ReleaseDynamicRHI()
+void FRenderTargetPool::ReleaseRHI()
 {
 	check(IsInRenderingThread());
 	DeferredDeleteArray.Empty();
diff --git a/Engine/Source/Runtime/RenderCore/Private/VisualizeTexture.cpp b/Engine/Source/Runtime/RenderCore/Private/VisualizeTexture.cpp
--- a/Engine/Source/Runtime/RenderCore/Private/VisualizeTexture.cpp
+++ b/Engine/Source/Runtime/RenderCore/Private/VisualizeTexture.cpp
@@ -624,7 +624,7 @@ static EVisualisePSType GetVisualizePSType(const FRDGTextureDesc& Desc)
 	return EVisualisePSType::Texture3D;
 }
 
-void FVisualizeTexture::ReleaseDynamicRHI()
+void FVisualizeTexture::ReleaseRHI()
 {
 	Config = {};
 	Requested = {};
diff --git a/Engine/Source/Runtime/RenderCore/Public/RenderResource.h b/Engine/Source/Runtime/RenderCore/Public/RenderResource.h
--- a/Engine/Source/Runtime/RenderCore/Public/RenderResource.h
+++ b/Engine/Source/Runtime/RenderCore/Public/RenderResource.h
@@ -95,8 +95,8 @@ public:
 	 * Resources that need to initialize after a D3D device reset must implement this function.
 	 * This is only called by the rendering thread.
 	 */
-	virtual void InitDynamicRHI() {}
-	virtual void InitDynamicRHI(FRHICommandList& RHICmdList) { InitDynamicRHI(); }
+	UE_DEPRECATED(5.3, "InitDynamicRHI is now unified with InitRHI.")
+	virtual void InitDynamicRHI() final {}
 
 	/**
 	 * Releases the dynamic RHI resource and/or RHI render target resources used by this resource.
@@ -104,7 +104,8 @@ public:
 	 * Resources that need to release before a D3D device reset must implement this function.
 	 * This is only called by the rendering thread.
 	 */
-	virtual void ReleaseDynamicRHI() {}
+	UE_DEPRECATED(5.3, "ReleaseDynamicRHI is now unified with ReleaseRHI.")
+	virtual void ReleaseDynamicRHI() final {}
 
 	/**
 	 * Initializes the RHI resources used by this resource.
diff --git a/Engine/Source/Runtime/RenderCore/Public/RenderTargetPool.h b/Engine/Source/Runtime/RenderCore/Public/RenderTargetPool.h
--- a/Engine/Source/Runtime/RenderCore/Public/RenderTargetPool.h
+++ b/Engine/Source/Runtime/RenderCore/Public/RenderTargetPool.h
@@ -124,7 +124,7 @@ public:
 	 */
 	void TickPoolElements();
 	/** Free renderer resources */
-	void ReleaseDynamicRHI();
+	void ReleaseRHI();
 
 	/** Allows to remove a resource so it cannot be shared and gets released immediately instead a/some frame[s] later. */
 	void FreeUnusedResource(TRefCountPtr<IPooledRenderTarget>& In);
diff --git a/Engine/Source/Runtime/RenderCore/Public/UniformBuffer.h b/Engine/Source/Runtime/RenderCore/Public/UniformBuffer.h
--- a/Engine/Source/Runtime/RenderCore/Public/UniformBuffer.h
+++ b/Engine/Source/Runtime/RenderCore/Public/UniformBuffer.h
@@ -68,7 +68,7 @@ public:
 	}
 
 	// FRenderResource interface.
-	virtual void InitDynamicRHI() override
+	virtual void InitRHI() override
 	{
 		check(IsInRenderingThread());
 		UniformBufferRHI.SafeRelease();
@@ -77,7 +77,7 @@ public:
 			UniformBufferRHI = CreateUniformBufferImmediate<TBufferStruct>(*((const TBufferStruct*)Contents), BufferUsage);
 		}
 	}
-	virtual void ReleaseDynamicRHI() override
+	virtual void ReleaseRHI() override
 	{
 		UniformBufferRHI.SafeRelease();
 	}
@@ -101,7 +101,7 @@ public:
 
 protected:
 
-	/** Sets the contents of the uniform buffer. Used within calls to InitDynamicRHI */
+	/** Sets the contents of the uniform buffer. Used within calls to InitRHI */
 	void SetContentsNoUpdate(const TBufferStruct& NewContents)
 	{
 		check(IsInRenderingThread());
diff --git a/Engine/Source/Runtime/RenderCore/Public/VisualizeTexture.h b/Engine/Source/Runtime/RenderCore/Public/VisualizeTexture.h
--- a/Engine/Source/Runtime/RenderCore/Public/VisualizeTexture.h
+++ b/Engine/Source/Runtime/RenderCore/Public/VisualizeTexture.h
@@ -147,7 +147,7 @@ private:
 	/** Create a pass capturing a texture. */
 	void CreateContentCapturePass(FRDGBuilder& GraphBuilder, FRDGTextureRef Texture, uint32 CaptureId);
 
-	void ReleaseDynamicRHI() override;
+	void ReleaseRHI() override;
 
 	void Visualize(const FString& InName, TOptional<uint32> InVersion = {});
 
diff --git a/Engine/Source/Runtime/Renderer/Private/IESTextureManager.cpp b/Engine/Source/Runtime/Renderer/Private/IESTextureManager.cpp
--- a/Engine/Source/Runtime/Renderer/Private/IESTextureManager.cpp
+++ b/Engine/Source/Runtime/Renderer/Private/IESTextureManager.cpp
@@ -87,7 +87,7 @@ struct FIESTextureManager : public FRenderResource
 	{
 	}
 
-	virtual void ReleaseDynamicRHI()
+	virtual void ReleaseRHI()
 	{
 		AtlasTexture.SafeRelease();
 	}
diff --git a/Engine/Source/Runtime/Renderer/Private/IndirectLightingCache.cpp b/Engine/Source/Runtime/Renderer/Private/IndirectLightingCache.cpp
--- a/Engine/Source/Runtime/Renderer/Private/IndirectLightingCache.cpp
+++ b/Engine/Source/Runtime/Renderer/Private/IndirectLightingCache.cpp
@@ -157,7 +157,7 @@ FIndirectLightingCache::FIndirectLightingCache(ERHIFeatureLevel::Type InFeatureL
 	CacheSize = GLightingCacheDimension;
 }
 
-void FIndirectLightingCache::InitDynamicRHI()
+void FIndirectLightingCache::InitRHI()
 {
 	if (CanIndirectLightingCacheUseVolumeTexture(GetFeatureLevel()))
 	{
@@ -181,7 +181,7 @@ void FIndirectLightingCache::InitDynamicRHI()
 	}
 }
 
-void FIndirectLightingCache::ReleaseDynamicRHI()
+void FIndirectLightingCache::ReleaseRHI()
 {
 	GRenderTargetPool.FreeUnusedResource(Texture0);
 	GRenderTargetPool.FreeUnusedResource(Texture1);
diff --git a/Engine/Source/Runtime/Renderer/Private/LightMapRendering.cpp b/Engine/Source/Runtime/Renderer/Private/LightMapRendering.cpp
--- a/Engine/Source/Runtime/Renderer/Private/LightMapRendering.cpp
+++ b/Engine/Source/Runtime/Renderer/Private/LightMapRendering.cpp
@@ -729,13 +729,13 @@ void InterpolateVolumetricLightmap(
 	OutInterpolation.DirectionalLightShadowing = DirectionalLightShadowingUnpacked.R;
 }
 
-void FEmptyPrecomputedLightingUniformBuffer::InitDynamicRHI()
+void FEmptyPrecomputedLightingUniformBuffer::InitRHI()
 {
 	FPrecomputedLightingUniformParameters Parameters;
 	GetPrecomputedLightingParameters(GMaxRHIFeatureLevel, Parameters, NULL);
 	SetContentsNoUpdate(Parameters);
 
-	Super::InitDynamicRHI();
+	Super::InitRHI();
 }
 
 /** Global uniform buffer containing the default precomputed lighting data. */
@@ -845,13 +845,13 @@ void GetIndirectLightingCacheParameters(
 	}
 }
 
-void FEmptyIndirectLightingCacheUniformBuffer::InitDynamicRHI()
+void FEmptyIndirectLightingCacheUniformBuffer::InitRHI()
 {
 	FIndirectLightingCacheUniformParameters Parameters;
 	GetIndirectLightingCacheParameters(GMaxRHIFeatureLevel, Parameters, nullptr, nullptr, FVector(0, 0, 0), 0, nullptr);
 	SetContentsNoUpdate(Parameters);
 
-	Super::InitDynamicRHI();
+	Super::InitRHI();
 }
 
 /** */
diff --git a/Engine/Source/Runtime/Renderer/Private/LightMapRendering.h b/Engine/Source/Runtime/Renderer/Private/LightMapRendering.h
--- a/Engine/Source/Runtime/Renderer/Private/LightMapRendering.h
+++ b/Engine/Source/Runtime/Renderer/Private/LightMapRendering.h
@@ -47,7 +47,7 @@ class FEmptyPrecomputedLightingUniformBuffer : public TUniformBuffer< FPrecomput
 {
 	typedef TUniformBuffer< FPrecomputedLightingUniformParameters > Super;
 public:
-	virtual void InitDynamicRHI() override;
+	virtual void InitRHI() override;
 };
 
 /** Global uniform buffer containing the default precomputed lighting data. */
@@ -69,7 +69,7 @@ class FEmptyIndirectLightingCacheUniformBuffer : public TUniformBuffer< FIndirec
 {
 	typedef TUniformBuffer< FIndirectLightingCacheUniformParameters > Super;
 public:
-	virtual void InitDynamicRHI() override;
+	virtual void InitRHI() override;
 };
 
 /** Global uniform buffer containing the default precomputed lighting data. */
diff --git a/Engine/Source/Runtime/Renderer/Private/MeshDrawCommands.cpp b/Engine/Source/Runtime/Renderer/Private/MeshDrawCommands.cpp
--- a/Engine/Source/Runtime/Renderer/Private/MeshDrawCommands.cpp
+++ b/Engine/Source/Runtime/Renderer/Private/MeshDrawCommands.cpp
@@ -129,7 +129,7 @@ void FPrimitiveIdVertexBufferPool::DiscardAll()
 	}
 }
 
-void FPrimitiveIdVertexBufferPool::ReleaseDynamicRHI()
+void FPrimitiveIdVertexBufferPool::ReleaseRHI()
 {
 	Entries.Empty();
 }
diff --git a/Engine/Source/Runtime/Renderer/Private/MeshDrawCommands.h b/Engine/Source/Runtime/Renderer/Private/MeshDrawCommands.h
--- a/Engine/Source/Runtime/Renderer/Private/MeshDrawCommands.h
+++ b/Engine/Source/Runtime/Renderer/Private/MeshDrawCommands.h
@@ -38,7 +38,7 @@ public:
 	void ReturnToFreeList(FPrimitiveIdVertexBufferPoolEntry Entry);
 	RENDERER_API void DiscardAll();
 
-	virtual void ReleaseDynamicRHI() override;
+	virtual void ReleaseRHI() override;
 
 private:
 	uint32 DiscardId;
diff --git a/Engine/Source/Runtime/Renderer/Private/PostProcess/DiaphragmDOF.cpp b/Engine/Source/Runtime/Renderer/Private/PostProcess/DiaphragmDOF.cpp
--- a/Engine/Source/Runtime/Renderer/Private/PostProcess/DiaphragmDOF.cpp
+++ b/Engine/Source/Runtime/Renderer/Private/PostProcess/DiaphragmDOF.cpp
@@ -452,7 +452,7 @@ public:
 	FSpriteIndexBuffer<16> ScatterIndexBuffer;
 
 
-	virtual void InitDynamicRHI() override
+	virtual void InitRHI() override
 	{
 		if (!GRHISupportsRectTopology)
 		{
@@ -460,7 +460,7 @@ public:
 		}
 	}
 
-	virtual void ReleaseDynamicRHI() override
+	virtual void ReleaseRHI() override
 	{
 		if (!GRHISupportsRectTopology)
 		{
diff --git a/Engine/Source/Runtime/Renderer/Private/RectLightTextureManager.cpp b/Engine/Source/Runtime/Renderer/Private/RectLightTextureManager.cpp
--- a/Engine/Source/Runtime/Renderer/Private/RectLightTextureManager.cpp
+++ b/Engine/Source/Runtime/Renderer/Private/RectLightTextureManager.cpp
@@ -198,7 +198,7 @@ struct FRectLightTextureManager : public FRenderResource
 	bool bHasPendingAdds = false;
 	bool bHasPendingDeletes = false;
 
-	virtual void ReleaseDynamicRHI()
+	virtual void ReleaseRHI()
 	{
 		AtlasTexture.SafeRelease();
 	}
diff --git a/Engine/Source/Runtime/Renderer/Private/ReflectionEnvironment.cpp b/Engine/Source/Runtime/Renderer/Private/ReflectionEnvironment.cpp
--- a/Engine/Source/Runtime/Renderer/Private/ReflectionEnvironment.cpp
+++ b/Engine/Source/Runtime/Renderer/Private/ReflectionEnvironment.cpp
@@ -145,7 +145,7 @@ bool IsReflectionCaptureAvailable()
 	return (!AllowStaticLightingVar || AllowStaticLightingVar->GetInt() != 0);
 }
 
-void FReflectionEnvironmentCubemapArray::InitDynamicRHI()
+void FReflectionEnvironmentCubemapArray::InitRHI()
 {
 	if (SupportsTextureCubeArray(GetFeatureLevel()))
 	{
@@ -184,7 +184,7 @@ void FReflectionEnvironmentCubemapArray::ReleaseCubeArray()
 	GRenderTargetPool.FreeUnusedResource(ReflectionEnvs);
 }
 
-void FReflectionEnvironmentCubemapArray::ReleaseDynamicRHI()
+void FReflectionEnvironmentCubemapArray::ReleaseRHI()
 {
 	ReleaseCubeArray();
 }
@@ -537,13 +537,13 @@ void FReflectionEnvironmentCubemapArray::ResizeCubemapArrayGPU(uint32 InMaxCubem
 	check(IsInitialized());
 	check(InCubemapSize == CubemapSize);
 
-	// Take a reference to the old cubemap array and then release it to prevent it getting destroyed during InitDynamicRHI
+	// Take a reference to the old cubemap array and then release it to prevent it getting destroyed during InitRHI
 	TRefCountPtr<IPooledRenderTarget> OldReflectionEnvs = ReflectionEnvs;
 	ReflectionEnvs = nullptr;
 	int OldMaxCubemaps = MaxCubemaps;
 	MaxCubemaps = InMaxCubemaps;
 
-	InitDynamicRHI();
+	InitRHI();
 
 	FTextureRHIRef TexRef = OldReflectionEnvs->GetRHI();
 	FRHICommandListImmediate& RHICmdList = FRHICommandListExecutor::GetImmediateCommandList();
@@ -606,7 +606,7 @@ void FReflectionEnvironmentCubemapArray::UpdateMaxCubemaps(uint32 InMaxCubemaps,
 
 void FReflectionEnvironmentCubemapArray::Reset()
 {
-	ReleaseDynamicRHI();
+	ReleaseRHI();
 	MaxCubemaps = 0;
 	CubemapSize = 0;
 }
diff --git a/Engine/Source/Runtime/Renderer/Private/RendererScene.cpp b/Engine/Source/Runtime/Renderer/Private/RendererScene.cpp
--- a/Engine/Source/Runtime/Renderer/Private/RendererScene.cpp
+++ b/Engine/Source/Runtime/Renderer/Private/RendererScene.cpp
@@ -313,7 +313,7 @@ TRefCountPtr<IPooledRenderTarget> FPersistentSkyAtmosphereData::GetCurrentCamera
 FSceneViewState::FSceneViewState(ERHIFeatureLevel::Type FeatureLevel, FSceneViewState* ShareOriginTarget)
 	: OcclusionQueryPool(RHICreateRenderQueryPool(RQT_Occlusion))
 {
-	// Set FeatureLevel to a valid value, so we get Init/ReleaseDynamicRHI calls on FeatureLevel changes
+	// Set FeatureLevel to a valid value, so we get Init/ReleaseRHI calls on FeatureLevel changes
 	SetFeatureLevel(FeatureLevel);
 	
 	UniqueID = FSceneViewState_UniqueID.Increment();
diff --git a/Engine/Source/Runtime/Renderer/Private/SceneOcclusion.cpp b/Engine/Source/Runtime/Renderer/Private/SceneOcclusion.cpp
--- a/Engine/Source/Runtime/Renderer/Private/SceneOcclusion.cpp
+++ b/Engine/Source/Runtime/Renderer/Private/SceneOcclusion.cpp
@@ -796,7 +796,7 @@ void FHZBOcclusionTester::SetInvalidFrameNumber()
 	checkSlow(IsInvalidFrame());
 }
 
-void FHZBOcclusionTester::InitDynamicRHI()
+void FHZBOcclusionTester::InitRHI()
 {
 	if (GetFeatureLevel() >= ERHIFeatureLevel::SM5)
 	{
@@ -804,7 +804,7 @@ void FHZBOcclusionTester::InitDynamicRHI()
 	}
 }
 
-void FHZBOcclusionTester::ReleaseDynamicRHI()
+void FHZBOcclusionTester::ReleaseRHI()
 {
 	if (GetFeatureLevel() >= ERHIFeatureLevel::SM5)
 	{
@@ -1627,7 +1627,7 @@ FOcclusionFeedback::~FOcclusionFeedback()
 {
 }
 
-void FOcclusionFeedback::InitDynamicRHI()
+void FOcclusionFeedback::InitRHI()
 {
 	for (int32 i = 0; i < UE_ARRAY_COUNT(OcclusionBuffers); ++i)
 	{
@@ -1642,7 +1642,7 @@ void FOcclusionFeedback::InitDynamicRHI()
 		OcclusionVertexDeclarationRHI = PipelineStateCache::GetOrCreateVertexDeclaration(Elements);
 	}
 }
-void FOcclusionFeedback::ReleaseDynamicRHI()
+void FOcclusionFeedback::ReleaseRHI()
 {
 	BatchOcclusionQueries.Empty();
 
diff --git a/Engine/Source/Runtime/Renderer/Private/ScenePrivate.h b/Engine/Source/Runtime/Renderer/Private/ScenePrivate.h
--- a/Engine/Source/Runtime/Renderer/Private/ScenePrivate.h
+++ b/Engine/Source/Runtime/Renderer/Private/ScenePrivate.h
@@ -1116,7 +1116,7 @@ public:
 		DistanceFieldTemporalSampleIndex = 0;
 		PreExposure = 1.f;
 
-		ReleaseDynamicRHI();
+		ReleaseRHI();
 	}
 
 	void SetupDistanceFieldTemporalOffset(const FSceneViewFamily& Family)
@@ -1341,14 +1341,14 @@ PRAGMA_ENABLE_DEPRECATION_WARNINGS
 	}
 
 	// FRenderResource interface.
-	virtual void InitDynamicRHI() override
+	virtual void InitRHI() override
 	{
-		HZBOcclusionTests.InitDynamicRHI();
+		HZBOcclusionTests.InitRHI();
 	}
 
-	virtual void ReleaseDynamicRHI() override
+	virtual void ReleaseRHI() override
 	{
-		HZBOcclusionTests.ReleaseDynamicRHI();
+		HZBOcclusionTests.ReleaseRHI();
 		EyeAdaptationManager.SafeRelease();
 		OcclusionFeedback.ReleaseResource();
 PRAGMA_DISABLE_DEPRECATION_WARNINGS
@@ -1579,8 +1579,8 @@ public:
 		, CubemapSize(0)
 	{}
 
-	virtual void InitDynamicRHI() override;
-	virtual void ReleaseDynamicRHI() override;
+	virtual void InitRHI() override;
+	virtual void ReleaseRHI() override;
 
 	/** 
 	 * Updates the maximum number of cubemaps that this array is allocated for.
@@ -2208,8 +2208,8 @@ public:
 	FIndirectLightingCache(ERHIFeatureLevel::Type InFeatureLevel);
 
 	// FRenderResource interface
-	virtual void InitDynamicRHI();
-	virtual void ReleaseDynamicRHI();
+	virtual void InitRHI();
+	virtual void ReleaseRHI();
 
 	/** Allocates a block in the volume texture atlas for a primitive. */
 	FIndirectLightingCacheAllocation* AllocatePrimitive(const FPrimitiveSceneInfo* PrimitiveSceneInfo, bool bUnbuiltPreview);
diff --git a/Engine/Source/Runtime/Renderer/Private/SceneRendering.h b/Engine/Source/Runtime/Renderer/Private/SceneRendering.h
--- a/Engine/Source/Runtime/Renderer/Private/SceneRendering.h
+++ b/Engine/Source/Runtime/Renderer/Private/SceneRendering.h
@@ -477,8 +477,8 @@ public:
 	~FHZBOcclusionTester() {}
 
 	// FRenderResource interface
-	virtual void InitDynamicRHI() override;
-	virtual void ReleaseDynamicRHI() override;
+	virtual void InitRHI() override;
+	virtual void ReleaseRHI() override;
 	
 	uint32 GetNum() const { return Primitives.Num(); }
 
@@ -522,8 +522,8 @@ public:
 	~FOcclusionFeedback();
 
 	// FRenderResource interface
-	virtual void InitDynamicRHI() override;
-	virtual void ReleaseDynamicRHI() override;
+	virtual void InitRHI() override;
+	virtual void ReleaseRHI() override;
 
 	void AddPrimitive(FRHICommandList& RHICmdList, FPrimitiveComponentId PrimitiveId, const FVector& BoundsOrigin, const FVector& BoundsBoxExtent, FGlobalDynamicVertexBuffer& DynamicVertexBuffer);
 
diff --git a/Engine/Source/Runtime/Renderer/Private/ShadowRendering.cpp b/Engine/Source/Runtime/Renderer/Private/ShadowRendering.cpp
--- a/Engine/Source/Runtime/Renderer/Private/ShadowRendering.cpp
+++ b/Engine/Source/Runtime/Renderer/Private/ShadowRendering.cpp
@@ -2473,13 +2473,13 @@ void SetupTranslucentSelfShadowUniformParameters(const FProjectedShadowInfo* Sha
 	}
 }
 
-void FEmptyTranslucentSelfShadowUniformBuffer::InitDynamicRHI()
+void FEmptyTranslucentSelfShadowUniformBuffer::InitRHI()
 {
 	FTranslucentSelfShadowUniformParameters Parameters;
 	SetupTranslucentSelfShadowUniformParameters(nullptr, Parameters);
 	SetContentsNoUpdate(Parameters);
 
-	Super::InitDynamicRHI();
+	Super::InitRHI();
 }
 
 
diff --git a/Engine/Source/Runtime/Renderer/Private/ShadowRendering.h b/Engine/Source/Runtime/Renderer/Private/ShadowRendering.h
--- a/Engine/Source/Runtime/Renderer/Private/ShadowRendering.h
+++ b/Engine/Source/Runtime/Renderer/Private/ShadowRendering.h
@@ -1325,7 +1325,7 @@ class FEmptyTranslucentSelfShadowUniformBuffer : public TUniformBuffer< FTranslu
 {
 	typedef TUniformBuffer< FTranslucentSelfShadowUniformParameters > Super;
 public:
-	virtual void InitDynamicRHI() override;
+	virtual void InitRHI() override;
 };
 
 /** Global uniform buffer containing the default precomputed lighting data. */
diff --git a/Engine/Source/Runtime/Renderer/Private/SystemTextures.cpp b/Engine/Source/Runtime/Renderer/Private/SystemTextures.cpp
--- a/Engine/Source/Runtime/Renderer/Private/SystemTextures.cpp
+++ b/Engine/Source/Runtime/Renderer/Private/SystemTextures.cpp
@@ -900,7 +900,7 @@ void FSystemTextures::InitializeFeatureLevelDependentTextures(FRHICommandListImm
 	FeatureLevelInitializedTo = InFeatureLevel;
 }
 
-void FSystemTextures::ReleaseDynamicRHI()
+void FSystemTextures::ReleaseRHI()
 {
 	WhiteDummySRV.SafeRelease();
 	WhiteDummy.SafeRelease();
diff --git a/Engine/Source/Runtime/Renderer/Private/VT/VirtualTextureTest.cpp b/Engine/Source/Runtime/Renderer/Private/VT/VirtualTextureTest.cpp
--- a/Engine/Source/Runtime/Renderer/Private/VT/VirtualTextureTest.cpp
+++ b/Engine/Source/Runtime/Renderer/Private/VT/VirtualTextureTest.cpp
@@ -49,14 +49,14 @@ FVirtualTextureTestType::~FVirtualTextureTestType()
 	delete Pool;
 }
 
-void FVirtualTextureTestType::InitDynamicRHI()
+void FVirtualTextureTestType::InitRHI()
 {
 	FRHICommandListImmediate& RHICmdList = FRHICommandListExecutor::GetImmediateCommandList();
 	FPooledRenderTargetDesc Desc( FPooledRenderTargetDesc::Create2DDesc( PhysicalTextureSize, PhysicalTextureFormat, FClearValueBinding::None, TexCreate_None, TexCreate_SRGB | TexCreate_RenderTargetable | TexCreate_ShaderResource, false ) );
 	GRenderTargetPool.FindFreeElement( RHICmdList, Desc, PhysicalTexture, TEXT("PhysicalTexture") );
 }
 
-void FVirtualTextureTestType::ReleaseDynamicRHI()
+void FVirtualTextureTestType::ReleaseRHI()
 {
 	GRenderTargetPool.FreeUnusedResource( PhysicalTexture );
 }
diff --git a/Engine/Source/Runtime/Renderer/Private/VT/VirtualTextureTest.h b/Engine/Source/Runtime/Renderer/Private/VT/VirtualTextureTest.h
--- a/Engine/Source/Runtime/Renderer/Private/VT/VirtualTextureTest.h
+++ b/Engine/Source/Runtime/Renderer/Private/VT/VirtualTextureTest.h
@@ -32,8 +32,8 @@ public:
 					~FVirtualTextureTestType();
 
 	// FRenderResource interface
-	virtual void	InitDynamicRHI() override;
-	virtual void	ReleaseDynamicRHI() override;
+	virtual void	InitRHI() override;
+	virtual void	ReleaseRHI() override;
 	
 	FRHITexture*	GetPhysicalTexture() const		{ return PhysicalTexture->GetRHI(); }
 	FRHITexture*	GetPageTableTexture() const;
diff --git a/Engine/Source/Runtime/Renderer/Private/VariableRateShading/VariableRateShadingImageManager.cpp b/Engine/Source/Runtime/Renderer/Private/VariableRateShading/VariableRateShadingImageManager.cpp
--- a/Engine/Source/Runtime/Renderer/Private/VariableRateShading/VariableRateShadingImageManager.cpp
+++ b/Engine/Source/Runtime/Renderer/Private/VariableRateShading/VariableRateShadingImageManager.cpp
@@ -262,7 +262,7 @@ void FVariableRateShadingImageManager::InitRHI()
 	}
 }
 
-void FVariableRateShadingImageManager::ReleaseDynamicRHI()
+void FVariableRateShadingImageManager::ReleaseRHI()
 {
 	GRenderTargetPool.FreeUnusedResources();
 }
diff --git a/Engine/Source/Runtime/Renderer/Public/SceneRenderTargetParameters.h b/Engine/Source/Runtime/Renderer/Public/SceneRenderTargetParameters.h
--- a/Engine/Source/Runtime/Renderer/Public/SceneRenderTargetParameters.h
+++ b/Engine/Source/Runtime/Renderer/Public/SceneRenderTargetParameters.h
@@ -151,7 +151,7 @@ public:
 
 private:
 	void Release();
-	void ReleaseDynamicRHI() override { Release(); }
+	void ReleaseRHI() override { Release(); }
 
 	// Contains the resolved scene depth target.
 	TRefCountPtr<IPooledRenderTarget> Depth;
diff --git a/Engine/Source/Runtime/Renderer/Public/SystemTextures.h b/Engine/Source/Runtime/Renderer/Public/SystemTextures.h
--- a/Engine/Source/Runtime/Renderer/Public/SystemTextures.h
+++ b/Engine/Source/Runtime/Renderer/Public/SystemTextures.h
@@ -94,12 +94,12 @@ public:
 	/**
 	 * Release textures when device is lost/destroyed.
 	 */
-	virtual void ReleaseDynamicRHI();
+	virtual void ReleaseRHI();
 
 	// -----------
 
 	/**
-		Any Textures added here MUST be explicitly released on ReleaseDynamicRHI()!
+		Any Textures added here MUST be explicitly released on ReleaseRHI()!
 		Some RHIs need all their references released during destruction!
 	*/
 
diff --git a/Engine/Source/Runtime/Renderer/Public/VariableRateShadingImageManager.h b/Engine/Source/Runtime/Renderer/Public/VariableRateShadingImageManager.h
--- a/Engine/Source/Runtime/Renderer/Public/VariableRateShadingImageManager.h
+++ b/Engine/Source/Runtime/Renderer/Public/VariableRateShadingImageManager.h
@@ -78,7 +78,7 @@ public:
 	virtual ~FVariableRateShadingImageManager();
 
 	virtual void InitRHI() override;
-	virtual void ReleaseDynamicRHI() override;
+	virtual void ReleaseRHI() override;
 
 	/**
 	 * Get the combined VRS image for the specified view setup, target size/configuration, and application flags.
diff --git a/Engine/Source/Runtime/SlateRHIRenderer/Private/SlateElementIndexBuffer.cpp b/Engine/Source/Runtime/SlateRHIRenderer/Private/SlateElementIndexBuffer.cpp
--- a/Engine/Source/Runtime/SlateRHIRenderer/Private/SlateElementIndexBuffer.cpp
+++ b/Engine/Source/Runtime/SlateRHIRenderer/Private/SlateElementIndexBuffer.cpp
@@ -46,7 +46,7 @@ void FSlateElementIndexBuffer::Destroy()
 }
 
 /** Initializes the index buffers RHI resource. */
-void FSlateElementIndexBuffer::InitDynamicRHI()
+void FSlateElementIndexBuffer::InitRHI()
 {
 	checkSlow( IsInRenderingThread() );
 
@@ -99,7 +99,7 @@ void FSlateElementIndexBuffer::PreFillBuffer(int32 RequiredIndexCount, bool bShr
 }
 
 /** Releases the index buffers RHI resource. */
-void FSlateElementIndexBuffer::ReleaseDynamicRHI()
+void FSlateElementIndexBuffer::ReleaseRHI()
 {
 	IndexBufferRHI.SafeRelease();
 	SetBufferSize(0);
diff --git a/Engine/Source/Runtime/SlateRHIRenderer/Private/SlateElementIndexBuffer.h b/Engine/Source/Runtime/SlateRHIRenderer/Private/SlateElementIndexBuffer.h
--- a/Engine/Source/Runtime/SlateRHIRenderer/Private/SlateElementIndexBuffer.h
+++ b/Engine/Source/Runtime/SlateRHIRenderer/Private/SlateElementIndexBuffer.h
@@ -20,10 +20,10 @@ public:
 	void Destroy();
 
 	/** Initializes the index buffers RHI resource. */
-	virtual void InitDynamicRHI() override;
+	virtual void InitRHI() override;
 
 	/** Releases the index buffers RHI resource. */
-	virtual void ReleaseDynamicRHI() override;
+	virtual void ReleaseRHI() override;
 
 	/** Returns a friendly name for this buffer. */
 	virtual FString GetFriendlyName() const override { return TEXT("SlateElementIndices"); }
diff --git a/Engine/Source/Runtime/SlateRHIRenderer/Private/SlateElementVertexBuffer.h b/Engine/Source/Runtime/SlateRHIRenderer/Private/SlateElementVertexBuffer.h
--- a/Engine/Source/Runtime/SlateRHIRenderer/Private/SlateElementVertexBuffer.h
+++ b/Engine/Source/Runtime/SlateRHIRenderer/Private/SlateElementVertexBuffer.h
@@ -55,7 +55,7 @@ public:
 	}
 
 	/** Initializes the vertex buffers RHI resource. */
-	virtual void InitDynamicRHI()
+	virtual void InitRHI()
 	{
 		if( !IsValidRef(VertexBufferRHI) )
 		{
@@ -72,7 +72,7 @@ public:
 	}
 
 	/** Releases the vertex buffers RHI resource. */
-	virtual void ReleaseDynamicRHI()
+	virtual void ReleaseRHI()
 	{
 		VertexBufferRHI.SafeRelease();
 		SetBufferSize(0);
@@ -172,7 +172,7 @@ public:
 	~FSlateStencilClipVertexBuffer() {};
 
 	/** Initializes the vertex buffers RHI resource. */
-	virtual void InitDynamicRHI()
+	virtual void InitRHI()
 	{
 		if (!IsValidRef(VertexBufferRHI))
 		{
@@ -188,7 +188,7 @@ public:
 	}
 
 	/** Releases the vertex buffers RHI resource. */
-	virtual void ReleaseDynamicRHI()
+	virtual void ReleaseRHI()
 	{
 		VertexBufferRHI.SafeRelease();
 	}
diff --git a/Engine/Source/Runtime/SlateRHIRenderer/Private/SlatePostProcessResource.cpp b/Engine/Source/Runtime/SlateRHIRenderer/Private/SlatePostProcessResource.cpp
--- a/Engine/Source/Runtime/SlateRHIRenderer/Private/SlatePostProcessResource.cpp
+++ b/Engine/Source/Runtime/SlateRHIRenderer/Private/SlatePostProcessResource.cpp
@@ -69,11 +69,11 @@ void FSlatePostProcessResource::CleanUp()
 	BeginCleanup(this);
 }
 
-void FSlatePostProcessResource::InitDynamicRHI()
+void FSlatePostProcessResource::InitRHI()
 {
 }
 
-void FSlatePostProcessResource::ReleaseDynamicRHI()
+void FSlatePostProcessResource::ReleaseRHI()
 {
 	SET_MEMORY_STAT(STAT_SLATEPPRenderTargetMem, 0);
 
diff --git a/Engine/Source/Runtime/SlateRHIRenderer/Private/SlatePostProcessResource.h b/Engine/Source/Runtime/SlateRHIRenderer/Private/SlatePostProcessResource.h
--- a/Engine/Source/Runtime/SlateRHIRenderer/Private/SlatePostProcessResource.h
+++ b/Engine/Source/Runtime/SlateRHIRenderer/Private/SlatePostProcessResource.h
@@ -30,8 +30,8 @@ public:
 	void CleanUp();
 
 	/** FRenderResource interface */
-	virtual void InitDynamicRHI() override;
-	virtual void ReleaseDynamicRHI() override;
+	virtual void InitRHI() override;
+	virtual void ReleaseRHI() override;
 
 	/** FSlateShaderResource interface */
 	virtual uint32 GetWidth() const override { return RenderTargetSize.X; }
diff --git a/Engine/Source/Runtime/SlateRHIRenderer/Private/SlateRHIFontTexture.cpp b/Engine/Source/Runtime/SlateRHIRenderer/Private/SlateRHIFontTexture.cpp
--- a/Engine/Source/Runtime/SlateRHIRenderer/Private/SlateRHIFontTexture.cpp
+++ b/Engine/Source/Runtime/SlateRHIRenderer/Private/SlateRHIFontTexture.cpp
@@ -12,7 +12,7 @@ FSlateFontTextureRHIResource::FSlateFontTextureRHIResource(uint32 InWidth, uint3
 {
 }
 
-void FSlateFontTextureRHIResource::InitDynamicRHI()
+void FSlateFontTextureRHIResource::InitRHI()
 {
 	check( IsInRenderingThread() );
 
@@ -73,7 +73,7 @@ void FSlateFontTextureRHIResource::InitDynamicRHI()
 	}
 }
 
-void FSlateFontTextureRHIResource::ReleaseDynamicRHI()
+void FSlateFontTextureRHIResource::ReleaseRHI()
 {
 	check( IsInRenderingThread() );
 
diff --git a/Engine/Source/Runtime/SlateRHIRenderer/Private/SlateRHIFontTexture.h b/Engine/Source/Runtime/SlateRHIRenderer/Private/SlateRHIFontTexture.h
--- a/Engine/Source/Runtime/SlateRHIRenderer/Private/SlateRHIFontTexture.h
+++ b/Engine/Source/Runtime/SlateRHIRenderer/Private/SlateRHIFontTexture.h
@@ -8,7 +8,7 @@
 #include "TextureResource.h"
 
 /**
- * Override for font textures that saves data between Init and ReleaseDynamicRHI to 
+ * Override for font textures that saves data between Init and ReleaseRHI to 
  * ensure all characters in the font texture exist if the rendering resource has to be recreated 
  * between caching new characters
  */
@@ -32,8 +32,8 @@ public:
 	virtual FString GetFriendlyName() const override { return TEXT("FSlateFontTextureRHIResource"); }
 
 	/** FRenderResource interface */
-	virtual void InitDynamicRHI() override;
-	virtual void ReleaseDynamicRHI() override;
+	virtual void InitRHI() override;
+	virtual void ReleaseRHI() override;
 	
 	/** Returns whether the texture resource is 8-bit grayscale or 8-bit per-channel BGRA color */
 	bool IsGrayscale() const { return bIsGrayscale; }
@@ -47,7 +47,7 @@ private:
 	uint32 Height;
 	/** Whether this texture is 8-bit grayscale or 8-bit per-channel BGRA color */
 	bool bIsGrayscale;
-	/** Temporary data stored between Release and InitDynamicRHI */
+	/** Temporary data stored between Release and InitRHI */
 	TArray<uint8> TempData;
 };
 
diff --git a/Engine/Source/Runtime/WebBrowserTexture/Private/WebBrowserTextureResource.cpp b/Engine/Source/Runtime/WebBrowserTexture/Private/WebBrowserTextureResource.cpp
--- a/Engine/Source/Runtime/WebBrowserTexture/Private/WebBrowserTextureResource.cpp
+++ b/Engine/Source/Runtime/WebBrowserTexture/Private/WebBrowserTextureResource.cpp
@@ -112,7 +112,7 @@ uint32 FWebBrowserTextureResource::GetSizeY() const
 }
 
 
-void FWebBrowserTextureResource::InitDynamicRHI()
+void FWebBrowserTextureResource::InitRHI()
 {
 	// create the sampler state
 	FSamplerStateInitializerRHI SamplerStateInitializer(
@@ -126,7 +126,7 @@ void FWebBrowserTextureResource::InitDynamicRHI()
 }
 
 
-void FWebBrowserTextureResource::ReleaseDynamicRHI()
+void FWebBrowserTextureResource::ReleaseRHI()
 {
 	Cleared = false;
 
diff --git a/Engine/Source/Runtime/WebBrowserTexture/Public/WebBrowserTextureResource.h b/Engine/Source/Runtime/WebBrowserTexture/Public/WebBrowserTextureResource.h
--- a/Engine/Source/Runtime/WebBrowserTexture/Public/WebBrowserTextureResource.h
+++ b/Engine/Source/Runtime/WebBrowserTexture/Public/WebBrowserTextureResource.h
@@ -78,8 +78,8 @@ public:
 	virtual FString GetFriendlyName() const override;
 	virtual uint32 GetSizeX() const override;
 	virtual uint32 GetSizeY() const override;
-	virtual void InitDynamicRHI() override;
-	virtual void ReleaseDynamicRHI() override;
+	virtual void InitRHI() override;
+	virtual void ReleaseRHI() override;
 
 protected:
 
diff --git a/Engine/Source/Runtime/Windows/D3D11RHI/Private/D3D11ConstantBuffer.cpp b/Engine/Source/Runtime/Windows/D3D11RHI/Private/D3D11ConstantBuffer.cpp
--- a/Engine/Source/Runtime/Windows/D3D11RHI/Private/D3D11ConstantBuffer.cpp
+++ b/Engine/Source/Runtime/Windows/D3D11RHI/Private/D3D11ConstantBuffer.cpp
@@ -21,7 +21,7 @@ FD3D11ConstantBuffer::~FD3D11ConstantBuffer()
 /**
 * Creates a constant buffer on the device
 */
-void FD3D11ConstantBuffer::InitDynamicRHI()
+void FD3D11ConstantBuffer::InitRHI()
 {
 // New circular buffer system for faster constant uploads.  Avoids CopyResource and speeds things up considerably
 	// aligned for best performance
@@ -29,7 +29,7 @@ void FD3D11ConstantBuffer::InitDynamicRHI()
 	FMemory::Memzero(ShadowData, GetMaxSize());
 }
 
-void FD3D11ConstantBuffer::ReleaseDynamicRHI()
+void FD3D11ConstantBuffer::ReleaseRHI()
 {
 	if(ShadowData)
 	{
diff --git a/Engine/Source/Runtime/Windows/D3D11RHI/Private/D3D11Query.cpp b/Engine/Source/Runtime/Windows/D3D11RHI/Private/D3D11Query.cpp
--- a/Engine/Source/Runtime/Windows/D3D11RHI/Private/D3D11Query.cpp
+++ b/Engine/Source/Runtime/Windows/D3D11RHI/Private/D3D11Query.cpp
@@ -597,7 +597,7 @@ void FD3D11DynamicRHI::RHICalibrateTimers()
 /**
  * Initializes all D3D resources and if necessary, the static variables.
  */
-void FD3D11BufferedGPUTiming::InitDynamicRHI()
+void FD3D11BufferedGPUTiming::InitRHI()
 {
 	StaticInitialize(D3DRHI, PlatformStaticInitialize);
 
@@ -631,7 +631,7 @@ void FD3D11BufferedGPUTiming::InitDynamicRHI()
 /**
  * Releases all D3D resources.
  */
-void FD3D11BufferedGPUTiming::ReleaseDynamicRHI()
+void FD3D11BufferedGPUTiming::ReleaseRHI()
 {
 	if ( StartTimestamps && EndTimestamps )
 	{
@@ -804,7 +804,7 @@ D3D11_QUERY_DATA_TIMESTAMP_DISJOINT FD3D11DisjointTimeStampQuery::GetResult()
 	return DisjointQueryData;
 }
 
-void FD3D11DisjointTimeStampQuery::InitDynamicRHI()
+void FD3D11DisjointTimeStampQuery::InitRHI()
 {
 	D3D11_QUERY_DESC QueryDesc;
 	QueryDesc.Query = D3D11_QUERY_TIMESTAMP_DISJOINT;
@@ -813,7 +813,7 @@ void FD3D11DisjointTimeStampQuery::InitDynamicRHI()
 	VERIFYD3D11RESULT_EX(D3DRHI->GetDevice()->CreateQuery(&QueryDesc, DisjointQuery.GetInitReference()), D3DRHI->GetDevice());
 }
 
-void FD3D11DisjointTimeStampQuery::ReleaseDynamicRHI()
+void FD3D11DisjointTimeStampQuery::ReleaseRHI()
 {
 
 }
diff --git a/Engine/Source/Runtime/Windows/D3D11RHI/Private/D3D11RHIPrivate.h b/Engine/Source/Runtime/Windows/D3D11RHI/Private/D3D11RHIPrivate.h
--- a/Engine/Source/Runtime/Windows/D3D11RHI/Private/D3D11RHIPrivate.h
+++ b/Engine/Source/Runtime/Windows/D3D11RHI/Private/D3D11RHIPrivate.h
@@ -191,12 +191,12 @@ public:
 	/**
 	 * Initializes all D3D resources.
 	 */
-	virtual void InitDynamicRHI() override;
+	virtual void InitRHI() override;
 
 	/**
 	 * Releases all D3D resources.
 	 */
-	virtual void ReleaseDynamicRHI() override;
+	virtual void ReleaseRHI() override;
 
 	static void CalibrateTimers(FD3D11DynamicRHI* InD3DRHI);
 
@@ -236,12 +236,12 @@ public:
 	/**
 	 * Initializes all D3D resources.
 	 */
-	virtual void InitDynamicRHI() override;
+	virtual void InitRHI() override;
 
 	/**
 	 * Releases all D3D resources.
 	 */
-	virtual void ReleaseDynamicRHI() override;
+	virtual void ReleaseRHI() override;
 
 
 private:
@@ -260,12 +260,12 @@ public:
 		Timing(InRHI, 1)
 	{
 		// Initialize Buffered timestamp queries 
-		Timing.InitDynamicRHI(); // can't do this from the RHI thread
+		Timing.InitRHI(); // can't do this from the RHI thread
 	}
 
 	virtual ~FD3D11EventNode()
 	{
-		Timing.ReleaseDynamicRHI();  // can't do this from the RHI thread
+		Timing.ReleaseRHI();  // can't do this from the RHI thread
 	}
 
 	/** 
@@ -298,14 +298,14 @@ public:
 		RootEventTiming(InRHI, 1),
 		DisjointQuery(InRHI)
 	{
-		RootEventTiming.InitDynamicRHI();
-		DisjointQuery.InitDynamicRHI();
+		RootEventTiming.InitRHI();
+		DisjointQuery.InitRHI();
 	}
 
 	~FD3D11EventNodeFrame()
 	{
-		RootEventTiming.ReleaseDynamicRHI();
-		DisjointQuery.ReleaseDynamicRHI();
+		RootEventTiming.ReleaseRHI();
+		DisjointQuery.ReleaseRHI();
 	}
 
 	/** Start this frame of per tracking */
diff --git a/Engine/Source/Runtime/Windows/D3D11RHI/Private/Windows/WindowsD3D11ConstantBuffer.cpp b/Engine/Source/Runtime/Windows/D3D11RHI/Private/Windows/WindowsD3D11ConstantBuffer.cpp
--- a/Engine/Source/Runtime/Windows/D3D11RHI/Private/Windows/WindowsD3D11ConstantBuffer.cpp
+++ b/Engine/Source/Runtime/Windows/D3D11RHI/Private/Windows/WindowsD3D11ConstantBuffer.cpp
@@ -7,7 +7,7 @@
 #include "WindowsD3D11ConstantBuffer.h"
 #include "D3D11RHIPrivate.h"
 
-void FWinD3D11ConstantBuffer::InitDynamicRHI()
+void FWinD3D11ConstantBuffer::InitRHI()
 {
 	TRefCountPtr<ID3D11Buffer> CBuffer;
 
@@ -36,12 +36,12 @@ void FWinD3D11ConstantBuffer::InitDynamicRHI()
 		BufferDesc.ByteWidth = Align(BufferDesc.ByteWidth / 2, 16);
 	}
 
-	FD3D11ConstantBuffer::InitDynamicRHI();
+	FD3D11ConstantBuffer::InitRHI();
 }
 
-void FWinD3D11ConstantBuffer::ReleaseDynamicRHI()
+void FWinD3D11ConstantBuffer::ReleaseRHI()
 {
-	FD3D11ConstantBuffer::ReleaseDynamicRHI();
+	FD3D11ConstantBuffer::ReleaseRHI();
 
 	for (FSubBuffer& SubBuffer : SubBuffers)
 	{
diff --git a/Engine/Source/Runtime/Windows/D3D11RHI/Private/Windows/WindowsD3D11ConstantBuffer.h b/Engine/Source/Runtime/Windows/D3D11RHI/Private/Windows/WindowsD3D11ConstantBuffer.h
--- a/Engine/Source/Runtime/Windows/D3D11RHI/Private/Windows/WindowsD3D11ConstantBuffer.h
+++ b/Engine/Source/Runtime/Windows/D3D11RHI/Private/Windows/WindowsD3D11ConstantBuffer.h
@@ -19,8 +19,8 @@ public:
 	}
 
 	// FRenderResource interface.
-	virtual void	InitDynamicRHI() override;
-	virtual void	ReleaseDynamicRHI() override;
+	virtual void	InitRHI() override;
+	virtual void	ReleaseRHI() override;
 
 	/**
 	* Get the current pool buffer
diff --git a/Engine/Source/Runtime/Windows/D3D11RHI/Public/D3D11ConstantBuffer.h b/Engine/Source/Runtime/Windows/D3D11RHI/Public/D3D11ConstantBuffer.h
--- a/Engine/Source/Runtime/Windows/D3D11RHI/Public/D3D11ConstantBuffer.h
+++ b/Engine/Source/Runtime/Windows/D3D11RHI/Public/D3D11ConstantBuffer.h
@@ -37,8 +37,8 @@ public:
 	virtual ~FD3D11ConstantBuffer();
 
 	// FRenderResource interface.
-	virtual void	InitDynamicRHI() override;
-	virtual void	ReleaseDynamicRHI() override;
+	virtual void	InitRHI() override;
+	virtual void	ReleaseRHI() override;
 
 	static inline constexpr uint32 GetMaxSize() { return (uint32)Align(MAX_GLOBAL_CONSTANT_BUFFER_BYTE_SIZE, 16); }
 
