Fix mismatch is how global dynamic vertex buffer memory stat is updated.

diff --git a/Engine/Source/Runtime/RenderCore/Private/GlobalRenderResources.cpp b/Engine/Source/Runtime/RenderCore/Private/GlobalRenderResources.cpp
--- a/Engine/Source/Runtime/RenderCore/Private/GlobalRenderResources.cpp
+++ b/Engine/Source/Runtime/RenderCore/Private/GlobalRenderResources.cpp
@@ -784,7 +784,7 @@ struct FDynamicVertexBufferPool : public FRenderResource
 		if (bInitializeBuffer)
 		{
 			FoundBuffer->InitResource(RHICmdList);
-			TotalAllocatedMemory += SizeInBytes;
+			TotalAllocatedMemory += FoundBuffer->BufferSize;
 		}
 
 		FoundBuffer->Lock(RHICmdList);
