[shaders] fixes to preprocessing error message remapping/block recording functionality
- handle the case where filename-less line directives occur before any directives containing files. these will necessarily only occur in code explicitly appended by the shader system (specifically the case where minifier encounters a block of code needing to be persisted before it has processed any line directives) and so we don't worry about remapping error messages in such code (the filenames would be meaningless anyways).
- when remapping diagnostic messages handle the case where the sentinel filename occurs more than once (this is the case with psslc warnings  about shadowing previously declared variables)
- make HasDirective const (this was just an oversight)

diff --git a/Engine/Source/Runtime/RenderCore/Private/ShaderPreprocessTypes.cpp b/Engine/Source/Runtime/RenderCore/Private/ShaderPreprocessTypes.cpp
--- a/Engine/Source/Runtime/RenderCore/Private/ShaderPreprocessTypes.cpp
+++ b/Engine/Source/Runtime/RenderCore/Private/ShaderPreprocessTypes.cpp
@@ -10,16 +10,14 @@ static const FString LineDirectiveSentinel = FString::Printf(TEXT("#line 1 \"%s\
 
 void FShaderDiagnosticRemapper::Remap(FShaderCompilerError& Diagnostic) const
 {
-	if (!Diagnostic.ErrorLineString.IsEmpty() && !Diagnostic.ErrorVirtualFilePath.IsEmpty())
+	if (!Diagnostic.ErrorLineString.IsEmpty() && !Diagnostic.ErrorVirtualFilePath.IsEmpty() && Diagnostic.ErrorVirtualFilePath == FilenameSentinel)
 	{
-		if (Diagnostic.ErrorVirtualFilePath != FilenameSentinel)
-		{
-			// only attempt remapping if the filename sentinel matches; this will bypass remapping in the case where compilation was retried on FXC
-			// with a DXC precompile (this compilation will report error messages relative to an intermediate hlsl file, which is not a change in behaviour)
-			return;
-		}
+		// only attempt remapping if the filename sentinel matches; this will bypass remapping in the case where compilation was retried on FXC
+		// with a DXC precompile (this compilation will report error messages relative to an intermediate hlsl file, which is not a change in behaviour)
 
-		// Some backends give us already-parsed FShaderCompilerError objects which simplifies the remapping process somewhat
+		// Some backends give us already-parsed FShaderCompilerError objects, so in this case we need to remap the parsed 
+		// line number string & path (in addition to potentially needing to remap the message itself, in case the sentinel
+		// file name occurred in the original diagnostic multiple times)
 		const FString& ErrorLineStr = Diagnostic.ErrorLineString;
 		check(FChar::IsDigit(ErrorLineStr[0]));
 		int32 StrippedLineNum = 0, LineNumberEnd = 0;
@@ -38,47 +36,43 @@ void FShaderDiagnosticRemapper::Remap(FShaderCompilerError& Diagnostic) const
 		Diagnostic.ErrorLineString = MoveTemp(NewErrorLineString);
 		Diagnostic.ErrorVirtualFilePath = RemapData.Filename;
 	}
-	// need to parse/remap/reconstruct the full error message
-	else
-	{
-		const FString& Message = Diagnostic.StrippedErrorMessage;
-		int32 FilenameIndex = Message.Find(FilenameSentinel);
-		if (FilenameIndex < 0)
-		{
-			// only attempt remapping if we find our sentinel filename in the diagnostic message; otherwise it's likely
-			// a non-line-specific error that doesn't need remapping
-			return;
-		}
 
+	int32 FilenameIndex = Diagnostic.StrippedErrorMessage.Find(FilenameSentinel);
+	// if we find our sentinel filename in the diagnostic message it needs remapping
+	// need to loop in case the message contains the filename multiple times
+	while (FilenameIndex >= 0)
+	{
+		const FString& OriginalMessage = Diagnostic.StrippedErrorMessage;
 		// Parse the line number from the message and record start and end indices
 		// Code assumes that the next integral number found in the message string after the filename is always line number
 		int32 LineNumberStart = FilenameIndex + FilenameSentinelLen;
-		while (!FChar::IsDigit(Message[LineNumberStart]))
+		while (!FChar::IsDigit(OriginalMessage[LineNumberStart]))
 		{
 			++LineNumberStart;
 		}
 		int32 LineNumberEnd = LineNumberStart;
 		int32 StrippedLineNum = 0;
-		while (FChar::IsDigit(Message[LineNumberEnd]) && LineNumberEnd < Message.Len())
+		while (FChar::IsDigit(OriginalMessage[LineNumberEnd]) && LineNumberEnd < OriginalMessage.Len())
 		{
-			StrippedLineNum = StrippedLineNum * 10 + (Message[LineNumberEnd++] - TEXT('0'));
+			StrippedLineNum = StrippedLineNum * 10 + (OriginalMessage[LineNumberEnd++] - TEXT('0'));
 		}
 
 		FRemapData RemapData = GetRemapData(StrippedLineNum);
 
 		// Assume line number is the same number of digits for simplicity in reserve allocation size;
 		// it's an upper bound but in most practical cases worst case we're allocating an extra byte
-		int32 RemappedLen = Message.Len() - FilenameSentinelLen + RemapData.Filename.Len();
+		int32 RemappedLen = OriginalMessage.Len() - FilenameSentinelLen + RemapData.Filename.Len();
 		FString RemappedMessage;
 		RemappedMessage.Reserve(RemappedLen);
-		RemappedMessage.AppendChars(Message.GetCharArray().GetData(), FilenameIndex);
+		RemappedMessage.AppendChars(OriginalMessage.GetCharArray().GetData(), FilenameIndex);
 		RemappedMessage.Append(RemapData.Filename);
 		int32 SeparatorStart = FilenameIndex + FilenameSentinelLen;
-		RemappedMessage.Append(Message.GetCharArray().GetData() + SeparatorStart, LineNumberStart - SeparatorStart);
+		RemappedMessage.Append(OriginalMessage.GetCharArray().GetData() + SeparatorStart, LineNumberStart - SeparatorStart);
 		RemappedMessage.AppendInt(RemapData.LineNumber);
-		RemappedMessage.Append(Message.GetCharArray().GetData() + LineNumberEnd, Message.Len() - LineNumberEnd);
+		RemappedMessage.Append(OriginalMessage.GetCharArray().GetData() + LineNumberEnd, OriginalMessage.Len() - LineNumberEnd);
 
 		Diagnostic.StrippedErrorMessage = MoveTemp(RemappedMessage);
+		FilenameIndex = Diagnostic.StrippedErrorMessage.Find(FilenameSentinel);
 	}
 }
 
@@ -97,9 +91,15 @@ void FShaderDiagnosticRemapper::AddSourceBlock(int32 OriginalLineNum, int32 Stri
 	{
 		// sanity check that ranges are always added in order of stripped line numbers
 		check(Blocks.Num() == 0 || Blocks.Last().StrippedLineNum < StrippedLineNum);
-		// sanity check to ensure if we're not given a path we must have a previous block from which to retrieve the path
-		check(!OriginalPath.IsEmpty() || Blocks.Num() > 0);
-		Blocks.Add(FSourceBlock{ StrippedLineNum, OriginalLineNum, OriginalPath.IsEmpty() ? Blocks.Last().OriginalPath : MoveTemp(OriginalPath) });
+		// if we're not given a path we must have a previous block from which to retrieve the path;
+		// if we don't we don't bother adding the block. in practice this is only encountered in 
+		// code automatically appended to the source by the shader system (since the preprocessor
+		// always appends a line directive containing the file name when a file is loaded). we 
+		// assume such code will not contain errors.
+		if (!OriginalPath.IsEmpty() || Blocks.Num() > 0)
+		{
+			Blocks.Add(FSourceBlock{ StrippedLineNum, OriginalLineNum, OriginalPath.IsEmpty() ? Blocks.Last().OriginalPath : MoveTemp(OriginalPath) });
+		}
 	}
 }
 
@@ -122,6 +122,11 @@ FShaderDiagnosticRemapper::FRemapData FShaderDiagnosticRemapper::GetRemapData(in
 
 	const FSourceBlock& FoundBlock = Blocks[FoundIndex];
 
+	// Sanity check to validate the assumption made in AddSourceBlock - we should never be querying remap data for a line number
+	// that doesn't have an associated file path in the unstripped source; this implies a warning or error has been emitted in
+	// system-generated code and should be addressed by the developer prior to submission.
+	checkf(!FoundBlock.OriginalPath.IsEmpty(), TEXT("Unexpected compile error/warning found in system-generated code"));
+
 	// -1 when indexing stripped line offsets since line numbers are 1-based
 	return FRemapData{ FoundBlock.OriginalPath, FoundBlock.OriginalLineNum + StrippedLineOffsets[StrippedLineNum - 1] };
 }
diff --git a/Engine/Source/Runtime/RenderCore/Public/ShaderPreprocessTypes.h b/Engine/Source/Runtime/RenderCore/Public/ShaderPreprocessTypes.h
--- a/Engine/Source/Runtime/RenderCore/Public/ShaderPreprocessTypes.h
+++ b/Engine/Source/Runtime/RenderCore/Public/ShaderPreprocessTypes.h
@@ -76,7 +76,7 @@ public:
 		return PreprocessedSource;
 	}
 
-	inline bool HasDirective(const FString& Directive)
+	inline bool HasDirective(const FString& Directive) const
 	{
 		int32 NumberOfDirectives = PragmaDirectives.Num();
 		for (int32 i = 0; i < NumberOfDirectives; i++)
